[MotionDesign]
MovePlugin

#virtualized

diff --git a/Engine/Plugins/Experimental/Avalanche/Config/DefaultAvalanche.ini b/Engine/Plugins/Experimental/Avalanche/Config/DefaultAvalanche.ini
new file mode 100644
--- /dev/null
+++ b/Engine/Plugins/Experimental/Avalanche/Config/DefaultAvalanche.ini
@@ -0,0 +1,153 @@
+[CoreRedirects]
+
++ClassRedirects=(OldName="/Script/Avalanche.AvaAnimation",NewName="/Script/AvalancheSequence.AvaSequence")
++ClassRedirects=(OldName="/Script/Avalanche.AvaAnimationActor",NewName="/Script/AvalancheSequence.AvaSequencePlaybackActor")
+
++EnumRedirects=(OldName="/Script/Avalanche.EAvaAnimationPlayMode",NewName="/Script/AvalancheSequence.EAvaSequencePlayMode")
++EnumRedirects=(OldName="/Script/Avalanche.EAvaMarkRole",NewName="/Script/AvalancheSequence.EAvaMarkRole")
+
++StructRedirects=(OldName="/Script/Avalanche.AvaLevelSequenceBindingReference",NewName="/Script/AvalancheSequence.AvaAnimBindingReference")
++StructRedirects=(OldName="/Script/Avalanche.AvaLevelSequenceBindingReferenceArray",NewName="/Script/AvalancheSequence.AvaAnimBindingReferenceArray")
++StructRedirects=(OldName="/Script/Avalanche.AvaLevelSequenceBindingReferences",NewName="/Script/AvalancheSequence.AvaAnimBindingReferences")
++StructRedirects=(OldName="/Script/Avalanche.AvaAnimBindingReference",NewName="/Script/AvalancheSequence.AvaAnimBindingReference")
++StructRedirects=(OldName="/Script/Avalanche.AvaAnimBindingReferenceArray",NewName="/Script/AvalancheSequence.AvaAnimBindingReferenceArray")
++StructRedirects=(OldName="/Script/Avalanche.AvaAnimBindingReferences",NewName="/Script/AvalancheSequence.AvaAnimBindingReferences")
++StructRedirects=(OldName="/Script/Avalanche.AvaAnimPlayParams",NewName="/Script/AvalancheSequence.AvaSequencePlayParams")
+
++PropertyRedirects=(OldName="AvalancheRemoteControlValues.Values", NewName="AvalancheRemoteControlValues.EntityValues")
++PropertyRedirects=(OldName="AvaMediaOutputInfo.HostName", NewName="AvaMediaOutputInfo.ServerName")
++PropertyRedirects=(OldName="/Script/Avalanche.AvalanchePlaylist.Pages", NewName="Pages_DEPRECATED")
++PropertyRedirects=(OldName="/Script/AvalancheSceneTree.AvaSceneItem.ObjectPath", NewName="Id")
+
++ClassRedirects=(OldName="AvaTextBackgroundModifier",NewName="AvaContentBackgroundModifier")
++ClassRedirects=(OldName="/Script/AvalancheModifiers.AvaContentBackgroundModifier",NewName="/Script/AvalancheModifiers.AvaAutoSizeModifier")
++ClassRedirects=(OldName="/Script/AvalancheModifiers.AvaContentBackgroundModifierFactory",NewName="/Script/AvalancheModifiers.AvaAutoSizeModifierFactory")
++ClassRedirects=(OldName="/Script/AvalancheModifiersEditor.AvaContentBackgroundModifierExtension",NewName="/Script/AvalancheModifiersEditor.AvaAutoSizeModifierExtension")
++ClassRedirects=(OldName="/Script/Avalanche.AvaViewportCameraActor",NewName="/Script/Engine.CameraActor")
+
++ClassRedirects=(OldName="/Script/Avalanche.AvaClonerActor",NewName="/Script/AvalancheEffectors.AvaClonerActor")
++ClassRedirects=(OldName="/Script/Avalanche.AvaClonerComponent",NewName="/Script/AvalancheEffectors.AvaClonerComponent")
++ClassRedirects=(OldName="/Script/Avalanche.AvaEffectorActor",NewName="/Script/AvalancheEffectors.AvaEffectorActor")
++ClassRedirects=(OldName="/Script/Avalanche.AvaEffectorComponent",NewName="/Script/AvalancheEffectors.AvaEffectorComponent")
++ClassRedirects=(OldName="/Script/Avalanche.AvaSphereComponent",NewName="/Script/AvalancheEffectors.AvaSphereComponent")
++ClassRedirects=(OldName="/Script/AvalancheEffectors.AvaSphereComponent",NewName="/Script/Engine.SphereComponent")
++ClassRedirects=(OldName="/Script/AvalancheEffectors.AvaClonerSphereLayout",NewName="/Script/AvalancheEffectors.AvaClonerSphereUniformLayout")
++ClassRedirects=(OldName="/Script/Avalanche.MovieSceneAvaToolboxRectCornerSection",NewName="/Script/AvalancheShapes.MovieSceneAvaShapeRectCornerSection")
++ClassRedirects=(OldName="/Script/Avalanche.MovieSceneAvaToolboxRectCornerTrack",NewName="/Script/AvalancheShapes.MovieSceneAvaShapeRectCornerTrack")
++ClassRedirects=(OldName="/Script/Avalanche.AvaToolboxShapeActor",NewName="/Script/AvalancheShapes.AvaShapeActor")
++ClassRedirects=(OldName="/Script/Avalanche.AvaToolbox2DArrowDynamicMesh",NewName="/Script/AvalancheShapes.AvaShape2DArrowDynamicMesh")
++ClassRedirects=(OldName="/Script/Avalanche.AvaToolbox2DDynMeshBase",NewName="/Script/AvalancheShapes.AvaShape2DDynMeshBase")
++ClassRedirects=(OldName="/Script/Avalanche.AvaToolbox3DDynMeshBase",NewName="/Script/AvalancheShapes.AvaShape3DDynMeshBase")
++ClassRedirects=(OldName="/Script/Avalanche.AvaToolboxChevronDynamicMesh",NewName="/Script/AvalancheShapes.AvaShapeChevronDynamicMesh")
++ClassRedirects=(OldName="/Script/Avalanche.AvaToolboxConeDynamicMesh",NewName="/Script/AvalancheShapes.AvaShapeConeDynamicMesh")
++ClassRedirects=(OldName="/Script/Avalanche.AvaToolboxCubeDynamicMesh",NewName="/Script/AvalancheShapes.AvaShapeCubeDynamicMesh")
++ClassRedirects=(OldName="/Script/Avalanche.AvaToolboxDynamicMeshBase",NewName="/Script/AvalancheShapes.AvaShapeDynamicMeshBase")
++ClassRedirects=(OldName="/Script/Avalanche.AvaToolboxEllipseDynamicMesh",NewName="/Script/AvalancheShapes.AvaShapeEllipseDynamicMesh")
++ClassRedirects=(OldName="/Script/Avalanche.AvaToolboxIrregularPolygonDynamicMesh",NewName="/Script/AvalancheShapes.AvaShapeIrregularPolygonDynamicMesh")
++ClassRedirects=(OldName="/Script/Avalanche.AvaToolboxLineDynamicMesh",NewName="/Script/AvalancheShapes.AvaShapeLineDynamicMesh")
++ClassRedirects=(OldName="/Script/Avalanche.AvaToolboxNGonDynamicMesh",NewName="/Script/AvalancheShapes.AvaShapeNGonDynamicMesh")
++ClassRedirects=(OldName="/Script/Avalanche.AvaToolboxRectangleDynamicMesh",NewName="/Script/AvalancheShapes.AvaShapeRectangleDynamicMesh")
++ClassRedirects=(OldName="/Script/Avalanche.AvaToolboxRingDynamicMesh",NewName="/Script/AvalancheShapes.AvaShapeRingDynamicMesh")
++ClassRedirects=(OldName="/Script/Avalanche.AvaToolboxRoundedPolygonDynamicMesh",NewName="/Script/AvalancheShapes.AvaShapeRoundedPolygonDynamicMesh")
++ClassRedirects=(OldName="/Script/Avalanche.AvaToolboxSphereDynamicMesh",NewName="/Script/AvalancheShapes.AvaShapeSphereDynamicMesh")
++ClassRedirects=(OldName="/Script/Avalanche.AvaToolboxStarDynamicMesh",NewName="/Script/AvalancheShapes.AvaShapeStarDynamicMesh")
++ClassRedirects=(OldName="/Script/Avalanche.AvaToolboxTorusDynamicMesh",NewName="/Script/AvalancheShapes.AvaShapeTorusDynamicMesh")
++ClassRedirects=(OldName="/Script/Avalanche.AvalancheText3DActor",NewName="/Script/AvalancheText.AvaTextActor")
++ClassRedirects=(OldName="/Script/Avalanche.AvaText3DCharacterTransform",NewName="/Script/AvalancheText.AvaTextCharacterTransform")
++ClassRedirects=(OldName="/Script/Avalanche.AvaText3DComponent",NewName="/Script/AvalancheText.AvaText3DComponent")
++ClassRedirects=(OldName="/Script/Avalanche.AvalancheFontObject",NewName="/Script/AvalancheText.AvaFontObject")
++ClassRedirects=(OldName="/Script/Avalanche.AvaText3DMaterialHub",NewName="/Script/AvalancheText.AvaTextMaterialHub")
++ClassRedirects=(OldName="/Script/AvalancheText3D.AvalancheText3DActor",NewName="/Script/AvalancheText.AvaTextActor")
++ClassRedirects=(OldName="/Script/AvalancheText3D.AvaText3DCharacterTransform",NewName="/Script/AvalancheText.AvaTextCharacterTransform")
++ClassRedirects=(OldName="/Script/AvalancheText3D.AvaText3DComponent",NewName="/Script/AvalancheText.AvaText3DComponent")
++ClassRedirects=(OldName="/Script/AvalancheText3D.AvaText3DMaterialHub",NewName="/Script/AvalancheText.AvaTextMaterialHub")
++ClassRedirects=(OldName="/Script/AvalancheText3D.AvalancheFontObject",NewName="/Script/AvalancheText.AvaFontObject")
++ClassRedirects=(OldName="/Script/AvalanchePropertyControllers.PropertyControllerTextResolver",NewName="/Script/AvalanchePropertyControllers.AvaPropertyAnimatorTextResolver")
++ClassRedirects=(OldName="/Script/AvalanchePropertyControllers.PropertyControllerSequenceTimeSource",NewName="/Script/AvalanchePropertyControllers.AvaPropertyAnimatorSequenceTimeSource")
++ClassRedirects=(OldName="/Script/AvalancheModifiers.AvaCloneModifier",NewName="/Script/AvalancheModifiers.AvaPatternModifier")
+
++StructRedirects=(OldName="/Script/Avalanche.AvaGuideInfo",NewName="/Script/Avalanche.AvaViewportGuideInfo")
++StructRedirects=(OldName="/Script/Avalanche.AvaClonerAttachmentItem",NewName="/Script/AvalancheEffectors.AvaClonerAttachmentItem")
++StructRedirects=(OldName="/Script/Avalanche.AvaClonerAttachmentTree",NewName="/Script/AvalancheEffectors.AvaClonerAttachmentTree")
++StructRedirects=(OldName="/Script/Avalanche.AvaClonerEffectorDataInterfaces",NewName="/Script/AvalancheEffectors.AvaClonerEffectorDataInterfaces")
++StructRedirects=(OldName="/Script/Avalanche.AvaClonerSystem",NewName="/Script/AvalancheEffectors.AvaClonerSystem")
++StructRedirects=(OldName="/Script/Avalanche.AvaClonerSampleMeshOptions",NewName="/Script/AvalancheEffectors.AvaClonerSampleMeshOptions")
++StructRedirects=(OldName="/Script/Avalanche.AvaClonerSampleSplineOptions",NewName="/Script/AvalancheEffectors.AvaClonerSampleSplineOptions")
++StructRedirects=(OldName="/Script/Avalanche.AvaClonerGridConstraintSphere",NewName="/Script/AvalancheEffectors.AvaClonerGridConstraintSphere")
++StructRedirects=(OldName="/Script/Avalanche.AvaClonerGridConstraintCylinder",NewName="/Script/AvalancheEffectors.AvaClonerGridConstraintCylinder")
++StructRedirects=(OldName="/Script/Avalanche.AvaClonerGridConstraintTexture",NewName="/Script/AvalancheEffectors.AvaClonerGridConstraintTexture")
++StructRedirects=(OldName="/Script/Avalanche.AvaClonerGridLayoutOptions",NewName="/Script/AvalancheEffectors.AvaClonerGridLayoutOptions")
++StructRedirects=(OldName="/Script/Avalanche.AvaClonerLineLayoutOptions",NewName="/Script/AvalancheEffectors.AvaClonerLineLayoutOptions")
++StructRedirects=(OldName="/Script/Avalanche.AvaClonerCircleLayoutOptions",NewName="/Script/AvalancheEffectors.AvaClonerCircleLayoutOptions")
++StructRedirects=(OldName="/Script/Avalanche.AvaClonerCylinderLayoutOptions",NewName="/Script/AvalancheEffectors.AvaClonerCylinderLayoutOptions")
++StructRedirects=(OldName="/Script/Avalanche.AvaClonerSphereLayoutOptions",NewName="/Script/AvalancheEffectors.AvaClonerSphereLayoutOptions")
++StructRedirects=(OldName="/Script/Avalanche.AvaClonerHoneycombLayoutOptions",NewName="/Script/AvalancheEffectors.AvaClonerHoneycombLayoutOptions")
++StructRedirects=(OldName="/Script/Avalanche.MovieSceneAvaToolboxRectCornerSectionTemplate",NewName="/Script/AvalancheShapes.MovieSceneAvaShapeRectCornerSectionTemplate")
++StructRedirects=(OldName="/Script/Avalanche.AvaToolboxMaskMaterial",NewName="/Script/AvalancheShapes.AvaShapeMaskMaterial") -- ?!!
++StructRedirects=(OldName="/Script/Avalanche.AvaToolboxMesh",NewName="/Script/AvalancheShapes.AvaShapeMesh")
++StructRedirects=(OldName="/Script/Avalanche.AvaToolboxParametricMaterial",NewName="/Script/AvalancheShapes.AvaShapeParametricMaterial")
++StructRedirects=(OldName="/Script/Avalanche.AvaToolboxCachedVertex2D",NewName="/Script/AvalancheShapes.AvaShapeCachedVertex2D")
++StructRedirects=(OldName="/Script/Avalanche.AvaToolboxCachedVertex3D",NewName="/Script/AvalancheShapes.AvaShapeCachedVertex3D")
++StructRedirects=(OldName="/Script/Avalanche.AvaToolboxMaterialUVParameters",NewName="/Script/AvalancheShapes.AvaShapeMaterialUVParameters")
++StructRedirects=(OldName="/Script/Avalanche.AvaToolboxMeshData",NewName="/Script/AvalancheShapes.AvaShapeMeshData")
++StructRedirects=(OldName="/Script/Avalanche.AvaToolboxRoundedCornerSettings",NewName="/Script/AvalancheShapes.AvaShapeRoundedCornerSettings")
++StructRedirects=(OldName="/Script/Avalanche.AvaToolboxRoundedCorner",NewName="/Script/AvalancheShapes.AvaShapeRoundedCorner")
++StructRedirects=(OldName="/Script/Avalanche.AvaToolboxRectangleCornerSettings",NewName="/Script/AvalancheShapes.AvaShapeRectangleCornerSettings")
++StructRedirects=(OldName="/Script/Avalanche.AvaTextAlignment",NewName="/Script/AvalancheText.AvaTextAlignment")
++StructRedirects=(OldName="/Script/Avalanche.AvaLinearGradientSettings",NewName="/Script/AvalancheText.AvaLinearGradientSettings")
++StructRedirects=(OldName="/Script/Avalanche.AvaMaterialMaskSettings",NewName="/Script/AvalancheText.AvaMaterialMaskSettings")
++StructRedirects=(OldName="/Script/Avalanche.AvaTextField",NewName="/Script/AvalancheText.AvaTextField")
++StructRedirects=(OldName="/Script/Avalanche.AvalancheFont",NewName="/Script/AvalancheText.AvaFont")
++StructRedirects=(OldName="/Script/AvalancheText3D.AvalancheFont",NewName="/Script/AvalancheText.AvaFont")
++StructRedirects=(OldName="/Script/AvalancheText3D.AvaTextAlignment",NewName="/Script/AvalancheText.AvaTextAlignment")
++StructRedirects=(OldName="/Script/AvalancheText3D.AvaLinearGradientSettings",NewName="/Script/AvalancheText.AvaLinearGradientSettings")
++StructRedirects=(OldName="/Script/AvalancheText3D.AvaMaterialMaskSettings",NewName="/Script/AvalancheText.AvaMaterialMaskSettings")
++StructRedirects=(OldName="/Script/AvalancheText3D.AvaTextField",NewName="/Script/AvalancheText.AvaTextField")
++StructRedirects=(OldName="/Script/AvalancheModifiers.AvaCloneModifierCircleLayoutOptions",NewName="/Script/AvalancheModifiers.AvaPatternModifierCircleLayoutOptions")
++StructRedirects=(OldName="/Script/AvalancheModifiers.AvaCloneModifierGridLayoutOptions",NewName="/Script/AvalancheModifiers.AvaPatternModifierGridLayoutOptions")
++StructRedirects=(OldName="/Script/AvalancheModifiers.AvaCloneModifierLineLayoutOptions",NewName="/Script/AvalancheModifiers.AvaPatternModifierLineLayoutOptions")
+
++EnumRedirects=(OldName="/Script/Avalanche.EAvaToolboxGuideState",NewName="/Script/Avalanche.EAvaViewportGuideState")
++EnumRedirects=(OldName="/Script/Avalanche.EAvaClonerLayout",NewName="/Script/AvalancheEffectors.EAvaClonerLayout")
++EnumRedirects=(OldName="/Script/Avalanche.EAvaClonerAxis",NewName="/Script/AvalancheEffectors.EAvaClonerAxis")
++EnumRedirects=(OldName="/Script/Avalanche.EAvaClonerPlane",NewName="/Script/AvalancheEffectors.EAvaClonerPlane")
++EnumRedirects=(OldName="/Script/Avalanche.EAvaClonerMeshRenderMode",NewName="/Script/AvalancheEffectors.EAvaClonerMeshRenderMode")
++EnumRedirects=(OldName="/Script/Avalanche.EAvaClonerGridConstraint",NewName="/Script/AvalancheEffectors.EAvaClonerGridConstraint")
++EnumRedirects=(OldName="/Script/Avalanche.EAvaClonerEasing",NewName="/Script/AvalancheEffectors.EAvaClonerEasing")
++EnumRedirects=(OldName="/Script/Avalanche.EAvaClonerMeshAsset",NewName="/Script/AvalancheEffectors.EAvaClonerMeshAsset")
++EnumRedirects=(OldName="/Script/Avalanche.EAvaClonerMeshSampleData",NewName="/Script/AvalancheEffectors.EAvaClonerMeshSampleData")
++EnumRedirects=(OldName="/Script/Avalanche.EAvaClonerEffectorType",NewName="/Script/AvalancheEffectors.EAvaClonerEffectorType")
++EnumRedirects=(OldName="/Script/Avalanche.EAvaClonerEffectorMode",NewName="/Script/AvalancheEffectors.EAvaClonerEffectorMode")
++EnumRedirects=(OldName="/Script/Avalanche.EAvaClonerAttachmentStatus",NewName="/Script/AvalancheEffectors.EAvaClonerAttachmentStatus")
++EnumRedirects=(OldName="/Script/Avalanche.EAvaClonerTextureSampleChannel",NewName="/Script/AvalancheEffectors.EAvaClonerTextureSampleChannel")
++EnumRedirects=(OldName="/Script/Avalanche.EAvaClonerCompareMode",NewName="/Script/AvalancheEffectors.EAvaClonerCompareMode")
++EnumRedirects=(OldName="/Script/Avalanche.EAvaToolboxUVMode",NewName="/Script/AvalancheShapes.EAvaShapeUVMode")
++EnumRedirects=(OldName="/Script/Avalanche.EAvaToolboxCornerType",NewName="/Script/AvalancheShapes.EAvaShapeCornerType")
++EnumRedirects=(OldName="/Script/Avalanche.ParametricMaterialStyle",NewName="/Script/AvalancheShapes.EAvaShapeParametricMaterialStyle",ValueChanges=(("SOLID","Solid"),("LINEAR_GRADIENT","LinearGradient")))
++EnumRedirects=(OldName="/Script/Avalanche.EAvaDynamicMeshUpdateState",NewName="/Script/AvalancheShapes.EAvaDynamicMeshUpdateState")
++EnumRedirects=(OldName="/Script/Avalanche.EAvaOutlineType",NewName="/Script/AvalancheText.EAvaOutlineType")
++EnumRedirects=(OldName="/Script/Avalanche.EAvaTextColoringStyle",NewName="/Script/AvalancheText.EAvaTextColoringStyle")
++EnumRedirects=(OldName="/Script/Avalanche.EAvaMaterialMaskOrientation",NewName="/Script/AvalancheText.EAvaMaterialMaskOrientation")
++EnumRedirects=(OldName="/Script/Avalanche.EAvaGradientDirection",NewName="/Script/AvalancheText.EAvaGradientDirection")
++EnumRedirects=(OldName="/Script/Avalanche.EAvaTextLength",NewName="/Script/AvalancheText.EAvaTextLength")
++EnumRedirects=(OldName="/Script/Avalanche.EAvaTextCase",NewName="/Script/AvalancheText.EAvaTextCase")
++EnumRedirects=(OldName="/Script/Avalanche.EMaterialType",ValueChanges=(("DynamicMaterial","MaterialDesigner")))
++EnumRedirects=(OldName="/Script/AvalancheModifiers.EAvaCloneModifierPlane",NewName="/Script/AvalancheModifiers.EAvaPatternModifierPlane")
++EnumRedirects=(OldName="/Script/AvalancheModifiers.EAvaCloneModifierAxis",NewName="/Script/AvalancheModifiers.EAvaPatternModifierAxis")
++EnumRedirects=(OldName="/Script/AvalancheModifiers.EAvaCloneModifierLayout",NewName="/Script/AvalancheModifiers.EAvaPatternModifierLayout")
+
++PackageRedirects=(OldName="/Script/AvalancheText3D",NewName="/Script/AvalancheText")
++ClassRedirects=(OldName="/Script/AvalancheModifiers.AvaMaskModifier",NewName="/Script/AvalancheModifiers.AvaBooleanModifier")
++PropertyRedirects=(OldName="/Script/AvalancheEffectors.AvaClonerActor.ForceOrientationRate",NewName="/Script/AvalancheEffectors.AvaClonerActor.OrientationForceRate")
++PropertyRedirects=(OldName="/Script/AvalancheEffectors.AvaClonerActor.ForceOrientationMin",NewName="/Script/AvalancheEffectors.AvaClonerActor.OrientationForceMin")
++PropertyRedirects=(OldName="/Script/AvalancheEffectors.AvaClonerActor.ForceOrientationMax",NewName="/Script/AvalancheEffectors.AvaClonerActor.OrientationForceMax")
++PropertyRedirects=(OldName="/Script/AvalancheEffectors.AvaClonerActor.ForceVortexAmount",NewName="/Script/AvalancheEffectors.AvaClonerActor.VortexForceAmount")
++PropertyRedirects=(OldName="/Script/AvalancheEffectors.AvaClonerActor.ForceVortexAxis",NewName="/Script/AvalancheEffectors.AvaClonerActor.VortexForceAxis")
++PropertyRedirects=(OldName="/Script/AvalancheEffectors.AvaClonerActor.ForceCurlNoiseStrength",NewName="/Script/AvalancheEffectors.AvaClonerActor.CurlNoiseForceStrength")
++PropertyRedirects=(OldName="/Script/AvalancheEffectors.AvaClonerActor.ForceCurlNoiseFrequency",NewName="/Script/AvalancheEffectors.AvaClonerActor.CurlNoiseForceFrequency")
++PropertyRedirects=(OldName="/Script/AvalancheEffectors.AvaClonerActor.StepRotationDelta",NewName="/Script/AvalancheEffectors.AvaClonerActor.DeltaStepRotation")
++PropertyRedirects=(OldName="/Script/AvalancheEffectors.AvaClonerActor.StepScaleDelta",NewName="/Script/AvalancheEffectors.AvaClonerActor.DeltaStepScale")
++PropertyRedirects=(OldName="/Script/AvalancheEffectors.AvaEffectorActor.Strength",NewName="/Script/AvalancheEffectors.AvaEffectorActor.LocationStrength")
+
+[/Script/Avalanche.AvaAnimation]
+DefaultCompletionMode=KeepState
diff --git a/Engine/Plugins/Experimental/Avalanche/Shaders/DrawTextureWithAlpha.usf b/Engine/Plugins/Experimental/Avalanche/Shaders/DrawTextureWithAlpha.usf
new file mode 100644
--- /dev/null
+++ b/Engine/Plugins/Experimental/Avalanche/Shaders/DrawTextureWithAlpha.usf
@@ -0,0 +1,28 @@
+// Copyright Epic Games, Inc. All Rights Reserved.
+
+#include "/Engine/Private/Common.ush"
+
+Texture2D InputTexture;
+SamplerState InputSampler;
+
+float4 DrawTextureWithAlphaPS(
+	noperspective float4 UVAndScreenPos : TEXCOORD0
+	) : SV_Target0
+{
+	const float2 UV = UVAndScreenPos.xy;
+	const float4 SceneColorValue = Texture2DSample(InputTexture, InputSampler, UV);
+
+	// 	r.PostProcessing.PropagateAlpha
+	#if POST_PROCESS_ALPHA == 2 
+		return float4(SceneColorValue.rgb, saturate(1.0 - SceneColorValue.a));
+	#else
+		if (float(SceneColorValue.r) > 0.f || float(SceneColorValue.g) > 0.f || float(SceneColorValue.b) > 0.f)
+		{
+			return float4(SceneColorValue.rgb, 1.f);
+		}
+		else
+		{
+			return float4(0.f, 0.f, 0.f, 0.f);
+		}
+	#endif
+}
\ No newline at end of file
diff --git a/Engine/Plugins/Experimental/Avalanche/Source/Avalanche/Private/Archive/AvalancheArchive.cpp b/Engine/Plugins/Experimental/Avalanche/Source/Avalanche/Private/Archive/AvalancheArchive.cpp
new file mode 100644
--- /dev/null
+++ b/Engine/Plugins/Experimental/Avalanche/Source/Avalanche/Private/Archive/AvalancheArchive.cpp
@@ -0,0 +1,175 @@
+// Copyright Epic Games, Inc. All Rights Reserved.
+
+#include "AvalancheArchive.h"
+#include "AvaBlueprint.h"
+#include "AvaBlueprint_Serialize.h"
+#include "HAL/UnrealMemory.h"
+#include "UObject/PropertyPortFlags.h"
+#include "UObject/UnrealType.h"
+
+#define LOCTEXT_NAMESPACE "AvalancheArchive"
+
+DEFINE_LOG_CATEGORY_STATIC(LogAvalancheArchive, Log, All);
+
+FAvalancheArchive::FAvalancheArchive(FAvalancheWorldData& InWorldData, FAvalancheObjectData& InObjectData, UObject* InSerializedObject, bool bIsLoading)
+	: WorldData(InWorldData)
+	, ObjectData(InObjectData)
+	, SerializedObject(InSerializedObject)
+	, ExcludedPropertyFlags(CPF_BlueprintAssignable | CPF_Transient | CPF_Deprecated)
+{
+	Super::SetWantBinaryPropertySerialization(false);
+	Super::SetIsTransacting(false);
+	Super::SetIsPersistent(true);
+
+	ArNoDelta = true;
+
+	if (bIsLoading)
+	{
+		// Serialize properties that were valid in a previous version and are deprecated now. PostSerialize is responsible to migrate the data.
+		ExcludedPropertyFlags &= ~CPF_Deprecated;
+		SetPortFlags(PPF_UseDeprecatedProperties);
+
+		Super::SetIsLoading(true);
+		Super::SetIsSaving(false);
+
+		WorldData.VersionInfo.ApplyToArchive(*this);
+	}
+	else
+	{
+		Super::SetIsLoading(false);
+		Super::SetIsSaving(true);
+	}
+}
+
+FString FAvalancheArchive::GetArchiveName() const
+{
+	return TEXT("FAvalancheArchive");
+}
+
+int64 FAvalancheArchive::TotalSize()
+{
+	return ObjectData.SerializedData.Num();
+}
+
+int64 FAvalancheArchive::Tell()
+{
+	return DataIndex;
+}
+
+void FAvalancheArchive::Seek(int64 InPos)
+{
+	checkSlow(InPos <= TotalSize());
+	DataIndex = InPos;
+}
+
+bool FAvalancheArchive::ShouldSkipProperty(const FProperty* InProperty) const
+{
+	check(InProperty);
+
+	if (InProperty->HasAnyPropertyFlags(ExcludedPropertyFlags))
+	{
+		return true;
+	}
+
+	return false;
+}
+
+FArchive& FAvalancheArchive::operator<<(FName& Value)
+{
+	if (IsLoading())
+	{
+		FAvaObjectIndex NameIndex;
+		*this << NameIndex.Index;
+
+		if (!ensureAlwaysMsgf(WorldData.SerializedNames.IsValidIndex(NameIndex.Index), TEXT("Data appears to be corrupted")))
+		{
+			SetError();
+			return *this;
+		}
+		Value = WorldData.SerializedNames[NameIndex.Index];
+	}
+	else
+	{
+		FAvaObjectIndex NameIndex;
+		if (const FAvaObjectIndex* ExistingIndex = WorldData.NameIndexMap.Find(Value))
+		{
+			NameIndex = *ExistingIndex;
+		}
+		else
+		{
+			NameIndex.Index = WorldData.SerializedNames.Add(Value);
+			WorldData.NameIndexMap.Add(Value, NameIndex);
+		}
+
+		*this << NameIndex.Index;
+	}
+	return *this;
+}
+
+FArchive& FAvalancheArchive::operator<<(UObject*& Value)
+{
+	if (IsLoading())
+	{
+		FAvaObjectIndex ReferencedIndex;
+		*this << ReferencedIndex.Index;
+
+		if (!ensureAlwaysMsgf(WorldData.SerializedObjectReferences.IsValidIndex(ReferencedIndex.Index)
+			, TEXT("Data appears to be corrupted")))
+		{
+			SetError();
+			return *this;
+		}
+
+		const FSoftObjectPath& ObjectPath = WorldData.SerializedObjectReferences[ReferencedIndex.Index];
+		if (ObjectPath.IsNull())
+		{
+			Value = nullptr;
+			return *this;
+		}
+
+		Value = ResolveObjectDependency(ReferencedIndex);
+	}
+	else
+	{
+		FAvaObjectIndex ReferenceIndex = FAvaBlueprint_Serialize::AddObjectDependency(WorldData, Value);
+		*this << ReferenceIndex.Index;
+	}
+
+	return *this;
+}
+
+void FAvalancheArchive::Serialize(void* Data, int64 Length)
+{
+	if (Length <= 0)
+	{
+		return;
+	}
+
+	if (IsLoading())
+	{
+		if (!ensure(DataIndex + Length <= TotalSize()))
+		{
+			UE_LOG(LogAvalancheArchive, Error, TEXT("Unable to read %d bytes at index %d (Archive size: %d), missing %d bytes."),
+				Length, DataIndex, TotalSize(), DataIndex + Length - TotalSize());
+			SetError();
+			return;
+		}
+
+		FMemory::Memcpy(Data, &ObjectData.SerializedData[DataIndex], Length);
+		DataIndex += Length;
+	}
+	else
+	{
+		const int64 RequiredEndIndex = DataIndex + Length;
+		const int32 ToAlloc  = RequiredEndIndex - TotalSize();
+		if (ToAlloc > 0)
+		{
+			ObjectData.SerializedData.AddUninitialized(ToAlloc);
+		}
+
+		FMemory::Memcpy(&ObjectData.SerializedData[DataIndex], Data, Length);
+		DataIndex = RequiredEndIndex;
+	}
+}
+
+#undef LOCTEXT_NAMESPACE
diff --git a/Engine/Plugins/Experimental/Avalanche/Source/Avalanche/Private/Archive/AvalancheArchive.h b/Engine/Plugins/Experimental/Avalanche/Source/Avalanche/Private/Archive/AvalancheArchive.h
new file mode 100644
--- /dev/null
+++ b/Engine/Plugins/Experimental/Avalanche/Source/Avalanche/Private/Archive/AvalancheArchive.h
@@ -0,0 +1,50 @@
+// Copyright Epic Games, Inc. All Rights Reserved.
+
+#pragma once
+
+#include "Data/AvalancheDataDefines.h"
+#include "Serialization/ArchiveUObject.h"
+#include "UObject/ObjectMacros.h"
+
+struct FAvalancheObjectData;
+struct FAvalancheWorldData;
+class UObject;
+
+class FAvalancheArchive : public FArchiveUObject
+{
+	using Super = FArchiveUObject;
+	
+public:
+	
+	FAvalancheArchive(FAvalancheWorldData& InWorldData
+		, FAvalancheObjectData& InObjectData
+		, UObject* InSerializedObject
+		, bool bIsLoading);
+
+	// FArchive Interface
+	virtual FString GetArchiveName() const override;
+	virtual int64 TotalSize() override;
+	virtual int64 Tell() override;
+	virtual void Seek(int64 InPos) override;
+	virtual bool ShouldSkipProperty(const FProperty* InProperty) const override;
+	virtual FArchive& operator<<(FName& Value) override;
+	virtual FArchive& operator<<(UObject*& Value) override;
+	virtual void Serialize(void* Data, int64 Length) override;
+	// ~FArchive Interface
+
+	/* Allocates and serializes an object dependency, or gets the object, if it already exists. */
+	virtual UObject* ResolveObjectDependency(FAvaObjectIndex ObjectIndex) const = 0;
+	
+protected:
+
+	FAvalancheWorldData& WorldData;
+
+	FAvalancheObjectData& ObjectData;
+
+	UObject* SerializedObject;
+
+	/** Flags that would cause the Property to be excluded from the Archive */
+	EPropertyFlags ExcludedPropertyFlags = EPropertyFlags::CPF_None;
+
+	int64 DataIndex = 0;
+};
diff --git a/Engine/Plugins/Experimental/Avalanche/Source/Avalanche/Private/Archive/AvalancheReader.cpp b/Engine/Plugins/Experimental/Avalanche/Source/Avalanche/Private/Archive/AvalancheReader.cpp
new file mode 100644
--- /dev/null
+++ b/Engine/Plugins/Experimental/Avalanche/Source/Avalanche/Private/Archive/AvalancheReader.cpp
@@ -0,0 +1,15 @@
+// Copyright Epic Games, Inc. All Rights Reserved.
+
+#include "AvalancheReader.h"
+#include "AvaBlueprint_Serialize.h"
+
+FAvalancheReader::FAvalancheReader(FAvalancheWorldData& InWorldData, FAvalancheObjectData& InObjectData, UObject* InObject)
+	: Super(InWorldData, InObjectData, InObject, true)
+{
+	InObject->Serialize(*this);
+}
+
+UObject* FAvalancheReader::ResolveObjectDependency(FAvaObjectIndex ObjectIndex) const
+{
+	return FAvaBlueprint_Serialize::ResolveObjectDependency(WorldData, ObjectIndex);
+}
diff --git a/Engine/Plugins/Experimental/Avalanche/Source/Avalanche/Private/Archive/AvalancheReader.h b/Engine/Plugins/Experimental/Avalanche/Source/Avalanche/Private/Archive/AvalancheReader.h
new file mode 100644
--- /dev/null
+++ b/Engine/Plugins/Experimental/Avalanche/Source/Avalanche/Private/Archive/AvalancheReader.h
@@ -0,0 +1,18 @@
+// Copyright Epic Games, Inc. All Rights Reserved.
+
+#pragma once
+
+#include "AvalancheArchive.h"
+
+class FAvalancheReader : public FAvalancheArchive
+{
+	using Super = FAvalancheArchive;
+	
+public:
+	
+	FAvalancheReader(FAvalancheWorldData& InWorldData, FAvalancheObjectData& InObjectData, UObject* InObject);
+
+	// FAvalancheArchive Interface
+	virtual UObject* ResolveObjectDependency(FAvaObjectIndex ObjectIndex) const override;
+	// ~FAvalancheArchive Interface
+};
diff --git a/Engine/Plugins/Experimental/Avalanche/Source/Avalanche/Private/Archive/AvalancheWriter.cpp b/Engine/Plugins/Experimental/Avalanche/Source/Avalanche/Private/Archive/AvalancheWriter.cpp
new file mode 100644
--- /dev/null
+++ b/Engine/Plugins/Experimental/Avalanche/Source/Avalanche/Private/Archive/AvalancheWriter.cpp
@@ -0,0 +1,11 @@
+// Copyright Epic Games, Inc. All Rights Reserved.
+
+#include "AvalancheWriter.h"
+#include "Data/AvalancheObjectData.h"
+
+FAvalancheWriter::FAvalancheWriter(FAvalancheWorldData& InWorldData, FAvalancheObjectData& InObjectData, UObject* InObject)
+	: Super(InWorldData, InObjectData, InObject, false)
+{
+	InObject->Serialize(*this);
+	InObjectData.ObjectFlags = InObject->GetFlags();
+}
diff --git a/Engine/Plugins/Experimental/Avalanche/Source/Avalanche/Private/Archive/AvalancheWriter.h b/Engine/Plugins/Experimental/Avalanche/Source/Avalanche/Private/Archive/AvalancheWriter.h
new file mode 100644
--- /dev/null
+++ b/Engine/Plugins/Experimental/Avalanche/Source/Avalanche/Private/Archive/AvalancheWriter.h
@@ -0,0 +1,23 @@
+// Copyright Epic Games, Inc. All Rights Reserved.
+
+#pragma once
+
+#include "AvalancheArchive.h"
+#include "Misc/AssertionMacros.h"
+
+class FAvalancheWriter : public FAvalancheArchive
+{
+	using Super = FAvalancheArchive;
+
+public:
+	
+	FAvalancheWriter(FAvalancheWorldData& InWorldData, FAvalancheObjectData& InObjectData, UObject* InObject);
+
+	// FAvalancheArchive Interface
+	virtual UObject* ResolveObjectDependency(FAvaObjectIndex ObjectIndex) const override
+	{
+		checkNoEntry();
+		return nullptr;
+	}
+	// ~FAvalancheArchive Interface
+};
diff --git a/Engine/Plugins/Experimental/Avalanche/Source/Avalanche/Private/AvaActor.cpp b/Engine/Plugins/Experimental/Avalanche/Source/Avalanche/Private/AvaActor.cpp
new file mode 100644
--- /dev/null
+++ b/Engine/Plugins/Experimental/Avalanche/Source/Avalanche/Private/AvaActor.cpp
@@ -0,0 +1,38 @@
+// Copyright Epic Games, Inc. All Rights Reserved.
+
+#include "AvaActor.h"
+#include "AvaBlueprintGeneratedClass.h"
+#include "AvaSequencePlayer.h"
+
+AAvalancheActor::AAvalancheActor()
+{
+}
+
+bool AAvalancheActor::UpdateGeneratedClass(bool bForceUpdate)
+{
+	if ((bGeneratedClassUpdated && !bForceUpdate) || HasAnyFlags(RF_ClassDefaultObject))
+	{
+		return false;
+	}
+
+	UAvaBlueprintGeneratedClass* GeneratedClass = Cast<UAvaBlueprintGeneratedClass>(GetClass());
+	if (GeneratedClass)
+	{
+		GeneratedClass->UpdateProperties(this);
+	}
+
+	bGeneratedClassUpdated = true;
+	return true;
+}
+
+void AAvalancheActor::PostActorCreated()
+{
+	Super::PostActorCreated();
+	UpdateGeneratedClass();
+}
+
+void AAvalancheActor::PostDuplicate(EDuplicateMode::Type DuplicateMode)
+{
+	Super::PostDuplicate(DuplicateMode);
+	UpdateGeneratedClass();
+}
diff --git a/Engine/Plugins/Experimental/Avalanche/Source/Avalanche/Private/AvaActorUtils.cpp b/Engine/Plugins/Experimental/Avalanche/Source/Avalanche/Private/AvaActorUtils.cpp
new file mode 100644
--- /dev/null
+++ b/Engine/Plugins/Experimental/Avalanche/Source/Avalanche/Private/AvaActorUtils.cpp
@@ -0,0 +1,267 @@
+// Copyright Epic Games, Inc. All Rights Reserved.
+
+#include "AvaActorUtils.h"
+#include "AvaBlueprint.h"
+#include "AvaDefs.h"
+#include "AvaSceneSubsystem.h"
+#include "Components/PrimitiveComponent.h"
+#include "Components/SceneComponent.h"
+#include "Engine/World.h"
+#include "EngineUtils.h"
+#include "Framework/AvalancheGameInstance.h"
+#include "GameFramework/Actor.h"
+#include "Math/MathFwd.h"
+#include "Math/OrientedBox.h"
+
+#if WITH_EDITOR
+#include "AvaOutliner.h"
+#include "AvaOutlinerDefines.h"
+#include "AvaOutlinerSubsystem.h"
+#include "AvaOutlinerUtils.h"
+#include "Item/AvaOutlinerTreeRoot.h"
+#endif
+
+FOrientedBox FAvaActorUtils::MakeOrientedBox(const FBox& InLocalBox, const FTransform& InWorldTransform)
+{
+	FOrientedBox OutOrientedBox;
+
+	OutOrientedBox.Center = InWorldTransform.TransformPosition(InLocalBox.GetCenter());
+
+	OutOrientedBox.AxisX = InWorldTransform.TransformVector(FVector::UnitX());
+	OutOrientedBox.AxisY = InWorldTransform.TransformVector(FVector::UnitY());
+	OutOrientedBox.AxisZ = InWorldTransform.TransformVector(FVector::UnitZ());
+
+	OutOrientedBox.ExtentX = (InLocalBox.Max.X - InLocalBox.Min.X) / 2.f;
+	OutOrientedBox.ExtentY = (InLocalBox.Max.Y - InLocalBox.Min.Y) / 2.f;
+	OutOrientedBox.ExtentZ = (InLocalBox.Max.Z - InLocalBox.Min.Z) / 2.f;
+
+	return OutOrientedBox;
+}
+
+FBox FAvaActorUtils::GetActorLocalBoundingBox(const AActor* InActor, bool bIncludeFromChildActors, bool bMustBeRegistered)
+{
+	FBox Box(ForceInit);
+	Box.IsValid = 0;
+
+	if (!InActor || !InActor->GetRootComponent())
+	{
+		return Box;
+	}
+
+	FTransform ActorToWorld = InActor->GetTransform();
+	ActorToWorld.SetScale3D(FVector::OneVector);
+	const FTransform WorldToActor = ActorToWorld.Inverse();
+
+	uint32 FailedComponentCount = 0;
+	InActor->ForEachComponent<UPrimitiveComponent>(bIncludeFromChildActors, [&FailedComponentCount, bMustBeRegistered, &WorldToActor, &Box](const UPrimitiveComponent* InPrimComp)
+		{
+#if WITH_EDITOR
+			// Ignore Visualization Components, but don't consider them as failed components.
+			if (InPrimComp->IsVisualizationComponent())
+			{
+				return;
+			}
+#endif
+
+	if (InPrimComp->IsRegistered() || !bMustBeRegistered)
+	{
+		const FTransform ComponentToActor = InPrimComp->GetComponentTransform() * WorldToActor;
+		Box += InPrimComp->CalcBounds(ComponentToActor).GetBox();
+	}
+	else
+	{
+		FailedComponentCount++;
+	}
+		});
+
+	// Actors with no Failed Primitives should still return a valid Box with no Extents and 0,0,0 origin (local).
+	if (FailedComponentCount == 0)
+	{
+		Box.IsValid = 1;
+	}
+
+	return Box;
+}
+
+FBox FAvaActorUtils::GetComponentLocalBoundingBox(const USceneComponent* InComponent)
+{
+	FBox Box(ForceInit);
+	Box.IsValid = 0;
+
+	if (!InComponent)
+	{
+		return Box;
+	}
+
+	if (!InComponent->IsRegistered())
+	{
+		return Box;
+	}
+
+#if WITH_EDITOR
+	if (InComponent->IsVisualizationComponent())
+	{
+		return Box;
+	}
+#endif
+
+	// Pre-scale component to be consistent with actor bounding boxes
+	const FTransform ComponentTransform = FTransform(FRotator::ZeroRotator, FVector::ZeroVector, InComponent->GetComponentScale());
+	const FBoxSphereBounds BoxSphereBounds = InComponent->CalcBounds(ComponentTransform);
+	Box = BoxSphereBounds.GetBox();
+	Box.IsValid = 1;
+
+	return Box;
+}
+
+IAvaSceneInterface* FAvaActorUtils::GetSceneInterfaceFromActor(const AActor* InActor)
+{
+	if (!InActor)
+	{
+		return nullptr;
+	}
+
+	const UWorld* World = InActor->GetWorld();
+	if (!World)
+	{
+		return nullptr;
+	}
+
+	const UAvaSceneSubsystem* SceneSubsystem = World->GetSubsystem<UAvaSceneSubsystem>();
+	if (!SceneSubsystem)
+	{
+		return nullptr;
+	}
+
+	return SceneSubsystem->GetSceneInterface(InActor->GetLevel());
+}
+
+AActor* FAvaActorUtils::ActorFromReferenceContainer(AActor* const InActor, const EAvaReferenceContainer InReferenceContainer, const bool bInIgnoreHiddenActors)
+{
+	if (InReferenceContainer == EAvaReferenceContainer::Other || !IsValid(InActor))
+	{
+		return nullptr;
+	}
+
+	// Note: Using Typed Outer World, instead of GetWorld(), since the typed outer World could be a streamed in world
+	// and GetWorld() only returns the main world.
+	UWorld* const World = InActor->GetTypedOuter<UWorld>();
+	if (!IsValid(World))
+	{
+		return nullptr;
+	}
+
+	AActor* const ParentActor = InActor->GetAttachParentActor();
+
+	bool bIsOutlinerAttachedActors = false;
+	TArray<AActor*> AttachedActors;
+
+#if WITH_EDITOR
+	UAvaOutlinerSubsystem* const OutlinerSubsystem = World->GetSubsystem<UAvaOutlinerSubsystem>();
+	if (IsValid(OutlinerSubsystem))
+	{
+		TSharedPtr<FAvaOutliner> AvaOutliner = OutlinerSubsystem->GetOutliner();
+		if (AvaOutliner.IsValid())
+		{
+			bIsOutlinerAttachedActors = true;
+			if (IsValid(ParentActor))
+			{
+				AttachedActors = FAvaOutlinerUtils::EditorOutlinerChildActors(AvaOutliner, ParentActor);
+			}
+			else // no valid parent, use world actors instead
+			{
+				TArray<FAvaOutlinerItemPtr> OutlinerRootChildren = AvaOutliner->GetTreeRoot()->GetChildren();
+				const bool bValidRootActors = FAvaOutlinerUtils::EditorOutlinerItemsToActors(OutlinerRootChildren, AttachedActors);
+			}
+		}
+	}
+#endif
+
+	if (!bIsOutlinerAttachedActors)
+	{
+		if (const IAvaSceneInterface* SceneInterface = FAvaActorUtils::GetSceneInterfaceFromActor(InActor))
+		{
+			const FAvaSceneTree& SceneTree = SceneInterface->GetSceneTree();
+
+			if (!IsValid(ParentActor))
+			{
+				for (const int32 ChildId : SceneTree.GetRootNode().GetChildrenIndices())
+				{
+					const FAvaSceneItem* Item = SceneTree.GetItemAtIndex(ChildId);
+					AttachedActors.Add(Item->Resolve<AActor>(World));
+				}
+			}
+			else
+			{
+				AttachedActors = SceneTree.GetChildActors(ParentActor);
+			}
+		}
+	}
+
+	if (AttachedActors.Num() == 0)
+	{
+		return nullptr;
+	}
+
+	if (bInIgnoreHiddenActors)
+	{
+		for (int32 ChildIndex = AttachedActors.Num() - 1; ChildIndex >= 0; --ChildIndex)
+		{
+			if (AttachedActors[ChildIndex]->IsHidden()
+#if WITH_EDITOR
+				|| AttachedActors[ChildIndex]->IsTemporarilyHiddenInEditor()
+#endif
+				)
+			{
+				AttachedActors.RemoveAt(ChildIndex);
+			}
+		}
+	}
+
+	int32 FromActorIndex = -1;
+	if (!AttachedActors.Find(InActor, FromActorIndex))
+	{
+		return nullptr;
+	}
+
+	switch (InReferenceContainer)
+	{
+		case EAvaReferenceContainer::Previous:
+		{
+			const int32 PreviousActorIndex = FromActorIndex - 1;
+			if (!AttachedActors.IsValidIndex(PreviousActorIndex))
+			{
+				return ParentActor;
+			}
+			return AttachedActors[PreviousActorIndex];
+		}
+		case EAvaReferenceContainer::Next:
+		{
+			const int32 NextActorIndex = FromActorIndex + 1;
+			if (!AttachedActors.IsValidIndex(NextActorIndex))
+			{
+				return nullptr;
+			}
+			return AttachedActors[NextActorIndex];
+		}
+		case EAvaReferenceContainer::First:
+		{
+			if (!AttachedActors.IsValidIndex(0))
+			{
+				return nullptr;
+			}
+			return AttachedActors[0] != InActor ? AttachedActors[0] : nullptr;
+		}
+		case EAvaReferenceContainer::Last:
+		{
+			AActor* const LastChildActor = AttachedActors.Last();
+			if (!IsValid(LastChildActor))
+			{
+				return nullptr;
+			}
+			return LastChildActor != InActor ? LastChildActor : nullptr;
+		}
+	}
+
+	return nullptr;
+}
diff --git a/Engine/Plugins/Experimental/Avalanche/Source/Avalanche/Private/AvaBlueprint.cpp b/Engine/Plugins/Experimental/Avalanche/Source/Avalanche/Private/AvaBlueprint.cpp
new file mode 100644
--- /dev/null
+++ b/Engine/Plugins/Experimental/Avalanche/Source/Avalanche/Private/AvaBlueprint.cpp
@@ -0,0 +1,743 @@
+// Copyright Epic Games, Inc. All Rights Reserved.
+
+#include "AvaBlueprint.h"
+#include "Algo/ForEach.h"
+#include "Async/Async.h"
+#include "AudioDevice.h"
+#include "AvaActor.h"
+#include "AvaBlueprintGeneratedClass.h"
+#include "AvaBlueprintVersion.h"
+#include "AvaBlueprint_Serialize.h"
+#include "AvaRemoteControlRebind.h"
+#include "AvaRemoteControlUtils.h"
+#include "AvaSceneSubsystem.h"
+#include "AvaSequence.h"
+#include "AvaSequencePlaybackObject.h"
+#include "AvaSequencePlayer.h"
+#include "AvaSequenceSubsystem.h"
+#include "Camera/CameraActor.h"
+#include "Engine/SimpleConstructionScript.h"
+#include "EngineUtils.h"
+#include "IRemoteControlModule.h"
+#include "Misc/ScopedSlowTask.h"
+#include "RemoteControlActor.h"
+#include "RemoteControlPreset.h"
+#include "Templates/UnrealTemplate.h"
+#include "UObject/ObjectSaveContext.h"
+#include "UObject/UObjectThreadContext.h"
+
+#if WITH_EDITOR
+#include "Editor.h"
+#include "Kismet2/BlueprintEditorUtils.h"
+#include "Viewport/Interaction/AvaViewportGuide.h"
+#endif
+
+#define LOCTEXT_NAMESPACE "AvalancheBlueprint"
+
+DEFINE_LOG_CATEGORY_STATIC(LogAvalancheBlueprint, Log, All);
+
+UAvalancheBlueprint::UAvalancheBlueprint()
+{
+	BlueprintType = EBlueprintType::BPTYPE_Normal;
+
+	RemoteControlPreset = CreateDefaultSubobject<URemoteControlPreset>("RemoteControlPreset");
+}
+
+FTopLevelAssetPath UAvalancheBlueprint::GetWorldContextPath(const FSoftObjectPath& InSourceAssetPath)
+{
+	constexpr const TCHAR* WorldName     = TEXT("AvalancheWorld"); 
+	constexpr const TCHAR* PackagePrefix = TEXT("/Temp/AvaEditor/TransientPackage");
+
+	FName PackageName = *(PackagePrefix / InSourceAssetPath.GetLongPackageName());
+	return FTopLevelAssetPath(PackageName, WorldName);
+}
+
+void UAvalancheBlueprint::BeginDestroy()
+{
+	UAvaSequencePlayer::OnSequenceFinished().RemoveAll(this);
+	
+	UnregisterRemoteControlPreset();
+
+	Super::BeginDestroy();
+}
+
+void UAvalancheBlueprint::SetAvalancheWorld(UWorld* InWorld)
+{
+	World = InWorld;
+
+	if (RemoteControlPreset)
+	{
+		RemoteControlPreset->RebindUnboundEntities();
+	}
+
+	if (UAvaSceneSubsystem* SceneSubsystem = World->GetSubsystem<UAvaSceneSubsystem>())
+	{
+		SceneSubsystem->RegisterSceneInterface(World->PersistentLevel, this);
+	}
+}
+
+UWorld* UAvalancheBlueprint::GetAvalancheWorld() const
+{
+	return World;
+}
+
+AAvalancheActor* UAvalancheBlueprint::GetPlaceholderActor() const
+{
+#if WITH_EDITOR
+	// If Placeholder Actor is Invalid, or Worlds do not match, try to find the Placeholder Actor
+	if (ensure(World) && (!IsValid(PlaceholderActor) || PlaceholderActor->GetWorld() != World))
+	{
+		UAvalancheBlueprint* const MutableThis = const_cast<UAvalancheBlueprint*>(this);
+		MutableThis->PlaceholderActor = nullptr;
+
+		for (AAvalancheActor* const Actor : TActorRange<AAvalancheActor>(World))
+		{
+			if (Actor->GetClass()->ClassGeneratedBy == this)
+			{
+				MutableThis->PlaceholderActor = Actor;
+				break;
+			}
+		}
+	}
+#endif
+	return PlaceholderActor;
+}
+
+IAvaSequencePlaybackObject* UAvalancheBlueprint::GetScenePlayback() const
+{
+	// If the saved scene playback is valid use that
+	if (IAvaSequencePlaybackObject* const ScenePlayback = PlaybackObject.GetInterface())
+	{
+		return ScenePlayback;
+	}
+
+	UAvaSequenceSubsystem* const SequenceSubsystem = UAvaSequenceSubsystem::Get(World);
+	if (!SequenceSubsystem)
+	{
+		return nullptr;
+	}
+
+	UAvalancheBlueprint& MutableThis = *const_cast<UAvalancheBlueprint*>(this);
+	if (IAvaSequencePlaybackObject* const ScenePlayback = SequenceSubsystem->FindOrCreatePlaybackObject(World->PersistentLevel, MutableThis))
+	{
+		MutableThis.PlaybackObject.SetObject(ScenePlayback->ToUObject());
+		MutableThis.PlaybackObject.SetInterface(ScenePlayback);
+		return ScenePlayback;
+	}
+
+	return nullptr;
+}
+
+void UAvalancheBlueprint::DestroyPlaceholderActor()
+{
+	if (PlaceholderActor && World)
+	{
+		World->DestroyActor(PlaceholderActor, false, false);
+		World->BroadcastLevelsChanged();
+	}
+
+#if WITH_EDITOR
+	if (SimpleConstructionScript && PlaceholderActor == SimpleConstructionScript->GetComponentEditorActorInstance())
+	{
+		// Ensure that all editable component references are cleared
+		SimpleConstructionScript->ClearEditorComponentReferences();
+
+		// Clear the reference to the actor instance
+		SimpleConstructionScript->SetComponentEditorActorInstance(nullptr);
+	}
+#endif
+
+	PlaceholderActor = nullptr;
+}
+
+void UAvalancheBlueprint::UpdatePlaceholderActor(bool bInForceFullUpdate)
+{
+	AAvalancheActor* AvalancheActor = GetPlaceholderActor();
+
+	// if Actor is already valid and we are not doing a full update, skip
+	if (AvalancheActor && !bInForceFullUpdate)
+	{
+		return;
+	}
+
+#if WITH_EDITOR
+	// Signal that we're going to be constructing editor components
+	if (SimpleConstructionScript)
+	{
+		SimpleConstructionScript->BeginEditorComponentConstruction();
+	}
+#endif
+
+	// Spawn Placeholder Actor if current is invalid, or we're forcing full update
+	if (!AvalancheActor || bInForceFullUpdate)
+	{
+		// Destroy the previous actor instance
+		DestroyPlaceholderActor();
+
+		// Spawn a new placeholder actor based on the Blueprint's generated class if it's Actor-based
+		if (ensureMsgf(GeneratedClass && GeneratedClass->IsChildOf(AAvalancheActor::StaticClass())
+			, TEXT("Generated Class (%s) is not based on Avalanche Actor."), GeneratedClass ? *GeneratedClass->GetName() : TEXT("invalid")))
+		{
+			// Spawn an Actor based on the Blueprint's generated class
+			FActorSpawnParameters SpawnInfo;
+			SpawnInfo.SpawnCollisionHandlingOverride = ESpawnActorCollisionHandlingMethod::AlwaysSpawn;
+			SpawnInfo.ObjectFlags = RF_Transient | RF_Transactional;
+			SpawnInfo.bNoFail = true;
+
+			check(World);
+			{
+#if WITH_EDITOR
+				FMakeClassSpawnableOnScope TemporarilySpawnable(GeneratedClass);
+#endif
+				AvalancheActor   = World->SpawnActor<AAvalancheActor>(GeneratedClass, FVector::ZeroVector, FRotator::ZeroRotator, SpawnInfo);
+				PlaceholderActor = AvalancheActor;
+			}
+			check(AvalancheActor);
+
+			// Ensure that the actor is invisible as Placeholder
+			AvalancheActor->SetHidden(true);
+#if WITH_EDITOR
+			AvalancheActor->SetIsTemporarilyHiddenInEditor(true);
+#endif
+
+			// Prevent any audio from playing as a result of spawning
+			if (FAudioDeviceHandle AudioDevice = GEngine->GetMainAudioDevice())
+			{
+				AudioDevice->Flush(World);
+			}
+
+#if WITH_EDITOR
+			// Set the reference to the preview actor for component editing purposes
+			if (SimpleConstructionScript)
+			{
+				SimpleConstructionScript->SetComponentEditorActorInstance(AvalancheActor);
+			}
+#endif
+		}
+	}
+	else
+	{
+		check(AvalancheActor);
+		AvalancheActor->ReregisterAllComponents();
+#if WITH_EDITOR
+		AvalancheActor->RerunConstructionScripts();
+#endif
+	}
+
+#if WITH_EDITOR
+	// Signal that we're done constructing editor components
+	if (SimpleConstructionScript)
+	{
+		SimpleConstructionScript->EndEditorComponentConstruction();
+	}
+#endif
+}
+
+void UAvalancheBlueprint::PreDuplicate(FObjectDuplicationParameters& InDuplicationParameters)
+{
+	Super::PreDuplicate(InDuplicationParameters);
+	SaveAvalancheWorld();
+}
+
+void UAvalancheBlueprint::Serialize(FArchive& Ar)
+{
+	Ar.UsingCustomVersion(FAvalancheBlueprintVersion::GUID);
+	
+	Super::Serialize(Ar);
+}
+
+#if WITH_EDITORONLY_DATA
+void UAvalancheBlueprint::PreSave(FObjectPreSaveContext ObjectSaveContext)
+{
+	SaveAvalancheWorld();
+	Super::PreSave(ObjectSaveContext);
+}
+
+void UAvalancheBlueprint::PostSaveRoot(FObjectPostSaveRootContext ObjectSaveContext)
+{
+	Super::PostSaveRoot(ObjectSaveContext);
+	OnAvalanchePostSave.Broadcast();
+}
+#endif
+
+void UAvalancheBlueprint::SaveAvalancheWorld()
+{
+#if WITH_EDITOR
+	OnAvalanchePreSave.Broadcast();
+#endif
+
+	WorldData.World = GetAvalancheWorld();
+	if (!WorldData.World.IsValid())
+	{
+		return;
+	}
+
+	TGuardValue SaveGuard(bSavingWorld, true);
+
+#if WITH_EDITOR
+	FScopedSlowTask SaveWorldTask(2.f, LOCTEXT("SaveAvalancheWorld", "Saving Avalanche World"));
+	SaveWorldTask.MakeDialogDelayed(1.f);	
+#endif
+	
+	TArray<AActor*> ActorsInWorld;
+
+	//Get all Actors from all Levels
+	{
+		int32 NumActors = 0;
+		for (ULevel* const Level : WorldData.World->GetLevels())
+		{
+			if (Level)
+			{
+				NumActors += Level->Actors.Num();
+			}
+		}
+			
+		ActorsInWorld.Reserve(NumActors);
+			
+		for (ULevel* const Level : WorldData.World->GetLevels())
+		{
+			if (Level)
+			{
+				ActorsInWorld.Append(Level->Actors);
+			}
+		}
+	}
+
+	FAvalancheWorldData NewWorldData;
+	
+#if WITH_EDITOR	
+	SaveWorldTask.EnterProgressFrame(1.f);
+#endif
+	
+	//Save Actors
+	{
+		FScopedSlowTask SaveActorsTask(ActorsInWorld.Num(), LOCTEXT("SavingActors", "Saving Actors"));
+		SaveActorsTask.MakeDialogDelayed(1.f);
+			
+		Algo::ForEach(ActorsInWorld, [&NewWorldData, this, &SaveActorsTask](AActor* Actor)
+		{
+			SaveActorsTask.EnterProgressFrame();
+			if (FAvaBlueprint_Serialize::ShouldSaveActor(WorldData.World.Get(), Actor))
+			{
+				NewWorldData.ActorData.Add(Actor, FAvaBlueprint_Serialize::SaveActor(Actor, NewWorldData));
+			}
+		});
+	}
+
+	NewWorldData.FinalizeSave();
+
+#if WITH_EDITOR		
+	SaveWorldTask.EnterProgressFrame(1.f);
+	//Internal::ComputeActorHashes(ActorsInWorld, SnapshotData);
+	SetCreateDefaultScene(false);
+#endif
+
+	WorldData = MoveTemp(NewWorldData);
+}
+
+void UAvalancheBlueprint::LoadAvalancheWorld(UWorld* InOverrideWorld, TOptional<FInitActorFunctionRef> InInitActorFunction)
+{
+	WorldData.World = InOverrideWorld ? InOverrideWorld : GetAvalancheWorld();
+	if (!WorldData.World.IsValid())
+	{
+		return;
+	}
+
+	TGuardValue LoadGuard(bLoadingWorld, true);
+	
+#if WITH_EDITOR
+	FScopedSlowTask LoadWorld(WorldData.ActorData.Num(), LOCTEXT("LoadWorld", "Loading Avalanche World"));
+	LoadWorld.MakeDialogDelayed(1.f, false);
+#endif
+	
+	TMap<FSoftObjectPath, AActor*> RecreatedActors;
+	RecreatedActors.Reserve(WorldData.ActorData.Num());
+
+	const bool bIsEditorWorld = WorldData.World->IsEditorWorld();
+	
+	//1st Pass: Allocate/Spawn the Actors. Serialization is done in separate pass so Object References can be resolved correctly.
+	for (const TPair<FSoftObjectPath, FAvalancheActorData>& Pair : WorldData.ActorData)
+	{
+		const FSoftObjectPath& ActorPath = Pair.Key;
+		const FAvalancheActorData& ActorData = Pair.Value;
+		
+		UClass* ActorClass = ActorData.ActorClass.TryLoadClass<AActor>();
+		if (!ActorClass)
+		{
+			UE_LOG(LogAvalancheBlueprint
+				, Warning
+				, TEXT("Failed to resolve class '%s'. Was it removed?")
+				, *ActorData.ActorClass.ToString());
+			continue;
+		}
+	
+		//Handle Name Clash
+		if (UObject* const FoundObject = FindObject<UObject>(WorldData.World.Get(), *ActorPath.ToString()))
+		{			
+			// If it's not an actor then it's possibly a UObjectRedirector
+			AActor* const FoundActor = Cast<AActor>(FoundObject);
+			if (IsValid(FoundActor))
+			{
+#if WITH_EDITOR
+				if (bIsEditorWorld)
+				{
+					GEditor->SelectActor(FoundActor, /*bSelect =*/true, /*bNotifyForActor =*/false, /*bSelectEvenIfHidden =*/true);
+					constexpr bool bVerifyDeletionCanHappen = true;
+					constexpr bool bWarnAboutReferences = false;
+					GEditor->edactDeleteSelected(FoundActor->GetWorld(), bVerifyDeletionCanHappen, bWarnAboutReferences, bWarnAboutReferences);
+				}
+				else
+#endif
+				{
+					FoundActor->Destroy(true, true);
+				}
+			}
+			else
+			{
+				const FName NewName = MakeUniqueObjectName(FoundObject->GetOuter(), FoundObject->GetClass());
+				FoundObject->Rename(*NewName.ToString(), nullptr, REN_NonTransactional | REN_DontCreateRedirectors);
+			}
+		}
+		const FString& SubObjectPath = ActorPath.GetSubPathString();
+		const int32 LastDotIndex = SubObjectPath.Find(TEXT("."));
+		// Full string: /Game/MapName.MapName:PersistentLevel.StaticMeshActor_42
+		// SubObjectPath: PersistentLevel.StaticMeshActor_42 
+		checkf(LastDotIndex != INDEX_NONE, TEXT("There should always be at least one dot after PersistentLevel"));
+			
+		const int32 NameLength = SubObjectPath.Len() - LastDotIndex - 1;
+		const FString ActorName = SubObjectPath.Right(NameLength);
+		
+		ULevel* const OverrideLevel = WorldData.World->PersistentLevel;
+		const FName ActorFName = *ActorName;
+		
+		FActorSpawnParameters SpawnParameters;
+		SpawnParameters.Name = ActorFName;
+		SpawnParameters.OverrideLevel = OverrideLevel;
+		SpawnParameters.bNoFail = true;
+		SpawnParameters.NameMode = FActorSpawnParameters::ESpawnActorNameMode::Required_ErrorAndReturnNull;
+		SpawnParameters.SpawnCollisionHandlingOverride = ESpawnActorCollisionHandlingMethod::AlwaysSpawn;
+		SpawnParameters.ObjectFlags = ActorData.GetObjectFlags() | RF_WasLoaded;
+
+		if (AActor* const RecreatedActor = WorldData.World->SpawnActor(ActorClass, nullptr, SpawnParameters))
+		{
+			//OnPostRecreateActor(RecreatedActor);
+			RecreatedActors.Add(ActorPath, RecreatedActor);
+			WorldData.CachedActors.Add(ActorPath, RecreatedActor);
+		}
+	}
+	
+	//2nd Pass: Serialization 
+	for (const TPair<FSoftObjectPath, FAvalancheActorData>& Pair : WorldData.ActorData)
+	{
+#if WITH_EDITOR
+		LoadWorld.EnterProgressFrame();
+#endif
+
+		const FSoftObjectPath& ActorPath = Pair.Key;
+		const FAvalancheActorData& ActorData = Pair.Value;
+		
+		if (AActor* const * const RecreatedActor = RecreatedActors.Find(ActorPath))
+		{
+			AActor* const Actor = *RecreatedActor;
+			
+			UE_LOG(LogAvalancheBlueprint
+				, Verbose
+				, TEXT("========== Apply recreated %s ==========")
+				, *Actor->GetPathName());
+
+			//Allocate Components
+			for (const TPair<FAvaObjectIndex, FAvalancheComponentData>& ComponentPair : ActorData.ComponentData)
+			{
+				const FAvaObjectIndex& ReferenceIndex = ComponentPair.Key;
+				const FAvalancheComponentData& ComponentData = ComponentPair.Value;
+				FAvalancheSubObjectData* SubObjectData = WorldData.SubObjects.Find(ReferenceIndex);
+				
+				if (ensure(SubObjectData))
+				{
+					const FSoftObjectPath& ComponentPath = WorldData.SerializedObjectReferences[ReferenceIndex.Index];
+					UActorComponent* const Component = FAvaBlueprint_Serialize::FindOrAllocateComponent(WorldData, Actor, ActorData
+						, ComponentPath, *SubObjectData, ComponentData);
+					
+					if (Component)
+					{
+						UE_LOG(LogAvalancheBlueprint
+							, Verbose
+							, TEXT("Found or Allocated Component %s")
+							, *Component->GetPathName());
+					}
+					else
+					{
+						UE_LOG(LogAvalancheBlueprint
+							, Verbose
+							, TEXT("Could not Find or Allocate from Component Path %s")
+							, *ComponentPath.ToString());
+					}
+				}
+			}
+
+
+			const AActor* AttachParentBeforeLoad = Actor->GetAttachParentActor();
+			FAvaBlueprint_Serialize::LoadActor(Actor, ActorPath, WorldData);
+			const AActor* AttachParentAfterLoad = Actor->GetAttachParentActor();
+
+			Actor->UnregisterAllComponents();
+			
+			if (AttachParentAfterLoad && Actor->GetRootComponent())
+			{
+				USceneComponent* const AttachComponent = Actor->GetRootComponent()->GetAttachParent();
+				const FName AttachSocket = Actor->GetAttachParentSocketName();
+				
+				Actor->DetachFromActor(FDetachmentTransformRules::KeepRelativeTransform);
+				Actor->AttachToComponent(AttachComponent
+					, FAttachmentTransformRules::KeepRelativeTransform
+					, AttachSocket);
+			}
+
+#if WITH_EDITOR
+			Actor->RerunConstructionScripts();
+#endif
+			Actor->RegisterAllComponents();
+			Actor->UpdateComponentTransforms();
+			
+			{
+				// GAllowActorScriptExecutionInEditor must be temporarily true so call to UserConstructionScript isn't skipped
+				FEditorScriptExecutionGuard AllowConstructionScript;
+				Actor->UserConstructionScript();
+			}
+
+			if (InInitActorFunction.IsSet())
+			{
+				(*InInitActorFunction)(*Actor, ActorData);
+			}
+
+#if WITH_EDITOR
+			// Otherwise actor will show up with internal object name, e.g. actor previously called Cube will be StaticMeshActor1
+			Actor->SetActorLabel(ActorData.ActorLabel);
+			Actor->SetIsTemporarilyHiddenInEditor(ActorData.bEditorVisibility);
+			// Recreated actors have invalid lightning cache... e.g. recreated point lights will show error image (S_LightError)
+			Actor->InvalidateLightingCacheDetailed(false);
+#endif
+		}
+		else
+		{
+			UE_LOG(LogAvalancheBlueprint
+				, Error
+				, TEXT("Failed to recreate actor %s")
+				, *ActorPath.ToString());
+		}
+	}
+
+	// Post Spawn, handle AvaBlueprint data only if the World was not overridden,
+	// as specifying an Override World is treated more as a read-only load, rather than a load & fix
+	if (!InOverrideWorld)
+	{
+		if (RemoteControlPreset)
+		{
+			RemoteControlPreset->RebindUnboundEntities();
+			FAvaRemoteControlRebind::ResolveAllFieldPathInfos(RemoteControlPreset);
+		}
+
+		UWorld* const PlaybackContext = WorldData.World.Get();
+		ensureAlways(PlaybackContext == World);
+
+		for (UAvaSequence* const Animation : GetSequences())
+		{
+			if (IsValid(Animation))
+			{
+				Animation->MigrateLegacyBindings(PlaybackContext);
+			}
+		}
+	}
+}
+
+void UAvalancheBlueprint::RegisterRemoteControlPreset()
+{
+	if (!HasAnyFlags(EObjectFlags::RF_ClassDefaultObject))
+	{
+		FAvaRemoteControlUtils::RegisterRemoteControlPreset(RemoteControlPreset, /*bInEnsureUniqueId*/ true);
+	}
+}
+
+void UAvalancheBlueprint::UnregisterRemoteControlPreset()
+{
+	if (!HasAnyFlags(EObjectFlags::RF_ClassDefaultObject))
+	{
+		FAvaRemoteControlUtils::UnregisterRemoteControlPreset(RemoteControlPreset);
+	}
+}
+
+void UAvalancheBlueprint::PostLoadSubobjects(FObjectInstancingGraph* OuterInstanceGraph)
+{
+	ScheduleRebuildSequenceTree();
+}
+
+ULevel* UAvalancheBlueprint::GetSceneLevel() const
+{
+	return World ? World->PersistentLevel : nullptr;
+}
+
+void UAvalancheBlueprint::SetStartupCameraName(FName InCameraName)
+{
+	if (InCameraName == StartupCameraName)
+	{
+		return;
+	}
+
+#if WITH_EDITOR
+	StartupCameraName = InCameraName;
+#endif
+}
+
+#if WITH_EDITOR
+UClass* UAvalancheBlueprint::GetBlueprintClass() const
+{
+	return UAvaBlueprintGeneratedClass::StaticClass();
+}
+
+void UAvalancheBlueprint::GetReparentingRules(TSet<const UClass*>& AllowedChildrenOfClasses,
+	TSet<const UClass*>& DisallowedChildrenOfClasses) const
+{
+	AllowedChildrenOfClasses.Add(UAvalancheBlueprint::StaticClass());
+}
+#endif
+
+IAvaSequencePlaybackObject* UAvalancheBlueprint::GetPlaybackObject() const
+{
+	return GetScenePlayback();
+}
+
+UObject* UAvalancheBlueprint::ToUObject()
+{
+	return this;
+}
+
+UWorld* UAvalancheBlueprint::GetContextWorld() const
+{
+	return GetAvalancheWorld();
+}
+
+bool UAvalancheBlueprint::CreateDirectorInstance(UAvaSequence& InSequence
+	, IMovieScenePlayer& InPlayer
+	, const FMovieSceneSequenceID& InSequenceID
+	, UObject*& OutDirectorInstance)
+{
+	OutDirectorInstance = GetPlaceholderActor();
+	return true;
+}
+
+void UAvalancheBlueprint::SetDefaultSequence(UAvaSequence* InSequence)
+{
+	if (IsValid(InSequence))
+	{
+		AddSequence(InSequence);
+		DefaultAnimationIndex = Animations.Find(InSequence);
+	}
+}
+
+UAvaSequence* UAvalancheBlueprint::GetDefaultSequence() const
+{
+	if (Animations.IsValidIndex(DefaultAnimationIndex))
+	{
+		return Animations[DefaultAnimationIndex];
+	}
+	return nullptr;
+}
+
+bool UAvalancheBlueprint::AddSequence(UAvaSequence* InSequence)
+{
+	if (IsValid(InSequence) && !Animations.Contains(InSequence))
+	{
+		Animations.Add(InSequence);
+		ScheduleRebuildSequenceTree();
+		return true;
+	}
+	
+	return false;
+}
+
+void UAvalancheBlueprint::RemoveSequence(UAvaSequence* InSequence)
+{
+	Animations.Remove(InSequence);
+	if (InSequence)
+	{
+		InSequence->OnSequenceRemoved();
+	}
+	ScheduleRebuildSequenceTree();
+}
+
+FName UAvalancheBlueprint::GetSequenceProviderDebugName() const
+{
+	return GetFName();
+}
+
+#if WITH_EDITOR
+void UAvalancheBlueprint::OnEditorSequencerCreated(const TSharedPtr<ISequencer>& InSequencer)
+{
+	EditorSequencer = InSequencer;
+	RebuildSequenceTree();
+}
+
+TSharedPtr<ISequencer> UAvalancheBlueprint::GetEditorSequencer() const
+{
+	return EditorSequencer.Pin();
+}
+
+bool UAvalancheBlueprint::GetDirectorBlueprint(UAvaSequence& InSequence, UBlueprint*& OutBlueprint)
+{
+	OutBlueprint = this;
+	return true;
+}
+#endif
+
+void UAvalancheBlueprint::RebuildSequenceTree()
+{
+	bPendingAnimTreeUpdate = false;
+	IAvaSequenceProvider::RebuildSequenceTree();
+}
+
+void UAvalancheBlueprint::ScheduleRebuildSequenceTree()
+{
+	// Bail if the Async Task hasn't executed yet
+	if (bPendingAnimTreeUpdate)
+	{
+		return;
+	}
+
+	bPendingAnimTreeUpdate = true;
+
+	TWeakObjectPtr<UAvalancheBlueprint> ThisWeak(this);
+	AsyncTask(ENamedThreads::GameThread, [ThisWeak]()->void
+	{
+		UAvalancheBlueprint* const This = ThisWeak.Get();
+
+		// Check if we already have Rebuilt the Animation Tree while waiting for this Async Task
+		if (This && This->bPendingAnimTreeUpdate)
+		{
+			This->RebuildSequenceTree();
+		}
+	});
+}
+
+UWorld* UAvalancheBlueprint::GetWorld() const
+{
+	return World.Get();
+}
+
+#if WITH_EDITOR
+void UAvalancheBlueprint::SetGuideOffset(int32 InGuideIdx, float InGuideOffset)
+{
+	if (GuideInfos.IsValidIndex(InGuideIdx))
+	{
+		Modify();
+		GuideInfos[InGuideIdx].OffsetFraction = InGuideOffset;
+	}
+}
+
+void UAvalancheBlueprint::SetGuides(const TArray<FAvaViewportGuideInfo_Deprecated>& InNewGuideInfos)
+{
+	Modify();
+	GuideInfos = InNewGuideInfos;
+}
+#endif
+
+#undef LOCTEXT_NAMESPACE
diff --git a/Engine/Plugins/Experimental/Avalanche/Source/Avalanche/Private/AvaBlueprintGeneratedClass.cpp b/Engine/Plugins/Experimental/Avalanche/Source/Avalanche/Private/AvaBlueprintGeneratedClass.cpp
new file mode 100644
--- /dev/null
+++ b/Engine/Plugins/Experimental/Avalanche/Source/Avalanche/Private/AvaBlueprintGeneratedClass.cpp
@@ -0,0 +1,32 @@
+// Copyright Epic Games, Inc. All Rights Reserved.
+
+#include "AvaBlueprintGeneratedClass.h"
+#include "AvaActor.h"
+#include "AvaSequence.h"
+
+void UAvaBlueprintGeneratedClass::UpdateProperties(AAvalancheActor* InActor)
+{
+	UAvaBlueprintGeneratedClass* const GeneratedClass = Cast<UAvaBlueprintGeneratedClass>(InActor->GetClass());
+
+	TMap<FName, FObjectPropertyBase*> ObjectPropertyMap;
+	for (TFieldIterator<FObjectPropertyBase> Iter(GeneratedClass, EFieldIterationFlags::Default); Iter; ++Iter)
+	{
+		check(*Iter);
+		ensureMsgf(!ObjectPropertyMap.Contains(Iter->GetFName()), TEXT("There are properties with the same names: '%s'"), *Iter->GetName());
+		ObjectPropertyMap.Add(Iter->GetFName(), *Iter);
+	}
+
+	// Animation Properties
+	for (UAvaSequence* const Animation : GeneratedClass->Animations)
+	{
+		if (Animation)
+		{
+			// Find property with the same name as the animation and assign the animation to it.
+			if (FObjectPropertyBase* const* Property = ObjectPropertyMap.Find(Animation->GetFName()))
+			{
+				check(*Property);
+				(*Property)->SetObjectPropertyValue_InContainer(InActor, Animation);
+			}
+		}
+	}
+}
diff --git a/Engine/Plugins/Experimental/Avalanche/Source/Avalanche/Private/AvaBlueprintVersion.cpp b/Engine/Plugins/Experimental/Avalanche/Source/Avalanche/Private/AvaBlueprintVersion.cpp
new file mode 100644
--- /dev/null
+++ b/Engine/Plugins/Experimental/Avalanche/Source/Avalanche/Private/AvaBlueprintVersion.cpp
@@ -0,0 +1,8 @@
+// Copyright Epic Games, Inc. All Rights Reserved.
+
+#include "AvaBlueprintVersion.h"
+#include "Serialization/CustomVersion.h"
+
+const FGuid FAvalancheBlueprintVersion::GUID(0xA248F02D, 0x294D4EA5, 0x99C9A713, 0x55277CD7);
+FCustomVersionRegistration GRegisterAvalancheBlueprintVersion(FAvalancheBlueprintVersion::GUID, FAvalancheBlueprintVersion::LatestVersion,
+	TEXT("AvalancheBlueprint"));
diff --git a/Engine/Plugins/Experimental/Avalanche/Source/Avalanche/Private/AvaBlueprintVersion.h b/Engine/Plugins/Experimental/Avalanche/Source/Avalanche/Private/AvaBlueprintVersion.h
new file mode 100644
--- /dev/null
+++ b/Engine/Plugins/Experimental/Avalanche/Source/Avalanche/Private/AvaBlueprintVersion.h
@@ -0,0 +1,24 @@
+// Copyright Epic Games, Inc. All Rights Reserved.
+
+#pragma once
+
+#include "Misc/Guid.h"
+
+/** Version used for serializing Avalanche Blueprint Data */
+struct FAvalancheBlueprintVersion
+{
+private:
+	FAvalancheBlueprintVersion() = delete;
+
+public:
+	enum Type : uint8
+	{
+		PreVersioning = 0, // From before file versioning was implemented
+
+		//new versions can be added above here
+		VersionPlusOne,
+		LatestVersion = VersionPlusOne - 1
+	};
+
+	const static FGuid GUID;
+};
diff --git a/Engine/Plugins/Experimental/Avalanche/Source/Avalanche/Private/AvaBlueprint_Serialize.cpp b/Engine/Plugins/Experimental/Avalanche/Source/Avalanche/Private/AvaBlueprint_Serialize.cpp
new file mode 100644
--- /dev/null
+++ b/Engine/Plugins/Experimental/Avalanche/Source/Avalanche/Private/AvaBlueprint_Serialize.cpp
@@ -0,0 +1,782 @@
+// Copyright Epic Games, Inc. All Rights Reserved.
+
+#include "AvaBlueprint_Serialize.h"
+#include "AudioDevice.h"
+#include "EngineUtils.h"
+#include "MovieScene.h"
+#include "AI/NavigationSystemConfig.h"
+#include "Algo/AllOf.h"
+#include "AvaSequence.h"
+#include "Archive/AvalancheReader.h"
+#include "Archive/AvalancheWriter.h"
+#include "Data/AvalancheActorData.h"
+#include "Data/AvalancheComponentData.h"
+#include "Data/AvalancheDataDefines.h"
+#include "Data/AvalancheObjectData.h"
+#include "Data/AvalancheSubObjectData.h"
+#include "Data/AvalancheWorldData.h"
+#include "GameFramework/PhysicsVolume.h"
+#include "Engine/BrushBuilder.h"
+#include "Misc/ScopedSlowTask.h"
+#include "Particles/ParticleEventManager.h"
+#include "Templates/NonNullPointer.h"
+#include "UObject/ObjectSaveContext.h"
+
+#if WITH_EDITORONLY_DATA
+#include "Kismet2/ComponentEditorUtils.h"
+#endif
+
+#define LOCTEXT_NAMESPACE "AvaBlueprint_Serialize"
+
+DEFINE_LOG_CATEGORY_STATIC(LogAvaBlueprint_Serialize, Log, All);
+
+FSoftObjectPath FAvaBlueprint_Serialize::SetActorInPath(AActor* NewActor, const FSoftObjectPath& OriginalObjectPath)
+{
+	const static FString PersistentLevelString("PersistentLevel.");
+	const int32 PersistentLevelStringLength = PersistentLevelString.Len();
+	const FString& SubPathString            = OriginalObjectPath.GetSubPathString();
+	const int32 IndexOfPersistentLevelInfo  = SubPathString.Find(PersistentLevelString, ESearchCase::CaseSensitive);
+	if (IndexOfPersistentLevelInfo == INDEX_NONE)
+	{
+		return {};
+	}
+
+	const int32 DotAfterActorNameIndex = SubPathString.Find(TEXT("."), ESearchCase::IgnoreCase, ESearchDir::FromStart,
+		IndexOfPersistentLevelInfo + PersistentLevelStringLength);
+
+	const FSoftObjectPath PathToNewActor(NewActor);
+	// PersistentLevel.StaticMeshActor_42.StaticMeshComponent becomes .StaticMeshComponent
+	const FString PathAfterOriginalActor = SubPathString.Right(SubPathString.Len() - DotAfterActorNameIndex);
+	return FSoftObjectPath(PathToNewActor.GetAssetPath(), PathToNewActor.GetSubPathString() + PathAfterOriginalActor);
+}
+
+UObject* FAvaBlueprint_Serialize::ResolveObjectDependency(FAvalancheWorldData& WorldData, FAvaObjectIndex ObjectIndex)
+{
+	if (!ensure(WorldData.SerializedObjectReferences.IsValidIndex(ObjectIndex.Index)))
+	{
+		return nullptr;
+	}
+
+	const FSoftObjectPath& OriginalObjectPath = WorldData.SerializedObjectReferences[ObjectIndex.Index];
+
+	bool bIsPathToActorSubObject;
+	const TOptional<FSoftObjectPath> PathToActor = ExtractActorFromPath(OriginalObjectPath, bIsPathToActorSubObject);
+
+	const bool bIsExternalAssetReference = !PathToActor.IsSet();
+	if (bIsExternalAssetReference)
+	{
+		UObject* ExternalReference      = OriginalObjectPath.ResolveObject();
+		const bool bNeedsToLoadFromDisk = ExternalReference == nullptr;
+		if (bNeedsToLoadFromDisk)
+		{
+			ExternalReference = OriginalObjectPath.TryLoad();
+		}
+
+		// We're supposed to be dealing with an external (asset) reference, e.g. a UMaterial.
+		ensureAlwaysMsgf(ExternalReference == nullptr
+			|| !ExternalReference->IsA<AActor>()
+			&& !ExternalReference->IsA<UActorComponent>()
+			, TEXT("Something is wrong. We just checked that the reference is not a world object but it is an actor or component."));
+
+		return ExternalReference;
+	}
+
+	TWeakObjectPtr<AActor> CachedActor;
+	{
+		const FSoftObjectPath& OriginalActorPath = *PathToActor;
+
+		FAvalancheActorData* const ActorData = WorldData.ActorData.Find(OriginalActorPath);
+		if (!ActorData)
+		{
+			UE_LOG(LogTemp, Warning, TEXT("No save data found for actor %s"), *OriginalActorPath.ToString());
+			return nullptr;
+		}
+
+		TWeakObjectPtr<AActor>& ActorWeak = WorldData.CachedActors.FindOrAdd(OriginalActorPath);
+		if (!ActorWeak.IsValid())
+		{
+			const FSoftClassPath SoftClassPath(ActorData->ActorClass);
+			UClass* TargetClass = SoftClassPath.TryLoadClass<AActor>();
+			if (!TargetClass)
+			{
+				UE_LOG(LogTemp
+					, Error
+					, TEXT("Unknown class %s. Mostly likely it is referencing a class that was deleted.")
+					, *SoftClassPath.ToString());
+				return nullptr;
+			}
+
+			FActorSpawnParameters SpawnParams;
+			SpawnParams.Template                       = Cast<AActor>(TargetClass->GetDefaultObject());
+			SpawnParams.NameMode                       = FActorSpawnParameters::ESpawnActorNameMode::Requested;
+			SpawnParams.SpawnCollisionHandlingOverride = ESpawnActorCollisionHandlingMethod::AlwaysSpawn;
+			if (ensureMsgf(SpawnParams.Template, TEXT("Failed to get class default. This should not happen. Investigate.")))
+			{
+				UClass* ClassToUse = SpawnParams.Template->GetClass();
+				SpawnParams.Name   = *FString("AvalancheObjectInstance_").Append(
+					*MakeUniqueObjectName(WorldData.World.Get(), ClassToUse).ToString());
+				ActorWeak = WorldData.World->SpawnActor<AActor>(ClassToUse, SpawnParams);
+			}
+			else
+			{
+				ActorWeak = WorldData.World->SpawnActor<AActor>(TargetClass, SpawnParams);
+			}
+
+			if (ensureMsgf(ActorWeak.IsValid(), TEXT("Failed to spawn actor of class '%s'"), *ActorData->ActorClass.ToString()))
+			{
+#if WITH_EDITOR
+				// Hide this actor so external systems can see that this components should not render, i.e. make USceneComponent::ShouldRender return false
+				ActorWeak->SetIsTemporarilyHiddenInEditor(true);
+#endif
+
+				for (const TPair<FAvaObjectIndex, FAvalancheComponentData>& Pair : ActorData->ComponentData)
+				{
+					const FAvaObjectIndex& ReferenceIndex = Pair.Key;
+
+					FAvalancheSubObjectData* SubObjectData = WorldData.SubObjects.Find(ReferenceIndex);
+					if (ensure(SubObjectData))
+					{
+						const FSoftObjectPath& ComponentPath = WorldData.SerializedObjectReferences[ReferenceIndex.Index];
+						FindOrAllocateComponent(WorldData
+							, ActorWeak.Get()
+							, *ActorData
+							, ComponentPath
+							, *SubObjectData
+							, Pair.Value);
+					}
+				}
+			}
+		}
+
+		CachedActor = ActorWeak;
+	}
+
+	if (!CachedActor.IsValid())
+	{
+		return nullptr;
+	}
+
+	AActor* SnapshotActor     = CachedActor.Get();
+	const bool bIsPathToActor = !bIsPathToActorSubObject;
+	if (bIsPathToActor)
+	{
+		return SnapshotActor;
+	}
+
+	// Reference to default subobjects, such as components, or an object we already allocated.
+	const FSoftObjectPath PathToSubObject = SetActorInPath(SnapshotActor, OriginalObjectPath);
+	if (UObject* ExistingSubObject = PathToSubObject.ResolveObject())
+	{
+		FAvalancheSubObjectData* SubObjectData = WorldData.SubObjects.Find(ObjectIndex.Index);
+		if (!SubObjectData || SubObjectData->bWasSkippedClass)
+		{
+			return ExistingSubObject;
+		}
+
+		if (SubObjectData->Class.ResolveClass() != ExistingSubObject->GetClass())
+		{
+			UE_LOG(LogTemp, Warning
+				, TEXT("Skipping serialisation of subobject because classes are different. Object '%s' will not contain values saved.")
+				, *ExistingSubObject->GetName());
+			return ExistingSubObject;
+		}
+
+		TWeakObjectPtr<UObject>& SubobjectCache = WorldData.CachedSubObjects.FindOrAdd(
+			WorldData.SerializedObjectReferences[ObjectIndex.Index]);
+		if (!SubobjectCache.IsValid())
+		{
+			SubobjectCache = ExistingSubObject;
+			FAvalancheReader Reader(WorldData, *SubObjectData, ExistingSubObject);
+			if (Reader.IsError())
+			{
+				UE_LOG(LogAvaBlueprint_Serialize, Error,
+					TEXT("An error occured loading Avalanche SubObject \"%s\" of Class \"%s\" (Path: \"%s\")."),
+					*GetNameSafe(ExistingSubObject), *SubObjectData->Class.ToString(),
+					*WorldData.SerializedObjectReferences[ObjectIndex.Index].ToString());
+			}
+		}
+
+		return ExistingSubObject;
+	}
+
+	// The sub object is instanced: recreate it recursively
+	return CreateSubObject(WorldData, ObjectIndex, PathToSubObject);
+}
+
+FString FAvaBlueprint_Serialize::ExtractLastSubObjectName(const FSoftObjectPath& ObjectPath)
+{
+	const FString& SubPathString = ObjectPath.GetSubPathString();
+	const int32 LastDotIndex     = SubPathString.Find(TEXT("."), ESearchCase::IgnoreCase, ESearchDir::FromEnd);
+	if (LastDotIndex == INDEX_NONE)
+	{
+		return SubPathString;
+	}
+
+	return SubPathString.RightChop(LastDotIndex + 1);
+}
+
+UObject* FAvaBlueprint_Serialize::CreateSubObject(FAvalancheWorldData& WorldData, FAvaObjectIndex ObjectIndex,
+	const FSoftObjectPath& PathToSubObject)
+{
+	FAvalancheSubObjectData* SubObjectData = WorldData.SubObjects.Find(ObjectIndex.Index);
+	if (!SubObjectData || SubObjectData->bWasSkippedClass)
+	{
+		return nullptr;
+	}
+
+	if (const TWeakObjectPtr<UObject>* SubObjectCache = WorldData.CachedSubObjects.Find(PathToSubObject))
+	{
+		if (SubObjectCache->IsValid())
+		{
+			return SubObjectCache->Get();
+		}
+	}
+
+	const UClass* Class = SubObjectData->Class.TryLoadClass<UObject>();
+	if (!Class)
+	{
+		UE_LOG(LogTemp, Warning, TEXT("Class '%s' not found. Maybe it was removed?"), *SubObjectData->Class.ToString());
+		return nullptr;
+	}
+
+	const FAvaObjectIndex OuterIndex = SubObjectData->OuterIndex;
+	check(OuterIndex.Index != ObjectIndex.Index);
+	UObject* SubObjectOuter = ResolveObjectDependency(WorldData, OuterIndex);
+	if (!SubObjectOuter)
+	{
+		UE_LOG(LogTemp, Warning
+			, TEXT("Failed to create '%s' because its outer could not be created.")
+			, *PathToSubObject.ToString());
+		return nullptr;
+	}
+
+	const FName SubObjectName        = *ExtractLastSubObjectName(PathToSubObject);
+	const bool bSubObjectNameIsTaken = [SubObjectOuter, SubObjectName]
+	{
+		TArray<UObject*> Subobjects;
+		GetObjectsWithOuter(SubObjectOuter, Subobjects, true);
+		for (const UObject* Object : Subobjects)
+		{
+			if (Object->GetFName() == SubObjectName)
+			{
+				return true;
+			}
+		}
+		return false;
+	}();
+	if (bSubObjectNameIsTaken)
+	{
+		UE_LOG(LogTemp
+			, Warning, TEXT("Failed to create '%s' because subobject name was already taken.")
+			, *PathToSubObject.ToString());
+		return nullptr;
+	}
+
+	UObject* SubObject = NewObject<UObject>(SubObjectOuter, Class, SubObjectName);
+
+	TWeakObjectPtr<UObject>& CachedSubObject = WorldData.CachedSubObjects.
+	                                                     FindOrAdd(WorldData.SerializedObjectReferences[ObjectIndex.Index]);
+	CachedSubObject = SubObject;
+
+	FAvalancheReader Reader(WorldData, *SubObjectData, SubObject);
+	if (Reader.IsError())
+	{
+		UE_LOG(LogAvaBlueprint_Serialize, Error,
+			TEXT("An error occured loading Avalanche SubObject \"%s\" of Class \"%s\" (Path: \"%s\")."),
+			*GetNameSafe(SubObject), *SubObjectData->Class.ToString(),
+			*WorldData.SerializedObjectReferences[ObjectIndex.Index].ToString());
+	}
+	return SubObject;
+}
+
+FString FAvaBlueprint_Serialize::ExtractRootToLeafComponentPath(const FSoftObjectPath& ComponentPath)
+{
+	const static FString PersistentLevelString("PersistentLevel.");
+	const int32 PersistentLevelStringLength = PersistentLevelString.Len();
+
+	// /Game/MapName.MapName:PersistentLevel.StaticMeshActor_42.StaticMeshComponent
+	// becomes PersistentLevel.StaticMeshActor_42.StaticMeshComponent
+	const FString& SubPathString           = ComponentPath.GetSubPathString();
+	const int32 IndexOfPersistentLevelInfo = SubPathString.Find(PersistentLevelString, ESearchCase::CaseSensitive);
+	if (IndexOfPersistentLevelInfo == INDEX_NONE)
+	{
+		return {};
+	}
+
+	// Example: /Game/MapName.MapName:PersistentLevel.StaticMeshActor_42.StaticMeshComponent
+	const int32 DotAfterActorNameIndex = SubPathString.Find(TEXT("."
+			), ESearchCase::IgnoreCase
+		, ESearchDir::FromStart
+		, IndexOfPersistentLevelInfo + PersistentLevelStringLength);
+
+	// PersistentLevel.SomeActor.SomeParentComp.SomeChildComp becomes SomeParentComp.SomeChildComp
+	return ensure(DotAfterActorNameIndex != INDEX_NONE)
+		       ? ComponentPath.GetSubPathString().RightChop(DotAfterActorNameIndex)
+		       : FString();
+}
+
+UActorComponent* FAvaBlueprint_Serialize::FindMatchingComponent(AActor* ActorToSearchOn, const FSoftObjectPath& ComponentPath)
+{
+	const FString RootToLeafPath = ExtractRootToLeafComponentPath(ComponentPath);
+	if (!ensure(!RootToLeafPath.IsEmpty()))
+	{
+		return nullptr;
+	}
+
+	for (UActorComponent* Component : ActorToSearchOn->GetComponents())
+	{
+		const FString OtherRootToLeaf = ExtractRootToLeafComponentPath(Component);
+		if (RootToLeafPath.Equals(OtherRootToLeaf))
+		{
+			return Component;
+		}
+	}
+	return nullptr;
+}
+
+void FAvaBlueprint_Serialize::SaveObject(FAvalancheWorldData& WorldData, FAvalancheObjectData& ObjectData, UObject* Object)
+{
+	FAvalancheWriter Writer(WorldData, ObjectData, Object);
+
+	for (const FCustomVersion& CustomVersion : Writer.GetCustomVersions().GetAllVersions())
+	{
+		WorldData.VersionInfo.UsedCustomVersions.Add(CustomVersion.Key);
+	}
+}
+
+void FAvaBlueprint_Serialize::LoadObject(FAvalancheWorldData& WorldData, FAvalancheObjectData& ObjectData, UObject* Object)
+{
+	FAvalancheReader Reader(WorldData, ObjectData, Object);
+	if (Reader.IsError())
+	{
+		UE_LOG(LogAvaBlueprint_Serialize, Error,
+			TEXT("An error occured loading Avalanche SubObject \"%s\"."),
+			*GetNameSafe(Object));
+	}
+}
+
+bool FAvaBlueprint_Serialize::ShouldSaveActor(const UWorld* InWorld, const AActor* InActor)
+{
+	if (InWorld && InActor)
+	{
+		return InActor != InWorld->GetDefaultPhysicsVolume()
+			&& !InActor->HasAnyFlags(RF_Transient)
+#if WITH_EDITOR
+			&& InActor->IsEditable()
+			&& !InActor->bIsEditorPreviewActor
+#endif
+			&& InActor != InWorld->GetDefaultBrush()
+			&& InActor != InWorld->GetWorldSettings()
+			&& InActor != InWorld->MyParticleEventManager;
+	}
+	return false;
+}
+
+bool FAvaBlueprint_Serialize::ShouldSaveComponent(const UActorComponent* InActorComponent)
+{
+	if (InActorComponent)
+	{
+		const bool bIsRootComponent = InActorComponent->GetOwner() && InActorComponent->GetOwner()->GetRootComponent() == InActorComponent;
+		const bool bAllowedCreationMethod = InActorComponent->CreationMethod != EComponentCreationMethod::UserConstructionScript;
+		const bool bNotTransient = !InActorComponent->HasAnyFlags(RF_Transient);
+
+#if WITH_EDITORONLY_DATA
+		constexpr bool bAllowUserConstructionScriptComps = false;
+
+		const UActorComponent* ParentComponent = nullptr;
+		if (const USceneComponent* SceneComponent = Cast<USceneComponent>(InActorComponent))
+		{
+			const bool bHasAttachParent = !!SceneComponent->GetAttachParent();
+			const bool bHasOwner        = !!SceneComponent->GetOwner();
+
+			//Check that we're Attached to a Component that is owned by our Owner Actor
+			if (bHasAttachParent && bHasOwner
+				&& SceneComponent->GetAttachParent()->GetOwner() == SceneComponent->GetOwner())
+			{
+				ParentComponent = SceneComponent->GetAttachParent();
+			}
+		}
+
+		const bool bCanEditComponentInstance = FComponentEditorUtils::CanEditComponentInstance(InActorComponent
+			, ParentComponent
+			, bAllowUserConstructionScriptComps);
+
+#else
+		const bool bCanEditComponentInstance = true;
+#endif
+
+		// Always Save Root Components as this affects Actor's Transform
+		return bIsRootComponent || (bAllowedCreationMethod && bNotTransient && bCanEditComponentInstance);
+	}
+	return false;
+}
+
+bool FAvaBlueprint_Serialize::ShouldSaveSubObject(UObject* InSubObject)
+{
+	check(InSubObject);
+	checkf(!InSubObject->IsA<UClass>(), TEXT("Do you have a typo on your code?"));
+
+	if (const UActorComponent* Component = Cast<UActorComponent>(InSubObject))
+	{
+		return ShouldSaveComponent(Component);
+	}
+
+	static const TSet UnsupportedClasses =
+		{
+			// Does not have state. Referenced by AVolumes, ABrushes, e.g. APostProcessVolume or ALightmassImportanceVolume
+			UBrushBuilder::StaticClass(),
+			UNavigationSystemConfig::StaticClass()
+		};
+
+	const bool bIsSupportedClass = Algo::AllOf(UnsupportedClasses, [InSubObject](UClass* Class)
+	{
+		return !InSubObject->IsA(Class);
+	});
+
+	return bIsSupportedClass && !InSubObject->HasAnyFlags(RF_Transient);
+}
+
+FAvaObjectIndex FAvaBlueprint_Serialize::AddOrFindObjectReference(FAvalancheWorldData& WorldData, const FSoftObjectPath& ObjectPath)
+{
+	if (const FAvaObjectIndex* ExistingIndex = WorldData.ObjectReferenceIndexMap.Find(ObjectPath))
+	{
+		return *ExistingIndex;
+	}
+	const FAvaObjectIndex Index(WorldData.SerializedObjectReferences.Add(ObjectPath));
+	WorldData.ObjectReferenceIndexMap.Add(ObjectPath, Index);
+	return Index;
+}
+
+TOptional<FSoftObjectPath> FAvaBlueprint_Serialize::ExtractActorFromPath(const FSoftObjectPath& InObjectPath, bool& bIsPathToActorSubObject)
+{
+	const static FString PersistentLevelString("PersistentLevel.");
+	const int32 PersistentLevelStringLength = PersistentLevelString.Len();
+	// /Game/MapName.MapName:PersistentLevel.StaticMeshActor_42.StaticMeshComponent becomes PersistentLevel.StaticMeshActor_42.StaticMeshComponent
+	const FString& SubPathString           = InObjectPath.GetSubPathString();
+	const int32 IndexOfPersistentLevelInfo = SubPathString.Find(PersistentLevelString, ESearchCase::CaseSensitive);
+	if (IndexOfPersistentLevelInfo == INDEX_NONE)
+	{
+		return {};
+	}
+
+	// Example: /Game/MapName.MapName:PersistentLevel.StaticMeshActor_42.StaticMeshComponent
+	const int32 DotAfterActorNameIndex = SubPathString.Find(TEXT("."), ESearchCase::IgnoreCase, ESearchDir::FromStart,
+		IndexOfPersistentLevelInfo + PersistentLevelStringLength);
+	const bool bPathIsToActor = DotAfterActorNameIndex == INDEX_NONE;
+
+	bIsPathToActorSubObject = !bPathIsToActor;
+	return bPathIsToActor
+		       ?
+		       // Example /Game/MapName.MapName:PersistentLevel.StaticMeshActor_42
+		       InObjectPath
+		       :
+		       // Converts /Game/MapName.MapName:PersistentLevel.StaticMeshActor_42.StaticMeshComponent to /Game/MapName.MapName:PersistentLevel.StaticMeshActor_42
+		       FSoftObjectPath(InObjectPath.GetAssetPath(), SubPathString.Left(DotAfterActorNameIndex));
+}
+
+UActorComponent* FAvaBlueprint_Serialize::FindOrAllocateComponent(const FAvalancheWorldData& InWorldData
+	, AActor* InRecreatedActor
+	, const FAvalancheActorData& InActorData
+	, const FSoftObjectPath& InComponentPath
+	, FAvalancheSubObjectData& SubObjectData
+	, const FAvalancheComponentData& ComponentData)
+{
+	UActorComponent* OutComponent = FindMatchingComponent(InRecreatedActor, InComponentPath);
+
+	if (!OutComponent)
+	{
+		bool bIsOwnedByComponent;
+		const FSoftObjectPath& ComponentOuterPath = InWorldData.SerializedObjectReferences[SubObjectData.OuterIndex.Index];
+		TOptional<TNonNullPtr<const FAvalancheActorData>> ComponentOuterData;
+
+		//Find Saved Actor Data Using Object Path
+		if (const TOptional<FSoftObjectPath> PathToActor = ExtractActorFromPath(ComponentOuterPath, bIsOwnedByComponent))
+		{
+			const FAvalancheActorData* const Result = InWorldData.ActorData.Find(*PathToActor);
+
+			UE_CLOG(Result == nullptr
+				, LogTemp
+				, Warning
+				, TEXT(
+					"Path %s looks like an actor path but no data was saved for it. Maybe it was a reference to an auto-generated actor, e.g. a brush or volume present in all worlds by default?"
+				)
+				, *ComponentOuterPath.ToString());
+
+			if (Result)
+			{
+				ComponentOuterData = TOptional<TNonNullPtr<const FAvalancheActorData>>(Result);
+			}
+		}
+
+		const bool bIsOwnedByOtherActor = ComponentOuterData.Get(nullptr) != &InActorData;
+
+		if (!ensureAlwaysMsgf(ComponentOuterData
+				, TEXT("Failed to recreate component %s because its outer %s did not have any associated actor data. Investigate.")
+				, *InComponentPath.ToString(), *ComponentOuterPath.ToString())
+
+			|| !ensureMsgf(!bIsOwnedByOtherActor
+				, TEXT(
+					"Failed to recreate component %s because saved data indicates it is owned by another actor %s. Components normally are owned by the actor they're attached to. Investigate."
+				)
+				, *InComponentPath.ToString(), *ComponentOuterPath.ToString()))
+		{
+			return nullptr;
+		}
+
+		UObject* ComponentOuter = InRecreatedActor;
+
+		if (bIsOwnedByComponent)
+		{
+			bool bFoundPath = false;
+			for (const TPair<FAvaObjectIndex, FAvalancheComponentData>& ComponentPair : InActorData.ComponentData)
+			{
+				const FSoftObjectPath& SavedComponentPath = InWorldData.SerializedObjectReferences[ComponentPair.Key.Index];
+				if (SavedComponentPath == ComponentOuterPath)
+				{
+					bFoundPath = true;
+
+					const FAvaObjectIndex ReferenceIndex                    = ComponentPair.Key;
+					const FAvalancheSubObjectData* const FoundSubObjectData = InWorldData.SubObjects.Find(ReferenceIndex);
+
+					if (ensure(FoundSubObjectData))
+					{
+						ComponentOuter = FindOrAllocateComponent(InWorldData
+							, InRecreatedActor
+							, InActorData
+							, ComponentOuterPath
+							, SubObjectData
+							, ComponentPair.Value);
+					}
+				}
+			}
+
+			UE_CLOG(bFoundPath
+				, LogTemp
+				, Error
+				, TEXT("Failed to find outer for component %s")
+				, *ComponentOuterPath.ToString());
+		}
+
+		if (ComponentOuter)
+		{
+			//Allocate Component
+			{
+				const UClass* ComponentClass = SubObjectData.Class.TryLoadClass<UActorComponent>();
+
+				if (!ensureMsgf(ComponentClass, TEXT("Component Class is invalid for Outer %s"), *ComponentOuter->GetName()))
+				{
+					return nullptr;
+				}
+				
+				FName ComponentName;
+				{
+					const FString& SubPathString = InComponentPath.GetSubPathString();
+					const int32 LastDotIndex     = SubPathString.Find(TEXT("."), ESearchCase::IgnoreCase, ESearchDir::FromEnd);
+					if (LastDotIndex == INDEX_NONE)
+					{
+						ComponentName = *SubPathString;
+					}
+					else
+					{
+						ComponentName = *SubPathString.RightChop(LastDotIndex + 1);
+					}
+				}
+
+				const EObjectFlags ObjectFlags = SubObjectData.GetObjectFlags();
+
+				// RF_Transactional must be set when the object is created so the creation transacted correctly
+				OutComponent = NewObject<UActorComponent>(ComponentOuter
+					, ComponentClass
+					, ComponentName
+					, ObjectFlags | RF_Transactional | RF_WasLoaded);
+
+				OutComponent->SetFlags(ObjectFlags);
+			}
+
+			// UActorComponent::PostInitProperties implicitly calls AddOwnedComponent but we have to manually add it to the other arrays
+			if (OutComponent->CreationMethod != ComponentData.CreationMethod)
+			{
+				OutComponent->CreationMethod = ComponentData.CreationMethod;
+				switch (OutComponent->CreationMethod)
+				{
+					case EComponentCreationMethod::Instance: InRecreatedActor->AddInstanceComponent(OutComponent);
+						break;
+
+					case EComponentCreationMethod::SimpleConstructionScript: InRecreatedActor->BlueprintCreatedComponents.AddUnique(
+							OutComponent);
+						break;
+
+					case EComponentCreationMethod::UserConstructionScript:
+						checkf(false, TEXT("Component created in construction script currently unsupported"));
+						break;
+
+					case EComponentCreationMethod::Native:
+					default: break;
+				}
+			}
+		}
+	}
+
+	return OutComponent;
+}
+
+void FAvaBlueprint_Serialize::AddSubObjectDependency(FAvalancheWorldData& WorldData, UObject* ReferenceFromOriginalObject,
+	FAvaObjectIndex Index)
+{
+	if (!ShouldSaveSubObject(ReferenceFromOriginalObject))
+	{
+		WorldData.SubObjects.Add(Index, FAvalancheSubObjectData::MakeSkippedSubObjectData());
+		return;
+	}
+
+	const bool bNeedsToSaveSubObject = !WorldData.SubObjects.Contains(Index);
+	if (bNeedsToSaveSubObject)
+	{
+		// Avoid infinite recursion
+		WorldData.SubObjects.Add(Index);
+
+		// Owning actor keeps reference to saved subobjects so they can be recreated faster when snapshot is loaded
+		if (!ReferenceFromOriginalObject->IsA<UActorComponent>())
+		{
+			const AActor* const OwningActor = ReferenceFromOriginalObject->GetTypedOuter<AActor>();
+			check(OwningActor);
+			WorldData.ActorData.FindOrAdd(OwningActor).OwnedSubobjects.Add(Index);
+		}
+
+		// Important: first allocate SubobjectData on stack...
+		FAvalancheSubObjectData SubObjectData;
+		const UClass* Class      = ReferenceFromOriginalObject->GetClass();
+		SubObjectData.Class      = Class;
+		SubObjectData.OuterIndex = AddObjectDependency(WorldData, ReferenceFromOriginalObject->GetOuter());
+
+		// ... because serialisation recursively calls this function. FWorldSnapshotData::Subobjects is possibly reallocated
+		SaveObject(WorldData, SubObjectData, ReferenceFromOriginalObject);
+
+		WorldData.SubObjects[Index] = MoveTemp(SubObjectData); // MoveTemp not profiled
+	}
+}
+
+FAvaObjectIndex FAvaBlueprint_Serialize::AddObjectDependency(FAvalancheWorldData& WorldData
+	, UObject* ReferenceFromOriginalObject
+	, bool bCheckWhetherSubObject)
+{
+	// Even if ShouldSave later returns false for this object, we want to track it
+	const FAvaObjectIndex OutIndex = AddOrFindObjectReference(WorldData, ReferenceFromOriginalObject);
+	if (bCheckWhetherSubObject && ReferenceFromOriginalObject && ReferenceFromOriginalObject->GetTypedOuter<AActor>())
+	{
+		AddSubObjectDependency(WorldData, ReferenceFromOriginalObject, OutIndex);
+	}
+	return OutIndex;
+}
+
+void FAvaBlueprint_Serialize::LoadActor(AActor* InActor, const FSoftObjectPath& InActorPath, FAvalancheWorldData& WorldData)
+{
+	check(InActor);
+
+	UE_LOG(LogTemp
+		, Verbose
+		, TEXT("========== Deserialize Actor %s ==========")
+		, *InActorPath.ToString());
+
+	FAvalancheActorData& ActorData = WorldData.ActorData[InActorPath];
+
+	LoadObject(WorldData, ActorData, InActor);
+
+#if WITH_EDITOR
+	UE_LOG(LogTemp
+		, Verbose
+		, TEXT("ActorLabel is \"%s\" for \"%s\" (editor object path \"%s\")")
+		, *InActor->GetActorLabel()
+		, *InActor->GetPathName()
+		, *InActorPath.ToString());
+#endif
+
+	//Deserialize Components
+	for (const TPair<FAvaObjectIndex, FAvalancheComponentData>& ComponentPair : ActorData.ComponentData)
+	{
+		const EComponentCreationMethod CreationMethod = ComponentPair.Value.CreationMethod;
+
+		// Construction script components are not supported
+		if (CreationMethod != EComponentCreationMethod::UserConstructionScript)
+		{
+			const FAvaObjectIndex SubObjectIndex         = ComponentPair.Key;
+			const FSoftObjectPath& OriginalComponentPath = WorldData.SerializedObjectReferences[SubObjectIndex.Index];
+
+			if (UActorComponent* const ComponentToRestore = FindMatchingComponent(InActor, OriginalComponentPath))
+			{
+				FAvalancheSubObjectData& SubObjectData = WorldData.SubObjects[SubObjectIndex];
+				LoadObject(WorldData, SubObjectData, ComponentToRestore);
+				if (UPrimitiveComponent* PrimitiveComponent = Cast<UPrimitiveComponent>(ComponentToRestore))
+				{
+#if WITH_EDITOR
+					PrimitiveComponent->UpdateCollisionProfile();
+#else
+					PrimitiveComponent->BodyInstance.LoadProfileData(false);
+#endif
+				}
+			}
+		}
+	}
+
+	InActor->PostLoad();
+	TArray<UObject*> SubObjects;
+	GetObjectsWithOuter(InActor, SubObjects);
+	for (UObject* const Object : SubObjects)
+	{
+		if (Object)
+		{
+			Object->PostLoad();
+		}
+	}
+}
+
+TOptional<FAvalancheComponentData> FAvaBlueprint_Serialize::SaveComponent(UActorComponent* InComponent)
+{
+	check(InComponent);
+
+	if (InComponent->CreationMethod == EComponentCreationMethod::UserConstructionScript)
+	{
+		UE_LOG(LogTemp, Warning
+			, TEXT("Components created dynamically in the construction script are not supported (%s). Skipping...")
+			, *InComponent->GetPathName());
+		return TOptional<FAvalancheComponentData>{};
+	}
+
+	FAvalancheComponentData OutComponentData;
+	OutComponentData.CreationMethod = InComponent->CreationMethod;
+	return OutComponentData;
+}
+
+FAvalancheActorData FAvaBlueprint_Serialize::SaveActor(AActor* InActor, FAvalancheWorldData& WorldData)
+{
+	check(InActor);
+
+	FAvalancheActorData OutActorData;
+
+	const UClass* const ActorClass = InActor->GetClass();
+	OutActorData.ActorClass        = ActorClass;
+
+#if WITH_EDITORONLY_DATA
+	OutActorData.ActorLabel        = InActor->GetActorLabel();
+	OutActorData.bEditorVisibility = InActor->IsTemporarilyHiddenInEditor();
+#endif
+
+	SaveObject(WorldData, OutActorData, InActor);
+
+	TInlineComponentArray<UActorComponent*> Components;
+	InActor->GetComponents(Components);
+	for (UActorComponent* Component : Components)
+	{
+		if (ShouldSaveComponent(Component))
+		{
+			if (TOptional<FAvalancheComponentData> SerializedComponentData = SaveComponent(Component))
+			{
+				const FAvaObjectIndex ComponentIndex = AddObjectDependency(WorldData, Component);
+				OutActorData.ComponentData.Add(ComponentIndex, *SerializedComponentData);
+			}
+		}
+	}
+	return OutActorData;
+}
+
+#undef LOCTEXT_NAMESPACE
diff --git a/Engine/Plugins/Experimental/Avalanche/Source/Avalanche/Private/AvaBlueprint_Serialize.h b/Engine/Plugins/Experimental/Avalanche/Source/Avalanche/Private/AvaBlueprint_Serialize.h
new file mode 100644
--- /dev/null
+++ b/Engine/Plugins/Experimental/Avalanche/Source/Avalanche/Private/AvaBlueprint_Serialize.h
@@ -0,0 +1,61 @@
+// Copyright Epic Games, Inc. All Rights Reserved.
+
+#pragma once
+
+#include "Data/AvalancheDataDefines.h"
+
+class AActor;
+class FString;
+class UActorComponent;
+class UObject;
+class UWorld;
+struct FAvalancheActorData;
+struct FAvalancheComponentData;
+struct FAvalancheObjectData;
+struct FAvalancheSubObjectData;
+struct FAvalancheWorldData;
+struct FSoftObjectPath;
+template<typename OptionalType> struct TOptional;
+
+struct FAvaBlueprint_Serialize
+{
+	/**
+	 * Takes an existing path to an actor's subobjects and replaces the actor bit with the path to another actor.
+	 * E.g. /Game/MapName.MapName:PersistentLevel.StaticMeshActor_42.StaticMeshComponent could become /Game/MapName.MapName:PersistentLevel.SomeOtherActor.StaticMeshComponent
+	 */
+	static FSoftObjectPath SetActorInPath(AActor* NewActor, const FSoftObjectPath& OriginalObjectPath);
+
+	static UObject* ResolveObjectDependency(FAvalancheWorldData& WorldData, FAvaObjectIndex ObjectIndex);
+
+	static FString ExtractLastSubObjectName(const FSoftObjectPath& ObjectPath);
+
+	static UObject* CreateSubObject(FAvalancheWorldData& WorldData, FAvaObjectIndex ObjectIndex, const FSoftObjectPath& PathToSubObject);
+
+	static FString ExtractRootToLeafComponentPath(const FSoftObjectPath& ComponentPath);
+	static UActorComponent* FindMatchingComponent(AActor* ActorToSearchOn, const FSoftObjectPath& ComponentPath);
+
+	static void SaveObject(FAvalancheWorldData& WorldData, FAvalancheObjectData& ObjectData, UObject* Object);
+	static void LoadObject(FAvalancheWorldData& WorldData, FAvalancheObjectData& ObjectData, UObject* Object);
+
+	static bool ShouldSaveActor(const UWorld* InWorld, const AActor* InActor);
+	static bool ShouldSaveComponent(const UActorComponent* InActorComponent);
+	static bool ShouldSaveSubObject(UObject* InSubObject);
+
+	static UActorComponent* FindOrAllocateComponent(const FAvalancheWorldData& InWorldData
+		, AActor* InRecreatedActor
+		, const FAvalancheActorData& InActorData
+		, const FSoftObjectPath& InComponentPath
+		, FAvalancheSubObjectData& SubObjectData
+		, const FAvalancheComponentData& ComponentData);
+
+	static FAvaObjectIndex AddOrFindObjectReference(FAvalancheWorldData& WorldData, const FSoftObjectPath& ObjectPath);
+	static TOptional<FSoftObjectPath> ExtractActorFromPath(const FSoftObjectPath& InObjectPath, bool& bIsPathToActorSubObject);
+	static void AddSubObjectDependency(FAvalancheWorldData& WorldData, UObject* ReferenceFromOriginalObject, FAvaObjectIndex Index);
+	static FAvaObjectIndex AddObjectDependency(FAvalancheWorldData& WorldData
+		, UObject* ReferenceFromOriginalObject
+		, bool bCheckWhetherSubObject = true);
+
+	static void LoadActor(AActor* InActor, const FSoftObjectPath& InActorPath, FAvalancheWorldData& WorldData);
+	static TOptional<FAvalancheComponentData> SaveComponent(UActorComponent* InComponent);
+	static FAvalancheActorData SaveActor(AActor* InActor, FAvalancheWorldData& WorldData);
+};
