[MotionDesign]

diff --git a/Engine/Plugins/Experimental/Avalanche/Source/Avalanche/Private/AvaScene.cpp b/Engine/Plugins/Experimental/Avalanche/Source/Avalanche/Private/AvaScene.cpp
new file mode 100644
--- /dev/null
+++ b/Engine/Plugins/Experimental/Avalanche/Source/Avalanche/Private/AvaScene.cpp
@@ -0,0 +1,472 @@
+// Copyright Epic Games, Inc. All Rights Reserved.
+
+#include "AvaScene.h"
+#include "AvaAssetTags.h"
+#include "AvaRemoteControlUtils.h"
+#include "AvaSceneSettings.h"
+#include "AvaSceneState.h"
+#include "AvaSceneSubsystem.h"
+#include "AvaSequencePlaybackObject.h"
+#include "AvaSequenceSubsystem.h"
+#include "AvaWorldSubsystemUtils.h"
+#include "Containers/Ticker.h"
+#include "EngineUtils.h"
+#include "RemoteControlPreset.h"
+#include "UObject/AssetRegistryTagsContext.h"
+#include "UObject/UObjectThreadContext.h"
+
+#if WITH_EDITOR
+#include "AvaBlueprint.h"
+#include "AvaField.h"
+#include "Misc/ScopedSlowTask.h"
+#include "RemoteControlBinding.h"
+#endif
+
+DEFINE_LOG_CATEGORY_STATIC(LogAvaScene, Log, All);
+
+#define LOCTEXT_NAMESPACE "AvaScene"
+
+AAvaScene* AAvaScene::GetScene(ULevel* InLevel, bool bInCreateSceneIfNotFound)
+{
+	if (!IsValid(InLevel))
+	{
+		return nullptr;
+	}
+
+	AAvaScene* ExistingScene = nullptr;
+
+	// Return the Existing Scene, or nullptr if not found and not creating a new scene
+	if (InLevel->Actors.FindItemByClass<AAvaScene>(&ExistingScene) || !bInCreateSceneIfNotFound)
+	{
+		return ExistingScene;
+	}
+
+	UWorld* const World = InLevel->GetWorld();
+	if (!IsValid(World))
+	{
+		return nullptr;
+	}
+
+	FActorSpawnParameters SpawnParameters;
+	SpawnParameters.OverrideLevel = InLevel;
+#if WITH_EDITOR
+	SpawnParameters.bHideFromSceneOutliner = true;
+#endif
+
+	AAvaScene* const NewScene = World->SpawnActor<AAvaScene>(SpawnParameters);
+	return NewScene;
+}
+
+#if WITH_EDITOR
+AAvaScene* AAvaScene::CreateFromBlueprint(ULevel* InNewLevel, UAvalancheBlueprint* InBlueprint)
+{
+	if (!IsValid(InNewLevel) || !IsValid(InNewLevel->OwningWorld) || !IsValid(InBlueprint))
+	{
+		return nullptr;
+	}
+
+	FScopedSlowTask SlowTask(100, LOCTEXT("ExportingBlueprintToWorld", "Exporting Blueprint to World"));
+
+	// Load Blueprint World Actors into the new Level 
+	SlowTask.EnterProgressFrame(80, LOCTEXT("CopyingActors", "Copying Blueprint Actors"));
+	{
+		auto InitActor = [](AActor& InActor, const FAvalancheActorData& InActorData)
+		{
+			InActor.bHiddenEd = InActorData.bEditorVisibility;
+		};
+
+		InBlueprint->LoadAvalancheWorld(InNewLevel->OwningWorld, UAvalancheBlueprint::FInitActorFunctionRef(InitActor));
+	}
+
+	// Create Scene Actor in Level
+	AAvaScene* Scene;
+	SlowTask.EnterProgressFrame(5, LOCTEXT("CreatingScene", "Creating Level Scene"));
+	{
+		Scene = GetScene(InNewLevel, /*bInCreateSceneIfNotFound=*/true);
+		check(Scene);
+
+		Scene->SceneTree         = InBlueprint->GetSceneTree();
+		Scene->OutlinerData      = InBlueprint->GetEditorData();
+		Scene->StartupCameraName = InBlueprint->GetStartupCameraName();
+	}
+
+	// Copy Sequences
+	SlowTask.EnterProgressFrame(10, LOCTEXT("CopyingSequences", "Copying Blueprint Sequences"));
+	{
+		// Copy to remove all invalid sequences
+		TArray<UAvaSequence*> SourceSequences = InBlueprint->GetSequences();
+		SourceSequences.RemoveAll([](UAvaSequence* InSequence) { return !IsValid(InSequence); });
+
+		// Expectation is for this new scene to have no sequences
+		ensureAlways(Scene->Animations.IsEmpty());
+		Scene->Animations.Empty(SourceSequences.Num());
+		Scene->DefaultSequenceIndex = SourceSequences.Find(InBlueprint->GetDefaultSequence());
+
+		const FTopLevelAssetPath OldContext = UAvalancheBlueprint::GetWorldContextPath(InBlueprint);
+		const FTopLevelAssetPath NewContext(InNewLevel->OwningWorld);
+
+		for (UAvaSequence* const SourceSequence : SourceSequences)
+		{
+			UAvaSequence* const NewSequence = DuplicateObject<UAvaSequence>(SourceSequence, Scene);
+			ensure(IsValid(NewSequence));
+
+			NewSequence->MigrateLegacyBindings(/*PlaybackContext*/InNewLevel->OwningWorld);
+
+			NewSequence->UpdateBindings(&OldContext, NewContext);
+
+			// Instead of calling AAvaScene::AddSequence which has some check overhead and schedules a Tree Rebuild, handle here manually
+			Scene->Animations.Add(NewSequence);
+		}
+
+		if (!Scene->Animations.IsValidIndex(Scene->DefaultSequenceIndex))
+		{
+			Scene->DefaultSequenceIndex = 0;
+		}
+
+		check(SourceSequences.Num() == Scene->Animations.Num());
+
+		for (int32 Index = 0; Index < SourceSequences.Num(); ++Index)
+		{
+			UAvaSequence* const SourceSequence = SourceSequences[Index];
+			UAvaSequence* const TargetSequence = Scene->Animations[Index];
+
+			// On Duplicate, the Child Anim property is marked as Duplicate Transient so should've not duplicated over
+			ensure(TargetSequence->GetChildren().IsEmpty());
+
+			for (const TWeakObjectPtr<UAvaSequence>& SourceChildren : SourceSequence->GetChildren())
+			{
+				const int32 ChildIndex = SourceSequences.Find(SourceChildren.Get());
+				if (ChildIndex != INDEX_NONE)
+				{
+					TargetSequence->AddChild(Scene->Animations[ChildIndex]);	
+				}
+			}
+		}
+
+		Scene->RebuildSequenceTree();
+	}
+
+	// Copy Remote Control Preset
+	SlowTask.EnterProgressFrame(5, LOCTEXT("CopyingRemoteControlPreset", "Copying Blueprint Remote Control Preset"));
+	{
+		URemoteControlPreset* const SourcePreset = InBlueprint->GetRemoteControlPreset();
+		if (IsValid(SourcePreset))
+		{
+			const EObjectFlags Flags   = IsValid(Scene->RemoteControlPreset) ? Scene->RemoteControlPreset->GetFlags() : SourcePreset->GetFlags();
+			Scene->RemoteControlPreset = DuplicateObject<URemoteControlPreset>(SourcePreset, Scene);
+			Scene->RemoteControlPreset->SetFlags(Flags);
+
+			FMapProperty* LevelMapProperty  = UE::AvaCore::GetProperty<URemoteControlLevelDependantBinding, FMapProperty>(TEXT("SubLevelSelectionMapByPath"));
+			FMapProperty* ObjectMapProperty = UE::AvaCore::GetProperty<URemoteControlLevelDependantBinding, FMapProperty>(TEXT("BoundObjectMapByPath"));
+			FProperty* LastLevelProperty    = UE::AvaCore::GetProperty<URemoteControlLevelDependantBinding>(TEXT("LevelWithLastSuccessfulResolve"));
+			FProperty* LastObjectProperty   = UE::AvaCore::GetProperty<URemoteControlBinding>(TEXT("LastBoundObjectPath"));
+
+			FSoftObjectPath NullObjectPath;
+			FSoftObjectPtr NullObjectPtr;
+
+			auto ResetMap = [](UObject* InObject, FMapProperty* InMapProperty)
+			{
+				FScriptMapHelper MapHelper(InMapProperty, InMapProperty->ContainerPtrToValuePtr<void*>(InObject));
+				MapHelper.EmptyValues();
+			};
+
+			for (TArray<TObjectPtr<URemoteControlBinding>>::TIterator BindingIter =  Scene->RemoteControlPreset->Bindings.CreateIterator(); BindingIter; ++BindingIter)
+			{
+				if (URemoteControlLevelDependantBinding* LevelBinding = Cast<URemoteControlLevelDependantBinding>(*BindingIter))
+				{
+					ResetMap(LevelBinding, LevelMapProperty);
+					ResetMap(LevelBinding, ObjectMapProperty);
+
+					LastLevelProperty->SetValue_InContainer(LevelBinding, &NullObjectPtr);
+					LastObjectProperty->SetValue_InContainer(LevelBinding, &NullObjectPath);
+				}
+			}
+
+			Scene->RemoteControlPreset->RebindUnboundEntities();
+		}
+	}
+
+	return Scene;
+}
+#endif
+
+AAvaScene::AAvaScene()
+{
+	SceneSettings = CreateDefaultSubobject<UAvaSceneSettings>(TEXT("SceneSettings"));
+
+	SceneState = CreateDefaultSubobject<UAvaSceneState>(TEXT("SceneState"));
+
+	RemoteControlPreset = CreateDefaultSubobject<URemoteControlPreset>(TEXT("RemoteControlPreset"));
+
+	StartupCameraName = NAME_None;
+
+#if WITH_EDITOR
+	if (!HasAnyFlags(EObjectFlags::RF_ClassDefaultObject))
+	{
+		PreWorldRenameDelegate = FWorldDelegates::OnPreWorldRename.AddUObject(this, &AAvaScene::OnWorldRenamed);
+		WorldTagGetterDelegate = UObject::FAssetRegistryTag::OnGetExtraObjectTagsWithContext.AddUObject(this, &AAvaScene::OnGetWorldTags);
+	}
+#endif
+}
+
+IAvaSequencePlaybackObject* AAvaScene::GetScenePlayback() const
+{
+	// If the saved scene playback is valid use that
+	if (IAvaSequencePlaybackObject* const ScenePlayback = PlaybackObject.GetInterface())
+	{
+		return ScenePlayback;
+	}
+
+	UAvaSequenceSubsystem* const SequenceSubsystem = UAvaSequenceSubsystem::Get(GetWorld());
+	if (!SequenceSubsystem)
+	{
+		return nullptr;
+	}
+
+	AAvaScene& MutableThis = *const_cast<AAvaScene*>(this);
+	if (IAvaSequencePlaybackObject* const ScenePlayback = SequenceSubsystem->FindOrCreatePlaybackObject(GetLevel(), MutableThis))
+	{
+		MutableThis.PlaybackObject.SetObject(ScenePlayback->ToUObject());
+		MutableThis.PlaybackObject.SetInterface(ScenePlayback);
+		return ScenePlayback;
+	}
+
+	return nullptr;
+}
+
+#if WITH_EDITOR
+void AAvaScene::OnWorldRenamed(UWorld* InWorld, const TCHAR* InName, UObject* InNewOuter, ERenameFlags InFlags, bool& bOutShouldFailRename)
+{
+	if (FUObjectThreadContext::Get().IsRoutingPostLoad || InWorld != GetWorld())
+	{
+		return;
+	}
+
+	for (UAvaSequence* Sequence : Animations)
+	{
+		if (Sequence)
+		{
+			Sequence->OnOuterWorldRenamed(InName, InNewOuter, InFlags, bOutShouldFailRename);
+		}
+	}
+}
+
+void AAvaScene::OnGetWorldTags(FAssetRegistryTagsContext Context) const
+{
+	const UObject* InWorld = Context.GetObject();
+	if (InWorld != GetTypedOuter<UWorld>())
+	{
+		return;
+	}
+
+	using namespace UE::Ava;
+	Context.AddTag(UObject::FAssetRegistryTag(AssetTags::AvalancheScene, AssetTags::Values::Enabled, UObject::FAssetRegistryTag::TT_Alphabetical));
+}
+#endif
+
+ULevel* AAvaScene::GetSceneLevel() const
+{
+	return GetLevel();
+}
+
+IAvaSequencePlaybackObject* AAvaScene::GetPlaybackObject() const
+{
+	return GetScenePlayback();
+}
+
+UObject* AAvaScene::ToUObject()
+{
+	return this;
+}
+
+UWorld* AAvaScene::GetContextWorld() const
+{
+	return GetWorld();
+}
+
+bool AAvaScene::CreateDirectorInstance(UAvaSequence& InSequence
+	, IMovieScenePlayer& InPlayer
+	, const FMovieSceneSequenceID& InSequenceID
+	, UObject*& OutDirectorInstance)
+{
+	// Allow ULevelSequence::CreateDirectorInstance to be called
+	return false;
+}
+
+bool AAvaScene::AddSequence(UAvaSequence* InSequence)
+{
+	if (IsValid(InSequence) && !Animations.Contains(InSequence))
+	{
+		Animations.Add(InSequence);
+		ScheduleRebuildSequenceTree();
+		return true;
+	}
+	return false;
+}
+
+void AAvaScene::RemoveSequence(UAvaSequence* InSequence)
+{
+	Animations.Remove(InSequence);
+	if (InSequence)
+	{
+		InSequence->OnSequenceRemoved();
+	}
+	ScheduleRebuildSequenceTree();
+}
+
+void AAvaScene::SetDefaultSequence(UAvaSequence* InSequence)
+{
+	if (IsValid(InSequence))
+	{
+		AddSequence(InSequence);
+		DefaultSequenceIndex = Animations.Find(InSequence);
+	}
+}
+
+UAvaSequence* AAvaScene::GetDefaultSequence() const
+{
+	if (Animations.IsValidIndex(DefaultSequenceIndex))
+	{
+		return Animations[DefaultSequenceIndex];
+	}
+	return nullptr;
+}
+
+FName AAvaScene::GetSequenceProviderDebugName() const
+{
+	return GetFName();
+}
+
+#if WITH_EDITOR
+void AAvaScene::OnEditorSequencerCreated(const TSharedPtr<ISequencer>& InSequencer)
+{
+	EditorSequencer = InSequencer;
+	RebuildSequenceTree();
+}
+#endif
+
+void AAvaScene::ScheduleRebuildSequenceTree()
+{
+	// Bail if the Deferred Rebuild is pending and hasn't executed yet
+	if (bPendingAnimTreeUpdate)
+	{
+		return;
+	}
+
+	bPendingAnimTreeUpdate = true;
+
+	TWeakObjectPtr<AAvaScene> ThisWeak(this);
+
+	FTSTicker::GetCoreTicker().AddTicker(FTickerDelegate::CreateLambda(
+		[ThisWeak](float InDeltaTime)->bool
+		{
+			AAvaScene* const This = ThisWeak.Get();
+
+			// Check if we already have Rebuilt the Animation Tree in between when this was added and when it was executed
+			if (This && This->bPendingAnimTreeUpdate)
+			{
+				This->RebuildSequenceTree();
+			}
+
+			// Return false for one time execution
+			return false;
+		}));
+}
+
+void AAvaScene::RebuildSequenceTree()
+{
+	bPendingAnimTreeUpdate = false;
+	IAvaSequenceProvider::RebuildSequenceTree();
+}
+
+void AAvaScene::BeginPlay()
+{
+	Super::BeginPlay();
+
+	if (SceneState)
+	{
+		SceneState->OnBeginPlay(SceneSettings);
+	}
+}
+
+void AAvaScene::PostActorCreated()
+{
+	Super::PostActorCreated();
+
+	if (!HasAnyFlags(EObjectFlags::RF_ClassDefaultObject))
+	{
+		FAvaRemoteControlUtils::RegisterRemoteControlPreset(RemoteControlPreset, /*bInEnsureUniqueId*/ false);
+	}
+
+	// Register AAvaScenes created past Subsystem Initialization
+	if (UAvaSceneSubsystem* SceneSubsystem = FAvaWorldSubsystemUtils::GetWorldSubsystem<UAvaSceneSubsystem>(this))
+	{
+		SceneSubsystem->RegisterSceneInterface(GetLevel(), this);
+	}
+}
+
+void AAvaScene::PostLoad()
+{
+	Super::PostLoad();
+
+	if (!HasAnyFlags(EObjectFlags::RF_ClassDefaultObject))
+	{
+		FAvaRemoteControlUtils::RegisterRemoteControlPreset(RemoteControlPreset, /*bInEnsureUniqueId*/ true);
+	}
+
+	// Register AAvaScenes created past Subsystem Initialization
+	if (UAvaSceneSubsystem* SceneSubsystem = FAvaWorldSubsystemUtils::GetWorldSubsystem<UAvaSceneSubsystem>(this))
+	{
+		SceneSubsystem->RegisterSceneInterface(GetLevel(), this);
+	}
+}
+
+void AAvaScene::PostDuplicate(EDuplicateMode::Type InDuplicateMode)
+{
+	Super::PostDuplicate(InDuplicateMode);
+	if (!HasAnyFlags(EObjectFlags::RF_ClassDefaultObject))
+	{
+		FAvaRemoteControlUtils::RegisterRemoteControlPreset(RemoteControlPreset, /*bInEnsureUniqueId*/ true);
+	}
+}
+
+void AAvaScene::PostEditImport()
+{
+	Super::PostEditImport();
+	if (!HasAnyFlags(EObjectFlags::RF_ClassDefaultObject))
+	{
+		FAvaRemoteControlUtils::RegisterRemoteControlPreset(RemoteControlPreset, /*bInEnsureUniqueId*/ true);
+	}
+}
+
+void AAvaScene::BeginDestroy()
+{
+	Super::BeginDestroy();
+
+	if (!HasAnyFlags(EObjectFlags::RF_ClassDefaultObject))
+	{
+		FAvaRemoteControlUtils::UnregisterRemoteControlPreset(RemoteControlPreset);
+	}
+
+#if WITH_EDITOR
+	FWorldDelegates::OnPreWorldRename.Remove(PreWorldRenameDelegate);
+	PreWorldRenameDelegate.Reset();
+
+	UObject::FAssetRegistryTag::OnGetExtraObjectTagsWithContext.Remove(WorldTagGetterDelegate);
+	WorldTagGetterDelegate.Reset();
+#endif
+}
+
+#if WITH_EDITOR
+void AAvaScene::SetStartupCameraName(FName InName)
+{
+	StartupCameraName = InName;
+}
+
+void AAvaScene::SetViewportGuideData(const TConstArrayView<FAvaViewportGuideInfo_Deprecated>& InGuideData)
+{
+	GuideData = InGuideData;
+}
+#endif
+
+#undef LOCTEXT_NAMESPACE
diff --git a/Engine/Plugins/Experimental/Avalanche/Source/Avalanche/Private/AvaSceneSettings.cpp b/Engine/Plugins/Experimental/Avalanche/Source/Avalanche/Private/AvaSceneSettings.cpp
new file mode 100644
--- /dev/null
+++ b/Engine/Plugins/Experimental/Avalanche/Source/Avalanche/Private/AvaSceneSettings.cpp
@@ -0,0 +1,8 @@
+// Copyright Epic Games, Inc. All Rights Reserved.
+
+#include "AvaSceneSettings.h"
+
+void UAvaSceneSettings::SetTagAttributes(const FAvaTagHandleContainer& InTagAttributes)
+{
+	TagAttributes = InTagAttributes;
+}
diff --git a/Engine/Plugins/Experimental/Avalanche/Source/Avalanche/Private/AvaSceneState.cpp b/Engine/Plugins/Experimental/Avalanche/Source/Avalanche/Private/AvaSceneState.cpp
new file mode 100644
--- /dev/null
+++ b/Engine/Plugins/Experimental/Avalanche/Source/Avalanche/Private/AvaSceneState.cpp
@@ -0,0 +1,47 @@
+// Copyright Epic Games, Inc. All Rights Reserved.
+
+#include "AvaSceneState.h"
+#include "AvaSceneSettings.h"
+#include "AvaTagHandle.h"
+
+void UAvaSceneState::OnBeginPlay(UAvaSceneSettings* InSceneSettings)
+{
+	SceneSettingsWeak = InSceneSettings;
+}
+
+bool UAvaSceneState::AddTagAttribute(const FAvaTagHandle& InTagHandle)
+{
+	if (const FAvaTag* Tag = InTagHandle.GetTag())
+	{
+		bool bAlreadyInSet;
+		ActiveTagAttributes.Add(*Tag, &bAlreadyInSet);
+		return !bAlreadyInSet;
+	}
+	return false;
+}
+
+bool UAvaSceneState::RemoveTagAttribute(const FAvaTagHandle& InTagHandle)
+{
+	if (const FAvaTag* Tag = InTagHandle.GetTag())
+	{
+		return ActiveTagAttributes.Remove(*Tag) > 0;
+	}
+	return false;
+}
+
+bool UAvaSceneState::ContainsTagAttribute(const FAvaTagHandle& InTagHandle) const
+{
+	const FAvaTag* Tag = InTagHandle.GetTag();
+	if (!Tag)
+	{
+		return false;
+	}
+
+	if (ActiveTagAttributes.Contains(*Tag))
+	{
+		return true;
+	}
+
+	UAvaSceneSettings* SceneSettings = SceneSettingsWeak.Get();
+	return SceneSettings && SceneSettings->GetTagAttributes().ContainsTag(InTagHandle);
+}
diff --git a/Engine/Plugins/Experimental/Avalanche/Source/Avalanche/Private/Data/AvalancheVersionInfo.cpp b/Engine/Plugins/Experimental/Avalanche/Source/Avalanche/Private/Data/AvalancheVersionInfo.cpp
new file mode 100644
--- /dev/null
+++ b/Engine/Plugins/Experimental/Avalanche/Source/Avalanche/Private/Data/AvalancheVersionInfo.cpp
@@ -0,0 +1,85 @@
+// Copyright Epic Games, Inc. All Rights Reserved.
+
+#include "Data/AvalancheVersionInfo.h"
+#include "Misc/EngineVersion.h"
+#include "Serialization/CustomVersion.h"
+#include "UObject/ObjectVersion.h"
+
+void FAvalancheFileVersionInfo::UpdateToLatest()
+{
+	FileVersionUE4      = GPackageFileUEVersion.FileVersionUE4;
+	FileVersionUE5      = GPackageFileUEVersion.FileVersionUE5;
+	FileVersionLicensee = GPackageFileLicenseeUEVersion;
+}
+
+void FAvalancheEngineVersionInfo::UpdateToLatest()
+{
+	const FEngineVersion& EngineVersion(FEngineVersion::Current());
+	Major      = EngineVersion.GetMajor();
+	Minor      = EngineVersion.GetMinor();
+	Patch      = EngineVersion.GetPatch();
+	Changelist = EngineVersion.GetChangelist();
+}
+
+void FAvalancheCustomVersionInfo::SetCustomVersion(const FCustomVersion& InVersion)
+{
+	FriendlyName = InVersion.GetFriendlyName();
+	Key          = InVersion.Key;
+	Version      = InVersion.Version;
+}
+
+void FAvalancheVersionInfo::UpdateToLatest()
+{
+	FileVersion.UpdateToLatest();
+	EngineVersion.UpdateToLatest();
+
+	if (!UsedCustomVersions.IsEmpty())
+	{
+		CustomVersions.Empty(UsedCustomVersions.Num());
+		for (const FGuid& UsedVersionGuid : UsedCustomVersions)
+		{
+			const TOptional<FCustomVersion> CustomVersion = FCurrentCustomVersions::Get(UsedVersionGuid);
+			if (CustomVersion.IsSet())
+			{
+				CustomVersions.AddDefaulted_GetRef().SetCustomVersion(CustomVersion.GetValue());
+			}
+		}
+	}
+	else
+	{
+		FCustomVersionArray CurrentCustomVersions = FCurrentCustomVersions::GetAll().GetAllVersions();
+		CustomVersions.Empty(CurrentCustomVersions.Num());
+	
+		for (const FCustomVersion& CustomVersion : CurrentCustomVersions)
+		{
+			CustomVersions.AddDefaulted_GetRef().SetCustomVersion(CustomVersion);
+		}
+	}
+}
+
+void FAvalancheVersionInfo::ApplyToArchive(FArchive& Archive) const
+{
+	//If CustomVersions is Empty, it's pre-versioning
+	if (CustomVersions.IsEmpty())
+	{
+		return;
+	}
+	
+	FPackageFileVersion UEVersion;
+	UEVersion.FileVersionUE4 = FileVersion.FileVersionUE4;
+	UEVersion.FileVersionUE5 = FileVersion.FileVersionUE5;
+
+	Archive.SetUEVer(UEVersion);
+	Archive.SetLicenseeUEVer(FileVersion.FileVersionLicensee);
+	Archive.SetEngineVer(FEngineVersionBase(EngineVersion.Major, EngineVersion.Minor, EngineVersion.Patch, EngineVersion.Changelist));
+
+	FCustomVersionContainer CustomVersionContainer;
+	for (const FAvalancheCustomVersionInfo& CustomVersion : CustomVersions)
+	{
+		CustomVersionContainer.SetVersion(CustomVersion.Key
+			, CustomVersion.Version
+			, CustomVersion.FriendlyName);
+	}
+	
+	Archive.SetCustomVersions(CustomVersionContainer);
+}
diff --git a/Engine/Plugins/Experimental/Avalanche/Source/Avalanche/Private/Data/AvalancheWorldData.cpp b/Engine/Plugins/Experimental/Avalanche/Source/Avalanche/Private/Data/AvalancheWorldData.cpp
new file mode 100644
--- /dev/null
+++ b/Engine/Plugins/Experimental/Avalanche/Source/Avalanche/Private/Data/AvalancheWorldData.cpp
@@ -0,0 +1,17 @@
+// Copyright Epic Games, Inc. All Rights Reserved.
+
+#include "Data/AvalancheWorldData.h"
+
+void FAvalancheWorldData::FinalizeSave()
+{
+	VersionInfo.UpdateToLatest();
+}
+
+void FAvalancheWorldData::ResetTransientData()
+{
+	World.Reset();
+	NameIndexMap.Reset();
+	ObjectReferenceIndexMap.Reset();
+	CachedActors.Reset();
+	CachedSubObjects.Reset();
+}
diff --git a/Engine/Plugins/Experimental/Avalanche/Source/Avalanche/Private/Framework/AvaSoftAssetPtr.cpp b/Engine/Plugins/Experimental/Avalanche/Source/Avalanche/Private/Framework/AvaSoftAssetPtr.cpp
new file mode 100644
--- /dev/null
+++ b/Engine/Plugins/Experimental/Avalanche/Source/Avalanche/Private/Framework/AvaSoftAssetPtr.cpp
@@ -0,0 +1,70 @@
+// Copyright Epic Games, Inc. All Rights Reserved.
+
+#include "Framework/AvaSoftAssetPtr.h"
+#include "AvaBlueprint.h"
+
+EAvalancheAssetType FAvaSoftAssetPath::GetAssetTypeFromClass(const FSoftClassPath& InAssetClassPath, bool bInLoadIfUnknown)
+{
+	static const FSoftClassPath AvalancheClassPath(UAvalancheBlueprint::StaticClass());
+	static const FSoftClassPath WorldClassPath(UWorld::StaticClass());
+
+	// Try comparing the asset paths directly.
+	// That is probably the fastest way, but can fail for derived types.
+	if (InAssetClassPath == AvalancheClassPath)
+	{
+		return EAvalancheAssetType::Blueprint;
+	}
+	if (InAssetClassPath == WorldClassPath)
+	{
+		return EAvalancheAssetType::World;
+	}
+
+	// Try to load the asset class and compare possible derived types.
+	if (InAssetClassPath.IsValid())
+	{
+		const UClass* AssetClass = InAssetClassPath.ResolveClass();
+		if (AssetClass->IsChildOf(UAvalancheBlueprint::StaticClass()))
+		{
+			return EAvalancheAssetType::Blueprint;
+		}
+		if (AssetClass->IsChildOf(UWorld::StaticClass()))
+		{
+			return EAvalancheAssetType::World;
+		}
+	}
+	return EAvalancheAssetType::Unknown;
+}
+
+EAvalancheAssetType FAvaSoftAssetPtr::GetAssetType(bool bInLoadIfUnknown) const
+{
+	const EAvalancheAssetType AssetTypeFromClass = FAvaSoftAssetPath::GetAssetTypeFromClass(AssetClassPath, bInLoadIfUnknown);
+	if (AssetTypeFromClass != EAvalancheAssetType::Unknown)
+	{
+		return AssetTypeFromClass;
+	}
+	
+	if (bInLoadIfUnknown)
+	{
+		// Todo: problem with loading the asset sync here.
+		// Todo: Investigate LoadPackageAsync.
+		// This will cause a hitch. For level asset, we want to avoid that.
+		UObject* LoadedSourceAsset;
+		{
+			TRACE_CPUPROFILER_EVENT_SCOPE(FAvaSoftAssetPtr::LoadSourceAsset);
+			const FSoftObjectPtr SourceAvalancheAsset(AssetPtr.ToSoftObjectPath());
+			LoadedSourceAsset = SourceAvalancheAsset.LoadSynchronous();
+		}
+	
+		if (Cast<UAvalancheBlueprint>(LoadedSourceAsset))
+		{
+			return EAvalancheAssetType::Blueprint;
+		}
+	
+		if (Cast<UWorld>(LoadedSourceAsset))
+		{
+			return EAvalancheAssetType::World;
+		}
+	}
+	
+	return EAvalancheAssetType::Unknown;
+}
\ No newline at end of file
diff --git a/Engine/Plugins/Experimental/Avalanche/Source/Avalanche/Private/Framework/AvalancheGameInstance.cpp b/Engine/Plugins/Experimental/Avalanche/Source/Avalanche/Private/Framework/AvalancheGameInstance.cpp
new file mode 100644
--- /dev/null
+++ b/Engine/Plugins/Experimental/Avalanche/Source/Avalanche/Private/Framework/AvalancheGameInstance.cpp
@@ -0,0 +1,511 @@
+// Copyright Epic Games, Inc. All Rights Reserved.
+
+#include "Framework/AvalancheGameInstance.h"
+
+#include "AudioDevice.h"
+#include "AvaBlueprint.h"
+#include "AvaRemoteControlRebind.h"
+#include "Components/ReflectionCaptureComponent.h"
+#include "Components/SkyLightComponent.h"
+#include "Engine/Console.h"
+#include "Engine/Engine.h"
+#include "Engine/LocalPlayer.h"
+#include "Engine/TextureRenderTarget2D.h"
+#include "Framework/Application/SlateApplication.h"
+#include "Misc/TimeGuard.h"
+#include "Slate/SceneViewport.h"
+#include "Viewport/AvaCameraManager.h"
+
+namespace UE::AvalancheGameInstance::Private
+{
+	/* Frame number marked where a synchronous asset loading occured. */
+	static uint64 MarkedSynchronousAssetLoadingFrame = 0;
+
+	/* Frame number marked where delta seconds are ignored. */
+	static uint64 MarkedIgnoreDeltaSecondsFrame = 0;
+
+	static void UpdateFrameMarkers()
+	{
+		// If the last frame had sync asset loading, mark current frame to ignore delta seconds.
+		if (MarkedSynchronousAssetLoadingFrame != 0 && GFrameCounter > MarkedSynchronousAssetLoadingFrame)
+		{
+			MarkedIgnoreDeltaSecondsFrame = GFrameCounter;
+			MarkedSynchronousAssetLoadingFrame = 0;
+		}
+
+		// Reset marker if stale.
+		if (MarkedIgnoreDeltaSecondsFrame != 0 && GFrameCounter > MarkedIgnoreDeltaSecondsFrame)
+		{
+			MarkedIgnoreDeltaSecondsFrame = 0;
+		}
+	}
+	
+	static void MarkSynchronousAssetLoadingThisFrame()
+	{
+		UpdateFrameMarkers();
+		MarkedSynchronousAssetLoadingFrame = GFrameCounter;
+	}
+	
+	static bool ShouldIgnoreDeltaSecondsForCurrentFrame()
+	{
+		UpdateFrameMarkers();
+		return MarkedIgnoreDeltaSecondsFrame != 0 && GFrameCounter == MarkedIgnoreDeltaSecondsFrame;
+	}
+}
+
+UAvalancheGameInstance::FOnAvaGameInstanceEvent UAvalancheGameInstance::OnEndPlay;
+UAvalancheGameInstance::FOnAvaGameInstanceEvent UAvalancheGameInstance::OnRenderTargetReady;
+
+UAvalancheGameInstance* UAvalancheGameInstance::Create(UObject* InOuter, const TSoftObjectPtr<UAvalancheBlueprint>& InAvalancheBlueprintTemplate)
+{
+	UAvalancheGameInstance* GameInstance = NewObject<UAvalancheGameInstance>(InOuter ? InOuter : GEngine);
+	GameInstance->CreateWorld();
+	GameInstance->LoadAvalancheBlueprint(InAvalancheBlueprintTemplate);
+	return GameInstance;
+}
+
+UAvalancheGameInstance* UAvalancheGameInstance::Create(UObject* InOuter)
+{
+	UAvalancheGameInstance* GameInstance = NewObject<UAvalancheGameInstance>(InOuter ? InOuter : GEngine);
+
+	// We need to call Init() (and corresponding Shutdown() later) to get SubsystemCollection to initialize.
+	// Not doing it because it might have side effects.
+	//GameInstance->Init();
+
+	GameInstance->CreateWorld();
+	return GameInstance;
+}
+
+bool UAvalancheGameInstance::CreateWorld()
+{
+	if (bWorldCreated)
+	{
+		return false;
+	}
+	
+	TRACE_CPUPROFILER_EVENT_SCOPE(UAvalancheGameInstance::CreateWorld);
+
+	if (!EnginePreExitHandle.IsValid())
+	{
+		EnginePreExitHandle = FCoreDelegates::OnEnginePreExit.AddUObject(this, &UAvalancheGameInstance::OnEnginePreExit);
+	}
+	
+	constexpr EWorldType::Type WorldType = EWorldType::Game;
+	
+	{
+		TRACE_CPUPROFILER_EVENT_SCOPE(UAvalancheGameInstance::LoadInstance::InitWorld);
+		
+		const FName AvalancheWorldName = MakeUniqueObjectName(GetOuter(), UWorld::StaticClass(), TEXT("AvalancheGameInstanceWorld"));
+		
+		PlayWorld = NewObject<UWorld>(this, AvalancheWorldName, RF_Transient);
+		check(PlayWorld);
+		PlayWorld->WorldType = WorldType;
+		PlayWorld->InitializeNewWorld(UWorld::InitializationValues()
+			.AllowAudioPlayback(true)
+			.CreatePhysicsScene(true)
+			.RequiresHitProxies(false)
+			.CreateNavigation(true)
+			.CreateAISystem(false)	// Disabling AI System until supported in Avalanche.
+			.ShouldSimulatePhysics(true)
+			.SetTransactional(false)
+			//TODO: World Settings
+		);
+
+		// Disallow the Engine to tick this World as this Game Instance will be in charge of ticking it
+		PlayWorld->SetShouldTick(false);
+	}
+
+	WorldContext = &GEngine->CreateNewWorldContext(WorldType);
+	check(WorldContext);
+	WorldContext->OwningGameInstance = this;
+	WorldContext->SetCurrentWorld(PlayWorld.Get());
+
+	PlayWorld->SetGameInstance(this);
+	PlayWorld->SetGameMode(FURL());
+
+	bWorldCreated = true;
+	return true;
+}
+
+bool UAvalancheGameInstance::LoadAvalancheBlueprint(const TSoftObjectPtr<UAvalancheBlueprint>& InSourceAvalancheBlueprint)
+{
+	if (ManagedAvalancheBlueprint && SourceAvalancheBlueprint == InSourceAvalancheBlueprint)
+	{
+		return false;
+	}
+
+	if (!bWorldCreated)
+	{
+		CreateWorld();
+	}
+
+	UAvalancheBlueprint* Source;
+	{
+		TRACE_CPUPROFILER_EVENT_SCOPE(UAvalancheGameInstance::LoadAsset::LoadSourceAvaBp);
+		Source = InSourceAvalancheBlueprint.LoadSynchronous();
+		check(Source);
+	}
+
+#if WITH_EDITORONLY_DATA
+	FGuardValue_Bitfield(Source->bDuplicatingReadOnly, true);
+#endif
+
+	{
+		TRACE_CPUPROFILER_EVENT_SCOPE(UAvalancheGameInstance::LoadAsset::DupAvaBp);
+		const FName AvaInstanceName = MakeUniqueObjectName(this, UAvalancheBlueprint::StaticClass(), Source->GetFName());
+		ManagedAvalancheBlueprint = Cast<UAvalancheBlueprint>(StaticDuplicateObject(Source, this, AvaInstanceName, RF_Transient));
+	}
+	
+	{
+		TRACE_CPUPROFILER_EVENT_SCOPE(UAvalancheGameInstance::LoadAsset::LoadAvaBpWorld);
+		ManagedAvalancheBlueprint->SetAvalancheWorld(PlayWorld.Get());
+		ManagedAvalancheBlueprint->LoadAvalancheWorld();
+	}
+	
+	ManagedAvalancheBlueprint->UpdatePlaceholderActor();
+
+	{
+		TRACE_CPUPROFILER_EVENT_SCOPE(UAvalancheGameInstance::LoadAsset::RebindRCP);
+		ManagedAvalancheBlueprint->RegisterRemoteControlPreset();
+		FAvaRemoteControlRebind::RebindUnboundEntities(ManagedAvalancheBlueprint->GetRemoteControlPreset());
+	}
+	
+	MarkSynchronousAssetLoadingThisFrame();
+
+	SourceAvalancheBlueprint = InSourceAvalancheBlueprint;
+	return true;
+}
+
+bool UAvalancheGameInstance::BeginPlayWorld(const FAvalancheInstancePlaySettings& InWorldPlaySettings)
+{
+	// Make sure we don't have pending unload or stop requests left over in the game instance. 
+	CancelWorldRequests();
+	
+	if (bWorldPlaying)
+	{
+		return false;
+	}
+
+	TRACE_CPUPROFILER_EVENT_SCOPE(UAvalancheGameInstance::BeginPlay);
+	TRACE_BOOKMARK(TEXT("UAvalancheGameInstance::BeginPlay"));
+	
+	ViewportClient = NewObject<UAvalancheGameViewportClient>(GEngine, NAME_None, RF_Transient);
+	check(ViewportClient);
+
+	// Note: UGameViewportClient::Init ignores "bCreateNewAudioDevice" parameter and always create a device for the world. 
+	ViewportClient->Init(*WorldContext, this);
+	ViewportClient->SetRenderTarget(InWorldPlaySettings.RenderTarget);
+	ViewportClient->bIsPlayInEditorViewport = false;
+	FAvaViewportQualitySettings QualitySettingsMutable = InWorldPlaySettings.QualitySettings; 
+	QualitySettingsMutable.Apply(ViewportClient->EngineShowFlags);
+	
+#if ALLOW_CONSOLE
+	// Create the viewport's console.
+	ViewportClient->ViewportConsole = NewObject<UConsole>(ViewportClient.Get(), GEngine->ConsoleClass);
+	// register console to get all log messages
+	GLog->AddOutputDevice(ViewportClient->ViewportConsole);
+#endif
+
+	TSharedPtr<SViewport> ViewportWidget;
+	Viewport = MakeShareable(ViewportClient->CreateGameViewport(ViewportWidget));
+	Viewport->SetInitialSize(InWorldPlaySettings.ViewportSize);
+
+	PlayWorld->InitializeActorsForPlay(PlayWorld->URL);
+	PlayWorld->BeginPlay();
+	
+	FCoreDelegates::OnEndFrame.AddUObject(this, &UAvalancheGameInstance::OnEndFrameTick);
+
+	bWorldPlaying = true;
+	PlayingChannelName = InWorldPlaySettings.ChannelName;
+	return true;
+}
+
+void UAvalancheGameInstance::RequestEndPlayWorld(bool bForceImmediate)
+{
+	if (bWorldPlaying)
+	{
+		if (bForceImmediate)
+		{
+			EndPlayWorld();
+		}
+		else
+		{
+			bRequestEndPlayWorld = true;
+		}
+	}
+}
+
+void UAvalancheGameInstance::BeginPlayAvalancheBlueprint(const TSoftObjectPtr<UAvalancheBlueprint>& InSourceAvalancheBlueprint, const FAvalancheInstancePlaySettings& InWorldPlaySettings)
+{
+	// World States: Created -> Playing
+	// Asset States: Unloaded -> Loading -> Loaded -> Playing (Visible)
+	const FSoftObjectPath& SourceAssetPath = InSourceAvalancheBlueprint.ToSoftObjectPath();
+
+	if (!bWorldPlaying)
+	{
+		BeginPlayWorld(InWorldPlaySettings);
+	}
+	else
+	{
+		UpdateRenderTarget(InWorldPlaySettings.RenderTarget);
+		UpdateSceneViewportSize(InWorldPlaySettings.ViewportSize);
+	}
+
+	// Check if the asset is loaded.
+	if (SourceAvalancheBlueprint != InSourceAvalancheBlueprint)
+	{
+		LoadAvalancheBlueprint(InSourceAvalancheBlueprint);
+	}	
+	
+	if (ManagedAvalancheBlueprint && SourceAvalancheBlueprint == InSourceAvalancheBlueprint)
+	{
+		// Verify if this can be changed on the fly.
+		ManagedAvalancheBlueprint->GetViewportQualitySettings().Apply(ViewportClient->EngineShowFlags);
+	
+		// Old code using ava camera manager. To retire.
+		constexpr bool bIsCanvasController = false;
+		ViewportClient->GetCameraManager()->Init(ManagedAvalancheBlueprint->GetPlaybackObject(), bIsCanvasController);
+#if WITH_EDITOR
+		ViewportClient->GetCameraManager()->SetDefaultViewTarget(PlayWorld, ManagedAvalancheBlueprint->GetStartupCameraName());
+#endif
+	}
+}
+
+void UAvalancheGameInstance::RequestUnloadWorld(bool bForceImmediate)
+{
+	if (bWorldCreated)
+	{
+		if (bForceImmediate)
+		{
+			UnloadWorld();
+		}
+		else
+		{
+			bRequestUnloadWorld = true;
+		}
+	}
+}
+
+void UAvalancheGameInstance::CancelWorldRequests()
+{
+	bRequestEndPlayWorld = false;
+	bRequestUnloadWorld = false;
+}
+
+void UAvalancheGameInstance::UpdateRenderTarget(UTextureRenderTarget2D* InRenderTarget)
+{
+	if (ViewportClient)
+	{
+		ViewportClient->SetRenderTarget(InRenderTarget);
+	}
+}
+
+UTextureRenderTarget2D* UAvalancheGameInstance::GetRenderTarget() const
+{
+	if (ViewportClient)
+	{
+		return ViewportClient->GetRenderTarget();
+	}
+	return nullptr;
+}
+
+void UAvalancheGameInstance::UpdateSceneViewportSize(const FIntPoint& InViewportSize)
+{
+	if (Viewport.IsValid() && Viewport->GetSize() != InViewportSize)
+	{
+		// Note: calling either SetViewportSize or SetFixedViewportSize doesn't work
+		// because ViewportWidget is null in this case. And we can't call ResizeViewport
+		// either because it is private. Calling the only function we can call.		
+		Viewport->UpdateViewportRHI(false, InViewportSize.X, InViewportSize.Y, EWindowMode::Type::Windowed, PF_Unknown);
+	}
+}
+
+void UAvalancheGameInstance::MarkSynchronousAssetLoadingThisFrame()
+{
+	UE::AvalancheGameInstance::Private::MarkSynchronousAssetLoadingThisFrame();
+}
+
+void UAvalancheGameInstance::Tick(float DeltaSeconds)
+{
+	if (bIsTicking)
+	{
+		return;
+	}
+	TGuardValue TickGuard(bIsTicking, true);
+
+	if (!bWorldPlaying)
+	{
+		return;
+	}
+
+	if (!PlayWorld)
+	{
+		EndPlayWorld();
+		return;
+	}
+	
+	{
+		QUICK_SCOPE_CYCLE_COUNTER(STAT_AvalancheGame_Tick_SetDropDetail);
+		ViewportClient->SetDropDetail(DeltaSeconds);
+	}
+	
+	/*
+	 * Things Removed Here because they're currently unneeded and will possibly never be
+	 * TickWorldTravel
+	 */
+
+	{
+		SCOPE_TIME_GUARD(TEXT("UAvalancheGameInstance::Tick - World Tick"));
+		PlayWorld->Tick(ELevelTick::LEVELTICK_All, DeltaSeconds);
+	}
+
+	//Removed check for Dedicated Server & Commandlet
+	{
+		QUICK_SCOPE_CYCLE_COUNTER(STAT_UGameEngine_Tick_CheckCaptures);
+		
+		// Only update reflection captures in game once all 'always loaded' levels have been loaded
+		// This won't work with actual level streaming though
+		if (PlayWorld->AreAlwaysLoadedLevelsLoaded())
+		{
+			// Update sky light first because it's considered direct lighting, sky diffuse will be visible in reflection capture indirect specular
+			USkyLightComponent::UpdateSkyCaptureContents(PlayWorld.Get());
+			UReflectionCaptureComponent::UpdateReflectionCaptureContents(PlayWorld.Get());
+		}
+	}
+
+	//Rest of the stuff between above Lighting Update and GameViewport Tick removed
+
+	{
+		SCOPE_TIME_GUARD(TEXT("UAvalancheGameInstance::Tick - TickViewport"));
+		QUICK_SCOPE_CYCLE_COUNTER(STAT_AvalancheGame_Tick_TickViewport);
+		ViewportClient->Tick(DeltaSeconds);
+	}
+	
+	{
+		Viewport->Draw();
+
+		if (!bIsRenderTargetReady)
+		{
+			if (!RenderTargetFence.IsValid())
+			{
+				RenderTargetFence = MakeUnique<FRenderCommandFence>();
+				RenderTargetFence->BeginFence();
+			}
+			else
+			{
+				bIsRenderTargetReady = RenderTargetFence->IsFenceComplete();
+				OnRenderTargetReady.Broadcast(this, PlayingChannelName);
+			}
+		}
+	}
+
+	if (bRequestEndPlayWorld)
+	{
+		EndPlayWorld();
+	}
+	if (bRequestUnloadWorld)
+	{
+		UnloadWorld();
+	}
+}
+
+void UAvalancheGameInstance::UnloadWorld()
+{
+	if (bWorldPlaying)
+	{
+		EndPlayWorld();
+	}
+
+	bWorldCreated = false;
+	bRequestUnloadWorld = false;
+
+	if (IsValid(ManagedAvalancheBlueprint.Get()))
+	{
+		ManagedAvalancheBlueprint->UnregisterRemoteControlPreset();
+	}
+	SourceAvalancheBlueprint.Reset();
+	ManagedAvalancheBlueprint = nullptr;
+	
+	if (PlayWorld)
+	{
+		if (FAudioDeviceHandle AudioDevice = PlayWorld->GetAudioDevice())
+		{
+			AudioDevice->Flush(PlayWorld, false);
+		}
+
+		// Need to do this before destroying the world context apparently.
+		PlayWorld->SetShouldForceUnloadStreamingLevels(true);
+		PlayWorld->FlushLevelStreaming();
+
+		GEngine->DestroyWorldContext(PlayWorld.Get());
+		PlayWorld->DestroyWorld(true);
+	}
+
+	PlayWorld = nullptr;
+	WorldContext = nullptr;
+}
+
+void UAvalancheGameInstance::EndPlayWorld()
+{
+	bWorldPlaying = false;
+	bRequestEndPlayWorld = false;
+
+	FCoreDelegates::OnEndFrame.RemoveAll(this);
+
+	OnEndPlay.Broadcast(this, PlayingChannelName);
+	PlayingChannelName = NAME_None;
+
+	Viewport.Reset();
+	RenderTargetFence.Reset();
+	bIsRenderTargetReady = false;
+
+#if ALLOW_CONSOLE
+	if (ViewportClient)
+	{
+		GLog->RemoveOutputDevice(ViewportClient->ViewportConsole);
+	}
+#endif
+
+	ViewportClient = nullptr;
+}
+
+void UAvalancheGameInstance::OnEndFrameTick()
+{
+	double DeltaSeconds = FApp::GetDeltaTime();
+
+	if (UE::AvalancheGameInstance::Private::ShouldIgnoreDeltaSecondsForCurrentFrame())
+	{
+		DeltaSeconds = LastDeltaSeconds;
+	}
+	else
+	{
+		LastDeltaSeconds = DeltaSeconds;
+	}
+	
+	Tick(DeltaSeconds);
+}
+
+void UAvalancheGameInstance::OnEnginePreExit()
+{
+	if (PlayWorld && PlayWorld->GetAudioDevice())
+	{
+		// Temporarily set the GIsRequestingExit to false so that the DLLs like XAudio2Dll are not unloaded
+		// for ref: FMixerPlatformXAudio2::TeardownHardware unloads the dll if Engine Exit is requested
+		PRAGMA_DISABLE_DEPRECATION_WARNINGS
+		TGuardValue<bool> RequestExitGuard(GIsRequestingExit, false);
+		PRAGMA_ENABLE_DEPRECATION_WARNINGS
+		const FAudioDeviceHandle EmptyHandle;
+		PlayWorld->SetAudioDevice(EmptyHandle);
+	}
+}
+
+void UAvalancheGameInstance::BeginDestroy()
+{
+	FCoreDelegates::OnEnginePreExit.Remove(EnginePreExitHandle);
+	EnginePreExitHandle.Reset();
+	
+	EndPlayWorld();
+	UnloadWorld();
+	
+	Super::BeginDestroy();
+}
diff --git a/Engine/Plugins/Experimental/Avalanche/Source/Avalanche/Private/Framework/AvalancheGameViewportClient.cpp b/Engine/Plugins/Experimental/Avalanche/Source/Avalanche/Private/Framework/AvalancheGameViewportClient.cpp
new file mode 100644
--- /dev/null
+++ b/Engine/Plugins/Experimental/Avalanche/Source/Avalanche/Private/Framework/AvalancheGameViewportClient.cpp
@@ -0,0 +1,536 @@
+// Copyright Epic Games, Inc. All Rights Reserved.
+
+#include "Framework/AvalancheGameViewportClient.h"
+#include "Camera/CameraActor.h"
+#include "Camera/CameraComponent.h"
+#include "Camera/CameraPhotography.h"
+#include "Components/LineBatchComponent.h"
+#include "ContentStreaming.h"
+#include "DynamicResolutionState.h"
+#include "Engine/Canvas.h"
+#include "Engine/LocalPlayer.h"
+#include "Engine/TextureRenderTarget2D.h"
+#include "EngineModule.h"
+#include "EngineUtils.h"
+#include "Framework/Application/SlateApplication.h"
+#include "GameFramework/PlayerController.h"
+#include "IAvaModule.h"
+#include "IHeadMountedDisplay.h"
+#include "IXRCamera.h"
+#include "IXRTrackingSystem.h"
+#include "LegacyScreenPercentageDriver.h"
+#include "Misc/EngineVersionComparison.h"
+#include "SceneViewExtension.h"
+#include "TextureResource.h"
+#include "UObject/Package.h"
+#include "UObject/UObjectGlobals.h"
+
+/** Util to find named canvas in transient package, and create if not found */
+static UCanvas* GetCanvasByName(FName CanvasName)
+{
+	// Cache to avoid FString/FName conversions/compares
+	static TMap<FName, UCanvas*> CanvasMap;
+	UCanvas** FoundCanvas = CanvasMap.Find(CanvasName);
+	if (!FoundCanvas)
+	{
+		UCanvas* CanvasObject = FindObject<UCanvas>(GetTransientPackage(), *CanvasName.ToString(), false);
+		if (!CanvasObject)
+		{
+			CanvasObject = NewObject<UCanvas>(GetTransientPackage(), CanvasName);
+			CanvasObject->AddToRoot();
+		}
+
+		CanvasMap.Add(CanvasName, CanvasObject);
+		return CanvasObject;
+	}
+
+	return *FoundCanvas;
+}
+
+UAvalancheGameViewportClient::UAvalancheGameViewportClient()
+	: CameraManager(MakeShared<FAvaCameraManager>())
+{
+}
+
+void UAvalancheGameViewportClient::Draw(FViewport* InViewport, FCanvas* Canvas)
+{
+	// Override the Canvas Render Target with ours
+	if (RenderTarget.IsValid())
+	{
+		Canvas->SetRenderTarget_GameThread(RenderTarget->GameThread_GetRenderTargetResource());
+	}
+	
+	// Allow HMD to modify the view later, just before rendering
+	const bool bStereoRendering = GEngine->IsStereoscopic3D(InViewport);
+	
+	// Create a temporary canvas if there isn't already one.
+	static FName CanvasObjectName(TEXT("CanvasObject"));
+	UCanvas* CanvasObject = GetCanvasByName(CanvasObjectName);
+	CanvasObject->Canvas = Canvas;
+
+	// Create temp debug canvas object
+	FCanvas* DebugCanvas = InViewport->GetDebugCanvas();
+	FIntPoint DebugCanvasSize = InViewport->GetSizeXY();
+	if (bStereoRendering && GEngine->XRSystem.IsValid() && GEngine->XRSystem->GetHMDDevice())
+	{
+		DebugCanvasSize = GEngine->XRSystem->GetHMDDevice()->GetIdealDebugCanvasRenderTargetSize();
+	}
+
+	static FName DebugCanvasObjectName(TEXT("DebugCanvasObject"));
+	UCanvas* DebugCanvasObject = GetCanvasByName(DebugCanvasObjectName);
+	DebugCanvasObject->Init(DebugCanvasSize.X, DebugCanvasSize.Y, nullptr, DebugCanvas);
+	if (DebugCanvas)
+	{
+		DebugCanvas->SetScaledToRenderTarget(bStereoRendering);
+		DebugCanvas->SetStereoRendering(bStereoRendering);
+	}
+	Canvas->SetScaledToRenderTarget(bStereoRendering);
+	Canvas->SetStereoRendering(bStereoRendering);
+	
+	if (World == nullptr)
+	{
+		return;
+	}
+
+	constexpr bool bCaptureNeedsSceneColor = false;
+	constexpr ESceneCaptureSource SceneCaptureSource = ESceneCaptureSource::SCS_FinalColorLDR;
+	constexpr ESceneCaptureCompositeMode SceneCaptureCompositeMode = ESceneCaptureCompositeMode::SCCM_Overwrite;
+
+	/** When enabled, the scene capture will composite into the render target instead of overwriting its contents. */
+
+	//TODO: World->GetTime() seems to return zeros because we have not Begun Play...
+	FGameTime Time = FGameTime::GetTimeSinceAppStart();
+
+	// Setup a FSceneViewFamily/FSceneView for the viewport.
+	FSceneViewFamilyContext ViewFamilyContext(FSceneViewFamily::ConstructionValues(Canvas->GetRenderTarget()
+		, World->Scene
+		, EngineShowFlags)
+		.SetRealtimeUpdate(true)
+		.SetResolveScene(!bCaptureNeedsSceneColor)
+		.SetTime(Time));
+
+	ViewFamilyContext.SceneCaptureSource = SceneCaptureSource;
+	ViewFamilyContext.SceneCaptureCompositeMode = SceneCaptureCompositeMode;
+	ViewFamilyContext.DebugDPIScale = GetDPIScale();
+	ViewFamilyContext.EngineShowFlags = EngineShowFlags;
+
+#if WITH_EDITOR
+	if (GIsEditor)
+	{
+		// Force enable view family show flag for HighDPI derived's screen percentage.
+		ViewFamilyContext.EngineShowFlags.ScreenPercentage = true;
+	}
+#endif
+
+	if (!bStereoRendering)
+	{
+		// stereo is enabled, as many HMDs require this for proper visuals
+		ViewFamilyContext.EngineShowFlags.SetScreenPercentage(false);
+	}
+
+	ViewFamilyContext.ViewExtensions = GEngine->ViewExtensions->GatherActiveExtensions(FSceneViewExtensionContext(InViewport));
+	
+	for (const TSharedRef<ISceneViewExtension>& ViewExt : ViewFamilyContext.ViewExtensions)
+	{
+		ViewExt->SetupViewFamily(ViewFamilyContext);
+	}
+
+	if (bStereoRendering && GEngine->XRSystem.IsValid() && GEngine->XRSystem->GetHMDDevice())
+	{
+		// Allow HMD to modify screen settings
+		GEngine->XRSystem->GetHMDDevice()->UpdateScreenSettings(InViewport);
+	}
+
+	ViewFamilyContext.ViewMode = EViewModeIndex::VMI_Lit;
+	EngineShowFlagOverride(ESFIM_Game, ViewFamilyContext.ViewMode, ViewFamilyContext.EngineShowFlags, false);
+
+	bool bFinalScreenPercentageShowFlag;
+	bool bUsesDynamicResolution = false;
+	// Setup the screen percentage and upscaling method for the view family.
+	{
+		checkf(ViewFamilyContext.GetScreenPercentageInterface() == nullptr,
+			TEXT("Some code has tried to set up an alien screen percentage driver, that could be wrong if not supported very well by the RHI."));
+
+		// Force screen percentage show flag to be turned off if not supported.
+		if (!ViewFamilyContext.SupportsScreenPercentage())
+		{
+			ViewFamilyContext.EngineShowFlags.ScreenPercentage = false;
+		}
+
+		// Set up secondary resolution fraction for the view family.
+		if (!bStereoRendering && ViewFamilyContext.SupportsScreenPercentage())
+		{
+			// Automatically compute secondary resolution fraction from DPI.
+			ViewFamilyContext.SecondaryViewFraction = GetDPIDerivedResolutionFraction();
+		}
+
+		bFinalScreenPercentageShowFlag = ViewFamilyContext.EngineShowFlags.ScreenPercentage;
+	}
+
+	TArray<FSceneView*> Views;
+	{
+		const int32 NumViews = bStereoRendering ? GEngine->StereoRenderingDevice->GetDesiredNumberOfViews(bStereoRendering) : 1;
+		for (int32 ViewIndex = 0; ViewIndex < NumViews; ++ViewIndex)
+		{
+			FSceneView* const View = CalcSceneView(&ViewFamilyContext
+				, InViewport
+				, nullptr
+				, bStereoRendering ? ViewIndex : INDEX_NONE);
+
+			if (View)
+			{
+				Views.Add(View);
+
+				View->CameraConstrainedViewRect = View->UnscaledViewRect;
+
+				// Add view information for resource streaming. Allow up to 5X boost for small FOV.
+				const float StreamingScale = 1.f / FMath::Clamp<float>(View->LODDistanceFactor, .2f, 1.f);
+
+				IStreamingManager::Get().AddViewInformation(View->ViewMatrices.GetViewOrigin()
+					, View->UnscaledViewRect.Width()
+					, View->UnscaledViewRect.Width() * View->ViewMatrices.GetProjectionMatrix().M[0][0]
+					, StreamingScale);
+
+				World->ViewLocationsRenderedLastFrame.Add(View->ViewMatrices.GetViewOrigin());
+			}
+		}
+	}
+
+	//FinalizeViews(&ViewFamily); TODO Possible Interface
+
+	// Update level streaming.
+	World->UpdateLevelStreaming();
+
+	Canvas->Clear(FLinearColor::Transparent);
+
+	// If a screen percentage interface was not set by one of the view extension, then set the legacy one.
+	if (ViewFamilyContext.GetScreenPercentageInterface() == nullptr)
+	{
+		float GlobalResolutionFraction = 1.0f;
+
+		if (ViewFamilyContext.EngineShowFlags.ScreenPercentage && !bDisableWorldRendering && ViewFamilyContext.Views.Num() > 0)
+		{
+			// Get global view fraction.
+			FStaticResolutionFractionHeuristic StaticHeuristic;
+			StaticHeuristic.Settings.PullRunTimeRenderingSettings(GetViewStatusForScreenPercentage());
+			StaticHeuristic.PullViewFamilyRenderingSettings(ViewFamilyContext);
+			StaticHeuristic.DPIScale = GetDPIScale();
+
+			GlobalResolutionFraction = StaticHeuristic.ResolveResolutionFraction();
+		}
+
+		ViewFamilyContext.SetScreenPercentageInterface(new FLegacyScreenPercentageDriver(
+			ViewFamilyContext, GlobalResolutionFraction));
+	}
+
+	check(ViewFamilyContext.GetScreenPercentageInterface() != nullptr);
+
+	// Make sure the engine show flag for screen percentage is still what it was when setting up the screen percentage interface
+	ViewFamilyContext.EngineShowFlags.ScreenPercentage = bFinalScreenPercentageShowFlag;
+
+	if (bStereoRendering && bUsesDynamicResolution)
+	{
+		// Change screen percentage method to raw output when doing dynamic resolution with VR if not using TAA upsample.
+		for (FSceneView* const View : Views)
+		{
+			if (View->PrimaryScreenPercentageMethod == EPrimaryScreenPercentageMethod::SpatialUpscale)
+			{
+				View->PrimaryScreenPercentageMethod = EPrimaryScreenPercentageMethod::RawOutput;
+			}
+		}
+	}
+
+	ViewFamilyContext.bIsHDR = InViewport->IsHDRViewport();
+
+	if (!bDisableWorldRendering && FSlateApplication::Get().GetPlatformApplication()->IsAllowedToRender()) //-V560
+	{
+		for (const TSharedRef<ISceneViewExtension>& ViewExt : ViewFamilyContext.ViewExtensions)
+		{
+			for (FSceneView* const View : Views)
+			{
+				ViewExt->SetupView(ViewFamilyContext, *View);
+			}
+		}
+		GetRendererModule().BeginRenderingViewFamily(Canvas, &ViewFamilyContext);
+	}
+	else
+	{
+		GetRendererModule().PerFrameCleanupIfSkipRenderer();
+	}
+
+	// Beyond this point, only UI rendering independent from dynamic resolution.
+	GEngine->EmitDynamicResolutionEvent(EDynamicResolutionStateEvent::EndDynamicResolutionRendering);
+
+	// Remove temporary debug lines.
+	if (World->LineBatcher)
+	{
+		World->LineBatcher->Flush();
+	}
+	if (World->ForegroundLineBatcher)
+	{
+		World->ForegroundLineBatcher->Flush();
+	}
+
+	// Render Stats HUD in the main canvas so that it gets captured
+	// and is displayed in the broadcast channel's outputs.
+	if (!Views.IsEmpty() && IAvaModule::Get().ShouldShowRuntimeStats())
+	{
+		DrawStatsHUD( World, InViewport, Canvas, nullptr, DebugProperties, Views[0]->ViewLocation, Views[0]->ViewRotation);
+	}
+	
+	// Ensure canvas has been flushed before rendering UI
+	Canvas->Flush_GameThread();
+
+	CameraManager->ResetCameraCut();
+}
+
+bool UAvalancheGameViewportClient::IsStatEnabled(const FString& InName) const
+{
+	// The IAvaModule holds the runtime stats. We want them persistent across all viewports.
+	return IAvaModule::Get().IsRuntimeStatEnabled(InName);
+}
+
+void UAvalancheGameViewportClient::SetCameraCutThisFrame()
+{
+	bCameraCutThisFrame = true;
+}
+
+void UAvalancheGameViewportClient::SetRenderTarget(UTextureRenderTarget2D* InRenderTarget)
+{
+	RenderTarget = InRenderTarget;
+}
+
+bool UAvalancheGameViewportClient::CalcSceneViewInitOptions(FSceneViewInitOptions& OutInitOptions
+	, FViewport* InViewport
+	, FViewElementDrawer* ViewDrawer
+	, int32 StereoViewIndex)
+{
+	QUICK_SCOPE_CYCLE_COUNTER(STAT_AvaMediaCalcSceneViewInitOptions);
+	if (InViewport == nullptr || !World)
+	{
+		return false;
+	}
+
+	// get the projection data
+	if (GetProjectionData(InViewport, /*inout*/ OutInitOptions, StereoViewIndex) == false)
+	{
+		// Return false if this we didn't get back the info we needed
+		return false;
+	}
+
+	// return if we have an invalid view rect
+	if (!OutInitOptions.IsValidViewRectangle())
+	{
+		return false;
+	}
+
+	// Was there a camera cut this frame?
+	OutInitOptions.bInCameraCut = CameraManager->ShouldEnableCameraCut() || bCameraCutThisFrame;
+	bCameraCutThisFrame = false;
+
+	if (GEngine->StereoRenderingDevice.IsValid())
+	{
+		OutInitOptions.StereoPass = GEngine->StereoRenderingDevice->GetViewPassForIndex(StereoViewIndex != INDEX_NONE, StereoViewIndex);
+	}
+
+	const uint32 ViewIndex = StereoViewIndex != INDEX_NONE ? StereoViewIndex : 0;
+
+	// Make sure the ViewStates array has enough elements for the given ViewIndex.
+	{
+		const int32 RequiredViewStates = (ViewIndex + 1) - ViewStates.Num();
+
+		if (RequiredViewStates > 0)
+		{
+			ViewStates.AddDefaulted(RequiredViewStates);
+		}
+	}
+
+	// Allocate the current ViewState if necessary
+	if (ViewStates[ViewIndex].GetReference() == nullptr)
+	{
+#if UE_VERSION_OLDER_THAN(5, 3, 0)
+		const ERHIFeatureLevel::Type FeatureLevel = World ? World->FeatureLevel.GetValue() : GMaxRHIFeatureLevel;
+#else
+		const ERHIFeatureLevel::Type FeatureLevel = World ? World->GetFeatureLevel() : GMaxRHIFeatureLevel;
+#endif
+		ViewStates[ViewIndex].Allocate(FeatureLevel);
+	}
+
+	OutInitOptions.SceneViewStateInterface = ViewStates[ViewIndex].GetReference();
+	OutInitOptions.ViewActor = CameraManager->GetViewTarget();
+	OutInitOptions.PlayerIndex = 0;
+	OutInitOptions.ViewElementDrawer = ViewDrawer;
+	OutInitOptions.BackgroundColor = FLinearColor::Transparent;
+	OutInitOptions.LODDistanceFactor = 1.f;
+	OutInitOptions.StereoViewIndex = StereoViewIndex;
+	OutInitOptions.WorldToMetersScale = World->GetWorldSettings()->WorldToMeters;
+	OutInitOptions.CursorPos = InViewport->HasMouseCapture() ? FIntPoint(-1, -1) : FIntPoint(InViewport->GetMouseX(), InViewport->GetMouseY());
+	OutInitOptions.OriginOffsetThisFrame = World->OriginOffsetThisFrame;
+	return true;
+}
+
+bool UAvalancheGameViewportClient::GetProjectionData(FViewport* InViewport
+	, FSceneViewProjectionData& ProjectionData
+	, int32 StereoViewIndex) const
+{
+		// If the actor
+	if ((InViewport == nullptr) || (InViewport->GetSizeXY().X == 0) || (InViewport->GetSizeXY().Y == 0))
+	{
+		return false;
+	}
+
+	const FVector2D Origin(0.f, 0.f);
+	const FVector2D Size(1.f, 1.f);
+
+	int32 X = FMath::TruncToInt(Origin.X * InViewport->GetSizeXY().X);
+	int32 Y = FMath::TruncToInt(Origin.Y * InViewport->GetSizeXY().Y);
+
+	X += InViewport->GetInitialPositionXY().X;
+	Y += InViewport->GetInitialPositionXY().Y;
+
+	uint32 SizeX = FMath::TruncToInt(Size.X * InViewport->GetSizeXY().X);
+	uint32 SizeY = FMath::TruncToInt(Size.Y * InViewport->GetSizeXY().Y);
+
+	FIntRect UnconstrainedRectangle = FIntRect(X, Y, X + SizeX, Y + SizeY);
+
+	ProjectionData.SetViewRectangle(UnconstrainedRectangle);
+
+	// Get the viewpoint.
+	FMinimalViewInfo ViewInfo;
+	CameraManager->SetCameraViewInfo(ViewInfo);
+
+	// If stereo rendering is enabled, update the size and offset appropriately for this pass
+	const bool bNeedStereo = StereoViewIndex != INDEX_NONE && GEngine->IsStereoscopic3D();
+	const bool bIsHeadTrackingAllowed = GEngine->XRSystem.IsValid()
+		&& (World != nullptr
+			? GEngine->XRSystem->IsHeadTrackingAllowedForWorld(*World)
+			: GEngine->XRSystem->IsHeadTrackingAllowed());
+
+	if (bNeedStereo)
+	{
+		GEngine->StereoRenderingDevice->AdjustViewRect(StereoViewIndex, X, Y, SizeX, SizeY);
+	}
+
+    FVector StereoViewLocation = ViewInfo.Location;
+    if (bNeedStereo || bIsHeadTrackingAllowed)
+    {
+	    TSharedPtr<IXRCamera> XRCamera = GEngine->XRSystem.IsValid()
+    		? GEngine->XRSystem->GetXRCamera()
+    		: nullptr;
+
+		if (XRCamera.IsValid())
+		{
+			const bool bHasActiveCamera = !!CameraManager->GetCachedCameraComponent();
+			XRCamera->UseImplicitHMDPosition(bHasActiveCamera);
+		}
+
+		if (GEngine->StereoRenderingDevice.IsValid())
+		{
+			GEngine->StereoRenderingDevice->CalculateStereoViewOffset(StereoViewIndex, ViewInfo.Rotation, World->GetWorldSettings()->WorldToMeters, StereoViewLocation);
+		}
+    }
+
+	// Create the view matrix
+	ProjectionData.ViewOrigin = StereoViewLocation;
+	ProjectionData.ViewRotationMatrix = FInverseRotationMatrix(ViewInfo.Rotation) * FMatrix(
+		FPlane(0,	0,	1,	0),
+		FPlane(1,	0,	0,	0),
+		FPlane(0,	1,	0,	0),
+		FPlane(0,	0,	0,	1));
+
+	// @todo viewext this use case needs to be revisited
+	if (!bNeedStereo)
+	{
+		// Create the projection matrix (and possibly constrain the view rectangle)
+		FMinimalViewInfo::CalculateProjectionMatrixGivenView(ViewInfo
+			, GetDefault<ULocalPlayer>()->AspectRatioAxisConstraint
+			, InViewport
+			, ProjectionData);
+
+		for (const TSharedRef<ISceneViewExtension>& ViewExt : GEngine->ViewExtensions->GatherActiveExtensions(FSceneViewExtensionContext(InViewport)))
+        {
+			ViewExt->SetupViewProjectionMatrix(ProjectionData);
+		};
+	}
+	else
+	{
+		// Let the stereoscopic rendering device handle creating its own projection matrix, as needed
+		ProjectionData.ProjectionMatrix = GEngine->StereoRenderingDevice->GetStereoProjectionMatrix(StereoViewIndex);
+
+		// calculate the out rect
+		ProjectionData.SetViewRectangle(FIntRect(X, Y, X + SizeX, Y + SizeY));
+	}
+	return true;
+}
+
+FSceneView* UAvalancheGameViewportClient::CalcSceneView(FSceneViewFamily* ViewFamily
+	, FViewport* InViewport
+	, FViewElementDrawer* ViewDrawer
+	, int32 StereoViewIndex)
+{
+	//QUICK_SCOPE_CYCLE_COUNTER(STAT_AvaMediaCalcSceneView);
+	FSceneViewInitOptions ViewInitOptions;
+	if (!CalcSceneViewInitOptions(ViewInitOptions, InViewport, ViewDrawer, StereoViewIndex))
+	{
+		return nullptr;
+	}
+
+	// Get the viewpoint...technically doing this twice
+	// but it makes GetProjectionData better
+	FMinimalViewInfo ViewInfo;
+	CameraManager->SetCameraViewInfo(ViewInfo);
+
+	ViewInitOptions.ViewLocation = ViewInfo.Location;
+	ViewInitOptions.ViewRotation = ViewInfo.Rotation;
+	ViewInitOptions.bUseFieldOfViewForLOD = ViewInfo.bUseFieldOfViewForLOD;
+	ViewInitOptions.FOV = ViewInfo.FOV;
+	ViewInitOptions.DesiredFOV = ViewInfo.DesiredFOV;
+
+	// Fill out the rest of the view init options
+	ViewInitOptions.ViewFamily = ViewFamily;
+
+	FSceneView* const View = new FSceneView(ViewInitOptions);
+
+	//Pass on the previous view transform from the view info (probably provided by the camera if set)
+	View->PreviousViewTransform = ViewInfo.PreviousViewTransform;
+
+	ViewFamily->Views.Add(View);
+
+	{
+		View->StartFinalPostprocessSettings(ViewInfo.Location);
+
+		//Post Process Setting
+		//TODO?
+
+		//	CAMERA OVERRIDE
+		View->OverridePostProcessSettings(ViewInfo.PostProcessSettings, ViewInfo.PostProcessBlendWeight);
+		CameraManager->ApplyExtraPostProcessBlends(View);
+		FCameraPhotographyManager::Get().UpdatePostProcessing(View->FinalPostProcessSettings);
+
+		if (GEngine->StereoRenderingDevice.IsValid())
+		{
+			FPostProcessSettings StereoDeviceOverridePostProcessingSettings;
+			float BlendWeight = 1.0f;
+
+			const bool StereoSettingsAvailable = GEngine->StereoRenderingDevice->OverrideFinalPostprocessSettings(&StereoDeviceOverridePostProcessingSettings
+				, View->StereoPass
+				, View->StereoViewIndex
+				, BlendWeight);
+
+			if (StereoSettingsAvailable)
+			{
+				View->OverridePostProcessSettings(StereoDeviceOverridePostProcessingSettings, BlendWeight);
+			}
+		}
+
+
+		View->EndFinalPostprocessSettings(ViewInitOptions);
+	}
+
+	for (int32 ViewExt = 0; ViewExt < ViewFamily->ViewExtensions.Num(); ++ViewExt)
+	{
+		ViewFamily->ViewExtensions[ViewExt]->SetupView(*ViewFamily, *View);
+	}
+	
+	return View;
+}
diff --git a/Engine/Plugins/Experimental/Avalanche/Source/Avalanche/Private/Transition/Conditions/AvaSceneContainsTagAttributeCondition.cpp b/Engine/Plugins/Experimental/Avalanche/Source/Avalanche/Private/Transition/Conditions/AvaSceneContainsTagAttributeCondition.cpp
new file mode 100644
--- /dev/null
+++ b/Engine/Plugins/Experimental/Avalanche/Source/Avalanche/Private/Transition/Conditions/AvaSceneContainsTagAttributeCondition.cpp
@@ -0,0 +1,131 @@
+// Copyright Epic Games, Inc. All Rights Reserved.
+
+#include "Transition/Conditions/AvaSceneContainsTagAttributeCondition.h"
+#include "AvaSceneState.h"
+#include "AvaSceneSubsystem.h"
+#include "AvaTransitionLayer.h"
+#include "AvaTransitionLayerUtils.h"
+#include "AvaTransitionSubsystem.h"
+#include "Behavior/AvaTransitionBehaviorInstance.h"
+#include "IAvaSceneInterface.h"
+#include "StateTreeExecutionContext.h"
+#include "StateTreeLinker.h"
+
+#define LOCTEXT_NAMESPACE "AvaSceneContainsTagAttributeCondition"
+
+FText FAvaSceneContainsTagAttributeCondition::GenerateDescription() const
+{
+	FFormatNamedArguments Arguments;
+
+	switch (SceneType)
+	{
+	case EAvaTransitionSceneType::This:
+		Arguments.Add(TEXT("Scene"), LOCTEXT("ThisScene", "this scene"));
+		break;
+
+	case EAvaTransitionSceneType::Other:
+		Arguments.Add(TEXT("Scene"), FText::Format(LOCTEXT("OtherScene", "scene in {0}"), FAvaTransitionLayerUtils::GetLayerQueryText(LayerType, *SpecificLayers.ToString())));
+		break;
+	}
+
+	if (bInvert)
+	{
+		Arguments.Add(TEXT("Contains"), LOCTEXT("ContainsInvert", "does NOT contain"));
+	}
+	else
+	{
+		Arguments.Add(TEXT("Contains"), LOCTEXT("Contains", "contains"));
+	}
+
+	Arguments.Add(TEXT("TagAttribute"), FText::FromName(TagAttribute.ToName()));
+
+	return FText::Format(LOCTEXT("ConditionDescription", "{Scene} {Contains} tag attribute '{TagAttribute}'"), Arguments);
+}
+
+bool FAvaSceneContainsTagAttributeCondition::Link(FStateTreeLinker& InLinker)
+{
+	FAvaTransitionCondition::Link(InLinker);
+	InLinker.LinkExternalData(SceneSubsystemHandle);
+	return true;
+}
+
+bool FAvaSceneContainsTagAttributeCondition::TestCondition(FStateTreeExecutionContext& InContext) const
+{
+	return ContainsTagAttribute(InContext) ^ bInvert;
+}
+
+bool FAvaSceneContainsTagAttributeCondition::ContainsTagAttribute(FStateTreeExecutionContext& InContext) const
+{
+	TArray<const FAvaTransitionScene*> TransitionScenes = GetTransitionScenes(InContext);
+	if (TransitionScenes.IsEmpty())
+	{
+		return false;
+	}
+
+	const UAvaSceneSubsystem& SceneSubsystem = InContext.GetExternalData(SceneSubsystemHandle);
+
+	for (const FAvaTransitionScene* TransitionScene : TransitionScenes)
+	{
+		if (!TransitionScene)
+		{
+			continue;
+		}
+
+		const IAvaSceneInterface* Scene = SceneSubsystem.GetSceneInterface(TransitionScene->GetLevel());
+		if (!Scene)
+		{
+			continue;
+		}
+
+		const UAvaSceneState* SceneState = Scene->GetSceneState();
+		if (SceneState && SceneState->ContainsTagAttribute(TagAttribute))
+		{
+			return true;
+		}
+	}
+
+	return false;
+}
+
+TArray<const FAvaTransitionScene*> FAvaSceneContainsTagAttributeCondition::GetTransitionScenes(FStateTreeExecutionContext& InContext) const
+{
+	const FAvaTransitionContext& TransitionContext = InContext.GetExternalData(TransitionContextHandle);
+
+	if (SceneType == EAvaTransitionSceneType::This)
+	{
+		return { TransitionContext.GetTransitionScene() };
+	}
+
+	ensureMsgf(SceneType == EAvaTransitionSceneType::Other
+		, TEXT("FAvaSceneContainsAttributeCondition::GetTargetSceneContexts did not recognize the provided transition scene type")
+		, *UEnum::GetValueAsString(SceneType));
+
+	// Get all the Behavior Instances from Query 
+	TArray<const FAvaTransitionBehaviorInstance*> BehaviorInstances;
+	{
+		UAvaTransitionSubsystem& TransitionSubsystem = InContext.GetExternalData(TransitionSubsystemHandle);
+
+		FAvaTransitionLayerComparator Comparator = FAvaTransitionLayerUtils::BuildComparator(TransitionContext, LayerType, SpecificLayers);
+
+		BehaviorInstances = FAvaTransitionLayerUtils::QueryBehaviorInstances(TransitionSubsystem, Comparator);
+	}
+
+	TArray<const FAvaTransitionScene*> TransitionScenes;
+	TransitionScenes.Reserve(BehaviorInstances.Num());
+
+	for (const FAvaTransitionBehaviorInstance* BehaviorInstance : BehaviorInstances)
+	{
+		const FAvaTransitionContext& OtherTransitionContext = BehaviorInstance->GetTransitionContext();
+		const FAvaTransitionScene* TransitionScene = OtherTransitionContext.GetTransitionScene();
+
+		// do not add if scene is marked as needing discard
+		if (TransitionScene && !TransitionScene->HasAllFlags(EAvaTransitionSceneFlags::NeedsDiscard))
+		{
+			TransitionScenes.AddUnique(TransitionScene);
+		}
+	}
+
+	return TransitionScenes;
+}
+
+#undef LOCTEXT_NAMESPACE
diff --git a/Engine/Plugins/Experimental/Avalanche/Source/Avalanche/Private/Transition/Conditions/AvaTransitionRCControllerMatchCondition.cpp b/Engine/Plugins/Experimental/Avalanche/Source/Avalanche/Private/Transition/Conditions/AvaTransitionRCControllerMatchCondition.cpp
new file mode 100644
--- /dev/null
+++ b/Engine/Plugins/Experimental/Avalanche/Source/Avalanche/Private/Transition/Conditions/AvaTransitionRCControllerMatchCondition.cpp
@@ -0,0 +1,121 @@
+// Copyright Epic Games, Inc. All Rights Reserved.
+
+#include "Transition/Conditions/AvaTransitionRCControllerMatchCondition.h"
+#include "AvaSceneSubsystem.h"
+#include "AvaTransitionContext.h"
+#include "AvaTransitionLayer.h"
+#include "AvaTransitionLayerUtils.h"
+#include "AvaTransitionLog.h"
+#include "AvaTransitionScene.h"
+#include "AvaTransitionSubsystem.h"
+#include "Behavior/AvaTransitionBehaviorInstance.h"
+#include "IAvaSceneInterface.h"
+#include "RCVirtualProperty.h"
+#include "RemoteControlPreset.h"
+#include "StateTreeExecutionContext.h"
+#include "StateTreeLinker.h"
+#include "Transition/Extensions/IAvaTransitionRCExtension.h"
+
+#define LOCTEXT_NAMESPACE "AvaTransitionRCControllerMatchCondition"
+
+FText FAvaTransitionRCControllerMatchCondition::GenerateDescription() const
+{
+	return FText::Format(LOCTEXT("ConditionDescription", "'{0}' is {1}")
+		, ControllerId.ToText()
+		, UEnum::GetDisplayValueAsText(ValueComparisonType).ToLower());
+}
+
+bool FAvaTransitionRCControllerMatchCondition::Link(FStateTreeLinker& InLinker)
+{
+	FAvaTransitionCondition::Link(InLinker);
+	InLinker.LinkExternalData(SceneSubsystemHandle);
+	return true;
+}
+
+bool FAvaTransitionRCControllerMatchCondition::TestCondition(FStateTreeExecutionContext& InContext) const
+{
+	const FAvaTransitionContext& TransitionContext = InContext.GetExternalData(TransitionContextHandle);
+	UAvaTransitionSubsystem& TransitionSubsystem   = InContext.GetExternalData(TransitionSubsystemHandle);
+	const UAvaSceneSubsystem& SceneSubsystem       = InContext.GetExternalData(SceneSubsystemHandle);
+	const FAvaTransitionScene* TransitionScene     = TransitionContext.GetTransitionScene();
+
+	URCVirtualPropertyBase* Controller = GetController(SceneSubsystem, TransitionScene);
+	if (!Controller)
+	{
+		return false;
+	}
+
+	// Get all the Behavior Instances in the same Layer
+	TArray<const FAvaTransitionBehaviorInstance*> BehaviorInstances;
+	{
+		FAvaTransitionLayerComparator Comparator = FAvaTransitionLayerUtils::BuildComparator(TransitionContext, EAvaTransitionLayerCompareType::Same, FAvaTagHandle());
+		BehaviorInstances = FAvaTransitionLayerUtils::QueryBehaviorInstances(TransitionSubsystem, Comparator);
+	}
+
+	if (BehaviorInstances.IsEmpty())
+	{
+		return false;
+	}
+
+	// Optional Extension to override Controller Comparison
+	IAvaRCTransitionExtension* const RCTransitionExtension = TransitionScene->FindExtension<IAvaRCTransitionExtension>();
+
+	for (const FAvaTransitionBehaviorInstance* BehaviorInstance : BehaviorInstances)
+	{
+		check(BehaviorInstance);
+
+		const FAvaTransitionScene* OtherTransitionScene = BehaviorInstance->GetTransitionContext().GetTransitionScene();
+		if (!OtherTransitionScene)
+		{
+			continue;
+		}
+
+		EAvaTransitionComparisonResult Result;
+		if (RCTransitionExtension)
+		{
+			Result = RCTransitionExtension->CompareControllers(Controller->Id
+				, *TransitionScene
+				, *OtherTransitionScene);
+		}
+		else if (URCVirtualPropertyBase* OtherController = GetController(SceneSubsystem, OtherTransitionScene))
+		{
+			Result = Controller->IsValueEqual(OtherController)
+				? EAvaTransitionComparisonResult::Same
+				: EAvaTransitionComparisonResult::Different;
+		}
+		else
+		{
+			Result = EAvaTransitionComparisonResult::None;
+		}
+
+		if (ValueComparisonType == Result)
+		{
+			return true;
+		}
+	}
+
+	return false;
+}
+
+URCVirtualPropertyBase* FAvaTransitionRCControllerMatchCondition::GetController(const UAvaSceneSubsystem& InSceneSubsystem, const FAvaTransitionScene* InTransitionScene) const
+{
+	if (!InTransitionScene)
+	{
+		return nullptr;
+	}
+
+	IAvaSceneInterface* SceneInterface = InSceneSubsystem.GetSceneInterface(InTransitionScene->GetLevel());
+	if (!SceneInterface)
+	{
+		return nullptr;
+	}
+
+	if (URemoteControlPreset* RemoteControlPreset = SceneInterface->GetRemoteControlPreset())
+	{
+		return ControllerId.FindController(RemoteControlPreset);
+	}
+
+	return nullptr;
+}
+
+#undef LOCTEXT_NAMESPACE
diff --git a/Engine/Plugins/Experimental/Avalanche/Source/Avalanche/Private/Transition/Tasks/AvaSceneAddTagAttributeTask.cpp b/Engine/Plugins/Experimental/Avalanche/Source/Avalanche/Private/Transition/Tasks/AvaSceneAddTagAttributeTask.cpp
new file mode 100644
--- /dev/null
+++ b/Engine/Plugins/Experimental/Avalanche/Source/Avalanche/Private/Transition/Tasks/AvaSceneAddTagAttributeTask.cpp
@@ -0,0 +1,36 @@
+// Copyright Epic Games, Inc. All Rights Reserved.
+
+#include "Transition/Tasks/AvaSceneAddTagAttributeTask.h"
+#include "AvaSceneState.h"
+#include "IAvaSceneInterface.h"
+
+#define LOCTEXT_NAMESPACE "AvaSceneAddTagAttributeTask"
+
+FText FAvaSceneAddTagAttributeTask::GenerateDescription() const
+{
+	return FText::Format(LOCTEXT("TaskDescription", "Add '{0}' tag attribute to this scene"), FText::FromName(TagAttribute.ToName()));
+}
+
+EStateTreeRunStatus FAvaSceneAddTagAttributeTask::EnterState(FStateTreeExecutionContext& InContext, const FStateTreeTransitionResult& InTransition) const
+{
+	IAvaSceneInterface* Scene = GetScene(InContext);
+	if (!Scene)
+	{
+		return EStateTreeRunStatus::Failed;
+	}
+
+	UAvaSceneState* SceneState = Scene->GetSceneState();
+	if (!SceneState)
+	{
+		return EStateTreeRunStatus::Failed;
+	}
+
+	if (SceneState->AddTagAttribute(TagAttribute))
+	{
+		return EStateTreeRunStatus::Succeeded;
+	}
+
+	return EStateTreeRunStatus::Failed;
+}
+
+#undef LOCTEXT_NAMESPACE
diff --git a/Engine/Plugins/Experimental/Avalanche/Source/Avalanche/Private/Transition/Tasks/AvaSceneRemoveTagAttributeTask.cpp b/Engine/Plugins/Experimental/Avalanche/Source/Avalanche/Private/Transition/Tasks/AvaSceneRemoveTagAttributeTask.cpp
new file mode 100644
--- /dev/null
+++ b/Engine/Plugins/Experimental/Avalanche/Source/Avalanche/Private/Transition/Tasks/AvaSceneRemoveTagAttributeTask.cpp
@@ -0,0 +1,36 @@
+// Copyright Epic Games, Inc. All Rights Reserved.
+
+#include "Transition/Tasks/AvaSceneRemoveTagAttributeTask.h"
+#include "AvaSceneState.h"
+#include "IAvaSceneInterface.h"
+
+#define LOCTEXT_NAMESPACE "AvaSceneRemoveTagAttributeTask"
+
+FText FAvaSceneRemoveTagAttributeTask::GenerateDescription() const
+{
+	return FText::Format(LOCTEXT("TaskDescription", "Remove '{0}' tag attribute from this scene"), FText::FromName(TagAttribute.ToName()));
+}
+
+EStateTreeRunStatus FAvaSceneRemoveTagAttributeTask::EnterState(FStateTreeExecutionContext& InContext, const FStateTreeTransitionResult& InTransition) const
+{
+	IAvaSceneInterface* Scene = GetScene(InContext);
+	if (!Scene)
+	{
+		return EStateTreeRunStatus::Failed;
+	}
+
+	UAvaSceneState* SceneState = Scene->GetSceneState();
+	if (!SceneState)
+	{
+		return EStateTreeRunStatus::Failed;
+	}
+
+	if (SceneState->RemoveTagAttribute(TagAttribute))
+	{
+		return EStateTreeRunStatus::Succeeded;
+	}
+
+	return EStateTreeRunStatus::Failed;
+}
+
+#undef LOCTEXT_NAMESPACE
diff --git a/Engine/Plugins/Experimental/Avalanche/Source/Avalanche/Private/Viewport/AvaCameraManager.cpp b/Engine/Plugins/Experimental/Avalanche/Source/Avalanche/Private/Viewport/AvaCameraManager.cpp
new file mode 100644
--- /dev/null
+++ b/Engine/Plugins/Experimental/Avalanche/Source/Avalanche/Private/Viewport/AvaCameraManager.cpp
@@ -0,0 +1,219 @@
+// Copyright Epic Games, Inc. All Rights Reserved.
+
+#include "Viewport/AvaCameraManager.h"
+#include "AvaActor.h"
+#include "AvaSequencePlaybackObject.h"
+#include "Camera/CameraComponent.h"
+#include "Camera/CameraTypes.h"
+#include "CineCameraActor.h"
+#include "Engine/Scene.h"
+#include "EngineUtils.h"
+#include "SceneView.h"
+
+#if WITH_EDITOR
+#include "UObject/UObjectGlobals.h"
+#endif
+
+FAvaCameraManager::FAvaCameraManager()
+{
+	ViewTargetWeak = nullptr;
+	CameraComponentWeak = nullptr;
+	bCameraCut = true; // Start with a camera cut
+}
+
+void FAvaCameraManager::Init(IAvaSequencePlaybackObject* InPlaybackObject, bool bInCanvasController)
+{
+	if (InPlaybackObject)
+	{
+		InPlaybackObject->GetOnCameraCut().AddSP(this, &FAvaCameraManager::OnUpdateCameraCut);
+	}
+}
+
+bool FAvaCameraManager::HasViewTarget() const
+{
+	return ViewTargetWeak.IsValid();
+}
+
+AActor* FAvaCameraManager::GetViewTarget() const
+{
+	return ViewTargetWeak.Get();
+}
+
+bool FAvaCameraManager::SetViewTarget(AActor* InNewViewTarget)
+{
+	AActor* CurrentViewTarget = GetViewTarget();
+
+	if (InNewViewTarget == CurrentViewTarget)
+	{
+		return false;
+	}
+
+	ViewTargetWeak = nullptr;
+	CameraComponentWeak = nullptr;
+	bCameraCut = true; // set bCameraCut for 2 frames or it doesn't work
+
+	if (IsValid(InNewViewTarget))
+	{
+		ViewTargetWeak = InNewViewTarget;
+
+		TArray<UCameraComponent*> CameraComponents;
+		InNewViewTarget->GetComponents<UCameraComponent>(CameraComponents);
+
+		for (UCameraComponent* CameraComponent : CameraComponents)
+		{
+			if (!CameraComponent->IsActive())
+			{
+				continue;
+			}
+
+			CameraComponentWeak = CameraComponent;
+			break;
+		}
+	}
+
+	OnViewTargetChanged.Broadcast(this, InNewViewTarget);
+
+	return true;
+}
+
+UCameraComponent* FAvaCameraManager::GetCachedCameraComponent() const
+{
+	return CameraComponentWeak.Get();
+}
+
+#if WITH_EDITOR
+bool FAvaCameraManager::SetDefaultViewTarget(UWorld* InWorld, FName InDefaultViewTargetName)
+{
+	if (!IsValid(InWorld))
+	{
+		UE_LOG(LogViewport, Error, TEXT("Error setting default view target for Avalanche BP: Null world."));
+		return false;
+	}
+
+	if (InDefaultViewTargetName == NAME_None)
+	{
+		return false;
+	}
+
+	for (AActor* Actor : TActorRange<AActor>(InWorld))
+	{
+		if (Actor->GetFName() == InDefaultViewTargetName)
+		{
+			SetViewTarget(Actor);
+			return true;
+		}
+	}
+
+	UE_LOG(LogViewport, Error, TEXT("Error setting default view target for Avalanche BP: Startup camera was not found."));
+	return false;
+}
+#endif
+
+FVector FAvaCameraManager::GetViewLocation() const
+{
+	FMinimalViewInfo ViewInfo;
+	SetCameraViewInfo(ViewInfo);
+
+	return ViewInfo.Location;
+}
+
+FRotator FAvaCameraManager::GetViewRotation() const
+{
+	FMinimalViewInfo ViewInfo;
+	SetCameraViewInfo(ViewInfo);
+
+	return ViewInfo.Rotation;
+}
+
+void FAvaCameraManager::GetViewPoint(FVector& OutLocation, FRotator& OutRotation) const
+{
+	FMinimalViewInfo ViewInfo;
+	SetCameraViewInfo(ViewInfo);
+
+	OutLocation = ViewInfo.Location;
+	OutRotation = ViewInfo.Rotation;
+}
+
+void FAvaCameraManager::SetViewPoint(const FVector& InLocation, const FRotator& InRotation)
+{
+	if (AActor* ViewTarget = ViewTargetWeak.Get())
+	{
+		ViewTargetWeak->SetActorLocation(InLocation);
+		ViewTargetWeak->SetActorRotation(InRotation);
+	}
+}
+
+void FAvaCameraManager::SetCameraViewInfo(FMinimalViewInfo& InOutViewInfo) const
+{
+	if (CameraComponentWeak.IsValid())
+	{
+		UWorld* World = CameraComponentWeak->GetWorld();
+		CameraComponentWeak->GetCameraView(IsValid(World) ? World->GetDeltaSeconds() : 0.f, InOutViewInfo);
+	}
+	else if (ViewTargetWeak.IsValid())
+	{
+		InOutViewInfo.Location = ViewTargetWeak->GetActorLocation();
+		InOutViewInfo.Rotation = ViewTargetWeak->GetActorRotation();
+	}
+	else
+	{
+		InOutViewInfo.Location = FVector::ZeroVector;
+		InOutViewInfo.Rotation = FRotator::ZeroRotator;
+	}
+
+	// We store the originally desired FOV as other classes may adjust to account for ultra-wide aspect ratios
+	InOutViewInfo.DesiredFOV = InOutViewInfo.FOV;
+}
+
+void FAvaCameraManager::GetExtraPostProcessBlends(TArray<FPostProcessSettings>& OutSettings,
+	TArray<float>& OutWeights) const
+{
+	OutSettings.Empty();
+	OutWeights.Empty();
+
+	if (CameraComponentWeak.IsValid())
+	{
+		CameraComponentWeak->GetExtraPostProcessBlends(OutSettings, OutWeights);
+	}
+}
+
+void FAvaCameraManager::ApplyExtraPostProcessBlends(FSceneView* InView) const
+{
+	if (UCameraComponent* CameraComponent = CameraComponentWeak.Get())
+	{
+		TArray<FPostProcessSettings> OutSettings;
+		TArray<float> OutWeights;
+		CameraComponent->GetExtraPostProcessBlends(OutSettings, OutWeights);
+
+		for (int32 Idx = 0; Idx < OutSettings.Num(); ++Idx)
+		{
+			InView->OverridePostProcessSettings(OutSettings[Idx], OutWeights[Idx]);
+		}
+	}
+}
+
+TArray<AActor*> FAvaCameraManager::GetAvailableCameras(UWorld* InWorld)
+{
+	TArray<AActor*> ActorsWithActiveCameraComponents;
+
+	if (IsValid(InWorld))
+	{
+		for (AActor* Actor : TActorRange<AActor>(InWorld))
+		{
+			if (Actor->HasActiveCameraComponent())
+			{
+				ActorsWithActiveCameraComponents.Add(Actor);
+			}
+		}
+	}
+
+	return ActorsWithActiveCameraComponents;
+}
+
+void FAvaCameraManager::OnUpdateCameraCut(UObject* InCameraObject, bool bInJump)
+{
+	if (AActor* CameraActor = Cast<AActor>(InCameraObject))
+	{
+		SetViewTarget(CameraActor);
+	}
+}
diff --git a/Engine/Plugins/Experimental/Avalanche/Source/Avalanche/Public/AvaActor.h b/Engine/Plugins/Experimental/Avalanche/Source/Avalanche/Public/AvaActor.h
new file mode 100644
--- /dev/null
+++ b/Engine/Plugins/Experimental/Avalanche/Source/Avalanche/Public/AvaActor.h
@@ -0,0 +1,34 @@
+// Copyright Epic Games, Inc. All Rights Reserved.
+
+#pragma once
+
+#include "GameFramework/Actor.h"
+#include "AvaActor.generated.h"
+
+/** Avalanche Actor */
+UCLASS(MinimalAPI, DisplayName = "Motion Design Actor")
+class AAvalancheActor : public AActor
+{
+	GENERATED_BODY()
+
+public:
+	AAvalancheActor();
+
+	void SetIsPlaceholder(bool bInIsPlaceholder = false) { bIsPlaceholder = bInIsPlaceholder; }
+
+	bool GetIsPlaceholder() const { return bIsPlaceholder; }
+
+	bool UpdateGeneratedClass(bool bForceUpdate = false);
+
+	//~ Begin AActor
+	virtual void PostActorCreated() override;
+	//~ End AActor
+
+	//~ Begin UObject
+	virtual void PostDuplicate(EDuplicateMode::Type DuplicateMode) override;
+	//~ End UObject
+
+private:
+	bool bGeneratedClassUpdated = false;
+	bool bIsPlaceholder = false;
+};
diff --git a/Engine/Plugins/Experimental/Avalanche/Source/Avalanche/Public/AvaActorUtils.h b/Engine/Plugins/Experimental/Avalanche/Source/Avalanche/Public/AvaActorUtils.h
new file mode 100644
--- /dev/null
+++ b/Engine/Plugins/Experimental/Avalanche/Source/Avalanche/Public/AvaActorUtils.h
@@ -0,0 +1,27 @@
+// Copyright Epic Games, Inc. All Rights Reserved.
+
+#pragma once
+
+#include "Containers/Array.h"
+#include "Containers/Map.h"
+#include "Math/MathFwd.h"
+
+class AActor;
+class IAvaSceneInterface;
+class USceneComponent;
+enum class EAvaReferenceContainer : uint8;
+struct FOrientedBox;
+
+struct AVALANCHE_API FAvaActorUtils
+{
+	static FOrientedBox MakeOrientedBox(const FBox& InLocalBox, const FTransform& InWorldTransform);
+
+	static FBox GetActorLocalBoundingBox(const AActor* InActor, bool bIncludeFromChildActors = false, bool bMustBeRegistered = true);
+
+	static FBox GetComponentLocalBoundingBox(const USceneComponent* InComponent);
+
+	static IAvaSceneInterface* GetSceneInterfaceFromActor(const AActor* InActor);
+
+	/** Finds an actor relative to another based on the type of reference container used. */
+	static AActor* ActorFromReferenceContainer(AActor* const InActor, const EAvaReferenceContainer InReferenceContainer, const bool bInIgnoreHiddenActors = false);
+};
diff --git a/Engine/Plugins/Experimental/Avalanche/Source/Avalanche/Public/AvaBlueprint.h b/Engine/Plugins/Experimental/Avalanche/Source/Avalanche/Public/AvaBlueprint.h
new file mode 100644
--- /dev/null
+++ b/Engine/Plugins/Experimental/Avalanche/Source/Avalanche/Public/AvaBlueprint.h
@@ -0,0 +1,217 @@
+// Copyright Epic Games, Inc. All Rights Reserved.
+
+#pragma once
+
+#include "AvaSceneTree.h"
+#include "AvaSequenceShared.h"
+#include "Data/AvalancheWorldData.h"
+#include "Engine/Blueprint.h"
+#include "IAvaSceneInterface.h"
+#include "IAvaSequenceProvider.h"
+#include "Types/SlateEnums.h"
+#include "Viewport/AvaViewportQualitySettings.h"
+#include "AvaBlueprint.generated.h"
+
+class AAvalancheActor;
+class IAvaSequencePlaybackObject;
+class ISequencer;
+class UAvaSequence;
+class UAvaSequencePlayer;
+class URemoteControlPreset;
+class UTexture2D;
+
+#if WITH_EDITOR
+struct FAvaViewportGuideInfo_Deprecated;
+#endif
+
+UCLASS(BlueprintType, DisplayName = "Motion Design Blueprint")
+class AVALANCHE_API UAvalancheBlueprint : public UBlueprint, public IAvaSequenceProvider, public IAvaSceneInterface
+{
+	GENERATED_BODY()
+
+public:
+	UAvalancheBlueprint();
+
+	static FTopLevelAssetPath GetWorldContextPath(const FSoftObjectPath& InSourceAssetPath);
+
+#if WITH_EDITORONLY_DATA
+	/** Called on Pre Save to Notify that we have Updated the Avalanche World for Save. */
+	FSimpleMulticastDelegate OnAvalanchePreSave;
+
+	/** Called on Post Save to Notify that we have Saved the Avalanche World. */
+	FSimpleMulticastDelegate OnAvalanchePostSave;
+
+	virtual void PreSave(FObjectPreSaveContext ObjectSaveContext) override;
+	virtual void PostSaveRoot(FObjectPostSaveRootContext ObjectSaveContext) override;
+
+	const TArray<uint8>& GetEditorData() const { return EditorData; }
+	TArray<uint8>& GetEditorDataMutable() { return EditorData; }
+
+	void SetCreateDefaultScene(bool bInCreateDefaultScene) { bCreateDefaultScene = bInCreateDefaultScene; }
+	bool ShouldCreateDefaultScene() const { return bCreateDefaultScene; }
+#endif
+
+	void SetAvalancheWorld(UWorld* InWorld);
+
+	UWorld* GetAvalancheWorld() const;
+
+	AAvalancheActor* GetPlaceholderActor() const;
+	IAvaSequencePlaybackObject* GetScenePlayback() const;
+
+	void DestroyPlaceholderActor();
+	
+	void UpdatePlaceholderActor(bool bInForceFullUpdate = false);
+
+	/**
+	 * Saves the Avalanche World State into the UAvalancheBlueprint::WorldData property, does NOT serialize to disk
+	 * Should be used prior to actions like duplicating or copying the UObject to get the latest world state duplicated/copied over
+	 */
+	void SaveAvalancheWorld();
+
+	using FInitActorFunctionRef = TFunctionRef<void(AActor&, const FAvalancheActorData&)>;
+	void LoadAvalancheWorld(UWorld* InOverrideWorld = nullptr, TOptional<FInitActorFunctionRef> InInitActorFunction = {});
+
+	bool IsSavingAvalancheWorld() const { return bSavingWorld; }
+	bool IsLoadingAvalancheWorld() const { return bLoadingWorld; }
+
+	void RegisterRemoteControlPreset();
+	void UnregisterRemoteControlPreset();
+
+	const FName& GetStartupCameraName() const { return StartupCameraName; }
+	void SetStartupCameraName(FName InCameraName);
+
+#if WITH_EDITOR
+	//~ Begin UBlueprint
+	virtual bool SupportedByDefaultBlueprintFactory() const override { return false; }
+	virtual UClass* GetBlueprintClass() const override;
+	virtual void GetReparentingRules(TSet<const UClass*>& AllowedChildrenOfClasses, TSet<const UClass*>& DisallowedChildrenOfClasses) const override;
+	//~ End UBlueprint
+#endif
+
+	//~ Begin UObject
+	virtual void BeginDestroy() override;
+	virtual void PreDuplicate(FObjectDuplicationParameters& InDuplicationParameters) override;
+	virtual void Serialize(FArchive& Ar) override;
+	virtual void PostLoadSubobjects(FObjectInstancingGraph* OuterInstanceGraph) override;
+	//~ End UObject
+
+	//~ Begin IAvaSceneInterface
+	virtual ULevel* GetSceneLevel() const override;
+	virtual UAvaSceneSettings* GetSceneSettings() const { return nullptr; }
+	virtual UAvaSceneState* GetSceneState() const override { return nullptr; }
+	virtual FAvaSceneTree& GetSceneTree() override { return SceneTree; }
+	virtual const FAvaSceneTree& GetSceneTree() const override { return SceneTree; }
+	virtual IAvaSequencePlaybackObject* GetPlaybackObject() const override;
+	virtual IAvaSequenceProvider* GetSequenceProvider() override { return this; }
+	virtual const IAvaSequenceProvider* GetSequenceProvider() const override { return this; }
+	virtual URemoteControlPreset* GetRemoteControlPreset() const override { return RemoteControlPreset; }
+	//~ End IAvaSceneInterface
+
+	//~ Begin IAvaSequenceProvider
+	virtual UObject* ToUObject() override;
+	virtual UWorld* GetContextWorld() const override;
+	virtual bool CreateDirectorInstance(UAvaSequence& InSequence, IMovieScenePlayer& InPlayer, const FMovieSceneSequenceID& InSequenceID, UObject*& OutDirectorInstance) override;
+	virtual bool AddSequence(UAvaSequence* InSequence) override;
+	virtual void RemoveSequence(UAvaSequence* InSequence) override;
+	virtual void SetDefaultSequence(UAvaSequence* InSequence) override;
+	virtual UAvaSequence* GetDefaultSequence() const override;
+	virtual const TArray<TObjectPtr<UAvaSequence>>& GetSequences() const override { return Animations; }
+	virtual const TArray<TWeakObjectPtr<UAvaSequence>>& GetRootSequences() const override { return RootAnimations; }
+	virtual TArray<TWeakObjectPtr<UAvaSequence>>& GetRootSequencesMutable() override { return RootAnimations; }
+	virtual FName GetSequenceProviderDebugName() const override;
+#if WITH_EDITOR
+	virtual void OnEditorSequencerCreated(const TSharedPtr<ISequencer>& InSequencer) override;
+	virtual TSharedPtr<ISequencer> GetEditorSequencer() const override;
+	virtual bool GetDirectorBlueprint(UAvaSequence& InSequence, UBlueprint*& OutBlueprint) override;
+#endif
+	virtual void RebuildSequenceTree() override;
+	virtual void ScheduleRebuildSequenceTree() override;
+	virtual FSimpleMulticastDelegate& GetOnSequenceTreeRebuilt() override { return OnAnimationsUpdated; }
+	//~ End IAvaAnimationManager
+
+	virtual UWorld* GetWorld() const override;
+
+	UTexture2D* GetThumbnailImage() const { return ThumbnailImage; };
+	void SetThumbnailImage(UTexture2D* InTexture) { ThumbnailImage = InTexture; }
+
+	FAvaViewportQualitySettings GetViewportQualitySettings() const { return ViewportQualitySettings; }
+	void SetViewportQualitySettings(const FAvaViewportQualitySettings InSettings) { ViewportQualitySettings = InSettings; }
+
+#if WITH_EDITOR
+	const TArray<FAvaViewportGuideInfo_Deprecated>& GetGuideInfos() const { return GuideInfos; }
+	void SetGuideOffset(int32 InGuideIdx, float InGuideOffset);
+	void SetGuides(const TArray<FAvaViewportGuideInfo_Deprecated>& InNewGuideInfos);
+#endif
+
+private:
+	UPROPERTY(Transient, DuplicateTransient, TextExportTransient)
+	TObjectPtr<UWorld> World;
+
+	/**
+	 * Placeholder Actor to use when the Blueprint is used for other means than to have it spawned on a different Level.
+	 * You can think of this as the Preview Actor of BlueprintEditor, but available at runtime and for other use cases beyond BP Editor
+	 * Examples include the Blueprint Editor itself, Playback and the Avalanche UMG Widget.
+	 * Placeholder should be null when an BP Avalanche Actor is spawned into the World as the Actor to use should be the spawned one instead of this Placeholder
+	 */
+	UPROPERTY(Transient, DuplicateTransient, TextExportTransient)
+	TObjectPtr<AAvalancheActor> PlaceholderActor;
+
+	/** The Base Playback Scene that is always present to Play Animations */
+	UPROPERTY()
+	TScriptInterface<IAvaSequencePlaybackObject> PlaybackObject;
+
+	UPROPERTY()
+	FAvalancheWorldData WorldData;
+
+	UPROPERTY()
+	FAvaSceneTree SceneTree;
+	
+	UPROPERTY()
+	TObjectPtr<UTexture2D> ThumbnailImage;
+
+	/** A List of All Animations, including those that are nested in other Animations */
+	UPROPERTY()
+	TArray<TObjectPtr<UAvaSequence>> Animations;
+
+	/** The index to the animation to use as the default animation */
+	UPROPERTY()
+	int32 DefaultAnimationIndex = 0;
+
+	/** A list of only the Root Animations (those without Parent Animations) */
+	UPROPERTY(Transient, DuplicateTransient, TextExportTransient)
+	TArray<TWeakObjectPtr<UAvaSequence>> RootAnimations;
+
+	UPROPERTY(VisibleAnywhere, Instanced, Category = "Motion Design")
+	TObjectPtr<URemoteControlPreset> RemoteControlPreset;
+
+	UPROPERTY()
+	FName StartupCameraName = NAME_None;
+
+#if WITH_EDITOR
+	TWeakPtr<ISequencer> EditorSequencer;
+#endif
+
+#if WITH_EDITORONLY_DATA
+	UPROPERTY()
+	TArray<uint8> EditorData;
+
+	UPROPERTY()
+	bool bCreateDefaultScene = false;
+#endif
+
+	/** Only the flags specified in these settings will be applied to the viewport client ShowFlags. */
+	UPROPERTY()
+	FAvaViewportQualitySettings ViewportQualitySettings;
+
+#if WITH_EDITORONLY_DATA
+	UPROPERTY()
+	TArray<FAvaViewportGuideInfo_Deprecated> GuideInfos;
+#endif
+
+	FSimpleMulticastDelegate OnAnimationsUpdated;
+
+	bool bLoadingWorld = false;
+	bool bSavingWorld = false;
+	bool bStoppingAllAnimations = false;
+	bool bPendingAnimTreeUpdate = false;
+};
diff --git a/Engine/Plugins/Experimental/Avalanche/Source/Avalanche/Public/AvaBlueprintGeneratedClass.h b/Engine/Plugins/Experimental/Avalanche/Source/Avalanche/Public/AvaBlueprintGeneratedClass.h
new file mode 100644
--- /dev/null
+++ b/Engine/Plugins/Experimental/Avalanche/Source/Avalanche/Public/AvaBlueprintGeneratedClass.h
@@ -0,0 +1,25 @@
+// Copyright Epic Games, Inc. All Rights Reserved.
+
+#pragma once
+
+#include "Containers/Array.h"
+#include "Engine/BlueprintGeneratedClass.h"
+#include "UObject/ObjectPtr.h"
+#include "UObject/Package.h"
+
+#include "AvaBlueprintGeneratedClass.generated.h"
+
+class AAvalancheActor;
+class UAvaSequence;
+
+UCLASS(MinimalAPI, DisplayName = "Motion Design Blueprint Generated Class")
+class UAvaBlueprintGeneratedClass : public UBlueprintGeneratedClass
+{
+	GENERATED_BODY()
+
+public:
+	void UpdateProperties(AAvalancheActor* InActor);
+
+	UPROPERTY()
+	TArray<TObjectPtr<UAvaSequence>> Animations;
+};
diff --git a/Engine/Plugins/Experimental/Avalanche/Source/Avalanche/Public/AvaDefs.h b/Engine/Plugins/Experimental/Avalanche/Source/Avalanche/Public/AvaDefs.h
new file mode 100644
--- /dev/null
+++ b/Engine/Plugins/Experimental/Avalanche/Source/Avalanche/Public/AvaDefs.h
@@ -0,0 +1,349 @@
+﻿// Copyright Epic Games, Inc. All Rights Reserved.
+
+#pragma once
+
+#include <type_traits>
+#include "UObject/ObjectMacros.h"
+
+#include "AvaDefs.generated.h"
+
+// Single-axis exclusive bitflag map
+namespace UE::Avalanche::AnchorPoints
+{
+	constexpr int32 None       = 0;
+	constexpr int32 Left       = 1 << 0; // Y- = 1
+	constexpr int32 HMiddle    = 1 << 1; // Y/2 = 2
+	constexpr int32 Right      = 1 << 2; // Y+ = 4
+	constexpr int32 Top        = 1 << 3; // Z+ = 8
+	constexpr int32 VMiddle    = 1 << 4; // Z/2 = 16
+	constexpr int32 Bottom     = 1 << 5; // Z- = 32
+	constexpr int32 Horizontal = Left | HMiddle | Right; // Y- to Y+ = 7
+	constexpr int32 Vertical   = Top | VMiddle | Bottom; // Z+ to Z- = 56
+	constexpr int32 Custom     = 1 << 6; // Custom = 64
+}
+
+// UType 2D space point map. None should not be selectable in the UI. 
+// Z: Top, Bottom, VMiddle & Y: Left, Right, HMiddle
+// Expressed as ZY eg TopLeft
+UENUM(BlueprintType)
+enum class EAvaAnchors : uint8
+{
+	None = UE::Avalanche::AnchorPoints::None UMETA(Hidden),
+	TopLeft = UE::Avalanche::AnchorPoints::Top + UE::Avalanche::AnchorPoints::Left,
+	Top = UE::Avalanche::AnchorPoints::Top + UE::Avalanche::AnchorPoints::HMiddle,
+	TopRight = UE::Avalanche::AnchorPoints::Top + UE::Avalanche::AnchorPoints::Right,
+	Left = UE::Avalanche::AnchorPoints::VMiddle + UE::Avalanche::AnchorPoints::Left,
+	Center = UE::Avalanche::AnchorPoints::VMiddle + UE::Avalanche::AnchorPoints::HMiddle,
+	Right = UE::Avalanche::AnchorPoints::VMiddle + UE::Avalanche::AnchorPoints::Right,
+	BottomLeft = UE::Avalanche::AnchorPoints::Bottom + UE::Avalanche::AnchorPoints::Left,
+	Bottom = UE::Avalanche::AnchorPoints::Bottom + UE::Avalanche::AnchorPoints::HMiddle,
+	BottomRight = UE::Avalanche::AnchorPoints::Bottom + UE::Avalanche::AnchorPoints::Right,
+	Custom = UE::Avalanche::AnchorPoints::Custom
+};
+
+UENUM(BlueprintType)
+enum class EAvaLayoutPlacementCenter : uint8
+{
+	Bounds,
+	Pivot
+};
+
+UENUM(BlueprintType)
+enum class EAvaVerticalAlignment : uint8
+{
+	Top,
+	Center,
+	Bottom,
+};
+
+UENUM(BlueprintType)
+enum class EAvaHorizontalAlignment : uint8
+{
+	Left,
+	Center,
+	Right,
+};
+
+UENUM(BlueprintType)
+enum class EAvaDepthAlignment : uint8
+{
+	Front,
+	Center,
+	Back
+};
+
+UENUM(BlueprintType)
+enum class EAvaLayoutSizeBounds : uint8
+{
+	Explicit,
+	ScaleHeightToChildren,
+	ScaleWidthToChildren,
+	ScaleToChildren,
+};
+
+UENUM(BlueprintType)
+enum class EAvaLayoutRelationshipBounds : uint8
+{
+	DoNotScale,
+	ScaleToHeight,
+	ScaleToWidth,
+	ScaleToFit,
+	StretchHeight,
+	StretchWidth,
+	Stretch
+};
+
+/** Specifies the method for finding a reference actor based on it's position in the parent's hierarchy. */
+UENUM(BlueprintType)
+enum class EAvaReferenceContainer : uint8
+{
+	/** Uses the previous actor in the parent's hierarchy. */
+	Previous,
+	/** Uses the next actor in the parent's hierarchy. */
+	Next,
+	/** Uses the first actor in the parent's hierarchy. */
+	First,
+	/** Uses the last actor in the parent's hierarchy. */
+	Last,
+	/** Uses a specified reference actor set by the user. */
+	Other
+};
+
+/** Specifies a corner for a 2D rectangle. */
+UENUM(BlueprintType)
+enum class EAvaCorner2D : uint8
+{
+	TopLeft,
+	TopRight,
+	BottomLeft,
+	BottomRight
+};
+
+/** Specifies a 3D axis. */
+UENUM(BlueprintType)
+enum class EAvaAxis : uint8
+{
+	Horizontal,
+	Vertical,
+	Depth
+};
+
+template<typename Enum>
+constexpr auto ToUnderlyingType(Enum Val)
+{
+	return static_cast<std::underlying_type_t<Enum>>(Val);
+}
+
+typedef int32 AvaAlignment; // xyz
+
+// alignment for X (Depth)
+template<typename AvaAlignment>
+constexpr EAvaDepthAlignment GetDAlignment(AvaAlignment InAlignment)
+{
+	return static_cast<EAvaDepthAlignment>((0xFF & InAlignment) >> 0);
+}
+
+// alignment for Y (Horizontal)
+template<typename AvaAlignment>
+constexpr EAvaHorizontalAlignment GetHAlignment(AvaAlignment InAlignment)
+{
+	return static_cast<EAvaHorizontalAlignment>((0xFF00 & InAlignment) >> 8);
+}
+
+// alignment for Z (Vertical)
+template<typename AvaAlignment>
+constexpr EAvaVerticalAlignment GetVAlignment(AvaAlignment InAlignment)
+{
+	return static_cast<EAvaVerticalAlignment>((0xFF0000 & InAlignment) >> 16);
+}
+
+// gets an anchor location from an alignment and a size
+template<typename AvaAlignment>
+constexpr FVector GetLocationFromAlignment(AvaAlignment InAlignment, FVector Size3D)
+{
+	FVector AnchorLocation = FVector::ZeroVector;
+	
+	// get alignment enums from value
+	const EAvaDepthAlignment Depth = GetDAlignment(InAlignment);
+	// X
+	switch (Depth)
+	{
+	case EAvaDepthAlignment::Back:
+		AnchorLocation.X = -Size3D.X / 2.f;
+		break;
+	case EAvaDepthAlignment::Center:
+		AnchorLocation.X = 0.f;
+		break;
+	case EAvaDepthAlignment::Front:
+		AnchorLocation.X = Size3D.X / 2.f;
+		break;
+	default:
+		return FVector::ZeroVector;
+	}
+
+	const EAvaHorizontalAlignment Horizontal = GetHAlignment(InAlignment);
+	// Y
+	switch (Horizontal)
+	{
+	case EAvaHorizontalAlignment::Left:
+		AnchorLocation.Y = -Size3D.Y / 2.f;
+		break;
+	case EAvaHorizontalAlignment::Center:
+		AnchorLocation.Y = 0.f;
+		break;
+	case EAvaHorizontalAlignment::Right:
+		AnchorLocation.Y = Size3D.Y / 2.f;
+		break;
+	default:
+		return FVector::ZeroVector;
+	}
+
+	const EAvaVerticalAlignment Vertical = GetVAlignment(InAlignment);
+	// Z
+	switch (Vertical)
+	{
+	case EAvaVerticalAlignment::Bottom:
+		AnchorLocation.Z = -Size3D.Z / 2.f;
+		break;
+	case EAvaVerticalAlignment::Center:
+		AnchorLocation.Z = 0.f;
+		break;
+	case EAvaVerticalAlignment::Top:
+		AnchorLocation.Z = Size3D.Z / 2.f;
+		break;
+	default:
+		return FVector::ZeroVector;
+	}
+
+	return AnchorLocation;
+}
+
+// create alignment value from enums
+constexpr AvaAlignment MakeAlignment(EAvaDepthAlignment Depth, EAvaHorizontalAlignment Horizontal, EAvaVerticalAlignment Vertical)
+{
+	return AvaAlignment(static_cast<uint8>(Vertical) << 16 | static_cast<uint8>(Horizontal) << 8 | static_cast<uint8>(Depth) << 0);
+}
+
+UENUM(BlueprintType)
+enum class EAvaColorStyle : uint8
+{
+	None,
+	Solid,
+	LinearGradient
+};
+
+USTRUCT(BlueprintType)
+struct FAvaColorChangeData
+{
+	GENERATED_BODY()
+
+	UPROPERTY(Category = "Motion Design", EditAnywhere, BlueprintReadWrite)
+	EAvaColorStyle ColorStyle;
+
+	UPROPERTY(Category = "Motion Design", EditAnywhere, BlueprintReadWrite)
+	FLinearColor PrimaryColor;
+	
+	UPROPERTY(Category = "Motion Design", EditAnywhere, BlueprintReadWrite)
+	FLinearColor SecondaryColor;
+
+	UPROPERTY(Category = "Motion Design", EditAnywhere, BlueprintReadWrite)
+	bool bIsUnlit;
+
+	FAvaColorChangeData()
+		: ColorStyle(EAvaColorStyle::None), PrimaryColor(FLinearColor::White), SecondaryColor(FLinearColor::White), bIsUnlit(true)
+	{
+	}
+
+	FAvaColorChangeData(EAvaColorStyle InColorStyle, const FLinearColor& InPrimaryColor, const FLinearColor& InSecondaryColor, bool bInIsUnlit)
+		: ColorStyle(InColorStyle), PrimaryColor(InPrimaryColor), SecondaryColor(InSecondaryColor), bIsUnlit(bInIsUnlit)
+	{
+	}
+
+	bool operator==(const FAvaColorChangeData& InOtherColorData) const
+	{
+		return
+			ColorStyle     == InOtherColorData.ColorStyle &&
+			PrimaryColor   == InOtherColorData.PrimaryColor &&
+			SecondaryColor == InOtherColorData.SecondaryColor &&
+			bIsUnlit       == InOtherColorData.bIsUnlit;
+	};
+	
+	bool operator!=(const FAvaColorChangeData& InOtherColorData) const
+	{
+		return !operator==(InOtherColorData);
+	};
+};
+
+/** Specifies a set of anchor alignments, one for each 3D axis. */
+USTRUCT(BlueprintType)
+struct FAvaAnchorAlignment
+{
+	GENERATED_BODY()
+
+	UPROPERTY()
+	bool bUseHorizontal = true;
+	
+	UPROPERTY()
+	bool bUseVertical = true;
+	
+	UPROPERTY()
+	bool bUseDepth = true;
+
+	UPROPERTY(EditAnywhere, BlueprintReadWrite, Category="AnchorAlignment", meta=(EditCondition="bUseHorizontal", EditConditionHides))
+	EAvaHorizontalAlignment Horizontal = EAvaHorizontalAlignment::Center;
+
+	UPROPERTY(EditAnywhere, BlueprintReadWrite, Category="AnchorAlignment", meta=(EditCondition="bUseVertical", EditConditionHides))
+	EAvaVerticalAlignment Vertical = EAvaVerticalAlignment::Center;
+
+	UPROPERTY(EditAnywhere, BlueprintReadWrite, Category="AnchorAlignment", meta=(EditCondition="bUseDepth", EditConditionHides))
+	EAvaDepthAlignment Depth = EAvaDepthAlignment::Center;
+
+	FAvaAnchorAlignment() {}
+	FAvaAnchorAlignment(const EAvaHorizontalAlignment InHorizontal, const EAvaVerticalAlignment InVertical, const EAvaDepthAlignment InDepth)
+		: Horizontal(InHorizontal), Vertical(InVertical), Depth(InDepth)
+	{}
+
+	bool operator==(const FAvaAnchorAlignment& Other) const
+	{
+		return Horizontal == Other.Horizontal && Vertical == Other.Vertical && Depth == Other.Depth;
+	}
+
+	FORCEINLINE bool IsHorizontalMatch(const FAvaAnchorAlignment& Other) const { return Horizontal == Other.Horizontal; }
+	FORCEINLINE bool IsVerticalMatch(const FAvaAnchorAlignment& Other) const { return Vertical == Other.Vertical; }
+	FORCEINLINE bool IsDepthMatch(const FAvaAnchorAlignment& Other) const { return Depth == Other.Depth; }
+
+	/** Returns a point on the bounds extent determined by the alignment properties of the structure. */
+	FVector LocalBoundsOffset(const FBox& InBounds, const bool bInInverted = false) const
+	{
+		const FVector Extent = InBounds.GetExtent();
+
+		FVector BoundsOffset = FVector::ZeroVector;
+
+		auto HorizontalAlign = [bInInverted, Extent]() -> double { return bInInverted ? Extent.Y : -Extent.Y; };
+		auto VerticalAlign = [bInInverted, Extent]() -> double { return bInInverted ? Extent.Z : -Extent.Z; };
+		auto DepthAlign = [bInInverted, Extent]() -> double { return bInInverted ? Extent.X : -Extent.X; };
+
+		switch (Horizontal)
+		{
+		case EAvaHorizontalAlignment::Left: BoundsOffset.Y = HorizontalAlign(); break;
+		case EAvaHorizontalAlignment::Center: break;
+		case EAvaHorizontalAlignment::Right: BoundsOffset.Y = -HorizontalAlign(); break;
+		}
+
+		switch (Vertical)
+		{
+		case EAvaVerticalAlignment::Top: BoundsOffset.Z = VerticalAlign(); break;
+		case EAvaVerticalAlignment::Center: break;
+		case EAvaVerticalAlignment::Bottom: BoundsOffset.Z = -VerticalAlign(); break;
+		}
+
+		switch (Depth)
+		{
+		case EAvaDepthAlignment::Front: BoundsOffset.X = DepthAlign(); break;
+		case EAvaDepthAlignment::Center: break;
+		case EAvaDepthAlignment::Back: BoundsOffset.X = -DepthAlign(); break;
+		}
+
+		return BoundsOffset;
+	}
+};
diff --git a/Engine/Plugins/Experimental/Avalanche/Source/Avalanche/Public/AvaLog.h b/Engine/Plugins/Experimental/Avalanche/Source/Avalanche/Public/AvaLog.h
new file mode 100644
--- /dev/null
+++ b/Engine/Plugins/Experimental/Avalanche/Source/Avalanche/Public/AvaLog.h
@@ -0,0 +1,7 @@
+﻿// Copyright Epic Games, Inc. All Rights Reserved.
+
+#pragma once
+
+#include "Logging/LogMacros.h"
+
+AVALANCHE_API DECLARE_LOG_CATEGORY_EXTERN(LogAva, Log, All);
