[MotionDesign]

diff --git a/Engine/Plugins/Experimental/Avalanche/Source/Avalanche/Public/AvaOutlinerUtils.h b/Engine/Plugins/Experimental/Avalanche/Source/Avalanche/Public/AvaOutlinerUtils.h
new file mode 100644
--- /dev/null
+++ b/Engine/Plugins/Experimental/Avalanche/Source/Avalanche/Public/AvaOutlinerUtils.h
@@ -0,0 +1,34 @@
+// Copyright Epic Games, Inc. All Rights Reserved.
+
+#pragma once
+
+#if WITH_EDITOR
+
+#include "AvaOutlinerDefines.h"
+#include "Containers/Array.h"
+#include "Templates/SharedPointer.h"
+#include "UObject/WeakObjectPtr.h"
+
+class AActor;
+class FAvaOutliner;
+
+/** Editor Only Outliner Access Functions */
+struct AVALANCHE_API FAvaOutlinerUtils
+{
+	/** Returns the Avalanche outliner from the outliner module at runtime in the editor. */
+	static TSharedPtr<FAvaOutliner> EditorGetOutliner(const UWorld* const InWorld);
+
+	/** Returns an array of direct child actors attached to a parent actor or all root level actors if the parent actor is null. */
+	static TArray<AActor*> EditorOutlinerChildActors(TSharedPtr<FAvaOutliner> InOutliner, AActor* const InParentActor = nullptr);
+
+	/** Returns true if the Avalanche BP editor is isolating actors and the array of actors currently being isolated. */
+	static bool EditorActorIsolationInfo(TSharedPtr<FAvaOutliner> InOutliner, TArray<TWeakObjectPtr<const AActor>>& OutIsolatedActors);
+
+	/**
+	 * Gets an array of actors from an array of outliner items while maintaining index ordering (for actor outliner items only).
+	 * Returns false and an empty OutActors if any outliner items are invalid or any actor objects they represent are invalid.
+	 */
+	static bool EditorOutlinerItemsToActors(const TArray<FAvaOutlinerItemPtr>& InItems, TArray<AActor*>& OutActors);
+};
+
+#endif
diff --git a/Engine/Plugins/Experimental/Avalanche/Source/Avalanche/Public/AvaSceneSettings.h b/Engine/Plugins/Experimental/Avalanche/Source/Avalanche/Public/AvaSceneSettings.h
new file mode 100644
--- /dev/null
+++ b/Engine/Plugins/Experimental/Avalanche/Source/Avalanche/Public/AvaSceneSettings.h
@@ -0,0 +1,28 @@
+// Copyright Epic Games, Inc. All Rights Reserved.
+
+#pragma once
+
+#include "AvaTagHandleContainer.h"
+#include "UObject/Object.h"
+#include "AvaSceneSettings.generated.h"
+
+/** Object containing information about its Scene */
+UCLASS(MinimalAPI)
+class UAvaSceneSettings : public UObject
+{
+	GENERATED_BODY()
+
+public:
+	const FAvaTagHandleContainer& GetTagAttributes() const
+	{
+		return TagAttributes;
+	}
+
+	UFUNCTION()
+	void SetTagAttributes(const FAvaTagHandleContainer& InTagAttributes);
+
+private:
+	/** Tags describing the Scene */
+	UPROPERTY(EditAnywhere, Setter="SetTagAttributes", Category="Scene", meta=(AllowPrivateAccess="true"))
+	FAvaTagHandleContainer TagAttributes;
+};
diff --git a/Engine/Plugins/Experimental/Avalanche/Source/Avalanche/Public/AvaSceneState.h b/Engine/Plugins/Experimental/Avalanche/Source/Avalanche/Public/AvaSceneState.h
new file mode 100644
--- /dev/null
+++ b/Engine/Plugins/Experimental/Avalanche/Source/Avalanche/Public/AvaSceneState.h
@@ -0,0 +1,32 @@
+// Copyright Epic Games, Inc. All Rights Reserved.
+
+#pragma once
+
+#include "AvaTagHandleContainer.h"
+#include "UObject/Object.h"
+#include "AvaSceneState.generated.h"
+
+class UAvaSceneSettings;
+
+/** Object providing State information of the Scene */
+UCLASS(MinimalAPI)
+class UAvaSceneState : public UObject
+{
+	GENERATED_BODY()
+
+public:
+	void OnBeginPlay(UAvaSceneSettings* InSceneSettings);
+
+	AVALANCHE_API bool AddTagAttribute(const FAvaTagHandle& InTagHandle);
+
+	AVALANCHE_API bool RemoveTagAttribute(const FAvaTagHandle& InTagHandle);
+
+	AVALANCHE_API bool ContainsTagAttribute(const FAvaTagHandle& InTagHandle) const;
+
+private:
+	/** Weak pointer to the same scene's Settings. Used to query for default tag attributes of the scene */
+	TWeakObjectPtr<UAvaSceneSettings> SceneSettingsWeak;
+
+	/** Active Tag Attributes that were added to the scene, separate from Scene Settings */
+	TSet<FAvaTag> ActiveTagAttributes;
+};
diff --git a/Engine/Plugins/Experimental/Avalanche/Source/Avalanche/Public/AvaSceneSubsystem.h b/Engine/Plugins/Experimental/Avalanche/Source/Avalanche/Public/AvaSceneSubsystem.h
new file mode 100644
--- /dev/null
+++ b/Engine/Plugins/Experimental/Avalanche/Source/Avalanche/Public/AvaSceneSubsystem.h
@@ -0,0 +1,39 @@
+// Copyright Epic Games, Inc. All Rights Reserved.
+
+#pragma once
+
+#include "Containers/Map.h"
+#include "Subsystems/WorldSubsystem.h"
+#include "UObject/ObjectPtr.h"
+#include "UObject/WeakInterfacePtr.h"
+#include "AvaSceneSubsystem.generated.h"
+
+class IAvaSceneInterface;
+class ULevel;
+
+UCLASS(MinimalAPI, DisplayName = "Motion Design Scene Subsystem")
+class UAvaSceneSubsystem : public UWorldSubsystem
+{
+	GENERATED_BODY()
+
+public:
+	void RegisterSceneInterface(ULevel* InLevel, IAvaSceneInterface* InSceneInterface);
+
+	/** Gets the Scene Interface for the World's Persistent Level */
+	AVALANCHE_API IAvaSceneInterface* GetSceneInterface() const;
+
+	/** Gets the Scene Interface for the provided Level */
+	AVALANCHE_API IAvaSceneInterface* GetSceneInterface(ULevel* InLevel) const;
+
+	/** Returns the first interface found in the given Level, if any */
+	AVALANCHE_API static IAvaSceneInterface* FindSceneInterface(ULevel* InLevel);
+
+protected:
+	//~ Begin UWorldSubsystem
+	AVALANCHE_API virtual void PostInitialize() override;
+	AVALANCHE_API virtual bool DoesSupportWorldType(const EWorldType::Type InWorldType) const override;
+	//~ End UWorldSubsystem
+
+private:
+	TMap<TWeakObjectPtr<ULevel>, TWeakInterfacePtr<IAvaSceneInterface>> SceneInterfaces;
+};
diff --git a/Engine/Plugins/Experimental/Avalanche/Source/Avalanche/Public/Data/AvalancheActorData.h b/Engine/Plugins/Experimental/Avalanche/Source/Avalanche/Public/Data/AvalancheActorData.h
new file mode 100644
--- /dev/null
+++ b/Engine/Plugins/Experimental/Avalanche/Source/Avalanche/Public/Data/AvalancheActorData.h
@@ -0,0 +1,37 @@
+// Copyright Epic Games, Inc. All Rights Reserved.
+
+#pragma once
+
+#include "AvalancheComponentData.h"
+#include "AvalancheDataDefines.h"
+#include "AvalancheObjectData.h"
+#include "Containers/Array.h"
+#include "Containers/Map.h"
+#include "Containers/UnrealString.h"
+#include "UObject/SoftObjectPath.h"
+#include "AvalancheActorData.generated.h"
+
+USTRUCT()
+struct FAvalancheActorData : public FAvalancheObjectData
+{
+	GENERATED_BODY()
+
+#if WITH_EDITORONLY_DATA
+	UPROPERTY()
+	FString ActorLabel;
+
+	UPROPERTY()
+	bool bEditorVisibility = false;
+#endif
+	
+	UPROPERTY()
+	FSoftClassPath ActorClass;
+
+	/** Additional Info for Components */
+	UPROPERTY()
+	TMap<FAvaObjectIndex, FAvalancheComponentData> ComponentData;
+
+	/** Indices to all Non-Component SubObjects */
+	UPROPERTY()
+	TArray<FAvaObjectIndex> OwnedSubobjects;
+};
diff --git a/Engine/Plugins/Experimental/Avalanche/Source/Avalanche/Public/Data/AvalancheComponentData.h b/Engine/Plugins/Experimental/Avalanche/Source/Avalanche/Public/Data/AvalancheComponentData.h
new file mode 100644
--- /dev/null
+++ b/Engine/Plugins/Experimental/Avalanche/Source/Avalanche/Public/Data/AvalancheComponentData.h
@@ -0,0 +1,16 @@
+// Copyright Epic Games, Inc. All Rights Reserved.
+
+#pragma once
+
+#include "ComponentInstanceDataCache.h"
+#include "AvalancheComponentData.generated.h"
+
+USTRUCT()
+struct FAvalancheComponentData
+{
+	GENERATED_BODY()
+
+	/** Describes how the component was created */
+	UPROPERTY()
+	EComponentCreationMethod CreationMethod{};
+};
diff --git a/Engine/Plugins/Experimental/Avalanche/Source/Avalanche/Public/Data/AvalancheDataDefines.h b/Engine/Plugins/Experimental/Avalanche/Source/Avalanche/Public/Data/AvalancheDataDefines.h
new file mode 100644
--- /dev/null
+++ b/Engine/Plugins/Experimental/Avalanche/Source/Avalanche/Public/Data/AvalancheDataDefines.h
@@ -0,0 +1,28 @@
+// Copyright Epic Games, Inc. All Rights Reserved.
+
+#pragma once
+
+#include "AvalancheDataDefines.generated.h"
+
+USTRUCT()
+struct FAvaObjectIndex
+{
+	GENERATED_BODY()
+
+	FAvaObjectIndex() {}
+
+	FAvaObjectIndex(int32 Index) : Index(Index) {}
+
+	bool operator==(const FAvaObjectIndex& Other) const
+	{
+		return Index == Other.Index;
+	}
+	
+	friend uint32 GetTypeHash(const FAvaObjectIndex& Value)
+    {
+    	return GetTypeHash(Value.Index);
+    }
+
+	UPROPERTY()
+	int32 Index = INDEX_NONE;
+};
diff --git a/Engine/Plugins/Experimental/Avalanche/Source/Avalanche/Public/Data/AvalancheObjectData.h b/Engine/Plugins/Experimental/Avalanche/Source/Avalanche/Public/Data/AvalancheObjectData.h
new file mode 100644
--- /dev/null
+++ b/Engine/Plugins/Experimental/Avalanche/Source/Avalanche/Public/Data/AvalancheObjectData.h
@@ -0,0 +1,21 @@
+// Copyright Epic Games, Inc. All Rights Reserved.
+
+#pragma once
+
+#include "Containers/Array.h"
+#include "UObject/ObjectMacros.h"
+#include "AvalancheObjectData.generated.h"
+
+USTRUCT()
+struct FAvalancheObjectData
+{
+	GENERATED_BODY()
+
+	UPROPERTY()
+	TArray<uint8> SerializedData;
+
+	UPROPERTY()
+	uint64 ObjectFlags{};
+
+	EObjectFlags GetObjectFlags() const { return static_cast<EObjectFlags>(ObjectFlags); }
+};
diff --git a/Engine/Plugins/Experimental/Avalanche/Source/Avalanche/Public/Data/AvalancheSubObjectData.h b/Engine/Plugins/Experimental/Avalanche/Source/Avalanche/Public/Data/AvalancheSubObjectData.h
new file mode 100644
--- /dev/null
+++ b/Engine/Plugins/Experimental/Avalanche/Source/Avalanche/Public/Data/AvalancheSubObjectData.h
@@ -0,0 +1,32 @@
+// Copyright Epic Games, Inc. All Rights Reserved.
+
+#pragma once
+
+#include "AvalancheDataDefines.h"
+#include "AvalancheObjectData.h"
+#include "UObject/SoftObjectPath.h"
+#include "AvalancheSubObjectData.generated.h"
+
+USTRUCT()
+struct FAvalancheSubObjectData : public FAvalancheObjectData
+{
+	GENERATED_BODY()
+
+	static FAvalancheSubObjectData MakeSkippedSubObjectData()
+	{
+		FAvalancheSubObjectData OutData;
+		OutData.bWasSkippedClass = true;
+		return OutData;
+	}
+
+	/** Index to FAvalancheWorldData::SerializedObjectReferences */
+	UPROPERTY()
+	FAvaObjectIndex OuterIndex;
+
+	UPROPERTY()
+	FSoftClassPath Class;
+
+	/** Whether this class was marked as unsupported when the snapshot was taken */
+	UPROPERTY()
+	bool bWasSkippedClass = false;
+};
diff --git a/Engine/Plugins/Experimental/Avalanche/Source/Avalanche/Public/Data/AvalancheVersionInfo.h b/Engine/Plugins/Experimental/Avalanche/Source/Avalanche/Public/Data/AvalancheVersionInfo.h
new file mode 100644
--- /dev/null
+++ b/Engine/Plugins/Experimental/Avalanche/Source/Avalanche/Public/Data/AvalancheVersionInfo.h
@@ -0,0 +1,122 @@
+// Copyright Epic Games, Inc. All Rights Reserved.
+
+#pragma once
+
+#include "Misc/Guid.h"
+#include "UObject/NameTypes.h"
+#include "AvalancheVersionInfo.generated.h"
+
+class FArchive;
+struct FCustomVersion;
+
+/** Holds file version information */
+USTRUCT()
+struct FAvalancheFileVersionInfo
+{
+	GENERATED_BODY()
+
+	FAvalancheFileVersionInfo()
+		: FileVersionUE4(0)
+		, FileVersionUE5(0)
+		, FileVersionLicensee(0)
+	{
+	}
+	
+	void UpdateToLatest();
+
+	UPROPERTY()
+	int32 FileVersionUE4;
+
+	UPROPERTY()
+	int32 FileVersionUE5;
+	
+	/* Licensee file version */
+	UPROPERTY()
+	int32 FileVersionLicensee;
+};
+
+/** Holds engine version information */
+USTRUCT()
+struct FAvalancheEngineVersionInfo
+{
+	GENERATED_BODY()
+
+	FAvalancheEngineVersionInfo()
+		: Major(0)
+		, Minor(0)
+		, Patch(0)
+		, Changelist(0)
+	{
+	}
+	
+	void UpdateToLatest();
+
+	/** Major version number */
+	UPROPERTY()
+	uint16 Major;
+
+	/** Minor version number */
+	UPROPERTY()
+	uint16 Minor;
+
+	/** Patch version number */
+	UPROPERTY()
+	uint16 Patch;
+
+	/** Changelist number. This is used to arbitrate when Major/Minor/Patch version numbers match */
+	UPROPERTY()
+	uint32 Changelist;
+};
+
+/** Holds custom version information */
+USTRUCT()
+struct FAvalancheCustomVersionInfo
+{
+	GENERATED_BODY()
+
+	FAvalancheCustomVersionInfo()
+		: Version(0)
+	{
+	}
+	
+	void SetCustomVersion(const FCustomVersion& InVersion);
+	
+	/** Unique custom key */
+	UPROPERTY()
+	FGuid Key;
+
+	/** Custom version */
+	UPROPERTY()
+	int32 Version;
+
+	/** Friendly name of the version */
+	UPROPERTY()
+	FName FriendlyName;
+};
+
+/** Holds version information for a session */
+USTRUCT()
+struct FAvalancheVersionInfo
+{
+	GENERATED_BODY()
+
+	/** Update all Versions (File, Engine, Custom) to Latest */
+	void UpdateToLatest();
+
+	void ApplyToArchive(FArchive& Archive) const;
+
+	/** File version info */
+	UPROPERTY()
+	FAvalancheFileVersionInfo FileVersion;
+
+	/** Engine version info */
+	UPROPERTY()
+	FAvalancheEngineVersionInfo EngineVersion;
+
+	/** Custom version info */
+	UPROPERTY()
+	TArray<FAvalancheCustomVersionInfo> CustomVersions;
+
+	/** Keep track of the custom versions used by the assets while saving. */
+	TSet<FGuid> UsedCustomVersions;
+};
\ No newline at end of file
diff --git a/Engine/Plugins/Experimental/Avalanche/Source/Avalanche/Public/Data/AvalancheWorldData.h b/Engine/Plugins/Experimental/Avalanche/Source/Avalanche/Public/Data/AvalancheWorldData.h
new file mode 100644
--- /dev/null
+++ b/Engine/Plugins/Experimental/Avalanche/Source/Avalanche/Public/Data/AvalancheWorldData.h
@@ -0,0 +1,85 @@
+// Copyright Epic Games, Inc. All Rights Reserved.
+
+#pragma once
+
+#include "AvalancheActorData.h"
+#include "AvalancheSubObjectData.h"
+#include "AvalancheVersionInfo.h"
+#include "Containers/Array.h"
+#include "Containers/Map.h"
+#include "UObject/NameTypes.h"
+#include "UObject/SoftObjectPath.h"
+#include "UObject/WeakObjectPtr.h"
+#include "AvalancheWorldData.generated.h"
+
+class AActor;
+class UObject;
+class UWorld;
+
+USTRUCT()
+struct FAvalancheWorldData
+{
+	GENERATED_BODY()
+	
+	/**
+	 *	Called after saving the world to perform final processing.
+	 */
+	void FinalizeSave();
+	
+	/**
+	 * Resets the Transient Members to zero, since DuplicateTransient only works on UCLASS not USTRUCT.
+	 */
+	void ResetTransientData();
+	
+	/**
+	* Holds the Version Info this World Data was last serialized with
+	*/
+	UPROPERTY()
+	FAvalancheVersionInfo VersionInfo;
+
+	/**
+	* Transient World used to load the World Data in
+	*/
+	UPROPERTY(Transient)
+	TWeakObjectPtr<UWorld> World;
+
+	/**
+	* Holds serialized actor data.
+	* Maps the original actor's path to its serialized data.
+	*/
+	UPROPERTY()
+	TMap<FSoftObjectPath, FAvalancheActorData> ActorData;
+
+	UPROPERTY()
+	TArray<FName> SerializedNames;
+
+	/* 
+	 * List of all Object References in the World
+	 * This is an Untranslated version so we ned to change the World Name Part (when it was serialized)
+	 * of the Path to the New World Name
+	 */
+	UPROPERTY()
+	TArray<FSoftObjectPath> SerializedObjectReferences;
+
+	/*
+	 * Key: Index to SerializedObjectReferences List
+	 * Value: SubObject Information for the Object Reference
+	 * (This Container effectively only contains SubObjects so there will be indices not found here)
+	 */
+	UPROPERTY()
+	TMap<FAvaObjectIndex, FAvalancheSubObjectData> SubObjects;
+
+	/** Binds every entry in SerializedNames to its index to speed up look up */
+	UPROPERTY(Transient)
+	TMap<FName, FAvaObjectIndex> NameIndexMap;
+
+	/** Binds every entry in SerializedObjectReferences to its index to speed up look up */
+	UPROPERTY(Transient)
+	TMap<FSoftObjectPath, FAvaObjectIndex> ObjectReferenceIndexMap;
+
+	UPROPERTY(Transient)
+	TMap<FSoftObjectPath, TWeakObjectPtr<AActor>> CachedActors;
+
+	UPROPERTY(Transient)
+	TMap<FSoftObjectPath, TWeakObjectPtr<UObject>> CachedSubObjects;
+};
diff --git a/Engine/Plugins/Experimental/Avalanche/Source/Avalanche/Public/Framework/AvaNullActor.h b/Engine/Plugins/Experimental/Avalanche/Source/Avalanche/Public/Framework/AvaNullActor.h
new file mode 100644
--- /dev/null
+++ b/Engine/Plugins/Experimental/Avalanche/Source/Avalanche/Public/Framework/AvaNullActor.h
@@ -0,0 +1,27 @@
+// Copyright Epic Games, Inc. All Rights Reserved.
+
+#pragma once
+
+#include "Containers/UnrealString.h"
+#include "GameFramework/Actor.h"
+#include "AvaNullActor.generated.h"
+
+class UAvaNullComponent;
+
+/**
+ * Avalanche Null Actor (Empty Group)
+ */
+UCLASS(MinimalAPI, DisplayName = "Motion Design Null Actor")
+class AAvaNullActor : public AActor
+{
+	GENERATED_BODY()
+
+public:
+	static const FString DefaultLabel;
+
+	AAvaNullActor();
+
+protected:
+	UPROPERTY(VisibleAnywhere, BlueprintReadOnly, Category = "Motion Design")
+	TObjectPtr<UAvaNullComponent> NullComponent;
+};
diff --git a/Engine/Plugins/Experimental/Avalanche/Source/Avalanche/Public/Framework/AvaSoftAssetPtr.h b/Engine/Plugins/Experimental/Avalanche/Source/Avalanche/Public/Framework/AvaSoftAssetPtr.h
new file mode 100644
--- /dev/null
+++ b/Engine/Plugins/Experimental/Avalanche/Source/Avalanche/Public/Framework/AvaSoftAssetPtr.h
@@ -0,0 +1,59 @@
+// Copyright Epic Games, Inc. All Rights Reserved.
+
+#pragma once
+
+#include "UObject/SoftObjectPtr.h"
+
+enum class EAvalancheAssetType
+{
+	Unknown,
+	Blueprint,
+	World
+};
+
+/**
+ * Avalanche Extension to FSoftObjectPath.
+ * Adding the class information along to resolve between blueprint and level without loading anything.
+ */
+struct FAvaSoftAssetPath
+{
+	FSoftClassPath AssetClassPath;
+	FSoftObjectPath AssetPath;
+
+	bool IsNull() const { return AssetPath.IsNull(); }
+	const FSoftObjectPath& ToSoftObjectPath() const { return AssetPath; }
+	EAvalancheAssetType GetAssetType(bool bInLoadIfUnknown = true) const { return GetAssetTypeFromClass(AssetClassPath, bInLoadIfUnknown);}
+
+	AVALANCHE_API static EAvalancheAssetType GetAssetTypeFromClass(const FSoftClassPath& InAssetClassPath, bool bInLoadIfUnknown = true);
+};
+
+/**
+ * Avalanche Extension to FSoftObjectPtr.
+ * Adding the class information along to resolve between blueprint and level without loading anything.
+ */
+
+struct FAvaSoftAssetPtr
+{
+	FSoftClassPath AssetClassPath;
+
+	// We need to use the template, it is the only way we can assign it from another SoftObjetPtr.
+	TSoftObjectPtr<UObject> AssetPtr;
+
+	// Warning: IsValid() for an AssetPtr will cause the asset to be loaded.
+	// Using IsNull() is the only option we have.	
+	bool IsNull() const { return AssetPtr.IsNull(); }
+	
+	const FSoftObjectPath& ToSoftObjectPath() const { return AssetPtr.ToSoftObjectPath(); }
+
+	/**
+	 * @brief Determines the asset type from the given class path or by loading the object.
+	 * @param bInLoadIfUnknown If the type can't be determined by the class path, the object is loaded to check it's type.
+	 * @return One of the recognised asset type for Avalanche playables.
+	 */
+	AVALANCHE_API EAvalancheAssetType GetAssetType(bool bInLoadIfUnknown = true) const;
+
+	FORCEINLINE bool operator==(const FAvaSoftAssetPtr& Rhs) const
+	{
+		return AssetClassPath == Rhs.AssetClassPath && AssetPtr == Rhs.AssetPtr;
+	}
+};
\ No newline at end of file
diff --git a/Engine/Plugins/Experimental/Avalanche/Source/Avalanche/Public/Framework/AvalancheGameInstance.h b/Engine/Plugins/Experimental/Avalanche/Source/Avalanche/Public/Framework/AvalancheGameInstance.h
new file mode 100644
--- /dev/null
+++ b/Engine/Plugins/Experimental/Avalanche/Source/Avalanche/Public/Framework/AvalancheGameInstance.h
@@ -0,0 +1,187 @@
+// Copyright Epic Games, Inc. All Rights Reserved.
+
+#pragma once
+
+#include "AvalancheGameViewportClient.h"
+#include "Engine/GameInstance.h"
+#include "Engine/World.h"
+#include "AvalancheGameInstance.generated.h"
+
+class FRenderCommandFence;
+class FSceneViewport;
+class UAvalancheBlueprint;
+class UTextureRenderTarget2D;
+struct FAvaViewportQualitySettings;
+struct FAvalancheInstanceSettings;
+
+struct FAvalancheInstancePlaySettings
+{
+	const FAvalancheInstanceSettings& Settings;
+	FName ChannelName;
+	UTextureRenderTarget2D* RenderTarget;
+	FIntPoint ViewportSize;
+	const FAvaViewportQualitySettings& QualitySettings;
+};
+
+UCLASS(DisplayName = "Motion Design Game Instance")
+class AVALANCHE_API UAvalancheGameInstance : public UGameInstance
+{
+	GENERATED_BODY()
+
+public:
+	/**
+	 * Creates a new Avalanche Game instance from the given Avalanche Template Asset.
+	 * Supported asset types: UAvalancheBlueprint, UWorld (Level).
+	 *
+	 * Remark: For Avalanche Blueprints, the instanced embedded RemoteControlPreset will be registered
+	 * under the name of the outer Package Name. It is thus required that each instance be contained in
+	 * a unique package to avoid stomping the registered RCPs.
+	 */
+	static UAvalancheGameInstance* Create(UObject* InOuter);
+
+	// Legacy - will create and load an avalanche blueprint.
+	static UAvalancheGameInstance* Create(UObject* InOuter, const TSoftObjectPtr<UAvalancheBlueprint>& InAvalancheBlueprintTemplate);
+	
+	/** Create the game world */
+	bool CreateWorld();
+
+	bool IsWorldCreated() const { return bWorldCreated; }
+
+	UWorld* GetPlayWorld() const { return PlayWorld.Get(); }
+
+	// Legacy
+	bool LoadAvalancheBlueprint(const TSoftObjectPtr<UAvalancheBlueprint>& InSourceAvalancheBlueprint);
+
+	// Legacy
+	void BeginPlayAvalancheBlueprint(const TSoftObjectPtr<UAvalancheBlueprint>& InSourceAvalancheBlueprint, const FAvalancheInstancePlaySettings& InWorldPlaySettings);
+	
+	bool BeginPlayWorld(const FAvalancheInstancePlaySettings& InWorldPlaySettings);
+
+	/**
+	 * The end play is normally requested to be done on the next Tick(). This is to
+	 * avoid having the world destroyed while within the Tick() of that world.
+	 * However, when shutting down, we need to force the end play to be done immediately because
+	 * there will not a be a next Tick().
+	 */
+	void RequestEndPlayWorld(bool bForceImmediate);
+	
+	void RequestUnloadWorld(bool bForceImmediate);
+
+	bool IsWorldPlaying() const { return bWorldPlaying; }
+
+	/**
+	 * Cancel any pending EndPlayWorld or UnloadWorld requests. 
+	 */
+	void CancelWorldRequests();
+
+	/**
+	* The output channels can be resized during playback, if this happens, the
+	* scene viewport needs to be updated to reflect the change.
+	*/
+	void UpdateSceneViewportSize(const FIntPoint& InViewportSize);
+
+	void UpdateRenderTarget(UTextureRenderTarget2D* InRenderTarget);
+
+	UTextureRenderTarget2D* GetRenderTarget() const;
+
+	/**
+	 * Returns true if the render target has been rendered and is ready to be used for output.
+	 */
+	bool IsRenderTargetReady() const { return bIsRenderTargetReady;}
+
+	/**
+	 * Reset the flag so the OnRenderTargetReady event can be called again on the next render.
+	 */
+	void ResetRenderTargetReady() { bIsRenderTargetReady = false; }
+
+	/**
+	 * @brief Get the currently playing scene viewport.
+	 * @remark Only valid within a BeginPlay()/EndPlay() scope. Null otherwise.
+	 */
+	TSharedPtr<FSceneViewport> GetSceneViewport() const { return Viewport; }
+
+	/**
+	 * @brief Get the currently playing game viewport client.
+	 * @remark Only valid within a BeginPlay()/EndPlay() scope. Null otherwise.
+	 */
+	UAvalancheGameViewportClient* GetAvalancheGameViewportClient() const { return ViewportClient; }
+	
+	/*
+	 * @return The Avalanche Blueprint used in play. Can be null if called outside play
+	 */
+	UAvalancheBlueprint* GetAvalancheBlueprint() const { return ManagedAvalancheBlueprint; }
+
+	DECLARE_MULTICAST_DELEGATE_TwoParams(FOnAvaGameInstanceEvent, UAvalancheGameInstance* /*InGameInstance*/, FName /*InChannelName*/);
+
+	/** Event called on EndPlayWorld(). */
+	static FOnAvaGameInstanceEvent& GetOnEndPlay() { return OnEndPlay; }
+
+	/** Event called when the render target becomes ready. */
+	static FOnAvaGameInstanceEvent& GetOnRenderTargetReady() { return OnRenderTargetReady; }
+
+	/*
+	 * Mark the current frame as a frame in which an asset was loaded synchronously.
+	 * When an asset has been loaded synchronously (i.e. avalanche blueprint), it will cause a
+	 * large delta seconds to occur on the next tick. By marking that particular frame it can be
+	 * used on the next tick to clamp the delta seconds to avoid animations skipping by the
+	 * amount of the sync load time.
+	 */
+	void MarkSynchronousAssetLoadingThisFrame();
+
+protected:
+	void Tick(float DeltaSeconds);
+
+	void UnloadWorld();
+	void EndPlayWorld();
+	void OnEndFrameTick();
+	void OnEnginePreExit();
+	
+	//~ Begin UObject
+	virtual void BeginDestroy() override;
+	//~ End UObject
+
+	/*
+	 * A Soft Ptr to the Source Avalanche (i.e. the actual asset).
+	 * It's a Soft Ptr because it can become stale if it's only a Weak Object Ptr
+	 * and we want to reload it if that happens
+	 */
+	TSoftObjectPtr<UAvalancheBlueprint> SourceAvalancheBlueprint;
+
+	/*
+	 * The Managed Avalanche, a Duplication of the Source Avalanche when Begin Play was called.
+	 * This shouldn't be GCd while in play, as it contains Animation and other data used during play.
+	 */
+	UPROPERTY(Transient)
+	TObjectPtr<UAvalancheBlueprint> ManagedAvalancheBlueprint;
+
+	UPROPERTY(Transient)
+	TObjectPtr<UWorld> PlayWorld;
+
+	UPROPERTY(Transient)
+	TObjectPtr<UAvalancheGameViewportClient> ViewportClient;
+
+	TSharedPtr<FSceneViewport> Viewport;
+
+	/** Channel this game instance is playing on. */
+	FName PlayingChannelName;
+
+	bool bWorldCreated = false;
+	bool bRequestUnloadWorld = false;
+	bool bWorldPlaying = false;
+	bool bRequestEndPlayWorld = false;
+
+	/** Tick reentrancy guard. */
+	bool bIsTicking = false;
+	
+	double LastDeltaSeconds = 0.0;
+
+	/** Render command fence to ensure the render target has been rendered. */
+	TUniquePtr<FRenderCommandFence> RenderTargetFence;
+	bool bIsRenderTargetReady = false;
+
+	/** Handle to delegate for EnginePreExit event. */
+	FDelegateHandle EnginePreExitHandle;
+
+	static FOnAvaGameInstanceEvent OnEndPlay;
+	static FOnAvaGameInstanceEvent OnRenderTargetReady;
+};
diff --git a/Engine/Plugins/Experimental/Avalanche/Source/Avalanche/Public/Framework/AvalancheGameViewportClient.h b/Engine/Plugins/Experimental/Avalanche/Source/Avalanche/Public/Framework/AvalancheGameViewportClient.h
new file mode 100644
--- /dev/null
+++ b/Engine/Plugins/Experimental/Avalanche/Source/Avalanche/Public/Framework/AvalancheGameViewportClient.h
@@ -0,0 +1,61 @@
+// Copyright Epic Games, Inc. All Rights Reserved.
+
+#pragma once
+
+#include "Engine/GameViewportClient.h"
+#include "Viewport/AvaCameraManager.h"
+#include "SceneTypes.h"
+#include "AvalancheGameViewportClient.generated.h"
+
+class AAvaViewportCameraActor;
+class FPrimitiveComponentId;
+class FViewElementDrawer;
+class UTextureRenderTarget2D;
+struct FMinimalViewInfo;
+struct FSceneViewInitOptions;
+struct FSceneViewProjectionData;
+
+UCLASS(DisplayName = "Motion Design Game Viewport Client")
+class AVALANCHE_API UAvalancheGameViewportClient : public UGameViewportClient
+{
+	GENERATED_BODY()
+	
+public:
+	UAvalancheGameViewportClient();
+
+	//~ Begin FViewportClient
+	virtual void Draw(FViewport* Viewport,FCanvas* Canvas) override;
+	virtual bool IsStatEnabled(const FString& InName) const override;
+	//~ End FViewportClient
+
+	void SetCameraCutThisFrame();
+
+	void SetRenderTarget(UTextureRenderTarget2D* InRenderTarget);
+
+	UTextureRenderTarget2D* GetRenderTarget() const { return RenderTarget.Get(); }
+
+	TSharedRef<FAvaCameraManager> GetCameraManager() const { return CameraManager; }
+
+private:
+	bool CalcSceneViewInitOptions(FSceneViewInitOptions& OutInitOptions
+		, FViewport* InViewport
+		, FViewElementDrawer* ViewDrawer = nullptr
+		, int32 StereoViewIndex = INDEX_NONE);
+
+	bool GetProjectionData(FViewport* InViewport
+		, FSceneViewProjectionData& ProjectionData
+		, int32 StereoViewIndex = INDEX_NONE) const;
+
+	FSceneView* CalcSceneView(FSceneViewFamily* ViewFamily
+		, FViewport* InViewport
+		, FViewElementDrawer* ViewDrawer = nullptr
+		, int32 StereoViewIndex = INDEX_NONE);
+
+	TArray<FSceneViewStateReference> ViewStates;
+
+	TWeakObjectPtr<UTextureRenderTarget2D> RenderTarget;
+
+	TSharedRef<FAvaCameraManager> CameraManager;
+
+	bool bCameraCutThisFrame = false;
+};
diff --git a/Engine/Plugins/Experimental/Avalanche/Source/Avalanche/Public/Framework/AvalancheInstanceSettings.h b/Engine/Plugins/Experimental/Avalanche/Source/Avalanche/Public/Framework/AvalancheInstanceSettings.h
new file mode 100644
--- /dev/null
+++ b/Engine/Plugins/Experimental/Avalanche/Source/Avalanche/Public/Framework/AvalancheInstanceSettings.h
@@ -0,0 +1,18 @@
+// Copyright Epic Games, Inc. All Rights Reserved.
+
+#pragma once
+
+#include "AvalancheInstanceSettings.generated.h"
+
+/**
+ * Settings applied when instancing an avalanche blueprint for playback.
+ */
+USTRUCT()
+struct FAvalancheInstanceSettings
+{
+	GENERATED_BODY()
+
+	/** Enable loading dependent levels as sub-playables. */
+	UPROPERTY(config, EditAnywhere, Category = Settings)
+	bool bEnableLoadSubPlayables = true;
+};
diff --git a/Engine/Plugins/Experimental/Avalanche/Source/Avalanche/Public/Transition/Conditions/AvaSceneContainsTagAttributeCondition.h b/Engine/Plugins/Experimental/Avalanche/Source/Avalanche/Public/Transition/Conditions/AvaSceneContainsTagAttributeCondition.h
new file mode 100644
--- /dev/null
+++ b/Engine/Plugins/Experimental/Avalanche/Source/Avalanche/Public/Transition/Conditions/AvaSceneContainsTagAttributeCondition.h
@@ -0,0 +1,65 @@
+// Copyright Epic Games, Inc. All Rights Reserved.
+
+#pragma once
+
+#include "AvaTagHandle.h"
+#include "AvaTagHandleContainer.h"
+#include "AvaTransitionEnums.h"
+#include "Conditions/AvaTransitionCondition.h"
+#include "AvaSceneContainsTagAttributeCondition.generated.h"
+
+class UAvaSceneSubsystem;
+struct FAvaTransitionScene;
+
+USTRUCT()
+struct FAvaSceneContainsAttributeConditionInstanceData
+{
+	GENERATED_BODY()
+};
+
+USTRUCT(DisplayName="Scene contains tag attribute", Category="Scene Attributes")
+struct AVALANCHE_API FAvaSceneContainsTagAttributeCondition : public FAvaTransitionCondition
+{
+	GENERATED_BODY()
+
+	using FInstanceDataType = FAvaSceneContainsAttributeConditionInstanceData;
+
+	//~ Begin FAvaTransitionCondition
+	virtual FText GenerateDescription() const override;
+	//~ End FAvaTransitionCondition
+
+	//~ Begin FStateTreeNodeBase
+	virtual const UStruct* GetInstanceDataType() const override { return FInstanceDataType::StaticStruct(); }
+	virtual bool Link(FStateTreeLinker& InLinker) override;
+	//~ End FStateTreeNodeBase
+
+	//~ Begin FStateTreeConditionBase
+	virtual bool TestCondition(FStateTreeExecutionContext& InContext) const override;
+	//~ End FStateTreeConditionBase
+
+	bool ContainsTagAttribute(FStateTreeExecutionContext& InContext) const;
+
+	TArray<const FAvaTransitionScene*> GetTransitionScenes(FStateTreeExecutionContext& InContext) const;
+
+	/** Whether the scene to check should be this scene or other scene */
+	UPROPERTY(EditAnywhere, Category="Parameter")
+	EAvaTransitionSceneType SceneType = EAvaTransitionSceneType::This;
+
+	/** Which Layer should be queried for the Scene Attributes */
+	UPROPERTY(EditAnywhere, Category="Parameter", meta=(EditCondition="SceneType == EAvaTransitionSceneType::Other", EditConditionHides))
+	EAvaTransitionLayerCompareType LayerType = EAvaTransitionLayerCompareType::Same;
+
+	/** Specific layer tags to check */
+	UPROPERTY(EditAnywhere, Category="Parameter", meta=(EditCondition="SceneType == EAvaTransitionSceneType::Other && LayerType==EAvaTransitionLayerCompareType::MatchingTag", EditConditionHides))
+	FAvaTagHandleContainer SpecificLayers;
+
+	/** The Tag Attribute to check if it's contained in the scene(s) or not */
+	UPROPERTY(EditAnywhere, Category="Parameter")
+	FAvaTagHandle TagAttribute;
+
+	/** Whether to invert the condition (i.e. if true, would become "does NOT contain attribute")*/
+	UPROPERTY(EditAnywhere, Category="Parameter")
+	bool bInvert = false;
+
+	TStateTreeExternalDataHandle<UAvaSceneSubsystem> SceneSubsystemHandle;
+};
diff --git a/Engine/Plugins/Experimental/Avalanche/Source/Avalanche/Public/Viewport/AvaCameraManager.h b/Engine/Plugins/Experimental/Avalanche/Source/Avalanche/Public/Viewport/AvaCameraManager.h
new file mode 100644
--- /dev/null
+++ b/Engine/Plugins/Experimental/Avalanche/Source/Avalanche/Public/Viewport/AvaCameraManager.h
@@ -0,0 +1,67 @@
+// Copyright Epic Games, Inc. All Rights Reserved.
+
+#pragma once
+
+#include "Math/Vector.h"
+#include "Templates/SharedPointer.h"
+#include "UObject/WeakObjectPtr.h"
+#include "UObject/WeakObjectPtrTemplates.h"
+
+class AActor;
+class FAvaCameraManager;
+class FName;
+class FSceneView;
+class IAvaSequencePlaybackObject;
+class UAvaCanvasComponent;
+class UCameraComponent;
+class UWorld;
+struct FMinimalViewInfo;
+struct FPostProcessSettings;
+
+DECLARE_MULTICAST_DELEGATE_TwoParams(FAvaOnViewTargetChanged, FAvaCameraManager*, AActor*);
+
+class AVALANCHE_API FAvaCameraManager : public TSharedFromThis<FAvaCameraManager>
+{
+public:
+	FAvaCameraManager();
+	virtual ~FAvaCameraManager() = default;
+
+	void Init(IAvaSequencePlaybackObject* InPlaybackObject, bool bInCanvasController);
+
+	bool HasViewTarget() const;
+	AActor* GetViewTarget() const;
+	bool SetViewTarget(AActor* InNewViewTarget);
+
+	UCameraComponent* GetCachedCameraComponent() const;
+
+#if WITH_EDITOR
+	bool SetDefaultViewTarget(UWorld* InWorld, FName InDefaultViewTargetName);
+#endif
+
+	void SetCameraCut() { bCameraCut = true; }
+	bool ShouldEnableCameraCut() const { return bCameraCut; }
+	void ResetCameraCut() { bCameraCut = false; }
+
+	FVector GetViewLocation() const;
+	FRotator GetViewRotation() const;
+	void GetViewPoint(FVector& OutLocation, FRotator& OutRotation) const;
+	void SetViewPoint(const FVector& InLocation, const FRotator& InRotation);
+
+	FAvaOnViewTargetChanged& GetOnViewTargetChanged() { return OnViewTargetChanged; }
+
+	void SetCameraViewInfo(FMinimalViewInfo& InOutViewInfo) const;
+	void GetExtraPostProcessBlends(TArray<FPostProcessSettings>& OutSettings, TArray<float>& OutWeights) const;
+
+	/** Gets the extra blends from the camera, if we have any. */
+	void ApplyExtraPostProcessBlends(FSceneView* InView) const;
+
+	virtual TArray<AActor*> GetAvailableCameras(UWorld* InWorld);
+
+protected:
+	TWeakObjectPtr<AActor> ViewTargetWeak;
+	TWeakObjectPtr<UCameraComponent> CameraComponentWeak;
+	bool bCameraCut;
+	FAvaOnViewTargetChanged OnViewTargetChanged;
+
+	void OnUpdateCameraCut(UObject* InCameraObject, bool bInJump);
+};
diff --git a/Engine/Plugins/Experimental/Avalanche/Source/Avalanche/Public/Viewport/AvaCineCameraActor.h b/Engine/Plugins/Experimental/Avalanche/Source/Avalanche/Public/Viewport/AvaCineCameraActor.h
new file mode 100644
--- /dev/null
+++ b/Engine/Plugins/Experimental/Avalanche/Source/Avalanche/Public/Viewport/AvaCineCameraActor.h
@@ -0,0 +1,82 @@
+﻿// Copyright Epic Games, Inc. All Rights Reserved.
+
+#pragma once
+
+#include "CineCameraActor.h"
+#include "CineCameraComponent.h"
+#include "AvaCineCameraActor.generated.h"
+
+/** 
+ * Motion Design Cine Camera Actor is derived from Cine Camera Actor.
+ * Its function is to provide a Cine Camera which can be used right away inside Motion Design.
+ * This is done by customizing some of its default values.
+ * In particular, Motion Design Editor configuration property "Camera Distance" is used to initialize camera position and manual focus.
+ * See: Editor Preferences --> Motion Design --> Editor Settings --> Camera Distance
+ */
+UCLASS(DisplayName = "Motion Design Cine Camera Actor")
+class AVALANCHE_API AAvaCineCameraActor : public ACineCameraActor
+{
+	friend class UAvalancheEditorSettings;
+	
+	GENERATED_BODY()
+
+public:
+	AAvaCineCameraActor(const FObjectInitializer& ObjectInitializer) : Super(ObjectInitializer)
+	{
+		/* todo: when the reset to defaults issue will be solved, Configure(GetDefaultCameraDistance()) could be called here to to initialize defaults */
+	}
+
+	/**
+	 * Initialize the camera with Avalanche scene default values: field of view, camera position, focus distance.
+	 * @param InCameraDistance: camera distance value, used to initialize camera position and manual focus distance
+	 */
+	void Configure(float InCameraDistance)
+	{
+		// initialize default camera distance - todo: we might use this value later for a custom reset?
+		DefaultCameraDistance = InCameraDistance;
+		
+		if (UCineCameraComponent* const CineCameraComp = GetCineCameraComponent())
+		{
+			CineCameraComp->SetFieldOfView(90.0f);
+
+			// smaller camera mesh, todo: this will be substituted with a custom visualizer, or other solution, to visualize frustum + direction
+			CineCameraComp->SetWorldScale3D(FVector(0.5f));
+
+			// using default ManualFocusDistance value from Default CineCameraComponent, since that is handled by UAvalancheEditorSettings
+			CineCameraComp->FocusSettings.ManualFocusDistance = DefaultCameraDistance;
+
+			FVector CameraPosition = FVector::ZeroVector;
+			CameraPosition.X = -DefaultCameraDistance;
+			SetActorLocation(CameraPosition);
+		}
+	}
+
+private:
+	/**
+	 * Updates some defaults of the Mutable CDO's CineCameraComponent, based on the provided CameraDistance value
+	 */
+	static void SetDefaultCameraDistance(float InCameraDistance)
+	{
+		AAvaCineCameraActor* const DefaultObject = GetMutableDefault<AAvaCineCameraActor>();
+
+		FVector CameraPosition = FVector::ZeroVector;
+		CameraPosition.X = -InCameraDistance;
+		DefaultObject->SetActorLocation(CameraPosition);
+		
+		if (UCineCameraComponent* const DefaultCineCameraComponent = DefaultObject->GetCineCameraComponent())
+		{
+			DefaultCineCameraComponent->FocusSettings.ManualFocusDistance = InCameraDistance;
+		}
+
+		DefaultObject->DefaultCameraDistance = InCameraDistance;
+	}
+
+	static float GetDefaultCameraDistance()
+	{
+		return GetMutableDefault<AAvaCineCameraActor>()->DefaultCameraDistance;
+	}
+
+	/** The CameraDistance value used when configuring this AvaCineCamera*/
+	UPROPERTY()
+	float DefaultCameraDistance;
+};
diff --git a/Engine/Plugins/Experimental/Avalanche/Source/Avalanche/Public/Viewport/AvaPostProcessVolume.h b/Engine/Plugins/Experimental/Avalanche/Source/Avalanche/Public/Viewport/AvaPostProcessVolume.h
new file mode 100644
--- /dev/null
+++ b/Engine/Plugins/Experimental/Avalanche/Source/Avalanche/Public/Viewport/AvaPostProcessVolume.h
@@ -0,0 +1,45 @@
+﻿// Copyright Epic Games, Inc. All Rights Reserved.
+
+#pragma once
+
+#include "Engine/PostProcessVolume.h"
+#include "AvaPostProcessVolume.generated.h"
+
+/** 
+ * Motion Design Post Process Volume is derived from Post Process Volume.
+ * Its function is to provide a Post Process which can be used right away inside Motion Design.
+ * This is done by customizing some of the default values.
+ */
+UCLASS(MinimalAPI, DisplayName = "Motion Design Post Process Volume")
+class AAvaPostProcessVolume : public APostProcessVolume
+{
+	GENERATED_BODY()
+
+public:
+	AAvaPostProcessVolume()
+	{
+		BlendWeight = 1.f;
+		bUnbound = true;
+
+		Settings.bOverride_AutoExposureBias = 1;
+		Settings.AutoExposureBias = 0.f;
+
+		Settings.bOverride_AutoExposureMinBrightness = 1;
+		Settings.AutoExposureMinBrightness = 1.f;
+
+		Settings.bOverride_AutoExposureMaxBrightness = 1;
+		Settings.AutoExposureMaxBrightness = 1.f;
+
+		Settings.bOverride_ExpandGamut = 1;
+		Settings.ExpandGamut = 0.f;
+
+		Settings.bOverride_ToneCurveAmount = 1;
+		Settings.ToneCurveAmount = 0.f;
+
+		Settings.bOverride_BlueCorrection = 1;
+		Settings.BlueCorrection = 0.f;
+
+		Settings.bOverride_VignetteIntensity = 1;
+		Settings.VignetteIntensity = 0.f;
+	}
+};
diff --git a/Engine/Plugins/Experimental/Avalanche/Source/Avalanche/Public/Viewport/Interaction/AvaViewportGuide.h b/Engine/Plugins/Experimental/Avalanche/Source/Avalanche/Public/Viewport/Interaction/AvaViewportGuide.h
new file mode 100644
--- /dev/null
+++ b/Engine/Plugins/Experimental/Avalanche/Source/Avalanche/Public/Viewport/Interaction/AvaViewportGuide.h
@@ -0,0 +1,32 @@
+// Copyright Epic Games, Inc. All Rights Reserved.
+
+#pragma once
+
+#include "Containers/EnumAsByte.h"
+#include "Types/SlateEnums.h"
+#include "AvaViewportGuide.generated.h"
+
+UENUM()
+enum class EAvaViewportGuideState_Deprecated : uint8
+{
+	Disabled,
+	Enabled,
+	SnappedTo
+};
+
+USTRUCT()
+struct FAvaViewportGuideInfo_Deprecated
+{
+	GENERATED_BODY()
+
+	UPROPERTY()
+	TEnumAsByte<EOrientation> Orientation = EOrientation::Orient_Horizontal;
+
+	UPROPERTY()
+	float OffsetFraction = 0.0f;
+
+	UPROPERTY()
+	EAvaViewportGuideState_Deprecated State = EAvaViewportGuideState_Deprecated::Disabled;
+
+	bool IsEnabled() const { return State != EAvaViewportGuideState_Deprecated::Disabled; }
+};
diff --git a/Engine/Plugins/Experimental/Avalanche/Source/Avalanche/Public/Viewport/Interaction/IAvaViewportDataProvider.h b/Engine/Plugins/Experimental/Avalanche/Source/Avalanche/Public/Viewport/Interaction/IAvaViewportDataProvider.h
new file mode 100644
--- /dev/null
+++ b/Engine/Plugins/Experimental/Avalanche/Source/Avalanche/Public/Viewport/Interaction/IAvaViewportDataProvider.h
@@ -0,0 +1,39 @@
+// Copyright Epic Games, Inc. All Rights Reserved.
+
+#pragma once
+
+#include "UObject/Interface.h"
+#if WITH_EDITOR
+#include "Containers/ArrayView.h"
+#endif
+#include "IAvaViewportDataProvider.generated.h"
+
+class UObject;
+
+#if WITH_EDITOR
+struct FAvaViewportGuideInfo_Deprecated;
+#endif
+
+/** Interface for Objects that use UAvaSequence and need to be handled by IAvaSequencer */
+UINTERFACE(MinimalAPI, NotBlueprintType, meta = (CannotImplementInterfaceInBlueprint))
+class UAvaViewportDataProvider : public UInterface
+{
+	GENERATED_BODY()
+};
+
+class IAvaViewportDataProvider
+{
+public:
+	GENERATED_BODY()
+
+	virtual UObject* ToUObject() = 0;
+
+	virtual FName GetStartupCameraName() const = 0;
+
+#if WITH_EDITOR
+	virtual void SetStartupCameraName(FName InName) = 0;
+
+	virtual TConstArrayView<FAvaViewportGuideInfo_Deprecated> GetViewportGuideData() const = 0;
+	virtual void SetViewportGuideData(const TConstArrayView<FAvaViewportGuideInfo_Deprecated>& InGuideData) = 0;
+#endif
+};
diff --git a/Engine/Plugins/Experimental/Avalanche/Source/Avalanche/Public/Viewport/Interaction/IAvaViewportDataProxy.h b/Engine/Plugins/Experimental/Avalanche/Source/Avalanche/Public/Viewport/Interaction/IAvaViewportDataProxy.h
new file mode 100644
--- /dev/null
+++ b/Engine/Plugins/Experimental/Avalanche/Source/Avalanche/Public/Viewport/Interaction/IAvaViewportDataProxy.h
@@ -0,0 +1,20 @@
+// Copyright Epic Games, Inc. All Rights Reserved.
+
+#pragma once
+
+#include "AvaType.h"
+
+class IAvaViewportDataProvider;
+
+/**
+ *
+ */
+class IAvaViewportDataProxy
+{
+public:
+	UE_AVA_TYPE(IAvaViewportDataProxy);
+
+#if WITH_EDITOR
+	virtual IAvaViewportDataProvider* GetViewportDataProvider() const = 0;
+#endif
+};
diff --git a/Engine/Plugins/Experimental/Avalanche/Source/AvalancheEditor/Private/Alignment/SAvaAnchorAlignment.cpp b/Engine/Plugins/Experimental/Avalanche/Source/AvalancheEditor/Private/Alignment/SAvaAnchorAlignment.cpp
new file mode 100644
--- /dev/null
+++ b/Engine/Plugins/Experimental/Avalanche/Source/AvalancheEditor/Private/Alignment/SAvaAnchorAlignment.cpp
@@ -0,0 +1,151 @@
+// Copyright Epic Games, Inc. All Rights Reserved.
+
+#include "SAvaAnchorAlignment.h"
+#include "SAvaDepthAlignment.h"
+#include "SAvaHorizontalAlignment.h"
+#include "SAvaVerticalAlignment.h"
+#include "SlateOptMacros.h"
+#include "Widgets/SBoxPanel.h"
+
+#define LOCTEXT_NAMESPACE "SAvaAnchorAlignment"
+
+BEGIN_SLATE_FUNCTION_BUILD_OPTIMIZATION
+void SAvaAnchorAlignment::Construct(const FArguments& InArgs)
+{
+	Anchors = InArgs._Anchors;
+
+	OnAnchorChanged = InArgs._OnAnchorChanged;
+
+	ChildSlot
+	.HAlign(HAlign_Left)
+	.VAlign(VAlign_Center)
+	[
+		SNew(SVerticalBox)
+		+ SVerticalBox::Slot()
+		.AutoHeight()
+		.Padding(0.0f, 2.0f, 0.0f, 1.0f)
+		[
+			SNew(SAvaHorizontalAlignment)
+			.UniformPadding(InArgs._UniformPadding)
+			.Visibility(this, &SAvaAnchorAlignment::GetHorizontalVisibility)
+			.Alignment(this, &SAvaAnchorAlignment::GetHorizontalAlignment)
+			.OnAlignmentChanged(this, &SAvaAnchorAlignment::OnHorizontalAlignmentChanged)
+		]
+		+ SVerticalBox::Slot()
+		.AutoHeight()
+		.Padding(0.0f, 0.0f, 0.0f, 1.0f)
+		[
+			SNew(SAvaVerticalAlignment)
+			.UniformPadding(InArgs._UniformPadding)
+			.Visibility(this, &SAvaAnchorAlignment::GetVerticalVisibility)
+			.Alignment(this, &SAvaAnchorAlignment::GetVerticalAlignment)
+			.OnAlignmentChanged(this, &SAvaAnchorAlignment::OnVerticalAlignmentChanged)
+		]
+		+ SVerticalBox::Slot()
+		.AutoHeight()
+		.Padding(0.0f, 0.0f, 0.0f, 2.0f)
+		[
+			SNew(SAvaDepthAlignment)
+			.UniformPadding(InArgs._UniformPadding)
+			.Visibility(this, &SAvaAnchorAlignment::GetDepthVisibility)
+			.Alignment(this, &SAvaAnchorAlignment::GetDepthAlignment)
+			.OnAlignmentChanged(this, &SAvaAnchorAlignment::OnDepthAlignmentChanged)
+		]
+	];
+}
+END_SLATE_FUNCTION_BUILD_OPTIMIZATION
+
+EVisibility SAvaAnchorAlignment::GetHorizontalVisibility() const
+{
+	if (!Anchors.IsSet())
+	{
+		return EVisibility::SelfHitTestInvisible;
+	}
+	
+	return Anchors.Get().bUseHorizontal ? EVisibility::SelfHitTestInvisible : EVisibility::Collapsed; 
+}
+
+EVisibility SAvaAnchorAlignment::GetVerticalVisibility() const
+{
+	if (!Anchors.IsSet())
+	{
+		return EVisibility::SelfHitTestInvisible;
+	}
+	
+	return Anchors.Get().bUseVertical ? EVisibility::SelfHitTestInvisible : EVisibility::Collapsed; 
+}
+
+EVisibility SAvaAnchorAlignment::GetDepthVisibility() const
+{
+	if (!Anchors.IsSet())
+	{
+		return EVisibility::SelfHitTestInvisible;
+	}
+	
+	return Anchors.Get().bUseDepth ? EVisibility::SelfHitTestInvisible : EVisibility::Collapsed; 
+}
+
+EAvaHorizontalAlignment SAvaAnchorAlignment::GetHorizontalAlignment() const
+{
+	return Anchors.IsSet() ? Anchors.Get().Horizontal : EAvaHorizontalAlignment::Center;
+}
+
+EAvaVerticalAlignment SAvaAnchorAlignment::GetVerticalAlignment() const
+{
+	return Anchors.IsSet() ? Anchors.Get().Vertical : EAvaVerticalAlignment::Center;
+}
+
+EAvaDepthAlignment SAvaAnchorAlignment::GetDepthAlignment() const
+{
+	return Anchors.IsSet() ? Anchors.Get().Depth : EAvaDepthAlignment::Center;
+}
+
+void SAvaAnchorAlignment::OnHorizontalAlignmentChanged(const EAvaHorizontalAlignment InAlignment)
+{
+	if (!Anchors.IsSet())
+	{
+		return;
+	}
+
+	FAvaAnchorAlignment AnchorAlignment = Anchors.Get();
+	AnchorAlignment.Horizontal = InAlignment;
+
+	if (OnAnchorChanged.IsBound())
+	{
+		OnAnchorChanged.Execute(AnchorAlignment);
+	}
+}
+
+void SAvaAnchorAlignment::OnVerticalAlignmentChanged(const EAvaVerticalAlignment InAlignment)
+{
+	if (!Anchors.IsSet())
+	{
+		return;
+	}
+
+	FAvaAnchorAlignment AnchorAlignment = Anchors.Get();
+	AnchorAlignment.Vertical = InAlignment;
+	
+	if (OnAnchorChanged.IsBound())
+	{
+		OnAnchorChanged.Execute(AnchorAlignment);
+	}
+}
+
+void SAvaAnchorAlignment::OnDepthAlignmentChanged(const EAvaDepthAlignment InAlignment)
+{
+	if (!Anchors.IsSet())
+	{
+		return;
+	}
+
+	FAvaAnchorAlignment AnchorAlignment = Anchors.Get();
+	AnchorAlignment.Depth = InAlignment;
+	
+	if (OnAnchorChanged.IsBound())
+	{
+		OnAnchorChanged.Execute(AnchorAlignment);
+	}
+}
+
+#undef LOCTEXT_NAMESPACE
diff --git a/Engine/Plugins/Experimental/Avalanche/Source/AvalancheEditor/Private/Alignment/SAvaAnchorAlignment.h b/Engine/Plugins/Experimental/Avalanche/Source/AvalancheEditor/Private/Alignment/SAvaAnchorAlignment.h
new file mode 100644
--- /dev/null
+++ b/Engine/Plugins/Experimental/Avalanche/Source/AvalancheEditor/Private/Alignment/SAvaAnchorAlignment.h
@@ -0,0 +1,47 @@
+// Copyright Epic Games, Inc. All Rights Reserved.
+
+#pragma once
+
+#include "AvaDefs.h"
+#include "Delegates/Delegate.h"
+#include "Widgets/SCompoundWidget.h"
+
+class SUniformGridPanel;
+
+/**
+ * Avalanche Anchor Alignment
+ * 
+ * Widget that holds holds axis alignment widgets for Left/Center/Right, Top/Center/Bottom, and Front/Center/Back.
+ * One alignment button set per axis row (Horizontal, Vertical, Depth).
+ */
+class SAvaAnchorAlignment : public SCompoundWidget
+{
+public:
+	DECLARE_DELEGATE_OneParam(FOnAnchorChanged, FAvaAnchorAlignment)
+
+	SLATE_BEGIN_ARGS(SAvaAnchorAlignment)
+	{}
+		SLATE_ATTRIBUTE(FAvaAnchorAlignment, Anchors)
+		SLATE_ATTRIBUTE(FMargin, UniformPadding)
+		SLATE_EVENT(FOnAnchorChanged, OnAnchorChanged)
+	SLATE_END_ARGS()
+
+	void Construct(const FArguments& InArgs);
+
+	EVisibility GetHorizontalVisibility() const;
+	EVisibility GetVerticalVisibility() const;
+	EVisibility GetDepthVisibility() const;
+
+	EAvaHorizontalAlignment GetHorizontalAlignment() const;
+	EAvaVerticalAlignment GetVerticalAlignment() const;
+	EAvaDepthAlignment GetDepthAlignment() const;
+
+	void OnHorizontalAlignmentChanged(const EAvaHorizontalAlignment InAlignment);
+	void OnVerticalAlignmentChanged(const EAvaVerticalAlignment InAlignment);
+	void OnDepthAlignmentChanged(const EAvaDepthAlignment InAlignment);
+
+protected:
+	TAttribute<FAvaAnchorAlignment> Anchors;
+
+	FOnAnchorChanged OnAnchorChanged;
+};
diff --git a/Engine/Plugins/Experimental/Avalanche/Source/AvalancheEditor/Private/Alignment/SAvaDepthAlignment.cpp b/Engine/Plugins/Experimental/Avalanche/Source/AvalancheEditor/Private/Alignment/SAvaDepthAlignment.cpp
new file mode 100644
--- /dev/null
+++ b/Engine/Plugins/Experimental/Avalanche/Source/AvalancheEditor/Private/Alignment/SAvaDepthAlignment.cpp
@@ -0,0 +1,50 @@
+// Copyright Epic Games, Inc. All Rights Reserved.
+
+#include "SAvaDepthAlignment.h"
+#include "AvaDefs.h"
+#include "AvaEditorStyle.h"
+#include "SlateOptMacros.h"
+#include "Widgets/Input/SSegmentedControl.h"
+
+#define LOCTEXT_NAMESPACE "SAvaDepthAlignment"
+
+BEGIN_SLATE_FUNCTION_BUILD_OPTIMIZATION
+
+void SAvaDepthAlignment::Construct(const FArguments& InArgs)
+{
+	Alignment = InArgs._Alignment;
+	OnAlignmentChanged = InArgs._OnAlignmentChanged;
+
+	ChildSlot
+	.HAlign(HAlign_Left)
+	.VAlign(VAlign_Center)
+	[
+		SNew(SSegmentedControl<EAvaDepthAlignment>)
+		.UniformPadding(InArgs._UniformPadding)
+		.Value(this, &SAvaDepthAlignment::GetCurrentAlignment)
+		.OnValueChanged(this, &SAvaDepthAlignment::OnCurrentAlignmentChanged)
+		+ SSegmentedControl<EAvaDepthAlignment>::Slot(EAvaDepthAlignment::Front)
+			.Icon(FAvaEditorStyle::Get().GetBrush("AvalancheIcons.Alignment.Translation.Front"))
+			.ToolTip(LOCTEXT("DAlignFront", "Front Align Depth"))
+		+ SSegmentedControl<EAvaDepthAlignment>::Slot(EAvaDepthAlignment::Center)
+			.Icon(FAvaEditorStyle::Get().GetBrush("AvalancheIcons.Alignment.Translation.Center_X"))
+			.ToolTip(LOCTEXT("DAlignCenter", "Center Align Depth"))
+		+ SSegmentedControl<EAvaDepthAlignment>::Slot(EAvaDepthAlignment::Back)
+			.Icon(FAvaEditorStyle::Get().GetBrush("AvalancheIcons.Alignment.Translation.Back"))
+			.ToolTip(LOCTEXT("DAlignBack", "Back Align Depth"))
+	];
+}
+
+END_SLATE_FUNCTION_BUILD_OPTIMIZATION
+
+EAvaDepthAlignment SAvaDepthAlignment::GetCurrentAlignment() const
+{
+	return Alignment.Get(EAvaDepthAlignment::Center);
+}
+
+void SAvaDepthAlignment::OnCurrentAlignmentChanged(const EAvaDepthAlignment NewAlignment)
+{
+	OnAlignmentChanged.ExecuteIfBound(NewAlignment);
+}
+
+#undef LOCTEXT_NAMESPACE
diff --git a/Engine/Plugins/Experimental/Avalanche/Source/AvalancheEditor/Private/Alignment/SAvaDepthAlignment.h b/Engine/Plugins/Experimental/Avalanche/Source/AvalancheEditor/Private/Alignment/SAvaDepthAlignment.h
new file mode 100644
--- /dev/null
+++ b/Engine/Plugins/Experimental/Avalanche/Source/AvalancheEditor/Private/Alignment/SAvaDepthAlignment.h
@@ -0,0 +1,37 @@
+// Copyright Epic Games, Inc. All Rights Reserved.
+
+#pragma once
+
+#include "Delegates/Delegate.h"
+#include "Widgets/SCompoundWidget.h"
+
+enum class EAvaDepthAlignment : uint8;
+
+/**
+ * Avalanche Depth Axis Alignment
+ * 
+ * Widget that holds three buttons for either Front, Center, and Back.
+ */
+class SAvaDepthAlignment : public SCompoundWidget
+{
+public:
+	DECLARE_DELEGATE_OneParam(FOnDepthAlignmentChanged, EAvaDepthAlignment)
+
+	SLATE_BEGIN_ARGS(SAvaDepthAlignment)
+	{}
+		SLATE_ATTRIBUTE(EAvaDepthAlignment, Alignment)
+		SLATE_ATTRIBUTE(FMargin, UniformPadding)
+		SLATE_EVENT(FOnDepthAlignmentChanged, OnAlignmentChanged)
+	SLATE_END_ARGS()
+
+	void Construct(const FArguments& InArgs);
+
+protected:
+	TAttribute<EAvaDepthAlignment> Alignment;
+
+	FOnDepthAlignmentChanged OnAlignmentChanged;
+
+	EAvaDepthAlignment GetCurrentAlignment() const;
+
+	void OnCurrentAlignmentChanged(const EAvaDepthAlignment NewAlignment);
+};
diff --git a/Engine/Plugins/Experimental/Avalanche/Source/AvalancheEditor/Private/Alignment/SAvaHorizontalAlignment.cpp b/Engine/Plugins/Experimental/Avalanche/Source/AvalancheEditor/Private/Alignment/SAvaHorizontalAlignment.cpp
new file mode 100644
--- /dev/null
+++ b/Engine/Plugins/Experimental/Avalanche/Source/AvalancheEditor/Private/Alignment/SAvaHorizontalAlignment.cpp
@@ -0,0 +1,50 @@
+// Copyright Epic Games, Inc. All Rights Reserved.
+
+#include "SAvaHorizontalAlignment.h"
+#include "AvaDefs.h"
+#include "AvaEditorStyle.h"
+#include "SlateOptMacros.h"
+#include "Widgets/Input/SSegmentedControl.h"
+
+#define LOCTEXT_NAMESPACE "SAvaHorizontalAlignment"
+
+BEGIN_SLATE_FUNCTION_BUILD_OPTIMIZATION
+
+void SAvaHorizontalAlignment::Construct(const FArguments& InArgs)
+{
+	Alignment = InArgs._Alignment;
+	OnAlignmentChanged = InArgs._OnAlignmentChanged;
+
+	ChildSlot
+	.HAlign(HAlign_Left)
+	.VAlign(VAlign_Center)
+	[
+		SNew(SSegmentedControl<EAvaHorizontalAlignment>)
+		.UniformPadding(InArgs._UniformPadding)
+		.Value(this, &SAvaHorizontalAlignment::GetCurrentAlignment)
+		.OnValueChanged(this, &SAvaHorizontalAlignment::OnCurrentAlignmentChanged)
+		+ SSegmentedControl<EAvaHorizontalAlignment>::Slot(EAvaHorizontalAlignment::Left)
+			.Icon(FAvaEditorStyle::Get().GetBrush("AvalancheIcons.Alignment.Left"))
+			.ToolTip(LOCTEXT("HAlignLeft", "Left Align Horizontally"))
+		+ SSegmentedControl<EAvaHorizontalAlignment>::Slot(EAvaHorizontalAlignment::Center)
+			.Icon(FAvaEditorStyle::Get().GetBrush("AvalancheIcons.Alignment.Center_Y"))
+			.ToolTip(LOCTEXT("HAlignCenter", "Center Align Horizontally"))
+		+ SSegmentedControl<EAvaHorizontalAlignment>::Slot(EAvaHorizontalAlignment::Right)
+			.Icon(FAvaEditorStyle::Get().GetBrush("AvalancheIcons.Alignment.Right"))
+			.ToolTip(LOCTEXT("HAlignRight", "Right Align Horizontally"))
+	];
+}
+
+END_SLATE_FUNCTION_BUILD_OPTIMIZATION
+
+EAvaHorizontalAlignment SAvaHorizontalAlignment::GetCurrentAlignment() const
+{
+	return Alignment.Get(EAvaHorizontalAlignment::Center);
+}
+
+void SAvaHorizontalAlignment::OnCurrentAlignmentChanged(const EAvaHorizontalAlignment NewAlignment)
+{
+	OnAlignmentChanged.ExecuteIfBound(NewAlignment);
+}
+
+#undef LOCTEXT_NAMESPACE
diff --git a/Engine/Plugins/Experimental/Avalanche/Source/AvalancheEditor/Private/Alignment/SAvaHorizontalAlignment.h b/Engine/Plugins/Experimental/Avalanche/Source/AvalancheEditor/Private/Alignment/SAvaHorizontalAlignment.h
new file mode 100644
--- /dev/null
+++ b/Engine/Plugins/Experimental/Avalanche/Source/AvalancheEditor/Private/Alignment/SAvaHorizontalAlignment.h
@@ -0,0 +1,37 @@
+// Copyright Epic Games, Inc. All Rights Reserved.
+
+#pragma once
+
+#include "Delegates/Delegate.h"
+#include "Widgets/SCompoundWidget.h"
+
+enum class EAvaHorizontalAlignment : uint8;
+
+/**
+ * Avalanche Horizontal Axis Alignment
+ * 
+ * Widget that holds three buttons for either Left, Center, and Right.
+ */
+class SAvaHorizontalAlignment : public SCompoundWidget
+{
+public:
+	DECLARE_DELEGATE_OneParam(FOnHorizontalAlignmentChanged, EAvaHorizontalAlignment)
+
+	SLATE_BEGIN_ARGS(SAvaHorizontalAlignment)
+	{}
+		SLATE_ATTRIBUTE(EAvaHorizontalAlignment, Alignment)
+		SLATE_ATTRIBUTE(FMargin, UniformPadding)
+		SLATE_EVENT(FOnHorizontalAlignmentChanged, OnAlignmentChanged)
+	SLATE_END_ARGS()
+
+	void Construct(const FArguments& InArgs);
+
+protected:
+	TAttribute<EAvaHorizontalAlignment> Alignment;
+
+	FOnHorizontalAlignmentChanged OnAlignmentChanged;
+
+	EAvaHorizontalAlignment GetCurrentAlignment() const;
+
+	void OnCurrentAlignmentChanged(const EAvaHorizontalAlignment NewAlignment);
+};
diff --git a/Engine/Plugins/Experimental/Avalanche/Source/AvalancheEditor/Private/Alignment/SAvaVerticalAlignment.cpp b/Engine/Plugins/Experimental/Avalanche/Source/AvalancheEditor/Private/Alignment/SAvaVerticalAlignment.cpp
new file mode 100644
--- /dev/null
+++ b/Engine/Plugins/Experimental/Avalanche/Source/AvalancheEditor/Private/Alignment/SAvaVerticalAlignment.cpp
@@ -0,0 +1,50 @@
+// Copyright Epic Games, Inc. All Rights Reserved.
+
+#include "SAvaVerticalAlignment.h"
+#include "AvaDefs.h"
+#include "AvaEditorStyle.h"
+#include "SlateOptMacros.h"
+#include "Widgets/Input/SSegmentedControl.h"
+
+#define LOCTEXT_NAMESPACE "SAvaVerticalAlignment"
+
+BEGIN_SLATE_FUNCTION_BUILD_OPTIMIZATION
+
+void SAvaVerticalAlignment::Construct(const FArguments& InArgs)
+{
+	Alignment = InArgs._Alignment;
+	OnAlignmentChanged = InArgs._OnAlignmentChanged;
+
+	ChildSlot
+	.HAlign(HAlign_Left)
+	.VAlign(VAlign_Center)
+	[
+		SNew(SSegmentedControl<EAvaVerticalAlignment>)
+		.UniformPadding(InArgs._UniformPadding)
+		.Value(this, &SAvaVerticalAlignment::GetCurrentAlignment)
+		.OnValueChanged(this, &SAvaVerticalAlignment::OnCurrentAlignmentChanged)
+		+ SSegmentedControl<EAvaVerticalAlignment>::Slot(EAvaVerticalAlignment::Top)
+			.Icon(FAvaEditorStyle::Get().GetBrush("AvalancheIcons.Alignment.Top"))
+			.ToolTip(LOCTEXT("VAlignTop", "Top Align Vertically"))
+		+ SSegmentedControl<EAvaVerticalAlignment>::Slot(EAvaVerticalAlignment::Center)
+			.Icon(FAvaEditorStyle::Get().GetBrush("AvalancheIcons.Alignment.Center_Z"))
+			.ToolTip(LOCTEXT("VAlignCenter", "Center Align Vertically"))
+		+ SSegmentedControl<EAvaVerticalAlignment>::Slot(EAvaVerticalAlignment::Bottom)
+			.Icon(FAvaEditorStyle::Get().GetBrush("AvalancheIcons.Alignment.Bottom"))
+			.ToolTip(LOCTEXT("VAlignBottom", "Bottom Align Vertically"))
+	];
+}
+
+END_SLATE_FUNCTION_BUILD_OPTIMIZATION
+
+EAvaVerticalAlignment SAvaVerticalAlignment::GetCurrentAlignment() const
+{
+	return Alignment.Get(EAvaVerticalAlignment::Center);
+}
+
+void SAvaVerticalAlignment::OnCurrentAlignmentChanged(const EAvaVerticalAlignment NewAlignment)
+{
+	OnAlignmentChanged.ExecuteIfBound(NewAlignment);
+}
+
+#undef LOCTEXT_NAMESPACE
diff --git a/Engine/Plugins/Experimental/Avalanche/Source/AvalancheEditor/Private/Alignment/SAvaVerticalAlignment.h b/Engine/Plugins/Experimental/Avalanche/Source/AvalancheEditor/Private/Alignment/SAvaVerticalAlignment.h
new file mode 100644
--- /dev/null
+++ b/Engine/Plugins/Experimental/Avalanche/Source/AvalancheEditor/Private/Alignment/SAvaVerticalAlignment.h
@@ -0,0 +1,36 @@
+// Copyright Epic Games, Inc. All Rights Reserved.
+
+#pragma once
+
+#include "Delegates/Delegate.h"
+#include "Widgets/SCompoundWidget.h"
+
+enum class EAvaVerticalAlignment : uint8;
+
+/**
+ * Avalanche Vertical Axis Alignment
+ * 
+ * Widget that holds three buttons for either Top, Center, and Bottom.
+ */
+class SAvaVerticalAlignment : public SCompoundWidget
+{
+public:
+	DECLARE_DELEGATE_OneParam(FOnVerticalAlignmentChanged, EAvaVerticalAlignment)
+
+	SLATE_BEGIN_ARGS(SAvaVerticalAlignment)
+	{}
+		SLATE_ATTRIBUTE(EAvaVerticalAlignment, Alignment)
+		SLATE_ATTRIBUTE(FMargin, UniformPadding)
+		SLATE_EVENT(FOnVerticalAlignmentChanged, OnAlignmentChanged)
+	SLATE_END_ARGS()
+
+	void Construct(const FArguments& InArgs);
+
+protected:
+	TAttribute<EAvaVerticalAlignment> Alignment;
+	FOnVerticalAlignmentChanged OnAlignmentChanged;
+
+	EAvaVerticalAlignment GetCurrentAlignment() const;
+
+	void OnCurrentAlignmentChanged(const EAvaVerticalAlignment NewAlignment);
+};
diff --git a/Engine/Plugins/Experimental/Avalanche/Source/AvalancheEditor/Private/AssetTypeActions/AssetTypeActions_AvaBlueprint.cpp b/Engine/Plugins/Experimental/Avalanche/Source/AvalancheEditor/Private/AssetTypeActions/AssetTypeActions_AvaBlueprint.cpp
new file mode 100644
--- /dev/null
+++ b/Engine/Plugins/Experimental/Avalanche/Source/AvalancheEditor/Private/AssetTypeActions/AssetTypeActions_AvaBlueprint.cpp
@@ -0,0 +1,63 @@
+// Copyright Epic Games, Inc. All Rights Reserved.
+
+#include "AssetTypeActions_AvaBlueprint.h"
+#include "AvaBlueprint.h"
+#include "AvaEditorFunctionLibrary.h"
+#include "Misc/MessageDialog.h"
+#include "SBlueprintDiff.h"
+#include "Styling/SlateIconFinder.h"
+#include "ToolMenu.h"
+#include "ToolMenuSection.h"
+
+#define LOCTEXT_NAMESPACE "AssetTypeActions_AvaBlueprint"
+
+FText FAssetTypeActions_AvaBlueprint::GetName() const
+{
+	return LOCTEXT("AssetTypeActions_AvaBlueprint", "Motion Design Blueprint");
+}
+
+UClass* FAssetTypeActions_AvaBlueprint::GetSupportedClass() const
+{
+	return UAvalancheBlueprint::StaticClass();
+}
+
+void FAssetTypeActions_AvaBlueprint::GetActions(const TArray<UObject*>& InObjects, FToolMenuSection& Section)
+{
+	const TArray<TWeakObjectPtr<UAvalancheBlueprint>> AvalancheBlueprints = GetTypedWeakObjectPtrs<UAvalancheBlueprint>(InObjects);
+
+	Section.AddMenuEntry(TEXT("AvaBlueprint_ExportToWorld")
+		, LOCTEXT("AvaBlueprint_ExportToLevelLabel", "Export to World")
+		, LOCTEXT("AvaBlueprint_ExportToLevelTooltip", "Creates a World Asset for this Motion Design Scene populated with the Motion Design Scene Data")
+		, FSlateIconFinder::FindIconForClass(UWorld::StaticClass())
+		, FUIAction(FExecuteAction::CreateStatic(&FAvaEditorFunctionLibrary::ExportAvaBlueprintsToWorld, AvalancheBlueprints)));
+}
+
+void FAssetTypeActions_AvaBlueprint::OpenAssetEditor(const TArray<UObject*>& InObjects, TSharedPtr<IToolkitHost> EditWithinLevelEditor)
+{
+	TSharedRef<FSimpleAssetEditor> NewEditor = MakeShared<FSimpleAssetEditor>();
+	NewEditor->InitEditor(EToolkitMode::Standalone, nullptr, InObjects, FSimpleAssetEditor::FGetDetailsViewObjects());
+	NewEditor->SetPropertyEditingEnabledDelegate(FIsPropertyEditingEnabled::CreateLambda([]{ return false; }));
+}
+
+void FAssetTypeActions_AvaBlueprint::PerformAssetDiff(UObject* Asset1, UObject* Asset2,
+	const FRevisionInfo& OldRevision, const FRevisionInfo& NewRevision) const
+{
+	UBlueprint* OldBlueprint = CastChecked<UBlueprint>(Asset1);
+	UBlueprint* NewBlueprint = CastChecked<UBlueprint>(Asset2);
+
+	// sometimes we're comparing different revisions of one single asset (other
+	// times we're comparing two completely separate assets altogether)
+	const bool bIsSingleAsset = (NewBlueprint->GetName() == OldBlueprint->GetName());
+
+	FText WindowTitle = LOCTEXT("NamelessAvalancheBlueprintDiff", "Motion Design Blueprint Diff");
+	// if we're diffing one asset against itself
+	if (bIsSingleAsset)
+	{
+		// identify the assumed single asset in the window's title
+		WindowTitle = FText::Format(LOCTEXT("AvalancheBlueprintDiff", "{0} - Motion Design Blueprint Diff"), FText::FromString(NewBlueprint->GetName()));
+	}
+
+	SBlueprintDiff::CreateDiffWindow(WindowTitle, OldBlueprint, NewBlueprint, OldRevision, NewRevision);
+}
+
+#undef LOCTEXT_NAMESPACE
diff --git a/Engine/Plugins/Experimental/Avalanche/Source/AvalancheEditor/Private/AssetTypeActions/AssetTypeActions_AvaBlueprint.h b/Engine/Plugins/Experimental/Avalanche/Source/AvalancheEditor/Private/AssetTypeActions/AssetTypeActions_AvaBlueprint.h
new file mode 100644
--- /dev/null
+++ b/Engine/Plugins/Experimental/Avalanche/Source/AvalancheEditor/Private/AssetTypeActions/AssetTypeActions_AvaBlueprint.h
@@ -0,0 +1,21 @@
+// Copyright Epic Games, Inc. All Rights Reserved.
+
+#pragma once
+
+#include "AssetTypeActions_Base.h"
+
+/** We override Base, and not Blueprint, to not allow creation of Child Classes. */
+class FAssetTypeActions_AvaBlueprint : public FAssetTypeActions_Base
+{
+public:
+	//~ Begin IAssetTypeActions
+	virtual FText GetName() const override;
+	virtual FColor GetTypeColor() const override { return FColor(44, 89, 180); }
+	virtual UClass* GetSupportedClass() const override;
+	virtual bool HasActions(const TArray<UObject*>& InObjects) const override { return true; }
+	virtual void GetActions(const TArray<UObject*>& InObjects, FToolMenuSection& Section) override;
+	virtual void OpenAssetEditor(const TArray<UObject*>& InObjects, TSharedPtr<IToolkitHost> EditWithinLevelEditor = TSharedPtr<IToolkitHost>()) override;
+	virtual uint32 GetCategories() override { return EAssetTypeCategories::Blueprint; }
+	virtual void PerformAssetDiff(UObject* Asset1, UObject* Asset2, const FRevisionInfo& OldRevision, const FRevisionInfo& NewRevision) const override;
+	//~ End IAssetTypeActions
+};
diff --git a/Engine/Plugins/Experimental/Avalanche/Source/AvalancheEditor/Private/AvaEditorFunctionLibrary.cpp b/Engine/Plugins/Experimental/Avalanche/Source/AvalancheEditor/Private/AvaEditorFunctionLibrary.cpp
new file mode 100644
--- /dev/null
+++ b/Engine/Plugins/Experimental/Avalanche/Source/AvalancheEditor/Private/AvaEditorFunctionLibrary.cpp
@@ -0,0 +1,80 @@
+// Copyright Epic Games, Inc. All Rights Reserved.
+
+#include "AvaEditorFunctionLibrary.h"
+#include "AssetToolsModule.h"
+#include "AvaBlueprint.h"
+#include "AvaScene.h"
+#include "Containers/Array.h"
+#include "Containers/UnrealString.h"
+#include "Editor.h"
+#include "Engine/World.h"
+#include "Factories/WorldFactory.h"
+#include "Misc/PackageName.h"
+#include "Misc/ScopedSlowTask.h"
+#include "Modules/ModuleManager.h"
+#include "UObject/Package.h"
+#include "UObject/WeakObjectPtr.h"
+
+#define LOCTEXT_NAMESPACE "AvaEditorFunctionLibrary"
+
+namespace UE::AvaEditor::Private
+{
+	void ExportBlueprint(UFactory* InWorldFactory, UAvalancheBlueprint* InBlueprint)
+	{
+		if (!IsValid(InBlueprint) || !IsValid(InWorldFactory))
+		{
+			return;
+		}
+
+		const UPackage* const BlueprintPackage = InBlueprint->GetPackage();
+
+		const FString Path = FPackageName::GetLongPackagePath(BlueprintPackage->GetName()) + TEXT("/Exported");
+
+		IAssetTools& AssetTools = FModuleManager::GetModuleChecked<FAssetToolsModule>("AssetTools").Get();
+
+		UWorld* const NewWorld = Cast<UWorld>(AssetTools.CreateAsset(InBlueprint->GetName()
+			, Path, UWorld::StaticClass(), InWorldFactory));
+
+		if (!IsValid(NewWorld))
+		{
+			return;
+		}
+
+		AAvaScene::CreateFromBlueprint(NewWorld->PersistentLevel, InBlueprint);
+	}
+}
+
+void FAvaEditorFunctionLibrary::ExportAvaBlueprintsToWorld(TArray<TWeakObjectPtr<UAvalancheBlueprint>> InWeakBlueprints)
+{
+	UFactory* const WorldFactory = NewObject<UWorldFactory>(GetTransientPackage());
+
+	FEditorDelegates::OnConfigureNewAssetProperties.Broadcast(WorldFactory);
+
+	if (!WorldFactory || !WorldFactory->ConfigureProperties())
+	{
+		return;
+	}
+
+	FEditorDelegates::OnNewAssetCreated.Broadcast(WorldFactory);
+
+	FScopedSlowTask SlowTask(InWeakBlueprints.Num(), LOCTEXT("ExportingBlueprintsToWorld", "Exporting Blueprint(s) to World"));
+	SlowTask.MakeDialog();
+
+	for (const TWeakObjectPtr<UAvalancheBlueprint>& WeakBlueprint : InWeakBlueprints)
+	{
+		if (UAvalancheBlueprint* Blueprint = WeakBlueprint.Get())
+		{
+			SlowTask.EnterProgressFrame(1
+				, FText::Format(LOCTEXT("ExportingBlueprintToWorld", "Exporting Blueprint '{0}'")
+				, FText::FromString(Blueprint->GetFriendlyName())));
+
+			UE::AvaEditor::Private::ExportBlueprint(WorldFactory, Blueprint);
+		}
+		else
+		{
+			SlowTask.EnterProgressFrame(1);
+		}
+	}
+}
+
+#undef LOCTEXT_NAMESPACE
diff --git a/Engine/Plugins/Experimental/Avalanche/Source/AvalancheEditor/Private/AvaEditorStyle.cpp b/Engine/Plugins/Experimental/Avalanche/Source/AvalancheEditor/Private/AvaEditorStyle.cpp
new file mode 100644
--- /dev/null
+++ b/Engine/Plugins/Experimental/Avalanche/Source/AvalancheEditor/Private/AvaEditorStyle.cpp
@@ -0,0 +1,301 @@
+// Copyright Epic Games, Inc. All Rights Reserved.
+
+#include "AvaEditorStyle.h"
+#include "AvaClonerEffectorShared.h"
+#include "Brushes/SlateImageBrush.h"
+#include "Engine/Texture2D.h"
+#include "Framework/Application/SlateApplication.h"
+#include "Interfaces/IPluginManager.h"
+#include "MediaPlate.h"
+#include "SlateOptMacros.h"
+#include "Styling/AppStyle.h"
+#include "Styling/SlateIconFinder.h"
+#include "Styling/SlateStyleMacros.h"
+#include "Styling/StyleColors.h"
+#include "Styling/SlateStyle.h"
+#include "Styling/SlateStyleRegistry.h"
+
+TSharedPtr<FAvaEditorStyle> FAvaEditorStyle::StyleInstance = nullptr;
+
+void FAvaEditorStyle::Shutdown()
+{
+	FSlateStyleRegistry::UnRegisterSlateStyle(*StyleInstance);
+	ensure(StyleInstance.IsUnique());
+	StyleInstance.Reset();
+}
+
+const FAvaEditorStyle& FAvaEditorStyle::Get()
+{
+	if (!StyleInstance.IsValid())
+	{
+		StyleInstance = MakeShared<FAvaEditorStyle>();
+		FSlateStyleRegistry::RegisterSlateStyle(*StyleInstance);
+	}
+
+	return *StyleInstance;
+}
+
+FAvaEditorStyle::FAvaEditorStyle()
+	: FSlateStyleSet(TEXT("AvalancheEditor"))
+{
+	Init();
+}
+
+namespace UE::AvalancheEditor::Private
+{
+	const FVector2f Icon16(16.f);
+	const FVector2f Icon20(20.f);
+	const FVector2f Icon22(22.f);
+	const FVector2f Icon25(25.f);
+	const FVector2f Icon32(32.f);
+	const FVector2f Icon40(40.f);
+	const FVector2f Icon64(64.f);
+}
+
+BEGIN_SLATE_FUNCTION_BUILD_OPTIMIZATION
+
+void FAvaEditorStyle::Init()
+{
+	using namespace UE::AvalancheEditor::Private;
+
+	TSharedPtr<IPlugin> Plugin = IPluginManager::Get().FindPlugin(TEXT("Avalanche"));
+	check(Plugin.IsValid());
+
+	if (Plugin.IsValid())
+	{
+		SetContentRoot(Plugin->GetBaseDir() / TEXT("Resources"));
+	}
+
+	SetCoreContentRoot(FPaths::EngineContentDir() / TEXT("Editor/Slate"));
+
+	// Custom Class Icons
+	Set("ClassThumbnail.AvalancheActor", new IMAGE_BRUSH("Icons/Icon64", Icon64));
+	Set("ClassThumbnail.AvalancheBlueprint", new IMAGE_BRUSH("Icons/Icon64", Icon64));
+	Set("ClassIcon.AvalancheActor", new IMAGE_BRUSH("Icons/Icon16", Icon16));
+	Set("ClassIcon.AvalancheBlueprint", new IMAGE_BRUSH("Icons/Icon16", Icon16));
+	Set("ClassIcon.AvaNullActor", new CORE_IMAGE_BRUSH(TEXT("Icons/SequencerIcons/icon_Sequencer_Move_24x"), Icon16));
+	Set("ClassIcon.AvaToolboxStarDynamicMesh", new IMAGE_BRUSH_SVG("Icons/ToolboxIcons/Favorite", Icon16));
+	Set("ClassIcon.AvaToolboxLineDynamicMesh", new IMAGE_BRUSH_SVG("Icons/ToolboxIcons/minus", Icon16));
+	Set("ClassIcon.AvaToolboxRectangleDynamicMesh", new IMAGE_BRUSH("Icons/ToolboxIcons/rectangle", Icon16));
+	Set("ClassIcon.AvaToolboxNGonDynamicMesh", new IMAGE_BRUSH("Icons/ToolboxIcons/regularpolygon", Icon16));
+	Set("ClassIcon.AvaToolboxIrregularPolygonDynamicMesh", new IMAGE_BRUSH("Icons/ToolboxIcons/irregularpolygon", Icon16));
+	Set("ClassIcon.AvaToolboxEllipseDynamicMesh", new IMAGE_BRUSH("Icons/ToolboxIcons/circle", Icon16));
+	Set("ClassIcon.AvaToolboxRingDynamicMesh", new IMAGE_BRUSH("Icons/ToolboxIcons/ring", Icon16));
+	Set("ClassIcon.AvaTextActor", new IMAGE_BRUSH("Icons/ToolboxIcons/3d-text", Icon16));
+	Set("ClassIcon.AvaShapeActor", new IMAGE_BRUSH("Icons/ToolboxIcons/regularpolygon", Icon16));
+	Set("ClassIcon.AvaClonerActor", new IMAGE_BRUSH_SVG("Icons/ToolboxIcons/cloner", Icon16));
+	Set("ClassIcon.AvaEffectorActor", new IMAGE_BRUSH_SVG("Icons/ToolboxIcons/effector", Icon16));
+
+	// Toolbox Icons
+	Set("AvalancheIcons.Chevron",          new IMAGE_BRUSH("Icons/ToolboxIcons/chevron",          Icon16));
+	Set("AvalancheIcons.RegularPolygon",   new IMAGE_BRUSH("Icons/ToolboxIcons/regularpolygon",   Icon16));
+	Set("AvalancheIcons.IrregularPolygon", new IMAGE_BRUSH("Icons/ToolboxIcons/irregularpolygon", Icon16));
+	Set("AvalancheIcons.Canvas",           new IMAGE_BRUSH("Icons/ToolboxIcons/Canvas",           Icon16));
+	Set("AvalancheIcons.LayoutGrid",       new IMAGE_BRUSH("Icons/ToolboxIcons/LayoutGrid",       Icon16));
+	Set("AvalancheIcons.Ellipse",          new IMAGE_BRUSH("Icons/ToolboxIcons/circle",           Icon16));
+	Set("AvalancheIcons.Ring",             new IMAGE_BRUSH("Icons/ToolboxIcons/ring",             Icon16));
+	Set("AvalancheIcons.Arrow",            new IMAGE_BRUSH("Icons/ToolboxIcons/arrow",            Icon16));
+	Set("AvalancheIcons.2DArrow",          new IMAGE_BRUSH("Icons/ToolboxIcons/arrow",            Icon16));
+	Set("AvalancheIcons.Sphere",           new IMAGE_BRUSH("Icons/ToolboxIcons/sphere",           Icon16));
+	Set("AvalancheIcons.Pyramid",          new IMAGE_BRUSH("Icons/ToolboxIcons/pyramid",          Icon16));
+	Set("AvalancheIcons.Cone",             new IMAGE_BRUSH("Icons/ToolboxIcons/cone",             Icon16));
+	Set("AvalancheIcons.Torus",            new IMAGE_BRUSH("Icons/ToolboxIcons/torus",            Icon16));
+	Set("AvalancheIcons.Cube",             new IMAGE_BRUSH("Icons/ToolboxIcons/cube",             Icon16));
+	Set("AvalancheIcons.Rectangle",        new IMAGE_BRUSH("Icons/ToolboxIcons/rectangle",        Icon16));
+	Set("AvalancheIcons.Plane",            new IMAGE_BRUSH("Icons/ToolboxIcons/plane",            Icon16));
+	Set("AvalancheIcons.Cylinder",         new IMAGE_BRUSH("Icons/ToolboxIcons/cylinder",         Icon16));
+	Set("AvalancheIcons.Line",             new IMAGE_BRUSH_SVG("Icons/ToolboxIcons/minus",             Icon16));
+	Set("AvalancheIcons.Star",             new IMAGE_BRUSH_SVG("Icons/ToolboxIcons/Favorite",          Icon16));
+	Set("AvalancheIcons.3DText",           new IMAGE_BRUSH_SVG("Icons/ToolboxIcons/Toolbar_Text",      Icon16));
+	Set("AvalancheIcons.RoundedBox",       new IMAGE_BRUSH_SVG("Icons/ToolboxIcons/Toolbar_Primitive", Icon16));
+	
+	Set("AvalancheIcons.Alignment.Translation.TopLeft",     new IMAGE_BRUSH("Icons/DetailsPanelIcons/TopLeft",     Icon16));
+	Set("AvalancheIcons.Alignment.Translation.Top",         new IMAGE_BRUSH("Icons/DetailsPanelIcons/Top",         Icon16));
+	Set("AvalancheIcons.Alignment.Translation.TopRight",    new IMAGE_BRUSH("Icons/DetailsPanelIcons/TopRight",    Icon16));
+	Set("AvalancheIcons.Alignment.Translation.Left",        new IMAGE_BRUSH("Icons/DetailsPanelIcons/Left",        Icon16));
+	Set("AvalancheIcons.Alignment.Translation.Center",      new IMAGE_BRUSH("Icons/DetailsPanelIcons/Center",      Icon16));
+	Set("AvalancheIcons.Alignment.Translation.Right",       new IMAGE_BRUSH("Icons/DetailsPanelIcons/Right",       Icon16));
+	Set("AvalancheIcons.Alignment.Translation.BottomLeft",  new IMAGE_BRUSH("Icons/DetailsPanelIcons/BottomLeft",  Icon16));
+	Set("AvalancheIcons.Alignment.Translation.Bottom",      new IMAGE_BRUSH("Icons/DetailsPanelIcons/Bottom",      Icon16));
+	Set("AvalancheIcons.Alignment.Translation.BottomRight", new IMAGE_BRUSH("Icons/DetailsPanelIcons/BottomRight", Icon16));
+	
+	Set("AvalancheIcons.Alignment.Translation.Back",        new IMAGE_BRUSH_SVG("Icons/DetailsPanelIcons/StackBack",   Icon22));
+	Set("AvalancheIcons.Alignment.Translation.Center_X",    new IMAGE_BRUSH_SVG("Icons/DetailsPanelIcons/StackCenter", Icon22));
+	Set("AvalancheIcons.Alignment.Translation.Front",       new IMAGE_BRUSH_SVG("Icons/DetailsPanelIcons/StackFront",  Icon22));
+
+	Set("AvalancheIcons.Alignment.Center_YZ",   new IMAGE_BRUSH_SVG("Icons/PaletteIcons/AlignCenterVertAndHoriz", Icon22));
+	Set("AvalancheIcons.Alignment.Left",        new IMAGE_BRUSH_SVG("Icons/PaletteIcons/AlignLeft",               Icon22));
+	Set("AvalancheIcons.Alignment.Center_Y",    new IMAGE_BRUSH_SVG("Icons/PaletteIcons/AlignCenterHoriz",        Icon22));
+	Set("AvalancheIcons.Alignment.Right",       new IMAGE_BRUSH_SVG("Icons/PaletteIcons/AlignRight",              Icon22));
+	Set("AvalancheIcons.Alignment.Top",         new IMAGE_BRUSH_SVG("Icons/PaletteIcons/AlignTop",                Icon22));
+	Set("AvalancheIcons.Alignment.Center_Z",    new IMAGE_BRUSH_SVG("Icons/PaletteIcons/AlignCenterVert",         Icon22));
+	Set("AvalancheIcons.Alignment.Bottom",      new IMAGE_BRUSH_SVG("Icons/PaletteIcons/AlignBottom",             Icon22));
+	Set("AvalancheIcons.Alignment.DistributeX", new IMAGE_BRUSH_SVG("Icons/PaletteIcons/DistributeDepth",         Icon22));
+	Set("AvalancheIcons.Alignment.DistributeY", new IMAGE_BRUSH_SVG("Icons/PaletteIcons/DistributeHorizontal",    Icon22));
+	Set("AvalancheIcons.Alignment.DistributeZ", new IMAGE_BRUSH_SVG("Icons/PaletteIcons/DistributeVertical",      Icon22));
+	
+	Set("AvalancheIcons.Alignment.Rotation.Actor.Roll",   new IMAGE_BRUSH_SVG("Icons/PaletteIcons/RotationAlignActorRoll",   Icon22));
+	Set("AvalancheIcons.Alignment.Rotation.Actor.Pitch",  new IMAGE_BRUSH_SVG("Icons/PaletteIcons/RotationAlignActorPitch",  Icon22));
+	Set("AvalancheIcons.Alignment.Rotation.Actor.Yaw",    new IMAGE_BRUSH_SVG("Icons/PaletteIcons/RotationAlignActorYaw",    Icon22));
+	Set("AvalancheIcons.Alignment.Rotation.Actor.All",    new IMAGE_BRUSH_SVG("Icons/PaletteIcons/RotationAlignActorAll",    Icon22));
+	Set("AvalancheIcons.Alignment.Rotation.Camera.Roll",  new IMAGE_BRUSH_SVG("Icons/PaletteIcons/RotationAlignCameraRoll",  Icon22));
+	Set("AvalancheIcons.Alignment.Rotation.Camera.Pitch", new IMAGE_BRUSH_SVG("Icons/PaletteIcons/RotationAlignCameraPitch", Icon22));
+	Set("AvalancheIcons.Alignment.Rotation.Camera.Yaw",   new IMAGE_BRUSH_SVG("Icons/PaletteIcons/RotationAlignCameraYaw",   Icon22));
+	Set("AvalancheIcons.Alignment.Rotation.Camera.All",   new IMAGE_BRUSH_SVG("Icons/PaletteIcons/RotationAlignCameraAll",   Icon22));
+
+	// Grid Icons
+	Set("AvalancheIcons.Grid.Toggle",           new IMAGE_BRUSH_SVG("Icons/PaletteIcons/GridToggle",              Icon22));
+	Set("AvalancheIcons.Grid.AlwaysShowToggle", new IMAGE_BRUSH_SVG("Icons/PaletteIcons/GridAlwaysVisibleToggle", Icon22));
+
+	// Snap Icons
+	Set("AvalancheIcons.Snap.Toggle",       new IMAGE_BRUSH_SVG("Icons/PaletteIcons/SnapToggle",       Icon22));
+	Set("AvalancheIcons.Snap.ToggleGrid",   new IMAGE_BRUSH_SVG("Icons/PaletteIcons/SnapGridToggle",   Icon22));
+	Set("AvalancheIcons.Snap.ToggleActor",  new IMAGE_BRUSH_SVG("Icons/PaletteIcons/SnapActorToggle",  Icon22));
+	Set("AvalancheIcons.Snap.ToggleScreen", new IMAGE_BRUSH_SVG("Icons/PaletteIcons/SnapScreenToggle", Icon22));
+
+	// Screen Icons
+	Set("AvalancheIcons.Screen.SizeToScreen",        new IMAGE_BRUSH_SVG("Icons/PaletteIcons/SizeToScreen",        Icon22));
+	Set("AvalancheIcons.Screen.SizeToScreenStretch", new IMAGE_BRUSH_SVG("Icons/PaletteIcons/SizeToScreenStretch", Icon22));
+	Set("AvalancheIcons.Screen.FitToScreen",         new IMAGE_BRUSH_SVG("Icons/PaletteIcons/FitToScreen", Icon22));
+
+	// Logs Icons
+	Set("AvalancheIcons.Logs.Pin",   new IMAGE_BRUSH_SVG("Icons/LogsIcons/Pin",   Icon16));
+	Set("AvalancheIcons.Logs.Unpin", new IMAGE_BRUSH_SVG("Icons/LogsIcons/Unpin", Icon16));
+	
+	// Editor Icons
+	Set("AvalancheIcons.Editor.ButtonExpander",        new IMAGE_BRUSH("Icons/EditorIcons/ButtonExpander",          FVector2f(7.0f)));
+	Set("AvalancheIcons.Editor.ToolHightlight",        new IMAGE_BRUSH("Icons/EditorIcons/ToolIconHighlight",       Icon40));
+	Set("AvalancheIcons.Editor.GameView",              new IMAGE_BRUSH_SVG("Icons/EditorIcons/GameView",            Icon25));
+	Set("AvalancheIcons.Editor.Visualizers",           new IMAGE_BRUSH_SVG("Icons/EditorIcons/Visualizers",         Icon25));
+	Set("AvalancheIcons.Editor.Guides",                new IMAGE_BRUSH_SVG("Icons/EditorIcons/Guides",              Icon25));
+	Set("AvalancheIcons.Editor.Billboards",            new IMAGE_BRUSH("Icons/EditorIcons/Billboards",              Icon25));
+	Set("AvalancheIcons.Editor.SelectionLock",         new IMAGE_BRUSH("Icons/EditorIcons/SelectionLock",           Icon25));
+	Set("AvalancheIcons.Editor.IsolateActors",         new IMAGE_BRUSH("Icons/EditorIcons/IsolateActors",           Icon25));
+	Set("AvalancheIcons.Editor.Settings",              new IMAGE_BRUSH("Icons/EditorIcons/Settings",                Icon25));
+	Set("AvalancheIcons.Editor.Favorites",             new IMAGE_BRUSH("Icons/EditorIcons/Favorites",               FVector2f(9.0f, 14.0f)));
+	Set("AvalancheIcons.Editor.RemoteControlEditor",   new IMAGE_BRUSH("Icons/EditorIcons/RemoteControl",           Icon25));
+	Set("AvalancheIcons.Editor.BoundingBoxes",         new IMAGE_BRUSH_SVG("Icons/EditorIcons/Mode_Bounding-Box",   Icon25));
+	Set("AvalancheIcons.Editor.SafeFrames",            new IMAGE_BRUSH_SVG("Icons/EditorIcons/Mode_Broadcast-Safe", Icon25));
+	Set("AvalancheIcons.Editor.KeyPreview",            new IMAGE_BRUSH_SVG("Icons/EditorIcons/Mode_Key-Preview",    Icon25));
+	Set("AvalancheIcons.Editor.Stats",                 new IMAGE_BRUSH_SVG("Icons/EditorIcons/Mode_Performance",    Icon25));
+	Set("AvalancheIcons.Editor.Snapshot",              new IMAGE_BRUSH_SVG("Icons/EditorIcons/Mode_Snapshot",       Icon25));
+	Set("AvalancheIcons.Editor.WireframeMode",         new IMAGE_BRUSH_SVG("Icons/EditorIcons/Mode_Wireframe",      Icon25));
+	Set("AvalancheIcons.Editor.Snap.Toggle",           new IMAGE_BRUSH_SVG("Icons/EditorIcons/Snapping_Toggle",     Icon25));
+	Set("AvalancheIcons.Editor.Grid.Toggle",           new IMAGE_BRUSH_SVG("Icons/EditorIcons/Snapping_Grid",       Icon25));
+	Set("AvalancheIcons.Editor.Radio.BlackBackground", new IMAGE_BRUSH_SVG("Icons/EditorIcons/radio-background",    Icon16, FStyleColors::Foldout.GetSpecifiedColor()));
+	Set("AvalancheIcons.Editor.Mask.Toggle",           new IMAGE_BRUSH("Icons/EditorIcons/Mode_Mask",               Icon25));
+	Set("AvalancheIcons.Editor.PivotMode",             new IMAGE_BRUSH("Icons/EditorIcons/Pivot_Mode",              Icon25));
+	
+	Set("AvalancheIcons.Lock2d", new IMAGE_BRUSH("Icons/DetailsPanelIcons/Lock2d", Icon16));
+	Set("AvalancheIcons.Lock3d", new IMAGE_BRUSH("Icons/DetailsPanelIcons/Lock3d", Icon16));
+	Set("AvalancheIcons.Unlock", new IMAGE_BRUSH("Icons/DetailsPanelIcons/Unlock", Icon16));
+
+	Set("AvalancheEditor.Thumbnail.Invalid", new FSlateRoundedBoxBrush(FStyleColors::Transparent, 4.0f, FLinearColor(1.0f, 0.2f, 0.2f, 1.0f), 1.0f));
+
+	Set("AvalancheEditor.StaticMeshToolsCategory", new CORE_IMAGE_BRUSH_SVG("Starship/AssetIcons/StaticMeshActor_16", Icon16));
+	Set("AvalancheEditor.CameraToolsCategory",     new CORE_IMAGE_BRUSH_SVG("Starship/AssetIcons/CameraActor_16", Icon16));
+	Set("AvalancheEditor.LightsToolsCategory",     new CORE_IMAGE_BRUSH_SVG("Starship/AssetIcons/PointLight_16", Icon16));
+
+	Set("AvalancheEditor.CubeTool",     new IMAGE_BRUSH("Icons/ToolboxIcons/cube",     Icon16));
+	Set("AvalancheEditor.SphereTool",   new IMAGE_BRUSH("Icons/ToolboxIcons/sphere",   Icon16));
+	Set("AvalancheEditor.CylinderTool", new IMAGE_BRUSH("Icons/ToolboxIcons/cylinder", Icon16));
+	Set("AvalancheEditor.ConeTool",     new IMAGE_BRUSH("Icons/ToolboxIcons/cone",     Icon20));
+	Set("AvalancheEditor.PlaneTool",    new IMAGE_BRUSH("Icons/ToolboxIcons/plane",    Icon20));
+
+	Set("AvalancheEditor.CameraTool",               new CORE_IMAGE_BRUSH_SVG("Starship/AssetIcons/CameraActor_16",       Icon16));
+	Set("AvalancheEditor.CineCameraTool",           new CORE_IMAGE_BRUSH_SVG("Starship/AssetIcons/CineCameraActor_16",   Icon16));
+	Set("AvalancheEditor.CameraRigCraneTool",       new CORE_IMAGE_BRUSH_SVG("Starship/AssetIcons/CameraRig_Crane_16",   Icon16));
+	Set("AvalancheEditor.CameraRigRailTool",        new CORE_IMAGE_BRUSH_SVG("Starship/AssetIcons/CameraRig_Rail_16",    Icon16));
+	Set("AvalancheEditor.CameraShakeSourceTool",    new CORE_IMAGE_BRUSH_SVG("Starship/AssetIcons/CameraShake",          Icon16));
+	Set("AvalancheEditor.AvaPostProcessVolumeTool", new CORE_IMAGE_BRUSH_SVG("Starship/AssetIcons/PostProcessVolume_16", Icon16));
+
+	Set("AvalancheEditor.PointLightTool",       new CORE_IMAGE_BRUSH_SVG("Starship/AssetIcons/PointLight_16",       Icon16));
+	Set("AvalancheEditor.DirectionalLightTool", new CORE_IMAGE_BRUSH_SVG("Starship/AssetIcons/DirectionalLight_16", Icon16));
+	Set("AvalancheEditor.RectLightTool",        new CORE_IMAGE_BRUSH_SVG("Starship/AssetIcons/RectLight_16",        Icon16));
+	Set("AvalancheEditor.SpotLightTool",        new CORE_IMAGE_BRUSH_SVG("Starship/AssetIcons/SpotLight_16",        Icon16));
+	Set("AvalancheEditor.SkyLightTool",         new CORE_IMAGE_BRUSH_SVG("Starship/AssetIcons/SkyLight_16",         Icon16));
+
+	// Easing
+	if (const UEnum* EasingEnum = StaticEnum<EAvaClonerEasing>())
+	{
+		for (int32 Idx = 0; Idx < EasingEnum->GetMaxEnumValue(); Idx++)
+		{
+			FText ValueText;
+			EasingEnum->GetDisplayValueAsText(static_cast<EAvaClonerEasing>(Idx), ValueText);
+			const FString EasingString = ValueText.ToString().Replace(TEXT(" "), TEXT(""));
+			const FName StyleName(TEXT("AvalancheIcons.Easing.") + EasingString);
+			
+			Set(StyleName, new IMAGE_BRUSH_SVG("Icons/ClonerIcons/" + EasingString, Icon32));
+		}
+	}
+	
+	// Colors
+	Set("AvalancheEditor.PalettesTab.ExpanderHeader", new FSlateColorBrush(FStyleColors::Header));
+
+	// Buttons
+	const FTextBlockStyle& AppStyle_ContentBrowserTopBarFont = FAppStyle::Get().GetWidgetStyle<FTextBlockStyle>("ContentBrowser.TopBar.Font");
+	FLinearColor ButtonTextColor = AppStyle_ContentBrowserTopBarFont.ColorAndOpacity.GetSpecifiedColor();
+	ButtonTextColor.A /= 2;
+	FLinearColor ButtonShadowColorAndOpacity = AppStyle_ContentBrowserTopBarFont.ShadowColorAndOpacity;
+	ButtonShadowColorAndOpacity.A /= 2;
+	Set("AvalancheEditor.Button.TextStyle", FTextBlockStyle(AppStyle_ContentBrowserTopBarFont)
+		.SetColorAndOpacity(ButtonTextColor)
+		.SetShadowColorAndOpacity(ButtonShadowColorAndOpacity));
+
+	const FButtonStyle& AppStyle_SimpleButton = FAppStyle::Get().GetWidgetStyle<FButtonStyle>("SimpleButton");
+
+	Set("AvalancheEditor.BorderlessButton", FButtonStyle(AppStyle_SimpleButton)
+		.SetNormalPadding(0.0f)
+		.SetPressedPadding(0.0f));
+
+	Set("AvalancheEditor.HighlightButton", FButtonStyle(AppStyle_SimpleButton)
+		.SetNormal(FSlateColorBrush(FStyleColors::Secondary))
+		.SetHovered(FSlateColorBrush(FStyleColors::Hover))
+		.SetPressed(FSlateColorBrush(FStyleColors::Header))
+		.SetDisabled(FSlateColorBrush(FStyleColors::Dropdown)));
+
+
+	Set("AvalancheEditor.SuperBarButton", FButtonStyle(AppStyle_SimpleButton)
+		.SetNormal(FSlateNoResource())
+		.SetHovered(FSlateColorBrush(FStyleColors::Secondary))
+		.SetPressed(FSlateColorBrush(FStyleColors::Hover))
+		.SetDisabled(FSlateNoResource()));
+
+	Set("AvalancheEditor.SuperBarButton.Rounded", FButtonStyle(AppStyle_SimpleButton)
+		.SetNormal(FSlateNoResource())
+		.SetHovered(FSlateRoundedBoxBrush(FStyleColors::Secondary, 4.0f))
+		.SetPressed(FSlateRoundedBoxBrush(FStyleColors::Hover, 4.0f))
+		.SetDisabled(FSlateNoResource()));
+
+	Set("AvalancheEditor.SuperBarSubMenuButton", FButtonStyle(AppStyle_SimpleButton)
+		.SetNormal(FSlateNoResource())
+		.SetHovered(FSlateNoResource())
+		.SetPressed(FSlateNoResource())
+		.SetDisabled(FSlateNoResource()));
+
+	Set("AvalancheEditor.DarkButton", FButtonStyle(AppStyle_SimpleButton)
+		.SetNormal(FSlateColorBrush(FStyleColors::Recessed))
+		.SetHovered(FSlateColorBrush(FStyleColors::Hover))
+		.SetPressed(FSlateColorBrush(FStyleColors::Header))
+		.SetDisabled(FSlateColorBrush(FStyleColors::Dropdown)));
+
+	const FButtonStyle& AppStyle_FlatButton = FAppStyle::Get().GetWidgetStyle<FButtonStyle>("FlatButton");
+	Set("LevelSnapshotsEditor.RemoveFilterButton", FButtonStyle(AppStyle_FlatButton)
+		.SetNormal(FSlateNoResource())
+		.SetNormalPadding(FMargin(0, 1.5f))
+		.SetPressedPadding(FMargin(0, 1.5f)));
+
+	// Text
+	const FTextBlockStyle& AppStyle_GraphCompactNodeTitle = FAppStyle::Get().GetWidgetStyle<FTextBlockStyle>("Graph.CompactNode.Title");
+	Set("AvalancheEditor.FilterRow.And", FTextBlockStyle(AppStyle_GraphCompactNodeTitle)
+		.SetFont(DEFAULT_FONT("BoldCondensed", 16)));
+	Set("AvalancheEditor.FilterRow.Or", FTextBlockStyle(AppStyle_GraphCompactNodeTitle)
+		.SetFont(DEFAULT_FONT("BoldCondensed", 18)));
+
+	// Check Boxes
+	const FCheckBoxStyle& AppStyle_RadioButton = FAppStyle::GetWidgetStyle<FCheckBoxStyle>("RadioButton");
+	Set("AvalancheEditor.BlackRadioButton", FCheckBoxStyle(AppStyle_RadioButton)
+		.SetBackgroundImage(*GetBrush("AvalancheIcons.Editor.Radio.BlackBackground")));
+}
+
+END_SLATE_FUNCTION_BUILD_OPTIMIZATION
diff --git a/Engine/Plugins/Experimental/Avalanche/Source/AvalancheEditor/Private/AvaEditorStyle.h b/Engine/Plugins/Experimental/Avalanche/Source/AvalancheEditor/Private/AvaEditorStyle.h
new file mode 100644
--- /dev/null
+++ b/Engine/Plugins/Experimental/Avalanche/Source/AvalancheEditor/Private/AvaEditorStyle.h
@@ -0,0 +1,21 @@
+// Copyright Epic Games, Inc. All Rights Reserved.
+
+#pragma once
+
+#include "Styling/ISlateStyle.h"
+#include "Styling/SlateStyle.h"
+
+class FAvaEditorStyle : public FSlateStyleSet
+{
+public:
+	static const FAvaEditorStyle& Get();
+
+	static void Shutdown();
+
+	FAvaEditorStyle();
+
+private:
+	static TSharedPtr<FAvaEditorStyle> StyleInstance;
+
+	void Init();
+};
diff --git a/Engine/Plugins/Experimental/Avalanche/Source/AvalancheEditor/Private/AvalancheEditorSettings.cpp b/Engine/Plugins/Experimental/Avalanche/Source/AvalancheEditor/Private/AvalancheEditorSettings.cpp
new file mode 100644
--- /dev/null
+++ b/Engine/Plugins/Experimental/Avalanche/Source/AvalancheEditor/Private/AvalancheEditorSettings.cpp
@@ -0,0 +1,75 @@
+// Copyright Epic Games, Inc. All Rights Reserved.
+
+#include "AvalancheEditorSettings.h"
+#include "ISettingsModule.h"
+#include "Modules/ModuleManager.h"
+#include "Viewport/AvaCineCameraActor.h"
+
+#define LOCTEXT_NAMESPACE "AvalancheEditorSettings"
+
+UAvalancheEditorSettings::UAvalancheEditorSettings()
+{
+	CategoryName = TEXT("Motion Design");
+	SectionName = TEXT("Editor");
+}
+
+UAvalancheEditorSettings* UAvalancheEditorSettings::Get()
+{
+	UAvalancheEditorSettings* DefaultSettings = GetMutableDefault<UAvalancheEditorSettings>();
+	static bool bInitialized = false;
+	if (!bInitialized)
+	{
+		bInitialized = true;
+		DefaultSettings->SetFlags(RF_Transactional);
+	}
+	return DefaultSettings;
+}
+
+void UAvalancheEditorSettings::OpenEditorSettingsWindow() const
+{
+	static ISettingsModule& SettingsModule = FModuleManager::LoadModuleChecked<ISettingsModule>(TEXT("Settings"));
+	SettingsModule.ShowViewer(GetContainerName(), CategoryName, SectionName);
+}
+
+void UAvalancheEditorSettings::PostInitProperties()
+{
+	Super::PostInitProperties();
+
+	if (const FName PresetNameNoLumen = TEXT("No Lumen");
+		!ViewportQualityPresets.Contains(PresetNameNoLumen))
+	{
+		ViewportQualityPresets.Add(PresetNameNoLumen, FAvaViewportQualitySettings::Preset(PresetNameNoLumen));
+	}
+
+	if (const FName PresetNameReduced = TEXT("Reduced");
+		!ViewportQualityPresets.Contains(PresetNameReduced))
+	{
+		ViewportQualityPresets.Add(PresetNameReduced, FAvaViewportQualitySettings::Preset(PresetNameReduced));
+	}
+
+	// Verifying integrity and sorting will ensure any new/removed entries are handled correctly.
+	DefaultViewportQualitySettings.VerifyIntegrity();
+	DefaultViewportQualitySettings.SortFeaturesByDisplayText();
+
+	for (TPair<FName, FAvaViewportQualitySettings>& Preset : ViewportQualityPresets)
+	{
+		FAvaViewportQualitySettings& Settings = Preset.Value;
+		Settings.VerifyIntegrity();
+		Settings.SortFeaturesByDisplayText();
+	}
+}
+
+void UAvalancheEditorSettings::PostEditChangeProperty(FPropertyChangedEvent& PropertyChangedEvent)
+{
+	Super::PostEditChangeProperty(PropertyChangedEvent);
+
+	OnChanged.Broadcast(this, PropertyChangedEvent.GetPropertyName());
+}
+
+// todo: when the reset to defaults issue will be solved, this function could be used from both PostInitProperties and PECP to update AAvaCineCameraActor defaults
+void UAvalancheEditorSettings::UpdateAvaCineCameraDefaults() const
+{
+	AAvaCineCameraActor::SetDefaultCameraDistance(CameraDistance);
+}
+
+#undef LOCTEXT_NAMESPACE
diff --git a/Engine/Plugins/Experimental/Avalanche/Source/AvalancheEditor/Private/AvalancheEditorSettings.h b/Engine/Plugins/Experimental/Avalanche/Source/AvalancheEditor/Private/AvalancheEditorSettings.h
new file mode 100644
--- /dev/null
+++ b/Engine/Plugins/Experimental/Avalanche/Source/AvalancheEditor/Private/AvalancheEditorSettings.h
@@ -0,0 +1,97 @@
+// Copyright Epic Games, Inc. All Rights Reserved.
+
+#pragma once
+
+#include "AvaDefs.h"
+#include "CoreMinimal.h"
+#include "Engine/DeveloperSettings.h"
+#include "Viewport/AvaViewportQualitySettings.h"
+#include "AvalancheEditorSettings.generated.h"
+
+USTRUCT()
+struct FAvaPaletteSettings
+{
+	GENERATED_BODY()
+
+	UPROPERTY(Config, VisibleAnywhere, Category = "Interface")
+	FName Name;
+
+	// Name is the TMap key
+	UPROPERTY(Config, EditAnywhere, Category = "Interface")
+	bool bEnablePalette = true;
+
+	UPROPERTY(Config)
+	bool bExpanded = false;
+};
+
+USTRUCT()
+struct FAvaPaletteTabSettings
+{
+	GENERATED_BODY()
+
+	UPROPERTY(Config, EditAnywhere, EditFixedSize, Category = "Interface", meta = (TitleProperty = "Name", FullyExpand))
+	TArray<FAvaPaletteSettings> Palettes;
+};
+
+/**
+ * Avalanche Editor Settings
+ */
+UCLASS(Config=EditorPerProjectUserSettings, meta = (DisplayName = "Editor"))
+class UAvalancheEditorSettings : public UDeveloperSettings
+{
+	GENERATED_BODY()
+
+public:
+	UAvalancheEditorSettings();
+
+	virtual ~UAvalancheEditorSettings() override = default;
+
+	static UAvalancheEditorSettings* Get();
+
+	/** Whether to Automatically Include the Attached Actors when performing Edit Actions such as Cut, Copy, Duplicate. */
+	UPROPERTY(Config, EditAnywhere, BlueprintReadWrite, Category = "Behavior")
+	bool bAutoIncludeAttachedActorsInEditActions = true;
+
+	/** When Grouping Actors with a Null Actor, whether to keep the relative transform of these Actors */
+	UPROPERTY(Config, EditAnywhere, BlueprintReadWrite, Category = "Behavior")
+	bool bKeepRelativeTransformWhenGrouping = false;
+
+	/*
+	 * Distance from the camera that new actors are created via the toolbox or drag and drop.
+	 * Also sets the distance from the origin that new Camera Preview Viewport cameras are created.
+	 */
+	UPROPERTY(Config, EditAnywhere, BlueprintReadWrite, Category = "Behavior")
+	float CameraDistance = 500.0f;
+
+	/**  */
+	UPROPERTY(Config, EditAnywhere, BlueprintReadWrite, Category = "Toolbox")
+	TArray<FString> FavoriteTools;
+
+	/**  */
+	UPROPERTY(Config, EditAnywhere, EditFixedSize, Category = "Interface", meta = (EditFixedOrder, ShowOnlyInnerProperties))
+	TArray<FAvaPaletteTabSettings> PaletteTabs;
+
+	/** Default viewport quality settings for all newly created Avalanche blueprints. */
+	UPROPERTY(Config, EditAnywhere, BlueprintReadWrite, Category = "Quality")
+	FAvaViewportQualitySettings DefaultViewportQualitySettings = FAvaViewportQualitySettings(true);
+
+	/** Viewport quality settings user presets. */
+	UPROPERTY(Config, EditAnywhere, BlueprintReadWrite, Category = "Quality")
+	TMap<FName, FAvaViewportQualitySettings> ViewportQualityPresets;
+
+	DECLARE_MULTICAST_DELEGATE_TwoParams(FOnSettingsChanged, const UAvalancheEditorSettings* InSettings, FName InSetting)
+	FOnSettingsChanged OnChanged;
+
+	//~ Begin UObject
+	virtual void PostInitProperties() override;
+	virtual void PostEditChangeProperty(FPropertyChangedEvent& PropertyChangedEvent) override;
+	//~ End UObject
+
+	void OpenEditorSettingsWindow() const;
+
+	/**
+	 * AvaCineCamera uses CameraDistance property to setup manual focus in its default object.
+	 * This function updates that default value to match the one from AvalancheEditorSettings
+	 */
+	void UpdateAvaCineCameraDefaults() const;
+};
diff --git a/Engine/Plugins/Experimental/Avalanche/Source/AvalancheEditor/Private/DetailView/Customizations/AvaAnchorAlignmentPropertyTypeCustomization.cpp b/Engine/Plugins/Experimental/Avalanche/Source/AvalancheEditor/Private/DetailView/Customizations/AvaAnchorAlignmentPropertyTypeCustomization.cpp
new file mode 100644
--- /dev/null
+++ b/Engine/Plugins/Experimental/Avalanche/Source/AvalancheEditor/Private/DetailView/Customizations/AvaAnchorAlignmentPropertyTypeCustomization.cpp
@@ -0,0 +1,82 @@
+// Copyright Epic Games, Inc. All Rights Reserved.
+
+#include "AvaAnchorAlignmentPropertyTypeCustomization.h"
+#include "Alignment/SAvaAnchorAlignment.h"
+#include "AvaDefs.h"
+#include "DetailWidgetRow.h"
+#include "IDetailChildrenBuilder.h"
+#include "PropertyHandle.h"
+#include "UObject/UnrealType.h"
+
+#define LOCTEXT_NAMESPACE "AvaAnchorAlignmentPropertyTypeCustomization"
+
+TSharedRef<IPropertyTypeCustomization> FAvaAnchorAlignmentPropertyTypeCustomization::MakeInstance()
+{
+	return MakeShared<FAvaAnchorAlignmentPropertyTypeCustomization>();
+}
+
+void FAvaAnchorAlignmentPropertyTypeCustomization::CustomizeHeader(TSharedRef<IPropertyHandle> InStructPropertyHandle
+	, FDetailWidgetRow& HeaderRow
+	, IPropertyTypeCustomizationUtils& StructCustomizationUtils)
+{
+	StructPropertyHandle = InStructPropertyHandle;
+
+	check(StructPropertyHandle.IsValid());
+
+	HeaderRow
+		.NameContent()
+		[
+			StructPropertyHandle->CreatePropertyNameWidget()
+		]
+		.ValueContent()
+		[
+			SNew(SAvaAnchorAlignment)
+			.UniformPadding(FMargin(5.0f, 2.0f))
+			.Anchors(this, &FAvaAnchorAlignmentPropertyTypeCustomization::GetAnchors)
+			.OnAnchorChanged(this, &FAvaAnchorAlignmentPropertyTypeCustomization::OnAnchorChanged)
+		];
+}
+
+FAvaAnchorAlignment FAvaAnchorAlignmentPropertyTypeCustomization::GetAnchors() const
+{
+	if (!StructPropertyHandle.IsValid())
+	{
+		return FAvaAnchorAlignment();
+	}
+	
+	void* AnchorsPtr = nullptr;
+	StructPropertyHandle->GetValueData(AnchorsPtr);
+	const FAvaAnchorAlignment* Anchors = static_cast<FAvaAnchorAlignment*>(AnchorsPtr);
+	
+	if (!Anchors)
+	{
+		return FAvaAnchorAlignment();
+	}
+	
+	return *Anchors;
+}
+
+void FAvaAnchorAlignmentPropertyTypeCustomization::OnAnchorChanged(const FAvaAnchorAlignment NewAnchor) const
+{
+	if (!StructPropertyHandle.IsValid())
+	{
+		return;
+	}
+	
+	void* AnchorsPtr = nullptr;
+	StructPropertyHandle->GetValueData(AnchorsPtr);
+	FAvaAnchorAlignment* Anchors = static_cast<FAvaAnchorAlignment*>(AnchorsPtr);
+
+	if (!Anchors)
+	{
+		return;
+	}
+
+	StructPropertyHandle->NotifyPreChange();
+
+	*Anchors = NewAnchor;
+
+	StructPropertyHandle->NotifyPostChange(EPropertyChangeType::ValueSet);
+}
+
+#undef LOCTEXT_NAMESPACE
diff --git a/Engine/Plugins/Experimental/Avalanche/Source/AvalancheEditor/Private/DetailView/Customizations/AvaAnchorAlignmentPropertyTypeCustomization.h b/Engine/Plugins/Experimental/Avalanche/Source/AvalancheEditor/Private/DetailView/Customizations/AvaAnchorAlignmentPropertyTypeCustomization.h
new file mode 100644
--- /dev/null
+++ b/Engine/Plugins/Experimental/Avalanche/Source/AvalancheEditor/Private/DetailView/Customizations/AvaAnchorAlignmentPropertyTypeCustomization.h
@@ -0,0 +1,37 @@
+// Copyright Epic Games, Inc. All Rights Reserved.
+
+#pragma once
+
+#include "PropertyEditorModule.h"
+#include "IPropertyTypeCustomization.h"
+
+struct FAvaAnchorAlignment;
+class FDetailWidgetRow;
+class IDetailChildrenBuilder;
+class IPropertyHandle;
+
+/**
+ * Avalanche Anchor Alignment Customization
+ * 
+ * Shows a grid of buttons for setting anchor alignment in the following layout:
+ * [HLeft ] [HCenter] [HRight ]
+ * [VTop  ] [VCenter] [VBottom]
+ * [DFront] [DCenter] [DBack  ]
+*/
+class FAvaAnchorAlignmentPropertyTypeCustomization : public IPropertyTypeCustomization
+{
+public:
+	static TSharedRef<IPropertyTypeCustomization> MakeInstance();
+
+	//~ Begin IPropertyTypeCustomization
+	virtual void CustomizeHeader(TSharedRef<IPropertyHandle> InStructPropertyHandle, FDetailWidgetRow& HeaderRow, IPropertyTypeCustomizationUtils& StructCustomizationUtils) override;
+	virtual void CustomizeChildren(TSharedRef<IPropertyHandle> InStructPropertyHandle, IDetailChildrenBuilder& StructBuilder, IPropertyTypeCustomizationUtils& StructCustomizationUtils) override {}
+	//~ End IPropertyTypeCustomization
+
+private:
+	TSharedPtr<IPropertyHandle> StructPropertyHandle;
+
+	FAvaAnchorAlignment GetAnchors() const;
+
+	void OnAnchorChanged(const FAvaAnchorAlignment NewAnchor) const;
+};
