[MotionDesign]

diff --git a/Engine/Plugins/Experimental/Avalanche/Source/AvalancheEditor/Private/DetailView/Customizations/AvaMeshesDetailCustomization.cpp b/Engine/Plugins/Experimental/Avalanche/Source/AvalancheEditor/Private/DetailView/Customizations/AvaMeshesDetailCustomization.cpp
new file mode 100644
--- /dev/null
+++ b/Engine/Plugins/Experimental/Avalanche/Source/AvalancheEditor/Private/DetailView/Customizations/AvaMeshesDetailCustomization.cpp
@@ -0,0 +1,318 @@
+// Copyright Epic Games, Inc. All Rights Reserved.
+
+#include "AvaMeshesDetailCustomization.h"
+
+#include "AssetRegistry/AssetRegistryModule.h"
+#include "AssetToolsModule.h"
+#include "DetailView/Widgets/SAvaDynamicMaterialWidget.h"
+#include "Components/DynamicMeshComponent.h"
+#include "DetailCategoryBuilder.h"
+#include "DetailLayoutBuilder.h"
+#include "DetailWidgetRow.h"
+#include "Dialogs/DlgPickAssetPath.h"
+#include "DynamicMeshToMeshDescription.h"
+#include "Engine/StaticMesh.h"
+#include "IAssetTools.h"
+#include "IDetailGroup.h"
+#include "Material/DynamicMaterialInstance.h"
+#include "StaticMeshAttributes.h"
+#include "AvaShapeActor.h"
+#include "DynamicMeshes/AvaShapeDynMeshBase.h"
+#include "UObject/Object.h"
+#include "Widgets/Input/SButton.h"
+#include "Widgets/Text/STextBlock.h"
+
+#define LOCTEXT_NAMESPACE "MeshesDetailsCustomization"
+
+void FAvaMeshesDetailCustomization::CustomizeDetails(IDetailLayoutBuilder& DetailBuilder)
+{
+	TSharedPtr<IPropertyHandle> MeshDatasHandle = DetailBuilder.GetProperty(
+		GET_MEMBER_NAME_CHECKED(UAvaShapeDynamicMeshBase, MeshDatas),
+		UAvaShapeDynamicMeshBase::StaticClass()
+	);
+
+	if (!MeshDatasHandle.IsValid())
+	{
+		return;
+	}
+
+	SelectedDynamicMeshes.Empty();
+	TArray<FName> MeshNames;
+	UAvaShapeDynamicMeshBase* DynMesh = nullptr;
+
+	for (const TWeakObjectPtr<UObject>& Object : DetailBuilder.GetSelectedObjects())
+	{
+		if (!Object.IsValid())
+		{
+			continue;
+		}
+
+		DynMesh = Cast<UAvaShapeDynamicMeshBase>(Object.Get());
+
+		if (!DynMesh)
+		{
+			AAvaShapeActor* ShapeActor = Cast<AAvaShapeActor>(Object.Get());
+
+			if (ShapeActor)
+			{
+				if (ShapeActor->GetDynamicMesh())
+				{
+					DynMesh = Cast<UAvaShapeDynamicMeshBase>(ShapeActor->GetDynamicMesh());
+				}
+			}
+		}
+
+		if (DynMesh)
+		{
+			MeshNames.Append(DynMesh->GetMeshDataNames());
+			SelectedDynamicMeshes.Add(DynMesh);
+		}
+	}
+
+	IDetailPropertyRow* MeshDatasRow = DetailBuilder.EditDefaultProperty(MeshDatasHandle);
+	DetailBuilder.HideProperty(MeshDatasHandle);
+
+	IDetailCategoryBuilder& CategoryBuilder = DetailBuilder.EditCategory(FName("Material"));
+
+	// use primary material everywhere
+	TSharedPtr<IPropertyHandle> UsePrimaryMaterialEverywhereHandle = DetailBuilder.GetProperty(
+		GET_MEMBER_NAME_CHECKED(UAvaShapeDynamicMeshBase, bUsePrimaryMaterialEverywhere),
+		UAvaShapeDynamicMeshBase::StaticClass()
+	);
+
+	// fix bug where toggle would appear in other group with same name
+	if (UsePrimaryMaterialEverywhereHandle.IsValid())
+	{
+		DetailBuilder.HideProperty(UsePrimaryMaterialEverywhereHandle);
+	}
+
+	if (MeshDatasRow)
+	{
+		MeshDatasRow->Visibility(EVisibility::Hidden);
+
+		TSharedPtr<IPropertyHandleMap> MapHandle = MeshDatasHandle->AsMap();
+		uint32 Count = 0;
+		MapHandle->GetNumElements(Count);
+
+		uint32 CountNames = MeshNames.Num();
+
+		if (Count < CountNames)
+		{
+			return;
+		}
+
+		for (uint32 Index = 0; Index < Count; ++Index)
+		{
+			FString MeshName = TEXT("");
+
+			if (Index > 0 || !DynMesh->GetUsePrimaryMaterialEverywhere())
+			{
+				MeshName = MeshNames[Index].ToString() + TEXT(" ");
+			}
+
+			TSharedPtr<IPropertyHandle> MeshPropertyHandle = MeshDatasHandle->GetChildHandle(Index);
+			check(MeshPropertyHandle.IsValid());
+
+			// Material Type
+			TSharedPtr<IPropertyHandle> MaterialTypeHandle = MeshPropertyHandle->GetChildHandle(GET_MEMBER_NAME_CHECKED(FAvaShapeMeshData, MaterialType));
+			check(MaterialTypeHandle.IsValid());
+
+			FString MaterialTypeName = MeshName + TEXT("Material Type");
+			FDetailWidgetRow& MaterialTypeRow = CategoryBuilder.AddCustomRow(FText::FromString(MaterialTypeName));
+
+			MaterialTypeRow.NameContent()[MaterialTypeHandle->CreatePropertyNameWidget(FText::FromString(MaterialTypeName))];
+			MaterialTypeRow.ValueContent()[MaterialTypeHandle->CreatePropertyValueWidget()];
+			MaterialTypeRow.Visibility(MakeAttributeLambda([=]() {
+				return MaterialTypeHandle->IsEditable() ? EVisibility::Visible : EVisibility::Hidden;
+			}));
+
+			// Material Asset
+			TSharedPtr<IPropertyHandle> MaterialHandle = MeshPropertyHandle->GetChildHandle(GET_MEMBER_NAME_CHECKED(FAvaShapeMeshData, Material));
+			check(MaterialHandle.IsValid());
+
+			FString MaterialName = MeshName + TEXT("Material Asset");
+			FDetailWidgetRow& MaterialRow = CategoryBuilder.AddCustomRow(FText::FromString(MaterialName));
+
+			MaterialRow.NameContent()
+			[
+				MaterialHandle->CreatePropertyNameWidget(FText::FromString(MaterialName))
+			];
+
+			MaterialRow.ValueContent()
+			[
+				SNew(SAvaDynamicMaterialWidget, MaterialHandle.ToSharedRef())
+			];
+
+			MaterialRow.Visibility(MakeAttributeLambda([MaterialHandle]()
+			{
+				return MaterialHandle->IsEditable() ? EVisibility::Visible : EVisibility::Hidden;
+			}));
+
+			// Parametric Material Color
+			TSharedPtr<IPropertyHandle> ParametricMaterialHandle = MeshPropertyHandle->GetChildHandle(GET_MEMBER_NAME_CHECKED(FAvaShapeMeshData, ParametricMaterial));
+			check(ParametricMaterialHandle.IsValid());
+
+			uint32 NumChildren = 0;
+			ParametricMaterialHandle->GetNumChildren(NumChildren);
+
+			for (uint32 ChildIdx = 0; ChildIdx < NumChildren; ChildIdx++)
+			{
+				TSharedPtr<IPropertyHandle> ParametricChildHandle = ParametricMaterialHandle->GetChildHandle(ChildIdx);
+				check(ParametricChildHandle.IsValid());
+
+				IDetailPropertyRow& NewParametricRow = CategoryBuilder.AddProperty(ParametricChildHandle.ToSharedRef());
+				FString ParametricRowName = MeshName + ParametricChildHandle->GetPropertyDisplayName().ToString();
+				NewParametricRow.DisplayName(FText::FromString(ParametricRowName));
+				NewParametricRow.Visibility(MakeAttributeLambda([=]() {
+					return ParametricMaterialHandle->IsEditable() ?
+						(ParametricChildHandle->IsEditable() ? EVisibility::Visible : EVisibility::Hidden) :
+						EVisibility::Hidden;
+				}));
+			}
+
+			// Use primary uv params
+			TSharedPtr<IPropertyHandle> UsePrimaryUVParamsHandle = MeshPropertyHandle->GetChildHandle(GET_MEMBER_NAME_CHECKED(FAvaShapeMeshData, bOverridePrimaryUVParams));
+			check(UsePrimaryUVParamsHandle.IsValid());
+
+			FString UsePrimaryParamsName = MeshName + TEXT("Override UV");
+			FDetailWidgetRow& UsePrimaryUVParamsRow = CategoryBuilder.AddCustomRow(FText::FromString(UsePrimaryParamsName));
+
+			UsePrimaryUVParamsRow.NameContent()[UsePrimaryUVParamsHandle->CreatePropertyNameWidget(FText::FromString(UsePrimaryParamsName))];
+			UsePrimaryUVParamsRow.ValueContent()[UsePrimaryUVParamsHandle->CreatePropertyValueWidget()];
+			UsePrimaryUVParamsRow.Visibility(MakeAttributeLambda([=]() {
+				return UsePrimaryUVParamsHandle->IsEditable() ? EVisibility::Visible : EVisibility::Hidden;
+			}));
+
+			// Only Add it the first time at this specific point
+			if (UsePrimaryMaterialEverywhereHandle.IsValid() && Index == 0 && MeshNames.Num() > 1)
+			{
+				CategoryBuilder.AddProperty(UsePrimaryMaterialEverywhereHandle);
+			}
+
+			// uv params
+			TSharedPtr<IPropertyHandle> MaterialUVHandle = MeshPropertyHandle->GetChildHandle(GET_MEMBER_NAME_CHECKED(FAvaShapeMeshData, MaterialUVParams));
+			check(MaterialUVHandle.IsValid());
+
+			IDetailPropertyRow& MaterialUVRow = CategoryBuilder.AddProperty(MaterialUVHandle.ToSharedRef());
+			FString MaterialUVName = MeshName + TEXT("Material UV");
+			MaterialUVRow.DisplayName(FText::FromString(MaterialUVName));
+			MaterialUVRow.Visibility(MakeAttributeLambda([=]() {
+				return MaterialUVHandle->IsEditable() ? EVisibility::Visible : EVisibility::Hidden;
+			}));
+
+			// Separator row
+			FDetailWidgetRow& SeparatorRow = CategoryBuilder.AddCustomRow(FText::GetEmpty());
+			SeparatorRow.WholeRowContent()[
+				SNullWidget::NullWidget
+			];
+
+			// visibility for the separator row
+			SeparatorRow.Visibility(MakeAttributeLambda([=]() {
+				return UsePrimaryUVParamsHandle->IsEditable() || MaterialUVHandle->IsEditable() ? EVisibility::Visible : EVisibility::Hidden;
+			}));
+		}
+	}
+
+	// StaticMesh Export button
+
+	static const FText ConvertToStaticMeshText = LOCTEXT("ConvertToStaticMesh", "Create StaticMesh");
+
+	IDetailGroup& GroupBuilderExport = CategoryBuilder.AddGroup(FName("Export"), FText::FromString(TEXT("Export")), false, true);
+	GroupBuilderExport.ToggleExpansion(true);
+
+	GroupBuilderExport.AddWidgetRow()
+	.NameContent()
+	[
+		SNullWidget::NullWidget
+	]
+	.ValueContent()
+	.VAlign(VAlign_Center)
+	.MaxDesiredWidth(250)
+	[
+		SNew(SButton)
+		.VAlign(VAlign_Center)
+		.ToolTipText(LOCTEXT("ConvertToStaticMeshTooltip", "Create a new StaticMesh asset using current geometry from this DynamicMeshComponent. Does not modify instance."))
+		.OnClicked(this, &FAvaMeshesDetailCustomization::OnConvertToStaticMeshClicked)
+		.IsEnabled(this, &FAvaMeshesDetailCustomization::CanConvertToStaticMesh)
+		.Content()
+		[
+			SNew(STextBlock)
+			.Text(ConvertToStaticMeshText)
+		]
+	];
+}
+
+FReply FAvaMeshesDetailCustomization::OnConvertToStaticMeshClicked()
+{
+	if (!CanConvertToStaticMesh())
+	{
+		return FReply::Handled();
+	}
+
+	TWeakObjectPtr<UAvaShapeDynamicMeshBase> DynMeshBase = SelectedDynamicMeshes[0];
+
+	if (!DynMeshBase.IsValid())
+	{
+		return FReply::Handled();
+	}
+
+	// generate name for asset
+	FString NewNameSuggestion = TEXT("SM_Avalanche") + DynMeshBase->GetMeshName();
+	FString PackageName = FString(TEXT("/Game/Meshes/")) + NewNameSuggestion;
+	FString AssetName;
+
+	FAssetToolsModule& AssetToolsModule = FModuleManager::LoadModuleChecked<FAssetToolsModule>("AssetTools");
+	AssetToolsModule.Get().CreateUniqueAssetName(PackageName, TEXT(""), PackageName, AssetName);
+
+	TSharedPtr<SDlgPickAssetPath> PickAssetPathWidget =
+		SNew(SDlgPickAssetPath)
+		.Title(LOCTEXT("ConvertToStaticMeshPickName", "Choose New StaticMesh Location"))
+		.DefaultAssetPath(FText::FromString(PackageName));
+
+	if (PickAssetPathWidget->ShowModal() != EAppReturnType::Ok)
+	{
+		return FReply::Handled();
+	}
+
+	// get input name provided by user
+	FString UserPackageName = PickAssetPathWidget->GetFullAssetPath().ToString();
+	FName MeshName(*FPackageName::GetLongPackageAssetName(UserPackageName));
+
+	// is input name valid ?
+	if (MeshName == NAME_None)
+	{
+		// Use default if invalid
+		UserPackageName = PackageName;
+		MeshName = *AssetName;
+	}
+
+	UE::Geometry::FDynamicMesh3* MeshIn = DynMeshBase->GetShapeMeshComponent()->GetMesh();
+
+	// empty mesh do not export
+	if (!MeshIn || MeshIn->TriangleCount() == 0)
+	{
+		return FReply::Handled();
+	}
+
+	// find/create package
+	UPackage* Package = CreatePackage(*UserPackageName);
+	check(Package);
+
+	// Create StaticMesh object
+	UStaticMesh* StaticMesh = NewObject<UStaticMesh>(Package, MeshName, RF_Public | RF_Standalone);
+
+	if (DynMeshBase->ExportToStaticMesh(StaticMesh))
+	{
+		// Notify asset registry of new asset
+		FAssetRegistryModule::AssetCreated(StaticMesh);
+	}
+
+	return FReply::Handled();
+}
+
+bool FAvaMeshesDetailCustomization::CanConvertToStaticMesh() const
+{
+	return SelectedDynamicMeshes.Num() == 1;
+}
+
+#undef LOCTEXT_NAMESPACE
\ No newline at end of file
diff --git a/Engine/Plugins/Experimental/Avalanche/Source/AvalancheEditor/Private/DetailView/Customizations/AvaMeshesDetailCustomization.h b/Engine/Plugins/Experimental/Avalanche/Source/AvalancheEditor/Private/DetailView/Customizations/AvaMeshesDetailCustomization.h
new file mode 100644
--- /dev/null
+++ b/Engine/Plugins/Experimental/Avalanche/Source/AvalancheEditor/Private/DetailView/Customizations/AvaMeshesDetailCustomization.h
@@ -0,0 +1,37 @@
+// Copyright Epic Games, Inc. All Rights Reserved.
+
+#pragma once
+
+#include "Delegates/Delegate.h"
+#include "IDetailCustomization.h"
+#include "Templates/SharedPointer.h"
+#include "UObject/WeakObjectPtr.h"
+
+class FReply;
+class UAvaShapeDynamicMeshBase;
+class UDynamicMaterialModel;
+
+/* Used to create the details materials meshes widget and export to StaticMesh */
+class FAvaMeshesDetailCustomization : public IDetailCustomization
+{
+public:
+	static TSharedRef<IDetailCustomization> MakeInstance()
+	{
+		return MakeShared<FAvaMeshesDetailCustomization>();
+	}
+
+	FAvaMeshesDetailCustomization()
+	{
+	}
+
+	virtual void CustomizeDetails(IDetailLayoutBuilder& DetailBuilder) override;
+
+protected:
+	// Handler when the convert button is clicked
+	FReply OnConvertToStaticMeshClicked();
+
+	// Enable or disable the button when selected object is not compatible
+	bool CanConvertToStaticMesh() const;
+
+	TArray<TWeakObjectPtr<UAvaShapeDynamicMeshBase>> SelectedDynamicMeshes;
+};
\ No newline at end of file
diff --git a/Engine/Plugins/Experimental/Avalanche/Source/AvalancheEditor/Private/DetailView/Customizations/AvaVectorPropertyTypeCustomization.cpp b/Engine/Plugins/Experimental/Avalanche/Source/AvalancheEditor/Private/DetailView/Customizations/AvaVectorPropertyTypeCustomization.cpp
new file mode 100644
--- /dev/null
+++ b/Engine/Plugins/Experimental/Avalanche/Source/AvalancheEditor/Private/DetailView/Customizations/AvaVectorPropertyTypeCustomization.cpp
@@ -0,0 +1,774 @@
+// Copyright Epic Games, Inc. All Rights Reserved.
+
+#include "AvaVectorPropertyTypeCustomization.h"
+#include "AvaEditorStyle.h"
+#include "AvaEditorSubsystem.h"
+#include "AvaEditorViewportUtils.h"
+#include "AvaViewportUtils.h"
+#include "DetailLayoutBuilder.h"
+#include "DetailWidgetRow.h"
+#include "DynamicMeshes/AvaShape2DDynMeshBase.h"
+#include "DynamicMeshes/AvaShape3DDynMeshBase.h"
+#include "Editor.h"
+#include "IDetailChildrenBuilder.h"
+#include "Viewport/AvaViewportExtension.h"
+#include "ViewportClient/IAvaViewportClient.h"
+#include "Widgets/Input/SButton.h"
+#include "Widgets/Input/SComboButton.h"
+#include "Widgets/Input/STextComboBox.h"
+#include "Widgets/Layout/SScaleBox.h"
+
+#define LOCTEXT_NAMESPACE "AvaVectorPropertyTypeCustomization"
+
+void FAvaVectorPropertyTypeCustomization::CustomizeHeader(TSharedRef<IPropertyHandle> StructPropertyHandle
+	, FDetailWidgetRow& HeaderRow
+	, IPropertyTypeCustomizationUtils& StructCustomizationUtils)
+{
+	VectorPropertyHandle = StructPropertyHandle;
+	const FString Type = StructPropertyHandle->GetProperty()->GetCPPType();
+	bIsVector3d = Type == "FVector";
+	SelectedObjectNum = VectorPropertyHandle->GetNumPerObjectValues();
+
+	if (SelectedObjectNum > 0)
+	{
+		TArray<UObject*> OuterObjects;
+		VectorPropertyHandle->GetOuterObjects(OuterObjects);
+
+		// All outer objects should have the same world? So just use the first
+		if (OuterObjects.Num() > 0)
+		{
+			UAvaEditorSubsystem* AvaEditorSubsystem = UAvaEditorSubsystem::Get(OuterObjects[0]);
+			if (AvaEditorSubsystem)
+			{
+				TSharedPtr<FAvaViewportExtension> ViewportExtension = AvaEditorSubsystem->FindExtension<FAvaViewportExtension>();
+				if (ViewportExtension.IsValid())
+				{
+					TArray<TSharedPtr<IAvaViewportClient>> ViewportClients = ViewportExtension->GetViewportClients();
+					if (ViewportClients.Num() > 0)
+					{
+						ViewportClient = ViewportClients[0];
+					}
+				}
+			}
+		}
+	}
+
+	bool bPreserveRatio = false;
+	if (StructPropertyHandle->HasMetaData("AllowPreserveRatio"))
+	{
+		bPreserveRatio = true;
+		
+		static const FVector2D ImageSize(16.f);
+
+		// fill available space
+		HeaderRow.NameWidget.HorizontalAlignment = EHorizontalAlignment::HAlign_Fill;
+		HeaderRow.NameWidget.VerticalAlignment = EVerticalAlignment::VAlign_Fill;
+
+		auto ComboBoxButtonBuilder = [this](const ERatioMode ButtonMode)->TSharedRef<SButton>
+		{
+			const FSlateBrush* ButtonImage = GetComboButtonBrush(ButtonMode);
+			const FText& ButtonText = GetComboButtonText(ButtonMode);
+			
+			return SNew(SButton)
+					.HAlign(EHorizontalAlignment::HAlign_Fill)
+					.VAlign(EVerticalAlignment::VAlign_Fill)
+					.Cursor(EMouseCursor::Hand)
+					.OnClicked(this, &FAvaVectorPropertyTypeCustomization::OnComboButtonClicked, ButtonMode)
+					.ContentPadding(0.f)
+					[
+						SNew(SHorizontalBox)
+						+ SHorizontalBox::Slot()
+						.FillWidth(0.25f)
+						[
+							SNew(SScaleBox)
+							[
+								SNew(SImage)
+								.DesiredSizeOverride(ImageSize)
+								.Image(ButtonImage)
+							]
+						]
+						+ SHorizontalBox::Slot()
+						.FillWidth(1.f)
+						[
+							SNew(STextBlock)
+							.Justification(ETextJustify::Center)
+							.Text(ButtonText)
+							.Margin(FMargin(5.f, 0.f))
+						]
+					];
+		};
+
+		TArray<ERatioMode> RatioModes {
+			ERatioMode::None,
+			ERatioMode::PreserveXY
+		};
+		
+		if (bIsVector3d)
+		{
+			RatioModes.Add(ERatioMode::PreserveYZ);
+			RatioModes.Add(ERatioMode::PreserveXZ);
+			RatioModes.Add(ERatioMode::PreserveXYZ);
+		}
+
+		const TSharedRef<SVerticalBox> ButtonVerticalBox = SNew(SVerticalBox);
+		for (const ERatioMode& Mode : RatioModes)
+		{
+			ButtonVerticalBox->AddSlot()
+			.Padding(0.f)
+			[
+				ComboBoxButtonBuilder(Mode)
+			];
+		}
+		
+		HeaderRow.NameContent()[
+			SNew(SHorizontalBox)
+			+ SHorizontalBox::Slot()
+			.FillWidth(1.f)
+			[
+				StructPropertyHandle->CreatePropertyNameWidget()
+			]
+			+ SHorizontalBox::Slot()
+			.FillWidth(1.f)
+			[
+				SAssignNew(ComboButton, SComboButton)
+				.HAlign(EHorizontalAlignment::HAlign_Fill)
+				.VAlign(EVerticalAlignment::VAlign_Center)
+				.Cursor(EMouseCursor::Hand)
+				.Method(EPopupMethod::UseCurrentWindow)
+				.ButtonContent()
+				[
+					SNew(SHorizontalBox)
+					+ SHorizontalBox::Slot()
+					.AutoWidth()
+					[
+						SNew(SScaleBox)
+						[
+							SNew(SImage)
+							.DesiredSizeOverride(ImageSize)
+							.Image(this, &FAvaVectorPropertyTypeCustomization::GetCurrentComboButtonBrush)
+						]
+					]
+					+ SHorizontalBox::Slot()
+					.AutoWidth()
+					[
+						SNew(STextBlock)
+						.Justification(ETextJustify::Center)
+						.Text(this, &FAvaVectorPropertyTypeCustomization::GetCurrentComboButtonText)
+						.Margin(FMargin(5.f, 0.f))
+					]
+				]
+				.MenuContent()
+				[
+					ButtonVerticalBox
+				]
+			]
+		];
+	}
+	else
+	{
+		HeaderRow.NameContent()[StructPropertyHandle->CreatePropertyNameWidget()];
+	}
+
+	// fill available space
+	HeaderRow.ValueWidget.HorizontalAlignment = EHorizontalAlignment::HAlign_Fill;
+	HeaderRow.ValueWidget.VerticalAlignment = EVerticalAlignment::VAlign_Fill;
+	RatioMode = ERatioMode::None;
+	
+	if (bIsVector3d)
+	{
+		bPixelSizeProperty = VectorPropertyHandle->GetProperty()->GetFName() == GET_MEMBER_NAME_CHECKED(UAvaShape3DDynMeshBase, PixelSize3D);
+		XPropertyHandle = StructPropertyHandle->GetChildHandle(GET_MEMBER_NAME_CHECKED(FVector, X));
+		YPropertyHandle = StructPropertyHandle->GetChildHandle(GET_MEMBER_NAME_CHECKED(FVector, Y));
+		ZPropertyHandle = StructPropertyHandle->GetChildHandle(GET_MEMBER_NAME_CHECKED(FVector, Z));
+
+		if (bPreserveRatio)
+		{
+			if (const FString& Value = VectorPropertyHandle->GetMetaDataProperty()->GetMetaData(FName("VectorRatioMode")); !Value.IsEmpty())
+			{
+				RatioMode = static_cast<ERatioMode>(FCString::Atoi(*Value));
+			}
+			else
+			{
+				RatioMode = ERatioMode::PreserveXYZ;
+			}
+		}
+		
+		if (StructPropertyHandle->HasMetaData("ClampMin"))
+		{
+			MinVectorClamp = FVector(StructPropertyHandle->GetFloatMetaData("ClampMin"));
+		}
+		
+		if (StructPropertyHandle->HasMetaData("ClampMax"))
+		{
+			MaxVectorClamp = FVector(StructPropertyHandle->GetFloatMetaData("ClampMax"));
+		}
+
+		float SpinDelta = 1.f;
+		if (StructPropertyHandle->HasMetaData("Delta"))
+		{
+			SpinDelta = StructPropertyHandle->GetFloatMetaData("Delta");
+		}
+		// compute spin delta based on min and max value in percentage (100%)
+		else if (MinVectorClamp.IsSet() && MaxVectorClamp.IsSet())
+		{
+			SpinDelta = (MaxVectorClamp->X - MinVectorClamp->X) / 100.f;
+		}
+		
+		HeaderRow.ValueContent()
+		[
+			SNew(SNumericVectorInputBox3D)
+				.Font(IDetailLayoutBuilder::GetDetailFont())
+				.X(this, &FAvaVectorPropertyTypeCustomization::GetVectorComponent,  static_cast<uint8>(0)) // X
+				.Y(this, &FAvaVectorPropertyTypeCustomization::GetVectorComponent,  static_cast<uint8>(1)) // Y
+				.Z(this, &FAvaVectorPropertyTypeCustomization::GetVectorComponent,  static_cast<uint8>(2)) // Z
+				.bColorAxisLabels(bPreserveRatio)
+				.MinVector(MinVectorClamp)
+				.MaxVector(MaxVectorClamp)
+				.MinSliderVector(MinVectorClamp)
+				.MaxSliderVector(MaxVectorClamp)
+				.OnXChanged(this, &FAvaVectorPropertyTypeCustomization::SetVectorComponent,  static_cast<uint8>(0))
+				.OnYChanged(this, &FAvaVectorPropertyTypeCustomization::SetVectorComponent,  static_cast<uint8>(1))
+				.OnZChanged(this, &FAvaVectorPropertyTypeCustomization::SetVectorComponent,  static_cast<uint8>(2))
+				.OnXCommitted(this, &FAvaVectorPropertyTypeCustomization::SetVectorComponent,  static_cast<uint8>(0))
+				.OnYCommitted(this, &FAvaVectorPropertyTypeCustomization::SetVectorComponent,  static_cast<uint8>(1))
+				.OnZCommitted(this, &FAvaVectorPropertyTypeCustomization::SetVectorComponent,  static_cast<uint8>(2))
+				.AllowSpin(true)
+				.SpinDelta(SpinDelta)
+				.IsEnabled(this, &FAvaVectorPropertyTypeCustomization::CanEditValue)
+				.OnBeginSliderMovement(this, &FAvaVectorPropertyTypeCustomization::OnBeginSliderMovement)
+				.OnEndSliderMovement(this, &FAvaVectorPropertyTypeCustomization::OnEndSliderMovement)
+		];
+	}
+	else
+	{
+		bPixelSizeProperty = VectorPropertyHandle->GetProperty()->GetFName() == GET_MEMBER_NAME_CHECKED(UAvaShape2DDynMeshBase, PixelSize2D);
+		XPropertyHandle = StructPropertyHandle->GetChildHandle(GET_MEMBER_NAME_CHECKED(FVector2D, X));
+		YPropertyHandle = StructPropertyHandle->GetChildHandle(GET_MEMBER_NAME_CHECKED(FVector2D, Y));
+
+		if (bPreserveRatio)
+		{
+			if (const FString& Value = VectorPropertyHandle->GetMetaDataProperty()->GetMetaData(FName("VectorRatioMode")); !Value.IsEmpty())
+			{
+				RatioMode = static_cast<ERatioMode>(FCString::Atoi(*Value));
+			}
+			else
+			{
+				RatioMode = ERatioMode::PreserveXY;
+			}
+		}
+		
+		if (StructPropertyHandle->HasMetaData("ClampMin"))
+		{
+			MinVector2DClamp = FVector2D(StructPropertyHandle->GetFloatMetaData("ClampMin"));
+		}
+		
+		if (StructPropertyHandle->HasMetaData("ClampMax"))
+		{
+			MaxVector2DClamp = FVector2D(StructPropertyHandle->GetFloatMetaData("ClampMax"));
+		}
+
+		float SpinDelta = 1.f;
+		if (StructPropertyHandle->HasMetaData("Delta"))
+		{
+			SpinDelta = StructPropertyHandle->GetFloatMetaData("Delta");
+		}
+		// compute spin delta based on min and max value in percentage (100%)
+		else if (MinVector2DClamp.IsSet() && MaxVector2DClamp.IsSet())
+		{
+			SpinDelta = (MaxVector2DClamp->X - MinVector2DClamp->X) / 100.f;
+		}
+
+		HeaderRow.ValueContent()
+		[
+			SNew(SNumericVectorInputBox2D)
+				.Font(IDetailLayoutBuilder::GetDetailFont())
+				.X(this, &FAvaVectorPropertyTypeCustomization::GetVectorComponent, static_cast<uint8>(0)) // X
+				.Y(this, &FAvaVectorPropertyTypeCustomization::GetVectorComponent,  static_cast<uint8>(1)) // Y
+				.bColorAxisLabels(bPreserveRatio)
+				.MinVector(MinVector2DClamp)
+				.MaxVector(MaxVector2DClamp)
+				.MinSliderVector(MinVector2DClamp)
+				.MaxSliderVector(MaxVector2DClamp)
+				.OnXChanged(this, &FAvaVectorPropertyTypeCustomization::SetVectorComponent,  static_cast<uint8>(0))
+				.OnYChanged(this, &FAvaVectorPropertyTypeCustomization::SetVectorComponent,  static_cast<uint8>(1))
+				.OnXCommitted(this, &FAvaVectorPropertyTypeCustomization::SetVectorComponent,  static_cast<uint8>(0))
+				.OnYCommitted(this, &FAvaVectorPropertyTypeCustomization::SetVectorComponent,  static_cast<uint8>(1))
+				.AllowSpin(true)
+				.SpinDelta(SpinDelta)
+				.IsEnabled(this, &FAvaVectorPropertyTypeCustomization::CanEditValue)
+				.OnBeginSliderMovement(this, &FAvaVectorPropertyTypeCustomization::OnBeginSliderMovement)
+				.OnEndSliderMovement(this, &FAvaVectorPropertyTypeCustomization::OnEndSliderMovement)
+		];
+	}
+}
+
+void FAvaVectorPropertyTypeCustomization::CustomizeChildren(TSharedRef<IPropertyHandle> StructPropertyHandle,
+	IDetailChildrenBuilder& StructBuilder, IPropertyTypeCustomizationUtils& StructCustomizationUtils)
+{
+}
+
+FReply FAvaVectorPropertyTypeCustomization::OnComboButtonClicked(const ERatioMode NewMode)
+{
+	RatioMode = NewMode;
+	const FString MetaData = FString::FromInt(static_cast<int32>(RatioMode));
+	VectorPropertyHandle->GetMetaDataProperty()->SetMetaData(FName("VectorRatioMode"), *MetaData);
+	ComboButton->SetIsOpen(false);
+	return FReply::Handled();
+}
+
+const FSlateBrush* FAvaVectorPropertyTypeCustomization::GetComboButtonBrush(const ERatioMode Mode) const
+{
+	if (Mode == ERatioMode::None)
+	{
+		return FAvaEditorStyle::Get().GetBrush("AvalancheIcons.Unlock");
+	}
+	else if (Mode == ERatioMode::PreserveXY || Mode == ERatioMode::PreserveXZ || Mode == ERatioMode::PreserveYZ)
+	{
+		return FAvaEditorStyle::Get().GetBrush("AvalancheIcons.Lock2d");
+	}
+	else
+	{
+		return FAvaEditorStyle::Get().GetBrush("AvalancheIcons.Lock3d");
+	}
+}
+
+FText FAvaVectorPropertyTypeCustomization::GetComboButtonText(const ERatioMode Mode) const
+{
+	switch (Mode)
+	{
+		case(ERatioMode::PreserveXY):
+			return FText::FromString("XY ");
+			break;
+		case(ERatioMode::PreserveYZ):
+			return FText::FromString("YZ ");
+			break;
+		case(ERatioMode::PreserveXZ):
+			return FText::FromString("XZ ");
+			break;
+		case(ERatioMode::PreserveXYZ):
+			return FText::FromString("XYZ");
+			break;
+		default:
+			return FText::FromString("Free");
+			break;
+	}
+}
+
+const FSlateBrush* FAvaVectorPropertyTypeCustomization::GetCurrentComboButtonBrush() const
+{
+	return GetComboButtonBrush(RatioMode);
+}
+
+FText FAvaVectorPropertyTypeCustomization::GetCurrentComboButtonText() const
+{
+	return GetComboButtonText(RatioMode);
+}
+
+TOptional<double> FAvaVectorPropertyTypeCustomization::GetVectorComponent(const uint8 Component) const
+{
+	if (!VectorPropertyHandle.IsValid() || SelectedObjectNum == 0)
+	{
+		return TOptional<double>();
+	}
+	double OutValue = 0.f;
+	switch(Component)
+	{
+		case 0:
+			if (XPropertyHandle->GetValue(OutValue) != FPropertyAccess::Success)
+			{
+				return TOptional<double>();
+			}
+		break;
+		case 1:
+			if (YPropertyHandle->GetValue(OutValue) != FPropertyAccess::Success)
+			{
+				return TOptional<double>();
+			}
+		break;
+		case 2:
+			if (ZPropertyHandle->GetValue(OutValue) != FPropertyAccess::Success)
+			{
+				return TOptional<double>();
+			}
+		break;
+		default:
+			return TOptional<double>();
+		break;
+	}
+	// handle specific case
+	if (bPixelSizeProperty)
+	{
+		return MeshSizeToPixelSize(OutValue);
+	}
+	else
+	{
+		return OutValue;
+	}
+}
+
+void FAvaVectorPropertyTypeCustomization::SetVectorComponent(double NewValue, const uint8 Component)
+{
+	if (bMovingSlider)
+	{
+		SetVectorComponent(NewValue, ETextCommit::Default, Component);
+	}
+}
+
+void FAvaVectorPropertyTypeCustomization::SetVectorComponent(double NewValue, ETextCommit::Type CommitType, const uint8 Component)
+{
+	const bool bFinalCommit = CommitType == ETextCommit::OnEnter || CommitType == ETextCommit::OnUserMovedFocus;
+	if (!bFinalCommit && CommitType != ETextCommit::Default)
+	{
+		return;
+	}
+	if (!VectorPropertyHandle.IsValid())
+	{
+		return;
+	}
+	// init here in case we only input the value without slider movement
+	if ((!bIsVector3d && Begin2DValues.Num() == 0) ||
+		(bIsVector3d && Begin3DValues.Num() == 0))
+	{
+		InitVectorValuesForRatio();
+	}
+	if (SelectedObjectNum == 0)
+	{
+		return;
+	}
+	// handle interactive debounce to avoid slow behaviour
+	LastComponentValueSet = Component;
+	if (DebounceValueSet > 0 && !bFinalCommit)
+	{
+		DebounceValueSet--;
+		return;
+	}
+	DebounceValueSet = FMath::Clamp(SelectedObjectNum > 1 ? (SelectedObjectNum * MULTI_OBJECT_DEBOUNCE) : SINGLE_OBJECT_DEBOUNCE, 0, 255);
+	// handle specific case for pixel size property
+	if (bPixelSizeProperty)
+	{
+		NewValue = PixelSizeToMeshSize(NewValue);
+	}
+	if (!bMovingSlider)
+	{
+		GEditor->BeginTransaction(VectorPropertyHandle->GetPropertyDisplayName());
+	}
+	// update objects value for property, we handle transaction ourselves to batch properties changes together
+	const EPropertyValueSetFlags::Type Flags = bMovingSlider ? EPropertyValueSetFlags::InteractiveChange : EPropertyValueSetFlags::NotTransactable;
+	SetComponentValue(NewValue, Component, Flags);
+	if (!bMovingSlider)
+	{
+		GEditor->EndTransaction();
+		ResetVectorValuesForRatio();
+	}
+}
+
+void FAvaVectorPropertyTypeCustomization::OnBeginSliderMovement()
+{
+	LastComponentValueSet = INVALID_COMPONENT_IDX;
+	DebounceValueSet = 0;
+	InitVectorValuesForRatio();
+	bMovingSlider = true;
+	GEditor->BeginTransaction(VectorPropertyHandle->GetPropertyDisplayName());
+}
+
+void FAvaVectorPropertyTypeCustomization::OnEndSliderMovement(double NewValue)
+{
+	DebounceValueSet = 0;
+	bMovingSlider = false;
+	// set final value like enter pressed
+	if (LastComponentValueSet != INVALID_COMPONENT_IDX)
+	{
+		SetVectorComponent(NewValue, ETextCommit::OnEnter, LastComponentValueSet);
+	}
+	// end started transactions during process
+	while(GEditor->IsTransactionActive())
+	{
+		GEditor->EndTransaction();
+	}
+	ResetVectorValuesForRatio();
+}
+
+bool FAvaVectorPropertyTypeCustomization::CanEditValue() const
+{
+	if (!VectorPropertyHandle.IsValid())
+	{
+		return false;
+	}
+	
+	if (bPixelSizeProperty && SelectedObjectNum > 0 && ViewportClient.IsValid())
+	{
+		if (FAvaViewportUtils::IsValidViewportSize(ViewportClient.Pin()->GetVirtualViewportSize()))
+		{
+			return VectorPropertyHandle->IsEditable();
+		}
+
+		return false;
+	}
+	
+	return VectorPropertyHandle->IsEditable();
+}
+
+void FAvaVectorPropertyTypeCustomization::InitVectorValuesForRatio()
+{
+	Begin2DValues.Empty();
+	Begin3DValues.Empty();
+	
+	TArray<FString> OutValues;
+	VectorPropertyHandle->GetPerObjectValues(OutValues);
+	SelectedObjectNum = VectorPropertyHandle->GetNumPerObjectValues();
+	
+	if (SelectedObjectNum > 0)
+	{
+		for (const FString& Val : OutValues)
+		{
+			if (bIsVector3d)
+			{
+				FVector Vector;
+				if (Vector.InitFromString(Val))
+				{
+					Begin3DValues.Add(Vector);
+				}
+				else
+				{
+					TOptional<FVector> OptionalVector;
+					Begin3DValues.Add(OptionalVector);
+				}
+			}
+			else
+			{
+				FVector2D Vector;
+				if (Vector.InitFromString(Val))
+				{
+					Begin2DValues.Add(Vector);
+				}
+				else
+				{
+					TOptional<FVector2D> OptionalVector;
+					Begin2DValues.Add(OptionalVector);
+				}
+			}
+		}
+	}
+}
+
+void FAvaVectorPropertyTypeCustomization::ResetVectorValuesForRatio()
+{
+	Begin2DValues.Empty();
+	Begin3DValues.Empty();
+}
+
+void FAvaVectorPropertyTypeCustomization::SetComponentValue(const double NewValue, const uint8 Component, const EPropertyValueSetFlags::Type Flags)
+{
+	// check if we are preserving ratio for current component change
+	bool bPreserveRatio = false;
+	switch (Component)
+	{
+		case 0:
+			bPreserveRatio = ((RatioMode & ERatioMode::X) != ERatioMode::None);
+		break;
+		case 1:
+			bPreserveRatio = ((RatioMode & ERatioMode::Y) != ERatioMode::None);
+		break;
+		case 2:
+			bPreserveRatio = ((RatioMode & ERatioMode::Z) != ERatioMode::None);
+		break;
+		default:
+		break;
+	}
+	const TArray<bool> PreserveRatios
+	{
+		(RatioMode & ERatioMode::X) != ERatioMode::None && bPreserveRatio,
+		(RatioMode & ERatioMode::Y) != ERatioMode::None && bPreserveRatio,
+		(RatioMode & ERatioMode::Z) != ERatioMode::None && bPreserveRatio
+	};
+	// set property per object since we need to handle ratios for each objects
+	const uint8 MaxComponentCount = bIsVector3d ? 3 : 2;
+	for (int32 ObjIdx = 0; ObjIdx < SelectedObjectNum; ObjIdx++)
+	{
+		if (!bIsVector3d && !Begin2DValues.IsValidIndex(ObjIdx) && !Begin2DValues[ObjIdx].IsSet())
+		{
+			continue;
+		}
+		if (bIsVector3d && !Begin3DValues.IsValidIndex(ObjIdx) && !Begin3DValues[ObjIdx].IsSet())
+		{
+			continue;
+		}
+		// compute clamped ratio for value change
+		const double ClampedRatio = GetClampedRatioValueChange(ObjIdx, NewValue, Component, PreserveRatios);
+		// loop for each component (X,Y,Z)
+		for (uint8 ComponentIdx = 0; ComponentIdx < MaxComponentCount; ComponentIdx++)
+		{
+			// only assign value to specific component, skip others
+			if (!bPreserveRatio && ComponentIdx != Component)
+			{
+				continue;
+			}
+			// compute new component value
+			const double NewComponentValue = GetClampedComponentValue(ObjIdx, NewValue, ClampedRatio, ComponentIdx, Component);
+			switch(ComponentIdx)
+			{
+				case 0:
+					if (PreserveRatios[0] || ComponentIdx == Component)
+					{
+						XPropertyHandle->SetPerObjectValue(ObjIdx, FString::SanitizeFloat(NewComponentValue), Flags);
+					}
+					break;
+				case 1:
+					if (PreserveRatios[1] || ComponentIdx == Component)
+					{
+						YPropertyHandle->SetPerObjectValue(ObjIdx, FString::SanitizeFloat(NewComponentValue), Flags);
+					}
+					break;
+				case 2:
+					if (PreserveRatios[2] || ComponentIdx == Component)
+					{
+						ZPropertyHandle->SetPerObjectValue(ObjIdx, FString::SanitizeFloat(NewComponentValue), Flags);
+					}
+					break;
+				default:;
+			}
+		}
+	}
+}
+
+double FAvaVectorPropertyTypeCustomization::GetClampedRatioValueChange(const int32 ObjectIdx, const double NewValue, const uint8 Component, const TArray<bool>& PreserveRatios) const
+{
+	double Ratio = 1;
+	// get pre change value for this component
+	if (bIsVector3d)
+	{
+		FVector BeginValue = Begin3DValues[ObjectIdx].GetValue();
+		if (BeginValue[Component] != 0)
+		{
+			Ratio = NewValue / BeginValue[Component];
+		}
+		// apply min/max clamp
+		if (MinVectorClamp.IsSet() || MaxVectorClamp.IsSet())
+		{
+			for (int32 ComponentIdx = 0; ComponentIdx < 3; ComponentIdx++)
+			{
+				if (PreserveRatios[ComponentIdx] || ComponentIdx == Component)
+				{
+					const double EndValue = BeginValue[ComponentIdx] * Ratio;
+					if (MinVectorClamp.IsSet())
+					{
+						const double MinValue = MinVectorClamp.GetValue()[ComponentIdx];
+						if (EndValue < MinValue)
+						{
+							Ratio = MinValue / BeginValue[ComponentIdx];
+						}
+					}
+					if (MaxVectorClamp.IsSet())
+					{
+						const double MaxValue = MaxVectorClamp.GetValue()[ComponentIdx];
+						if (EndValue > MaxValue)
+						{
+							Ratio = MaxValue / BeginValue[ComponentIdx];
+						}
+					}
+				}
+			}	
+		}
+	}
+	else
+	{
+		FVector2D BeginValue = Begin2DValues[ObjectIdx].GetValue();
+		if (BeginValue[Component] != 0)
+		{
+			Ratio = NewValue / BeginValue[Component];
+		}
+		// apply min/max clamp
+		if (MinVector2DClamp.IsSet() || MaxVector2DClamp.IsSet())
+		{
+			for (int32 ComponentIdx = 0; ComponentIdx < 2; ComponentIdx++)
+			{
+				if (PreserveRatios[ComponentIdx] || ComponentIdx == Component)
+				{
+					const double EndValue = BeginValue[ComponentIdx] * Ratio;
+					if (MinVector2DClamp.IsSet())
+					{
+						const double MinValue = MinVector2DClamp.GetValue()[ComponentIdx];
+						if (EndValue < MinValue)
+						{
+							Ratio = MinValue / BeginValue[ComponentIdx];
+						}
+					}
+					if (MaxVector2DClamp.IsSet())
+					{
+						const double MaxValue = MaxVector2DClamp.GetValue()[ComponentIdx];
+						if (EndValue > MaxValue)
+						{
+							Ratio = MaxValue / BeginValue[ComponentIdx];
+						}
+					}
+				}
+			}
+		}
+	}
+	return Ratio;
+}
+
+double FAvaVectorPropertyTypeCustomization::GetClampedComponentValue(const int32 ObjectIdx, double NewValue, const double Ratio, const uint8 ComponentIdx, const uint8 OriginalComponent)
+{
+	const double OldValue = (bIsVector3d ? Begin3DValues[ObjectIdx].GetValue()[ComponentIdx] : Begin2DValues[ObjectIdx].GetValue()[ComponentIdx]);
+	const double SliderOriginalValue = (bIsVector3d ? Begin3DValues[ObjectIdx].GetValue()[OriginalComponent] : Begin2DValues[ObjectIdx].GetValue()[OriginalComponent]);
+	if (SliderOriginalValue == 0 && OldValue == 0)
+	{
+		if (bIsVector3d)
+		{
+			if (MinVectorClamp.IsSet())
+			{
+				NewValue = FMath::Max(NewValue, MinVectorClamp.GetValue()[ComponentIdx]);
+			}
+			if (MaxVectorClamp.IsSet())
+			{
+				NewValue = FMath::Min(NewValue, MaxVectorClamp.GetValue()[ComponentIdx]);
+			}
+		}
+		else
+		{
+			if (MinVector2DClamp.IsSet())
+			{
+				NewValue = FMath::Max(NewValue, MinVector2DClamp.GetValue()[ComponentIdx]);
+			}
+			if (MaxVector2DClamp.IsSet())
+			{
+				NewValue = FMath::Min(NewValue, MaxVector2DClamp.GetValue()[ComponentIdx]);
+			}
+		}
+		return NewValue;
+	}
+	return OldValue * Ratio;
+}
+
+double FAvaVectorPropertyTypeCustomization::MeshSizeToPixelSize(double MeshSize) const
+{
+	if (ViewportClient.IsValid())
+	{
+		double PixelSize;
+
+		if (FAvaEditorViewportUtils::MeshSizeToPixelSize(ViewportClient.Pin().ToSharedRef(), MeshSize, PixelSize))
+		{
+			return PixelSize;
+		}
+	}
+
+	return MeshSize;
+}
+
+double FAvaVectorPropertyTypeCustomization::PixelSizeToMeshSize(double PixelSize) const
+{
+	if (ViewportClient.IsValid())
+	{
+		double MeshSize;
+
+		if (FAvaEditorViewportUtils::PixelSizeToMeshSize(ViewportClient.Pin().ToSharedRef(), PixelSize, MeshSize))
+		{
+			return MeshSize;
+		}
+	}
+
+	return PixelSize;
+}
+
+#undef LOCTEXT_NAMESPACE
diff --git a/Engine/Plugins/Experimental/Avalanche/Source/AvalancheEditor/Private/DetailView/Customizations/AvaVectorPropertyTypeCustomization.h b/Engine/Plugins/Experimental/Avalanche/Source/AvalancheEditor/Private/DetailView/Customizations/AvaVectorPropertyTypeCustomization.h
new file mode 100644
--- /dev/null
+++ b/Engine/Plugins/Experimental/Avalanche/Source/AvalancheEditor/Private/DetailView/Customizations/AvaVectorPropertyTypeCustomization.h
@@ -0,0 +1,140 @@
+// Copyright Epic Games, Inc. All Rights Reserved.
+
+#pragma once
+
+#include "CoreMinimal.h"
+#include "IPropertyTypeCustomization.h"
+#include "PropertyEditorModule.h"
+#include "PropertyHandle.h"
+#include "Types/SlateEnums.h"
+#include "Widgets/Input/SVectorInputBox.h"
+#include "AvaVectorPropertyTypeCustomization.generated.h"
+
+class IAvaViewportClient;
+class SButton;
+class SComboButton;
+
+class FAvaVectorPropertyTypeIdentifier : public IPropertyTypeIdentifier
+{
+public:
+	virtual bool IsPropertyTypeCustomized(const IPropertyHandle& InPropertyHandle) const override
+	{
+		static const TArray<FName> PropertyMetaTags = { 
+			TEXT("AllowPreserveRatio"), 
+			TEXT("VectorRatioMode"), 
+			TEXT("Delta")
+		};
+
+		for (const FName& MetaTag : PropertyMetaTags)
+		{
+			if (InPropertyHandle.HasMetaData(MetaTag))
+			{
+				return true;
+			}
+		}
+
+		return false;
+	}
+};
+
+UENUM()
+enum class ERatioMode : uint8
+{
+	None = 0,                        // Free
+	X = 1 << 0,
+	Y = 1 << 1,
+	Z = 1 << 2,
+	PreserveXY = X | Y,              // Lock XY
+	PreserveYZ = Y | Z,              // Lock YZ  (3D)
+	PreserveXZ = X | Z,              // Lock XZ  (3D)
+	PreserveXYZ = X | Y | Z          // Lock XYZ (3D)
+};
+
+ENUM_CLASS_FLAGS(ERatioMode)
+
+class FAvaVectorPropertyTypeCustomization : public IPropertyTypeCustomization
+{
+public:
+	typedef SNumericVectorInputBox<double, UE::Math::TVector2<double>, 2> SNumericVectorInputBox2D;
+	typedef SNumericVectorInputBox<FVector::FReal, UE::Math::TVector<FVector::FReal>, 3> SNumericVectorInputBox3D;
+
+	static constexpr uint8 MULTI_OBJECT_DEBOUNCE  = 3;
+	static constexpr uint8 SINGLE_OBJECT_DEBOUNCE = 2;
+	static constexpr uint8 INVALID_COMPONENT_IDX  = 5;
+
+	static TSharedRef<IPropertyTypeCustomization> MakeInstance()
+	{
+		return MakeShared<FAvaVectorPropertyTypeCustomization>();
+	}
+
+	explicit FAvaVectorPropertyTypeCustomization()
+	{
+	}
+
+	// BEGIN IPropertyTypeCustomization interface
+	virtual void CustomizeHeader(TSharedRef<IPropertyHandle> StructPropertyHandle,
+		FDetailWidgetRow& HeaderRow,
+		IPropertyTypeCustomizationUtils& StructCustomizationUtils) override;
+	virtual void CustomizeChildren(TSharedRef<IPropertyHandle> StructPropertyHandle,
+		IDetailChildrenBuilder& StructBuilder,
+		IPropertyTypeCustomizationUtils& StructCustomizationUtils) override;
+	// END IPropertyTypeCustomization interface
+	
+protected:
+	TWeakPtr<IAvaViewportClient> ViewportClient;
+
+	// used to close the dropdown menu
+	TSharedPtr<SComboButton> ComboButton;
+	
+	TSharedPtr<IPropertyHandle> VectorPropertyHandle;
+	TSharedPtr<IPropertyHandle> XPropertyHandle;
+	TSharedPtr<IPropertyHandle> YPropertyHandle;
+	TSharedPtr<IPropertyHandle> ZPropertyHandle;
+
+	// optional begin values to compute ratios change
+	TArray<TOptional<FVector>> Begin3DValues;
+	TArray<TOptional<FVector2D>> Begin2DValues;
+	
+	int32 SelectedObjectNum = 0;
+	uint8 DebounceValueSet = 0;
+	uint8 LastComponentValueSet = INDEX_NONE;
+	bool bMovingSlider = false;
+	bool bIsVector3d = false;
+	ERatioMode RatioMode = ERatioMode::None;
+	// specific case to handle that needs conversion
+	bool bPixelSizeProperty = false;
+
+	// optional clamp values
+	TOptional<FVector> MinVectorClamp;
+	TOptional<FVector> MaxVectorClamp;
+	TOptional<FVector2D> MinVector2DClamp;
+	TOptional<FVector2D> MaxVector2DClamp;
+
+	TOptional<double> GetVectorComponent(const uint8 Component) const;
+	void SetVectorComponent(double NewValue, const uint8 Component);
+	void SetVectorComponent(double NewValue, ETextCommit::Type CommitType, const uint8 Component);
+
+	void OnBeginSliderMovement();
+	void OnEndSliderMovement(double NewValue);
+
+	FReply OnComboButtonClicked(const ERatioMode NewMode);
+	const FSlateBrush* GetComboButtonBrush(const ERatioMode Mode) const;
+	FText GetComboButtonText(const ERatioMode Mode) const;
+	const FSlateBrush* GetCurrentComboButtonBrush() const;
+	FText GetCurrentComboButtonText() const;
+
+	bool CanEditValue() const;
+	void InitVectorValuesForRatio();
+	void ResetVectorValuesForRatio();
+
+	void SetComponentValue(const double NewValue, const uint8 Component, const EPropertyValueSetFlags::Type Flags);
+
+	// get the correct clamped ratio if a component hit min/max value
+	double GetClampedRatioValueChange(const int32 ObjectIdx, const double NewValue, const uint8 Component, const TArray<bool>& PreserveRatios) const;
+	// get the new clamped value if a component original value is zero since ratio * 0 = 0
+	double GetClampedComponentValue(const int32 ObjectIdx, double NewValue, const double Ratio, const uint8 ComponentIdx, const uint8 OriginalComponent);
+
+	// special case for the pixel property only available in editor
+	double MeshSizeToPixelSize(double MeshSize) const;
+	double PixelSizeToMeshSize(double PixelSize) const;
+};
diff --git a/Engine/Plugins/Experimental/Avalanche/Source/AvalancheEditor/Private/DetailView/Customizations/AvaViewportQualitySettingsPropertyTypeCustomization.cpp b/Engine/Plugins/Experimental/Avalanche/Source/AvalancheEditor/Private/DetailView/Customizations/AvaViewportQualitySettingsPropertyTypeCustomization.cpp
new file mode 100644
--- /dev/null
+++ b/Engine/Plugins/Experimental/Avalanche/Source/AvalancheEditor/Private/DetailView/Customizations/AvaViewportQualitySettingsPropertyTypeCustomization.cpp
@@ -0,0 +1,62 @@
+// Copyright Epic Games, Inc. All Rights Reserved.
+
+#include "AvaViewportQualitySettingsPropertyTypeCustomization.h"
+#include "DetailLayoutBuilder.h"
+#include "DetailWidgetRow.h"
+#include "IDetailChildrenBuilder.h"
+#include "PropertyCustomizationHelpers.h"
+#include "PropertyHandle.h"
+#include "Viewport/AvaViewportQualitySettings.h"
+#include "Widgets/Text/STextBlock.h"
+
+#define LOCTEXT_NAMESPACE "AvaViewportQualitySettingsPropertyTypeCustomization"
+
+TSharedRef<IPropertyTypeCustomization> FAvaViewportQualitySettingsPropertyTypeCustomization::MakeInstance()
+{
+	return MakeShared<FAvaViewportQualitySettingsPropertyTypeCustomization>();
+}
+
+void FAvaViewportQualitySettingsPropertyTypeCustomization::CustomizeHeader(TSharedRef<IPropertyHandle> InStructPropertyHandle, FDetailWidgetRow& InHeaderRow, IPropertyTypeCustomizationUtils& StructCustomizationUtils)
+{
+	InHeaderRow.NameContent()
+		[
+			InStructPropertyHandle->CreatePropertyNameWidget()
+		];
+}
+
+void FAvaViewportQualitySettingsPropertyTypeCustomization::CustomizeChildren(TSharedRef<IPropertyHandle> InStructPropertyHandle, IDetailChildrenBuilder& InDetailBuilder, IPropertyTypeCustomizationUtils& InStructCustomizationUtils)
+{
+	TSharedPtr<IPropertyHandle> FeaturesProperty = InStructPropertyHandle->GetChildHandle(GET_MEMBER_NAME_CHECKED(FAvaViewportQualitySettings, Features));
+
+	const TSharedRef<FDetailArrayBuilder> FeaturesArrayBuilder = MakeShared<FDetailArrayBuilder>(FeaturesProperty.ToSharedRef(), /*InGenerateHeader*/ false, /*InDisplayResetToDefault*/ true, /*InDisplayElementNum*/ false);
+
+	FeaturesArrayBuilder->OnGenerateArrayElementWidget(FOnGenerateArrayElementWidget::CreateLambda([](TSharedRef<IPropertyHandle> InElementPropertyHandle, int32 InArrayIndex, IDetailChildrenBuilder& InChildrenBuilder)
+	{
+		TSharedPtr<IPropertyHandle> NameProperty  = InElementPropertyHandle->GetChildHandle(GET_MEMBER_NAME_CHECKED(FAvaViewportQualitySettingsFeature, Name));
+		TSharedPtr<IPropertyHandle> ValueProperty = InElementPropertyHandle->GetChildHandle(GET_MEMBER_NAME_CHECKED(FAvaViewportQualitySettingsFeature, Enabled));
+
+		FString FeatureName;
+		NameProperty->GetValue(FeatureName);
+
+		FText NameText, TooltipText;
+		FAvaViewportQualitySettings::FeatureNameAndTooltipText(FeatureName, NameText, TooltipText);
+
+		InChildrenBuilder.AddProperty(InElementPropertyHandle)
+			.ToolTip(TooltipText)
+			.CustomWidget()
+			.NameContent()
+			[
+				SNew(STextBlock)
+				.Text(NameText)
+				.Font(IDetailLayoutBuilder::GetDetailFont())
+			]
+			.ValueContent()
+			[
+				ValueProperty->CreatePropertyValueWidget(/*bDisplayPropertyButtons*/false)
+			];
+	}));
+
+	InDetailBuilder.AddCustomBuilder(FeaturesArrayBuilder);
+}
+
+#undef LOCTEXT_NAMESPACE
diff --git a/Engine/Plugins/Experimental/Avalanche/Source/AvalancheEditor/Private/DetailView/Customizations/AvaViewportQualitySettingsPropertyTypeCustomization.h b/Engine/Plugins/Experimental/Avalanche/Source/AvalancheEditor/Private/DetailView/Customizations/AvaViewportQualitySettingsPropertyTypeCustomization.h
new file mode 100644
--- /dev/null
+++ b/Engine/Plugins/Experimental/Avalanche/Source/AvalancheEditor/Private/DetailView/Customizations/AvaViewportQualitySettingsPropertyTypeCustomization.h
@@ -0,0 +1,16 @@
+// Copyright Epic Games, Inc. All Rights Reserved.
+
+#pragma once
+
+#include "IPropertyTypeCustomization.h"
+
+class FAvaViewportQualitySettingsPropertyTypeCustomization : public IPropertyTypeCustomization
+{
+public:
+	static TSharedRef<IPropertyTypeCustomization> MakeInstance();
+
+	//~ Begin IPropertyTypeCustomization
+	virtual void CustomizeHeader(TSharedRef<IPropertyHandle> InStructPropertyHandle, FDetailWidgetRow& InHeaderRow, IPropertyTypeCustomizationUtils& InStructCustomizationUtils) override;
+	virtual void CustomizeChildren(TSharedRef<IPropertyHandle> InStructPropertyHandle, IDetailChildrenBuilder& InDetailBuilder, IPropertyTypeCustomizationUtils& InStructCustomizationUtils) override;
+	//~ End IPropertyTypeCustomization
+};
diff --git a/Engine/Plugins/Experimental/Avalanche/Source/AvalancheEditor/Private/OperatorStack/AvaOperatorStackExtension.cpp b/Engine/Plugins/Experimental/Avalanche/Source/AvalancheEditor/Private/OperatorStack/AvaOperatorStackExtension.cpp
new file mode 100644
--- /dev/null
+++ b/Engine/Plugins/Experimental/Avalanche/Source/AvalancheEditor/Private/OperatorStack/AvaOperatorStackExtension.cpp
@@ -0,0 +1,204 @@
+// Copyright Epic Games, Inc. All Rights Reserved.
+
+#include "AvaOperatorStackExtension.h"
+#include "AvaEditorCommands.h"
+#include "AvaOperatorStackTabSpawner.h"
+#include "Components/PropertyAnimatorCoreComponent.h"
+#include "EditorModeManager.h"
+#include "Framework/Docking/LayoutExtender.h"
+#include "LevelEditor.h"
+#include "Modifiers/ActorModifierCoreBase.h"
+#include "Modifiers/ActorModifierCoreComponent.h"
+#include "Selection.h"
+#include "Subsystems/OperatorStackEditorSubsystem.h"
+#include "Subsystems/PropertyAnimatorCoreSubsystem.h"
+#include "ToolMenu.h"
+#include "ToolMenuEntry.h"
+#include "ToolMenuSection.h"
+#include "Widgets/SAvaOperatorStackTab.h"
+#include "Widgets/SOperatorStackEditorWidget.h"
+
+#define LOCTEXT_NAMESPACE "AvaOperatorStackExtension"
+
+FAvaOperatorStackExtension::FAvaOperatorStackExtension()
+	: AnimatorCommands(MakeShared<FUICommandList>())
+{
+}
+
+TSharedPtr<SDockTab> FAvaOperatorStackExtension::OpenTab()
+{
+	const TSharedPtr<IAvaEditor>& Editor = GetEditor();
+
+	if (!Editor.IsValid())
+	{
+		return nullptr;
+	}
+
+	TSharedPtr<FTabManager> TabManager = Editor->GetTabManager();
+
+	if (!TabManager.IsValid())
+	{
+		return nullptr;
+	}
+
+	return TabManager->TryInvokeTab(FAvaOperatorStackTabSpawner::GetTabID());
+}
+
+void FAvaOperatorStackExtension::Activate()
+{
+	FAvaEditorExtension::Activate();
+
+	USelection::SelectionChangedEvent.AddSP(this, &FAvaOperatorStackExtension::OnSelectionChanged);
+}
+
+void FAvaOperatorStackExtension::Deactivate()
+{
+	FAvaEditorExtension::Deactivate();
+
+	USelection::SelectionChangedEvent.RemoveAll(this);
+}
+
+void FAvaOperatorStackExtension::RegisterTabSpawners(const TSharedRef<IAvaEditor>& InEditor) const
+{
+	InEditor->AddTabSpawner<FAvaOperatorStackTabSpawner>(InEditor);
+}
+
+void FAvaOperatorStackExtension::ExtendLevelEditorLayout(FLayoutExtender& InExtender) const
+{
+	InExtender.ExtendLayout(LevelEditorTabIds::LevelEditorSceneOutliner
+		, ELayoutExtensionPosition::After
+		, FTabManager::FTab(FAvaOperatorStackTabSpawner::GetTabID(), ETabState::ClosedTab));
+}
+
+void FAvaOperatorStackExtension::ExtendToolbarMenu(UToolMenu& InMenu)
+{
+	FToolMenuSection& Section = InMenu.FindOrAddSection(DefaultSectionName);
+
+	FToolMenuEntry& Entry = Section.AddEntry(FToolMenuEntry::InitToolBarButton(
+		TEXT("OpenOperatorStackButton"),
+		FExecuteAction::CreateSPLambda(this, [this]() { OpenTab(); }),
+		LOCTEXT("OpenOperatorStackLabel", "Operator Stack"),
+		LOCTEXT("OpenOperatorStackTooltip", "Open the operator stack tab."),
+		FSlateIcon(FAppStyle::GetAppStyleSetName(), "ClassIcon.UserDefinedStruct")
+	));
+
+	Entry.StyleNameOverride = "CalloutToolbar";
+}
+
+void FAvaOperatorStackExtension::BindCommands(const TSharedRef<FUICommandList>& InCommandList)
+{
+	InCommandList->Append(AnimatorCommands);
+
+	const FAvaEditorCommands& EditorCommands = FAvaEditorCommands::Get();
+
+	AnimatorCommands->MapAction(EditorCommands.DisableAnimators
+		, FExecuteAction::CreateSP(this, &FAvaOperatorStackExtension::EnableAnimators, false));
+
+	AnimatorCommands->MapAction(EditorCommands.EnableAnimators
+		, FExecuteAction::CreateSP(this, &FAvaOperatorStackExtension::EnableAnimators, true));
+}
+
+void FAvaOperatorStackExtension::OnSelectionChanged(UObject* InSelection)
+{
+	const USelection* Selection = Cast<USelection>(InSelection);
+
+	if (!Selection || Selection->Num() == 0)
+	{
+		return;
+	}
+
+	const UObject* LastSelectedObject = Selection->GetSelectedObject(Selection->Num() - 1);
+
+	if (!LastSelectedObject)
+	{
+		return;
+	}
+
+	const UOperatorStackEditorSubsystem* OperatorStackEditorSubsystem = UOperatorStackEditorSubsystem::Get();
+
+	if (!OperatorStackEditorSubsystem)
+	{
+		return;
+	}
+
+	FName ActiveCustomization;
+	if (LastSelectedObject->IsA<UActorModifierCoreBase>()
+		|| LastSelectedObject->IsA<UActorModifierCoreComponent>())
+	{
+		ActiveCustomization = TEXT("Modifiers");
+	}
+	else if (LastSelectedObject->IsA<UPropertyAnimatorCoreBase>()
+		|| LastSelectedObject->IsA<UPropertyAnimatorCoreComponent>())
+	{
+		ActiveCustomization = TEXT("Animators");
+	}
+
+	if (ActiveCustomization.IsNone())
+	{
+		return;
+	}
+
+	// Lets invoke the operator stack tab
+	const TSharedPtr<SDockTab> InvokedTab = OpenTab();
+
+	if (!InvokedTab.IsValid())
+	{
+		return;
+	}
+
+	// Find the avalanche operator stack and set the active customization
+	OperatorStackEditorSubsystem->ForEachCustomizationWidget([ActiveCustomization](TSharedRef<SOperatorStackEditorWidget> InWidget)->bool
+	{
+		if (InWidget->GetPanelTag() == SAvaOperatorStackTab::PanelTag)
+		{
+			InWidget->SetActiveCustomization(ActiveCustomization);
+
+			// Stop looping for customization widgets
+			return false;
+		}
+
+		return true;
+	});
+}
+
+void FAvaOperatorStackExtension::EnableAnimators(bool bInEnable) const
+{
+	const UWorld* const World = GetWorld();
+	if (!IsValid(World))
+	{
+		return;
+	}
+
+	const FEditorModeTools* ModeTools = GetEditorModeTools();
+	if (!ModeTools)
+	{
+		return;
+	}
+
+	const UTypedElementSelectionSet* SelectionSet = ModeTools->GetEditorSelectionSet();
+	if (!SelectionSet)
+	{
+		return;
+	}
+
+	const TSet<AActor*> SelectedActors(SelectionSet->GetSelectedObjects<AActor>());
+
+	UPropertyAnimatorCoreSubsystem* AnimatorSubsystem = UPropertyAnimatorCoreSubsystem::Get();
+
+	if (!AnimatorSubsystem)
+	{
+		return;
+	}
+
+	// If nothing selected, target all animators in level
+	if (!SelectedActors.IsEmpty())
+	{
+		AnimatorSubsystem->SetActorAnimatorsEnabled(SelectedActors, bInEnable, true);
+	}
+	else
+	{
+		AnimatorSubsystem->SetLevelAnimatorsEnabled(World, bInEnable, true);
+	}
+}
+
+#undef LOCTEXT_NAMESPACE
diff --git a/Engine/Plugins/Experimental/Avalanche/Source/AvalancheEditor/Private/OperatorStack/AvaOperatorStackTabSpawner.cpp b/Engine/Plugins/Experimental/Avalanche/Source/AvalancheEditor/Private/OperatorStack/AvaOperatorStackTabSpawner.cpp
new file mode 100644
--- /dev/null
+++ b/Engine/Plugins/Experimental/Avalanche/Source/AvalancheEditor/Private/OperatorStack/AvaOperatorStackTabSpawner.cpp
@@ -0,0 +1,41 @@
+﻿// Copyright Epic Games, Inc. All Rights Reserved.
+
+#include "AvaOperatorStackTabSpawner.h"
+#include "DetailView/AvaDetailsExtension.h"
+#include "OperatorStack/Widgets/SAvaOperatorStackTab.h"
+#include "Styling/AppStyle.h"
+#include "Styling/SlateIconFinder.h"
+
+#define LOCTEXT_NAMESPACE "AvaOperatorStackTabSpawner"
+
+FName FAvaOperatorStackTabSpawner::GetTabID()
+{
+	return TEXT("AvalancheOperatorStack");
+}
+
+FAvaOperatorStackTabSpawner::FAvaOperatorStackTabSpawner(const TSharedRef<IAvaEditor>& InEditor)
+	: FAvaTabSpawner(InEditor, FAvaOperatorStackTabSpawner::GetTabID())
+{
+	TabLabel       = LOCTEXT("TabLabel", "Operator Stack");
+	TabTooltipText = LOCTEXT("TabTooltip", "Operator Stack");
+	TabIcon        = FSlateIcon(FAppStyle::GetAppStyleSetName(), "ClassIcon.UserDefinedStruct");
+}
+
+TSharedRef<SWidget> FAvaOperatorStackTabSpawner::CreateTabBody()
+{
+	const TSharedPtr<IAvaEditor> Editor = EditorWeak.Pin();
+	if (!Editor.IsValid())
+	{
+		return GetNullWidget();
+	}
+
+	const TSharedPtr<FAvaDetailsExtension> DetailsExtension = Editor->FindExtension<FAvaDetailsExtension>();
+	if (!DetailsExtension.IsValid())
+	{
+		return GetNullWidget();
+	}
+
+	return SNew(SAvaOperatorStackTab, DetailsExtension.ToSharedRef());
+}
+
+#undef LOCTEXT_NAMESPACE
\ No newline at end of file
diff --git a/Engine/Plugins/Experimental/Avalanche/Source/AvalancheEditor/Private/OperatorStack/AvaOperatorStackTabSpawner.h b/Engine/Plugins/Experimental/Avalanche/Source/AvalancheEditor/Private/OperatorStack/AvaOperatorStackTabSpawner.h
new file mode 100644
--- /dev/null
+++ b/Engine/Plugins/Experimental/Avalanche/Source/AvalancheEditor/Private/OperatorStack/AvaOperatorStackTabSpawner.h
@@ -0,0 +1,17 @@
+﻿// Copyright Epic Games, Inc. All Rights Reserved.
+
+#pragma once
+
+#include "AvaTabSpawner.h"
+
+class FAvaOperatorStackTabSpawner : public FAvaTabSpawner
+{
+public:
+	static FName GetTabID();
+
+	explicit FAvaOperatorStackTabSpawner(const TSharedRef<IAvaEditor>& InEditor);
+
+	//~ Begin IAvaTabSpawner
+	virtual TSharedRef<SWidget> CreateTabBody() override;
+	//~ End IAvaTabSpawner
+};
\ No newline at end of file
diff --git a/Engine/Plugins/Experimental/Avalanche/Source/AvalancheEditor/Private/OperatorStack/Widgets/SAvaOperatorStackTab.cpp b/Engine/Plugins/Experimental/Avalanche/Source/AvalancheEditor/Private/OperatorStack/Widgets/SAvaOperatorStackTab.cpp
new file mode 100644
--- /dev/null
+++ b/Engine/Plugins/Experimental/Avalanche/Source/AvalancheEditor/Private/OperatorStack/Widgets/SAvaOperatorStackTab.cpp
@@ -0,0 +1,150 @@
+// Copyright Epic Games, Inc. All Rights Reserved.
+
+#include "SAvaOperatorStackTab.h"
+
+#include "Animators/PropertyAnimatorCoreBase.h"
+#include "Contexts/OperatorStackEditorContext.h"
+#include "DetailView/IAvaDetailsProvider.h"
+#include "EditorModeManager.h"
+#include "Items/OperatorStackEditorItem.h"
+#include "Items/OperatorStackEditorObjectItem.h"
+#include "Modifiers/ActorModifierCoreStack.h"
+#include "Selection.h"
+#include "Selection/AvaEditorSelection.h"
+#include "Subsystems/OperatorStackEditorSubsystem.h"
+#include "Widgets/SOperatorStackEditorWidget.h"
+
+#define LOCTEXT_NAMESPACE "SAvaModifierStackTab"
+
+void SAvaOperatorStackTab::Construct(const FArguments& InArgs
+	, const TSharedPtr<IAvaDetailsProvider>& InProvider)
+{
+	DetailsProviderWeak = InProvider;
+	
+	UOperatorStackEditorSubsystem* OperatorStackSubsystem = UOperatorStackEditorSubsystem::Get();
+
+	check(InProvider.IsValid() && OperatorStackSubsystem);
+
+	USelection::SelectionChangedEvent.AddSP(this, &SAvaOperatorStackTab::RefreshSelection);
+
+	// Modifiers delegates
+	UActorModifierCoreStack::OnModifierAddedDelegate.AddSP(this, &SAvaOperatorStackTab::OnModifierUpdated);
+	UActorModifierCoreStack::OnModifierMovedDelegate.AddSP(this, &SAvaOperatorStackTab::OnModifierUpdated);
+	UActorModifierCoreStack::OnModifierRemovedDelegate.AddSP(this, &SAvaOperatorStackTab::OnModifierUpdated);
+
+	// Property controllers delegates
+	UPropertyAnimatorCoreBase::OnAnimatorCreatedDelegate.AddSP(this, &SAvaOperatorStackTab::OnControllerUpdated);
+	UPropertyAnimatorCoreBase::OnAnimatorRemovedDelegate.AddSP(this, &SAvaOperatorStackTab::OnControllerUpdated);
+	UPropertyAnimatorCoreBase::OnAnimatorRenamedDelegate.AddSP(this, &SAvaOperatorStackTab::OnControllerUpdated);
+
+	const TSharedPtr<IDetailKeyframeHandler> KeyframeHandler = InProvider->GetDetailsKeyframeHandler();
+	
+	OperatorStack = OperatorStackSubsystem->GenerateWidget();
+	OperatorStack->SetKeyframeHandler(KeyframeHandler);
+	OperatorStack->SetPanelTag(SAvaOperatorStackTab::PanelTag);
+	
+	ChildSlot
+	[
+		OperatorStack.ToSharedRef()
+	];
+
+	RefreshSelection(nullptr);
+}
+
+SAvaOperatorStackTab::~SAvaOperatorStackTab()
+{
+	USelection::SelectionChangedEvent.RemoveAll(this);
+	
+	UActorModifierCoreStack::OnModifierAddedDelegate.RemoveAll(this);
+	UActorModifierCoreStack::OnModifierMovedDelegate.RemoveAll(this);
+	UActorModifierCoreStack::OnModifierRemovedDelegate.RemoveAll(this);
+	
+	UPropertyAnimatorCoreBase::OnAnimatorCreatedDelegate.RemoveAll(this);
+	UPropertyAnimatorCoreBase::OnAnimatorRemovedDelegate.RemoveAll(this);
+	UPropertyAnimatorCoreBase::OnAnimatorRenamedDelegate.RemoveAll(this);
+}
+
+void SAvaOperatorStackTab::RefreshSelection(UObject* InSelectionObject) const
+{
+	const TSharedPtr<IAvaDetailsProvider> DetailsProvider = DetailsProviderWeak.Pin();
+	if (!DetailsProvider.IsValid())
+	{
+		return;
+	}
+	
+	FEditorModeTools* ModeTools = DetailsProvider->GetDetailsModeTools();
+	if (!ModeTools)
+	{
+		return;
+	}
+
+	const FAvaEditorSelection EditorSelection(*ModeTools, InSelectionObject);
+	if (!EditorSelection.IsValid())
+	{
+		return;
+	}
+
+	const TArray<UObject*> SelectedObjects = EditorSelection.GetSelectedObjects<UObject, EAvaSelectionSource::All>();
+
+	TArray<FOperatorStackEditorItemPtr> SelectedItems;
+	Algo::Transform(SelectedObjects, SelectedItems, [](UObject* InObject)
+	{
+		return MakeShared<FOperatorStackEditorObjectItem>(InObject);
+	});
+
+	const FOperatorStackEditorContext Context(SelectedItems);
+	OperatorStack->SetContext(Context);
+}
+
+void SAvaOperatorStackTab::OnModifierUpdated(UActorModifierCoreBase* InUpdatedItem) const
+{
+	RefreshCurrentSelection(InUpdatedItem);
+}
+
+void SAvaOperatorStackTab::OnControllerUpdated(UPropertyAnimatorCoreBase* InController) const
+{
+	RefreshCurrentSelection(InController);
+}
+
+void SAvaOperatorStackTab::RefreshCurrentSelection(const UObject* InObject) const
+{
+	const TSharedPtr<IAvaDetailsProvider> DetailsProvider = DetailsProviderWeak.Pin();
+	if (!DetailsProvider.IsValid())
+	{
+		return;
+	}
+	
+	FEditorModeTools* ModeTools = DetailsProvider->GetDetailsModeTools();
+	if (!InObject || !OperatorStack.IsValid() || !ModeTools)
+	{
+		return;
+	}
+
+	const AActor* OwningActor = InObject->GetTypedOuter<AActor>();
+	if (!OwningActor)
+	{
+		return;
+	}
+
+	const FOperatorStackEditorContextPtr Context = OperatorStack->GetContext();
+	if (!Context.IsValid())
+	{
+		return;
+	}
+
+	const FAvaEditorSelection EditorSelection(*ModeTools, ModeTools->GetSelectedActors());
+	if (!EditorSelection.IsValid())
+	{
+		return;
+	}
+
+	const TArray<AActor*> SelectedActors = EditorSelection.GetSelectedObjects<AActor, EAvaSelectionSource::All>();
+	if (!SelectedActors.Contains(OwningActor))
+	{
+		return;
+	}
+
+	OperatorStack->RefreshContext();
+}
+
+#undef LOCTEXT_NAMESPACE
\ No newline at end of file
diff --git a/Engine/Plugins/Experimental/Avalanche/Source/AvalancheEditor/Private/OperatorStack/Widgets/SAvaOperatorStackTab.h b/Engine/Plugins/Experimental/Avalanche/Source/AvalancheEditor/Private/OperatorStack/Widgets/SAvaOperatorStackTab.h
new file mode 100644
--- /dev/null
+++ b/Engine/Plugins/Experimental/Avalanche/Source/AvalancheEditor/Private/OperatorStack/Widgets/SAvaOperatorStackTab.h
@@ -0,0 +1,38 @@
+﻿// Copyright Epic Games, Inc. All Rights Reserved.
+
+#pragma once
+
+#include "Templates/SharedPointer.h"
+#include "Widgets/SCompoundWidget.h"
+
+class FEditorModeTools;
+class IAvaDetailsProvider;
+class IDetailKeyframeHandler;
+class SOperatorStackEditorWidget;
+class UPropertyAnimatorCoreBase;
+class UObject;
+class UActorModifierCoreBase;
+
+/** Widget for the operator stack */
+class SAvaOperatorStackTab : public SCompoundWidget
+{
+public:
+	static inline const FName PanelTag = TEXT("AvaOperatorStackTab");
+	
+	SLATE_BEGIN_ARGS(SAvaOperatorStackTab) {}
+	SLATE_END_ARGS()
+
+	void Construct(const FArguments& InArgs, const TSharedPtr<IAvaDetailsProvider>& InProvider);
+
+	virtual ~SAvaOperatorStackTab() override;
+
+private:
+	void RefreshSelection(UObject* InSelectionObject) const;
+	void OnModifierUpdated(UActorModifierCoreBase* InUpdatedItem) const;
+	void OnControllerUpdated(UPropertyAnimatorCoreBase* InController) const;
+	void RefreshCurrentSelection(const UObject* InObject) const;
+	
+	TWeakPtr<IAvaDetailsProvider> DetailsProviderWeak;
+
+	TSharedPtr<SOperatorStackEditorWidget> OperatorStack;
+};
\ No newline at end of file
diff --git a/Engine/Plugins/Experimental/Avalanche/Source/AvalancheEditor/Private/Outliner/AvaOutlinerTabSpawner.cpp b/Engine/Plugins/Experimental/Avalanche/Source/AvalancheEditor/Private/Outliner/AvaOutlinerTabSpawner.cpp
new file mode 100644
--- /dev/null
+++ b/Engine/Plugins/Experimental/Avalanche/Source/AvalancheEditor/Private/Outliner/AvaOutlinerTabSpawner.cpp
@@ -0,0 +1,91 @@
+﻿// Copyright Epic Games, Inc. All Rights Reserved.
+
+#include "AvaOutlinerTabSpawner.h"
+#include "AvaOutliner.h"
+#include "AvaOutlinerExtension.h"
+#include "AvaOutlinerView.h"
+#include "Framework/Docking/TabManager.h"
+#include "Framework/Docking/WorkspaceItem.h"
+#include "IAvaEditor.h"
+#include "Styling/AppStyle.h"
+
+#define LOCTEXT_NAMESPACE "AvaOutlinerTabSpawner"
+
+FName FAvaOutlinerTabSpawner::GetTabID(int32 InOutlinerId)
+{
+	return *FString::Printf(TEXT("AvalancheOutliner%d"), InOutlinerId + 1);
+}
+
+FAvaOutlinerTabSpawner::FAvaOutlinerTabSpawner(int32 InOutlinerId, const TSharedRef<IAvaEditor>& InEditor)
+	: FAvaTabSpawner(InEditor, FAvaOutlinerTabSpawner::GetTabID(InOutlinerId))
+	, OutlinerViewId(InOutlinerId)
+{
+	TabLabel = (OutlinerViewId == 0)
+		? LOCTEXT("FirstTabLabel", "Motion Design Outliner")
+		: FText::Format(LOCTEXT("OtherTabLabels", "Motion Design Outliner {0}"), FText::AsNumber(OutlinerViewId + 1));
+
+	TabTooltipText = LOCTEXT("TabTooltip", "Motion Design Outliner");
+
+	TabIcon = FSlateIcon(FAppStyle::GetAppStyleSetName(), "LevelEditor.Tabs.Outliner");
+}
+
+TSharedRef<SWidget> FAvaOutlinerTabSpawner::CreateTabBody()
+{
+	const TSharedPtr<IAvaEditor> Editor = EditorWeak.Pin();
+	if (!ensure(Editor.IsValid()))
+	{
+		return GetNullWidget();
+	}
+
+	const TSharedPtr<FAvaOutlinerExtension> OutlinerExtension = Editor->FindExtension<FAvaOutlinerExtension>();
+	if (!ensure(OutlinerExtension.IsValid()))
+	{
+		return GetNullWidget();
+	}
+
+	const TSharedPtr<FAvaOutliner> Outliner = OutlinerExtension->GetAvaOutliner();
+	if (!ensure(Outliner.IsValid()))
+	{
+		return GetNullWidget();
+	}
+
+	TSharedPtr<FAvaOutlinerView> OutlinerView = Outliner->GetOutlinerView(OutlinerViewId);
+	if (!OutlinerView.IsValid())
+	{
+		OutlinerView = Outliner->RegisterOutlinerView(OutlinerViewId);
+	}
+
+	if (!ensure(OutlinerView.IsValid()))
+	{
+		return GetNullWidget();
+	}
+	return OutlinerView->GetOutlinerWidget();
+}
+
+FTabSpawnerEntry& FAvaOutlinerTabSpawner::RegisterTabSpawner(const TSharedRef<FTabManager>& InTabManager, const TSharedPtr<FWorkspaceItem>& InWorkspaceMenu)
+{
+	TSharedPtr<FWorkspaceItem> GroupItem;
+
+	const FText OutlinerGroupName = LOCTEXT("SubMenuLabel", "Motion Design Outliner");
+
+	for (const TSharedRef<FWorkspaceItem>& Item : InWorkspaceMenu->GetChildItems())
+	{
+		if (Item->GetDisplayName().ToString() == OutlinerGroupName.ToString())
+		{
+			GroupItem = Item;
+			break;
+		}
+	}
+
+	if (!GroupItem.IsValid())
+	{
+		GroupItem = InWorkspaceMenu->AddGroup(OutlinerGroupName
+			, LOCTEXT("SubMenuTooltip", "Open a new Motion Design Outliner")
+			, FSlateIcon(FAppStyle::GetAppStyleSetName(), "LevelEditor.Tabs.Outliner"));
+	}
+
+	return FAvaTabSpawner::RegisterTabSpawner(InTabManager, nullptr)
+		.SetGroup(GroupItem.ToSharedRef());
+}
+
+#undef LOCTEXT_NAMESPACE
diff --git a/Engine/Plugins/Experimental/Avalanche/Source/AvalancheEditor/Private/Palettes/AvaPaletteExtension.cpp b/Engine/Plugins/Experimental/Avalanche/Source/AvalancheEditor/Private/Palettes/AvaPaletteExtension.cpp
new file mode 100644
--- /dev/null
+++ b/Engine/Plugins/Experimental/Avalanche/Source/AvalancheEditor/Private/Palettes/AvaPaletteExtension.cpp
@@ -0,0 +1,162 @@
+// Copyright Epic Games, Inc. All Rights Reserved.
+
+#include "Palettes/AvaPaletteExtension.h"
+#include "AvaPaletteTabSpawner.h"
+#include "Editor.h"
+#include "Framework/Docking/LayoutExtender.h"
+#include "LevelEditor.h"
+#include "Selection/AvaEditorSelection.h"
+#include "Widgets/AvaViewportColorPickerActorClassRegistry.h"
+#include "Widgets/AvaViewportColorPickerDelegates.h"
+
+FAvaPaletteExtension::~FAvaPaletteExtension()
+{
+	FAvaViewportColorPickerDelegates::GetOnColorPicked().RemoveAll(this);
+}
+
+void FAvaPaletteExtension::Construct(const TSharedRef<IAvaEditor>& InEditor)
+{
+	Super::Construct(InEditor);
+
+	RegisterPalettes();
+}
+
+void FAvaPaletteExtension::Activate()
+{
+	Super::Activate();
+
+	FAvaViewportColorPickerDelegates::GetOnColorPicked().AddSP(this, &FAvaPaletteExtension::OnColorPicked);
+}
+
+void FAvaPaletteExtension::Deactivate()
+{
+	Super::Deactivate();
+
+	FAvaViewportColorPickerDelegates::GetOnColorPicked().RemoveAll(this);
+}
+
+void FAvaPaletteExtension::RegisterTabSpawners(const TSharedRef<IAvaEditor>& InEditor) const
+{
+	Super::RegisterTabSpawners(InEditor);
+
+	for (const FAvaPaletteTabInfo& PaletteTabInfo : RegisteredPalettes)
+	{
+		InEditor->AddTabSpawner<FAvaPaletteTabSpawner>(InEditor, PaletteTabInfo);
+	}
+}
+
+void FAvaPaletteExtension::ExtendLevelEditorLayout(FLayoutExtender& InExtender) const
+{
+	Super::ExtendLevelEditorLayout(InExtender);
+
+	FName PreviousTab = LevelEditorTabIds::PlacementBrowser;
+	ELayoutExtensionPosition Placement = ELayoutExtensionPosition::Below;
+
+	for (const FAvaPaletteTabInfo& PaletteTabInfo : RegisteredPalettes)
+	{
+		InExtender.ExtendLayout(
+			PreviousTab
+			, Placement
+			, FTabManager::FTab(PaletteTabInfo.Id, ETabState::OpenedTab)
+		);
+
+		PreviousTab = PaletteTabInfo.Id;
+		Placement = ELayoutExtensionPosition::After;
+	}
+}
+
+void FAvaPaletteExtension::NotifyOnSelectionChanged(const FAvaEditorSelection& InSelection)
+{
+	Super::NotifyOnSelectionChanged(InSelection);
+
+	TSharedPtr<IToolkitHost> ToolkitHost = GetToolkitHost();
+
+	if (!ToolkitHost.IsValid())
+	{
+		return;
+	}
+
+	TArray<AActor*> SelectedActors = InSelection.GetSelectedObjects<AActor>();
+
+	if (SelectedActors.Num() == 1)
+	{
+		const AActor* const ColorSource = SelectedActors[0];
+		FAvaColorChangeData NewColorData;
+
+		if (FAvaViewportColorPickerActorClassRegistry::GetColorDataFromActor(ColorSource, NewColorData))
+		{
+			FAvaViewportColorPickerDelegates::BroadcastColorSourceSelected(ToolkitHost.ToSharedRef(), NewColorData);
+		}
+	}
+}
+
+bool FAvaPaletteExtension::RegisterPalette(const FAvaPaletteTabInfo& InPaletteTabInfo)
+{
+	const bool bAlreadyRegistered = RegisteredPalettes.ContainsByPredicate(
+		[&InPaletteTabInfo](const FAvaPaletteTabInfo& InRegisteredPalette)
+		{
+			return InPaletteTabInfo.Id == InRegisteredPalette.Id;
+		}
+	);
+
+	if (!ensureMsgf(!bAlreadyRegistered, TEXT("Motion Design palette already registered: %s"), *InPaletteTabInfo.Id.ToString()))
+	{
+		return false;
+	}
+
+	RegisteredPalettes.Add(InPaletteTabInfo);
+
+	return true;
+}
+
+void FAvaPaletteExtension::RegisterPalettes()
+{
+}
+
+void FAvaPaletteExtension::OnColorPicked(const TSharedRef<IToolkitHost>& InEditor, const FAvaColorChangeData& InNewColorData)
+{
+	TSharedPtr<IAvaEditor> Editor = GetEditor();
+
+	if (!Editor.IsValid())
+	{
+		return;
+	}
+
+	if (Editor->GetToolkitHost() != InEditor)
+	{
+		return;
+	}
+
+	FEditorModeTools* ModeTools = Editor->GetEditorModeTools();
+
+	if (!ModeTools)
+	{
+		return;
+	}
+
+	USelection* ActorSelection = ModeTools->GetSelectedActors();
+
+	if (!ActorSelection)
+	{
+		return;
+	}
+
+	TArray<AActor*> SelectedActors;
+	ActorSelection->GetSelectedObjects(SelectedActors);
+	bool bHasColourReceiver = false;
+
+	for (AActor* Actor : SelectedActors)
+	{
+		if (FAvaViewportColorPickerActorClassRegistry::ApplyColorDataToActor(Actor, InNewColorData))
+		{
+			bHasColourReceiver = true;
+		}
+	}
+
+	if (!bHasColourReceiver)
+	{
+		return;
+	}
+
+	GEditor->RedrawLevelEditingViewports();
+}
diff --git a/Engine/Plugins/Experimental/Avalanche/Source/AvalancheEditor/Private/Palettes/AvaPaletteExtension.h b/Engine/Plugins/Experimental/Avalanche/Source/AvalancheEditor/Private/Palettes/AvaPaletteExtension.h
new file mode 100644
--- /dev/null
+++ b/Engine/Plugins/Experimental/Avalanche/Source/AvalancheEditor/Private/Palettes/AvaPaletteExtension.h
@@ -0,0 +1,68 @@
+// Copyright Epic Games, Inc. All Rights Reserved.
+
+#pragma once
+
+#include "AvaType.h"
+#include "Containers/ContainersFwd.h"
+#include "IAvaEditorExtension.h"
+#include "Internationalization/Text.h"
+#include "Templates/SharedPointer.h"
+#include "Textures/SlateIcon.h"
+
+class FEditorViewportClient;
+class SAvaColorPickerPalette;
+class SAvaPaletteBase;
+class SWidget;
+struct FAvaColorChangeData;
+
+DECLARE_DELEGATE_RetVal_OneParam(TSharedRef<SAvaPaletteBase>, FAvaOnPaletteConstruct, const TSharedRef<IToolkitHost>&)
+
+struct FAvaPaletteInfo
+{
+	FText Title;
+	FAvaOnPaletteConstruct ConstructionCallback;
+};
+
+struct FAvaPaletteTabInfo
+{
+	FName Id;
+	FText Name;
+	FText ToolTip;
+	FSlateIcon Icon;
+	FAvaOnPaletteConstruct ConstructionCallback;
+
+	FAvaPaletteInfo GetPaletteInfo() const
+	{
+		return {
+			Name,
+			ConstructionCallback
+		};
+	}
+};
+
+class FAvaPaletteExtension : public FAvaEditorExtension
+{
+public:
+	UE_AVA_INHERITS_WITH_SUPER(FAvaPaletteExtension, FAvaEditorExtension);
+
+	virtual ~FAvaPaletteExtension() override;
+
+	//~ Begin IAvaEditorExtension
+	virtual void Construct(const TSharedRef<IAvaEditor>& InEditor) override;
+	virtual void Activate() override;
+	virtual void Deactivate() override;
+	virtual void RegisterTabSpawners(const TSharedRef<IAvaEditor>& InEditor) const override;
+	virtual void ExtendLevelEditorLayout(FLayoutExtender& InExtender) const override;
+	virtual void NotifyOnSelectionChanged(const FAvaEditorSelection& InSelection) override;
+	//~ End IAvaEditorExtension
+
+	bool RegisterPalette(const FAvaPaletteTabInfo& InPaletteTabInfo);
+	TConstArrayView<FAvaPaletteTabInfo> GetRegisteredPalettes() const { return RegisteredPalettes; }
+
+protected:
+	TArray<FAvaPaletteTabInfo> RegisteredPalettes;
+
+	void RegisterPalettes();
+
+	void OnColorPicked(const TSharedRef<IToolkitHost>& InEditor, const FAvaColorChangeData& InNewColorData);
+};
diff --git a/Engine/Plugins/Experimental/Avalanche/Source/AvalancheEditor/Private/Palettes/AvaPaletteTabSpawner.cpp b/Engine/Plugins/Experimental/Avalanche/Source/AvalancheEditor/Private/Palettes/AvaPaletteTabSpawner.cpp
new file mode 100644
--- /dev/null
+++ b/Engine/Plugins/Experimental/Avalanche/Source/AvalancheEditor/Private/Palettes/AvaPaletteTabSpawner.cpp
@@ -0,0 +1,82 @@
+// Copyright Epic Games, Inc. All Rights Reserved.
+
+#include "Palettes/AvaPaletteTabSpawner.h"
+#include "AvaPaletteExtension.h"
+#include "Framework/Docking/TabManager.h"
+#include "Framework/Docking/WorkspaceItem.h"
+#include "Styling/AppStyle.h"
+#include "Widgets/SAvaPaletteBase.h"
+
+#define LOCTEXT_NAMESPACE "AvaPaletteTabSpawner"
+
+FAvaPaletteTabSpawner::FAvaPaletteTabSpawner(const TSharedRef<IAvaEditor>& InEditor, const FAvaPaletteTabInfo& InPaletteTabInfo)
+	: FAvaTabSpawner(InEditor, InPaletteTabInfo.Id)
+	, PaletteTabInfo(InPaletteTabInfo)
+{
+	TabLabel = InPaletteTabInfo.Name;
+	TabTooltipText = InPaletteTabInfo.ToolTip.IsEmpty() ? InPaletteTabInfo.Name : InPaletteTabInfo.ToolTip;
+	TabIcon = InPaletteTabInfo.Icon;
+}
+
+FName FAvaPaletteTabSpawner::GetTabID()
+{
+	return PaletteTabInfo.Id;
+}
+
+TSharedRef<SWidget> FAvaPaletteTabSpawner::CreateTabBody()
+{
+	if (!PaletteTabInfo.ConstructionCallback.IsBound())
+	{
+		return GetNullWidget();
+	}
+
+	TSharedPtr<IAvaEditor> Editor = EditorWeak.Pin();
+
+	if (!Editor.IsValid())
+	{
+		return GetNullWidget();
+	}
+
+	TSharedPtr<IToolkitHost> ToolkitHost = Editor->GetToolkitHost();
+
+	if (!ToolkitHost.IsValid())
+	{
+		return GetNullWidget();
+	}
+
+	TSharedRef<SAvaPaletteBase> PaletteWidget = PaletteTabInfo.ConstructionCallback.Execute(ToolkitHost.ToSharedRef());
+	PaletteWidget->OnCreatedInEditor(EditorWeak.Pin());
+
+	return PaletteWidget;
+}
+
+FTabSpawnerEntry& FAvaPaletteTabSpawner::RegisterTabSpawner(const TSharedRef<FTabManager>& InTabManager,
+	const TSharedPtr<FWorkspaceItem>& InWorkspaceMenu)
+{
+	TSharedPtr<FWorkspaceItem> GroupItem;
+
+	const FText PalettesGroupName = LOCTEXT("SubMenuLabel", "Palettes");
+
+	for (const TSharedRef<FWorkspaceItem>& Item : InWorkspaceMenu->GetChildItems())
+	{
+		if (Item->GetDisplayName().ToString() == PalettesGroupName.ToString())
+		{
+			GroupItem = Item;
+			break;
+		}
+	}
+
+	if (!GroupItem.IsValid())
+	{
+		GroupItem = InWorkspaceMenu->AddGroup(
+			PalettesGroupName
+			, LOCTEXT("SubMenuTooltip", "Open an Avalanche palette")
+			, FSlateIcon(FSlateIcon(FAppStyle::GetAppStyleSetName(), "Icons.FolderClosed"))
+		);
+	}
+
+	return FAvaTabSpawner::RegisterTabSpawner(InTabManager, nullptr)
+		.SetGroup(GroupItem.ToSharedRef());
+}
+
+#undef LOCTEXT_NAMESPACE
diff --git a/Engine/Plugins/Experimental/Avalanche/Source/AvalancheEditor/Private/Palettes/AvaPaletteTabSpawner.h b/Engine/Plugins/Experimental/Avalanche/Source/AvalancheEditor/Private/Palettes/AvaPaletteTabSpawner.h
new file mode 100644
--- /dev/null
+++ b/Engine/Plugins/Experimental/Avalanche/Source/AvalancheEditor/Private/Palettes/AvaPaletteTabSpawner.h
@@ -0,0 +1,23 @@
+// Copyright Epic Games, Inc. All Rights Reserved.
+
+#pragma once
+
+#include "AvaPaletteExtension.h"
+#include "AvaTabSpawner.h"
+
+class FAvaPaletteTabSpawner : public FAvaTabSpawner
+{
+public:
+	explicit FAvaPaletteTabSpawner(const TSharedRef<IAvaEditor>& InEditor, const FAvaPaletteTabInfo& InPaletteTabInfo);
+
+	FName GetTabID();
+
+	//~ Begin IAvaTabSpawner
+	virtual TSharedRef<SWidget> CreateTabBody() override;
+	virtual FTabSpawnerEntry& RegisterTabSpawner(const TSharedRef<FTabManager>& InTabManager, 
+		const TSharedPtr<FWorkspaceItem>& InWorkspaceMenu) override;
+	//~ End IAvaTabSpawner
+
+protected:
+	const FAvaPaletteTabInfo PaletteTabInfo;
+};
diff --git a/Engine/Plugins/Experimental/Avalanche/Source/AvalancheEditor/Private/Palettes/Widgets/SAvaPaletteBase.cpp b/Engine/Plugins/Experimental/Avalanche/Source/AvalancheEditor/Private/Palettes/Widgets/SAvaPaletteBase.cpp
new file mode 100644
--- /dev/null
+++ b/Engine/Plugins/Experimental/Avalanche/Source/AvalancheEditor/Private/Palettes/Widgets/SAvaPaletteBase.cpp
@@ -0,0 +1,8 @@
+// Copyright Epic Games, Inc. All Rights Reserved.
+
+#include "SAvaPaletteBase.h"
+#include "Palettes/AvaPaletteExtension.h"
+
+void SAvaPaletteBase::PrivateRegisterAttributes(struct FSlateAttributeDescriptor::FInitializer&)
+{
+}
diff --git a/Engine/Plugins/Experimental/Avalanche/Source/AvalancheEditor/Private/Palettes/Widgets/SAvaPaletteBase.h b/Engine/Plugins/Experimental/Avalanche/Source/AvalancheEditor/Private/Palettes/Widgets/SAvaPaletteBase.h
new file mode 100644
--- /dev/null
+++ b/Engine/Plugins/Experimental/Avalanche/Source/AvalancheEditor/Private/Palettes/Widgets/SAvaPaletteBase.h
@@ -0,0 +1,36 @@
+// Copyright Epic Games, Inc. All Rights Reserved.
+
+#pragma once
+
+#include "Templates/SharedPointer.h"
+#include "Widgets/SCompoundWidget.h"
+
+class IAvaEditor;
+class IToolkitHost;
+struct FAvaPaletteInfo;
+struct FAvaPaletteTabInfo;
+
+class SAvaPaletteBase : public SCompoundWidget
+{
+public:
+	SLATE_DECLARE_WIDGET(SAvaPaletteBase, SCompoundWidget)
+
+	SLATE_BEGIN_ARGS(SAvaPaletteBase) {}
+	SLATE_END_ARGS()
+
+	void Construct(const FArguments& Args, const TSharedRef<IToolkitHost>& InToolkitHost)
+	{
+		ToolkitHostWeak = InToolkitHost;
+	}
+
+	virtual const FAvaPaletteTabInfo& GetPaletteTabInfo() const = 0;
+
+	TSharedPtr<IToolkitHost> GetToolkitHost() const { return ToolkitHostWeak.Pin(); }
+
+	virtual TSharedPtr<IAvaEditor> GetEditor() const { return EditorWeak.Pin(); }
+	virtual void OnCreatedInEditor(TSharedPtr<IAvaEditor> InEditor) { EditorWeak = InEditor; }
+
+protected:
+	TWeakPtr<IToolkitHost> ToolkitHostWeak;
+	TWeakPtr<IAvaEditor> EditorWeak;
+};
diff --git a/Engine/Plugins/Experimental/Avalanche/Source/AvalancheEditor/Private/Transition/AvaTransitionPlaylistExtension.cpp b/Engine/Plugins/Experimental/Avalanche/Source/AvalancheEditor/Private/Transition/AvaTransitionPlaylistExtension.cpp
new file mode 100644
--- /dev/null
+++ b/Engine/Plugins/Experimental/Avalanche/Source/AvalancheEditor/Private/Transition/AvaTransitionPlaylistExtension.cpp
@@ -0,0 +1,177 @@
+// Copyright Epic Games, Inc. All Rights Reserved.
+
+#include "AvaTransitionPlaylistExtension.h"
+#include "AvaTransitionLog.h"
+#include "AvaTransitionSubsystem.h"
+#include "AvaTransitionTree.h"
+#include "Behavior/IAvaTransitionBehavior.h"
+#include "IAvaMediaEditorModule.h"
+#include "Playback/AvaMediaPlaybackUtils.h"
+#include "Playlist/AvaPlaylistEditor.h"
+#include "Playlist/AvalanchePlaylist.h"
+#include "Playlist/Pages/AvaPlaylistPageContext.h"
+#include "Styling/SlateIconFinder.h"
+#include "ToolMenu.h"
+#include "ToolMenuSection.h"
+#include "ToolMenus.h"
+
+#define LOCTEXT_NAMESPACE "AvaTransitionPlaylistExtension"
+
+FAvaTransitionPlaylistExtension::FAvaTransitionPlaylistExtension()
+	: ExtensionSectionName(TEXT("AvaTransitionPlaylistExtension"))
+{
+}
+
+FAvaTransitionPlaylistExtension::~FAvaTransitionPlaylistExtension()
+{
+	Shutdown();
+}
+
+void FAvaTransitionPlaylistExtension::Startup()
+{
+	UToolMenus* ToolMenus = UToolMenus::Get();
+	check(ToolMenus);
+
+	UToolMenu* PageContextMenu = ToolMenus->ExtendMenu(IAvaMediaEditorModule::GetPlaylistPageMenuName());
+	check(PageContextMenu);
+
+	PageContextMenuWeak = PageContextMenu;
+
+	PageContextMenu->AddDynamicSection(ExtensionSectionName
+		, FNewToolMenuDelegate::CreateRaw(this, &FAvaTransitionPlaylistExtension::ExtendPageContextMenu)
+		, FToolMenuInsert(TEXT("PageListOperations"), EToolMenuInsertType::After));
+}
+
+void FAvaTransitionPlaylistExtension::Shutdown()
+{
+	if (UToolMenu* ToolMenu = PageContextMenuWeak.Get())
+	{
+		ToolMenu->RemoveSection(ExtensionSectionName);
+		PageContextMenuWeak = nullptr;
+	}
+}
+
+void FAvaTransitionPlaylistExtension::ExtendPageContextMenu(UToolMenu* InMenu)
+{
+	if (!InMenu)
+	{
+		return;
+	}
+
+	UAvaPlaylistPageContext* PageContext = InMenu->FindContext<UAvaPlaylistPageContext>();
+	if (!PageContext)
+	{
+		return;
+	}
+
+	TWeakPtr<FAvaPlaylistEditor> PlaylistEditorWeak = PageContext->GetPlaylistEditor();
+	if (!PlaylistEditorWeak.IsValid())
+	{
+		return;
+	}
+
+	FToolMenuSection& Section = InMenu->AddSection(TEXT("TransitionLogicSection"), LOCTEXT("TransitionLogicSectionLabel", "Transition Logic"));
+	Section.AddMenuEntry(TEXT("OpenTransitionLogic")
+		, LOCTEXT("OpenTransitionLogicLabel", "Open Transition Logic")
+		, LOCTEXT("OpenTransitionLogicTooltip", "Opens the Transition Logic Editor for the selected pages")
+		, FSlateIconFinder::FindIconForClass(UAvaTransitionTree::StaticClass())
+		, FExecuteAction::CreateRaw(this, &FAvaTransitionPlaylistExtension::OpenTransitionTree, PlaylistEditorWeak));
+}
+
+void FAvaTransitionPlaylistExtension::OpenTransitionTree(TWeakPtr<FAvaPlaylistEditor> InPlaylistEditorWeak)
+{
+	TSharedPtr<FAvaPlaylistEditor> PlaylistEditor = InPlaylistEditorWeak.Pin();
+	if (!PlaylistEditor.IsValid())
+	{
+		UE_LOG(LogAvaTransition, Warning, TEXT("Unable to open Transition Logic. Playlist Editor is invalid"));
+		return;
+	}
+
+	UAvalanchePlaylist* Playlist = PlaylistEditor->GetPlaylist();
+	if (!Playlist)
+	{
+		UE_LOG(LogAvaTransition, Warning, TEXT("Unable to open Transition Logic. Playlist is invalid"));
+		return;
+	}
+
+	TConstArrayView<int32> SelectedPageIds = PlaylistEditor->GetSelectedPagesOnFocusedWidget();
+	if (SelectedPageIds.IsEmpty())
+	{
+		UE_LOG(LogAvaTransition, Log, TEXT("Playlist '%s' did not have selected pages on focused widget."), *Playlist->GetName());
+		return;
+	}
+
+	const FString PlaylistName = Playlist->GetName();
+
+	for (int32 PageId : SelectedPageIds)
+	{
+		const FAvalanchePage& Page = Playlist->GetPage(PageId);
+		if (!Page.IsValidPage())
+		{
+			return;
+		}
+
+		TArray<FSoftObjectPath> PageAssetPaths = Page.GetAvalancheAssetPaths(Playlist);
+		for (const FSoftObjectPath& PageAssetPath : PageAssetPaths)
+		{
+			OpenTransitionTree(PageAssetPath, *PlaylistName, PageId);
+		}
+	}
+}
+
+void FAvaTransitionPlaylistExtension::OpenTransitionTree(const FSoftObjectPath& InPageAssetPath, const TCHAR* InPlaylistName, int32 InPageId)
+{
+	if (!FAvaMediaPlaybackUtils::IsMapAsset(InPageAssetPath.GetLongPackageName()))
+	{
+		UE_LOG(LogAvaTransition, Warning
+			, TEXT("Page asset path '%s' (Page Id: '%d', Playlist '%s) is not a map asset and will not be processed")
+			, *InPageAssetPath.ToString()
+			, InPageId
+			, InPlaylistName);
+		return;
+	}
+
+	UWorld* const World = Cast<UWorld>(InPageAssetPath.TryLoad());
+	if (!World)
+	{
+		UE_LOG(LogAvaTransition, Warning
+			, TEXT("World path '%s' (Page Id: '%d', Playlist '%s) did not load a valid world and will not be processed")
+			, *InPageAssetPath.ToString()
+			, InPageId
+			, InPlaylistName);
+		return;
+	}
+
+	IAvaTransitionBehavior* TransitionBehavior = UAvaTransitionSubsystem::FindTransitionBehavior(World->PersistentLevel);
+	if (!TransitionBehavior)
+	{
+		UE_LOG(LogAvaTransition, Warning
+			, TEXT("World '%s' (Page Id: '%d', Playlist '%s) does not have a valid Transition Behavior")
+			, *InPageAssetPath.ToString()
+			, InPageId
+			, InPlaylistName);
+		return;
+	}
+
+	UAvaTransitionTree* TransitionTree = TransitionBehavior->GetTransitionTree();
+	if (!TransitionTree)
+	{
+		UE_LOG(LogAvaTransition, Warning
+			, TEXT("Transition Behavior of World '%s' (Page Id: '%d', Playlist '%s) did not have a valid Transition Tree")
+			, *InPageAssetPath.ToString()
+			, InPageId
+			, InPlaylistName);
+		return;
+	}
+
+	UAssetEditorSubsystem* AssetEditorSubsystem = GEditor->GetEditorSubsystem<UAssetEditorSubsystem>();
+	check(AssetEditorSubsystem);
+
+	AssetEditorSubsystem->OpenEditorForAsset(TransitionTree
+		, EToolkitMode::Standalone
+		, TSharedPtr<IToolkitHost>()
+		, /*bShowProgressWindow*/true
+		, EAssetTypeActivationOpenedMethod::View);
+}
+
+#undef LOCTEXT_NAMESPACE
diff --git a/Engine/Plugins/Experimental/Avalanche/Source/AvalancheEditor/Private/Transition/AvaTransitionPlaylistExtension.h b/Engine/Plugins/Experimental/Avalanche/Source/AvalancheEditor/Private/Transition/AvaTransitionPlaylistExtension.h
new file mode 100644
--- /dev/null
+++ b/Engine/Plugins/Experimental/Avalanche/Source/AvalancheEditor/Private/Transition/AvaTransitionPlaylistExtension.h
@@ -0,0 +1,33 @@
+// Copyright Epic Games, Inc. All Rights Reserved.
+
+#pragma once
+
+#include "UObject/NameTypes.h"
+#include "UObject/WeakObjectPtr.h"
+
+class FAvaPlaylistEditor;
+class UToolMenu;
+class UAvalanchePlaylist;
+
+class FAvaTransitionPlaylistExtension
+{
+public:
+	FAvaTransitionPlaylistExtension();
+
+	~FAvaTransitionPlaylistExtension();
+
+	void Startup();
+
+	void Shutdown();
+
+private:
+	void ExtendPageContextMenu(UToolMenu* InMenu);
+
+	void OpenTransitionTree(TWeakPtr<FAvaPlaylistEditor> InPlaylistEditorWeak);
+
+	void OpenTransitionTree(const FSoftObjectPath& InPageAssetPath, const TCHAR* InPlaylistName, int32 InPageId);
+
+	TWeakObjectPtr<UToolMenu> PageContextMenuWeak;
+
+	const FName ExtensionSectionName;
+};
diff --git a/Engine/Plugins/Experimental/Avalanche/Source/AvalancheEditor/Public/AvaEditorFunctionLibrary.h b/Engine/Plugins/Experimental/Avalanche/Source/AvalancheEditor/Public/AvaEditorFunctionLibrary.h
new file mode 100644
--- /dev/null
+++ b/Engine/Plugins/Experimental/Avalanche/Source/AvalancheEditor/Public/AvaEditorFunctionLibrary.h
@@ -0,0 +1,19 @@
+// Copyright Epic Games, Inc. All Rights Reserved.
+
+#pragma once
+
+#include "Containers/Array.h"
+#include "Misc/EngineVersionComparison.h"
+#include "UObject/WeakObjectPtrFwd.h"
+
+class UAvalancheBlueprint;
+
+class FAvaEditorFunctionLibrary
+{
+public:
+	/**
+	 * Creates World Assets based on the Avalanche Blueprints provided
+	 * @param InWeakBlueprints the Avalanche Blueprints to export
+	 */
+	AVALANCHEEDITOR_API static void ExportAvaBlueprintsToWorld(TArray<TWeakObjectPtr<UAvalancheBlueprint>> InWeakBlueprints);
+};
diff --git a/Engine/Plugins/Experimental/Avalanche/Source/AvalancheEditorCore/Private/LevelEditor/AvaLevelEditorStyle.cpp b/Engine/Plugins/Experimental/Avalanche/Source/AvalancheEditorCore/Private/LevelEditor/AvaLevelEditorStyle.cpp
new file mode 100644
--- /dev/null
+++ b/Engine/Plugins/Experimental/Avalanche/Source/AvalancheEditorCore/Private/LevelEditor/AvaLevelEditorStyle.cpp
@@ -0,0 +1,26 @@
+// Copyright Epic Games, Inc. All Rights Reserved.
+
+#include "AvaLevelEditorStyle.h"
+#include "Misc/Paths.h"
+#include "Styling/SlateStyleMacros.h"
+#include "Styling/SlateStyleRegistry.h"
+
+FAvaLevelEditorStyle::FAvaLevelEditorStyle()
+	: FSlateStyleSet(TEXT("AvaLevelEditor"))
+{
+	const FVector2f Icon16x16(16.f, 16.f);
+
+	ContentRootDir     = FPaths::EngineContentDir() / TEXT("Editor/Slate");
+	CoreContentRootDir = FPaths::EngineContentDir() / TEXT("Slate");
+
+	Set("AvaLevelEditor.CreateScene"    , new CORE_IMAGE_BRUSH_SVG("Starship/Common/plus", Icon16x16));
+	Set("AvaLevelEditor.ActivateScene"  , new CORE_IMAGE_BRUSH_SVG("Starship/Common/play", Icon16x16));
+	Set("AvaLevelEditor.DeactivateScene", new CORE_IMAGE_BRUSH_SVG("Starship/Common/stop", Icon16x16));
+
+	FSlateStyleRegistry::RegisterSlateStyle(*this);
+}
+
+FAvaLevelEditorStyle::~FAvaLevelEditorStyle()
+{
+	FSlateStyleRegistry::UnRegisterSlateStyle(*this);
+}
diff --git a/Engine/Plugins/Experimental/Avalanche/Source/AvalancheEditorCore/Private/LevelEditor/AvaLevelEditorStyle.h b/Engine/Plugins/Experimental/Avalanche/Source/AvalancheEditorCore/Private/LevelEditor/AvaLevelEditorStyle.h
new file mode 100644
--- /dev/null
+++ b/Engine/Plugins/Experimental/Avalanche/Source/AvalancheEditorCore/Private/LevelEditor/AvaLevelEditorStyle.h
@@ -0,0 +1,19 @@
+// Copyright Epic Games, Inc. All Rights Reserved.
+
+#pragma once
+
+#include "Styling/SlateStyle.h"
+
+class FAvaLevelEditorStyle : public FSlateStyleSet
+{
+public:
+	FAvaLevelEditorStyle();
+
+	virtual ~FAvaLevelEditorStyle() override;
+
+	static FAvaLevelEditorStyle& Get()
+	{
+		static FAvaLevelEditorStyle StyleSet;
+		return StyleSet;
+	}
+};
diff --git a/Engine/Plugins/Experimental/Avalanche/Source/AvalancheEditorCore/Private/Widgets/InputDataTypes/AvaUserInputDataText.cpp b/Engine/Plugins/Experimental/Avalanche/Source/AvalancheEditorCore/Private/Widgets/InputDataTypes/AvaUserInputDataText.cpp
new file mode 100644
--- /dev/null
+++ b/Engine/Plugins/Experimental/Avalanche/Source/AvalancheEditorCore/Private/Widgets/InputDataTypes/AvaUserInputDataText.cpp
@@ -0,0 +1,104 @@
+// Copyright Epic Games, Inc. All Rights Reserved.
+
+#include "Widgets/InputDataTypes/AvaUserInputDataText.h"
+#include "Widgets/Input/SEditableText.h"
+#include "Widgets/Input/SEditableTextBox.h"
+#include "Widgets/Layout/SBorder.h"
+#include "Widgets/Layout/SBox.h"
+#include "Widgets/SBoxPanel.h"
+#include "Widgets/SOverlay.h"
+#include "Widgets/Text/SMultiLineEditableText.h"
+
+#define LOCTEXT_NAMESPACE "FAvaUserInputTextData"
+
+FAvaUserInputTextData::FAvaUserInputTextData(const FText& InValue, bool bInAllowMultiline, TOptional<int32> InMaxLength)
+	: Value(InValue)
+	, bAllowMultiline(bInAllowMultiline)
+	, MaxLength(InMaxLength)
+{
+}
+
+const FText& FAvaUserInputTextData::GetValue() const
+{
+	return Value;
+}
+
+TSharedRef<SWidget> FAvaUserInputTextData::CreateInputWidget()
+{
+	if (bAllowMultiline)
+	{
+		TSharedRef<SScrollBar> ScrollBar = SNew(SScrollBar)
+			.Orientation(Orient_Vertical);
+
+		return SNew(SBox)
+			.WidthOverride(500.f)
+			.HeightOverride(200.f)
+			[
+				SNew(SOverlay)
+				+ SOverlay::Slot()
+				[
+					SNew(SBorder)
+					.BorderImage(&FCoreStyle::Get().GetWidgetStyle<FEditableTextBoxStyle>("NormalEditableTextBox").BackgroundImageFocused)
+				]
+				+ SOverlay::Slot()
+				[
+					SNew(SHorizontalBox)
+					+ SHorizontalBox::Slot()
+					.FillWidth(1.f)
+					.Padding(5.f)
+					[
+						SNew(SMultiLineEditableText)
+						.Text(Value)
+						.OnTextChanged(this, &FAvaUserInputTextData::OnTextChanged)
+						.OnTextCommitted(this, &FAvaUserInputTextData::OnTextCommitted)
+						.AllowMultiLine(true)
+						.VScrollBar(ScrollBar)
+					]
+					+ SHorizontalBox::Slot()
+					.AutoWidth()
+					[
+						ScrollBar
+					]
+				]
+			];
+	}
+
+	return SNew(SBox)
+		.WidthOverride(FMath::Max(200.f, MaxLength.IsSet() ? MaxLength.GetValue() * 8.f : 200.f))
+		.HAlign(HAlign_Fill)
+		[
+			SNew(SEditableTextBox)
+			.Text(Value)
+			.OnTextChanged(this, &FAvaUserInputTextData::OnTextChanged)
+			.OnTextCommitted(this, &FAvaUserInputTextData::OnTextCommitted)
+			.OnVerifyTextChanged(this, &FAvaUserInputTextData::OnTextVerify)
+		];
+}
+
+void FAvaUserInputTextData::OnTextChanged(const FText& InValue)
+{
+	Value = InValue;
+}
+
+void FAvaUserInputTextData::OnTextCommitted(const FText& InValue, ETextCommit::Type InCommitType)
+{
+	Value = InValue;
+
+	if (!bAllowMultiline && InCommitType == ETextCommit::OnEnter)
+	{
+		OnUserCommit();
+	}
+}
+
+bool FAvaUserInputTextData::OnTextVerify(const FText& InValue, FText& OutErrorText)
+{
+	if (MaxLength.IsSet() && InValue.ToString().Len() > MaxLength.GetValue())
+	{
+		OutErrorText = LOCTEXT("InputTooLong", "Input too long");
+		return false;
+	}
+
+	return true;
+}
+
+#undef LOCTEXT_NAMESPACE
diff --git a/Engine/Plugins/Experimental/Avalanche/Source/AvalancheEditorCore/Private/Widgets/InputDataTypes/AvaUserInputDataTypeBase.cpp b/Engine/Plugins/Experimental/Avalanche/Source/AvalancheEditorCore/Private/Widgets/InputDataTypes/AvaUserInputDataTypeBase.cpp
new file mode 100644
--- /dev/null
+++ b/Engine/Plugins/Experimental/Avalanche/Source/AvalancheEditorCore/Private/Widgets/InputDataTypes/AvaUserInputDataTypeBase.cpp
@@ -0,0 +1,12 @@
+// Copyright Epic Games, Inc. All Rights Reserved.
+
+#include "Widgets/InputDataTypes/AvaUserInputDataTypeBase.h"
+
+FAvaUserInputDataTypeBase::FAvaUserInputDataTypeBase()
+{
+}
+
+void FAvaUserInputDataTypeBase::OnUserCommit()
+{
+	OnCommit.ExecuteIfBound();
+}
diff --git a/Engine/Plugins/Experimental/Avalanche/Source/AvalancheEditorCore/Public/Widgets/InputDataTypes/AvaUserInputDataNumber.h b/Engine/Plugins/Experimental/Avalanche/Source/AvalancheEditorCore/Public/Widgets/InputDataTypes/AvaUserInputDataNumber.h
new file mode 100644
--- /dev/null
+++ b/Engine/Plugins/Experimental/Avalanche/Source/AvalancheEditorCore/Public/Widgets/InputDataTypes/AvaUserInputDataNumber.h
@@ -0,0 +1,59 @@
+// Copyright Epic Games, Inc. All Rights Reserved.
+
+#pragma once
+
+#include "Widgets/InputDataTypes/AvaUserInputDataTypeBase.h"
+#include "Misc/Optional.h"
+#include "Widgets/Input/SSpinBox.h"
+#include "Widgets/Layout/SBox.h"
+
+template<typename InNumberType>
+struct FAvaUserInputNumberData : public FAvaUserInputDataTypeBase
+{
+	using ThisStruct = FAvaUserInputNumberData<InNumberType>;
+
+	FAvaUserInputNumberData(InNumberType InValue, TOptional<InNumberType> InMinValue = TOptional<InNumberType>(),
+		TOptional<InNumberType> InMaxValue = TOptional<InNumberType>())
+		: Value(InValue)
+		, MinValue(InMinValue)
+		, MaxValue(InMaxValue)
+	{
+	}
+
+	InNumberType GetValue() const { return Value; }
+
+	//~ Begin FAvaUserInputDataTypeBase
+	virtual TSharedRef<SWidget> CreateInputWidget()
+	{
+		return SNew(SBox)
+			.WidthOverride(200.f)
+			.HAlign(HAlign_Fill)
+			[
+				SNew(SSpinBox<InNumberType>)
+				.Value(Value)
+				.MinValue(MinValue)
+				.MaxSliderValue(MinValue)
+				.MaxValue(MaxValue)
+				.MaxSliderValue(MaxValue)
+				.EnableSlider(true)
+				.OnValueChanged(this, &ThisStruct::OnValueChanged)
+				.OnValueCommitted(this, &ThisStruct::OnValueCommitted)
+			];
+	}
+	//~ End FAvaUserInputDataTypeBase
+
+protected:
+	InNumberType Value;
+	TOptional<InNumberType> MinValue;
+	TOptional<InNumberType> MaxValue;
+
+	void OnValueChanged(InNumberType InValue)
+	{
+		Value = InValue;
+	}
+
+	void OnValueCommitted(InNumberType InValue, ETextCommit::Type InCommitType)
+	{
+		Value = InValue;
+	}
+};
\ No newline at end of file
diff --git a/Engine/Plugins/Experimental/Avalanche/Source/AvalancheEditorCore/Public/Widgets/InputDataTypes/AvaUserInputDataText.h b/Engine/Plugins/Experimental/Avalanche/Source/AvalancheEditorCore/Public/Widgets/InputDataTypes/AvaUserInputDataText.h
new file mode 100644
--- /dev/null
+++ b/Engine/Plugins/Experimental/Avalanche/Source/AvalancheEditorCore/Public/Widgets/InputDataTypes/AvaUserInputDataText.h
@@ -0,0 +1,33 @@
+// Copyright Epic Games, Inc. All Rights Reserved.
+
+#pragma once
+
+#include "Widgets/InputDataTypes/AvaUserInputDataTypeBase.h"
+#include "Internationalization/Text.h"
+#include "Misc/Optional.h"
+
+namespace ETextCommit { enum Type : int; }
+
+struct AVALANCHEEDITORCORE_API FAvaUserInputTextData : public FAvaUserInputDataTypeBase
+{
+	FAvaUserInputTextData(const FText& InValue, bool bInAllowMultiline = false, TOptional<int32> InMaxLength = TOptional<int32>());
+
+	virtual ~FAvaUserInputTextData() override = default;
+
+	const FText& GetValue() const;
+
+	//~ Begin FAvaUserInputDataTypeBase
+	virtual TSharedRef<SWidget> CreateInputWidget();
+	//~ End FAvaUserInputDataTypeBase
+
+protected:
+	FText Value;
+	bool bAllowMultiline;
+	TOptional<int32> MaxLength;
+
+	void OnTextChanged(const FText& InValue);
+
+	void OnTextCommitted(const FText& InValue, ETextCommit::Type InCommitType);
+
+	bool OnTextVerify(const FText& InValue, FText& OutErrorText);
+};
\ No newline at end of file
diff --git a/Engine/Plugins/Experimental/Avalanche/Source/AvalancheEditorCore/Public/Widgets/InputDataTypes/AvaUserInputDataTypeBase.h b/Engine/Plugins/Experimental/Avalanche/Source/AvalancheEditorCore/Public/Widgets/InputDataTypes/AvaUserInputDataTypeBase.h
new file mode 100644
--- /dev/null
+++ b/Engine/Plugins/Experimental/Avalanche/Source/AvalancheEditorCore/Public/Widgets/InputDataTypes/AvaUserInputDataTypeBase.h
@@ -0,0 +1,26 @@
+// Copyright Epic Games, Inc. All Rights Reserved.
+
+#pragma once
+
+#include "Delegates/Delegate.h"
+#include "Delegates/DelegateCombinations.h"
+#include "Templates/SharedPointer.h"
+
+class SWidget;
+
+struct AVALANCHEEDITORCORE_API FAvaUserInputDataTypeBase : public TSharedFromThis<FAvaUserInputDataTypeBase>
+{
+	friend class SAvaUserInputDialog;
+
+	FAvaUserInputDataTypeBase();
+
+	virtual ~FAvaUserInputDataTypeBase() = default;
+
+	virtual TSharedRef<SWidget> CreateInputWidget() = 0;
+
+protected:
+	DECLARE_DELEGATE(FOnCommit)
+	FOnCommit OnCommit;
+
+	void OnUserCommit();
+};
diff --git a/Engine/Plugins/Experimental/Avalanche/Source/AvalancheEditorCore/Public/Widgets/SAvaUserInputDialog.h b/Engine/Plugins/Experimental/Avalanche/Source/AvalancheEditorCore/Public/Widgets/SAvaUserInputDialog.h
new file mode 100644
--- /dev/null
+++ b/Engine/Plugins/Experimental/Avalanche/Source/AvalancheEditorCore/Public/Widgets/SAvaUserInputDialog.h
@@ -0,0 +1,37 @@
+// Copyright Epic Games, Inc. All Rights Reserved.
+
+#pragma once
+
+#include "Widgets/SCompoundWidget.h"
+#include "Widgets/InputDataTypes/AvaUserInputDataTypeBase.h"
+
+class SAvaUserInputDialog : public SCompoundWidget
+{
+	SLATE_DECLARE_WIDGET_API(SAvaUserInputDialog, SCompoundWidget, AVALANCHEEDITORCORE_API)
+
+public:
+	static AVALANCHEEDITORCORE_API bool CreateModalDialog(const TSharedPtr<SWidget>& InParent, const FText& InTitle, const FText& InPrompt,
+		const TSharedRef<FAvaUserInputDataTypeBase>& InInputType);
+
+	SLATE_BEGIN_ARGS(SAvaUserInputDialog) {}
+		SLATE_ARGUMENT(FText, Prompt)
+	SLATE_END_ARGS()
+
+	AVALANCHEEDITORCORE_API void Construct(const FArguments& InArgs, const TSharedRef<FAvaUserInputDataTypeBase>& InInputType);
+
+	AVALANCHEEDITORCORE_API TSharedPtr<FAvaUserInputDataTypeBase> GetInputType() const;
+
+	AVALANCHEEDITORCORE_API bool WasAccepted() const;
+
+private:
+	TSharedPtr<FAvaUserInputDataTypeBase> InputType;
+	bool bAccepted = false;
+
+	FReply OnOkayClicked();
+
+	FReply OnCancelClicked();
+
+	void OnUserCommit();
+
+	void Close(bool bInAccepted);
+};
