[MotionDesign]

diff --git a/Engine/Plugins/Experimental/Avalanche/Source/AvalancheEffectors/Private/AvaClonerEffectorShared.cpp b/Engine/Plugins/Experimental/Avalanche/Source/AvalancheEffectors/Private/AvaClonerEffectorShared.cpp
new file mode 100644
--- /dev/null
+++ b/Engine/Plugins/Experimental/Avalanche/Source/AvalancheEffectors/Private/AvaClonerEffectorShared.cpp
@@ -0,0 +1,204 @@
+// Copyright Epic Games, Inc. All Rights Reserved.
+
+#include "AvaClonerEffectorShared.h"
+
+#include "Components/BillboardComponent.h"
+#include "Engine/Texture2D.h"
+#include "NiagaraDataChannelAccessor.h"
+#include "NiagaraSystem.h"
+#include "PrimitiveSceneProxy.h"
+#include "SceneManagement.h"
+#include "UObject/Package.h"
+#include "UObject/UObjectGlobals.h"
+
+void UE::ClonerEffector::SetBillboardComponentSprite(const AActor* InActor, const FString& InTexturePath)
+{
+	if (!InActor)
+	{
+		return;
+	}
+
+	UBillboardComponent* BillboardComponent = InActor->FindComponentByClass<UBillboardComponent>();
+
+	if (!BillboardComponent)
+	{
+		return;
+	}
+
+	UTexture2D* SpriteTexture = FindObject<UTexture2D>(nullptr, *InTexturePath);
+
+	if (!SpriteTexture)
+	{
+		SpriteTexture = LoadObject<UTexture2D>(nullptr, *InTexturePath);
+	}
+
+	if (SpriteTexture && SpriteTexture != BillboardComponent->Sprite)
+	{
+		BillboardComponent->SetSprite(SpriteTexture);
+	}
+}
+
+void UE::ClonerEffector::SetBillboardComponentVisibility(const AActor* InActor, bool bInVisibility)
+{
+	if (!InActor)
+	{
+		return;
+	}
+
+	UBillboardComponent* BillboardComponent = InActor->FindComponentByClass<UBillboardComponent>();
+
+	if (!BillboardComponent)
+	{
+		return;
+	}
+
+	BillboardComponent->SetVisibility(bInVisibility, false);
+}
+
+void FAvaClonerEffectorChannelData::Write(UNiagaraDataChannelWriter* InWriter) const
+{
+	if (!InWriter)
+	{
+		return;
+	}
+
+	/** General */
+	InWriter->WriteInt(EasingName, Identifier, static_cast<int32>(Easing));
+	InWriter->WriteInt(ModeName, Identifier, static_cast<int32>(Mode));
+	InWriter->WriteInt(TypeName, Identifier, static_cast<int32>(Type));
+	InWriter->WriteFloat(MagnitudeName, Identifier, Magnitude);
+	InWriter->WriteVector(InnerExtentName, Identifier, InnerExtent);
+	InWriter->WriteVector(OuterExtentName, Identifier, OuterExtent);
+	InWriter->WriteVector(LocationDeltaName, Identifier, LocationDelta);
+	InWriter->WriteQuat(RotationDeltaName, Identifier, RotationDelta);
+	InWriter->WriteVector(ScaleDeltaName, Identifier, ScaleDelta);
+	InWriter->WritePosition(LocationName, Identifier, Location);
+	InWriter->WriteQuat(RotationName, Identifier, Rotation);
+	InWriter->WriteVector(ScaleName, Identifier, Scale);
+	InWriter->WriteFloat(FrequencyName, Identifier, Frequency);
+	InWriter->WriteVector(PanName, Identifier, Pan);
+
+	/** Forces */
+	InWriter->WriteFloat(OrientationForceRateName, Identifier, OrientationForceRate);
+	InWriter->WriteVector(OrientationForceMinName, Identifier, OrientationForceMin);
+	InWriter->WriteVector(OrientationForceMaxName, Identifier, OrientationForceMax);
+	InWriter->WriteFloat(VortexForceAmountName, Identifier, VortexForceAmount);
+	InWriter->WriteVector(VortexForceAxisName, Identifier, VortexForceAxis);
+	InWriter->WriteFloat(CurlNoiseForceStrengthName, Identifier, CurlNoiseForceStrength);
+	InWriter->WriteFloat(CurlNoiseForceFrequencyName, Identifier, CurlNoiseForceFrequency);
+	InWriter->WriteFloat(AttractionForceStrengthName, Identifier, AttractionForceStrength);
+	InWriter->WriteFloat(AttractionForceFalloffName, Identifier, AttractionForceFalloff);
+	InWriter->WriteVector(GravityForceAccelerationName, Identifier, GravityForceAcceleration);
+}
+
+void FAvaClonerEffectorChannelData::Read(const UNiagaraDataChannelReader* InReader)
+{
+	if (!InReader)
+	{
+		return;
+	}
+
+	/** Effector */
+	bool bIsValid;
+	Easing = static_cast<EAvaClonerEasing>(InReader->ReadInt(EasingName, Identifier, bIsValid));
+	Mode = static_cast<EAvaClonerEffectorMode>(InReader->ReadInt(ModeName, Identifier, bIsValid));
+	Type = static_cast<EAvaClonerEffectorType>(InReader->ReadInt(TypeName, Identifier, bIsValid));
+	Magnitude = InReader->ReadFloat(MagnitudeName, Identifier, bIsValid);
+	InnerExtent = InReader->ReadVector(InnerExtentName, Identifier, bIsValid);
+	OuterExtent = InReader->ReadVector(OuterExtentName, Identifier, bIsValid);
+	LocationDelta = InReader->ReadVector(LocationDeltaName, Identifier, bIsValid);
+	RotationDelta = InReader->ReadQuat(RotationDeltaName, Identifier, bIsValid);
+	ScaleDelta = InReader->ReadVector(ScaleDeltaName, Identifier, bIsValid);
+	Location = InReader->ReadPosition(LocationName, Identifier, bIsValid);
+	Rotation = InReader->ReadQuat(RotationName, Identifier, bIsValid);
+	Scale = InReader->ReadVector(ScaleName, Identifier, bIsValid);
+	Frequency = InReader->ReadFloat(FrequencyName, Identifier, bIsValid);
+	Pan = InReader->ReadVector(PanName, Identifier, bIsValid);
+
+	/** Forces */
+	OrientationForceRate = InReader->ReadFloat(OrientationForceRateName, Identifier, bIsValid);
+	OrientationForceMin = InReader->ReadVector(OrientationForceMinName, Identifier, bIsValid);
+	OrientationForceMax = InReader->ReadVector(OrientationForceMaxName, Identifier, bIsValid);
+	VortexForceAmount = InReader->ReadFloat(VortexForceAmountName, Identifier, bIsValid);
+	VortexForceAxis = InReader->ReadVector(VortexForceAxisName, Identifier, bIsValid);
+	CurlNoiseForceStrength = InReader->ReadFloat(CurlNoiseForceStrengthName, Identifier, bIsValid);
+	CurlNoiseForceFrequency = InReader->ReadFloat(CurlNoiseForceFrequencyName, Identifier, bIsValid);
+	AttractionForceStrength = InReader->ReadFloat(AttractionForceStrengthName, Identifier, bIsValid);
+	AttractionForceFalloff = InReader->ReadFloat(AttractionForceFalloffName, Identifier, bIsValid);
+	GravityForceAcceleration = InReader->ReadVector(GravityForceAccelerationName, Identifier, bIsValid);
+}
+
+FAvaClonerEffectorDataInterfaces::FAvaClonerEffectorDataInterfaces(const UNiagaraSystem* InSystem)
+{
+	check(InSystem);
+	const FNiagaraUserRedirectionParameterStore& UserParameterStore = InSystem->GetExposedParameters();
+
+	static const FNiagaraVariable EffectorIndexDIVar(FNiagaraTypeDefinition(UNiagaraDataInterfaceArrayInt32::StaticClass()), FAvaClonerEffectorDataInterfaces::IndexName);
+	UNiagaraDataInterfaceArrayInt32* IndexArrayDI = CastChecked<UNiagaraDataInterfaceArrayInt32>(UserParameterStore.GetDataInterface(EffectorIndexDIVar));
+	DataInterfaces.Add(EffectorIndexDIVar.GetName(), IndexArrayDI);
+}
+
+void FAvaClonerEffectorDataInterfaces::CopyTo(FAvaClonerEffectorDataInterfaces& InOther) const
+{
+	for (const TPair<FName, TObjectPtr<UNiagaraDataInterface>>& DataInterfacePair : DataInterfaces)
+	{
+		const TObjectPtr<UNiagaraDataInterface> DataInterface = DataInterfacePair.Value;
+		const TObjectPtr<UNiagaraDataInterface>* OtherDataInterface = InOther.DataInterfaces.Find(DataInterfacePair.Key);
+
+		if (!OtherDataInterface || !DataInterface)
+		{
+			continue;
+		}
+
+		DataInterface->CopyTo(*OtherDataInterface);
+	}
+}
+
+void FAvaClonerEffectorDataInterfaces::Resize(int32 InSize) const
+{
+	if (UNiagaraDataInterfaceArrayInt32* IndexArray = GetIndexArray())
+	{
+		IndexArray->GetArrayReference().SetNum(InSize);
+	}
+}
+
+void FAvaClonerEffectorDataInterfaces::Remove(int32 InIndex) const
+{
+	UNiagaraDataInterfaceArrayInt32* IndexArray = GetIndexArray();
+	if (IndexArray && IndexArray->GetArrayReference().IsValidIndex(InIndex))
+	{
+		IndexArray->GetArrayReference().RemoveAt(InIndex);
+	}
+}
+
+bool FAvaClonerEffectorDataInterfaces::IsValid() const
+{
+	for (const TPair<FName, TObjectPtr<UNiagaraDataInterface>>& DataInterfacePair : DataInterfaces)
+	{
+		if (!DataInterfacePair.Value.Get())
+		{
+			return false;
+		}
+	}
+
+	return !DataInterfaces.IsEmpty();
+}
+
+int32 FAvaClonerEffectorDataInterfaces::Num() const
+{
+	UNiagaraDataInterfaceArrayInt32* IndexDI = GetIndexArray();
+
+	const int32 IndexSize = IndexDI->GetArrayReference().Num();
+
+	return IndexSize;
+}
+
+UNiagaraDataInterfaceArrayInt32* FAvaClonerEffectorDataInterfaces::GetIndexArray() const
+{
+	if (const TObjectPtr<UNiagaraDataInterface>* Array = DataInterfaces.Find(FAvaClonerEffectorDataInterfaces::IndexName))
+	{
+		return Cast<UNiagaraDataInterfaceArrayInt32>(Array->Get());
+	}
+
+	return nullptr;
+}
diff --git a/Engine/Plugins/Experimental/Avalanche/Source/AvalancheEffectors/Private/AvalancheEffectorsModule.cpp b/Engine/Plugins/Experimental/Avalanche/Source/AvalancheEffectors/Private/AvalancheEffectorsModule.cpp
new file mode 100644
--- /dev/null
+++ b/Engine/Plugins/Experimental/Avalanche/Source/AvalancheEffectors/Private/AvalancheEffectorsModule.cpp
@@ -0,0 +1,5 @@
+// Copyright Epic Games, Inc. All Rights Reserved.
+
+#include "Modules/ModuleManager.h"
+
+IMPLEMENT_MODULE(FDefaultModuleImpl, AvalancheEffectors)
diff --git a/Engine/Plugins/Experimental/Avalanche/Source/AvalancheEffectors/Private/Cloner/AvaClonerActor.cpp b/Engine/Plugins/Experimental/Avalanche/Source/AvalancheEffectors/Private/Cloner/AvaClonerActor.cpp
new file mode 100644
--- /dev/null
+++ b/Engine/Plugins/Experimental/Avalanche/Source/AvalancheEffectors/Private/Cloner/AvaClonerActor.cpp
@@ -0,0 +1,1404 @@
+// Copyright Epic Games, Inc. All Rights Reserved.
+
+#include "Cloner/AvaClonerActor.h"
+
+#include "Cloner/AvaClonerComponent.h"
+#include "Cloner/Layouts/AvaClonerCircleLayout.h"
+#include "Cloner/Layouts/AvaClonerCylinderLayout.h"
+#include "Cloner/Layouts/AvaClonerGridLayout.h"
+#include "Cloner/Layouts/AvaClonerHoneycombLayout.h"
+#include "Cloner/Layouts/AvaClonerLineLayout.h"
+#include "Cloner/Layouts/AvaClonerMeshLayout.h"
+#include "Cloner/Layouts/AvaClonerSphereUniformLayout.h"
+#include "Cloner/Layouts/AvaClonerSplineLayout.h"
+#include "Effector/AvaEffectorActor.h"
+#include "Engine/StaticMeshActor.h"
+#include "NiagaraDataInterfaceCurve.h"
+#include "NiagaraDataInterfaceSkeletalMesh.h"
+#include "NiagaraMeshRendererProperties.h"
+#include "NiagaraSystem.h"
+#include "Subsystems/AvaClonerSubsystem.h"
+#include "UObject/ConstructorHelpers.h"
+
+#if WITH_EDITOR
+#include "Editor.h"
+#include "Selection.h"
+#include "Settings/LevelEditorViewportSettings.h"
+#endif
+
+DEFINE_LOG_CATEGORY_STATIC(LogAvaClonerActor, Log, All);
+
+// Sets default values
+AAvaClonerActor::AAvaClonerActor()
+{
+	SetCanBeDamaged(false);
+	PrimaryActorTick.bCanEverTick          = true;
+	PrimaryActorTick.bStartWithTickEnabled = true;
+	PrimaryActorTick.bTickEvenWhenPaused   = true;
+
+	ClonerComponent = CreateDefaultSubobject<UAvaClonerComponent>(TEXT("AvaClonerComponent"));
+	SetRootComponent(ClonerComponent);
+
+	// Default Scale Curve
+	LifetimeScaleCurve.AddKey(0, 1.f);
+	LifetimeScaleCurve.AddKey(1, 0.f);
+
+	if (!IsTemplate())
+	{
+#if WITH_EDITOR
+		if (GEditor)
+		{
+			GEditor->GetSelectedActors()->SelectionChangedEvent.AddUObject(this, &AAvaClonerActor::OnEditorSelectionChanged);
+		}
+#endif
+
+		ClonerComponent->TransformUpdated.AddUObject(this, &AAvaClonerActor::OnClonerTransformed);
+		ClonerComponent->OnClonerMeshUpdated.AddUObject(this, &AAvaClonerActor::OnClonerMeshUpdated);
+
+		const TArray<FString> LayoutNames = GetClonerLayoutNames();
+
+		// Apply default layout
+		LayoutName = !LayoutNames.IsEmpty() ? FName(LayoutNames[0]) : NAME_None;
+
+#if WITH_EDITOR
+		if (const UAvaClonerSubsystem* ClonerSubsystem = UAvaClonerSubsystem::Get())
+		{
+			UAvaClonerSubsystem::OnCVarChangedDelegate.AddUObject(this, &AAvaClonerActor::OnCVarChanged);
+			bReduceMotionGhosting = ClonerSubsystem->IsNoFlickerEnabled();
+		}
+#endif
+	}
+}
+
+#if WITH_EDITOR
+FString AAvaClonerActor::GetDefaultActorLabel() const
+{
+	return DefaultLabel;
+}
+
+const TAvaPropertyChangeDispatcher<AAvaClonerActor> AAvaClonerActor::PropertyChangeDispatcher =
+{
+	{ GET_MEMBER_NAME_CHECKED(AAvaClonerActor, bEnabled), &AAvaClonerActor::OnEnabledChanged },
+	{ GET_MEMBER_NAME_CHECKED(AAvaClonerActor, Seed), &AAvaClonerActor::OnSeedChanged },
+	/** Layout */
+	{ GET_MEMBER_NAME_CHECKED(AAvaClonerActor, LayoutName), &AAvaClonerActor::OnLayoutNameChanged },
+	/** Advanced */
+	{ GET_MEMBER_NAME_CHECKED(AAvaClonerActor, bVisualizerSpriteVisible), &AAvaClonerActor::OnVisualizerSpriteVisibleChanged },
+	{ GET_MEMBER_NAME_CHECKED(AAvaClonerActor, bReduceMotionGhosting), &AAvaClonerActor::OnReduceMotionGhostingChanged },
+	/** Spawn */
+	{ GET_MEMBER_NAME_CHECKED(AAvaClonerActor, SpawnLoopMode), &AAvaClonerActor::OnSpawnOptionsChanged },
+	{ GET_MEMBER_NAME_CHECKED(AAvaClonerActor, SpawnLoopInterval), &AAvaClonerActor::OnSpawnOptionsChanged },
+	{ GET_MEMBER_NAME_CHECKED(AAvaClonerActor, SpawnLoopIterations), &AAvaClonerActor::OnSpawnOptionsChanged },
+	{ GET_MEMBER_NAME_CHECKED(AAvaClonerActor, SpawnBehaviorMode), &AAvaClonerActor::OnSpawnOptionsChanged },
+	{ GET_MEMBER_NAME_CHECKED(AAvaClonerActor, SpawnRate), &AAvaClonerActor::OnSpawnOptionsChanged },
+	/** Lifetime */
+	{ GET_MEMBER_NAME_CHECKED(AAvaClonerActor, bLifetimeEnabled), &AAvaClonerActor::OnLifetimeOptionsChanged },
+	{ GET_MEMBER_NAME_CHECKED(AAvaClonerActor, LifetimeMin), &AAvaClonerActor::OnLifetimeOptionsChanged },
+	{ GET_MEMBER_NAME_CHECKED(AAvaClonerActor, LifetimeMax), &AAvaClonerActor::OnLifetimeOptionsChanged },
+	{ GET_MEMBER_NAME_CHECKED(AAvaClonerActor, bLifetimeScaleEnabled), &AAvaClonerActor::OnLifetimeOptionsChanged },
+	/** Renderer */
+	{ GET_MEMBER_NAME_CHECKED(AAvaClonerActor, MeshRenderMode), &AAvaClonerActor::OnMeshRenderModeChanged },
+	{ GET_MEMBER_NAME_CHECKED(AAvaClonerActor, MeshFacingMode), &AAvaClonerActor::OnMeshRendererOptionsChanged },
+	{ GET_MEMBER_NAME_CHECKED(AAvaClonerActor, bMeshCastShadows), &AAvaClonerActor::OnMeshRendererOptionsChanged },
+	{ GET_MEMBER_NAME_CHECKED(AAvaClonerActor, DefaultMeshes), &AAvaClonerActor::OnDefaultMeshesChanged },
+	/** Progress */
+	{ GET_MEMBER_NAME_CHECKED(AAvaClonerActor, bInvertProgress), &AAvaClonerActor::OnProgressChanged },
+	{ GET_MEMBER_NAME_CHECKED(AAvaClonerActor, Progress), &AAvaClonerActor::OnProgressChanged },
+	/** Step */
+	{ GET_MEMBER_NAME_CHECKED(AAvaClonerActor, bDeltaStepEnabled), &AAvaClonerActor::OnDeltaStepChanged },
+	{ GET_MEMBER_NAME_CHECKED(AAvaClonerActor, DeltaStepRotation), &AAvaClonerActor::OnDeltaStepChanged },
+	{ GET_MEMBER_NAME_CHECKED(AAvaClonerActor, DeltaStepScale), &AAvaClonerActor::OnDeltaStepChanged },
+	/** Range */
+	{ GET_MEMBER_NAME_CHECKED(AAvaClonerActor, bRangeEnabled), &AAvaClonerActor::OnRangeOptionsChanged },
+	{ GET_MEMBER_NAME_CHECKED(AAvaClonerActor, RangeOffsetMin), &AAvaClonerActor::OnRangeOptionsChanged },
+	{ GET_MEMBER_NAME_CHECKED(AAvaClonerActor, RangeOffsetMax), &AAvaClonerActor::OnRangeOptionsChanged },
+	{ GET_MEMBER_NAME_CHECKED(AAvaClonerActor, RangeRotationMin), &AAvaClonerActor::OnRangeOptionsChanged },
+	{ GET_MEMBER_NAME_CHECKED(AAvaClonerActor, RangeRotationMax), &AAvaClonerActor::OnRangeOptionsChanged },
+	{ GET_MEMBER_NAME_CHECKED(AAvaClonerActor, RangeScaleMin), &AAvaClonerActor::OnRangeOptionsChanged },
+	{ GET_MEMBER_NAME_CHECKED(AAvaClonerActor, RangeScaleMax), &AAvaClonerActor::OnRangeOptionsChanged },
+	{ GET_MEMBER_NAME_CHECKED(AAvaClonerActor, bRangeScaleUniform), &AAvaClonerActor::OnRangeOptionsChanged },
+	{ GET_MEMBER_NAME_CHECKED(AAvaClonerActor, RangeScaleUniformMin), &AAvaClonerActor::OnRangeOptionsChanged },
+	{ GET_MEMBER_NAME_CHECKED(AAvaClonerActor, RangeScaleUniformMax), &AAvaClonerActor::OnRangeOptionsChanged },
+};
+
+void AAvaClonerActor::PostEditChangeProperty(FPropertyChangedEvent& InPropertyChangedEvent)
+{
+	Super::PostEditChangeProperty(InPropertyChangedEvent);
+
+	PropertyChangeDispatcher.OnPropertyChanged(this, InPropertyChangedEvent);
+}
+#endif
+
+void AAvaClonerActor::PostInitProperties()
+{
+	Super::PostInitProperties();
+
+	if (HasAnyFlags(RF_ClassDefaultObject))
+	{
+		// Register new type def for niagara
+
+		constexpr ENiagaraTypeRegistryFlags MeshFlags =
+			ENiagaraTypeRegistryFlags::AllowAnyVariable |
+			ENiagaraTypeRegistryFlags::AllowParameter;
+
+		FNiagaraTypeRegistry::Register(FNiagaraTypeDefinition(StaticEnum<EAvaClonerMeshRenderMode>()), MeshFlags);
+		FNiagaraTypeRegistry::Register(FNiagaraTypeDefinition(StaticEnum<EAvaClonerGridConstraint>()), MeshFlags);
+		FNiagaraTypeRegistry::Register(FNiagaraTypeDefinition(StaticEnum<EAvaClonerPlane>()), MeshFlags);
+		FNiagaraTypeRegistry::Register(FNiagaraTypeDefinition(StaticEnum<EAvaClonerAxis>()), MeshFlags);
+		FNiagaraTypeRegistry::Register(FNiagaraTypeDefinition(StaticEnum<EAvaClonerEasing>()), MeshFlags);
+		FNiagaraTypeRegistry::Register(FNiagaraTypeDefinition(StaticEnum<EAvaClonerMeshAsset>()), MeshFlags);
+		FNiagaraTypeRegistry::Register(FNiagaraTypeDefinition(StaticEnum<EAvaClonerMeshSampleData>()), MeshFlags);
+		FNiagaraTypeRegistry::Register(FNiagaraTypeDefinition(StaticEnum<EAvaClonerEffectorType>()), MeshFlags);
+		FNiagaraTypeRegistry::Register(FNiagaraTypeDefinition(StaticEnum<EAvaClonerTextureSampleChannel>()), MeshFlags);
+		FNiagaraTypeRegistry::Register(FNiagaraTypeDefinition(StaticEnum<EAvaClonerCompareMode>()), MeshFlags);
+		FNiagaraTypeRegistry::Register(FNiagaraTypeDefinition(StaticEnum<EAvaClonerEffectorMode>()), MeshFlags);
+		FNiagaraTypeRegistry::Register(FNiagaraTypeDefinition(StaticEnum<EAvaClonerSpawnLoopMode>()), MeshFlags);
+		FNiagaraTypeRegistry::Register(FNiagaraTypeDefinition(StaticEnum<EAvaClonerSpawnBehaviorMode>()), MeshFlags);
+	}
+}
+
+void AAvaClonerActor::PostLoad()
+{
+	Super::PostLoad();
+
+	// Migrate old properties to new layout system
+	if (!bDeprecatedPropertiesMigrated)
+	{
+		// Grid options
+		if (UAvaClonerGridLayout* GridLayout = FindOrAddLayout<UAvaClonerGridLayout>())
+		{
+			GridLayout->SetCountX(GridOptions_DEPRECATED.Count.X);
+			GridLayout->SetCountY(GridOptions_DEPRECATED.Count.Y);
+			GridLayout->SetCountZ(GridOptions_DEPRECATED.Count.Z);
+			GridLayout->SetSpacingX(GridOptions_DEPRECATED.Spacing.X);
+			GridLayout->SetSpacingY(GridOptions_DEPRECATED.Spacing.Y);
+			GridLayout->SetSpacingZ(GridOptions_DEPRECATED.Spacing.Z);
+			GridLayout->SetConstraint(GridOptions_DEPRECATED.Constraint);
+			GridLayout->SetInvertConstraint(GridOptions_DEPRECATED.bInvertConstraint);
+			GridLayout->SetSphereConstraint(GridOptions_DEPRECATED.SphereConstraint);
+			GridLayout->SetCylinderConstraint(GridOptions_DEPRECATED.CylinderConstraint);
+			GridLayout->SetTextureConstraint(GridOptions_DEPRECATED.TextureConstraint);
+		}
+
+		// Line options
+		if (UAvaClonerLineLayout* LineLayout = FindOrAddLayout<UAvaClonerLineLayout>())
+		{
+			LineLayout->SetCount(LineOptions_DEPRECATED.Count);
+			LineLayout->SetSpacing(LineOptions_DEPRECATED.Spacing);
+			LineLayout->SetAxis(LineOptions_DEPRECATED.Axis);
+			LineLayout->SetDirection(LineOptions_DEPRECATED.Direction);
+			LineLayout->SetRotation(LineOptions_DEPRECATED.Rotation);
+		}
+
+		// Circle options
+		if (UAvaClonerCircleLayout* CircleLayout = FindOrAddLayout<UAvaClonerCircleLayout>())
+		{
+			CircleLayout->SetCount(CircleOptions_DEPRECATED.Count);
+			CircleLayout->SetRadius(CircleOptions_DEPRECATED.Radius);
+			CircleLayout->SetAngleStart(CircleOptions_DEPRECATED.AngleStart);
+			CircleLayout->SetAngleRatio(CircleOptions_DEPRECATED.AngleRatio);
+			CircleLayout->SetOrientMesh(CircleOptions_DEPRECATED.bOrientMesh);
+			CircleLayout->SetPlane(CircleOptions_DEPRECATED.Plane);
+			CircleLayout->SetRotation(CircleOptions_DEPRECATED.Rotation);
+			CircleLayout->SetScale(CircleOptions_DEPRECATED.Scale);
+		}
+
+		// Cylinder options
+		if (UAvaClonerCylinderLayout* CylinderLayout = FindOrAddLayout<UAvaClonerCylinderLayout>())
+		{
+			CylinderLayout->SetBaseCount(CylinderOptions_DEPRECATED.BaseCount);
+			CylinderLayout->SetHeightCount(CylinderOptions_DEPRECATED.HeightCount);
+			CylinderLayout->SetHeight(CylinderOptions_DEPRECATED.Height);
+			CylinderLayout->SetRadius(CylinderOptions_DEPRECATED.Radius);
+			CylinderLayout->SetAngleStart(CylinderOptions_DEPRECATED.AngleStart);
+			CylinderLayout->SetAngleRatio(CylinderOptions_DEPRECATED.AngleRatio);
+			CylinderLayout->SetOrientMesh(CylinderOptions_DEPRECATED.bOrientMesh);
+			CylinderLayout->SetPlane(CylinderOptions_DEPRECATED.Plane);
+			CylinderLayout->SetRotation(CylinderOptions_DEPRECATED.Rotation);
+			CylinderLayout->SetScale(CylinderOptions_DEPRECATED.Scale);
+		}
+
+		// Sphere options
+		if (UAvaClonerSphereUniformLayout* SphereLayout = FindOrAddLayout<UAvaClonerSphereUniformLayout>())
+		{
+			SphereLayout->SetCount(SphereOptions_DEPRECATED.Count);
+			SphereLayout->SetRadius(SphereOptions_DEPRECATED.Radius);
+			SphereLayout->SetRatio(SphereOptions_DEPRECATED.Ratio);
+			SphereLayout->SetOrientMesh(SphereOptions_DEPRECATED.bOrientMesh);
+			SphereLayout->SetRotation(SphereOptions_DEPRECATED.Rotation);
+			SphereLayout->SetScale(SphereOptions_DEPRECATED.Scale);
+		}
+
+		// Honeycomb options
+		if (UAvaClonerHoneycombLayout* HoneycombLayout = FindOrAddLayout<UAvaClonerHoneycombLayout>())
+		{
+			HoneycombLayout->SetPlane(HoneycombOptions_DEPRECATED.Plane);
+			HoneycombLayout->SetWidthCount(HoneycombOptions_DEPRECATED.WidthCount);
+			HoneycombLayout->SetHeightCount(HoneycombOptions_DEPRECATED.HeightCount);
+			HoneycombLayout->SetWidthOffset(HoneycombOptions_DEPRECATED.WidthOffset);
+			HoneycombLayout->SetHeightOffset(HoneycombOptions_DEPRECATED.HeightOffset);
+			HoneycombLayout->SetHeightSpacing(HoneycombOptions_DEPRECATED.HeightSpacing);
+			HoneycombLayout->SetWidthSpacing(HoneycombOptions_DEPRECATED.WidthSpacing);
+		}
+
+		// Sample mesh options
+		if (UAvaClonerMeshLayout* MeshLayout = FindOrAddLayout<UAvaClonerMeshLayout>())
+		{
+			MeshLayout->SetCount(SampleMeshOptions_DEPRECATED.Count);
+			MeshLayout->SetAsset(SampleMeshOptions_DEPRECATED.Asset);
+			MeshLayout->SetSampleData(SampleMeshOptions_DEPRECATED.SampleData);
+			MeshLayout->SetSampleActorWeak(SampleMeshOptions_DEPRECATED.SampleActor);
+		}
+
+		// Sample spline options
+		if (UAvaClonerSplineLayout* SplineLayout = FindOrAddLayout<UAvaClonerSplineLayout>())
+		{
+			SplineLayout->SetCount(SampleSplineOptions_DEPRECATED.Count);
+			SplineLayout->SetSplineActorWeak(SampleSplineOptions_DEPRECATED.SplineActor);
+		}
+
+		// Set previous layout
+		if (const UAvaClonerSubsystem* Subsystem = UAvaClonerSubsystem::Get())
+		{
+			const int32 LayoutIdx = static_cast<int32>(Layout_DEPRECATED);
+			SetLayoutName(Subsystem->GetLayoutNames()[LayoutIdx]);
+		}
+
+		bDeprecatedPropertiesMigrated = true;
+
+		UE_LOG(LogAvaClonerActor, Log, TEXT("%s : Cloner deprecated properties migrated"), *GetActorNameOrLabel());
+	}
+
+	for (TMap<FName, TObjectPtr<UAvaClonerLayoutBase>>::TIterator It(LayoutInstances); It; ++It)
+	{
+		if (!It->Value)
+		{
+			It.RemoveCurrent();
+			continue;
+		}
+
+		// Update layout name if changed
+		const FName NewLayoutName = It->Value->GetLayoutName();
+
+		if (It->Key != NewLayoutName)
+		{
+			LayoutInstances.Add(NewLayoutName, It->Value);
+			It.RemoveCurrent();
+
+			UE_LOG(LogAvaClonerActor, Log, TEXT("%s : Cloner layout name changed %s"), *GetActorNameOrLabel(), *NewLayoutName.ToString());
+		}
+	}
+}
+
+void AAvaClonerActor::PostActorCreated()
+{
+	Super::PostActorCreated();
+
+#if WITH_EDITOR
+	SpawnDefaultActorAttached();
+#endif
+}
+
+void AAvaClonerActor::Tick(float InDeltaTime)
+{
+	Super::Tick(InDeltaTime);
+
+	if (!bClonerInitialized)
+	{
+		InitializeCloner();
+	}
+
+	if (!bEnabled)
+	{
+		return;
+	}
+
+	TreeUpdateDeltaTime += InDeltaTime;
+	if (TreeUpdateDeltaTime >= TreeUpdateInterval)
+	{
+		TreeUpdateDeltaTime -= TreeUpdateInterval;
+
+		if (ClonerComponent)
+		{
+			ClonerComponent->UpdateClonerAttachmentTree();
+			ClonerComponent->UpdateClonerRenderState();
+		}
+	}
+
+	if (bNeedsRefresh)
+	{
+		bNeedsRefresh = false;
+		RequestClonerUpdate(true);
+	}
+}
+
+void AAvaClonerActor::UpdateLayoutOptions()
+{
+	OnMeshRenderModeChanged();
+	OnSeedChanged();
+
+	// update layouts options
+	if (UAvaClonerLayoutBase* CurrentLayout = ClonerComponent->GetClonerActiveLayout())
+	{
+		CurrentLayout->UpdateLayoutParameters();
+	}
+
+	OnMeshRendererOptionsChanged();
+	OnDefaultMeshesChanged();
+	OnDeltaStepChanged();
+	OnRangeOptionsChanged();
+	OnSpawnOptionsChanged();
+	OnLifetimeOptionsChanged();
+}
+
+void AAvaClonerActor::UpdateClonerEffectors()
+{
+	Effectors.RemoveAll([](const TWeakObjectPtr<AAvaEffectorActor>& InEffector)
+	{
+		return !InEffector.IsValid();
+	});
+
+	if (const UAvaClonerLayoutBase* LayoutSystem = GetActiveLayout())
+	{
+		// Effectors could be registered before system is loaded
+		const int32 EffectorCount = GetEffectorCount();
+
+		if (LayoutSystem->GetDataInterfaces().Num() != EffectorCount)
+		{
+			LayoutSystem->GetDataInterfaces().Resize(EffectorCount);
+
+			RequestClonerUpdate();
+		}
+	}
+
+	ForEachEffector([this](AAvaEffectorActor* InEffector, int32 InIdx)
+	{
+		InEffector->OnClonerUpdated(this, InIdx);
+		return true;
+	});
+}
+
+void AAvaClonerActor::SetTreeUpdateInterval(float InInterval)
+{
+	if (InInterval == TreeUpdateInterval)
+	{
+		return;
+	}
+
+	TreeUpdateInterval = InInterval;
+}
+
+void AAvaClonerActor::SetMeshRenderMode(EAvaClonerMeshRenderMode InMode)
+{
+	if (InMode == MeshRenderMode)
+	{
+		return;
+	}
+
+	MeshRenderMode = InMode;
+	OnMeshRenderModeChanged();
+}
+
+void AAvaClonerActor::SetMeshFacingMode(ENiagaraMeshFacingMode InMode)
+{
+	if (MeshFacingMode == InMode)
+	{
+		return;
+	}
+
+	MeshFacingMode = InMode;
+	OnMeshRendererOptionsChanged();
+}
+
+void AAvaClonerActor::SetMeshCastShadows(bool InbCastShadows)
+{
+	if (bMeshCastShadows == InbCastShadows)
+	{
+		return;
+	}
+
+	bMeshCastShadows = InbCastShadows;
+	OnMeshRendererOptionsChanged();
+}
+
+void AAvaClonerActor::SetDefaultMeshes(const TArray<TObjectPtr<UStaticMesh>>& InMeshes)
+{
+	DefaultMeshes = InMeshes;
+	OnDefaultMeshesChanged();
+}
+
+void AAvaClonerActor::BP_SetDefaultMeshes(const TArray<UStaticMesh*>& InMeshes)
+{
+	DefaultMeshes.Empty(InMeshes.Num());
+	Algo::Transform(InMeshes, DefaultMeshes, [](UStaticMesh* InMesh)->TObjectPtr<UStaticMesh>
+	{
+		return InMesh;
+	});
+	OnDefaultMeshesChanged();
+}
+
+TArray<UStaticMesh*> AAvaClonerActor::BP_GetDefaultMeshes() const
+{
+	TArray<UStaticMesh*> Meshes;
+	Meshes.Reserve(DefaultMeshes.Num());
+	Algo::Transform(DefaultMeshes, Meshes, [](const TObjectPtr<UStaticMesh>& InMesh)->UStaticMesh*
+	{
+		return InMesh;
+	});
+	return Meshes;
+}
+
+void AAvaClonerActor::SetSeed(int32 InSeed)
+{
+	if (InSeed == Seed)
+	{
+		return;
+	}
+
+	Seed = InSeed;
+	OnSeedChanged();
+}
+
+void AAvaClonerActor::SetDeltaStepEnabled(bool bInEnabled)
+{
+	if (bDeltaStepEnabled == bInEnabled)
+	{
+		return;
+	}
+
+	bDeltaStepEnabled = bInEnabled;
+	OnDeltaStepChanged();
+}
+
+void AAvaClonerActor::SetDeltaStepRotation(const FRotator& InRotation)
+{
+	if (InRotation == DeltaStepRotation)
+	{
+		return;
+	}
+
+	DeltaStepRotation = InRotation;
+	OnDeltaStepChanged();
+}
+
+void AAvaClonerActor::SetDeltaStepScale(const FVector& InScale)
+{
+	if (InScale == DeltaStepScale)
+	{
+		return;
+	}
+
+	DeltaStepScale = InScale;
+	OnDeltaStepChanged();
+}
+
+void AAvaClonerActor::SetEnabled(bool bInEnable)
+{
+	if (bInEnable == bEnabled)
+	{
+		return;
+	}
+
+	bEnabled = bInEnable;
+	OnEnabledChanged();
+}
+
+void AAvaClonerActor::SetInvertProgress(bool bInInvertProgress)
+{
+	if (bInvertProgress == bInInvertProgress)
+	{
+		return;
+	}
+
+	bInvertProgress = bInInvertProgress;
+	OnProgressChanged();
+}
+
+void AAvaClonerActor::SetProgress(float InProgress)
+{
+	if (Progress == InProgress)
+	{
+		return;
+	}
+
+	if (InProgress < 0.f || InProgress > 1.f)
+	{
+		return;
+	}
+
+	Progress = InProgress;
+	OnProgressChanged();
+}
+
+void AAvaClonerActor::SetLayoutName(FName InLayoutName)
+{
+	if (LayoutName == InLayoutName)
+	{
+		return;
+	}
+
+	const TArray<FString> LayoutNames = GetClonerLayoutNames();
+	if (!LayoutNames.Contains(InLayoutName))
+	{
+		return;
+	}
+
+	LayoutName = InLayoutName;
+	OnLayoutNameChanged();
+}
+
+void AAvaClonerActor::SetRangeEnabled(bool bInRangeEnabled)
+{
+	if (bRangeEnabled == bInRangeEnabled)
+	{
+		return;
+	}
+
+	bRangeEnabled = bInRangeEnabled;
+	OnRangeOptionsChanged();
+}
+
+void AAvaClonerActor::SetRangeOffsetMin(const FVector& InRangeOffsetMin)
+{
+	if (RangeOffsetMin == InRangeOffsetMin)
+	{
+		return;
+	}
+
+	RangeOffsetMin = InRangeOffsetMin;
+	OnRangeOptionsChanged();
+}
+
+void AAvaClonerActor::SetRangeOffsetMax(const FVector& InRangeOffsetMax)
+{
+	if (RangeOffsetMax == InRangeOffsetMax)
+	{
+		return;
+	}
+
+	RangeOffsetMax = InRangeOffsetMax;
+	OnRangeOptionsChanged();
+}
+
+void AAvaClonerActor::SetRangeRotationMin(const FRotator& InRangeRotationMin)
+{
+	if (RangeRotationMin == InRangeRotationMin)
+	{
+		return;
+	}
+
+	RangeRotationMin = InRangeRotationMin;
+	OnRangeOptionsChanged();
+}
+
+void AAvaClonerActor::SetRangeRotationMax(const FRotator& InRangeRotationMax)
+{
+	if (RangeRotationMax == InRangeRotationMax)
+	{
+		return;
+	}
+
+	RangeRotationMax = InRangeRotationMax;
+	OnRangeOptionsChanged();
+}
+
+void AAvaClonerActor::SetRangeScaleUniform(bool bInRangeScaleUniform)
+{
+	if (bRangeScaleUniform == bInRangeScaleUniform)
+	{
+		return;
+	}
+
+	bRangeScaleUniform = bInRangeScaleUniform;
+	OnRangeOptionsChanged();
+}
+
+void AAvaClonerActor::SetRangeScaleMin(const FVector& InRangeScaleMin)
+{
+	if (RangeScaleMin == InRangeScaleMin)
+	{
+		return;
+	}
+
+	RangeScaleMin = InRangeScaleMin;
+	OnRangeOptionsChanged();
+}
+
+void AAvaClonerActor::SetRangeScaleMax(const FVector& InRangeScaleMax)
+{
+	if (RangeScaleMax == InRangeScaleMax)
+	{
+		return;
+	}
+
+	RangeScaleMax = InRangeScaleMax;
+	OnRangeOptionsChanged();
+}
+
+void AAvaClonerActor::SetRangeScaleUniformMin(float InRangeScaleUniformMin)
+{
+	if (RangeScaleUniformMin == InRangeScaleUniformMin)
+	{
+		return;
+	}
+
+	RangeScaleUniformMin = InRangeScaleUniformMin;
+	OnRangeOptionsChanged();
+}
+
+void AAvaClonerActor::SetRangeScaleUniformMax(float InRangeScaleUniformMax)
+{
+	if (RangeScaleUniformMax == InRangeScaleUniformMax)
+	{
+		return;
+	}
+
+	RangeScaleUniformMax = InRangeScaleUniformMax;
+	OnRangeOptionsChanged();
+}
+
+void AAvaClonerActor::OnMeshRenderModeChanged()
+{
+	if (!bEnabled)
+	{
+		return;
+	}
+
+	if (const UAvaClonerLayoutBase* LayoutSystem = GetActiveLayout())
+	{
+		FNiagaraUserRedirectionParameterStore& ExposedParameters = LayoutSystem->GetSystem()->GetExposedParameters();
+
+		static const FNiagaraVariable MeshModeVar(FNiagaraTypeDefinition(StaticEnum<EAvaClonerMeshRenderMode>()), TEXT("MeshRenderMode"));
+		ExposedParameters.SetParameterValue<int32>(static_cast<int32>(MeshRenderMode), MeshModeVar);
+
+		RequestClonerUpdate();
+	}
+}
+
+void AAvaClonerActor::OnDeltaStepChanged()
+{
+	if (!ClonerComponent || !bEnabled)
+	{
+		return;
+	}
+
+	if (bDeltaStepEnabled)
+	{
+		ClonerComponent->SetVectorParameter(TEXT("DeltaStepRotation"), FVector(DeltaStepRotation.Roll, DeltaStepRotation.Pitch, DeltaStepRotation.Yaw));
+
+		ClonerComponent->SetVectorParameter(TEXT("DeltaStepScale"), DeltaStepScale);
+	}
+	else
+	{
+		ClonerComponent->SetVectorParameter(TEXT("DeltaStepRotation"), FVector::ZeroVector);
+
+		ClonerComponent->SetVectorParameter(TEXT("DeltaStepScale"), FVector::ZeroVector);
+	}
+
+	RequestClonerUpdate();
+}
+
+void AAvaClonerActor::OnSeedChanged()
+{
+	if (!ClonerComponent || !bEnabled)
+	{
+		return;
+	}
+
+	ClonerComponent->SetRandomSeedOffset(Seed);
+
+	RequestClonerUpdate();
+}
+
+void AAvaClonerActor::OnProgressChanged()
+{
+	if (!ClonerComponent || !bEnabled)
+	{
+		return;
+	}
+
+	ClonerComponent->SetFloatParameter(TEXT("ParticleProgress"), Progress * (bInvertProgress ? -1 : 1));
+}
+
+void AAvaClonerActor::SetSpawnLoopMode(EAvaClonerSpawnLoopMode InMode)
+{
+	if (SpawnLoopMode == InMode)
+	{
+		return;
+	}
+
+	SpawnLoopMode = InMode;
+	OnSpawnOptionsChanged();
+}
+
+void AAvaClonerActor::SetSpawnLoopIterations(int32 InIterations)
+{
+	if (SpawnLoopIterations == InIterations)
+	{
+		return;
+	}
+
+	if (InIterations < 1)
+	{
+		return;
+	}
+
+	SpawnLoopIterations = InIterations;
+	OnSpawnOptionsChanged();
+}
+
+void AAvaClonerActor::SetSpawnLoopInterval(float InInterval)
+{
+	if (SpawnLoopInterval == InInterval)
+	{
+		return;
+	}
+
+	if (InInterval < 0.f)
+	{
+		return;
+	}
+
+	SpawnLoopInterval = InInterval;
+	OnSpawnOptionsChanged();
+}
+
+void AAvaClonerActor::SetSpawnBehaviorMode(EAvaClonerSpawnBehaviorMode InMode)
+{
+	if (SpawnBehaviorMode == InMode)
+	{
+		return;
+	}
+
+	SpawnBehaviorMode = InMode;
+	OnSpawnOptionsChanged();
+}
+
+void AAvaClonerActor::SetSpawnRate(float InRate)
+{
+	if (SpawnRate == InRate)
+	{
+		return;
+	}
+
+	if (InRate < 0)
+	{
+		return;
+	}
+
+	SpawnRate = InRate;
+	OnSpawnOptionsChanged();
+}
+
+void AAvaClonerActor::SetLifetimeEnabled(bool bInEnabled)
+{
+	if (bLifetimeEnabled == bInEnabled)
+	{
+		return;
+	}
+
+	bLifetimeEnabled = bInEnabled;
+	OnLifetimeOptionsChanged();
+}
+
+void AAvaClonerActor::SetLifetimeMin(float InMin)
+{
+	if (LifetimeMin == InMin)
+	{
+		return;
+	}
+
+	if (InMin < 0)
+	{
+		return;
+	}
+
+	LifetimeMin = InMin;
+	OnLifetimeOptionsChanged();
+}
+
+void AAvaClonerActor::SetLifetimeMax(float InMax)
+{
+	if (LifetimeMax == InMax)
+	{
+		return;
+	}
+
+	if (InMax < 0)
+	{
+		return;
+	}
+
+	LifetimeMax = InMax;
+	OnLifetimeOptionsChanged();
+}
+
+void AAvaClonerActor::SetLifetimeScaleEnabled(bool bInEnabled)
+{
+	if (bLifetimeScaleEnabled == bInEnabled)
+	{
+		return;
+	}
+
+	bLifetimeScaleEnabled = bInEnabled;
+	OnLifetimeOptionsChanged();
+}
+
+void AAvaClonerActor::SetLifetimeScaleCurve(const FRichCurve& InCurve)
+{
+	LifetimeScaleCurve = InCurve;
+	OnLifetimeOptionsChanged();
+}
+
+UAvaClonerLayoutBase* AAvaClonerActor::GetActiveLayout() const
+{
+	if (ClonerComponent)
+	{
+		return ClonerComponent->GetClonerActiveLayout();
+	}
+
+	return nullptr;
+}
+
+const FAvaClonerEffectorDataInterfaces* AAvaClonerActor::GetEffectorDataInterfaces() const
+{
+	if (const UAvaClonerLayoutBase* LayoutSystem = GetActiveLayout())
+	{
+		return &LayoutSystem->GetDataInterfaces();
+	}
+
+	return nullptr;
+}
+
+int32 AAvaClonerActor::GetMeshCount() const
+{
+	if (ClonerComponent)
+	{
+		if (const UAvaClonerLayoutBase* LayoutSystem = GetActiveLayout())
+		{
+			return LayoutSystem->GetMeshRenderer()->Meshes.Num();
+		}
+	}
+	return 0;
+}
+
+int32 AAvaClonerActor::RegisterEffector(AAvaEffectorActor* InEffector)
+{
+	if (!InEffector)
+	{
+		return INDEX_NONE;
+	}
+
+	int32 Index = Effectors.Find(InEffector);
+
+	if (Index != INDEX_NONE)
+	{
+		return Index;
+	}
+
+	Index = Effectors.Add(InEffector);
+
+	if (const UAvaClonerLayoutBase* LayoutSystem = GetActiveLayout())
+	{
+		const int32 EffectorCount = GetEffectorCount();
+		LayoutSystem->GetDataInterfaces().Resize(EffectorCount);
+	}
+
+	InEffector->OnClonerLinked(this, Index);
+
+	constexpr bool bImmediateUpdate = true;
+	RequestClonerUpdate(bImmediateUpdate);
+
+	return Index;
+}
+
+bool AAvaClonerActor::UnregisterEffector(AAvaEffectorActor* InEffector)
+{
+	if (!InEffector)
+	{
+		return false;
+	}
+
+	const int32 OldIndex = Effectors.Find(InEffector);
+
+	if (OldIndex == INDEX_NONE)
+	{
+		return false;
+	}
+
+	Effectors.RemoveAt(OldIndex);
+
+	if (const UAvaClonerLayoutBase* LayoutSystem = GetActiveLayout())
+	{
+		LayoutSystem->GetDataInterfaces().Remove(OldIndex);
+	}
+
+	InEffector->OnClonerUnlinked(this, OldIndex);
+
+	constexpr bool bImmediateUpdate = true;
+	RequestClonerUpdate(bImmediateUpdate);
+
+	return true;
+}
+
+bool AAvaClonerActor::IsEffectorRegistered(const AAvaEffectorActor* InEffector) const
+{
+	return InEffector && Effectors.Contains(InEffector);
+}
+
+int32 AAvaClonerActor::GetEffectorIndex(AAvaEffectorActor* InEffector) const
+{
+	return InEffector ? Effectors.Find(InEffector) : INDEX_NONE;
+}
+
+int32 AAvaClonerActor::GetEffectorCount() const
+{
+	return Effectors.Num();
+}
+
+void AAvaClonerActor::ForEachEffector(TFunctionRef<bool(AAvaEffectorActor*, int32)> InFunction)
+{
+	for (int32 Idx = 0; Idx < Effectors.Num(); Idx++)
+	{
+		AAvaEffectorActor* Effector = Effectors[Idx].Get();
+		if (!Effector)
+		{
+			continue;
+		}
+
+		if (!InFunction(Effector, Idx))
+		{
+			return;
+		}
+	}
+}
+
+void AAvaClonerActor::OnEnabledChanged()
+{
+	if (bEnabled)
+	{
+		OnLayoutNameChanged();
+	}
+	else
+	{
+		ClonerComponent->DeactivateImmediate();
+		ClonerComponent->SetAsset(nullptr);
+	}
+}
+
+void AAvaClonerActor::OnLayoutNameChanged()
+{
+	if (!ClonerComponent || !bEnabled)
+	{
+		return;
+	}
+
+	UAvaClonerLayoutBase* NewActiveLayout = FindOrAddLayout(LayoutName);
+
+	// Apply layout
+	if (ClonerComponent->SetClonerActiveLayout(NewActiveLayout))
+	{
+		ActiveLayout = ClonerComponent->GetClonerActiveLayout();
+		OnClonerSystemChanged();
+	}
+}
+
+void AAvaClonerActor::OnSpawnOptionsChanged()
+{
+	if (!ClonerComponent || !bEnabled)
+	{
+		return;
+	}
+
+	if (const UAvaClonerLayoutBase* ActiveSystem = GetActiveLayout())
+	{
+		FNiagaraUserRedirectionParameterStore& ExposedParameters = ActiveSystem->GetSystem()->GetExposedParameters();
+
+		ClonerComponent->SetFloatParameter(TEXT("SpawnLoopInterval"), SpawnLoopInterval);
+
+		ClonerComponent->SetIntParameter(TEXT("SpawnLoopIterations"), SpawnLoopIterations);
+
+		ClonerComponent->SetFloatParameter(TEXT("SpawnRate"), SpawnRate);
+
+		const FNiagaraVariable SpawnBehaviorModeVar(FNiagaraTypeDefinition(StaticEnum<EAvaClonerSpawnBehaviorMode>()), TEXT("SpawnBehaviorMode"));
+		ExposedParameters.SetParameterValue<int32>(static_cast<int32>(SpawnBehaviorMode), SpawnBehaviorModeVar);
+
+		const FNiagaraVariable SpawnLoopModeVar(FNiagaraTypeDefinition(StaticEnum<EAvaClonerSpawnLoopMode>()), TEXT("SpawnLoopMode"));
+		ExposedParameters.SetParameterValue<int32>(static_cast<int32>(SpawnLoopMode), SpawnLoopModeVar);
+
+		RequestClonerUpdate();
+	}
+}
+
+void AAvaClonerActor::OnLifetimeScaleCurveChanged()
+{
+	if (const UNiagaraDataInterfaceCurve* LifetimeCurve = LifetimeScaleCurveDIWeak.Get())
+	{
+		LifetimeScaleCurve = LifetimeCurve->Curve;
+		OnLifetimeOptionsChanged();
+	}
+}
+
+void AAvaClonerActor::OnLifetimeOptionsChanged()
+{
+	if (!ClonerComponent || !bEnabled)
+	{
+		return;
+	}
+
+	LifetimeMin = FMath::Max(0, LifetimeMin);
+	LifetimeMax = FMath::Max(LifetimeMin, LifetimeMax);
+
+	ClonerComponent->SetBoolParameter(TEXT("LifetimeEnabled"), bLifetimeEnabled);
+
+	ClonerComponent->SetFloatParameter(TEXT("LifetimeMin"), LifetimeMin);
+
+	ClonerComponent->SetFloatParameter(TEXT("LifetimeMax"), LifetimeMax);
+
+	ClonerComponent->SetBoolParameter(TEXT("LifetimeScaleEnabled"), bLifetimeEnabled && bLifetimeScaleEnabled);
+
+	if (const UAvaClonerLayoutBase* ActiveSystem = GetActiveLayout())
+	{
+		const FNiagaraUserRedirectionParameterStore& ExposedParameters = ActiveSystem->GetSystem()->GetExposedParameters();
+
+		static const FNiagaraVariable LifetimeScaleCurveVar(FNiagaraTypeDefinition(UNiagaraDataInterfaceCurve::StaticClass()), TEXT("LifetimeScaleCurve"));
+
+#if WITH_EDITOR
+		if (UNiagaraDataInterfaceCurve* LifetimeCurve = LifetimeScaleCurveDIWeak.Get())
+		{
+			LifetimeCurve->OnChanged().RemoveAll(this);
+		}
+#endif
+
+		if (UNiagaraDataInterfaceCurve* LifetimeScaleCurveDI = Cast<UNiagaraDataInterfaceCurve>(ExposedParameters.GetDataInterface(LifetimeScaleCurveVar)))
+		{
+			LifetimeScaleCurveDIWeak = LifetimeScaleCurveDI;
+			LifetimeScaleCurveDI->Curve = LifetimeScaleCurve;
+
+#if WITH_EDITOR
+			LifetimeScaleCurveDI->UpdateLUT();
+			LifetimeScaleCurveDI->OnChanged().AddUObject(this, &AAvaClonerActor::OnLifetimeScaleCurveChanged);
+#endif
+		}
+	}
+
+	RequestClonerUpdate();
+}
+
+void AAvaClonerActor::OnClonerTransformed(USceneComponent*, EUpdateTransformFlags, ETeleportType)
+{
+	UpdateClonerEffectors();
+}
+
+void AAvaClonerActor::OnClonerMeshUpdated()
+{
+	RequestClonerUpdate();
+}
+
+void AAvaClonerActor::OnClonerSystemChanged()
+{
+	UpdateLayoutOptions();
+	UpdateClonerEffectors();
+}
+
+void AAvaClonerActor::RequestClonerUpdate(bool bInImmediate)
+{
+	if (!bEnabled || !ClonerComponent)
+	{
+		return;
+	}
+
+	if (bInImmediate)
+	{
+		bNeedsRefresh = false;
+		ClonerComponent->RefreshUserParameters();
+	}
+	else
+	{
+		bNeedsRefresh = true;
+	}
+}
+
+void AAvaClonerActor::OnRangeOptionsChanged()
+{
+	if (!ClonerComponent || !bEnabled)
+	{
+		return;
+	}
+
+	RangeScaleUniformMin = FMath::Clamp(RangeScaleUniformMin, UE_KINDA_SMALL_NUMBER, RangeScaleUniformMax);
+	RangeScaleUniformMax = FMath::Max3(RangeScaleUniformMin, RangeScaleUniformMax, UE_KINDA_SMALL_NUMBER);
+
+	RangeScaleMin.X = FMath::Clamp(RangeScaleMin.X, UE_KINDA_SMALL_NUMBER, RangeScaleMax.X);
+	RangeScaleMin.Y = FMath::Clamp(RangeScaleMin.Y, UE_KINDA_SMALL_NUMBER, RangeScaleMax.Y);
+	RangeScaleMin.Z = FMath::Clamp(RangeScaleMin.Z, UE_KINDA_SMALL_NUMBER, RangeScaleMax.Z);
+
+	RangeScaleMax.X = FMath::Max3<double>(RangeScaleMin.X, RangeScaleMax.X, UE_KINDA_SMALL_NUMBER);
+	RangeScaleMax.Y = FMath::Max3<double>(RangeScaleMin.Y, RangeScaleMax.Y, UE_KINDA_SMALL_NUMBER);
+	RangeScaleMax.Z = FMath::Max3<double>(RangeScaleMin.Z, RangeScaleMax.Z, UE_KINDA_SMALL_NUMBER);
+
+	ClonerComponent->SetBoolParameter(TEXT("RangeEnabled"), bRangeEnabled);
+
+	ClonerComponent->SetVectorParameter(TEXT("RangeOffsetMin"), RangeOffsetMin);
+
+	ClonerComponent->SetVectorParameter(TEXT("RangeOffsetMax"), RangeOffsetMax);
+
+	ClonerComponent->SetVariableQuat(TEXT("RangeRotationMin"), RangeRotationMin.Quaternion());
+
+	ClonerComponent->SetVariableQuat(TEXT("RangeRotationMax"), RangeRotationMax.Quaternion());
+
+	ClonerComponent->SetBoolParameter(TEXT("RangeScaleUniform"), bRangeScaleUniform);
+
+	ClonerComponent->SetVectorParameter(TEXT("RangeScaleMin"), RangeScaleMin);
+
+	ClonerComponent->SetVectorParameter(TEXT("RangeScaleMax"), RangeScaleMax);
+
+	ClonerComponent->SetFloatParameter(TEXT("RangeScaleUniformMin"), RangeScaleUniformMin);
+
+	ClonerComponent->SetFloatParameter(TEXT("RangeScaleUniformMax"), RangeScaleUniformMax);
+
+	RequestClonerUpdate();
+}
+
+bool AAvaClonerActor::IsClonerValid() const
+{
+	return IsValid(ClonerComponent.Get()) && IsValid(ClonerComponent->GetAsset());
+}
+
+void AAvaClonerActor::OnMeshRendererOptionsChanged()
+{
+	const UAvaClonerLayoutBase* CurrentLayout = ClonerComponent->GetClonerActiveLayout();
+
+	if (!CurrentLayout)
+	{
+		return;
+	}
+
+	UNiagaraMeshRendererProperties* MeshRenderer = CurrentLayout->GetMeshRenderer();
+
+	if (!MeshRenderer)
+	{
+		return;
+	}
+
+	MeshRenderer->FacingMode = MeshFacingMode;
+	MeshRenderer->bCastShadows = bMeshCastShadows;
+
+	RequestClonerUpdate();
+}
+
+void AAvaClonerActor::OnDefaultMeshesChanged()
+{
+	// Will force an update
+	TreeUpdateDeltaTime = TreeUpdateInterval;
+}
+
+#if WITH_EDITOR
+void AAvaClonerActor::OnReduceMotionGhostingChanged()
+{
+	if (UAvaClonerSubsystem* ClonerSubsystem = UAvaClonerSubsystem::Get())
+	{
+		if (bReduceMotionGhosting)
+		{
+			ClonerSubsystem->EnableNoFlicker();
+		}
+		else
+		{
+			ClonerSubsystem->DisableNoFlicker();
+		}
+	}
+}
+
+void AAvaClonerActor::OnCVarChanged()
+{
+	if (const UAvaClonerSubsystem* ClonerSubsystem = UAvaClonerSubsystem::Get())
+	{
+		bReduceMotionGhosting = ClonerSubsystem->IsNoFlickerEnabled();
+	}
+}
+
+void AAvaClonerActor::ForceUpdateCloner()
+{
+	if (ClonerComponent)
+	{
+		constexpr bool bReset = true;
+		ClonerComponent->UpdateClonerAttachmentTree(bReset);
+		TreeUpdateDeltaTime = TreeUpdateInterval;
+	}
+}
+
+void AAvaClonerActor::SpawnLinkedEffector()
+{
+	UWorld* ClonerWorld = GetWorld();
+	if (!ClonerWorld)
+	{
+		return;
+	}
+
+	FActorSpawnParameters Params;
+	Params.bTemporaryEditorActor = false;
+
+	const FVector ClonerLocation = GetActorLocation();
+	const FRotator ClonerRotation = GetActorRotation();
+
+	AAvaEffectorActor* EffectorActor = ClonerWorld->SpawnActor<AAvaEffectorActor>(AAvaEffectorActor::StaticClass(), ClonerLocation, ClonerRotation, Params);
+
+	if (!EffectorActor)
+	{
+		return;
+	}
+
+	EffectorActor->LinkCloner(this);
+	FActorLabelUtilities::RenameExistingActor(EffectorActor, EffectorActor->GetDefaultActorLabel(), true);
+
+	// Set default offset for visual feedback
+	EffectorActor->SetOffset(FVector(0, 0, 100));
+}
+
+void AAvaClonerActor::SpawnDefaultActorAttached()
+{
+	// Only spawn if world is valid and not a preview actor
+	UWorld* World = GetWorld();
+	if (!World || bIsEditorPreviewActor)
+	{
+		return;
+	}
+
+	// Find or load cube mesh
+	constexpr const TCHAR* DefaultStaticMeshPath = TEXT("/Script/Engine.StaticMesh'/Engine/BasicShapes/Cube.Cube'");
+	UStaticMesh* DefaultStaticMesh = LoadObject<UStaticMesh>(nullptr, DefaultStaticMeshPath);
+
+	if (!DefaultStaticMesh)
+	{
+		return;
+	}
+
+	// Spawn attached actor with same flags as this actor
+	FActorSpawnParameters SpawnParameters;
+	SpawnParameters.Owner = this;
+	SpawnParameters.ObjectFlags = GetFlags();
+	SpawnParameters.bTemporaryEditorActor = false;
+
+	if (AStaticMeshActor* DefaultActorAttached = World->SpawnActor<AStaticMeshActor>(GetActorLocation(), GetActorRotation(), SpawnParameters))
+	{
+		UStaticMeshComponent* StaticMeshComponent = DefaultActorAttached->GetStaticMeshComponent();
+		StaticMeshComponent->SetStaticMesh(DefaultStaticMesh);
+
+		DefaultActorAttached->SetMobility(EComponentMobility::Movable);
+		DefaultActorAttached->AttachToActor(this, FAttachmentTransformRules::KeepWorldTransform);
+
+		FActorLabelUtilities::SetActorLabelUnique(DefaultActorAttached, TEXT("DefaultCube"));
+	}
+}
+
+void AAvaClonerActor::OnEditorSelectionChanged(UObject* InSelection)
+{
+	if (const USelection* ActorSelection = Cast<USelection>(InSelection))
+	{
+		if (ULevelEditorViewportSettings* ViewportSettings = GetMutableDefault<ULevelEditorViewportSettings>())
+		{
+			if (ActorSelection->Num() == 1 && ActorSelection->GetSelectedObject(0) == this)
+			{
+				UseSelectionOutline = ViewportSettings->bUseSelectionOutline;
+				ViewportSettings->bUseSelectionOutline = false;
+			}
+			else if (UseSelectionOutline.IsSet())
+			{
+				ViewportSettings->bUseSelectionOutline = UseSelectionOutline.GetValue();
+				UseSelectionOutline.Reset();
+			}
+		}
+	}
+}
+#endif
+
+void AAvaClonerActor::OnVisualizerSpriteVisibleChanged()
+{
+#if WITH_EDITOR
+	UE::ClonerEffector::SetBillboardComponentSprite(this, TEXT("/Script/Engine.Texture2D'/Avalanche/ClonerResources/Textures/T_ClonerIcon.T_ClonerIcon'"));
+	UE::ClonerEffector::SetBillboardComponentVisibility(this, bVisualizerSpriteVisible);
+#endif
+}
+
+TArray<FString> AAvaClonerActor::GetClonerLayoutNames() const
+{
+	TArray<FString> LayoutNamesStrings;
+
+	if (const UAvaClonerSubsystem* Subsystem = UAvaClonerSubsystem::Get())
+	{
+		const TArray<FName>& LayoutNames = Subsystem->GetLayoutNames();
+		LayoutNamesStrings.Reserve(LayoutNames.Num());
+
+		Algo::Transform(LayoutNames, LayoutNamesStrings, [](const FName& InName)
+		{
+			return InName.ToString();
+		});
+	}
+
+	return LayoutNamesStrings;
+}
+
+UAvaClonerLayoutBase* AAvaClonerActor::FindOrAddLayout(TSubclassOf<UAvaClonerLayoutBase> InClass)
+{
+	const UAvaClonerSubsystem* Subsystem = UAvaClonerSubsystem::Get();
+
+	if (!Subsystem)
+	{
+		return nullptr;
+	}
+
+	const FName ClassLayoutName = Subsystem->FindLayoutName(InClass);
+
+	if (ClassLayoutName.IsNone())
+	{
+		return nullptr;
+	}
+
+	return FindOrAddLayout(ClassLayoutName);
+}
+
+UAvaClonerLayoutBase* AAvaClonerActor::FindOrAddLayout(FName InLayoutName)
+{
+	UAvaClonerSubsystem* Subsystem = UAvaClonerSubsystem::Get();
+	if (!Subsystem)
+	{
+		return nullptr;
+	}
+
+	// Check cached layout instances
+	UAvaClonerLayoutBase* NewActiveLayout = nullptr;
+	if (TObjectPtr<UAvaClonerLayoutBase> const* LayoutInstance = LayoutInstances.Find(InLayoutName))
+	{
+		NewActiveLayout = *LayoutInstance;
+	}
+
+	// Create new layout instance and cache it
+	if (!NewActiveLayout)
+	{
+		NewActiveLayout = Subsystem->CreateNewLayout(InLayoutName, this);
+		LayoutInstances.Add(InLayoutName, NewActiveLayout);
+	}
+
+	return NewActiveLayout;
+}
+
+void AAvaClonerActor::InitializeCloner()
+{
+	if (bClonerInitialized)
+	{
+		return;
+	}
+
+	OnVisualizerSpriteVisibleChanged();
+
+#if WITH_EDITOR
+	OnReduceMotionGhostingChanged();
+#endif
+
+	// For new cloner instances no need to migrate anything
+	bDeprecatedPropertiesMigrated = true;
+
+	bClonerInitialized = true;
+
+	OnLayoutNameChanged();
+}
diff --git a/Engine/Plugins/Experimental/Avalanche/Source/AvalancheEffectors/Private/Cloner/AvaClonerComponent.cpp b/Engine/Plugins/Experimental/Avalanche/Source/AvalancheEffectors/Private/Cloner/AvaClonerComponent.cpp
new file mode 100644
--- /dev/null
+++ b/Engine/Plugins/Experimental/Avalanche/Source/AvalancheEffectors/Private/Cloner/AvaClonerComponent.cpp
@@ -0,0 +1,1437 @@
+// Copyright Epic Games, Inc. All Rights Reserved.
+
+#include "Cloner/AvaClonerComponent.h"
+#include "Async/Async.h"
+#include "AvaActorUtils.h"
+#include "AvaBlueprint.h"
+#include "AvaSceneTree.h"
+#include "Cloner/AvaClonerActor.h"
+#include "Cloner/Layouts/AvaClonerLayoutBase.h"
+#include "Components/BrushComponent.h"
+#include "Components/DynamicMeshComponent.h"
+#include "Components/SkeletalMeshComponent.h"
+#include "DynamicMesh/MeshNormals.h"
+#include "DynamicMesh/MeshTransforms.h"
+#include "DynamicMesh/Operations/MergeCoincidentMeshEdges.h"
+#include "DynamicMeshEditor.h"
+#include "GeometryScript/MeshAssetFunctions.h"
+#include "GeometryScript/MeshBasicEditFunctions.h"
+#include "GeometryScript/SceneUtilityFunctions.h"
+#include "Materials/Material.h"
+#include "NiagaraMeshRendererProperties.h"
+#include "NiagaraSystem.h"
+#include "ProceduralMeshComponent.h"
+#include "StaticMeshOperations.h"
+#include "UDynamicMesh.h"
+#include "UObject/Package.h"
+
+#if WITH_EDITOR
+#include "AvaOutliner.h"
+#include "AvaOutlinerDefines.h"
+#include "AvaOutlinerSubsystem.h"
+#include "Framework/Notifications/NotificationManager.h"
+#include "Item/AvaOutlinerActor.h"
+#include "Materials/Material.h"
+#include "Misc/MessageDialog.h"
+#include "Widgets/Notifications/SNotificationList.h"
+#endif
+
+DEFINE_LOG_CATEGORY_STATIC(LogAvaClonerComponent, Log, All);
+
+#define LOCTEXT_NAMESPACE "AvaEffectorActorVisualizer"
+
+UAvaClonerComponent::UAvaClonerComponent()
+	: UNiagaraComponent()
+{
+	CastShadow = true;
+	bReceivesDecals = true;
+
+#if WITH_EDITOR
+	// Do not show bounding box around cloner for better visibility
+	SetIsVisualizationComponent(true);
+#endif
+
+	bIsEditorOnly = false;
+
+	// Show sprite for this component to visualize it when empty
+#if WITH_EDITORONLY_DATA
+	bVisualizeComponent = true;
+#endif
+
+	if (!IsTemplate())
+	{
+		// Bind to delegate to detect material changes
+#if WITH_EDITOR
+		FCoreUObjectDelegates::OnObjectPropertyChanged.RemoveAll(this);
+		FCoreUObjectDelegates::OnObjectPropertyChanged.AddUObject(this, &UAvaClonerComponent::OnActorPropertyChanged);
+#endif
+
+		USceneComponent::MarkRenderStateDirtyEvent.AddUObject(this, &UAvaClonerComponent::OnRenderStateDirty);
+	}
+}
+
+void UAvaClonerComponent::PostLoad()
+{
+	SetAsset(nullptr);
+
+	Super::PostLoad();
+}
+
+void UAvaClonerComponent::UpdateClonerRenderState()
+{
+	/**
+	 * Perform a mesh update when asset is valid,
+	 * An update is not already ongoing,
+	 * Meshes are out of date after an attachment tree update,
+	 * Tree is up to date
+	 */
+	if (!GetAsset()
+		|| bClonerMeshesUpdating
+		|| !bClonerMeshesDirty
+		|| ClonerTree.Status != EAvaClonerAttachmentStatus::Updated)
+	{
+		return;
+	}
+
+	UpdateDirtyMeshesAsync();
+}
+
+void UAvaClonerComponent::RefreshUserParameters() const
+{
+	if (UNiagaraSystem* ActiveSystem = GetAsset())
+	{
+		FNiagaraUserRedirectionParameterStore& UserParameterStore = ActiveSystem->GetExposedParameters();
+		UserParameterStore.PostGenericEditChange();
+	}
+}
+
+void UAvaClonerComponent::OnRenderStateDirty(UActorComponent& InComponent)
+{
+	const AActor* OwningActor = InComponent.GetOwner();
+	const AAvaClonerActor* ClonerActor = GetOuterAAvaClonerActor();
+
+	if (!IsValid(OwningActor) || !IsValid(ClonerActor))
+	{
+		return;
+	}
+
+	if (OwningActor->GetWorld() != ClonerActor->GetWorld())
+	{
+		return;
+	}
+
+	/* Owning actor is updated and is attached to cloner
+		or owning actor is detached from cloner */
+	if (OwningActor->IsAttachedTo(ClonerActor)
+		|| ClonerTree.ItemAttachmentMap.Contains(OwningActor))
+	{
+		UE_LOG(LogAvaClonerComponent, Log, TEXT("%s : Detected attachment change for %s"), *ClonerActor->GetActorNameOrLabel(), *OwningActor->GetActorNameOrLabel());
+
+		// Only update attachment tree
+		UpdateClonerAttachmentTree();
+	}
+}
+
+void UAvaClonerComponent::UpdateClonerAttachmentTree(bool bInReset)
+{
+	if (ClonerTree.Status == EAvaClonerAttachmentStatus::Updated)
+	{
+		ClonerTree.Status = EAvaClonerAttachmentStatus::Outdated;
+	}
+
+	if (bInReset)
+	{
+		ClonerTree.Reset();
+		ClonerTree.Status = EAvaClonerAttachmentStatus::Outdated;
+	}
+
+	UpdateAttachmentTree();
+}
+
+void UAvaClonerComponent::UpdateAttachmentTree()
+{
+	if (ClonerTree.Status != EAvaClonerAttachmentStatus::Outdated)
+	{
+		return;
+	}
+	ClonerTree.Status = EAvaClonerAttachmentStatus::Updating;
+
+	// Invalidate all, to see what is outdated and what is still invalid
+	for (TPair<TWeakObjectPtr<AActor>, FAvaClonerAttachmentItem>& AttachmentPair : ClonerTree.ItemAttachmentMap)
+	{
+		AttachmentPair.Value.Status = EAvaClonerAttachmentStatus::Invalid;
+	}
+
+	// Update root attachment items
+	TArray<AActor*> RootChildren;
+	GetOrderedRootActors(RootChildren);
+
+	TArray<TObjectPtr<UStaticMesh>> NewCombinesMeshes;
+	TArray<TWeakObjectPtr<AActor>> NewRootActors;
+	for (int32 RootIdx = 0; RootIdx < RootChildren.Num(); RootIdx++)
+	{
+		AActor* RootChild = RootChildren[RootIdx];
+
+		UpdateActorAttachment(RootChild, nullptr);
+
+		// Lets find the old root idx
+		const int32 OldIdx = ClonerTree.RootActors.Find(RootChild);
+		TObjectPtr<UStaticMesh> CombineBakedMesh = nullptr;
+		if (OldIdx != INDEX_NONE)
+		{
+			CombineBakedMesh = ClonerTree.MergedBakedMeshes[OldIdx];
+
+			// Did we rearrange stuff ?
+			if (RootIdx != OldIdx)
+			{
+				bClonerMeshesDirty = true;
+			}
+		}
+		NewCombinesMeshes.Add(CombineBakedMesh);
+		NewRootActors.Add(RootChild);
+	}
+
+	// Did we remove any root actors ?
+	if (ClonerTree.RootActors.Num() != NewRootActors.Num())
+	{
+		bClonerMeshesDirty = true;
+	}
+
+	// Did we need to update meshes
+	TArray<TWeakObjectPtr<AActor>> ClonedActors;
+	ClonerTree.ItemAttachmentMap.GenerateKeyArray(ClonedActors);
+	for (const TWeakObjectPtr<AActor>& ClonedActor : ClonedActors)
+	{
+		FAvaClonerAttachmentItem* ClonedItem = ClonerTree.ItemAttachmentMap.Find(ClonedActor);
+
+		if (!ClonedItem)
+		{
+			continue;
+		}
+
+		if (ClonedItem->Status == EAvaClonerAttachmentStatus::Invalid)
+		{
+			InvalidateBakedStaticMesh(ClonedActor.Get());
+			UnbindActorDelegates(ClonedActor.Get());
+			ClonerTree.ItemAttachmentMap.Remove(ClonedItem->ItemActor);
+		}
+		else if (ClonedItem->Status == EAvaClonerAttachmentStatus::Outdated)
+		{
+			if (ClonedItem->MeshStatus == EAvaClonerAttachmentStatus::Outdated)
+			{
+				DirtyItemAttachments.Add(ClonedItem);
+				InvalidateBakedStaticMesh(ClonedActor.Get());
+			}
+			bClonerMeshesDirty = true;
+			ClonedItem->Status = EAvaClonerAttachmentStatus::Updated;
+		}
+	}
+
+	// Did we remove an attachment ?
+	if (ClonedActors.Num() != ClonerTree.ItemAttachmentMap.Num())
+	{
+		bClonerMeshesDirty = true;
+	}
+
+	if (!DirtyItemAttachments.IsEmpty())
+	{
+		bClonerMeshesDirty = true;
+	}
+
+	// Use default meshes if nothing is attached
+	if (NewRootActors.IsEmpty() && NewCombinesMeshes.IsEmpty())
+	{
+		if (const AAvaClonerActor* ClonerActor = GetOuterAAvaClonerActor())
+		{
+			const TArray<TObjectPtr<UStaticMesh>>& DefaultMeshes = ClonerActor->GetDefaultMeshes();
+
+			if (ClonerTree.MergedBakedMeshes.Num() != DefaultMeshes.Num())
+			{
+				bClonerMeshesDirty = true;
+			}
+
+			for (const TObjectPtr<UStaticMesh>& DefaultMesh : DefaultMeshes)
+			{
+				const int32 Idx = NewCombinesMeshes.Add(DefaultMesh);
+
+				// Only update when there is a change
+				if (!ClonerTree.MergedBakedMeshes.IsValidIndex(Idx)
+					|| DefaultMesh != ClonerTree.MergedBakedMeshes[Idx])
+				{
+					bClonerMeshesDirty = true;
+				}
+			}
+		}
+	}
+
+	ClonerTree.RootActors = NewRootActors;
+	ClonerTree.MergedBakedMeshes = NewCombinesMeshes;
+	ClonerTree.Status = EAvaClonerAttachmentStatus::Updated;
+}
+
+void UAvaClonerComponent::UpdateActorAttachment(AActor* InActor, AActor* InParent)
+{
+	if (!InActor)
+	{
+		return;
+	}
+
+	const AActor* ClonerActor = GetOwner();
+	const FTransform& ClonerTransform = ClonerActor->GetActorTransform();
+
+	// Here order is not important
+	TArray<AActor*> ChildrenActors;
+	InActor->GetAttachedActors(ChildrenActors, true, false);
+
+	FAvaClonerAttachmentItem* AttachmentItem = ClonerTree.ItemAttachmentMap.Find(InActor);
+	if (AttachmentItem)
+	{
+		AttachmentItem->Status = EAvaClonerAttachmentStatus::Updated;
+
+		// Check Root is the same
+		const bool bIsRoot = InParent == nullptr;
+		if (AttachmentItem->bRootItem != bIsRoot)
+		{
+			InvalidateBakedStaticMesh(InActor);
+			AttachmentItem->bRootItem = bIsRoot;
+			AttachmentItem->Status = EAvaClonerAttachmentStatus::Outdated;
+		}
+
+		// Check parent is the same
+		if (AttachmentItem->ParentActor.Get() != InParent)
+		{
+			InvalidateBakedStaticMesh(InParent);
+			InvalidateBakedStaticMesh(AttachmentItem->ParentActor.Get());
+			AttachmentItem->ParentActor = InParent;
+			AttachmentItem->Status = EAvaClonerAttachmentStatus::Outdated;
+		}
+
+		// Check transform is the same
+		const FTransform ActorTransform = InActor->GetActorTransform().GetRelativeTransform(ClonerTransform);
+		if (!ActorTransform.Equals(AttachmentItem->ActorTransform))
+		{
+			// invalidate if not root, else change transform in mesh renderer
+			if (!bIsRoot)
+			{
+				InvalidateBakedStaticMesh(InActor);
+			}
+			AttachmentItem->ActorTransform = ActorTransform;
+			AttachmentItem->Status = EAvaClonerAttachmentStatus::Outdated;
+		}
+	}
+	else
+	{
+		AttachmentItem = &ClonerTree.ItemAttachmentMap.Add(InActor, FAvaClonerAttachmentItem());
+		AttachmentItem->ItemActor = InActor;
+		AttachmentItem->ParentActor = InParent;
+		AttachmentItem->ActorTransform = InActor->GetActorTransform().GetRelativeTransform(ClonerTransform);
+		AttachmentItem->MeshStatus = EAvaClonerAttachmentStatus::Outdated;
+		AttachmentItem->bRootItem = InParent == nullptr;
+		AttachmentItem->Status = EAvaClonerAttachmentStatus::Outdated;
+		InvalidateBakedStaticMesh(InActor);
+		BindActorDelegates(InActor);
+	}
+
+	if (AttachmentItem->ChildrenActors.Num() != ChildrenActors.Num())
+	{
+		InvalidateBakedStaticMesh(InActor);
+	}
+
+	AttachmentItem->ChildrenActors.Empty(ChildrenActors.Num());
+	for (AActor* ChildActor : ChildrenActors)
+	{
+		AttachmentItem->ChildrenActors.Add(ChildActor);
+		UpdateActorAttachment(ChildActor, InActor);
+	}
+}
+
+void UAvaClonerComponent::BindActorDelegates(AActor* InActor)
+{
+	if (!InActor)
+	{
+		return;
+	}
+
+#if WITH_EDITOR
+	// Hide new actor when attached to cloner
+	InActor->SetIsTemporarilyHiddenInEditor(true);
+#endif
+	InActor->SetActorHiddenInGame(true);
+
+	InActor->OnDestroyed.AddUniqueDynamic(this, &UAvaClonerComponent::OnActorDestroyed);
+
+#if WITH_EDITOR
+	// Detect static mesh change
+	TArray<UStaticMeshComponent*> StaticMeshComponents;
+	InActor->GetComponents(StaticMeshComponents, false);
+	for (UStaticMeshComponent* StaticMeshComponent : StaticMeshComponents)
+	{
+		if (!StaticMeshComponent->OnStaticMeshChanged().IsBoundToObject(this))
+		{
+			StaticMeshComponent->OnStaticMeshChanged().AddUObject(this, &UAvaClonerComponent::OnMeshChanged, InActor);
+		}
+	}
+#endif
+
+	// Detect dynamic mesh change
+	TArray<UDynamicMeshComponent*> DynamicMeshComponents;
+	InActor->GetComponents(DynamicMeshComponents, false);
+	for (UDynamicMeshComponent* DynamicMeshComponent : DynamicMeshComponents)
+	{
+		if (!DynamicMeshComponent->OnMeshChanged.IsBoundToObject(this))
+		{
+			UStaticMeshComponent* NullComponent = nullptr;
+			DynamicMeshComponent->OnMeshChanged.AddUObject(this, &UAvaClonerComponent::OnMeshChanged, NullComponent, InActor);
+		}
+	}
+}
+
+void UAvaClonerComponent::UnbindActorDelegates(AActor* InActor) const
+{
+	if (!InActor)
+	{
+		return;
+	}
+
+#if WITH_EDITOR
+	// Show new actor when detached to cloner
+	InActor->SetIsTemporarilyHiddenInEditor(false);
+#endif
+	InActor->SetActorHiddenInGame(false);
+
+	InActor->OnDestroyed.RemoveAll(this);
+
+#if WITH_EDITOR
+	TArray<UStaticMeshComponent*> StaticMeshComponents;
+	InActor->GetComponents(StaticMeshComponents, false);
+	for (UStaticMeshComponent* StaticMeshComponent : StaticMeshComponents)
+	{
+		StaticMeshComponent->OnStaticMeshChanged().RemoveAll(this);
+	}
+#endif
+
+	TArray<UDynamicMeshComponent*> DynamicMeshComponents;
+	InActor->GetComponents(DynamicMeshComponents, false);
+	for (UDynamicMeshComponent* DynamicMeshComponent : DynamicMeshComponents)
+	{
+		DynamicMeshComponent->OnMeshChanged.RemoveAll(this);
+	}
+}
+
+void UAvaClonerComponent::OnActorDestroyed(AActor* InDestroyedActor)
+{
+	if (ClonerTree.ItemAttachmentMap.Contains(InDestroyedActor))
+	{
+		InvalidateBakedStaticMesh(InDestroyedActor);
+		UnbindActorDelegates(InDestroyedActor);
+		ClonerTree.ItemAttachmentMap.Remove(InDestroyedActor);
+		bClonerMeshesDirty = true;
+	}
+}
+
+#if WITH_EDITOR
+void UAvaClonerComponent::OnActorPropertyChanged(UObject* InObject, FPropertyChangedEvent& InPropertyChangedEvent)
+{
+	OnMaterialChanged(InObject);
+}
+#endif
+
+void UAvaClonerComponent::OnMaterialChanged(UObject* InObject)
+{
+	const AActor* ClonerActor = GetOwner();
+
+	if (!ClonerActor)
+	{
+		return;
+	}
+
+	if (AActor* ActorChanged = Cast<AActor>(InObject))
+	{
+		if (FAvaClonerAttachmentItem* Item = ClonerTree.ItemAttachmentMap.Find(ActorChanged))
+		{
+			int32 MatIdx = 0;
+			TArray<TWeakObjectPtr<UMaterialInterface>> NewMaterials;
+			bool bMaterialChanged = false;
+
+			TArray<UDynamicMeshComponent*> DynamicMeshComponents;
+			ActorChanged->GetComponents(DynamicMeshComponents, false);
+			for (const UDynamicMeshComponent* DynamicMeshComponent : DynamicMeshComponents)
+			{
+				for (UMaterialInterface* Material : DynamicMeshComponent->GetMaterials())
+				{
+					if (!Item->BakedMaterials.IsValidIndex(MatIdx) || Item->BakedMaterials[MatIdx] != Material)
+					{
+						bMaterialChanged = true;
+					}
+					NewMaterials.Add(Material);
+					MatIdx++;
+				}
+			}
+
+			TArray<USkeletalMeshComponent*> SkeletalMeshComponents;
+			ActorChanged->GetComponents(SkeletalMeshComponents, false);
+			for (const USkeletalMeshComponent* SkeletalMeshComponent : SkeletalMeshComponents)
+			{
+				for (UMaterialInterface* Material : SkeletalMeshComponent->GetMaterials())
+				{
+					if (!Item->BakedMaterials.IsValidIndex(MatIdx) || Item->BakedMaterials[MatIdx] != Material)
+					{
+						bMaterialChanged = true;
+					}
+					NewMaterials.Add(Material);
+					MatIdx++;
+				}
+			}
+
+			TArray<UBrushComponent*> BrushComponents;
+			ActorChanged->GetComponents(BrushComponents, false);
+			for (const UBrushComponent* BrushComponent : BrushComponents)
+			{
+				for (int32 Idx = 0; Idx < BrushComponent->GetNumMaterials(); Idx++)
+				{
+					if (!Item->BakedMaterials.IsValidIndex(MatIdx) || Item->BakedMaterials[MatIdx] != BrushComponent->GetMaterial(Idx))
+					{
+						bMaterialChanged = true;
+					}
+					NewMaterials.Add(BrushComponent->GetMaterial(Idx));
+					MatIdx++;
+				}
+			}
+
+			TArray<UProceduralMeshComponent*> ProceduralMeshComponents;
+			ActorChanged->GetComponents(ProceduralMeshComponents, false);
+			for (const UProceduralMeshComponent* ProceduralMeshComponent : ProceduralMeshComponents)
+			{
+				if (ProceduralMeshComponent->GetNumSections() == 0)
+				{
+					continue;
+				}
+				for (UMaterialInterface* Material : ProceduralMeshComponent->GetMaterials())
+				{
+					if (!Item->BakedMaterials.IsValidIndex(MatIdx) || Item->BakedMaterials[MatIdx] != Material)
+					{
+						bMaterialChanged = true;
+					}
+					NewMaterials.Add(Material);
+					MatIdx++;
+				}
+			}
+
+			TArray<UStaticMeshComponent*> StaticMeshComponents;
+			ActorChanged->GetComponents(StaticMeshComponents, false);
+			for (const UStaticMeshComponent* StaticMeshComponent : StaticMeshComponents)
+			{
+				for (UMaterialInterface* Material : StaticMeshComponent->GetMaterials())
+				{
+					if (!Item->BakedMaterials.IsValidIndex(MatIdx) || Item->BakedMaterials[MatIdx] != Material)
+					{
+						bMaterialChanged = true;
+					}
+					NewMaterials.Add(Material);
+					MatIdx++;
+				}
+			}
+
+			if (bMaterialChanged)
+			{
+				UE_LOG(LogAvaClonerComponent, Log, TEXT("%s : Detected material change for %s"), *ClonerActor->GetActorNameOrLabel(), *ActorChanged->GetActorNameOrLabel());
+
+				if (NewMaterials.Num() == Item->BakedMaterials.Num())
+				{
+					Item->BakedMaterials = NewMaterials;
+				}
+				else
+				{
+					Item->MeshStatus = EAvaClonerAttachmentStatus::Outdated;
+				}
+
+				InvalidateBakedStaticMesh(ActorChanged);
+			}
+		}
+	}
+}
+
+void UAvaClonerComponent::OnTransformUpdated(USceneComponent* InUpdatedComponent, EUpdateTransformFlags InUpdateTransformFlags, ETeleportType InTeleport)
+{
+	InvalidateBakedStaticMesh(InUpdatedComponent->GetOwner());
+}
+
+void UAvaClonerComponent::OnMeshChanged(UStaticMeshComponent*, AActor* InActor)
+{
+	if (!InActor)
+	{
+		return;
+	}
+
+	const AActor* ClonerActor = GetOwner();
+
+	if (!ClonerActor)
+	{
+		return;
+	}
+
+	if (FAvaClonerAttachmentItem* Item = ClonerTree.ItemAttachmentMap.Find(InActor))
+	{
+		UE_LOG(LogAvaClonerComponent, Log, TEXT("%s : Detected mesh change for %s"), *ClonerActor->GetActorNameOrLabel(), *InActor->GetActorNameOrLabel());
+
+		Item->MeshStatus = EAvaClonerAttachmentStatus::Outdated;
+        InvalidateBakedStaticMesh(InActor);
+		DirtyItemAttachments.Add(Item);
+	}
+}
+
+void UAvaClonerComponent::GetOrderedRootActors(TArray<AActor*>& OutActors) const
+{
+	const AActor* ClonerActor = GetOwner();
+	if (!ClonerActor)
+	{
+		return;
+	}
+
+	UWorld* ClonerWorld = ClonerActor->GetTypedOuter<UWorld>();
+	if (!ClonerWorld)
+	{
+		return;
+	}
+
+#if WITH_EDITOR
+	if (const UAvaOutlinerSubsystem* const OutlinerSubsystem = ClonerWorld->GetSubsystem<UAvaOutlinerSubsystem>())
+	{
+		const TSharedPtr<FAvaOutliner> AvaOutliner = OutlinerSubsystem->GetOutliner();
+		if (AvaOutliner.IsValid())
+		{
+			const FAvaOutlinerItemPtr OutlinerClonerItem = AvaOutliner->FindItem(ClonerActor);
+			if (OutlinerClonerItem.IsValid())
+			{
+				for (FAvaOutlinerItemPtr OutlinerChild : OutlinerClonerItem->GetChildren())
+				{
+					if (!OutlinerChild.IsValid() || !OutlinerChild->IsA<FAvaOutlinerActor>())
+					{
+						continue;
+					}
+
+					if (const FAvaOutlinerActor* OutlinerActor = OutlinerChild->CastTo<FAvaOutlinerActor>())
+					{
+						OutActors.Add(OutlinerActor->GetActor());
+					}
+				}
+			}
+			return;
+		}
+	}
+#endif
+
+	if (const IAvaSceneInterface* SceneInterface = FAvaActorUtils::GetSceneInterfaceFromActor(ClonerActor))
+	{
+		TArray<AActor*> AttachedActors;
+		ClonerActor->GetAttachedActors(AttachedActors, true, false);
+		OutActors.Reserve(AttachedActors.Num());
+		const FAvaSceneTree& SceneTree = SceneInterface->GetSceneTree();
+
+		for (AActor* ChildActor : AttachedActors)
+		{
+			if (const FAvaSceneTreeNode* ChildNode = SceneTree.FindTreeNode(FAvaSceneItem(ChildActor, ClonerWorld)))
+			{
+				OutActors.Insert(ChildActor, ChildNode->GetLocalIndex());
+			}
+		}
+	}
+	else
+	{
+		// Fallback outside Avalanche
+		ClonerActor->GetAttachedActors(OutActors, true, false);
+	}
+}
+
+AActor* UAvaClonerComponent::GetRootActor(AActor* InActor) const
+{
+	if (InActor == nullptr)
+	{
+		return nullptr;
+	}
+	if (const FAvaClonerAttachmentItem* Item = ClonerTree.ItemAttachmentMap.Find(InActor))
+	{
+		if (Item->bRootItem)
+		{
+			return InActor;
+		}
+
+		return GetRootActor(Item->ParentActor.Get());
+	}
+	return nullptr;
+}
+
+void UAvaClonerComponent::InvalidateBakedStaticMesh(AActor* InActor)
+{
+	if (!InActor)
+	{
+		return;
+	}
+
+	if (const FAvaClonerAttachmentItem* FoundItem = ClonerTree.ItemAttachmentMap.Find(InActor))
+	{
+		if (FoundItem->bRootItem || !FoundItem->ParentActor.IsValid())
+		{
+			const int32 RootIdx = ClonerTree.RootActors.Find(InActor);
+			if (ClonerTree.MergedBakedMeshes.IsValidIndex(RootIdx))
+			{
+				ClonerTree.MergedBakedMeshes[RootIdx] = nullptr;
+				bClonerMeshesDirty = true;
+			}
+		}
+		else
+		{
+			InvalidateBakedStaticMesh(FoundItem->ParentActor.Get());
+		}
+	}
+}
+
+void UAvaClonerComponent::UpdateDirtyMeshesAsync()
+{
+	if (bClonerMeshesUpdating)
+	{
+		return;
+	}
+
+	bClonerMeshesUpdating = true;
+	TSet<FAvaClonerAttachmentItem*> UpdateItems = DirtyItemAttachments;
+	DirtyItemAttachments.Empty();
+
+	// Update baked dynamic meshes on other thread
+	TWeakObjectPtr<UAvaClonerComponent> ThisWeak(this);
+	Async(EAsyncExecution::TaskGraph, [ThisWeak, UpdateItems]()->bool
+	{
+		if (!ThisWeak.IsValid())
+		{
+			return false;
+		}
+
+		// update actor baked dynamic meshes
+		for (const FAvaClonerAttachmentItem* Item : UpdateItems)
+		{
+			if (!Item || !Item->ItemActor.IsValid())
+			{
+				continue;
+			}
+
+			ThisWeak->UpdateActorBakedDynamicMesh(Item->ItemActor.Get());
+		}
+
+		// Create baked static mesh on main thread (required)
+		Async(EAsyncExecution::TaskGraphMainThread, [ThisWeak]()
+		{
+			if (!ThisWeak.IsValid())
+			{
+				return;
+			}
+
+			// Update actors baked static mesh
+			for (int32 Idx = 0; Idx < ThisWeak->ClonerTree.RootActors.Num(); Idx++)
+			{
+				AActor* RootActor = ThisWeak->ClonerTree.RootActors[Idx].Get();
+				if (!ThisWeak->ClonerTree.MergedBakedMeshes[Idx].Get())
+				{
+					ThisWeak->UpdateRootActorBakedStaticMesh(RootActor);
+				}
+			}
+
+			// update niagara asset
+			ThisWeak->OnDirtyMeshesUpdated();
+		});
+
+		return true;
+	});
+}
+
+void UAvaClonerComponent::OnDirtyMeshesUpdated()
+{
+	bClonerMeshesUpdating = false;
+
+	// Update niagara parameters
+	UpdateClonerMeshes();
+}
+
+void UAvaClonerComponent::UpdateActorBakedDynamicMesh(AActor* InActor)
+{
+	if (!InActor)
+	{
+		return;
+	}
+
+	FAvaClonerAttachmentItem* AttachmentItem = ClonerTree.ItemAttachmentMap.Find(InActor);
+
+	if (!AttachmentItem || AttachmentItem->MeshStatus != EAvaClonerAttachmentStatus::Outdated)
+	{
+		return;
+	}
+
+	AActor* ClonerActor = GetOwner();
+
+	if (!ClonerActor)
+	{
+		return;
+	}
+
+	AttachmentItem->MeshStatus = EAvaClonerAttachmentStatus::Updating;
+
+	UE_LOG(LogAvaClonerComponent, Log, TEXT("%s : Updating baked actor mesh %s"), *ClonerActor->GetActorNameOrLabel(), *InActor->GetActorNameOrLabel());
+
+	using namespace UE::Geometry;
+
+	TArray<FDynamicMesh3> ConvertedMeshes;
+	TArray<TWeakObjectPtr<UMaterialInterface>> MeshesMaterials;
+
+	const FTransform& SourceTransform = InActor->GetActorTransform();
+	UDynamicMesh* OutputDynamicMesh = NewObject<UDynamicMesh>();
+
+	static FGeometryScriptCopyMeshFromAssetOptions OutputMeshOptions;
+	OutputMeshOptions.bIgnoreRemoveDegenerates = false;
+	OutputMeshOptions.bRequestTangents = false;
+	OutputMeshOptions.bApplyBuildSettings = false;
+
+	// Dynamic mesh components
+	{
+		TArray<UDynamicMeshComponent*> Components;
+		InActor->GetComponents(Components, false);
+		for (UDynamicMeshComponent* Component : Components)
+		{
+			const UDynamicMesh* DynamicMesh = Component->GetDynamicMesh();
+			// Transform the new mesh relative to the component
+			const FTransform RelativeTransform = Component->GetComponentTransform().GetRelativeTransform(SourceTransform);
+			MeshesMaterials.Append(Component->GetMaterials());
+			// Create a copy
+			DynamicMesh->ProcessMesh([&ConvertedMeshes, RelativeTransform](const FDynamicMesh3& EditMesh)
+			{
+				FDynamicMesh3 CopyMesh = EditMesh;
+				MeshTransforms::ApplyTransform(CopyMesh, RelativeTransform);
+				ConvertedMeshes.Add(MoveTemp(CopyMesh));
+			});
+		}
+	}
+
+	// Skeletal mesh components
+	{
+		static FGeometryScriptMeshReadLOD SkeletalMeshLOD;
+		SkeletalMeshLOD.LODType = EGeometryScriptLODType::SourceModel;
+
+		TArray<USkeletalMeshComponent*> Components;
+		InActor->GetComponents(Components, false);
+		for (USkeletalMeshComponent* Component : Components)
+		{
+			USkeletalMesh* SkeletalMesh = Component->GetSkeletalMeshAsset();
+
+			// convert to dynamic mesh
+			EGeometryScriptOutcomePins OutResult;
+			UGeometryScriptLibrary_StaticMeshFunctions::CopyMeshFromSkeletalMesh(SkeletalMesh, OutputDynamicMesh, OutputMeshOptions, SkeletalMeshLOD, OutResult);
+			if (OutResult == EGeometryScriptOutcomePins::Success)
+			{
+				// Transform the new mesh relative to the component
+				const FTransform RelativeTransform = Component->GetComponentTransform().GetRelativeTransform(SourceTransform);
+				MeshesMaterials.Append(Component->GetMaterials());
+				OutputDynamicMesh->EditMesh([&ConvertedMeshes, RelativeTransform](FDynamicMesh3& EditMesh)
+				{
+					MeshTransforms::ApplyTransform(EditMesh, RelativeTransform);
+					ConvertedMeshes.Add(MoveTemp(EditMesh));
+					// replace by empty mesh
+					FDynamicMesh3 EmptyMesh;
+					EditMesh = MoveTemp(EmptyMesh);
+				}, EDynamicMeshChangeType::GeneralEdit, EDynamicMeshAttributeChangeFlags::Unknown, true);
+			}
+		}
+	}
+
+	// Brush components
+	{
+		static const FGeometryScriptCopyMeshFromComponentOptions Options;
+		static FTransform Transform;
+
+		TArray<UBrushComponent*> Components;
+		InActor->GetComponents(Components, false);
+		for (UBrushComponent* Component : Components)
+		{
+			// convert to dynamic mesh
+			EGeometryScriptOutcomePins OutResult;
+			UGeometryScriptLibrary_SceneUtilityFunctions::CopyMeshFromComponent(Component, OutputDynamicMesh, Options, false, Transform, OutResult);
+			if (OutResult == EGeometryScriptOutcomePins::Success)
+			{
+				// Transform the new mesh relative to the component
+				const FTransform RelativeTransform = Component->GetComponentTransform().GetRelativeTransform(SourceTransform);
+				for (int32 MatIdx = 0; MatIdx < Component->GetNumMaterials(); MatIdx++)
+				{
+					MeshesMaterials.Add(Component->GetMaterial(MatIdx));
+				}
+				OutputDynamicMesh->EditMesh([&ConvertedMeshes, RelativeTransform](FDynamicMesh3& EditMesh)
+				{
+					MeshTransforms::ApplyTransform(EditMesh, RelativeTransform);
+					ConvertedMeshes.Add(MoveTemp(EditMesh));
+					// replace by empty mesh
+					FDynamicMesh3 EmptyMesh;
+					EditMesh = MoveTemp(EmptyMesh);
+				}, EDynamicMeshChangeType::GeneralEdit, EDynamicMeshAttributeChangeFlags::Unknown, true);
+			}
+		}
+	}
+
+	// Procedural mesh components
+	{
+		TArray<UProceduralMeshComponent*> Components;
+		InActor->GetComponents(Components, false);
+		for (UProceduralMeshComponent* Component : Components)
+		{
+			const int32 SectionCount = Component->GetNumSections();
+			if (SectionCount == 0)
+			{
+				continue;
+			}
+
+			// Transform the new mesh relative to the component
+			const FTransform RelativeTransform = Component->GetComponentTransform().GetRelativeTransform(SourceTransform);
+
+			FDynamicMesh3 Mesh;
+			Mesh.EnableAttributes();
+			Mesh.Attributes()->EnablePrimaryColors();
+			Mesh.Attributes()->EnableMaterialID();
+			Mesh.Attributes()->SetNumNormalLayers(1);
+			Mesh.Attributes()->SetNumUVLayers(1);
+			Mesh.Attributes()->SetNumPolygroupLayers(1);
+			Mesh.Attributes()->EnableTangents();
+
+			FDynamicMeshColorOverlay* ColorOverlay = Mesh.Attributes()->PrimaryColors();
+			FDynamicMeshNormalOverlay* NormalOverlay = Mesh.Attributes()->PrimaryNormals();
+			FDynamicMeshUVOverlay* UVOverlay = Mesh.Attributes()->PrimaryUV();
+			FDynamicMeshMaterialAttribute* MaterialAttr = Mesh.Attributes()->GetMaterialID();
+			FDynamicMeshPolygroupAttribute* GroupAttr = Mesh.Attributes()->GetPolygroupLayer(0);
+			FDynamicMeshNormalOverlay* TangentOverlay = Mesh.Attributes()->PrimaryTangents();
+
+			for (int32 SectionIdx = 0; SectionIdx < SectionCount; SectionIdx++)
+			{
+				if (FProcMeshSection* Section = Component->GetProcMeshSection(SectionIdx))
+				{
+					if (Section->bSectionVisible)
+					{
+						TArray<int32> VtxIds;
+						TArray<int32> NormalIds;
+						TArray<int32> ColorIds;
+						TArray<int32> UVIds;
+						TArray<int32> TaIds;
+
+						// copy vertices data (position, normal, color, UV, tangent)
+						for (FProcMeshVertex& SectionVertex : Section->ProcVertexBuffer)
+						{
+							int32 VId = Mesh.AppendVertex(SectionVertex.Position);
+							VtxIds.Add(VId);
+
+							int32 NId = NormalOverlay->AppendElement(static_cast<FVector3f>(SectionVertex.Normal));
+							NormalIds.Add(NId);
+
+							int32 CId = ColorOverlay->AppendElement(static_cast<FVector4f>(SectionVertex.Color));
+							ColorIds.Add(CId);
+
+							int32 UVId = UVOverlay->AppendElement(static_cast<FVector2f>(SectionVertex.UV0));
+							UVIds.Add(UVId);
+
+							int32 TaId = TangentOverlay->AppendElement(static_cast<FVector3f>(SectionVertex.Tangent.TangentX));
+							TaIds.Add(TaId);
+						}
+
+						// copy tris data
+						if (Section->ProcIndexBuffer.Num() % 3 != 0)
+						{
+							continue;
+						}
+						for (int32 Idx = 0; Idx < Section->ProcIndexBuffer.Num(); Idx+=3)
+						{
+							int32 VIdx1 = Section->ProcIndexBuffer[Idx];
+							int32 VIdx2 = Section->ProcIndexBuffer[Idx + 1];
+							int32 VIdx3 = Section->ProcIndexBuffer[Idx + 2];
+
+							int32 VId1 = VtxIds[VIdx1];
+							int32 VId2 = VtxIds[VIdx2];
+							int32 VId3 = VtxIds[VIdx3];
+
+							int32 TId = Mesh.AppendTriangle(VId1, VId2, VId3, SectionIdx);
+							if (TId < 0)
+							{
+								continue;
+							}
+
+							NormalOverlay->SetTriangle(TId, FIndex3i(NormalIds[VIdx1], NormalIds[VIdx2], NormalIds[VIdx3]), true);
+							ColorOverlay->SetTriangle(TId, FIndex3i(ColorIds[VIdx1], ColorIds[VIdx2], ColorIds[VIdx3]), true);
+							UVOverlay->SetTriangle(TId, FIndex3i(UVIds[VIdx1], UVIds[VIdx2], UVIds[VIdx3]), true);
+							TangentOverlay->SetTriangle(TId, FIndex3i(TaIds[VIdx1], TaIds[VIdx2], TaIds[VIdx3]), true);
+
+							MaterialAttr->SetValue(TId, SectionIdx);
+							GroupAttr->SetValue(TId, SectionIdx);
+						}
+					}
+				}
+			}
+
+			if (Mesh.TriangleCount() > 0)
+			{
+				MeshTransforms::ApplyTransform(Mesh, RelativeTransform);
+				MeshesMaterials.Append(Component->GetMaterials());
+				ConvertedMeshes.Add(MoveTemp(Mesh));
+			}
+		}
+	}
+
+	// Static mesh components
+	{
+		static FGeometryScriptMeshReadLOD StaticMeshLOD;
+		StaticMeshLOD.LODType = EGeometryScriptLODType::RenderData;
+
+		TArray<UStaticMeshComponent*> Components;
+		InActor->GetComponents(Components, false);
+		for (UStaticMeshComponent* Component : Components)
+		{
+			UStaticMesh* StaticMesh = Component->GetStaticMesh();
+			// convert to dynamic mesh
+			EGeometryScriptOutcomePins OutResult;
+			UGeometryScriptLibrary_StaticMeshFunctions::CopyMeshFromStaticMesh(StaticMesh, OutputDynamicMesh, OutputMeshOptions, StaticMeshLOD, OutResult);
+			if (OutResult == EGeometryScriptOutcomePins::Success)
+			{
+				// Transform the new mesh relative to the component
+				const FTransform RelativeTransform = Component->GetComponentTransform().GetRelativeTransform(SourceTransform);
+				MeshesMaterials.Append(Component->GetMaterials());
+				OutputDynamicMesh->EditMesh([&ConvertedMeshes, RelativeTransform](FDynamicMesh3& EditMesh)
+				{
+					MeshTransforms::ApplyTransform(EditMesh, RelativeTransform);
+					ConvertedMeshes.Add(MoveTemp(EditMesh));
+					// replace by empty mesh
+					FDynamicMesh3 EmptyMesh;
+					EditMesh = MoveTemp(EmptyMesh);
+				}, EDynamicMeshChangeType::GeneralEdit, EDynamicMeshAttributeChangeFlags::Unknown, true);
+			}
+		}
+	}
+
+	// Lets combine all meshes components from this actor together
+	OutputDynamicMesh->EditMesh([&ConvertedMeshes](FDynamicMesh3& InMergedMesh)
+	{
+		InMergedMesh.Clear();
+		InMergedMesh.EnableAttributes();
+		InMergedMesh.Attributes()->SetNumNormalLayers(1);
+
+		FDynamicMeshEditor Editor(&InMergedMesh);
+		FGeometryScriptAppendMeshOptions AppendOptions;
+		AppendOptions.CombineMode = EGeometryScriptCombineAttributesMode::EnableAllMatching;
+		int32 MaterialCount = 0;
+
+		// Convert meshes
+		for (int32 MeshIdx = 0; MeshIdx < ConvertedMeshes.Num(); MeshIdx++)
+		{
+			const FDynamicMesh3& ConvertedMesh = ConvertedMeshes[MeshIdx];
+
+			// Enable matching attributes & append mesh
+			FMeshIndexMappings TmpMappings;
+			AppendOptions.UpdateAttributesForCombineMode(InMergedMesh, ConvertedMesh);
+			Editor.AppendMesh(&ConvertedMesh, TmpMappings);
+
+			// Fix triangles materials linking
+			if (ConvertedMesh.HasAttributes() && ConvertedMesh.Attributes()->HasMaterialID())
+			{
+				const FDynamicMeshMaterialAttribute* FromMaterialIDAttrib = ConvertedMesh.Attributes()->GetMaterialID();
+				FDynamicMeshMaterialAttribute* ToMaterialIDAttrib = InMergedMesh.Attributes()->GetMaterialID();
+				TMap<int32, int32> MaterialMap;
+				for (const TPair<int32, int32>& FromToTId : TmpMappings.GetTriangleMap().GetForwardMap())
+				{
+					const int32 FromMatId = FromMaterialIDAttrib->GetValue(FromToTId.Key);
+					const int32 ToMatId = FromMatId + MaterialCount;
+					MaterialMap.Add(FromMatId, ToMatId);
+					ToMaterialIDAttrib->SetNewValue(FromToTId.Value, ToMatId);
+				}
+				MaterialCount += MaterialMap.Num();
+			}
+		}
+
+		if (InMergedMesh.TriangleCount() > 0)
+		{
+			// Merge shared edges
+			FMergeCoincidentMeshEdges WeldOp(&InMergedMesh);
+			WeldOp.Apply();
+		}
+	}, EDynamicMeshChangeType::GeneralEdit, EDynamicMeshAttributeChangeFlags::Unknown, true);
+
+	AttachmentItem->BakedMesh = OutputDynamicMesh;
+	AttachmentItem->BakedMaterials = MeshesMaterials;
+	// Was the mesh invalidated during the update process ?
+	AttachmentItem->MeshStatus = AttachmentItem->MeshStatus == EAvaClonerAttachmentStatus::Outdated ? EAvaClonerAttachmentStatus::Outdated : EAvaClonerAttachmentStatus::Updated;
+	InvalidateBakedStaticMesh(AttachmentItem->ItemActor.Get());
+}
+
+void UAvaClonerComponent::UpdateRootActorBakedStaticMesh(AActor* InRootActor)
+{
+	if (!InRootActor)
+	{
+		return;
+	}
+
+	const int32 RootIdx = ClonerTree.RootActors.Find(InRootActor);
+
+	if (RootIdx == INDEX_NONE)
+	{
+		return;
+	}
+
+	const FAvaClonerAttachmentItem* RootAttachmentItem = ClonerTree.ItemAttachmentMap.Find(InRootActor);
+
+	if (!RootAttachmentItem)
+	{
+		return;
+	}
+
+	const AActor* ClonerActor = GetOwner();
+
+	if (!ClonerActor)
+	{
+		return;
+	}
+
+	UE_LOG(LogAvaClonerComponent, Log, TEXT("%s : Updating root merged baked mesh %s"), *ClonerActor->GetActorNameOrLabel(), *InRootActor->GetActorNameOrLabel());
+
+	using namespace UE::Geometry;
+
+#if WITH_EDITOR
+	TSet<int32> MaterialsMissingNiagaraUsageFlag;
+#endif
+
+	TArray<FAvaClonerAttachmentItem*> AttachmentItems;
+	GetActorAttachmentItems(InRootActor, AttachmentItems);
+
+	// Lets combine all dynamic meshes from item attachments and self into one mesh
+	UDynamicMesh* MergedAttachmentMesh = NewObject<UDynamicMesh>();
+	TArray<TObjectPtr<UMaterialInterface>> CombineMaterials;
+	MergedAttachmentMesh->EditMesh([&AttachmentItems, &CombineMaterials, this
+#if WITH_EDITOR
+		, &MaterialsMissingNiagaraUsageFlag
+#endif
+		](FDynamicMesh3& InMergedMesh)
+	{
+		InMergedMesh.Clear();
+		InMergedMesh.EnableAttributes();
+		InMergedMesh.Attributes()->SetNumNormalLayers(1);
+
+		FDynamicMeshEditor Editor(&InMergedMesh);
+		FGeometryScriptAppendMeshOptions AppendOptions;
+		AppendOptions.CombineMode = EGeometryScriptCombineAttributesMode::EnableAllMatching;
+		int32 MaterialCount = 0;
+
+		// Convert meshes
+		for (int32 MeshIdx = 0; MeshIdx < AttachmentItems.Num(); MeshIdx++)
+		{
+			FAvaClonerAttachmentItem* AttachmentItem = AttachmentItems[MeshIdx];
+			if (!AttachmentItem)
+			{
+				continue;
+			}
+
+			const UDynamicMesh* BakedDynamicMesh = AttachmentItem->BakedMesh;
+			if (!BakedDynamicMesh)
+			{
+				continue;
+			}
+
+			BakedDynamicMesh->ProcessMesh([&AttachmentItem, &CombineMaterials, &MaterialCount, &AppendOptions, &InMergedMesh, &Editor, this
+#if WITH_EDITOR
+				, &MaterialsMissingNiagaraUsageFlag
+#endif
+				](const FDynamicMesh3& InProcessMesh)
+			{
+				FDynamicMesh3 ConvertedMesh = InProcessMesh;
+
+				if (AttachmentItem->ParentActor.IsValid())
+				{
+					const FTransform& ParentTransform = AttachmentItem->ParentActor->GetTransform();
+					const FTransform ParentChildTransform = AttachmentItem->ItemActor->GetTransform().GetRelativeTransform(ParentTransform);
+					MeshTransforms::ApplyTransform(ConvertedMesh, ParentChildTransform);
+				}
+				else
+				{
+					MeshTransforms::ApplyTransform(ConvertedMesh, FTransform::Identity);
+				}
+
+				// Enable matching attributes & append mesh
+				FMeshIndexMappings TmpMappings;
+				AppendOptions.UpdateAttributesForCombineMode(InMergedMesh, ConvertedMesh);
+				Editor.AppendMesh(&ConvertedMesh, TmpMappings);
+
+				// Fix triangles materials linking
+				if (ConvertedMesh.HasAttributes() && ConvertedMesh.Attributes()->HasMaterialID())
+				{
+					const FDynamicMeshMaterialAttribute* FromMaterialIDAttrib = ConvertedMesh.Attributes()->GetMaterialID();
+					FDynamicMeshMaterialAttribute* ToMaterialIDAttrib = InMergedMesh.Attributes()->GetMaterialID();
+					TMap<int32, int32> MaterialMap;
+					for (const TPair<int32, int32>& FromToTId : TmpMappings.GetTriangleMap().GetForwardMap())
+					{
+						const int32 FromMatId = FromMaterialIDAttrib->GetValue(FromToTId.Key);
+						const int32 ToMatId = FromMatId + MaterialCount;
+						MaterialMap.Add(FromMatId, ToMatId);
+						ToMaterialIDAttrib->SetNewValue(FromToTId.Value, ToMatId);
+					}
+					MaterialCount += MaterialMap.Num();
+
+					for (TWeakObjectPtr<UMaterialInterface>& BakedMaterial : AttachmentItem->BakedMaterials)
+					{
+						UMaterialInterface* MaterialInterface = BakedMaterial.Get();
+						int32 MatIdx = CombineMaterials.Add(MaterialInterface);
+#if WITH_EDITOR
+						if (MaterialInterface)
+						{
+							const UMaterial* Material = MaterialInterface->GetMaterial();
+							if (Material && !Material->GetUsageByFlag(EMaterialUsage::MATUSAGE_NiagaraMeshParticles))
+							{
+								MaterialsMissingNiagaraUsageFlag.Add(MatIdx);
+							}
+						}
+#endif
+					}
+				}
+			});
+		}
+	}, EDynamicMeshChangeType::GeneralEdit, EDynamicMeshAttributeChangeFlags::Unknown, true);
+
+#if WITH_EDITOR
+	if (!MaterialsMissingNiagaraUsageFlag.IsEmpty())
+	{
+		int32 ReadOnlyMaterialCount = 0;
+
+		for (const int32 MaterialIdx : MaterialsMissingNiagaraUsageFlag)
+		{
+			UMaterialInterface* Material = CombineMaterials[MaterialIdx];
+			if (!Material)
+			{
+				continue;
+			}
+
+			if (!IsMaterialDirtyable(Material))
+			{
+				ReadOnlyMaterialCount++;
+				CombineMaterials[MaterialIdx] = nullptr;
+				UE_LOG(LogAvaClonerComponent, Warning, TEXT("%s : The following materials (%s) on actor (%s) does not have the required usage flag (bUsedWithNiagaraMeshParticles) to work with the cloner, this material cannot be dirtied due to its read-only location, skipping material and proceeding"), *ClonerActor->GetActorNameOrLabel(), Material ? *Material->GetMaterial()->GetPathName() : TEXT("Invalid Material"), *InRootActor->GetActorNameOrLabel());
+			}
+			else
+			{
+				UE_LOG(LogAvaClonerComponent, Log, TEXT("%s : The following materials (%s) on actor (%s) does not have the required usage flag (bUsedWithNiagaraMeshParticles) to work with the cloner, dirtying material and proceeding, please save the asset for working runtime result"), *ClonerActor->GetActorNameOrLabel(), Material ? *Material->GetMaterial()->GetPathName() : TEXT("Invalid Material"), *InRootActor->GetActorNameOrLabel());
+			}
+		}
+
+		if (ReadOnlyMaterialCount > 0)
+		{
+			FNotificationInfo NotificationInfo(FText::Format(LOCTEXT("MaterialsMissingUsageFlag", "Detected {0} read-only material(s) with missing niagara usage flag required to work properly with cloner (See logs)"), ReadOnlyMaterialCount));
+			NotificationInfo.ExpireDuration = 5.f;
+			NotificationInfo.bFireAndForget = true;
+			NotificationInfo.Image = FAppStyle::GetBrush("Icons.WarningWithColor");
+			FSlateNotificationManager::Get().AddNotification(NotificationInfo);
+		}
+	}
+#endif
+
+	// Convert merged mesh to static mesh
+	UStaticMesh* BakedAttachmentMesh = NewObject<UStaticMesh>();
+	if (MergedAttachmentMesh->GetTriangleCount())
+	{
+		// export options
+		FGeometryScriptCopyMeshToAssetOptions AssetOptions;
+		AssetOptions.bReplaceMaterials = true;
+		AssetOptions.bEnableRecomputeNormals = false;
+		AssetOptions.bEnableRecomputeTangents = false;
+		AssetOptions.bEnableRemoveDegenerates = true;
+		AssetOptions.bEmitTransaction = false;
+		AssetOptions.bApplyNaniteSettings = false;
+		AssetOptions.bDeferMeshPostEditChange = false;
+		AssetOptions.NewMaterials = CombineMaterials;
+		// LOD options
+		static FGeometryScriptMeshWriteLOD TargetLOD;
+		TargetLOD.LODIndex = 0;
+		TargetLOD.bWriteHiResSource = false;
+		// result
+		EGeometryScriptOutcomePins OutResult;
+		UGeometryScriptLibrary_StaticMeshFunctions::CopyMeshToStaticMesh(MergedAttachmentMesh, BakedAttachmentMesh, AssetOptions, TargetLOD, OutResult);
+
+		// Fix camera culling and avoid killing particles when bounds partially in view
+		const float NewRadiusBounds = BakedAttachmentMesh->GetBounds().SphereRadius * 2;
+		BakedAttachmentMesh->SetExtendedBounds(FBox(-FVector(NewRadiusBounds), FVector(NewRadiusBounds)));
+
+#if WITH_EDITORONLY_DATA
+		// Compute normals and tangents
+		FMeshDescription* MeshDescription = BakedAttachmentMesh->GetMeshDescription(0);
+		FStaticMeshOperations::ComputeTriangleTangentsAndNormals(*MeshDescription);
+		FStaticMeshOperations::ComputeTangentsAndNormals(*MeshDescription, EComputeNTBsFlags::Normals);
+#endif
+
+		bClonerMeshesDirty = true;
+	}
+
+	MergedAttachmentMesh->MarkAsGarbage();
+	ClonerTree.MergedBakedMeshes[RootIdx] = BakedAttachmentMesh;
+}
+
+void UAvaClonerComponent::GetActorAttachmentItems(AActor* InActor, TArray<FAvaClonerAttachmentItem*>& OutAttachmentItems)
+{
+	if (!InActor)
+	{
+		return;
+	}
+
+	FAvaClonerAttachmentItem* AttachmentItem = ClonerTree.ItemAttachmentMap.Find(InActor);
+
+	if (!AttachmentItem)
+	{
+		return;
+	}
+
+	OutAttachmentItems.Add(AttachmentItem);
+
+	for (TWeakObjectPtr<AActor>& ChildActor : AttachmentItem->ChildrenActors)
+	{
+		if (ChildActor.IsValid())
+		{
+			GetActorAttachmentItems(ChildActor.Get(), OutAttachmentItems);
+		}
+	}
+}
+
+bool UAvaClonerComponent::IsAllMergedMeshesValid() const
+{
+	for (const TObjectPtr<UStaticMesh>& MergedMesh : ClonerTree.MergedBakedMeshes)
+	{
+		if (!MergedMesh)
+		{
+			return false;
+		}
+	}
+	return true;
+}
+
+bool UAvaClonerComponent::IsMaterialDirtyable(const UMaterialInterface* InMaterial) const
+{
+	const UMaterial* BaseMaterial = InMaterial->GetMaterial();
+	const FString ContentFolder = FPaths::ConvertRelativePathToFull(FPaths::ProjectContentDir());
+
+	const UPackage* MaterialPackage = BaseMaterial->GetPackage();
+	const FPackagePath& LoadedPath = MaterialPackage->GetLoadedPath();
+	const FString PackagePath = FPaths::ConvertRelativePathToFull(LoadedPath.GetLocalFullPath());
+	const FString MaterialPath = BaseMaterial->GetPathName();
+
+	const bool bTransientPackage = MaterialPackage == GetTransientPackage() || MaterialPath.StartsWith("/Temp/");
+	const bool bContentFolder = PackagePath.StartsWith(ContentFolder);
+
+	return bTransientPackage || bContentFolder;
+}
+
+void UAvaClonerComponent::UpdateClonerMeshes()
+{
+	const AActor* ClonerActor = GetOwner();
+	if (!ClonerActor)
+	{
+		return;
+	}
+
+	if (!bClonerMeshesDirty)
+	{
+		return;
+	}
+
+	if (!GetAsset() || !ActiveLayout)
+	{
+		return;
+	}
+
+	UNiagaraMeshRendererProperties* MeshRenderer = ActiveLayout->GetMeshRenderer();
+
+	if (!MeshRenderer)
+	{
+		UE_LOG(LogAvaClonerComponent, Warning, TEXT("%s : Invalid mesh renderer for cloner system"), *ClonerActor->GetActorNameOrLabel());
+		return;
+	}
+
+	// Set new number of meshes in renderer
+	static const FName MeshNumName(TEXT("MeshNum"));
+	SetIntParameter(MeshNumName, ClonerTree.MergedBakedMeshes.Num());
+
+	// Resize mesh array properly
+	if (MeshRenderer->Meshes.Num() > ClonerTree.MergedBakedMeshes.Num())
+	{
+		MeshRenderer->Meshes.SetNum(ClonerTree.MergedBakedMeshes.Num());
+	}
+
+	// Set baked meshes in mesh renderer array
+	for (int32 Idx = 0; Idx < ClonerTree.MergedBakedMeshes.Num(); Idx++)
+	{
+		UStaticMesh* StaticMesh = ClonerTree.MergedBakedMeshes[Idx];
+		FNiagaraMeshRendererMeshProperties& MeshProperties = !MeshRenderer->Meshes.IsValidIndex(Idx) ? MeshRenderer->Meshes.AddDefaulted_GetRef() : MeshRenderer->Meshes[Idx];
+		MeshProperties.Mesh = StaticMesh && StaticMesh->GetNumTriangles(0) > 0 ? StaticMesh : nullptr;
+
+		if (ClonerTree.RootActors.IsValidIndex(Idx))
+		{
+			if (const FAvaClonerAttachmentItem* Item = ClonerTree.ItemAttachmentMap.Find(ClonerTree.RootActors[Idx]))
+			{
+				MeshProperties.Rotation = Item->ActorTransform.Rotator();
+				MeshProperties.Scale = Item->ActorTransform.GetScale3D();
+			}
+		}
+	}
+
+#if WITH_EDITORONLY_DATA
+	MeshRenderer->OnMeshChanged();
+#endif
+
+	bClonerMeshesDirty = !DirtyItemAttachments.IsEmpty();
+
+	UE_LOG(LogAvaClonerComponent, Log, TEXT("%s : Cloner mesh updated %i"), *ClonerActor->GetActorNameOrLabel(), ClonerTree.MergedBakedMeshes.Num());
+
+	OnClonerMeshUpdated.Broadcast();
+}
+
+bool UAvaClonerComponent::SetClonerActiveLayout(UAvaClonerLayoutBase* InLayout)
+{
+	if (!InLayout)
+	{
+		return false;
+	}
+
+	const AActor* ClonerActor = GetOwner();
+	if (!ClonerActor)
+	{
+		return false;
+	}
+
+	if (!InLayout->IsLayoutLoaded())
+	{
+		// Load system
+		if (!InLayout->LoadLayout())
+		{
+			UE_LOG(LogAvaClonerComponent, Warning, TEXT("%s : Cloner layout system failed to load %s - %s"), *GetOwner()->GetActorNameOrLabel(), *InLayout->GetLayoutName().ToString(), *InLayout->GetLayoutAssetPath());
+			return false;
+		}
+
+		UE_LOG(LogAvaClonerComponent, Log, TEXT("%s : Cloner layout system loaded %s - %s"), *GetOwner()->GetActorNameOrLabel(), *InLayout->GetLayoutName().ToString(), *InLayout->GetLayoutAssetPath());
+
+		OnClonerSystemLoaded.Broadcast(InLayout);
+	}
+
+	// Copy data interfaces to new layout
+	if (ActiveLayout && ActiveLayout->IsLayoutLoaded())
+	{
+		ActiveLayout->CopyTo(InLayout);
+	}
+
+	// Deactivate previous layout
+	if (ActiveLayout && ActiveLayout->IsLayoutActive())
+	{
+		ActiveLayout->DeactivateLayout();
+	}
+
+	// Activate new layout
+	InLayout->ActivateLayout();
+
+	ActiveLayout = InLayout;
+	bClonerMeshesDirty = true;
+
+	UE_LOG(LogAvaClonerComponent, Log, TEXT("%s : Cloner layout system changed %s - %s"), *ClonerActor->GetActorNameOrLabel(), *InLayout->GetLayoutName().ToString(), *InLayout->GetLayoutAssetPath());
+
+	return true;
+}
+
+#undef LOCTEXT_NAMESPACE
diff --git a/Engine/Plugins/Experimental/Avalanche/Source/AvalancheEffectors/Private/Cloner/Layouts/AvaClonerCircleLayout.cpp b/Engine/Plugins/Experimental/Avalanche/Source/AvalancheEffectors/Private/Cloner/Layouts/AvaClonerCircleLayout.cpp
new file mode 100644
--- /dev/null
+++ b/Engine/Plugins/Experimental/Avalanche/Source/AvalancheEffectors/Private/Cloner/Layouts/AvaClonerCircleLayout.cpp
@@ -0,0 +1,147 @@
+// Copyright Epic Games, Inc. All Rights Reserved.
+
+#include "Cloner/Layouts/AvaClonerCircleLayout.h"
+#include "Cloner/AvaClonerComponent.h"
+
+void UAvaClonerCircleLayout::SetCount(int32 InCount)
+{
+	if (Count == InCount)
+	{
+		return;
+	}
+	
+	Count = InCount;
+	UpdateLayoutParameters();
+}
+
+void UAvaClonerCircleLayout::SetRadius(float InRadius)
+{
+	if (Radius == InRadius)
+	{
+		return;
+	}
+
+	Radius = InRadius;
+	UpdateLayoutParameters();
+}
+
+void UAvaClonerCircleLayout::SetAngleStart(float InAngleStart)
+{
+	if (AngleStart == InAngleStart)
+	{
+		return;
+	}
+
+	AngleStart = InAngleStart;
+	UpdateLayoutParameters();
+}
+
+void UAvaClonerCircleLayout::SetAngleRatio(float InAngleRatio)
+{
+	if (AngleRatio == InAngleRatio)
+	{
+		return;
+	}
+
+	AngleRatio = InAngleRatio;
+	UpdateLayoutParameters();
+}
+
+void UAvaClonerCircleLayout::SetOrientMesh(bool bInOrientMesh)
+{
+	if (bOrientMesh == bInOrientMesh)
+	{
+		return;
+	}
+
+	bOrientMesh = bInOrientMesh;
+	UpdateLayoutParameters();
+}
+
+void UAvaClonerCircleLayout::SetPlane(EAvaClonerPlane InPlane)
+{
+	if (Plane == InPlane)
+	{
+		return;
+	}
+
+	Plane = InPlane;
+	UpdateLayoutParameters();
+}
+
+void UAvaClonerCircleLayout::SetRotation(const FRotator& InRotation)
+{
+	if (Rotation == InRotation)
+	{
+		return;
+	}
+
+	Rotation = InRotation;
+	UpdateLayoutParameters();
+}
+
+void UAvaClonerCircleLayout::SetScale(const FVector& InScale)
+{
+	if (Scale == InScale)
+	{
+		return;
+	}
+
+	Scale = InScale;
+	UpdateLayoutParameters();
+}
+
+#if WITH_EDITOR
+const TAvaPropertyChangeDispatcher<UAvaClonerCircleLayout> UAvaClonerCircleLayout::PropertyChangeDispatcher =
+{
+	{ GET_MEMBER_NAME_CHECKED(UAvaClonerCircleLayout, Count), &UAvaClonerCircleLayout::OnLayoutPropertyChanged },
+	{ GET_MEMBER_NAME_CHECKED(UAvaClonerCircleLayout, Radius), &UAvaClonerCircleLayout::OnLayoutPropertyChanged },
+	{ GET_MEMBER_NAME_CHECKED(UAvaClonerCircleLayout, AngleStart), &UAvaClonerCircleLayout::OnLayoutPropertyChanged },
+	{ GET_MEMBER_NAME_CHECKED(UAvaClonerCircleLayout, AngleRatio), &UAvaClonerCircleLayout::OnLayoutPropertyChanged },
+	{ GET_MEMBER_NAME_CHECKED(UAvaClonerCircleLayout, bOrientMesh), &UAvaClonerCircleLayout::OnLayoutPropertyChanged },
+	{ GET_MEMBER_NAME_CHECKED(UAvaClonerCircleLayout, Plane), &UAvaClonerCircleLayout::OnLayoutPropertyChanged },
+	{ GET_MEMBER_NAME_CHECKED(UAvaClonerCircleLayout, Rotation), &UAvaClonerCircleLayout::OnLayoutPropertyChanged },
+	{ GET_MEMBER_NAME_CHECKED(UAvaClonerCircleLayout, Scale), &UAvaClonerCircleLayout::OnLayoutPropertyChanged },
+};
+
+void UAvaClonerCircleLayout::PostEditChangeProperty(FPropertyChangedEvent& InPropertyChangedEvent)
+{
+	Super::PostEditChangeProperty(InPropertyChangedEvent);
+
+	PropertyChangeDispatcher.OnPropertyChanged(this, InPropertyChangedEvent);
+}
+#endif
+
+void UAvaClonerCircleLayout::OnLayoutParametersChanged(UAvaClonerComponent* InComponent)
+{
+	Super::OnLayoutParametersChanged(InComponent);
+
+	InComponent->SetIntParameter(TEXT("CircleCount"), Count);
+		
+	InComponent->SetFloatParameter(TEXT("CircleRadius"), Radius);
+		
+	InComponent->SetFloatParameter(TEXT("CircleStart"), AngleStart);
+
+	InComponent->SetFloatParameter(TEXT("CircleRatio"), AngleRatio);
+		
+	InComponent->SetBoolParameter(TEXT("MeshOrientAxisEnable"), bOrientMesh);
+
+	FVector CircleRotation(Rotation.Yaw, Rotation.Pitch, Rotation.Roll);
+
+	if (Plane == EAvaClonerPlane::XY)
+	{
+		CircleRotation = FVector(0);
+	}
+	else if (Plane == EAvaClonerPlane::YZ)
+	{
+		CircleRotation = FVector(0, 90, 0);
+	}
+	else if (Plane == EAvaClonerPlane::XZ)
+	{
+		CircleRotation = FVector(0, 0, 90);
+	}
+
+	InComponent->SetVectorParameter(TEXT("CircleRotation"), CircleRotation);
+
+	InComponent->SetVectorParameter(TEXT("CircleScale"), Scale);
+}
