[MotionDesign]

diff --git a/Engine/Plugins/Experimental/Avalanche/Source/AvalancheEffectors/Private/Cloner/Layouts/AvaClonerCylinderLayout.cpp b/Engine/Plugins/Experimental/Avalanche/Source/AvalancheEffectors/Private/Cloner/Layouts/AvaClonerCylinderLayout.cpp
new file mode 100644
--- /dev/null
+++ b/Engine/Plugins/Experimental/Avalanche/Source/AvalancheEffectors/Private/Cloner/Layouts/AvaClonerCylinderLayout.cpp
@@ -0,0 +1,175 @@
+// Copyright Epic Games, Inc. All Rights Reserved.
+
+#include "Cloner/Layouts/AvaClonerCylinderLayout.h"
+#include "Cloner/AvaClonerComponent.h"
+
+void UAvaClonerCylinderLayout::SetBaseCount(int32 InBaseCount)
+{
+	if (BaseCount == InBaseCount)
+	{
+		return;
+	}
+
+	BaseCount = InBaseCount;
+	UpdateLayoutParameters();
+}
+
+void UAvaClonerCylinderLayout::SetHeightCount(int32 InHeightCount)
+{
+	if (HeightCount == InHeightCount)
+	{
+		return;
+	}
+
+	HeightCount = InHeightCount;
+	UpdateLayoutParameters();
+}
+
+void UAvaClonerCylinderLayout::SetHeight(float InHeight)
+{
+	if (Height == InHeight)
+	{
+		return;
+	}
+
+	Height = InHeight;
+	UpdateLayoutParameters();
+}
+
+void UAvaClonerCylinderLayout::SetRadius(float InRadius)
+{
+	if (Radius == InRadius)
+	{
+		return;
+	}
+
+	Radius = InRadius;
+	UpdateLayoutParameters();
+}
+
+void UAvaClonerCylinderLayout::SetAngleStart(float InAngleStart)
+{
+	if (AngleStart == InAngleStart)
+	{
+		return;
+	}
+
+	AngleStart = InAngleStart;
+	UpdateLayoutParameters();
+}
+
+void UAvaClonerCylinderLayout::SetAngleRatio(float InAngleRatio)
+{
+	if (AngleRatio == InAngleRatio)
+	{
+		return;
+	}
+
+	AngleRatio = InAngleRatio;
+	UpdateLayoutParameters();
+}
+
+void UAvaClonerCylinderLayout::SetOrientMesh(bool bInOrientMesh)
+{
+	if (bOrientMesh == bInOrientMesh)
+	{
+		return;
+	}
+
+	bOrientMesh = bInOrientMesh;
+	UpdateLayoutParameters();
+}
+
+void UAvaClonerCylinderLayout::SetPlane(EAvaClonerPlane InPlane)
+{
+	if (Plane == InPlane)
+	{
+		return;
+	}
+
+	Plane = InPlane;
+	UpdateLayoutParameters();
+}
+
+void UAvaClonerCylinderLayout::SetRotation(const FRotator& InRotation)
+{
+	if (Rotation == InRotation)
+	{
+		return;
+	}
+
+	Rotation = InRotation;
+	UpdateLayoutParameters();
+}
+
+void UAvaClonerCylinderLayout::SetScale(const FVector& InScale)
+{
+	if (Scale == InScale)
+	{
+		return;
+	}
+
+	Scale = InScale;
+	UpdateLayoutParameters();
+}
+
+#if WITH_EDITOR
+const TAvaPropertyChangeDispatcher<UAvaClonerCylinderLayout> UAvaClonerCylinderLayout::PropertyChangeDispatcher =
+{
+	{ GET_MEMBER_NAME_CHECKED(UAvaClonerCylinderLayout, BaseCount), &UAvaClonerCylinderLayout::OnLayoutPropertyChanged },
+	{ GET_MEMBER_NAME_CHECKED(UAvaClonerCylinderLayout, HeightCount), &UAvaClonerCylinderLayout::OnLayoutPropertyChanged },
+	{ GET_MEMBER_NAME_CHECKED(UAvaClonerCylinderLayout, Height), &UAvaClonerCylinderLayout::OnLayoutPropertyChanged },
+	{ GET_MEMBER_NAME_CHECKED(UAvaClonerCylinderLayout, Radius), &UAvaClonerCylinderLayout::OnLayoutPropertyChanged },
+	{ GET_MEMBER_NAME_CHECKED(UAvaClonerCylinderLayout, AngleStart), &UAvaClonerCylinderLayout::OnLayoutPropertyChanged },
+	{ GET_MEMBER_NAME_CHECKED(UAvaClonerCylinderLayout, AngleRatio), &UAvaClonerCylinderLayout::OnLayoutPropertyChanged },
+	{ GET_MEMBER_NAME_CHECKED(UAvaClonerCylinderLayout, bOrientMesh), &UAvaClonerCylinderLayout::OnLayoutPropertyChanged },
+	{ GET_MEMBER_NAME_CHECKED(UAvaClonerCylinderLayout, Plane), &UAvaClonerCylinderLayout::OnLayoutPropertyChanged },
+	{ GET_MEMBER_NAME_CHECKED(UAvaClonerCylinderLayout, Rotation), &UAvaClonerCylinderLayout::OnLayoutPropertyChanged },
+	{ GET_MEMBER_NAME_CHECKED(UAvaClonerCylinderLayout, Scale), &UAvaClonerCylinderLayout::OnLayoutPropertyChanged },
+};
+
+void UAvaClonerCylinderLayout::PostEditChangeProperty(FPropertyChangedEvent& InPropertyChangedEvent)
+{
+	Super::PostEditChangeProperty(InPropertyChangedEvent);
+
+	PropertyChangeDispatcher.OnPropertyChanged(this, InPropertyChangedEvent);
+}
+#endif
+
+void UAvaClonerCylinderLayout::OnLayoutParametersChanged(UAvaClonerComponent* InComponent)
+{
+	Super::OnLayoutParametersChanged(InComponent);
+
+	InComponent->SetIntParameter(TEXT("CylinderBaseCount"), BaseCount);
+		
+	InComponent->SetIntParameter(TEXT("CylinderHeightCount"), HeightCount);
+
+	InComponent->SetFloatParameter(TEXT("CylinderHeight"), Height);
+		
+	InComponent->SetFloatParameter(TEXT("CylinderRadius"), Radius);
+
+	InComponent->SetFloatParameter(TEXT("CylinderRatio"), AngleRatio);
+
+	InComponent->SetFloatParameter(TEXT("CylinderStart"), AngleStart);
+
+	InComponent->SetBoolParameter(TEXT("MeshOrientAxisEnable"), bOrientMesh);
+
+	FVector CylinderRotation(Rotation.Yaw, Rotation.Pitch, Rotation.Roll);
+
+	if (Plane == EAvaClonerPlane::XY)
+	{
+		CylinderRotation = FVector(0);
+	}
+	else if (Plane == EAvaClonerPlane::YZ)
+	{
+		CylinderRotation = FVector(0, 90, 0);
+	}
+	else if (Plane == EAvaClonerPlane::XZ)
+	{
+		CylinderRotation = FVector(0, 0, 90);
+	}
+		
+	InComponent->SetVectorParameter(TEXT("CylinderRotation"), CylinderRotation);
+
+	InComponent->SetVectorParameter(TEXT("CylinderScale"), Scale);
+}
diff --git a/Engine/Plugins/Experimental/Avalanche/Source/AvalancheEffectors/Private/Cloner/Layouts/AvaClonerGridLayout.cpp b/Engine/Plugins/Experimental/Avalanche/Source/AvalancheEffectors/Private/Cloner/Layouts/AvaClonerGridLayout.cpp
new file mode 100644
--- /dev/null
+++ b/Engine/Plugins/Experimental/Avalanche/Source/AvalancheEffectors/Private/Cloner/Layouts/AvaClonerGridLayout.cpp
@@ -0,0 +1,180 @@
+// Copyright Epic Games, Inc. All Rights Reserved.
+
+#include "Cloner/Layouts/AvaClonerGridLayout.h"
+#include "Cloner/AvaClonerComponent.h"
+#include "NiagaraDataInterfaceTexture.h"
+#include "NiagaraSystem.h"
+
+void UAvaClonerGridLayout::SetCountX(int32 InCountX)
+{
+	if (CountX == InCountX)
+	{
+		return;
+	}
+
+	CountX = InCountX;
+	UpdateLayoutParameters();
+}
+
+void UAvaClonerGridLayout::SetCountY(int32 InCountY)
+{
+	if (CountY == InCountY)
+	{
+		return;
+	}
+
+	CountY = InCountY;
+	UpdateLayoutParameters();
+}
+
+void UAvaClonerGridLayout::SetCountZ(int32 InCountZ)
+{
+	if (CountZ == InCountZ)
+	{
+		return;
+	}
+
+	CountZ = InCountZ;
+	UpdateLayoutParameters();
+}
+
+void UAvaClonerGridLayout::SetSpacingX(float InSpacingX)
+{
+	if (SpacingX == InSpacingX)
+	{
+		return;
+	}
+
+	SpacingX = InSpacingX;
+	UpdateLayoutParameters();
+}
+
+void UAvaClonerGridLayout::SetSpacingY(float InSpacingY)
+{
+	if (SpacingY == InSpacingY)
+	{
+		return;
+	}
+
+	SpacingY = InSpacingY;
+	UpdateLayoutParameters();
+}
+
+void UAvaClonerGridLayout::SetSpacingZ(float InSpacingZ)
+{
+	if (SpacingZ == InSpacingZ)
+	{
+		return;
+	}
+
+	SpacingZ = InSpacingZ;
+	UpdateLayoutParameters();	
+}
+
+void UAvaClonerGridLayout::SetConstraint(EAvaClonerGridConstraint InConstraint)
+{
+	if (Constraint == InConstraint)
+	{
+		return;
+	}
+
+	Constraint = InConstraint;
+	UpdateLayoutParameters();
+}
+
+void UAvaClonerGridLayout::SetInvertConstraint(bool bInInvertConstraint)
+{
+	if (bInvertConstraint == bInInvertConstraint)
+	{
+		return;
+	}
+
+	bInvertConstraint = bInInvertConstraint;
+	UpdateLayoutParameters();
+}
+
+void UAvaClonerGridLayout::SetSphereConstraint(const FAvaClonerGridConstraintSphere& InConstraint)
+{
+	SphereConstraint = InConstraint;
+	UpdateLayoutParameters();
+}
+
+void UAvaClonerGridLayout::SetCylinderConstraint(const FAvaClonerGridConstraintCylinder& InConstraint)
+{
+	CylinderConstraint = InConstraint;
+	UpdateLayoutParameters();
+}
+
+void UAvaClonerGridLayout::SetTextureConstraint(const FAvaClonerGridConstraintTexture& InConstraint)
+{
+	TextureConstraint = InConstraint;
+	UpdateLayoutParameters();
+}
+
+#if WITH_EDITOR
+const TAvaPropertyChangeDispatcher<UAvaClonerGridLayout> UAvaClonerGridLayout::PropertyChangeDispatcher =
+{
+	{ GET_MEMBER_NAME_CHECKED(UAvaClonerGridLayout, CountX), &UAvaClonerGridLayout::OnLayoutPropertyChanged },
+	{ GET_MEMBER_NAME_CHECKED(UAvaClonerGridLayout, CountY), &UAvaClonerGridLayout::OnLayoutPropertyChanged },
+	{ GET_MEMBER_NAME_CHECKED(UAvaClonerGridLayout, CountZ), &UAvaClonerGridLayout::OnLayoutPropertyChanged },
+	{ GET_MEMBER_NAME_CHECKED(UAvaClonerGridLayout, SpacingX), &UAvaClonerGridLayout::OnLayoutPropertyChanged },
+	{ GET_MEMBER_NAME_CHECKED(UAvaClonerGridLayout, SpacingY), &UAvaClonerGridLayout::OnLayoutPropertyChanged },
+	{ GET_MEMBER_NAME_CHECKED(UAvaClonerGridLayout, SpacingZ), &UAvaClonerGridLayout::OnLayoutPropertyChanged },
+	{ GET_MEMBER_NAME_CHECKED(UAvaClonerGridLayout, Constraint), &UAvaClonerGridLayout::OnLayoutPropertyChanged },
+	{ GET_MEMBER_NAME_CHECKED(UAvaClonerGridLayout, bInvertConstraint), &UAvaClonerGridLayout::OnLayoutPropertyChanged },
+	{ GET_MEMBER_NAME_CHECKED(UAvaClonerGridLayout, SphereConstraint), &UAvaClonerGridLayout::OnLayoutPropertyChanged },
+	{ GET_MEMBER_NAME_CHECKED(UAvaClonerGridLayout, CylinderConstraint), &UAvaClonerGridLayout::OnLayoutPropertyChanged },
+	{ GET_MEMBER_NAME_CHECKED(UAvaClonerGridLayout, TextureConstraint), &UAvaClonerGridLayout::OnLayoutPropertyChanged },
+};
+
+void UAvaClonerGridLayout::PostEditChangeProperty(FPropertyChangedEvent& InPropertyChangedEvent)
+{
+	Super::PostEditChangeProperty(InPropertyChangedEvent);
+
+	PropertyChangeDispatcher.OnPropertyChanged(this, InPropertyChangedEvent);
+}
+#endif
+
+void UAvaClonerGridLayout::OnLayoutParametersChanged(UAvaClonerComponent* InComponent)
+{
+	Super::OnLayoutParametersChanged(InComponent);
+
+	InComponent->SetIntParameter(TEXT("GridCountX"), CountX);
+	
+	InComponent->SetIntParameter(TEXT("GridCountY"), CountY);
+
+	InComponent->SetIntParameter(TEXT("GridCountZ"), CountZ);
+	
+	InComponent->SetVectorParameter(TEXT("GridSpacing"), FVector(SpacingX, SpacingY, SpacingZ));
+	
+	FNiagaraUserRedirectionParameterStore& ExposedParameters = InComponent->GetAsset()->GetExposedParameters();
+	static const FNiagaraVariable ConstraintVar(FNiagaraTypeDefinition(StaticEnum<EAvaClonerGridConstraint>()), TEXT("Constraint"));
+	ExposedParameters.SetParameterValue<int32>(static_cast<int32>(Constraint), ConstraintVar);
+
+	InComponent->SetBoolParameter(TEXT("ConstraintInvert"), Constraint != EAvaClonerGridConstraint::None ? bInvertConstraint : false);
+
+	InComponent->SetVectorParameter(TEXT("ConstraintCylinderCenter"), CylinderConstraint.Center);
+	
+	InComponent->SetFloatParameter(TEXT("ConstraintCylinderHeight"), CylinderConstraint.Height);
+	
+	InComponent->SetFloatParameter(TEXT("ConstraintCylinderRadius"), CylinderConstraint.Radius);
+
+	InComponent->SetVectorParameter(TEXT("ConstraintSphereCenter"), SphereConstraint.Center);
+	
+	InComponent->SetFloatParameter(TEXT("ConstraintSphereRadius"), SphereConstraint.Radius);
+	
+	static const FNiagaraVariable ConstraintTextureChannelVar(FNiagaraTypeDefinition(StaticEnum<EAvaClonerTextureSampleChannel>()), TEXT("ConstraintTextureChannel"));
+	ExposedParameters.SetParameterValue<int32>(static_cast<int32>(TextureConstraint.Channel), ConstraintTextureChannelVar);
+	
+	static const FNiagaraVariable ConstraintTextureCompareModeVar(FNiagaraTypeDefinition(StaticEnum<EAvaClonerCompareMode>()), TEXT("ConstraintTextureCompareMode"));
+	ExposedParameters.SetParameterValue<int32>(static_cast<int32>(TextureConstraint.CompareMode), ConstraintTextureCompareModeVar);
+	
+	static const FNiagaraVariable ConstraintTexturePlaneVar(FNiagaraTypeDefinition(StaticEnum<EAvaClonerPlane>()), TEXT("ConstraintTexturePlane"));
+	ExposedParameters.SetParameterValue<int32>(static_cast<int32>(TextureConstraint.Plane), ConstraintTexturePlaneVar);
+
+	InComponent->SetFloatParameter(TEXT("ConstraintTextureThreshold"), TextureConstraint.Threshold);
+	
+	static const FNiagaraVariable ConstraintTextureSamplerVar(FNiagaraTypeDefinition(UNiagaraDataInterfaceTexture::StaticClass()), TEXT("ConstraintTextureSampler"));
+	UNiagaraDataInterfaceTexture* TextureSamplerDI = Cast<UNiagaraDataInterfaceTexture>(ExposedParameters.GetDataInterface(ConstraintTextureSamplerVar));
+	TextureSamplerDI->SetTexture(TextureConstraint.Texture.Get());
+}
diff --git a/Engine/Plugins/Experimental/Avalanche/Source/AvalancheEffectors/Private/Cloner/Layouts/AvaClonerHoneycombLayout.cpp b/Engine/Plugins/Experimental/Avalanche/Source/AvalancheEffectors/Private/Cloner/Layouts/AvaClonerHoneycombLayout.cpp
new file mode 100644
--- /dev/null
+++ b/Engine/Plugins/Experimental/Avalanche/Source/AvalancheEffectors/Private/Cloner/Layouts/AvaClonerHoneycombLayout.cpp
@@ -0,0 +1,123 @@
+// Copyright Epic Games, Inc. All Rights Reserved.
+
+#include "Cloner/Layouts/AvaClonerHoneycombLayout.h"
+#include "Cloner/AvaClonerComponent.h"
+#include "NiagaraSystem.h"
+
+void UAvaClonerHoneycombLayout::SetPlane(EAvaClonerPlane InPlane)
+{
+	if (Plane == InPlane)
+	{
+		return;
+	}
+
+	Plane = InPlane;
+	UpdateLayoutParameters();
+}
+
+void UAvaClonerHoneycombLayout::SetWidthCount(int32 InWidthCount)
+{
+	if (WidthCount == InWidthCount)
+	{
+		return;
+	}
+
+	WidthCount = InWidthCount;
+	UpdateLayoutParameters();
+}
+
+void UAvaClonerHoneycombLayout::SetHeightCount(int32 InHeightCount)
+{
+	if (HeightCount == InHeightCount)
+	{
+		return;
+	}
+
+	HeightCount = InHeightCount;
+	UpdateLayoutParameters();
+}
+
+void UAvaClonerHoneycombLayout::SetWidthOffset(float InWidthOffset)
+{
+	if (WidthOffset == InWidthOffset)
+	{
+		return;
+	}
+
+	WidthOffset = InWidthOffset;
+	UpdateLayoutParameters();
+}
+
+void UAvaClonerHoneycombLayout::SetHeightOffset(float InHeightOffset)
+{
+	if (HeightOffset == InHeightOffset)
+	{
+		return;
+	}
+
+	HeightOffset = InHeightOffset;
+	UpdateLayoutParameters();
+}
+
+void UAvaClonerHoneycombLayout::SetHeightSpacing(float InHeightSpacing)
+{
+	if (HeightSpacing == InHeightSpacing)
+	{
+		return;
+	}
+
+	HeightSpacing = InHeightSpacing;
+	UpdateLayoutParameters();
+}
+
+void UAvaClonerHoneycombLayout::SetWidthSpacing(float InWidthSpacing)
+{
+	if (WidthSpacing == InWidthSpacing)
+	{
+		return;
+	}
+
+	WidthSpacing = InWidthSpacing;
+	UpdateLayoutParameters();
+}
+
+#if WITH_EDITOR
+const TAvaPropertyChangeDispatcher<UAvaClonerHoneycombLayout> UAvaClonerHoneycombLayout::PropertyChangeDispatcher =
+{
+	{ GET_MEMBER_NAME_CHECKED(UAvaClonerHoneycombLayout, Plane), &UAvaClonerHoneycombLayout::OnLayoutPropertyChanged },
+	{ GET_MEMBER_NAME_CHECKED(UAvaClonerHoneycombLayout, WidthCount), &UAvaClonerHoneycombLayout::OnLayoutPropertyChanged },
+	{ GET_MEMBER_NAME_CHECKED(UAvaClonerHoneycombLayout, HeightCount), &UAvaClonerHoneycombLayout::OnLayoutPropertyChanged },
+	{ GET_MEMBER_NAME_CHECKED(UAvaClonerHoneycombLayout, WidthOffset), &UAvaClonerHoneycombLayout::OnLayoutPropertyChanged },
+	{ GET_MEMBER_NAME_CHECKED(UAvaClonerHoneycombLayout, HeightOffset), &UAvaClonerHoneycombLayout::OnLayoutPropertyChanged },
+	{ GET_MEMBER_NAME_CHECKED(UAvaClonerHoneycombLayout, WidthSpacing), &UAvaClonerHoneycombLayout::OnLayoutPropertyChanged },
+	{ GET_MEMBER_NAME_CHECKED(UAvaClonerHoneycombLayout, HeightSpacing), &UAvaClonerHoneycombLayout::OnLayoutPropertyChanged },
+};
+
+void UAvaClonerHoneycombLayout::PostEditChangeProperty(FPropertyChangedEvent& InPropertyChangedEvent)
+{
+	Super::PostEditChangeProperty(InPropertyChangedEvent);
+
+	PropertyChangeDispatcher.OnPropertyChanged(this, InPropertyChangedEvent);
+}
+#endif
+
+void UAvaClonerHoneycombLayout::OnLayoutParametersChanged(UAvaClonerComponent* InComponent)
+{
+	Super::OnLayoutParametersChanged(InComponent);
+	
+	InComponent->SetIntParameter(TEXT("HoneycombWidthCount"), WidthCount);
+		
+	InComponent->SetIntParameter(TEXT("HoneycombHeightCount"), HeightCount);
+		
+	InComponent->SetFloatParameter(TEXT("HoneycombWidthOffset"), WidthOffset);
+		
+	InComponent->SetFloatParameter(TEXT("HoneycombHeightOffset"), HeightOffset);
+
+	InComponent->SetFloatParameter(TEXT("HoneycombWidthSpacing"), WidthSpacing);
+
+	InComponent->SetFloatParameter(TEXT("HoneycombHeightSpacing"), HeightSpacing);
+
+	FNiagaraUserRedirectionParameterStore& ExposedParameters = InComponent->GetAsset()->GetExposedParameters();
+	static const FNiagaraVariable HoneycombPlaneVar(FNiagaraTypeDefinition(StaticEnum<EAvaClonerPlane>()), TEXT("HoneycombPlane"));
+	ExposedParameters.SetParameterValue<int32>(static_cast<int32>(Plane), HoneycombPlaneVar);
+}
diff --git a/Engine/Plugins/Experimental/Avalanche/Source/AvalancheEffectors/Private/Cloner/Layouts/AvaClonerLayoutBase.cpp b/Engine/Plugins/Experimental/Avalanche/Source/AvalancheEffectors/Private/Cloner/Layouts/AvaClonerLayoutBase.cpp
new file mode 100644
--- /dev/null
+++ b/Engine/Plugins/Experimental/Avalanche/Source/AvalancheEffectors/Private/Cloner/Layouts/AvaClonerLayoutBase.cpp
@@ -0,0 +1,290 @@
+// Copyright Epic Games, Inc. All Rights Reserved.
+
+#include "Cloner/Layouts/AvaClonerLayoutBase.h"
+#include "Cloner/AvaClonerActor.h"
+#include "Cloner/AvaClonerComponent.h"
+#include "NiagaraEmitter.h"
+#include "NiagaraMeshRendererProperties.h"
+#include "NiagaraSystem.h"
+#include "UObject/Package.h"
+#include "UObject/UObjectGlobals.h"
+
+DEFINE_LOG_CATEGORY_STATIC(LogAvaClonerLayoutBase, Log, All);
+
+bool UAvaClonerLayoutBase::IsLayoutValid() const
+{
+	if (LayoutName.IsNone() || LayoutAssetPath.IsEmpty())
+	{
+		return false;
+	}
+
+	static const FString LayoutBaseAssetPath = TEXT("/Script/Niagara.NiagaraSystem'/Avalanche/ClonerResources/Systems/NS_ClonerBase.NS_ClonerBase'");
+
+	// Get the template niagara asset
+	const UNiagaraSystem* TemplateNiagaraSystem = LoadSystemPath(LayoutAssetPath);
+
+	// Get the base niagara asset
+	const UNiagaraSystem* BaseNiagaraSystem = LoadSystemPath(LayoutBaseAssetPath);
+
+	if (!TemplateNiagaraSystem || !BaseNiagaraSystem)
+	{
+		UE_LOG(LogAvaClonerLayoutBase, Warning, TEXT("Cloner layout %s : Template system (%s) or base system (%s) is invalid"), *LayoutName.ToString(), *LayoutAssetPath, *LayoutBaseAssetPath);
+		return false;
+	}
+
+	// Compare parameters : template should have base parameters
+	bool bIsSystemBasedOnBaseAsset = true;
+
+	{
+		TArray<FNiagaraVariable> TemplateSystemParameters;
+		TemplateNiagaraSystem->GetExposedParameters().GetParameters(TemplateSystemParameters);
+
+		TArray<FNiagaraVariable> BaseSystemParameters;
+		BaseNiagaraSystem->GetExposedParameters().GetParameters(BaseSystemParameters);
+
+		for (const FNiagaraVariable& SystemParameter : BaseSystemParameters)
+		{
+			if (!TemplateSystemParameters.Contains(SystemParameter))
+			{
+				bIsSystemBasedOnBaseAsset = false;
+				UE_LOG(LogAvaClonerLayoutBase, Warning, TEXT("Cloner layout %s : Template system (%s) missing parameter (%s) from base system (%s)"), *LayoutName.ToString(), *LayoutAssetPath, *SystemParameter.ToString(), *LayoutBaseAssetPath);
+				break;
+			}
+		}
+	}
+
+	if (!bIsSystemBasedOnBaseAsset)
+	{
+		UE_LOG(LogAvaClonerLayoutBase, Warning, TEXT("Cloner layout %s : Template system (%s) is not based off base system (%s)"), *LayoutName.ToString(), *LayoutAssetPath, *LayoutBaseAssetPath);
+	}
+	
+	return bIsSystemBasedOnBaseAsset;
+}
+
+bool UAvaClonerLayoutBase::IsLayoutLoaded() const
+{
+	return !IsTemplate() && NiagaraSystem && MeshRenderer && DataInterfaces.IsValid();
+}
+
+bool UAvaClonerLayoutBase::LoadLayout()
+{
+	if (IsLayoutLoaded())
+	{
+		return true;
+	}
+	
+	if (LayoutAssetPath.IsEmpty())
+	{
+		return false;
+	}
+	
+	// Get the template niagara asset
+	const UNiagaraSystem* TemplateNiagaraSystem = FindObject<UNiagaraSystem>(nullptr, *LayoutAssetPath);
+	if (!TemplateNiagaraSystem)
+	{
+		TemplateNiagaraSystem = LoadObject<UNiagaraSystem>(nullptr, *LayoutAssetPath);
+	}
+
+	if (!TemplateNiagaraSystem)
+	{
+		return false;
+	}
+
+	// Copy template asset since we will modify it, and we do not want other cloner instances to have the new modified asset
+	FObjectDuplicationParameters Parameters = InitStaticDuplicateObjectParams(TemplateNiagaraSystem, GetTransientPackage(), NAME_None, RF_Transient, nullptr, EDuplicateMode::Normal);
+	NiagaraSystem = Cast<UNiagaraSystem>(StaticDuplicateObjectEx(Parameters));
+
+	if (!NiagaraSystem)
+	{
+		return false;
+	}
+
+	for (FNiagaraEmitterHandle& SystemEmitterHandle : NiagaraSystem->GetEmitterHandles())
+	{
+		if (const FVersionedNiagaraEmitterData* EmitterData = SystemEmitterHandle.GetEmitterData())
+		{
+			for (UNiagaraRendererProperties* EmitterRenderer : EmitterData->GetRenderers())
+			{
+				if (UNiagaraMeshRendererProperties* EmitterMeshRenderer = Cast<UNiagaraMeshRendererProperties>(EmitterRenderer))
+				{
+					EmitterMeshRenderer->Meshes.Empty();
+#if WITH_EDITORONLY_DATA
+					EmitterMeshRenderer->OnMeshChanged();
+#endif
+
+					MeshRenderer = EmitterMeshRenderer;
+					DataInterfaces = FAvaClonerEffectorDataInterfaces(NiagaraSystem);
+
+					OnLayoutLoaded();
+					
+					return true;
+				}
+			}
+		}
+	}
+	
+	return false;
+}
+
+bool UAvaClonerLayoutBase::UnloadLayout()
+{
+	if (!IsLayoutLoaded())
+	{
+		return false;
+	}
+
+	// Cannot unload while active
+	if (IsLayoutActive())
+	{
+		return false;
+	}
+	
+	DataInterfaces = FAvaClonerEffectorDataInterfaces();
+	MeshRenderer = nullptr;
+	NiagaraSystem = nullptr;
+	
+	OnLayoutUnloaded();
+
+	return true;
+}
+
+bool UAvaClonerLayoutBase::IsLayoutActive() const
+{
+	const UAvaClonerComponent* Component = GetClonerComponent();
+
+	if (!Component)
+	{
+		return false;	
+	}
+	
+	return IsLayoutLoaded() && Component->GetAsset() == NiagaraSystem;
+}
+
+bool UAvaClonerLayoutBase::ActivateLayout()
+{
+	if (IsLayoutActive())
+	{
+		return false;
+	}
+
+	// Load layout first
+	if (!IsLayoutLoaded())
+	{
+		return false;
+	}
+
+	UAvaClonerComponent* ClonerComponent = GetClonerComponent();
+
+	if (!ClonerComponent)
+	{
+		return false;
+	}
+
+	ClonerComponent->SetAsset(NiagaraSystem);
+
+	OnLayoutActive();
+
+	return true;
+}
+
+bool UAvaClonerLayoutBase::DeactivateLayout()
+{
+	if (!IsLayoutActive())
+	{
+		return false;
+	}
+
+	UAvaClonerComponent* ClonerComponent = GetClonerComponent();
+
+	if (!ClonerComponent)
+	{
+		return false;
+	}
+
+	ClonerComponent->SetAsset(nullptr);
+
+	OnLayoutInactive();
+
+	return true;
+}
+
+bool UAvaClonerLayoutBase::CopyTo(UAvaClonerLayoutBase* InOtherLayout) const
+{
+	if (!InOtherLayout || InOtherLayout == this)
+	{
+		return false;
+	}
+
+	if (!IsLayoutLoaded())
+	{
+		return false;
+	}
+
+	const UNiagaraSystem* OtherSystem = InOtherLayout->NiagaraSystem;
+	if (!OtherSystem)
+	{
+		return false;
+	}
+	
+	DataInterfaces.CopyTo(InOtherLayout->DataInterfaces);
+
+	return true;
+}
+
+void UAvaClonerLayoutBase::OnLayoutPropertyChanged()
+{
+	UpdateLayoutParameters();
+}
+
+UNiagaraSystem* UAvaClonerLayoutBase::LoadSystemPath(const FString& InPath) const
+{
+	UNiagaraSystem* LoadedNiagaraSystem = FindObject<UNiagaraSystem>(nullptr, *InPath);
+	
+	if (!LoadedNiagaraSystem)
+	{
+		LoadedNiagaraSystem = LoadObject<UNiagaraSystem>(nullptr, *InPath);
+	}
+
+	return LoadedNiagaraSystem;
+}
+
+AAvaClonerActor* UAvaClonerLayoutBase::GetClonerActor() const
+{
+	return GetTypedOuter<AAvaClonerActor>();
+}
+
+UAvaClonerComponent* UAvaClonerLayoutBase::GetClonerComponent() const
+{
+	if (const AAvaClonerActor* Cloner = GetClonerActor())
+	{
+		return Cloner->GetClonerComponent();
+	}
+
+	return nullptr;
+}
+
+void UAvaClonerLayoutBase::UpdateLayoutParameters(bool bInUpdateCloner, bool bInImmediate)
+{
+	if (!IsLayoutActive())
+	{
+		return;	
+	}
+	
+	if (UAvaClonerComponent* ClonerComponent = GetClonerComponent())
+	{
+		OnLayoutParametersChanged(ClonerComponent);
+
+		if (bInUpdateCloner)
+		{
+			RequestClonerUpdate(bInImmediate);
+		}
+	}
+}
+
+void UAvaClonerLayoutBase::RequestClonerUpdate(bool bInImmediate) const
+{
+	if (AAvaClonerActor* Cloner = GetClonerActor())
+	{
+		Cloner->RequestClonerUpdate(bInImmediate);
+	}
+}
diff --git a/Engine/Plugins/Experimental/Avalanche/Source/AvalancheEffectors/Private/Cloner/Layouts/AvaClonerLineLayout.cpp b/Engine/Plugins/Experimental/Avalanche/Source/AvalancheEffectors/Private/Cloner/Layouts/AvaClonerLineLayout.cpp
new file mode 100644
--- /dev/null
+++ b/Engine/Plugins/Experimental/Avalanche/Source/AvalancheEffectors/Private/Cloner/Layouts/AvaClonerLineLayout.cpp
@@ -0,0 +1,108 @@
+// Copyright Epic Games, Inc. All Rights Reserved.
+
+#include "Cloner/Layouts/AvaClonerLineLayout.h"
+#include "Cloner/AvaClonerComponent.h"
+
+void UAvaClonerLineLayout::SetCount(int32 InCount)
+{
+	if (Count == InCount)
+	{
+		return;
+	}
+
+	Count = InCount;
+	UpdateLayoutParameters();
+}
+
+void UAvaClonerLineLayout::SetSpacing(float InSpacing)
+{
+	if (Spacing == InSpacing)
+	{
+		return;
+	}
+
+	Spacing = InSpacing;
+	UpdateLayoutParameters();
+}
+
+void UAvaClonerLineLayout::SetAxis(EAvaClonerAxis InAxis)
+{
+	if (Axis == InAxis)
+	{
+		return;
+	}
+
+	Axis = InAxis;
+	UpdateLayoutParameters();
+}
+
+void UAvaClonerLineLayout::SetDirection(const FVector& InDirection)
+{
+	if (Direction == InDirection)
+	{
+		return;
+	}
+
+	Direction = InDirection;
+	UpdateLayoutParameters();
+}
+
+void UAvaClonerLineLayout::SetRotation(const FRotator& InRotation)
+{
+	if (Rotation == InRotation)
+	{
+		return;
+	}
+
+	Rotation = InRotation;
+	UpdateLayoutParameters();
+}
+
+#if WITH_EDITOR
+const TAvaPropertyChangeDispatcher<UAvaClonerLineLayout> UAvaClonerLineLayout::PropertyChangeDispatcher =
+{
+	{ GET_MEMBER_NAME_CHECKED(UAvaClonerLineLayout, Count), &UAvaClonerLineLayout::OnLayoutPropertyChanged },
+	{ GET_MEMBER_NAME_CHECKED(UAvaClonerLineLayout, Spacing), &UAvaClonerLineLayout::OnLayoutPropertyChanged },
+	{ GET_MEMBER_NAME_CHECKED(UAvaClonerLineLayout, Axis), &UAvaClonerLineLayout::OnLayoutPropertyChanged },
+	{ GET_MEMBER_NAME_CHECKED(UAvaClonerLineLayout, Direction), &UAvaClonerLineLayout::OnLayoutPropertyChanged },
+	{ GET_MEMBER_NAME_CHECKED(UAvaClonerLineLayout, Rotation), &UAvaClonerLineLayout::OnLayoutPropertyChanged },
+};
+
+void UAvaClonerLineLayout::PostEditChangeProperty(FPropertyChangedEvent& InPropertyChangedEvent)
+{
+	Super::PostEditChangeProperty(InPropertyChangedEvent);
+
+	PropertyChangeDispatcher.OnPropertyChanged(this, InPropertyChangedEvent);
+}
+#endif
+
+void UAvaClonerLineLayout::OnLayoutParametersChanged(UAvaClonerComponent* InComponent)
+{
+	Super::OnLayoutParametersChanged(InComponent);
+
+	InComponent->SetIntParameter(TEXT("LineCount"), Count);
+		
+	InComponent->SetFloatParameter(TEXT("LineSpacing"), Spacing);
+
+	FVector LineAxis;
+	if (Axis == EAvaClonerAxis::X)
+	{
+		LineAxis = FVector::XAxisVector;
+	}
+	else if (Axis == EAvaClonerAxis::Y)
+	{
+		LineAxis = FVector::YAxisVector;
+	}
+	else if (Axis == EAvaClonerAxis::Z)
+	{
+		LineAxis = FVector::ZAxisVector;
+	}
+	else
+	{
+		LineAxis = Direction.GetSafeNormal();
+	}
+
+	InComponent->SetVectorParameter(TEXT("LineAxis"), LineAxis);
+
+	InComponent->SetVectorParameter(TEXT("LineRotation"), FVector(Rotation.Roll, Rotation.Pitch, Rotation.Yaw));
+}
diff --git a/Engine/Plugins/Experimental/Avalanche/Source/AvalancheEffectors/Private/Cloner/Layouts/AvaClonerMeshLayout.cpp b/Engine/Plugins/Experimental/Avalanche/Source/AvalancheEffectors/Private/Cloner/Layouts/AvaClonerMeshLayout.cpp
new file mode 100644
--- /dev/null
+++ b/Engine/Plugins/Experimental/Avalanche/Source/AvalancheEffectors/Private/Cloner/Layouts/AvaClonerMeshLayout.cpp
@@ -0,0 +1,156 @@
+// Copyright Epic Games, Inc. All Rights Reserved.
+
+#include "Cloner/Layouts/AvaClonerMeshLayout.h"
+#include "Cloner/AvaClonerActor.h"
+#include "Cloner/AvaClonerComponent.h"
+#include "DataInterface/NiagaraDataInterfaceActorComponent.h"
+#include "DataInterface/NiagaraDataInterfaceStaticMesh.h"
+#include "NiagaraDataInterfaceSkeletalMesh.h"
+#include "NiagaraSystem.h"
+
+void UAvaClonerMeshLayout::SetCount(int32 InCount)
+{
+	if (Count == InCount)
+	{
+		return;
+	}
+
+	Count = InCount;
+	UpdateLayoutParameters();
+}
+
+void UAvaClonerMeshLayout::SetAsset(EAvaClonerMeshAsset InAsset)
+{
+	if (Asset == InAsset)
+	{
+		return;
+	}
+
+	Asset = InAsset;
+	UpdateLayoutParameters();
+}
+
+void UAvaClonerMeshLayout::SetSampleData(EAvaClonerMeshSampleData InSampleData)
+{
+	if (SampleData == InSampleData)
+	{
+		return;
+	}
+
+	SampleData = InSampleData;
+	UpdateLayoutParameters();
+}
+
+void UAvaClonerMeshLayout::SetSampleActorWeak(const TWeakObjectPtr<AActor>& InSampleActor)
+{
+	if (SampleActorWeak == InSampleActor)
+	{
+		return;
+	}
+
+	SampleActorWeak = InSampleActor;
+	UpdateLayoutParameters();
+}
+
+void UAvaClonerMeshLayout::SetSampleActor(AActor* InActor)
+{
+	SetSampleActorWeak(InActor);
+}
+
+#if WITH_EDITOR
+const TAvaPropertyChangeDispatcher<UAvaClonerMeshLayout> UAvaClonerMeshLayout::PropertyChangeDispatcher =
+{
+	{ GET_MEMBER_NAME_CHECKED(UAvaClonerMeshLayout, Count), &UAvaClonerMeshLayout::OnLayoutPropertyChanged },
+	{ GET_MEMBER_NAME_CHECKED(UAvaClonerMeshLayout, Asset), &UAvaClonerMeshLayout::OnLayoutPropertyChanged },
+	{ GET_MEMBER_NAME_CHECKED(UAvaClonerMeshLayout, SampleData), &UAvaClonerMeshLayout::OnLayoutPropertyChanged },
+	{ GET_MEMBER_NAME_CHECKED(UAvaClonerMeshLayout, SampleActorWeak), &UAvaClonerMeshLayout::OnLayoutPropertyChanged },
+};
+
+void UAvaClonerMeshLayout::PostEditChangeProperty(FPropertyChangedEvent& InPropertyChangedEvent)
+{
+	Super::PostEditChangeProperty(InPropertyChangedEvent);
+
+	PropertyChangeDispatcher.OnPropertyChanged(this, InPropertyChangedEvent);
+}
+#endif
+
+void UAvaClonerMeshLayout::OnLayoutInactive()
+{
+	Super::OnLayoutInactive();
+
+	// unbind
+	if (USceneComponent* SceneComponent = SceneComponentWeak.Get())
+	{
+		SceneComponent->TransformUpdated.RemoveAll(this);
+	}
+}
+
+void UAvaClonerMeshLayout::OnLayoutParametersChanged(UAvaClonerComponent* InComponent)
+{
+	Super::OnLayoutParametersChanged(InComponent);
+
+	InComponent->SetIntParameter(TEXT("SampleMeshCount"), Count);
+	
+	FNiagaraUserRedirectionParameterStore& ExposedParameters = InComponent->GetAsset()->GetExposedParameters();
+	
+	static const FNiagaraVariable SampleMeshAssetVar(FNiagaraTypeDefinition(StaticEnum<EAvaClonerMeshAsset>()), TEXT("SampleMeshAsset"));
+	ExposedParameters.SetParameterValue<int32>(static_cast<int32>(Asset), SampleMeshAssetVar);
+	
+	static const FNiagaraVariable SampleMeshDataVar(FNiagaraTypeDefinition(StaticEnum<EAvaClonerMeshSampleData>()), TEXT("SampleMeshData"));
+	ExposedParameters.SetParameterValue<int32>(static_cast<int32>(SampleData), SampleMeshDataVar);
+	
+	static const FNiagaraVariable SampleMeshActorVar(FNiagaraTypeDefinition(UNiagaraDataInterfaceActorComponent::StaticClass()), TEXT("SampleMeshActor"));
+	UNiagaraDataInterfaceActorComponent* ActorMeshDI = Cast<UNiagaraDataInterfaceActorComponent>(ExposedParameters.GetDataInterface(SampleMeshActorVar));
+
+	// unbind
+	if (USceneComponent* SceneComponent = SceneComponentWeak.Get())
+	{
+		SceneComponent->TransformUpdated.RemoveAll(this);
+	}
+	SceneComponentWeak = nullptr;
+	
+	// bind
+	AActor* SampleActor = SampleActorWeak.Get();
+	if (SampleActor && SampleActor->GetRootComponent())
+	{
+		ActorMeshDI->SourceActor = SampleActor;
+		SceneComponentWeak = SampleActor->GetRootComponent();
+		SceneComponentWeak->TransformUpdated.AddUObject(this, &UAvaClonerMeshLayout::OnSampleMeshTransformed);
+	}
+	else
+	{
+		SampleActorWeak.Reset();
+		SceneComponentWeak.Reset();
+	}
+
+	if (Asset == EAvaClonerMeshAsset::StaticMesh)
+	{
+		static const FNiagaraVariable SampleMeshStaticVar(FNiagaraTypeDefinition(UNiagaraDataInterfaceStaticMesh::StaticClass()), TEXT("SampleMeshStatic"));
+		UNiagaraDataInterfaceStaticMesh* StaticMeshDI = Cast<UNiagaraDataInterfaceStaticMesh>(ExposedParameters.GetDataInterface(SampleMeshStaticVar));
+		if (SampleActorWeak.IsValid())
+		{
+			if (UStaticMeshComponent* StaticMeshComponent = Cast<UStaticMeshComponent>(SampleActorWeak->GetComponentByClass(UStaticMeshComponent::StaticClass())))
+			{
+				StaticMeshDI->SetSourceComponentFromBlueprints(StaticMeshComponent);
+			}
+		}
+	}
+
+	if (Asset == EAvaClonerMeshAsset::SkeletalMesh)
+	{
+		static const FNiagaraVariable SampleMeshSkeletalVar(FNiagaraTypeDefinition(UNiagaraDataInterfaceSkeletalMesh::StaticClass()), TEXT("SampleMeshSkeletal"));
+		UNiagaraDataInterfaceSkeletalMesh* SkeletalMeshDI = Cast<UNiagaraDataInterfaceSkeletalMesh>(ExposedParameters.GetDataInterface(SampleMeshSkeletalVar));
+		if (SampleActorWeak.IsValid())
+		{
+			if (USkeletalMeshComponent* SkeletalMeshComponent = Cast<USkeletalMeshComponent>(SampleActorWeak->GetComponentByClass(USkeletalMeshComponent::StaticClass())))
+			{
+				SkeletalMeshDI->SetSourceComponentFromBlueprints(SkeletalMeshComponent);
+			}
+		}
+	}
+}
+
+void UAvaClonerMeshLayout::OnSampleMeshTransformed(USceneComponent* InComponent, EUpdateTransformFlags InFlags, ETeleportType InType)
+{
+	UpdateLayoutParameters();
+}
diff --git a/Engine/Plugins/Experimental/Avalanche/Source/AvalancheEffectors/Private/Cloner/Layouts/AvaClonerSphereRandomLayout.cpp b/Engine/Plugins/Experimental/Avalanche/Source/AvalancheEffectors/Private/Cloner/Layouts/AvaClonerSphereRandomLayout.cpp
new file mode 100644
--- /dev/null
+++ b/Engine/Plugins/Experimental/Avalanche/Source/AvalancheEffectors/Private/Cloner/Layouts/AvaClonerSphereRandomLayout.cpp
@@ -0,0 +1,149 @@
+// Copyright Epic Games, Inc. All Rights Reserved.
+
+#include "Cloner/Layouts/AvaClonerSphereRandomLayout.h"
+#include "Cloner/AvaClonerComponent.h"
+
+void UAvaClonerSphereRandomLayout::SetCount(int32 InCount)
+{
+	if (Count == InCount)
+	{
+		return;
+	}
+
+	Count = InCount;
+	UpdateLayoutParameters();
+}
+
+void UAvaClonerSphereRandomLayout::SetRadius(float InRadius)
+{
+	if (Radius == InRadius)
+	{
+		return;
+	}
+
+	Radius = InRadius;
+	UpdateLayoutParameters();
+}
+
+void UAvaClonerSphereRandomLayout::SetDistribution(float InDistribution)
+{
+	if (Distribution == InDistribution)
+	{
+		return;
+	}
+
+	if (InDistribution < 0.f || InDistribution > 1.f)
+	{
+		return;
+	}
+
+	Distribution = InDistribution;
+	UpdateLayoutParameters();
+}
+
+void UAvaClonerSphereRandomLayout::SetLongitude(float InLongitude)
+{
+	if (Longitude == InLongitude)
+	{
+		return;
+	}
+
+	if (InLongitude < 0.f || InLongitude > 1.f)
+	{
+		return;
+	}
+
+	Longitude = InLongitude;
+	UpdateLayoutParameters();
+}
+
+void UAvaClonerSphereRandomLayout::SetLatitude(float InLatitude)
+{
+	if (Latitude == InLatitude)
+	{
+		return;
+	}
+
+	if (InLatitude < 0.f || InLatitude > 1.f)
+	{
+		return;
+	}
+
+	Latitude = InLatitude;
+	UpdateLayoutParameters();
+}
+
+void UAvaClonerSphereRandomLayout::SetOrientMesh(bool bInOrientMesh)
+{
+	if (bOrientMesh == bInOrientMesh)
+	{
+		return;
+	}
+
+	bOrientMesh = bInOrientMesh;
+	UpdateLayoutParameters();
+}
+
+void UAvaClonerSphereRandomLayout::SetRotation(const FRotator& InRotation)
+{
+	if (Rotation == InRotation)
+	{
+		return;
+	}
+
+	Rotation = InRotation;
+	UpdateLayoutParameters();
+}
+
+void UAvaClonerSphereRandomLayout::SetScale(const FVector& InScale)
+{
+	if (Scale == InScale)
+	{
+		return;
+	}
+
+	Scale = InScale;
+	UpdateLayoutParameters();
+}
+
+#if WITH_EDITOR
+const TAvaPropertyChangeDispatcher<UAvaClonerSphereRandomLayout> UAvaClonerSphereRandomLayout::PropertyChangeDispatcher =
+{
+	{ GET_MEMBER_NAME_CHECKED(UAvaClonerSphereRandomLayout, Count), &UAvaClonerSphereRandomLayout::OnLayoutPropertyChanged },
+	{ GET_MEMBER_NAME_CHECKED(UAvaClonerSphereRandomLayout, Radius), &UAvaClonerSphereRandomLayout::OnLayoutPropertyChanged },
+	{ GET_MEMBER_NAME_CHECKED(UAvaClonerSphereRandomLayout, Distribution), &UAvaClonerSphereRandomLayout::OnLayoutPropertyChanged },
+	{ GET_MEMBER_NAME_CHECKED(UAvaClonerSphereRandomLayout, Longitude), &UAvaClonerSphereRandomLayout::OnLayoutPropertyChanged },
+	{ GET_MEMBER_NAME_CHECKED(UAvaClonerSphereRandomLayout, Latitude), &UAvaClonerSphereRandomLayout::OnLayoutPropertyChanged },
+	{ GET_MEMBER_NAME_CHECKED(UAvaClonerSphereRandomLayout, bOrientMesh), &UAvaClonerSphereRandomLayout::OnLayoutPropertyChanged },
+	{ GET_MEMBER_NAME_CHECKED(UAvaClonerSphereRandomLayout, Rotation), &UAvaClonerSphereRandomLayout::OnLayoutPropertyChanged },
+	{ GET_MEMBER_NAME_CHECKED(UAvaClonerSphereRandomLayout, Scale), &UAvaClonerSphereRandomLayout::OnLayoutPropertyChanged },
+};
+
+void UAvaClonerSphereRandomLayout::PostEditChangeProperty(FPropertyChangedEvent& InPropertyChangedEvent)
+{
+	Super::PostEditChangeProperty(InPropertyChangedEvent);
+
+	PropertyChangeDispatcher.OnPropertyChanged(this, InPropertyChangedEvent);
+}
+#endif
+
+void UAvaClonerSphereRandomLayout::OnLayoutParametersChanged(UAvaClonerComponent* InComponent)
+{
+	Super::OnLayoutParametersChanged(InComponent);
+
+	InComponent->SetIntParameter(TEXT("SphereCount"), Count);
+			
+	InComponent->SetFloatParameter(TEXT("SphereRadius"), Radius);
+		
+	InComponent->SetFloatParameter(TEXT("SphereDistribution"), Distribution);
+	
+	InComponent->SetFloatParameter(TEXT("SphereLongitude"), Longitude);
+	
+	InComponent->SetFloatParameter(TEXT("SphereLatitude"), Latitude);
+		
+	InComponent->SetBoolParameter(TEXT("MeshOrientAxisEnable"), bOrientMesh);
+		
+	InComponent->SetVectorParameter(TEXT("SphereRotation"), FVector(Rotation.Yaw, Rotation.Pitch, Rotation.Roll));
+		
+	InComponent->SetVectorParameter(TEXT("SphereScale"), Scale);
+}
diff --git a/Engine/Plugins/Experimental/Avalanche/Source/AvalancheEffectors/Private/Cloner/Layouts/AvaClonerSphereUniformLayout.cpp b/Engine/Plugins/Experimental/Avalanche/Source/AvalancheEffectors/Private/Cloner/Layouts/AvaClonerSphereUniformLayout.cpp
new file mode 100644
--- /dev/null
+++ b/Engine/Plugins/Experimental/Avalanche/Source/AvalancheEffectors/Private/Cloner/Layouts/AvaClonerSphereUniformLayout.cpp
@@ -0,0 +1,106 @@
+// Copyright Epic Games, Inc. All Rights Reserved.
+
+#include "Cloner/Layouts/AvaClonerSphereUniformLayout.h"
+#include "Cloner/AvaClonerComponent.h"
+
+void UAvaClonerSphereUniformLayout::SetCount(int32 InCount)
+{
+	if (Count == InCount)
+	{
+		return;
+	}
+
+	Count = InCount;
+	UpdateLayoutParameters();
+}
+
+void UAvaClonerSphereUniformLayout::SetRadius(float InRadius)
+{
+	if (Radius == InRadius)
+	{
+		return;
+	}
+
+	Radius = InRadius;
+	UpdateLayoutParameters();
+}
+
+void UAvaClonerSphereUniformLayout::SetRatio(float InRatio)
+{
+	if (Ratio == InRatio)
+	{
+		return;
+	}
+
+	Ratio = InRatio;
+	UpdateLayoutParameters();
+}
+
+void UAvaClonerSphereUniformLayout::SetOrientMesh(bool bInOrientMesh)
+{
+	if (bOrientMesh == bInOrientMesh)
+	{
+		return;
+	}
+
+	bOrientMesh = bInOrientMesh;
+	UpdateLayoutParameters();
+}
+
+void UAvaClonerSphereUniformLayout::SetRotation(const FRotator& InRotation)
+{
+	if (Rotation == InRotation)
+	{
+		return;
+	}
+
+	Rotation = InRotation;
+	UpdateLayoutParameters();
+}
+
+void UAvaClonerSphereUniformLayout::SetScale(const FVector& InScale)
+{
+	if (Scale == InScale)
+	{
+		return;
+	}
+
+	Scale = InScale;
+	UpdateLayoutParameters();
+}
+
+#if WITH_EDITOR
+const TAvaPropertyChangeDispatcher<UAvaClonerSphereUniformLayout> UAvaClonerSphereUniformLayout::PropertyChangeDispatcher =
+{
+	{ GET_MEMBER_NAME_CHECKED(UAvaClonerSphereUniformLayout, Count), &UAvaClonerSphereUniformLayout::OnLayoutPropertyChanged },
+	{ GET_MEMBER_NAME_CHECKED(UAvaClonerSphereUniformLayout, Radius), &UAvaClonerSphereUniformLayout::OnLayoutPropertyChanged },
+	{ GET_MEMBER_NAME_CHECKED(UAvaClonerSphereUniformLayout, Ratio), &UAvaClonerSphereUniformLayout::OnLayoutPropertyChanged },
+	{ GET_MEMBER_NAME_CHECKED(UAvaClonerSphereUniformLayout, bOrientMesh), &UAvaClonerSphereUniformLayout::OnLayoutPropertyChanged },
+	{ GET_MEMBER_NAME_CHECKED(UAvaClonerSphereUniformLayout, Rotation), &UAvaClonerSphereUniformLayout::OnLayoutPropertyChanged },
+	{ GET_MEMBER_NAME_CHECKED(UAvaClonerSphereUniformLayout, Scale), &UAvaClonerSphereUniformLayout::OnLayoutPropertyChanged },
+};
+
+void UAvaClonerSphereUniformLayout::PostEditChangeProperty(FPropertyChangedEvent& InPropertyChangedEvent)
+{
+	Super::PostEditChangeProperty(InPropertyChangedEvent);
+
+	PropertyChangeDispatcher.OnPropertyChanged(this, InPropertyChangedEvent);
+}
+#endif
+
+void UAvaClonerSphereUniformLayout::OnLayoutParametersChanged(UAvaClonerComponent* InComponent)
+{
+	Super::OnLayoutParametersChanged(InComponent);
+
+	InComponent->SetIntParameter(TEXT("SphereCount"), Count);
+			
+	InComponent->SetFloatParameter(TEXT("SphereRadius"), Radius);
+		
+	InComponent->SetFloatParameter(TEXT("SphereRatio"), Ratio);
+		
+	InComponent->SetBoolParameter(TEXT("MeshOrientAxisEnable"), bOrientMesh);
+		
+	InComponent->SetVectorParameter(TEXT("SphereRotation"), FVector(Rotation.Yaw, Rotation.Pitch, Rotation.Roll));
+		
+	InComponent->SetVectorParameter(TEXT("SphereScale"), Scale);
+}
diff --git a/Engine/Plugins/Experimental/Avalanche/Source/AvalancheEffectors/Private/Cloner/Layouts/AvaClonerSplineLayout.cpp b/Engine/Plugins/Experimental/Avalanche/Source/AvalancheEffectors/Private/Cloner/Layouts/AvaClonerSplineLayout.cpp
new file mode 100644
--- /dev/null
+++ b/Engine/Plugins/Experimental/Avalanche/Source/AvalancheEffectors/Private/Cloner/Layouts/AvaClonerSplineLayout.cpp
@@ -0,0 +1,182 @@
+// Copyright Epic Games, Inc. All Rights Reserved.
+
+#include "Cloner/Layouts/AvaClonerSplineLayout.h"
+
+#include "Cloner/AvaClonerActor.h"
+#include "Cloner/AvaClonerComponent.h"
+#include "Engine/World.h"
+#include "GameFramework/Actor.h"
+#include "NiagaraDataInterfaceSpline.h"
+#include "NiagaraSystem.h"
+
+#if WITH_EDITOR
+#include "Editor/EditorEngine.h"
+#endif
+
+void UAvaClonerSplineLayout::SetCount(int32 InCount)
+{
+	if (Count == InCount)
+	{
+		return;
+	}
+
+	Count = InCount;
+	UpdateLayoutParameters();
+}
+
+void UAvaClonerSplineLayout::SetSplineActorWeak(const TWeakObjectPtr<AActor>& InSplineActor)
+{
+	if (SplineActorWeak == InSplineActor)
+	{
+		return;
+	}
+
+	SplineActorWeak = InSplineActor;
+	UpdateLayoutParameters();
+}
+
+void UAvaClonerSplineLayout::SetSplineActor(AActor* InSplineActor)
+{
+	SetSplineActorWeak(InSplineActor);
+}
+
+void UAvaClonerSplineLayout::SetOrientMesh(bool bInOrientMesh)
+{
+	if (bOrientMesh == bInOrientMesh)
+	{
+		return;
+	}
+
+	bOrientMesh = bInOrientMesh;
+	UpdateLayoutParameters();
+}
+
+#if WITH_EDITOR
+void UAvaClonerSplineLayout::SpawnLinkedSplineActor()
+{
+	const AAvaClonerActor* ClonerActor = GetClonerActor();
+	
+	if (!ClonerActor)
+	{
+		return;
+	}
+
+	UWorld* ClonerWorld = ClonerActor->GetWorld();
+	
+	FActorSpawnParameters Params;
+	Params.bTemporaryEditorActor = false;
+
+	AActor* SpawnedSplineActor = ClonerWorld->SpawnActor<AActor>(FVector::ZeroVector, FRotator::ZeroRotator, Params);
+	
+	if (!SpawnedSplineActor)
+	{
+		return;
+	}
+
+	// Construct the new component and attach as needed
+	USplineComponent* const NewComponent = NewObject<USplineComponent>(SpawnedSplineActor
+		, USplineComponent::StaticClass()
+		, MakeUniqueObjectName(SpawnedSplineActor, USplineComponent::StaticClass(), TEXT("SplineComponent"))
+		, RF_Transactional);
+
+	SpawnedSplineActor->SetRootComponent(NewComponent);
+	
+	// Add to SerializedComponents array so it gets saved
+	SpawnedSplineActor->AddInstanceComponent(NewComponent);
+	NewComponent->OnComponentCreated();
+	NewComponent->RegisterComponent();
+
+	// Rerun construction scripts
+	SpawnedSplineActor->RerunConstructionScripts();
+
+	SpawnedSplineActor->SetActorLocation(ClonerActor->GetActorLocation());
+	SpawnedSplineActor->SetActorRotation(ClonerActor->GetActorRotation());
+	
+	SetSplineActorWeak(SpawnedSplineActor);
+	FActorLabelUtilities::RenameExistingActor(SpawnedSplineActor, TEXT("SplineActor"), true);
+}
+
+const TAvaPropertyChangeDispatcher<UAvaClonerSplineLayout> UAvaClonerSplineLayout::PropertyChangeDispatcher =
+{
+	{ GET_MEMBER_NAME_CHECKED(UAvaClonerSplineLayout, Count), &UAvaClonerSplineLayout::OnLayoutPropertyChanged },
+	{ GET_MEMBER_NAME_CHECKED(UAvaClonerSplineLayout, SplineActorWeak), &UAvaClonerSplineLayout::OnLayoutPropertyChanged },
+	{ GET_MEMBER_NAME_CHECKED(UAvaClonerSplineLayout, bOrientMesh), &UAvaClonerSplineLayout::OnLayoutPropertyChanged },
+};
+
+void UAvaClonerSplineLayout::PostEditChangeProperty(FPropertyChangedEvent& InPropertyChangedEvent)
+{
+	Super::PostEditChangeProperty(InPropertyChangedEvent);
+
+	PropertyChangeDispatcher.OnPropertyChanged(this, InPropertyChangedEvent);
+}
+#endif
+
+void UAvaClonerSplineLayout::OnLayoutInactive()
+{
+	Super::OnLayoutInactive();
+
+	// unbind
+	if (USplineComponent* SplineComponent = SplineComponentWeak.Get())
+	{
+		SplineComponent->TransformUpdated.RemoveAll(this);
+		USceneComponent::MarkRenderStateDirtyEvent.RemoveAll(this);
+	}
+}
+
+void UAvaClonerSplineLayout::OnLayoutParametersChanged(UAvaClonerComponent* InComponent)
+{
+	Super::OnLayoutParametersChanged(InComponent);
+
+	InComponent->SetIntParameter(TEXT("SampleSplineCount"), Count);
+
+	const FNiagaraUserRedirectionParameterStore& ExposedParameters = InComponent->GetAsset()->GetExposedParameters();
+	static const FNiagaraVariable SampleSplineVar(FNiagaraTypeDefinition(UNiagaraDataInterfaceSpline::StaticClass()), TEXT("SampleSpline"));
+	UNiagaraDataInterfaceSpline* SplineDI = Cast<UNiagaraDataInterfaceSpline>(ExposedParameters.GetDataInterface(SampleSplineVar));
+
+	InComponent->SetBoolParameter(TEXT("MeshOrientAxisEnable"), bOrientMesh);
+	
+	// unbind
+	SplineDI->Source = nullptr;
+	
+	if (USplineComponent* SplineComponent = SplineComponentWeak.Get())
+    {
+		SplineComponent->TransformUpdated.RemoveAll(this);
+		SplineComponentWeak.Reset();
+    }
+	
+	SplineComponentWeak = nullptr;
+
+	// bind
+	if (AActor* SplineActor = SplineActorWeak.Get())
+	{
+		if (USplineComponent* SplineComponent = SplineActor->FindComponentByClass<USplineComponent>())
+		{
+			SplineDI->Source = SplineActor;
+			SplineComponentWeak = SplineComponent;
+			
+			SplineComponent->TransformUpdated.RemoveAll(this);
+			SplineComponent->TransformUpdated.AddUObject(this, &UAvaClonerSplineLayout::OnSampleSplineTransformed);
+
+			USceneComponent::MarkRenderStateDirtyEvent.RemoveAll(this);
+			USceneComponent::MarkRenderStateDirtyEvent.AddUObject(this, &UAvaClonerSplineLayout::OnSampleSplineRenderStateUpdated);
+
+			if (AAvaClonerActor* ClonerActor = GetClonerActor())
+			{
+				ClonerActor->SetActorTransform(SplineActor->GetActorTransform());
+			}
+		}
+	}
+}
+
+void UAvaClonerSplineLayout::OnSampleSplineTransformed(USceneComponent* InComponent, EUpdateTransformFlags InFlags, ETeleportType InType)
+{
+	UpdateLayoutParameters();
+}
+
+void UAvaClonerSplineLayout::OnSampleSplineRenderStateUpdated(UActorComponent& InComponent)
+{
+	if (SplineActorWeak.IsValid() && InComponent.GetOwner() == SplineActorWeak.Get())
+	{
+		UpdateLayoutParameters();
+	}
+}
diff --git a/Engine/Plugins/Experimental/Avalanche/Source/AvalancheEffectors/Private/Effector/AvaEffectorActor.cpp b/Engine/Plugins/Experimental/Avalanche/Source/AvalancheEffectors/Private/Effector/AvaEffectorActor.cpp
new file mode 100644
--- /dev/null
+++ b/Engine/Plugins/Experimental/Avalanche/Source/AvalancheEffectors/Private/Effector/AvaEffectorActor.cpp
@@ -0,0 +1,1224 @@
+// Copyright Epic Games, Inc. All Rights Reserved.
+
+#include "Effector/AvaEffectorActor.h"
+
+#include "Cloner/AvaClonerActor.h"
+#include "Components/BoxComponent.h"
+#include "Components/SphereComponent.h"
+#include "Effector/AvaEffectorComponent.h"
+#include "Math/Vector.h"
+#include "Subsystems/AvaEffectorSubsystem.h"
+
+DEFINE_LOG_CATEGORY_STATIC(LogAvalancheEffector, Log, All);
+
+AAvaEffectorActor::AAvaEffectorActor()
+{
+	SetCanBeDamaged(false);
+	PrimaryActorTick.bCanEverTick = false;
+
+	USceneComponent* SceneComponent = CreateDefaultSubobject<UAvaEffectorComponent>(TEXT("AvaEffectorComponent"));
+#if WITH_EDITORONLY_DATA
+	SceneComponent->bVisualizeComponent = true;
+#endif
+	SetRootComponent(SceneComponent);
+
+	// Sphere 
+	InnerSphereComponent = CreateDefaultSubobject<USphereComponent>(TEXT("AvaInnerSphereComponent"));;
+	InnerSphereComponent->ShapeColor = FColor::Blue;
+	InnerSphereComponent->SetLineThickness(VisualizerThickness);
+	InnerSphereComponent->SetSphereRadius(InnerRadius);
+	InnerSphereComponent->SetHiddenInGame(true);
+#if WITH_EDITOR
+	// Do not show bounding box around cloner for better visibility
+	InnerSphereComponent->SetIsVisualizationComponent(true);
+#endif
+	InnerSphereComponent->bIsEditorOnly = false;
+	InnerSphereComponent->SetupAttachment(SceneComponent);
+
+	OuterSphereComponent = CreateDefaultSubobject<USphereComponent>(TEXT("AvaOuterSphereComponent"));
+	OuterSphereComponent->ShapeColor = FColor::Red;
+	OuterSphereComponent->SetLineThickness(VisualizerThickness);
+	OuterSphereComponent->SetSphereRadius(OuterRadius);
+	OuterSphereComponent->SetHiddenInGame(true);
+#if WITH_EDITOR
+	// Do not show bounding box around cloner for better visibility
+	OuterSphereComponent->SetIsVisualizationComponent(true);
+#endif
+	OuterSphereComponent->bIsEditorOnly = false;
+	OuterSphereComponent->SetupAttachment(SceneComponent);
+
+	// Box
+	InnerBoxComponent = CreateDefaultSubobject<UBoxComponent>(TEXT("AvaInnerBoxComponent"));
+	InnerBoxComponent->ShapeColor = FColor::Blue;
+	InnerBoxComponent->SetLineThickness(VisualizerThickness);
+	InnerBoxComponent->SetBoxExtent(InnerExtent);
+	InnerBoxComponent->SetHiddenInGame(true);
+#if WITH_EDITOR
+	// Do not show bounding box around cloner for better visibility
+	InnerBoxComponent->SetIsVisualizationComponent(true);
+#endif
+	InnerBoxComponent->bIsEditorOnly = false;
+	InnerBoxComponent->SetupAttachment(SceneComponent);
+
+	OuterBoxComponent = CreateDefaultSubobject<UBoxComponent>(TEXT("AvaOuterBoxComponent"));
+	OuterBoxComponent->ShapeColor = FColor::Red;
+	OuterBoxComponent->SetLineThickness(VisualizerThickness);
+	OuterBoxComponent->SetBoxExtent(OuterExtent);
+	OuterBoxComponent->SetHiddenInGame(true);
+#if WITH_EDITOR
+	// Do not show bounding box around cloner for better visibility
+	OuterBoxComponent->SetIsVisualizationComponent(true);
+#endif
+	OuterBoxComponent->bIsEditorOnly = false;
+	OuterBoxComponent->SetupAttachment(SceneComponent);
+
+	// Plane
+	InnerPlaneComponent = CreateDefaultSubobject<UBoxComponent>(TEXT("AvaInnerPlaneComponent"));
+	InnerPlaneComponent->ShapeColor = FColor::Blue;
+	InnerPlaneComponent->SetLineThickness(VisualizerThickness);
+	InnerPlaneComponent->SetBoxExtent(InnerExtent);
+	InnerPlaneComponent->SetHiddenInGame(true);
+#if WITH_EDITOR
+	// Do not show bounding box around cloner for better visibility
+	InnerPlaneComponent->SetIsVisualizationComponent(true);
+#endif
+	InnerPlaneComponent->bIsEditorOnly = false;
+	InnerPlaneComponent->SetupAttachment(SceneComponent);
+
+	OuterPlaneComponent = CreateDefaultSubobject<UBoxComponent>(TEXT("AvaOuterPlaneComponent"));
+	OuterPlaneComponent->ShapeColor = FColor::Red;
+	OuterPlaneComponent->SetLineThickness(VisualizerThickness);
+	OuterPlaneComponent->SetBoxExtent(OuterExtent);
+	OuterPlaneComponent->SetHiddenInGame(true);
+#if WITH_EDITOR
+	// Do not show bounding box around cloner for better visibility
+	OuterPlaneComponent->SetIsVisualizationComponent(true);
+#endif
+	OuterPlaneComponent->bIsEditorOnly = false;
+	OuterPlaneComponent->SetupAttachment(SceneComponent);
+
+	if (!IsTemplate())
+	{
+		SceneComponent->TransformUpdated.AddUObject(this, &AAvaEffectorActor::OnEffectorTransformed);
+
+		RegisterToChannel();
+	}
+}
+
+#if WITH_EDITOR
+FString AAvaEffectorActor::GetDefaultActorLabel() const
+{
+	return DefaultLabel;
+}
+
+TAvaPropertyChangeDispatcher<AAvaEffectorActor> AAvaEffectorActor::PropertyChangeDispatcher =
+{
+	/** Effector */
+	{ GET_MEMBER_NAME_CHECKED(AAvaEffectorActor, bEnabled), &AAvaEffectorActor::OnEnabledChanged },
+	{ GET_MEMBER_NAME_CHECKED(AAvaEffectorActor, Magnitude), &AAvaEffectorActor::OnMagnitudeChanged },
+	{ GET_MEMBER_NAME_CHECKED(AAvaEffectorActor, Cloners), &AAvaEffectorActor::OnClonersChanged },
+	{ GET_MEMBER_NAME_CHECKED(AAvaEffectorActor, VisualizerThickness), &AAvaEffectorActor::OnVisualizerThicknessChanged },
+	{ GET_MEMBER_NAME_CHECKED(AAvaEffectorActor, bVisualizerSpriteVisible), &AAvaEffectorActor::OnVisualizerSpriteVisibleChanged },
+	/** Type */
+	{ GET_MEMBER_NAME_CHECKED(AAvaEffectorActor, Type), &AAvaEffectorActor::OnTypeChanged },
+	{ GET_MEMBER_NAME_CHECKED(AAvaEffectorActor, Easing), &AAvaEffectorActor::OnEasingChanged },
+	{ GET_MEMBER_NAME_CHECKED(AAvaEffectorActor, OuterRadius), &AAvaEffectorActor::OnSphereChanged },
+	{ GET_MEMBER_NAME_CHECKED(AAvaEffectorActor, InnerRadius), &AAvaEffectorActor::OnSphereChanged },
+	{ GET_MEMBER_NAME_CHECKED(AAvaEffectorActor, InnerExtent), &AAvaEffectorActor::OnBoxChanged },
+	{ GET_MEMBER_NAME_CHECKED(AAvaEffectorActor, OuterExtent), &AAvaEffectorActor::OnBoxChanged },
+	{ GET_MEMBER_NAME_CHECKED(AAvaEffectorActor, PlaneSpacing), &AAvaEffectorActor::OnPlaneChanged },
+	/** Mode */
+	{ GET_MEMBER_NAME_CHECKED(AAvaEffectorActor, Mode), &AAvaEffectorActor::OnModeChanged },
+	{ GET_MEMBER_NAME_CHECKED(AAvaEffectorActor, Offset), &AAvaEffectorActor::OnTransformOptionsChanged },
+	{ GET_MEMBER_NAME_CHECKED(AAvaEffectorActor, Rotation), &AAvaEffectorActor::OnTransformOptionsChanged },
+	{ GET_MEMBER_NAME_CHECKED(AAvaEffectorActor, Scale), &AAvaEffectorActor::OnTransformOptionsChanged },
+	{ GET_MEMBER_NAME_CHECKED(AAvaEffectorActor, TargetActorWeak), &AAvaEffectorActor::OnTargetActorChanged },
+	{ GET_MEMBER_NAME_CHECKED(AAvaEffectorActor, LocationStrength), &AAvaEffectorActor::OnNoiseFieldOptionsChanged },
+	{ GET_MEMBER_NAME_CHECKED(AAvaEffectorActor, RotationStrength), &AAvaEffectorActor::OnNoiseFieldOptionsChanged },
+	{ GET_MEMBER_NAME_CHECKED(AAvaEffectorActor, ScaleStrength), &AAvaEffectorActor::OnNoiseFieldOptionsChanged },
+	{ GET_MEMBER_NAME_CHECKED(AAvaEffectorActor, Pan), &AAvaEffectorActor::OnNoiseFieldOptionsChanged },
+	{ GET_MEMBER_NAME_CHECKED(AAvaEffectorActor, Frequency), &AAvaEffectorActor::OnNoiseFieldOptionsChanged },
+	/** Force */
+	{ GET_MEMBER_NAME_CHECKED(AAvaEffectorActor, bOrientationForceEnabled), &AAvaEffectorActor::OnForceEnabledChanged },
+	{ GET_MEMBER_NAME_CHECKED(AAvaEffectorActor, OrientationForceRate), &AAvaEffectorActor::OnForceOptionsChanged },
+	{ GET_MEMBER_NAME_CHECKED(AAvaEffectorActor, OrientationForceMin), &AAvaEffectorActor::OnForceOptionsChanged },
+	{ GET_MEMBER_NAME_CHECKED(AAvaEffectorActor, OrientationForceMax), &AAvaEffectorActor::OnForceOptionsChanged },
+	{ GET_MEMBER_NAME_CHECKED(AAvaEffectorActor, bVortexForceEnabled), &AAvaEffectorActor::OnForceEnabledChanged },
+	{ GET_MEMBER_NAME_CHECKED(AAvaEffectorActor, VortexForceAmount), &AAvaEffectorActor::OnForceOptionsChanged },
+	{ GET_MEMBER_NAME_CHECKED(AAvaEffectorActor, VortexForceAxis), &AAvaEffectorActor::OnForceOptionsChanged },
+	{ GET_MEMBER_NAME_CHECKED(AAvaEffectorActor, bCurlNoiseForceEnabled), &AAvaEffectorActor::OnForceEnabledChanged },
+	{ GET_MEMBER_NAME_CHECKED(AAvaEffectorActor, CurlNoiseForceStrength), &AAvaEffectorActor::OnForceOptionsChanged },
+	{ GET_MEMBER_NAME_CHECKED(AAvaEffectorActor, CurlNoiseForceFrequency), &AAvaEffectorActor::OnForceOptionsChanged },
+	{ GET_MEMBER_NAME_CHECKED(AAvaEffectorActor, bAttractionForceEnabled), &AAvaEffectorActor::OnForceEnabledChanged },
+	{ GET_MEMBER_NAME_CHECKED(AAvaEffectorActor, AttractionForceStrength), &AAvaEffectorActor::OnForceOptionsChanged },
+	{ GET_MEMBER_NAME_CHECKED(AAvaEffectorActor, AttractionForceFalloff), &AAvaEffectorActor::OnForceOptionsChanged },
+	{ GET_MEMBER_NAME_CHECKED(AAvaEffectorActor, bGravityForceEnabled), &AAvaEffectorActor::OnForceEnabledChanged },
+	{ GET_MEMBER_NAME_CHECKED(AAvaEffectorActor, GravityForceAcceleration), &AAvaEffectorActor::OnForceOptionsChanged },
+};
+
+void AAvaEffectorActor::PostEditChangeProperty(FPropertyChangedEvent& InPropertyChangedEvent)
+{
+	Super::PostEditChangeProperty(InPropertyChangedEvent);
+
+	PropertyChangeDispatcher.OnPropertyChanged(this, InPropertyChangedEvent);
+}
+
+void AAvaEffectorActor::PostEditImport()
+{
+	Super::PostEditImport();
+
+	RegisterToCloners();
+}
+
+void AAvaEffectorActor::PostEditUndo()
+{
+	Super::PostEditUndo();
+
+	RegisterToChannel();
+
+	OnClonersChanged();
+
+	// Undo causes cloner to blackout, so refresh state
+	ForEachCloner([](AAvaClonerActor* InCloner, int32 InIndex)
+	{
+		InCloner->ForceUpdateCloner();
+		return true;
+	});
+}
+#endif
+
+void AAvaEffectorActor::Destroyed()
+{
+	Super::Destroyed();
+
+	// Remove this effector from the effector channel
+	if (UAvaEffectorSubsystem* EffectorSubsystem = UAvaEffectorSubsystem::Get(GetWorld()))
+	{
+		EffectorSubsystem->UnregisterChannelEffector(this);
+	}
+
+	// Remove this effector from the cloners it is linked to
+	ForEachCloner([this](AAvaClonerActor* InCloner, int32 InIndex)
+	{
+		InCloner->UnregisterEffector(this);
+		return true;
+	});
+}
+
+void AAvaEffectorActor::PostActorCreated()
+{
+	Super::PostActorCreated();
+
+	OnEffectorChanged();
+}
+
+void AAvaEffectorActor::PostLoad()
+{
+	Super::PostLoad();
+
+	RegisterToCloners();
+}
+
+void AAvaEffectorActor::PostDuplicate(EDuplicateMode::Type InDuplicateMode)
+{
+	Super::PostDuplicate(InDuplicateMode);
+
+	RegisterToCloners();
+}
+
+void AAvaEffectorActor::LinkCloner(AAvaClonerActor* InCloner)
+{
+	if (!InCloner || Cloners.Contains(InCloner))
+	{
+		return;
+	}
+
+	Cloners.Add(InCloner);
+	OnClonersChanged();
+}
+
+void AAvaEffectorActor::UnlinkCloner(AAvaClonerActor* InCloner)
+{
+	if (!InCloner || !Cloners.Contains(InCloner))
+	{
+		return;
+	}
+
+	Cloners.Remove(InCloner);
+	OnClonersChanged();
+}
+
+void AAvaEffectorActor::OnEffectorIdentifierChanged()
+{
+	int32 Identifier = ChannelData.GetIdentifier();
+
+	if (Identifier == INDEX_NONE)
+	{
+		UE_LOG(LogAvalancheEffector, Log, TEXT("%s : Effector channel identifier is invalid"), *GetActorNameOrLabel());
+		return;
+	}
+
+	// Update effector identifier in linked cloners
+	ForEachCloner([Identifier](AAvaClonerActor* InCloner, int32 InEffectorIndex)
+	{
+		if (const FAvaClonerEffectorDataInterfaces* DataInterfaces = InCloner->GetEffectorDataInterfaces())
+		{
+			UNiagaraDataInterfaceArrayInt32* EffectorIndexDI = DataInterfaces->GetIndexArray();
+
+			// This effector has a specific index in the cloner
+			if (EffectorIndexDI->GetArrayReference().IsValidIndex(InEffectorIndex))
+			{
+				const int32 PrevIdentifier = EffectorIndexDI->GetArrayReference()[InEffectorIndex];
+
+				// Only update if identifier is different
+				if (PrevIdentifier != Identifier)
+				{
+					EffectorIndexDI->GetArrayReference()[InEffectorIndex] = Identifier;
+
+					constexpr bool bImmediateUpdate = true;
+					InCloner->RequestClonerUpdate(bImmediateUpdate);
+				}
+			}
+		}
+
+		return true;
+	});
+}
+
+void AAvaEffectorActor::SetType(EAvaClonerEffectorType InType)
+{
+	if (InType == Type)
+	{
+		return;
+	}
+
+	Type = InType;
+	OnTypeChanged();
+}
+
+void AAvaEffectorActor::SetEasing(EAvaClonerEasing InEasing)
+{
+	if (InEasing == Easing)
+	{
+		return;
+	}
+
+	Easing = InEasing;
+	OnEasingChanged();
+}
+
+void AAvaEffectorActor::SetMagnitude(float InMagnitude)
+{
+	if (FMath::IsNearlyEqual(InMagnitude, Magnitude))
+	{
+		return;
+	}
+
+	if (InMagnitude < 0.f || InMagnitude > 1.f)
+	{
+		return;
+	}
+
+	Magnitude = InMagnitude;
+	OnMagnitudeChanged();
+}
+
+void AAvaEffectorActor::SetOuterRadius(float InRadius)
+{
+	if (FMath::IsNearlyEqual(InRadius, OuterRadius))
+	{
+		return;
+	}
+
+	if (InRadius < 0.f || InRadius < InnerRadius)
+	{
+		return;
+	}
+
+	OuterRadius = InRadius;
+	OnSphereChanged();
+}
+
+void AAvaEffectorActor::SetInnerRadius(float InRadius)
+{
+	if (FMath::IsNearlyEqual(InRadius, InnerRadius))
+	{
+		return;
+	}
+
+	if (InRadius < 0.f || InRadius > OuterRadius)
+	{
+		return;
+	}
+
+	InnerRadius = InRadius;
+	OnSphereChanged();
+}
+
+void AAvaEffectorActor::SetInnerExtent(const FVector& InExtent)
+{
+	if (InExtent.Equals(InnerExtent))
+	{
+		return;
+	}
+
+	if (InExtent.X < 0 || InExtent.Y < 0 || InExtent.Z < 0)
+	{
+		return;
+	}
+
+	InnerExtent = InExtent;
+	OnBoxChanged();
+}
+
+void AAvaEffectorActor::SetOuterExtent(const FVector& InExtent)
+{
+	if (InExtent.Equals(OuterExtent))
+	{
+		return;
+	}
+
+	if (InExtent.X < 0 || InExtent.Y < 0 || InExtent.Z < 0)
+	{
+		return;
+	}
+
+	OuterExtent = InExtent;
+	OnBoxChanged();
+}
+
+void AAvaEffectorActor::SetPlaneSpacing(float InSpacing)
+{
+	if (FMath::IsNearlyEqual(InSpacing, PlaneSpacing))
+	{
+		return;
+	}
+
+	if (InSpacing < 0.f)
+	{
+		return;
+	}
+
+	PlaneSpacing = InSpacing;
+	OnPlaneChanged();
+}
+
+void AAvaEffectorActor::SetOffset(const FVector& InOffset)
+{
+	if (InOffset.Equals(Offset))
+	{
+		return;
+	}
+
+	Offset = InOffset;
+	OnTransformOptionsChanged();
+}
+
+void AAvaEffectorActor::SetRotation(const FRotator& InRotation)
+{
+	if (InRotation.Equals(Rotation))
+	{
+		return;
+	}
+
+	constexpr float MinRotation = -180.f;
+	if (InRotation.Pitch < MinRotation || InRotation.Roll < MinRotation || InRotation.Yaw < MinRotation)
+	{
+		return;
+	}
+
+	constexpr float MaxRotation = 180.f;
+	if (InRotation.Pitch > MaxRotation || InRotation.Roll > MaxRotation || InRotation.Yaw > MaxRotation)
+	{
+		return;
+	}
+
+	Rotation = InRotation;
+	OnTransformOptionsChanged();
+}
+
+void AAvaEffectorActor::SetScale(const FVector& InScale)
+{
+	if (InScale.Equals(Scale))
+	{
+		return;
+	}
+
+	if (InScale.X < 0.f || InScale.Y < 0.f || InScale.Z < 0.f)
+	{
+		return;
+	}
+
+	Scale = InScale;
+	OnTransformOptionsChanged();
+}
+
+void AAvaEffectorActor::SetEnabled(bool bInEnable)
+{
+	if (bInEnable == bEnabled)
+	{
+		return;
+	}
+
+	bEnabled = bInEnable;
+	OnEnabledChanged();
+}
+
+void AAvaEffectorActor::SetVisualizerThickness(float InThickness)
+{
+	if (FMath::IsNearlyEqual(VisualizerThickness, InThickness))
+	{
+		return;
+	}
+
+	if (VisualizerThickness < 0.1f || VisualizerThickness > 10.f)
+	{
+		return;
+	}
+
+	VisualizerThickness = InThickness;
+	OnVisualizerThicknessChanged();
+}
+
+void AAvaEffectorActor::SetMode(EAvaClonerEffectorMode InMode)
+{
+	if (Mode == InMode)
+	{
+		return;
+	}
+
+	Mode = InMode;
+	OnModeChanged();
+}
+
+void AAvaEffectorActor::SetTargetActor(AActor* InTargetActor)
+{
+	if (InTargetActor == TargetActorWeak.Get())
+	{
+		return;
+	}
+
+	TargetActorWeak = InTargetActor;
+	OnTargetActorChanged();
+}
+
+void AAvaEffectorActor::SetAttractionForceStrength(float InForceStrength)
+{
+	if (FMath::IsNearlyEqual(AttractionForceStrength, InForceStrength))
+	{
+		return;
+	}
+
+	AttractionForceStrength = InForceStrength;
+	OnForceOptionsChanged();
+}
+
+void AAvaEffectorActor::SetAttractionForceFalloff(float InForceFalloff)
+{
+	if (FMath::IsNearlyEqual(AttractionForceFalloff, InForceFalloff))
+	{
+		return;
+	}
+
+	AttractionForceFalloff = InForceFalloff;
+	OnForceOptionsChanged();
+}
+
+void AAvaEffectorActor::SetGravityForceEnabled(bool bInForceEnabled)
+{
+	if (bGravityForceEnabled == bInForceEnabled)
+	{
+		return;
+	}
+
+	bGravityForceEnabled = bInForceEnabled;
+	OnForceEnabledChanged();
+}
+
+void AAvaEffectorActor::SetGravityForceAcceleration(const FVector& InAcceleration)
+{
+	if (GravityForceAcceleration.Equals(InAcceleration))
+	{
+		return;
+	}
+
+	GravityForceAcceleration = InAcceleration;
+	OnForceOptionsChanged();
+}
+
+void AAvaEffectorActor::SetLocationStrength(const FVector& InStrength)
+{
+	if (LocationStrength.Equals(InStrength))
+	{
+		return;
+	}
+
+	LocationStrength = InStrength;
+	OnNoiseFieldOptionsChanged();
+}
+
+void AAvaEffectorActor::SetRotationStrength(const FRotator& InStrength)
+{
+	if (RotationStrength.Equals(InStrength))
+	{
+		return;
+	}
+
+	RotationStrength = InStrength;
+	OnNoiseFieldOptionsChanged();
+}
+
+void AAvaEffectorActor::SetScaleStrength(const FVector& InStrength)
+{
+	if (ScaleStrength.Equals(InStrength))
+	{
+		return;
+	}
+
+	ScaleStrength = InStrength;
+	OnNoiseFieldOptionsChanged();
+}
+
+void AAvaEffectorActor::SetPan(const FVector& InPan)
+{
+	if (Pan.Equals(InPan))
+	{
+		return;
+	}
+
+	Pan = InPan;
+	OnNoiseFieldOptionsChanged();
+}
+
+void AAvaEffectorActor::SetFrequency(float InFrequency)
+{
+	if (FMath::IsNearlyEqual(Frequency, InFrequency))
+	{
+		return;
+	}
+
+	if (InFrequency < 0.f)
+	{
+		return;
+	}
+
+	Frequency = InFrequency;
+	OnNoiseFieldOptionsChanged();
+}
+
+void AAvaEffectorActor::SetOrientationForceEnabled(bool bInForceEnabled)
+{
+	if (bOrientationForceEnabled == bInForceEnabled)
+	{
+		return;
+	}
+
+	bOrientationForceEnabled = bInForceEnabled;
+	OnForceEnabledChanged();
+}
+
+void AAvaEffectorActor::SetOrientationForceRate(float InForceOrientationRate)
+{
+	if (FMath::IsNearlyEqual(OrientationForceRate, InForceOrientationRate))
+	{
+		return;
+	}
+
+	OrientationForceRate = InForceOrientationRate;
+	OnForceOptionsChanged();
+}
+
+void AAvaEffectorActor::SetOrientationForceMin(const FVector& InForceOrientationMin)
+{
+	if (OrientationForceMin.Equals(InForceOrientationMin))
+	{
+		return;
+	}
+
+	OrientationForceMin = InForceOrientationMin;
+	OnForceOptionsChanged();
+}
+
+void AAvaEffectorActor::SetOrientationForceMax(const FVector& InForceOrientationMax)
+{
+	if (OrientationForceMax.Equals(InForceOrientationMax))
+	{
+		return;
+	}
+
+	OrientationForceMax = InForceOrientationMax;
+	OnForceOptionsChanged();
+}
+
+void AAvaEffectorActor::SetVortexForceEnabled(bool bInForceEnabled)
+{
+	if (bVortexForceEnabled == bInForceEnabled)
+	{
+		return;
+	}
+
+	bVortexForceEnabled = bInForceEnabled;
+	OnForceEnabledChanged();
+}
+
+void AAvaEffectorActor::SetVortexForceAmount(float InForceVortexAmount)
+{
+	if (FMath::IsNearlyEqual(VortexForceAmount, InForceVortexAmount))
+	{
+		return;
+	}
+
+	VortexForceAmount = InForceVortexAmount;
+	OnForceOptionsChanged();
+}
+
+void AAvaEffectorActor::SetVortexForceAxis(const FVector& InForceVortexAxis)
+{
+	if (VortexForceAxis.Equals(InForceVortexAxis))
+	{
+		return;
+	}
+
+	VortexForceAxis = InForceVortexAxis;
+	OnForceOptionsChanged();
+}
+
+void AAvaEffectorActor::SetCurlNoiseForceEnabled(bool bInForceEnabled)
+{
+	if (bCurlNoiseForceEnabled == bInForceEnabled)
+	{
+		return;
+	}
+
+	bCurlNoiseForceEnabled = bInForceEnabled;
+	OnForceEnabledChanged();
+}
+
+void AAvaEffectorActor::SetCurlNoiseForceStrength(float InForceCurlNoiseStrength)
+{
+	if (FMath::IsNearlyEqual(CurlNoiseForceStrength, InForceCurlNoiseStrength))
+	{
+		return;
+	}
+
+	CurlNoiseForceStrength = InForceCurlNoiseStrength;
+	OnForceOptionsChanged();
+}
+
+void AAvaEffectorActor::SetCurlNoiseForceFrequency(float InForceCurlNoiseFrequency)
+{
+	if (FMath::IsNearlyEqual(CurlNoiseForceFrequency, InForceCurlNoiseFrequency))
+	{
+		return;
+	}
+
+	CurlNoiseForceFrequency = InForceCurlNoiseFrequency;
+	OnForceOptionsChanged();
+}
+
+void AAvaEffectorActor::SetAttractionForceEnabled(bool bInForceEnabled)
+{
+	if (bAttractionForceEnabled == bInForceEnabled)
+	{
+		return;
+	}
+
+	bAttractionForceEnabled = bInForceEnabled;
+	OnForceEnabledChanged();
+}
+
+#if WITH_EDITOR
+void AAvaEffectorActor::SetVisualizerSpriteVisible(bool bInVisible)
+{
+	if (bVisualizerSpriteVisible == bInVisible)
+	{
+		return;
+	}
+
+	bVisualizerSpriteVisible = bInVisible;
+	OnVisualizerSpriteVisibleChanged();
+}
+#endif
+
+void AAvaEffectorActor::ForEachCloner(TFunctionRef<bool(AAvaClonerActor*, int32)> InFunction, bool bSkipIdxCheck)
+{
+	for (TSet<TWeakObjectPtr<AAvaClonerActor>>::TIterator It(InternalCloners); It; ++It)
+	{
+		AAvaClonerActor* Cloner = It->Get();
+		if (!Cloner)
+		{
+			continue;
+		}
+
+		const int32 Idx = Cloner->GetEffectorIndex(this);
+		if (!bSkipIdxCheck && Idx == INDEX_NONE)
+		{
+			continue;
+		}
+
+		if (!InFunction(Cloner, Idx))
+		{
+			return;
+		}
+	}
+}
+
+void AAvaEffectorActor::OnEffectorSubsystemInitialized(const UWorld* World)
+{
+	if (GetWorld() == World)
+	{
+		RegisterToChannel();
+	}
+}
+
+void AAvaEffectorActor::RegisterToChannel()
+{
+	UAvaEffectorSubsystem::OnSubsystemInitializedDelegate.RemoveAll(this);
+
+	// Register this effector to the effector channel
+	if (UAvaEffectorSubsystem* EffectorSubsystem = UAvaEffectorSubsystem::Get(GetWorld()))
+	{
+		EffectorSubsystem->RegisterChannelEffector(this);
+	}
+	else
+	{
+		UAvaEffectorSubsystem::OnSubsystemInitializedDelegate.AddUObject(this, &AAvaEffectorActor::OnEffectorSubsystemInitialized);
+	}
+}
+
+void AAvaEffectorActor::OnEffectorTransformed(USceneComponent* InUpdatedComponent, EUpdateTransformFlags InUpdateTransformFlags, ETeleportType InTeleport)
+{
+	OnTransformChanged();
+	// Update when scaled or rotated
+	OnSphereChanged();
+	OnBoxChanged();
+	OnPlaneChanged();
+	// update if self
+	const AActor* InternalTargetActor = InternalTargetActorWeak.Get();
+	if (InternalTargetActor == this)
+	{
+		OnTargetOptionsChanged();
+	}
+}
+
+void AAvaEffectorActor::OnEnabledChanged()
+{
+	if (bEnabled)
+	{
+		OnEffectorChanged();
+	}
+	else // Disabled
+	{
+		OnEffectorDisabled();
+		// Hide visualization components
+		OnTypeChanged();
+		// Editor
+		OnVisualizerThicknessChanged();
+		OnVisualizerSpriteVisibleChanged();
+	}
+}
+
+void AAvaEffectorActor::OnEffectorDisabled()
+{
+	OnMagnitudeChanged();
+}
+
+void AAvaEffectorActor::OnClonersChanged()
+{
+	TSet<TWeakObjectPtr<AAvaClonerActor>> AddedCloners = Cloners.Difference(InternalCloners);
+	TSet<TWeakObjectPtr<AAvaClonerActor>> RemovedCloners = InternalCloners.Difference(Cloners);
+
+	// Unlink removed cloners
+	for (const TWeakObjectPtr<AAvaClonerActor>& RemovedClonerWeak : RemovedCloners)
+	{
+		if (AAvaClonerActor* RemovedCloner = RemovedClonerWeak.Get())
+		{
+			RemovedCloner->UnregisterEffector(this);
+		}
+	}
+
+	// Link added cloners
+	for (const TWeakObjectPtr<AAvaClonerActor>& AddedClonerWeak : AddedCloners)
+	{
+		if (AAvaClonerActor* AddedCloner = AddedClonerWeak.Get())
+		{
+			AddedCloner->RegisterEffector(this);
+		}
+	}
+}
+
+void AAvaEffectorActor::OnModeChanged()
+{
+	ChannelData.Mode = Mode;
+
+	OnTransformOptionsChanged();
+	OnTargetActorChanged();
+	OnNoiseFieldOptionsChanged();
+}
+
+void AAvaEffectorActor::OnTypeChanged()
+{
+	// Sphere
+	InnerSphereComponent->SetVisibility(Type == EAvaClonerEffectorType::Sphere);
+	OuterSphereComponent->SetVisibility(Type == EAvaClonerEffectorType::Sphere);
+	// Box
+	InnerBoxComponent->SetVisibility(Type == EAvaClonerEffectorType::Box);
+	OuterBoxComponent->SetVisibility(Type == EAvaClonerEffectorType::Box);
+	// Plane
+	InnerPlaneComponent->SetVisibility(Type == EAvaClonerEffectorType::Plane);
+	OuterPlaneComponent->SetVisibility(Type == EAvaClonerEffectorType::Plane);
+
+	ChannelData.Type = Type;
+
+	// Update type data
+	OnSphereChanged();
+	OnBoxChanged();
+	OnPlaneChanged();
+}
+
+void AAvaEffectorActor::RefreshClonerParameters(AAvaClonerActor* InCloner, bool bInForce) const
+{
+	if (InCloner && (bEnabled || bInForce))
+	{
+		InCloner->RequestClonerUpdate();
+	}
+}
+
+void AAvaEffectorActor::RefreshClonersParameters(bool bInForce)
+{
+	ForEachCloner([this, bInForce](AAvaClonerActor* InCloner, int32 InIndex)
+	{
+		RefreshClonerParameters(InCloner, bInForce);
+		return true;
+	});
+}
+
+void AAvaEffectorActor::RegisterToCloners()
+{
+	for (TSet<TWeakObjectPtr<AAvaClonerActor>>::TIterator It = InternalCloners.CreateIterator(); It; ++It)
+	{
+		AAvaClonerActor* Cloner = It->Get();
+
+		if (!Cloner)
+		{
+			It.RemoveCurrent();
+			continue;
+		}
+
+		Cloner->RegisterEffector(this);
+	}
+
+	OnEffectorChanged();
+}
+
+void AAvaEffectorActor::OnEffectorChanged()
+{
+	if (!bEnabled)
+	{
+		OnEnabledChanged();
+	}
+	else
+	{
+		OnModeChanged();
+		OnTypeChanged();
+		OnEasingChanged();
+		OnTransformChanged();
+		OnMagnitudeChanged();
+		OnForceOptionsChanged();
+		// Editor
+		OnVisualizerThicknessChanged();
+		OnVisualizerSpriteVisibleChanged();
+	}
+}
+
+void AAvaEffectorActor::OnClonerLinked(AAvaClonerActor* InCloner, int32 InEffectorIdx)
+{
+	if (!InCloner)
+	{
+		return;
+	}
+
+	InternalCloners.Add(InCloner);
+
+	OnEffectorIdentifierChanged();
+
+	UE_LOG(LogAvalancheEffector, Log, TEXT("%s : Effector linked to Cloner %s"), *GetActorNameOrLabel(), *InCloner->GetActorNameOrLabel());
+}
+
+void AAvaEffectorActor::OnClonerUnlinked(AAvaClonerActor* InCloner, int32 InEffectorIdx)
+{
+	if (!InCloner)
+	{
+		return;
+	}
+
+	InternalCloners.Remove(InCloner);
+
+	UE_LOG(LogAvalancheEffector, Log, TEXT("%s : Effector unlinked from cloner %s"), *GetActorNameOrLabel(), *InCloner->GetActorNameOrLabel());
+}
+
+void AAvaEffectorActor::OnEasingChanged()
+{
+	ChannelData.Easing = Easing;
+}
+
+void AAvaEffectorActor::OnTransformChanged()
+{
+	ChannelData.Location = GetActorLocation();
+	ChannelData.Rotation = GetActorRotation().Quaternion();
+	ChannelData.Scale = GetActorScale();
+}
+
+void AAvaEffectorActor::OnBoxChanged()
+{
+	InnerExtent = ClampVector(InnerExtent, FVector::ZeroVector, OuterExtent);
+	// Max()
+	OuterExtent = OuterExtent.ComponentMax(InnerExtent);
+
+	InnerBoxComponent->SetBoxExtent(InnerExtent);
+	OuterBoxComponent->SetBoxExtent(OuterExtent);
+
+	if (Type != EAvaClonerEffectorType::Box)
+	{
+		return;
+	}
+
+	ChannelData.InnerExtent = GetInnerExtent();
+	ChannelData.OuterExtent = GetOuterExtent();
+}
+
+void AAvaEffectorActor::OnSphereChanged()
+{
+	InnerRadius = FMath::Clamp(InnerRadius, 0.f, OuterRadius);
+	OuterRadius = FMath::Max(OuterRadius, InnerRadius);
+
+	InnerSphereComponent->SetSphereRadius(InnerRadius);
+	OuterSphereComponent->SetSphereRadius(OuterRadius);
+
+	if (Type != EAvaClonerEffectorType::Sphere)
+	{
+		return;
+	}
+
+	ChannelData.InnerExtent = FVector(GetInnerRadius());
+	ChannelData.OuterExtent = FVector(GetOuterRadius());
+}
+
+void AAvaEffectorActor::OnPlaneChanged()
+{
+	static const FVector PlaneAxis = -FVector::LeftVector;
+	const FVector InnerPlane = PlaneAxis * FVector(-PlaneSpacing/2);
+	const FVector OuterPlane = PlaneAxis * FVector(PlaneSpacing/2);
+
+	static const FVector InnerSize(100, 0, 100);
+	InnerPlaneComponent->SetRelativeLocation(InnerPlane);
+	InnerPlaneComponent->SetBoxExtent(InnerSize);
+
+	static const FVector OuterSize(200, 0, 200);
+	OuterPlaneComponent->SetRelativeLocation(OuterPlane);
+	OuterPlaneComponent->SetBoxExtent(OuterSize);
+
+	if (Type != EAvaClonerEffectorType::Plane)
+	{
+		return;
+	}
+
+	ChannelData.InnerExtent = PlaneAxis;
+	ChannelData.OuterExtent = FVector(PlaneSpacing);
+}
+
+void AAvaEffectorActor::OnMagnitudeChanged()
+{
+	const float EffectorMagnitude = bEnabled ? GetMagnitude() : 0.f;
+	ChannelData.Magnitude = EffectorMagnitude;
+}
+
+void AAvaEffectorActor::OnTransformOptionsChanged()
+{
+	if (Mode != EAvaClonerEffectorMode::Default)
+	{
+		return;
+	}
+
+	FVector ScaleDelta = GetScale();
+	ScaleDelta.X = FMath::Max(ScaleDelta.X, UE_KINDA_SMALL_NUMBER);
+	ScaleDelta.Y = FMath::Max(ScaleDelta.Y, UE_KINDA_SMALL_NUMBER);
+	ScaleDelta.Z = FMath::Max(ScaleDelta.Z, UE_KINDA_SMALL_NUMBER);
+
+	ChannelData.LocationDelta = GetOffset();
+	ChannelData.RotationDelta = GetRotation().Quaternion();
+	ChannelData.ScaleDelta = ScaleDelta;
+}
+
+void AAvaEffectorActor::OnTargetActorChanged()
+{
+	AActor* TargetActor = TargetActorWeak.Get();
+	AActor* const InternalTargetActor = InternalTargetActorWeak.Get();
+	if (TargetActor && TargetActor == InternalTargetActor)
+	{
+		OnTargetOptionsChanged();
+		return;
+	}
+
+	// unbind, except if it is self
+	if (InternalTargetActor && InternalTargetActor->GetRootComponent())
+	{
+		if (InternalTargetActor != this)
+		{
+			InternalTargetActor->OnDestroyed.RemoveAll(this);
+			InternalTargetActor->GetRootComponent()->TransformUpdated.RemoveAll(this);
+			InternalTargetActorWeak.Reset();
+		}
+	}
+
+	// set self by default if invalid
+	if (!TargetActor || !TargetActor->GetRootComponent())
+	{
+		TargetActor = this;
+	}
+
+	// bind to transform event, do not bind to self since we already do that
+	if (TargetActor != this)
+	{
+		TargetActor->GetRootComponent()->TransformUpdated.RemoveAll(this);
+		TargetActor->GetRootComponent()->TransformUpdated.AddUObject(this, &AAvaEffectorActor::OnTargetActorTransformChanged);
+		TargetActor->OnDestroyed.RemoveAll(this);
+		TargetActor->OnDestroyed.AddUniqueDynamic(this, &AAvaEffectorActor::OnTargetActorDestroyed);
+	}
+
+	TargetActorWeak = TargetActor;
+	InternalTargetActorWeak = TargetActor;
+	OnTargetOptionsChanged();
+}
+
+void AAvaEffectorActor::OnTargetActorDestroyed(AActor* InActor)
+{
+	const AActor* TargetActor = TargetActorWeak.Get();
+	if (TargetActor == InActor)
+	{
+		TargetActorWeak = this;
+		OnTargetActorChanged();
+	}
+}
+
+void AAvaEffectorActor::OnNoiseFieldOptionsChanged()
+{
+	if (Mode != EAvaClonerEffectorMode::NoiseField)
+	{
+		return;
+	}
+
+	ChannelData.LocationDelta = LocationStrength;
+	ChannelData.RotationDelta = RotationStrength.Quaternion();
+	ChannelData.ScaleDelta = ScaleStrength;
+	ChannelData.Frequency = Frequency;
+	ChannelData.Pan = Pan;
+}
+
+void AAvaEffectorActor::OnTargetActorTransformChanged(USceneComponent*, EUpdateTransformFlags, ETeleportType)
+{
+	OnTargetOptionsChanged();
+}
+
+void AAvaEffectorActor::OnTargetOptionsChanged()
+{
+	if (Mode != EAvaClonerEffectorMode::Target)
+	{
+		return;
+	}
+
+	const AActor* InternalTargetActor = InternalTargetActorWeak.Get();
+	if (!InternalTargetActor)
+	{
+		return;
+	}
+
+	ChannelData.LocationDelta = InternalTargetActor->GetActorLocation();
+	ChannelData.RotationDelta = FQuat::Identity;
+	ChannelData.ScaleDelta = FVector::OneVector;
+}
+
+void AAvaEffectorActor::OnVisualizerThicknessChanged()
+{
+	VisualizerThickness = FMath::Clamp(VisualizerThickness, 0.1f, 10.f);
+
+	InnerSphereComponent->SetLineThickness(VisualizerThickness);
+	OuterSphereComponent->SetLineThickness(VisualizerThickness);
+	InnerBoxComponent->SetLineThickness(VisualizerThickness);
+	OuterBoxComponent->SetLineThickness(VisualizerThickness);
+	InnerPlaneComponent->SetLineThickness(VisualizerThickness);
+	OuterPlaneComponent->SetLineThickness(VisualizerThickness);
+}
+
+void AAvaEffectorActor::OnVisualizerSpriteVisibleChanged()
+{
+#if WITH_EDITOR
+	UE::ClonerEffector::SetBillboardComponentSprite(this, TEXT("/Script/Engine.Texture2D'/Avalanche/ClonerResources/Textures/T_EffectorIcon.T_EffectorIcon'"));
+	UE::ClonerEffector::SetBillboardComponentVisibility(this, bVisualizerSpriteVisible);
+#endif
+}
+
+void AAvaEffectorActor::OnForceOptionsChanged()
+{
+	if (bOrientationForceEnabled)
+	{
+		ChannelData.OrientationForceRate = OrientationForceRate;
+		ChannelData.OrientationForceMin = OrientationForceMin;
+		ChannelData.OrientationForceMax = OrientationForceMax;
+	}
+	else
+	{
+		ChannelData.OrientationForceRate = 0.f;
+		ChannelData.OrientationForceMin = FVector::ZeroVector;
+		ChannelData.OrientationForceMax = FVector::ZeroVector;
+	}
+
+	if (bVortexForceEnabled)
+	{
+		ChannelData.VortexForceAmount = VortexForceAmount;
+		ChannelData.VortexForceAxis = VortexForceAxis;
+	}
+	else
+	{
+		ChannelData.VortexForceAmount = 0.f;
+		ChannelData.VortexForceAxis = FVector::ZeroVector;
+	}
+
+	if (bCurlNoiseForceEnabled)
+	{
+		ChannelData.CurlNoiseForceStrength = CurlNoiseForceStrength;
+		ChannelData.CurlNoiseForceFrequency = CurlNoiseForceFrequency;
+	}
+	else
+	{
+		ChannelData.CurlNoiseForceStrength = 0.f;
+		ChannelData.CurlNoiseForceFrequency = 0.f;
+	}
+
+	if (bAttractionForceEnabled)
+	{
+		ChannelData.AttractionForceStrength = AttractionForceStrength;
+		ChannelData.AttractionForceFalloff = AttractionForceFalloff;
+	}
+	else
+	{
+		ChannelData.AttractionForceStrength = 0.f;
+		ChannelData.AttractionForceFalloff = 0.f;
+	}
+
+	if (bGravityForceEnabled)
+	{
+		ChannelData.GravityForceAcceleration = GravityForceAcceleration;
+	}
+	else
+	{
+		ChannelData.GravityForceAcceleration = FVector::ZeroVector;
+	}
+}
+
+void AAvaEffectorActor::OnForceEnabledChanged()
+{
+	// Refresh cloners to reset clones transform
+	RefreshClonersParameters();
+	OnForceOptionsChanged();
+}
diff --git a/Engine/Plugins/Experimental/Avalanche/Source/AvalancheEffectors/Private/Subsystems/AvaClonerSubsystem.cpp b/Engine/Plugins/Experimental/Avalanche/Source/AvalancheEffectors/Private/Subsystems/AvaClonerSubsystem.cpp
new file mode 100644
--- /dev/null
+++ b/Engine/Plugins/Experimental/Avalanche/Source/AvalancheEffectors/Private/Subsystems/AvaClonerSubsystem.cpp
@@ -0,0 +1,244 @@
+// Copyright Epic Games, Inc. All Rights Reserved.
+
+#include "Subsystems/AvaClonerSubsystem.h"
+
+#include "Cloner/AvaClonerActor.h"
+#include "Cloner/Layouts/AvaClonerCircleLayout.h"
+#include "Cloner/Layouts/AvaClonerCylinderLayout.h"
+#include "Cloner/Layouts/AvaClonerGridLayout.h"
+#include "Cloner/Layouts/AvaClonerHoneycombLayout.h"
+#include "Cloner/Layouts/AvaClonerLayoutBase.h"
+#include "Cloner/Layouts/AvaClonerLineLayout.h"
+#include "Cloner/Layouts/AvaClonerMeshLayout.h"
+#include "Cloner/Layouts/AvaClonerSphereRandomLayout.h"
+#include "Cloner/Layouts/AvaClonerSphereUniformLayout.h"
+#include "Cloner/Layouts/AvaClonerSplineLayout.h"
+#include "Engine/Engine.h"
+#include "UObject/Class.h"
+#include "UObject/UObjectIterator.h"
+
+#if WITH_EDITOR
+#include "HAL/IConsoleManager.h"
+
+UAvaClonerSubsystem::FOnCVarChanged UAvaClonerSubsystem::OnCVarChangedDelegate;
+#endif
+
+UAvaClonerSubsystem* UAvaClonerSubsystem::Get()
+{
+	if (GEngine)
+	{
+		return GEngine->GetEngineSubsystem<UAvaClonerSubsystem>();
+	}
+	
+	return nullptr;
+}
+
+void UAvaClonerSubsystem::Initialize(FSubsystemCollectionBase& Collection)
+{
+	Super::Initialize(Collection);
+
+	// Register them here to match old order of layout enum
+	RegisterLayoutClass(UAvaClonerGridLayout::StaticClass());
+	RegisterLayoutClass(UAvaClonerLineLayout::StaticClass());
+	RegisterLayoutClass(UAvaClonerCircleLayout::StaticClass());
+	RegisterLayoutClass(UAvaClonerCylinderLayout::StaticClass());
+	RegisterLayoutClass(UAvaClonerSphereUniformLayout::StaticClass());
+	RegisterLayoutClass(UAvaClonerHoneycombLayout::StaticClass());
+	RegisterLayoutClass(UAvaClonerMeshLayout::StaticClass());
+	RegisterLayoutClass(UAvaClonerSplineLayout::StaticClass());
+	RegisterLayoutClass(UAvaClonerSphereRandomLayout::StaticClass());
+
+	// Scan for new layouts
+	ScanForRegistrableClasses();
+
+#if WITH_EDITOR
+	CVarTSRShadingRejectionFlickeringPeriod = IConsoleManager::Get().FindConsoleVariable(TEXT("r.TSR.ShadingRejection.Flickering.Period"));
+
+	if (CVarTSRShadingRejectionFlickeringPeriod)
+	{
+		CVarTSRShadingRejectionFlickeringPeriod->OnChangedDelegate().AddUObject(this, &UAvaClonerSubsystem::OnTSRShadingRejectionFlickeringPeriodChanged);
+	}
+#endif
+}
+
+void UAvaClonerSubsystem::Deinitialize()
+{
+	Super::Deinitialize();
+
+#if WITH_EDITOR
+	if (CVarTSRShadingRejectionFlickeringPeriod)
+	{
+		CVarTSRShadingRejectionFlickeringPeriod->OnChangedDelegate().RemoveAll(this);
+	}
+#endif
+}
+
+bool UAvaClonerSubsystem::RegisterLayoutClass(const UClass* InClonerLayoutClass)
+{
+	if (!IsValid(InClonerLayoutClass))
+	{
+		return false;
+	}
+
+	if (!InClonerLayoutClass->IsChildOf(UAvaClonerLayoutBase::StaticClass())
+		|| InClonerLayoutClass->HasAnyClassFlags(CLASS_Abstract | CLASS_Deprecated | CLASS_NewerVersionExists))
+	{
+		return false;
+	}
+
+	if (IsLayoutClassRegistered(InClonerLayoutClass))
+	{
+		return false;
+	}
+
+	const UAvaClonerLayoutBase* CDO = InClonerLayoutClass->GetDefaultObject<UAvaClonerLayoutBase>();
+
+	if (!CDO)
+	{
+		return false;
+	}
+
+	// Check niagara asset is valid
+	if (!CDO->IsLayoutValid())
+	{
+		return false;
+	}
+
+	// Does not overwrite existing layouts
+	const FName LayoutName = CDO->GetLayoutName();
+	if (LayoutClasses.Contains(LayoutName))
+	{
+		return false;
+	}
+	
+	LayoutClasses.Add(LayoutName, CDO->GetClass());
+	
+	return true;
+}
+
+bool UAvaClonerSubsystem::UnregisterLayoutClass(const UClass* InClonerLayoutClass)
+{
+	if (!IsValid(InClonerLayoutClass))
+	{
+		return false;
+	}
+	
+	for (TMap<FName, TSubclassOf<UAvaClonerLayoutBase>>::TIterator It(LayoutClasses); It; ++It)
+	{
+		if (It->Value.Get() == InClonerLayoutClass)
+		{
+			It.RemoveCurrent();
+			return true;
+		}
+	}
+	
+	return false;
+}
+
+bool UAvaClonerSubsystem::IsLayoutClassRegistered(const UClass* InClonerLayoutClass)
+{
+	if (!IsValid(InClonerLayoutClass))
+	{
+		return false;
+	}
+	
+	for (const TPair<FName, TSubclassOf<UAvaClonerLayoutBase>>& LayoutClassPair : LayoutClasses)
+	{
+		if (LayoutClassPair.Value.Get() == InClonerLayoutClass)
+		{
+			return true;
+		}
+	}
+	
+	return false;
+}
+
+TArray<FName> UAvaClonerSubsystem::GetLayoutNames() const
+{
+	TArray<FName> LayoutNames;
+	LayoutClasses.GenerateKeyArray(LayoutNames);
+	return LayoutNames;
+}
+
+FName UAvaClonerSubsystem::FindLayoutName(TSubclassOf<UAvaClonerLayoutBase> InLayoutClass) const
+{
+	if (const FName* Key = LayoutClasses.FindKey(InLayoutClass))
+	{
+		return *Key;
+	}
+
+	return NAME_None;
+}
+
+UAvaClonerLayoutBase* UAvaClonerSubsystem::CreateNewLayout(FName InLayoutName, AAvaClonerActor* InClonerActor)
+{
+	if (!IsValid(InClonerActor))
+	{
+		return nullptr;
+	}
+
+	TSubclassOf<UAvaClonerLayoutBase> const* LayoutClass = LayoutClasses.Find(InLayoutName);
+	
+	if (!LayoutClass)
+	{
+		return nullptr;
+	}
+
+	return NewObject<UAvaClonerLayoutBase>(InClonerActor, LayoutClass->Get());
+}
+
+void UAvaClonerSubsystem::ScanForRegistrableClasses()
+{
+	for (const UClass* const Class : TObjectRange<UClass>())
+	{
+		RegisterLayoutClass(Class);
+	}
+}
+
+#if WITH_EDITOR
+void UAvaClonerSubsystem::EnableNoFlicker()
+{
+	if (IsNoFlickerEnabled())
+	{
+		return;
+	}
+
+	PreviousCVarValue = CVarTSRShadingRejectionFlickeringPeriod->GetInt();
+	CVarTSRShadingRejectionFlickeringPeriod->Set(NoFlicker);
+}
+
+void UAvaClonerSubsystem::DisableNoFlicker()
+{
+	if (!IsNoFlickerEnabled())
+	{
+		return;
+	}
+	
+	if (PreviousCVarValue.IsSet())
+	{
+		CVarTSRShadingRejectionFlickeringPeriod->Set(PreviousCVarValue.GetValue());
+	}
+	else
+	{
+		CVarTSRShadingRejectionFlickeringPeriod->Set(*CVarTSRShadingRejectionFlickeringPeriod->GetDefaultValue());
+	}
+}
+
+bool UAvaClonerSubsystem::IsNoFlickerEnabled() const
+{
+	if (!CVarTSRShadingRejectionFlickeringPeriod)
+	{
+		return false;
+	}
+
+	return CVarTSRShadingRejectionFlickeringPeriod->GetInt() == NoFlicker;
+}
+
+void UAvaClonerSubsystem::OnTSRShadingRejectionFlickeringPeriodChanged(IConsoleVariable* InCVar) const
+{
+	if (InCVar == CVarTSRShadingRejectionFlickeringPeriod)
+	{
+		OnCVarChangedDelegate.Broadcast();
+	}
+}
+#endif
\ No newline at end of file
diff --git a/Engine/Plugins/Experimental/Avalanche/Source/AvalancheEffectors/Private/Subsystems/AvaEffectorSubsystem.cpp b/Engine/Plugins/Experimental/Avalanche/Source/AvalancheEffectors/Private/Subsystems/AvaEffectorSubsystem.cpp
new file mode 100644
--- /dev/null
+++ b/Engine/Plugins/Experimental/Avalanche/Source/AvalancheEffectors/Private/Subsystems/AvaEffectorSubsystem.cpp
@@ -0,0 +1,153 @@
+// Copyright Epic Games, Inc. All Rights Reserved.
+
+#include "Subsystems/AvaEffectorSubsystem.h"
+
+#include "NiagaraDataChannel.h"
+#include "NiagaraDataChannelAccessor.h"
+#include "NiagaraDataChannelPublic.h"
+#include "Effector/AvaEffectorActor.h"
+
+DEFINE_LOG_CATEGORY_STATIC(LogAvaEffectorSubsystem, Log, All);
+
+UAvaEffectorSubsystem::FOnSubsystemInitialized UAvaEffectorSubsystem::OnSubsystemInitializedDelegate;
+
+UAvaEffectorSubsystem* UAvaEffectorSubsystem::Get(const UWorld* InWorld)
+{
+	if (InWorld)
+	{
+		return InWorld->GetSubsystem<UAvaEffectorSubsystem>();
+	}
+
+	return nullptr;
+}
+
+void UAvaEffectorSubsystem::Initialize(FSubsystemCollectionBase& Collection)
+{
+	Super::Initialize(Collection);
+
+	// Load niagara data channel asset for effectors and cache it
+	static const FString DataChannelAssetPath = TEXT("/Script/Niagara.NiagaraDataChannelAsset'/Avalanche/ClonerResources/Channels/NDC_Effector.NDC_Effector'");
+
+	EffectorDataChannelAsset = FindObject<UNiagaraDataChannelAsset>(nullptr, *DataChannelAssetPath);
+
+	if (!EffectorDataChannelAsset)
+	{
+		EffectorDataChannelAsset = LoadObject<UNiagaraDataChannelAsset>(nullptr, *DataChannelAssetPath);
+	}
+
+	check(EffectorDataChannelAsset->Get());
+}
+
+void UAvaEffectorSubsystem::PostInitialize()
+{
+	Super::PostInitialize();
+
+	if (const UWorld* World = GetWorld())
+	{
+		OnSubsystemInitializedDelegate.Broadcast(World);
+	}
+}
+
+bool UAvaEffectorSubsystem::RegisterChannelEffector(AAvaEffectorActor* InEffector)
+{
+	if (!InEffector)
+	{
+		return false;
+	}
+
+	const int32 EffectorIndex = EffectorsWeak.AddUnique(InEffector);
+	const bool bRegistered = EffectorIndex == EffectorsWeak.Num() - 1;
+
+	if (bRegistered)
+	{
+		UE_LOG(LogAvaEffectorSubsystem, Log, TEXT("%s effector registered in channel %i"), *InEffector->GetActorNameOrLabel(), EffectorIndex);
+	}
+
+	return bRegistered;
+}
+
+bool UAvaEffectorSubsystem::UnregisterChannelEffector(AAvaEffectorActor* InEffector)
+{
+	if (!InEffector)
+	{
+		return false;
+	}
+
+	const bool bUnregistered = EffectorsWeak.Remove(InEffector) > 0;
+
+	if (bUnregistered)
+	{
+		UE_LOG(LogAvaEffectorSubsystem, Log, TEXT("%s effector unregistered from channel"), *InEffector->GetActorNameOrLabel());
+	}
+
+	return bUnregistered;
+}
+
+void UAvaEffectorSubsystem::UpdateEffectorChannel()
+{
+	const UWorld* World = GetWorld();
+
+	if (!IsValid(World))
+	{
+		return;
+	}
+
+	if (EffectorsWeak.IsEmpty())
+	{
+		return;
+	}
+
+	// Reserve space in channel for each effectors
+	static const FNiagaraDataChannelSearchParameters SearchParameters;
+	UNiagaraDataChannelWriter* ChannelWriter = UNiagaraDataChannelLibrary::WriteToNiagaraDataChannel(World, EffectorDataChannelAsset.Get(), SearchParameters, EffectorsWeak.Num(), true, true, true, TEXT("AvaEffectorSubsystem"));
+
+	if (!ChannelWriter)
+	{
+		UE_LOG(LogAvaEffectorSubsystem, Warning, TEXT("Effector data channel writer is invalid"));
+		return;
+	}
+
+	// Remove invalid effectors and push updates to effector assigned channel indexes
+	int32 EffectorIndex = 0;
+	for (TArray<TWeakObjectPtr<AAvaEffectorActor>>::TIterator It(EffectorsWeak); It; ++It)
+	{
+		AAvaEffectorActor* Effector = It->Get();
+
+		if (!Effector)
+		{
+			It.RemoveCurrent();
+			continue;
+		}
+
+		FAvaClonerEffectorChannelData& ChannelData = Effector->GetEffectorChannelData();
+
+		const bool bIdentifierChanged = ChannelData.Identifier != EffectorIndex;
+
+		// Set channel before writing 
+		ChannelData.Identifier = EffectorIndex++;
+
+		// Push effector data to channel
+		ChannelData.Write(ChannelWriter);
+
+		// When changed, update cloners DI linked to this effector
+		if (bIdentifierChanged)
+		{
+			Effector->OnEffectorIdentifierChanged();
+		}
+	}
+}
+
+bool UAvaEffectorSubsystem::IsTickableInEditor() const
+{
+	return true;
+}
+
+TStatId UAvaEffectorSubsystem::GetStatId() const
+{
+	RETURN_QUICK_DECLARE_CYCLE_STAT(UAvaEffectorSubsystem, STATGROUP_Tickables);
+}
+
+void UAvaEffectorSubsystem::Tick(float InDeltaTime)
+{
+	UpdateEffectorChannel();
+}
