[MotionDesign]

diff --git a/Engine/Plugins/Experimental/Avalanche/Source/AvalancheEffectors/Public/AvaClonerEffectorShared.h b/Engine/Plugins/Experimental/Avalanche/Source/AvalancheEffectors/Public/AvaClonerEffectorShared.h
new file mode 100644
--- /dev/null
+++ b/Engine/Plugins/Experimental/Avalanche/Source/AvalancheEffectors/Public/AvaClonerEffectorShared.h
@@ -0,0 +1,665 @@
+// Copyright Epic Games, Inc. All Rights Reserved.
+
+#pragma once
+
+#include "Components/SphereComponent.h"
+#include "Engine/StaticMesh.h"
+#include "Materials/MaterialInterface.h"
+#include "NiagaraDataInterfaceArrayInt.h"
+#include "AvaClonerEffectorShared.generated.h"
+
+class UNiagaraDataChannelReader;
+class UNiagaraDataChannelWriter;
+class AActor;
+class UDynamicMesh;
+class UNiagaraSystem;
+class USplineComponent;
+
+/** Enumerate the different layouts available to dispose cloner instances */
+UENUM(BlueprintType)
+enum class EAvaClonerLayout : uint8
+{
+	Grid,
+	Line,
+	Circle,
+	Cylinder,
+	Sphere,
+	Honeycomb,
+	SampleMesh,
+	SampleSpline
+};
+
+/** Enumerate the axis available to use */
+UENUM(BlueprintType)
+enum class EAvaClonerAxis : uint8
+{
+	X,
+	Y,
+	Z,
+	Custom
+};
+
+/** Enumerate the planes available to use */
+UENUM(BlueprintType)
+enum class EAvaClonerPlane : uint8
+{
+	XY,
+	YZ,
+	XZ,
+	Custom
+};
+
+/** Enumerate the mesh render modes available to render a cloner instance */
+UENUM(BlueprintType)
+enum class EAvaClonerMeshRenderMode : uint8
+{
+	/** Iterate through each attachment mesh available */
+	Iterate,
+	/** Pick randomly through each attachment mesh available, update cloner seed for variations */
+	Random,
+	/** Blend based on the total cloner instances and attachment mesh available */
+	Blend
+};
+
+/** Enumerate the grid constraints available when cloner layout selected is grid */
+UENUM(BlueprintType)
+enum class EAvaClonerGridConstraint : uint8
+{
+	None,
+	Sphere,
+	Cylinder,
+	Texture
+};
+
+/* Enumerates all easing functions that are available to apply on weights
+See https://easings.net/ for curves visualizations and web open domain implementations
+Used as one enum to send the index to niagara as uint8 and apply easing directly in niagara */
+UENUM(BlueprintType)
+enum class EAvaClonerEasing : uint8
+{
+	Linear,
+
+	InSine,
+	OutSine,
+	InOutSine,
+
+	InQuad,
+	OutQuad,
+	InOutQuad,
+
+	InCubic,
+	OutCubic,
+	InOutCubic,
+
+	InQuart,
+	OutQuart,
+	InOutQuart,
+
+	InQuint,
+	OutQuint,
+	InOutQuint,
+
+	InExpo,
+	OutExpo,
+	InOutExpo,
+
+	InCirc,
+	OutCirc,
+	InOutCirc,
+
+	InBack,
+	OutBack,
+	InOutBack,
+
+	InElastic,
+	OutElastic,
+	InOutElastic,
+
+	InBounce,
+	OutBounce,
+	InOutBounce,
+
+	Random
+};
+
+/** Enumerate the mesh asset to look for when mesh layout is selected */
+UENUM(BlueprintType)
+enum class EAvaClonerMeshAsset : uint8
+{
+	StaticMesh,
+	SkeletalMesh
+};
+
+/** Enumerate the mesh sample dataset to pick from when mesh layout is selected */
+UENUM(BlueprintType)
+enum class EAvaClonerMeshSampleData : uint8
+{
+	Vertices,
+	Triangles,
+	Sockets,
+	Bones,
+	Sections
+};
+
+/** Enumerates the effector shapes available */
+UENUM(BlueprintType)
+enum class EAvaClonerEffectorType : uint8
+{
+	/** Clones inside the sphere radius will be affected by the effector */
+	Sphere,
+	/** Clones between two planes will be affected by the effector */
+	Plane,
+	/** Clones inside the box extent will be affected by the effector */
+	Box,
+	/** All clones will be affected by the effector with the same max weight */
+	Unbound
+};
+
+/** Enumerates the effector mode available */
+UENUM(BlueprintType)
+enum class EAvaClonerEffectorMode : uint8
+{
+	/** Control clones offset, rotation, scale manually */
+	Default,
+	/** Rotates clones towards a target actor */
+	Target,
+	/** Randomly applies curl noise across the field zone */
+	NoiseField
+};
+
+UENUM()
+enum class EAvaClonerAttachmentStatus : uint8
+{
+	/** Item should be removed, no longer valid */
+	Invalid,
+	/** Item should be updated, changes detected */
+	Outdated,
+	/** Item is up to date, no changes needed */
+	Updated,
+	/** Item is being updated at the moment */
+	Updating
+};
+
+/** Enumerate all texture channels to sample for constraint */
+UENUM()
+enum class EAvaClonerTextureSampleChannel : uint8
+{
+	RGBLuminance,
+	RGBAverage,
+	RGBMax,
+	R,
+	G,
+	B,
+	A
+};
+
+/** Enumerate all operation compare mode for constraint */
+UENUM()
+enum class EAvaClonerCompareMode : uint8
+{
+	Greater,
+	GreaterEqual,
+	Equal,
+	NotEqual,
+	Less,
+	LessEqual
+};
+
+/** Enumerates all modes to handle cloner spawn mode */
+UENUM()
+enum class EAvaClonerSpawnLoopMode : uint8
+{
+	/** Cloner spawns once and then enters idle mode */
+	Once,
+	/** Cloner spawns multiple times and then enters idle mode */
+	Multiple,
+	/** Cloner spawns infinitely and never enters idle mode */
+	Infinite
+};
+
+/** Enumerates all modes for how clones are spawned */
+UENUM()
+enum class EAvaClonerSpawnBehaviorMode : uint8
+{
+	/** Spawns instantly the number of clones needed for the layout */
+	Instant,
+	/** Spawns at a specific rate per second during the spawn loop duration */
+	Rate,
+};
+
+USTRUCT()
+struct FAvaClonerAttachmentItem
+{
+	GENERATED_BODY()
+
+	/** Children attached to this actor, order is not important here as they are combined into one */
+	UPROPERTY()
+	TSet<TWeakObjectPtr<AActor>> ChildrenActors;
+
+	/** Current cloner attached actor represented by this item */
+	UPROPERTY()
+	TWeakObjectPtr<AActor> ItemActor;
+
+	/** Parent of this item actor or null if below cloner */
+	UPROPERTY()
+	TWeakObjectPtr<AActor> ParentActor;
+
+	/** Actual baked mesh for the current actor item,
+		dynamic because it's easier for merging them together and avoid conversions again */
+	UPROPERTY()
+	TObjectPtr<UDynamicMesh> BakedMesh;
+
+	/** Actual baked materials for the current actor item */
+	UPROPERTY()
+	TArray<TWeakObjectPtr<UMaterialInterface>> BakedMaterials;
+
+	/** Status of the baked mesh, does it needs to be updated */
+	UPROPERTY()
+	EAvaClonerAttachmentStatus MeshStatus = EAvaClonerAttachmentStatus::Outdated;
+
+	/** Last actor item transform, used to trigger an update if changed */
+	UPROPERTY()
+	FTransform ActorTransform;
+
+	/** Item is root cloner actor */
+	UPROPERTY()
+	bool bRootItem = false;
+
+	/** Status of this attachment item */
+	UPROPERTY()
+	EAvaClonerAttachmentStatus Status = EAvaClonerAttachmentStatus::Outdated;
+
+	friend uint32 GetTypeHash(const FAvaClonerAttachmentItem& InItem)
+	{
+		return GetTypeHash(InItem.ItemActor);
+	}
+
+	bool operator==(const FAvaClonerAttachmentItem& InOther) const
+	{
+		return ItemActor.Get() == InOther.ItemActor.Get();
+	}
+};
+
+USTRUCT()
+struct FAvaClonerAttachmentTree
+{
+	GENERATED_BODY()
+
+	/** All cloner attached actor items */
+	UPROPERTY()
+	TMap<TWeakObjectPtr<AActor>, FAvaClonerAttachmentItem> ItemAttachmentMap;
+
+	/** Actors directly attached to the cloner actor, order is important here */
+	UPROPERTY()
+	TArray<TWeakObjectPtr<AActor>> RootActors;
+
+	/** Merged static meshes corresponding to the root actors for niagara,
+		result of merging dynamic meshes together */
+	UPROPERTY()
+	TArray<TObjectPtr<UStaticMesh>> MergedBakedMeshes;
+
+	/** Status of the cloner tree */
+	UPROPERTY()
+	EAvaClonerAttachmentStatus Status = EAvaClonerAttachmentStatus::Updated;
+
+	void Reset()
+	{
+		ItemAttachmentMap.Empty();
+		RootActors.Empty();
+		MergedBakedMeshes.Empty();
+		Status = EAvaClonerAttachmentStatus::Updated;
+	}
+};
+
+struct FAvaClonerEffectorChannelData
+{
+	friend class UAvaEffectorSubsystem;
+
+	/** General */
+	static constexpr const TCHAR* EasingName = TEXT("Easing");
+	static constexpr const TCHAR* ModeName = TEXT("Mode");
+	static constexpr const TCHAR* TypeName = TEXT("Type");
+	static constexpr const TCHAR* MagnitudeName = TEXT("Magnitude");
+	static constexpr const TCHAR* InnerExtentName = TEXT("InnerExtent");
+	static constexpr const TCHAR* OuterExtentName = TEXT("OuterExtent");
+	static constexpr const TCHAR* LocationDeltaName = TEXT("LocationDelta");
+	static constexpr const TCHAR* RotationDeltaName = TEXT("RotationDelta");
+	static constexpr const TCHAR* ScaleDeltaName = TEXT("ScaleDelta");
+	static constexpr const TCHAR* LocationName = TEXT("Location");
+	static constexpr const TCHAR* RotationName = TEXT("Rotation");
+	static constexpr const TCHAR* ScaleName = TEXT("Scale");
+	static constexpr const TCHAR* FrequencyName = TEXT("Frequency");
+	static constexpr const TCHAR* PanName = TEXT("Pan");
+
+	/** Forces */
+	static constexpr const TCHAR* OrientationForceRateName = TEXT("OrientationForceRate");
+	static constexpr const TCHAR* OrientationForceMinName = TEXT("OrientationForceMin");
+	static constexpr const TCHAR* OrientationForceMaxName = TEXT("OrientationForceMax");
+	static constexpr const TCHAR* VortexForceAmountName = TEXT("VortexForceAmount");
+	static constexpr const TCHAR* VortexForceAxisName = TEXT("VortexForceAxis");
+	static constexpr const TCHAR* CurlNoiseForceStrengthName = TEXT("CurlNoiseForceStrength");
+	static constexpr const TCHAR* CurlNoiseForceFrequencyName = TEXT("CurlNoiseForceFrequency");
+	static constexpr const TCHAR* AttractionForceStrengthName = TEXT("AttractionForceStrength");
+	static constexpr const TCHAR* AttractionForceFalloffName = TEXT("AttractionForceFalloff");
+	static constexpr const TCHAR* GravityForceAccelerationName = TEXT("GravityForceAcceleration");
+
+	int32 GetIdentifier() const
+	{
+		return Identifier;
+	}
+
+	/** General parameters */
+	EAvaClonerEasing Easing = EAvaClonerEasing::Linear;
+	EAvaClonerEffectorMode Mode = EAvaClonerEffectorMode::Default;
+	EAvaClonerEffectorType Type = EAvaClonerEffectorType::Sphere;
+	float Magnitude = 0.f;
+	FVector InnerExtent = FVector::ZeroVector;
+	FVector OuterExtent = FVector::ZeroVector;
+	FVector LocationDelta = FVector::ZeroVector;
+	FQuat RotationDelta = FQuat::Identity;
+	FVector ScaleDelta = FVector::OneVector;
+	FVector Location = FVector::ZeroVector;
+	FQuat Rotation = FQuat::Identity;
+	FVector Scale = FVector::OneVector;
+	float Frequency = 1.f;
+	FVector Pan = FVector::ZeroVector;
+
+	/** Forces parameters */
+	float OrientationForceRate = 0.f;
+	FVector OrientationForceMin = FVector::ZeroVector;
+	FVector OrientationForceMax = FVector::ZeroVector;
+	float VortexForceAmount = 0.f;
+	FVector VortexForceAxis = FVector::ZeroVector;
+	float CurlNoiseForceStrength = 0.f;
+	float CurlNoiseForceFrequency = 0.f;
+	float AttractionForceStrength = 0.f;
+	float AttractionForceFalloff = 0.f;
+	FVector GravityForceAcceleration = FVector::ZeroVector;
+
+protected:
+	/** Cache effector identifier to detect a change and update cloners DI */
+	int32 Identifier = INDEX_NONE;
+
+	void Write(UNiagaraDataChannelWriter* InWriter) const;
+	void Read(const UNiagaraDataChannelReader* InReader);
+};
+
+USTRUCT()
+struct FAvaClonerEffectorDataInterfaces
+{
+	friend class UAvaClonerLayoutBase;
+
+	GENERATED_BODY()
+
+	static inline const FName IndexName = TEXT("EffectorIndexArray");
+
+	explicit FAvaClonerEffectorDataInterfaces(const UNiagaraSystem* InSystem);
+	FAvaClonerEffectorDataInterfaces() = default;
+
+	void CopyTo(FAvaClonerEffectorDataInterfaces& InOther) const;
+	void Resize(int32 InSize) const;
+	void Remove(int32 InIndex) const;
+	bool IsValid() const;
+	int32 Num() const;
+
+	UNiagaraDataInterfaceArrayInt32* GetIndexArray() const;
+
+protected:
+	UPROPERTY()
+	TMap<FName, TObjectPtr<UNiagaraDataInterface>> DataInterfaces;
+};
+
+USTRUCT(BlueprintType)
+struct AVALANCHEEFFECTORS_API FAvaClonerSampleMeshOptions
+{
+	GENERATED_BODY()
+
+	UPROPERTY(BlueprintReadWrite, EditInstanceOnly, Category="Cloner")
+	int32 Count = 3 * 3 * 3;
+
+	UPROPERTY(BlueprintReadWrite, EditInstanceOnly, Category="Cloner")
+	EAvaClonerMeshAsset Asset = EAvaClonerMeshAsset::StaticMesh;
+
+	UPROPERTY(BlueprintReadWrite, EditInstanceOnly, Category="Cloner")
+	EAvaClonerMeshSampleData SampleData = EAvaClonerMeshSampleData::Vertices;
+
+	UPROPERTY(BlueprintReadWrite, EditInstanceOnly, Category="Cloner")
+	TWeakObjectPtr<AActor> SampleActor = nullptr;
+
+	UPROPERTY()
+	TWeakObjectPtr<USceneComponent> SceneComponent = nullptr;
+};
+
+USTRUCT(BlueprintType)
+struct AVALANCHEEFFECTORS_API FAvaClonerSampleSplineOptions
+{
+	GENERATED_BODY()
+
+	UPROPERTY(BlueprintReadWrite, EditInstanceOnly, Category="Cloner")
+	int32 Count = 3 * 3;
+
+	UPROPERTY(BlueprintReadWrite, EditInstanceOnly, Category="Cloner")
+	TWeakObjectPtr<AActor> SplineActor = nullptr;
+
+	UPROPERTY()
+	TWeakObjectPtr<USplineComponent> SplineComponent = nullptr;
+};
+
+USTRUCT(BlueprintType)
+struct AVALANCHEEFFECTORS_API FAvaClonerGridConstraintSphere
+{
+	GENERATED_BODY()
+
+	UPROPERTY(BlueprintReadWrite, EditInstanceOnly, Category="Cloner")
+	float Radius = 400.f;
+
+	UPROPERTY(BlueprintReadWrite, EditInstanceOnly, Category="Cloner")
+	FVector Center = FVector(0);
+};
+
+USTRUCT(BlueprintType)
+struct AVALANCHEEFFECTORS_API FAvaClonerGridConstraintCylinder
+{
+	GENERATED_BODY()
+
+	UPROPERTY(BlueprintReadWrite, EditInstanceOnly, Category="Cloner")
+	float Radius = 400.f;
+
+	UPROPERTY(BlueprintReadWrite, EditInstanceOnly, Category="Cloner")
+	float Height = 800.f;
+
+	UPROPERTY(BlueprintReadWrite, EditInstanceOnly, Category="Cloner")
+	FVector Center = FVector(0);
+};
+
+USTRUCT(BlueprintType)
+struct AVALANCHEEFFECTORS_API FAvaClonerGridConstraintTexture
+{
+	GENERATED_BODY()
+
+	UPROPERTY(BlueprintReadWrite, EditInstanceOnly, Category="Cloner")
+	TWeakObjectPtr<UTexture> Texture;
+
+	UPROPERTY(BlueprintReadWrite, EditInstanceOnly, Category="Cloner")
+	EAvaClonerTextureSampleChannel Channel = EAvaClonerTextureSampleChannel::RGBLuminance;
+
+	UPROPERTY(BlueprintReadWrite, EditInstanceOnly, Category="Cloner")
+	EAvaClonerPlane Plane = EAvaClonerPlane::XY;
+
+	UPROPERTY(BlueprintReadWrite, EditInstanceOnly, Category="Cloner")
+	EAvaClonerCompareMode CompareMode = EAvaClonerCompareMode::Greater;
+
+	UPROPERTY(BlueprintReadWrite, EditInstanceOnly, Category="Cloner")
+	float Threshold = 0.f;
+};
+
+USTRUCT(BlueprintType)
+struct AVALANCHEEFFECTORS_API FAvaClonerGridLayoutOptions
+{
+	GENERATED_BODY()
+
+	UPROPERTY(BlueprintReadWrite, EditInstanceOnly, Category="Cloner")
+	FIntVector Count = FIntVector(3, 3, 3);
+
+	UPROPERTY(BlueprintReadWrite, EditInstanceOnly, Category="Cloner")
+	FVector Spacing = FVector(400.f, 400.f, 400.f);
+
+	UPROPERTY(BlueprintReadWrite, EditInstanceOnly, Category="Cloner")
+	EAvaClonerGridConstraint Constraint = EAvaClonerGridConstraint::None;
+
+	UPROPERTY(BlueprintReadWrite, EditInstanceOnly, Category="Cloner", meta=(EditCondition="Constraint != EAvaClonerGridConstraint::None", EditConditionHides))
+	bool bInvertConstraint = false;
+
+	UPROPERTY(BlueprintReadWrite, EditInstanceOnly, Category="Cloner", meta=(EditCondition="Constraint == EAvaClonerGridConstraint::Sphere", EditConditionHides))
+	FAvaClonerGridConstraintSphere SphereConstraint;
+
+	UPROPERTY(BlueprintReadWrite, EditInstanceOnly, Category="Cloner", meta=(EditCondition="Constraint == EAvaClonerGridConstraint::Cylinder", EditConditionHides))
+	FAvaClonerGridConstraintCylinder CylinderConstraint;
+
+	UPROPERTY(BlueprintReadWrite, EditInstanceOnly, Category="Cloner", meta=(EditCondition="Constraint == EAvaClonerGridConstraint::Texture", EditConditionHides))
+	FAvaClonerGridConstraintTexture TextureConstraint;
+};
+
+USTRUCT(BlueprintType)
+struct AVALANCHEEFFECTORS_API FAvaClonerLineLayoutOptions
+{
+	GENERATED_BODY()
+
+	UPROPERTY(BlueprintReadWrite, EditInstanceOnly, Category="Cloner")
+	int32 Count = 3 * 3;
+
+	UPROPERTY(BlueprintReadWrite, EditInstanceOnly, Category="Cloner")
+	float Spacing = 400.f;
+
+	UPROPERTY(BlueprintReadWrite, EditInstanceOnly, Category="Cloner")
+	EAvaClonerAxis Axis = EAvaClonerAxis::Y;
+
+	UPROPERTY(BlueprintReadWrite, EditInstanceOnly, Category="Cloner", meta=(ClampMin="0", ClampMax="1", EditCondition="Axis == EAvaClonerAxis::Custom", EditConditionHides))
+	FVector Direction = FVector::YAxisVector;
+
+	UPROPERTY(BlueprintReadWrite, EditInstanceOnly, Category="Cloner")
+	FRotator Rotation = FRotator(0.f);
+};
+
+USTRUCT(BlueprintType)
+struct AVALANCHEEFFECTORS_API FAvaClonerCircleLayoutOptions
+{
+	GENERATED_BODY()
+
+	UPROPERTY(BlueprintReadWrite, EditInstanceOnly, Category="Cloner")
+	int32 Count = 3 * 3;
+
+	UPROPERTY(BlueprintReadWrite, EditInstanceOnly, Category="Cloner")
+	float Radius = 400.f;
+
+	UPROPERTY(BlueprintReadWrite, EditInstanceOnly, Category="Cloner")
+	float AngleStart = 0.f;
+
+	UPROPERTY(BlueprintReadWrite, EditInstanceOnly, Category="Cloner")
+	float AngleRatio = 1.f;
+
+	UPROPERTY(BlueprintReadWrite, EditInstanceOnly, Category="Cloner")
+	bool bOrientMesh = false;
+
+	UPROPERTY(BlueprintReadWrite, EditInstanceOnly, Category="Cloner")
+	EAvaClonerPlane Plane = EAvaClonerPlane::XY;
+
+	UPROPERTY(BlueprintReadWrite, EditInstanceOnly, Category="Cloner", meta=(EditCondition="Plane == EAvaClonerPlane::Custom", EditConditionHides))
+	FRotator Rotation = FRotator(0.f, 0.f, 0.f);
+
+	UPROPERTY(BlueprintReadWrite, EditInstanceOnly, Category="Cloner", meta=(ClampMin="0", AllowPreserveRatio, Delta="0.01"))
+	FVector Scale = FVector(1.f, 1.f, 1.f);
+};
+
+USTRUCT(BlueprintType)
+struct AVALANCHEEFFECTORS_API FAvaClonerCylinderLayoutOptions
+{
+	GENERATED_BODY()
+
+	UPROPERTY(BlueprintReadWrite, EditInstanceOnly, Category="Cloner")
+	int32 BaseCount = 3 * 3;
+
+	UPROPERTY(BlueprintReadWrite, EditInstanceOnly, Category="Cloner")
+	int32 HeightCount = 3;
+
+	UPROPERTY(BlueprintReadWrite, EditInstanceOnly, Category="Cloner")
+	float Height = 400.f;
+
+	UPROPERTY(BlueprintReadWrite, EditInstanceOnly, Category="Cloner")
+	float Radius = 400.f;
+
+	UPROPERTY(BlueprintReadWrite, EditInstanceOnly, Category="Cloner")
+	float AngleStart = 0.f;
+
+	UPROPERTY(BlueprintReadWrite, EditInstanceOnly, Category="Cloner")
+	float AngleRatio = 1.f;
+
+	UPROPERTY(BlueprintReadWrite, EditInstanceOnly, Category="Cloner")
+	bool bOrientMesh = false;
+
+	UPROPERTY(BlueprintReadWrite, EditInstanceOnly, Category="Cloner")
+	EAvaClonerPlane Plane = EAvaClonerPlane::XY;
+
+	UPROPERTY(BlueprintReadWrite, EditInstanceOnly, Category="Cloner", meta=(EditCondition="Plane == EAvaClonerPlane::Custom", EditConditionHides))
+	FRotator Rotation = FRotator(0.f, 0.f, 0.f);
+
+	UPROPERTY(BlueprintReadWrite, EditInstanceOnly, Category="Cloner", meta=(ClampMin="0", AllowPreserveRatio, Delta="0.01"))
+	FVector Scale = FVector(1.f, 1.f, 1.f);
+};
+
+USTRUCT(BlueprintType)
+struct AVALANCHEEFFECTORS_API FAvaClonerSphereLayoutOptions
+{
+	GENERATED_BODY()
+
+	UPROPERTY(BlueprintReadWrite, EditInstanceOnly, Category="Cloner")
+	int32 Count = 3 * 3 * 3;
+
+	UPROPERTY(BlueprintReadWrite, EditInstanceOnly, Category="Cloner")
+	float Radius = 400.f;
+
+	UPROPERTY(BlueprintReadWrite, EditInstanceOnly, Category="Cloner", meta=(ClampMin="0", ClampMax="1"))
+	float Ratio = 1.f;
+
+	UPROPERTY(BlueprintReadWrite, EditInstanceOnly, Category="Cloner")
+	bool bOrientMesh = false;
+
+	UPROPERTY(BlueprintReadWrite, EditInstanceOnly, Category="Cloner")
+	FRotator Rotation = FRotator(0.f, 0.f, 0.f);
+
+	UPROPERTY(BlueprintReadWrite, EditInstanceOnly, Category="Cloner", meta=(ClampMin="0", AllowPreserveRatio, Delta="0.01"))
+	FVector Scale = FVector(1.f, 1.f, 1.f);
+};
+
+USTRUCT(BlueprintType)
+struct AVALANCHEEFFECTORS_API FAvaClonerHoneycombLayoutOptions
+{
+	GENERATED_BODY()
+
+	UPROPERTY(BlueprintReadWrite, EditInstanceOnly, Category="Cloner")
+	EAvaClonerPlane Plane = EAvaClonerPlane::XY;
+
+	UPROPERTY(BlueprintReadWrite, EditInstanceOnly, Category="Cloner")
+	int32 WidthCount = 3;
+
+	UPROPERTY(BlueprintReadWrite, EditInstanceOnly, Category="Cloner")
+	int32 HeightCount = 3;
+
+	UPROPERTY(BlueprintReadWrite, EditInstanceOnly, Category="Cloner")
+	float WidthOffset = 0.f;
+
+	UPROPERTY(BlueprintReadWrite, EditInstanceOnly, Category="Cloner")
+	float HeightOffset = 0.5f;
+
+	UPROPERTY(BlueprintReadWrite, EditInstanceOnly, Category="Cloner")
+	float WidthSpacing = 400.f;
+
+	UPROPERTY(BlueprintReadWrite, EditInstanceOnly, Category="Cloner")
+	float HeightSpacing = 400.f;
+};
+
+namespace UE::ClonerEffector
+{
+	void SetBillboardComponentSprite(const AActor* InActor, const FString& InTexturePath);
+	void SetBillboardComponentVisibility(const AActor* InActor, bool bInVisibility);
+}
\ No newline at end of file
diff --git a/Engine/Plugins/Experimental/Avalanche/Source/AvalancheEffectors/Public/Cloner/AvaClonerActor.h b/Engine/Plugins/Experimental/Avalanche/Source/AvalancheEffectors/Public/Cloner/AvaClonerActor.h
new file mode 100644
--- /dev/null
+++ b/Engine/Plugins/Experimental/Avalanche/Source/AvalancheEffectors/Public/Cloner/AvaClonerActor.h
@@ -0,0 +1,702 @@
+// Copyright Epic Games, Inc. All Rights Reserved.
+
+#pragma once
+
+#include "AvaClonerEffectorShared.h"
+#include "Curves/RichCurve.h"
+#include "GameFramework/Actor.h"
+#include "Layouts/AvaClonerLayoutBase.h"
+#include "NiagaraMeshRendererProperties.h"
+#include "AvaPropertyChangeDispatcher.h"
+#include "AvaClonerActor.generated.h"
+
+class AActor;
+class AAvaEffectorActor;
+class UAvaClonerComponent;
+class UAvaClonerLayoutBase;
+class UNiagaraDataInterfaceCurve;
+class USceneComponent;
+struct FAvaClonerSystem;
+
+UCLASS(BlueprintType, HideCategories=(Rendering,Replication,Collision,HLOD,Physics,Networking,Input,Actor,Cooking,LevelInstance,Streaming), DisplayName = "Motion Design Cloner Actor")
+class AVALANCHEEFFECTORS_API AAvaClonerActor : public AActor
+{
+	GENERATED_BODY()
+
+	friend class AAvaEffectorActor;
+	friend class UAvaClonerLayoutBase;
+
+	friend class FAvaClonerActorVisualizer;
+	friend class FAvaClonerDetailCustomization;
+
+public:
+	static inline const FString DefaultLabel = TEXT("Cloner");
+
+	AAvaClonerActor();
+
+	UFUNCTION(BlueprintCallable, Category="Cloner")
+	void SetTreeUpdateInterval(float InInterval);
+
+	UFUNCTION(BlueprintPure, Category="Cloner")
+	float GetTreeUpdateInterval() const
+	{
+		return TreeUpdateInterval;
+	}
+
+	UFUNCTION(BlueprintCallable, Category="Cloner")
+	void SetMeshRenderMode(EAvaClonerMeshRenderMode InMode);
+
+	UFUNCTION(BlueprintPure, Category="Cloner")
+	EAvaClonerMeshRenderMode GetMeshRenderMode() const
+	{
+		return MeshRenderMode;
+	}
+
+	UFUNCTION(BlueprintCallable, Category="Cloner")
+	void SetMeshFacingMode(ENiagaraMeshFacingMode InMode);
+
+	UFUNCTION(BlueprintPure, Category="Cloner")
+	ENiagaraMeshFacingMode GetMeshFacingMode() const
+	{
+		return MeshFacingMode;
+	}
+
+	UFUNCTION(BlueprintCallable, Category="Cloner")
+	void SetMeshCastShadows(bool InbCastShadows);
+
+	UFUNCTION(BlueprintPure, Category="Cloner")
+	bool GetMeshCastShadows() const
+	{
+		return bMeshCastShadows;
+	}
+
+	void SetDefaultMeshes(const TArray<TObjectPtr<UStaticMesh>>& InMeshes);
+
+	const TArray<TObjectPtr<UStaticMesh>>& GetDefaultMeshes() const
+	{
+		return DefaultMeshes;
+	}
+
+	UFUNCTION(BlueprintCallable, Category="Cloner", meta=(DisplayName="SetDefaultMeshes"))
+	void BP_SetDefaultMeshes(const TArray<UStaticMesh*>& InMeshes);
+
+	UFUNCTION(BlueprintPure, Category="Cloner", meta=(DisplayName="GetDefaultMeshes"))
+	TArray<UStaticMesh*> BP_GetDefaultMeshes() const;
+
+	UFUNCTION(BlueprintCallable, Category="Cloner")
+	void SetSeed(int32 InSeed);
+
+	UFUNCTION(BlueprintPure, Category="Cloner")
+	int32 GetSeed() const
+	{
+		return Seed;
+	}
+
+	UFUNCTION(BlueprintCallable, Category="Cloner")
+	void SetDeltaStepEnabled(bool bInEnabled);
+
+	UFUNCTION(BlueprintPure, Category="Cloner")
+	bool GetDeltaStepEnabled() const
+	{
+		return bDeltaStepEnabled;
+	}
+
+	UFUNCTION(BlueprintCallable, Category="Cloner")
+	void SetDeltaStepRotation(const FRotator& InRotation);
+
+	UFUNCTION(BlueprintPure, Category="Cloner")
+	FRotator GetDeltaStepRotation() const
+	{
+		return DeltaStepRotation;
+	}
+
+	UFUNCTION(BlueprintCallable, Category="Cloner")
+	void SetDeltaStepScale(const FVector& InScale);
+
+	UFUNCTION(BlueprintPure, Category="Cloner")
+	FVector GetDeltaStepScale() const
+	{
+		return DeltaStepScale;
+	}
+
+	UFUNCTION(BlueprintCallable, Category="Cloner")
+	void SetEnabled(bool bInEnable);
+
+	UFUNCTION(BlueprintPure, Category="Cloner")
+	bool GetEnabled() const
+	{
+		return bEnabled;
+	}
+
+	UFUNCTION(BlueprintCallable, Category="Cloner")
+	void SetInvertProgress(bool bInInvertProgress);
+
+	UFUNCTION(BlueprintPure, Category="Cloner")
+	bool GetInvertProgress() const
+	{
+		return bInvertProgress;
+	}
+
+	UFUNCTION(BlueprintCallable, Category="Cloner")
+	void SetProgress(float InProgress);
+
+	UFUNCTION(BlueprintPure, Category="Cloner")
+	float GetProgress() const
+	{
+		return Progress;
+	}
+
+	UFUNCTION(BlueprintCallable, Category="Cloner")
+	void SetLayoutName(FName InLayoutName);
+
+	UFUNCTION(BlueprintPure, Category="Cloner")
+	FName GetLayoutName() const
+	{
+		return LayoutName;
+	}
+
+	UFUNCTION(BlueprintCallable, Category="Cloner")
+	void SetRangeEnabled(bool bInRangeEnabled);
+
+	UFUNCTION(BlueprintPure, Category="Cloner")
+	bool GetRangeEnabled() const
+	{
+		return bRangeEnabled;
+	}
+
+	UFUNCTION(BlueprintCallable, Category="Cloner")
+	void SetRangeOffsetMin(const FVector& InRangeOffsetMin);
+
+	UFUNCTION(BlueprintPure, Category="Cloner")
+	const FVector& GetRangeOffsetMin() const
+	{
+		return RangeOffsetMin;
+	}
+
+	UFUNCTION(BlueprintCallable, Category="Cloner")
+	void SetRangeOffsetMax(const FVector& InRangeOffsetMax);
+
+	UFUNCTION(BlueprintPure, Category="Cloner")
+	const FVector& GetRangeOffsetMax() const
+	{
+		return RangeOffsetMax;
+	}
+
+	UFUNCTION(BlueprintCallable, Category="Cloner")
+	void SetRangeRotationMin(const FRotator& InRangeRotationMin);
+
+	UFUNCTION(BlueprintPure, Category="Cloner")
+	const FRotator& GetRangeRotationMin() const
+	{
+		return RangeRotationMin;
+	}
+
+	UFUNCTION(BlueprintCallable, Category="Cloner")
+	void SetRangeRotationMax(const FRotator& InRangeRotationMax);
+
+	UFUNCTION(BlueprintPure, Category="Cloner")
+	const FRotator& GetRangeRotationMax() const
+	{
+		return RangeRotationMax;
+	}
+
+	UFUNCTION(BlueprintCallable, Category="Cloner")
+	void SetRangeScaleUniform(bool bInRangeScaleUniform);
+
+	UFUNCTION(BlueprintPure, Category="Cloner")
+	bool GetRangeScaleUniform() const
+	{
+		return bRangeScaleUniform;
+	}
+
+	UFUNCTION(BlueprintCallable, Category="Cloner")
+	void SetRangeScaleMin(const FVector& InRangeScaleMin);
+
+	UFUNCTION(BlueprintPure, Category="Cloner")
+	const FVector& GetRangeScaleMin() const
+	{
+		return RangeScaleMin;
+	}
+
+	UFUNCTION(BlueprintCallable, Category="Cloner")
+	void SetRangeScaleMax(const FVector& InRangeScaleMax);
+
+	UFUNCTION(BlueprintPure, Category="Cloner")
+	const FVector& GetRangeScaleMax() const
+	{
+		return RangeScaleMax;
+	}
+
+	UFUNCTION(BlueprintCallable, Category="Cloner")
+	void SetRangeScaleUniformMin(float InRangeScaleUniformMin);
+
+	UFUNCTION(BlueprintPure, Category="Cloner")
+	float GetRangeScaleUniformMin() const
+	{
+		return RangeScaleUniformMin;
+	}
+
+	UFUNCTION(BlueprintCallable, Category="Cloner")
+	void SetRangeScaleUniformMax(float InRangeScaleUniformMax);
+
+	UFUNCTION(BlueprintPure, Category="Cloner")
+	float GetRangeScaleUniformMax() const
+	{
+		return RangeScaleUniformMax;
+	}
+
+	UFUNCTION(BlueprintCallable, Category="Cloner")
+	void SetSpawnLoopMode(EAvaClonerSpawnLoopMode InMode);
+
+	UFUNCTION(BlueprintPure, Category="Cloner")
+	EAvaClonerSpawnLoopMode GetSpawnLoopMode() const
+	{
+		return SpawnLoopMode;
+	}
+
+	UFUNCTION(BlueprintCallable, Category="Cloner")
+	void SetSpawnLoopIterations(int32 InIterations);
+
+	UFUNCTION(BlueprintPure, Category="Cloner")
+	int32 GetSpawnLoopIterations() const
+	{
+		return SpawnLoopIterations;
+	}
+
+	UFUNCTION(BlueprintCallable, Category="Cloner")
+	void SetSpawnLoopInterval(float InInterval);
+
+	UFUNCTION(BlueprintPure, Category="Cloner")
+	float GetSpawnLoopInterval() const
+	{
+		return SpawnLoopInterval;
+	}
+
+	UFUNCTION(BlueprintCallable, Category="Cloner")
+	void SetSpawnBehaviorMode(EAvaClonerSpawnBehaviorMode InMode);
+
+	UFUNCTION(BlueprintPure, Category="Cloner")
+	EAvaClonerSpawnBehaviorMode GetSpawnBehaviorMode() const
+	{
+		return SpawnBehaviorMode;
+	}
+
+	UFUNCTION(BlueprintCallable, Category="Cloner")
+	void SetSpawnRate(float InRate);
+
+	UFUNCTION(BlueprintPure, Category="Cloner")
+	float GetSpawnRate() const
+	{
+		return SpawnRate;
+	}
+
+	UFUNCTION(BlueprintCallable, Category="Cloner")
+	void SetLifetimeEnabled(bool bInEnabled);
+
+	UFUNCTION(BlueprintPure, Category="Cloner")
+	bool GetLifetimeEnabled() const
+	{
+		return bLifetimeEnabled;
+	}
+
+	UFUNCTION(BlueprintCallable, Category="Cloner")
+	void SetLifetimeMin(float InMin);
+
+	UFUNCTION(BlueprintPure, Category="Cloner")
+	float GetLifetimeMin() const
+	{
+		return LifetimeMin;
+	}
+
+	UFUNCTION(BlueprintCallable, Category="Cloner")
+	void SetLifetimeMax(float InMax);
+
+	UFUNCTION(BlueprintPure, Category="Cloner")
+	float GetLifetimeMax() const
+	{
+		return LifetimeMax;
+	}
+
+	UFUNCTION(BlueprintCallable, Category="Cloner")
+	void SetLifetimeScaleEnabled(bool bInEnabled);
+
+	UFUNCTION(BlueprintPure, Category="Cloner")
+	bool GetLifetimeScaleEnabled() const
+	{
+		return bLifetimeScaleEnabled;
+	}
+
+	UFUNCTION()
+	void SetLifetimeScaleCurve(const FRichCurve& InCurve);
+
+	UFUNCTION()
+	const FRichCurve& GetLifetimeScaleCurve() const
+	{
+		return LifetimeScaleCurve;
+	}
+
+	UFUNCTION(BlueprintPure, Category="Cloner")
+	UAvaClonerLayoutBase* GetActiveLayout() const;
+
+	template<
+		typename InLayoutClass
+		UE_REQUIRES(TIsDerivedFrom<InLayoutClass, UAvaClonerLayoutBase>::Value)>
+	bool IsActiveLayout() const
+	{
+		if (const UAvaClonerLayoutBase* CurrentLayout = GetActiveLayout())
+		{
+			return CurrentLayout->GetClass() == InLayoutClass::StaticClass();
+		}
+
+		return false;
+	}
+
+	template<
+		typename InLayoutClass
+		UE_REQUIRES(TIsDerivedFrom<InLayoutClass, UAvaClonerLayoutBase>::Value)>
+	InLayoutClass* GetActiveLayout() const
+	{
+		return Cast<InLayoutClass>(GetActiveLayout());
+	}
+
+	/** Returns the number of meshes this cloner currently handles */
+	UFUNCTION(BlueprintPure, Category="Cloner")
+	int32 GetMeshCount() const;
+
+	/** Gets the number of effectors applied on this cloner */
+	UFUNCTION(BlueprintPure, Category="Cloner")
+	int32 GetEffectorCount() const;
+
+#if WITH_EDITOR
+	/** This will force an update of the cloner attachment tree */
+	UFUNCTION(CallInEditor, Category="Cloner")
+	void ForceUpdateCloner();
+
+	/** This will spawn an effector actor and link it to this cloner */
+	UFUNCTION(CallInEditor, Category="Cloner")
+	void SpawnLinkedEffector();
+
+	/** When the cloner actor is created, we spawn a default actor attached */
+	void SpawnDefaultActorAttached();
+#endif
+
+	bool GetClonerInitialized() const
+	{
+		return bClonerInitialized;
+	}
+
+protected:
+	//~ Begin UObject
+	virtual void PostInitProperties() override;
+	virtual void PostLoad() override;
+#if WITH_EDITOR
+	virtual void PostEditChangeProperty(FPropertyChangedEvent& InPropertyChangedEvent) override;
+#endif
+	//~ End UObject
+
+	//~ Begin AActor
+	virtual void PostActorCreated() override;
+	/** Allow ticking in editor viewport too */
+	virtual bool ShouldTickIfViewportsOnly() const override
+	{
+		return true;
+	}
+	/** Tick needed to keep the attachment tree updated */
+	virtual void Tick(float DeltaSeconds) override;
+#if WITH_EDITOR
+	virtual FString GetDefaultActorLabel() const override;
+#endif
+	//~ End AActor
+
+	UAvaClonerComponent* GetClonerComponent() const
+	{
+		return ClonerComponent;
+	}
+
+	UNiagaraDataInterfaceCurve* GetLifetimeScaleCurveDI() const
+	{
+		return LifetimeScaleCurveDIWeak.Get();
+	}
+
+	/** Will force a system update to refresh user parameters */
+	void RequestClonerUpdate(bool bInImmediate = false);
+
+	void UpdateLayoutOptions();
+	void UpdateClonerEffectors();
+
+	/** Used by effector actors to apply transformations to this cloner instances */
+	const FAvaClonerEffectorDataInterfaces* GetEffectorDataInterfaces() const;
+
+	/** Register a new effector and return the index of this effector to match niagara data interface */
+	int32 RegisterEffector(AAvaEffectorActor* InEffector);
+
+	/** Unregister the effector and reorganize the niagara data interface to match with the new array indexes */
+	bool UnregisterEffector(AAvaEffectorActor* InEffector);
+
+	/** Checks if an effector is registered with this cloner */
+	bool IsEffectorRegistered(const AAvaEffectorActor* InEffector) const;
+
+	/** Gets the index of an effector in this cloner */
+	int32 GetEffectorIndex(AAvaEffectorActor* InEffector) const;
+
+	/** For each valid effector that has this cloner registered */
+	void ForEachEffector(TFunctionRef<bool(AAvaEffectorActor*, int32)> InFunction);
+
+	void OnClonerTransformed(USceneComponent*, EUpdateTransformFlags, ETeleportType);
+	void OnClonerMeshUpdated();
+	void OnClonerSystemChanged();
+
+	void OnEnabledChanged();
+	void OnMeshRenderModeChanged();
+	void OnMeshRendererOptionsChanged();
+	void OnDefaultMeshesChanged();
+	void OnSeedChanged();
+	void OnProgressChanged();
+	void OnDeltaStepChanged();
+	void OnRangeOptionsChanged();
+	void OnLayoutNameChanged();
+	void OnSpawnOptionsChanged();
+	void OnLifetimeScaleCurveChanged();
+	void OnLifetimeOptionsChanged();
+
+	/** Update sprite visibility of this cloner */
+	void OnVisualizerSpriteVisibleChanged();
+
+	bool IsClonerValid() const;
+
+#if WITH_EDITOR
+	void OnReduceMotionGhostingChanged();
+	void OnCVarChanged();
+#endif
+
+	/** Is this cloner enabled/disabled */
+	UPROPERTY(EditAnywhere, BlueprintReadWrite, Setter="SetEnabled", Getter="GetEnabled", Category="Cloner")
+	bool bEnabled = true;
+
+	/** Interval to update the attachment tree and update the cloner meshes */
+	UPROPERTY(EditAnywhere, BlueprintReadWrite, Setter="SetTreeUpdateInterval", Getter="GetTreeUpdateInterval", Category="Cloner", AdvancedDisplay, meta=(ClampMin="0.0"))
+	float TreeUpdateInterval = 0.2f;
+
+	/** Cloner instance seed for random deterministic patterns */
+	UPROPERTY(EditAnywhere, BlueprintReadWrite, Setter="SetSeed", Getter="GetSeed", Category="Cloner")
+	int32 Seed = 0;
+
+	/** Indicates how we select the mesh to render on each clones */
+	UPROPERTY(EditAnywhere, BlueprintReadWrite, Setter="SetMeshRenderMode", Getter="GetMeshRenderMode", Category="Renderer")
+	EAvaClonerMeshRenderMode MeshRenderMode = EAvaClonerMeshRenderMode::Iterate;
+
+	/** Mode to indicate how clones facing is determined */
+	UPROPERTY(EditAnywhere, BlueprintReadWrite, Setter="SetMeshFacingMode", Getter="GetMeshFacingMode", Category="Renderer")
+	ENiagaraMeshFacingMode MeshFacingMode = ENiagaraMeshFacingMode::Default;
+
+	/** Whether clones cast shadow, disabling will result in better performance */
+	UPROPERTY(EditAnywhere, BlueprintReadWrite, Setter="SetMeshCastShadows", Getter="GetMeshCastShadows", Category="Renderer")
+	bool bMeshCastShadows = true;
+
+	/** When nothing is attached to the cloner, these meshes are used as default */
+	UPROPERTY(EditAnywhere, BlueprintReadWrite, Setter="SetDefaultMeshes", Getter="GetDefaultMeshes", BlueprintSetter="BP_SetDefaultMeshes", BlueprintGetter="BP_GetDefaultMeshes", Category="Renderer")
+	TArray<TObjectPtr<UStaticMesh>> DefaultMeshes;
+
+	/** How many times do we spawn clones */
+	UPROPERTY(EditAnywhere, BlueprintReadWrite, Setter="SetSpawnLoopMode", Getter="GetSpawnLoopMode", DisplayName="LoopMode", Category="Spawn")
+	EAvaClonerSpawnLoopMode SpawnLoopMode = EAvaClonerSpawnLoopMode::Once;
+
+	/** Amount of spawn iterations for clones */
+	UPROPERTY(EditAnywhere, BlueprintReadWrite, Setter="SetSpawnLoopIterations", Getter="GetSpawnLoopIterations", DisplayName="LoopIterations", Category="Spawn", meta=(ClampMin="1", EditCondition="SpawnLoopMode == EAvaClonerSpawnLoopMode::Multiple", EditConditionHides))
+	int32 SpawnLoopIterations = 1;
+
+	/** Interval/Duration of spawn for clones */
+	UPROPERTY(EditAnywhere, BlueprintReadWrite, Setter="SetSpawnLoopInterval", Getter="GetSpawnLoopInterval", DisplayName="LoopInterval", Category="Spawn", meta=(ClampMin="0"))
+	float SpawnLoopInterval = 1.f;
+
+	/** How does spawn occurs */
+	UPROPERTY(EditAnywhere, BlueprintReadWrite, Setter="SetSpawnBehaviorMode", Getter="GetSpawnBehaviorMode", DisplayName="BehaviorMode", Category="Spawn")
+	EAvaClonerSpawnBehaviorMode SpawnBehaviorMode = EAvaClonerSpawnBehaviorMode::Instant;
+
+	/** How many clones to spawn each seconds */
+	UPROPERTY(EditAnywhere, BlueprintReadWrite, Setter="SetSpawnRate", Getter="GetSpawnRate", DisplayName="Rate", Category="Spawn", meta=(ClampMin="0", EditCondition="SpawnBehaviorMode == EAvaClonerSpawnBehaviorMode::Rate", EditConditionHides))
+	float SpawnRate = 1.f;
+
+	/** Do we destroy the clones after a specific duration */
+	UPROPERTY(EditAnywhere, BlueprintReadWrite, Setter="SetLifetimeEnabled", Getter="GetLifetimeEnabled", DisplayName="Enabled", Category="Lifetime")
+	bool bLifetimeEnabled = false;
+
+	/** Minimum lifetime for a clone */
+	UPROPERTY(EditAnywhere, BlueprintReadWrite, Setter="SetLifetimeMin", Getter="GetLifetimeMin", DisplayName="Min", Category="Lifetime", meta=(ClampMin="0", EditCondition="bLifetimeEnabled", EditConditionHides))
+	float LifetimeMin = 0.25f;
+
+	/** Maximum lifetime for a clone */
+	UPROPERTY(EditAnywhere, BlueprintReadWrite, Setter="SetLifetimeMax", Getter="GetLifetimeMax", DisplayName="Max", Category="Lifetime", meta=(ClampMin="0", EditCondition="bLifetimeEnabled", EditConditionHides))
+	float LifetimeMax = 1.5f;
+
+	/** Enable scale by lifetime */
+	UPROPERTY(EditAnywhere, BlueprintReadWrite, Setter="SetLifetimeScaleEnabled", Getter="GetLifetimeScaleEnabled", DisplayName="ScaleEnabled", Category="Lifetime", meta=(EditCondition="bLifetimeEnabled", EditConditionHides))
+	bool bLifetimeScaleEnabled = false;
+
+	/** Used to expose the scale curve editor in details panel */
+	UPROPERTY(Transient, DuplicateTransient, TextExportTransient)
+	TWeakObjectPtr<UNiagaraDataInterfaceCurve> LifetimeScaleCurveDIWeak;
+
+	UPROPERTY(Setter="SetLifetimeScaleCurve", Getter="GetLifetimeScaleCurve")
+	FRichCurve LifetimeScaleCurve;
+
+	/** Name of the layout to use */
+	UPROPERTY(EditAnywhere, BlueprintReadWrite, Setter="SetLayoutName", Getter="GetLayoutName", Category="Layout", meta=(GetOptions="GetClonerLayoutNames"))
+	FName LayoutName = NAME_None;
+
+	/** Active layout */
+	UPROPERTY(VisibleInstanceOnly, BlueprintReadOnly, Instanced, Transient, DuplicateTransient, Category="Layout", meta=(DisplayAfter="LayoutName"))
+	TObjectPtr<UAvaClonerLayoutBase> ActiveLayout;
+
+	/** Invert progress behaviour */
+	UPROPERTY(EditAnywhere, BlueprintReadWrite, Setter="SetInvertProgress", Getter="GetInvertProgress", Category="Progress")
+	bool bInvertProgress = false;
+
+	/** Changes visibility of instances based on the total count, 1.f = 100% = all instances visible */
+	UPROPERTY(EditAnywhere, BlueprintReadWrite, Setter="SetProgress", Getter="GetProgress", Category="Progress", meta=(ClampMin="0", ClampMax="1"))
+	float Progress = 1.f;
+
+	/** Enable steps to add delta variation on each clone instance */
+	UPROPERTY(EditAnywhere, BlueprintReadWrite, Setter="SetDeltaStepEnabled", Getter="GetDeltaStepEnabled", DisplayName="Enabled", Category="Step")
+	bool bDeltaStepEnabled = false;
+
+	/** Amount of rotation difference between one step and the next one */
+	UPROPERTY(EditAnywhere, BlueprintReadWrite, Setter="SetDeltaStepRotation", Getter="GetDeltaStepRotation", DisplayName="Rotation", Category="Step", meta=(EditCondition="bDeltaStepEnabled", EditConditionHides))
+	FRotator DeltaStepRotation = FRotator(0.f);
+
+	/** Amount of scale difference between one step and the next one */
+	UPROPERTY(EditAnywhere, BlueprintReadWrite, Setter="SetDeltaStepScale", Getter="GetDeltaStepScale", DisplayName="Scale", Category="Step", meta=(AllowPreserveRatio, Delta="0.0001", EditCondition="bDeltaStepEnabled", EditConditionHides))
+	FVector DeltaStepScale = FVector(0.f);
+
+	/** Use random range transforms for each clones */
+	UPROPERTY(EditAnywhere, BlueprintReadWrite, Setter="SetRangeEnabled", Getter="GetRangeEnabled", DisplayName="Enabled", Category="Range")
+	bool bRangeEnabled = false;
+
+	UPROPERTY(EditAnywhere, BlueprintReadWrite, Setter="SetRangeOffsetMin", Getter="GetRangeOffsetMin", DisplayName="OffsetMin", Category="Range", meta=(EditCondition="bRangeEnabled", EditConditionHides))
+	FVector RangeOffsetMin = FVector::ZeroVector;
+
+	UPROPERTY(EditAnywhere, BlueprintReadWrite, Setter="SetRangeOffsetMax", Getter="GetRangeOffsetMax", DisplayName="OffsetMax", Category="Range", meta=(EditCondition="bRangeEnabled", EditConditionHides))
+	FVector RangeOffsetMax = FVector::ZeroVector;
+
+	UPROPERTY(EditAnywhere, BlueprintReadWrite, Setter="SetRangeRotationMin", Getter="GetRangeRotationMin", DisplayName="RotationMin", Category="Range", meta=(EditCondition="bRangeEnabled", EditConditionHides))
+	FRotator RangeRotationMin = FRotator::ZeroRotator;
+
+	UPROPERTY(EditAnywhere, BlueprintReadWrite, Setter="SetRangeRotationMax", Getter="GetRangeRotationMax", DisplayName="RotationMax", Category="Range", meta=(EditCondition="bRangeEnabled", EditConditionHides))
+	FRotator RangeRotationMax = FRotator::ZeroRotator;
+
+	UPROPERTY(EditAnywhere, BlueprintReadWrite, Setter="SetRangeScaleUniform", Getter="GetRangeScaleUniform", DisplayName="ScaleUniformEnabled", Category="Range", meta=(EditCondition="bRangeEnabled", EditConditionHides))
+	bool bRangeScaleUniform = false;
+
+	UPROPERTY(EditAnywhere, BlueprintReadWrite, Setter="SetRangeScaleMin", Getter="GetRangeScaleMin", DisplayName="ScaleMin", Category="Range", meta=(AllowPreserveRatio, Delta="0.0001", ClampMin="0", EditCondition="bRangeEnabled && !bRangeScaleUniform", EditConditionHides))
+	FVector RangeScaleMin = FVector::OneVector;
+
+	UPROPERTY(EditAnywhere, BlueprintReadWrite, Setter="SetRangeScaleMax", Getter="GetRangeScaleMax", DisplayName="ScaleMax", Category="Range", meta=(AllowPreserveRatio, Delta="0.0001", ClampMin="0", EditCondition="bRangeEnabled && !bRangeScaleUniform", EditConditionHides))
+	FVector RangeScaleMax = FVector::OneVector;
+
+	UPROPERTY(EditAnywhere, BlueprintReadWrite, Setter="SetRangeScaleUniformMin", Getter="GetRangeScaleUniformMin", DisplayName="ScaleMin", Category="Range", meta=(Delta="0.0001", ClampMin="0", EditCondition="bRangeEnabled && bRangeScaleUniform", EditConditionHides))
+	float RangeScaleUniformMin = 1.f;
+
+	UPROPERTY(EditAnywhere, BlueprintReadWrite, Setter="SetRangeScaleUniformMax", Getter="GetRangeScaleUniformMax", DisplayName="ScaleMax", Category="Range", meta=(Delta="0.0001", ClampMin="0", EditCondition="bRangeEnabled && bRangeScaleUniform", EditConditionHides))
+	float RangeScaleUniformMax = 1.f;
+
+#if WITH_EDITORONLY_DATA
+	/** Toggle the sprite to visualize and click on this cloner */
+	UPROPERTY(EditAnywhere, BlueprintReadWrite, AdvancedDisplay, Category="Cloner", meta=(AllowPrivateAccess = "true"))
+	bool bVisualizerSpriteVisible = true;
+
+	/** Reduces the r.TSR.ShadingRejection.Flickering.Period from 3 (default) to 1 if enabled to avoid ghosting artifacts when moving */
+	UPROPERTY(EditAnywhere, BlueprintReadWrite, AdvancedDisplay, Category="Cloner", meta=(AllowPrivateAccess = "true"))
+	bool bReduceMotionGhosting = false;
+#endif
+
+private:
+#if WITH_EDITOR
+	/** Hide outline selection when this cloner only is selected in viewport */
+	void OnEditorSelectionChanged(UObject* InSelection);
+
+	TOptional<bool> UseSelectionOutline;
+#endif
+
+	/** Gets all layout names available */
+	UFUNCTION()
+	TArray<FString> GetClonerLayoutNames() const;
+
+	template<
+		typename InLayoutClass
+		UE_REQUIRES(TIsDerivedFrom<InLayoutClass, UAvaClonerLayoutBase>::Value)>
+	InLayoutClass* FindOrAddLayout()
+	{
+		return Cast<InLayoutClass>(FindOrAddLayout(InLayoutClass::StaticClass()));
+	}
+
+	/** Find or add a layout by its class */
+	UAvaClonerLayoutBase* FindOrAddLayout(TSubclassOf<UAvaClonerLayoutBase> InClass);
+
+	/** Find or add a layout by its name */
+	UAvaClonerLayoutBase* FindOrAddLayout(FName InLayoutName);
+
+	/** Initiate and perform operation */
+	void InitializeCloner();
+
+	/** Effectors that applies to this cloner, index match niagara data interface arrays */
+	UPROPERTY(Transient)
+	TArray<TWeakObjectPtr<AAvaEffectorActor>> Effectors;
+
+	/** Cloner niagara component */
+	UPROPERTY()
+	TObjectPtr<UAvaClonerComponent> ClonerComponent;
+
+	/** Previously used layout instances cached */
+	UPROPERTY()
+	TMap<FName, TObjectPtr<UAvaClonerLayoutBase>> LayoutInstances;
+
+	/**
+	* Below properties are deprecated and no longer in use,
+	* they will be migrated to new layout system on load,
+	* And they will be all removed in future version
+	*/
+	UPROPERTY(meta=(DeprecatedProperty, DeprecationMessage="Use LayoutName instead"))
+	EAvaClonerLayout Layout_DEPRECATED = EAvaClonerLayout::Grid;
+
+	UPROPERTY(meta=(DeprecatedProperty, DeprecationMessage="Use UAvaClonerGridLayout instead"))
+	FAvaClonerGridLayoutOptions GridOptions_DEPRECATED;
+
+	UPROPERTY(meta=(DeprecatedProperty, DeprecationMessage="Use UAvaClonerLineLayout instead"))
+	FAvaClonerLineLayoutOptions LineOptions_DEPRECATED;
+
+	UPROPERTY(meta=(DeprecatedProperty, DeprecationMessage="Use UAvaClonerCircleLayout instead"))
+	FAvaClonerCircleLayoutOptions CircleOptions_DEPRECATED;
+
+	UPROPERTY(meta=(DeprecatedProperty, DeprecationMessage="Use UAvaClonerCylinderLayout instead"))
+	FAvaClonerCylinderLayoutOptions CylinderOptions_DEPRECATED;
+
+	UPROPERTY(meta=(DeprecatedProperty, DeprecationMessage="Use UAvaClonerSphereLayout instead"))
+	FAvaClonerSphereLayoutOptions SphereOptions_DEPRECATED;
+
+	UPROPERTY(meta=(DeprecatedProperty, DeprecationMessage="Use UAvaClonerHoneycombLayout instead"))
+	FAvaClonerHoneycombLayoutOptions HoneycombOptions_DEPRECATED;
+
+	UPROPERTY(meta=(DeprecatedProperty, DeprecationMessage="Use UAvaClonerMeshLayout instead"))
+	FAvaClonerSampleMeshOptions SampleMeshOptions_DEPRECATED;
+
+	UPROPERTY(meta=(DeprecatedProperty, DeprecationMessage="Use UAvaClonerSplineLayout instead"))
+	FAvaClonerSampleSplineOptions SampleSplineOptions_DEPRECATED;
+
+	/**
+	 * Used to confirm migration of deprecated properties above to new layout system
+	 * Newly created cloner will not migrate anything and skip this task
+	 */
+	UPROPERTY()
+	bool bDeprecatedPropertiesMigrated = false;
+
+	float TreeUpdateDeltaTime = 0.f;
+
+	bool bNeedsRefresh = false;
+
+	bool bClonerInitialized = false;
+
+#if WITH_EDITOR
+	/** Used for PECP */
+	static const TAvaPropertyChangeDispatcher<AAvaClonerActor> PropertyChangeDispatcher;
+#endif
+};
diff --git a/Engine/Plugins/Experimental/Avalanche/Source/AvalancheEffectors/Public/Cloner/AvaClonerComponent.h b/Engine/Plugins/Experimental/Avalanche/Source/AvalancheEffectors/Public/Cloner/AvaClonerComponent.h
new file mode 100644
--- /dev/null
+++ b/Engine/Plugins/Experimental/Avalanche/Source/AvalancheEffectors/Public/Cloner/AvaClonerComponent.h
@@ -0,0 +1,124 @@
+// Copyright Epic Games, Inc. All Rights Reserved.
+
+#pragma once
+
+#include "AvaClonerEffectorShared.h"
+#include "Layouts/AvaClonerLayoutBase.h"
+#include "NiagaraComponent.h"
+#include "AvaClonerComponent.generated.h"
+
+class UAvaClonerLayoutBase;
+class UMaterialInterface;
+
+UCLASS(Within=AvaClonerActor, HideCategories=(Niagara))
+class AVALANCHEEFFECTORS_API UAvaClonerComponent : public UNiagaraComponent
+{
+	GENERATED_BODY()
+
+public:
+	/** Called when meshes have been updated */
+	DECLARE_MULTICAST_DELEGATE(FAvaClonerMeshUpdated)
+	static inline FAvaClonerMeshUpdated OnClonerMeshUpdated;
+
+	/** Called when cloner layout system is loaded */
+	DECLARE_MULTICAST_DELEGATE_OneParam(FAvaClonerSystemLoaded, UAvaClonerLayoutBase* /** InLayout */)
+	static inline FAvaClonerSystemLoaded OnClonerSystemLoaded;
+
+	UAvaClonerComponent();
+
+	//~ Begin UObject
+	virtual void PostLoad() override;
+	//~ End UObject
+
+	/**
+	 * Triggers an update of the attachment tree to detect updated items
+	 * If reset is true, clears the attachment tree and rebuilds it otherwise diff update
+	 */
+	void UpdateClonerAttachmentTree(bool bInReset = false);
+
+	/** Called to trigger an update of cloner rendering state if tree */
+	void UpdateClonerRenderState();
+
+	/** Sets the layout to use for this cloner simulation */
+	bool SetClonerActiveLayout(UAvaClonerLayoutBase* InLayout);
+
+	UAvaClonerLayoutBase* GetClonerActiveLayout() const
+	{
+		return ActiveLayout;
+	}
+
+	/** Forces a refresh of the active system parameters in niagara store */
+	void RefreshUserParameters() const;
+
+protected:
+	void OnRenderStateDirty(UActorComponent& InComponent);
+
+	void UpdateAttachmentTree();
+	void UpdateActorAttachment(AActor* InActor, AActor* InParent);
+
+	void BindActorDelegates(AActor* InActor);
+	void UnbindActorDelegates(AActor* InActor) const;
+
+	void OnTransformUpdated(USceneComponent* InUpdatedComponent, EUpdateTransformFlags InUpdateTransformFlags, ETeleportType InTeleport);
+
+	UFUNCTION()
+	void OnActorDestroyed(AActor* InDestroyedActor);
+
+#if WITH_EDITOR
+	void OnActorPropertyChanged(UObject* InObject, FPropertyChangedEvent& InPropertyChangedEvent);
+#endif
+
+	/** Called when a material has changed in any cloned actors */
+	void OnMaterialChanged(UObject* InObject);
+
+	/** Called when a mesh has been updated and cloner needs to update instances */
+	void OnMeshChanged(UStaticMeshComponent*, AActor* InActor);
+
+	/** Get cloner direct children in the correct order */
+	void GetOrderedRootActors(TArray<AActor*>& OutActors) const;
+
+	/** Get the root cloner actor (direct child) for a specific actor */
+	AActor* GetRootActor(AActor* InActor) const;
+
+	/** Resets the root baked static mesh to be regenerated later */
+	void InvalidateBakedStaticMesh(AActor* InActor);
+
+	/** Runs async update to rebuild dirty meshes */
+	void UpdateDirtyMeshesAsync();
+	void OnDirtyMeshesUpdated();
+
+	/** Merges all primitive components from an actor to dynamic mesh, does not recurse */
+	void UpdateActorBakedDynamicMesh(AActor* InActor);
+
+	/** Merges all baked dynamic meshes from children and self into one static mesh to use as niagara mesh input */
+	void UpdateRootActorBakedStaticMesh(AActor* InRootActor);
+
+	/** Gets all attachment items based on an actor, will recurse */
+	void GetActorAttachmentItems(AActor* InActor, TArray<FAvaClonerAttachmentItem*>& OutAttachmentItems);
+
+	/** Update niagara asset static meshes */
+	void UpdateClonerMeshes();
+
+	/** Checks that all root static meshes are valid */
+	bool IsAllMergedMeshesValid() const;
+
+	/** Only materials transient or part of the content folder can be dirtied, engine or plugins cannot */
+	bool IsMaterialDirtyable(const UMaterialInterface* InMaterial) const;
+
+	/** Attachment tree view */
+	UPROPERTY(Transient, NonTransactional)
+	FAvaClonerAttachmentTree ClonerTree;
+
+	/** Current active layout */
+	UPROPERTY(Transient, NonTransactional)
+	TObjectPtr<UAvaClonerLayoutBase> ActiveLayout;
+
+	/** Attachment items that are dirty and need an update */
+	TSet<FAvaClonerAttachmentItem*> DirtyItemAttachments;
+
+	/** Asset meshes needs update */
+	bool bClonerMeshesDirty = true;
+
+	/** State of the baked dynamic and static mesh creation */
+	bool bClonerMeshesUpdating = false;
+};
\ No newline at end of file
diff --git a/Engine/Plugins/Experimental/Avalanche/Source/AvalancheEffectors/Public/Cloner/Layouts/AvaClonerCircleLayout.h b/Engine/Plugins/Experimental/Avalanche/Source/AvalancheEffectors/Public/Cloner/Layouts/AvaClonerCircleLayout.h
new file mode 100644
--- /dev/null
+++ b/Engine/Plugins/Experimental/Avalanche/Source/AvalancheEffectors/Public/Cloner/Layouts/AvaClonerCircleLayout.h
@@ -0,0 +1,135 @@
+// Copyright Epic Games, Inc. All Rights Reserved.
+
+#pragma once
+
+#include "AvaClonerEffectorShared.h"
+#include "AvaClonerLayoutBase.h"
+#include "AvaPropertyChangeDispatcher.h"
+#include "AvaClonerCircleLayout.generated.h"
+
+UCLASS(BlueprintType)
+class AVALANCHEEFFECTORS_API UAvaClonerCircleLayout : public UAvaClonerLayoutBase
+{
+	GENERATED_BODY()
+
+	friend class FAvaClonerActorVisualizer;
+	
+public:
+	UAvaClonerCircleLayout()
+		: UAvaClonerLayoutBase(
+			TEXT("Circle")
+			, TEXT("/Script/Niagara.NiagaraSystem'/Avalanche/ClonerResources/Systems/NS_ClonerCircle.NS_ClonerCircle'")
+		)
+	{}
+	
+	UFUNCTION(BlueprintCallable, Category="Cloner|Layout|Circle")
+	void SetCount(int32 InCount);
+	
+	UFUNCTION(BlueprintPure, Category="Cloner|Layout|Circle")
+	int32 GetCount() const
+	{
+		return Count;
+	}
+
+	UFUNCTION(BlueprintCallable, Category="Cloner|Layout|Circle")
+	void SetRadius(float InRadius);
+	
+	UFUNCTION(BlueprintPure, Category="Cloner|Layout|Circle")
+	float GetRadius() const
+	{
+		return Radius;
+	}
+
+	UFUNCTION(BlueprintCallable, Category="Cloner|Layout|Circle")
+	void SetAngleStart(float InAngleStart);
+
+	UFUNCTION(BlueprintPure, Category="Cloner|Layout|Circle")
+	float GetAngleStart() const
+	{
+		return AngleStart;
+	}
+
+	UFUNCTION(BlueprintCallable, Category="Cloner|Layout|Circle")
+	void SetAngleRatio(float InAngleRatio);
+	
+	UFUNCTION(BlueprintPure, Category="Cloner|Layout|Circle")
+	float GetAngleRatio() const
+	{
+		return AngleRatio;
+	}
+
+	UFUNCTION(BlueprintCallable, Category="Cloner|Layout|Circle")
+	void SetOrientMesh(bool bInOrientMesh);
+	
+	UFUNCTION(BlueprintPure, Category="Cloner|Layout|Circle")
+	bool GetOrientMesh() const
+	{
+		return bOrientMesh;
+	}
+
+	UFUNCTION(BlueprintCallable, Category="Cloner|Layout|Circle")
+	void SetPlane(EAvaClonerPlane InPlane);
+	
+	UFUNCTION(BlueprintPure, Category="Cloner|Layout|Circle")
+	EAvaClonerPlane GetPlane() const
+	{
+		return Plane;
+	}
+
+	UFUNCTION(BlueprintCallable, Category="Cloner|Layout|Circle")
+	void SetRotation(const FRotator& InRotation);
+
+	UFUNCTION(BlueprintPure, Category="Cloner|Layout|Circle")
+	const FRotator& GetRotation() const
+	{
+		return Rotation;
+	}
+	
+	UFUNCTION(BlueprintCallable, Category="Cloner|Layout|Circle")
+	void SetScale(const FVector& InScale);
+
+	UFUNCTION(BlueprintPure, Category="Cloner|Layout|Circle")
+	const FVector& GetScale() const
+	{
+		return Scale;
+	}
+
+	//~ Begin UObject
+#if WITH_EDITOR
+	virtual void PostEditChangeProperty(FPropertyChangedEvent& InPropertyChangedEvent) override;
+#endif
+	//~ End UObject
+	
+protected:
+	virtual void OnLayoutParametersChanged(UAvaClonerComponent* InComponent) override;
+	
+	UPROPERTY(EditInstanceOnly, BlueprintReadWrite, Setter="SetCount", Getter="GetCount", Category="Layout")
+	int32 Count = 3 * 3;
+
+	UPROPERTY(EditInstanceOnly, BlueprintReadWrite, Setter="SetRadius", Getter="GetRadius", Category="Layout")
+	float Radius = 200.f;
+
+	UPROPERTY(EditInstanceOnly, BlueprintReadWrite, Setter="SetAngleStart", Getter="GetAngleStart", Category="Layout")
+	float AngleStart = 0.f;
+
+	UPROPERTY(EditInstanceOnly, BlueprintReadWrite, Setter="SetAngleRatio", Getter="GetAngleRatio", Category="Layout")
+	float AngleRatio = 1.f;
+
+	UPROPERTY(EditInstanceOnly, BlueprintReadWrite, Setter="SetOrientMesh", Getter="GetOrientMesh", Category="Layout")
+	bool bOrientMesh = true;
+
+	UPROPERTY(EditInstanceOnly, BlueprintReadWrite, Setter="SetPlane", Getter="GetPlane", Category="Layout")
+	EAvaClonerPlane Plane = EAvaClonerPlane::XY;
+
+	UPROPERTY(EditInstanceOnly, BlueprintReadWrite, Setter="SetRotation", Getter="GetRotation", Category="Layout", meta=(EditCondition="Plane == EAvaClonerPlane::Custom", EditConditionHides))
+	FRotator Rotation = FRotator(0.f, 0.f, 0.f);
+	
+	UPROPERTY(EditInstanceOnly, BlueprintReadWrite, Setter="SetScale", Getter="GetScale", Category="Layout", meta=(ClampMin="0", AllowPreserveRatio, Delta="0.01"))
+	FVector Scale = FVector(1.f, 1.f, 1.f);
+
+private:
+#if WITH_EDITOR
+	/** Used for PECP */
+	static const TAvaPropertyChangeDispatcher<UAvaClonerCircleLayout> PropertyChangeDispatcher;
+#endif
+};
diff --git a/Engine/Plugins/Experimental/Avalanche/Source/AvalancheEffectors/Public/Cloner/Layouts/AvaClonerCylinderLayout.h b/Engine/Plugins/Experimental/Avalanche/Source/AvalancheEffectors/Public/Cloner/Layouts/AvaClonerCylinderLayout.h
new file mode 100644
--- /dev/null
+++ b/Engine/Plugins/Experimental/Avalanche/Source/AvalancheEffectors/Public/Cloner/Layouts/AvaClonerCylinderLayout.h
@@ -0,0 +1,159 @@
+// Copyright Epic Games, Inc. All Rights Reserved.
+
+#pragma once
+
+#include "AvaClonerEffectorShared.h"
+#include "AvaClonerLayoutBase.h"
+#include "AvaPropertyChangeDispatcher.h"
+#include "AvaClonerCylinderLayout.generated.h"
+
+UCLASS(BlueprintType)
+class AVALANCHEEFFECTORS_API UAvaClonerCylinderLayout : public UAvaClonerLayoutBase
+{
+	GENERATED_BODY()
+
+	friend class FAvaClonerActorVisualizer;
+	
+public:
+	UAvaClonerCylinderLayout()
+		: UAvaClonerLayoutBase(
+			TEXT("Cylinder")
+			, TEXT("/Script/Niagara.NiagaraSystem'/Avalanche/ClonerResources/Systems/NS_ClonerCylinder.NS_ClonerCylinder'")
+		)
+	{}
+	
+	UFUNCTION(BlueprintCallable, Category="Cloner|Layout|Cylinder")
+	void SetBaseCount(int32 InBaseCount);
+
+	UFUNCTION(BlueprintPure, Category="Cloner|Layout|Cylinder")
+	int32 GetBaseCount() const
+	{
+		return BaseCount;
+	}
+
+	UFUNCTION(BlueprintCallable, Category="Cloner|Layout|Cylinder")
+	void SetHeightCount(int32 InHeightCount);
+	
+	UFUNCTION(BlueprintPure, Category="Cloner|Layout|Cylinder")
+	int32 GetHeightCount() const
+	{
+		return HeightCount;
+	}
+
+	UFUNCTION(BlueprintCallable, Category="Cloner|Layout|Cylinder")
+	void SetHeight(float InHeight);
+	
+	UFUNCTION(BlueprintPure, Category="Cloner|Layout|Cylinder")
+	float GetHeight() const
+	{
+		return Height;
+	}
+
+	UFUNCTION(BlueprintCallable, Category="Cloner|Layout|Cylinder")
+	void SetRadius(float InRadius);
+	
+	UFUNCTION(BlueprintPure, Category="Cloner|Layout|Cylinder")
+	float GetRadius() const
+	{
+		return Radius;
+	}
+
+	UFUNCTION(BlueprintCallable, Category="Cloner|Layout|Cylinder")
+	void SetAngleStart(float InAngleStart);
+
+	UFUNCTION(BlueprintPure, Category="Cloner|Layout|Cylinder")
+	float GetAngleStart() const
+	{
+		return AngleStart;
+	}
+
+	UFUNCTION(BlueprintCallable, Category="Cloner|Layout|Cylinder")
+	void SetAngleRatio(float InAngleRatio);
+
+	UFUNCTION(BlueprintPure, Category="Cloner|Layout|Cylinder")
+	float GetAngleRatio() const
+	{
+		return AngleRatio;
+	}
+
+	UFUNCTION(BlueprintCallable, Category="Cloner|Layout|Cylinder")
+	void SetOrientMesh(bool bInOrientMesh);
+	
+	UFUNCTION(BlueprintPure, Category="Cloner|Layout|Cylinder")
+	bool GetOrientMesh() const
+	{
+		return bOrientMesh;
+	}
+
+	UFUNCTION(BlueprintCallable, Category="Cloner|Layout|Cylinder")
+	void SetPlane(EAvaClonerPlane InPlane);
+
+	UFUNCTION(BlueprintPure, Category="Cloner|Layout|Cylinder")
+	EAvaClonerPlane GetPlane() const
+	{
+		return Plane;
+	}
+
+	UFUNCTION(BlueprintCallable, Category="Cloner|Layout|Cylinder")
+	void SetRotation(const FRotator& InRotation);
+	
+	UFUNCTION(BlueprintPure, Category="Cloner|Layout|Cylinder")
+	const FRotator& GetRotation() const
+	{
+		return Rotation;
+	}
+
+	UFUNCTION(BlueprintCallable, Category="Cloner|Layout|Cylinder")
+	void SetScale(const FVector& InScale);
+
+	UFUNCTION(BlueprintPure, Category="Cloner|Layout|Cylinder")
+	const FVector& GetScale() const
+	{
+		return Scale;
+	}
+
+	//~ Begin UObject
+#if WITH_EDITOR
+	virtual void PostEditChangeProperty(FPropertyChangedEvent& InPropertyChangedEvent) override;
+#endif
+	//~ End UObject
+
+protected:
+	virtual void OnLayoutParametersChanged(UAvaClonerComponent* InComponent) override;
+	
+	UPROPERTY(EditInstanceOnly, BlueprintReadWrite, Setter="SetBaseCount", Getter="GetBaseCount", Category="Layout")
+	int32 BaseCount = 3 * 3;
+
+	UPROPERTY(EditInstanceOnly, BlueprintReadWrite, Setter="SetHeightCount", Getter="GetHeightCount", Category="Layout")
+	int32 HeightCount = 3;
+
+	UPROPERTY(EditInstanceOnly, BlueprintReadWrite, Setter="SetHeight", Getter="GetHeight", Category="Layout")
+	float Height = 400.f;
+	
+	UPROPERTY(EditInstanceOnly, BlueprintReadWrite, Setter="SetRadius", Getter="GetRadius", Category="Layout")
+	float Radius = 200.f;
+
+	UPROPERTY(EditInstanceOnly, BlueprintReadWrite, Setter="SetAngleStart", Getter="GetAngleStart", Category="Layout")
+	float AngleStart = 0.f;
+	
+	UPROPERTY(EditInstanceOnly, BlueprintReadWrite, Setter="SetAngleRatio", Getter="GetAngleRatio", Category="Layout")
+	float AngleRatio = 1.f;
+	
+	UPROPERTY(EditInstanceOnly, BlueprintReadWrite, Setter="SetOrientMesh", Getter="GetOrientMesh", Category="Layout")
+	bool bOrientMesh = true;
+	
+	UPROPERTY(EditInstanceOnly, BlueprintReadWrite, Setter="SetPlane", Getter="GetPlane", Category="Layout")
+	EAvaClonerPlane Plane = EAvaClonerPlane::XY;
+	
+	UPROPERTY(EditInstanceOnly, BlueprintReadWrite, Setter="SetRotation", Getter="GetRotation", Category="Layout", meta=(EditCondition="Plane == EAvaClonerPlane::Custom", EditConditionHides))
+	FRotator Rotation = FRotator(0.f, 0.f, 0.f);
+	
+	UPROPERTY(EditInstanceOnly, BlueprintReadWrite, Setter="SetScale", Getter="GetScale", Category="Layout", meta=(ClampMin="0", AllowPreserveRatio, Delta="0.01"))
+	FVector Scale = FVector(1.f, 1.f, 1.f);
+
+private:
+#if WITH_EDITOR
+	/** Used for PECP */
+	static const TAvaPropertyChangeDispatcher<UAvaClonerCylinderLayout> PropertyChangeDispatcher;
+#endif
+};
\ No newline at end of file
diff --git a/Engine/Plugins/Experimental/Avalanche/Source/AvalancheEffectors/Public/Cloner/Layouts/AvaClonerGridLayout.h b/Engine/Plugins/Experimental/Avalanche/Source/AvalancheEffectors/Public/Cloner/Layouts/AvaClonerGridLayout.h
new file mode 100644
--- /dev/null
+++ b/Engine/Plugins/Experimental/Avalanche/Source/AvalancheEffectors/Public/Cloner/Layouts/AvaClonerGridLayout.h
@@ -0,0 +1,171 @@
+// Copyright Epic Games, Inc. All Rights Reserved.
+
+#pragma once
+
+#include "AvaClonerEffectorShared.h"
+#include "AvaClonerLayoutBase.h"
+#include "AvaPropertyChangeDispatcher.h"
+#include "AvaClonerGridLayout.generated.h"
+
+UCLASS(BlueprintType)
+class AVALANCHEEFFECTORS_API UAvaClonerGridLayout : public UAvaClonerLayoutBase
+{
+	GENERATED_BODY()
+
+	friend class FAvaClonerActorVisualizer;
+	
+public:
+	UAvaClonerGridLayout()
+		: UAvaClonerLayoutBase(
+			TEXT("Grid")
+			, TEXT("/Script/Niagara.NiagaraSystem'/Avalanche/ClonerResources/Systems/NS_ClonerGrid.NS_ClonerGrid'")
+		)
+	{}
+
+	UFUNCTION(BlueprintCallable, Category="Cloner|Layout|Grid")
+	void SetCountX(int32 InCountX);
+	
+	UFUNCTION(BlueprintPure, Category="Cloner|Layout|Grid")
+	int32 GetCountX() const
+	{
+		return CountX;
+	}
+
+	UFUNCTION(BlueprintCallable, Category="Cloner|Layout|Grid")
+	void SetCountY(int32 InCountY);
+
+	UFUNCTION(BlueprintPure, Category="Cloner|Layout|Grid")
+	int32 GetCountY() const
+	{
+		return CountY;
+	}
+	
+	UFUNCTION(BlueprintCallable, Category="Cloner|Layout|Grid")
+	void SetCountZ(int32 InCountZ);
+	
+	UFUNCTION(BlueprintPure, Category="Cloner|Layout|Grid")
+	int32 GetCountZ() const
+	{
+		return CountZ;
+	}
+
+	UFUNCTION(BlueprintCallable, Category="Cloner|Layout|Grid")
+	void SetSpacingX(float InSpacingX);
+	
+	UFUNCTION(BlueprintPure, Category="Cloner|Layout|Grid")
+	float GetSpacingX() const
+	{
+		return SpacingX;
+	}
+
+	UFUNCTION(BlueprintCallable, Category="Cloner|Layout|Grid")
+	void SetSpacingY(float InSpacingY);
+	
+	UFUNCTION(BlueprintPure, Category="Cloner|Layout|Grid")
+	float GetSpacingY() const
+	{
+		return SpacingY;
+	}
+
+	UFUNCTION(BlueprintCallable, Category="Cloner|Layout|Grid")
+	void SetSpacingZ(float InSpacingZ);
+	
+	UFUNCTION(BlueprintPure, Category="Cloner|Layout|Grid")
+	float GetSpacingZ() const
+	{
+		return SpacingZ;
+	}
+
+	UFUNCTION(BlueprintCallable, Category="Cloner|Layout|Grid")
+	void SetConstraint(EAvaClonerGridConstraint InConstraint);
+	
+	UFUNCTION(BlueprintPure, Category="Cloner|Layout|Grid")
+	EAvaClonerGridConstraint GetConstraint() const
+	{
+		return Constraint;
+	}
+
+	UFUNCTION(BlueprintCallable, Category="Cloner|Layout|Grid")
+	void SetInvertConstraint(bool bInInvertConstraint);
+	
+	UFUNCTION(BlueprintPure, Category="Cloner|Layout|Grid")
+	bool GetInvertConstraint() const
+	{
+		return bInvertConstraint;
+	}
+
+	UFUNCTION(BlueprintCallable, Category="Cloner|Layout|Grid")
+	void SetSphereConstraint(const FAvaClonerGridConstraintSphere& InConstraint);
+	
+	UFUNCTION(BlueprintPure, Category="Cloner|Layout|Grid")
+	const FAvaClonerGridConstraintSphere& GetSphereConstraint() const
+	{
+		return SphereConstraint;
+	}
+
+	UFUNCTION(BlueprintCallable, Category="Cloner|Layout|Grid")
+	void SetCylinderConstraint(const FAvaClonerGridConstraintCylinder& InConstraint);
+
+	UFUNCTION(BlueprintPure, Category="Cloner|Layout|Grid")
+	const FAvaClonerGridConstraintCylinder& GetCylinderConstraint() const
+	{
+		return CylinderConstraint;
+	}
+
+	UFUNCTION(BlueprintCallable, Category="Cloner|Layout|Grid")
+	void SetTextureConstraint(const FAvaClonerGridConstraintTexture& InConstraint);
+	
+	UFUNCTION(BlueprintPure, Category="Cloner|Layout|Grid")
+	const FAvaClonerGridConstraintTexture& GetTextureConstraint() const
+	{
+		return TextureConstraint;
+	}
+	
+	//~ Begin UObject
+#if WITH_EDITOR
+	virtual void PostEditChangeProperty(FPropertyChangedEvent& InPropertyChangedEvent) override;
+#endif
+	//~ End UObject
+
+protected:
+	virtual void OnLayoutParametersChanged(UAvaClonerComponent* InComponent) override;
+	
+	UPROPERTY(EditInstanceOnly, BlueprintReadWrite, Setter="SetCountX", Getter="GetCountX", Category="Layout")
+	int32 CountX = 3;
+
+	UPROPERTY(EditInstanceOnly, BlueprintReadWrite, Setter="SetCountY", Getter="GetCountY", Category="Layout")
+	int32 CountY = 3;
+
+	UPROPERTY(EditInstanceOnly, BlueprintReadWrite, Setter="SetCountZ", Getter="GetCountZ", Category="Layout")
+	int32 CountZ = 3;
+
+	UPROPERTY(EditInstanceOnly, BlueprintReadWrite, Setter="SetSpacingX", Getter="GetSpacingX", Category="Layout")
+	float SpacingX = 105.f;
+
+	UPROPERTY(EditInstanceOnly, BlueprintReadWrite, Setter="SetSpacingY", Getter="GetSpacingY", Category="Layout")
+	float SpacingY = 105.f;
+
+	UPROPERTY(EditInstanceOnly, BlueprintReadWrite, Setter="SetSpacingZ", Getter="GetSpacingZ", Category="Layout")
+	float SpacingZ = 105.f;
+
+	UPROPERTY(EditInstanceOnly, BlueprintReadWrite, Setter="SetConstraint", Getter="GetConstraint", Category="Layout")
+	EAvaClonerGridConstraint Constraint = EAvaClonerGridConstraint::None;
+
+	UPROPERTY(EditInstanceOnly, BlueprintReadWrite, Setter="SetInvertConstraint", Getter="GetInvertConstraint", Category="Layout", meta=(EditCondition="Constraint != EAvaClonerGridConstraint::None", EditConditionHides))
+	bool bInvertConstraint = false;
+
+	UPROPERTY(EditInstanceOnly, BlueprintReadWrite, Setter="SetSphereConstraint", Getter="GetSphereConstraint", Category="Layout", meta=(EditCondition="Constraint == EAvaClonerGridConstraint::Sphere", EditConditionHides))
+	FAvaClonerGridConstraintSphere SphereConstraint;
+
+	UPROPERTY(EditInstanceOnly, BlueprintReadWrite, Setter="SetCylinderConstraint", Getter="GetCylinderConstraint", Category="Layout", meta=(EditCondition="Constraint == EAvaClonerGridConstraint::Cylinder", EditConditionHides))
+	FAvaClonerGridConstraintCylinder CylinderConstraint;
+
+	UPROPERTY(EditInstanceOnly, BlueprintReadWrite, Setter="SetTextureConstraint", Getter="GetTextureConstraint", Category="Layout", meta=(EditCondition="Constraint == EAvaClonerGridConstraint::Texture", EditConditionHides))
+	FAvaClonerGridConstraintTexture TextureConstraint;
+
+private:
+#if WITH_EDITOR
+	/** Used for PECP */
+	static const TAvaPropertyChangeDispatcher<UAvaClonerGridLayout> PropertyChangeDispatcher;
+#endif
+};
\ No newline at end of file
diff --git a/Engine/Plugins/Experimental/Avalanche/Source/AvalancheEffectors/Public/Cloner/Layouts/AvaClonerHoneycombLayout.h b/Engine/Plugins/Experimental/Avalanche/Source/AvalancheEffectors/Public/Cloner/Layouts/AvaClonerHoneycombLayout.h
new file mode 100644
--- /dev/null
+++ b/Engine/Plugins/Experimental/Avalanche/Source/AvalancheEffectors/Public/Cloner/Layouts/AvaClonerHoneycombLayout.h
@@ -0,0 +1,123 @@
+// Copyright Epic Games, Inc. All Rights Reserved.
+
+#pragma once
+
+#include "AvaClonerEffectorShared.h"
+#include "AvaClonerLayoutBase.h"
+#include "AvaPropertyChangeDispatcher.h"
+#include "AvaClonerHoneycombLayout.generated.h"
+
+UCLASS(BlueprintType)
+class AVALANCHEEFFECTORS_API UAvaClonerHoneycombLayout : public UAvaClonerLayoutBase
+{
+	GENERATED_BODY()
+
+	friend class FAvaClonerActorVisualizer;
+	
+public:
+	UAvaClonerHoneycombLayout()
+		: UAvaClonerLayoutBase(
+			TEXT("Honeycomb")
+			, TEXT("/Script/Niagara.NiagaraSystem'/Avalanche/ClonerResources/Systems/NS_ClonerHoneycomb.NS_ClonerHoneycomb'")
+		)
+	{}
+	
+	UFUNCTION(BlueprintCallable, Category="Cloner|Layout|Honeycomb")
+	void SetPlane(EAvaClonerPlane InPlane);
+
+	UFUNCTION(BlueprintPure, Category="Cloner|Layout|Honeycomb")
+	EAvaClonerPlane GetPlane() const
+	{
+		return Plane;
+	}
+
+	UFUNCTION(BlueprintCallable, Category="Cloner|Layout|Honeycomb")
+	void SetWidthCount(int32 InWidthCount);
+	
+	UFUNCTION(BlueprintPure, Category="Cloner|Layout|Honeycomb")
+	int32 GetWidthCount() const
+	{
+		return WidthCount;
+	}
+
+	UFUNCTION(BlueprintCallable, Category="Cloner|Layout|Honeycomb")
+	void SetHeightCount(int32 InHeightCount);
+	
+	UFUNCTION(BlueprintPure, Category="Cloner|Layout|Honeycomb")
+	int32 GetHeightCount() const
+	{
+		return HeightCount;
+	}
+
+	UFUNCTION(BlueprintCallable, Category="Cloner|Layout|Honeycomb")
+	void SetWidthOffset(float InWidthOffset);
+	
+	UFUNCTION(BlueprintPure, Category="Cloner|Layout|Honeycomb")
+	float GetWidthOffset() const
+	{
+		return WidthOffset;
+	}
+
+	UFUNCTION(BlueprintCallable, Category="Cloner|Layout|Honeycomb")
+	void SetHeightOffset(float InHeightOffset);
+	
+	UFUNCTION(BlueprintPure, Category="Cloner|Layout|Honeycomb")
+	float GetHeightOffset() const
+	{
+		return HeightOffset;
+	}
+
+	UFUNCTION(BlueprintCallable, Category="Cloner|Layout|Honeycomb")
+	void SetHeightSpacing(float InHeightSpacing);
+	
+	UFUNCTION(BlueprintPure, Category="Cloner|Layout|Honeycomb")
+	float GetHeightSpacing() const
+	{
+		return HeightSpacing;
+	}
+
+	UFUNCTION(BlueprintCallable, Category="Cloner|Layout|Honeycomb")
+	void SetWidthSpacing(float InWidthSpacing);
+
+	UFUNCTION(BlueprintPure, Category="Cloner|Layout|Honeycomb")
+	float GetWidthSpacing() const
+	{
+		return WidthSpacing;
+	}
+
+	//~ Begin UObject
+#if WITH_EDITOR
+	virtual void PostEditChangeProperty(FPropertyChangedEvent& InPropertyChangedEvent) override;
+#endif
+	//~ End UObject
+
+protected:
+	virtual void OnLayoutParametersChanged(UAvaClonerComponent* InComponent) override;
+	
+	UPROPERTY(EditInstanceOnly, BlueprintReadWrite, Setter="SetPlane", Getter="GetPlane", Category="Layout")
+	EAvaClonerPlane Plane = EAvaClonerPlane::XY;
+
+	UPROPERTY(EditInstanceOnly, BlueprintReadWrite, Setter="SetWidthCount", Getter="GetWidthCount", Category="Layout")
+	int32 WidthCount = 3;
+
+	UPROPERTY(EditInstanceOnly, BlueprintReadWrite, Setter="SetHeightCount", Getter="GetHeightCount", Category="Layout")
+	int32 HeightCount = 3;
+	
+	UPROPERTY(EditInstanceOnly, BlueprintReadWrite, Setter="SetWidthOffset", Getter="GetWidthOffset", Category="Layout")
+	float WidthOffset = 0.f;
+
+	UPROPERTY(EditInstanceOnly, BlueprintReadWrite, Setter="SetHeightOffset", Getter="GetHeightOffset", Category="Layout")
+	float HeightOffset = 0.5f;
+
+	UPROPERTY(EditInstanceOnly, BlueprintReadWrite, Setter="SetWidthSpacing", Getter="GetWidthSpacing", Category="Layout")
+	float WidthSpacing = 105.f;
+	
+	UPROPERTY(EditInstanceOnly, BlueprintReadWrite, Setter="SetHeightSpacing", Getter="GetHeightSpacing", Category="Layout")
+	float HeightSpacing = 105.f;
+
+private:
+#if WITH_EDITOR
+	/** Used for PECP */
+	static const TAvaPropertyChangeDispatcher<UAvaClonerHoneycombLayout> PropertyChangeDispatcher;
+#endif
+};
\ No newline at end of file
diff --git a/Engine/Plugins/Experimental/Avalanche/Source/AvalancheEffectors/Public/Cloner/Layouts/AvaClonerLayoutBase.h b/Engine/Plugins/Experimental/Avalanche/Source/AvalancheEffectors/Public/Cloner/Layouts/AvaClonerLayoutBase.h
new file mode 100644
--- /dev/null
+++ b/Engine/Plugins/Experimental/Avalanche/Source/AvalancheEffectors/Public/Cloner/Layouts/AvaClonerLayoutBase.h
@@ -0,0 +1,142 @@
+// Copyright Epic Games, Inc. All Rights Reserved.
+
+#pragma once
+
+#include "AvaClonerEffectorShared.h"
+#include "UObject/Object.h"
+#include "AvaClonerLayoutBase.generated.h"
+
+class AAvaClonerActor;
+class UAvaClonerComponent;
+class UNiagaraMeshRendererProperties;
+class UNiagaraSystem;
+
+/**
+ * Base class for layouts available in the cloner actor
+ * Steps to add a new layout :
+ * 1. Create a new system that extends from NS_ClonerBase and expose all the parent parameters (examples can be found in Content)
+ * 2. Extend this layout class and give it a unique name with the newly created system path
+ * 3. Expose all new system specific parameters in the layout extended class and update them when required
+ * Your new layout is ready and will be available in the cloner in the layout dropdown
+ */
+UCLASS(Abstract, BlueprintType, Within=AvaClonerActor, AutoExpandCategories=("Layout"))
+class AVALANCHEEFFECTORS_API UAvaClonerLayoutBase : public UObject
+{
+	GENERATED_BODY()
+
+public:
+	UAvaClonerLayoutBase()
+		: UAvaClonerLayoutBase(NAME_None, FString())
+	{}
+	
+	UAvaClonerLayoutBase(const FName& InLayoutName, const FString& InLayoutAssetPath)
+		: LayoutName(InLayoutName)
+		, LayoutAssetPath(InLayoutAssetPath)
+	{}
+
+	UFUNCTION(BlueprintPure, Category="Cloner|Layout")
+	FName GetLayoutName() const
+	{
+		return LayoutName;
+	}
+
+	FString GetLayoutAssetPath() const
+	{
+		return LayoutAssetPath;
+	}
+
+	UNiagaraSystem* GetSystem() const
+	{
+		return NiagaraSystem;
+	}
+	
+	UNiagaraMeshRendererProperties* GetMeshRenderer() const
+	{
+		return MeshRenderer;
+	}
+
+	const FAvaClonerEffectorDataInterfaces& GetDataInterfaces() const
+	{
+		return DataInterfaces;
+	}
+
+	/** Get the cloner actor using this layout */
+	UFUNCTION(BlueprintPure, Category="Cloner|Layout")
+	AAvaClonerActor* GetClonerActor() const;
+
+	/** Get the cloner component using this layout */
+	UAvaClonerComponent* GetClonerComponent() const;
+
+	/** Request refresh layout next tick */
+	void UpdateLayoutParameters(bool bInUpdateCloner = true, bool bInImmediate = false);
+
+	/** Updates the cloner, forcing a reset of the system */
+	void RequestClonerUpdate(bool bInImmediate = false) const;
+
+	/* Checks if the niagara system asset is valid and usable with the cloner */
+	bool IsLayoutValid() const;
+
+	/** Is this layout system cached and ready to be used */
+	UFUNCTION(BlueprintPure, Category="Cloner|Layout")
+	bool IsLayoutLoaded() const;
+
+	/** Load this layout system if not already loaded */
+	bool LoadLayout();
+
+	/** Free the loaded system and return to idle state */
+	bool UnloadLayout();
+
+	/** Is this layout system in use within the cloner */
+	UFUNCTION(BlueprintPure, Category="Cloner|Layout")
+	bool IsLayoutActive() const;
+
+	/** Activate this layout system on the cloner, must be loaded first */
+	bool ActivateLayout();
+
+	/** Deactivate this layout system if active */
+	bool DeactivateLayout();
+
+	/** Copies this layout data interfaces to other layout */
+	bool CopyTo(UAvaClonerLayoutBase* InOtherLayout) const;
+
+protected:
+	/** Called once after layout is loaded */
+	virtual void OnLayoutLoaded() {}
+
+	/** Called once after layout is unloaded */
+	virtual void OnLayoutUnloaded() {}
+
+	/** Called after layout becomes active */
+	virtual void OnLayoutActive() {}
+
+	/** Called after layout becomes inactive */
+	virtual void OnLayoutInactive() {}
+
+	/** Called to reapply layout parameters */
+	virtual void OnLayoutParametersChanged(UAvaClonerComponent* InComponent) {}
+
+	void OnLayoutPropertyChanged();
+	
+private:
+	UNiagaraSystem* LoadSystemPath(const FString& InPath) const;
+	
+	/** Layout name to display in layout options */
+	UPROPERTY(Transient)
+	FName LayoutName;
+
+	/** Niagara System asset path for this layout */
+	UPROPERTY(Transient)
+	FString LayoutAssetPath;
+
+	/** Niagara system representing this layout */
+	UPROPERTY(Transient)
+	TObjectPtr<UNiagaraSystem> NiagaraSystem;
+
+	/** Mesh renderer in this niagara system */
+	UPROPERTY(Transient)
+	TObjectPtr<UNiagaraMeshRendererProperties> MeshRenderer;
+
+	/** Data interfaces used by the effectors */
+	UPROPERTY(Transient)
+	FAvaClonerEffectorDataInterfaces DataInterfaces;
+};
\ No newline at end of file
diff --git a/Engine/Plugins/Experimental/Avalanche/Source/AvalancheEffectors/Public/Cloner/Layouts/AvaClonerLineLayout.h b/Engine/Plugins/Experimental/Avalanche/Source/AvalancheEffectors/Public/Cloner/Layouts/AvaClonerLineLayout.h
new file mode 100644
--- /dev/null
+++ b/Engine/Plugins/Experimental/Avalanche/Source/AvalancheEffectors/Public/Cloner/Layouts/AvaClonerLineLayout.h
@@ -0,0 +1,99 @@
+// Copyright Epic Games, Inc. All Rights Reserved.
+
+#pragma once
+
+#include "AvaClonerEffectorShared.h"
+#include "AvaClonerLayoutBase.h"
+#include "AvaPropertyChangeDispatcher.h"
+#include "AvaClonerLineLayout.generated.h"
+
+UCLASS(BlueprintType)
+class AVALANCHEEFFECTORS_API UAvaClonerLineLayout : public UAvaClonerLayoutBase
+{
+	GENERATED_BODY()
+
+	friend class FAvaClonerActorVisualizer;
+	
+public:
+	UAvaClonerLineLayout()
+		: UAvaClonerLayoutBase(
+			TEXT("Line")
+			, TEXT("/Script/Niagara.NiagaraSystem'/Avalanche/ClonerResources/Systems/NS_ClonerLine.NS_ClonerLine'")
+		)
+	{}
+	
+	UFUNCTION(BlueprintCallable, Category="Cloner|Layout|Line")
+	void SetCount(int32 InCount);
+
+	UFUNCTION(BlueprintPure, Category="Cloner|Layout|Line")
+	int32 GetCount() const
+	{
+		return Count;
+	}
+
+	UFUNCTION(BlueprintCallable, Category="Cloner|Layout|Line")
+	void SetSpacing(float InSpacing);
+	
+	UFUNCTION(BlueprintPure, Category="Cloner|Layout|Line")
+	float GetSpacing() const
+	{
+		return Spacing;
+	}
+
+	UFUNCTION(BlueprintCallable, Category="Cloner|Layout|Line")
+	void SetAxis(EAvaClonerAxis InAxis);
+
+	UFUNCTION(BlueprintPure, Category="Cloner|Layout|Line")
+	EAvaClonerAxis GetAxis() const
+	{
+		return Axis;
+	}
+
+	UFUNCTION(BlueprintCallable, Category="Cloner|Layout|Line")
+	void SetDirection(const FVector& InDirection);
+
+	UFUNCTION(BlueprintPure, Category="Cloner|Layout|Line")
+	const FVector& GetDirection() const
+	{
+		return Direction;
+	}
+
+	UFUNCTION(BlueprintCallable, Category="Cloner|Layout|Line")
+	void SetRotation(const FRotator& InRotation);
+
+	UFUNCTION(BlueprintPure, Category="Cloner|Layout|Line")
+	const FRotator& GetRotation() const
+	{
+		return Rotation;
+	}
+	
+	//~ Begin UObject
+#if WITH_EDITOR
+	virtual void PostEditChangeProperty(FPropertyChangedEvent& InPropertyChangedEvent) override;
+#endif
+	//~ End UObject
+
+protected:
+	virtual void OnLayoutParametersChanged(UAvaClonerComponent* InComponent) override;
+	
+	UPROPERTY(EditInstanceOnly, BlueprintReadWrite, Setter="SetCount", Getter="GetCount", Category="Layout")
+    int32 Count = 10;
+    	
+	UPROPERTY(EditInstanceOnly, BlueprintReadWrite, Setter="SetSpacing", Getter="GetSpacing", Category="Layout")
+	float Spacing = 105.f;
+    	
+	UPROPERTY(EditInstanceOnly, BlueprintReadWrite, Setter="SetAxis", Getter="GetAxis", Category="Layout")
+	EAvaClonerAxis Axis = EAvaClonerAxis::Y;
+    
+	UPROPERTY(EditInstanceOnly, BlueprintReadWrite, Setter="SetDirection", Getter="GetDirection", Category="Layout", meta=(ClampMin="0", ClampMax="1", EditCondition="Axis == EAvaClonerAxis::Custom", EditConditionHides))
+	FVector Direction = FVector::YAxisVector;
+    
+	UPROPERTY(EditInstanceOnly, BlueprintReadWrite, Setter="SetRotation", Getter="GetRotation", Category="Layout")
+	FRotator Rotation = FRotator(0.f);
+	
+private:
+#if WITH_EDITOR
+	/** Used for PECP */
+	static const TAvaPropertyChangeDispatcher<UAvaClonerLineLayout> PropertyChangeDispatcher;
+#endif
+};
\ No newline at end of file
diff --git a/Engine/Plugins/Experimental/Avalanche/Source/AvalancheEffectors/Public/Cloner/Layouts/AvaClonerMeshLayout.h b/Engine/Plugins/Experimental/Avalanche/Source/AvalancheEffectors/Public/Cloner/Layouts/AvaClonerMeshLayout.h
new file mode 100644
--- /dev/null
+++ b/Engine/Plugins/Experimental/Avalanche/Source/AvalancheEffectors/Public/Cloner/Layouts/AvaClonerMeshLayout.h
@@ -0,0 +1,105 @@
+// Copyright Epic Games, Inc. All Rights Reserved.
+
+#pragma once
+
+#include "AvaClonerEffectorShared.h"
+#include "AvaClonerLayoutBase.h"
+#include "AvaPropertyChangeDispatcher.h"
+#include "UObject/WeakObjectPtrTemplates.h"
+#include "AvaClonerMeshLayout.generated.h"
+
+class AActor;
+class USceneComponent;
+
+UCLASS(BlueprintType)
+class AVALANCHEEFFECTORS_API UAvaClonerMeshLayout : public UAvaClonerLayoutBase
+{
+	GENERATED_BODY()
+
+	friend class FAvaClonerActorVisualizer;
+	
+public:
+	UAvaClonerMeshLayout()
+		: UAvaClonerLayoutBase(
+			TEXT("Mesh")
+			, TEXT("/Script/Niagara.NiagaraSystem'/Avalanche/ClonerResources/Systems/NS_ClonerSampleMesh.NS_ClonerSampleMesh'")
+		)
+	{}
+	
+	UFUNCTION(BlueprintCallable, Category="Cloner|Layout|Mesh")
+	void SetCount(int32 InCount);
+	
+	UFUNCTION(BlueprintPure, Category="Cloner|Layout|Mesh")
+	int32 GetCount() const
+	{
+		return Count;
+	}
+
+	UFUNCTION(BlueprintCallable, Category="Cloner|Layout|Mesh")
+	void SetAsset(EAvaClonerMeshAsset InAsset);
+
+	UFUNCTION(BlueprintPure, Category="Cloner|Layout|Mesh")
+	EAvaClonerMeshAsset GetAsset() const
+	{
+		return Asset;
+	}
+	
+	UFUNCTION(BlueprintCallable, Category="Cloner|Layout|Mesh")
+	void SetSampleData(EAvaClonerMeshSampleData InSampleData);
+	
+	UFUNCTION(BlueprintPure, Category="Cloner|Layout|Mesh")
+	EAvaClonerMeshSampleData GetSampleData() const
+	{
+		return SampleData;
+	}
+
+	UFUNCTION()
+	void SetSampleActorWeak(const TWeakObjectPtr<AActor>& InSampleActor);
+	
+	TWeakObjectPtr<AActor> GetSampleActorWeak() const
+	{
+		return SampleActorWeak;
+	}
+
+	UFUNCTION(BlueprintCallable, Category="Cloner|Layout|Mesh")
+	void SetSampleActor(AActor* InActor);
+
+	UFUNCTION(BlueprintPure, Category="Cloner|Layout|Mesh")
+	AActor* GetSampleActor() const
+	{
+		return SampleActorWeak.Get();
+	}
+
+	//~ Begin UObject
+#if WITH_EDITOR
+	virtual void PostEditChangeProperty(FPropertyChangedEvent& InPropertyChangedEvent) override;
+#endif
+	//~ End UObject
+
+protected:
+	virtual void OnLayoutInactive() override;
+	virtual void OnLayoutParametersChanged(UAvaClonerComponent* InComponent) override;
+	
+	void OnSampleMeshTransformed(USceneComponent* InComponent, EUpdateTransformFlags InFlags, ETeleportType InType);
+	
+	UPROPERTY(EditInstanceOnly, BlueprintReadWrite, Setter="SetCount", Getter="GetCount", Category="Layout")
+	int32 Count = 3 * 3 * 3;
+
+	UPROPERTY(EditInstanceOnly, BlueprintReadWrite, Setter="SetAsset", Getter="GetAsset", Category="Layout")
+	EAvaClonerMeshAsset Asset = EAvaClonerMeshAsset::StaticMesh;
+
+	UPROPERTY(EditInstanceOnly, BlueprintReadWrite, Setter="SetSampleData", Getter="GetSampleData", Category="Layout")
+	EAvaClonerMeshSampleData SampleData = EAvaClonerMeshSampleData::Vertices;
+
+	UPROPERTY(EditInstanceOnly, BlueprintReadWrite, Setter="SetSampleActorWeak", Getter="GetSampleActorWeak", DisplayName="SampleActor", Category="Layout")
+	TWeakObjectPtr<AActor> SampleActorWeak;
+	
+	UPROPERTY(Transient)
+	TWeakObjectPtr<USceneComponent> SceneComponentWeak;
+
+private:
+#if WITH_EDITOR
+	/** Used for PECP */
+	static const TAvaPropertyChangeDispatcher<UAvaClonerMeshLayout> PropertyChangeDispatcher;
+#endif
+};
\ No newline at end of file
diff --git a/Engine/Plugins/Experimental/Avalanche/Source/AvalancheEffectors/Public/Cloner/Layouts/AvaClonerSphereRandomLayout.h b/Engine/Plugins/Experimental/Avalanche/Source/AvalancheEffectors/Public/Cloner/Layouts/AvaClonerSphereRandomLayout.h
new file mode 100644
--- /dev/null
+++ b/Engine/Plugins/Experimental/Avalanche/Source/AvalancheEffectors/Public/Cloner/Layouts/AvaClonerSphereRandomLayout.h
@@ -0,0 +1,134 @@
+// Copyright Epic Games, Inc. All Rights Reserved.
+
+#pragma once
+
+#include "AvaClonerLayoutBase.h"
+#include "AvaPropertyChangeDispatcher.h"
+#include "AvaClonerSphereRandomLayout.generated.h"
+
+UCLASS(BlueprintType)
+class AVALANCHEEFFECTORS_API UAvaClonerSphereRandomLayout : public UAvaClonerLayoutBase
+{
+	GENERATED_BODY()
+
+	friend class FAvaClonerActorVisualizer;
+	
+public:
+	UAvaClonerSphereRandomLayout()
+		: UAvaClonerLayoutBase(
+			TEXT("SphereRandom")
+			, TEXT("/Script/Niagara.NiagaraSystem'/Avalanche/ClonerResources/Systems/NS_ClonerSphereRandom.NS_ClonerSphereRandom'")
+		)
+	{}
+	
+	UFUNCTION(BlueprintCallable, Category="Cloner|Layout|SphereRandom")
+	void SetCount(int32 InCount);
+
+	UFUNCTION(BlueprintPure, Category="Cloner|Layout|SphereRandom")
+	int32 GetCount() const
+	{
+		return Count;
+	}
+
+	UFUNCTION(BlueprintCallable, Category="Cloner|Layout|SphereRandom")
+	void SetRadius(float InRadius);
+	
+	UFUNCTION(BlueprintPure, Category="Cloner|Layout|SphereRandom")
+	float GetRadius() const
+	{
+		return Radius;
+	}
+
+	UFUNCTION(BlueprintCallable, Category="Cloner|Layout|SphereRandom")
+	void SetDistribution(float InDistribution);
+
+	UFUNCTION(BlueprintPure, Category="Cloner|Layout|SphereRandom")
+	float GetDistribution() const
+	{
+		return Distribution;
+	}
+
+	UFUNCTION(BlueprintCallable, Category="Cloner|Layout|SphereRandom")
+	void SetLongitude(float InLongitude);
+
+	UFUNCTION(BlueprintPure, Category="Cloner|Layout|SphereRandom")
+	float GetLongitude() const
+	{
+		return Longitude;
+	}
+
+	UFUNCTION(BlueprintCallable, Category="Cloner|Layout|SphereRandom")
+	void SetLatitude(float InLatitude);
+
+	UFUNCTION(BlueprintPure, Category="Cloner|Layout|SphereRandom")
+	float GetLatitude() const
+	{
+		return Latitude;
+	}
+
+	UFUNCTION(BlueprintCallable, Category="Cloner|Layout|SphereRandom")
+	void SetOrientMesh(bool InbOrientMesh);
+
+	UFUNCTION(BlueprintPure, Category="Cloner|Layout|SphereRandom")
+	bool GetOrientMesh() const
+	{
+		return bOrientMesh;
+	}
+
+	UFUNCTION(BlueprintCallable, Category="Cloner|Layout|SphereRandom")
+	void SetRotation(const FRotator& InRotation);
+	
+	UFUNCTION(BlueprintPure, Category="Cloner|Layout|SphereRandom")
+	FRotator GetRotation() const
+	{
+		return Rotation;
+	}
+
+	UFUNCTION(BlueprintCallable, Category="Cloner|Layout|SphereRandom")
+	void SetScale(const FVector& InScale);
+	
+	UFUNCTION(BlueprintPure, Category="Cloner|Layout|SphereRandom")
+	FVector GetScale() const
+	{
+		return Scale;
+	}
+
+	//~ Begin UObject
+#if WITH_EDITOR
+	virtual void PostEditChangeProperty(FPropertyChangedEvent& InPropertyChangedEvent) override;
+#endif
+	//~ End UObject
+	
+protected:
+	virtual void OnLayoutParametersChanged(UAvaClonerComponent* InComponent) override;
+	
+	UPROPERTY(EditInstanceOnly, BlueprintReadWrite, Setter="SetCount", Getter="GetCount", Category="Layout")
+	int32 Count = 30;
+	
+	UPROPERTY(EditInstanceOnly, BlueprintReadWrite, Setter="SetRadius", Getter="GetRadius", Category="Layout")
+	float Radius = 200.f;
+
+	UPROPERTY(EditInstanceOnly, BlueprintReadWrite, Setter="SetDistribution", Getter="GetDistribution", Category="Layout", meta=(ClampMin="0", ClampMax="1"))
+	float Distribution = 1.f;
+
+	UPROPERTY(EditInstanceOnly, BlueprintReadWrite, Setter="SetLongitude", Getter="GetLongitude", Category="Layout", meta=(ClampMin="0", ClampMax="1"))
+	float Longitude = 1.f;
+
+	UPROPERTY(EditInstanceOnly, BlueprintReadWrite, Setter="SetLatitude", Getter="GetLatitude", Category="Layout", meta=(ClampMin="0", ClampMax="1"))
+	float Latitude = 1.f;
+
+	UPROPERTY(EditInstanceOnly, BlueprintReadWrite, Setter="SetOrientMesh", Getter="GetOrientMesh", Category="Layout")
+	bool bOrientMesh = false;
+	
+	UPROPERTY(EditInstanceOnly, BlueprintReadWrite, Setter="SetRotation", Getter="GetRotation", Category="Layout")
+	FRotator Rotation = FRotator(0.f, 0.f, 0.f);
+	
+	UPROPERTY(EditInstanceOnly, BlueprintReadWrite, Setter="SetScale", Getter="GetScale", Category="Layout", meta=(ClampMin="0", AllowPreserveRatio, Delta="0.01"))
+	FVector Scale = FVector(1.f, 1.f, 1.f);
+
+private:
+#if WITH_EDITOR
+	/** Used for PECP */
+	static const TAvaPropertyChangeDispatcher<UAvaClonerSphereRandomLayout> PropertyChangeDispatcher;
+#endif
+};
\ No newline at end of file
diff --git a/Engine/Plugins/Experimental/Avalanche/Source/AvalancheEffectors/Public/Cloner/Layouts/AvaClonerSphereUniformLayout.h b/Engine/Plugins/Experimental/Avalanche/Source/AvalancheEffectors/Public/Cloner/Layouts/AvaClonerSphereUniformLayout.h
new file mode 100644
--- /dev/null
+++ b/Engine/Plugins/Experimental/Avalanche/Source/AvalancheEffectors/Public/Cloner/Layouts/AvaClonerSphereUniformLayout.h
@@ -0,0 +1,110 @@
+// Copyright Epic Games, Inc. All Rights Reserved.
+
+#pragma once
+
+#include "AvaClonerLayoutBase.h"
+#include "AvaPropertyChangeDispatcher.h"
+#include "AvaClonerSphereUniformLayout.generated.h"
+
+UCLASS(BlueprintType)
+class AVALANCHEEFFECTORS_API UAvaClonerSphereUniformLayout : public UAvaClonerLayoutBase
+{
+	GENERATED_BODY()
+
+	friend class FAvaClonerActorVisualizer;
+	
+public:
+	UAvaClonerSphereUniformLayout()
+		: UAvaClonerLayoutBase(
+			TEXT("SphereUniform")
+			, TEXT("/Script/Niagara.NiagaraSystem'/Avalanche/ClonerResources/Systems/NS_ClonerSphere.NS_ClonerSphere'")
+		)
+	{}
+	
+	UFUNCTION(BlueprintCallable, Category="Cloner|Layout|SphereUniform")
+	void SetCount(int32 InCount);
+
+	UFUNCTION(BlueprintPure, Category="Cloner|Layout|SphereUniform")
+	int32 GetCount() const
+	{
+		return Count;
+	}
+
+	UFUNCTION(BlueprintCallable, Category="Cloner|Layout|SphereUniform")
+	void SetRadius(float InRadius);
+	
+	UFUNCTION(BlueprintPure, Category="Cloner|Layout|SphereUniform")
+	float GetRadius() const
+	{
+		return Radius;
+	}
+
+	UFUNCTION(BlueprintCallable, Category="Cloner|Layout|SphereUniform")
+	void SetRatio(float InRatio);
+
+	UFUNCTION(BlueprintPure, Category="Cloner|Layout|SphereUniform")
+	float GetRatio() const
+	{
+		return Ratio;
+	}
+
+	UFUNCTION(BlueprintCallable, Category="Cloner|Layout|SphereUniform")
+	void SetOrientMesh(bool InbOrientMesh);
+
+	UFUNCTION(BlueprintPure, Category="Cloner|Layout|SphereUniform")
+	bool GetOrientMesh() const
+	{
+		return bOrientMesh;
+	}
+
+	UFUNCTION(BlueprintCallable, Category="Cloner|Layout|SphereUniform")
+	void SetRotation(const FRotator& InRotation);
+	
+	UFUNCTION(BlueprintPure, Category="Cloner|Layout|SphereUniform")
+	FRotator GetRotation() const
+	{
+		return Rotation;
+	}
+
+	UFUNCTION(BlueprintCallable, Category="Cloner|Layout|SphereUniform")
+	void SetScale(const FVector& InScale);
+	
+	UFUNCTION(BlueprintPure, Category="Cloner|Layout|SphereUniform")
+	FVector GetScale() const
+	{
+		return Scale;
+	}
+
+	//~ Begin UObject
+#if WITH_EDITOR
+	virtual void PostEditChangeProperty(FPropertyChangedEvent& InPropertyChangedEvent) override;
+#endif
+	//~ End UObject
+	
+protected:
+	virtual void OnLayoutParametersChanged(UAvaClonerComponent* InComponent) override;
+	
+	UPROPERTY(EditInstanceOnly, BlueprintReadWrite, Setter="SetCount", Getter="GetCount", Category="Layout")
+	int32 Count = 30;
+	
+	UPROPERTY(EditInstanceOnly, BlueprintReadWrite, Setter="SetRadius", Getter="GetRadius", Category="Layout")
+	float Radius = 200.f;
+
+	UPROPERTY(EditInstanceOnly, BlueprintReadWrite, Setter="SetRatio", Getter="GetRatio", Category="Layout", meta=(ClampMin="0", ClampMax="1"))
+	float Ratio = 1.f;
+
+	UPROPERTY(EditInstanceOnly, BlueprintReadWrite, Setter="SetOrientMesh", Getter="GetOrientMesh", Category="Layout")
+	bool bOrientMesh = false;
+	
+	UPROPERTY(EditInstanceOnly, BlueprintReadWrite, Setter="SetRotation", Getter="GetRotation", Category="Layout")
+	FRotator Rotation = FRotator(0.f, 0.f, 0.f);
+	
+	UPROPERTY(EditInstanceOnly, BlueprintReadWrite, Setter="SetScale", Getter="GetScale", Category="Layout", meta=(ClampMin="0", AllowPreserveRatio, Delta="0.01"))
+	FVector Scale = FVector(1.f, 1.f, 1.f);
+
+private:
+#if WITH_EDITOR
+	/** Used for PECP */
+	static const TAvaPropertyChangeDispatcher<UAvaClonerSphereUniformLayout> PropertyChangeDispatcher;
+#endif
+};
\ No newline at end of file
diff --git a/Engine/Plugins/Experimental/Avalanche/Source/AvalancheEffectors/Public/Cloner/Layouts/AvaClonerSplineLayout.h b/Engine/Plugins/Experimental/Avalanche/Source/AvalancheEffectors/Public/Cloner/Layouts/AvaClonerSplineLayout.h
new file mode 100644
--- /dev/null
+++ b/Engine/Plugins/Experimental/Avalanche/Source/AvalancheEffectors/Public/Cloner/Layouts/AvaClonerSplineLayout.h
@@ -0,0 +1,97 @@
+// Copyright Epic Games, Inc. All Rights Reserved.
+
+#pragma once
+
+#include "AvaClonerLayoutBase.h"
+#include "AvaPropertyChangeDispatcher.h"
+#include "UObject/WeakObjectPtrTemplates.h"
+#include "AvaClonerSplineLayout.generated.h"
+
+class AActor;
+class USplineComponent;
+
+UCLASS(BlueprintType)
+class AVALANCHEEFFECTORS_API UAvaClonerSplineLayout : public UAvaClonerLayoutBase
+{
+	GENERATED_BODY()
+
+	friend class FAvaClonerActorVisualizer;
+	
+public:
+	UAvaClonerSplineLayout()
+		: UAvaClonerLayoutBase(
+			TEXT("Spline")
+			, TEXT("/Script/Niagara.NiagaraSystem'/Avalanche/ClonerResources/Systems/NS_ClonerSampleSpline.NS_ClonerSampleSpline'")
+		)
+	{}
+	
+	UFUNCTION(BlueprintCallable, Category="Cloner|Layout|Spline")
+	void SetCount(int32 InCount);
+	
+	UFUNCTION(BlueprintPure, Category="Cloner|Layout|Spline")
+	int32 GetCount() const
+	{
+		return Count;
+	}
+
+	UFUNCTION()
+	void SetSplineActorWeak(const TWeakObjectPtr<AActor>& InSplineActor);
+
+	TWeakObjectPtr<AActor> GetSplineActorWeak() const
+	{
+		return SplineActorWeak;
+	}
+	
+	UFUNCTION(BlueprintCallable, Category="Cloner|Layout|Spline")
+	void SetSplineActor(AActor* InSplineActor);
+	
+	UFUNCTION(BlueprintPure, Category="Cloner|Layout|Spline")
+	AActor* GetSplineActor() const
+	{
+		return SplineActorWeak.Get();
+	}
+	
+	UFUNCTION(BlueprintCallable, Category="Cloner|Layout|Spline")
+	void SetOrientMesh(bool bInOrientMesh);
+	
+	UFUNCTION(BlueprintPure, Category="Cloner|Layout|Spline")
+	bool GetOrientMesh() const
+	{
+		return bOrientMesh;
+	}
+
+#if WITH_EDITOR
+	/** Spawns an actor with a spline component and linked it to this cloner layout */
+	UFUNCTION(CallInEditor, Category="Cloner|Layout")
+	void SpawnLinkedSplineActor();
+
+	//~ Begin UObject
+	virtual void PostEditChangeProperty(FPropertyChangedEvent& InPropertyChangedEvent) override;
+#endif
+	//~ End UObject
+
+protected:
+	virtual void OnLayoutInactive() override;
+	virtual void OnLayoutParametersChanged(UAvaClonerComponent* InComponent) override;
+	
+	void OnSampleSplineTransformed(USceneComponent* InComponent, EUpdateTransformFlags InFlags, ETeleportType InType);
+	void OnSampleSplineRenderStateUpdated(UActorComponent& InComponent);
+	
+	UPROPERTY(EditInstanceOnly, BlueprintReadWrite, Setter="SetCount", Getter="GetCount", Category="Layout")
+	int32 Count = 3 * 3;
+
+	UPROPERTY(EditInstanceOnly, BlueprintReadWrite, Setter="SetSplineActorWeak", Getter="GetSplineActorWeak", DisplayName="SampleActor", Category="Layout")
+	TWeakObjectPtr<AActor> SplineActorWeak;
+
+	UPROPERTY(EditInstanceOnly, BlueprintReadWrite, Setter="SetOrientMesh", Getter="GetOrientMesh", Category="Layout")
+    bool bOrientMesh = false;
+
+	UPROPERTY(Transient)
+	TWeakObjectPtr<USplineComponent> SplineComponentWeak;
+
+private:
+#if WITH_EDITOR
+	/** Used for PECP */
+	static const TAvaPropertyChangeDispatcher<UAvaClonerSplineLayout> PropertyChangeDispatcher;
+#endif
+};
