[MotionDesign]

diff --git a/Engine/Plugins/Experimental/Avalanche/Source/AvalancheEffectors/Public/Effector/AvaEffectorActor.h b/Engine/Plugins/Experimental/Avalanche/Source/AvalancheEffectors/Public/Effector/AvaEffectorActor.h
new file mode 100644
--- /dev/null
+++ b/Engine/Plugins/Experimental/Avalanche/Source/AvalancheEffectors/Public/Effector/AvaEffectorActor.h
@@ -0,0 +1,651 @@
+// Copyright Epic Games, Inc. All Rights Reserved.
+
+#pragma once
+
+#include "AvaClonerEffectorShared.h"
+#include "AvaPropertyChangeDispatcher.h"
+#include "GameFramework/Actor.h"
+#include "AvaEffectorActor.generated.h"
+
+class UNiagaraDataChannelWriter;
+class AAvaClonerActor;
+class USphereComponent;
+class UBoxComponent;
+
+UCLASS(BlueprintType, HideCategories=(Rendering,Replication,Collision,HLOD,Physics,Networking,Input,Actor,Cooking,LevelInstance), DisplayName = "Motion Design Effector Actor")
+class AVALANCHEEFFECTORS_API AAvaEffectorActor : public AActor
+{
+	GENERATED_BODY()
+
+	friend class AAvaClonerActor;
+	friend class FAvaEffectorDetailCustomization;
+	friend class FAvaEffectorActorVisualizer;
+	friend class UAvaEffectorSubsystem;
+
+public:
+	static inline const FString DefaultLabel = TEXT("Effector");
+
+	AAvaEffectorActor();
+
+	void LinkCloner(AAvaClonerActor* InCloner);
+	void UnlinkCloner(AAvaClonerActor* InCloner);
+
+	UFUNCTION(BlueprintCallable, Category="Effector")
+	void SetType(EAvaClonerEffectorType InType);
+
+	UFUNCTION(BlueprintPure, Category="Effector")
+	EAvaClonerEffectorType GetType() const
+	{
+		return Type;
+	}
+
+	UFUNCTION(BlueprintCallable, Category="Effector")
+	void SetEasing(EAvaClonerEasing InEasing);
+
+	UFUNCTION(BlueprintPure, Category="Effector")
+	EAvaClonerEasing GetEasing() const
+	{
+		return Easing;
+	}
+
+	UFUNCTION(BlueprintCallable, Category="Effector")
+	void SetMagnitude(float InMagnitude);
+
+	UFUNCTION(BlueprintPure, Category="Effector")
+	float GetMagnitude() const
+	{
+		return Magnitude;
+	}
+
+	UFUNCTION(BlueprintCallable, Category="Effector")
+	void SetOuterRadius(float InRadius);
+
+	UFUNCTION(BlueprintPure, Category="Effector")
+	float GetOuterRadius() const
+	{
+		return OuterRadius;
+	}
+
+	UFUNCTION(BlueprintCallable, Category="Effector")
+	void SetInnerRadius(float InRadius);
+
+	UFUNCTION(BlueprintPure, Category="Effector")
+	float GetInnerRadius() const
+	{
+		return InnerRadius;
+	}
+
+	UFUNCTION(BlueprintCallable, Category="Effector")
+	void SetInnerExtent(const FVector& InExtent);
+
+	UFUNCTION(BlueprintPure, Category="Effector")
+	FVector GetInnerExtent() const
+	{
+		return InnerExtent;
+	}
+
+	UFUNCTION(BlueprintCallable, Category="Effector")
+	void SetOuterExtent(const FVector& InExtent);
+
+	UFUNCTION(BlueprintPure, Category="Effector")
+	FVector GetOuterExtent() const
+	{
+		return OuterExtent;
+	}
+
+	UFUNCTION(BlueprintCallable, Category="Effector")
+	void SetPlaneSpacing(float InSpacing);
+
+	UFUNCTION(BlueprintPure, Category="Effector")
+	float GetPlaneSpacing() const
+	{
+		return PlaneSpacing;
+	}
+
+	UFUNCTION(BlueprintCallable, Category="Effector")
+	void SetOffset(const FVector& InOffset);
+
+	UFUNCTION(BlueprintPure, Category="Effector")
+	FVector GetOffset() const
+	{
+		return Offset;
+	}
+
+	UFUNCTION(BlueprintCallable, Category="Effector")
+	void SetRotation(const FRotator& InRotation);
+
+	UFUNCTION(BlueprintPure, Category="Effector")
+	FRotator GetRotation() const
+	{
+		return Rotation;
+	}
+
+	UFUNCTION(BlueprintCallable, Category="Effector")
+	void SetScale(const FVector& InScale);
+
+	UFUNCTION(BlueprintPure, Category="Effector")
+	FVector GetScale() const
+	{
+		return Scale;
+	}
+
+	UFUNCTION(BlueprintCallable, Category="Effector")
+	void SetEnabled(bool bInEnable);
+
+	UFUNCTION(BlueprintPure, Category="Effector")
+	bool GetEnabled() const
+	{
+		return bEnabled;
+	}
+
+	UFUNCTION(BlueprintCallable, Category="Effector")
+	void SetVisualizerThickness(float InThickness);
+
+	UFUNCTION(BlueprintPure, Category="Effector")
+	float GetVisualizerThickness() const
+	{
+		return VisualizerThickness;
+	}
+
+	UFUNCTION(BlueprintCallable, Category="Effector")
+	void SetMode(EAvaClonerEffectorMode InMode);
+
+	UFUNCTION(BlueprintPure, Category="Effector")
+	EAvaClonerEffectorMode GetMode() const
+	{
+		return Mode;
+	}
+
+	UFUNCTION(BlueprintCallable, Category="Effector")
+	void SetTargetActor(AActor* InTargetActor);
+
+	UFUNCTION(BlueprintPure, Category="Effector")
+	AActor* GetTargetActor() const
+	{
+		return TargetActorWeak.Get();
+	}
+
+	UFUNCTION(BlueprintCallable, Category="Cloner")
+	void SetOrientationForceEnabled(bool bInForceEnabled);
+
+	UFUNCTION(BlueprintPure, Category="Cloner")
+	bool GetOrientationForceEnabled() const
+	{
+		return bOrientationForceEnabled;
+	}
+
+	UFUNCTION(BlueprintCallable, Category="Cloner")
+	void SetOrientationForceRate(float InForceOrientationRate);
+
+	UFUNCTION(BlueprintPure, Category="Cloner")
+	float GetOrientationForceRate() const
+	{
+		return OrientationForceRate;
+	}
+
+	UFUNCTION(BlueprintCallable, Category="Cloner")
+	void SetOrientationForceMin(const FVector& InForceOrientationMin);
+
+	UFUNCTION(BlueprintPure, Category="Cloner")
+	const FVector& GetOrientationForceMin() const
+	{
+		return OrientationForceMin;
+	}
+
+	UFUNCTION(BlueprintCallable, Category="Cloner")
+	void SetOrientationForceMax(const FVector& InForceOrientationMax);
+
+	UFUNCTION(BlueprintPure, Category="Cloner")
+	const FVector& GetOrientationForceMax() const
+	{
+		return OrientationForceMax;
+	}
+
+	UFUNCTION(BlueprintCallable, Category="Cloner")
+	void SetVortexForceEnabled(bool bInForceEnabled);
+
+	UFUNCTION(BlueprintPure, Category="Cloner")
+	bool GetVortexForceEnabled() const
+	{
+		return bVortexForceEnabled;
+	}
+
+	UFUNCTION(BlueprintCallable, Category="Cloner")
+	void SetVortexForceAmount(float InForceVortexAmount);
+
+	UFUNCTION(BlueprintPure, Category="Cloner")
+	float GetVortexForceAmount() const
+	{
+		return VortexForceAmount;
+	}
+
+	UFUNCTION(BlueprintCallable, Category="Cloner")
+	void SetVortexForceAxis(const FVector& InForceVortexAxis);
+
+	UFUNCTION(BlueprintPure, Category="Cloner")
+	const FVector& GetVortexForceAxis() const
+	{
+		return VortexForceAxis;
+	}
+
+	UFUNCTION(BlueprintCallable, Category="Cloner")
+	void SetCurlNoiseForceEnabled(bool bInForceEnabled);
+
+	UFUNCTION(BlueprintPure, Category="Cloner")
+	bool GetCurlNoiseForceEnabled() const
+	{
+		return bCurlNoiseForceEnabled;
+	}
+
+	UFUNCTION(BlueprintCallable, Category="Cloner")
+	void SetCurlNoiseForceStrength(float InForceCurlNoiseStrength);
+
+	UFUNCTION(BlueprintPure, Category="Cloner")
+	float GetCurlNoiseForceStrength() const
+	{
+		return CurlNoiseForceStrength;
+	}
+
+	UFUNCTION(BlueprintCallable, Category="Cloner")
+	void SetCurlNoiseForceFrequency(float InForceCurlNoiseFrequency);
+
+	UFUNCTION(BlueprintPure, Category="Cloner")
+	float GetCurlNoiseForceFrequency() const
+	{
+		return CurlNoiseForceFrequency;
+	}
+
+	UFUNCTION(BlueprintCallable, Category="Cloner")
+	void SetAttractionForceEnabled(bool bInForceEnabled);
+
+	UFUNCTION(BlueprintPure, Category="Cloner")
+	bool GetAttractionForceEnabled() const
+	{
+		return bAttractionForceEnabled;
+	}
+
+	UFUNCTION(BlueprintCallable, Category="Cloner")
+	void SetAttractionForceStrength(float InForceStrength);
+
+	UFUNCTION(BlueprintPure, Category="Cloner")
+	float GetAttractionForceStrength() const
+	{
+		return AttractionForceStrength;
+	}
+
+	UFUNCTION(BlueprintCallable, Category="Cloner")
+	void SetAttractionForceFalloff(float InForceFalloff);
+
+	UFUNCTION(BlueprintPure, Category="Cloner")
+	float GetAttractionForceFalloff() const
+	{
+		return AttractionForceFalloff;
+	}
+
+	UFUNCTION(BlueprintCallable, Category="Cloner")
+	void SetGravityForceEnabled(bool bInForceEnabled);
+
+	UFUNCTION(BlueprintPure, Category="Cloner")
+	bool GetGravityForceEnabled() const
+	{
+		return bGravityForceEnabled;
+	}
+
+	UFUNCTION(BlueprintCallable, Category="Cloner")
+	void SetGravityForceAcceleration(const FVector& InAcceleration);
+
+	UFUNCTION(BlueprintPure, Category="Cloner")
+	FVector GetGravityForceAcceleration() const
+	{
+		return GravityForceAcceleration;
+	}
+
+	UFUNCTION(BlueprintCallable, Category="Cloner")
+	void SetLocationStrength(const FVector& InStrength);
+
+	UFUNCTION(BlueprintPure, Category="Cloner")
+	FVector GetLocationStrength() const
+	{
+		return LocationStrength;
+	}
+
+	UFUNCTION(BlueprintCallable, Category="Cloner")
+	void SetRotationStrength(const FRotator& InStrength);
+
+	UFUNCTION(BlueprintPure, Category="Cloner")
+	FRotator GetRotationStrength() const
+	{
+		return RotationStrength;
+	}
+
+	UFUNCTION(BlueprintCallable, Category="Cloner")
+	void SetScaleStrength(const FVector& InStrength);
+
+	UFUNCTION(BlueprintPure, Category="Cloner")
+	FVector GetScaleStrength() const
+	{
+		return ScaleStrength;
+	}
+
+	UFUNCTION(BlueprintCallable, Category="Cloner")
+	void SetPan(const FVector& InPan);
+
+	UFUNCTION(BlueprintPure, Category="Cloner")
+	FVector GetPan() const
+	{
+		return Pan;
+	}
+
+	UFUNCTION(BlueprintCallable, Category="Cloner")
+	void SetFrequency(float InFrequency);
+
+	UFUNCTION(BlueprintPure, Category="Cloner")
+	float GetFrequency() const
+	{
+		return Frequency;
+	}
+
+#if WITH_EDITOR
+	UFUNCTION(BlueprintCallable, Category="Effector")
+	void SetVisualizerSpriteVisible(bool bInVisible);
+
+	UFUNCTION(BlueprintPure, Category="Effector")
+	bool GetVisualizerSpriteVisible() const
+	{
+		return bVisualizerSpriteVisible;
+	}
+#endif
+
+protected:
+	//~ Begin AActor
+	virtual void Destroyed() override;
+	virtual void PostActorCreated() override;
+#if WITH_EDITOR
+	virtual FString GetDefaultActorLabel() const override;
+#endif
+	//~ End AActor
+
+	//~ Begin UObject
+#if WITH_EDITOR
+	virtual void PostEditChangeProperty(FPropertyChangedEvent& InPropertyChangedEvent) override;
+	virtual void PostEditImport() override;
+	virtual void PostEditUndo() override;
+#endif
+	virtual void PostLoad() override;
+	virtual void PostDuplicate(EDuplicateMode::Type InDuplicateMode) override;
+	//~ End UObject
+
+	/** For each valid cloner that has this effector registered, loops in the internal cloner array */
+	void ForEachCloner(TFunctionRef<bool(AAvaClonerActor*, int32)> InFunction, bool bSkipIdxCheck = false);
+
+	void OnEffectorSubsystemInitialized(const UWorld* World);
+
+	/** Registers this effector to data channel */
+	void RegisterToChannel();
+
+	/** Called when this effector identifier changed to update linked cloners */
+	void OnEffectorIdentifierChanged();
+
+	FAvaClonerEffectorChannelData& GetEffectorChannelData()
+	{
+		return ChannelData;
+	}
+
+	void OnEffectorTransformed(USceneComponent* InUpdatedComponent, EUpdateTransformFlags InUpdateTransformFlags, ETeleportType InTeleport);
+
+	void OnClonerUpdated(AAvaClonerActor* InCloner, int32 InEffectorIdx) {}
+	void OnClonerLinked(AAvaClonerActor* InCloner, int32 InEffectorIdx);
+	void OnClonerUnlinked(AAvaClonerActor* InCloner, int32 InEffectorIdx);
+
+	/** Will update cloner system only if effector is enabled or force update is true, update will happen next tick */
+	void RefreshClonerParameters(AAvaClonerActor* InCloner, bool bInForce = false) const;
+	void RefreshClonersParameters(bool bInForce = false);
+
+	/** Will register this effector with linked clones */
+	void RegisterToCloners();
+
+	/** Will update all values of this effector on all clones */
+	void OnEffectorChanged();
+
+	/** Called when effector is enabled/disabled */
+	void OnEnabledChanged();
+
+	/** Called when we disable this effector */
+	void OnEffectorDisabled();
+
+	/** Called when cloners set is updated */
+	void OnClonersChanged();
+
+	/** Update mode for this effector */
+	void OnModeChanged();
+
+	/** Update shape type for this effector */
+	void OnTypeChanged();
+
+	/** Update values for box type effector */
+	void OnBoxChanged();
+
+	/** Update values for sphere type effector */
+	void OnSphereChanged();
+
+	/** Update values for plane type effector */
+	void OnPlaneChanged();
+
+	/** Called when transform of effector has changed */
+	void OnTransformChanged();
+
+	/** Update values for easing */
+	void OnEasingChanged();
+
+	/** Update values for magnitude */
+	void OnMagnitudeChanged();
+
+	/** Called when default mode options are changed */
+	void OnTransformOptionsChanged();
+
+	/** Update actor targeted by this effector on cloners */
+	void OnTargetActorChanged();
+
+	/** Called when target mode options are changed */
+	void OnTargetOptionsChanged();
+
+	/** Called when noise field mode options are changed */
+	void OnNoiseFieldOptionsChanged();
+
+	/** Called when force options are changed */
+	void OnForceOptionsChanged();
+
+	/** Called when force enabled state changed */
+	void OnForceEnabledChanged();
+
+	/** Update thickness of components of this effector */
+	void OnVisualizerThicknessChanged();
+
+	/** Update sprite visibility of this effector */
+	void OnVisualizerSpriteVisibleChanged();
+
+	/** Is this effector enabled/disabled on linked cloners */
+	UPROPERTY(EditAnywhere, BlueprintReadWrite, Setter="SetEnabled", Getter="GetEnabled", Category="Effector")
+	bool bEnabled = true;
+
+	/** The ratio effect of the effector on clones */
+	UPROPERTY(EditAnywhere, BlueprintReadWrite, Setter="SetMagnitude", Getter="GetMagnitude", Category="Effector", meta=(ClampMin="0", ClampMax="1"))
+	float Magnitude = 1.f;
+
+	/** Cloners affected by this effector */
+	UPROPERTY(EditAnywhere, Category="Effector", meta=(AllowPrivateAccess = "true"))
+	TSet<TWeakObjectPtr<AAvaClonerActor>> Cloners;
+
+	/** Type of effector to apply on cloners instances */
+	UPROPERTY(EditAnywhere, BlueprintReadWrite, Setter="SetType", Getter="GetType", Category="Type")
+	EAvaClonerEffectorType Type = EAvaClonerEffectorType::Sphere;
+
+	/** Weight easing function applied to lerp transforms */
+	UPROPERTY(EditAnywhere, BlueprintReadWrite, Setter="SetEasing", Getter="GetEasing", Category="Type", meta=(EditCondition="Type != EAvaClonerEffectorType::Unbound", EditConditionHides))
+	EAvaClonerEasing Easing = EAvaClonerEasing::Linear;
+
+	/** Inner radius of sphere, all clones inside will be affected with a maximum weight */
+	UPROPERTY(EditAnywhere, BlueprintReadWrite, Setter="SetInnerRadius", Getter="GetInnerRadius", Category="Type", meta=(ClampMin="0", EditCondition="Type == EAvaClonerEffectorType::Sphere", EditConditionHides))
+	float InnerRadius = 50.f;
+
+	/** Outer radius of sphere, all clones outside will not be affected */
+	UPROPERTY(EditAnywhere, BlueprintReadWrite, Setter="SetOuterRadius", Getter="GetOuterRadius", Category="Type", meta=(ClampMin="0", EditCondition="Type == EAvaClonerEffectorType::Sphere", EditConditionHides))
+	float OuterRadius = 200.f;
+
+	/** Inner extent of box, all clones inside will be affected with a maximum weight */
+	UPROPERTY(EditAnywhere, BlueprintReadWrite, Setter="SetInnerExtent", Getter="GetInnerExtent", Category="Type", meta=(ClampMin="0", EditCondition="Type == EAvaClonerEffectorType::Box", EditConditionHides))
+	FVector InnerExtent = FVector(50.f);
+
+	/** Outer extent of box, all clones outside will not be affected */
+	UPROPERTY(EditAnywhere, BlueprintReadWrite, Setter="SetOuterExtent", Getter="GetOuterExtent", Category="Type", meta=(ClampMin="0", EditCondition="Type == EAvaClonerEffectorType::Box", EditConditionHides))
+	FVector OuterExtent = FVector(200.f);
+
+	/** Plane spacing, everything inside this zone will be affected */
+	UPROPERTY(EditAnywhere, BlueprintReadWrite, Setter="SetPlaneSpacing", Getter="GetPlaneSpacing", Category="Type", meta=(ClampMin="0", EditCondition="Type == EAvaClonerEffectorType::Plane", EditConditionHides))
+	float PlaneSpacing = 200.f;
+
+	/** Mode of effector for each clones instances */
+	UPROPERTY(EditAnywhere, BlueprintReadWrite, Setter="SetMode", Getter="GetMode", Category="Mode")
+	EAvaClonerEffectorMode Mode = EAvaClonerEffectorMode::Default;
+
+	/** Offset applied on affected clones */
+	UPROPERTY(EditAnywhere, BlueprintReadWrite, Setter="SetOffset", Getter="GetOffset", Category="Mode", meta=(EditCondition="Mode == EAvaClonerEffectorMode::Default", EditConditionHides))
+	FVector Offset = FVector::ZeroVector;
+
+	/** Rotation applied on affected clones */
+	UPROPERTY(EditAnywhere, BlueprintReadWrite, Setter="SetRotation", Getter="GetRotation", Category="Mode", meta=(ClampMin="-180", ClampMax="180", UIMin="-180", UIMax="180", EditCondition="Mode == EAvaClonerEffectorMode::Default", EditConditionHides))
+	FRotator Rotation = FRotator::ZeroRotator;
+
+	/** Scale applied on affected clones */
+	UPROPERTY(EditAnywhere, BlueprintReadWrite, Setter="SetScale", Getter="GetScale", Category="Mode", meta=(ClampMin="0", AllowPreserveRatio, Delta="0.01", EditCondition="Mode == EAvaClonerEffectorMode::Default", EditConditionHides))
+	FVector Scale = FVector::OneVector;
+
+	/** The actor to track when mode is set to target */
+	UPROPERTY(EditAnywhere, BlueprintReadWrite, Category="Mode", meta=(DisplayName="TargetActor", EditCondition="Mode == EAvaClonerEffectorMode::Target", EditConditionHides))
+	TWeakObjectPtr<AActor> TargetActorWeak = nullptr;
+
+	UPROPERTY()
+	TWeakObjectPtr<AActor> InternalTargetActorWeak = nullptr;
+
+	/** Amplitude of the noise field for location */
+	UPROPERTY(EditAnywhere, BlueprintReadWrite, Setter="SetLocationStrength", Getter="GetLocationStrength", Category="Mode", meta=(EditCondition="Mode == EAvaClonerEffectorMode::NoiseField", EditConditionHides))
+	FVector LocationStrength = FVector::ZeroVector;
+
+	/** Amplitude of the noise field for rotation */
+	UPROPERTY(EditAnywhere, BlueprintReadWrite, Category="Mode", meta=(EditCondition="Mode == EAvaClonerEffectorMode::NoiseField", EditConditionHides))
+	FRotator RotationStrength = FRotator::ZeroRotator;
+
+	/** Amplitude of the noise field for scale */
+	UPROPERTY(EditAnywhere, BlueprintReadWrite, Category="Mode", meta=(EditCondition="Mode == EAvaClonerEffectorMode::NoiseField", EditConditionHides))
+	FVector ScaleStrength = FVector::OneVector;
+
+	/** Panning to offset the noise field sampling */
+	UPROPERTY(EditAnywhere, BlueprintReadWrite, Setter="SetPan", Getter="GetPan", Category="Mode", meta=(EditCondition="Mode == EAvaClonerEffectorMode::NoiseField", EditConditionHides))
+	FVector Pan = FVector::ZeroVector;
+
+	/** Intensity of the noise field */
+	UPROPERTY(EditAnywhere, BlueprintReadWrite, Setter="SetFrequency", Getter="GetFrequency", Category="Mode", meta=(ClampMin="0", EditCondition="Mode == EAvaClonerEffectorMode::NoiseField", EditConditionHides))
+	float Frequency = 0.5f;
+
+	/** Enable orientation force to allow each clone instance to rotate around its pivot */
+	UPROPERTY(EditAnywhere, BlueprintReadWrite, Setter="SetOrientationForceEnabled", Getter="GetOrientationForceEnabled", Category="Force")
+	bool bOrientationForceEnabled = false;
+
+	UPROPERTY(EditAnywhere, BlueprintReadWrite, Setter="SetOrientationForceRate", Getter="GetOrientationForceRate", Category="Force", meta=(Delta="0.0001", ClampMin="0", EditCondition="bOrientationForceEnabled", EditConditionHides))
+	float OrientationForceRate = 1.f;
+
+	UPROPERTY(EditAnywhere, BlueprintReadWrite, Setter="SetOrientationForceMin", Getter="GetOrientationForceMin", Category="Force", meta=(EditCondition="bOrientationForceEnabled", EditConditionHides))
+	FVector OrientationForceMin = FVector(-0.1f);
+
+	UPROPERTY(EditAnywhere, BlueprintReadWrite, Setter="SetOrientationForceMax", Getter="GetOrientationForceMax", Category="Force", meta=(EditCondition="bOrientationForceEnabled", EditConditionHides))
+	FVector OrientationForceMax = FVector(0.1f);
+
+	/** Enable vortex force to allow each clone instance to rotate around a specific axis on the cloner pivot */
+	UPROPERTY(EditAnywhere, BlueprintReadWrite, Setter="SetVortexForceEnabled", Getter="GetVortexForceEnabled", Category="Force")
+	bool bVortexForceEnabled = false;
+
+	UPROPERTY(EditAnywhere, BlueprintReadWrite, Setter="SetVortexForceAmount", Getter="GetVortexForceAmount", Category="Force", meta=(EditCondition="bVortexForceEnabled", EditConditionHides))
+	float VortexForceAmount = 10000.f;
+
+	UPROPERTY(EditAnywhere, BlueprintReadWrite, Setter="SetVortexForceAxis", Getter="GetVortexForceAxis", Category="Force", meta=(EditCondition="bVortexForceEnabled", EditConditionHides))
+	FVector VortexForceAxis = FVector::ZAxisVector;
+
+	/** Enable curl noise force to allow each clone instance to add random location variation */
+	UPROPERTY(EditAnywhere, BlueprintReadWrite, Setter="SetCurlNoiseForceEnabled", Getter="GetCurlNoiseForceEnabled", Category="Force")
+	bool bCurlNoiseForceEnabled = false;
+
+	UPROPERTY(EditAnywhere, BlueprintReadWrite, Setter="SetCurlNoiseForceStrength", Getter="GetCurlNoiseForceStrength", Category="Force", meta=(EditCondition="bCurlNoiseForceEnabled", EditConditionHides))
+	float CurlNoiseForceStrength = 1000.f;
+
+	UPROPERTY(EditAnywhere, BlueprintReadWrite, Setter="SetCurlNoiseForceFrequency", Getter="GetCurlNoiseForceFrequency", Category="Force", meta=(EditCondition="bCurlNoiseForceEnabled", EditConditionHides))
+	float CurlNoiseForceFrequency = 10.f;
+
+	/** Enable attraction force to allow each clone instances to gravitate toward a location */
+	UPROPERTY(EditAnywhere, BlueprintReadWrite, Setter="SetAttractionForceEnabled", Getter="GetAttractionForceEnabled", Category="Force")
+	bool bAttractionForceEnabled = false;
+
+	UPROPERTY(EditAnywhere, BlueprintReadWrite, Setter="SetAttractionForceStrength", Getter="GetAttractionForceStrength", Category="Force", meta=(EditCondition="bAttractionForceEnabled", EditConditionHides))
+	float AttractionForceStrength = 1000.f;
+
+	UPROPERTY(EditAnywhere, BlueprintReadWrite, Setter="SetAttractionForceFalloff", Getter="GetAttractionForceFalloff", Category="Force", meta=(EditCondition="bAttractionForceEnabled", EditConditionHides))
+	float AttractionForceFalloff = 0.1f;
+
+	/** Enable gravity force to pull particles based on an acceleration vector */
+	UPROPERTY(EditAnywhere, BlueprintReadWrite, Setter="SetGravityForceEnabled", Getter="GetGravityForceEnabled", Category="Force")
+	bool bGravityForceEnabled = false;
+
+	UPROPERTY(EditAnywhere, BlueprintReadWrite, Setter="SetGravityForceAcceleration", Getter="GetGravityForceAcceleration", Category="Force", meta=(EditCondition="bGravityForceEnabled", EditConditionHides))
+	FVector GravityForceAcceleration = FVector(0, 0, -980.f);
+
+	/** Thickness of components visualizers */
+	UPROPERTY(EditAnywhere, BlueprintReadWrite, AdvancedDisplay, Setter="SetVisualizerThickness", Getter="GetVisualizerThickness", Category="Effector", meta=(ClampMin="0.1", ClampMax="10.0", Delta="0.1"))
+	float VisualizerThickness = 2.f;
+
+#if WITH_EDITORONLY_DATA
+	/** Toggle the sprite to visualize and click on this effector */
+	UPROPERTY(EditAnywhere, BlueprintReadWrite, AdvancedDisplay, Category="Effector")
+	bool bVisualizerSpriteVisible = true;
+#endif
+
+private:
+	void OnTargetActorTransformChanged(USceneComponent* InUpdatedComponent, EUpdateTransformFlags InFlags, ETeleportType InTeleport);
+
+	UFUNCTION()
+	void OnTargetActorDestroyed(AActor* InActor);
+
+	/** Sphere component for sphere maximum weight */
+	UPROPERTY()
+	TObjectPtr<USphereComponent> InnerSphereComponent;
+
+	/** Sphere component for sphere minimum weight */
+	UPROPERTY()
+	TObjectPtr<USphereComponent> OuterSphereComponent;
+
+	/** Box component for box maximum weight */
+	UPROPERTY()
+	TObjectPtr<UBoxComponent> InnerBoxComponent;
+
+	/** Box component for box minimum weight */
+	UPROPERTY()
+	TObjectPtr<UBoxComponent> OuterBoxComponent;
+
+	/** Plane components for plane minimum weight */
+	UPROPERTY()
+	TObjectPtr<UBoxComponent> InnerPlaneComponent;
+
+	/** Plane components for plane maximum weight */
+	UPROPERTY()
+	TObjectPtr<UBoxComponent> OuterPlaneComponent;
+
+	/** Internal cloners array */
+	UPROPERTY(NonTransactional)
+	TSet<TWeakObjectPtr<AAvaClonerActor>> InternalCloners;
+
+	/** Transient effector channel data */
+	FAvaClonerEffectorChannelData ChannelData;
+
+#if WITH_EDITOR
+	/** Used for PECP */
+	static TAvaPropertyChangeDispatcher<AAvaEffectorActor> PropertyChangeDispatcher;
+#endif
+};
diff --git a/Engine/Plugins/Experimental/Avalanche/Source/AvalancheEffectors/Public/Effector/AvaEffectorComponent.h b/Engine/Plugins/Experimental/Avalanche/Source/AvalancheEffectors/Public/Effector/AvaEffectorComponent.h
new file mode 100644
--- /dev/null
+++ b/Engine/Plugins/Experimental/Avalanche/Source/AvalancheEffectors/Public/Effector/AvaEffectorComponent.h
@@ -0,0 +1,13 @@
+// Copyright Epic Games, Inc. All Rights Reserved.
+
+#pragma once
+
+#include "Components/SceneComponent.h"
+#include "AvaEffectorComponent.generated.h"
+
+/** Class used to define a custom visualizer for this component and the effector actor */
+UCLASS(MinimalAPI, Within = AvaEffectorActor)
+class UAvaEffectorComponent : public USceneComponent
+{
+	GENERATED_BODY()
+};
\ No newline at end of file
diff --git a/Engine/Plugins/Experimental/Avalanche/Source/AvalancheEffectors/Public/Subsystems/AvaClonerSubsystem.h b/Engine/Plugins/Experimental/Avalanche/Source/AvalancheEffectors/Public/Subsystems/AvaClonerSubsystem.h
new file mode 100644
--- /dev/null
+++ b/Engine/Plugins/Experimental/Avalanche/Source/AvalancheEffectors/Public/Subsystems/AvaClonerSubsystem.h
@@ -0,0 +1,66 @@
+// Copyright Epic Games, Inc. All Rights Reserved.
+
+#pragma once
+
+#include "Cloner/Layouts/AvaClonerLayoutBase.h"
+#include "Subsystems/EngineSubsystem.h"
+#include "UObject/WeakObjectPtrTemplates.h"
+#include "AvaClonerSubsystem.generated.h"
+
+UCLASS()
+class AVALANCHEEFFECTORS_API UAvaClonerSubsystem : public UEngineSubsystem
+{
+	GENERATED_BODY()
+
+public:
+	static inline constexpr int32 NoFlicker = 1;
+
+#if WITH_EDITOR
+	DECLARE_MULTICAST_DELEGATE(FOnCVarChanged)
+	static FOnCVarChanged OnCVarChangedDelegate;
+#endif
+
+	/** Get this subsystem instance */
+	static UAvaClonerSubsystem* Get();
+
+	//~ Begin USubsystem
+	virtual void Initialize(FSubsystemCollectionBase& Collection) override;
+	virtual void Deinitialize() override;
+	//~ End USubsystem
+
+	bool RegisterLayoutClass(const UClass* InClonerLayoutClass);
+	bool UnregisterLayoutClass(const UClass* InClonerLayoutClass);
+	bool IsLayoutClassRegistered(const UClass* InClonerLayoutClass);
+
+	/** Get available cloner layout names to use in dropdown */
+	TArray<FName> GetLayoutNames() const;
+
+	/** Based on a layout class, find layout name */
+	FName FindLayoutName(TSubclassOf<UAvaClonerLayoutBase> InLayoutClass) const;
+
+	/** Creates a new layout instance for a cloner actor */
+	UAvaClonerLayoutBase* CreateNewLayout(FName InLayoutName, AAvaClonerActor* InClonerActor);
+
+#if WITH_EDITOR
+	void EnableNoFlicker();
+	void DisableNoFlicker();
+	bool IsNoFlickerEnabled() const;
+
+	void OnTSRShadingRejectionFlickeringPeriodChanged(IConsoleVariable* InCVar) const;
+#endif
+
+protected:
+	void ScanForRegistrableClasses();
+
+	/** Linking name to the layout class */
+	UPROPERTY()
+	TMap<FName, TSubclassOf<UAvaClonerLayoutBase>> LayoutClasses;
+
+#if WITH_EDITOR
+	/** Allows to reduce ghosting artifacts when moving clone instances */
+	IConsoleVariable* CVarTSRShadingRejectionFlickeringPeriod = nullptr;
+
+	/** Previous value to restore it when disabled */
+	TOptional<int32> PreviousCVarValue;
+#endif
+};
\ No newline at end of file
diff --git a/Engine/Plugins/Experimental/Avalanche/Source/AvalancheEffectors/Public/Subsystems/AvaEffectorSubsystem.h b/Engine/Plugins/Experimental/Avalanche/Source/AvalancheEffectors/Public/Subsystems/AvaEffectorSubsystem.h
new file mode 100644
--- /dev/null
+++ b/Engine/Plugins/Experimental/Avalanche/Source/AvalancheEffectors/Public/Subsystems/AvaEffectorSubsystem.h
@@ -0,0 +1,53 @@
+// Copyright Epic Games, Inc. All Rights Reserved.
+
+#pragma once
+
+#include "Engine/World.h"
+#include "Subsystems/WorldSubsystem.h"
+#include "UObject/WeakObjectPtrTemplates.h"
+#include "AvaEffectorSubsystem.generated.h"
+
+class AAvaEffectorActor;
+class UNiagaraDataChannelAsset;
+
+UCLASS()
+class AVALANCHEEFFECTORS_API UAvaEffectorSubsystem : public UTickableWorldSubsystem
+{
+	GENERATED_BODY()
+
+public:
+	DECLARE_MULTICAST_DELEGATE_OneParam(FOnSubsystemInitialized, const UWorld*)
+	static FOnSubsystemInitialized OnSubsystemInitializedDelegate;
+
+	/** Get this subsystem instance */
+	static UAvaEffectorSubsystem* Get(const UWorld* InWorld = GWorld);
+
+	//~ Begin USubsystem
+	virtual void Initialize(FSubsystemCollectionBase& Collection) override;
+	virtual void PostInitialize() override;
+	//~ End USubsystem
+
+	/** Registers an effector actor to use it within a effector channel */
+	bool RegisterChannelEffector(AAvaEffectorActor* InEffector);
+
+	/** Unregister an effector actor used within a effector channel */
+	bool UnregisterChannelEffector(AAvaEffectorActor* InEffector);
+
+protected:
+	//~ Begin FTickableGameObject interface	
+	virtual bool IsTickableInEditor() const override;
+	virtual TStatId GetStatId() const override;
+	virtual void Tick(float InDeltaTime) override;
+	//~ End FTickableGameObject interface
+
+	/** Updates all registered effectors */
+	void UpdateEffectorChannel();
+
+	/** Ordered effectors included in this channel */
+	UPROPERTY()
+	TArray<TWeakObjectPtr<AAvaEffectorActor>> EffectorsWeak;
+
+	/** This represents the data channel structure for effector */
+	UPROPERTY()
+	TObjectPtr<UNiagaraDataChannelAsset> EffectorDataChannelAsset;
+};
\ No newline at end of file
diff --git a/Engine/Plugins/Experimental/Avalanche/Source/AvalancheEffectorsEditor/Private/AvaEffectorsEditorCommands.cpp b/Engine/Plugins/Experimental/Avalanche/Source/AvalancheEffectorsEditor/Private/AvaEffectorsEditorCommands.cpp
new file mode 100644
--- /dev/null
+++ b/Engine/Plugins/Experimental/Avalanche/Source/AvalancheEffectorsEditor/Private/AvaEffectorsEditorCommands.cpp
@@ -0,0 +1,33 @@
+// Copyright Epic Games, Inc. All Rights Reserved.
+
+#include "AvaEffectorsEditorCommands.h"
+#include "AvaEffectorsEditorStyle.h"
+
+#define LOCTEXT_NAMESPACE "AvaEffectorsEditorCommands"
+
+FAvaEffectorsEditorCommands::FAvaEffectorsEditorCommands()
+	: TCommands<FAvaEffectorsEditorCommands>(
+		TEXT("AvalancheEffectorsEditor")
+		, LOCTEXT("AvalancheEffectorsEditor", "Motion Design Effects Editor")
+		, NAME_None
+		, FAvaEffectorsEditorStyle::Get().GetStyleSetName()
+	)
+{
+}
+
+void FAvaEffectorsEditorCommands::RegisterCommands()
+{
+	UI_COMMAND(Tool_Actor_Effector
+		, "Effector Actor"
+		, "Create an Effector Actor in the viewport."
+		, EUserInterfaceActionType::ToggleButton
+		, FInputChord());
+
+	UI_COMMAND(Tool_Actor_Cloner
+		, "Cloner Actor"
+		, "Create a Cloner Actor in the viewport."
+		, EUserInterfaceActionType::ToggleButton
+		, FInputChord());
+}
+
+#undef LOCTEXT_NAMESPACE
diff --git a/Engine/Plugins/Experimental/Avalanche/Source/AvalancheEffectorsEditor/Private/AvaEffectorsEditorStyle.cpp b/Engine/Plugins/Experimental/Avalanche/Source/AvalancheEffectorsEditor/Private/AvaEffectorsEditorStyle.cpp
new file mode 100644
--- /dev/null
+++ b/Engine/Plugins/Experimental/Avalanche/Source/AvalancheEffectorsEditor/Private/AvaEffectorsEditorStyle.cpp
@@ -0,0 +1,49 @@
+// Copyright Epic Games, Inc. All Rights Reserved.
+
+#include "AvaEffectorsEditorStyle.h"
+#include "AvaClonerEffectorShared.h"
+#include "Interfaces/IPluginManager.h"
+#include "Misc/Paths.h"
+#include "Styling/SlateStyleMacros.h"
+#include "Styling/StyleColors.h"
+#include "Styling/SlateTypes.h"
+#include "Styling/SlateStyleRegistry.h"
+
+FAvaEffectorsEditorStyle::FAvaEffectorsEditorStyle()
+	: FSlateStyleSet(TEXT("AvalancheEffectorsEditor"))
+{
+	const FVector2f Icon20x20(20.0f, 20.0f);
+	const FVector2f Icon32x32(32.0f, 32.0f);
+	
+	const TSharedPtr<IPlugin> Plugin = IPluginManager::Get().FindPlugin(TEXT("Avalanche"));
+	
+	check(Plugin.IsValid());
+	
+	ContentRootDir = FPaths::Combine(Plugin->GetBaseDir(), TEXT("Resources"));
+
+	Set("AvalancheEffectorsEditor.Tool_Actor_Effector", new IMAGE_BRUSH_SVG("Icons/ToolboxIcons/effector", Icon20x20));
+	Set("AvalancheEffectorsEditor.Tool_Actor_Cloner", new IMAGE_BRUSH_SVG("Icons/ToolboxIcons/cloner", Icon20x20));
+	Set("Tool_Actor_Effector", new IMAGE_BRUSH_SVG("Icons/ToolboxIcons/effector", Icon20x20));
+	Set("Tool_Actor_Cloner", new IMAGE_BRUSH_SVG("Icons/ToolboxIcons/cloner", Icon20x20));
+
+	// Easing
+	if (const UEnum* EasingEnum = StaticEnum<EAvaClonerEasing>())
+	{
+		for (int32 Idx = 0; Idx < EasingEnum->GetMaxEnumValue(); Idx++)
+		{
+			FText ValueText;
+			EasingEnum->GetDisplayValueAsText(static_cast<EAvaClonerEasing>(Idx), ValueText);
+			const FString EasingString = ValueText.ToString().Replace(TEXT(" "), TEXT(""));
+			const FName EasingName(TEXT("AvalancheIcons.Easing.") + EasingString);
+			
+			Set(EasingName, new IMAGE_BRUSH_SVG("Icons/ClonerIcons/" + EasingString, Icon32x32));
+		}
+	}
+	
+	FSlateStyleRegistry::RegisterSlateStyle(*this);
+}
+
+FAvaEffectorsEditorStyle::~FAvaEffectorsEditorStyle()
+{
+	FSlateStyleRegistry::UnRegisterSlateStyle(*this);
+}
\ No newline at end of file
diff --git a/Engine/Plugins/Experimental/Avalanche/Source/AvalancheEffectorsEditor/Private/AvaEffectorsEditorStyle.h b/Engine/Plugins/Experimental/Avalanche/Source/AvalancheEffectorsEditor/Private/AvaEffectorsEditorStyle.h
new file mode 100644
--- /dev/null
+++ b/Engine/Plugins/Experimental/Avalanche/Source/AvalancheEffectorsEditor/Private/AvaEffectorsEditorStyle.h
@@ -0,0 +1,19 @@
+// Copyright Epic Games, Inc. All Rights Reserved.
+
+#pragma once
+
+#include "Styling/SlateStyle.h"
+
+class FAvaEffectorsEditorStyle : public FSlateStyleSet
+{
+public:
+	FAvaEffectorsEditorStyle();
+	
+	virtual ~FAvaEffectorsEditorStyle() override;
+	
+	static FAvaEffectorsEditorStyle& Get()
+	{
+		static FAvaEffectorsEditorStyle StyleSet;
+		return StyleSet;
+	}
+};
\ No newline at end of file
diff --git a/Engine/Plugins/Experimental/Avalanche/Source/AvalancheEffectorsEditor/Private/AvalancheEffectorsEditorModule.cpp b/Engine/Plugins/Experimental/Avalanche/Source/AvalancheEffectorsEditor/Private/AvalancheEffectorsEditorModule.cpp
new file mode 100644
--- /dev/null
+++ b/Engine/Plugins/Experimental/Avalanche/Source/AvalancheEffectorsEditor/Private/AvalancheEffectorsEditorModule.cpp
@@ -0,0 +1,111 @@
+// Copyright Epic Games, Inc. All Rights Reserved.
+
+#include "AvalancheEffectorsEditorModule.h"
+#include "AvaEffectorsEditorCommands.h"
+#include "AvaEffectorsEditorStyle.h"
+#include "AvaInteractiveToolsDelegates.h"
+#include "AvalancheShapesEditorModule.h"
+#include "Cloner/AvaClonerActor.h"
+#include "Cloner/AvaClonerActorTool.h"
+#include "Cloner/AvaClonerActorVis.h"
+#include "Cloner/AvaClonerComponent.h"
+#include "Cloner/AvaClonerDetailCustomization.h"
+#include "ComponentVisualizers.h"
+#include "Effector/AvaEffectorActorTool.h"
+#include "Effector/AvaEffectorActorVis.h"
+#include "Effector/AvaEffectorComponent.h"
+#include "Effector/AvaEffectorDetailCustomization.h"
+#include "Framework/Application/SlateApplication.h"
+#include "IAvalancheComponentVisualizersModule.h"
+#include "Modules/ModuleManager.h"
+#include "PropertyEditorModule.h"
+
+#define LOCTEXT_NAMESPACE "AvalancheEffectorsEditorModule"
+
+FAvalancheEffectorsEditorModule& FAvalancheEffectorsEditorModule::Get()
+{
+	static const FName ModuleName = TEXT("AvalancheEffectorsEditor");
+	return FModuleManager::LoadModuleChecked<FAvalancheEffectorsEditorModule>(ModuleName);
+}
+
+void FAvalancheEffectorsEditorModule::StartupModule()
+{
+	FCoreDelegates::OnPostEngineInit.AddRaw(this, &FAvalancheEffectorsEditorModule::PostEngineInit);
+
+	RegisterCustomLayouts();
+	FAvaEffectorsEditorStyle::Get();
+	FAvaEffectorsEditorCommands::Register();
+
+	FAvaInteractiveToolsDelegates::GetRegisterToolsDelegate().AddRaw(this, &FAvalancheEffectorsEditorModule::RegisterTools);
+}
+
+void FAvalancheEffectorsEditorModule::ShutdownModule()
+{
+	FCoreDelegates::OnPostEngineInit.RemoveAll(this);
+
+	if (UObjectInitialized() && !IsEngineExitRequested())
+	{
+		UnregisterCustomLayouts();
+	}
+
+	FAvaEffectorsEditorCommands::Unregister();
+	
+	FAvaInteractiveToolsDelegates::GetRegisterToolsDelegate().RemoveAll(this);
+}
+
+void FAvalancheEffectorsEditorModule::PostEngineInit()
+{
+	if (FSlateApplication::IsInitialized())
+	{
+		RegisterComponentVisualizers();
+	}
+}
+
+void FAvalancheEffectorsEditorModule::RegisterTools(IAvalancheInteractiveToolsModule* InModule)
+{
+	InModule->RegisterTool(
+		IAvalancheInteractiveToolsModule::Get().CategoryNameActor,
+		GetDefault<UAvaEffectorActorTool>()->GetToolParameters()
+	);
+
+	InModule->RegisterTool(
+		IAvalancheInteractiveToolsModule::Get().CategoryNameActor,
+		GetDefault<UAvaClonerActorTool>()->GetToolParameters()
+	);
+}
+
+void FAvalancheEffectorsEditorModule::RegisterComponentVisualizers()
+{
+	IAvalancheComponentVisualizersModule::RegisterComponentVisualizer<UAvaEffectorComponent, FAvaEffectorActorVisualizer>(&Visualizers);
+	IAvalancheComponentVisualizersModule::RegisterComponentVisualizer<UAvaClonerComponent, FAvaClonerActorVisualizer>(&Visualizers);
+}
+
+namespace UE::AvalancheEffectorsEditor::Private
+{
+	static const FName PropertyEditorModuleName("PropertyEditor");
+}
+
+void FAvalancheEffectorsEditorModule::RegisterCustomLayouts()
+{
+	FPropertyEditorModule& PropertyModule = FModuleManager::GetModuleChecked<FPropertyEditorModule>(UE::AvalancheEffectorsEditor::Private::PropertyEditorModuleName);
+
+	// Cloner/effector
+	PropertyModule.RegisterCustomClassLayout(AAvaEffectorActor::StaticClass()->GetFName(), FOnGetDetailCustomizationInstance::CreateStatic(&FAvaEffectorDetailCustomization::MakeInstance));
+	PropertyModule.RegisterCustomClassLayout(AAvaClonerActor::StaticClass()->GetFName(), FOnGetDetailCustomizationInstance::CreateStatic(&FAvaClonerDetailCustomization::MakeInstance));
+}
+
+void FAvalancheEffectorsEditorModule::UnregisterCustomLayouts()
+{
+	if (FModuleManager::Get().IsModuleLoaded(UE::AvalancheEffectorsEditor::Private::PropertyEditorModuleName))
+	{
+		FPropertyEditorModule& PropertyModule = FModuleManager::GetModuleChecked<FPropertyEditorModule>(UE::AvalancheEffectorsEditor::Private::PropertyEditorModuleName);
+
+		// Cloner/effector
+		PropertyModule.UnregisterCustomClassLayout(AAvaEffectorActor::StaticClass()->GetFName());
+		PropertyModule.UnregisterCustomClassLayout(AAvaClonerActor::StaticClass()->GetFName());
+	}
+}
+
+#undef LOCTEXT_NAMESPACE
+
+IMPLEMENT_MODULE(FAvalancheEffectorsEditorModule, AvalancheEffectorsEditorModule)
diff --git a/Engine/Plugins/Experimental/Avalanche/Source/AvalancheEffectorsEditor/Private/Cloner/AvaClonerActorTool.cpp b/Engine/Plugins/Experimental/Avalanche/Source/AvalancheEffectorsEditor/Private/Cloner/AvaClonerActorTool.cpp
new file mode 100644
--- /dev/null
+++ b/Engine/Plugins/Experimental/Avalanche/Source/AvalancheEffectorsEditor/Private/Cloner/AvaClonerActorTool.cpp
@@ -0,0 +1,37 @@
+// Copyright Epic Games, Inc. All Rights Reserved.
+
+#include "AvaClonerActorTool.h"
+#include "AvaEffectorsEditorCommands.h"
+#include "AvaInteractiveToolsSettings.h"
+#include "Builders/AvaInteractiveToolsToolBuilder.h"
+#include "Cloner/AvaClonerActor.h"
+
+UAvaClonerActorTool::UAvaClonerActorTool()
+{
+	ActorClass = AAvaClonerActor::StaticClass();
+}
+
+bool UAvaClonerActorTool::UseIdentityRotation() const
+{
+	return ConditionalIdentityRotation();
+}
+
+FName UAvaClonerActorTool::GetCategoryName()
+{
+	return IAvalancheInteractiveToolsModule::CategoryNameActor;
+}
+
+FAvaInteractiveToolsToolParameters UAvaClonerActorTool::GetToolParameters() const
+{
+	return {
+		FAvaEffectorsEditorCommands::Get().Tool_Actor_Cloner,
+		TEXT("Clone Actor Tool"),
+		3000,
+		FAvalancheInteractiveToolsCreateBuilder::CreateLambda(
+			[](UEdMode* InEdMode)
+			{
+				return UAvaInteractiveToolsToolBuilder::CreateToolBuilder<UAvaClonerActorTool>(InEdMode);
+			}),
+		ActorClass
+	};
+}
diff --git a/Engine/Plugins/Experimental/Avalanche/Source/AvalancheEffectorsEditor/Private/Cloner/AvaClonerActorTool.h b/Engine/Plugins/Experimental/Avalanche/Source/AvalancheEffectorsEditor/Private/Cloner/AvaClonerActorTool.h
new file mode 100644
--- /dev/null
+++ b/Engine/Plugins/Experimental/Avalanche/Source/AvalancheEffectorsEditor/Private/Cloner/AvaClonerActorTool.h
@@ -0,0 +1,22 @@
+// Copyright Epic Games, Inc. All Rights Reserved.
+
+#pragma once
+
+#include "Tools/AvaInteractiveToolsActorPointToolBase.h"
+#include "AvaClonerActorTool.generated.h"
+
+UCLASS()
+class UAvaClonerActorTool : public UAvaInteractiveToolsActorPointToolBase
+{
+	GENERATED_BODY()
+
+public:
+	UAvaClonerActorTool();
+
+	//~ Begin UAvaInteractiveToolsToolBase
+	virtual bool UseIdentityLocation() const override { return false; }
+	virtual bool UseIdentityRotation() const override;
+	virtual FName GetCategoryName() override;
+	virtual FAvaInteractiveToolsToolParameters GetToolParameters() const override;
+	//~ End UAvaInteractiveToolsToolBase
+};
diff --git a/Engine/Plugins/Experimental/Avalanche/Source/AvalancheEffectorsEditor/Private/Cloner/AvaClonerActorVis.cpp b/Engine/Plugins/Experimental/Avalanche/Source/AvalancheEffectorsEditor/Private/Cloner/AvaClonerActorVis.cpp
new file mode 100644
--- /dev/null
+++ b/Engine/Plugins/Experimental/Avalanche/Source/AvalancheEffectorsEditor/Private/Cloner/AvaClonerActorVis.cpp
@@ -0,0 +1,956 @@
+// Copyright Epic Games, Inc. All Rights Reserved. 
+
+#include "Cloner/AvaClonerActorVis.h"
+
+#include "AvaField.h"
+#include "AvalancheShapesEditorModule.h"
+#include "Cloner/AvaClonerActor.h"
+#include "Cloner/AvaClonerComponent.h"
+#include "Cloner/Layouts/AvaClonerCircleLayout.h"
+#include "Cloner/Layouts/AvaClonerCylinderLayout.h"
+#include "Cloner/Layouts/AvaClonerGridLayout.h"
+#include "Cloner/Layouts/AvaClonerHoneycombLayout.h"
+#include "Cloner/Layouts/AvaClonerLineLayout.h"
+#include "Cloner/Layouts/AvaClonerSphereUniformLayout.h"
+#include "EditorViewportClient.h"
+#include "Engine/Texture2D.h"
+#include "Framework/Application/SlateApplication.h"
+#include "IAvalancheComponentVisualizersModule.h"
+#include "ScopedTransaction.h"
+#include "TextureResource.h"
+
+IMPLEMENT_HIT_PROXY(HAvaClonerActorSpacingHitProxy, HAvaHitProxy);
+
+#define LOCTEXT_NAMESPACE "AvaClonerActorVisualizer"
+
+FAvaClonerActorVisualizer::FAvaClonerActorVisualizer()
+	: FAvaVisualizerBase()
+{
+}
+
+void FAvaClonerActorVisualizer::StoreInitialValues()
+{
+	Super::StoreInitialValues();
+	
+	if (GetEditedComponent() == nullptr)
+	{
+		return;
+	}
+
+	const AAvaClonerActor* ClonerActor = ClonerActorWeak.Get();
+
+	if (!ClonerActor)
+	{
+		return;
+	}
+
+	if (const UAvaClonerGridLayout* GridLayout = ClonerActor->GetActiveLayout<UAvaClonerGridLayout>())
+	{
+		InitialSpacing = FVector(GridLayout->GetSpacingX(), GridLayout->GetSpacingY(), GridLayout->GetSpacingZ());
+	}
+	else if (const UAvaClonerLineLayout* LineLayout = ClonerActor->GetActiveLayout<UAvaClonerLineLayout>())
+	{
+		InitialSpacing = FVector(LineLayout->GetSpacing());
+	}
+	else if (const UAvaClonerHoneycombLayout* HoneycombLayout = ClonerActor->GetActiveLayout<UAvaClonerHoneycombLayout>())
+	{
+		const EAvaClonerPlane Plane = HoneycombLayout->GetPlane();
+		
+		FVector Spacing;
+		if (Plane == EAvaClonerPlane::XY)
+		{
+			Spacing = FVector(HoneycombLayout->GetWidthSpacing(), HoneycombLayout->GetHeightSpacing(), 0);
+		}
+		else if (Plane == EAvaClonerPlane::YZ)
+		{
+			Spacing = FVector(0, HoneycombLayout->GetWidthSpacing(), HoneycombLayout->GetHeightSpacing());
+		}
+		else if (Plane == EAvaClonerPlane::XZ)
+		{
+			Spacing = FVector(HoneycombLayout->GetWidthSpacing(), 0, HoneycombLayout->GetHeightSpacing());
+		}
+		
+		InitialSpacing = Spacing;
+	}
+	else if (const UAvaClonerCircleLayout* CircleLayout = ClonerActor->GetActiveLayout<UAvaClonerCircleLayout>())
+	{
+		InitialSpacing = FVector(CircleLayout->GetRadius());
+	}
+	else if (const UAvaClonerCylinderLayout* CylinderLayout = ClonerActor->GetActiveLayout<UAvaClonerCylinderLayout>())
+	{
+		InitialSpacing = FVector(0, CylinderLayout->GetRadius(), CylinderLayout->GetHeight());
+	}
+	else if (const UAvaClonerSphereUniformLayout* SphereLayout = ClonerActor->GetActiveLayout<UAvaClonerSphereUniformLayout>())
+	{
+		InitialSpacing = FVector(0, SphereLayout->GetRadius(), 0);
+	}
+}
+
+FBox FAvaClonerActorVisualizer::GetComponentBounds(const UActorComponent* InComponent) const
+{
+	if (const UAvaClonerComponent* ClonerComponent = Cast<UAvaClonerComponent>(InComponent))
+	{
+		if (const AAvaClonerActor* ClonerActor = Cast<AAvaClonerActor>(ClonerComponent->GetOwner()))
+		{
+			FVector Origin;
+			FVector Extent;
+			ClonerActor->GetActorBounds(false, Origin, Extent);
+			return FBox(-Extent, Extent);
+		}
+	}
+
+	return Super::GetComponentBounds(InComponent);
+}
+
+bool FAvaClonerActorVisualizer::HandleInputDeltaInternal(FEditorViewportClient* InViewportClient, FViewport* InViewport, const FVector& InAccumulatedTranslation, const FRotator& InAccumulatedRotation, const FVector& InAccumulatedScale)
+{
+	if (!FSlateApplication::Get().GetPressedMouseButtons().Contains(EKeys::LeftMouseButton))
+	{
+		return false;
+	}
+
+	const AAvaClonerActor* ClonerActor = ClonerActorWeak.Get();
+
+	if (!ClonerActor)
+	{
+		EndEditing();
+	}
+	else if (GetViewportWidgetMode(InViewportClient) == UE::Widget::WM_Translate)
+	{
+		const EAxisList::Type AxisList = GetViewportWidgetAxisList(InViewportClient);
+			
+		if (UAvaClonerGridLayout* GridLayout = ClonerActor->GetActiveLayout<UAvaClonerGridLayout>())
+		{
+			if (AxisList & EAxisList::X)
+			{
+				const float SpacingX = InitialSpacing.X + InAccumulatedTranslation.X / FMath::Max(1, GridLayout->GetCountX() / 2.f);
+				GridLayout->SetSpacingX(SpacingX);
+				OnPropertyModified(GridLayout, GET_MEMBER_NAME_CHECKED(UAvaClonerGridLayout, SpacingX));
+			}
+			else if (AxisList & EAxisList::Y)
+			{
+				const float SpacingY = InitialSpacing.Y + InAccumulatedTranslation.Y / FMath::Max(1, GridLayout->GetCountY() / 2.f);
+				GridLayout->SetSpacingY(SpacingY);
+				OnPropertyModified(GridLayout, GET_MEMBER_NAME_CHECKED(UAvaClonerGridLayout, SpacingY));
+			}
+			else if (AxisList & EAxisList::Z)
+			{
+				const float SpacingZ = InitialSpacing.Z + InAccumulatedTranslation.Z / FMath::Max(1, GridLayout->GetCountZ() / 2.f);
+				GridLayout->SetSpacingZ(SpacingZ);
+				OnPropertyModified(GridLayout, GET_MEMBER_NAME_CHECKED(UAvaClonerGridLayout, SpacingZ));
+			}
+			
+			return true;
+		}
+		else if (UAvaClonerLineLayout* LineLayout = ClonerActor->GetActiveLayout<UAvaClonerLineLayout>())
+		{
+			const int32 Count = LineLayout->GetCount();
+			float Spacing = LineLayout->GetSpacing();
+            				
+			if (AxisList & EAxisList::X)
+			{
+				Spacing = InitialSpacing.X + InAccumulatedTranslation.X / FMath::Max(1, Count);
+			}
+			if (AxisList & EAxisList::Y)
+			{
+				Spacing = InitialSpacing.Y + InAccumulatedTranslation.Y / FMath::Max(1, Count);
+			}
+			if (AxisList & EAxisList::Z)
+			{
+				Spacing = InitialSpacing.Z + InAccumulatedTranslation.Z / FMath::Max(1, Count);
+			}
+			
+			LineLayout->SetSpacing(Spacing);
+			OnPropertyModified(LineLayout, GET_MEMBER_NAME_CHECKED(UAvaClonerLineLayout, Spacing));
+
+			return true;
+		}
+		else if (UAvaClonerHoneycombLayout* HoneycombLayout = ClonerActor->GetActiveLayout<UAvaClonerHoneycombLayout>())
+		{
+			const EAvaClonerPlane Plane = HoneycombLayout->GetPlane();
+			const int32 WidthCount = HoneycombLayout->GetWidthCount();
+			const int32 HeightCount = HoneycombLayout->GetHeightCount();
+			float WidthSpacing = HoneycombLayout->GetWidthSpacing();
+			float HeightSpacing = HoneycombLayout->GetHeightSpacing();
+			
+			if (Plane == EAvaClonerPlane::XY)
+			{
+				WidthSpacing = InitialSpacing.X + InAccumulatedTranslation.X / FMath::Max(1, WidthCount / 2.f);
+				HeightSpacing = InitialSpacing.Y + InAccumulatedTranslation.Y / FMath::Max(1, HeightCount / 2.f);
+			}
+			else if (Plane == EAvaClonerPlane::YZ)
+			{
+				WidthSpacing = InitialSpacing.Y + InAccumulatedTranslation.Y / FMath::Max(1, WidthCount / 2.f);
+				HeightSpacing = InitialSpacing.Z + InAccumulatedTranslation.Z / FMath::Max(1, HeightCount / 2.f);
+			}
+			else if (Plane == EAvaClonerPlane::XZ)
+			{
+				WidthSpacing = InitialSpacing.X + InAccumulatedTranslation.X / FMath::Max(1, WidthCount / 2.f);
+				HeightSpacing = InitialSpacing.Z + InAccumulatedTranslation.Z / FMath::Max(1, HeightCount / 2.f);
+			}
+			
+			HoneycombLayout->SetWidthSpacing(WidthSpacing);
+			HoneycombLayout->SetHeightSpacing(HeightSpacing);
+
+			OnPropertyModified(HoneycombLayout, GET_MEMBER_NAME_CHECKED(UAvaClonerHoneycombLayout, WidthSpacing));
+			OnPropertyModified(HoneycombLayout, GET_MEMBER_NAME_CHECKED(UAvaClonerHoneycombLayout, HeightSpacing));
+
+			return true;
+		}
+		else if (UAvaClonerCircleLayout* CircleLayout = ClonerActor->GetActiveLayout<UAvaClonerCircleLayout>())
+		{
+			float Radius = CircleLayout->GetRadius();
+			
+			if (AxisList & EAxisList::X)
+			{
+				Radius = InitialSpacing.X + InAccumulatedTranslation.X;
+			}
+			else if (AxisList & EAxisList::Y)
+			{
+				Radius = InitialSpacing.Y + InAccumulatedTranslation.Y;
+			}
+			else if (AxisList & EAxisList::Z)
+			{
+				Radius = InitialSpacing.Z + InAccumulatedTranslation.Z;
+			}
+			
+			CircleLayout->SetRadius(Radius);
+			OnPropertyModified(CircleLayout, GET_MEMBER_NAME_CHECKED(UAvaClonerCircleLayout, Radius));
+
+			return true;
+		}
+		else if (UAvaClonerCylinderLayout* CylinderLayout = ClonerActor->GetActiveLayout<UAvaClonerCylinderLayout>())
+		{
+			const EAvaClonerPlane Plane = CylinderLayout->GetPlane();
+			float Radius = CylinderLayout->GetRadius();
+			float Height = CylinderLayout->GetHeight();
+			
+			if (Plane == EAvaClonerPlane::XY)
+			{
+				if (AxisList & EAxisList::Y)
+				{
+					Radius = InitialSpacing.Y + InAccumulatedTranslation.Y;
+				}
+				else
+				{
+					Height = InitialSpacing.Z + InAccumulatedTranslation.Z;
+				}
+			}
+			else if (Plane == EAvaClonerPlane::YZ)
+			{
+				if (AxisList & EAxisList::Y)
+				{
+					Radius = InitialSpacing.Y + InAccumulatedTranslation.Y;
+				}
+				else
+				{
+					Height = InitialSpacing.Z + InAccumulatedTranslation.X;
+				}
+			}
+			else if (Plane == EAvaClonerPlane::XZ)
+			{
+				if (AxisList & EAxisList::Z)
+				{
+					Radius = InitialSpacing.Y + InAccumulatedTranslation.Z;
+				}
+				else
+				{
+					Height = InitialSpacing.Z + InAccumulatedTranslation.Y;
+				}
+			}
+
+			CylinderLayout->SetRadius(Radius);
+			CylinderLayout->SetHeight(Height);
+
+			OnPropertyModified(CylinderLayout, GET_MEMBER_NAME_CHECKED(UAvaClonerCylinderLayout, Radius));
+			OnPropertyModified(CylinderLayout, GET_MEMBER_NAME_CHECKED(UAvaClonerCylinderLayout, Height));
+
+			return true;
+		}
+		else if (UAvaClonerSphereUniformLayout* SphereLayout = ClonerActor->GetActiveLayout<UAvaClonerSphereUniformLayout>())
+		{
+			float Radius = SphereLayout->GetRadius();
+			
+			if (AxisList & EAxisList::Y)
+			{
+				Radius = InitialSpacing.Y + InAccumulatedTranslation.Y;
+			}
+			
+			SphereLayout->SetRadius(Radius);
+			OnPropertyModified(SphereLayout, GET_MEMBER_NAME_CHECKED(UAvaClonerSphereUniformLayout, Radius));
+
+			return true;
+		}
+	}
+	
+	return Super::HandleInputDeltaInternal(InViewportClient, InViewport, InAccumulatedTranslation, InAccumulatedRotation, InAccumulatedScale);
+}
+
+void FAvaClonerActorVisualizer::DrawVisualizationEditing(const UActorComponent* InComponent, const FSceneView* InView, FPrimitiveDrawInterface* InPDI, int32& InOutIconIndex)
+{
+	Super::DrawVisualizationEditing(InComponent, InView, InPDI, InOutIconIndex);
+
+	const UAvaClonerComponent* ClonerComponent = Cast<UAvaClonerComponent>(InComponent);
+
+	if (!ClonerComponent)
+	{
+		return;
+	}
+
+	const AAvaClonerActor* ClonerActor = Cast<AAvaClonerActor>(ClonerComponent->GetOwner());
+
+	if (!ClonerActor || !ClonerActor->GetEnabled() || ClonerActor->GetMeshCount() == 0)
+	{
+		return;
+	}
+
+	if (const UAvaClonerGridLayout* GridLayout = ClonerActor->GetActiveLayout<UAvaClonerGridLayout>())
+	{
+		if (GridLayout->GetCountX() > 0)
+		{
+			DrawSpacingButton(ClonerActor, InView, InPDI, InOutIconIndex, EAvaClonerAxis::X, FAvaVisualizerBase::Inactive);
+			InOutIconIndex++;
+		}
+		
+		if (GridLayout->GetCountY() > 0)
+		{
+			DrawSpacingButton(ClonerActor, InView, InPDI, InOutIconIndex, EAvaClonerAxis::Y, FAvaVisualizerBase::Inactive);
+			InOutIconIndex++;
+		}
+		
+		if (GridLayout->GetCountZ() > 0)
+		{
+			DrawSpacingButton(ClonerActor, InView, InPDI, InOutIconIndex, EAvaClonerAxis::Z, FAvaVisualizerBase::Inactive);
+			InOutIconIndex++;
+		}
+	}
+	else if (const UAvaClonerLineLayout* LineLayout = ClonerActor->GetActiveLayout<UAvaClonerLineLayout>())
+	{
+		if (LineLayout->GetCount() > 0)
+		{
+			DrawSpacingButton(ClonerActor, InView, InPDI, InOutIconIndex, LineLayout->GetAxis(), FAvaVisualizerBase::Inactive);
+			InOutIconIndex++;
+		}
+	}
+	else if (const UAvaClonerHoneycombLayout* HoneycombLayout = ClonerActor->GetActiveLayout<UAvaClonerHoneycombLayout>())
+	{
+		const EAvaClonerPlane Plane = HoneycombLayout->GetPlane();
+		
+		if (Plane == EAvaClonerPlane::XY)
+		{
+			DrawSpacingButton(ClonerActor, InView, InPDI, InOutIconIndex, EAvaClonerAxis::X, FAvaVisualizerBase::Inactive);
+			InOutIconIndex++;
+			DrawSpacingButton(ClonerActor, InView, InPDI, InOutIconIndex, EAvaClonerAxis::Y, FAvaVisualizerBase::Inactive);
+			InOutIconIndex++;
+		}
+		else if (Plane == EAvaClonerPlane::YZ)
+		{
+			DrawSpacingButton(ClonerActor, InView, InPDI, InOutIconIndex, EAvaClonerAxis::Y, FAvaVisualizerBase::Inactive);
+			InOutIconIndex++;
+			DrawSpacingButton(ClonerActor, InView, InPDI, InOutIconIndex, EAvaClonerAxis::Z, FAvaVisualizerBase::Inactive);
+			InOutIconIndex++;
+		}
+		else if (Plane == EAvaClonerPlane::XZ)
+		{
+			DrawSpacingButton(ClonerActor, InView, InPDI, InOutIconIndex, EAvaClonerAxis::X, FAvaVisualizerBase::Inactive);
+			InOutIconIndex++;
+			DrawSpacingButton(ClonerActor, InView, InPDI, InOutIconIndex, EAvaClonerAxis::Z, FAvaVisualizerBase::Inactive);
+			InOutIconIndex++;
+		}
+	}
+	else if (const UAvaClonerCircleLayout* CircleLayout = ClonerActor->GetActiveLayout<UAvaClonerCircleLayout>())
+	{
+		const EAvaClonerPlane Plane = CircleLayout->GetPlane();
+
+		if (Plane == EAvaClonerPlane::XY)
+		{
+			DrawSpacingButton(ClonerActor, InView, InPDI, InOutIconIndex, EAvaClonerAxis::Y, FAvaVisualizerBase::Inactive);
+			InOutIconIndex++;
+		}
+		else if (Plane == EAvaClonerPlane::YZ)
+		{
+			DrawSpacingButton(ClonerActor, InView, InPDI, InOutIconIndex, EAvaClonerAxis::Y, FAvaVisualizerBase::Inactive);
+			InOutIconIndex++;
+		}
+		else if (Plane == EAvaClonerPlane::XZ)
+		{
+			DrawSpacingButton(ClonerActor, InView, InPDI, InOutIconIndex, EAvaClonerAxis::Z, FAvaVisualizerBase::Inactive);
+			InOutIconIndex++;
+		}
+	}
+	else if (const UAvaClonerCylinderLayout* CylinderLayout = ClonerActor->GetActiveLayout<UAvaClonerCylinderLayout>())
+	{
+		const EAvaClonerPlane Plane = CylinderLayout->GetPlane();
+		
+		if (Plane == EAvaClonerPlane::XY)
+		{
+			DrawSpacingButton(ClonerActor, InView, InPDI, InOutIconIndex, EAvaClonerAxis::Y, FAvaVisualizerBase::Inactive);
+			InOutIconIndex++;
+			DrawSpacingButton(ClonerActor, InView, InPDI, InOutIconIndex, EAvaClonerAxis::Z, FAvaVisualizerBase::Inactive);
+			InOutIconIndex++;
+		}
+		else if (Plane == EAvaClonerPlane::YZ)
+		{
+			DrawSpacingButton(ClonerActor, InView, InPDI, InOutIconIndex, EAvaClonerAxis::Y, FAvaVisualizerBase::Inactive);
+			InOutIconIndex++;
+			DrawSpacingButton(ClonerActor, InView, InPDI, InOutIconIndex, EAvaClonerAxis::X, FAvaVisualizerBase::Inactive);
+			InOutIconIndex++;
+		}
+		else if (Plane == EAvaClonerPlane::XZ)
+		{
+			DrawSpacingButton(ClonerActor, InView, InPDI, InOutIconIndex, EAvaClonerAxis::Z, FAvaVisualizerBase::Inactive);
+			InOutIconIndex++;
+			DrawSpacingButton(ClonerActor, InView, InPDI, InOutIconIndex, EAvaClonerAxis::Y, FAvaVisualizerBase::Inactive);
+			InOutIconIndex++;
+		}
+	}
+	else if (ClonerActor->IsActiveLayout<UAvaClonerSphereUniformLayout>())
+	{
+		DrawSpacingButton(ClonerActor, InView, InPDI, InOutIconIndex, EAvaClonerAxis::Y, FAvaVisualizerBase::Inactive);
+		InOutIconIndex++;
+	}
+}
+
+void FAvaClonerActorVisualizer::DrawVisualizationNotEditing(const UActorComponent* InComponent, const FSceneView* InView, FPrimitiveDrawInterface* InPDI, int32& InOutIconIndex)
+{
+	Super::DrawVisualizationNotEditing(InComponent, InView, InPDI, InOutIconIndex);
+
+	const UAvaClonerComponent* ClonerComponent = Cast<UAvaClonerComponent>(InComponent);
+
+	if (!ClonerComponent)
+	{
+		return;
+	}
+
+	const AAvaClonerActor* ClonerActor = Cast<AAvaClonerActor>(ClonerComponent->GetOwner());
+
+	if (!ClonerActor || !ClonerActor->GetEnabled() || ClonerActor->GetMeshCount() == 0)
+	{
+		return;
+	}
+
+	if (const UAvaClonerGridLayout* GridLayout = ClonerActor->GetActiveLayout<UAvaClonerGridLayout>())
+	{
+		if (GridLayout->GetCountX() > 0)
+		{
+			DrawSpacingButton(ClonerActor, InView, InPDI, InOutIconIndex, EAvaClonerAxis::X, FAvaVisualizerBase::Inactive);
+			InOutIconIndex++;
+		}
+		
+		if (GridLayout->GetCountY() > 0)
+		{
+			DrawSpacingButton(ClonerActor, InView, InPDI, InOutIconIndex, EAvaClonerAxis::Y, FAvaVisualizerBase::Inactive);
+			InOutIconIndex++;
+		}
+		
+		if (GridLayout->GetCountZ() > 0)
+		{
+			DrawSpacingButton(ClonerActor, InView, InPDI, InOutIconIndex, EAvaClonerAxis::Z, FAvaVisualizerBase::Inactive);
+			InOutIconIndex++;
+		}
+	}
+	else if (const UAvaClonerLineLayout* LineLayout = ClonerActor->GetActiveLayout<UAvaClonerLineLayout>())
+	{
+		if (LineLayout->GetCount() > 0)
+		{
+			DrawSpacingButton(ClonerActor, InView, InPDI, InOutIconIndex, LineLayout->GetAxis(), FAvaVisualizerBase::Inactive);
+			InOutIconIndex++;
+		}
+	}
+	else if (const UAvaClonerHoneycombLayout* HoneycombLayout = ClonerActor->GetActiveLayout<UAvaClonerHoneycombLayout>())
+	{
+		const EAvaClonerPlane Plane = HoneycombLayout->GetPlane();
+		
+		if (Plane == EAvaClonerPlane::XY)
+		{
+			DrawSpacingButton(ClonerActor, InView, InPDI, InOutIconIndex, EAvaClonerAxis::X, FAvaVisualizerBase::Inactive);
+			InOutIconIndex++;
+			DrawSpacingButton(ClonerActor, InView, InPDI, InOutIconIndex, EAvaClonerAxis::Y, FAvaVisualizerBase::Inactive);
+			InOutIconIndex++;
+		}
+		else if (Plane == EAvaClonerPlane::YZ)
+		{
+			DrawSpacingButton(ClonerActor, InView, InPDI, InOutIconIndex, EAvaClonerAxis::Y, FAvaVisualizerBase::Inactive);
+			InOutIconIndex++;
+			DrawSpacingButton(ClonerActor, InView, InPDI, InOutIconIndex, EAvaClonerAxis::Z, FAvaVisualizerBase::Inactive);
+			InOutIconIndex++;
+		}
+		else if (Plane == EAvaClonerPlane::XZ)
+		{
+			DrawSpacingButton(ClonerActor, InView, InPDI, InOutIconIndex, EAvaClonerAxis::X, FAvaVisualizerBase::Inactive);
+			InOutIconIndex++;
+			DrawSpacingButton(ClonerActor, InView, InPDI, InOutIconIndex, EAvaClonerAxis::Z, FAvaVisualizerBase::Inactive);
+			InOutIconIndex++;
+		}
+	}
+	else if (const UAvaClonerCircleLayout* CircleLayout = ClonerActor->GetActiveLayout<UAvaClonerCircleLayout>())
+	{
+		const EAvaClonerPlane Plane = CircleLayout->GetPlane();
+		
+		if (Plane == EAvaClonerPlane::XY)
+		{
+			DrawSpacingButton(ClonerActor, InView, InPDI, InOutIconIndex, EAvaClonerAxis::Y, FAvaVisualizerBase::Inactive);
+			InOutIconIndex++;
+		}
+		else if (Plane == EAvaClonerPlane::YZ)
+		{
+			DrawSpacingButton(ClonerActor, InView, InPDI, InOutIconIndex, EAvaClonerAxis::Y, FAvaVisualizerBase::Inactive);
+			InOutIconIndex++;
+		}
+		else if (Plane == EAvaClonerPlane::XZ)
+		{
+			DrawSpacingButton(ClonerActor, InView, InPDI, InOutIconIndex, EAvaClonerAxis::Z, FAvaVisualizerBase::Inactive);
+			InOutIconIndex++;
+		}
+		else
+		{
+			DrawSpacingButton(ClonerActor, InView, InPDI, InOutIconIndex, EAvaClonerAxis::Custom, FAvaVisualizerBase::Inactive);
+			InOutIconIndex++;
+		}
+	}
+	else if (const UAvaClonerCylinderLayout* CylinderLayout = ClonerActor->GetActiveLayout<UAvaClonerCylinderLayout>())
+	{
+		const EAvaClonerPlane Plane = CylinderLayout->GetPlane();
+		
+		if (Plane == EAvaClonerPlane::XY)
+		{
+			DrawSpacingButton(ClonerActor, InView, InPDI, InOutIconIndex, EAvaClonerAxis::Y, FAvaVisualizerBase::Inactive);
+			InOutIconIndex++;
+			DrawSpacingButton(ClonerActor, InView, InPDI, InOutIconIndex, EAvaClonerAxis::Z, FAvaVisualizerBase::Inactive);
+			InOutIconIndex++;
+		}
+		else if (Plane == EAvaClonerPlane::YZ)
+		{
+			DrawSpacingButton(ClonerActor, InView, InPDI, InOutIconIndex, EAvaClonerAxis::Y, FAvaVisualizerBase::Inactive);
+			InOutIconIndex++;
+			DrawSpacingButton(ClonerActor, InView, InPDI, InOutIconIndex, EAvaClonerAxis::X, FAvaVisualizerBase::Inactive);
+			InOutIconIndex++;
+		}
+		else if (Plane == EAvaClonerPlane::XZ)
+		{
+			DrawSpacingButton(ClonerActor, InView, InPDI, InOutIconIndex, EAvaClonerAxis::Z, FAvaVisualizerBase::Inactive);
+			InOutIconIndex++;
+			DrawSpacingButton(ClonerActor, InView, InPDI, InOutIconIndex, EAvaClonerAxis::Y, FAvaVisualizerBase::Inactive);
+			InOutIconIndex++;
+		}
+	}
+	else if (ClonerActor->IsActiveLayout<UAvaClonerSphereUniformLayout>())
+	{
+		DrawSpacingButton(ClonerActor, InView, InPDI, InOutIconIndex, EAvaClonerAxis::Y, FAvaVisualizerBase::Inactive);
+		InOutIconIndex++;
+	}
+}
+
+FVector FAvaClonerActorVisualizer::GetHandleSpacingLocation(const AAvaClonerActor* InClonerActor, EAvaClonerAxis InAxis) const
+{
+	const FRotator ClonerRotation = InClonerActor->GetActorRotation();
+	const FVector ClonerScale = InClonerActor->GetActorScale();
+	FVector OutLocation = FVector::ZeroVector;
+	FVector Axis = InAxis == EAvaClonerAxis::X ? FVector::XAxisVector : (InAxis == EAvaClonerAxis::Z ? FVector::ZAxisVector : FVector::YAxisVector);
+	
+	if (const UAvaClonerGridLayout* GridLayout = InClonerActor->GetActiveLayout<UAvaClonerGridLayout>())
+	{
+		const FVector Spacing(GridLayout->GetSpacingX(), GridLayout->GetSpacingY(), GridLayout->GetSpacingZ());
+		const FVector Count(GridLayout->GetCountX(), GridLayout->GetCountY(), GridLayout->GetCountZ());
+
+		OutLocation = ClonerRotation.RotateVector(Axis * Spacing * Count / 2);
+	}
+	else if (const UAvaClonerLineLayout* LineLayout = InClonerActor->GetActiveLayout<UAvaClonerLineLayout>())
+	{
+		if (InAxis == EAvaClonerAxis::Custom)
+		{
+			Axis = LineLayout->GetDirection().GetSafeNormal();
+		}
+		
+		OutLocation = ClonerRotation.RotateVector(Axis) * LineLayout->GetSpacing() * FVector(LineLayout->GetCount());
+	}
+	else if (const UAvaClonerHoneycombLayout* HoneycombLayout = InClonerActor->GetActiveLayout<UAvaClonerHoneycombLayout>())
+	{
+		const EAvaClonerPlane Plane = HoneycombLayout->GetPlane();
+		const float WidthSpacing = HoneycombLayout->GetWidthSpacing();
+		const float HeightSpacing = HoneycombLayout->GetHeightSpacing();
+		const int32 WidthCount = HoneycombLayout->GetWidthCount();
+		const int32 HeightCount = HoneycombLayout->GetHeightCount();
+		
+		FVector Spacing;
+		FVector Count;
+		if (Plane == EAvaClonerPlane::XY)
+		{
+			Spacing = FVector(WidthSpacing, HeightSpacing, 0);
+			Count = FVector(WidthCount, HeightCount, 0);
+		}
+		else if (Plane == EAvaClonerPlane::YZ)
+		{
+			Spacing = FVector(0, WidthSpacing, HeightSpacing);
+			Count = FVector(0, WidthCount, HeightCount);
+		}
+		else if (Plane == EAvaClonerPlane::XZ)
+		{
+			Spacing = FVector(WidthSpacing, 0, HeightSpacing);
+			Count = FVector(WidthCount, 0, HeightCount);
+		}
+		
+		OutLocation = ClonerRotation.RotateVector(Axis * Spacing * Count/2);
+	}
+	else if (const UAvaClonerCircleLayout* CircleLayout = InClonerActor->GetActiveLayout<UAvaClonerCircleLayout>())
+	{
+		FRotator Rotation = CircleLayout->GetRotation();
+		const FVector Scale = CircleLayout->GetScale();
+		const EAvaClonerPlane Plane = CircleLayout->GetPlane();
+		const float Radius = CircleLayout->GetRadius();
+		
+		if (Plane == EAvaClonerPlane::XY)
+		{
+			Rotation = FRotator::ZeroRotator;
+		}
+		else if (Plane == EAvaClonerPlane::YZ)
+		{
+			Rotation = FRotator(90, 0, 0);
+		}
+		else if (Plane == EAvaClonerPlane::XZ)
+		{
+			Rotation = FRotator(0, 90, 0);
+		}
+		
+		OutLocation = ClonerRotation.RotateVector(Rotation.RotateVector(Axis * Scale)) * Radius;
+	}
+	else if (const UAvaClonerCylinderLayout* CylinderLayout = InClonerActor->GetActiveLayout<UAvaClonerCylinderLayout>())
+	{
+		FRotator Rotation = CylinderLayout->GetRotation();
+		const FVector Scale = CylinderLayout->GetScale();
+		const EAvaClonerPlane Plane = CylinderLayout->GetPlane();
+		const float Radius = CylinderLayout->GetRadius();
+		const float Height = CylinderLayout->GetHeight();
+		float Dim = 0.f;
+		
+		if (Plane == EAvaClonerPlane::XY)
+		{
+			Rotation = FRotator::ZeroRotator;
+			if (InAxis == EAvaClonerAxis::Y)
+			{
+				Dim = Radius;
+			}
+			else
+			{
+				Dim = Height / 2;
+			}
+		}
+		else if (Plane == EAvaClonerPlane::YZ)
+		{
+			if (InAxis == EAvaClonerAxis::Y)
+			{
+				Rotation = FRotator(90, 0, 0);
+				Dim = Radius;
+			}
+			else
+			{
+				Rotation = FRotator(0, 0, 90);
+				Dim = Height / 2;
+			}
+		}
+		else if (Plane == EAvaClonerPlane::XZ)
+		{
+			if (InAxis == EAvaClonerAxis::Z)
+			{
+				Rotation = FRotator(0, 90, 0);
+				Dim = Radius;
+			}
+			else
+			{
+				Rotation = FRotator(90, 0, 0);
+				Dim = Height / 2;
+			}
+		}
+		
+		OutLocation = ClonerRotation.RotateVector(Rotation.RotateVector(Axis * Scale)) * Dim;
+	}
+	else if (const UAvaClonerSphereUniformLayout* SphereLayout = InClonerActor->GetActiveLayout<UAvaClonerSphereUniformLayout>())
+	{
+		const float Radius = SphereLayout->GetRadius();
+		const FVector Scale = SphereLayout->GetScale();
+		
+		OutLocation = ClonerRotation.RotateVector(Axis) * Scale * Radius;
+	}
+	
+	return InClonerActor->GetActorLocation() + OutLocation * ClonerScale;
+}
+
+void FAvaClonerActorVisualizer::DrawSpacingButton(const AAvaClonerActor* InClonerActor, const FSceneView* InView, FPrimitiveDrawInterface* InPDI, int32 InIconIndex, EAvaClonerAxis InAxis, FLinearColor InColor) const
+{
+	UTexture2D* SpacingSprite = IAvalancheComponentVisualizersModule::Get().GetSettings()->GetVisualizerSprite(FAvalancheShapesEditorModule::BevelSprite);
+
+	if (!SpacingSprite || !SpacingSprite->GetResource() || !InClonerActor->GetEnabled() || InClonerActor->GetMeshCount() == 0)
+	{
+		return;
+	}
+
+	FVector IconLocation;
+	float IconSize;
+	GetIconMetrics(InView, InIconIndex, IconLocation, IconSize);
+
+	IconLocation = GetHandleSpacingLocation(InClonerActor, InAxis);
+
+	InPDI->SetHitProxy(new HAvaClonerActorSpacingHitProxy(InClonerActor->GetRootComponent(), InAxis));
+	InPDI->DrawSprite(IconLocation, IconSize, IconSize, SpacingSprite->GetResource(), InColor, SDPG_Foreground, 0, 0, 0, 0, SE_BLEND_Opaque);
+	InPDI->SetHitProxy(nullptr);
+}
+
+void FAvaClonerActorVisualizer::OnPropertyModified(UObject* InPropertyObject, FName InPropertyName, EPropertyChangeType::Type InType)
+{
+	if (!InPropertyObject)
+	{
+		return;
+	}
+	
+	InPropertyObject->Modify();
+	bHasBeenModified = true;
+	FProperty* PropertyModified = UE::AvaCore::GetProperty<UObject>(InPropertyName);
+	NotifyPropertyModified(InPropertyObject, PropertyModified, InType);
+}
+
+UActorComponent* FAvaClonerActorVisualizer::GetEditedComponent() const
+{
+	return ClonerActorWeak.IsValid() ? ClonerActorWeak->GetRootComponent() : nullptr;
+}
+
+TMap<UObject*, TArray<FProperty*>> FAvaClonerActorVisualizer::GatherEditableProperties(UObject* InObject) const
+{
+	if (UAvaClonerComponent* ClonerComponent = Cast<UAvaClonerComponent>(InObject))
+	{
+		if (UAvaClonerLayoutBase* Layout = ClonerComponent->GetClonerActiveLayout())
+		{
+			TArray<FName> PropertyNames;
+
+			if (UAvaClonerGridLayout* GridLayout = Cast<UAvaClonerGridLayout>(Layout))
+			{
+				PropertyNames = {
+					GET_MEMBER_NAME_CHECKED(UAvaClonerGridLayout, SpacingX),
+					GET_MEMBER_NAME_CHECKED(UAvaClonerGridLayout, SpacingY),
+					GET_MEMBER_NAME_CHECKED(UAvaClonerGridLayout, SpacingZ)
+				};
+			}
+			else if (UAvaClonerLineLayout* LineLayout = Cast<UAvaClonerLineLayout>(Layout))
+			{
+				PropertyNames = {GET_MEMBER_NAME_CHECKED(UAvaClonerLineLayout, Spacing)};
+			}
+			else if (UAvaClonerHoneycombLayout* HoneycombLayout = Cast<UAvaClonerHoneycombLayout>(Layout))
+			{
+				PropertyNames = {
+					GET_MEMBER_NAME_CHECKED(UAvaClonerHoneycombLayout, WidthSpacing),
+					GET_MEMBER_NAME_CHECKED(UAvaClonerHoneycombLayout, HeightSpacing)
+				};
+			}
+			else if (UAvaClonerCircleLayout* CircleLayout = Cast<UAvaClonerCircleLayout>(Layout))
+			{
+				PropertyNames = {GET_MEMBER_NAME_CHECKED(UAvaClonerCircleLayout, Radius)};
+			}
+			else if (UAvaClonerCylinderLayout* CylinderLayout = Cast<UAvaClonerCylinderLayout>(Layout))
+			{
+				PropertyNames = {
+					GET_MEMBER_NAME_CHECKED(UAvaClonerCylinderLayout, Radius),
+					GET_MEMBER_NAME_CHECKED(UAvaClonerCylinderLayout, Height)
+				};
+			}
+
+			TArray<FProperty*> Properties;
+			Properties.Reserve(PropertyNames.Num());
+
+			UClass* LayoutClass = Layout->GetClass();
+
+			for (FName PropertyName : PropertyNames)
+			{
+				if (FProperty* Property = LayoutClass->FindPropertyByName(PropertyName))
+				{
+					Properties.Add(Property);
+				}
+			}
+
+			return {{Layout, Properties}};
+		}
+	}
+
+	return {};
+}
+
+bool FAvaClonerActorVisualizer::VisProxyHandleClick(FEditorViewportClient* InViewportClient, HComponentVisProxy* InVisProxy, const FViewportClick& InClick)
+{
+	if (InClick.GetKey() != EKeys::LeftMouseButton)
+	{
+		EndEditing();
+		return Super::VisProxyHandleClick(InViewportClient, InVisProxy, InClick);
+	}
+
+	MeshType* DynMesh = Cast<MeshType>(const_cast<UActorComponent*>(InVisProxy->Component.Get()));
+
+	if (DynMesh == nullptr)
+	{
+		return Super::VisProxyHandleClick(InViewportClient, InVisProxy, InClick);
+	}
+
+	if (InVisProxy->IsA(HAvaClonerActorSpacingHitProxy::StaticGetType()))
+	{
+		EndEditing();
+		ClonerActorWeak = Cast<AAvaClonerActor>(DynMesh->GetOwner());
+		bEditingSpacing = true;
+		EditingAxis = static_cast<HAvaClonerActorSpacingHitProxy*>(InVisProxy)->Axis;
+		StartEditing(InViewportClient, DynMesh);
+		
+		return true;
+	}
+
+	return Super::VisProxyHandleClick(InViewportClient, InVisProxy, InClick);
+}
+
+bool FAvaClonerActorVisualizer::GetWidgetLocation(const FEditorViewportClient* InViewportClient, FVector& OutLocation) const
+{
+	if (const AAvaClonerActor* ClonerActor = ClonerActorWeak.Get())
+	{
+		if (bEditingSpacing)
+		{
+			OutLocation = GetHandleSpacingLocation(ClonerActor, EditingAxis);
+			return true;
+		}
+	}
+	
+	return Super::GetWidgetLocation(InViewportClient, OutLocation);
+}
+
+bool FAvaClonerActorVisualizer::GetWidgetMode(const FEditorViewportClient* InViewportClient,
+	UE::Widget::EWidgetMode& OutMode) const
+{
+	if (bEditingSpacing)
+	{
+		OutMode = UE::Widget::EWidgetMode::WM_Translate;
+		return true;
+	}
+	
+	return Super::GetWidgetMode(InViewportClient, OutMode);
+}
+
+bool FAvaClonerActorVisualizer::GetWidgetAxisList(const FEditorViewportClient* InViewportClient,
+	UE::Widget::EWidgetMode InWidgetMode, EAxisList::Type& OutAxisList) const
+{
+	if (bEditingSpacing)
+	{
+		if (EditingAxis == EAvaClonerAxis::X)
+		{
+			OutAxisList = EAxisList::Type::X;
+		}
+		else if (EditingAxis == EAvaClonerAxis::Y)
+		{
+			OutAxisList = EAxisList::Type::Y;
+		}
+		else if (EditingAxis == EAvaClonerAxis::Z)
+		{
+			OutAxisList = EAxisList::Type::Z;
+		}
+		else if (EditingAxis == EAvaClonerAxis::Custom)
+		{
+			OutAxisList = EAxisList::Type::XYZ;
+		}
+		return true;
+	}
+	
+	return Super::GetWidgetAxisList(InViewportClient, InWidgetMode, OutAxisList);
+}
+
+bool FAvaClonerActorVisualizer::GetWidgetAxisListDragOverride(const FEditorViewportClient* InViewportClient,
+	UE::Widget::EWidgetMode InWidgetMode, EAxisList::Type& OutAxisList) const
+{
+	return Super::GetWidgetAxisListDragOverride(InViewportClient, InWidgetMode, OutAxisList);
+}
+
+bool FAvaClonerActorVisualizer::ResetValue(FEditorViewportClient* InViewportClient, HHitProxy* InHitProxy)
+{
+	if (!InHitProxy->IsA(HAvaClonerActorSpacingHitProxy::StaticGetType()))
+	{
+		return Super::ResetValue(InViewportClient, InHitProxy);
+	}
+	
+	const HAvaClonerActorSpacingHitProxy* ComponentHitProxy = static_cast<HAvaClonerActorSpacingHitProxy*>(InHitProxy);
+	
+	if (!ComponentHitProxy->Component.IsValid() || !ComponentHitProxy->Component->IsA<UAvaClonerComponent>())
+	{
+		return Super::ResetValue(InViewportClient, InHitProxy);
+	}
+
+	if (const AAvaClonerActor* ClonerActor = Cast<AAvaClonerActor>(ComponentHitProxy->Component->GetOwner()))
+	{
+		if (UAvaClonerGridLayout* GridLayout = ClonerActor->GetActiveLayout<UAvaClonerGridLayout>())
+		{
+			FScopedTransaction Transaction(LOCTEXT("VisualizerResetValue", "Visualizer Reset Value"));
+			GridLayout->SetFlags(RF_Transactional);
+			
+			GridLayout->SetSpacingX(100.f);
+			GridLayout->SetSpacingY(100.f);
+			GridLayout->SetSpacingZ(100.f);
+			
+			OnPropertyModified(GridLayout, GET_MEMBER_NAME_CHECKED(UAvaClonerGridLayout, SpacingX), EPropertyChangeType::ValueSet);
+			OnPropertyModified(GridLayout, GET_MEMBER_NAME_CHECKED(UAvaClonerGridLayout, SpacingY), EPropertyChangeType::ValueSet);
+			OnPropertyModified(GridLayout, GET_MEMBER_NAME_CHECKED(UAvaClonerGridLayout, SpacingZ), EPropertyChangeType::ValueSet);
+		}
+		else if (UAvaClonerLineLayout* LineLayout = ClonerActor->GetActiveLayout<UAvaClonerLineLayout>())
+		{
+			FScopedTransaction Transaction(LOCTEXT("VisualizerResetValue", "Visualizer Reset Value"));
+			LineLayout->SetFlags(RF_Transactional);
+
+			LineLayout->SetSpacing(500.f);
+			
+			OnPropertyModified(LineLayout, GET_MEMBER_NAME_CHECKED(UAvaClonerLineLayout, Spacing), EPropertyChangeType::ValueSet);
+		}
+		else if (UAvaClonerHoneycombLayout* HoneycombLayout = ClonerActor->GetActiveLayout<UAvaClonerHoneycombLayout>())
+		{
+			FScopedTransaction Transaction(LOCTEXT( "VisualizerResetValue", "Visualizer Reset Value"));
+			HoneycombLayout->SetFlags(RF_Transactional);
+			
+			HoneycombLayout->SetWidthSpacing(100.f);
+			HoneycombLayout->SetHeightSpacing(100.f);
+			
+			OnPropertyModified(HoneycombLayout, GET_MEMBER_NAME_CHECKED(UAvaClonerHoneycombLayout, WidthSpacing), EPropertyChangeType::ValueSet);
+			OnPropertyModified(HoneycombLayout, GET_MEMBER_NAME_CHECKED(UAvaClonerHoneycombLayout, HeightSpacing), EPropertyChangeType::ValueSet);
+		}
+		else if (UAvaClonerCircleLayout* CircleLayout = ClonerActor->GetActiveLayout<UAvaClonerCircleLayout>())
+		{
+			FScopedTransaction Transaction(LOCTEXT("VisualizerResetValue", "Visualizer Reset Value"));
+			CircleLayout->SetFlags(RF_Transactional);
+			
+			CircleLayout->SetRadius(500.f);
+			
+			OnPropertyModified(CircleLayout, GET_MEMBER_NAME_CHECKED(UAvaClonerCircleLayout, Radius), EPropertyChangeType::ValueSet);
+		}
+		else if (UAvaClonerCylinderLayout* CylinderLayout = ClonerActor->GetActiveLayout<UAvaClonerCylinderLayout>())
+		{
+			FScopedTransaction Transaction(LOCTEXT("VisualizerResetValue", "Visualizer Reset Value"));
+			CylinderLayout->SetFlags(RF_Transactional);
+			
+			CylinderLayout->SetRadius(500.f);
+			CylinderLayout->SetHeight(1000.f);
+
+			OnPropertyModified(CylinderLayout, GET_MEMBER_NAME_CHECKED(UAvaClonerCylinderLayout, Radius), EPropertyChangeType::ValueSet);
+			OnPropertyModified(CylinderLayout, GET_MEMBER_NAME_CHECKED(UAvaClonerCylinderLayout, Height), EPropertyChangeType::ValueSet);
+		}
+	}
+	
+	return true;
+}
+
+bool FAvaClonerActorVisualizer::IsEditing() const
+{
+	if (bEditingSpacing)
+	{
+		return true;
+	}
+	
+	return Super::IsEditing();
+}
+
+void FAvaClonerActorVisualizer::EndEditing()
+{
+	Super::EndEditing();
+
+	ClonerActorWeak.Reset();
+	InitialSpacing = FVector::ZeroVector;
+	bEditingSpacing = false;
+}
+
+#undef LOCTEXT_NAMESPACE
\ No newline at end of file
diff --git a/Engine/Plugins/Experimental/Avalanche/Source/AvalancheEffectorsEditor/Private/Cloner/AvaClonerDetailCustomization.cpp b/Engine/Plugins/Experimental/Avalanche/Source/AvalancheEffectorsEditor/Private/Cloner/AvaClonerDetailCustomization.cpp
new file mode 100644
--- /dev/null
+++ b/Engine/Plugins/Experimental/Avalanche/Source/AvalancheEffectorsEditor/Private/Cloner/AvaClonerDetailCustomization.cpp
@@ -0,0 +1,198 @@
+// Copyright Epic Games, Inc. All Rights Reserved.
+
+#include "Cloner/AvaClonerDetailCustomization.h"
+#include "Cloner/AvaClonerActor.h"
+#include "DetailLayoutBuilder.h"
+#include "DetailWidgetRow.h"
+#include "Modules/ModuleManager.h"
+#include "NiagaraDataInterfaceCurve.h"
+#include "PropertyEditorModule.h"
+#include "Widgets/Input/SButton.h"
+
+#define LOCTEXT_NAMESPACE "AvaClonerDetailCustomization"
+
+void FAvaClonerDetailCustomization::CustomizeDetails(IDetailLayoutBuilder& InDetailBuilder)
+{
+	static const TArray<FName> CategoriesPriority =
+	{
+		TEXT("Cloner"),
+		TEXT("Layout"),
+		TEXT("Renderer"),
+		TEXT("Progress"),
+		TEXT("Step"),
+		TEXT("Range"),
+		TEXT("Spawn"),
+		TEXT("Lifetime"),
+	};
+
+	// Sort categories
+	int32 Order = InDetailBuilder.EditCategory(TEXT("Transform")).GetSortOrder();
+	for (const FName& Category : CategoriesPriority)
+	{
+		IDetailCategoryBuilder& CategoryBuilder = InDetailBuilder.EditCategory(Category);
+		CategoryBuilder.SetSortOrder(++Order);
+	}
+
+	// Gather all cloner actor selected
+	TArray<TWeakObjectPtr<AAvaClonerActor>> SelectedClonersWeak = InDetailBuilder.GetSelectedObjectsOfType<AAvaClonerActor>();
+
+	LayoutFunctionNames.Empty();
+
+	if (SelectedClonersWeak.IsEmpty())
+	{
+		return;
+	}
+
+	// Look for layout UFunctions and group them by name
+	for (const TWeakObjectPtr<AAvaClonerActor>& ClonerWeak : SelectedClonersWeak)
+	{
+		const AAvaClonerActor* Cloner = ClonerWeak.Get();
+		if (!Cloner)
+		{
+			continue;
+		}
+
+		// Look for CallInEditor functions in active layout
+		UAvaClonerLayoutBase* ActiveLayout = Cloner->GetActiveLayout();
+		if (!ActiveLayout)
+		{
+			continue;
+		}
+
+		// Iterate through all UFunctions in the class and its parent classes.
+		for (UFunction* Function : TFieldRange<UFunction>(ActiveLayout->GetClass(), EFieldIteratorFlags::ExcludeSuper))
+		{
+			// Only CallInEditor function with 0 parameters
+			if (Function && Function->HasMetaData("CallInEditor") && Function->NumParms == 0)
+			{
+				FName FunctionName = Function->GetFName();
+				TMap<TWeakObjectPtr<UObject>, TWeakObjectPtr<UFunction>>& ObjectFunctions = LayoutFunctionNames.FindOrAdd(FunctionName);
+				ObjectFunctions.Add(ActiveLayout, Function);
+			}
+		}
+	}
+
+	if (!LayoutFunctionNames.IsEmpty())
+	{
+		// Add buttons for selected cloners
+		const TSharedPtr<SHorizontalBox> FunctionsWidget = SNew(SHorizontalBox);
+		for (const TPair<FName, TMap<TWeakObjectPtr<UObject>, TWeakObjectPtr<UFunction>>>& LayoutFunctionNamesPair : LayoutFunctionNames)
+		{
+			const FText ButtonLabel = FText::FromString(FName::NameToDisplayString(LayoutFunctionNamesPair.Key.ToString(), false));
+
+			FunctionsWidget->AddSlot()
+			.AutoWidth()
+			[
+				SNew(SButton)
+				.Text(ButtonLabel)
+				.OnClicked(this, &FAvaClonerDetailCustomization::OnExecuteFunction, LayoutFunctionNamesPair.Key)
+			];
+		}
+
+		IDetailCategoryBuilder& ClonerCategoryBuilder = InDetailBuilder.EditCategory("Cloner");
+
+		ClonerCategoryBuilder
+			.AddCustomRow(FText::FromString("LayoutFunctions"))
+			.WholeRowContent()
+			[
+				FunctionsWidget.ToSharedRef()
+			];
+	}
+
+	if (SelectedClonersWeak.Num() != 1)
+	{
+		return;
+	}
+
+	const AAvaClonerActor* Cloner = SelectedClonersWeak[0].Get();
+
+	if (!Cloner)
+	{
+		return;
+	}
+
+	UNiagaraDataInterfaceCurve* LifetimeScaleCurve = Cloner->GetLifetimeScaleCurveDI();
+
+	if (!LifetimeScaleCurve)
+	{
+		return;
+	}
+
+	const TArray<UObject*> ShowObjects {LifetimeScaleCurve};
+	IDetailCategoryBuilder& LifetimeCategory = InDetailBuilder.EditCategory(TEXT("Lifetime"));
+	LifetimeCategory.SetShowAdvanced(true);
+
+	FAddPropertyParams Params;
+	Params.HideRootObjectNode(true);
+	Params.CreateCategoryNodes(false);
+
+	IDetailPropertyRow* CurveRow = LifetimeCategory.AddExternalObjects(ShowObjects, EPropertyLocation::Advanced, Params);
+	CurveRow->Visibility(MakeAttributeLambda([Cloner]()
+	{
+		return Cloner
+			&& Cloner->GetLifetimeEnabled()
+			&& Cloner->GetLifetimeScaleEnabled()
+			? EVisibility::Visible
+			: EVisibility::Collapsed;
+	}));
+}
+
+FReply FAvaClonerDetailCustomization::OnExecuteFunction(FName InFunctionName)
+{
+	if (TMap<TWeakObjectPtr<UObject>, TWeakObjectPtr<UFunction>> const* ObjectFunctions = LayoutFunctionNames.Find(InFunctionName))
+	{
+		for (const TPair<TWeakObjectPtr<UObject>, TWeakObjectPtr<UFunction>>& ObjectFunction : *ObjectFunctions)
+		{
+			UObject* Object = ObjectFunction.Key.Get();
+			UFunction* Function = ObjectFunction.Value.Get();
+
+			if (!Object || !Function)
+			{
+				continue;
+			}
+
+			Object->ProcessEvent(Function, nullptr);
+		}
+	}
+
+	return FReply::Handled();
+}
+
+void FAvaClonerDetailCustomization::RegisterCustomSections() const
+{
+	static const FName PropertyEditor("PropertyEditor");
+	FPropertyEditorModule& PropertyModule = FModuleManager::GetModuleChecked<FPropertyEditorModule>(PropertyEditor);
+
+	static const FName ClassName = AAvaClonerActor::StaticClass()->GetFName();
+
+	const TSharedRef<FPropertySection> StreamingSection = PropertyModule.FindOrCreateSection(ClassName, "Streaming", LOCTEXT("Cloner.Streaming", "Streaming"));
+	StreamingSection->RemoveCategory("World Partition");
+	StreamingSection->RemoveCategory("Data Layers");
+	StreamingSection->RemoveCategory("HLOD");
+
+	const TSharedRef<FPropertySection> GeneralSection = PropertyModule.FindOrCreateSection(ClassName, "General", LOCTEXT("Cloner.General", "General"));
+	GeneralSection->AddCategory("Cloner");
+
+	const TSharedRef<FPropertySection> LayoutSection = PropertyModule.FindOrCreateSection(ClassName, "Layouts", LOCTEXT("Cloner.Layouts", "Layouts"));
+	LayoutSection->AddCategory("Layout");
+
+	const TSharedRef<FPropertySection> RendererSection = PropertyModule.FindOrCreateSection(ClassName, "Renderer", LOCTEXT("Cloner.Renderer", "Renderer"));
+	RendererSection->AddCategory("Renderer");
+
+	const TSharedRef<FPropertySection> ProgressSection = PropertyModule.FindOrCreateSection(ClassName, "Progress", LOCTEXT("Cloner.Progress", "Progress"));
+	ProgressSection->AddCategory("Progress");
+
+	const TSharedRef<FPropertySection> StepSection = PropertyModule.FindOrCreateSection(ClassName, "Step", LOCTEXT("Cloner.Step", "Step"));
+	StepSection->AddCategory("Step");
+
+	const TSharedRef<FPropertySection> RangeSection = PropertyModule.FindOrCreateSection(ClassName, "Range", LOCTEXT("Cloner.Range", "Range"));
+	RangeSection->AddCategory("Range");
+
+	const TSharedRef<FPropertySection> SpawnSection = PropertyModule.FindOrCreateSection(ClassName, "Spawn", LOCTEXT("Cloner.Spawn", "Spawn"));
+	SpawnSection->AddCategory("Spawn");
+
+	const TSharedRef<FPropertySection> LifetimeSection = PropertyModule.FindOrCreateSection(ClassName, "Lifetime", LOCTEXT("Cloner.Lifetime", "Lifetime"));
+	LifetimeSection->AddCategory("Lifetime");
+}
+
+#undef LOCTEXT_NAMESPACE
diff --git a/Engine/Plugins/Experimental/Avalanche/Source/AvalancheEffectorsEditor/Private/Cloner/AvaClonerDetailCustomization.h b/Engine/Plugins/Experimental/Avalanche/Source/AvalancheEffectorsEditor/Private/Cloner/AvaClonerDetailCustomization.h
new file mode 100644
--- /dev/null
+++ b/Engine/Plugins/Experimental/Avalanche/Source/AvalancheEffectorsEditor/Private/Cloner/AvaClonerDetailCustomization.h
@@ -0,0 +1,37 @@
+// Copyright Epic Games, Inc. All Rights Reserved.
+
+#pragma once
+
+#include "IDetailCustomization.h"
+#include "Templates/SharedPointer.h"
+#include "UObject/Class.h"
+#include "UObject/Object.h"
+#include "UObject/WeakObjectPtr.h"
+
+class FReply;
+
+/* Used to customize cloner actor properties in details panel */
+class FAvaClonerDetailCustomization : public IDetailCustomization
+{
+public:
+	static TSharedRef<IDetailCustomization> MakeInstance()
+	{
+		return MakeShared<FAvaClonerDetailCustomization>();
+	}
+	
+	explicit FAvaClonerDetailCustomization()
+	{
+		RegisterCustomSections();
+	}
+
+	virtual void CustomizeDetails(IDetailLayoutBuilder& InDetailBuilder) override;
+
+protected:
+	/** Execute ufunction with that name on selected objects */
+	FReply OnExecuteFunction(FName InFunctionName);
+
+	/** Register custom categories section in details panel */
+	void RegisterCustomSections() const;
+	
+	TMap<FName, TMap<TWeakObjectPtr<UObject>, TWeakObjectPtr<UFunction>>> LayoutFunctionNames;
+};
diff --git a/Engine/Plugins/Experimental/Avalanche/Source/AvalancheEffectorsEditor/Private/Effector/AvaEffectorActorTool.cpp b/Engine/Plugins/Experimental/Avalanche/Source/AvalancheEffectorsEditor/Private/Effector/AvaEffectorActorTool.cpp
new file mode 100644
--- /dev/null
+++ b/Engine/Plugins/Experimental/Avalanche/Source/AvalancheEffectorsEditor/Private/Effector/AvaEffectorActorTool.cpp
@@ -0,0 +1,37 @@
+// Copyright Epic Games, Inc. All Rights Reserved.
+
+#include "AvaEffectorActorTool.h"
+#include "AvaEffectorsEditorCommands.h"
+#include "AvaInteractiveToolsSettings.h"
+#include "Builders/AvaInteractiveToolsToolBuilder.h"
+#include "Effector/AvaEffectorActor.h"
+
+UAvaEffectorActorTool::UAvaEffectorActorTool()
+{
+	ActorClass = AAvaEffectorActor::StaticClass();
+}
+
+bool UAvaEffectorActorTool::UseIdentityRotation() const
+{
+	return ConditionalIdentityRotation();
+}
+
+FName UAvaEffectorActorTool::GetCategoryName()
+{
+	return IAvalancheInteractiveToolsModule::CategoryNameActor;
+}
+
+FAvaInteractiveToolsToolParameters UAvaEffectorActorTool::GetToolParameters() const
+{
+	return {
+		FAvaEffectorsEditorCommands::Get().Tool_Actor_Effector,
+		TEXT("Effector Actor Tool"),
+		4000,
+		FAvalancheInteractiveToolsCreateBuilder::CreateLambda(
+			[](UEdMode* InEdMode)
+			{
+				return UAvaInteractiveToolsToolBuilder::CreateToolBuilder<UAvaEffectorActorTool>(InEdMode);
+			}),
+		ActorClass
+	};
+}
diff --git a/Engine/Plugins/Experimental/Avalanche/Source/AvalancheEffectorsEditor/Private/Effector/AvaEffectorActorTool.h b/Engine/Plugins/Experimental/Avalanche/Source/AvalancheEffectorsEditor/Private/Effector/AvaEffectorActorTool.h
new file mode 100644
--- /dev/null
+++ b/Engine/Plugins/Experimental/Avalanche/Source/AvalancheEffectorsEditor/Private/Effector/AvaEffectorActorTool.h
@@ -0,0 +1,22 @@
+// Copyright Epic Games, Inc. All Rights Reserved.
+
+#pragma once
+
+#include "Tools/AvaInteractiveToolsActorPointToolBase.h"
+#include "AvaEffectorActorTool.generated.h"
+
+UCLASS()
+class UAvaEffectorActorTool : public UAvaInteractiveToolsActorPointToolBase
+{
+	GENERATED_BODY()
+
+public:
+	UAvaEffectorActorTool();
+
+	//~ Begin UAvaInteractiveToolsToolBase
+	virtual bool UseIdentityLocation() const override { return false; }
+	virtual bool UseIdentityRotation() const override;
+	virtual FName GetCategoryName() override;
+	virtual FAvaInteractiveToolsToolParameters GetToolParameters() const override;
+	//~ End UAvaInteractiveToolsToolBase
+};
