[MotionDesign]

diff --git a/Engine/Plugins/Experimental/Avalanche/Source/AvalancheEffectorsEditor/Private/Effector/AvaEffectorActorVis.cpp b/Engine/Plugins/Experimental/Avalanche/Source/AvalancheEffectorsEditor/Private/Effector/AvaEffectorActorVis.cpp
new file mode 100644
--- /dev/null
+++ b/Engine/Plugins/Experimental/Avalanche/Source/AvalancheEffectorsEditor/Private/Effector/AvaEffectorActorVis.cpp
@@ -0,0 +1,469 @@
+// Copyright Epic Games, Inc. All Rights Reserved.
+
+#include "Effector/AvaEffectorActorVis.h"
+
+#include "AvaField.h"
+#include "AvalancheShapesEditorModule.h"
+#include "AvaVisBase.h"
+#include "EditorViewportClient.h"
+#include "Effector/AvaEffectorComponent.h"
+#include "Engine/Texture2D.h"
+#include "Framework/Application/SlateApplication.h"
+#include "IAvalancheComponentVisualizersModule.h"
+#include "ScopedTransaction.h"
+#include "TextureResource.h"
+
+IMPLEMENT_HIT_PROXY(HAvaEffectorActorZoneHitProxy, HAvaHitProxy);
+
+#define LOCTEXT_NAMESPACE "AvaEffectorActorVisualizer"
+
+FAvaEffectorActorVisualizer::FAvaEffectorActorVisualizer()
+	: FAvaVisualizerBase()
+{
+	using namespace UE::AvaCore;
+	InnerRadiusProperty  = GetProperty<AAvaEffectorActor>(GET_MEMBER_NAME_CHECKED(AAvaEffectorActor, InnerRadius));
+	OuterRadiusProperty  = GetProperty<AAvaEffectorActor>(GET_MEMBER_NAME_CHECKED(AAvaEffectorActor, OuterRadius));
+	InnerExtentProperty  = GetProperty<AAvaEffectorActor>(GET_MEMBER_NAME_CHECKED(AAvaEffectorActor, InnerExtent));
+	OuterExtentProperty  = GetProperty<AAvaEffectorActor>(GET_MEMBER_NAME_CHECKED(AAvaEffectorActor, OuterExtent));
+	PlaneSpacingProperty = GetProperty<AAvaEffectorActor>(GET_MEMBER_NAME_CHECKED(AAvaEffectorActor, PlaneSpacing));
+}
+
+void FAvaEffectorActorVisualizer::StoreInitialValues()
+{
+	Super::StoreInitialValues();
+	
+	if (GetEditedComponent() == nullptr)
+	{
+		return;
+	}
+
+	InitialInnerRadius = EffectorActorWeak->GetInnerRadius();
+	InitialOuterRadius = EffectorActorWeak->GetOuterRadius();
+	InitialInnerExtent = EffectorActorWeak->GetInnerExtent();
+	InitialOuterExtent = EffectorActorWeak->GetOuterExtent();
+	InitialPlaneSpacing = EffectorActorWeak->GetPlaneSpacing();
+}
+
+FBox FAvaEffectorActorVisualizer::GetComponentBounds(const UActorComponent* InComponent) const
+{
+	if (const UAvaEffectorComponent* EffectorComponent = Cast<UAvaEffectorComponent>(InComponent))
+	{
+		if (const AAvaEffectorActor* EffectorActor = Cast<AAvaEffectorActor>(EffectorComponent->GetOwner()))
+		{
+			if (EffectorActor->GetType() == EAvaClonerEffectorType::Box)
+			{
+				return FBox(-EffectorActor->GetOuterExtent(), EffectorActor->GetOuterExtent());
+			}
+			if (EffectorActor->GetType() == EAvaClonerEffectorType::Sphere)
+			{
+				return FBox(-FVector(EffectorActor->GetOuterRadius() / 2), FVector(EffectorActor->GetOuterRadius() / 2));
+			}
+			if (EffectorActor->GetType() == EAvaClonerEffectorType::Plane)
+			{
+				return FBox(-FVector(EffectorActor->GetPlaneSpacing() / 2), FVector(EffectorActor->GetPlaneSpacing() / 2));
+			}
+		}
+	}
+	
+	return Super::GetComponentBounds(InComponent);
+}
+
+bool FAvaEffectorActorVisualizer::HandleInputDeltaInternal(FEditorViewportClient* InViewportClient, FViewport* InViewport, const FVector& InAccumulatedTranslation, const FRotator& InAccumulatedRotation, const FVector& InAccumulatedScale)
+{
+	if (!FSlateApplication::Get().GetPressedMouseButtons().Contains(EKeys::LeftMouseButton))
+	{
+		return false;
+	}
+
+	if (AAvaEffectorActor* EffectorActor = EffectorActorWeak.Get())
+	{
+		if (GetViewportWidgetMode(InViewportClient) == UE::Widget::WM_Translate)
+		{
+			if (GetViewportWidgetAxisList(InViewportClient) & EAxisList::XYZ)
+			{
+				if (EffectorActor->GetType() == EAvaClonerEffectorType::Box)
+				{
+					if (bEditingInnerZone)
+					{
+						EffectorActor->SetInnerExtent(InitialInnerExtent + InAccumulatedTranslation);
+						EffectorActor->Modify();
+						bHasBeenModified = true;
+						NotifyPropertyModified(EffectorActor, InnerExtentProperty, EPropertyChangeType::Interactive);
+					}
+					else if (bEditingOuterZone)
+					{
+						EffectorActor->SetOuterExtent(InitialOuterExtent + InAccumulatedTranslation);
+						EffectorActor->Modify();
+						bHasBeenModified = true;
+						NotifyPropertyModified(EffectorActor, OuterExtentProperty, EPropertyChangeType::Interactive);
+					}
+					
+					return true;
+				}
+			}
+			if (GetViewportWidgetAxisList(InViewportClient) & EAxisList::Y)
+			{
+				if (EffectorActor->GetType() == EAvaClonerEffectorType::Plane)
+				{
+					if (bEditingOuterZone || bEditingInnerZone)
+					{
+						EffectorActor->SetPlaneSpacing(InitialPlaneSpacing + InAccumulatedTranslation.Y);
+						EffectorActor->Modify();
+						bHasBeenModified = true;
+						NotifyPropertyModified(EffectorActor, PlaneSpacingProperty, EPropertyChangeType::Interactive);
+					}
+					
+					return true;
+				}
+				if (EffectorActor->GetType() == EAvaClonerEffectorType::Sphere)
+				{
+					if (bEditingInnerZone)
+					{
+						EffectorActor->SetInnerRadius(InitialInnerRadius + InAccumulatedTranslation.Y);
+						EffectorActor->Modify();
+						bHasBeenModified = true;
+						NotifyPropertyModified(EffectorActor, InnerRadiusProperty, EPropertyChangeType::Interactive);
+					}
+					else if (bEditingOuterZone)
+					{
+						EffectorActor->SetOuterRadius(InitialOuterRadius + InAccumulatedTranslation.Y);
+						EffectorActor->Modify();
+						bHasBeenModified = true;
+						NotifyPropertyModified(EffectorActor, OuterRadiusProperty, EPropertyChangeType::Interactive);
+					}
+					
+					return true;
+				}
+			}
+		}
+	}
+	else
+	{
+		EndEditing();
+	}
+	
+	return Super::HandleInputDeltaInternal(InViewportClient, InViewport, InAccumulatedTranslation, InAccumulatedRotation, InAccumulatedScale);
+}
+
+void FAvaEffectorActorVisualizer::DrawVisualizationEditing(const UActorComponent* InComponent, const FSceneView* InView, FPrimitiveDrawInterface* InPDI, int32& InOutIconIndex)
+{
+	Super::DrawVisualizationEditing(InComponent, InView, InPDI, InOutIconIndex);
+
+	const UAvaEffectorComponent* EffectorComponent = Cast<UAvaEffectorComponent>(InComponent);
+
+	if (!EffectorComponent)
+	{
+		return;
+	}
+
+	const AAvaEffectorActor* EffectorActor = Cast<AAvaEffectorActor>(EffectorComponent->GetOwner());
+
+	if (!EffectorActor)
+	{
+		return;
+	}
+
+	if (EffectorActor->GetType() != EAvaClonerEffectorType::Plane)
+	{
+		if (!bEditingInnerZone)
+		{
+			DrawZoneButton(EffectorActor, InView, InPDI, InOutIconIndex, true, FAvaVisualizerBase::Inactive);
+			InOutIconIndex++;
+		}
+	}
+
+	if (!bEditingOuterZone)
+	{
+		DrawZoneButton(EffectorActor, InView, InPDI, InOutIconIndex, false, bEditingOuterZone ? FAvaVisualizerBase::Active : FAvaVisualizerBase::Inactive);
+		InOutIconIndex++;
+	}
+}
+
+void FAvaEffectorActorVisualizer::DrawVisualizationNotEditing(const UActorComponent* InComponent, const FSceneView* InView, FPrimitiveDrawInterface* InPDI, int32& InOutIconIndex)
+{
+	Super::DrawVisualizationNotEditing(InComponent, InView, InPDI, InOutIconIndex);
+
+	const UAvaEffectorComponent* EffectorComponent = Cast<UAvaEffectorComponent>(InComponent);
+
+	if (!EffectorComponent)
+	{
+		return;
+	}
+
+	const AAvaEffectorActor* EffectorActor = Cast<AAvaEffectorActor>(EffectorComponent->GetOwner());
+
+	if (!EffectorActor)
+	{
+		return;
+	}
+
+	if (EffectorActor->GetType() != EAvaClonerEffectorType::Plane)
+	{
+		DrawZoneButton(EffectorActor, InView, InPDI, InOutIconIndex, true, FAvaVisualizerBase::Inactive);
+		InOutIconIndex++;
+	}
+	
+	DrawZoneButton(EffectorActor, InView, InPDI, InOutIconIndex, false, FAvaVisualizerBase::Inactive);
+	InOutIconIndex++;
+}
+
+FVector FAvaEffectorActorVisualizer::GetHandleZoneLocation(const AAvaEffectorActor* InEffectorActor, bool bInInnerSize) const
+{
+	const FVector EffectorScale = InEffectorActor->GetActorScale();
+	const FRotator EffectorRotation = InEffectorActor->GetActorRotation();
+	FVector OutLocation = InEffectorActor->GetActorLocation();
+
+	if (InEffectorActor->GetType() == EAvaClonerEffectorType::Box)
+	{
+		if (bInInnerSize)
+		{
+			OutLocation += EffectorRotation.RotateVector(InEffectorActor->GetInnerExtent()) * EffectorScale;
+		}
+		else
+			{
+			OutLocation += EffectorRotation.RotateVector(InEffectorActor->GetOuterExtent()) * EffectorScale;
+		}
+	}
+	else if (InEffectorActor->GetType() == EAvaClonerEffectorType::Plane)
+	{
+		const float ComponentScale = (EffectorRotation.RotateVector(-FVector::YAxisVector) * EffectorScale).Length();
+		OutLocation += EffectorRotation.RotateVector(FVector::YAxisVector) * (InEffectorActor->GetPlaneSpacing() / 2) * ComponentScale;
+	}
+	else if (InEffectorActor->GetType() == EAvaClonerEffectorType::Sphere)
+	{
+		const float MinComponentScale = FMath::Min<float>(FMath::Min<float>(EffectorScale.X, EffectorScale.Y), EffectorScale.Z);
+		if (bInInnerSize)
+		{
+			OutLocation += EffectorRotation.RotateVector(FVector::YAxisVector) * InEffectorActor->GetInnerRadius() * MinComponentScale;
+		}
+		else
+		{
+			OutLocation += EffectorRotation.RotateVector(FVector::YAxisVector) * InEffectorActor->GetOuterRadius() * MinComponentScale;
+		}
+	}
+	
+	return OutLocation;
+}
+
+void FAvaEffectorActorVisualizer::DrawZoneButton(const AAvaEffectorActor* InEffectorActor, const FSceneView* InView, FPrimitiveDrawInterface* InPDI, int32 InIconIndex, bool bInInnerZone, FLinearColor InColor) const
+{
+	UTexture2D* ZoneSprite = IAvalancheComponentVisualizersModule::Get().GetSettings()->GetVisualizerSprite(FAvalancheShapesEditorModule::BevelSprite);
+
+	if (!ZoneSprite || !ZoneSprite->GetResource())
+	{
+		return;
+	}
+
+	FVector IconLocation;
+	float IconSize;
+	GetIconMetrics(InView, InIconIndex, IconLocation, IconSize);
+
+	IconLocation = GetHandleZoneLocation(InEffectorActor, bInInnerZone);
+
+	InPDI->SetHitProxy(new HAvaEffectorActorZoneHitProxy(InEffectorActor->GetRootComponent(), bInInnerZone));
+	InPDI->DrawSprite(IconLocation, IconSize, IconSize, ZoneSprite->GetResource(), InColor, SDPG_Foreground, 0, 0, 0, 0, SE_BLEND_Opaque);
+	InPDI->SetHitProxy(nullptr);
+}
+
+UActorComponent* FAvaEffectorActorVisualizer::GetEditedComponent() const
+{
+	return EffectorActorWeak.IsValid() ? EffectorActorWeak->GetRootComponent() : nullptr;
+}
+
+TMap<UObject*, TArray<FProperty*>> FAvaEffectorActorVisualizer::GatherEditableProperties(UObject* InObject) const
+{
+	if (UAvaEffectorComponent* EffectorComponent = Cast<UAvaEffectorComponent>(InObject))
+	{
+		if (AAvaEffectorActor* EffectorActor = EffectorComponent->GetOuterAAvaEffectorActor())
+		{
+			switch (EffectorActor->GetType())
+			{
+				case EAvaClonerEffectorType::Plane:
+					return {{EffectorActor, {PlaneSpacingProperty}}};
+
+				case EAvaClonerEffectorType::Box:
+					return {{EffectorActor, {InnerExtentProperty, OuterExtentProperty}}};
+
+				case EAvaClonerEffectorType::Sphere:
+					return {{EffectorActor, {InnerRadiusProperty, OuterRadiusProperty}}};
+			}
+		}
+	}
+
+	return {};
+}
+
+bool FAvaEffectorActorVisualizer::VisProxyHandleClick(FEditorViewportClient* InViewportClient, HComponentVisProxy* InVisProxy, const FViewportClick& InClick)
+{
+	if (InClick.GetKey() != EKeys::LeftMouseButton)
+	{
+		EndEditing();
+		return Super::VisProxyHandleClick(InViewportClient, InVisProxy, InClick);
+	}
+
+	UActorComponent* Component = const_cast<UActorComponent*>(InVisProxy->Component.Get());
+
+	if (!Component || !Component->GetOwner()->IsA<AAvaEffectorActor>())
+	{
+		return Super::VisProxyHandleClick(InViewportClient, InVisProxy, InClick);
+	}
+
+	if (InVisProxy->IsA(HAvaEffectorActorZoneHitProxy::StaticGetType()))
+	{
+		EndEditing();
+		EffectorActorWeak = Cast<AAvaEffectorActor>(InVisProxy->Component->GetOwner());
+		bEditingInnerZone = static_cast<HAvaEffectorActorZoneHitProxy*>(InVisProxy)->bInnerZone;
+		bEditingOuterZone = !bEditingInnerZone;
+		StartEditing(InViewportClient, Component);
+		
+		return true;
+	}
+
+	return Super::VisProxyHandleClick(InViewportClient, InVisProxy, InClick);
+}
+
+bool FAvaEffectorActorVisualizer::GetWidgetLocation(const FEditorViewportClient* InViewportClient, FVector& OutLocation) const
+{
+	if (EffectorActorWeak.IsValid())
+	{
+		OutLocation = GetHandleZoneLocation(EffectorActorWeak.Get(), bEditingInnerZone);
+		return true;
+	}
+	
+	return Super::GetWidgetLocation(InViewportClient, OutLocation);
+}
+
+bool FAvaEffectorActorVisualizer::GetWidgetMode(const FEditorViewportClient* InViewportClient,
+	UE::Widget::EWidgetMode& OutMode) const
+{
+	if (bEditingOuterZone || bEditingInnerZone)
+	{
+		OutMode = UE::Widget::EWidgetMode::WM_Translate;
+		return true;
+	}
+	
+	return Super::GetWidgetMode(InViewportClient, OutMode);
+}
+
+bool FAvaEffectorActorVisualizer::GetWidgetAxisList(const FEditorViewportClient* InViewportClient,
+	UE::Widget::EWidgetMode InWidgetMode, EAxisList::Type& OutAxisList) const
+{
+	if (bEditingOuterZone || bEditingInnerZone)
+	{
+		if (EffectorActorWeak->GetType() != EAvaClonerEffectorType::Box)
+		{
+			OutAxisList = EAxisList::Type::Y;
+		}
+		else
+		{
+			OutAxisList = EAxisList::Type::XYZ;
+		}
+		
+		return true;
+	}
+	
+	return Super::GetWidgetAxisList(InViewportClient, InWidgetMode, OutAxisList);
+}
+
+bool FAvaEffectorActorVisualizer::GetWidgetAxisListDragOverride(const FEditorViewportClient* InViewportClient,
+	UE::Widget::EWidgetMode InWidgetMode, EAxisList::Type& OutAxisList) const
+{
+	if (bEditingOuterZone || bEditingInnerZone)
+	{
+		if (EffectorActorWeak->GetType() != EAvaClonerEffectorType::Box)
+		{
+			OutAxisList = EAxisList::Type::Y;
+			return true;
+		}
+	}
+	
+	return Super::GetWidgetAxisListDragOverride(InViewportClient, InWidgetMode, OutAxisList);
+}
+
+bool FAvaEffectorActorVisualizer::ResetValue(FEditorViewportClient* InViewportClient, HHitProxy* InHitProxy)
+{
+	if (!InHitProxy->IsA(HAvaEffectorActorZoneHitProxy::StaticGetType()))
+	{
+		return Super::ResetValue(InViewportClient, InHitProxy);
+	}
+	
+	const HAvaEffectorActorZoneHitProxy* ComponentHitProxy = static_cast<HAvaEffectorActorZoneHitProxy*>(InHitProxy);
+	
+	if (!ComponentHitProxy->Component.IsValid() || !ComponentHitProxy->Component->IsA<UAvaEffectorComponent>())
+	{
+		return Super::ResetValue(InViewportClient, InHitProxy);
+	}
+
+	if (AAvaEffectorActor* EffectorActor = Cast<AAvaEffectorActor>(ComponentHitProxy->Component->GetOwner()))
+	{
+		if (EffectorActor->GetType() == EAvaClonerEffectorType::Box)
+		{
+			if (ComponentHitProxy->bInnerZone)
+			{
+				FScopedTransaction Transaction(LOCTEXT("VisualizerResetValue", "Visualizer Reset Value"));
+				EffectorActor->SetFlags(RF_Transactional);
+				EffectorActor->SetInnerExtent(FVector(50.f));
+				EffectorActor->Modify();
+				NotifyPropertyModified(EffectorActor, InnerExtentProperty, EPropertyChangeType::ValueSet);
+			}
+			else
+			{
+				FScopedTransaction Transaction(LOCTEXT("VisualizerResetValue", "Visualizer Reset Value"));
+				EffectorActor->SetFlags(RF_Transactional);
+				EffectorActor->SetOuterExtent(FVector(200.f));
+				EffectorActor->Modify();
+				NotifyPropertyModified(EffectorActor, OuterExtentProperty, EPropertyChangeType::ValueSet);
+			}
+		}
+		else if (EffectorActor->GetType() == EAvaClonerEffectorType::Plane)
+		{
+			FScopedTransaction Transaction(LOCTEXT("VisualizerResetValue", "Visualizer Reset Value"));
+			EffectorActor->SetFlags(RF_Transactional);
+			EffectorActor->SetPlaneSpacing(200.f);
+			EffectorActor->Modify();
+			NotifyPropertyModified(EffectorActor, PlaneSpacingProperty, EPropertyChangeType::ValueSet);
+		}
+		else if (EffectorActor->GetType() == EAvaClonerEffectorType::Sphere)
+		{
+			if (ComponentHitProxy->bInnerZone)
+			{
+				FScopedTransaction Transaction(LOCTEXT("VisualizerResetValue", "Visualizer Reset Value"));
+				EffectorActor->SetFlags(RF_Transactional);
+				EffectorActor->SetInnerRadius(50.f);
+				EffectorActor->Modify();
+				NotifyPropertyModified(EffectorActor, InnerRadiusProperty, EPropertyChangeType::ValueSet);
+			}
+			else
+			{
+				FScopedTransaction Transaction(LOCTEXT("VisualizerResetValue", "Visualizer Reset Value"));
+				EffectorActor->SetFlags(RF_Transactional);
+				EffectorActor->SetOuterRadius(200.f);
+				EffectorActor->Modify();
+				NotifyPropertyModified(EffectorActor, OuterRadiusProperty, EPropertyChangeType::ValueSet);
+			}
+		}
+	}
+	
+	return true;
+}
+
+bool FAvaEffectorActorVisualizer::IsEditing() const
+{
+	if (bEditingInnerZone || bEditingOuterZone)
+	{
+		return true;
+	}
+	
+	return Super::IsEditing();
+}
+
+void FAvaEffectorActorVisualizer::EndEditing()
+{
+	Super::EndEditing();
+
+	EffectorActorWeak.Reset();
+	bEditingInnerZone = false;
+	bEditingOuterZone = false;
+}
+
+#undef LOCTEXT_NAMESPACE
diff --git a/Engine/Plugins/Experimental/Avalanche/Source/AvalancheEffectorsEditor/Private/Effector/AvaEffectorDetailCustomization.cpp b/Engine/Plugins/Experimental/Avalanche/Source/AvalancheEffectorsEditor/Private/Effector/AvaEffectorDetailCustomization.cpp
new file mode 100644
--- /dev/null
+++ b/Engine/Plugins/Experimental/Avalanche/Source/AvalancheEffectorsEditor/Private/Effector/AvaEffectorDetailCustomization.cpp
@@ -0,0 +1,249 @@
+// Copyright Epic Games, Inc. All Rights Reserved.
+
+#include "AvaEffectorDetailCustomization.h"
+
+#include "AvaEffectorsEditorStyle.h"
+#include "DetailLayoutBuilder.h"
+#include "DetailWidgetRow.h"
+#include "Effector/AvaEffectorActor.h"
+#include "Modules/ModuleManager.h"
+#include "PropertyEditorModule.h"
+#include "PropertyHandle.h"
+#include "Styling/SlateBrush.h"
+#include "Widgets/Images/SImage.h"
+#include "Widgets/Input/SComboBox.h"
+
+#define LOCTEXT_NAMESPACE "AvaClonerEffectorDetailCustomization"
+
+DEFINE_LOG_CATEGORY_STATIC(LogAvaClonerEffectorDetailCustomization, Log, All);
+
+void FAvaEffectorDetailCustomization::CustomizeDetails(IDetailLayoutBuilder& DetailBuilder)
+{
+	EasingPropertyHandle = DetailBuilder.GetProperty(
+		GET_MEMBER_NAME_CHECKED(AAvaEffectorActor, Easing),
+		AAvaEffectorActor::StaticClass()
+	);
+
+	if (!EasingPropertyHandle.IsValid())
+	{
+		return;
+	}
+
+	if (IDetailPropertyRow* EasingRow = DetailBuilder.EditDefaultProperty(EasingPropertyHandle))
+	{
+		PopulateEasingInfos();
+
+		FDetailWidgetRow& CustomWidget = EasingRow->CustomWidget();
+
+		CustomWidget.NameContent()
+		.HAlign(EHorizontalAlignment::HAlign_Fill)
+		.VAlign(EVerticalAlignment::VAlign_Fill)
+		[
+			EasingPropertyHandle->CreatePropertyNameWidget()
+		];
+
+		CustomWidget.ValueContent()
+		.HAlign(EHorizontalAlignment::HAlign_Fill)
+		.VAlign(EVerticalAlignment::VAlign_Fill)
+		[
+			SNew(SComboBox<FName>)
+			.OptionsSource(&EasingNames)
+			.InitiallySelectedItem(GetCurrentEasingName())
+			.ToolTipText(LOCTEXT("EasingTooltip", "Easings sorted from most dramatic to least and specials at the end"))
+			.OnGenerateWidget(this, &FAvaEffectorDetailCustomization::OnGenerateEasingEntry)
+			.OnSelectionChanged(this, &FAvaEffectorDetailCustomization::OnSelectionChanged)
+			.Content()
+			[
+				OnGenerateEasingEntry(NAME_None)
+			]
+		];
+	}
+}
+
+void FAvaEffectorDetailCustomization::RegisterCustomSections() const
+{
+	static const FName PropertyEditor("PropertyEditor");
+	FPropertyEditorModule& PropertyModule = FModuleManager::GetModuleChecked<FPropertyEditorModule>(PropertyEditor);
+
+	static const FName ClassName = AAvaEffectorActor::StaticClass()->GetFName();
+
+	const TSharedRef<FPropertySection> StreamingSection = PropertyModule.FindOrCreateSection(ClassName, "Streaming", LOCTEXT("Effector.Streaming", "Streaming"));
+	StreamingSection->RemoveCategory("World Partition");
+	StreamingSection->RemoveCategory("Data Layers");
+	StreamingSection->RemoveCategory("HLOD");
+
+	const TSharedRef<FPropertySection> EffectorSection = PropertyModule.FindOrCreateSection(ClassName, "General", LOCTEXT("Effector.General", "General"));
+	EffectorSection->AddCategory(TEXT("Effector"));
+
+	const TSharedRef<FPropertySection> TypeSection = PropertyModule.FindOrCreateSection(ClassName, "Type", LOCTEXT("Effector.Type", "Type"));
+	TypeSection->AddCategory(TEXT("Type"));
+
+	const TSharedRef<FPropertySection> ModeSection = PropertyModule.FindOrCreateSection(ClassName, "Mode", LOCTEXT("Effector.Mode", "Mode"));
+	ModeSection->AddCategory(TEXT("Mode"));
+
+	const TSharedRef<FPropertySection> ForceSection = PropertyModule.FindOrCreateSection(ClassName, "Forces", LOCTEXT("Effector.Forces", "Forces"));
+	ForceSection->AddCategory(TEXT("Force"));
+}
+
+void FAvaEffectorDetailCustomization::PopulateEasingInfos()
+{
+	EasingEnum = StaticEnum<EAvaClonerEasing>();
+	if (EasingEnum.IsValid())
+	{
+		EasingNames.Empty();
+
+		// Sort from most dramatic to least IN then OUT then IN OUT, then specials
+		static const TArray<EAvaClonerEasing> SortedEasings
+		{
+			EAvaClonerEasing::InExpo,
+			EAvaClonerEasing::InCirc,
+			EAvaClonerEasing::InQuint,
+			EAvaClonerEasing::InQuart,
+			EAvaClonerEasing::InQuad,
+			EAvaClonerEasing::InCubic,
+			EAvaClonerEasing::InSine,
+			EAvaClonerEasing::OutExpo,
+			EAvaClonerEasing::OutCirc,
+			EAvaClonerEasing::OutQuint,
+			EAvaClonerEasing::OutQuart,
+			EAvaClonerEasing::OutQuad,
+			EAvaClonerEasing::OutCubic,
+			EAvaClonerEasing::OutSine,
+			EAvaClonerEasing::InOutExpo,
+			EAvaClonerEasing::InOutCirc,
+			EAvaClonerEasing::InOutQuint,
+			EAvaClonerEasing::InOutQuart,
+			EAvaClonerEasing::InOutQuad,
+			EAvaClonerEasing::InOutCubic,
+			EAvaClonerEasing::InOutSine,
+			EAvaClonerEasing::Linear,
+			EAvaClonerEasing::InBounce,
+			EAvaClonerEasing::InBack,
+			EAvaClonerEasing::InElastic,
+			EAvaClonerEasing::OutBounce,
+			EAvaClonerEasing::OutBack,
+			EAvaClonerEasing::OutElastic,
+			EAvaClonerEasing::InOutBounce,
+			EAvaClonerEasing::InOutBack,
+			EAvaClonerEasing::InOutElastic,
+			EAvaClonerEasing::Random
+		};
+
+		check(EasingEnum->GetMaxEnumValue() == SortedEasings.Num());
+
+		for (const EAvaClonerEasing& Easing : SortedEasings)
+		{
+			EasingNames.Add(FName(EasingEnum->GetNameStringByValue(static_cast<uint8>(Easing))));
+		}
+	}
+}
+
+TSharedRef<SWidget> FAvaEffectorDetailCustomization::OnGenerateEasingEntry(FName InName) const
+{
+	TSharedPtr<SWidget> ImageWidget;
+	TSharedPtr<SWidget> TextWidget;
+
+	static const FVector2D ImageSizeClosed(16.f, 16.f);
+	static const FVector2D ImageSizeOpened(32.f, 32.f);
+
+	// If none = update with current value else set fixed value once
+	if (InName == NAME_None)
+	{
+		SAssignNew(ImageWidget, SImage)
+		.DesiredSizeOverride(ImageSizeClosed)
+		.Image(this, &FAvaEffectorDetailCustomization::GetEasingImage, InName);
+
+		SAssignNew(TextWidget, STextBlock)
+		.Text(this, &FAvaEffectorDetailCustomization::GetEasingText, InName);
+	}
+	else
+	{
+		SAssignNew(ImageWidget, SImage)
+		.DesiredSizeOverride(ImageSizeOpened)
+		.Image(GetEasingImage(InName));
+
+		SAssignNew(TextWidget, STextBlock)
+		.Text(GetEasingText(InName));
+	}
+
+	return SNew(SHorizontalBox)
+		+ SHorizontalBox::Slot()
+		.AutoWidth()
+		.Padding(2.f)
+		[
+			ImageWidget.ToSharedRef()
+		]
+		+ SHorizontalBox::Slot()
+		.FillWidth(1.f)
+		.Padding(5.f, 2.f)
+		.HAlign(EHorizontalAlignment::HAlign_Left)
+		.VAlign(EVerticalAlignment::VAlign_Center)
+		[
+			TextWidget.ToSharedRef()
+		];
+}
+
+void FAvaEffectorDetailCustomization::OnSelectionChanged(FName InSelection, ESelectInfo::Type InSelectInfo) const
+{
+	if (EasingEnum.IsValid() && EasingPropertyHandle.IsValid())
+	{
+		const int32 Idx = EasingEnum->GetValueByNameString(InSelection.ToString());
+		if (Idx != INDEX_NONE && EasingPropertyHandle->SetValue(static_cast<uint8>(Idx), EPropertyValueSetFlags::DefaultFlags) != FPropertyAccess::Success)
+		{
+			UE_LOG(LogAvaClonerEffectorDetailCustomization, Warning, TEXT("ClonerEffectorDetailCustomization : Cannot set property value %s on selection"), *InSelection.ToString())
+		}
+	}
+}
+
+FName FAvaEffectorDetailCustomization::GetCurrentEasingName() const
+{
+	uint8 CurrentValue;
+	if (EasingPropertyHandle.IsValid() && EasingPropertyHandle->GetValue(CurrentValue) == FPropertyAccess::Result::Success)
+	{
+		if (EasingEnum.IsValid())
+		{
+			return FName(EasingEnum->GetNameStringByValue(CurrentValue));
+		}
+	}
+	return NAME_None;
+}
+
+const FSlateBrush* FAvaEffectorDetailCustomization::GetEasingImage(FName InName) const
+{
+	if (InName == NAME_None)
+	{
+		InName = GetCurrentEasingName();
+
+		// multiple values
+		if (InName == NAME_None)
+		{
+			return nullptr;
+		}
+	}
+
+	return FAvaEffectorsEditorStyle::Get().GetBrush(FName(TEXT("AvalancheIcons.Easing.") + InName.ToString()));
+}
+
+FText FAvaEffectorDetailCustomization::GetEasingText(FName InName) const
+{
+	if (InName == NAME_None)
+	{
+		InName = GetCurrentEasingName();
+
+		// multiple values
+		if (InName == NAME_None)
+		{
+			return LOCTEXT("MultipleValue", "Multiple values selected");
+		}
+	}
+
+	if (EasingEnum.IsValid())
+	{
+		const int32 EnumValue = EasingEnum->GetValueByNameString(InName.ToString());
+		return EasingEnum->GetDisplayNameTextByValue(EnumValue);
+	}
+
+	return FText::GetEmpty();
+}
+
+#undef LOCTEXT_NAMESPACE
diff --git a/Engine/Plugins/Experimental/Avalanche/Source/AvalancheEffectorsEditor/Private/Effector/AvaEffectorDetailCustomization.h b/Engine/Plugins/Experimental/Avalanche/Source/AvalancheEffectorsEditor/Private/Effector/AvaEffectorDetailCustomization.h
new file mode 100644
--- /dev/null
+++ b/Engine/Plugins/Experimental/Avalanche/Source/AvalancheEffectorsEditor/Private/Effector/AvaEffectorDetailCustomization.h
@@ -0,0 +1,45 @@
+// Copyright Epic Games, Inc. All Rights Reserved.
+
+#pragma once
+
+#include "IDetailCustomization.h"
+#include "Templates/SharedPointer.h"
+#include "Types/SlateEnums.h"
+#include "UObject/WeakObjectPtr.h"
+
+class IPropertyHandle;
+class SWidget;
+class UEnum;
+struct FSlateBrush;
+
+/* Used to customize effector actor properties in details panel */
+class FAvaEffectorDetailCustomization : public IDetailCustomization
+{
+public:
+	static TSharedRef<IDetailCustomization> MakeInstance()
+	{
+		return MakeShared<FAvaEffectorDetailCustomization>();
+	}
+
+	explicit FAvaEffectorDetailCustomization()
+	{
+		RegisterCustomSections();
+	}
+
+	virtual void CustomizeDetails(IDetailLayoutBuilder& DetailBuilder) override;
+
+protected:
+	void RegisterCustomSections() const;
+
+	void PopulateEasingInfos();
+	FName GetCurrentEasingName() const;
+	const FSlateBrush* GetEasingImage(FName InName) const;
+	FText GetEasingText(FName InName) const;
+
+	TSharedRef<SWidget> OnGenerateEasingEntry(FName InName) const;
+	void OnSelectionChanged(FName InSelection, ESelectInfo::Type InSelectInfo) const;
+
+	TWeakObjectPtr<UEnum> EasingEnum;
+	TArray<FName> EasingNames;
+	TSharedPtr<IPropertyHandle> EasingPropertyHandle;
+};
diff --git a/Engine/Plugins/Experimental/Avalanche/Source/AvalancheEffectorsEditor/Public/AvalancheEffectorsEditorModule.h b/Engine/Plugins/Experimental/Avalanche/Source/AvalancheEffectorsEditor/Public/AvalancheEffectorsEditorModule.h
new file mode 100644
--- /dev/null
+++ b/Engine/Plugins/Experimental/Avalanche/Source/AvalancheEffectorsEditor/Public/AvalancheEffectorsEditorModule.h
@@ -0,0 +1,31 @@
+// Copyright Epic Games, Inc. All Rights Reserved.
+
+#pragma once
+
+#include "IAvalancheInteractiveToolsModule.h"
+#include "Modules/ModuleInterface.h"
+
+class FComponentVisualizer;
+
+class FAvalancheEffectorsEditorModule : public IModuleInterface
+{
+public:
+	static FAvalancheEffectorsEditorModule& Get();
+
+	//~ Begin IModuleInterface interface
+	virtual void StartupModule() override;
+	virtual void ShutdownModule() override;
+	//~ End IModuleInterface interface
+
+private:
+	void PostEngineInit();
+
+	void RegisterTools(IAvalancheInteractiveToolsModule* InModule);
+
+	void RegisterComponentVisualizers();
+
+	void RegisterCustomLayouts();
+	void UnregisterCustomLayouts();
+
+	TArray<TSharedPtr<FComponentVisualizer>> Visualizers;
+};
diff --git a/Engine/Plugins/Experimental/Avalanche/Source/AvalancheInteractiveTools/Private/AvaInteractiveToolsStyle.cpp b/Engine/Plugins/Experimental/Avalanche/Source/AvalancheInteractiveTools/Private/AvaInteractiveToolsStyle.cpp
new file mode 100644
--- /dev/null
+++ b/Engine/Plugins/Experimental/Avalanche/Source/AvalancheInteractiveTools/Private/AvaInteractiveToolsStyle.cpp
@@ -0,0 +1,90 @@
+// Copyright Epic Games, Inc. All Rights Reserved.
+
+#include "AvaInteractiveToolsStyle.h"
+#include "Interfaces/IPluginManager.h"
+#include "Styling/SlateStyleRegistry.h"
+
+#define LOCTEXT_NAMESPACE "AvaInteractiveToolsStyle"
+
+TSharedPtr<FSlateStyleSet> FAvaInteractiveToolsStyle::StyleInstance;
+
+void FAvaInteractiveToolsStyle::Initialize()
+{
+	StyleInstance = Create();
+	FSlateStyleRegistry::RegisterSlateStyle(Get());
+}
+
+void FAvaInteractiveToolsStyle::Shutdown()
+{
+	FSlateStyleRegistry::UnRegisterSlateStyle(Get());
+	StyleInstance.Reset();
+}
+
+FName FAvaInteractiveToolsStyle::GetStyleSetName()
+{
+	static FName StyleSetName(TEXT("AvalancheInteractiveTools"));
+	return StyleSetName;
+}
+
+const ISlateStyle& FAvaInteractiveToolsStyle::Get()
+{
+	if (StyleInstance.IsValid() == false)
+	{
+		Initialize();
+	}
+
+	return *StyleInstance;
+}
+
+const FSlateBrush* FAvaInteractiveToolsStyle::GetBrush(FName PropertyName, const ANSICHAR* Specifier)
+{
+	return StyleInstance->GetBrush(PropertyName, Specifier);
+}
+
+#define IMAGE_BRUSH(RelativePath, ... ) FSlateImageBrush(Style->RootToContentDir(RelativePath, TEXT(".png")), __VA_ARGS__)
+#define IMAGE_BRUSH_SVG(RelativePath, ...) FSlateVectorImageBrush(Style->RootToContentDir(RelativePath, TEXT(".svg")), __VA_ARGS__)
+
+#define CORE_IMAGE_BRUSH(RelativePath, ...) FSlateImageBrush(Style->RootToCoreContentDir(RelativePath, TEXT(".png")), __VA_ARGS__)
+#define CORE_IMAGE_BRUSH_SVG(RelativePath, ...) FSlateVectorImageBrush(Style->RootToCoreContentDir(RelativePath, TEXT(".svg")), __VA_ARGS__)
+
+TSharedRef<FSlateStyleSet> FAvaInteractiveToolsStyle::Create()
+{
+	const FVector2f Icon16x16(16.0f, 16.0f);
+	const FVector2f Icon20x20(20.0f, 20.0f);
+
+	TSharedRef<FSlateStyleSet> Style = MakeShared<FSlateStyleSet>(GetStyleSetName());
+
+	TSharedPtr<IPlugin> Plugin = IPluginManager::Get().FindPlugin(TEXT("Avalanche"));
+	check(Plugin.IsValid());
+	if (Plugin.IsValid())
+	{
+		Style->SetContentRoot(Plugin->GetBaseDir() / TEXT("Resources"));
+	}
+
+	Style->SetCoreContentRoot(FPaths::EngineContentDir() / TEXT("Editor/Slate"));
+
+	// Path Constants
+	const FString UEContentEditorSlatePath = FPaths::EngineContentDir() / TEXT("Editor/Slate");
+	const FString UEContentEditorSlateIconsPath = UEContentEditorSlatePath / TEXT("Icons");
+
+	// Categories
+	Style->Set("AvalancheInteractiveTools.Category_2D",     new IMAGE_BRUSH("Icons/ToolboxIcons/rectangle", Icon20x20));
+	Style->Set("AvalancheInteractiveTools.Category_3D",     new IMAGE_BRUSH("Icons/ToolboxIcons/cube", Icon20x20));
+	Style->Set("AvalancheInteractiveTools.Category_Actor",  new CORE_IMAGE_BRUSH_SVG("Starship/AssetIcons/Actor_16", Icon16x16));
+	Style->Set("AvalancheInteractiveTools.Category_Layout", new IMAGE_BRUSH("Icons/ToolboxIcons/layoutgrid", Icon20x20));
+
+	// Actor Tools
+	Style->Set("AvalancheInteractiveTools.Tool_Actor_Null", new CORE_IMAGE_BRUSH(TEXT("Icons/SequencerIcons/icon_Sequencer_Move_24x"), Icon16x16));
+	Style->Set("AvalancheInteractiveTools.Tool_Actor_Spline", new IMAGE_BRUSH_SVG("Icons/ToolboxIcons/Toolbar_Spline", Icon20x20));
+	Style->Set("Tool_Actor_Null", new CORE_IMAGE_BRUSH(TEXT("Icons/SequencerIcons/icon_Sequencer_Move_24x"), Icon20x20));
+	Style->Set("Tool_Actor_Spline", new IMAGE_BRUSH_SVG("Icons/ToolboxIcons/Toolbar_Spline", Icon20x20));
+
+	return Style;
+}
+
+#undef IMAGE_BRUSH
+#undef IMAGE_BRUSH_SVG
+#undef CORE_IMAGE_BRUSH
+#undef CORE_IMAGE_BRUSH_SVG
+
+#undef LOCTEXT_NAMESPACE
diff --git a/Engine/Plugins/Experimental/Avalanche/Source/AvalancheInteractiveTools/Private/AvaInteractiveToolsStyle.h b/Engine/Plugins/Experimental/Avalanche/Source/AvalancheInteractiveTools/Private/AvaInteractiveToolsStyle.h
new file mode 100644
--- /dev/null
+++ b/Engine/Plugins/Experimental/Avalanche/Source/AvalancheInteractiveTools/Private/AvaInteractiveToolsStyle.h
@@ -0,0 +1,25 @@
+// Copyright Epic Games, Inc. All Rights Reserved.
+
+#pragma once
+
+#include "Styling/ISlateStyle.h"
+#include "Styling/SlateStyle.h"
+
+struct FSlateBrush;
+
+class FAvaInteractiveToolsStyle
+{
+public:
+	static void Initialize();
+	static void Shutdown();
+
+	static FName GetStyleSetName();
+	static const ISlateStyle& Get();
+
+	static const FSlateBrush* GetBrush(FName PropertyName, const ANSICHAR* Specifier = nullptr);
+
+private:
+	static TSharedPtr<FSlateStyleSet> StyleInstance;
+
+	static TSharedRef<FSlateStyleSet> Create();
+};
diff --git a/Engine/Plugins/Experimental/Avalanche/Source/AvalancheInteractiveTools/Private/Tools/AvaInteractiveToolsActorTool.cpp b/Engine/Plugins/Experimental/Avalanche/Source/AvalancheInteractiveTools/Private/Tools/AvaInteractiveToolsActorTool.cpp
new file mode 100644
--- /dev/null
+++ b/Engine/Plugins/Experimental/Avalanche/Source/AvalancheInteractiveTools/Private/Tools/AvaInteractiveToolsActorTool.cpp
@@ -0,0 +1,22 @@
+// Copyright Epic Games, Inc. All Rights Reserved.
+
+#include "Tools/AvaInteractiveToolsActorTool.h"
+#include "AvalancheInteractiveToolsModule.h"
+#include "Builders/AvaInteractiveToolsActorToolBuilder.h"
+
+UAvaInteractiveToolsActorTool::UAvaInteractiveToolsActorTool()
+{
+}
+
+FAvaInteractiveToolsToolParameters UAvaInteractiveToolsActorTool::GetToolParameters() const
+{
+	// This is never really needed, it should work regardless. OOP consistency!
+	return UAvaInteractiveToolsActorToolBuilder::CreateToolParameters(
+		Category,
+		Command,
+		Identifier,
+		Priority,
+		ActorClass,
+		GetClass()
+	);
+}
diff --git a/Engine/Plugins/Experimental/Avalanche/Source/AvalancheInteractiveTools/Private/Tools/AvaInteractiveToolsActorToolNull.cpp b/Engine/Plugins/Experimental/Avalanche/Source/AvalancheInteractiveTools/Private/Tools/AvaInteractiveToolsActorToolNull.cpp
new file mode 100644
--- /dev/null
+++ b/Engine/Plugins/Experimental/Avalanche/Source/AvalancheInteractiveTools/Private/Tools/AvaInteractiveToolsActorToolNull.cpp
@@ -0,0 +1,32 @@
+// Copyright Epic Games, Inc. All Rights Reserved.
+
+#include "Tools/AvaInteractiveToolsActorToolNull.h"
+#include "AvaInteractiveToolsCommands.h"
+#include "AvalancheInteractiveToolsModule.h"
+#include "Builders/AvaInteractiveToolsToolBuilder.h"
+#include "Framework/AvaNullActor.h"
+
+UAvaInteractiveToolsActorToolNull::UAvaInteractiveToolsActorToolNull()
+{
+	ActorClass = AAvaNullActor::StaticClass();
+}
+
+FName UAvaInteractiveToolsActorToolNull::GetCategoryName()
+{
+	return FAvalancheInteractiveToolsModule::CategoryNameActor;
+}
+
+FAvaInteractiveToolsToolParameters UAvaInteractiveToolsActorToolNull::GetToolParameters() const
+{
+	return {
+		FAvaInteractiveToolsCommands::Get().Tool_Actor_Null,
+		TEXT("Null Actor Tool"),
+		1000,
+		FAvalancheInteractiveToolsCreateBuilder::CreateLambda(
+			[](UEdMode* InEdMode)
+			{
+				return UAvaInteractiveToolsToolBuilder::CreateToolBuilder<UAvaInteractiveToolsActorToolNull>(InEdMode);
+			}),
+		ActorClass
+	};
+}
diff --git a/Engine/Plugins/Experimental/Avalanche/Source/AvalancheInteractiveTools/Private/Tools/AvaInteractiveToolsActorToolSpline.cpp b/Engine/Plugins/Experimental/Avalanche/Source/AvalancheInteractiveTools/Private/Tools/AvaInteractiveToolsActorToolSpline.cpp
new file mode 100644
--- /dev/null
+++ b/Engine/Plugins/Experimental/Avalanche/Source/AvalancheInteractiveTools/Private/Tools/AvaInteractiveToolsActorToolSpline.cpp
@@ -0,0 +1,32 @@
+// Copyright Epic Games, Inc. All Rights Reserved.
+
+#include "Tools/AvaInteractiveToolsActorToolSpline.h"
+#include "AvaInteractiveToolsCommands.h"
+#include "AvalancheInteractiveToolsModule.h"
+#include "Builders/AvaInteractiveToolsToolBuilder.h"
+#include "Framework/AvaSplineActor.h"
+
+UAvaInteractiveToolsActorToolSpline::UAvaInteractiveToolsActorToolSpline()
+{
+	ActorClass = AAvaSplineActor::StaticClass();
+}
+
+FName UAvaInteractiveToolsActorToolSpline::GetCategoryName()
+{
+	return FAvalancheInteractiveToolsModule::CategoryNameActor;
+}
+
+FAvaInteractiveToolsToolParameters UAvaInteractiveToolsActorToolSpline::GetToolParameters() const
+{
+	return {
+		FAvaInteractiveToolsCommands::Get().Tool_Actor_Spline,
+		TEXT("Spline Actor Tool"),
+		5000,
+		FAvalancheInteractiveToolsCreateBuilder::CreateLambda(
+			[](UEdMode* InEdMode)
+			{
+				return UAvaInteractiveToolsToolBuilder::CreateToolBuilder<UAvaInteractiveToolsActorToolSpline>(InEdMode);
+			}),
+		ActorClass
+	};
+}
diff --git a/Engine/Plugins/Experimental/Avalanche/Source/AvalancheInteractiveTools/Private/Tools/AvaInteractiveToolsToolBase.cpp b/Engine/Plugins/Experimental/Avalanche/Source/AvalancheInteractiveTools/Private/Tools/AvaInteractiveToolsToolBase.cpp
new file mode 100644
--- /dev/null
+++ b/Engine/Plugins/Experimental/Avalanche/Source/AvalancheInteractiveTools/Private/Tools/AvaInteractiveToolsToolBase.cpp
@@ -0,0 +1,750 @@
+// Copyright Epic Games, Inc. All Rights Reserved.
+
+#include "Tools/AvaInteractiveToolsToolBase.h"
+#include "AvaInteractiveToolsSettings.h"
+#include "AvalancheInteractiveToolsModule.h"
+#include "AvaViewportUtils.h"
+#include "BaseBehaviors/SingleClickBehavior.h"
+#include "BaseBehaviors/SingleKeyCaptureBehavior.h"
+#include "ContextObjectStore.h"
+#include "Editor/EditorEngine.h"
+#include "EditorViewportClient.h"
+#include "EdMode/AvaInteractiveToolsEdMode.h"
+#include "Engine/World.h"
+#include "Framework/Application/SlateApplication.h"
+#include "IAvalancheInteractiveToolsModeDetailsObject.h"
+#include "IAvalancheInteractiveToolsModeDetailsObjectProvider.h"
+#include "InteractiveToolManager.h"
+#include "InteractiveToolsContext.h"
+#include "Planners/AvaInteractiveToolsToolViewportAreaPlanner.h"
+#include "Planners/AvaInteractiveToolsToolViewportPlanner.h"
+#include "Planners/AvaInteractiveToolsToolViewportPointPlanner.h"
+#include "ToolContextInterfaces.h"
+#include "Toolkits/BaseToolkit.h"
+#include "UnrealClient.h"
+#include "ViewportClient/IAvaViewportClient.h"
+
+#define LOCTEXT_NAMESPACE "AvaInteractiveToolsToolBase"
+
+namespace UE::AvalancheInteractiveTools::Private
+{
+	static const FName AvaITFToolPresetMenuName = "AvaITFToolPresetMenu";
+}
+
+UAvaInteractiveToolsRightClickBehavior::UAvaInteractiveToolsRightClickBehavior()
+{
+	HitTestOnRelease = false;
+	SetUseRightMouseButton();
+}
+
+void UAvaInteractiveToolsRightClickBehavior::Clicked(const FInputDeviceState& Input, const FInputCaptureData& Data)
+{
+	Target->OnUpdateModifierState(UAvaInteractiveToolsToolBase::BID_Cancel, true);
+}
+
+UObject* UAvaInteractiveToolsToolBase::GetDetailsObjectFromActor(AActor* InActor)
+{
+	if (!InActor)
+	{
+		return nullptr;
+	}
+
+	if (InActor->Implements<UAvalancheInteractiveToolsModeDetailsObject>())
+	{
+		return InActor;
+	}
+	else if (InActor->Implements<UAvalancheInteractiveToolsModeDetailsObjectProvider>())
+	{
+		return IAvalancheInteractiveToolsModeDetailsObjectProvider::Execute_GetModeDetailsObject(InActor);
+	}
+
+	TArray<UActorComponent*> Components;
+	InActor->GetComponents<UActorComponent>(Components);
+
+	for (UActorComponent* Component : Components)
+	{
+		if (Component->Implements<UAvalancheInteractiveToolsModeDetailsObject>())
+		{
+			return Component;
+			break;
+		}
+		else if (Component->Implements<UAvalancheInteractiveToolsModeDetailsObjectProvider>())
+		{
+			if (UObject* DetailsObject = IAvalancheInteractiveToolsModeDetailsObjectProvider::Execute_GetModeDetailsObject(Component))
+			{
+				return DetailsObject;
+			}
+		}
+	}
+
+	return nullptr;
+}
+
+void UAvaInteractiveToolsToolBase::Setup()
+{
+	Super::Setup();
+
+	UInteractiveToolManager* ToolManager = GetToolManager();
+	bool bReactivated = false;
+	UAvaInteractiveToolsEdMode* AvaInteractiveToolsEdMode = nullptr;
+
+	if (UContextObjectStore* ContextStore = ToolManager->GetContextObjectStore())
+	{
+		AvaInteractiveToolsEdMode = Cast<UAvaInteractiveToolsEdMode>(ContextStore->FindContextByClass(UAvaInteractiveToolsEdMode::StaticClass()));
+
+		if (AvaInteractiveToolsEdMode)
+		{
+			bReactivated = AvaInteractiveToolsEdMode->GetLastActiveTool() == ToolManager->GetActiveToolName(EToolSide::Left);
+		}
+	}
+
+	if (!CanActivate(bReactivated))
+	{
+		ToolManager->PostActiveToolShutdownRequest(this, EToolShutdownType::Cancel);
+	}
+
+	if (AvaInteractiveToolsEdMode)
+	{
+		AvaInteractiveToolsEdMode->OnToolSetup(this);
+	}
+
+	Activate(bReactivated);
+}
+
+void UAvaInteractiveToolsToolBase::Shutdown(EToolShutdownType ShutdownType)
+{
+	switch (ShutdownType)
+	{
+		case EToolShutdownType::Cancel:
+			OnCancel();
+			break;
+
+		case EToolShutdownType::Accept:
+		case EToolShutdownType::Completed:
+			OnComplete();
+			break;
+	}
+
+	if (ViewportPlanner)
+	{
+		ViewportPlanner->Shutdown(ShutdownType);
+		ViewportPlanner = nullptr;
+	}
+
+	Super::Shutdown(ShutdownType);
+}
+
+void UAvaInteractiveToolsToolBase::DrawHUD(FCanvas* InCanvas, IToolsContextRenderAPI* InRenderAPI)
+{
+	Super::DrawHUD(InCanvas, InRenderAPI);
+
+	if (ViewportPlanner)
+	{
+		ViewportPlanner->DrawHUD(InCanvas, InRenderAPI);
+	}
+}
+
+void UAvaInteractiveToolsToolBase::Render(IToolsContextRenderAPI* InRenderAPI)
+{
+	Super::Render(InRenderAPI);
+
+	if (ViewportPlanner)
+	{
+		ViewportPlanner->Render(InRenderAPI);
+	}
+}
+
+void UAvaInteractiveToolsToolBase::OnTick(float InDeltaTime)
+{
+	Super::OnTick(InDeltaTime);
+
+	if (ViewportPlanner)
+	{
+		ViewportPlanner->OnTick(InDeltaTime);
+	}
+}
+
+bool UAvaInteractiveToolsToolBase::SupportsDefaultAction() const
+{
+	return ViewportPlannerClass == UAvaInteractiveToolsToolViewportPointPlanner::StaticClass()
+		|| ViewportPlannerClass == UAvaInteractiveToolsToolViewportAreaPlanner::StaticClass();
+}
+
+void UAvaInteractiveToolsToolBase::DefaultAction()
+{
+	RequestShutdown(EToolShutdownType::Completed);
+}
+
+bool UAvaInteractiveToolsToolBase::ConditionalIdentityRotation() const
+{
+	switch (GetDefault<UAvaInteractiveToolsSettings>()->DefaultActionActorAlignment)
+	{
+		case EAvaInteractiveToolsDefaultActionAlignment::Camera:
+			if (!bPerformingDefaultAction)
+			{
+				return false;
+			}
+
+			return !IsAvalancheViewport();
+
+		default:
+		case EAvaInteractiveToolsDefaultActionAlignment::Axis:
+			return bPerformingDefaultAction;
+	}
+}
+
+AActor* UAvaInteractiveToolsToolBase::SpawnActor(TSubclassOf<AActor> InActorClass, bool bInPreview, FString* InActorLabelOverride) const
+{
+	if (IToolsContextQueriesAPI* ContextAPI = GetToolManager()->GetContextQueriesAPI())
+	{
+		if (FViewport* Viewport = ContextAPI->GetFocusedViewport())
+		{
+			if (TSharedPtr<IAvaViewportClient> AvaViewportClient = FAvaViewportUtils::GetAvaViewportClient(Viewport))
+			{
+				const FVector2f ViewportSize = AvaViewportClient->GetViewportSize();
+				return SpawnActor(InActorClass, EAvaViewportStatus::Focused, ViewportSize * 0.5, bInPreview, InActorLabelOverride);
+			}
+		}
+	}
+
+	return nullptr;
+}
+
+AActor* UAvaInteractiveToolsToolBase::SpawnActor(TSubclassOf<AActor> InActorClass, EAvaViewportStatus InViewportStatus, 
+	const FVector2f& InViewportPosition, bool bInPreview, FString* InActorLabelOverride) const
+{
+	UWorld* World = nullptr;;
+	FVector SpawnLocation = FVector::ZeroVector;
+	FRotator SpawnRotation = FRotator::ZeroRotator;
+	if (IToolsContextQueriesAPI* ContextAPI = GetToolManager()->GetContextQueriesAPI())
+	{
+		const bool bUseIdentityLocation = UseIdentityLocation();
+		const bool bUseIdentityRotation = UseIdentityRotation();
+
+		if (FViewport* Viewport = ContextAPI->GetFocusedViewport())
+		{
+			FVector CameraForward = FVector(1, 0, 0);
+
+			if (TSharedPtr<IAvaViewportClient> AvaViewportClient = FAvaViewportUtils::GetAvaViewportClient(Viewport))
+			{
+				CameraForward = AvaViewportClient->GetViewportViewTransform().TransformVectorNoScale(FVector::ForwardVector);
+				CameraForward.Z = 0;
+			}
+
+			const bool bSuccess = ViewportPositionToWorldPositionAndOrientation(
+				InViewportStatus,
+				InViewportPosition,
+				World,
+				SpawnLocation,
+				SpawnRotation
+			);
+
+			if (!bSuccess)
+			{
+				return nullptr;
+			}
+
+			if (bUseIdentityLocation)
+			{
+				SpawnLocation = FVector::ZeroVector;
+			}
+
+			if (bUseIdentityRotation)
+			{
+				if (CameraForward.Dot(FVector::ForwardVector) >= 0)
+				{
+					SpawnRotation = FRotator::ZeroRotator;
+				}
+				else
+				{
+					SpawnRotation = FRotator(0, 180, 0);
+				}
+			}
+		}
+	}
+
+	FActorSpawnParameters SpawnParams;
+	SpawnParams.bNoFail = true;
+
+	const bool bUseOverride = !bInPreview && InActorLabelOverride && !InActorLabelOverride->IsEmpty();
+
+	if (bInPreview)
+	{
+		SpawnParams.bHideFromSceneOutliner = true;
+		SpawnParams.bTemporaryEditorActor = true;
+		SpawnParams.Name = FName(TEXT("AvaITFPreviewActor"));
+		SpawnParams.NameMode = FActorSpawnParameters::ESpawnActorNameMode::Requested;
+		SpawnParams.ObjectFlags |= RF_Transient;
+	}
+	else if (bUseOverride)
+	{
+		SpawnParams.Name = FName(*InActorLabelOverride);
+		SpawnParams.NameMode = FActorSpawnParameters::ESpawnActorNameMode::Requested;
+		SpawnParams.ObjectFlags |= RF_Transactional;
+	}
+
+	AActor* NewActor = World->SpawnActor<AActor>(InActorClass, SpawnLocation, SpawnRotation, SpawnParams);
+
+	if (bUseOverride)
+	{
+		FActorLabelUtilities::SetActorLabelUnique(NewActor, *InActorLabelOverride);
+	}
+	else
+	{
+		FActorLabelUtilities::SetActorLabelUnique(NewActor, NewActor->GetDefaultActorLabel());
+	}
+
+	if (UObject* DetailsObject = GetDetailsObjectFromActor(NewActor))
+	{
+		SetToolkitSettingsObject(DetailsObject);
+	}
+
+	return NewActor;
+}
+
+void UAvaInteractiveToolsToolBase::SetToolkitSettingsObject(UObject* InObject) const
+{
+	if (UInteractiveToolManager* ToolManager = GetToolManager())
+	{
+		if (UContextObjectStore* ContextStore = ToolManager->GetContextObjectStore())
+		{
+			if (UAvaInteractiveToolsEdMode* AvaInteractiveToolsEdMode =
+				Cast<UAvaInteractiveToolsEdMode>(ContextStore->FindContextByClass(UAvaInteractiveToolsEdMode::StaticClass())))
+			{
+				if (TSharedPtr<FModeToolkit> Toolkit = AvaInteractiveToolsEdMode->GetToolkit().Pin())
+				{
+					Toolkit->SetModeSettingsObject(InObject);
+				}
+			}
+		}
+	}
+}
+
+bool UAvaInteractiveToolsToolBase::ViewportPositionToWorldPositionAndOrientation(EAvaViewportStatus InViewportStatus, 
+	const FVector2f& InViewportPosition, UWorld*& OutWorld, FVector& OutPosition, FRotator& OutRotation) const
+{
+	const float CameraDistance = GetDefault<UAvaInteractiveToolsSettings>()->CameraDistance;
+
+	return ViewportPositionToWorldPositionAndOrientation(InViewportStatus, InViewportPosition, CameraDistance, OutWorld, OutPosition, OutRotation);
+}
+
+bool UAvaInteractiveToolsToolBase::ViewportPositionToWorldPositionAndOrientation(EAvaViewportStatus InViewportStatus,
+	const FVector2f& InViewportPosition, float InDistance, UWorld*& OutWorld, FVector& OutPosition, FRotator& OutRotation) const
+{
+	if (IToolsContextQueriesAPI* ContextAPI = GetToolManager()->GetContextQueriesAPI())
+	{
+		if (UWorld* World = ContextAPI->GetCurrentEditingWorld())
+		{
+			if (TSharedPtr<IAvaViewportClient> AvaViewportClient = FAvaViewportUtils::GetAvaViewportClient(GetViewport(InViewportStatus)))
+			{
+				if (const FEditorViewportClient* EditorViewportClient = AvaViewportClient->AsEditorViewportClient())
+				{
+					const FViewportCameraTransform ViewTransform = EditorViewportClient->GetViewTransform();
+					const FVector2f ViewportSize = AvaViewportClient->GetViewportSize();
+
+					OutWorld = World;
+					OutRotation = ViewTransform.GetRotation();
+					OutPosition = AvaViewportClient->ViewportPositionToWorldPosition(InViewportPosition, InDistance);
+
+					return true;
+				}
+			}
+		}
+	}
+
+	return false;
+}
+
+FViewport* UAvaInteractiveToolsToolBase::GetViewport(EAvaViewportStatus InViewportStatus) const
+{
+	if (UInteractiveToolManager* ToolManager = GetToolManager())
+	{
+		if (IToolsContextQueriesAPI* ContextAPI = ToolManager->GetContextQueriesAPI())
+		{
+			FViewport* Viewport = nullptr;
+
+			switch (InViewportStatus)
+			{
+				case EAvaViewportStatus::Hovered:
+					Viewport = ContextAPI->GetHoveredViewport();
+					break;
+
+				case EAvaViewportStatus::Focused:
+					Viewport = ContextAPI->GetFocusedViewport();
+					break;
+
+				default:
+					checkNoEntry();
+					return nullptr;
+			}
+
+			if (FAvaViewportUtils::GetAsEditorViewportClient(Viewport))
+			{
+				return Viewport;
+			}
+		}
+	}
+
+	return nullptr;
+}
+
+void UAvaInteractiveToolsToolBase::OnViewportPlannerComplete()
+{
+	if (PreviewActor)
+	{
+		PreviewActor->Destroy();
+	}
+
+	RequestShutdown(EToolShutdownType::Completed);
+}
+
+FInputRayHit UAvaInteractiveToolsToolBase::CanBeginSingleClickAndDragSequence(const FInputDeviceRay& InPressPos)
+{
+	// Always hits every place in the viewport
+	return FInputRayHit(0);
+}
+
+void UAvaInteractiveToolsToolBase::OnClickPress(const FInputDeviceRay& InPressPos)
+{
+	// Nothing
+}
+
+void UAvaInteractiveToolsToolBase::OnDragStart(const FInputDeviceRay& InDragPos)
+{
+	if (LeftClickBehavior)
+	{
+		// Fake a click at the start position
+		OnClickRelease(LeftClickBehavior->GetInitialMouseDownRay(), true);
+	}
+}
+
+void UAvaInteractiveToolsToolBase::OnClickDrag(const FInputDeviceRay& InDragPos)
+{
+}
+
+void UAvaInteractiveToolsToolBase::OnClickRelease(const FInputDeviceRay& InReleasePos, bool bInIsDragOperation)
+{
+	if (ViewportPlanner)
+	{
+		ViewportPlanner->OnClicked(InReleasePos);
+	}
+}
+
+void UAvaInteractiveToolsToolBase::OnTerminateSingleClickAndDragSequence()
+{
+	RequestShutdown(EToolShutdownType::Cancel);
+}
+
+FInputRayHit UAvaInteractiveToolsToolBase::IsHitByClick(const FInputDeviceRay& InClickPos)
+{
+	return FInputRayHit(0);
+}
+
+void UAvaInteractiveToolsToolBase::OnClicked(const FInputDeviceRay& InClickPos)
+{
+	// Right click cancel
+	RequestShutdown(EToolShutdownType::Cancel);
+}
+
+void UAvaInteractiveToolsToolBase::OnUpdateModifierState(int InModifierID, bool bInIsOn)
+{
+	switch (InModifierID)
+	{
+		case BID_Cancel:
+			if (bInIsOn)
+			{
+				RequestShutdown(EToolShutdownType::Cancel);
+			}
+
+			break;
+	}
+}
+
+bool UAvaInteractiveToolsToolBase::CanActivate(bool bInReactivate) const
+{
+	if (!bInReactivate)
+	{
+
+		return true;
+	}
+
+	return SupportsDefaultAction();
+}
+
+void UAvaInteractiveToolsToolBase::Activate(bool bInReactivate)
+{
+	if (!CanActivate(bInReactivate))
+	{
+		RequestShutdown(EToolShutdownType::Cancel);
+		return;
+	}
+
+	FAvalancheInteractiveToolsModule::Get().OnToolActivated();
+
+	const bool bUsePresetMenu = ShouldUsePresetMenu();
+	const bool bSupportsDefaultAction = !bUsePresetMenu && SupportsDefaultAction();
+	const bool bForceDefaultAction = !bUsePresetMenu && ShouldForceDefaultAction();
+
+	if (bUsePresetMenu)
+	{
+		ShowPresetMenu();
+	}
+	else if (!bForceDefaultAction && (!bInReactivate || !bSupportsDefaultAction))
+	{
+		OnActivate();
+
+		if (!OnBegin())
+		{
+			RequestShutdown(EToolShutdownType::Cancel);
+			return;
+		}
+
+		OnPostBegin();
+	}
+	else if (bSupportsDefaultAction)
+	{
+		bPerformingDefaultAction = true;
+		DefaultAction();
+	}
+	else
+	{
+		if (bForceDefaultAction)
+		{
+			UE_LOG(LogAvaInteractiveTools, Warning, TEXT("Alt used to force a tool into using the default action, but it does not support the default action."));
+		}
+
+		RequestShutdown(EToolShutdownType::Cancel);
+	}
+}
+
+void UAvaInteractiveToolsToolBase::OnActivate()
+{
+}
+
+bool UAvaInteractiveToolsToolBase::OnBegin()
+{
+	if (ViewportPlannerClass == nullptr)
+	{
+		return false;
+	}
+
+	BeginTransaction();
+	return true;
+}
+
+void UAvaInteractiveToolsToolBase::OnPostBegin()
+{
+	LeftClickBehavior = NewObject<UAvaSingleClickAndDragBehavior>(this);
+	LeftClickBehavior->Initialize(this);
+	LeftClickBehavior->bSupportsDrag = ViewportPlannerClass.Get() && ViewportPlannerClass->IsChildOf<UAvaInteractiveToolsToolViewportAreaPlanner>();
+	AddInputBehavior(LeftClickBehavior);
+
+	RightClickBehavior = NewObject<UAvaInteractiveToolsRightClickBehavior>(this);
+	RightClickBehavior->Initialize(this);
+	AddInputBehavior(RightClickBehavior);
+
+	EscapeKeyBehavior = NewObject<USingleKeyCaptureBehavior>(this);
+	EscapeKeyBehavior->Initialize(static_cast<IClickBehaviorTarget*>(this), BID_Cancel, EKeys::Escape);
+	AddInputBehavior(EscapeKeyBehavior);
+
+	ViewportPlanner = NewObject<UAvaInteractiveToolsToolViewportPlanner>(this, ViewportPlannerClass);
+	ViewportPlanner->Setup(this);
+}
+
+void UAvaInteractiveToolsToolBase::OnCancel()
+{
+	CancelTransaction();
+
+	if (PreviewActor)
+	{
+		PreviewActor->Destroy();
+		PreviewActor = nullptr;
+	}
+
+	SetToolkitSettingsObject(nullptr);
+}
+
+void UAvaInteractiveToolsToolBase::OnComplete()
+{
+	EndTransaction();
+
+	if (PreviewActor)
+	{
+		PreviewActor->Destroy();
+		PreviewActor = nullptr;
+	}
+}
+
+void UAvaInteractiveToolsToolBase::BeginTransaction()
+{
+	GetToolManager()->BeginUndoTransaction(LOCTEXT("AvalancheInteractiveToolsTool", "Motion Design Interactive Tools Box Tool"));
+}
+
+void UAvaInteractiveToolsToolBase::EndTransaction()
+{
+	GetToolManager()->EndUndoTransaction();
+}
+
+void UAvaInteractiveToolsToolBase::CancelTransaction()
+{
+	// Doesn't exist
+	//GetToolManager()->CancelUndoTransaction();
+	GetToolManager()->EndUndoTransaction();
+}
+
+void UAvaInteractiveToolsToolBase::RequestShutdown(EToolShutdownType InShutdownType)
+{
+	if (UContextObjectStore* ContextStore = GetToolManager()->GetContextObjectStore())
+	{
+		if (UAvaInteractiveToolsEdMode* AvaInteractiveToolsEdMode = Cast<UAvaInteractiveToolsEdMode>(ContextStore->FindContextByClass(UAvaInteractiveToolsEdMode::StaticClass())))
+		{
+			AvaInteractiveToolsEdMode->OnToolShutdown(this, InShutdownType);
+		}
+	}
+
+	if (ViewportPlanner)
+	{
+		ViewportPlanner->Shutdown(InShutdownType);
+	}
+
+	SetToolkitSettingsObject(nullptr);
+
+	GetToolManager()->PostActiveToolShutdownRequest(this, InShutdownType);
+
+	FAvalancheInteractiveToolsModule::Get().OnToolDeactivated();
+}
+
+bool UAvaInteractiveToolsToolBase::IsAvalancheViewport() const
+{
+	if (IToolsContextQueriesAPI* ContextAPI = GetToolManager()->GetContextQueriesAPI())
+	{
+		if (FViewport* Viewport = ContextAPI->GetFocusedViewport())
+		{
+			if (TSharedPtr<IAvaViewportClient> AvaViewportClient = FAvaViewportUtils::GetAvaViewportClient(Viewport))
+			{
+				return AvaViewportClient->IsAvalancheViewport();
+			}
+		}
+	}
+
+	return false;
+}
+
+bool UAvaInteractiveToolsToolBase::ShouldForceDefaultAction() const
+{
+	return FSlateApplication::Get().GetModifierKeys().IsAltDown();
+}
+
+bool UAvaInteractiveToolsToolBase::ShouldUsePresetMenu() const
+{
+	const FAvaInteractiveToolsToolParameters ToolParams = GetToolParameters();
+
+	if (ToolParams.Presets.IsEmpty())
+	{
+		return false;
+	}
+
+	return SupportsDefaultAction() && FSlateApplication::Get().GetModifierKeys().IsShiftDown();
+}
+
+void UAvaInteractiveToolsToolBase::ShowPresetMenu()
+{
+	using namespace UE::AvalancheInteractiveTools::Private;
+
+	RegisterPresetMenu();
+
+	TSharedPtr<SWidget> Parent = FSlateApplication::Get().FindBestParentWindowForDialogs(nullptr);
+
+	if (!Parent.IsValid())
+	{
+		return;
+	}
+
+	FToolMenuContext Context;
+	Context.AddObject(this);
+
+	FSlateApplication::Get().PushMenu(
+		Parent.ToSharedRef(),
+		FWidgetPath(),
+		UToolMenus::Get()->GenerateWidget(AvaITFToolPresetMenuName, Context),
+		FSlateApplication::Get().GetCursorPos(),
+		FPopupTransitionEffect::ContextMenu
+	);
+}
+
+void UAvaInteractiveToolsToolBase::OnPresetSelected(TStrongObjectPtr<UAvaInteractiveToolsToolBase> InToolPtr, FName InPresetName)
+{
+	UAvaInteractiveToolsToolBase* Tool = InToolPtr.Get();
+
+	if (!Tool)
+	{
+		return;
+	}
+
+	const FAvaInteractiveToolsToolParameters ToolParams = Tool->GetToolParameters();
+	const TSharedRef<FAvaInteractiveToolsToolPresetBase>* Preset = ToolParams.Presets.Find(InPresetName);
+
+	if (!Preset)
+	{
+		return;
+	}
+
+	Tool->bPerformingDefaultAction = true;
+	Tool->DefaultAction();
+
+	(*Preset)->ApplyPreset(Tool->GetSpawnedActor());
+}
+
+void UAvaInteractiveToolsToolBase::RegisterPresetMenu()
+{
+	using namespace UE::AvalancheInteractiveTools::Private;
+
+	UToolMenus* ToolMenus = UToolMenus::Get();
+
+	if (ToolMenus->IsMenuRegistered(AvaITFToolPresetMenuName))
+	{
+		return;
+	}
+
+	UToolMenu* PresetMenu = ToolMenus->RegisterMenu(AvaITFToolPresetMenuName);
+
+	static const FName PresetSectionName = "Preset";
+
+	PresetMenu->AddDynamicSection(
+		PresetSectionName,
+		FNewSectionConstructChoice(FNewToolMenuDelegate::CreateLambda(
+			[](UToolMenu* InMenu)
+			{
+				UAvaInteractiveToolsToolBase* Tool = InMenu->FindContext<UAvaInteractiveToolsToolBase>();
+
+				if (!Tool)
+				{
+					return;
+				}
+
+				const FAvaInteractiveToolsToolParameters ToolParams = Tool->GetToolParameters();
+
+				if (ToolParams.Presets.IsEmpty())
+				{
+					return;
+				}
+
+				TStrongObjectPtr<UAvaInteractiveToolsToolBase> ToolPtr = TStrongObjectPtr<UAvaInteractiveToolsToolBase>(Tool);
+				FToolMenuSection& Section = InMenu->FindOrAddSection("Presets");
+
+				for (const TPair<FName, TSharedRef<FAvaInteractiveToolsToolPresetBase>>& PresetPair : ToolParams.Presets)
+				{
+					Section.AddMenuEntry(
+						PresetPair.Key,
+						PresetPair.Value->GetName(),
+						PresetPair.Value->GetDescription(),
+						ToolParams.UICommand->GetIcon(),
+						FToolUIActionChoice(FExecuteAction::CreateStatic(&UAvaInteractiveToolsToolBase::OnPresetSelected, ToolPtr, PresetPair.Key))
+					);
+				}
+			}
+		))
+	);
+}
+
+#undef LOCTEXT_NAMESPACE
diff --git a/Engine/Plugins/Experimental/Avalanche/Source/AvalancheInteractiveTools/Public/Tools/AvaInteractiveToolsToolPresetBase.h b/Engine/Plugins/Experimental/Avalanche/Source/AvalancheInteractiveTools/Public/Tools/AvaInteractiveToolsToolPresetBase.h
new file mode 100644
--- /dev/null
+++ b/Engine/Plugins/Experimental/Avalanche/Source/AvalancheInteractiveTools/Public/Tools/AvaInteractiveToolsToolPresetBase.h
@@ -0,0 +1,24 @@
+// Copyright Epic Games, Inc. All Rights Reserved.
+
+#pragma once
+
+#include "Internationalization/Text.h"
+#include "Templates/SharedPointer.h"
+
+class AActor;
+
+class FAvaInteractiveToolsToolPresetBase : public TSharedFromThis<FAvaInteractiveToolsToolPresetBase>
+{
+public:
+	virtual ~FAvaInteractiveToolsToolPresetBase() = default;
+
+	virtual void ApplyPreset(AActor* InActor) const = 0;
+
+	const FText& GetName() { return Name; }
+
+	const FText& GetDescription() { return Description; }
+
+protected:
+	FText Name;
+	FText Description;
+};
\ No newline at end of file
diff --git a/Engine/Plugins/Experimental/Avalanche/Source/AvalancheInteractiveToolsRuntime/Public/IAvalancheInteractiveToolsModeDetailsObject.h b/Engine/Plugins/Experimental/Avalanche/Source/AvalancheInteractiveToolsRuntime/Public/IAvalancheInteractiveToolsModeDetailsObject.h
new file mode 100644
--- /dev/null
+++ b/Engine/Plugins/Experimental/Avalanche/Source/AvalancheInteractiveToolsRuntime/Public/IAvalancheInteractiveToolsModeDetailsObject.h
@@ -0,0 +1,17 @@
+// Copyright Epic Games, Inc. All Rights Reserved.
+
+#pragma once
+
+#include "UObject/Interface.h"
+#include "IAvalancheInteractiveToolsModeDetailsObject.generated.h"
+
+UINTERFACE(MinimalAPI)
+class UAvalancheInteractiveToolsModeDetailsObject : public UInterface
+{
+	GENERATED_BODY()
+};
+
+class IAvalancheInteractiveToolsModeDetailsObject
+{
+	GENERATED_BODY()
+};
diff --git a/Engine/Plugins/Experimental/Avalanche/Source/AvalancheInteractiveToolsRuntime/Public/IAvalancheInteractiveToolsModeDetailsObjectProvider.h b/Engine/Plugins/Experimental/Avalanche/Source/AvalancheInteractiveToolsRuntime/Public/IAvalancheInteractiveToolsModeDetailsObjectProvider.h
new file mode 100644
--- /dev/null
+++ b/Engine/Plugins/Experimental/Avalanche/Source/AvalancheInteractiveToolsRuntime/Public/IAvalancheInteractiveToolsModeDetailsObjectProvider.h
@@ -0,0 +1,25 @@
+// Copyright Epic Games, Inc. All Rights Reserved.
+
+#pragma once
+
+#include "UObject/Interface.h"
+#include "IAvalancheInteractiveToolsModeDetailsObjectProvider.generated.h"
+
+class UObject;
+
+UINTERFACE(MinimalAPI)
+class UAvalancheInteractiveToolsModeDetailsObjectProvider : public UInterface
+{
+	GENERATED_BODY()
+};
+
+class AVALANCHEINTERACTIVETOOLSRUNTIME_API IAvalancheInteractiveToolsModeDetailsObjectProvider
+{
+	GENERATED_BODY()
+
+public:
+	UFUNCTION(BlueprintNativeEvent, CallInEditor, Category = "Motion Design Interactive Tools")
+	UObject* GetModeDetailsObject() const;
+
+	virtual UObject* GetModeDetailsObject_Implementation() const PURE_VIRTUAL(IAvalancheInteractiveToolsModeDetailsObjectProvider::GetModeDetailsObject, return nullptr;)
+};
diff --git a/Engine/Plugins/Experimental/Avalanche/Source/AvalancheLevelViewport/Private/AvaLevelViewportStyle.cpp b/Engine/Plugins/Experimental/Avalanche/Source/AvalancheLevelViewport/Private/AvaLevelViewportStyle.cpp
new file mode 100644
--- /dev/null
+++ b/Engine/Plugins/Experimental/Avalanche/Source/AvalancheLevelViewport/Private/AvaLevelViewportStyle.cpp
@@ -0,0 +1,195 @@
+// Copyright Epic Games, Inc. All Rights Reserved.
+
+#include "AvaLevelViewportStyle.h"
+#include "Interfaces/IPluginManager.h"
+#include "Styling/AppStyle.h"
+#include "Styling/SlateStyleRegistry.h"
+#include "Styling/StyleColors.h"
+#include "Styling/ToolBarStyle.h"
+
+#define LOCTEXT_NAMESPACE "AvaLevelViewportStyle"
+
+TSharedPtr<FSlateStyleSet> FAvaLevelViewportStyle::StyleInstance;
+
+void FAvaLevelViewportStyle::Initialize()
+{
+	StyleInstance = Create();
+	FSlateStyleRegistry::RegisterSlateStyle(Get());
+}
+
+void FAvaLevelViewportStyle::Shutdown()
+{
+	FSlateStyleRegistry::UnRegisterSlateStyle(Get());
+	StyleInstance.Reset();
+}
+
+FName FAvaLevelViewportStyle::GetStyleSetName()
+{
+	static FName StyleSetName(TEXT("AvalancheLevelViewport"));
+	return StyleSetName;
+}
+
+const ISlateStyle& FAvaLevelViewportStyle::Get()
+{
+	if (StyleInstance.IsValid() == false)
+	{
+		Initialize();
+	}
+
+	return *StyleInstance;
+}
+
+const FSlateBrush* FAvaLevelViewportStyle::GetBrush(FName PropertyName, const ANSICHAR* Specifier)
+{
+	return StyleInstance->GetBrush(PropertyName, Specifier);
+}
+
+#define IMAGE_BRUSH(RelativePath, ... ) FSlateImageBrush(Style->RootToContentDir(RelativePath, TEXT(".png")), __VA_ARGS__)
+#define IMAGE_BRUSH_SVG(RelativePath, ...) FSlateVectorImageBrush(Style->RootToContentDir(RelativePath, TEXT(".svg")), __VA_ARGS__)
+
+TSharedRef<FSlateStyleSet> FAvaLevelViewportStyle::Create()
+{
+	const FVector2f Icon16x16(16.0f, 16.0f);
+	const FVector2f Icon20x20(20.0f, 20.0f);
+	const FVector2f Icon22x22(22.0f, 22.0f);
+	const FVector2f Icon25x25(25.0f, 25.0f);
+
+	TSharedRef<FSlateStyleSet> Style = MakeShared<FSlateStyleSet>(GetStyleSetName());
+
+	TSharedPtr<IPlugin> Plugin = IPluginManager::Get().FindPlugin(TEXT("Avalanche"));
+	check(Plugin.IsValid());
+	if (Plugin.IsValid())
+	{
+		Style->SetContentRoot(Plugin->GetBaseDir() / TEXT("Resources"));
+	}
+
+	// Path Constants
+	const FString UEContentEditorSlatePath = FPaths::EngineContentDir() / TEXT("Editor/Slate");
+	const FString UEContentEditorSlateIconsPath = UEContentEditorSlatePath / TEXT("Icons");
+
+	// Grid Icons
+	Style->Set("Button.ToggleGrid", new IMAGE_BRUSH_SVG("Icons/PaletteIcons/GridToggle", Icon22x22));
+	Style->Set("Button.ToggleGridAlwaysVisible", new IMAGE_BRUSH_SVG("Icons/PaletteIcons/GridAlwaysVisibleToggle", Icon22x22));
+
+	// Snap Icons
+	Style->Set("Button.ToggleSnap", new IMAGE_BRUSH_SVG("Icons/PaletteIcons/SnapToggle", Icon22x22));
+	Style->Set("Button.ToggleGrid", new IMAGE_BRUSH_SVG("Icons/PaletteIcons/SnapGridToggle", Icon22x22));
+	Style->Set("Button.ToggleScreen", new IMAGE_BRUSH_SVG("Icons/PaletteIcons/SnapScreenToggle", Icon22x22));
+	Style->Set("Button.ToggleActor", new IMAGE_BRUSH_SVG("Icons/PaletteIcons/SnapActorToggle", Icon22x22));
+
+	// Viewport
+	Style->Set("Button.GameView", new IMAGE_BRUSH_SVG("Icons/EditorIcons/GameView", Icon25x25));
+	Style->Set("Button.Visualizers", new IMAGE_BRUSH_SVG("Icons/EditorIcons/Visualizers", Icon25x25));
+	Style->Set("Button.Guides", new IMAGE_BRUSH_SVG("Icons/EditorIcons/Guides", Icon25x25));
+	Style->Set("Button.Billboards", new IMAGE_BRUSH("Icons/EditorIcons/Billboards", Icon25x25));
+	Style->Set("Button.IsolateActors", new IMAGE_BRUSH("Icons/EditorIcons/IsolateActors", Icon25x25));
+	Style->Set("Button.BoundingBoxes", new IMAGE_BRUSH_SVG("Icons/EditorIcons/Mode_Bounding-Box", Icon25x25));
+	Style->Set("Button.SafeFrames", new IMAGE_BRUSH_SVG("Icons/EditorIcons/Mode_Broadcast-Safe", Icon25x25));
+	Style->Set("Button.KeyPreview", new IMAGE_BRUSH_SVG("Icons/EditorIcons/Mode_Key-Preview", Icon25x25));
+	Style->Set("Button.Snapshot", new IMAGE_BRUSH_SVG("Icons/EditorIcons/Mode_Snapshot", Icon25x25));
+	Style->Set("Button.WireframeMode", new IMAGE_BRUSH_SVG("Icons/EditorIcons/Mode_Wireframe", Icon25x25));
+	Style->Set("Button.Mask.Toggle", new IMAGE_BRUSH("Icons/EditorIcons/Mode_Mask", Icon25x25));
+
+	// Editor
+	Style->Set("Button.SelectionLock", new IMAGE_BRUSH("Icons/EditorIcons/SelectionLock", Icon25x25));
+	Style->Set("Button.PivotMode", new IMAGE_BRUSH("Icons/EditorIcons/Pivot_Mode", Icon25x25));
+
+	// Alignment
+	Style->Set("AvalancheIcons.Alignment.Translation.TopLeft",     new IMAGE_BRUSH("Icons/DetailsPanelIcons/TopLeft",     Icon16x16));
+	Style->Set("AvalancheIcons.Alignment.Translation.Top",         new IMAGE_BRUSH("Icons/DetailsPanelIcons/Top",         Icon16x16));
+	Style->Set("AvalancheIcons.Alignment.Translation.TopRight",    new IMAGE_BRUSH("Icons/DetailsPanelIcons/TopRight",    Icon16x16));
+	Style->Set("AvalancheIcons.Alignment.Translation.Left",        new IMAGE_BRUSH("Icons/DetailsPanelIcons/Left",        Icon16x16));
+	Style->Set("AvalancheIcons.Alignment.Translation.Center",      new IMAGE_BRUSH("Icons/DetailsPanelIcons/Center",      Icon16x16));
+	Style->Set("AvalancheIcons.Alignment.Translation.Right",       new IMAGE_BRUSH("Icons/DetailsPanelIcons/Right",       Icon16x16));
+	Style->Set("AvalancheIcons.Alignment.Translation.BottomLeft",  new IMAGE_BRUSH("Icons/DetailsPanelIcons/BottomLeft",  Icon16x16));
+	Style->Set("AvalancheIcons.Alignment.Translation.Bottom",      new IMAGE_BRUSH("Icons/DetailsPanelIcons/Bottom",      Icon16x16));
+	Style->Set("AvalancheIcons.Alignment.Translation.BottomRight", new IMAGE_BRUSH("Icons/DetailsPanelIcons/BottomRight", Icon16x16));
+	
+	Style->Set("AvalancheIcons.Alignment.Translation.Back",        new IMAGE_BRUSH_SVG("Icons/DetailsPanelIcons/StackBack",   Icon22x22));
+	Style->Set("AvalancheIcons.Alignment.Translation.Center_X",    new IMAGE_BRUSH_SVG("Icons/DetailsPanelIcons/StackCenter", Icon22x22));
+	Style->Set("AvalancheIcons.Alignment.Translation.Front",       new IMAGE_BRUSH_SVG("Icons/DetailsPanelIcons/StackFront",  Icon22x22));
+
+	Style->Set("AvalancheIcons.Alignment.Center_YZ",   new IMAGE_BRUSH_SVG("Icons/PaletteIcons/AlignCenterVertAndHoriz", Icon22x22));
+	Style->Set("AvalancheIcons.Alignment.Left",        new IMAGE_BRUSH_SVG("Icons/PaletteIcons/AlignLeft",               Icon22x22));
+	Style->Set("AvalancheIcons.Alignment.Center_Y",    new IMAGE_BRUSH_SVG("Icons/PaletteIcons/AlignCenterHoriz",        Icon22x22));
+	Style->Set("AvalancheIcons.Alignment.Right",       new IMAGE_BRUSH_SVG("Icons/PaletteIcons/AlignRight",              Icon22x22));
+	Style->Set("AvalancheIcons.Alignment.Top",         new IMAGE_BRUSH_SVG("Icons/PaletteIcons/AlignTop",                Icon22x22));
+	Style->Set("AvalancheIcons.Alignment.Center_Z",    new IMAGE_BRUSH_SVG("Icons/PaletteIcons/AlignCenterVert",         Icon22x22));
+	Style->Set("AvalancheIcons.Alignment.Bottom",      new IMAGE_BRUSH_SVG("Icons/PaletteIcons/AlignBottom",             Icon22x22));
+	Style->Set("AvalancheIcons.Alignment.DistributeX", new IMAGE_BRUSH_SVG("Icons/PaletteIcons/DistributeDepth",         Icon22x22));
+	Style->Set("AvalancheIcons.Alignment.DistributeY", new IMAGE_BRUSH_SVG("Icons/PaletteIcons/DistributeHorizontal",    Icon22x22));
+	Style->Set("AvalancheIcons.Alignment.DistributeZ", new IMAGE_BRUSH_SVG("Icons/PaletteIcons/DistributeVertical",      Icon22x22));
+	
+	Style->Set("AvalancheIcons.Alignment.Rotation.Actor.Roll",   new IMAGE_BRUSH_SVG("Icons/PaletteIcons/RotationAlignActorRoll",   Icon22x22));
+	Style->Set("AvalancheIcons.Alignment.Rotation.Actor.Pitch",  new IMAGE_BRUSH_SVG("Icons/PaletteIcons/RotationAlignActorPitch",  Icon22x22));
+	Style->Set("AvalancheIcons.Alignment.Rotation.Actor.Yaw",    new IMAGE_BRUSH_SVG("Icons/PaletteIcons/RotationAlignActorYaw",    Icon22x22));
+	Style->Set("AvalancheIcons.Alignment.Rotation.Actor.All",    new IMAGE_BRUSH_SVG("Icons/PaletteIcons/RotationAlignActorAll",    Icon22x22));
+	Style->Set("AvalancheIcons.Alignment.Rotation.Camera.Roll",  new IMAGE_BRUSH_SVG("Icons/PaletteIcons/RotationAlignCameraRoll",  Icon22x22));
+	Style->Set("AvalancheIcons.Alignment.Rotation.Camera.Pitch", new IMAGE_BRUSH_SVG("Icons/PaletteIcons/RotationAlignCameraPitch", Icon22x22));
+	Style->Set("AvalancheIcons.Alignment.Rotation.Camera.Yaw",   new IMAGE_BRUSH_SVG("Icons/PaletteIcons/RotationAlignCameraYaw",   Icon22x22));
+	Style->Set("AvalancheIcons.Alignment.Rotation.Camera.All",   new IMAGE_BRUSH_SVG("Icons/PaletteIcons/RotationAlignCameraAll",   Icon22x22));
+
+	// Screen Icons
+	Style->Set("AvalancheIcons.Screen.SizeToScreen",        new IMAGE_BRUSH_SVG("Icons/PaletteIcons/SizeToScreen",        Icon22x22));
+	Style->Set("AvalancheIcons.Screen.SizeToScreenStretch", new IMAGE_BRUSH_SVG("Icons/PaletteIcons/SizeToScreenStretch", Icon22x22));
+	Style->Set("AvalancheIcons.Screen.FitToScreen",         new IMAGE_BRUSH_SVG("Icons/PaletteIcons/FitToScreen",         Icon22x22));
+
+	// Color picker icons
+	Style->Set("AvalancheIcons.ColorPicker.SolidColors",    new IMAGE_BRUSH("Icons/EditorIcons/SolidColors",    Icon20x20));
+	Style->Set("AvalancheIcons.ColorPicker.LinearGradient", new IMAGE_BRUSH("Icons/EditorIcons/LinearGradient", Icon20x20));
+
+	// Post process icons
+	Style->Set("AvalancheIcons.PostProcess.RGB",   new IMAGE_BRUSH("Icons/ViewportIcons/RGBSquare", Icon20x20));
+	Style->Set("AvalancheIcons.PostProcess.Red",   new IMAGE_BRUSH("Icons/ViewportIcons/RedSquare",   Icon20x20));
+	Style->Set("AvalancheIcons.PostProcess.Green", new IMAGE_BRUSH("Icons/ViewportIcons/GreenSquare", Icon20x20));
+	Style->Set("AvalancheIcons.PostProcess.Blue",  new IMAGE_BRUSH("Icons/ViewportIcons/BlueSquare",  Icon20x20));
+	Style->Set("AvalancheIcons.PostProcess.Alpha", new IMAGE_BRUSH("Icons/ViewportIcons/AlphaSquare", Icon20x20));
+
+	// StatusBar ToolMenu
+	{
+    	FToolBarStyle StatusToolbarStyle = FAppStyle::Get().GetWidgetStyle<FToolBarStyle>("AssetEditorToolbar");
+
+    	StatusToolbarStyle.SetButtonPadding(       FMargin(0.0f, 0.0f));
+    	StatusToolbarStyle.SetCheckBoxPadding(     FMargin(0.0f, 0.0f));
+    	StatusToolbarStyle.SetComboButtonPadding(  FMargin(0.0f, 0.0f));
+    	StatusToolbarStyle.SetIndentedBlockPadding(FMargin(0.0f, 0.0f));
+    	StatusToolbarStyle.SetBlockPadding(        FMargin(0.0f, 0.0f));
+    	StatusToolbarStyle.SetSeparatorPadding(    FMargin(0.0f, 0.0f));
+    	StatusToolbarStyle.bShowLabels = false;
+    	StatusToolbarStyle.SetBackground(FSlateColorBrush(FStyleColors::Transparent));
+    	StatusToolbarStyle.SetBackgroundPadding(0);
+    	StatusToolbarStyle.SetButtonPadding(0);
+    	StatusToolbarStyle.SetCheckBoxPadding(0);
+    	
+    	Style->Set("StatusBar", StatusToolbarStyle);
+    }
+
+	FCheckBoxStyle CheckBoxStyle = FAppStyle::Get().GetWidgetStyle<FCheckBoxStyle>("DetailsView.SectionButton");
+	CheckBoxStyle.Padding.Left = 10.f;
+	CheckBoxStyle.Padding.Right = 10.f;
+
+	Style->Set("Avalanche.Alignment.Context", CheckBoxStyle);
+	
+	FButtonStyle ButtonStyle = FAppStyle::Get().GetWidgetStyle<FButtonStyle>("Button");
+	ButtonStyle.NormalPadding.Left = 4.f;
+	ButtonStyle.NormalPadding.Right = 4.f;
+	ButtonStyle.PressedPadding = ButtonStyle.NormalPadding;
+
+	Style->Set("Avalanche.Alignment.Button", ButtonStyle);
+
+	FButtonStyle GuidePresetMenuStyle = FAppStyle::GetWidgetStyle<FButtonStyle>("FlatButton");
+	GuidePresetMenuStyle.SetNormalPadding(2.f);
+	GuidePresetMenuStyle.SetPressedPadding(2.f);
+	GuidePresetMenuStyle.SetDisabled(GuidePresetMenuStyle.Normal);
+	GuidePresetMenuStyle.Hovered.TintColor = FSlateColor(EStyleColor::Highlight).GetSpecifiedColor();
+	GuidePresetMenuStyle.SetPressed(GuidePresetMenuStyle.Hovered);
+
+	Style->Set("Avalanche.Menu.GuidePreset.Button", GuidePresetMenuStyle);
+
+	return Style;
+}
+
+#undef IMAGE_BRUSH
+#undef IMAGE_BRUSH_SVG
+
+#undef LOCTEXT_NAMESPACE
diff --git a/Engine/Plugins/Experimental/Avalanche/Source/AvalancheLevelViewport/Private/AvaLevelViewportStyle.h b/Engine/Plugins/Experimental/Avalanche/Source/AvalancheLevelViewport/Private/AvaLevelViewportStyle.h
new file mode 100644
--- /dev/null
+++ b/Engine/Plugins/Experimental/Avalanche/Source/AvalancheLevelViewport/Private/AvaLevelViewportStyle.h
@@ -0,0 +1,25 @@
+// Copyright Epic Games, Inc. All Rights Reserved.
+
+#pragma once
+
+#include "Styling/ISlateStyle.h"
+#include "Styling/SlateStyle.h"
+
+struct FSlateBrush;
+
+class FAvaLevelViewportStyle
+{
+public:
+	static void Initialize();
+	static void Shutdown();
+
+	static FName GetStyleSetName();
+	static const ISlateStyle& Get();
+
+	static const FSlateBrush* GetBrush(FName PropertyName, const ANSICHAR* Specifier = nullptr);
+
+private:
+	static TSharedPtr<FSlateStyleSet> StyleInstance;
+
+	static TSharedRef<FSlateStyleSet> Create();
+};
diff --git a/Engine/Plugins/Experimental/Avalanche/Source/AvalancheLevelViewport/Private/AvaViewportCameraHistory.cpp b/Engine/Plugins/Experimental/Avalanche/Source/AvalancheLevelViewport/Private/AvaViewportCameraHistory.cpp
new file mode 100644
--- /dev/null
+++ b/Engine/Plugins/Experimental/Avalanche/Source/AvalancheLevelViewport/Private/AvaViewportCameraHistory.cpp
@@ -0,0 +1,236 @@
+// Copyright Epic Games, Inc. All Rights Reserved.
+
+#include "AvaViewportCameraHistory.h"
+
+#include "AvaLevelViewportCommands.h"
+#include "AvaLog.h"
+#include "Framework/Notifications/NotificationManager.h"
+#include "LevelEditor.h"
+#include "LevelEditor/AvaLevelEditorUtils.h"
+#include "Math/Transform.h"
+#include "Widgets/Notifications/SNotificationList.h"
+
+#define LOCTEXT_NAMESPACE "AvaViewportCameraHistory"
+
+namespace UE::AvalancheEditor::Private
+{
+	static FText UndoRedoMessage = LOCTEXT("UndoRedoMessage", "Camera Movement");
+
+	/** Modulus that correctly wraps negative values. */
+	int32 WrapIndex(int32 InValue, const int32 InNum)
+	{
+		if (InNum == 0)
+		{
+			return 0;
+		}
+		
+		const int32 Result = ((InValue %= InNum) < 0) ? InValue + InNum : InValue;
+		return Result;
+	}
+}
+
+FAvaViewportCameraHistory::FAvaViewportCameraHistory()
+{
+	FEditorDelegates::OnMapOpened.AddRaw(this, &FAvaViewportCameraHistory::MapOpened);
+	FCoreDelegates::OnPostEngineInit.AddRaw(this, &FAvaViewportCameraHistory::PostEngineInit);
+}
+
+FAvaViewportCameraHistory::~FAvaViewportCameraHistory()
+{
+	FEditorDelegates::OnMapOpened.RemoveAll(this);
+	FCoreDelegates::OnPostEngineInit.RemoveAll(this);
+	
+	UnbindCommands();
+}
+
+void FAvaViewportCameraHistory::BindCommands()
+{
+	const FLevelEditorModule* const LevelEditorModule = FAvaLevelEditorUtils::GetLevelEditorModule();
+	if (!LevelEditorModule)
+	{
+		return;
+	}
+	
+	const FAvaLevelViewportCommands& AvaLevelViewportCommands = FAvaLevelViewportCommands::Get();
+	const TSharedRef<FUICommandList> LevelEditorActions = LevelEditorModule->GetGlobalLevelEditorActions();
+
+	LevelEditorActions->MapAction(AvaLevelViewportCommands.CameraTransformUndo
+		, FExecuteAction::CreateSP(this, &FAvaViewportCameraHistory::ExecuteCameraTransformUndo));
+
+	LevelEditorActions->MapAction(AvaLevelViewportCommands.CameraTransformRedo
+		, FExecuteAction::CreateSP(this, &FAvaViewportCameraHistory::ExecuteCameraTransformRedo));
+
+	BindCameraTransformDelegates();
+}
+
+void FAvaViewportCameraHistory::UnbindCommands()
+{
+	UnbindCameraTransformDelegates();
+	
+	const FLevelEditorModule* const LevelEditorModule = FAvaLevelEditorUtils::GetLevelEditorModule();
+	if (!LevelEditorModule)
+	{
+		return;
+	}
+	
+	const FAvaLevelViewportCommands& AvaLevelViewportCommands = FAvaLevelViewportCommands::Get();
+	const TSharedRef<FUICommandList> LevelEditorActions = LevelEditorModule->GetGlobalLevelEditorActions();
+
+	LevelEditorActions->UnmapAction(AvaLevelViewportCommands.CameraTransformUndo);
+	LevelEditorActions->UnmapAction(AvaLevelViewportCommands.CameraTransformRedo);
+}
+
+void FAvaViewportCameraHistory::MapOpened(const FString& InMapName, bool bIsTemplate)
+{
+	Reset();
+}
+
+void FAvaViewportCameraHistory::PostEngineInit()
+{
+	BindCommands();
+}
+
+void FAvaViewportCameraHistory::Reset()
+{
+	CameraTransformHistory.Reset();
+	CameraTransformHistoryIndex = 0;
+	CameraTransformHistoryHeadIndex = 1;
+}
+
+void FAvaViewportCameraHistory::BindCameraTransformDelegates()
+{
+	check(GEditor);
+	
+	if (GEditor)
+	{
+		OnBeginCameraTransformHandle = GEditor->OnBeginCameraMovement().AddSP(this, &FAvaViewportCameraHistory::OnBeginCameraTransform);
+		OnEndCameraTransformHandle = GEditor->OnEndCameraMovement().AddSP(this, &FAvaViewportCameraHistory::OnEndCameraTransform);
+	}
+}
+
+void FAvaViewportCameraHistory::UnbindCameraTransformDelegates()
+{
+	if (GEditor)
+	{
+		GEditor->OnBeginCameraMovement().RemoveAll(this);
+		GEditor->OnEndCameraMovement().RemoveAll(this);
+	}
+}
+
+void FAvaViewportCameraHistory::OnBeginCameraTransform(UObject& InCameraObject)
+{
+	if (!::IsValid(&InCameraObject))
+	{
+		return;
+	}
+
+	if (AActor* CameraActor = Cast<AActor>(&InCameraObject))
+	{
+		if (CameraTransformHistory.IsEmpty())
+		{
+			CameraTransformHistoryIndex = CameraTransformHistory.Add({ CameraActor, CameraActor->GetActorTransform() });
+			CameraTransformHistoryHeadIndex = 1;
+		}
+	}
+}
+
+void FAvaViewportCameraHistory::OnEndCameraTransform(UObject& InCameraObject)
+{
+	if (!::IsValid(&InCameraObject))
+	{
+		return;
+	}
+
+	if (AActor* CameraActor = Cast<AActor>(&InCameraObject))
+	{
+		// New item replaces next
+		int32 CameraTransformIndexToAdd = UE::AvalancheEditor::Private::WrapIndex(CameraTransformHistoryIndex + 1, CameraTransformHistory.Max());
+
+		CameraTransformIndexToAdd = CameraTransformHistory.Insert({ CameraActor, CameraActor->GetActorTransform() }, CameraTransformIndexToAdd);
+		
+		CameraTransformHistoryIndex = CameraTransformIndexToAdd;
+		CameraTransformHistoryHeadIndex = UE::AvalancheEditor::Private::WrapIndex(CameraTransformHistoryIndex + 1, CameraTransformHistory.Max());
+
+		UE_LOG(LogAva, Display, TEXT("Saved camera transform at %i, head %i, %s"), CameraTransformHistoryIndex, CameraTransformHistoryHeadIndex, *CameraActor->GetActorTransform().Rotator().ToString());
+	}
+}
+
+void FAvaViewportCameraHistory::ExecuteCameraTransformUndo()
+{
+	const int32 CameraTransformIndexToRestore = UE::AvalancheEditor::Private::WrapIndex(CameraTransformHistoryIndex - 1, CameraTransformHistory.Max());
+	if (CameraTransformHistory.IsEmpty()
+		|| !CameraTransformHistory.IsValidIndex(CameraTransformIndexToRestore)
+		|| (CameraTransformHistory.Num() > 1 && CameraTransformIndexToRestore == CameraTransformHistoryHeadIndex))
+	{
+		return;
+	}
+
+	const TPair<TWeakObjectPtr<AActor>, FTransform> UndoPair = CameraTransformHistory[CameraTransformIndexToRestore];
+	if (AActor* CameraActor = UndoPair.Key.Get())
+	{
+		CameraTransformHistoryIndex = CameraTransformIndexToRestore;
+		CameraActor->SetActorTransform(UndoPair.Value);
+		NotifyUndo();
+
+		UE_LOG(LogAva, Display, TEXT("Undo camera transform at %i, %s"), CameraTransformIndexToRestore, *CameraActor->GetActorTransform().Rotator().ToString());
+	}
+}
+
+void FAvaViewportCameraHistory::ExecuteCameraTransformRedo()
+{
+	const int32 CameraTransformIndexToRestore = UE::AvalancheEditor::Private::WrapIndex(CameraTransformHistoryIndex + 1, CameraTransformHistory.Max());
+	if (CameraTransformHistory.IsEmpty()
+		|| !CameraTransformHistory.IsValidIndex(CameraTransformIndexToRestore)
+		|| CameraTransformIndexToRestore == CameraTransformHistoryHeadIndex)
+	{
+		return;
+	}
+
+	const TPair<TWeakObjectPtr<AActor>, FTransform> RedoPair = CameraTransformHistory[CameraTransformIndexToRestore];	
+	if (AActor* CameraActor = RedoPair.Key.Get())
+	{
+		CameraTransformHistoryIndex = CameraTransformIndexToRestore;
+		CameraActor->SetActorTransform(RedoPair.Value);
+		NotifyRedo();
+
+		UE_LOG(LogAva, Display, TEXT("Redo camera transform at %i, %s"), CameraTransformIndexToRestore, *CameraActor->GetActorTransform().Rotator().ToString());
+	}
+}
+
+void FAvaViewportCameraHistory::NotifyUndo()
+{
+	const FText UndoMessage = NSLOCTEXT("UnrealEd", "UndoMessageFormat", "Undo: {0}");
+	Notify(FText::Format(UndoMessage, UE::AvalancheEditor::Private::UndoRedoMessage));
+}
+
+void FAvaViewportCameraHistory::NotifyRedo()
+{
+	const FText RedoMessage = NSLOCTEXT("UnrealEd", "RedoMessageFormat", "Redo: {0}");
+	Notify(FText::Format(RedoMessage, UE::AvalancheEditor::Private::UndoRedoMessage));
+}
+
+// @see: UEditorEngine::ShowUndoRedoNotification
+void FAvaViewportCameraHistory::Notify(const FText& InText)
+{
+	// Add a new notification item only if the previous one has expired or is otherwise done fading out (CS_None). This way multiple undo/redo notifications do not pollute the notification window.
+	if (!UndoRedoNotificationItem.IsValid() || UndoRedoNotificationItem->GetCompletionState() == SNotificationItem::CS_None)
+	{
+		FNotificationInfo Info(InText);
+		Info.bUseLargeFont = false;
+		Info.bUseSuccessFailIcons = false;
+
+		UndoRedoNotificationItem = FSlateNotificationManager::Get().AddNotification(Info);
+	}
+
+	if (UndoRedoNotificationItem.IsValid())
+	{
+		// Update the text and completion state to reflect current info
+		UndoRedoNotificationItem->SetText(InText);
+		UndoRedoNotificationItem->SetCompletionState(SNotificationItem::CS_Success);
+
+		// Restart the fade animation for the current undo/redo notification
+		UndoRedoNotificationItem->ExpireAndFadeout();
+	}
+}
+
+#undef LOCTEXT_NAMESPACE
diff --git a/Engine/Plugins/Experimental/Avalanche/Source/AvalancheLevelViewport/Private/AvalancheLevelViewportModule.cpp b/Engine/Plugins/Experimental/Avalanche/Source/AvalancheLevelViewport/Private/AvalancheLevelViewportModule.cpp
new file mode 100644
--- /dev/null
+++ b/Engine/Plugins/Experimental/Avalanche/Source/AvalancheLevelViewport/Private/AvalancheLevelViewportModule.cpp
@@ -0,0 +1,73 @@
+// Copyright Epic Games, Inc. All Rights Reserved.
+
+#include "AvalancheLevelViewportModule.h"
+
+#include "AvaLevelViewportCommands.h"
+#include "AvaLevelViewportStyle.h"
+#include "AvaViewportCameraHistory.h"
+#include "AvaViewportUtils.h"
+#include "Modules/ModuleManager.h"
+#include "ToolMenus.h"
+#include "ViewportClient/AvaLevelViewportClient.h"
+
+DEFINE_LOG_CATEGORY(AvaLevelViewportLog);
+
+namespace UE::AvalancheLevelViewport::Private
+{
+	TSharedPtr<IAvaViewportClient> GetAsAvaLevelViewportClient(FEditorViewportClient* InViewportClient)
+	{
+		if (FAvaLevelViewportClient::IsAvaLevelViewportClient(InViewportClient))
+		{
+			return static_cast<FAvaLevelViewportClient*>(InViewportClient)->AsShared();
+		}
+
+		return nullptr;
+	}
+}
+
+void FAvalancheLevelViewportModule::StartupModule()
+{
+	FAvaLevelViewportStyle::Initialize();
+	FAvaLevelViewportCommands::Register();
+
+	ViewportCameraHistory = MakeShared<FAvaViewportCameraHistory>();
+
+	UToolMenus::RegisterStartupCallback(FSimpleMulticastDelegate::FDelegate::CreateRaw(this, &FAvalancheLevelViewportModule::RegisterMenus));
+
+	AvaLevelViewportClientCasterDelegateHandle = FAvaViewportUtils::RegisterViewportClientCaster(
+		&UE::AvalancheLevelViewport::Private::GetAsAvaLevelViewportClient
+	);
+}
+
+void FAvalancheLevelViewportModule::ShutdownModule()
+{
+	UToolMenus::UnRegisterStartupCallback(this);
+	UToolMenus::UnregisterOwner(this);
+
+	ViewportCameraHistory.Reset();
+
+	FAvaLevelViewportStyle::Shutdown();
+	FAvaLevelViewportCommands::Unregister();
+
+	if (AvaLevelViewportClientCasterDelegateHandle.IsValid())
+	{
+		FAvaViewportUtils::UnregisterViewportClientCaster(AvaLevelViewportClientCasterDelegateHandle);
+		AvaLevelViewportClientCasterDelegateHandle.Reset();
+	}
+}
+
+void FAvalancheLevelViewportModule::RegisterMenus()
+{
+	// Owner will be used for cleanup in call to UToolMenus::UnregisterOwner
+	FToolMenuOwnerScoped OwnerScoped(this);
+	
+	// Extension point for SAvaLevelViewportStatusBar
+	{
+		static FName StatusBarMenuName = UE::AvaLevelViewport::Internal::StatusBarMenuName;
+		UToolMenu* Menu = UToolMenus::Get()->RegisterMenu(StatusBarMenuName, NAME_None, EMultiBoxType::ToolBar, false);
+		Menu->SetStyleSet(&FAvaLevelViewportStyle::Get());
+		Menu->StyleName = "StatusBar";
+	}
+}
+
+IMPLEMENT_MODULE(FAvalancheLevelViewportModule, AvalancheLevelViewport)
diff --git a/Engine/Plugins/Experimental/Avalanche/Source/AvalancheLevelViewport/Private/AvalancheLevelViewportModule.h b/Engine/Plugins/Experimental/Avalanche/Source/AvalancheLevelViewport/Private/AvalancheLevelViewportModule.h
new file mode 100644
--- /dev/null
+++ b/Engine/Plugins/Experimental/Avalanche/Source/AvalancheLevelViewport/Private/AvalancheLevelViewportModule.h
@@ -0,0 +1,34 @@
+// Copyright Epic Games, Inc. All Rights Reserved.
+
+#pragma once
+
+#include "Delegates/IDelegateInstance.h"
+#include "Logging/LogMacros.h"
+#include "Modules/ModuleInterface.h"
+#include "Templates/SharedPointerFwd.h"
+#include "AvaViewportCameraHistory.h"
+
+DECLARE_LOG_CATEGORY_EXTERN(AvaLevelViewportLog, Log, All);
+
+namespace UE::AvaLevelViewport::Internal
+{
+	static FName StatusBarMenuName = TEXT("AvalancheLevelViewport.StatusBar");
+}
+
+class FAvalancheLevelViewportModule : public IModuleInterface
+{
+public:
+	//~ Begin IModuleInterface interface
+	virtual void StartupModule() override;
+	virtual void ShutdownModule() override;
+	//~ End IModuleInterface interface
+
+private:
+	void RegisterMenus();
+
+protected:
+	FDelegateHandle AvaLevelViewportClientCasterDelegateHandle;
+
+	/** Handles viewport camera undo/redo */
+	TSharedPtr<FAvaViewportCameraHistory> ViewportCameraHistory;
+};
diff --git a/Engine/Plugins/Experimental/Avalanche/Source/AvalancheLevelViewport/Private/Widgets/AvaViewportColorPickerDelegates.cpp b/Engine/Plugins/Experimental/Avalanche/Source/AvalancheLevelViewport/Private/Widgets/AvaViewportColorPickerDelegates.cpp
new file mode 100644
--- /dev/null
+++ b/Engine/Plugins/Experimental/Avalanche/Source/AvalancheLevelViewport/Private/Widgets/AvaViewportColorPickerDelegates.cpp
@@ -0,0 +1,38 @@
+// Copyright Epic Games, Inc. All Rights Reserved.
+
+#include "Widgets/AvaViewportColorPickerDelegates.h"
+
+namespace UE::AvalancheLevelViewport::Private
+{
+	static FAvaOnColorPicked OnColorPicked;
+	static FAvaOnColorSourceSelected OnColorSourceSelected;
+
+	static FAvaColorChangeData LastColorData = FAvaColorChangeData();
+}
+
+FAvaOnColorPicked& FAvaViewportColorPickerDelegates::GetOnColorPicked()
+{
+	return UE::AvalancheLevelViewport::Private::OnColorPicked;
+}
+
+void FAvaViewportColorPickerDelegates::BroadcastColorPicked(const TSharedRef<IToolkitHost>& InEditor, const FAvaColorChangeData& InNewColorData)
+{
+	UE::AvalancheLevelViewport::Private::OnColorPicked.Broadcast(InEditor, InNewColorData);
+	UE::AvalancheLevelViewport::Private::LastColorData = InNewColorData;
+}
+
+FAvaOnColorPicked& FAvaViewportColorPickerDelegates::GetOnColorSourceSelected()
+{
+	return UE::AvalancheLevelViewport::Private::OnColorSourceSelected;
+}
+
+void FAvaViewportColorPickerDelegates::BroadcastColorSourceSelected(const TSharedRef<IToolkitHost>& InEditor, const FAvaColorChangeData& InNewColorData)
+{
+	UE::AvalancheLevelViewport::Private::OnColorPicked.Broadcast(InEditor, InNewColorData);
+	UE::AvalancheLevelViewport::Private::LastColorData = InNewColorData;
+}
+
+FAvaColorChangeData FAvaViewportColorPickerDelegates::GetLastColorData()
+{
+	return UE::AvalancheLevelViewport::Private::LastColorData;
+}
diff --git a/Engine/Plugins/Experimental/Avalanche/Source/AvalancheLevelViewport/Public/Widgets/AvaViewportColorPickerActorClassRegistry.cpp b/Engine/Plugins/Experimental/Avalanche/Source/AvalancheLevelViewport/Public/Widgets/AvaViewportColorPickerActorClassRegistry.cpp
new file mode 100644
--- /dev/null
+++ b/Engine/Plugins/Experimental/Avalanche/Source/AvalancheLevelViewport/Public/Widgets/AvaViewportColorPickerActorClassRegistry.cpp
@@ -0,0 +1,85 @@
+// Copyright Epic Games, Inc. All Rights Reserved.
+
+#include "Widgets/AvaViewportColorPickerActorClassRegistry.h"
+#include "AvaDefs.h"
+#include "Containers/Map.h"
+#include "Templates/SharedPointer.h"
+
+namespace UE::AvalancheLevelViewport::Private
+{
+	FAvaColorChangeData DefaultColorData;
+	TMap<UClass*, TSharedRef<FAvaViewportColorPickerActorAdapter>> Adapters;
+
+	TSharedPtr<FAvaViewportColorPickerActorAdapter> FindAdapterForClass(UClass* InClass)
+	{
+		if (!InClass)
+		{
+			return nullptr;
+		}
+
+		for (UClass* Class = InClass; Class != AActor::StaticClass(); Class = Class->GetSuperClass())
+		{
+			if (const TSharedRef<FAvaViewportColorPickerActorAdapter>* Adapter = Adapters.Find(Class))
+			{
+				return *Adapter;
+			}
+		}
+
+		return nullptr;
+	}
+}
+
+FAvaColorChangeData FAvaViewportColorPickerActorAdapter::GetColorData(const AActor* InActor) const
+{
+	using namespace UE::AvalancheLevelViewport::Private;
+
+	return DefaultColorData;
+}
+
+void FAvaViewportColorPickerActorAdapter::SetColorData(AActor* InActor, const FAvaColorChangeData& InColorData) const
+{
+}
+
+void FAvaViewportColorPickerActorClassRegistry::RegisterClassAdapter(UClass* InClass, 
+	const TSharedRef<FAvaViewportColorPickerActorAdapter>& InAdapter)
+{
+	using namespace UE::AvalancheLevelViewport::Private;
+
+	Adapters.Add(InClass, InAdapter);
+}
+
+bool FAvaViewportColorPickerActorClassRegistry::ApplyColorDataToActor(AActor* InActor, const FAvaColorChangeData& InColorData)
+{
+	if (!InActor)
+	{
+		return false;
+	}
+
+	using namespace UE::AvalancheLevelViewport::Private;
+
+	if (TSharedPtr<FAvaViewportColorPickerActorAdapter> Adapter = FindAdapterForClass(InActor->GetClass()))
+	{
+		Adapter->SetColorData(InActor, InColorData);
+		return true;
+	}
+
+	return false;
+}
+
+bool FAvaViewportColorPickerActorClassRegistry::GetColorDataFromActor(const AActor* InActor, FAvaColorChangeData& OutColorData)
+{
+	if (!InActor)
+	{
+		return false;
+	}
+
+	using namespace UE::AvalancheLevelViewport::Private;
+
+	if (TSharedPtr<FAvaViewportColorPickerActorAdapter> Adapter = FindAdapterForClass(InActor->GetClass()))
+	{
+		OutColorData = Adapter->GetColorData(InActor);
+		return true;
+	}
+
+	return false;
+}
diff --git a/Engine/Plugins/Experimental/Avalanche/Source/AvalancheLevelViewport/Public/Widgets/AvaViewportColorPickerActorClassRegistry.h b/Engine/Plugins/Experimental/Avalanche/Source/AvalancheLevelViewport/Public/Widgets/AvaViewportColorPickerActorClassRegistry.h
new file mode 100644
--- /dev/null
+++ b/Engine/Plugins/Experimental/Avalanche/Source/AvalancheLevelViewport/Public/Widgets/AvaViewportColorPickerActorClassRegistry.h
@@ -0,0 +1,81 @@
+// Copyright Epic Games, Inc. All Rights Reserved.
+
+#pragma once
+
+#include "AvaDefs.h"
+#include "GameFramework/Actor.h"
+#include "Templates/SharedPointerFwd.h"
+
+class AActor;
+class UClass;
+
+struct AVALANCHELEVELVIEWPORT_API FAvaViewportColorPickerActorAdapter
+{
+	virtual ~FAvaViewportColorPickerActorAdapter() = default;
+
+	virtual FAvaColorChangeData GetColorData(const AActor* InActor) const;
+	virtual void SetColorData(AActor* InActor, const FAvaColorChangeData& InColorData) const;
+};
+
+template<typename InType>
+struct TAvaViewportColorPickerActorAdapter : public FAvaViewportColorPickerActorAdapter
+{
+	virtual ~TAvaViewportColorPickerActorAdapter() = default;
+
+	virtual FAvaColorChangeData GetColorData(const AActor* InActor) const override
+	{
+		if (const InType* CastActor = Cast<InType>(InActor))
+		{
+			return CastActor->GetColorData();
+		}
+
+		return FAvaColorChangeData();
+	}
+
+	virtual void SetColorData(AActor* InActor, const FAvaColorChangeData& InColorData) const override
+	{
+		if (InType* CastActor = Cast<InType>(InActor))
+		{
+			CastActor->SetColorData(InColorData);
+		}
+	}
+};
+
+template<typename InComponentType>
+struct TAvaViewportColorPickerActorComponentAdapter : public FAvaViewportColorPickerActorAdapter
+{
+	virtual ~TAvaViewportColorPickerActorComponentAdapter() = default;
+
+	virtual FAvaColorChangeData GetColorData(const AActor* InActor) const override
+	{
+		if (IsValid(InActor))
+		{
+			if (InComponentType* Component = InActor->FindComponentByClass<InComponentType>())
+			{
+				return Component->GetColorData();
+			}
+		}
+
+		return FAvaColorChangeData();
+	}
+
+	virtual void SetColorData(AActor* InActor, const FAvaColorChangeData& InColorData) const override
+	{
+		if (IsValid(InActor))
+		{
+			if (InComponentType* Component = InActor->FindComponentByClass<InComponentType>())
+			{
+				Component->SetColorData(InColorData);
+			}
+		}
+	}
+};
+
+struct AVALANCHELEVELVIEWPORT_API FAvaViewportColorPickerActorClassRegistry
+{
+	static void RegisterClassAdapter(UClass* InClass, const TSharedRef<FAvaViewportColorPickerActorAdapter>& InAdapter);
+
+	static bool ApplyColorDataToActor(AActor* InActor, const FAvaColorChangeData& InColorData);
+
+	static bool GetColorDataFromActor(const AActor* InActor, FAvaColorChangeData& OutColorData);
+};
diff --git a/Engine/Plugins/Experimental/Avalanche/Source/AvalancheLevelViewport/Public/Widgets/AvaViewportColorPickerDelegates.h b/Engine/Plugins/Experimental/Avalanche/Source/AvalancheLevelViewport/Public/Widgets/AvaViewportColorPickerDelegates.h
new file mode 100644
--- /dev/null
+++ b/Engine/Plugins/Experimental/Avalanche/Source/AvalancheLevelViewport/Public/Widgets/AvaViewportColorPickerDelegates.h
@@ -0,0 +1,26 @@
+// Copyright Epic Games, Inc. All Rights Reserved.
+
+#pragma once
+
+#include "AvaDefs.h"
+#include "Delegates/Delegate.h"
+#include "Delegates/DelegateCombinations.h"
+#include "Templates/SharedPointerFwd.h"
+
+class IToolkitHost;
+
+DECLARE_MULTICAST_DELEGATE_TwoParams(FAvaOnColorPicked, const TSharedRef<IToolkitHost>& Editor, const FAvaColorChangeData& NewColorData)
+DECLARE_MULTICAST_DELEGATE_TwoParams(FAvaOnColorSourceSelected, const TSharedRef<IToolkitHost>& Editor, const FAvaColorChangeData& NewColorData)
+
+struct AVALANCHELEVELVIEWPORT_API FAvaViewportColorPickerDelegates
+{
+	static FAvaOnColorPicked& GetOnColorPicked();
+
+	static void BroadcastColorPicked(const TSharedRef<IToolkitHost>& InEditor, const FAvaColorChangeData& InNewColorData);
+
+	static FAvaOnColorPicked& GetOnColorSourceSelected();
+
+	static void BroadcastColorSourceSelected(const TSharedRef<IToolkitHost>& InEditor, const FAvaColorChangeData& InNewColorData);
+
+	static FAvaColorChangeData GetLastColorData();
+};
diff --git a/Engine/Plugins/Experimental/Avalanche/Source/AvalancheMRQ/Private/AvaMRQRundownPageSetting.cpp b/Engine/Plugins/Experimental/Avalanche/Source/AvalancheMRQ/Private/AvaMRQRundownPageSetting.cpp
new file mode 100644
--- /dev/null
+++ b/Engine/Plugins/Experimental/Avalanche/Source/AvalancheMRQ/Private/AvaMRQRundownPageSetting.cpp
@@ -0,0 +1,140 @@
+// Copyright Epic Games, Inc. All Rights Reserved.
+
+#include "AvaMRQRundownPageSetting.h"
+#include "AvaRemoteControlRebind.h"
+#include "AvaRemoteControlUtils.h"
+#include "AvaScene.h"
+#include "AvaSceneSubsystem.h"
+#include "Engine/Engine.h"
+#include "Engine/Level.h"
+#include "IAvaSceneInterface.h"
+#include "LevelSequence.h"
+#include "Misc/CoreMiscDefines.h"
+#include "Misc/PackageName.h"
+#include "MoviePipeline.h"
+#include "Playlist/AvalanchePlaylist.h"
+
+DEFINE_LOG_CATEGORY_STATIC(LogAvaMRQSequenceData, Log, All);
+
+namespace UE::AvaMRQ::Private
+{
+	int32 FindPieInstanceId(UWorld* InWorld)
+	{
+		check(GEngine);
+		for (const FWorldContext& WorldContext : GEngine->GetWorldContexts())
+		{
+			// Return the first PIE world that we can find
+			if (WorldContext.WorldType == EWorldType::PIE && InWorld == WorldContext.World())
+			{
+				return WorldContext.PIEInstance;
+			}
+		}
+		return INDEX_NONE;
+	}
+
+	FSoftObjectPath RemovePiePrefix(const FSoftObjectPath& InPiePath, int32 InPieInstanceId)
+	{
+		FString NewPath = InPiePath.ToString();
+		if (FPackageName::GetLongPackageAssetName(NewPath).StartsWith(PLAYWORLD_PACKAGE_PREFIX))
+		{
+			FString PiePrefix = FString::Printf(TEXT("%s_%d_"), PLAYWORLD_PACKAGE_PREFIX, InPieInstanceId);
+			NewPath = NewPath.Replace(*PiePrefix, TEXT(""), ESearchCase::CaseSensitive);
+		}
+		return FSoftObjectPath(NewPath);
+	}
+}
+
+void UAvaMRQRundownPageSetting::SetupForPipelineImpl(UMoviePipeline* InPipeline)
+{
+	UWorld* World = GetWorld();
+	if (!InPipeline || !World)
+	{
+		return;
+	}
+
+	ULevelSequence* LevelSequence = InPipeline->GetTargetSequence();
+	if (!LevelSequence)
+	{
+		return;
+	}
+
+	UAvalanchePlaylist* Rundown = RundownPage.Rundown.LoadSynchronous();
+	if (!Rundown)
+	{
+		UE_LOG(LogAvaMRQSequenceData, Error
+			, TEXT("Motion Design Rundown '%s' was not valid and could not be loaded.")
+			, *RundownPage.Rundown.ToString());
+		return;
+	}
+
+	const FAvalanchePage& Page = Rundown->GetPage(RundownPage.PageId);
+	if (!Page.IsValidPage())
+	{
+		UE_LOG(LogAvaMRQSequenceData, Error
+			, TEXT("Motion Design Rundown '%s' did not have a valid page in %d.")
+			, *Rundown->GetFullName()
+			, RundownPage.PageId);
+		return;
+	}
+
+	// Ensure Page Source Asset matches
+	int32 PieInstanceId = UE::AvaMRQ::Private::FindPieInstanceId(World);
+	if (PieInstanceId != INDEX_NONE)
+	{
+		FSoftObjectPath SourcePath  = Page.GetAvalancheAssetPath(Rundown);
+		FSoftObjectPath CurrentPath = UE::AvaMRQ::Private::RemovePiePrefix(World, PieInstanceId);
+
+		if (CurrentPath != SourcePath)
+		{
+			UE_LOG(LogAvaMRQSequenceData, Error
+				, TEXT("Asset path '%s' in Page '%d' of Avalanche Rundown '%s' did not match the provided PIE sanitized world path '%s'")
+				, *SourcePath.ToString()
+				, Page.GetPageId()
+				, *Rundown->GetFullName()
+				, *CurrentPath.ToString());
+			return;
+		}
+	}
+
+	ULevel* const Level = World->PersistentLevel;
+
+	IAvaSceneInterface* Scene = UAvaSceneSubsystem::FindSceneInterface(Level);
+	if (!Scene)
+	{
+		UE_LOG(LogAvaMRQSequenceData, Error
+			, TEXT("Motion Design Scene Interface was not found in the provided World '%s'.")
+			, *World->GetFullName());
+		return;
+	}
+
+	URemoteControlPreset* RemoteControlPreset = Scene->GetRemoteControlPreset();
+	if (!RemoteControlPreset)
+	{
+		UE_LOG(LogAvaMRQSequenceData, Error
+			, TEXT("A Remote Control Preset was not found in the provided World '%s'.")
+			, *World->GetFullName());
+		return;
+	}
+
+	FAvaRemoteControlRebind::RebindUnboundEntities(RemoteControlPreset, Level);
+	Page.GetRemoteControlValues().ApplyEntityValuesToRemoteControlPreset(RemoteControlPreset);
+
+	// todo: ideally, _getUObject should not be used here, and instead the Viewport Data Provider gotten directly through an interface call
+	if (IAvaViewportDataProvider* ViewportDataProvider = Cast<IAvaViewportDataProvider>(Scene->_getUObject()))
+	{
+		if (APlayerController* PlayerController = World->GetFirstPlayerController())
+		{
+			FName StartupCameraName = ViewportDataProvider->GetStartupCameraName();
+
+			TObjectPtr<AActor>* ViewTarget = World->PersistentLevel->Actors.FindByPredicate([StartupCameraName](TObjectPtr<AActor> InActor)
+			{
+				return InActor && InActor->GetFName() == StartupCameraName;
+			});
+
+			if (ViewTarget)
+			{
+				PlayerController->SetViewTarget(*ViewTarget);
+			}
+		}
+	}
+}
