[MotionDesign]

diff --git a/Engine/Plugins/Experimental/Avalanche/Source/AvalancheMRQEditor/Private/AvaMRQEditorModule.cpp b/Engine/Plugins/Experimental/Avalanche/Source/AvalancheMRQEditor/Private/AvaMRQEditorModule.cpp
new file mode 100644
--- /dev/null
+++ b/Engine/Plugins/Experimental/Avalanche/Source/AvalancheMRQEditor/Private/AvaMRQEditorModule.cpp
@@ -0,0 +1,120 @@
+// Copyright Epic Games, Inc. All Rights Reserved.
+
+#include "AvaMRQEditorModule.h"
+#include "AvaMRQEditorCommands.h"
+#include "AvaMRQEditorPlaylistUtils.h"
+#include "Framework/MultiBox/MultiBoxExtender.h"
+#include "IAvaMediaEditorModule.h"
+#include "Modules/ModuleManager.h"
+#include "Playlist/AvaPlaylistEditor.h"
+#include "Playlist/AvalanchePlaylist.h"
+#include "Toolkits/AssetEditorToolkit.h"
+
+void FAvaMRQEditorModule::StartupModule()
+{
+	FAvaMRQEditorCommands::Register();
+
+	TSharedPtr<FExtensibilityManager> PlaylistToolbarExtensibility = IAvaMediaEditorModule::Get().GetPlaylistToolBarExtensibilityManager();
+	if (PlaylistToolbarExtensibility.IsValid())
+	{
+		PlaylistToolbarExtensibilityWeak = PlaylistToolbarExtensibility;
+
+		TArray<FAssetEditorExtender>& ExtenderDelegates = PlaylistToolbarExtensibility->GetExtenderDelegates();
+		ExtenderDelegates.Add(FAssetEditorExtender::CreateStatic(&FAvaMRQEditorModule::ExtendPlaylistToolbar));
+
+		PlaylistToolbarExtenderHandle = ExtenderDelegates.Last().GetHandle();
+	}
+}
+
+void FAvaMRQEditorModule::ShutdownModule()
+{
+	FAvaMRQEditorCommands::Unregister();
+
+	if (TSharedPtr<FExtensibilityManager> PlaylistToolbarExtensibility = PlaylistToolbarExtensibilityWeak.Pin())
+	{
+		if (PlaylistToolbarExtenderHandle.IsValid())
+		{
+			PlaylistToolbarExtensibility->GetExtenderDelegates().RemoveAll([this](const FAssetEditorExtender& InExtender)
+			{
+				return InExtender.GetHandle() == PlaylistToolbarExtenderHandle;
+			});
+			PlaylistToolbarExtenderHandle.Reset();
+		}
+		PlaylistToolbarExtensibilityWeak.Reset();
+	}
+}
+
+TSharedRef<FExtender> FAvaMRQEditorModule::ExtendPlaylistToolbar(const TSharedRef<FUICommandList> InCommandList, const TArray<UObject*> InObjects)
+{
+	TSharedRef<FExtender> Extender = MakeShared<FExtender>();
+
+	UAssetEditorSubsystem* const AssetEditorSubsystem = GEditor ? GEditor->GetEditorSubsystem<UAssetEditorSubsystem>() : nullptr;
+	if (!AssetEditorSubsystem)
+	{
+		return Extender;
+	}
+
+	TSharedRef<FAvaMRQPlaylistContext> Context = MakeShared<FAvaMRQPlaylistContext>();
+	Context->PlaylistEditors.Reserve(InObjects.Num());
+
+	const FName PlaylistEditorName = TEXT("AvaPlaylistEditor");
+
+	// Gather Playlist Editors
+	for (UObject* Object : InObjects)
+	{
+		if (UAvalanchePlaylist* Playlist = Cast<UAvalanchePlaylist>(Object))
+		{
+			IAssetEditorInstance* AssetEditor = AssetEditorSubsystem->FindEditorForAsset(Playlist, /*bFocusIfOpen*/false);
+			if (!AssetEditor || AssetEditor->GetEditorName() != PlaylistEditorName)
+			{
+				continue;
+			}
+
+			FAvaPlaylistEditor* PlaylistEditor = static_cast<FAvaPlaylistEditor*>(AssetEditor);
+			Context->PlaylistEditors.Add(StaticCastWeakPtr<FAvaPlaylistEditor>(PlaylistEditor->AsWeak()));
+		}
+	}
+
+	if (Context->PlaylistEditors.IsEmpty())
+	{
+		return Extender;
+	}
+
+	Extender->AddToolBarExtension("Pages"
+		, EExtensionHook::After
+		, CreatePlaylistActions(Context)
+		, FToolBarExtensionDelegate::CreateLambda([](FToolBarBuilder& ToolbarBuilder)
+		{
+			const FAvaMRQEditorCommands& MRQEditorCommands = FAvaMRQEditorCommands::Get();
+			ToolbarBuilder.AddToolBarButton(MRQEditorCommands.RenderSelectedPages);
+		})
+	);
+
+	return Extender;
+}
+
+TSharedRef<FUICommandList> FAvaMRQEditorModule::CreatePlaylistActions(TSharedRef<FAvaMRQPlaylistContext> InContext)
+{
+	const FAvaMRQEditorCommands& MRQEditorCommands = FAvaMRQEditorCommands::Get();
+
+	TSharedRef<FUICommandList> CommandList = MakeShared<FUICommandList>();
+
+	static auto HasAnySelectedPage = [](const TWeakPtr<const FAvaPlaylistEditor>& InPlaylistEditorWeak)
+	{
+		TSharedPtr<const FAvaPlaylistEditor> PlaylistEditor = InPlaylistEditorWeak.Pin();
+		return PlaylistEditor.IsValid() && PlaylistEditor->GetSelectedPagesOnActiveSubListWidget().Num() > 0;
+	};
+
+	FCanExecuteAction CanExecuteSelectedPageAction = FCanExecuteAction::CreateLambda([InContext]
+	{
+		return InContext->PlaylistEditors.ContainsByPredicate(HasAnySelectedPage);
+	});
+
+	CommandList->MapAction(MRQEditorCommands.RenderSelectedPages
+		, FExecuteAction::CreateLambda([InContext]{ FAvaMRQEditorPlaylistUtils::RenderSelectedPages(InContext->PlaylistEditors); })
+		, CanExecuteSelectedPageAction);
+
+	return CommandList;
+}
+
+IMPLEMENT_MODULE(FAvaMRQEditorModule, AvalancheMRQEditor)
diff --git a/Engine/Plugins/Experimental/Avalanche/Source/AvalancheMRQEditor/Private/AvaMRQEditorModule.h b/Engine/Plugins/Experimental/Avalanche/Source/AvalancheMRQEditor/Private/AvaMRQEditorModule.h
new file mode 100644
--- /dev/null
+++ b/Engine/Plugins/Experimental/Avalanche/Source/AvalancheMRQEditor/Private/AvaMRQEditorModule.h
@@ -0,0 +1,36 @@
+// Copyright Epic Games, Inc. All Rights Reserved.
+
+#pragma once
+
+#include "Containers/Array.h"
+#include "Delegates/Delegate.h"
+#include "Modules/ModuleInterface.h"
+#include "Templates/SharedPointer.h"
+
+class FAvaPlaylistEditor;
+class FExtender;
+class FExtensibilityManager;
+class FUICommandList;
+class UAvalanchePlaylist;
+class UObject;
+
+struct FAvaMRQPlaylistContext
+{
+    TArray<TWeakPtr<const FAvaPlaylistEditor>, TInlineAllocator<1>> PlaylistEditors;
+};
+
+class FAvaMRQEditorModule : public IModuleInterface
+{
+    //~ Begin IModuleInterface
+    virtual void StartupModule() override;
+    virtual void ShutdownModule() override;
+    //~ End IModuleInterface
+
+    static TSharedRef<FExtender> ExtendPlaylistToolbar(const TSharedRef<FUICommandList> InCommandList, const TArray<UObject*> InObjects);
+
+    static TSharedRef<FUICommandList> CreatePlaylistActions(TSharedRef<FAvaMRQPlaylistContext> InContext);
+
+    TWeakPtr<FExtensibilityManager> PlaylistToolbarExtensibilityWeak;
+
+    FDelegateHandle PlaylistToolbarExtenderHandle;
+};
diff --git a/Engine/Plugins/Experimental/Avalanche/Source/AvalancheMRQEditor/Private/AvaMRQEditorPlaylistUtils.cpp b/Engine/Plugins/Experimental/Avalanche/Source/AvalancheMRQEditor/Private/AvaMRQEditorPlaylistUtils.cpp
new file mode 100644
--- /dev/null
+++ b/Engine/Plugins/Experimental/Avalanche/Source/AvalancheMRQEditor/Private/AvaMRQEditorPlaylistUtils.cpp
@@ -0,0 +1,290 @@
+// Copyright Epic Games, Inc. All Rights Reserved.
+
+#include "AvaMRQEditorPlaylistUtils.h"
+#include "AvaMRQRundownPageSetting.h"
+#include "AvaSceneSubsystem.h"
+#include "AvaSequencerUtils.h"
+#include "EditorWorldUtils.h"
+#include "IAvaSceneInterface.h"
+#include "IAvaSequenceProvider.h"
+#include "MoviePipelineDeferredPasses.h"
+#include "MoviePipelineImageSequenceOutput.h"
+#include "MoviePipelineOutputSetting.h"
+#include "MoviePipelinePIEExecutor.h"
+#include "MoviePipelineQueueSubsystem.h"
+#include "Playback/AvaMediaPlaybackUtils.h"
+#include "Playlist/AvaPlaylistEditor.h"
+
+DEFINE_LOG_CATEGORY_STATIC(LogAvaMRQEditorPlaylist, Log, All);
+
+namespace UE::AvaMRQEditor::Private
+{
+	void CleanupRootedWorlds(TConstArrayView<TWeakObjectPtr<UWorld>> InRootedWorlds)
+	{
+		for (const TWeakObjectPtr<UWorld>& RootedWorldWeak : InRootedWorlds)
+		{
+			if (UWorld* const RootedWorld = RootedWorldWeak.Get())
+			{
+				RootedWorld->RemoveFromRoot();
+			}
+		}
+	}
+
+	struct FAvaMRQScopedRender
+	{
+		FAvaMRQScopedRender()
+		{
+			QueueSubsystem = GEditor ? GEditor->GetEditorSubsystem<UMoviePipelineQueueSubsystem>() : nullptr;
+			ensureMsgf(QueueSubsystem, TEXT("Not able to access UMoviePipelineQueueSubsystem (returned null)"));
+		}
+
+		~FAvaMRQScopedRender()
+		{
+			if (JobAllocationCount > 0)
+			{
+				UMoviePipelineExecutorBase* Executor = QueueSubsystem->RenderQueueWithExecutor(UMoviePipelinePIEExecutor::StaticClass());
+				if (Executor && !RootedWorlds.IsEmpty())
+				{
+					Executor->OnExecutorFinished().AddLambda([RootedWorlds = MoveTemp(RootedWorlds)](UMoviePipelineExecutorBase*, bool)
+					{
+						CleanupRootedWorlds(RootedWorlds);
+					});
+				}
+			}
+		}
+
+		UMoviePipelineExecutorJob* AllocateJob()
+		{
+			UMoviePipelineQueue* const Queue = QueueSubsystem->GetQueue();
+			if (!Queue)
+			{
+				return nullptr;
+			}
+			if (JobAllocationCount == 0)
+			{
+				Queue->DeleteAllJobs();
+			}
+			if (UMoviePipelineExecutorJob* Job = Queue->AllocateNewJob(UMoviePipelineExecutorJob::StaticClass()))
+			{
+				++JobAllocationCount;
+				return Job;
+			}
+			return nullptr;
+		}
+
+		bool IsValid() const
+		{
+			return ::IsValid(QueueSubsystem) && ::IsValid(QueueSubsystem->GetQueue());
+		}
+
+		void EnsureRootedWorld(UWorld* InWorldChecked)
+		{
+			check(InWorldChecked);
+			if (!InWorldChecked->IsRooted())
+			{
+				InWorldChecked->AddToRoot();
+				RootedWorlds.Add(InWorldChecked);
+			}
+		}
+
+	private:		
+		UMoviePipelineQueueSubsystem* QueueSubsystem = nullptr;
+
+		uint32 JobAllocationCount = 0;
+
+		TArray<TWeakObjectPtr<UWorld>> RootedWorlds;
+	};
+
+	void RenderSequence(FAvaMRQScopedRender& InScopedRender, UWorld* InWorld, const UAvalanchePlaylist& InPlaylist, const FAvalanchePage& InPage, UAvaSequence* InSequence)
+	{
+		UMoviePipelineExecutorJob* Job = InScopedRender.AllocateJob();
+		if (!Job)
+		{
+			return;
+		}
+
+		constexpr const TCHAR* Separator = TEXT("-");
+
+		Job->JobName = InPlaylist.GetName() + Separator + TEXT("Page_") + FString::FromInt(InPage.GetPageId()) + Separator + InSequence->GetName();
+		Job->Map = InWorld;
+		Job->SetSequence(InSequence);
+
+		UMoviePipelinePrimaryConfig* MRQConfig = Job->GetConfiguration();
+		check(MRQConfig);
+
+		auto CreatePipelineSetting = [MRQConfig]<typename InSettingType>(InSettingType*& OutSettings)
+		{
+			OutSettings = Cast<InSettingType>(MRQConfig->FindOrAddSettingByClass(InSettingType::StaticClass()));
+		};
+
+		UMoviePipelineImageSequenceOutput_PNG* OutputPNG;
+		CreatePipelineSetting(OutputPNG);
+
+		UMoviePipelineDeferredPassBase* DeferredPass;
+		CreatePipelineSetting(DeferredPass);
+
+		UMoviePipelineOutputSetting* OutputSettings;
+		CreatePipelineSetting(OutputSettings);
+		OutputSettings->FileNameFormat = TEXT("{job_name}.{frame_number}");;
+
+		UAvaMRQRundownPageSetting* RundownPageSetting;
+		CreatePipelineSetting(RundownPageSetting);
+		RundownPageSetting->RundownPage.Rundown = &InPlaylist;
+		RundownPageSetting->RundownPage.PageId  = InPage.GetPageId();
+
+		UE_LOG(LogAvaMRQEditorPlaylist, Log
+			, TEXT("New MRQ Job '%s' created for World '%s' and LevelSequence '%s'")
+			, *Job->GetFullName()
+			, *InWorld->GetFullName()
+			, *InSequence->GetFullName());
+	}
+
+	void RenderPage(FAvaMRQScopedRender& InScopedRender, const UAvalanchePlaylist& InPlaylist, const FAvalanchePage& InPage)
+	{
+		FSoftObjectPath PageAssetPath = InPage.GetAvalancheAssetPath(&InPlaylist);
+		if (!FAvaMediaPlaybackUtils::IsMapAsset(PageAssetPath.GetLongPackageName()))
+		{
+			UE_LOG(LogAvaMRQEditorPlaylist, Error
+				, TEXT("Page asset path '%s' (Page Id: '%d', Playlist '%s) is not a map asset and will not be processed")
+				, *PageAssetPath.ToString()
+				, InPage.GetPageId()
+				, *InPlaylist.GetName());
+			return;
+		}
+
+		UWorld* const World = Cast<UWorld>(PageAssetPath.TryLoad());
+		if (!World)
+		{
+			UE_LOG(LogAvaMRQEditorPlaylist, Error
+				, TEXT("World path '%s' (Page Id: '%d', Playlist '%s) did not load a valid world and will not be processed")
+				, *PageAssetPath.ToString()
+				, InPage.GetPageId()
+				, *InPlaylist.GetName());
+			return;
+		}
+		InScopedRender.EnsureRootedWorld(World);
+
+		IAvaSceneInterface* SceneInterface = UAvaSceneSubsystem::FindSceneInterface(World->PersistentLevel);
+		if (!SceneInterface || !SceneInterface->GetSequenceProvider())
+		{
+			UE_LOG(LogAvaMRQEditorPlaylist, Error
+				, TEXT("World '%s' (Page Id: '%d', Playlist '%s) did not have a valid Scene Interface to retrieve the Sequences so will not be processed")
+				, *PageAssetPath.ToString()
+				, InPage.GetPageId()
+				, *InPlaylist.GetName());
+			return;
+		}
+
+		TConstArrayView<UAvaSequence*> Sequences = SceneInterface->GetSequenceProvider()->GetSequences();
+		for (UAvaSequence* Sequence : Sequences)
+		{
+			RenderSequence(InScopedRender, World, InPlaylist, InPage, Sequence);
+		}
+	}
+
+	void RenderPages(FAvaMRQScopedRender& InScopedRender, const UAvalanchePlaylist& InPlaylist, TConstArrayView<int32> InPageIds)
+	{
+		TSet<int32> ProcessedPages;
+		ProcessedPages.Reserve(InPageIds.Num());
+
+		for (int32 PageId : InPageIds)
+		{
+			bool bIsAlreadyInSet = false;
+			ProcessedPages.Add(PageId, &bIsAlreadyInSet);
+
+			if (bIsAlreadyInSet)
+			{
+				UE_LOG(LogAvaMRQEditorPlaylist, Warning
+					, TEXT("Page Id '%d' is repeated in the page ids to export for Playlist '%s' and will not be processed multiple times")
+					, PageId
+					, *InPlaylist.GetName());
+				continue;
+			}
+
+			const FAvalanchePage& Page = InPlaylist.GetPage(PageId);
+			if (!Page.IsValidPage())
+			{
+				UE_LOG(LogAvaMRQEditorPlaylist, Warning
+					, TEXT("Page Id '%d' is invalid for Playlist '%s' and will not be processed")
+					, PageId
+					, *InPlaylist.GetName());
+				continue;
+			}
+
+			RenderPage(InScopedRender, InPlaylist, Page);
+		}
+	}
+}
+
+void FAvaMRQEditorPlaylistUtils::RenderSelectedPages(TConstArrayView<TWeakPtr<const FAvaPlaylistEditor>> InPlaylistEditors)
+{
+	if (InPlaylistEditors.IsEmpty())
+	{
+		return;
+	}
+
+	using namespace UE::AvaMRQEditor;
+	TSet<TWeakPtr<const FAvaPlaylistEditor>> ProcessedPlaylistEditors;
+	ProcessedPlaylistEditors.Reserve(InPlaylistEditors.Num());
+
+	Private::FAvaMRQScopedRender ScopedRender;
+	if (!ScopedRender.IsValid())
+	{
+		return;
+	}
+
+	for (TWeakPtr<const FAvaPlaylistEditor> PlaylistEditorWeak : InPlaylistEditors)
+	{
+		TSharedPtr<const FAvaPlaylistEditor> PlaylistEditor = PlaylistEditorWeak.Pin();
+		if (!PlaylistEditor || !PlaylistEditor->IsPlaylistValid())
+		{
+			continue;
+		}
+
+		UAvalanchePlaylist* Playlist = PlaylistEditor->GetPlaylist();
+		check(Playlist);
+
+		bool bIsAlreadyInSet = false;
+		ProcessedPlaylistEditors.Add(PlaylistEditorWeak, &bIsAlreadyInSet);
+
+		if (bIsAlreadyInSet)
+		{
+			UE_LOG(LogAvaMRQEditorPlaylist, Warning
+				, TEXT("Playlist Editor '%s' is repeated in the editors to export and will not be processed again")
+				, *Playlist->GetName());
+			continue;
+		}
+
+		TConstArrayView<int32> PageIds = PlaylistEditor->GetSelectedPagesOnActiveSubListWidget();
+		Private::RenderPages(ScopedRender, *Playlist, PageIds);
+	}
+}
+
+void FAvaMRQEditorPlaylistUtils::RenderPages(const UAvalanchePlaylist& InPlaylist, TConstArrayView<int32> InPageIds)
+{
+	if (InPageIds.IsEmpty())
+	{
+		return;
+	}
+
+	using namespace UE::AvaMRQEditor;
+	Private::FAvaMRQScopedRender ScopedRender;
+	if (!ScopedRender.IsValid())
+	{
+		return;
+	}
+
+	Private::RenderPages(ScopedRender, InPlaylist, InPageIds);
+}
+
+void FAvaMRQEditorPlaylistUtils::RenderPage(const UAvalanchePlaylist& InPlaylist, const FAvalanchePage& InPage)
+{
+	using namespace UE::AvaMRQEditor;
+	Private::FAvaMRQScopedRender ScopedRender;
+	if (!ScopedRender.IsValid())
+	{
+		return;
+	}
+
+	Private::RenderPage(ScopedRender, InPlaylist, InPage);
+}
diff --git a/Engine/Plugins/Experimental/Avalanche/Source/AvalancheMRQEditor/Private/AvaMRQEditorStyle.cpp b/Engine/Plugins/Experimental/Avalanche/Source/AvalancheMRQEditor/Private/AvaMRQEditorStyle.cpp
new file mode 100644
--- /dev/null
+++ b/Engine/Plugins/Experimental/Avalanche/Source/AvalancheMRQEditor/Private/AvaMRQEditorStyle.cpp
@@ -0,0 +1,23 @@
+// Copyright Epic Games, Inc. All Rights Reserved.
+
+#include "AvaMRQEditorStyle.h"
+#include "Misc/Paths.h"
+#include "Styling/SlateStyleMacros.h"
+#include "Styling/SlateStyleRegistry.h"
+
+FAvaMRQEditorStyle::FAvaMRQEditorStyle()
+	: FSlateStyleSet(TEXT("AvaMRQEditor"))
+{
+	const FVector2f Icon20(20.f);
+
+	ContentRootDir     = FPaths::EngineContentDir() / TEXT("Editor/Slate");
+	CoreContentRootDir = FPaths::EngineContentDir() / TEXT("Slate");
+
+	Set("AvaMRQEditorCommands.RenderSelectedPages", new IMAGE_BRUSH_SVG("Starship/MainToolbar/cinematics", Icon20));
+	FSlateStyleRegistry::RegisterSlateStyle(*this);
+}
+
+FAvaMRQEditorStyle::~FAvaMRQEditorStyle()
+{
+	FSlateStyleRegistry::UnRegisterSlateStyle(*this);
+}
diff --git a/Engine/Plugins/Experimental/Avalanche/Source/AvalancheMRQEditor/Private/AvaMRQEditorStyle.h b/Engine/Plugins/Experimental/Avalanche/Source/AvalancheMRQEditor/Private/AvaMRQEditorStyle.h
new file mode 100644
--- /dev/null
+++ b/Engine/Plugins/Experimental/Avalanche/Source/AvalancheMRQEditor/Private/AvaMRQEditorStyle.h
@@ -0,0 +1,19 @@
+// Copyright Epic Games, Inc. All Rights Reserved.
+
+#pragma once
+
+#include "Styling/SlateStyle.h"
+
+class FAvaMRQEditorStyle : public FSlateStyleSet
+{
+public:
+	FAvaMRQEditorStyle();
+
+	virtual ~FAvaMRQEditorStyle() override;
+
+	static FAvaMRQEditorStyle& Get()
+	{
+		static FAvaMRQEditorStyle StyleSet;
+		return StyleSet;
+	}
+};
diff --git a/Engine/Plugins/Experimental/Avalanche/Source/AvalancheMRQEditor/Public/AvaMRQEditorPlaylistUtils.h b/Engine/Plugins/Experimental/Avalanche/Source/AvalancheMRQEditor/Public/AvaMRQEditorPlaylistUtils.h
new file mode 100644
--- /dev/null
+++ b/Engine/Plugins/Experimental/Avalanche/Source/AvalancheMRQEditor/Public/AvaMRQEditorPlaylistUtils.h
@@ -0,0 +1,20 @@
+// Copyright Epic Games, Inc. All Rights Reserved.
+
+#pragma once
+
+#include "Containers/ArrayView.h"
+#include "Templates/SharedPointerFwd.h"
+
+class FAvaPlaylistEditor;
+class UAvaSequence;
+class UAvalanchePlaylist;
+struct FAvalanchePage;
+
+struct FAvaMRQEditorPlaylistUtils
+{
+	AVALANCHEMRQEDITOR_API static void RenderSelectedPages(TConstArrayView<TWeakPtr<const FAvaPlaylistEditor>> InPlaylistEditors);
+
+	AVALANCHEMRQEDITOR_API static void RenderPages(const UAvalanchePlaylist& InPlaylist, TConstArrayView<int32> InPageIds);
+
+	AVALANCHEMRQEDITOR_API static void RenderPage(const UAvalanchePlaylist& InPlaylist, const FAvalanchePage& InPage);
+};
diff --git a/Engine/Plugins/Experimental/Avalanche/Source/AvalancheMask/Private/AvaMaskLog.h b/Engine/Plugins/Experimental/Avalanche/Source/AvalancheMask/Private/AvaMaskLog.h
new file mode 100644
--- /dev/null
+++ b/Engine/Plugins/Experimental/Avalanche/Source/AvalancheMask/Private/AvaMaskLog.h
@@ -0,0 +1,7 @@
+// Copyright Epic Games, Inc. All Rights Reserved.
+
+#pragma once
+
+#include "Logging/LogMacros.h"
+
+DECLARE_LOG_CATEGORY_EXTERN(LogAvaMask, VeryVerbose, All);
diff --git a/Engine/Plugins/Experimental/Avalanche/Source/AvalancheMask/Private/AvaMaskSettings.cpp b/Engine/Plugins/Experimental/Avalanche/Source/AvalancheMask/Private/AvaMaskSettings.cpp
new file mode 100644
--- /dev/null
+++ b/Engine/Plugins/Experimental/Avalanche/Source/AvalancheMask/Private/AvaMaskSettings.cpp
@@ -0,0 +1,50 @@
+// Copyright Epic Games, Inc. All Rights Reserved.
+
+#include "AvaMaskSettings.h"
+
+#include "AvaMaskLog.h"
+
+#if WITH_EDITORONLY_DATA
+#include "AssetRegistry/AssetData.h"
+#include "AssetRegistry/AssetRegistryModule.h"
+#include "Interfaces/IPluginManager.h"
+#include "Materials/MaterialFunctionInterface.h"
+#endif
+
+UAvalancheMaskSettings::UAvalancheMaskSettings()
+{
+    CategoryName = TEXT("Plugins");
+    SectionName = TEXT("AvalancheMask");
+}
+
+#if WITH_EDITORONLY_DATA
+UMaterialFunctionInterface* UAvalancheMaskSettings::GetMaterialFunction()
+{
+	if (UMaterialFunctionInterface* StrongMaterialFunction = MaterialFunction.Get())
+	{
+		return StrongMaterialFunction;		
+	}
+
+	UMaterialFunctionInterface* DefaultMaterialFunction = GetDefaultMaterialFunction();
+	MaterialFunction = DefaultMaterialFunction;
+	
+	return DefaultMaterialFunction;
+}
+
+UMaterialFunctionInterface* UAvalancheMaskSettings::GetDefaultMaterialFunction()
+{
+	// /Script/Engine.MaterialFunction'/GeometryMask/GeometryMask/MF_ApplyGeometryMask.MF_ApplyGeometryMask'
+	// const FString AssetPath = FString(TEXT("/GeometryMask")) / TEXT("GeometryMask") / TEXT("MF_ApplyGeometryMask");
+	// const FString AssetPath = TEXT("/Script/Engine.MaterialFunction'/GeometryMask/GeometryMask/MF_ApplyGeometryMask.MF_ApplyGeometryMask'");
+	const FString AssetPath = TEXT("/Script/Engine.MaterialFunction'/Avalanche/MaskResources/MF_ApplyMask2D_Single.MF_ApplyMask2D_Single'");
+	const FAssetRegistryModule& AssetRegistryModule = FModuleManager::LoadModuleChecked<FAssetRegistryModule>("AssetRegistry");
+	if (const FAssetData AssetData = AssetRegistryModule.Get().GetAssetByObjectPath(FSoftObjectPath(AssetPath));
+		!AssetData.IsValid())
+	{
+		UE_LOG(LogAvaMask, Error, TEXT("Default MaterialFunction asset is invalid or not found: '%s'"), *AssetPath);
+		return nullptr;
+	}
+	
+	return Cast<UMaterialFunctionInterface>(StaticLoadObject(UMaterialFunctionInterface::StaticClass(), this, *AssetPath));
+}
+#endif
diff --git a/Engine/Plugins/Experimental/Avalanche/Source/AvalancheMask/Private/AvaMaskSettings.h b/Engine/Plugins/Experimental/Avalanche/Source/AvalancheMask/Private/AvaMaskSettings.h
new file mode 100644
--- /dev/null
+++ b/Engine/Plugins/Experimental/Avalanche/Source/AvalancheMask/Private/AvaMaskSettings.h
@@ -0,0 +1,39 @@
+// Copyright Epic Games, Inc. All Rights Reserved.
+
+#pragma once
+
+#include "Engine/DeveloperSettings.h"
+#if WITH_EDITORONLY_DATA
+#include "UObject/SoftObjectPtr.h"
+#endif
+
+#include "AvaMaskSettings.generated.h"
+
+#if WITH_EDITORONLY_DATA
+class UMaterialFunctionInterface;
+#endif
+
+/** Settings for Avalanche Mask */
+UCLASS(Config = Engine)
+class UAvalancheMaskSettings
+    : public UDeveloperSettings
+{
+    GENERATED_BODY()
+
+public:
+    UAvalancheMaskSettings();
+
+#if WITH_EDITORONLY_DATA
+	/** Material Function to use to expect or add to a material. */
+	UPROPERTY(Config, EditAnywhere, BlueprintReadOnly, NoClear, Category = "Material")
+	TSoftObjectPtr<UMaterialFunctionInterface> MaterialFunction;
+
+	/** Get the user specified or default material function. */
+	UMaterialFunctionInterface* GetMaterialFunction();
+#endif
+
+private:
+#if WITH_EDITORONLY_DATA
+	UMaterialFunctionInterface* GetDefaultMaterialFunction();
+#endif
+};
diff --git a/Engine/Plugins/Experimental/Avalanche/Source/AvalancheMask/Private/AvaMaskTypes.cpp b/Engine/Plugins/Experimental/Avalanche/Source/AvalancheMask/Private/AvaMaskTypes.cpp
new file mode 100644
--- /dev/null
+++ b/Engine/Plugins/Experimental/Avalanche/Source/AvalancheMask/Private/AvaMaskTypes.cpp
@@ -0,0 +1,86 @@
+﻿// Copyright Epic Games, Inc. All Rights Reserved.
+
+#include "AvaMaskTypes.h"
+
+#include "Engine/Texture.h"
+#include "Materials/MaterialInterface.h"
+
+bool FAvaMask2DMaterialParameters::SetFromMaterial(
+	const UMaterialInterface* InMaterial
+	, const FMaterialParameterInfo& InTextureParameter
+	, const FMaterialParameterInfo& InChannelParameter
+	, const FMaterialParameterInfo& InInvertParameter
+	, const FMaterialParameterInfo& InBaseOpacityParameter
+	, const FMaterialParameterInfo& InPaddingParameter
+	, const FMaterialParameterInfo& InFeatherParameter)
+{
+	if (!ensureAlways(InMaterial))
+	{
+		return false;
+	}
+
+	// Texture
+	{
+		UTexture* FoundTexture = nullptr;
+		if (!InMaterial->GetTextureParameterValue(InTextureParameter, FoundTexture))
+		{
+			return false;
+		}
+		Texture = FoundTexture;
+	}
+
+	// Base Opacity
+	{
+		float FoundValue = 0.0f;
+		if (!InMaterial->GetScalarParameterValue(InBaseOpacityParameter, FoundValue))
+		{
+			return false;
+		}
+		BaseOpacity = FoundValue;
+	}
+
+	// Channel
+	{
+		FLinearColor FoundValue;
+		if (!InMaterial->GetVectorParameterValue(InChannelParameter, FoundValue))
+		{
+			return false;
+		}
+		ChannelAsVector = FoundValue;
+		Channel = UE::GeometryMask::VectorToMaskChannel(ChannelAsVector);
+	}
+
+	// Invert
+	{
+		float FoundValue = 0.0f;
+		if (!InMaterial->GetScalarParameterValue(InInvertParameter, FoundValue))
+		{
+			return false;
+		}
+		bInvert = FoundValue > 0.1;
+	}
+
+	// Padding
+	{
+		FLinearColor FoundValue;
+		if (!InMaterial->GetVectorParameterValue(InPaddingParameter, FoundValue))
+		{
+			return false;			
+		}
+		Padding = FVector2f(FoundValue.R, FoundValue.G);
+	}
+
+	// Feathering
+	{
+		FLinearColor FoundValue;
+		if (!InMaterial->GetVectorParameterValue(InFeatherParameter, FoundValue))
+		{
+			return false;
+		}
+		bApplyFeathering = !FMath::IsNearlyZero(FoundValue.R);
+		OuterFeatherRadius = FoundValue.G;
+		InnerFeatherRadius = FoundValue.B;
+	}
+
+	return true;
+}
diff --git a/Engine/Plugins/Experimental/Avalanche/Source/AvalancheMask/Private/AvaMaskUtilities.h b/Engine/Plugins/Experimental/Avalanche/Source/AvalancheMask/Private/AvaMaskUtilities.h
new file mode 100644
--- /dev/null
+++ b/Engine/Plugins/Experimental/Avalanche/Source/AvalancheMask/Private/AvaMaskUtilities.h
@@ -0,0 +1,106 @@
+﻿// Copyright Epic Games, Inc. All Rights Reserved.
+
+#pragma once
+
+#include "Components/ActorComponent.h"
+#include "Engine/EngineTypes.h"
+#include "GeometryMaskTypes.h"
+#include "MaterialTypes.h"
+#include "Templates/SubclassOf.h"
+#include "Templates/UnrealTypeTraits.h"
+
+class UMaterialFunctionInterface;
+class UDynamicMaterialInstance;
+class UMaterialInstance;
+
+namespace UE::AvaMask::Internal
+{
+	/** Note that this assumes the parameters have the same name for all materials, which they should if using the provided material function. */
+	static FMaterialParameterInfo TextureParameterInfo = { TEXT("Mask_Texture") };
+	static FMaterialParameterInfo BaseOpacityParameterInfo = { TEXT("Mask_BaseOpacity") };
+	static FMaterialParameterInfo ChannelParameterInfo = { TEXT("Mask_Channel") };
+	static FMaterialParameterInfo InvertParameterInfo = { TEXT("Mask_Invert") };
+	static FMaterialParameterInfo PaddingParameterInfo = { TEXT("Mask_Padding") };
+	static FMaterialParameterInfo FeatherParameterInfo = { TEXT("Mask_Feather") };
+	
+	static TMap<FLinearColor, EGeometryMaskColorChannel> VectorToMaskChannelEnum =
+	{
+		{ {1, 0, 0, 0}, EGeometryMaskColorChannel::Red },
+		{ {0, 1, 0, 0}, EGeometryMaskColorChannel::Green },
+		{ {0, 0, 1, 0}, EGeometryMaskColorChannel::Blue },
+		{ {0, 0, 0, 1}, EGeometryMaskColorChannel::Alpha }
+	};
+	
+	static TMap<EGeometryMaskColorChannel, FLinearColor> MaskChannelEnumToVector =
+	{
+		{ EGeometryMaskColorChannel::Red, {1, 0, 0, 0} },
+		{ EGeometryMaskColorChannel::Green, {0, 1, 0, 0} },
+		{ EGeometryMaskColorChannel::Blue, {0, 0, 1, 0} },
+		{ EGeometryMaskColorChannel::Alpha, {0, 0, 0, 1} }
+	};
+
+	static FName HandleTag = TEXT("Modifier.Mask2D");
+
+	/** Utility function, will set "OtherActor" if the components owner differs from the provided one. */
+	FSoftComponentReference MakeComponentReference(const AActor* InOwner, const UActorComponent* InComponent);
+
+	/** Creates a unique (deterministic) key based on the provided parameters. */
+	uint32 MakeMaterialInstanceKey(const UMaterialInterface* InMaterial, const FName InMaskChannelName, EBlendMode InBlendMode, const int32 InSeed = 128);
+
+	UActorComponent* FindOrAddComponent(const TSubclassOf<UActorComponent>& InComponentClass, AActor* InActor);
+
+	template<
+		typename InComponentClass
+		UE_REQUIRES(std::is_base_of_v<UActorComponent, InComponentClass>)>
+	InComponentClass* FindOrAddComponent(AActor* InActor)
+	{
+		return Cast<InComponentClass>(FindOrAddComponent(InComponentClass::StaticClass(), InActor));
+	}
+
+	void RemoveComponentByInterface(const TSubclassOf<UInterface>& InInterfaceType, const AActor* InActor);
+
+	template<
+		typename InInterfaceType
+		UE_REQUIRES(std::is_base_of_v<UInterface, InInterfaceType>)>
+	void RemoveComponentByInterface(AActor* InActor)
+	{
+		RemoveComponentByInterface(InInterfaceType::StaticClass(), InActor);
+	}
+
+	FString GetGeneratedMaterialPath();
+
+	FString GetBlendModeString(const EBlendMode InBlendMode);
+
+	/** Returns whichever blend mode satisfies the requirement AND user specified blend mode */
+	EBlendMode GetTargetBlendMode(const EBlendMode InFromMaterial, const EBlendMode InRequired);
+
+	template <typename ElementType, typename KeyFuncs, typename AllocatorType, typename LookupType, typename LambdaType>
+	ElementType& FindOrAddByLambda(TSet<ElementType, KeyFuncs, AllocatorType>& Set, const LookupType& Key, LambdaType Lambda)
+	{
+		uint32 KeyHash = KeyFuncs::GetKeyHash(Key);
+		if (ElementType* Found = Set.FindByHash(Key))
+		{
+			return *Found;
+		}
+
+		ElementType NewElement = Invoke(Lambda, Key);
+		check(KeyFuncs::GetKeyHash(NewElement) == KeyHash);
+		FSetElementId NewElementId = Set.AddByHash(KeyHash, MoveTemp(NewElement));
+		return Set[NewElementId];
+	}
+
+	template <typename KeyType, typename ValueType, typename AllocatorType, typename KeyFuncs, typename InitKeyType, typename LambdaType>
+	ValueType& FindOrAddByLambda(TMap<KeyType, ValueType, AllocatorType, KeyFuncs>& Map, InitKeyType&& Key, LambdaType Lambda)
+	{
+		uint32 KeyHash = KeyFuncs::GetKeyHash(Key);
+		if (ValueType* Found = Map.FindByHash(KeyHash, Key))
+		{
+			return *Found;
+		}
+
+		ValueType NewElement = Invoke(Lambda);
+		return Map.AddByHash(KeyHash, Forward<InitKeyType>(Key), MoveTemp(NewElement));
+	}
+
+	UMaterialInterface* GetNonTransientParentMaterial(const UMaterialInstance* InMaterialInstance);
+}
diff --git a/Engine/Plugins/Experimental/Avalanche/Source/AvalancheMask/Private/Handling/AvaMaskParametricMaterialHandle.cpp b/Engine/Plugins/Experimental/Avalanche/Source/AvalancheMask/Private/Handling/AvaMaskParametricMaterialHandle.cpp
new file mode 100644
--- /dev/null
+++ b/Engine/Plugins/Experimental/Avalanche/Source/AvalancheMask/Private/Handling/AvaMaskParametricMaterialHandle.cpp
@@ -0,0 +1,279 @@
+﻿// Copyright Epic Games, Inc. All Rights Reserved.
+
+#include "Handling/AvaMaskParametricMaterialHandle.h"
+
+#include "AvaMaskUtilities.h"
+#include "AvaShapeParametricMaterial.h"
+#include "Engine/Texture.h"
+#include "Handling/AvaParametricMaterialHandle.h"
+#include "Materials/MaterialInstanceDynamic.h"
+
+FAvaMaskParametricMaterialHandle::FAvaMaskParametricMaterialHandle(const FStructView& InParametricMaterial)
+	: FAvaParametricMaterialHandle(InParametricMaterial)
+{
+	if (FAvaShapeParametricMaterial* ParametricMtl = GetParametricMaterial())
+	{
+		if (UMaterialInstanceDynamic* MID = ParametricMtl->GetMaterial())
+		{
+			LastShadingModel = MID->GetShadingModels();
+			LastBlendMode = MID->GetBlendMode();
+			LastMaterialInstance = MID;
+		}
+
+		if (!ParametricMtl->OnMaterialParameterChanged().IsBoundToObject(this))
+		{
+			MaterialChangedHandle = ParametricMtl->OnMaterialParameterChanged().AddRaw(this, &FAvaMaskParametricMaterialHandle::OnMaterialChanged);
+		}
+	}
+}
+
+FAvaMaskParametricMaterialHandle::~FAvaMaskParametricMaterialHandle()
+{
+	// @note: if this doesn't cover all situations, the handle needs to be aware of the struct's outer to check validity
+	if (!GIsGarbageCollecting)
+	{
+		if (FAvaShapeParametricMaterial* ParametricMtl = GetParametricMaterial())
+		{
+			ParametricMtl->OnMaterialParameterChanged().RemoveAll(this);
+		}
+	}
+}
+
+bool FAvaMaskParametricMaterialHandle::GetMaskParameters(FAvaMask2DMaterialParameters& OutParameters)
+{
+	if (const UMaterialInstanceDynamic* MaterialInstance = WeakMaterialInstance.Get())
+	{
+		OutParameters.SetFromMaterial(
+			MaterialInstance
+			, UE::AvaMask::Internal::TextureParameterInfo
+			, UE::AvaMask::Internal::ChannelParameterInfo
+			, UE::AvaMask::Internal::InvertParameterInfo
+			, UE::AvaMask::Internal::BaseOpacityParameterInfo
+			, UE::AvaMask::Internal::PaddingParameterInfo
+			, UE::AvaMask::Internal::FeatherParameterInfo);
+		return true;
+	}
+
+	return false;
+}
+
+bool FAvaMaskParametricMaterialHandle::SetMaskParameters(
+	UTexture* InTexture
+	, float InBaseOpacity
+	, EGeometryMaskColorChannel InChannel
+	, bool bInInverted
+	, const FVector2f& InPadding
+	, bool bInApplyFeathering
+	, float InOuterFeathering
+	, float InInnerFeathering)
+{
+	if (UMaterialInstanceDynamic* MaterialInstance = GetMaterialInstance())
+	{
+		UE_LOG(LogAvaMask, Display, TEXT("SetParameters: Texture:%s, Channel:%s"), InTexture ? *InTexture->GetName() : TEXT("(None)"), *UE::AvaMask::Internal::MaskChannelEnumToVector[InChannel].ToString());
+
+		UTexture* DefaultTexture = nullptr;
+		MaterialInstance->GetTextureParameterDefaultValue(UE::AvaMask::Internal::TextureParameterInfo, DefaultTexture);
+
+		InTexture = InTexture ? InTexture : DefaultTexture;
+
+		MaterialInstance->SetTextureParameterValueByInfo(UE::AvaMask::Internal::TextureParameterInfo, InTexture);
+		MaterialInstance->SetVectorParameterValueByInfo(UE::AvaMask::Internal::ChannelParameterInfo, UE::AvaMask::Internal::MaskChannelEnumToVector[InChannel]);
+		MaterialInstance->SetScalarParameterValueByInfo(UE::AvaMask::Internal::InvertParameterInfo, bInInverted ? 1.0f : 0.0f);
+		MaterialInstance->SetScalarParameterValueByInfo(UE::AvaMask::Internal::BaseOpacityParameterInfo, FMath::Clamp(InBaseOpacity, 0.0f, 1.0f));
+		MaterialInstance->SetVectorParameterValueByInfo(UE::AvaMask::Internal::PaddingParameterInfo, FLinearColor(FMath::Max(0, InPadding.X), FMath::Max(0, InPadding.Y), 0, 0));
+		MaterialInstance->SetVectorParameterValueByInfo(UE::AvaMask::Internal::FeatherParameterInfo, FLinearColor(bInApplyFeathering ? 1.0f : 0.0f, FMath::Max(0, InOuterFeathering), FMath::Max(0, InInnerFeathering), FMath::Max(0, FMath::Max(InOuterFeathering, InInnerFeathering))));
+
+		return true;
+	}
+
+	return false;
+}
+
+const EBlendMode FAvaMaskParametricMaterialHandle::GetBlendMode()
+{
+	if (const FAvaShapeParametricMaterial* ParametricMtl = GetParametricMaterial())
+	{
+		return ParametricMtl->GetUseTranslucentMaterial() ? EBlendMode::BLEND_Translucent : EBlendMode::BLEND_Opaque;
+	}
+	
+	if (const UMaterialInterface* MaterialInstance = GetMaterialInstance())
+	{
+		return MaterialInstance->GetBlendMode();
+	}
+	
+	if (const UMaterialInterface* Material = GetMaterial())
+	{
+		return Material->GetBlendMode();
+	}
+
+	return EBlendMode::BLEND_Opaque;
+}
+
+void FAvaMaskParametricMaterialHandle::SetBlendMode(const EBlendMode InBlendMode)
+{
+	const EBlendMode TargetBlendMode = UE::AvaMask::Internal::GetTargetBlendMode( GetBlendMode(), InBlendMode);
+	if (FAvaShapeParametricMaterial* ParametricMtl = GetParametricMaterial())
+	{
+		ParametricMtl->SetUseAutoTranslucency(false); // Disable so that user changes won't override the mask
+		ParametricMtl->SetUseTranslucentMaterial(TargetBlendMode != EBlendMode::BLEND_Opaque);
+		return;
+	}
+	
+	if (UMaterialInstanceDynamic* MaterialInstance = GetMaterialInstance())
+	{
+		MaterialInstance->BlendMode = TargetBlendMode;
+		MaterialInstance->BasePropertyOverrides.BlendMode = TargetBlendMode;
+	}
+}
+
+void FAvaMaskParametricMaterialHandle::SetChildMaterial(UMaterialInstanceDynamic* InChildMaterial)
+{
+	ChildMaterial = InChildMaterial;
+}
+
+bool FAvaMaskParametricMaterialHandle::SaveOriginalState(const FStructView& InHandleData)
+{
+	// @note: this assumes the current material is the original
+	if (const FAvaShapeParametricMaterial* ParametricMtl = GetParametricMaterial())
+	{
+		ensureAlways(IsValid());
+		
+		if (FHandleData* MaterialHandleData = InHandleData.GetPtr<FHandleData>())
+		{
+			MaterialHandleData->OriginalBlendMode = GetBlendMode();
+
+			const UMaterialInterface* Material = ParametricMtl->GetMaterial();
+
+			// Can happen when Parametric material not yet initialized - but initialization is accounted for elsewhere in the Handle
+			if (!Material)
+			{
+				Material = ParametricMtl->GetDefaultMaterial();
+			}
+
+			MaterialHandleData->OriginalMaskMaterialParameters.SetFromMaterial(
+				Material
+				, UE::AvaMask::Internal::TextureParameterInfo
+				, UE::AvaMask::Internal::ChannelParameterInfo
+				, UE::AvaMask::Internal::InvertParameterInfo
+				, UE::AvaMask::Internal::BaseOpacityParameterInfo
+				, UE::AvaMask::Internal::PaddingParameterInfo
+				, UE::AvaMask::Internal::FeatherParameterInfo);
+
+			MaterialHandleData->OriginalMaterialStyle = ParametricMtl->GetStyle();
+			MaterialHandleData->OriginalBlendMode = GetBlendMode();
+
+			return true;
+		}
+	}
+	
+	return false;
+}
+
+bool FAvaMaskParametricMaterialHandle::ApplyOriginalState(
+	const FStructView& InHandleData
+	, TUniqueFunction<void(UMaterialInterface*)>&& InMaterialSetter)
+{
+	if (FAvaShapeParametricMaterial* ParametricMtl = GetParametricMaterial())
+	{
+		if (const FHandleData* MaterialHandleData = InHandleData.GetPtr<const FHandleData>())
+		{
+			SetBlendMode(MaterialHandleData->OriginalBlendMode);
+			
+			SetMaskParameters(
+				MaterialHandleData->OriginalMaskMaterialParameters.Texture
+				, MaterialHandleData->OriginalMaskMaterialParameters.BaseOpacity
+				, MaterialHandleData->OriginalMaskMaterialParameters.Channel
+				, MaterialHandleData->OriginalMaskMaterialParameters.bInvert
+				, MaterialHandleData->OriginalMaskMaterialParameters.Padding
+				, MaterialHandleData->OriginalMaskMaterialParameters.bApplyFeathering
+				, MaterialHandleData->OriginalMaskMaterialParameters.OuterFeatherRadius
+				, MaterialHandleData->OriginalMaskMaterialParameters.InnerFeatherRadius);
+
+			ParametricMtl->SetStyle(MaterialHandleData->OriginalMaterialStyle);
+			ParametricMtl->SetUseAutoTranslucency(true);
+			ParametricMtl->SetUseTranslucentMaterial(MaterialHandleData->OriginalBlendMode == EBlendMode::BLEND_Translucent);
+
+			InMaterialSetter(ParametricMtl->GetMaterial());
+
+			return true;
+		}
+	}
+
+	return false;
+}
+
+bool FAvaMaskParametricMaterialHandle::ApplyModifiedState(
+	const FAvaMask2DSubjectParameters& InModifiedParameters
+	, const FStructView& InHandleData
+	, TUniqueFunction<void(UMaterialInterface*)>&& InMaterialSetter)
+{
+	if (FAvaShapeParametricMaterial* ParametricMtl = GetParametricMaterial())
+	{
+		if (InHandleData.GetPtr<const FHandleData>())
+		{
+			SetBlendMode(InModifiedParameters.MaterialParameters.BlendMode);
+
+			SetMaskParameters(
+				InModifiedParameters.MaterialParameters.Texture
+				, InModifiedParameters.MaterialParameters.BaseOpacity
+				, InModifiedParameters.MaterialParameters.Channel
+				, InModifiedParameters.MaterialParameters.bInvert
+				, InModifiedParameters.MaterialParameters.Padding
+				, InModifiedParameters.MaterialParameters.bApplyFeathering
+				, InModifiedParameters.MaterialParameters.OuterFeatherRadius
+				, InModifiedParameters.MaterialParameters.InnerFeatherRadius);
+
+			InMaterialSetter(ParametricMtl->GetMaterial());
+
+			return true;
+		}
+	}
+
+	return false;
+}
+
+bool FAvaMaskParametricMaterialHandle::IsValid() const
+{
+	return FAvaParametricMaterialHandle::IsValid();
+}
+
+bool FAvaMaskParametricMaterialHandle::IsSupported(
+	const UStruct* InStruct
+	, const TVariant<UObject*, FStructView>& InInstance
+	, FName InTag)
+{
+	return InTag == UE::AvaMask::Internal::HandleTag
+		&& FAvaParametricMaterialHandle::IsSupported(InStruct, InInstance, InTag);
+}
+
+UMaterialInstanceDynamic* FAvaMaskParametricMaterialHandle::GetMaterialInstance()
+{
+	return FAvaParametricMaterialHandle::GetMaterialInstance();
+}
+
+void FAvaMaskParametricMaterialHandle::OnMaterialChanged()
+{
+	if (const FAvaShapeParametricMaterial* ParametricMtl = GetParametricMaterial())
+	{
+		if (UMaterialInstanceDynamic* MaterialInstance = ParametricMtl->GetMaterial())
+		{
+			const FMaterialShadingModelField ShadingModel = MaterialInstance->GetShadingModels();
+			const EBlendMode BlendMode = MaterialInstance->GetBlendMode();
+
+			if (UMaterialInstance* LastMaterial = LastMaterialInstance.Get())
+			{
+				MaterialInstance->CopyInterpParameters(LastMaterial);	
+			}
+
+			LastShadingModel = ShadingModel;
+			LastBlendMode = BlendMode;
+			LastMaterialInstance = MaterialInstance;
+		}
+
+		if (UMaterialInstanceDynamic* ChildMtl = ChildMaterial.Get())
+		{
+			// Write the changed params to the Child Material
+			ParametricMtl->SetMaterialParameterValues(ChildMtl);
+		}
+	}
+}
diff --git a/Engine/Plugins/Experimental/Avalanche/Source/AvalancheMask/Private/Handling/AvaMaskText3DActorMaterialCollectionHandle.cpp b/Engine/Plugins/Experimental/Avalanche/Source/AvalancheMask/Private/Handling/AvaMaskText3DActorMaterialCollectionHandle.cpp
new file mode 100644
--- /dev/null
+++ b/Engine/Plugins/Experimental/Avalanche/Source/AvalancheMask/Private/Handling/AvaMaskText3DActorMaterialCollectionHandle.cpp
@@ -0,0 +1,833 @@
+﻿// Copyright Epic Games, Inc. All Rights Reserved.
+
+#include "Handling/AvaMaskText3DActorMaterialCollectionHandle.h"
+
+#include "AvaMaskLog.h"
+#include "AvaMaskUtilities.h"
+#include "AvaObjectHandleSubsystem.h"
+#include "AvaText3DComponent.h"
+#include "Components/PrimitiveComponent.h"
+#include "GameFramework/Actor.h"
+#include "Handling/AvaHandleUtilities.h"
+#include "IAvaMaskMaterialCollectionHandle.h"
+#include "IAvaMaskMaterialHandle.h"
+#include "MaterialHub/AvaTextMaterialHub.h"
+#include "Materials/MaterialInterface.h"
+#include "Mesh.h"
+#include "Text3DComponent.h"
+
+namespace UE::AvaMask::Private
+{
+	EAvaTextTranslucency GetTargetTranslucencyType(
+		const EAvaTextTranslucency InFromMaterial
+		, const EBlendMode InRequired)
+	{
+		if (InRequired == EBlendMode::BLEND_Opaque)
+		{
+			return InFromMaterial;
+		}
+		
+		if (InFromMaterial == EAvaTextTranslucency::None)
+		{
+			return EAvaTextTranslucency::Translucent;
+		}
+
+		return InFromMaterial;
+	}
+}
+
+FAvaMaskText3DActorMaterialCollectionHandle::FAvaMaskText3DActorMaterialCollectionHandle(AActor* InActor)
+	: WeakActor(InActor)
+	, WeakComponent(InActor->GetComponentByClass<UText3DComponent>())
+{
+}
+
+TArray<TObjectPtr<UMaterialInterface>> FAvaMaskText3DActorMaterialCollectionHandle::GetMaterials()
+{
+	TArray<TObjectPtr<UMaterialInterface>> Materials;
+	Materials.Reserve(GetNumMaterials());
+	
+	if (const UText3DComponent* TextComponent = GetComponent())
+	{
+		Materials.Emplace(TextComponent->GetFrontMaterial());
+		Materials.Emplace(TextComponent->GetBevelMaterial());
+		Materials.Emplace(TextComponent->GetExtrudeMaterial());
+		Materials.Emplace(TextComponent->GetBackMaterial());
+	}
+
+	return Materials;
+}
+
+TArray<TSharedPtr<IAvaMaskMaterialHandle>> FAvaMaskText3DActorMaterialCollectionHandle::GetMaterialHandles()
+{
+	if (MaterialHandles.IsEmpty() || !UE::Ava::Internal::IsHandleValid(MaterialHandles.Last()))
+	{
+		if (const UText3DComponent* TextComponent = GetComponent())
+		{
+			UAvaObjectHandleSubsystem* HandleSubsystem = UE::Ava::Internal::GetObjectHandleSubsystem();
+			if (!HandleSubsystem)
+			{
+				return MaterialHandles;
+			}
+			
+			MaterialHandles.Init(nullptr, static_cast<int32>(EText3DGroupType::TypeCount));
+			MaterialHandles[static_cast<int32>(EText3DGroupType::Front)] = HandleSubsystem->MakeHandle<IAvaMaskMaterialHandle>(TextComponent->GetFrontMaterial(), UE::AvaMask::Internal::HandleTag);
+			MaterialHandles[static_cast<int32>(EText3DGroupType::Bevel)] = HandleSubsystem->MakeHandle<IAvaMaskMaterialHandle>(TextComponent->GetBevelMaterial(), UE::AvaMask::Internal::HandleTag);
+			MaterialHandles[static_cast<int32>(EText3DGroupType::Extrude)] = HandleSubsystem->MakeHandle<IAvaMaskMaterialHandle>(TextComponent->GetExtrudeMaterial(), UE::AvaMask::Internal::HandleTag);
+			MaterialHandles[static_cast<int32>(EText3DGroupType::Back)] = HandleSubsystem->MakeHandle<IAvaMaskMaterialHandle>(TextComponent->GetBackMaterial(), UE::AvaMask::Internal::HandleTag);
+		}
+	}
+
+	return MaterialHandles;
+}
+
+void FAvaMaskText3DActorMaterialCollectionHandle::SetMaterial(
+	const FSoftComponentReference& InComponent
+	, const int32 InSlotIdx
+	, UMaterialInterface* InMaterial)
+{
+	if (!ensure(InSlotIdx < 4))
+	{
+		UE_LOG(LogAvaMask, Error, TEXT("SlotIdx out of range, should be from 0-3, was %u"), InSlotIdx);
+		return;
+	}
+
+	if (UText3DComponent* TextComponent = GetComponent())
+	{
+		switch (InSlotIdx)
+		{
+		case 0:
+			TextComponent->SetFrontMaterial(InMaterial);
+
+		case 1:
+			TextComponent->SetBevelMaterial(InMaterial);
+
+		case 2:
+			TextComponent->SetExtrudeMaterial(InMaterial);
+
+		case 3:
+			TextComponent->SetBackMaterial(InMaterial);
+
+		default:
+			return;
+		}
+	}
+}
+
+void FAvaMaskText3DActorMaterialCollectionHandle::SetMaterials(
+	const TArray<TObjectPtr<UMaterialInterface>>& InMaterials
+	, const TBitArray<>& InSetToggle)
+{
+	if (!ensure(InMaterials.Num() == GetNumMaterials()))
+	{
+		UE_LOG(LogAvaMask, Warning, TEXT("Expected 4 materials, got %u"), InMaterials.Num());
+		return;
+	}
+
+	if (!ensure(InSetToggle.Num() == GetNumMaterials()))
+	{
+		UE_LOG(LogAvaMask, Warning, TEXT("Expected 4 materials, got %u"), InSetToggle.Num());
+		return;
+	}
+
+	if (UText3DComponent* TextComponent = GetComponent())
+	{
+		if (InSetToggle[0])
+		{
+			TextComponent->SetFrontMaterial(InMaterials[0]);	
+		}
+
+		if (InSetToggle[1])
+		{
+			TextComponent->SetBevelMaterial(InMaterials[1]);	
+		}
+
+		if (InSetToggle[2])
+		{
+			TextComponent->SetExtrudeMaterial(InMaterials[2]);	
+		}
+
+		if (InSetToggle[3])
+		{
+			TextComponent->SetBackMaterial(InMaterials[3]);	
+		}
+	}
+}
+
+int32 FAvaMaskText3DActorMaterialCollectionHandle::GetNumMaterials() const
+{
+	constexpr int32 TextMaterialCount = static_cast<int32>(EText3DGroupType::TypeCount);
+	static_assert(TextMaterialCount == 4);
+	
+	return TextMaterialCount;
+}
+
+void FAvaMaskText3DActorMaterialCollectionHandle::ForEachMaterial(
+	TFunctionRef<bool(
+		const FSoftComponentReference& InComponent
+		, const int32 InIdx
+		, UMaterialInterface* InMaterial)> InFunction)
+{
+	if (const UText3DComponent* TextComponent = GetComponent())
+	{
+		const FSoftComponentReference ComponentReference = UE::AvaMask::Internal::MakeComponentReference(TextComponent->GetOwner(), TextComponent);
+		if (!InFunction(ComponentReference, static_cast<int32>(EText3DGroupType::Front), TextComponent->GetFrontMaterial()))
+		{
+			return;
+		}
+
+		if (!InFunction(ComponentReference, static_cast<int32>(EText3DGroupType::Bevel), TextComponent->GetBevelMaterial()))
+		{
+			return;
+		}
+
+		if (!InFunction(ComponentReference, static_cast<int32>(EText3DGroupType::Extrude), TextComponent->GetExtrudeMaterial()))
+		{
+			return;
+		}
+		
+		InFunction(ComponentReference, static_cast<int32>(EText3DGroupType::Back), TextComponent->GetBackMaterial());
+	}
+}
+
+void FAvaMaskText3DActorMaterialCollectionHandle::ForEachMaterialHandle(TFunctionRef<bool(const FSoftComponentReference& InComponent, const int32 InIdx, const bool bInIsSlotOccupied, const TSharedPtr<IAvaMaskMaterialHandle>& InMaterialHandle)> InFunction)
+{
+	if (const UText3DComponent* TextComponent = GetComponent())
+	{
+		// Effectively refreshes material handles if necessary
+		GetMaterialHandles();
+		
+		const FSoftComponentReference ComponentReference = UE::AvaMask::Internal::MakeComponentReference(TextComponent->GetOwner(), TextComponent);
+		TSharedPtr<IAvaMaskMaterialHandle> MaterialHandle = MaterialHandles[static_cast<int32>(EText3DGroupType::Front)];
+		if (!InFunction(ComponentReference, static_cast<int32>(EText3DGroupType::Front), MaterialHandle.IsValid(), MaterialHandle))
+		{
+			return;
+		}
+
+		MaterialHandle = MaterialHandles[static_cast<int32>(EText3DGroupType::Bevel)];
+		if (!InFunction(ComponentReference, static_cast<int32>(EText3DGroupType::Bevel), MaterialHandle.IsValid(), MaterialHandle))
+		{
+			return;
+		}
+
+		MaterialHandle = MaterialHandles[static_cast<int32>(EText3DGroupType::Extrude)];
+		if (!InFunction(ComponentReference, static_cast<int32>(EText3DGroupType::Extrude), MaterialHandle.IsValid(), MaterialHandle))
+		{
+			return;
+		}
+
+		MaterialHandle = MaterialHandles[static_cast<int32>(EText3DGroupType::Back)];
+		InFunction(ComponentReference, static_cast<int32>(EText3DGroupType::Back), MaterialHandle.IsValid(), MaterialHandle);
+	}
+}
+
+void FAvaMaskText3DActorMaterialCollectionHandle::MapEachMaterial(
+	TFunctionRef<UMaterialInterface*(
+		const FSoftComponentReference& InComponent
+		, const int32 InIdx
+		, UMaterialInterface* InMaterial)> InFunction)
+{	
+	if (const UText3DComponent* TextComponent = GetComponent())
+	{
+		TArray<TObjectPtr<UMaterialInterface>> MappedMaterials;
+		MappedMaterials.Init(nullptr, static_cast<TArray<TObjectPtr<UMaterialInterface>>::SizeType>(EText3DGroupType::TypeCount));
+
+		TBitArray<> SetFlags;
+        SetFlags.Init(false, MappedMaterials.Num());
+
+		const FSoftComponentReference ComponentReference = UE::AvaMask::Internal::MakeComponentReference(TextComponent->GetOwner(), TextComponent);
+		
+		MappedMaterials[0] = InFunction(ComponentReference, static_cast<int32>(EText3DGroupType::Front), TextComponent->GetFrontMaterial());
+		SetFlags[0] = MappedMaterials[0] != nullptr;
+
+		MappedMaterials[1] = InFunction(ComponentReference, static_cast<int32>(EText3DGroupType::Bevel), TextComponent->GetBevelMaterial());
+		SetFlags[1] = MappedMaterials[1] != nullptr;
+
+		MappedMaterials[2] = InFunction(ComponentReference, static_cast<int32>(EText3DGroupType::Extrude), TextComponent->GetExtrudeMaterial());
+		SetFlags[2] = MappedMaterials[2] != nullptr;
+
+		MappedMaterials[3] = InFunction(ComponentReference, static_cast<int32>(EText3DGroupType::Back), TextComponent->GetBackMaterial());
+		SetFlags[3] = MappedMaterials[3] != nullptr;
+
+		SetMaterials(MappedMaterials, SetFlags);
+	}
+}
+
+void FAvaMaskText3DActorMaterialCollectionHandle::MapEachMaterialHandle(TFunctionRef<TSharedPtr<IAvaMaskMaterialHandle>(const FSoftComponentReference& InComponent, const int32 InIdx, const bool bInIsSlotOccupied, const TSharedPtr<IAvaMaskMaterialHandle>& InMaterialHandle)> InFunction)
+{
+	if (const UText3DComponent* TextComponent = GetComponent())
+	{
+		// Effectively refreshes material handles if necessary
+		GetMaterialHandles();
+
+		TArray<TObjectPtr<UMaterialInterface>> MappedMaterials;
+		MappedMaterials.Reserve(GetNumMaterials());
+
+		TBitArray<> SetFlags;
+		SetFlags.Init(false, GetNumMaterials());
+
+		const FSoftComponentReference ComponentReference = UE::AvaMask::Internal::MakeComponentReference(TextComponent->GetOwner(), TextComponent);
+
+		auto TryMap = [this, InFunction, &ComponentReference, &MappedMaterials, &SetFlags](const EText3DGroupType InMaterialGroup)
+		{
+			int32 InMaterialGroupIdx = static_cast<int32>(InMaterialGroup);			
+			
+			TSharedPtr<IAvaMaskMaterialHandle>& MaterialHandle = MaterialHandles[InMaterialGroupIdx];
+			if (ensureAlways(UE::Ava::Internal::IsHandleValid(MaterialHandle)))
+			{
+				const TSharedPtr<IAvaMaskMaterialHandle> MappedMaterialHandle = InFunction(ComponentReference, InMaterialGroupIdx, true, MaterialHandle);
+				UMaterialInterface* RemappedMaterial = nullptr;
+				if (MappedMaterialHandle.IsValid())
+				{
+					RemappedMaterial = MappedMaterialHandle->GetMaterial();	
+				}
+				else
+				{
+					RemappedMaterial = MaterialHandle->GetMaterial();
+				}
+
+				MappedMaterials[InMaterialGroupIdx] = RemappedMaterial;
+				SetFlags[InMaterialGroupIdx] = RemappedMaterial != nullptr;
+			}	
+		};
+
+		TryMap(EText3DGroupType::Front);
+		TryMap(EText3DGroupType::Bevel);
+		TryMap(EText3DGroupType::Extrude);
+		TryMap(EText3DGroupType::Back);
+
+		SetMaterials(MappedMaterials, SetFlags);
+	}
+}
+
+bool FAvaMaskText3DActorMaterialCollectionHandle::IsValid() const
+{
+	return WeakActor.IsValid() && WeakComponent.IsValid();
+}
+
+bool FAvaMaskText3DActorMaterialCollectionHandle::IsSupported(
+	const UStruct* InStruct
+	, const TVariant<UObject*, FStructView>& InInstance
+	, FName InTag)
+{
+	if (InTag == UE::AvaMask::Internal::HandleTag
+		&& InStruct
+		&& ::IsValid(InStruct)
+		&& InStruct->IsChildOf<AActor>()
+		&& InInstance.TryGet<UObject*>() != nullptr)
+	{
+		if (const AActor* Actor = Cast<AActor>(InInstance.Get<UObject*>()))
+		{
+			return Actor->GetComponentByClass<UText3DComponent>() != nullptr;
+		}
+	}
+
+	return false;
+}
+
+FStructView FAvaMaskText3DActorMaterialCollectionHandle::GetMaterialHandleData(
+	FHandleData* InParentHandleData
+	, const FSoftComponentReference& InComponent
+	, const int32 InSlotIdx)
+{
+	if (InParentHandleData)
+	{
+		return InParentHandleData->GroupMaterialData[InSlotIdx];
+	}
+
+	return nullptr;
+}
+
+FStructView FAvaMaskText3DActorMaterialCollectionHandle::GetOrAddMaterialHandleData(
+	FHandleData* InParentHandleData
+	, const TSharedPtr<IAvaMaskMaterialHandle>& InMaterialHandle
+	, const FSoftComponentReference& InComponent
+	, const int32 InSlotIdx)
+{
+	if (InParentHandleData)
+	{
+		return InParentHandleData->GroupMaterialData.FindOrAdd(InSlotIdx, InMaterialHandle->MakeDataStruct());
+	}
+
+	return nullptr;
+}
+
+UText3DComponent* FAvaMaskText3DActorMaterialCollectionHandle::GetComponent()
+{
+	if (UText3DComponent* Component = WeakComponent.Get())
+	{
+		return Component;
+	}
+
+	if (const AActor* Actor = WeakActor.Get())
+	{
+		WeakComponent = Actor->GetComponentByClass<UText3DComponent>();
+		if (UText3DComponent* Component = WeakComponent.Get())
+		{
+			return Component;
+		}
+	}
+
+	return nullptr;
+}
+
+FAvaMaskAvaTextActorMaterialCollectionHandle::FAvaMaskAvaTextActorMaterialCollectionHandle(AActor* InActor)
+	: WeakActor(InActor)
+	, WeakComponent(InActor->GetComponentByClass<UAvaText3DComponent>())
+{
+	if (UAvaText3DComponent* TextComponent = GetComponent())
+	{
+		TextComponent->GetMaterialProviderDelegate().BindRaw(this, &FAvaMaskAvaTextActorMaterialCollectionHandle::ProvideMaterialWithSettings);
+	}
+}
+
+FAvaMaskAvaTextActorMaterialCollectionHandle::~FAvaMaskAvaTextActorMaterialCollectionHandle()
+{
+	if (UAvaText3DComponent* TextComponent = GetComponent())
+	{
+		TextComponent->GetMaterialProviderDelegate().Unbind();
+	}
+}
+
+TArray<TObjectPtr<UMaterialInterface>> FAvaMaskAvaTextActorMaterialCollectionHandle::GetMaterials()
+{
+	TArray<TObjectPtr<UMaterialInterface>> Materials;
+	Materials.Reserve(GetNumMaterials());
+	
+	if (const UAvaText3DComponent* TextComponent = GetComponent())
+	{
+		Materials.Emplace(TextComponent->GetFrontMaterial());
+		Materials.Emplace(TextComponent->GetBevelMaterial());
+		Materials.Emplace(TextComponent->GetExtrudeMaterial());
+		Materials.Emplace(TextComponent->GetBackMaterial());
+	}
+
+	return Materials;
+}
+
+TArray<TSharedPtr<IAvaMaskMaterialHandle>> FAvaMaskAvaTextActorMaterialCollectionHandle::GetMaterialHandles()
+{
+	if (MaterialHandles.IsEmpty() || !UE::Ava::Internal::IsHandleValid(MaterialHandles.Last()))
+	{
+		if (const UAvaText3DComponent* TextComponent = GetComponent())
+		{
+			MaterialHandles.Init(nullptr, static_cast<int32>(EText3DGroupType::TypeCount));
+
+			UAvaObjectHandleSubsystem* HandleSubsystem = UE::Ava::Internal::GetObjectHandleSubsystem();
+			if (!HandleSubsystem)
+			{
+				return MaterialHandles;
+			}
+
+			UMaterialInterface* FrontMaterial = TextComponent->GetFrontMaterial();
+			MaterialHandles[static_cast<int32>(EText3DGroupType::Front)] = HandleSubsystem->MakeHandle<IAvaMaskMaterialHandle>(FrontMaterial, UE::AvaMask::Internal::HandleTag);
+
+			UMaterialInterface* BevelMaterial = TextComponent->GetBevelMaterial();
+			MaterialHandles[static_cast<int32>(EText3DGroupType::Bevel)] = HandleSubsystem->MakeHandle<IAvaMaskMaterialHandle>(BevelMaterial, UE::AvaMask::Internal::HandleTag);
+
+			UMaterialInterface* ExtrudeMaterial = TextComponent->GetBevelMaterial();
+			MaterialHandles[static_cast<int32>(EText3DGroupType::Extrude)] = HandleSubsystem->MakeHandle<IAvaMaskMaterialHandle>(ExtrudeMaterial, UE::AvaMask::Internal::HandleTag);
+
+			UMaterialInterface* BackMaterial = TextComponent->GetBevelMaterial();
+			MaterialHandles[static_cast<int32>(EText3DGroupType::Back)] = HandleSubsystem->MakeHandle<IAvaMaskMaterialHandle>(BackMaterial, UE::AvaMask::Internal::HandleTag);
+		}
+	}
+
+	return MaterialHandles;
+}
+
+int32 FAvaMaskAvaTextActorMaterialCollectionHandle::GetNumMaterials() const
+{
+	constexpr int32 TextMaterialCount = static_cast<int32>(EText3DGroupType::TypeCount);
+	static_assert(TextMaterialCount == 4);
+	
+	return TextMaterialCount;
+}
+
+void FAvaMaskAvaTextActorMaterialCollectionHandle::SetMaterial(
+	const FSoftComponentReference& InComponent
+	, const int32 InSlotIdx
+	, UMaterialInterface* InMaterial)
+{
+	if (!ensure(InSlotIdx < 4))
+	{
+		UE_LOG(LogAvaMask, Error, TEXT("SlotIdx out of range, should be from 0-3, was %u"), InSlotIdx);
+		return;
+	}
+
+	if (UAvaText3DComponent* TextComponent = GetComponent())
+    {
+		switch (InSlotIdx)
+		{
+		case 0:
+			TextComponent->SetFrontMaterial(InMaterial);
+
+		case 1:
+			TextComponent->SetBevelMaterial(InMaterial);
+
+		case 2:
+			TextComponent->SetExtrudeMaterial(InMaterial);
+
+		case 3:
+			TextComponent->SetBackMaterial(InMaterial);
+
+		default:
+			return;
+		}
+    }
+}
+
+void FAvaMaskAvaTextActorMaterialCollectionHandle::SetMaterials(
+	const TArray<TObjectPtr<UMaterialInterface>>& InMaterials
+	, const TBitArray<>& InSetToggle)
+{
+	if (!ensure(InMaterials.Num() == GetNumMaterials()))
+	{
+		UE_LOG(LogAvaMask, Warning, TEXT("Expected 4 materials, got %u"), InMaterials.Num());
+		return;
+	}
+
+	if (!ensure(InSetToggle.Num() == GetNumMaterials()))
+	{
+		UE_LOG(LogAvaMask, Warning, TEXT("Expected 4 materials, got %u"), InSetToggle.Num());
+		return;
+	}
+	
+	if (UAvaText3DComponent* TextComponent = GetComponent())
+	{
+		if (InSetToggle[0])
+		{
+			TextComponent->SetFrontMaterial(InMaterials[0]);	
+		}
+
+		if (InSetToggle[1])
+		{
+			TextComponent->SetBevelMaterial(InMaterials[1]);	
+		}
+
+		if (InSetToggle[2])
+		{
+			TextComponent->SetExtrudeMaterial(InMaterials[2]);	
+		}
+
+		if (InSetToggle[3])
+		{
+			TextComponent->SetBackMaterial(InMaterials[3]);	
+		}
+	}
+}
+
+void FAvaMaskAvaTextActorMaterialCollectionHandle::ForEachMaterial(
+	TFunctionRef<bool(
+		const FSoftComponentReference& InComponent
+		, const int32 InIdx
+		, UMaterialInterface* InMaterial)> InFunction)
+{
+	if (const UAvaText3DComponent* TextComponent = GetComponent())
+	{
+		const FSoftComponentReference ComponentReference = UE::AvaMask::Internal::MakeComponentReference(TextComponent->GetOwner(), TextComponent);
+		if (!InFunction(ComponentReference, static_cast<int32>(EText3DGroupType::Front), TextComponent->GetFrontMaterial()))
+		{
+			return;
+		}
+		
+		if (!InFunction(ComponentReference, static_cast<int32>(EText3DGroupType::Bevel), TextComponent->GetBevelMaterial()))
+		{
+			return;
+		}
+		
+		if (!InFunction(ComponentReference, static_cast<int32>(EText3DGroupType::Extrude), TextComponent->GetExtrudeMaterial()))
+		{
+			return;
+		}
+		
+		InFunction(ComponentReference, static_cast<int32>(EText3DGroupType::Back), TextComponent->GetBackMaterial());
+	}
+}
+
+void FAvaMaskAvaTextActorMaterialCollectionHandle::ForEachMaterialHandle(TFunctionRef<bool(const FSoftComponentReference& InComponent, const int32 InIdx, const bool bInIsSlotOccupied, const TSharedPtr<IAvaMaskMaterialHandle>& InMaterialHandle)> InFunction)
+{
+	if (const UAvaText3DComponent* TextComponent = GetComponent())
+	{
+		// Effectively refreshes material handles if necessary
+		GetMaterialHandles();
+		
+		const FSoftComponentReference ComponentReference = UE::AvaMask::Internal::MakeComponentReference(TextComponent->GetOwner(), TextComponent);
+		TSharedPtr<IAvaMaskMaterialHandle>& MaterialHandle = MaterialHandles[static_cast<int32>(EText3DGroupType::Front)];
+		if (ensureAlways(UE::Ava::Internal::IsHandleValid(MaterialHandle)))
+		{
+			if (!InFunction(ComponentReference, static_cast<int32>(EText3DGroupType::Front), true, MaterialHandle))
+			{
+				return;
+			}
+		}
+
+		MaterialHandle = MaterialHandles[static_cast<int32>(EText3DGroupType::Bevel)];
+		if (ensureAlways(UE::Ava::Internal::IsHandleValid(MaterialHandle)))
+		{
+			if (!InFunction(ComponentReference, static_cast<int32>(EText3DGroupType::Bevel), true, MaterialHandle))
+			{
+				return;
+			}
+		}
+
+		MaterialHandle = MaterialHandles[static_cast<int32>(EText3DGroupType::Extrude)];
+		if (ensureAlways(UE::Ava::Internal::IsHandleValid(MaterialHandle)))
+		{
+			if (!InFunction(ComponentReference, static_cast<int32>(EText3DGroupType::Extrude), true, MaterialHandle))
+			{
+				return;
+			}
+		}
+
+		MaterialHandle = MaterialHandles[static_cast<int32>(EText3DGroupType::Back)];
+		if (ensureAlways(UE::Ava::Internal::IsHandleValid(MaterialHandle)))
+		{
+			InFunction(ComponentReference, static_cast<int32>(EText3DGroupType::Back), true, MaterialHandle);
+		}
+	}
+}
+
+void FAvaMaskAvaTextActorMaterialCollectionHandle::MapEachMaterial(
+	TFunctionRef<UMaterialInterface*(
+		const FSoftComponentReference& InComponent
+		, const int32 InIdx
+		, UMaterialInterface* InMaterial)> InFunction)
+{
+	if (const UAvaText3DComponent* TextComponent = GetComponent())
+	{
+		TArray<TObjectPtr<UMaterialInterface>> MappedMaterials;
+		MappedMaterials.Init(nullptr, static_cast<int32>(EText3DGroupType::TypeCount));
+
+		TBitArray<> SetFlags;
+		SetFlags.Init(false, MappedMaterials.Num());
+
+		const FSoftComponentReference ComponentReference = UE::AvaMask::Internal::MakeComponentReference(TextComponent->GetOwner(), TextComponent);
+		
+		MappedMaterials[0] = InFunction(ComponentReference, static_cast<int32>(EText3DGroupType::Front), TextComponent->GetFrontMaterial());
+		SetFlags[0] = MappedMaterials[0] != nullptr;
+		
+		MappedMaterials[1] = InFunction(ComponentReference, static_cast<int32>(EText3DGroupType::Bevel), TextComponent->GetBevelMaterial());
+		SetFlags[1] = MappedMaterials[1] != nullptr;
+		
+		MappedMaterials[2] = InFunction(ComponentReference, static_cast<int32>(EText3DGroupType::Extrude), TextComponent->GetExtrudeMaterial());
+		SetFlags[2] = MappedMaterials[2] != nullptr;
+		
+		MappedMaterials[3] = InFunction(ComponentReference, static_cast<int32>(EText3DGroupType::Back), TextComponent->GetBackMaterial());
+		SetFlags[3] = MappedMaterials[3] != nullptr;
+		
+		SetMaterials(MappedMaterials, SetFlags);
+	}
+}
+
+void FAvaMaskAvaTextActorMaterialCollectionHandle::MapEachMaterialHandle(TFunctionRef<TSharedPtr<IAvaMaskMaterialHandle>(const FSoftComponentReference& InComponent, const int32 InIdx, const bool bInIsSlotOccupied, const TSharedPtr<IAvaMaskMaterialHandle>& InMaterialHandle)> InFunction)
+{
+	if (const UAvaText3DComponent* TextComponent = GetComponent())
+	{		
+		// Effectively refreshes material handles if necessary
+		GetMaterialHandles();
+
+		TArray<TObjectPtr<UMaterialInterface>> MappedMaterials;
+		MappedMaterials.Reserve(GetNumMaterials());
+
+		TBitArray<> SetFlags;
+		SetFlags.Init(false, GetNumMaterials());
+
+		const FSoftComponentReference ComponentReference = UE::AvaMask::Internal::MakeComponentReference(TextComponent->GetOwner(), TextComponent);
+
+		auto TryMap = [this, InFunction, &ComponentReference, &MappedMaterials, &SetFlags](const EText3DGroupType InMaterialGroup)
+		{
+			int32 InMaterialGroupIdx = static_cast<int32>(InMaterialGroup);			
+			
+			TSharedPtr<IAvaMaskMaterialHandle>& MaterialHandle = MaterialHandles[InMaterialGroupIdx];
+			if (ensureAlways(UE::Ava::Internal::IsHandleValid(MaterialHandle)))
+			{
+				const TSharedPtr<IAvaMaskMaterialHandle> MappedMaterialHandle = InFunction(ComponentReference, InMaterialGroupIdx, true, MaterialHandle);
+				UMaterialInterface* RemappedMaterial = nullptr;
+				if (MappedMaterialHandle.IsValid())
+				{
+					RemappedMaterial = MappedMaterialHandle->GetMaterial();	
+				}
+				else
+				{
+					RemappedMaterial = MaterialHandle->GetMaterial();
+				}
+
+				MappedMaterials[InMaterialGroupIdx] = RemappedMaterial;
+				SetFlags[InMaterialGroupIdx] = RemappedMaterial != nullptr;
+			}	
+		};
+
+		TryMap(EText3DGroupType::Front);
+		TryMap(EText3DGroupType::Bevel);
+		TryMap(EText3DGroupType::Extrude);
+		TryMap(EText3DGroupType::Back);
+
+		SetMaterials(MappedMaterials, SetFlags);
+	}
+}
+
+bool FAvaMaskAvaTextActorMaterialCollectionHandle::SaveOriginalState(const FStructView& InHandleData)
+{
+	if (FHandleData* HandleData = InHandleData.GetPtr<FHandleData>())
+	{
+		ForEachMaterialHandle([this, HandleData](
+			const FSoftComponentReference& InComponent
+			, const int32 InSlotIdx
+			, const bool bInIsSlotOccupied
+			, const TSharedPtr<IAvaMaskMaterialHandle>& InMaterialHandle)
+		{
+			FInstancedStruct& MaterialHandleData = HandleData->GroupMaterialData.FindOrAdd(InSlotIdx, InMaterialHandle->MakeDataStruct());
+			return InMaterialHandle->SaveOriginalState(MaterialHandleData);
+		});
+
+		if (const UAvaText3DComponent* TextComponent = GetComponent())
+		{
+			HandleData->OriginalTranslucencyStyle = TextComponent->GetTranslucencyStyle();
+			return true;
+		}
+		
+		return true;
+	}
+
+	return false;
+}
+
+bool FAvaMaskAvaTextActorMaterialCollectionHandle::ApplyOriginalState(const FStructView& InHandleData)
+{
+	if (const FHandleData* HandleData = InHandleData.GetPtr<FHandleData>())
+	{
+		if (UAvaText3DComponent* TextComponent = GetComponent())
+		{
+			// Only restore if the original was NOT translucent AND the current alpha is 1.0
+			if (HandleData->OriginalTranslucencyStyle == EAvaTextTranslucency::None
+				&& FMath::IsNearlyEqual(TextComponent->GetOpacity(), 1.0f))
+			{
+				TextComponent->SetTranslucencyStyle(HandleData->OriginalTranslucencyStyle);	
+			}
+			TextComponent->RefreshMaterialInstances();
+			
+			Super::ApplyOriginalState(InHandleData);
+			
+			return true;
+		}
+
+		return false;
+	}
+
+	return false;
+}
+
+bool FAvaMaskAvaTextActorMaterialCollectionHandle::ApplyModifiedState(
+	const FAvaMask2DSubjectParameters& InModifiedParameters
+	, const FStructView& InHandleData)
+{
+	if (InHandleData.GetPtr<FHandleData>())
+	{
+		if (!Super::ApplyModifiedState(InModifiedParameters, InHandleData))
+		{
+			return false;
+		}
+
+		if (UAvaText3DComponent* TextComponent = GetComponent())
+		{
+			const EAvaTextTranslucency TargetTranslucencyStyle = UE::AvaMask::Private::GetTargetTranslucencyType(TextComponent->GetTranslucencyStyle(), InModifiedParameters.MaterialParameters.BlendMode);
+			TextComponent->SetTranslucencyStyle(TargetTranslucencyStyle);
+			TextComponent->RefreshMaterialInstances();
+			return true;
+		}
+
+		return true;
+	}
+
+	return false;
+}
+
+bool FAvaMaskAvaTextActorMaterialCollectionHandle::IsValid() const
+{
+	return WeakActor.IsValid() && WeakComponent.IsValid();
+}
+
+bool FAvaMaskAvaTextActorMaterialCollectionHandle::IsSupported(
+	const UStruct* InStruct
+	, const TVariant<UObject*, FStructView>& InInstance
+	, FName InTag)
+{
+	if (InTag == UE::AvaMask::Internal::HandleTag
+		&& InStruct
+		&& ::IsValid(InStruct)
+		&& InStruct->IsChildOf<AActor>()
+		&& InInstance.TryGet<UObject*>() != nullptr)
+	{
+		if (const AActor* Actor = Cast<AActor>(InInstance.Get<UObject*>()))
+		{
+			return Actor->GetComponentByClass<UAvaText3DComponent>() != nullptr;
+		}
+	}
+
+	return false;
+}
+
+FStructView FAvaMaskAvaTextActorMaterialCollectionHandle::GetMaterialHandleData(
+	FHandleData* InParentHandleData
+	, const FSoftComponentReference& InComponent
+	, const int32 InSlotIdx)
+{
+	if (InParentHandleData)
+	{
+		return InParentHandleData->GroupMaterialData[InSlotIdx];
+	}
+
+	return nullptr;	
+}
+
+FStructView FAvaMaskAvaTextActorMaterialCollectionHandle::GetOrAddMaterialHandleData(
+	FHandleData* InParentHandleData
+	, const TSharedPtr<IAvaMaskMaterialHandle>& InMaterialHandle
+	, const FSoftComponentReference& InComponent
+	, const int32 InSlotIdx)
+{
+	if (InParentHandleData)
+	{
+		return InParentHandleData->GroupMaterialData.FindOrAdd(InSlotIdx, InMaterialHandle->MakeDataStruct());
+	}
+
+	return nullptr;
+}
+
+UAvaText3DComponent* FAvaMaskAvaTextActorMaterialCollectionHandle::GetComponent()
+{
+	if (UAvaText3DComponent* Component = WeakComponent.Get())
+	{
+		return Component;
+	}
+
+	if (const AActor* Actor = WeakActor.Get())
+	{
+		WeakComponent = Actor->GetComponentByClass<UAvaText3DComponent>();
+		if (UAvaText3DComponent* Component = WeakComponent.Get())
+		{
+			return Component;
+		}
+	}
+
+	return nullptr;
+}
+
+UMaterialInterface* FAvaMaskAvaTextActorMaterialCollectionHandle::ProvideMaterialWithSettings(const UMaterialInterface* InPreviousMaterial, const FAvaTextMaterialSettings& InMaterialSettings)
+{
+	UMaterialInterface* AttemptedAssignedMaterial = UAvaTextMaterialHub::GetMaterial(InMaterialSettings);
+
+	if (OnSourceMaterialPreAssignment().IsBound())
+	{
+		UMaterialInterface* UserAssignedMaterial = OnSourceMaterialPreAssignment().Execute(InPreviousMaterial, AttemptedAssignedMaterial);
+
+		// Flag handles for refresh
+		MaterialHandles.Reset();
+		GetMaterialHandles();
+
+		return UserAssignedMaterial;
+	}
+
+	return nullptr;
+}
diff --git a/Engine/Plugins/Experimental/Avalanche/Source/AvalancheMask/Private/Handling/AvaObjectHandleSubsystem.cpp b/Engine/Plugins/Experimental/Avalanche/Source/AvalancheMask/Private/Handling/AvaObjectHandleSubsystem.cpp
new file mode 100644
--- /dev/null
+++ b/Engine/Plugins/Experimental/Avalanche/Source/AvalancheMask/Private/Handling/AvaObjectHandleSubsystem.cpp
@@ -0,0 +1,151 @@
+﻿// Copyright Epic Games, Inc. All Rights Reserved.
+
+#include "Handling/AvaObjectHandleSubsystem.h"
+
+#include "AvaHandleUtilities.h"
+#include "AvaMaskActorMaterialCollectionHandle.h"
+#include "AvaMaskAvaShapeMaterialCollectionHandle.h"
+#include "AvaMaskLog.h"
+#include "AvaMaskText3DActorMaterialCollectionHandle.h"
+#include "Engine/Engine.h"
+#include "GameFramework/Actor.h"
+#include "Handling/AvaMaskDesignedMaterialHandle.h"
+#include "Handling/AvaMaskMaterialInstanceHandle.h"
+#include "Handling/AvaMaskParametricMaterialHandle.h"
+#include "IAvaMaskMaterialCollectionHandle.h"
+#include "Material/DynamicMaterialInstance.h"
+#include "Materials/MaterialInterface.h"
+#include "UObject/UObjectIterator.h"
+
+void UAvaObjectHandleSubsystem::Initialize(FSubsystemCollectionBase& Collection)
+{
+	Super::Initialize(Collection);
+
+	FindObjectHandleFactories();
+}
+
+TSharedPtr<IAvaObjectHandle> UAvaObjectHandleSubsystem::MakeHandle(
+	const UStruct* InStruct
+	, const TVariant<UObject*, FStructView>& InInstance
+	, FName InTag)
+{
+	for (const TPair<FIsSupportedFunction, FMakeHandleFunction>& SupportedFactoryPair : ObjectHandleFactories)
+	{
+		if (SupportedFactoryPair.Key(InStruct, InInstance, InTag))
+		{
+			TSharedPtr<IAvaObjectHandle> Handle = SupportedFactoryPair.Value(InInstance);
+			if (!UE::Ava::Internal::IsHandleValid(Handle))
+			{
+				const FString ObjectName = GetInstanceName(InStruct, InInstance);
+				if (InTag.IsNone())
+				{
+					UE_LOG(LogAvaMask, Warning, TEXT("Object Handle for '%s' was created but invalid."), *ObjectName)
+				}
+				else
+				{
+					UE_LOG(LogAvaMask, Warning, TEXT("Object Handle for '%s' (with tag '%s') was created but invalid."), *ObjectName, *InTag.ToString())
+				}
+			}
+			return Handle;
+		}
+	}
+
+	UE_LOG(LogAvaMask, Display, TEXT("No ObjectHandle found for '%s'"), *InStruct->GetName());
+	
+	return nullptr;
+}
+
+// @note: registration order matters!
+void UAvaObjectHandleSubsystem::FindObjectHandleFactories()
+{
+	// Material Collection Handles
+	{
+		ObjectHandleFactories.Add({
+			FAvaMaskAvaTextActorMaterialCollectionHandle::IsSupported
+			, [](const TVariant<UObject*, FStructView>& InObject) -> TSharedPtr<IAvaMaskMaterialCollectionHandle>
+			{
+				return MakeShared<FAvaMaskAvaTextActorMaterialCollectionHandle>(Cast<AActor>(InObject.Get<UObject*>()));
+			}});
+	
+		ObjectHandleFactories.Add({
+			FAvaMaskAvaShapeMaterialCollectionHandle::IsSupported
+			, [](const TVariant<UObject*, FStructView>& InObject) -> TSharedPtr<IAvaMaskMaterialCollectionHandle>
+			{
+				return MakeShared<FAvaMaskAvaShapeMaterialCollectionHandle>(Cast<AActor>(InObject.Get<UObject*>()));
+			}});
+	
+		ObjectHandleFactories.Add({
+			FAvaMaskText3DActorMaterialCollectionHandle::IsSupported
+			, [](const TVariant<UObject*, FStructView>& InObject) -> TSharedPtr<IAvaMaskMaterialCollectionHandle>
+			{
+				return MakeShared<FAvaMaskText3DActorMaterialCollectionHandle>(Cast<AActor>(InObject.Get<UObject*>()));
+			}});
+	
+		ObjectHandleFactories.Add({
+			FAvaMaskActorMaterialCollectionHandle::IsSupported
+			, [](const TVariant<UObject*, FStructView>& InObject) -> TSharedPtr<IAvaMaskMaterialCollectionHandle>
+			{
+				return MakeShared<FAvaMaskActorMaterialCollectionHandle>(Cast<AActor>(InObject.Get<UObject*>()));
+			}});
+	}
+
+	// Material Handles
+	{
+		ObjectHandleFactories.Add({
+			FAvaMaskDesignedMaterialHandle::IsSupported
+			, [](const TVariant<UObject*, FStructView>& InMaterial) -> TSharedPtr<IAvaMaskMaterialHandle>
+			{
+				return MakeShared<FAvaMaskDesignedMaterialHandle>(Cast<UDynamicMaterialInstance>(InMaterial.Get<UObject*>()));
+			}});
+
+		ObjectHandleFactories.Add({
+			FAvaMaskParametricMaterialHandle::IsSupported
+			, [](const TVariant<UObject*, FStructView>& InMaterial) -> TSharedPtr<IAvaMaskMaterialHandle>
+			{
+				return MakeShared<FAvaMaskParametricMaterialHandle>(InMaterial.Get<FStructView>());
+			}});
+	
+		ObjectHandleFactories.Add({
+			FAvaMaskMaterialInstanceHandle::IsSupported
+			, [](const TVariant<UObject*, FStructView>& InMaterial) -> TSharedPtr<IAvaMaskMaterialHandle>
+			{
+				return MakeShared<FAvaMaskMaterialInstanceHandle>(Cast<UMaterialInterface>(InMaterial.Get<UObject*>()));
+			}});
+	}
+}
+
+FString UAvaObjectHandleSubsystem::GetInstanceName(const UStruct* InStruct, const TVariant<UObject*, FStructView>& InInstance)
+{
+	struct FNameGetter
+	{
+		FString operator()(const UObject* InValue) const
+		{
+			if (::IsValid(InValue))
+			{
+				return InValue->GetName();
+			}
+
+			return TEXT("");
+		}
+		
+		FString operator()(const FStructView& InValue) const
+		{
+			return InValue.GetScriptStruct()->GetName();
+		}
+	};
+
+	FNameGetter NameGetter;
+
+	if (FString Name = Visit(NameGetter, InInstance);
+		!Name.IsEmpty())
+	{
+		return Name;
+	}
+
+	if (const UStruct* Struct = InStruct)
+	{
+		return Struct->GetName();
+	}
+
+	return TEXT("");
+}
diff --git a/Engine/Plugins/Experimental/Avalanche/Source/AvalancheMask/Private/Handling/AvaObjectHandleSubsystem.h b/Engine/Plugins/Experimental/Avalanche/Source/AvalancheMask/Private/Handling/AvaObjectHandleSubsystem.h
new file mode 100644
--- /dev/null
+++ b/Engine/Plugins/Experimental/Avalanche/Source/AvalancheMask/Private/Handling/AvaObjectHandleSubsystem.h
@@ -0,0 +1,77 @@
+﻿// Copyright Epic Games, Inc. All Rights Reserved.
+
+#pragma once
+
+#include "AvaMaskLog.h"
+#include "IAvaObjectHandle.h"
+#include "Misc/TVariant.h"
+#include "StructView.h"
+#include "Subsystems/EngineSubsystem.h"
+
+#include "AvaObjectHandleSubsystem.generated.h"
+
+class UAvaObjectHandleSubsystem;
+
+/** Responsible for providing Handlers for a given UObject. */
+UCLASS()
+class UAvaObjectHandleSubsystem
+	: public UEngineSubsystem
+{
+	GENERATED_BODY()
+
+public:
+	// ~Begin USubsystem
+	virtual void Initialize(FSubsystemCollectionBase& Collection) override;
+	// ~End USubsystem
+
+	TSharedPtr<IAvaObjectHandle> MakeHandle(const UStruct* InStruct, const TVariant<UObject*, FStructView>& InInstance, FName InTag = NAME_None);
+
+	template <typename HandleType, typename ObjectType = UObject
+		UE_REQUIRES(!std::is_const_v<HandleType> && std::is_base_of_v<IAvaObjectHandle, HandleType>
+				&&	!std::is_const_v<ObjectType> && TModels_V<CStaticClassProvider, std::decay_t<ObjectType>>)>
+    TSharedPtr<HandleType> MakeHandle(ObjectType* InInstance, FName InTag = NAME_None)
+    {
+		if (!InInstance)
+		{
+			UE_LOG(LogAvaMask, Warning, TEXT("Invalid or null object provided to MakeHandle"));
+			return nullptr;
+		}
+
+		TVariant<UObject*, FStructView> Variant(TInPlaceType<UObject*>(), InInstance);
+		
+    	return StaticCastSharedPtr<HandleType>(MakeHandle(InInstance->GetClass(), MoveTemp(Variant), InTag));
+    }
+
+	template <typename HandleType, typename StructType
+		UE_REQUIRES(!std::is_const_v<HandleType> && std::is_base_of_v<IAvaObjectHandle, HandleType>
+				&&	!std::is_const_v<StructType> && TModels_V<CStaticStructProvider, std::decay_t<StructType>>)>
+	TSharedPtr<HandleType> MakeHandle(const FStructView& InInstance, FName InTag = NAME_None)
+	{
+		if (!InInstance.IsValid())
+		{
+			UE_LOG(LogAvaMask, Warning, TEXT("Invalid or null StructView provided to MakeHandle"));
+			return nullptr;
+		}
+
+		if (!ensure(InInstance.GetPtr<StructType>()))
+		{
+			UE_LOG(LogAvaMask, Warning, TEXT("StructView does not point to the expected struct type \'%s\'"), *StructType::StaticStruct()->GetName());
+			return nullptr;			
+		}
+		
+		TVariant<UObject*, FStructView> Variant(TInPlaceType<FStructView>(), InInstance);
+		
+		return StaticCastSharedPtr<HandleType>(MakeHandle(InInstance.GetScriptStruct(), MoveTemp(Variant), InTag));
+	}
+ 
+private:
+	void FindObjectHandleFactories();
+
+	FString GetInstanceName(const UStruct* InStruct, const TVariant<UObject*, FStructView>& InInstance);
+
+private:
+	using FIsSupportedFunction = TFunction<bool(const UStruct*, const TVariant<UObject*, FStructView>&, FName)>;
+	using FMakeHandleFunction = TFunction<TSharedPtr<IAvaObjectHandle>(const TVariant<UObject*, FStructView>&)>;
+
+	TArray<TPair<FIsSupportedFunction, FMakeHandleFunction>> ObjectHandleFactories;
+};
diff --git a/Engine/Plugins/Experimental/Avalanche/Source/AvalancheMask/Private/Handling/AvaParametricMaterialHandle.cpp b/Engine/Plugins/Experimental/Avalanche/Source/AvalancheMask/Private/Handling/AvaParametricMaterialHandle.cpp
new file mode 100644
--- /dev/null
+++ b/Engine/Plugins/Experimental/Avalanche/Source/AvalancheMask/Private/Handling/AvaParametricMaterialHandle.cpp
@@ -0,0 +1,57 @@
+﻿// Copyright Epic Games, Inc. All Rights Reserved.
+
+#include "Handling/AvaParametricMaterialHandle.h"
+
+#include "AvaShapeParametricMaterial.h"
+
+FAvaParametricMaterialHandle::FAvaParametricMaterialHandle(const FStructView& InParametricMaterial)
+	: FAvaMaterialInstanceHandle(InParametricMaterial.GetPtr<FAvaShapeParametricMaterial>()->GetMaterial())
+	, ParametricMaterial(InParametricMaterial)
+{
+}
+
+void FAvaParametricMaterialHandle::CopyParametersFrom(UMaterialInstance* InSourceMaterial)
+{
+	Super::CopyParametersFrom(InSourceMaterial);
+	
+	if (FAvaShapeParametricMaterial* ParametricMtl = GetParametricMaterial())
+	{
+		ParametricMtl->CopyFromMaterialParameters(InSourceMaterial);
+	}
+}
+
+bool FAvaParametricMaterialHandle::IsValid() const
+{
+	return this && ParametricMaterial.IsValid() && WeakMaterialInstance.IsValid();
+}
+
+bool FAvaParametricMaterialHandle::IsSupported(
+	const UStruct* InStruct
+	, const TVariant<UObject*, FStructView>& InInstance
+	, FName InTag)
+{
+	return InStruct
+		&& ::IsValid(InStruct)
+		&& InStruct->IsChildOf(FAvaShapeParametricMaterial::StaticStruct())
+		&& InInstance.TryGet<FStructView>() != nullptr;
+}
+
+FAvaShapeParametricMaterial* FAvaParametricMaterialHandle::GetParametricMaterial()
+{
+	if (ParametricMaterial.IsValid())
+	{
+		return ParametricMaterial.GetPtr<FAvaShapeParametricMaterial>();
+	}
+
+	return nullptr;
+}
+
+UMaterialInstanceDynamic* FAvaParametricMaterialHandle::GetMaterialInstance()
+{
+	if (const FAvaShapeParametricMaterial* ParametricMtl = GetParametricMaterial())
+	{
+		return ParametricMtl->GetMaterial();
+	}
+	
+	return FAvaMaterialInstanceHandle::GetMaterialInstance();
+}
diff --git a/Engine/Plugins/Experimental/Avalanche/Source/AvalancheMask/Private/Handling/AvaParametricMaterialHandle.h b/Engine/Plugins/Experimental/Avalanche/Source/AvalancheMask/Private/Handling/AvaParametricMaterialHandle.h
new file mode 100644
--- /dev/null
+++ b/Engine/Plugins/Experimental/Avalanche/Source/AvalancheMask/Private/Handling/AvaParametricMaterialHandle.h
@@ -0,0 +1,36 @@
+﻿// Copyright Epic Games, Inc. All Rights Reserved.
+
+#pragma once
+
+#include "AvaType.h"
+#include "Handling/AvaMaterialInstanceHandle.h"
+#include "IAvaMaskMaterialHandle.h"
+#include "IAvaObjectHandle.h"
+
+struct FAvaShapeParametricMaterial;
+
+class FAvaParametricMaterialHandle
+	: public FAvaMaterialInstanceHandle
+{
+public:
+	UE_AVA_INHERITS_WITH_SUPER(FAvaParametricMaterialHandle, FAvaMaterialInstanceHandle);
+	
+	explicit FAvaParametricMaterialHandle(const FStructView& InParametricMaterial);
+
+	// ~Begin IAvaMaterialHandle
+	virtual void CopyParametersFrom(UMaterialInstance* InSourceMaterial) override;
+	// ~End IAvaMaterialHandle
+
+	// ~Begin IAvaObjectHandle
+	virtual bool IsValid() const override;
+	// ~End IAvaObjectHandle
+
+	static bool IsSupported(const UStruct* InStruct, const TVariant<UObject*, FStructView>& InInstance, FName InTag = NAME_None);
+
+protected:
+	FAvaShapeParametricMaterial* GetParametricMaterial();
+	virtual UMaterialInstanceDynamic* GetMaterialInstance() override;
+
+protected:
+	FStructView ParametricMaterial = nullptr;
+};
diff --git a/Engine/Plugins/Experimental/Avalanche/Source/AvalancheMask/Private/Handling/IAvaMaskMaterialHandle.h b/Engine/Plugins/Experimental/Avalanche/Source/AvalancheMask/Private/Handling/IAvaMaskMaterialHandle.h
new file mode 100644
--- /dev/null
+++ b/Engine/Plugins/Experimental/Avalanche/Source/AvalancheMask/Private/Handling/IAvaMaskMaterialHandle.h
@@ -0,0 +1,163 @@
+﻿// Copyright Epic Games, Inc. All Rights Reserved.
+
+#pragma once
+
+#include "AvaMaskLog.h"
+#include "AvaMaskTypes.h"
+#include "AvaMaskUtilities.h"
+#include "Engine/Texture.h"
+#include "IAvaMaterialHandle.h"
+#include "InstancedStruct.h"
+#include "Materials/AvaMaterialUtils.h"
+#include "Materials/MaterialInstance.h"
+#include "Materials/MaterialInstanceDynamic.h"
+#include "StructView.h"
+
+class UAvaMaskMaterialInstanceSubsystem;
+
+namespace UE::AvaMask::Internal
+{
+	UAvaMaskMaterialInstanceSubsystem* GetMaterialInstanceSubsystem();
+}
+
+/** Interface between the mask system and an individual material. */
+class IAvaMaskMaterialHandle
+	: public IAvaMaterialHandle
+{
+public:
+	UE_AVA_INHERITS(IAvaMaskMaterialHandle, IAvaMaterialHandle);
+
+	virtual ~IAvaMaskMaterialHandle() override = default;
+
+	virtual FInstancedStruct MakeDataStruct() = 0;
+	virtual const UScriptStruct* GetDataStructType() = 0;
+
+	virtual bool GetMaskParameters(FAvaMask2DMaterialParameters& OutParameters) = 0;
+	virtual bool SetMaskParameters(UTexture* InTexture, float InBaseOpacity, EGeometryMaskColorChannel InChannel, bool bInInverted = false, const FVector2f& InPadding = FVector2f::Zero(), bool bInApplyFeathering = false, float InOuterFeathering = 0.0f, float InInnerFeathering = 0.0f) = 0;
+
+	virtual const EBlendMode GetBlendMode() = 0;
+	virtual void SetBlendMode(const EBlendMode InBlendMode) = 0;
+
+	/** Provides the opportunity to propagate changes to a derived material. */
+	virtual void SetChildMaterial(UMaterialInstanceDynamic* InChildMaterial) = 0;
+
+	virtual bool SaveOriginalState(const FStructView& InHandleData) = 0;
+	
+	virtual bool ApplyOriginalState(
+		const FStructView& InHandleData
+		, TUniqueFunction<void(UMaterialInterface*)>&& InMaterialSetter) = 0;
+	
+	virtual bool ApplyModifiedState(
+		const FAvaMask2DSubjectParameters& InModifiedParameters
+		, const FStructView& InHandleData
+		, TUniqueFunction<void(UMaterialInterface*)>&& InMaterialSetter) = 0;
+
+	virtual bool HasRequiredParameters(TArray<FString>& OutMissingParameterNames) = 0;	
+};
+
+template <typename HandleDataType>
+class TAvaMaskMaterialHandle
+	: public IAvaMaskMaterialHandle
+{
+	static_assert(TModels_V<CStaticStructProvider, HandleDataType>, "HandleDataType should be a UStruct");
+
+public:
+	UE_AVA_INHERITS_WITH_SUPER(TAvaMaskMaterialHandle, IAvaMaskMaterialHandle);
+	
+	using FHandleData = HandleDataType;
+
+public:
+	virtual FInstancedStruct MakeDataStruct() override;
+	virtual const UScriptStruct* GetDataStructType() override;
+
+	virtual bool GetMaskParameters(FAvaMask2DMaterialParameters& OutParameters) override;
+	virtual bool SetMaskParameters(UTexture* InTexture, float InBaseOpacity, EGeometryMaskColorChannel InChannel, bool bInInverted = false, const FVector2f& InPadding = FVector2f::Zero(), bool bInApplyFeathering = false, float InOuterFeathering = 0.0f, float InInnerFeathering = 0.0f) override;
+
+	virtual bool HasRequiredParameters(TArray<FString>& OutMissingParameterNames) override;
+
+protected:
+	virtual UMaterialInstanceDynamic* GetMaterialInstance() = 0;
+};
+
+template <typename HandleDataType>
+FInstancedStruct TAvaMaskMaterialHandle<HandleDataType>::MakeDataStruct()
+{
+	FInstancedStruct MaterialData = FInstancedStruct::Make<FHandleData>();
+	return MaterialData;
+}
+
+template <typename HandleDataType>
+const UScriptStruct* TAvaMaskMaterialHandle<HandleDataType>::GetDataStructType()
+{
+	return FHandleData::StaticStruct();
+}
+
+template <typename HandleDataType>
+bool TAvaMaskMaterialHandle<HandleDataType>::GetMaskParameters(FAvaMask2DMaterialParameters& OutParameters)
+{
+	if (const UMaterialInstanceDynamic* MaterialInstance = GetMaterialInstance())
+	{
+		OutParameters.SetFromMaterial(
+			MaterialInstance
+			, UE::AvaMask::Internal::TextureParameterInfo
+			, UE::AvaMask::Internal::ChannelParameterInfo
+			, UE::AvaMask::Internal::InvertParameterInfo
+			, UE::AvaMask::Internal::BaseOpacityParameterInfo
+			, UE::AvaMask::Internal::PaddingParameterInfo
+			, UE::AvaMask::Internal::FeatherParameterInfo);
+		return true;
+	}
+
+	return false;
+}
+
+template <typename HandleDataType>
+bool TAvaMaskMaterialHandle<HandleDataType>::SetMaskParameters(UTexture* InTexture
+	, float InBaseOpacity
+	, EGeometryMaskColorChannel InChannel
+	, bool bInInverted
+	, const FVector2f& InPadding
+	, bool bInApplyFeathering
+	, float InOuterFeathering
+	, float InInnerFeathering)
+{
+	if (UMaterialInstanceDynamic* MaterialInstance = GetMaterialInstance())
+	{
+		UE_LOG(LogAvaMask, Display, TEXT("SetParameters: Texture:%s, Channel:%s"), InTexture ? *InTexture->GetName() : TEXT("(None)"), *UE::AvaMask::Internal::MaskChannelEnumToVector[InChannel].ToString());
+
+		UTexture* DefaultTexture = nullptr;
+		MaterialInstance->GetTextureParameterDefaultValue(UE::AvaMask::Internal::TextureParameterInfo, DefaultTexture);
+
+		InTexture = InTexture ? InTexture : DefaultTexture;
+
+		MaterialInstance->SetTextureParameterValueByInfo(UE::AvaMask::Internal::TextureParameterInfo, InTexture);
+		MaterialInstance->SetVectorParameterValueByInfo(UE::AvaMask::Internal::ChannelParameterInfo, UE::AvaMask::Internal::MaskChannelEnumToVector[InChannel]);
+		MaterialInstance->SetScalarParameterValueByInfo(UE::AvaMask::Internal::InvertParameterInfo, bInInverted ? 1.0f : 0.0f);
+		MaterialInstance->SetScalarParameterValueByInfo(UE::AvaMask::Internal::BaseOpacityParameterInfo, FMath::Clamp(InBaseOpacity, 0.0f, 1.0f));
+		MaterialInstance->SetVectorParameterValueByInfo(UE::AvaMask::Internal::PaddingParameterInfo, FLinearColor(FMath::Max(0, InPadding.X), FMath::Max(0, InPadding.Y), 0, 0));
+		MaterialInstance->SetVectorParameterValueByInfo(UE::AvaMask::Internal::FeatherParameterInfo, FLinearColor(bInApplyFeathering ? 1.0f : 0.0f, FMath::Max(0, InOuterFeathering), FMath::Max(0, InInnerFeathering), FMath::Max(0, FMath::Max(InOuterFeathering, InInnerFeathering))));
+
+		return true;
+	}
+
+	return false;
+}
+
+template <typename HandleDataType>
+bool TAvaMaskMaterialHandle<HandleDataType>::HasRequiredParameters(TArray<FString>& OutMissingParameterNames)
+{
+	if (const UMaterialInterface* Material = GetMaterial())
+	{
+		return UE::Ava::MaterialHasParameters(
+			*Material
+			, {
+				{ UE::AvaMask::Internal::TextureParameterInfo.Name, EMaterialParameterType::Texture }
+				, { UE::AvaMask::Internal::BaseOpacityParameterInfo.Name, EMaterialParameterType::Scalar }
+				, { UE::AvaMask::Internal::ChannelParameterInfo.Name, EMaterialParameterType::Vector }
+				, { UE::AvaMask::Internal::InvertParameterInfo.Name, EMaterialParameterType::Scalar }
+			}
+			, OutMissingParameterNames);
+	}
+
+	return false;
+}
