[MotionDesign]

diff --git a/Engine/Plugins/Experimental/Avalanche/Source/AvalancheMask/Private/Mask2D/AvaMask2DBaseModifier.cpp b/Engine/Plugins/Experimental/Avalanche/Source/AvalancheMask/Private/Mask2D/AvaMask2DBaseModifier.cpp
new file mode 100644
--- /dev/null
+++ b/Engine/Plugins/Experimental/Avalanche/Source/AvalancheMask/Private/Mask2D/AvaMask2DBaseModifier.cpp
@@ -0,0 +1,654 @@
+// Copyright Epic Games, Inc. All Rights Reserved.
+
+#include "Mask2D/AvaMask2DBaseModifier.h"
+
+#include "AvaMaskLog.h"
+#include "AvaMaskSubsystem.h"
+#include "AvaMaskUtilities.h"
+#include "Components/PrimitiveComponent.h"
+#include "Engine/Engine.h"
+#include "Engine/TextureRenderTarget2D.h"
+#include "Framework/AvaGizmoComponent.h"
+#include "GeometryMaskCanvas.h"
+#include "GeometryMaskReadComponent.h"
+#include "GeometryMaskSubsystem.h"
+#include "GeometryMaskWriteComponent.h"
+#include "Handling/AvaHandleUtilities.h"
+#include "Handling/AvaObjectHandleSubsystem.h"
+#include "Handling/IAvaMaskMaterialCollectionHandle.h"
+#include "Handling/IAvaMaskMaterialHandle.h"
+#include "Materials/AvaMaskMaterialInstanceSubsystem.h"
+#include "Materials/Material.h"
+#include "Materials/MaterialInstanceDynamic.h"
+#include "Modifiers/ActorModifierCoreStack.h"
+#include "Subsystems/ActorModifierCoreSubsystem.h"
+
+#define LOCTEXT_NAMESPACE "AvaMask2DModifier"
+
+namespace UE::AvaMask::Private
+{
+	template <typename KeyType, typename ValueType>
+	ValueType& FindInPreviousOrAddToNew(
+		const TMap<KeyType, ValueType>& InPrevious
+		, TMap<KeyType, ValueType>& InNew
+		, const KeyType& InKey
+		, TIdentity_T<TUniqueFunction<ValueType()>>&& InValueFactory)
+	{
+		if (const ValueType* FoundValue = InPrevious.Find(InKey))
+		{
+			return InNew.Emplace(InKey, MoveTempIfPossible(*FoundValue));
+		}
+
+		if (ValueType* FoundValue = InNew.Find(InKey))
+		{
+			return *FoundValue;
+		}
+
+		return InNew.Emplace(InKey, InValueFactory());
+	}
+}
+
+#if WITH_EDITOR
+const TAvaPropertyChangeDispatcher<UAvaMask2DBaseModifier> UAvaMask2DBaseModifier::PropertyChangeDispatcher =
+{
+	{ GET_MEMBER_NAME_CHECKED(UAvaMask2DBaseModifier, bUseParentChannel), &UAvaMask2DBaseModifier::OnUseParentChannelChanged },
+	{ GET_MEMBER_NAME_CHECKED(UAvaMask2DBaseModifier, Channel), &UAvaMask2DBaseModifier::OnChannelChanged },
+	{ GET_MEMBER_NAME_CHECKED(UAvaMask2DBaseModifier, bInverted), &UAvaMask2DBaseModifier::OnInvertedChanged },
+	{ GET_MEMBER_NAME_CHECKED(UAvaMask2DBaseModifier, bUseBlur), &UAvaMask2DBaseModifier::OnBlurChanged },
+	{ GET_MEMBER_NAME_CHECKED(UAvaMask2DBaseModifier, BlurStrength), &UAvaMask2DBaseModifier::OnBlurChanged },
+	{ GET_MEMBER_NAME_CHECKED(UAvaMask2DBaseModifier, bUseFeathering), &UAvaMask2DBaseModifier::OnFeatherChanged },
+	{ GET_MEMBER_NAME_CHECKED(UAvaMask2DBaseModifier, OuterFeatherRadius), &UAvaMask2DBaseModifier::OnFeatherChanged },
+	{ GET_MEMBER_NAME_CHECKED(UAvaMask2DBaseModifier, InnerFeatherRadius), &UAvaMask2DBaseModifier::OnFeatherChanged },
+};
+#endif
+
+void UAvaMask2DBaseModifier::SetUseParentChannel(const bool bInUseParentChannel)
+{
+	if (bUseParentChannel != bInUseParentChannel)
+	{
+		bUseParentChannel = bInUseParentChannel;
+		OnUseParentChannelChanged();
+	}
+}
+
+const FName UAvaMask2DBaseModifier::GetChannel() const
+{
+	return bUseParentChannel ? ParentChannel : Channel;
+}
+
+void UAvaMask2DBaseModifier::SetChannel(FName InChannel)
+{
+	if (bUseParentChannel)
+	{
+		if (ParentChannel != InChannel)
+		{
+			ParentChannel = InChannel;
+			OnChannelChanged();
+		}
+	}
+	else
+	{
+		if (Channel != InChannel)
+		{
+			Channel = InChannel;
+			OnChannelChanged();
+		}
+	}
+}
+
+void UAvaMask2DBaseModifier::SetIsInverted(const bool bInInvert)
+{
+	if (bInverted != bInInvert)
+	{
+		bInverted = bInInvert;
+		OnInvertedChanged();
+	}
+}
+
+void UAvaMask2DBaseModifier::UseBlur(bool bInUseBlur)
+{
+	if (bUseBlur != bInUseBlur)
+	{
+		bUseBlur = bInUseBlur;
+		OnBlurChanged();
+	}
+}
+
+void UAvaMask2DBaseModifier::SetBlurStrength(float InBlurStrength)
+{
+	InBlurStrength = FMath::Max(0.0f, InBlurStrength);
+	if (!FMath::IsNearlyEqual(BlurStrength, InBlurStrength))
+	{
+		BlurStrength = InBlurStrength;
+		OnBlurChanged();
+	}
+}
+
+void UAvaMask2DBaseModifier::UseFeathering(bool bInUseFeathering)
+{
+	if (bUseFeathering != bInUseFeathering)
+	{
+		bUseFeathering = bInUseFeathering;
+		OnFeatherChanged();
+	}
+}
+
+void UAvaMask2DBaseModifier::SetOuterFeatherRadius(int32 InFeatherRadius)
+{
+	InFeatherRadius = FMath::Max(0, InFeatherRadius);
+	if (OuterFeatherRadius != InFeatherRadius)
+	{
+		OuterFeatherRadius = InFeatherRadius;
+		OnFeatherChanged();
+	}
+}
+
+void UAvaMask2DBaseModifier::SetInnerFeatherRadius(int32 InFeatherRadius)
+{
+	InFeatherRadius = FMath::Max(0, InFeatherRadius);
+	if (InnerFeatherRadius != InFeatherRadius)
+	{
+		InnerFeatherRadius = InFeatherRadius;
+		OnFeatherChanged();
+	}
+}
+
+void UAvaMask2DBaseModifier::OnBlurChanged()
+{
+	MarkModifierDirty();
+}
+
+void UAvaMask2DBaseModifier::OnFeatherChanged()
+{
+	MarkModifierDirty();
+}
+
+#if WITH_EDITOR
+void UAvaMask2DBaseModifier::PostEditChangeProperty(FPropertyChangedEvent& InPropertyChangedEvent)
+{
+	Super::PostEditChangeProperty(InPropertyChangedEvent);
+
+	PropertyChangeDispatcher.OnPropertyChanged(this, InPropertyChangedEvent);
+}
+#endif
+
+void UAvaMask2DBaseModifier::OnModifierAdded(EActorModifierCoreEnableReason InReason)
+{
+	Super::OnModifierAdded(InReason);
+
+	TRACE_BOOKMARK(TEXT("UAvaMask2DModifier::OnModifierAdded"));
+
+	OnChannelChanged();
+
+	if (InReason == EActorModifierCoreEnableReason::User
+		|| InReason == EActorModifierCoreEnableReason::Duplicate)
+	{
+		SetupChannelName();
+	}
+}
+
+void UAvaMask2DBaseModifier::OnModifierRemoved(EActorModifierCoreDisableReason InReason)
+{
+	Super::OnModifierRemoved(InReason);
+
+	TRACE_BOOKMARK(TEXT("UAvaMask2DModifier::OnModifierRemoved"));
+
+	TArray<AActor*> ActorDataActors;
+	ActorDataActors.Reserve(ActorData.Num());	
+	Algo::Transform(ActorData, ActorDataActors, [](const TPair<TWeakObjectPtr<AActor>, FAvaMask2DActorData>& InActorDataPair)
+	{
+		return InActorDataPair.Key.Get();
+	});
+	
+	for (AActor* Actor : ActorDataActors)
+	{
+		if (!Actor)
+		{
+			continue;
+		}
+		
+		RemoveFromActor(Actor);
+	}
+}
+
+void UAvaMask2DBaseModifier::OnModifiedActorTransformed()
+{
+	// This is needed to override parent behavior (to disable it)
+}
+
+void UAvaMask2DBaseModifier::SavePreState()
+{
+	Super::SavePreState();
+
+	TMap<TWeakObjectPtr<AActor>, FAvaMask2DActorData> PreviousActorDatas = ActorData;
+	
+	TMap<TWeakObjectPtr<AActor>, FAvaMask2DActorData> NewActorDatas;
+	NewActorDatas.Reserve(PreviousActorDatas.Num());
+
+	ForEachActor<AActor>(
+		[this, &NewActorDatas](AActor* InActor)
+		{
+			SaveActorPreState(InActor, NewActorDatas.Emplace(InActor));
+			return true;
+		}
+		, EActorModifierCoreLookup::SelfAndAllChildren);
+
+	ActorData = NewActorDatas;
+}
+
+void UAvaMask2DBaseModifier::RestorePreState()
+{
+	Super::RestorePreState();
+
+	TGuardValue<bool> RestoreStateGuard(bIsRestoring, true);
+	for (const TPair<TWeakObjectPtr<AActor>, FAvaMask2DActorData>& ActorDataPair : ActorData)
+	{
+		if (AActor* Actor = ActorDataPair.Key.Get())
+		{
+			RestoreActorPreState(Actor, ActorDataPair.Value);
+		}
+	}
+}
+
+void UAvaMask2DBaseModifier::Apply()
+{
+	if (UGeometryMaskCanvas* Canvas = GetCurrentCanvas())
+	{
+		Canvas->SetApplyBlur(bUseBlur);
+		Canvas->SetBlurStrength(BlurStrength);
+
+		Canvas->SetApplyFeather(bUseFeathering);
+		Canvas->SetOuterFeatherRadius(OuterFeatherRadius);
+		Canvas->SetInnerFeatherRadius(InnerFeatherRadius);
+	}
+}
+
+void UAvaMask2DBaseModifier::SaveActorPreState(AActor* InActor, FAvaMask2DActorData& InActorData)
+{
+	// Get MaterialCollectionHandle
+	const TSharedPtr<IAvaMaskMaterialCollectionHandle>& MaterialCollectionHandle =
+		UE::Ava::Internal::FindOrAddHandleByLambda(
+			MaterialCollectionHandles
+			, InActor
+			, [this, InActor]()
+			{
+				TSharedPtr<IAvaMaskMaterialCollectionHandle> Handle = GetObjectHandleSubsystem()->MakeHandle<IAvaMaskMaterialCollectionHandle>(InActor, UE::AvaMask::Internal::HandleTag);
+				Handle->OnSourceMaterialsChanged().BindUObject(this, &UAvaMask2DBaseModifier::OnMaterialsChanged);
+				return Handle;
+			});
+
+	FInstancedStruct& HandleData = MaterialCollectionHandleData.FindOrAdd(InActor, MaterialCollectionHandle->MakeDataStruct());
+	MaterialCollectionHandle->SaveOriginalState(HandleData);
+}
+
+void UAvaMask2DBaseModifier::RestoreActorPreState(AActor* InActor, const FAvaMask2DActorData& InActorData)
+{
+	// Get MaterialCollectionHandle
+    const TSharedPtr<IAvaMaskMaterialCollectionHandle>& MaterialCollectionHandle =
+    	UE::Ava::Internal::FindOrAddHandleByLambda(
+    		MaterialCollectionHandles
+    		, InActor
+    		, [this, InActor]()
+    		{
+    			return GetObjectHandleSubsystem()->MakeHandle<IAvaMaskMaterialCollectionHandle>(InActor, UE::AvaMask::Internal::HandleTag);
+    		});
+
+    FInstancedStruct& MaterialCollectionData = MaterialCollectionHandleData.FindOrAdd(InActor, MaterialCollectionHandle->MakeDataStruct());
+    MaterialCollectionHandle->ApplyOriginalState(MaterialCollectionData);
+}
+
+void UAvaMask2DBaseModifier::OnSceneTreeTrackedActorParentChanged(
+	int32 InIdx
+	, const TArray<TWeakObjectPtr<AActor>>& InPreviousParentActor
+	, const TArray<TWeakObjectPtr<AActor>>& InNewParentActor)
+{
+	Super::OnSceneTreeTrackedActorParentChanged(InIdx, InPreviousParentActor, InNewParentActor);
+
+	// @todo: remove stuff from child writers/readers if necessary
+
+	// We don't use the provided parent here, instead we just use the event to trigger custom parent discovery
+	if (bUseParentChannel)
+	{
+		if (TryResolveParentChannel())
+		{
+			OnChannelChanged();
+		}
+	}
+}
+
+void UAvaMask2DBaseModifier::OnSceneTreeTrackedActorChildrenChanged(int32 InIdx, const TSet<TWeakObjectPtr<AActor>>& InPreviousChildrenActors, const TSet<TWeakObjectPtr<AActor>>& InNewChildrenActors)
+{
+	Super::OnSceneTreeTrackedActorChildrenChanged(InIdx, InPreviousChildrenActors, InNewChildrenActors);
+
+	TSet<TWeakObjectPtr<AActor>> UnparentedActors = InPreviousChildrenActors.Difference(InNewChildrenActors);
+	for (TWeakObjectPtr<AActor>& ActorWeak : UnparentedActors)
+	{
+		if (AActor* Actor = ActorWeak.Get())
+		{
+			RemoveFromActor(Actor);
+		}
+	}
+}
+
+FName UAvaMask2DBaseModifier::GenerateUniqueMaskName() const
+{
+	const AActor* ModifierActor = GetModifiedActor();
+	if (!ModifierActor)
+	{
+		return NAME_None;
+	}
+
+	return FName(ModifierActor->GetActorNameOrLabel() + TEXT("_Mask"));
+}
+
+UAvaMask2DBaseModifier* UAvaMask2DBaseModifier::FindMaskModifierOnActor(const AActor* InActor)
+{
+	// Get from modifier, if present
+	if (const UActorModifierCoreSubsystem* ModifierSubsystem = UActorModifierCoreSubsystem::Get())
+	{
+		const UActorModifierCoreStack* ModifierStack = ModifierSubsystem->GetActorModifierStack(InActor);
+		if (!ModifierStack)
+		{
+			return nullptr;
+		}
+
+		UAvaMask2DBaseModifier* FoundMaskModifier = nullptr;
+		
+		TArray<UAvaMask2DBaseModifier*> FoundModifiers;
+		ModifierStack->GetClassModifiers<UAvaMask2DBaseModifier>(FoundModifiers);
+		if (!FoundModifiers.IsEmpty())
+		{
+			FoundMaskModifier = FoundModifiers.Last();
+		}
+
+		return FoundMaskModifier;
+	}
+
+	return nullptr;
+}
+
+void UAvaMask2DBaseModifier::OnModifierCDOSetup(FActorModifierCoreMetadata& InMetadata)
+{
+	Super::OnModifierCDOSetup(InMetadata);
+
+	InMetadata.SetName(TEXT("Mask"));
+	InMetadata.SetCategory(TEXT("Rendering"));
+#if WITH_EDITOR
+	InMetadata.SetDescription(LOCTEXT("ModifierDescription", "Allows to use a custom mask texture on attached actors materials"));
+#endif
+}
+
+UActorComponent* UAvaMask2DBaseModifier::FindOrAddMaskComponent(TSubclassOf<UActorComponent> InComponentClass, AActor* InActor)
+{
+	if (!InActor || !InComponentClass)
+	{
+		return nullptr;
+	}
+
+	// Skip adding to actors that already write to a mask
+	if (InComponentClass->ImplementsInterface(UGeometryMaskReadInterface::StaticClass()))
+	{
+		if (InActor->FindComponentByInterface(UGeometryMaskWriteInterface::StaticClass()))
+		{
+			UE_LOG(LogAvaMask, Display, TEXT("Attempting to add a Mask Read component to an Actor that has a Mask Write component: '%s'"), *InActor->GetName());
+			return nullptr;
+		}
+	}
+
+	UActorComponent* MaskComponent = UE::AvaMask::Internal::FindOrAddComponent(InComponentClass, InActor);
+	SetupMaskComponent(MaskComponent);
+
+	return MaskComponent;
+}
+
+bool UAvaMask2DBaseModifier::ActorSupportsMaskReadWrite(const AActor* InActor)
+{
+	// Currently only applies to primitive components
+	return InActor->FindComponentByClass<UPrimitiveComponent>() != nullptr;
+}
+
+bool UAvaMask2DBaseModifier::TryResolveParentChannel()
+{
+	bool bParentChannelWasFound = false;
+	if (const AActor* ActorModified = GetModifiedActor())
+	{
+		const AActor* Parent = ActorModified->GetAttachParentActor();
+
+		const IGeometryMaskReadInterface* ReadComponent = nullptr;
+		while (ReadComponent == nullptr && Parent)
+		{
+			ReadComponent = Cast<IGeometryMaskReadInterface>(Parent->FindComponentByInterface<UGeometryMaskReadInterface>());
+			Parent = Parent->GetAttachParentActor();
+		}
+
+		if (ReadComponent)
+		{
+			ParentChannel = ReadComponent->GetParameters().CanvasName;
+			bParentChannelWasFound = true;
+		}
+		else
+		{
+			// If parent invalid or not found, just use the previously specified one (not from parent)
+			ParentChannel = Channel;
+		}
+	}
+
+	return bParentChannelWasFound;
+}
+
+void UAvaMask2DBaseModifier::SetupMaskComponent(UActorComponent* InComponent)
+{
+	if (!InComponent)
+	{
+		return;
+	}
+
+	if (const UGeometryMaskCanvas* Canvas = GetCurrentCanvas())
+	{
+		if (FAvaMask2DActorData* StoredActorData = ActorData.Find(InComponent->GetOwner()))
+		{
+			StoredActorData->CanvasTextureWeak = Canvas->GetTexture();
+		}
+	}
+	else
+	{
+		UE_LOG(LogAvaMask, Warning, TEXT("Expected canvas to be valid"));
+	}
+}
+
+void UAvaMask2DBaseModifier::RemoveFromActor(AActor* InActor)
+{
+	ActorData.Remove(InActor);
+	
+	MaterialCollectionHandleData.Remove(InActor);
+	MaterialCollectionHandles.Remove(InActor);
+	
+	UE::AvaMask::Internal::RemoveComponentByInterface<UGeometryMaskReadInterface>(InActor);
+	UE::AvaMask::Internal::RemoveComponentByInterface<UGeometryMaskWriteInterface>(InActor);
+	if (UAvaGizmoComponent* FoundComponent = InActor->FindComponentByClass<UAvaGizmoComponent>())
+	{
+		FoundComponent->DestroyComponent();
+	}
+}
+
+void UAvaMask2DBaseModifier::OnChannelChanged()
+{
+	// Reset/invalidate cached
+	{
+		for (TPair<TWeakObjectPtr<AActor>, FAvaMask2DActorData>& ActorDataPair : ActorData)
+		{
+			ActorDataPair.Value.CanvasTextureWeak.Reset();
+		}
+		
+		LastResolvedCanvasName = NAME_None;
+		CanvasWeak.Reset();
+	}
+
+#if WITH_EDITOR
+	if (!Channel.IsNone())
+	{
+		if (UAvaMaskSubsystem* Subsystem = GEngine->GetEngineSubsystem<UAvaMaskSubsystem>())
+		{
+			Subsystem->SetLastSpecifiedChannelName(Channel);
+		}
+	}
+#endif
+
+	TryResolveCanvas();
+	
+	MarkModifierDirty();
+}
+
+void UAvaMask2DBaseModifier::OnInvertedChanged()
+{
+	MarkModifierDirty();
+}
+
+void UAvaMask2DBaseModifier::SetupChannelName()
+{
+	// Already named, leave as-is
+	if (!Channel.IsNone())
+	{
+		return;
+	}
+
+	AutoChannelName = GenerateUniqueMaskName();	
+
+#if WITH_EDITOR
+	if (const UAvaMaskSubsystem* Subsystem = GEngine->GetEngineSubsystem<UAvaMaskSubsystem>())
+	{
+		Channel = Subsystem->GetLastSpecifiedChannelName();
+	}
+
+	if (Channel.IsNone())
+	{
+		Channel = TEXT("0");
+	}
+#else
+	Channel = TEXT("0");
+#endif
+}
+
+void UAvaMask2DBaseModifier::OnUseParentChannelChanged()
+{
+	if (bUseParentChannel)
+	{
+		if (TryResolveParentChannel())
+		{
+			OnChannelChanged();
+		}
+	}
+
+	MarkModifierDirty();
+}
+
+void UAvaMask2DBaseModifier::OnMaskSetCanvas(const UGeometryMaskCanvas* InCanvas, AActor* InActor)
+{
+	if (!InCanvas)
+	{
+		return;
+	}
+
+	if (FAvaMask2DActorData* StoredActorData = ActorData.Find(InActor))
+	{
+		// Cache canvas texture
+		StoredActorData->CanvasTextureWeak = InCanvas->GetTexture();
+	}
+}
+
+void UAvaMask2DBaseModifier::TryResolveCanvas()
+{
+	if (UGeometryMaskSubsystem* GeometryMaskSubsystem = GEngine->GetEngineSubsystem<UGeometryMaskSubsystem>())
+	{
+		UGeometryMaskCanvas* Canvas = GeometryMaskSubsystem->GetNamedCanvas(GetChannel());
+		LastResolvedCanvasName = Canvas->GetCanvasName();
+		CanvasWeak = Canvas;
+	}
+}
+
+UTexture* UAvaMask2DBaseModifier::TryResolveCanvasTexture(AActor* InActor, FAvaMask2DActorData& InActorData)
+{
+	if (!InActor)
+	{
+		return nullptr;
+	}
+
+	// Get cached version
+	if (UTexture* CanvasTexture = InActorData.CanvasTextureWeak.Get())
+	{
+		return CanvasTexture;
+	}
+	
+	if (!GEngine)
+	{
+		return nullptr;
+	}
+
+	// Try get canvas texture if already available
+	UGeometryMaskSubsystem* MaskSubsystem = GEngine->GetEngineSubsystem<UGeometryMaskSubsystem>();
+	if (!MaskSubsystem)
+	{
+		return nullptr;
+	}
+			
+	if (const UGeometryMaskCanvas* Canvas = MaskSubsystem->GetNamedCanvas(Channel))
+	{
+		// Cache canvas texture
+		UTexture* CanvasTexture = Canvas->GetTexture();
+		InActorData.CanvasTextureWeak = CanvasTexture;
+
+		return CanvasTexture;
+	}
+
+	return nullptr;
+}
+
+UAvaObjectHandleSubsystem* UAvaMask2DBaseModifier::GetObjectHandleSubsystem()
+{
+	if (UAvaObjectHandleSubsystem* Subsystem = ObjectHandleSubsystem.Get())
+	{
+		return Subsystem;
+	}
+
+	ObjectHandleSubsystem = GEngine->GetEngineSubsystem<UAvaObjectHandleSubsystem>();
+	return ObjectHandleSubsystem.Get();
+}
+
+UAvaMaskMaterialInstanceSubsystem* UAvaMask2DBaseModifier::GetMaterialInstanceSubsystem()
+{
+	if (UAvaMaskMaterialInstanceSubsystem* Subsystem = MaterialInstanceSubsystem.Get())
+	{
+		return Subsystem;
+	}
+
+	MaterialInstanceSubsystem = GEngine->GetEngineSubsystem<UAvaMaskMaterialInstanceSubsystem>();
+	return MaterialInstanceSubsystem.Get();
+}
+
+UGeometryMaskCanvas* UAvaMask2DBaseModifier::GetCurrentCanvas()
+{
+	if (LastResolvedCanvasName.IsNone() || LastResolvedCanvasName != GetChannel())
+	{
+		LastResolvedCanvasName = NAME_None;
+		CanvasWeak.Reset();
+	}
+
+	if (!CanvasWeak.IsValid())
+	{
+		TryResolveCanvas();
+	}
+	
+	if (UGeometryMaskCanvas* Canvas = CanvasWeak.Get())
+	{
+		return Canvas;
+	}
+
+	return nullptr;
+}
+
+void UAvaMask2DBaseModifier::OnMaterialsChanged(UPrimitiveComponent* InPrimitiveComponent, const TArray<TSharedPtr<IAvaMaskMaterialHandle>>& InMaterialHandles)
+{
+	MarkModifierDirty();
+}
+
+#undef LOCTEXT_NAMESPACE
diff --git a/Engine/Plugins/Experimental/Avalanche/Source/AvalancheMask/Private/Mask2D/AvaMask2DReadModifier.cpp b/Engine/Plugins/Experimental/Avalanche/Source/AvalancheMask/Private/Mask2D/AvaMask2DReadModifier.cpp
new file mode 100644
--- /dev/null
+++ b/Engine/Plugins/Experimental/Avalanche/Source/AvalancheMask/Private/Mask2D/AvaMask2DReadModifier.cpp
@@ -0,0 +1,193 @@
+// Copyright Epic Games, Inc. All Rights Reserved.
+
+#include "Mask2D/AvaMask2DReadModifier.h"
+
+#include "AvaMaskUtilities.h"
+#include "Components/PrimitiveComponent.h"
+#include "Engine/Engine.h"
+#include "Engine/TextureRenderTarget2D.h"
+#include "GeometryMaskCanvas.h"
+#include "GeometryMaskCanvasResource.h"
+#include "GeometryMaskReadComponent.h"
+#include "Handling/AvaHandleUtilities.h"
+#include "Handling/AvaObjectHandleSubsystem.h"
+#include "Handling/IAvaMaskMaterialCollectionHandle.h"
+#include "Handling/IAvaMaskMaterialHandle.h"
+#include "Materials/Material.h"
+#include "Materials/MaterialInstanceDynamic.h"
+#include "Modifiers/ActorModifierCoreStack.h"
+
+#define LOCTEXT_NAMESPACE "AvaMask2DModifier"
+
+#if WITH_EDITOR
+const TAvaPropertyChangeDispatcher<UAvaMask2DReadModifier> UAvaMask2DReadModifier::PropertyChangeDispatcher =
+{
+	{ GET_MEMBER_NAME_CHECKED(UAvaMask2DReadModifier, BaseOpacity), &UAvaMask2DReadModifier::OnBaseOpacityChanged }
+};
+#endif
+
+void UAvaMask2DReadModifier::SetBaseOpacity(float InBaseOpacity)
+{
+	InBaseOpacity = FMath::Clamp(InBaseOpacity, 0.0f, 1.0f);
+	if (!FMath::IsNearlyEqual(BaseOpacity, InBaseOpacity))
+	{
+		BaseOpacity = InBaseOpacity;
+		OnBaseOpacityChanged();
+	}
+}
+
+#if WITH_EDITOR
+void UAvaMask2DReadModifier::PostEditChangeProperty(FPropertyChangedEvent& InPropertyChangedEvent)
+{
+	Super::PostEditChangeProperty(InPropertyChangedEvent);
+
+	UAvaMask2DReadModifier::PropertyChangeDispatcher.OnPropertyChanged(this, InPropertyChangedEvent);
+}
+#endif
+
+void UAvaMask2DReadModifier::Apply()
+{
+	Super::Apply();
+	
+	if (GetCurrentCanvas())
+	{
+		for (TPair<TWeakObjectPtr<AActor>, FAvaMask2DActorData>& ActorDataPair : ActorData)
+		{
+			if (AActor* Actor = ActorDataPair.Key.Get())
+			{
+				if (!ApplyRead(Actor, ActorDataPair.Value))
+				{
+					return;
+				}			
+			}
+		}
+	}
+
+	Next();
+}
+
+EBlendMode UAvaMask2DReadModifier::GetBlendMode() const
+{
+	return bUseBlur || bUseFeathering ? EBlendMode::BLEND_Translucent : EBlendMode::BLEND_Masked;
+}
+
+bool UAvaMask2DReadModifier::ApplyRead(AActor* InActor, FAvaMask2DActorData& InActorData)
+{
+	// Only add read/write component to actors with primitives
+	if (!ActorSupportsMaskReadWrite(InActor))
+	{
+		// Ok to continue, just nothing done for this actor
+		return true;
+	}
+	
+	if (const UGeometryMaskCanvas* Canvas = GetCurrentCanvas())
+	{
+		// Get MaterialCollectionHandle
+		const TSharedPtr<IAvaMaskMaterialCollectionHandle>& MaterialCollectionHandle =
+			UE::Ava::Internal::FindOrAddHandleByLambda(
+				MaterialCollectionHandles
+				, InActor
+				, [this, InActor]()
+				{
+					return GetObjectHandleSubsystem()->MakeHandle<IAvaMaskMaterialCollectionHandle>(InActor, UE::AvaMask::Internal::HandleTag);
+				});
+
+		const UTexture* CanvasTexture = TryResolveCanvasTexture(InActor, InActorData);
+		if (!CanvasTexture)
+		{
+			return true;
+		}
+
+		if (const UGeometryMaskCanvas* GeometryCanvas = GetCurrentCanvas())
+		{
+			FAvaMask2DSubjectParameters ApplyParameters;
+			ApplyParameters.MaterialParameters.BlendMode = GetBlendMode();		
+			ApplyParameters.MaterialParameters.CanvasName = GeometryCanvas->GetCanvasName();
+			ApplyParameters.MaterialParameters.Texture = GeometryCanvas->GetTexture();
+			ApplyParameters.MaterialParameters.BaseOpacity = BaseOpacity;
+			ApplyParameters.MaterialParameters.Channel = GeometryCanvas->GetColorChannel();
+			ApplyParameters.MaterialParameters.ChannelAsVector = UE::GeometryMask::MaskChannelEnumToVector[ApplyParameters.MaterialParameters.Channel];
+			ApplyParameters.MaterialParameters.bInvert = bInverted;
+
+			const FIntPoint ViewportSize = GeometryCanvas->GetResource()->GetViewportSize();
+			const FVector2f ViewportPadding = FVector2f::One() - FVector2f(ViewportSize + GeometryCanvas->GetResource()->GetViewportPadding()) / ViewportSize;
+			ApplyParameters.MaterialParameters.Padding = ViewportPadding;
+
+			ApplyParameters.MaterialParameters.bApplyFeathering = bUseFeathering;
+			ApplyParameters.MaterialParameters.OuterFeatherRadius = OuterFeatherRadius;
+			ApplyParameters.MaterialParameters.InnerFeatherRadius = InnerFeatherRadius;
+
+			FInstancedStruct& MaterialCollectionData = MaterialCollectionHandleData.FindOrAdd(InActor, MaterialCollectionHandle->MakeDataStruct());
+			MaterialCollectionHandle->ApplyModifiedState(ApplyParameters, MaterialCollectionData);
+		}
+
+		if (UGeometryMaskReadComponent* ReadComponent = FindOrAddMaskComponent<UGeometryMaskReadComponent>(InActor))
+		{
+			FGeometryMaskReadParameters Parameters = ReadComponent->GetParameters();
+			Parameters.CanvasName = Channel;
+			Parameters.ColorChannel = Canvas->GetColorChannel();
+			Parameters.bInvert = bInverted;
+
+			ReadComponent->SetParameters(Parameters);
+		}
+
+		// Only check for valid materials here, after ApplyModifiedState has occured once
+		FText FailReason;
+		if (!MaterialCollectionHandle->ValidateMaterials(FailReason))
+		{
+			Fail(FailReason);
+			return false;
+		}
+	}
+
+	return true;
+}
+
+void UAvaMask2DReadModifier::OnModifierCDOSetup(FActorModifierCoreMetadata& InMetadata)
+{
+	Super::OnModifierCDOSetup(InMetadata);
+
+	InMetadata.SetName(TEXT("MaskRead"));
+	InMetadata.SetCategory(TEXT("Rendering"));
+	InMetadata.DisallowAfter(TEXT("MaskWrite"));
+	InMetadata.DisallowBefore(TEXT("MaskWrite"));
+#if WITH_EDITOR
+	InMetadata.SetDisplayName(FText::FromString(TEXT("Mask (Apply)")));
+	InMetadata.SetDescription(LOCTEXT("ModifierDescription", "Allows to use a custom mask texture on attached actors materials"));
+#endif
+}
+
+void UAvaMask2DReadModifier::SetupMaskComponent(UActorComponent* InComponent)
+{
+	if (!InComponent)
+	{
+		return;
+	}
+
+	if (IGeometryMaskReadInterface* MaskReader = Cast<IGeometryMaskReadInterface>(InComponent))
+	{
+		SetupMaskReadComponent(MaskReader);
+	}
+	
+	UAvaMask2DBaseModifier::SetupMaskComponent(InComponent);
+}
+
+void UAvaMask2DReadModifier::SetupMaskReadComponent(IGeometryMaskReadInterface* InMaskReader)
+{
+	if (const UGeometryMaskCanvas* Canvas = GetCurrentCanvas())
+	{
+		// Set canvas name to read from
+		FGeometryMaskReadParameters ReadParameters = InMaskReader->GetParameters();
+		ReadParameters.CanvasName = Channel;
+		ReadParameters.ColorChannel = Canvas->GetColorChannel();
+
+		InMaskReader->SetParameters(ReadParameters);
+	}
+}
+
+void UAvaMask2DReadModifier::OnBaseOpacityChanged()
+{
+	MarkModifierDirty();
+}
+
+#undef LOCTEXT_NAMESPACE
diff --git a/Engine/Plugins/Experimental/Avalanche/Source/AvalancheMask/Private/Mask2D/AvaMask2DWriteModifier.cpp b/Engine/Plugins/Experimental/Avalanche/Source/AvalancheMask/Private/Mask2D/AvaMask2DWriteModifier.cpp
new file mode 100644
--- /dev/null
+++ b/Engine/Plugins/Experimental/Avalanche/Source/AvalancheMask/Private/Mask2D/AvaMask2DWriteModifier.cpp
@@ -0,0 +1,152 @@
+// Copyright Epic Games, Inc. All Rights Reserved.
+
+#include "Mask2D/AvaMask2DWriteModifier.h"
+
+#include "AvaMaskUtilities.h"
+#include "Components/PrimitiveComponent.h"
+#include "Engine/Engine.h"
+#include "Framework/AvaGizmoComponent.h"
+#include "GeometryMaskCanvas.h"
+#include "GeometryMaskWriteComponent.h"
+#include "Materials/Material.h"
+#include "Materials/MaterialInstanceDynamic.h"
+#include "Modifiers/ActorModifierCoreStack.h"
+
+#define LOCTEXT_NAMESPACE "AvaMask2DModifier"
+
+#if WITH_EDITOR
+const TAvaPropertyChangeDispatcher<UAvaMask2DWriteModifier> UAvaMask2DWriteModifier::PropertyChangeDispatcher =
+{
+	{ GET_MEMBER_NAME_CHECKED(UAvaMask2DWriteModifier, WriteOperation), &UAvaMask2DWriteModifier::OnWriteOperationChanged }
+};
+#endif
+
+#if WITH_EDITOR
+void UAvaMask2DWriteModifier::PostEditChangeProperty(FPropertyChangedEvent& InPropertyChangedEvent)
+{
+	Super::PostEditChangeProperty(InPropertyChangedEvent);
+
+	UAvaMask2DWriteModifier::PropertyChangeDispatcher.OnPropertyChanged(this, InPropertyChangedEvent);
+}
+#endif
+
+void UAvaMask2DWriteModifier::Apply()
+{
+	Super::Apply();
+	
+	if (GetCurrentCanvas())
+	{
+		for (TPair<TWeakObjectPtr<AActor>, FAvaMask2DActorData>& ActorDataPair : ActorData)
+		{
+			if (AActor* Actor = ActorDataPair.Key.Get())
+			{
+				if (!ApplyWrite(Actor, ActorDataPair.Value))
+				{
+					return;
+				}
+			}
+		}
+	}
+
+	Next();
+}
+
+bool UAvaMask2DWriteModifier::ApplyWrite(AActor* InActor, FAvaMask2DActorData& InActorData)
+{
+	// Only add read/write component to actors with primitives
+	if (ActorSupportsMaskReadWrite(InActor))
+	{
+		if (const UGeometryMaskCanvas* GeometryCanvas = GetCurrentCanvas())
+		{
+			const UTexture* CanvasTexture = TryResolveCanvasTexture(InActor, InActorData);
+			if (!CanvasTexture)
+			{
+				return true;
+			}
+		
+			if (UGeometryMaskWriteMeshComponent* WriteComponent = FindOrAddMaskComponent<UGeometryMaskWriteMeshComponent>(InActor))
+			{
+				FGeometryMaskWriteParameters Parameters = WriteComponent->GetParameters();
+				Parameters.CanvasName = Channel;
+				Parameters.ColorChannel = GeometryCanvas->GetColorChannel();
+				Parameters.OperationType = WriteOperation;
+
+				WriteComponent->SetParameters(Parameters);
+			}
+		}
+	}
+
+	if (UAvaGizmoComponent* VisualizationComponent = UE::AvaMask::Internal::FindOrAddComponent<UAvaGizmoComponent>(InActor))
+	{
+		VisualizationComponent->SetVisibleInEditor(true);
+		VisualizationComponent->SetHiddenInGame(false);
+		VisualizationComponent->SetShowWireframe(false);
+		VisualizationComponent->SetWireframeColor(FLinearColor::Red);
+		VisualizationComponent->SetsStencil(true);
+		VisualizationComponent->SetStencilId(150);
+		VisualizationComponent->SetRenderInMainPass(false);
+		VisualizationComponent->SetRenderDepth(false);
+		VisualizationComponent->SetMaterialToDefault();
+		VisualizationComponent->ApplyGizmoValues();
+	}
+
+	return true;
+}
+
+void UAvaMask2DWriteModifier::SetWriteOperation(const EGeometryMaskCompositeOperation InWriteOperation)
+{
+	if (WriteOperation != InWriteOperation)
+	{
+		WriteOperation = InWriteOperation;
+		OnWriteOperationChanged();
+	}
+}
+
+void UAvaMask2DWriteModifier::OnModifierCDOSetup(FActorModifierCoreMetadata& InMetadata)
+{
+	Super::OnModifierCDOSetup(InMetadata);
+
+	InMetadata.SetName(TEXT("MaskWrite"));
+	InMetadata.SetCategory(TEXT("Rendering"));
+	InMetadata.DisallowAfter(TEXT("MaskRead"));
+	InMetadata.DisallowBefore(TEXT("MaskRead"));
+#if WITH_EDITOR
+	InMetadata.SetDisplayName(FText::FromString(TEXT("Mask (Set)")));
+	InMetadata.SetDescription(LOCTEXT("ModifierDescription", "Allows to use a custom mask texture on attached actors materials"));
+#endif
+}
+
+void UAvaMask2DWriteModifier::SetupMaskComponent(UActorComponent* InComponent)
+{
+	if (!InComponent)
+	{
+		return;
+	}
+	
+	if (IGeometryMaskWriteInterface* MaskWriter = Cast<IGeometryMaskWriteInterface>(InComponent))
+	{
+		SetupMaskWriteComponent(MaskWriter);
+	}
+
+	Super::SetupMaskComponent(InComponent);
+}
+
+void UAvaMask2DWriteModifier::SetupMaskWriteComponent(IGeometryMaskWriteInterface* InMaskWriter)
+{
+	if (const UGeometryMaskCanvas* Canvas = GetCurrentCanvas())
+	{
+		// Set canvas name to write to
+		FGeometryMaskWriteParameters WriteParameters = InMaskWriter->GetParameters();
+		WriteParameters.CanvasName = Channel;
+		WriteParameters.ColorChannel = Canvas->GetColorChannel();
+
+		InMaskWriter->SetParameters(WriteParameters);
+	}
+}
+
+void UAvaMask2DWriteModifier::OnWriteOperationChanged()
+{
+	MarkModifierDirty();
+}
+
+#undef LOCTEXT_NAMESPACE
diff --git a/Engine/Plugins/Experimental/Avalanche/Source/AvalancheMask/Public/Mask2D/AvaMask2DBaseModifier.h b/Engine/Plugins/Experimental/Avalanche/Source/AvalancheMask/Public/Mask2D/AvaMask2DBaseModifier.h
new file mode 100644
--- /dev/null
+++ b/Engine/Plugins/Experimental/Avalanche/Source/AvalancheMask/Public/Mask2D/AvaMask2DBaseModifier.h
@@ -0,0 +1,235 @@
+// Copyright Epic Games, Inc. All Rights Reserved.
+
+#pragma once
+
+#include "AvaMaskTypes.h"
+#include "AvaPropertyChangeDispatcher.h"
+#include "GeometryMaskTypes.h"
+#include "Modifiers/AvaArrangeBaseModifier.h"
+#include "StructView.h"
+#include "UObject/ObjectPtr.h"
+
+#include "AvaMask2DBaseModifier.generated.h"
+
+class AActor;
+class IAvaActorHandle;
+class IAvaComponentHandle;
+class IAvaMaskMaterialCollectionHandle;
+class IAvaMaskMaterialHandle;
+class IGeometryMaskReadInterface;
+class IGeometryMaskWriteInterface;
+class UActorComponent;
+class UAvaMaskMaterialInstanceSubsystem;
+class UAvaObjectHandleSubsystem;
+class UGeometryMaskCanvas;
+class UMaterialInstanceDynamic;
+class UMaterialInterface;
+class UPrimitiveComponent;
+class UTexture;
+struct FAvaShapeParametricMaterial;
+struct FMaterialParameterInfo;
+
+UENUM(BlueprintType)
+enum class EAvaMask2DMode : uint8
+{
+	Read		UMETA(DisplayName = "Target", ToolTip = "Use the specified Mask Channel to apply to this geometry"),
+	Write		UMETA(DisplayName = "Source", ToolTip = "Use the specified Mask Channel to render this geometry to"),
+};
+
+/** Used for target actors to store essential information */
+USTRUCT()
+struct FAvaMask2DActorData
+{
+	GENERATED_BODY()
+
+	/** The canvas texture to apply to this actor materials */
+	UPROPERTY()
+	TWeakObjectPtr<UTexture> CanvasTextureWeak;
+
+	/** Original assigned materials, used to create instances from, and restore to when modifier removed. */
+	UPROPERTY()
+	TMap<FAvaMask2DComponentMaterialPath, TObjectPtr<UMaterialInterface>> OriginalMaterials;
+
+	/** Instantiated materials corresponding to the above. */
+	UPROPERTY()
+	TMap<FAvaMask2DComponentMaterialPath, TObjectPtr<UMaterialInstanceDynamic>> GeneratedMaterialInstances;
+};
+
+/** Uses scene actors to create a mask texture and applies it to attached actors */
+UCLASS(Abstract)
+class AVALANCHEMASK_API UAvaMask2DBaseModifier
+	: public UAvaArrangeBaseModifier
+{
+	GENERATED_BODY()
+
+	friend class UAvaMask2DModifierShared;
+	
+public:
+	//~ Begin UObject
+#if WITH_EDITOR
+	virtual void PostEditChangeProperty(FPropertyChangedEvent& InPropertyChangedEvent) override;
+#endif
+	//~ End UObject
+
+	bool UseParentChannel() const { return bUseParentChannel; }
+	void SetUseParentChannel(const bool bInUseParentChannel);
+
+	const FName GetChannel() const;
+	void SetChannel(FName InChannel);
+
+	bool IsInverted() const { return bInverted; }
+	void SetIsInverted(const bool bInInvert);
+
+	bool UseBlur() const { return bUseBlur; }
+	void UseBlur(bool bInUseBlur);
+
+	float GetBlurStrength() const { return BlurStrength; }
+	void SetBlurStrength(float InBlurStrength);
+
+	bool UseFeathering() const { return bUseFeathering; }
+	void UseFeathering(bool bInUseFeathering);
+
+	int32 GetOuterFeatherRadius() const { return OuterFeatherRadius; }
+	void SetOuterFeatherRadius(int32 InFeatherRadius);
+
+	int32 GetInnerFeatherRadius() const { return InnerFeatherRadius; }
+	void SetInnerFeatherRadius(int32 InFeatherRadius);
+
+	/** Utility to generate a unique mask name based on the associated Actor */
+	FName GenerateUniqueMaskName() const;
+
+	/** Utility to find an existing modifier on the provided actor or it's parent */
+	static UAvaMask2DBaseModifier* FindMaskModifierOnActor(const AActor* InActor);
+
+protected:
+	//~ Begin UActorModifierCoreBase
+	virtual void OnModifierCDOSetup(FActorModifierCoreMetadata& InMetadata) override;
+	virtual void OnModifierAdded(EActorModifierCoreEnableReason InReason) override;
+	virtual void OnModifierRemoved(EActorModifierCoreDisableReason InReason) override;
+	virtual void OnModifiedActorTransformed() override;
+	virtual void SavePreState() override;
+	virtual void RestorePreState() override;
+	virtual void Apply() override;
+	//~ End UActorModifierCoreBase
+
+	virtual void SaveActorPreState(AActor* InActor, FAvaMask2DActorData& InActorData);
+	virtual void RestoreActorPreState(AActor* InActor, const FAvaMask2DActorData& InActorData);
+
+	//~ Begin IAvaSceneTreeUpdateModifierExtension
+	virtual void OnSceneTreeTrackedActorParentChanged(int32 InIdx, const TArray<TWeakObjectPtr<AActor>>& InPreviousParentActor, const TArray<TWeakObjectPtr<AActor>>& InNewParentActor) override;
+	virtual void OnSceneTreeTrackedActorChildrenChanged(int32 InIdx, const TSet<TWeakObjectPtr<AActor>>& InPreviousChildrenActors, const TSet<TWeakObjectPtr<AActor>>& InNewChildrenActors) override;
+	//~ End IAvaSceneTreeUpdateModifierExtension
+
+	void OnUseParentChannelChanged();
+	void OnChannelChanged();
+	void OnInvertedChanged();
+	void OnBlurChanged();
+	void OnFeatherChanged();
+
+	void SetupChannelName();
+	virtual void SetupMaskComponent(UActorComponent* InComponent);
+	virtual void RemoveFromActor(AActor* InActor);
+
+	void OnMaskSetCanvas(const UGeometryMaskCanvas* InCanvas, AActor* InActor);
+
+	UActorComponent* FindOrAddMaskComponent(TSubclassOf<UActorComponent> InComponentClass, AActor* InActor);
+
+	template<
+		typename InComponentClass
+		UE_REQUIRES(std::is_base_of_v<UActorComponent, InComponentClass>)>
+	InComponentClass* FindOrAddMaskComponent(AActor* InActor)
+	{
+		return Cast<InComponentClass>(FindOrAddMaskComponent(InComponentClass::StaticClass(), InActor));
+	}
+
+	static bool ActorSupportsMaskReadWrite(const AActor* InActor);
+
+	/** Returns true if parent channel was found */
+	bool TryResolveParentChannel();
+
+	void TryResolveCanvas();
+	UTexture* TryResolveCanvasTexture(AActor* InActor, FAvaMask2DActorData& InActorData);
+
+	UAvaObjectHandleSubsystem* GetObjectHandleSubsystem();
+	UAvaMaskMaterialInstanceSubsystem* GetMaterialInstanceSubsystem();
+
+	/** Returns or resolves the currently referenced canvas. */
+	UGeometryMaskCanvas* GetCurrentCanvas();
+	
+	virtual void OnMaterialsChanged(UPrimitiveComponent* InPrimitiveComponent, const TArray<TSharedPtr<IAvaMaskMaterialHandle>>& InMaterialHandles);
+
+protected:
+	/** Whether to get the channel from the parent (first one that specifies a mask channel) */
+	UPROPERTY(EditInstanceOnly, BlueprintReadWrite, Getter = "UseParentChannel", Setter = "SetUseParentChannel", Category = "Mask2D", meta = (AllowPrivateAccess = "true", ValidEnumValues = "Self,Parent"))
+	bool bUseParentChannel = false;
+
+	/** Channel found when GetChannelFromParent is true */
+	UPROPERTY(VisibleAnywhere, Category = "Mask2D", meta = (EditCondition = "bUseParentChannel", EditConditionHides))
+	FName ParentChannel = NAME_None;
+
+	/** Channel to read to or write from */
+	UPROPERTY(EditInstanceOnly, BlueprintReadWrite, Getter, Setter, Category = "Mask2D", meta = (EditCondition = "!bUseParentChannel", EditConditionHides, AllowPrivateAccess = "true"))
+	FName Channel = NAME_None;
+
+	/** Whether to apply the mask as inverted (visible becomes invisible and vice versa) */
+	UPROPERTY(EditInstanceOnly, BlueprintReadWrite, Getter = "IsInverted", Setter = "SetIsInverted", Category = "Mask2D", meta = (AllowPrivateAccess = "true"))
+	bool bInverted = false;
+
+	UPROPERTY(EditInstanceOnly, BlueprintReadWrite, Getter = "UseBlur", Setter = "UseBlur", Category = "Mask2D|Shared", meta = (AllowPrivateAccess = "true"))
+	bool bUseBlur = false;
+
+	UPROPERTY(EditInstanceOnly, BlueprintReadWrite, Getter, Setter, Category = "Mask2D|Shared", meta = (ClampMin = 0.0, AllowPrivateAccess = "true", EditCondition = "bUseBlur", EditConditionHides))
+	float BlurStrength = 1.0f;
+
+	UPROPERTY(EditInstanceOnly, BlueprintReadWrite, Getter = "UseFeathering", Setter = "UseFeathering", Category = "Mask2D|Shared", meta = (AllowPrivateAccess = "true"))
+	bool bUseFeathering = false;
+
+	UPROPERTY(EditInstanceOnly, BlueprintReadWrite, Getter, Setter, Category = "Mask2D|Shared", meta = (ClampMin = 0, AllowPrivateAccess = "true", EditCondition = "bUseFeathering", EditConditionHides))
+	int32 OuterFeatherRadius = 16;
+
+	UPROPERTY(EditInstanceOnly, BlueprintReadWrite, Getter, Setter, Category = "Mask2D|Shared", meta = (ClampMin = 0, AllowPrivateAccess = "true", EditCondition = "bUseFeathering", EditConditionHides))
+	int32 InnerFeatherRadius = 16;
+
+#if WITH_EDITOR
+	/** Used for PECP */
+	static const TAvaPropertyChangeDispatcher<UAvaMask2DBaseModifier> PropertyChangeDispatcher;
+#endif
+
+	UPROPERTY(Transient)
+	TWeakObjectPtr<UAvaObjectHandleSubsystem> ObjectHandleSubsystem;
+	
+	UPROPERTY(Transient)
+	TWeakObjectPtr<UAvaMaskMaterialInstanceSubsystem> MaterialInstanceSubsystem;
+
+	/**
+	 * Cached auto channel name, used when the user adds the modifier and immediately sets the mode to Source,
+	 * in which case a different channel name is chosen.
+	 */
+	UPROPERTY(Transient)
+	FName AutoChannelName;
+
+	/** Cached actor data to apply/restore */
+	UPROPERTY(DuplicateTransient)
+	TMap<TWeakObjectPtr<AActor>, FAvaMask2DActorData> ActorData;
+
+	/** Reference to the last resolved canvas name, as stored in CanvasWeak */
+	UPROPERTY(Transient, DuplicateTransient)
+	FName LastResolvedCanvasName;
+
+	/** Reference to the Canvas used */
+	UPROPERTY(Transient, DuplicateTransient)
+	TWeakObjectPtr<UGeometryMaskCanvas> CanvasWeak;
+
+	/** Reference to the underlying canvas texture */
+	TWeakObjectPtr<UTexture> CanvasTextureWeak;
+
+	bool bIsRestoring = false;
+
+	UPROPERTY()
+	TMap<TWeakObjectPtr<UObject>, FInstancedStruct> MaterialCollectionHandleData;
+
+	UPROPERTY()
+	TMap<TWeakObjectPtr<UMaterialInterface>, FInstancedStruct> MaterialHandleData;
+
+	TMap<TObjectKey<UObject>, TSharedPtr<IAvaMaskMaterialCollectionHandle>> MaterialCollectionHandles;
+};
diff --git a/Engine/Plugins/Experimental/Avalanche/Source/AvalancheMask/Public/Mask2D/AvaMask2DReadModifier.h b/Engine/Plugins/Experimental/Avalanche/Source/AvalancheMask/Public/Mask2D/AvaMask2DReadModifier.h
new file mode 100644
--- /dev/null
+++ b/Engine/Plugins/Experimental/Avalanche/Source/AvalancheMask/Public/Mask2D/AvaMask2DReadModifier.h
@@ -0,0 +1,76 @@
+// Copyright Epic Games, Inc. All Rights Reserved.
+
+#pragma once
+
+#include "AvaMask2DBaseModifier.h"
+#include "AvaPropertyChangeDispatcher.h"
+#include "Modifiers/AvaArrangeBaseModifier.h"
+#include "UObject/ObjectPtr.h"
+#include "UObject/Package.h"
+#include "UObject/WeakObjectPtr.h"
+
+#include "AvaMask2DReadModifier.generated.h"
+
+class AActor;
+class IAvaActorHandle;
+class IAvaComponentHandle;
+class IAvaMaskMaterialCollectionHandle;
+class IAvaMaskMaterialHandle;
+class IGeometryMaskReadInterface;
+class IGeometryMaskWriteInterface;
+class UActorComponent;
+class UAvaMaskMaterialInstanceSubsystem;
+class UAvaObjectHandleSubsystem;
+class UGeometryMaskCanvas;
+class UMaterialInstanceDynamic;
+class UMaterialInterface;
+class UPrimitiveComponent;
+class UTexture;
+struct FAvaShapeParametricMaterial;
+struct FMaterialParameterInfo;
+
+/** Uses scene actors to create a mask texture and applies it to attached actors */
+UCLASS(BlueprintType)
+class AVALANCHEMASK_API UAvaMask2DReadModifier
+	: public UAvaMask2DBaseModifier
+{
+	GENERATED_BODY()
+
+	friend class UAvaMask2DModifierShared;
+	
+public:
+	//~ Begin UObject
+#if WITH_EDITOR
+	virtual void PostEditChangeProperty(FPropertyChangedEvent& InPropertyChangedEvent) override;
+#endif
+	//~ End UObject
+
+	const float GetBaseOpacity() const { return BaseOpacity; }
+	void SetBaseOpacity(float InBaseOpacity);
+
+protected:
+	//~ Begin UActorModifierCoreBase
+	virtual void OnModifierCDOSetup(FActorModifierCoreMetadata& InMetadata) override;
+	virtual void Apply() override;
+	//~ End UActorModifierCoreBase
+
+	/** Get target blend mode dependent on modifier settings. */
+	EBlendMode GetBlendMode() const;
+
+	void OnBaseOpacityChanged();
+
+	virtual void SetupMaskComponent(UActorComponent* InComponent) override;
+	void SetupMaskReadComponent(IGeometryMaskReadInterface* InMaskReader);
+
+	bool ApplyRead(AActor* InActor, FAvaMask2DActorData& InActorData);
+
+protected:
+	/** Base opacity/alpha to use in Read mode */
+	UPROPERTY(EditInstanceOnly, BlueprintReadWrite, Getter, Setter, Category = "Mask2D", meta = (ClampMin = 0.0, ClampMax = 1.0, AllowPrivateAccess = "true"))
+	float BaseOpacity = 0.0f;
+
+#if WITH_EDITOR
+	/** Used for PECP */
+	static const TAvaPropertyChangeDispatcher<UAvaMask2DReadModifier> PropertyChangeDispatcher;
+#endif
+};
diff --git a/Engine/Plugins/Experimental/Avalanche/Source/AvalancheMask/Public/Mask2D/AvaMask2DWriteModifier.h b/Engine/Plugins/Experimental/Avalanche/Source/AvalancheMask/Public/Mask2D/AvaMask2DWriteModifier.h
new file mode 100644
--- /dev/null
+++ b/Engine/Plugins/Experimental/Avalanche/Source/AvalancheMask/Public/Mask2D/AvaMask2DWriteModifier.h
@@ -0,0 +1,73 @@
+// Copyright Epic Games, Inc. All Rights Reserved.
+
+#pragma once
+
+#include "AvaMask2DBaseModifier.h"
+#include "AvaPropertyChangeDispatcher.h"
+#include "GeometryMaskTypes.h"
+#include "UObject/ObjectPtr.h"
+#include "UObject/Package.h"
+#include "UObject/WeakObjectPtr.h"
+
+#include "AvaMask2DWriteModifier.generated.h"
+
+class AActor;
+class IAvaActorHandle;
+class IAvaComponentHandle;
+class IAvaMaskMaterialCollectionHandle;
+class IAvaMaskMaterialHandle;
+class IGeometryMaskReadInterface;
+class IGeometryMaskWriteInterface;
+class UActorComponent;
+class UAvaMaskMaterialInstanceSubsystem;
+class UAvaObjectHandleSubsystem;
+class UGeometryMaskCanvas;
+class UMaterialInstanceDynamic;
+class UMaterialInterface;
+class UPrimitiveComponent;
+class UTexture;
+struct FAvaShapeParametricMaterial;
+struct FMaterialParameterInfo;
+
+/** Uses scene actors to create a mask texture and applies it to attached actors */
+UCLASS(BlueprintType)
+class AVALANCHEMASK_API UAvaMask2DWriteModifier
+	: public UAvaMask2DBaseModifier
+{
+	GENERATED_BODY()
+
+	friend class UAvaMask2DModifierShared;
+	
+public:
+	//~ Begin UObject
+#if WITH_EDITOR
+	virtual void PostEditChangeProperty(FPropertyChangedEvent& InPropertyChangedEvent) override;
+#endif
+	//~ End UObject
+
+	EGeometryMaskCompositeOperation GetWriteOperation() const { return WriteOperation; }
+	void SetWriteOperation(const EGeometryMaskCompositeOperation InWriteOperation);
+
+protected:
+	//~ Begin UActorModifierCoreBase
+	virtual void OnModifierCDOSetup(FActorModifierCoreMetadata& InMetadata) override;
+	virtual void Apply() override;
+	//~ End UActorModifierCoreBase
+
+	void OnWriteOperationChanged();
+
+	void SetupMaskComponent(UActorComponent* InComponent);
+	void SetupMaskWriteComponent(IGeometryMaskWriteInterface* InMaskWriter);
+
+	bool ApplyWrite(AActor* InActor, FAvaMask2DActorData& InActorData);
+
+protected:
+	/** How to write to the chosen mask channel */
+	UPROPERTY(EditInstanceOnly, BlueprintReadWrite, Getter, Setter, Category = "Mask2D", meta = (ValidEnumValues = "Add, Subtract", AllowPrivateAccess = "true"))
+	EGeometryMaskCompositeOperation WriteOperation = EGeometryMaskCompositeOperation::Add;
+
+#if WITH_EDITOR
+	/** Used for PECP */
+	static const TAvaPropertyChangeDispatcher<UAvaMask2DWriteModifier> PropertyChangeDispatcher;
+#endif
+};
diff --git a/Engine/Plugins/Experimental/Avalanche/Source/AvalancheMaskEditor/Private/AvaMaskEditorCommands.h b/Engine/Plugins/Experimental/Avalanche/Source/AvalancheMaskEditor/Private/AvaMaskEditorCommands.h
new file mode 100644
--- /dev/null
+++ b/Engine/Plugins/Experimental/Avalanche/Source/AvalancheMaskEditor/Private/AvaMaskEditorCommands.h
@@ -0,0 +1,27 @@
+﻿// Copyright Epic Games, Inc. All Rights Reserved.
+
+#pragma once
+
+#include "AvaMaskEditorStyle.h"
+#include "Framework/Commands/Commands.h"
+
+class FAvaMaskEditorCommands : public TCommands<FAvaMaskEditorCommands>
+{
+public:
+	FAvaMaskEditorCommands()
+		: TCommands<FAvaMaskEditorCommands>(
+			TEXT("AvalancheMaskEditor"),
+			NSLOCTEXT("Contexts", "AvalancheMaskEditor", "Motion Design Masking"),
+			NAME_None,
+			FAvalancheMaskEditorStyle::GetStyleSetName())
+	{
+	}
+
+	virtual void RegisterCommands() override;
+
+public:
+	TSharedPtr<FUICommandInfo> ToggleMaskMode;
+	TSharedPtr<FUICommandInfo> ToggleShowAllMasks;
+	TSharedPtr<FUICommandInfo> ToggleIsolateMask;
+	TSharedPtr<FUICommandInfo> ToggleEnableMask;
+};
diff --git a/Engine/Plugins/Experimental/Avalanche/Source/AvalancheMaskEditor/Private/AvaMaskEditorLog.cpp b/Engine/Plugins/Experimental/Avalanche/Source/AvalancheMaskEditor/Private/AvaMaskEditorLog.cpp
new file mode 100644
--- /dev/null
+++ b/Engine/Plugins/Experimental/Avalanche/Source/AvalancheMaskEditor/Private/AvaMaskEditorLog.cpp
@@ -0,0 +1,5 @@
+// Copyright Epic Games, Inc. All Rights Reserved.
+
+#include "AvaMaskEditorLog.h"
+
+DEFINE_LOG_CATEGORY(LogAvalancheMaskEditor);
diff --git a/Engine/Plugins/Experimental/Avalanche/Source/AvalancheMaskEditor/Private/AvaMaskEditorLog.h b/Engine/Plugins/Experimental/Avalanche/Source/AvalancheMaskEditor/Private/AvaMaskEditorLog.h
new file mode 100644
--- /dev/null
+++ b/Engine/Plugins/Experimental/Avalanche/Source/AvalancheMaskEditor/Private/AvaMaskEditorLog.h
@@ -0,0 +1,7 @@
+// Copyright Epic Games, Inc. All Rights Reserved.
+
+#pragma once
+
+#include "Logging/LogMacros.h"
+
+DECLARE_LOG_CATEGORY_EXTERN(LogAvalancheMaskEditor, VeryVerbose, All);
diff --git a/Engine/Plugins/Experimental/Avalanche/Source/AvalancheMaskEditor/Private/AvaMaskEditorMode.cpp b/Engine/Plugins/Experimental/Avalanche/Source/AvalancheMaskEditor/Private/AvaMaskEditorMode.cpp
new file mode 100644
--- /dev/null
+++ b/Engine/Plugins/Experimental/Avalanche/Source/AvalancheMaskEditor/Private/AvaMaskEditorMode.cpp
@@ -0,0 +1,488 @@
+﻿// Copyright Epic Games, Inc. All Rights Reserved.
+
+#include "AvaMaskEditorMode.h"
+
+#include "AvaMaskEditorCommands.h"
+#include "AvaMaskEditorLog.h"
+#include "AvaMaskEditorStyle.h"
+#include "AvaMaskEditorSubsystem.h"
+#include "Editor.h"
+#include "EditorModeManager.h"
+#include "Engine/World.h"
+#include "EngineUtils.h"
+#include "Framework/Application/SlateApplication.h"
+#include "Framework/AvaGizmoComponent.h"
+#include "GeometryMaskSubsystem.h"
+#include "IAvaMaskEditor.h"
+#include "IGeometryMaskWriteInterface.h"
+#include "LevelEditor.h"
+#include "Mask2D/AvaMask2DBaseModifier.h"
+#include "Modifiers/ActorModifierCoreStack.h"
+#include "Modules/ModuleManager.h"
+#include "Selection.h"
+#include "Subsystems/ActorModifierCoreSubsystem.h"
+#include "Widgets/Images/SImage.h"
+#include "Widgets/Input/SButton.h"
+#include "Widgets/SGeometryMaskCanvasPreview.h"
+#include "Widgets/SWidget.h"
+#include "Widgets/Text/STextBlock.h"
+
+#define LOCTEXT_NAMESPACE "AvalancheMaskEditorMode"
+
+const FEditorModeID UAvalancheMaskEditorMode::EM_AvalancheMaskEditorModeId(UE::AvalancheMaskEditor::AvalancheMaskEditorModeName);
+
+UAvalancheMaskEditorMode::UAvalancheMaskEditorMode()
+{
+	Info = FEditorModeInfo(UAvalancheMaskEditorMode::EM_AvalancheMaskEditorModeId,
+		LOCTEXT("AvalancheMaskEditorModeName", "AvalancheMaskEditorMode"),
+		FSlateIcon(),
+		false);
+}
+
+void UAvalancheMaskEditorMode::Enter()
+{
+	UEdMode::Enter();
+
+	if (GEditor && GEngine)
+	{
+		UTypedElementSelectionSet* SelectionSet = GetModeManager()->GetSelectedActors()->GetElementSelectionSet();
+		SelectionSet->OnChanged().AddUObject(this, &UAvalancheMaskEditorMode::OnSelectionChanged);
+		
+		WeakLastSelectedActor = SelectionSet->GetTopSelectedObject<AActor>();
+		WeakActorSelectionSet = SelectionSet;
+		
+		OnActorSpawnedHandle = GetWorld()->AddOnActorSpawnedHandler(FOnActorSpawned::FDelegate::CreateUObject(this, &UAvalancheMaskEditorMode::OnActorSpawned));
+
+		GEngine->OnLevelActorAdded().AddUObject(this, &UAvalancheMaskEditorMode::OnActorSpawned);
+	}
+
+	for (AActor* Actor : TActorRange<AActor>(GetWorld()))
+	{
+		if (Actor->FindComponentByInterface<UGeometryMaskWriteInterface>())
+		{
+			WeakMaskWriterActors.Add(Actor);
+			if (UAvaGizmoComponent* GizmoComponent = Actor->GetComponentByClass<UAvaGizmoComponent>())
+			{
+				GizmoComponent->SetVisibleInEditor(true);
+			}
+		}
+	}
+
+	const FLevelEditorModule& LevelEditorModule = FModuleManager::GetModuleChecked<FLevelEditorModule>("LevelEditor");
+	if (TSharedPtr<ILevelEditor> LevelEditor = LevelEditorModule.GetFirstLevelEditor())
+	{
+		FText DisplayText = LOCTEXT("MaskViewMode", "Mask View");
+    
+    	TSharedPtr<SWidget> ToolWidget = nullptr;
+    	{
+    		static FName ToolkitOverlayMenuName = UE::AvalancheMaskEditor::Internal::ToolkitOverlayMenuName;
+    		UToolMenu* Menu = UToolMenus::Get()->ExtendMenu(ToolkitOverlayMenuName);
+    		Menu->SetStyleSet(&FAvalancheMaskEditorStyle::Get());
+    		Menu->StyleName = "AvalancheMaskEditor.ViewportOverlayToolbar";
+    		{
+    			FToolMenuSection& Section = Menu->FindOrAddSection(TEXT("Default"));
+    			{
+    				FToolMenuEntry& Entry = Section.AddEntry(FToolMenuEntry::InitToolBarButton(FAvaMaskEditorCommands::Get().ToggleShowAllMasks));
+    				Entry.Icon.Set(FSlateIcon(FAvalancheMaskEditorStyle::GetStyleSetName(), TEXT("AvalancheMaskEditor.ToggleShowAllMasks")));
+    			}
+    
+    			{
+    				FToolMenuEntry& Entry = Section.AddEntry(FToolMenuEntry::InitToolBarButton(FAvaMaskEditorCommands::Get().ToggleEnableMask));
+    				Entry.Icon.Set(FSlateIcon(FAvalancheMaskEditorStyle::GetStyleSetName(), TEXT("AvalancheMaskEditor.ToggleDisableMask")));
+    			}
+    
+    			{
+    				FToolMenuEntry& Entry = Section.AddEntry(FToolMenuEntry::InitToolBarButton(FAvaMaskEditorCommands::Get().ToggleIsolateMask));
+    				Entry.Icon.Set(FSlateIcon(FAvalancheMaskEditorStyle::GetStyleSetName(), TEXT("AvalancheMaskEditor.ToggleIsolateMask")));
+    			}
+    		}
+    
+    		ToolWidget = UToolMenus::Get()->GenerateWidget(Menu);	
+    	}
+    	
+		// ViewportOverlay
+		SAssignNew(ViewportOverlayWidget, SOverlay)
+			+SOverlay::Slot(-2)
+			[
+				SAssignNew(CanvasPreviewWidget, SGeometryMaskCanvasPreview)
+				.Visibility(EVisibility::HitTestInvisible)
+				.Opacity(0.5f)
+			]
+	
+			+SOverlay::Slot()
+			[
+				SNew(SHorizontalBox)
+				+SHorizontalBox::Slot()
+				.HAlign(HAlign_Center)
+				.VAlign(VAlign_Bottom)
+				.Padding(FMargin(0.0f, 0.0f, 0.f, 15.f))
+				[
+					SNew(SBorder)
+					.BorderImage(FAppStyle::Get().GetBrush("EditorViewport.OverlayBrush"))
+					.Padding(8.f)
+					[
+						SNew(SHorizontalBox)
+						+SHorizontalBox::Slot()
+						.AutoWidth()
+						.VAlign(VAlign_Center)
+						.Padding(FMargin(0.f, 0.f, 8.f, 0.f))
+						[
+							SNew(SImage)
+							.Image(FAvalancheMaskEditorStyle::Get().GetBrush(TEXT("AvalancheMaskEditor.ToggleMaskMode")))
+						]
+
+						+SHorizontalBox::Slot()
+						.AutoWidth()
+						.VAlign(VAlign_Center)
+						.Padding(FMargin(0.f, 0.f, 8.f, 0.f))
+						[
+							SNew(STextBlock)
+							.Text(DisplayText)
+						]
+
+						+SHorizontalBox::Slot()
+						.AutoWidth()
+						.Padding(FMargin(2.0, 0.f, 0.f, 0.f))
+						[
+							SNew(SButton)
+							.ButtonStyle(FAppStyle::Get(), "PrimaryButton")
+							.TextStyle(FAppStyle::Get(), "DialogButtonText")
+							.Text(LOCTEXT("ExitEdit", "Exit"))
+							.ToolTipText(LOCTEXT("ExitTooltip", "Exit Mask Mode"))
+							.HAlign(HAlign_Center)
+							.VAlign(VAlign_Center)
+							.OnClicked_Lambda([EditorModeManager = GetModeManager()]() 
+							{
+								EditorModeManager->DeactivateMode(UAvalancheMaskEditorMode::EM_AvalancheMaskEditorModeId);						
+								return FReply::Handled(); 
+							})
+						]
+					]
+				]
+			];
+
+		const TSharedRef<SWidget> ViewportOverlayWidgetRef = ViewportOverlayWidget.ToSharedRef();
+		ViewportOverlayWidgetRef->SetTag(TEXT("AvaMaskOverlayWidget"));
+
+		LevelEditor->AddViewportOverlayWidget(ViewportOverlayWidgetRef);
+
+		if (const TSharedPtr<SOverlay> ParentOverlay = StaticCastSharedPtr<SOverlay>(ViewportOverlayWidget->GetParentWidget());
+			ParentOverlay.IsValid())
+		{
+			// Now that we have the Overlay widget, remove the slot
+			ParentOverlay->RemoveSlot(ViewportOverlayWidgetRef);
+
+			// Then re-add it, specifying zorder
+			ParentOverlay->AddSlot(-2)
+			[
+				ViewportOverlayWidgetRef
+			];
+		}
+
+		// Initialize, re-using the callback
+		OnSelectionChanged(WeakActorSelectionSet.Get());
+	}
+}
+
+void UAvalancheMaskEditorMode::Exit()
+{
+	if (GEditor && GEngine)
+	{
+		for (TWeakObjectPtr<AActor>& WeakActor : WeakMaskWriterActors)
+		{
+			if (const AActor* Actor = WeakActor.Get())
+			{
+				if (UAvaGizmoComponent* GizmoComponent = Actor->GetComponentByClass<UAvaGizmoComponent>())
+				{
+					GizmoComponent->SetVisibleInEditor(false);
+				}
+			}
+		}
+
+		// Remove selection events
+		if (UTypedElementSelectionSet* SelectionSet = WeakActorSelectionSet.Get())
+		{
+			SelectionSet->OnChanged().RemoveAll(this);
+		}
+
+		if (ViewportOverlayWidget.IsValid())
+		{
+			const FLevelEditorModule& LevelEditorModule = FModuleManager::GetModuleChecked<FLevelEditorModule>("LevelEditor");
+			if (TSharedPtr<ILevelEditor> LevelEditor = LevelEditorModule.GetFirstLevelEditor())
+			{
+				LevelEditor->RemoveViewportOverlayWidget(ViewportOverlayWidget.ToSharedRef());
+			}
+		}
+		
+		WeakLastSelectedActor = nullptr;
+		WeakMaskWriterActors.Reset();
+		WeakActorSelectionSet.Reset();
+
+		GetWorld()->RemoveOnActorSpawnedHandler(OnActorSpawnedHandle);
+
+		GEngine->OnLevelActorAdded().RemoveAll(this);
+	}
+
+	UEdMode::Exit();
+}
+
+bool UAvalancheMaskEditorMode::UsesToolkits() const
+{
+	return false;
+}
+
+void UAvalancheMaskEditorMode::CreateToolkit()
+{
+	// Toolkit = MakeShared<FAvalancheMaskEditorModeToolkit>();
+}
+
+void UAvalancheMaskEditorMode::ModeTick(float DeltaTime)
+{
+	Super::ModeTick(DeltaTime);
+}
+
+bool UAvalancheMaskEditorMode::IsCompatibleWith(FEditorModeID OtherModeID) const
+{
+	return true;
+}
+
+AActor* UAvalancheMaskEditorMode::GetActorToParentTo() const
+{
+	if (AActor* LastSelectedActor = WeakLastSelectedActor.Get())
+	{
+		return LastSelectedActor;
+	}
+
+	USelection* SelectedActors = GetModeManager()->GetSelectedActors();
+	if (SelectedActors->Num() > 0)
+	{
+		return SelectedActors->GetTop<AActor>();
+	}
+
+	return nullptr;
+}
+
+void UAvalancheMaskEditorMode::OnSelectionChanged(const UTypedElementSelectionSet* InSelectionSet)
+{
+	if (InSelectionSet)
+	{
+		PreviewCanvasName = NAME_None;
+		PreviewCanvasChannel = EGeometryMaskColorChannel::None;
+		
+		if (const UTypedElementSelectionSet* SelectionSet = WeakActorSelectionSet.Get())
+		{
+			TArray<UObject*> AllSelected = SelectionSet->GetSelectedObjects();
+			WeakLastSelectedActor = SelectionSet->GetTopSelectedObject<AActor>();
+			if (const AActor* FirstActor = SelectionSet->GetTopSelectedObject<AActor>())
+			{
+				SelectedMaskModifier = UAvaMask2DBaseModifier::FindMaskModifierOnActor(FirstActor);
+				if (SelectedMaskModifier.IsValid())
+				{
+					SelectedMaskCanvas = GetCanvasReferencedByActor(FirstActor);
+					for (const TWeakObjectPtr<AActor>& WeakWriterActor : WeakMaskWriterActors)
+					{
+						if (const AActor* WriterActor = WeakWriterActor.Get())
+						{
+							if (const IGeometryMaskWriteInterface* WriterObject = Cast<IGeometryMaskWriteInterface>(WriterActor->FindComponentByInterface<UGeometryMaskWriteInterface>()))
+							{
+								if (UAvaGizmoComponent* GizmoComponent = WriterActor->GetComponentByClass<UAvaGizmoComponent>())
+								{
+									FName WriterCanvasName = WriterObject->GetParameters().CanvasName;
+									GizmoComponent->SetsStencil(WriterCanvasName == SelectedMaskCanvas->GetCanvasName());
+									GizmoComponent->SetStencilId(WriterCanvasName != SelectedMaskCanvas->GetCanvasName() ? 0 : 150);
+								}
+							}
+						}
+					}
+
+					PreviewCanvasName = SelectedMaskCanvas->GetCanvasName();
+					PreviewCanvasChannel = SelectedMaskCanvas->GetColorChannel();
+				}
+			}
+			else
+			{
+				// Otherwise the selection has changed and the previous MaskModifier is no longer selected
+				SelectedMaskModifier = nullptr;
+				SelectedMaskCanvas = nullptr;
+			}
+		}
+
+		UpdatePreviewWidget();
+	}
+}
+
+UGeometryMaskCanvas* UAvalancheMaskEditorMode::GetCanvasReferencedByActor(const AActor* InActor)
+{
+	FName CanvasName = NAME_None;
+	if (const IGeometryMaskWriteInterface* WriteComponent = Cast<IGeometryMaskWriteInterface>(InActor->FindComponentByInterface<UGeometryMaskWriteInterface>()))
+	{
+		const FGeometryMaskWriteParameters WriteComponentParameters = WriteComponent->GetParameters();
+		CanvasName = WriteComponentParameters.CanvasName;
+	}	
+	else if (const IGeometryMaskReadInterface* ReadComponent = Cast<IGeometryMaskReadInterface>(InActor->FindComponentByInterface<UGeometryMaskReadInterface>()))
+	{
+		const FGeometryMaskReadParameters ReadComponentParameters = ReadComponent->GetParameters();
+		CanvasName = ReadComponentParameters.CanvasName;
+	}
+
+	if (UGeometryMaskCanvas* Canvas = GEngine->GetEngineSubsystem<UGeometryMaskSubsystem>()->GetNamedCanvas(CanvasName))
+	{
+		return Canvas;		
+	}
+
+	return nullptr;
+}
+
+void UAvalancheMaskEditorMode::UpdatePreviewWidget()
+{
+	if (CanvasPreviewWidget.IsValid())
+	{
+		CanvasPreviewWidget->SetCanvasName(PreviewCanvasName);
+		CanvasPreviewWidget->SetColorChannel(PreviewCanvasChannel);
+	}
+}
+
+void UAvalancheMaskEditorMode::OnActorSpawned(AActor* InActor)
+{
+	if (CanMaskSelected(WeakLastSelectedActor.Get()))
+	{
+		if (AddMaskToSelected({InActor}))
+		{
+			UpdatePreviewWidget();
+		}
+	}
+}
+
+bool UAvalancheMaskEditorMode::AddMaskToSelected(const TArray<AActor*>& InMaskingActors)
+{
+	// @note: This assumes CanMaskSelected passed
+	
+	AActor* ParentActor = GetActorToParentTo();
+	
+	UAvaMask2DBaseModifier* ParentMaskModifier = FindOrAddMaskModifier(ParentActor);
+	if (!ParentMaskModifier)
+	{
+		return false;
+	}
+
+	const UGeometryMaskCanvas* ParentCanvas = GetCanvasReferencedByActor(ParentActor);
+	const FName ChannelName = ParentMaskModifier->GetChannel();
+	const EGeometryMaskColorChannel ColorChannel = ParentCanvas ? ParentCanvas->GetColorChannel() : EGeometryMaskColorChannel::Red;
+
+	// @todo: replace SetMode
+	// ParentMaskModifier->SetMode(EAvaMask2DMode::Read);
+
+	for (AActor* PlacedActor : InMaskingActors)
+	{
+		// Parent to Actor
+		PlacedActor->AttachToActor(ParentActor, FAttachmentTransformRules::KeepWorldTransform);
+		
+		// Setup modifier
+		if (UAvaMask2DBaseModifier* MaskModifier = FindOrAddMaskModifier(PlacedActor))
+		{
+			// @todo: replace SetMode
+			// MaskModifier->SetMode(EAvaMask2DMode::Write);
+			MaskModifier->SetChannel(ChannelName);
+			MaskModifier->SetUseParentChannel(true);
+
+			// Flush unused canvases, in case temporary actors were created/destroyed
+			if (UGeometryMaskSubsystem* Subsystem = GEngine->GetEngineSubsystem<UGeometryMaskSubsystem>())
+			{
+				Subsystem->RemoveWithoutWriters();
+			}
+
+			PreviewCanvasName = ChannelName;
+			PreviewCanvasChannel = ColorChannel;
+		}
+	}
+		
+	return true;
+}
+
+UAvaMask2DBaseModifier* UAvalancheMaskEditorMode::FindOrAddMaskModifier(AActor* InActor)
+{
+	if (const UActorModifierCoreSubsystem* ModifierSubsystem = UActorModifierCoreSubsystem::Get())
+	{
+		UActorModifierCoreStack* ModifierStack = ModifierSubsystem->GetActorModifierStack(InActor);
+		if (!ModifierStack)
+		{
+			ModifierStack = ModifierSubsystem->AddActorModifierStack(InActor);
+		}
+
+		const FName MaskModifierName = GetDefault<UAvaMask2DBaseModifier>()->GetModifierName();
+		UAvaMask2DBaseModifier* MaskModifier = nullptr;
+		
+		TArray<UAvaMask2DBaseModifier*> FoundModifiers;
+		ModifierStack->GetClassModifiers<UAvaMask2DBaseModifier>(FoundModifiers);
+		if (!FoundModifiers.IsEmpty())
+		{
+			MaskModifier = FoundModifiers.Last();
+		}
+		else
+		{
+			FActorModifierCoreStackInsertOp InsertOp;
+			InsertOp.NewModifierName = MaskModifierName;
+			
+			MaskModifier = Cast<UAvaMask2DBaseModifier>(ModifierStack->InsertModifier(InsertOp));
+			if (!MaskModifier)
+			{
+				UE_LOG(LogAvalancheMaskEditor, Error, TEXT("Error inserting Mask modifier."));
+			}
+			else
+			{
+				MaskModifier->SetChannel(MaskModifier->GenerateUniqueMaskName());
+			}
+		}
+
+		return MaskModifier;
+	}
+
+	return nullptr;
+}
+
+bool UAvalancheMaskEditorMode::CanMaskSelected(AActor* InSelectedActor)
+{
+	auto CanMaskActor = [](AActor* InActor)
+	{
+		if (const UActorModifierCoreSubsystem* ModifierSubsystem = UActorModifierCoreSubsystem::Get())
+		{
+			const FName MaskModifierName = GetDefault<UAvaMask2DBaseModifier>()->GetModifierName();
+			
+			if (const UActorModifierCoreStack* ExistingModifierStack = ModifierSubsystem->GetActorModifierStack(InActor))
+			{
+				if (ExistingModifierStack->ContainsModifier(MaskModifierName))
+				{
+					return true;
+				}				
+			}
+			 
+			const TSet<FName> AllowedModifiers = ModifierSubsystem->GetAllowedModifiers(InActor);
+			if (AllowedModifiers.Contains(MaskModifierName))
+			{
+				return true;
+			}
+		}
+
+		return false;
+	};
+
+	// Check provided Actor if specified
+	if (InSelectedActor)
+	{
+		return CanMaskActor(InSelectedActor);
+	}
+
+	// Otherwise get from USelection
+	USelection* SelectedActors = GetModeManager()->GetSelectedActors();
+
+	// Can only add to a single Actor
+	if (SelectedActors->Num() != 1)
+	{
+		return false;
+	}
+
+	AActor* SelectedActor = SelectedActors->GetTop<AActor>();
+	return CanMaskActor(SelectedActor);
+}
+
+#undef LOCTEXT_NAMESPACE
diff --git a/Engine/Plugins/Experimental/Avalanche/Source/AvalancheMaskEditor/Private/AvaMaskEditorStyle.cpp b/Engine/Plugins/Experimental/Avalanche/Source/AvalancheMaskEditor/Private/AvaMaskEditorStyle.cpp
new file mode 100644
--- /dev/null
+++ b/Engine/Plugins/Experimental/Avalanche/Source/AvalancheMaskEditor/Private/AvaMaskEditorStyle.cpp
@@ -0,0 +1,91 @@
+﻿// Copyright Epic Games, Inc. All Rights Reserved.
+
+#include "AvaMaskEditorStyle.h"
+
+#include "Framework/Application/SlateApplication.h"
+#include "Interfaces/IPluginManager.h"
+#include "Styling/SlateStyleMacros.h"
+#include "Styling/SlateStyleRegistry.h"
+#include "Styling/StyleColors.h"
+#include "Styling/ToolBarStyle.h"
+
+namespace UE::AvalancheMaskEditor::Private
+{
+	static FName NAME_StyleSet = TEXT("AvalancheMaskEditorStyle");
+}
+
+#define RootToContentDir Style->RootToContentDir
+
+TSharedPtr<FSlateStyleSet> FAvalancheMaskEditorStyle::StyleInstance = nullptr;
+
+void FAvalancheMaskEditorStyle::Initialize()
+{
+	if (!StyleInstance.IsValid())
+	{
+		StyleInstance = Create();
+		FSlateStyleRegistry::RegisterSlateStyle(*StyleInstance);
+	}
+}
+
+void FAvalancheMaskEditorStyle::Shutdown()
+{
+	FSlateStyleRegistry::UnRegisterSlateStyle(*StyleInstance);
+	ensure(StyleInstance.IsUnique());
+	StyleInstance.Reset();
+}
+
+FName FAvalancheMaskEditorStyle::GetStyleSetName()
+{
+	return UE::AvalancheMaskEditor::Private::NAME_StyleSet;
+}
+
+const ISlateStyle& FAvalancheMaskEditorStyle::Get()
+{
+	return *StyleInstance;
+}
+
+const FVector2D Icon16x16(16.0f, 16.0f);
+const FVector2D Icon20x20(20.0f, 20.0f);
+const FVector2D Icon32x32(32.0f, 32.0f);
+
+TSharedRef<FSlateStyleSet> FAvalancheMaskEditorStyle::Create()
+{
+	TSharedRef<FSlateStyleSet> Style = MakeShared<FSlateStyleSet>(UE::AvalancheMaskEditor::Private::NAME_StyleSet);
+
+	FString ContentRoot = IPluginManager::Get().FindPlugin(UE_PLUGIN_NAME)->GetBaseDir() / TEXT("Resources");
+	Style->SetContentRoot(ContentRoot);
+
+	Style->Set("AvalancheMaskEditor.ToggleMaskMode.Small", new IMAGE_BRUSH_SVG(TEXT("Icons/MaskIcons/Mode_On"), Icon16x16));
+	Style->Set("AvalancheMaskEditor.ToggleMaskMode", new IMAGE_BRUSH_SVG(TEXT("Icons/MaskIcons/Mode_On"), Icon20x20));
+	Style->Set("AvalancheMaskEditor.ToggleShowAllMasks", new IMAGE_BRUSH_SVG(TEXT("Icons/MaskIcons/Disable"), Icon20x20));
+	Style->Set("AvalancheMaskEditor.ToggleDisableMask", new IMAGE_BRUSH_SVG(TEXT("Icons/MaskIcons/Disable"), Icon20x20));
+	Style->Set("AvalancheMaskEditor.ToggleIsolateMask", new IMAGE_BRUSH_SVG(TEXT("Icons/MaskIcons/Disable"), Icon20x20));
+
+	{
+		FToolBarStyle ViewportOverlayToolbarStyle = FAppStyle::Get().GetWidgetStyle<FToolBarStyle>("AssetEditorToolbar");
+
+		ViewportOverlayToolbarStyle.SetButtonPadding(       FMargin(0.0f, 0.0f));
+		ViewportOverlayToolbarStyle.SetCheckBoxPadding(     FMargin(0.0f, 0.0f));
+		ViewportOverlayToolbarStyle.SetComboButtonPadding(  FMargin(0.0f, 0.0f));
+		ViewportOverlayToolbarStyle.SetIndentedBlockPadding(FMargin(0.0f, 0.0f));
+		ViewportOverlayToolbarStyle.SetBlockPadding(        FMargin(0.0f, 0.0f));
+		ViewportOverlayToolbarStyle.SetSeparatorPadding(    FMargin(0.0f, 0.0f));
+		ViewportOverlayToolbarStyle.bShowLabels = false;
+		ViewportOverlayToolbarStyle.SetBackground(FSlateColorBrush(FStyleColors::Transparent));
+		ViewportOverlayToolbarStyle.SetBackgroundPadding(0);
+		ViewportOverlayToolbarStyle.SetButtonPadding(0);
+		ViewportOverlayToolbarStyle.SetCheckBoxPadding(0);
+		
+		Style->Set("AvalancheMaskEditor.ViewportOverlayToolbar", ViewportOverlayToolbarStyle);
+	}
+	
+	return Style;
+}
+
+void FAvalancheMaskEditorStyle::ReloadTextures()
+{
+	if (FSlateApplication::IsInitialized())
+	{
+		FSlateApplication::Get().GetRenderer()->ReloadTextureResources();
+	}
+}
diff --git a/Engine/Plugins/Experimental/Avalanche/Source/AvalancheMaskEditor/Private/AvaMaskEditorStyle.h b/Engine/Plugins/Experimental/Avalanche/Source/AvalancheMaskEditor/Private/AvaMaskEditorStyle.h
new file mode 100644
--- /dev/null
+++ b/Engine/Plugins/Experimental/Avalanche/Source/AvalancheMaskEditor/Private/AvaMaskEditorStyle.h
@@ -0,0 +1,29 @@
+﻿// Copyright Epic Games, Inc. All Rights Reserved.
+
+#pragma once
+
+#include "CoreMinimal.h"
+#include "Styling/SlateStyle.h"
+
+class FSlateStyleSet;
+
+class FAvalancheMaskEditorStyle
+{
+public:
+	static void Initialize();
+	static void Shutdown();
+
+	/** reloads textures used by slate renderer */
+	static void ReloadTextures();
+
+	/** @return The Slate style set for the Shooter game */
+	static const ISlateStyle& Get();
+
+	static FName GetStyleSetName();
+
+private:
+	static TSharedRef<FSlateStyleSet> Create();
+
+private:
+	static TSharedPtr<FSlateStyleSet> StyleInstance;
+};
diff --git a/Engine/Plugins/Experimental/Avalanche/Source/AvalancheMaskEditor/Public/IAvaMaskEditor.h b/Engine/Plugins/Experimental/Avalanche/Source/AvalancheMaskEditor/Public/IAvaMaskEditor.h
new file mode 100644
--- /dev/null
+++ b/Engine/Plugins/Experimental/Avalanche/Source/AvalancheMaskEditor/Public/IAvaMaskEditor.h
@@ -0,0 +1,10 @@
+// Copyright Epic Games, Inc. All Rights Reserved.
+
+#pragma once
+
+#include "UObject/NameTypes.h"
+
+namespace UE::AvalancheMaskEditor
+{
+	const static FName AvalancheMaskEditorModeName("EditMode.AvalancheMask");
+}
diff --git a/Engine/Plugins/Experimental/Avalanche/Source/AvalancheMedia/Private/AvaBroadcastProfile.cpp b/Engine/Plugins/Experimental/Avalanche/Source/AvalancheMedia/Private/AvaBroadcastProfile.cpp
new file mode 100644
--- /dev/null
+++ b/Engine/Plugins/Experimental/Avalanche/Source/AvalancheMedia/Private/AvaBroadcastProfile.cpp
@@ -0,0 +1,336 @@
+// Copyright Epic Games, Inc. All Rights Reserved.
+
+#include "AvaBroadcastProfile.h"
+#include "AvalancheBroadcast.h"
+
+#define LOCTEXT_NAMESPACE "AvalancheBroadcast"
+
+FAvaBroadcastProfile FAvaBroadcastProfile::NullProfile(nullptr, NAME_None);
+
+FAvaBroadcastProfile::FAvaBroadcastProfile(UAvalancheBroadcast* InBroadcast, FName InProfileName)
+	: ParentBroadcastWeak(InBroadcast)
+	, ProfileName(InProfileName)
+{
+}
+
+void FAvaBroadcastProfile::BeginDestroy()
+{
+	for (FAvaOutputChannel& Channel : GetLocalChannels())
+	{
+		Channel.ReleasePlaceholderResources();
+		Channel.ReleaseOutputs();
+		Channel.ReleasePlaceholderRenderTargets();
+	}
+}
+
+UAvalancheBroadcast& FAvaBroadcastProfile::GetBroadcast() const
+{
+	// Note: during the migration towards a non-singleton UAvalancheBroadcast,
+	// a fallback path to the singleton is kept, but it has an ensure to try and catch
+	// the code paths that lead to this.
+	UAvalancheBroadcast* ParentBroadcast = ParentBroadcastWeak.Get();
+	return ensure(ParentBroadcast) ? *ParentBroadcast : UAvalancheBroadcast::Get();
+}
+
+void FAvaBroadcastProfile::CopyProfiles(const FAvaBroadcastProfile& InSourceProfile, FAvaBroadcastProfile& OutTargetProfile)
+{
+	OutTargetProfile.Channels.Empty(InSourceProfile.Channels.Num());
+	for (const FAvaOutputChannel& SourceChannel : InSourceProfile.Channels)
+	{
+		FAvaOutputChannel& TargetChannel = OutTargetProfile.Channels.Add_GetRef(FAvaOutputChannel(&OutTargetProfile));
+		TargetChannel.SetChannelIndex(SourceChannel.GetChannelIndex());
+		
+		FAvaOutputChannel::DuplicateChannel(SourceChannel, TargetChannel);
+	}
+
+	// Assuming destination profile is not current yet.
+	OutTargetProfile.UpdateChannels(false);
+}
+
+bool FAvaBroadcastProfile::StartChannelBroadcast()
+{
+	bool bBroadcastStarted = false;
+	for (FAvaOutputChannel* Channel : GetChannels())
+	{
+		if (Channel->StartChannelBroadcast())
+		{
+			bBroadcastStarted = true;
+		}
+	}
+	return bBroadcastStarted;
+}
+
+void FAvaBroadcastProfile::StopChannelBroadcast()
+{
+	for (FAvaOutputChannel* Channel : GetChannels())
+	{
+		Channel->StopChannelBroadcast();
+	}
+}
+
+bool FAvaBroadcastProfile::IsBroadcastingAnyChannel() const
+{
+	for (const FAvaOutputChannel* Channel : GetChannels())
+	{
+		if (Channel->GetState() == EAvaChannelState::Live)
+		{
+			return true;
+		}
+	}
+	return false;
+}
+
+bool FAvaBroadcastProfile::IsBroadcastingAllChannels() const
+{
+	for (const FAvaOutputChannel* Channel : GetChannels())
+	{
+		if (Channel->GetState() != EAvaChannelState::Live)
+		{
+			return false;
+		}
+	}
+	return !GetChannels().IsEmpty();
+}
+
+bool FAvaBroadcastProfile::IsValidProfile() const
+{
+	return this != &FAvaBroadcastProfile::NullProfile && !ProfileName.IsNone();
+}
+
+namespace UE::AvaBroadcastProfile::Private
+{
+	// Since the channel index is serialized (now), we will initialize it to the
+	// channel's profile index only if it hasn't been set already.
+	inline void ConditionalSetChannelIndex(FAvaOutputChannel& InChannel, int32 InIndex)
+	{
+		if (InChannel.GetChannelIndex() == INDEX_NONE)
+		{
+			InChannel.SetChannelIndex(InIndex);
+		}
+	}
+}
+
+void FAvaBroadcastProfile::PostLoadProfile(bool bInIsProfileActive, UAvalancheBroadcast* InBroadcast)
+{
+	ParentBroadcastWeak = InBroadcast;
+	int32 Index = 0;
+	for (FAvaOutputChannel& Channel : Channels)
+	{
+		UE::AvaBroadcastProfile::Private::ConditionalSetChannelIndex(Channel, Index++);
+		Channel.PostLoadMediaOutputs(bInIsProfileActive, this);
+	}
+
+	ResolveChannels();
+}
+
+void FAvaBroadcastProfile::UpdateChannels(bool bInIsProfileActive)
+{
+	int32 Index = 0;
+	for (FAvaOutputChannel& Channel : Channels)
+	{
+		UE::AvaBroadcastProfile::Private::ConditionalSetChannelIndex(Channel, Index++);
+	}
+
+	ResolveChannels();
+
+	// Remark: include the pinned channels when updating resources.
+	for (FAvaOutputChannel* Channel : ResolvedChannels)
+	{
+		Channel->UpdateChannelResources(bInIsProfileActive);
+	}
+}
+
+void FAvaBroadcastProfile::ResolveChannels()
+{
+	const UAvalancheBroadcast& Broadcast = GetBroadcast();
+	ResolvedChannels.Empty(Broadcast.GetChannelNameCount());
+
+	for (int32 ChannelIndex = 0; ChannelIndex < Broadcast.GetChannelNameCount(); ++ChannelIndex)
+	{
+		const FName ChannelName = Broadcast.GetChannelName(ChannelIndex);
+		FAvaOutputChannel& Channel = GetChannelMutable(ChannelName);
+		if (Channel.IsValidChannel())
+		{
+			ResolvedChannels.Add(&Channel);
+		}
+	}
+}
+
+FAvaOutputChannel& FAvaBroadcastProfile::AddChannel(FName InChannelName)
+{
+	UAvalancheBroadcast& Broadcast = GetBroadcast();
+
+	int32 ChannelNameIndex;
+	if (!InChannelName.IsNone())
+	{
+		// If a name is provided, added it in the broadcast channel names.
+		// It may exist already, in which case the existing index is used.
+		ChannelNameIndex = Broadcast.AddChannelName(InChannelName);	
+	}
+	else
+	{
+		// Determine the last name index in this profile.
+		int32 LastChannelNameIndex = -1;
+		for (const FAvaOutputChannel& Channel : Channels)
+		{
+			if (Channel.GetChannelIndex() > LastChannelNameIndex)
+			{
+				LastChannelNameIndex = Channel.GetChannelIndex();
+			}
+		}
+		
+		ChannelNameIndex = LastChannelNameIndex + 1;
+
+		// Add the name to the broadcast if needed.
+		Broadcast.GetOrAddChannelName(ChannelNameIndex);
+	}
+
+	FAvaOutputChannel& Channel = Channels.Add_GetRef(FAvaOutputChannel(this));
+	Channel.SetChannelIndex(ChannelNameIndex);
+	Channel.UpdateChannelResources(Broadcast.GetCurrentProfileName() == GetName());
+	
+	Broadcast.UpdateChannelNames();
+	ResolveChannels();
+
+	Broadcast.QueueNotifyChange(EAvaBroadcastChange::ChannelGrid);
+	Broadcast.GetOnChannelsListChanged().Broadcast(*this);
+	
+	return Channel;
+}
+
+bool FAvaBroadcastProfile::RemoveChannel(FName InChannelName)
+{
+	UAvalancheBroadcast& Broadcast = GetBroadcast();
+	
+	const int32 ChannelIndex = GetLocalChannelIndexInProfile(InChannelName);
+
+	if (ChannelIndex == INDEX_NONE)
+	{
+		UE_LOG(LogAvaBroadcast, Error,
+			TEXT("Can't remove channel \"%s\" from profile \"%s\", not found in profile."),
+			*InChannelName.ToString(), *ProfileName.ToString());
+		return false;
+	}
+
+	if (!Channels.IsValidIndex(ChannelIndex))
+	{
+		UE_LOG(LogAvaBroadcast, Error,
+			TEXT("Can't remove channel \"%s\" (index %d) from profile \"%s\": invalid index."),
+			*InChannelName.ToString(), ChannelIndex, *ProfileName.ToString());
+		return false;
+	}
+	
+	Channels.RemoveAt(ChannelIndex);
+	Broadcast.UpdateChannelNames();	// Will update channel name indices in all profiles.
+
+	// Handle removing a pinned channel.
+	if (Broadcast.GetPinnedChannelProfileName(InChannelName) == GetName())
+	{
+		Broadcast.UnpinChannel(InChannelName);
+		Broadcast.RebuildProfiles();	// Channel names must be updated before calling this.
+	}
+	else
+	{
+		// If the channel is not pinned, resolve only this profile's channels.
+		ResolveChannels();
+	}
+
+	Broadcast.QueueNotifyChange(EAvaBroadcastChange::ChannelGrid);
+	Broadcast.GetOnChannelsListChanged().Broadcast(*this);
+	return true;
+}
+
+int32 FAvaBroadcastProfile::GetLocalChannelIndexInProfile(FName InChannelName) const
+{
+	return Channels.IndexOfByPredicate([InChannelName](const FAvaOutputChannel& InChannel)
+	{
+		return InChannel.GetChannelName() == InChannelName;
+	});
+}
+
+int32 FAvaBroadcastProfile::GetChannelIndexInProfile(FName InChannelName) const
+{
+	return ResolvedChannels.IndexOfByPredicate([InChannelName](const FAvaOutputChannel* InChannel)
+	{
+		return InChannel->GetChannelName() == InChannelName;
+	});
+}
+
+const FAvaOutputChannel& FAvaBroadcastProfile::GetLocalChannel(FName InChannelName) const
+{
+	const int32 ChannelIndex = GetLocalChannelIndexInProfile(InChannelName);
+	
+	if (Channels.IsValidIndex(ChannelIndex))
+	{
+		return Channels[ChannelIndex];
+	}
+	
+	return FAvaOutputChannel::NullChannel;
+}
+
+FAvaOutputChannel& FAvaBroadcastProfile::GetLocalChannelMutable(FName InChannelName)
+{
+	const int32 ChannelIndex = GetLocalChannelIndexInProfile(InChannelName);
+	
+	if (Channels.IsValidIndex(ChannelIndex))
+	{
+		return Channels[ChannelIndex];
+	}
+	
+	return FAvaOutputChannel::NullChannel;
+}
+
+const FAvaOutputChannel& FAvaBroadcastProfile::GetChannel(FName InChannelName) const
+{
+	UAvalancheBroadcast& Broadcast = GetBroadcast();
+	const FName PinnedChannelProfileName = Broadcast.GetPinnedChannelProfileName(InChannelName);
+	if (PinnedChannelProfileName != NAME_None)
+	{
+		return Broadcast.GetProfile(PinnedChannelProfileName).GetLocalChannel(InChannelName);
+	}	
+	return GetLocalChannel(InChannelName); 
+}
+
+FAvaOutputChannel& FAvaBroadcastProfile::GetChannelMutable(FName InChannelName)
+{
+	UAvalancheBroadcast& Broadcast = GetBroadcast();
+	const FName PinnedChannelProfileName = Broadcast.GetPinnedChannelProfileName(InChannelName);
+	if (PinnedChannelProfileName != NAME_None)
+	{
+		return Broadcast.GetProfile(PinnedChannelProfileName).GetLocalChannelMutable(InChannelName);
+	}	
+	return GetLocalChannelMutable(InChannelName);
+}
+
+FAvaOutputChannel& FAvaBroadcastProfile::GetOrAddChannel(FName InChannelName)
+{
+	FAvaOutputChannel& Channel = GetChannelMutable(InChannelName);
+	return Channel.IsValidChannel() ? Channel : AddChannel(InChannelName);
+}
+
+UMediaOutput* FAvaBroadcastProfile::AddChannelMediaOutput(FName InChannelName, const UClass* InMediaOutputClass, const FAvaMediaOutputInfo& InOutputInfo)
+{
+	FAvaOutputChannel& Channel = GetChannelMutable(InChannelName);
+	if (Channel.IsValidChannel())
+	{
+		return Channel.AddMediaOutput(InMediaOutputClass, InOutputInfo);
+	}
+	return nullptr;
+}
+
+int32 FAvaBroadcastProfile::RemoveChannelMediaOutputs(FName InChannelName, const TArray<UMediaOutput*>& InMediaOutputs)
+{
+	FAvaOutputChannel& Channel = GetChannelMutable(InChannelName);
+	if (Channel.IsValidChannel())
+	{
+		int32 RemovedCount = 0;
+		for (UMediaOutput* const MediaOutput : InMediaOutputs)
+		{
+			RemovedCount += Channel.RemoveMediaOutput(MediaOutput);
+		}		
+		return RemovedCount;
+	}
+	return 0;
+}
+
+#undef LOCTEXT_NAMESPACE
\ No newline at end of file
