[MotionDesign]

diff --git a/Engine/Plugins/Experimental/Avalanche/Source/AvalancheMedia/Private/AvaMediaModule.cpp b/Engine/Plugins/Experimental/Avalanche/Source/AvalancheMedia/Private/AvaMediaModule.cpp
new file mode 100644
--- /dev/null
+++ b/Engine/Plugins/Experimental/Avalanche/Source/AvalancheMedia/Private/AvaMediaModule.cpp
@@ -0,0 +1,731 @@
+// Copyright Epic Games, Inc. All Rights Reserved.
+
+#include "AvaMediaModule.h"
+
+#include "Application/ThrottleManager.h"
+#include "AvaMediaStyle.h"
+#include "AvalancheMediaSettings.h"
+#include "IMediaIOCoreModule.h"
+#include "Interfaces/IPluginManager.h"
+#include "Misc/CommandLine.h"
+#include "Misc/CoreDelegates.h"
+#include "Playback/AvaMediaPlaybackClientDelegates.h"
+#include "ShaderCore.h"
+
+#if WITH_EDITOR
+#include "Editor.h"
+#endif
+
+DEFINE_LOG_CATEGORY(LogAvaMedia);
+
+namespace UE::AvaMediaModule::Private
+{
+	// Command line parsing helper.
+	bool IsPlaybackServerManuallyStarted(FString& OutPlaybackServerName)
+	{
+		return FParse::Value(FCommandLine::Get(),TEXT("AvaMediaPlaybackServerStart="), OutPlaybackServerName) ||
+		 FParse::Param(FCommandLine::Get(), TEXT("AvaMediaPlaybackServerStart"));
+	}
+}
+
+FAvaMediaModule::FAvaMediaModule()
+: BroadcastSettingsBridge(this)
+{}
+
+void FAvaMediaModule::StartupModule()
+{
+	using namespace UE::AvaMediaModule::Private;
+	
+	FAvaMediaStyle::Initialize();
+	const TSharedPtr<IPlugin> Plugin = IPluginManager::Get().FindPlugin(TEXT("Avalanche"));
+	check(Plugin.IsValid());
+
+	const FString PluginShaderDir = FPaths::Combine(Plugin->GetBaseDir(), TEXT("Shaders"));
+	AddShaderSourceDirectoryMapping(TEXT("/Plugin/Avalanche"), PluginShaderDir);
+
+	IMediaIOCoreModule::Get().RegisterDeviceProvider(&AvaDisplayDeviceProvider);
+
+	ConsoleCmds.Add(IConsoleManager::Get().RegisterConsoleCommand(
+				TEXT("AvaMediaServer.Start"),
+				TEXT("Starts the playback server."),
+				FConsoleCommandWithArgsDelegate::CreateRaw(this, &FAvaMediaModule::StartPlaybackServerCommand),
+				ECVF_Default
+				));
+	ConsoleCmds.Add(IConsoleManager::Get().RegisterConsoleCommand(
+				TEXT("AvaMediaServer.Stop"),
+				TEXT("Stops the playback server."),
+				FConsoleCommandWithArgsDelegate::CreateRaw(this, &FAvaMediaModule::StopPlaybackServerCommand),
+				ECVF_Default
+				));
+	ConsoleCmds.Add(IConsoleManager::Get().RegisterConsoleCommand(
+				TEXT("AvaMediaClient.Start"),
+				TEXT("Starts the playback client."),
+				FConsoleCommandWithArgsDelegate::CreateRaw(this, &FAvaMediaModule::StartPlaybackClientCommand),
+				ECVF_Default
+				));
+	ConsoleCmds.Add(IConsoleManager::Get().RegisterConsoleCommand(
+				TEXT("AvaMediaClient.Stop"),
+				TEXT("Stops the playback client."),
+				FConsoleCommandWithArgsDelegate::CreateRaw(this, &FAvaMediaModule::StopPlaybackClientCommand),
+				ECVF_Default
+				));
+	ConsoleCmds.Add(IConsoleManager::Get().RegisterConsoleCommand(
+				TEXT("AvaMediaLocalServer.Launch"),
+				TEXT("Launches the local playback server."),
+				FConsoleCommandWithArgsDelegate::CreateRaw(this, &FAvaMediaModule::LaunchLocalPlaybackServerCommand),
+				ECVF_Default
+				));
+	ConsoleCmds.Add(IConsoleManager::Get().RegisterConsoleCommand(
+				TEXT("AvaMediaLocalServer.Stop"),
+				TEXT("Stops the local playback server."),
+				FConsoleCommandWithArgsDelegate::CreateRaw(this, &FAvaMediaModule::StopLocalPlaybackServerCommand),
+				ECVF_Default
+				));
+	ConsoleCmds.Add(IConsoleManager::Get().RegisterConsoleCommand(
+				TEXT("AvaMediaHttpServer.Start"),
+				TEXT("Starts the http playback server."),
+				FConsoleCommandWithArgsDelegate::CreateRaw(this, &FAvaMediaModule::StartHttpPlaybackServerCommand),
+				ECVF_Default
+				));
+	ConsoleCmds.Add(IConsoleManager::Get().RegisterConsoleCommand(
+				TEXT("AvaMediaHttpServer.Stop"),
+				TEXT("Stops the http playback server."),
+				FConsoleCommandWithArgsDelegate::CreateRaw(this, &FAvaMediaModule::StopHttpPlaybackServerCommand),
+				ECVF_Default
+				));
+	ConsoleCmds.Add(IConsoleManager::Get().RegisterConsoleCommand(
+				TEXT("AvaMediaDevices.Save"),
+				TEXT("Save Device Providers data."),
+				FConsoleCommandWithArgsDelegate::CreateRaw(this, &FAvaMediaModule::SaveDeviceProvidersCommand),
+				ECVF_Default
+				));
+	ConsoleCmds.Add(IConsoleManager::Get().RegisterConsoleCommand(
+				TEXT("AvaMediaDevices.Load"),
+				TEXT("Load device providers."),
+				FConsoleCommandWithArgsDelegate::CreateRaw(this, &FAvaMediaModule::LoadDeviceProvidersCommand),
+				ECVF_Default
+				));
+	ConsoleCmds.Add(IConsoleManager::Get().RegisterConsoleCommand(
+				TEXT("AvaMediaDevices.Unload"),
+				TEXT("Load device providers."),
+				FConsoleCommandWithArgsDelegate::CreateRaw(this, &FAvaMediaModule::UnloadDeviceProvidersCommand),
+				ECVF_Default
+				));
+	ConsoleCmds.Add(IConsoleManager::Get().RegisterConsoleCommand(
+				TEXT("AvaMediaDevices.List"),
+				TEXT("List device providers."),
+				FConsoleCommandWithArgsDelegate::CreateRaw(this, &FAvaMediaModule::ListDeviceProvidersCommand),
+				ECVF_Default
+				));
+	ConsoleCmds.Add(IConsoleManager::Get().RegisterConsoleCommand(
+				TEXT("AvaStat"),
+				TEXT("Enable engine performance statistics. Same as 'stat' command but will affect Avalanche Playback outputs and propagate to connected servers."),
+				FConsoleCommandWithArgsDelegate::CreateRaw(this, &FAvaMediaModule::HandleStatCommand),
+				ECVF_Default
+				));
+
+	FCoreDelegates::OnPostEngineInit.AddRaw(this, &FAvaMediaModule::PostEngineInit);
+	FCoreDelegates::OnEnginePreExit.AddRaw(this, &FAvaMediaModule::EnginePreExit);
+	
+	FString DummyServerName;
+	if (!AvaMediaPlaybackServer.IsValid() && IsPlaybackServerManuallyStarted(DummyServerName))
+	{
+		// Prevent throttling when server is started.
+		// This has to be done before any SLevelViewport are ticked since the cvar value is cached on first tick.
+		static const FSlateThrottleManager & ThrottleManager = FSlateThrottleManager::Get();
+		if (IConsoleVariable* AllowThrottling = IConsoleManager::Get().FindConsoleVariable(TEXT("Slate.bAllowThrottling")))
+		{
+			AllowThrottling->Set(0);
+			UE_LOG(LogAvaMedia, Log, TEXT("Setting Slate.bAllowThrottling to false."));
+		}
+	}
+
+	AvaMediaSync = MakeUnique<FAvaMediaSync>();
+
+	// StormSyncAvaBridge has some issues with servers in game mode. This option
+	// allows us to disable some of it until all the bugs are fixed.
+	if (FParse::Param(FCommandLine::Get(), TEXT("DisableAvaMediaSync")))
+	{
+		AvaMediaSync->SetFeatureEnabled(false);
+	}
+}
+
+void FAvaMediaModule::ShutdownModule()
+{
+	StopAllServices();
+	
+	FCoreDelegates::OnPostEngineInit.RemoveAll(this);
+	FCoreDelegates::OnEnginePreExit.RemoveAll(this);
+	
+	FAvaMediaStyle::Shutdown();
+
+	if (IMediaIOCoreModule::IsAvailable())
+	{
+		IMediaIOCoreModule::Get().UnregisterDeviceProvider(&AvaDisplayDeviceProvider);
+	}
+
+	for (IConsoleObject* ConsoleCmd : ConsoleCmds)
+	{
+		IConsoleManager::Get().UnregisterConsoleObject(ConsoleCmd);
+	}
+	ConsoleCmds.Empty();
+
+	AvaMediaSync.Reset();
+}
+
+void FAvaMediaModule::StartMediaPlaybackClient()
+{
+	StartPlaybackClientCommand({});
+}
+
+void FAvaMediaModule::StopMediaPlaybackClient()
+{
+	StopPlaybackClientCommand({});
+}
+
+void FAvaMediaModule::StartMediaPlaybackServer(const FString& InPlaybackServerName)
+{
+	StartPlaybackServerCommand({InPlaybackServerName});
+}
+
+void FAvaMediaModule::StopMediaPlaybackServer()
+{
+	StopPlaybackServerCommand({});
+}
+
+class FAvaMediaPlaybackClientDummy : public IAvaMediaPlaybackClient
+{
+	static const FString EmptyString;
+	static const TArray<FString> EmptyStringArray;
+public:
+	//IAvaMediaPlaybackClient
+	virtual int32 GetNumConnectedServers() const override { return 0; }
+	virtual TArray<FString> GetServerNames() const override { return TArray<FString>();}
+	virtual FMessageAddress GetServerAddress(const FString& InServerName) const override
+	{
+		FMessageAddress InvalidAddress;
+		InvalidAddress.Invalidate();
+		return InvalidAddress;
+	}
+	virtual bool HasServerUserData(const FString& InServerName, const FString& InKey) const override { return false;}
+	virtual const FString& GetServerUserData(const FString& InServerName, const FString& InKey) const override { return EmptyString; }
+	virtual bool HasUserData(const FString& InKey) const override { return false;}
+	virtual const FString& GetUserData(const FString& InKey) const override { return EmptyString; }
+	virtual void SetUserData(const FString& InKey, const FString& InData) override {}
+	virtual void RemoveUserData(const FString& InKey) override {}
+	virtual void BroadcastStatCommand(const FString& InCommand, bool bInBroadcastLocalState) override {}
+	virtual void RequestPlaybackAssetStatus(const FSoftObjectPath& InAssetPath, const FString& InChannelOrServerName, bool bInForceRefresh) override {}
+	virtual void RequestPlayback(const FGuid& InInstanceId, const FSoftObjectPath& InAssetPath, const FString& InChannelName, EAvaMediaPlaybackAction InAction, const FString& InArguments) override {}
+	virtual void RequestAnimPlayback(const FGuid& InInstanceId, const FSoftObjectPath& InAssetPath, const FString& InChannelName, const FAnimPlaySettings& InAnimSettings) override {};
+	virtual void RequestAnimAction(const FGuid& InInstanceId, const FSoftObjectPath& InAssetPath, const FString& InChannelName, const FString& InAnimationName, EAvaMediaAnimAction InAction) override {};
+	virtual void RequestRemoteControlUpdate(const FGuid& InInstanceId, const FSoftObjectPath& InAssetPath, const FString& InChannelName, const FAvalancheRemoteControlValues& InRemoteControlValues) override {};
+	virtual void RequestPlayableTransitionStart(const FGuid& InTransitionId, TArray<FGuid>&& InEnterInstanceIds, TArray<FGuid>&& InPlayingInstanceIds, TArray<FGuid>&& InExitInstanceIds, TArray<FAvalancheRemoteControlValues>&& InEnterValues, const FName& InChannelName, EAvalanchePlayableTransitionFlags InTransitionFlags) override {}
+	virtual void RequestPlayableTransitionStop(const FGuid& InTransitionId, const FName& InChannelName) override {}
+	virtual void RequestBroadcast(const FString& InProfile, const FName& InChannel, const TArray<UMediaOutput*>& InRemoteMediaOutputs, EAvaMediaBroadcastAction InAction) override {}
+	virtual bool IsMediaOutputRemoteFallback(const UMediaOutput* InMediaOutput) override { return false;}
+	virtual EAvaMediaIssueSeverity GetMediaOutputIssueSeverity(const FString& InServerName, const FString& InChannelName, const FGuid& InOutputGuid) const override { return EAvaMediaIssueSeverity::None;}
+	virtual const TArray<FString>& GetMediaOutputIssueMessages(const FString& InServerName, const FString& InChannelName, const FGuid& InOutputGuid) const override { return EmptyStringArray; }
+	virtual EAvaMediaOutputState GetMediaOutputState(const FString& InServerName, const FString& InChannelName, const FGuid& InOutputGuid) const override { return EAvaMediaOutputState::Error; }
+	virtual bool HasAnyServerOnlineForChannel(const FName& InChannelName) const override { return false;}
+   	virtual TOptional<EAvaMediaPlaybackStatus> GetRemotePlaybackStatus(const FGuid& InInstanceId, const FSoftObjectPath& InAssetPath, const FString& InChannelName, const FString& InServerName) const override
+	{
+		return EAvaMediaPlaybackStatus::Unknown;
+	}
+	virtual const FString* GetRemotePlaybackUserData(const FGuid& InInstanceId, const FSoftObjectPath& InAssetPath, const FString& InChannelName, const FString& InServerName) const override
+	{
+		return nullptr;
+	}
+	virtual TOptional<EAvaMediaPlaybackAssetStatus> GetRemotePlaybackAssetStatus(const FSoftObjectPath& InAssetPath, const FString& InServerName) const override
+	{
+		return EAvaMediaPlaybackAssetStatus::Unknown;
+	}
+	//~IAvaMediaPlaybackClient
+};
+const FString FAvaMediaPlaybackClientDummy::EmptyString;
+const TArray<FString> FAvaMediaPlaybackClientDummy::EmptyStringArray;
+
+IAvaMediaPlaybackClient& FAvaMediaModule::GetMediaPlaybackClient()
+{
+	if (AvaMediaPlaybackClient.IsValid())
+	{
+		return *AvaMediaPlaybackClient;
+	}
+	else
+	{
+		static FAvaMediaPlaybackClientDummy DummyClient;
+		return DummyClient;
+	}
+}
+
+const IMediaIOCoreDeviceProvider* FAvaMediaModule::GetDeviceProvider(FName InProviderName, const FMediaIOOutputConfiguration* InMediaIOOutputConfiguration) const
+{
+	if (AvaMediaPlaybackClient.IsValid() && InMediaIOOutputConfiguration)
+	{
+		const FAvaDeviceProviderWrapper* Wrapper = DeviceProviderProxyManager.GetDeviceProviderWrapper(InProviderName);
+		if (Wrapper)
+		{
+			return Wrapper->GetProviderForDeviceName(InMediaIOOutputConfiguration->MediaConfiguration.MediaConnection.Device.DeviceName);
+		}
+	}
+	return IMediaIOCoreModule::Get().GetDeviceProvider(InProviderName);
+}
+
+TArray<const IMediaIOCoreDeviceProvider*> FAvaMediaModule::GetDeviceProvidersForServer(const FString& InServerName) const
+{
+	if (AvaMediaPlaybackClient.IsValid())
+	{
+		return DeviceProviderProxyManager.GetDeviceProvidersForServer(InServerName);
+	}
+	else
+	{
+		return TArray<const IMediaIOCoreDeviceProvider*>();
+	}
+}
+
+FString FAvaMediaModule::GetServerNameForDevice(const FName& InDeviceProviderName, const FName& InDeviceName) const
+{
+	// (New Method) Search in the device provider proxies.
+	// This method doesn't assume the device name starts with the server name.
+	FString FoundServerName = DeviceProviderProxyManager.FindServerNameForDevice(InDeviceProviderName, InDeviceName);
+	if (!FoundServerName.IsEmpty())
+	{
+		return FoundServerName;
+	}
+
+	// Fallback (legacy method)
+	// Assumes the device name starts with the server name.
+	if (AvaMediaPlaybackClient.IsValid())
+	{
+		const FString DeviceName = InDeviceName.ToString();
+		TArray<FString> ServerNames = AvaMediaPlaybackClient->GetServerNames();
+		for (const FString& ServerName : ServerNames)
+		{
+			if (DeviceName.StartsWith(ServerName))
+			{
+				return ServerName;
+			}
+		}
+	}
+	
+	return FString();
+}
+
+bool FAvaMediaModule::IsLocalDevice(const FName& InDeviceProviderName, const FName& InDeviceName) const
+{
+	return DeviceProviderProxyManager.IsLocalDevice(InDeviceProviderName, InDeviceName);
+}
+
+void FAvaMediaModule::LaunchGameModeLocalPlaybackServer()
+{
+	// In order to talk to the server, we need the client.
+	if (!IsMediaPlaybackClientStarted())
+	{
+		// For now, we issue an error.
+		UE_LOG(LogAvaMedia, Error, TEXT("Playback Client must be started prior to starting a local playback server."));
+		return;
+
+		// TODO: We could automatically start the playback client.
+		// However, we would need to wait for the first round of server discovery to complete prior to
+		// search for the "local" server process. There are no events for this currently.
+		//StartPlaybackClientCommand({});
+	}
+
+	if (!LocalPlaybackServerProcess.IsValid())
+	{
+		LocalPlaybackServerProcess = FAvaMediaPlaybackServerProcess::FindOrCreate(*AvaMediaPlaybackClient);
+	}
+
+	if (LocalPlaybackServerProcess.IsValid())
+	{
+		if (!LocalPlaybackServerProcess->IsLaunched())
+		{
+			LocalPlaybackServerProcess->Launch();
+		}
+	}
+}
+
+void FAvaMediaModule::StopGameModeLocalPlaybackServer()
+{
+	// We may not have a local handle to the server process if it was started by
+	// another client instance.
+	if (!LocalPlaybackServerProcess.IsValid() && AvaMediaPlaybackClient.IsValid())
+	{
+		LocalPlaybackServerProcess = FAvaMediaPlaybackServerProcess::Find(*AvaMediaPlaybackClient);
+	}
+	
+	if (LocalPlaybackServerProcess.IsValid() && LocalPlaybackServerProcess->IsLaunched())
+	{
+		LocalPlaybackServerProcess->Stop();
+		UE_LOG(LogAvaMedia, Log, TEXT("Local playback server has been stopped."));
+	}
+
+	LocalPlaybackServerProcess.Reset();
+}
+
+bool FAvaMediaModule::IsGameModeLocalPlaybackServerLaunched() const
+{
+	return LocalPlaybackServerProcess.IsValid() && LocalPlaybackServerProcess->IsLaunched();
+}
+
+const IAvalancheBroadcastSettings& FAvaMediaModule::GetBroadcastSettings() const
+{
+	return BroadcastSettingsBridge;
+}
+
+const FAvalancheInstanceSettings& FAvaMediaModule::GetAvalancheInstanceSettings() const
+{
+	// if the server is enabled, fetch the settings from the currently connected client.
+	if (AvaMediaPlaybackServer.IsValid())
+	{
+		if (const FAvalancheInstanceSettings* SettingsFromClient = AvaMediaPlaybackServer->GetAvalancheInstanceSettings())
+		{
+			return *SettingsFromClient;
+		}
+	}
+	// Return the local settings.
+	return UAvalancheMediaSettings::Get().AvalancheInstanceSettings;
+}
+
+FAvaMediaPlaybackManager& FAvaMediaModule::GetLocalPlaybackManager() const
+{
+	check(LocalPlaybackManager.IsValid());
+	return *LocalPlaybackManager;
+}
+
+FAvalancheManagedInstanceCache& FAvaMediaModule::GetManagedInstanceCache() const
+{
+	check(ManagedInstanceCache.IsValid());
+	return *ManagedInstanceCache;
+}
+
+bool FAvaMediaModule::IsAvaMediaSyncProviderFeatureAvailable() const
+{
+	return AvaMediaSync.IsValid() && AvaMediaSync->IsFeatureAvailable(); 
+}
+
+IAvaMediaSyncProvider* FAvaMediaModule::GetAvaMediaSyncProvider() const
+{
+	check(AvaMediaSync.IsValid());
+	return AvaMediaSync->GetCurrentProvider();
+}
+
+void FAvaMediaModule::NotifyMapChangedEvent(UWorld* InWorld, EAvaMediaMapChangeType InEventType)
+{
+	OnMapChangedEvent.Broadcast(InWorld, InEventType);
+}
+
+void FAvaMediaModule::PostEngineInit()
+{
+	using namespace UE::AvaMediaModule::Private;
+
+	// This needs to happen late in the loading process, otherwise it fails.
+	const UAvalancheMediaSettings& Settings = UAvalancheMediaSettings::Get();
+
+	ManagedInstanceCache = MakeShared<FAvalancheManagedInstanceCache>();
+	
+	// Allow for specification of the server name in the command line.
+	// Command line has priority over project settings. 
+	FString PlaybackServerName = Settings.PlaybackServerName;
+	const bool bIsServerManuallyStarted = IsPlaybackServerManuallyStarted(PlaybackServerName);
+	const bool bIsClientManuallyStarted = FParse::Param(FCommandLine::Get(), TEXT("AvaMediaPlaybackClientStart"));
+
+	// Adding a command to suppress the client from auto-starting. This is used when spawning
+	// extra server process from the same project, while preventing extra clients.
+	const bool bIsClientAutoStartSuppressed = FParse::Param(FCommandLine::Get(), TEXT("AvaMediaPlaybackClientSuppress"));
+	
+	bool bShouldStartClient = bIsClientManuallyStarted || (Settings.bAutoStartPlaybackClient && !bIsClientAutoStartSuppressed && !IsRunningCommandlet());
+	bool bShouldStartServer = bIsServerManuallyStarted || (Settings.bAutoStartPlaybackServer && !IsRunningCommandlet());
+	
+	const bool bIsGameMode = !GIsEditor || IsRunningGame();
+
+	// The playback client and server can't both run in the same process.
+	// For the editor, we will keep the client and suppress the server.
+	// For the game, we will keep the server and suppress the client.
+
+	// In game mode, the auto start client is honored only if auto start server is not set.
+	// In editor mode, the auto start server is honored only if auto start client is not set.
+	
+	if (bShouldStartServer && bShouldStartClient)
+	{
+		if (bIsGameMode)
+		{
+			UE_LOG(LogAvaMedia, Log, TEXT("Auto start of Playback Client has been suppressed in game mode in favor of Playback Server."));
+			bShouldStartClient = false; // In game mode, the client is suppressed.
+
+		}
+		else
+		{
+			UE_LOG(LogAvaMedia, Log, TEXT("Auto start of Playback Server has been suppressed in editor mode in favor of Playback Client."));
+			bShouldStartServer = false; // In editor mode, the server is suppressed.
+		}
+	}
+	
+	if (bShouldStartClient)
+	{
+		StartPlaybackClientCommand({});
+	}
+
+	if (bShouldStartServer)
+	{
+		StartPlaybackServerCommand({PlaybackServerName});
+	}
+	
+	LocalPlaybackManager = MakeShared<FAvaMediaPlaybackManager>();
+
+#if WITH_EDITOR
+	// Capture Raw Ptr to avoid keeping ref count on capture
+	// Only Local Playback Manager should handle tear down for PIE End
+	FAvaMediaPlaybackManager* LocalPlaybackManagerRaw = LocalPlaybackManager.Get();
+	FEditorDelegates::PrePIEEnded.AddSPLambda(LocalPlaybackManagerRaw, [LocalPlaybackManagerRaw](const bool)
+	{
+		LocalPlaybackManagerRaw->OnParentWorldBeginTearDown();
+	});
+#endif
+
+	// Playback server required by Http server
+	if (AvaMediaPlaybackServer.IsValid() && Settings.bAutoStartWebServer)
+	{
+		StartHttpPlaybackServerCommand({});
+	}
+}
+
+void FAvaMediaModule::EnginePreExit()
+{
+	StopAllServices();
+}
+
+void FAvaMediaModule::StopAllServices()
+{
+	StopPlaybackServerCommand({});
+	StopPlaybackClientCommand({});
+	
+	if (LocalPlaybackManager)
+	{
+		LocalPlaybackManager->StartShuttingDown();
+		LocalPlaybackManager->StopAllPlaybacks(true);	
+	}
+	LocalPlaybackManager.Reset();
+	ManagedInstanceCache.Reset();
+}
+
+void FAvaMediaModule::StartPlaybackServerCommand(const TArray<FString>& InArgs)
+{
+	// Starting a playback server in the same process as playback client is forbidden.
+	if (AvaMediaPlaybackClient)
+	{
+		UE_LOG(LogAvaMedia, Error, TEXT("A Playback Server can't be started in the same process as a Playback Client."));
+		return;
+	}
+	
+	if (!AvaMediaPlaybackServer)
+	{
+		AvaMediaPlaybackServer = MakeShared<FAvaMediaPlaybackServer>();
+		AvaMediaPlaybackServer->Init(InArgs.Num() > 0 ? InArgs[0] : TEXT(""));
+		OnAvaMediaPlaybackServerStarted.Broadcast();
+		UE_LOG(LogAvaMedia, Log, TEXT("Playback Server Started"));
+	}
+}
+
+void FAvaMediaModule::StopPlaybackServerCommand(const TArray<FString>& InArgs)
+{
+	if (AvaMediaPlaybackServer)
+	{
+		AvaMediaPlaybackServer->StartShuttingDown();
+		AvaMediaPlaybackServer->StopBroadcast();
+		AvaMediaPlaybackServer->StopPlaybacks();
+		OnAvaMediaPlaybackServerStopped.Broadcast();
+	}
+	AvaMediaPlaybackServer.Reset();
+}
+
+void FAvaMediaModule::StartPlaybackClientCommand(const TArray<FString>& InArgs)
+{
+	// Starting a playback server in the same process as playback client is forbidden.
+	if (AvaMediaPlaybackServer)
+	{
+		UE_LOG(LogAvaMedia, Error, TEXT("A Playback Client can't be started in the same process as a Playback Server."));
+		return;
+	}
+
+	if (!AvaMediaPlaybackClient)
+	{
+		using namespace UE::AvaMediaPlaybackClient::Delegates;
+		if (GetOnConnectionEvent().IsBoundToObject(this))
+		{
+			GetOnConnectionEvent().AddRaw(this, &FAvaMediaModule::OnAvaMediaPlaybackClientConnectionEvent);
+		}
+
+		AvaMediaPlaybackClient = MakeShared<FAvaMediaPlaybackClient>(this);
+		AvaMediaPlaybackClient->Init();
+		OnAvaMediaPlaybackClientStarted.Broadcast();
+		UE_LOG(LogAvaMedia, Log, TEXT("Playback Client Started"));
+	}
+}
+
+void FAvaMediaModule::StopPlaybackClientCommand(const TArray<FString>& InArgs)
+{
+	if (AvaMediaPlaybackClient)
+	{
+		OnAvaMediaPlaybackClientStopped.Broadcast();
+	}
+	AvaMediaPlaybackClient.Reset();
+	UE::AvaMediaPlaybackClient::Delegates::GetOnConnectionEvent().RemoveAll(this);
+}
+
+void FAvaMediaModule::StartHttpPlaybackServerCommand(const TArray<FString>& InArgs)
+{
+	if (!AvaMediaHttpPlaybackServer)
+	{
+		AvaMediaHttpPlaybackServer = MakeShared<FAvaMediaHttpServer>();
+		
+		if (InArgs.Num() > 0)
+		{
+			// ...
+		}
+	}
+	
+	if (!AvaMediaHttpPlaybackServer->IsRunning())
+	{
+		const int32 Port = GetDefault<UAvalancheMediaSettings>()->HttpServerPort;
+		AvaMediaHttpPlaybackServer->Start(Port);
+		UE_LOG(LogAvaMedia, Log, TEXT("Http Playback Server Started"));
+	}
+}
+
+void FAvaMediaModule::StopHttpPlaybackServerCommand(const TArray<FString>& InArgs)
+{
+	AvaMediaHttpPlaybackServer.Reset();
+}
+
+
+void FAvaMediaModule::SaveDeviceProvidersCommand(const TArray<FString>& InArgs)
+{
+	FAvaDeviceProviderDataList OutProviders;
+	OutProviders.Populate(InArgs.Num() > 0 ? InArgs[0] : FPlatformProcess::ComputerName());
+	OutProviders.SaveToJson();	// Saves in the project's config folder.
+	OutProviders.SaveToXml();	// Saves in the project's config folder.
+}
+
+void FAvaMediaModule::LoadDeviceProvidersCommand(const TArray<FString>& InArgs)
+{
+	DeviceProviderProxyManager.TestInstall();
+}
+
+void FAvaMediaModule::UnloadDeviceProvidersCommand(const TArray<FString>& InArgs)
+{
+	DeviceProviderProxyManager.TestUninstall();
+}
+
+void FAvaMediaModule::ListDeviceProvidersCommand(const TArray<FString>& InArgs)
+{
+	DeviceProviderProxyManager.ListAllProviders();
+}
+
+void FAvaMediaModule::HandleStatCommand(const TArray<FString>& InArgs)
+{
+	if (InArgs.IsEmpty())
+	{
+		UE_LOG(LogAvaMedia, Error, TEXT("Stat Command: No arguments specified."));
+		return;
+	}
+
+	const bool bLocalCommandSucceeded = LocalPlaybackManager->HandleStatCommand(InArgs);
+	
+	if (AvaMediaPlaybackClient.IsValid())
+	{
+		AvaMediaPlaybackClient->BroadcastStatCommand(InArgs[0], bLocalCommandSucceeded);
+	}
+}
+
+void FAvaMediaModule::OnAvaMediaPlaybackClientConnectionEvent(IAvaMediaPlaybackClient& InPlaybackClient,
+		const UE::AvaMediaPlaybackClient::Delegates::FConnectionEventArgs& InArgs)
+{
+	using namespace UE::AvaMediaPlaybackClient::Delegates;
+	// When a playback server connection event occurs, update the status of the playback server process.
+	switch (InArgs.Event)
+	{
+	case EConnectionEvent::ServerConnected:
+		if (!IsGameModeLocalPlaybackServerLaunched() && IsMediaPlaybackClientStarted())
+		{
+			LocalPlaybackServerProcess = FAvaMediaPlaybackServerProcess::Find(*AvaMediaPlaybackClient);
+		}
+		break;
+	case EConnectionEvent::ServerDisconnected:
+		// Clean up the process handle.
+		if (LocalPlaybackServerProcess.IsValid() && !LocalPlaybackServerProcess->IsLaunched())
+		{
+			LocalPlaybackServerProcess.Reset();
+		}
+		break;
+	}
+}
+
+const FLinearColor& FAvaMediaModule::FLocalBroadcastSettings::GetChannelClearColor() const
+{
+	return UAvalancheMediaSettings::Get().ChannelClearColor;
+}
+
+EPixelFormat FAvaMediaModule::FLocalBroadcastSettings::GetDefaultPixelFormat() const
+{
+	return UAvalancheMediaSettings::Get().ChannelDefaultPixelFormat;
+}
+
+const FIntPoint& FAvaMediaModule::FLocalBroadcastSettings::GetDefaultResolution() const
+{
+	return UAvalancheMediaSettings::Get().ChannelDefaultResolution;
+}
+
+bool FAvaMediaModule::FLocalBroadcastSettings::IsDrawPlaceholderWidget() const
+{
+	return UAvalancheMediaSettings::Get().bDrawPlaceholderWidget;
+}
+
+const FSoftObjectPath& FAvaMediaModule::FLocalBroadcastSettings::GetPlaceholderWidgetClass() const
+{
+	return UAvalancheMediaSettings::Get().PlaceholderWidgetClass.ToSoftObjectPath();
+}
+
+const FLinearColor& FAvaMediaModule::FBroadcastSettingsBridge::GetChannelClearColor() const
+{
+	return GetSettings().GetChannelClearColor();
+}
+
+EPixelFormat FAvaMediaModule::FBroadcastSettingsBridge::GetDefaultPixelFormat() const
+{
+	return GetSettings().GetDefaultPixelFormat();
+}
+
+const FIntPoint& FAvaMediaModule::FBroadcastSettingsBridge::GetDefaultResolution() const
+{
+	return GetSettings().GetDefaultResolution();
+}
+
+bool FAvaMediaModule::FBroadcastSettingsBridge::IsDrawPlaceholderWidget() const
+{
+	return GetSettings().IsDrawPlaceholderWidget();
+}
+
+const FSoftObjectPath& FAvaMediaModule::FBroadcastSettingsBridge::GetPlaceholderWidgetClass() const
+{
+	return GetSettings().GetPlaceholderWidgetClass();
+}
+
+const IAvalancheBroadcastSettings& FAvaMediaModule::FBroadcastSettingsBridge::GetSettings() const
+{
+	// if server is enabled, fetch the setting from the currently connected client.
+	if (ParentModule->AvaMediaPlaybackServer.IsValid())
+	{
+		if (const IAvalancheBroadcastSettings* SettingsFromClient = ParentModule->AvaMediaPlaybackServer->GetBroadcastSettings())
+		{
+			return *SettingsFromClient;
+		}
+	}
+	return ParentModule->LocalBroadcastSettings;
+}
+
+IAvaDeviceProviderProxyManager& FAvaMediaModule::GetDeviceProviderProxyManager()
+{
+	return DeviceProviderProxyManager;
+}
+
+IMPLEMENT_MODULE(FAvaMediaModule, AvalancheMedia)
\ No newline at end of file
diff --git a/Engine/Plugins/Experimental/Avalanche/Source/AvalancheMedia/Private/AvaMediaModule.h b/Engine/Plugins/Experimental/Avalanche/Source/AvalancheMedia/Private/AvaMediaModule.h
new file mode 100644
--- /dev/null
+++ b/Engine/Plugins/Experimental/Avalanche/Source/AvalancheMedia/Private/AvaMediaModule.h
@@ -0,0 +1,154 @@
+// Copyright Epic Games, Inc. All Rights Reserved.
+
+#pragma once
+
+#include "Http/AvaMediaHttpServer.h"
+#include "IAvaMediaModule.h"
+#include "IAvalancheBroadcastSettings.h"
+#include "ModularFeature/AvaMediaSync.h"
+#include "ModuleDescriptor.h"
+#include "OutputDevices/AvaDeviceProviderProxy.h"
+#include "OutputDevices/AvaDisplayDeviceProvider.h"
+#include "Playback/AvaMediaPlaybackClient.h"
+#include "Playback/AvaMediaPlaybackClientDelegates.h"
+#include "Playback/AvaMediaPlaybackManager.h"
+#include "Playback/AvaMediaPlaybackServer.h"
+#include "Playback/AvaMediaPlaybackServerProcess.h"
+#include "Playlist/AvalancheManagedInstanceCache.h"
+
+class FAvaMediaModule : public IAvaMediaModule
+{
+public:
+	FAvaMediaModule();
+
+	//~ Begin IAvaMediaModule
+	virtual void StartupModule() override;
+	virtual void ShutdownModule() override;
+	virtual bool IsMediaPlaybackClientStarted() const override { return AvaMediaPlaybackClient.IsValid();}
+	virtual void StartMediaPlaybackClient() override;
+	virtual void StopMediaPlaybackClient() override;
+	virtual bool IsMediaPlaybackServerStarted() const override { return AvaMediaPlaybackServer.IsValid();}
+	virtual void StartMediaPlaybackServer(const FString& InPlaybackServerName) override;
+	virtual void StopMediaPlaybackServer() override;
+	
+	virtual IAvaMediaPlaybackClient& GetMediaPlaybackClient() override;
+	virtual TSharedPtr<FAvaMediaPlaybackServer> GetMediaPlaybackServer() const override { return AvaMediaPlaybackServer; }
+	virtual const IMediaIOCoreDeviceProvider* GetDeviceProvider(FName InProviderName, const FMediaIOOutputConfiguration* InMediaIOOutputConfiguration) const override;
+	virtual TArray<const IMediaIOCoreDeviceProvider*> GetDeviceProvidersForServer(const FString& InServerName) const override;
+	virtual FString GetServerNameForDevice(const FName& InDeviceProviderName, const FName& InDeviceName) const override;
+	virtual bool IsLocalDevice(const FName& InDeviceProviderName, const FName& InDeviceName) const override;
+	virtual void LaunchGameModeLocalPlaybackServer() override;
+	virtual void StopGameModeLocalPlaybackServer() override;
+	virtual bool IsGameModeLocalPlaybackServerLaunched() const override;
+	virtual const IAvalancheBroadcastSettings& GetBroadcastSettings() const override;
+	virtual const FAvalancheInstanceSettings& GetAvalancheInstanceSettings() const override;
+	virtual FAvaMediaPlaybackManager& GetLocalPlaybackManager() const override;
+	virtual FAvalancheManagedInstanceCache& GetManagedInstanceCache() const override;
+	virtual bool IsAvaMediaSyncProviderFeatureAvailable() const override;
+	virtual IAvaMediaSyncProvider* GetAvaMediaSyncProvider() const override;
+	virtual void NotifyMapChangedEvent(UWorld* InWorld, EAvaMediaMapChangeType InEventType) override;
+	virtual FOnAvaMediaSyncProviderChanged& GetOnAvaMediaSyncProviderChanged() override { return OnAvaMediaSyncProviderChanged; }
+	virtual FOnAvaMediaSyncPackageModified& GetOnAvaMediaSyncPackageModified() override { return OnAvaMediaSyncPackageModified; }
+	virtual FOnMapChangedEvent& GetOnMapChangedEvent() override { return OnMapChangedEvent; }
+	virtual FOnAvaMediaPlaybackClientStarted& GetOnAvaMediaPlaybackClientStarted() override { return OnAvaMediaPlaybackClientStarted; }
+	virtual FOnAvaMediaPlaybackClientStopped& GetAvaMediaPlaybackClientStopped() override { return OnAvaMediaPlaybackClientStopped; }
+	virtual FOnAvaMediaPlaybackServerStarted& GetOnAvaMediaPlaybackServerStarted() override { return OnAvaMediaPlaybackServerStarted; }
+	virtual FOnAvaMediaPlaybackServerStopped& GetAvaMediaPlaybackServerStopped() override { return OnAvaMediaPlaybackServerStopped; }
+	virtual FGetEditorViewportClient& GetEditorViewportClientDelegate() override { return GetEditorViewportClient; }
+	virtual IAvaDeviceProviderProxyManager& GetDeviceProviderProxyManager() override;
+	//~ End IAvaMediaModule
+	
+private:
+	void PostEngineInit();
+	void EnginePreExit();
+	void StopAllServices();
+	
+	// Command handlers
+	void StartPlaybackServerCommand(const TArray<FString>& InArgs);
+	void StopPlaybackServerCommand(const TArray<FString>& InArgs);
+	void StartPlaybackClientCommand(const TArray<FString>& InArgs);
+	void StopPlaybackClientCommand(const TArray<FString>& InArgs);
+	void LaunchLocalPlaybackServerCommand(const TArray<FString>& InArgs) { LaunchGameModeLocalPlaybackServer(); }
+	void StopLocalPlaybackServerCommand(const TArray<FString>& InArgs) { StopGameModeLocalPlaybackServer(); }
+
+	void StartHttpPlaybackServerCommand(const TArray<FString>& InArgs);
+	void StopHttpPlaybackServerCommand(const TArray<FString>& InArgs);
+	
+	void SaveDeviceProvidersCommand(const TArray<FString>& InArgs);
+	void LoadDeviceProvidersCommand(const TArray<FString>& InArgs);
+	void UnloadDeviceProvidersCommand(const TArray<FString>& InArgs);
+	void ListDeviceProvidersCommand(const TArray<FString>& InArgs);
+	void HandleStatCommand(const TArray<FString>& InArgs);
+
+	// Event Handlers
+	void OnAvaMediaPlaybackClientConnectionEvent(IAvaMediaPlaybackClient& InPlaybackClient,
+		const UE::AvaMediaPlaybackClient::Delegates::FConnectionEventArgs& InArgs);
+
+private:
+	FAvaDisplayDeviceProvider AvaDisplayDeviceProvider;
+	FAvaDeviceProviderProxyManager DeviceProviderProxyManager;
+	TUniquePtr<FAvaMediaSync> AvaMediaSync;
+
+	TArray<IConsoleObject*> ConsoleCmds;
+	
+	/**
+	 *	Wraps the local default UAvalancheMediaSettings.
+	 */
+	class FLocalBroadcastSettings : public IAvalancheBroadcastSettings
+	{
+	public:
+		virtual ~FLocalBroadcastSettings() override = default;
+
+		//~ Begin IAvalancheBroadcastSettings
+		virtual const FLinearColor& GetChannelClearColor() const override;
+		virtual EPixelFormat GetDefaultPixelFormat() const override;
+		virtual const FIntPoint& GetDefaultResolution() const override;
+		virtual bool IsDrawPlaceholderWidget() const override;
+		virtual const FSoftObjectPath& GetPlaceholderWidgetClass() const override;
+		//~ End IAvalancheBroadcastSettings
+	};
+	FLocalBroadcastSettings LocalBroadcastSettings;
+
+	/**
+	 * The purpose of the settings bridge is to provide a valid object
+	 * even if a client disconnects, it will either switch to another client
+	 * or use the local settings.
+	 */
+	class FBroadcastSettingsBridge final : public IAvalancheBroadcastSettings
+	{
+	public:
+		explicit FBroadcastSettingsBridge(FAvaMediaModule* Module) : ParentModule(Module) {}
+		virtual ~FBroadcastSettingsBridge() override = default;
+
+		//~ Begin IAvalancheBroadcastSettings
+		virtual const FLinearColor& GetChannelClearColor() const override;
+		virtual EPixelFormat GetDefaultPixelFormat() const override;
+		virtual const FIntPoint& GetDefaultResolution() const override;
+		virtual bool IsDrawPlaceholderWidget() const override;
+		virtual const FSoftObjectPath& GetPlaceholderWidgetClass() const override;
+		//~ End IAvalancheBroadcastSettings
+
+	private:
+		const IAvalancheBroadcastSettings& GetSettings() const;
+		
+		FAvaMediaModule* ParentModule = nullptr;
+	};
+	FBroadcastSettingsBridge BroadcastSettingsBridge;
+	
+	TSharedPtr<FAvaMediaPlaybackServer> AvaMediaPlaybackServer;	
+	TSharedPtr<FAvaMediaPlaybackClient> AvaMediaPlaybackClient;
+	TSharedPtr<FAvaMediaPlaybackServerProcess> LocalPlaybackServerProcess;
+	TSharedPtr<FAvaMediaPlaybackManager> LocalPlaybackManager;
+	TSharedPtr<FAvalancheManagedInstanceCache> ManagedInstanceCache;
+
+	TSharedPtr<FAvaMediaHttpServer> AvaMediaHttpPlaybackServer;
+
+	FOnAvaMediaSyncProviderChanged OnAvaMediaSyncProviderChanged;
+	FOnAvaMediaSyncPackageModified OnAvaMediaSyncPackageModified;
+	FOnMapChangedEvent OnMapChangedEvent;
+	FOnAvaMediaPlaybackClientStarted OnAvaMediaPlaybackClientStarted;
+	FOnAvaMediaPlaybackClientStopped OnAvaMediaPlaybackClientStopped;
+	FOnAvaMediaPlaybackServerStarted OnAvaMediaPlaybackServerStarted;
+	FOnAvaMediaPlaybackServerStopped OnAvaMediaPlaybackServerStopped;
+	FGetEditorViewportClient GetEditorViewportClient;
+};
diff --git a/Engine/Plugins/Experimental/Avalanche/Source/AvalancheMedia/Private/AvaMediaStyle.cpp b/Engine/Plugins/Experimental/Avalanche/Source/AvalancheMedia/Private/AvaMediaStyle.cpp
new file mode 100644
--- /dev/null
+++ b/Engine/Plugins/Experimental/Avalanche/Source/AvalancheMedia/Private/AvaMediaStyle.cpp
@@ -0,0 +1,76 @@
+// Copyright Epic Games, Inc. All Rights Reserved.
+
+#include "AvaMediaStyle.h"
+#include "Brushes/SlateImageBrush.h"
+#include "Interfaces/IPluginManager.h"
+#include "Styling/SlateStyle.h"
+#include "Styling/SlateStyleRegistry.h"
+
+TSharedPtr<FSlateStyleSet> FAvaMediaStyle::StyleInstance = nullptr;
+
+void FAvaMediaStyle::Initialize()
+{
+	if (!StyleInstance.IsValid())
+	{
+		StyleInstance = Create();
+		FSlateStyleRegistry::RegisterSlateStyle(*StyleInstance);
+	}
+}
+
+void FAvaMediaStyle::Shutdown()
+{
+	FSlateStyleRegistry::UnRegisterSlateStyle(*StyleInstance);
+	ensure(StyleInstance.IsUnique());
+	StyleInstance.Reset();
+}
+
+FName FAvaMediaStyle::GetStyleSetName()
+{
+	static FName StyleSetName(TEXT("AvalancheMedia"));
+	return StyleSetName;
+}
+
+const FLinearColor& FAvaMediaStyle::GetColor(FName PropertyName, const ANSICHAR* Specifier)
+{
+	return StyleInstance->GetColor(PropertyName, Specifier);
+}
+
+const FSlateBrush* FAvaMediaStyle::GetBrush(FName PropertyName, const ANSICHAR* Specifier)
+{
+	return StyleInstance->GetBrush(PropertyName, Specifier);
+}
+
+#define IMAGE_BRUSH_SVG(RelativePath, ...) FSlateVectorImageBrush(Style->RootToContentDir(RelativePath, TEXT(".svg")), __VA_ARGS__)
+#define IMAGE_BRUSH(RelativePath, ...) FSlateImageBrush(Style->RootToContentDir(RelativePath, TEXT(".png")), __VA_ARGS__)
+#define BOX_BRUSH(RelativePath, ...) FSlateBoxBrush(Style->RootToContentDir(RelativePath, TEXT(".png")), __VA_ARGS__)
+#define BORDER_BRUSH(RelativePath, ...) FSlateBorderBrush(Style->RootToContentDir(RelativePath, TEXT(".png")), __VA_ARGS__)
+
+const FVector2D Icon16x16(16.0f, 16.0f);
+const FVector2D Icon200x200(200.f, 200.f);
+
+TSharedRef<FSlateStyleSet> FAvaMediaStyle::Create()
+{
+	TSharedRef<FSlateStyleSet> Style = MakeShared<FSlateStyleSet>("AvalancheMedia");
+	
+	TSharedPtr<IPlugin> Plugin = IPluginManager::Get().FindPlugin(TEXT("Avalanche"));
+	check(Plugin.IsValid());
+	if (Plugin.IsValid())
+	{
+		Style->SetContentRoot(FPaths::Combine(Plugin->GetBaseDir(), TEXT("Resources")));
+	}
+	
+	Style->Set("AvalancheMedia.UnrealIcon", new IMAGE_BRUSH("Icons/MediaIcons/ue_logo", Icon200x200));
+	Style->Set("ClassIcon.AvalancheBroadcast", new IMAGE_BRUSH("Icons/MediaIcons/MediaOutput", Icon16x16));
+
+	return Style;
+}
+
+#undef IMAGE_BRUSH_SVG
+#undef IMAGE_BRUSH
+#undef BOX_BRUSH
+#undef BORDER_BRUSH
+
+const ISlateStyle& FAvaMediaStyle::Get()
+{
+	return *StyleInstance;
+}
diff --git a/Engine/Plugins/Experimental/Avalanche/Source/AvalancheMedia/Private/AvaMediaStyle.h b/Engine/Plugins/Experimental/Avalanche/Source/AvalancheMedia/Private/AvaMediaStyle.h
new file mode 100644
--- /dev/null
+++ b/Engine/Plugins/Experimental/Avalanche/Source/AvalancheMedia/Private/AvaMediaStyle.h
@@ -0,0 +1,37 @@
+// Copyright Epic Games, Inc. All Rights Reserved.
+
+#pragma once
+
+#include "Styling/ISlateStyle.h"
+#include "Styling/SlateStyle.h"
+
+struct FSlateBrush;
+class UTexture2D;
+
+class FAvaMediaStyle
+{
+public:
+
+	static void Initialize();
+
+	static void Shutdown();
+	
+	static const ISlateStyle& Get();
+
+	static FName GetStyleSetName();
+
+	static const FLinearColor& GetColor(FName PropertyName, const ANSICHAR* Specifier = nullptr);
+
+	static const FSlateBrush* GetBrush(FName PropertyName, const ANSICHAR* Specifier = nullptr);
+
+	template< typename WidgetStyleType >
+	static const WidgetStyleType& GetWidgetStyle(FName PropertyName, const ANSICHAR* Specifier = nullptr)
+	{
+		return StyleInstance->GetWidgetStyle<WidgetStyleType>(PropertyName, Specifier);
+	}
+
+private:
+
+	static TSharedRef<FSlateStyleSet> Create();
+	static TSharedPtr<FSlateStyleSet> StyleInstance;
+};
diff --git a/Engine/Plugins/Experimental/Avalanche/Source/AvalancheMedia/Private/AvaRenderTargetUtils.cpp b/Engine/Plugins/Experimental/Avalanche/Source/AvalancheMedia/Private/AvaRenderTargetUtils.cpp
new file mode 100644
--- /dev/null
+++ b/Engine/Plugins/Experimental/Avalanche/Source/AvalancheMedia/Private/AvaRenderTargetUtils.cpp
@@ -0,0 +1,180 @@
+// Copyright Epic Games, Inc. All Rights Reserved.
+
+#include "AvaRenderTargetUtils.h"
+
+#include "Engine/Engine.h"
+#include "Engine/TextureRenderTarget2D.h"
+#include "UObject/Package.h"
+
+namespace AvaRenderTargetUtils
+{
+	UTextureRenderTarget2D* CreateDefaultRenderTarget(FName InBaseName)
+	{
+		const FName TargetName = MakeUniqueObjectName(GetTransientPackage(), UTextureRenderTarget2D::StaticClass(), InBaseName);
+		UTextureRenderTarget2D* const RenderTarget = NewObject<UTextureRenderTarget2D>(GetTransientPackage(), TargetName);
+		// Linear gamma would normally be true for everything but RGB8_SRGB format.
+		RenderTarget->bForceLinearGamma = false;
+		RenderTarget->bAutoGenerateMips = false;
+		return RenderTarget;
+	}
+
+	bool SetupRenderTargetGamma(UTextureRenderTarget2D* InRenderTarget)
+	{
+		// Note: In CreateDefaultRenderTarget, we set bForceLinearGamma to false, meaning,
+		// the output is supposed to be gammatized (i.e. non linear).
+		// However, bForceLinearGamma is ignored for floating point textures
+		// so we have to override by specifying a target gamma instead.
+		if (IsFloatFormat(InRenderTarget))
+		{
+			// This will force the float texture to be gammatized like the integer textures,
+			// thus making it consistent. For Ava rendering, the outputs are expecting a gammatized texture,
+			// regardless of it's format.
+			const float TargetGamma = GEngine ? GEngine->GetDisplayGamma() : 2.2f;
+			if (InRenderTarget->TargetGamma != TargetGamma)
+			{
+				InRenderTarget->TargetGamma = TargetGamma;
+				return true;
+			}
+		}
+		return false;
+	}
+	
+	void UpdateRenderTarget(UTextureRenderTarget2D* InRenderTarget, const FIntPoint& InSize, EPixelFormat InFormat, const FLinearColor& InClearColor)
+	{
+		check(InSize.X > 0 && InSize.Y > 0);
+		check(InFormat != PF_Unknown);
+
+		if (!ensure(InRenderTarget))
+		{
+			return;
+		}
+
+		bool bNeedInitFormat = InRenderTarget->OverrideFormat != InFormat
+			|| InRenderTarget->ClearColor != InClearColor;
+		
+		bNeedInitFormat |= SetupRenderTargetGamma(InRenderTarget);
+
+		if (bNeedInitFormat)
+		{
+			InRenderTarget->ClearColor = InClearColor;
+			InRenderTarget->InitCustomFormat(InSize.X, InSize.Y, InFormat, InRenderTarget->bForceLinearGamma);
+			InRenderTarget->UpdateResourceImmediate();	// Doesn't seem to be necessary, but doesn't hurt either.
+		}
+		else
+		{
+			InRenderTarget->ResizeTarget(InSize.X, InSize.Y);
+		}
+	}
+	
+	FIntPoint GetRenderTargetSize(const UTextureRenderTarget2D* InRenderTarget)
+	{
+		if (ensure(IsValid(InRenderTarget)))
+		{
+			return FIntPoint(InRenderTarget->SizeX, InRenderTarget->SizeY);
+		}
+		return FIntPoint(0,0);
+	}
+
+	inline bool IsFloatFormat(ETextureRenderTargetFormat InFormat)
+	{
+		return InFormat == RTF_R16f || InFormat == RTF_RG16f || InFormat == RTF_RGBA16f
+			|| InFormat == RTF_R32f || InFormat == RTF_RG32f || InFormat == RTF_RGBA32f;
+	}
+	
+	bool IsFloatFormat(const UTextureRenderTarget2D* InRenderTarget)
+	{
+		if (ensure(IsValid(InRenderTarget)))
+		{
+			if (InRenderTarget->OverrideFormat != PF_Unknown)
+			{
+				// Note: Assumes HDR is float. However, upon inspection, it is missing PF_G16R16F, PF_G32R32F and PF_FloatR11G11B10.
+				const EPixelFormat Format = InRenderTarget->GetFormat();
+				return IsHDR(Format) || Format == PF_G16R16F || Format == PF_G32R32F || Format == PF_FloatR11G11B10;
+			}
+			return IsFloatFormat(InRenderTarget->RenderTargetFormat);
+		}
+		return false;
+	}
+
+	inline bool IsSingleChannelAlpha(EPixelFormat InFormat)
+	{
+		return (InFormat == PF_A1 || InFormat == PF_A8);
+	}
+
+	int32 GetNumColorComponents(EPixelFormat InFormat)
+	{
+		const FPixelFormatInfo& FormatInfo = GPixelFormats[InFormat];
+		// Exceptions: alpha only format have no color components.
+		// (Haven't seen luminance-alpha formats)
+		if (IsSingleChannelAlpha(InFormat) || FormatInfo.Supported == false)
+		{
+			return 0;
+		}
+		
+		// Maximum of 3 color components.
+		return (FormatInfo.NumComponents > 3) ? 3 : FormatInfo.NumComponents;
+	}
+
+	// Only works for channel that have the same number of bits per channel for
+	// all channels.
+	// Some exceptions: PF_R5G6B5_UNORM, PF_FloatR11G11B10.
+	// But those are not render target formats.
+	inline int32 GetNumBitsPerChannel(EPixelFormat InFormat)
+	{
+		const FPixelFormatInfo& FormatInfo = GPixelFormats[InFormat];
+		if (FormatInfo.Supported == false || FormatInfo.NumComponents == 0)
+		{
+			return 0;
+		}
+		
+		const int32 PixelsPerBlock = FormatInfo.BlockSizeX * FormatInfo.BlockSizeY * FormatInfo.BlockSizeZ;
+		if (PixelsPerBlock == 0)
+		{
+			return 0;
+		}
+		return 8 * FormatInfo.BlockBytes / (PixelsPerBlock * FormatInfo.NumComponents);
+	}
+	
+	// Note: we only have to deal with possible RGBA render target formats.
+	int32 GetNumColorChannelBits(EPixelFormat InFormat)
+	{
+		// Exceptions:
+		// - alpha only format have no color components:
+		if (IsSingleChannelAlpha(InFormat))
+		{
+			return 0;
+		}
+		// - This format doesn't have the same number of bits for alpha and color.
+		if (InFormat == PF_A2B10G10R10)
+		{
+			return 10;
+		}
+		
+		return GetNumBitsPerChannel(InFormat);
+	}
+
+	bool HasAlpha(EPixelFormat InFormat)
+	{
+		// Exception: single channel alpha.
+		if (IsSingleChannelAlpha(InFormat))
+		{
+			return true;
+		}
+
+		// Rule: 4 components formats have alpha.
+		const FPixelFormatInfo& FormatInfo = GPixelFormats[InFormat];
+		return FormatInfo.Supported == true && FormatInfo.NumComponents == 4;
+	}
+	
+	// Note: we only have to deal with possible render target formats.
+	int32 GetNumAlphaChannelBits(EPixelFormat InFormat)
+	{
+		// Exceptions:
+		if (InFormat == PF_A2B10G10R10)
+		{
+			return 2;
+		}
+		// Rule:
+		return HasAlpha(InFormat) ? GetNumBitsPerChannel(InFormat) : 0;
+	}	
+}
diff --git a/Engine/Plugins/Experimental/Avalanche/Source/AvalancheMedia/Private/AvaSerializationUtils.cpp b/Engine/Plugins/Experimental/Avalanche/Source/AvalancheMedia/Private/AvaSerializationUtils.cpp
new file mode 100644
--- /dev/null
+++ b/Engine/Plugins/Experimental/Avalanche/Source/AvalancheMedia/Private/AvaSerializationUtils.cpp
@@ -0,0 +1,73 @@
+// Copyright Epic Games, Inc. All Rights Reserved.
+
+#include "AvaSerializationUtils.h"
+
+#include "Serialization/StructuredArchive.h"
+#include "UObject/Object.h"
+#include "UObject/UnrealType.h"
+
+namespace UE::AvaSerializationUtils::Private
+{
+	void SerializeObject(FStructuredArchiveRecord& InRootRecord, UObject* InObject, const TFunction<bool(FProperty*)>& InShouldSkipProperty)
+	{
+		if (!InObject)
+		{
+			return;
+		}
+	
+		FStructuredArchiveRecord ObjectRecord = InRootRecord.EnterRecord(*InObject->GetName());
+		FArchive& UnderlyingArchive = InRootRecord.GetUnderlyingArchive();
+	
+		for (FProperty* Property = InObject->GetClass()->PropertyLink; Property; Property = Property->PropertyLinkNext)
+		{
+			if (!Property || Property->HasAnyPropertyFlags(CPF_Transient) || (InShouldSkipProperty && InShouldSkipProperty(Property)))
+			{
+				continue;
+			}
+
+			FStructuredArchiveSlot PropertyField = ObjectRecord.EnterField(*Property->GetName());
+		
+			for (int32 Index = 0; Index < Property->ArrayDim; ++Index)
+			{
+				UObject* const Archetype = InObject->GetArchetype();
+			
+				void* const Data = Property->ContainerPtrToValuePtr<void>(InObject, Index);
+				const void* const Defaults
+					= Property->ContainerPtrToValuePtrForDefaults<void>(Archetype->GetClass(), Archetype, Index);
+			
+				//if (!Property->Identical(Target, Default, UnderlyingArchive.GetPortFlags()))
+				{
+					FSerializedPropertyScope SerializedProperty(UnderlyingArchive, Property);
+					Property->SerializeItem(PropertyField, Data, Defaults);
+				}
+			}
+		}
+	}
+}
+
+void UE::AvaSerializationUtils::SerializeObject(FStructuredArchiveFormatter& InFormatter, UObject* InObject, const TFunction<bool(FProperty*)>& InShouldSkipProperty)
+{
+#if WITH_TEXT_ARCHIVE_SUPPORT
+	FStructuredArchive StructuredArchive(InFormatter);
+	FStructuredArchiveRecord RootRecord = StructuredArchive.Open().EnterRecord();
+	Private::SerializeObject(RootRecord, InObject, InShouldSkipProperty);
+	StructuredArchive.Close();
+#endif
+}
+
+void UE::AvaSerializationUtils::JsonValueConversion::BytesToString(const TArray<uint8>& InValueAsBytes, FString& OutValueAsString)
+{
+	// Reinterpret as TCHAR array. FJsonStructSerializerBackend uses UCS2CHAR which is compatible with TCHAR.
+	const TCHAR* ValueAsChars = reinterpret_cast<const TCHAR*>(InValueAsBytes.GetData());
+	int32 SizeInChar = InValueAsBytes.Num() / sizeof(TCHAR);
+
+	// Exclude terminating character, if any.
+	// Remark: From observation, the arrays from FJsonStructSerializerBackend don't seem to have a terminating character, but we check anyway.
+	if (SizeInChar > 0 && ValueAsChars[SizeInChar-1] == 0)
+	{
+		--SizeInChar; // Exclude terminating character.
+	}
+
+	OutValueAsString.Reset(SizeInChar);
+	OutValueAsString.AppendChars(ValueAsChars, SizeInChar);
+}
\ No newline at end of file
diff --git a/Engine/Plugins/Experimental/Avalanche/Source/AvalancheMedia/Private/AvalancheBroadcast.cpp b/Engine/Plugins/Experimental/Avalanche/Source/AvalancheMedia/Private/AvalancheBroadcast.cpp
new file mode 100644
--- /dev/null
+++ b/Engine/Plugins/Experimental/Avalanche/Source/AvalancheMedia/Private/AvalancheBroadcast.cpp
@@ -0,0 +1,673 @@
+// Copyright Epic Games, Inc. All Rights Reserved.
+
+#include "AvalancheBroadcast.h"
+
+#include "Async/Async.h"
+#include "AvaSerializationUtils.h"
+#include "Formatters/XmlArchiveInputFormatter.h"
+#include "Formatters/XmlArchiveOutputFormatter.h"
+#include "HAL/FileManager.h"
+#include "IAvaMediaModule.h"
+#include "MediaOutput.h"
+#include "Misc/Paths.h"
+#include "UObject/ObjectSaveContext.h"
+#include "UObject/Package.h"
+
+DEFINE_LOG_CATEGORY(LogAvaBroadcast);
+
+#define LOCTEXT_NAMESPACE "AvalancheBroadcast"
+
+namespace UE::AvalancheBroadcast::Private
+{
+	// Tested with server in game and editor modes: This function is called only when the server receives some broadcast queries.
+	static FString GetXmlSaveFilepath()
+	{
+		FString BroadcastConfigName;
+
+		// Allowing command line specification of the broadcast configuration file name. This allows
+		// starting the same project (from a shared location) with different configurations.
+		if (FParse::Value(FCommandLine::Get(),TEXT("AvalancheBroadcastConfig="), BroadcastConfigName))
+		{
+			// Ensure it has an xml extension.
+			if (!FPaths::GetExtension(BroadcastConfigName).Equals(TEXT("xml"), ESearchCase::IgnoreCase))
+			{
+				// Appending extension instead of using SetExtension to respect
+				// dot naming convention for xml/yaml config files.
+				BroadcastConfigName += TEXT(".xml");
+			}
+		}
+		else
+		{
+			// When launching the server from the same project location, we want to avoid loading the same
+			// broadcast configuration as the client. The server needs a clean configuration.
+			const bool bIsServerRunning = (IAvaMediaModule::IsModuleLoaded() && IAvaMediaModule::Get().IsMediaPlaybackServerStarted());	
+			BroadcastConfigName = bIsServerRunning ? TEXT("AvalancheServerBroadcastConfig.xml") : TEXT("AvalancheBroadcastConfig.xml");
+		}
+		
+		return FPaths::ProjectConfigDir() / BroadcastConfigName;
+	}
+}
+
+UAvalancheBroadcast& UAvalancheBroadcast::Get()
+{
+	static UAvalancheBroadcast* Broadcast = nullptr;
+	
+	if (!IsValid(Broadcast))
+	{
+		static const TCHAR* PackageName = TEXT("/Temp/AvalancheMedia/AvalancheBroadcast");
+		
+		UPackage* const BroadcastPackage = CreatePackage(PackageName);
+		BroadcastPackage->SetFlags(RF_Transient);
+		BroadcastPackage->AddToRoot();
+
+		//Don't Mark as Transient for "marking package dirty"
+		Broadcast = NewObject<UAvalancheBroadcast>(BroadcastPackage
+			, TEXT("AvalancheBroadcast")
+			, RF_Transactional | RF_Standalone);
+
+		Broadcast->AddToRoot();
+#if WITH_EDITOR
+		Broadcast->LoadBroadcast();
+#endif
+	}
+	
+	check(Broadcast);
+	return *Broadcast;
+}
+
+void UAvalancheBroadcast::BeginDestroy()
+{
+	for (TPair<FName, FAvaBroadcastProfile>& Pair : Profiles)
+	{
+		Pair.Value.BeginDestroy();
+	}
+	Super::BeginDestroy();
+}
+
+UAvalancheBroadcast* UAvalancheBroadcast::GetAvalancheBroadcast()
+{
+	return &UAvalancheBroadcast::Get();
+}
+
+void UAvalancheBroadcast::StartBroadcast()
+{
+	if (IsBroadcastingAllChannels())
+	{
+		return;
+	}
+
+	FAvaBroadcastProfile& Profile = GetCurrentProfile();
+	Profile.StartChannelBroadcast();
+}
+
+void UAvalancheBroadcast::StopBroadcast()
+{
+	GetCurrentProfile().StopChannelBroadcast();
+}
+
+bool UAvalancheBroadcast::IsBroadcastingAnyChannel() const
+{
+	return GetCurrentProfile().IsBroadcastingAnyChannel();
+}
+
+bool UAvalancheBroadcast::IsBroadcastingAllChannels() const
+{
+	return GetCurrentProfile().IsBroadcastingAllChannels();
+}
+
+bool UAvalancheBroadcast::ConditionalStartBroadcastChannel(const FName& InChannelName)
+{
+	FAvaOutputChannel& Channel = GetCurrentProfile().GetChannelMutable(InChannelName);
+	if (!Channel.IsValidChannel())
+	{
+		UE_LOG(LogAvaBroadcast, Error,
+			TEXT("Start Broadcast failed: Channel \"%s\" of Profile \"%s\" is not valid."),
+			*InChannelName.ToString(), *GetCurrentProfileName().ToString());
+		return false;
+	}
+	
+	if (Channel.GetState() == EAvaChannelState::Idle)
+	{
+		Channel.StartChannelBroadcast();
+	}
+	return true;
+}
+
+
+TArray<FName> UAvalancheBroadcast::GetProfileNames() const
+{
+	TArray<FName> ProfileNames;
+	Profiles.GetKeys(ProfileNames);
+	return ProfileNames;
+}
+
+const TMap<FName, FAvaBroadcastProfile>& UAvalancheBroadcast::GetProfiles() const
+{
+	return Profiles;
+}
+
+FName UAvalancheBroadcast::CreateProfile(FName InProfileName, bool bMakeCurrentProfile)
+{
+	if (InProfileName == NAME_None)
+	{
+		InProfileName = FName(TEXT("Profile"), 0);
+	}
+	
+	FAvaBroadcastProfile& Profile = CreateProfileInternal(InProfileName);
+	
+	if (bMakeCurrentProfile)
+	{
+		SetCurrentProfile(Profile.GetName());
+	}
+	
+	return Profile.GetName();
+}
+
+bool UAvalancheBroadcast::DuplicateProfile(FName InNewProfile, FName InTemplateProfile, bool bMakeCurrentProfile)
+{
+	const FAvaBroadcastProfile* const TemplateProfile = Profiles.Find(InTemplateProfile);
+
+	if (TemplateProfile)
+	{		
+		if (InNewProfile == NAME_None)
+		{
+			InNewProfile = InTemplateProfile;
+		}
+
+		FAvaBroadcastProfile& Profile = CreateProfileInternal(InNewProfile);
+		FAvaBroadcastProfile::CopyProfiles(*TemplateProfile, Profile);
+		
+		if (bMakeCurrentProfile)
+		{
+			SetCurrentProfile(Profile.GetName());
+		}
+		
+		return true;
+	}
+	
+	return false;
+}
+
+bool UAvalancheBroadcast::DuplicateCurrentProfile(FName InProfileName, bool bMakeCurrentProfile)
+{
+	return DuplicateProfile(InProfileName, CurrentProfile, bMakeCurrentProfile);
+}
+
+bool UAvalancheBroadcast::RemoveProfile(FName InProfileName)
+{
+	const bool bRemovingCurrentProfile = CurrentProfile == InProfileName;
+	const bool bIsLastRemainingProfile = Profiles.Num() == 1;
+	const bool bIsBroadcasting = IsBroadcastingAnyChannel();
+	
+	//The only condition that would prevent us from doing Removal is if we're currently Broadcasting and we want to remove Current Profile.
+	const bool bCanRemoveProfile = !bIsLastRemainingProfile && !(bIsBroadcasting && bRemovingCurrentProfile);
+	
+	if (bCanRemoveProfile)
+	{
+		const int32 RemoveCount = Profiles.Remove(InProfileName);
+
+		//If Removing Current Profile, we need to find a new Current Profile
+		if (bRemovingCurrentProfile)
+		{
+			CurrentProfile = NAME_None;
+			EnsureValidCurrentProfile();
+		}
+		
+		return RemoveCount > 0;
+	}
+	return false;
+}
+
+bool UAvalancheBroadcast::CanRenameProfile(FName InProfileName, FName InNewProfileName, FText* OutErrorMessage) const
+{
+	if (InNewProfileName.IsNone())
+	{
+		if (OutErrorMessage)
+		{
+			*OutErrorMessage = LOCTEXT("RenameError_ProfileNone", "Invalid profile name.");
+		}
+		return false;
+	}
+
+	if (Profiles.Contains(InNewProfileName))
+	{
+		if (OutErrorMessage)
+		{
+			*OutErrorMessage = LOCTEXT("RenameError_ProfileExists", "Profile name already exists.");
+		}
+		return false;
+	}
+	
+	const bool bRenamingCurrentProfile = (CurrentProfile == InProfileName);
+	const bool bIsBroadcasting = IsBroadcastingAnyChannel();
+
+	if (bIsBroadcasting && bRenamingCurrentProfile)
+	{
+		if (OutErrorMessage)
+		{
+			*OutErrorMessage = LOCTEXT("RenameError_ProfileInUse", "Profile is currently Broadcasting Channels.");
+		}
+		return false;
+	}
+	
+	return true;
+}
+
+bool UAvalancheBroadcast::RenameProfile(FName InProfileName, FName InNewProfileName)
+{
+	if (CanRenameProfile(InProfileName, InNewProfileName))
+	{
+		FAvaBroadcastProfile Profile = MoveTemp(Profiles[InProfileName]);
+		Profiles.Remove(InProfileName);
+		Profile.ProfileName = InNewProfileName;
+		Profiles.Add(InNewProfileName, MoveTemp(Profile));
+		
+		if (CurrentProfile == InProfileName)
+		{
+			CurrentProfile = InNewProfileName;
+		}
+		return true;
+	}
+	return false;
+}
+
+bool UAvalancheBroadcast::SetCurrentProfile(FName InProfileName)
+{
+	const bool bIsBroadcasting = IsBroadcastingAnyChannel();
+	
+	//Can only set a new Current Profile if not Broadcasting.
+	if (!bIsBroadcasting && CurrentProfile != InProfileName && Profiles.Contains(InProfileName))
+	{
+		if (GetCurrentProfile().IsValidProfile())
+		{
+			// Deallocate previous profile's resources.
+			GetCurrentProfile().UpdateChannels(false);
+		}
+		
+		CurrentProfile = InProfileName;
+		
+		// Allocate new profile's resources.
+		GetCurrentProfile().UpdateChannels(true);
+		
+		QueueNotifyChange(EAvaBroadcastChange::CurrentProfile);
+		return true;
+	}
+	return false;
+}
+
+FAvaBroadcastProfile& UAvalancheBroadcast::GetProfile(FName InProfileName)
+{
+	if (Profiles.Contains(InProfileName))
+	{
+		return Profiles[InProfileName];
+	}
+	return FAvaBroadcastProfile::NullProfile;
+}
+
+const FAvaBroadcastProfile& UAvalancheBroadcast::GetProfile(FName InProfileName) const
+{
+	if (Profiles.Contains(InProfileName))
+	{
+		return Profiles[InProfileName];
+	}
+	return FAvaBroadcastProfile::NullProfile;
+}
+
+#if WITH_EDITOR
+void UAvalancheBroadcast::LoadBroadcast()
+{
+	using namespace UE::AvalancheBroadcast::Private;
+	TUniquePtr<FArchive> FileReader(IFileManager::Get().CreateFileReader(*GetXmlSaveFilepath()));
+
+	if (FileReader.IsValid())
+	{
+		FXmlArchiveInputFormatter InputFormatter(*FileReader, this);
+		UE::AvaSerializationUtils::SerializeObject(InputFormatter, this);
+		FileReader->Close();
+	}
+	
+	if (Profiles.Num() > 0)
+	{
+		for (TPair<FName, FAvaBroadcastProfile>& Profile : Profiles)
+		{
+			const bool bIsProfileActive = (Profile.Key == CurrentProfile);
+			Profile.Value.PostLoadProfile(bIsProfileActive, this);
+		}
+	}
+	else
+	{
+		CreateProfile(NAME_None, true);
+		SaveBroadcast();
+	}
+	
+	EnsureValidCurrentProfile();
+	UpdateProfileNames();
+}
+
+void UAvalancheBroadcast::SaveBroadcast()
+{
+	bool bIsBroadcastSaved = false;
+	using namespace UE::AvalancheBroadcast::Private;
+	TUniquePtr<FArchive> FileWriter(IFileManager::Get().CreateFileWriter(*GetXmlSaveFilepath()));
+	if (FileWriter.IsValid())	
+	{
+		FXmlArchiveOutputFormatter XmlOutput(*FileWriter);
+		XmlOutput.SerializeObjectsInPlace(true); // We want the media outputs nested.
+		UE::AvaSerializationUtils::SerializeObject(XmlOutput, this);
+		bIsBroadcastSaved = XmlOutput.SaveDocumentToInnerArchive();
+		GetPackage()->SetDirtyFlag(false);
+		FileWriter->Close();
+	}
+	if (!bIsBroadcastSaved)
+	{
+		UE_LOG(LogAvaBroadcast, Error, TEXT("Failed to save broadcast configuration.")); 
+	}
+}
+
+FString UAvalancheBroadcast::GetBroadcastSaveFilepath() const
+{
+	return UE::AvalancheBroadcast::Private::GetXmlSaveFilepath();
+}
+#endif
+
+void UAvalancheBroadcast::QueueNotifyChange(EAvaBroadcastChange InChange)
+{
+	if (InChange != EAvaBroadcastChange::None)
+	{		
+		const bool bCreateAsyncTask = (QueuedBroadcastChanges == EAvaBroadcastChange::None);
+		EnumAddFlags(QueuedBroadcastChanges, InChange);
+
+		if (bCreateAsyncTask)
+		{
+			TWeakObjectPtr<UAvalancheBroadcast> ThisWeak(this);
+			AsyncTask(ENamedThreads::GameThread, [ThisWeak]()
+			{
+				if (ThisWeak.IsValid())
+				{
+					ThisWeak->OnBroadcastChanged.Broadcast(ThisWeak->QueuedBroadcastChanges);
+					ThisWeak->QueuedBroadcastChanges = EAvaBroadcastChange::None;
+				}
+			});	
+		}
+	}
+}
+
+FDelegateHandle UAvalancheBroadcast::AddChangeListener(FOnAvaBroadcastChanged::FDelegate&& InDelegate)
+{
+	return OnBroadcastChanged.Add(MoveTemp(InDelegate));
+}
+
+void UAvalancheBroadcast::RemoveChangeListener(FDelegateHandle InDelegateHandle)
+{
+	OnBroadcastChanged.Remove(InDelegateHandle);
+}
+
+void UAvalancheBroadcast::RemoveChangeListener(const void* InUserObject)
+{
+	OnBroadcastChanged.RemoveAll(InUserObject);
+}
+
+int32 UAvalancheBroadcast::GetChannelNameCount() const
+{
+	return ChannelNames.Num();
+}
+
+int32 UAvalancheBroadcast::GetChannelIndex(FName InChannelName) const
+{
+	return ChannelNames.Find(InChannelName);
+}
+
+FName UAvalancheBroadcast::GetChannelName(int32 ChannelIndex) const
+{
+	if (ChannelNames.IsValidIndex(ChannelIndex))
+	{
+		return ChannelNames[ChannelIndex];
+	}
+	return NAME_None;
+}
+
+FName UAvalancheBroadcast::GetOrAddChannelName(int32 ChannelIndex)
+{
+	if (ChannelIndex < 0)
+	{
+		return NAME_None;
+	}
+	
+	if (ChannelIndex < ChannelNames.Num())
+	{
+		return ChannelNames[ChannelIndex];
+	}
+	
+	//Generate new items and set a unique names for them
+	{
+		//Store the current names as a Set for Fast Search
+		TSet ChannelNamesSet(ChannelNames);
+
+		//Add the New Items as Defaulted
+		const int32 OldItemCount = ChannelNames.Num();
+		ChannelNames.AddDefaulted(ChannelIndex - OldItemCount + 1);
+		ChannelNamesSet.Reserve(ChannelNames.Num());
+		
+		FName UniqueName = TEXT("Channel");
+		uint32 UniqueIndex = 1;
+		
+		for (int32 Index = OldItemCount; Index <= ChannelIndex; ++Index)
+		{
+			do
+			{
+				UniqueName.SetNumber(UniqueIndex++);
+			}
+			while (ChannelNamesSet.Contains(UniqueName));
+			
+			ChannelNamesSet.Add(UniqueName);
+			ChannelNames[Index] = UniqueName;
+		}
+	}
+	
+	return ChannelNames[ChannelIndex];
+}
+
+int32 UAvalancheBroadcast::AddChannelName(FName InChannelName)
+{
+	return ChannelNames.AddUnique(InChannelName);
+}
+
+TArray<int32> UAvalancheBroadcast::BuildChannelIndices() const
+{
+	TArray<int32> ChannelIndices;
+	ChannelIndices.Reserve(ChannelNames.Num());
+	for (const TPair<FName, FAvaBroadcastProfile>& Pair : Profiles)
+	{
+		for (const FAvaOutputChannel& Channel : Pair.Value.Channels)
+		{
+			ChannelIndices.AddUnique(Channel.GetChannelIndex());
+		}
+	}
+	return ChannelIndices;
+}
+
+void UAvalancheBroadcast::UpdateChannelNames()
+{
+	// This is called when a channel is added or removed from a profile.
+	// We have to reconcile all channel names from all the profiles.
+
+	// First pass, build the new channel names list.
+	//
+	// Different profile may have sub-set of channels.
+	// Ex:
+	// Profile 1: channel1, channel3
+	// Profile 2: channel1, channel2
+	//	
+	// Sorting by channel indices. The channel indices are
+	// used for the connections (pins) in the playback graph.
+	// So we want to preserve that order.
+	TArray<int32> ChannelIndices = BuildChannelIndices();
+	ChannelIndices.Sort();
+
+	// Build the new channel names list from the sorted channel indices.
+	TArray<FName> NewChannelNames;
+	NewChannelNames.Reserve(ChannelIndices.Num());
+	for (const int32 Index : ChannelIndices)
+	{
+		NewChannelNames.Add(GetOrAddChannelName(Index));
+	}
+	
+	if (NewChannelNames == ChannelNames)
+	{
+		return;
+	}
+	
+	// Update the channel indices in all profiles. 
+	for (TPair<FName, FAvaBroadcastProfile>& Pair : Profiles)
+	{
+		for (FAvaOutputChannel& Channel : Pair.Value.Channels)
+		{
+			Channel.SetChannelIndex(NewChannelNames.Find(Channel.GetChannelName()));
+		}
+	}
+
+	TArray<FName> RemovedNames;
+	for (const FName& PreviousChannelName : ChannelNames)
+	{
+		if (!NewChannelNames.Contains(PreviousChannelName))
+		{
+			RemovedNames.Add(PreviousChannelName);
+		}
+	}
+	
+	// We can finally update the ChannelNames array.
+	// Channel.GetChannelName() (above) was still using old ChannelNames to find new indices.
+	ChannelNames = NewChannelNames;
+
+	// Housekeeping for internal data:
+	// Remove ChannelType and PinnedChannels entries that where removed.
+	for (const FName& RemovedChannelName : RemovedNames)
+	{
+		ChannelTypes.Remove(RemovedChannelName);
+		PinnedChannels.Remove(RemovedChannelName);
+	}
+}
+
+bool UAvalancheBroadcast::CanRenameChannel(FName InChannelName, FName InNewChannelName) const
+{
+	//Make sure new Channel name is valid, and we don't have it already in the List
+	return InNewChannelName != NAME_None
+		&& InChannelName != InNewChannelName
+		&& true == ChannelNames.Contains(InChannelName)
+		&& false == ChannelNames.Contains(InNewChannelName);
+}
+
+bool UAvalancheBroadcast::RenameChannel(FName InChannelName, FName InNewChannelName)
+{
+	if (CanRenameChannel(InChannelName, InNewChannelName))
+	{
+		const int32 Index = GetChannelIndex(InChannelName);
+		check(ChannelNames.IsValidIndex(Index));
+		
+		ChannelNames[Index] = InNewChannelName;
+
+		if (const EAvaBroadcastChannelType* ExistingChannelType = ChannelTypes.Find(InChannelName))
+		{
+			const EAvaBroadcastChannelType ExistingChannelTypeCopy = *ExistingChannelType;
+			ChannelTypes.Remove(InChannelName);
+			ChannelTypes.Add(InNewChannelName, ExistingChannelTypeCopy);
+		}
+
+		if (const FName* ExistingPinnedProfileName = PinnedChannels.Find(InChannelName))
+		{
+			const FName ExistingPinnedProfileNameCopy = *ExistingPinnedProfileName;
+			PinnedChannels.Remove(InChannelName);
+			PinnedChannels.Add(InNewChannelName, ExistingPinnedProfileNameCopy);
+		}
+		
+		QueueNotifyChange(EAvaBroadcastChange::ChannelRename);
+		return true;
+	}
+	return false;
+}
+
+void UAvalancheBroadcast::SetChannelType(FName InChannelName, EAvaBroadcastChannelType InChannelType)
+{
+	ChannelTypes.Add(InChannelName, InChannelType);
+}
+
+EAvaBroadcastChannelType UAvalancheBroadcast::GetChannelType(FName InChannelName) const
+{
+	const EAvaBroadcastChannelType* ChannelType = ChannelTypes.Find(InChannelName);
+	// For backward compatibility, if the channel type is not set, defaults to Broadcast.
+	return ChannelType ? *ChannelType : EAvaBroadcastChannelType::Program;
+}
+
+void UAvalancheBroadcast::PinChannel(FName InChannelName, FName InProfileName)
+{
+	PinnedChannels.Add(InChannelName, InProfileName);
+}
+
+void UAvalancheBroadcast::UnpinChannel(FName InChannelName)
+{
+	PinnedChannels.Remove(InChannelName);
+}
+
+FName UAvalancheBroadcast::GetPinnedChannelProfileName(FName InChannelName) const
+{
+	const FName* ProfileName = PinnedChannels.Find(InChannelName);
+	return ProfileName ? *ProfileName : NAME_None;
+}
+
+void UAvalancheBroadcast::RebuildProfiles()
+{
+	for (TPair<FName, FAvaBroadcastProfile>& Profile : Profiles)
+	{
+		const bool bIsProfileActive = Profile.Key == CurrentProfile;
+		Profile.Value.UpdateChannels(bIsProfileActive);
+	}
+}
+
+#if WITH_EDITOR
+void UAvalancheBroadcast::PostEditUndo()
+{
+	UObject::PostEditUndo();
+	for (TPair<FName, FAvaBroadcastProfile>& Profile : Profiles)
+	{
+		const bool bIsProfileActive = Profile.Key == CurrentProfile;
+		Profile.Value.UpdateChannels(bIsProfileActive);
+	}
+	QueueNotifyChange(EAvaBroadcastChange::All);
+}
+#endif
+
+FAvaBroadcastProfile& UAvalancheBroadcast::CreateProfileInternal(FName InProfileName)
+{	
+	uint32 UniqueIndex = FMath::Max(1, InProfileName.GetNumber());
+	
+	while (Profiles.Contains(InProfileName))
+	{
+		InProfileName.SetNumber(++UniqueIndex);
+	};
+		
+	FAvaBroadcastProfile& Profile = Profiles.Add(InProfileName, {this, InProfileName});
+	Profile.AddChannel();
+	return Profile;
+}
+
+void UAvalancheBroadcast::EnsureValidCurrentProfile()
+{
+	if (CurrentProfile == NAME_None || !Profiles.Contains(CurrentProfile))
+	{
+		for (TPair<FName, FAvaBroadcastProfile>& Profile : Profiles)
+		{
+			SetCurrentProfile(Profile.Key);
+			break;
+		}
+	}
+}
+
+void UAvalancheBroadcast::UpdateProfileNames()
+{
+	for (TPair<FName, FAvaBroadcastProfile>& Pair : Profiles)
+	{
+		Pair.Value.ProfileName = Pair.Key;
+	}
+}
+
+#undef LOCTEXT_NAMESPACE
diff --git a/Engine/Plugins/Experimental/Avalanche/Source/AvalancheMedia/Private/AvalancheBroadcastComponent.cpp b/Engine/Plugins/Experimental/Avalanche/Source/AvalancheMedia/Private/AvalancheBroadcastComponent.cpp
new file mode 100644
--- /dev/null
+++ b/Engine/Plugins/Experimental/Avalanche/Source/AvalancheMedia/Private/AvalancheBroadcastComponent.cpp
@@ -0,0 +1,77 @@
+// Copyright Epic Games, Inc. All Rights Reserved.
+
+#include "AvalancheBroadcastComponent.h"
+
+#include "AvalancheBroadcast.h"
+#include "Engine/World.h"
+
+UAvalancheBroadcastComponent::UAvalancheBroadcastComponent(const FObjectInitializer& ObjectInitializer)
+	: Super(ObjectInitializer)
+{
+	bWantsInitializeComponent = true;
+}
+
+bool UAvalancheBroadcastComponent::StartBroadcasting(FString& OutErrorMessage)
+{
+	UAvalancheBroadcast::Get().StartBroadcast();
+	
+	// Build the error/warning messages from all outputs of all channels.
+	const TArray<FAvaOutputChannel*>& Channels = UAvalancheBroadcast::Get().GetCurrentProfile().GetChannels();
+	for (const FAvaOutputChannel* Channel : Channels)
+	{
+		if (Channel->GetIssueSeverity() != EAvaMediaIssueSeverity::None)
+		{
+			FString OutputMessages;
+			const TArray<UMediaOutput*>& Outputs = Channel->GetMediaOutputs();
+			for (const UMediaOutput* Output : Outputs)
+			{
+				const TArray<FString>& OutputIssueMessages = Channel->GetMediaOutputIssueMessages(Output);
+				const EAvaMediaIssueSeverity OutputSeverity = Channel->GetMediaOutputIssueSeverity(Channel->GetMediaOutputState(Output), Output);
+				const FString OutputIssue = StaticEnum<EAvaMediaIssueSeverity>()->GetDisplayValueAsText(OutputSeverity).ToString();
+				for (const FString& OutputIssueMessage : OutputIssueMessages)
+				{
+					OutputMessages += FString::Format(TEXT(" - Output {0}: {1}\n"),{OutputIssue,OutputIssueMessage});
+				}
+			}
+
+			if (!OutputMessages.IsEmpty())
+			{
+				OutErrorMessage += FString::Format(TEXT("{0} on Channel {1}: \n"),
+					{	StaticEnum<EAvaMediaIssueSeverity>()->GetDisplayValueAsText(Channel->GetIssueSeverity()).ToString(),
+						Channel->GetChannelName().ToString()});
+			}
+			
+			OutErrorMessage += OutputMessages;
+		}
+	}
+	
+	// Return true if any of the channels is broadcasting.
+	return UAvalancheBroadcast::Get().IsBroadcastingAnyChannel();
+}
+
+bool UAvalancheBroadcastComponent::StopBroadcasting()
+{
+	UAvalancheBroadcast::Get().StopBroadcast();
+	return true;
+}
+
+void UAvalancheBroadcastComponent::InitializeComponent()
+{
+	Super::InitializeComponent();
+	FWorldDelegates::OnWorldBeginTearDown.AddUObject(this, &UAvalancheBroadcastComponent::OnWorldBeginTearDown);
+}
+
+void UAvalancheBroadcastComponent::UninitializeComponent()
+{
+	FWorldDelegates::OnWorldBeginTearDown.RemoveAll(this);
+	Super::UninitializeComponent();
+}
+
+void UAvalancheBroadcastComponent::OnWorldBeginTearDown(UWorld* InWorld)
+{
+	if (InWorld == GetWorld() && bStopBroadcastOnTearDown)
+	{
+		UAvalancheBroadcast::Get().StopBroadcast();
+	}
+}
+
diff --git a/Engine/Plugins/Experimental/Avalanche/Source/AvalancheMedia/Private/AvalancheMediaSettings.cpp b/Engine/Plugins/Experimental/Avalanche/Source/AvalancheMedia/Private/AvalancheMediaSettings.cpp
new file mode 100644
--- /dev/null
+++ b/Engine/Plugins/Experimental/Avalanche/Source/AvalancheMedia/Private/AvalancheMediaSettings.cpp
@@ -0,0 +1,49 @@
+// Copyright Epic Games, Inc. All Rights Reserved.
+
+#include "AvalancheMediaSettings.h"
+
+UAvalancheMediaSettings::UAvalancheMediaSettings()
+{
+	CategoryName = TEXT("Motion Design");
+	SectionName  = TEXT("Playback & Broadcast");
+}
+
+UAvalancheMediaSettings* UAvalancheMediaSettings::GetSingletonInstance()
+{
+	UAvalancheMediaSettings* DefaultSettings = GetMutableDefault<UAvalancheMediaSettings>();
+	static bool bInitialized = false;
+	if (!bInitialized)
+	{
+		bInitialized = true;
+		DefaultSettings->SetFlags(RF_Transactional);
+	}
+	return DefaultSettings;
+}
+
+ELogVerbosity::Type UAvalancheMediaSettings::ToLogVerbosity(EAvaMediaLogVerbosity InAvaMediaLogVerbosity)
+{
+	switch (InAvaMediaLogVerbosity)
+	{
+	case EAvaMediaLogVerbosity::NoLogging:
+		return ELogVerbosity::NoLogging;
+	case EAvaMediaLogVerbosity::Fatal:
+		return ELogVerbosity::Fatal;
+	case EAvaMediaLogVerbosity::Error:
+		return ELogVerbosity::Error;
+	case EAvaMediaLogVerbosity::Warning:
+		return ELogVerbosity::Warning;
+	case EAvaMediaLogVerbosity::Display:
+		return ELogVerbosity::Display;
+	case EAvaMediaLogVerbosity::Log:
+		return ELogVerbosity::Log;
+	case EAvaMediaLogVerbosity::Verbose:
+		return ELogVerbosity::Verbose;
+	case EAvaMediaLogVerbosity::VeryVerbose:
+		return ELogVerbosity::VeryVerbose;
+	default:
+		return ELogVerbosity::NoLogging;
+	}
+}
+
+
+
diff --git a/Engine/Plugins/Experimental/Avalanche/Source/AvalancheMedia/Private/Channel/AvaMediaOutputInfo.cpp b/Engine/Plugins/Experimental/Avalanche/Source/AvalancheMedia/Private/Channel/AvaMediaOutputInfo.cpp
new file mode 100644
--- /dev/null
+++ b/Engine/Plugins/Experimental/Avalanche/Source/AvalancheMedia/Private/Channel/AvaMediaOutputInfo.cpp
@@ -0,0 +1,17 @@
+// Copyright Epic Games, Inc. All Rights Reserved.
+
+#include "Channel/AvaMediaOutputInfo.h"
+#include "OutputDevices/AvaDeviceProviderProxy.h"
+
+bool FAvaMediaOutputInfo::IsRemote(const FString& InServerName)
+{
+	return !InServerName.IsEmpty() && InServerName != FAvaDeviceProviderProxyManager::LocalServerName;
+}
+
+void FAvaMediaOutputInfo::PostLoad()
+{
+	if (!Guid.IsValid())
+	{
+		Guid = FGuid::NewGuid();
+	}
+}
