[MotionDesign]

diff --git a/Engine/Plugins/Experimental/Avalanche/Source/AvalancheMedia/Private/Channel/AvaOutputChannel.cpp b/Engine/Plugins/Experimental/Avalanche/Source/AvalancheMedia/Private/Channel/AvaOutputChannel.cpp
new file mode 100644
--- /dev/null
+++ b/Engine/Plugins/Experimental/Avalanche/Source/AvalancheMedia/Private/Channel/AvaOutputChannel.cpp
@@ -0,0 +1,1357 @@
+// Copyright Epic Games, Inc. All Rights Reserved.
+
+#include "Channel/AvaOutputChannel.h"
+
+#include "AvaRenderTargetUtils.h"
+#include "AvalancheBroadcast.h"
+#include "AvalancheMediaSettings.h"
+#include "Blueprint/UserWidget.h"
+#include "Blueprint/WidgetBlueprintGeneratedClass.h"
+#include "Channel/AvalancheOutputChannelViewInterface.h"
+#include "Engine/TextureRenderTarget2D.h"
+#include "Extensions/UserWidgetExtension.h"
+#include "Framework/Application/SlateApplication.h"
+#include "IAvaMediaModule.h"
+#include "IAvalancheBroadcastSettings.h"
+#include "MediaCapture.h"
+#include "MediaOutput.h"
+#include "MediaShaders.h"
+#include "OutputDevices/AvaMediaOutputUtils.h"
+#include "OutputDevices/AvaRenderTargetMediaUtils.h"
+#include "Playback/AvaMediaPlayableGroup.h"
+#include "Playback/IAvaMediaPlaybackClient.h"
+#include "RHICommandList.h"
+#include "RenderResource.h"
+#include "ScreenPass.h"
+#include "Slate/WidgetRenderer.h"
+#include "TextureResource.h"
+#include "UObject/Package.h"
+
+DEFINE_LOG_CATEGORY_STATIC(LogAvaOutputChannel, Log, All);
+
+void FAvaMediaCapture::Reset(UMediaCapture* InMediaCapture, UMediaOutput* InMediaOutput, UTextureRenderTarget2D* InRenderTarget)
+{
+	if (IsValid(MediaCapture))
+	{
+		MediaCapture->OnStateChangedNative.RemoveAll(this);
+	}
+	
+	OnMediaCaptureStateChanged.Clear();
+	MediaCapture = InMediaCapture;
+	MediaOutput = InMediaOutput;
+	RenderTarget = InRenderTarget;
+	if (IsValid(MediaCapture))
+	{
+		MediaCapture->OnStateChangedNative.AddRaw(this, &FAvaMediaCapture::OnStateChangedNative);
+	}
+}
+
+void FAvaMediaCapture::OnStateChangedNative()
+{
+	if (OnMediaCaptureStateChanged.IsBound())
+	{
+		OnMediaCaptureStateChanged.Broadcast(MediaCapture, MediaOutput);
+	}
+}
+
+namespace Avalanche
+{
+	bool IsCapturing(const UMediaCapture* InMediaCapture)
+	{
+		if (IsValid(InMediaCapture))
+		{
+			const EMediaCaptureState CaptureState = InMediaCapture->GetState();
+
+			return CaptureState == EMediaCaptureState::Capturing
+				|| CaptureState == EMediaCaptureState::Preparing;
+		}
+		return false;
+	}
+
+	static void AddInvertAlphaConversionPass(FRDGBuilder& GraphBuilder
+		, const FRDGTextureRef& SourceRGBTexture
+		, const FResolveRect& ResolveRect
+		, const FVector2D& SizeU
+		, const FVector2D& SizeV
+		, FRDGTextureRef OutputResource)
+	{
+		// Rectangle area to use from source
+		const FIntRect ViewRect(ResolveRect.X1, ResolveRect.Y1, ResolveRect.X2, ResolveRect.Y2);
+
+		//Dummy ViewFamily/ViewInfo created to use built in Draw Screen/Texture Pass
+		FSceneViewFamily ViewFamily(FSceneViewFamily::ConstructionValues(nullptr, nullptr, FEngineShowFlags(ESFIM_Game))
+			.SetTime(FGameTime()));
+		FSceneViewInitOptions ViewInitOptions;
+		ViewInitOptions.ViewFamily = &ViewFamily;
+		ViewInitOptions.SetViewRectangle(ViewRect);
+		ViewInitOptions.ViewOrigin = FVector::ZeroVector;
+		ViewInitOptions.ViewRotationMatrix = FMatrix::Identity;
+		ViewInitOptions.ProjectionMatrix = FMatrix::Identity;
+		FSceneView ViewInfo(ViewInitOptions);
+
+		//At some point we should support color conversion (ocio) but for now we push incoming texture as is
+		constexpr bool bDoLinearToSRGB = false;
+
+		FGlobalShaderMap* GlobalShaderMap = GetGlobalShaderMap(GMaxRHIFeatureLevel);
+		TShaderMapRef<FScreenPassVS> VertexShader(GlobalShaderMap);
+
+		//Configure source/output viewport to get the right UV scaling from source texture to output texture
+		FScreenPassTextureViewport InputViewport(SourceRGBTexture, ViewRect);
+		FScreenPassTextureViewport OutputViewport(OutputResource);
+
+		// In cases where texture is converted from a format that doesn't have A channel, we want to force set it to 1.
+		EMediaCaptureConversionOperation MediaConversionOperation = EMediaCaptureConversionOperation::INVERT_ALPHA;
+		FModifyAlphaSwizzleRgbaPS::FPermutationDomain PermutationVector;
+		PermutationVector.Set<FModifyAlphaSwizzleRgbaPS::FConversionOp>(static_cast<int32>(MediaConversionOperation));
+
+		TShaderMapRef<FModifyAlphaSwizzleRgbaPS> PixelShader(GlobalShaderMap, PermutationVector);
+		FModifyAlphaSwizzleRgbaPS::FParameters* Parameters = PixelShader->AllocateAndSetParameters(GraphBuilder, SourceRGBTexture, OutputResource);
+		AddDrawScreenPass(GraphBuilder, RDG_EVENT_NAME("AvaOutputInvertAlpha"), ViewInfo, OutputViewport, InputViewport, VertexShader, PixelShader, Parameters);
+	}
+
+	void ExecuteInvertAlphaConversionPass(FRHICommandListImmediate& RHICmdList, UTextureRenderTarget2D* SourceRT, UTextureRenderTarget2D* DestinationRT)
+	{
+		FRDGBuilder GraphBuilder(RHICmdList);
+
+		const FTexture2DRHIRef SourceTexture = SourceRT->GetRenderTargetResource()->GetTexture2DRHI();
+		const FRDGTextureRef SourceRGBTexture = GraphBuilder.RegisterExternalTexture(CreateRenderTarget(SourceTexture, TEXT("SourceTexture")));
+
+		const FTexture2DRHIRef DestinationTexture = DestinationRT->GetRenderTargetResource()->GetTexture2DRHI();
+		const FRDGTextureRef OutputResource = GraphBuilder.RegisterExternalTexture(CreateRenderTarget(DestinationTexture, TEXT("DestTexture")));
+
+		const FVector2D SizeU = { 0.0f, 1.0f };
+		const FVector2D SizeV = { 0.0f, 1.0f };
+
+		const FResolveRect ResolveRect(0, 0, SourceTexture->GetSizeX(), SourceTexture->GetSizeY());
+
+		AddInvertAlphaConversionPass(GraphBuilder, SourceRGBTexture, ResolveRect, SizeU, SizeV, OutputResource);
+
+		GraphBuilder.Execute();
+	}
+
+	namespace MediaOutputInfoHelper
+	{
+		/**
+		 * Legacy support helper.
+		 * Initialize the media output info from a given media output object.
+		 * Note: this will only work correctly if the device can be found in
+		 * a device provider. If the device is on a remote server, the server
+		 * must be connected when this function is called.
+		*/
+		FAvaMediaOutputInfo InitFrom(const UMediaOutput* InMediaOutput)
+		{
+			FAvaMediaOutputInfo OutputInfo;
+
+			if (::IsValid(InMediaOutput))
+			{
+				const FString DeviceNameString = UE::AvaMediaOutputUtils::GetDeviceName(InMediaOutput);
+				if (!DeviceNameString.IsEmpty())
+				{
+					const FName DeviceName = *DeviceNameString;
+					const FName DeviceProviderName = UE::AvaMediaOutputUtils::GetDeviceProviderName(InMediaOutput);
+					const FString ServerName = IAvaMediaModule::Get().GetServerNameForDevice(DeviceProviderName, DeviceName);
+					if (!ServerName.IsEmpty() || IAvaMediaModule::Get().IsLocalDevice(DeviceProviderName, DeviceName))
+					{
+						OutputInfo.DeviceName = DeviceName;
+						OutputInfo.DeviceProviderName = DeviceProviderName;
+						OutputInfo.ServerName = ServerName;
+					}
+				}
+			}
+			return OutputInfo;
+		}
+	}
+}
+
+FAvaOutputChannel FAvaOutputChannel::NullChannel;
+FAvaOutputChannel::FOnAvaChannelChanged FAvaOutputChannel::OnChannelChanged;
+FAvaOutputChannel::FMediaOutputStateChanged FAvaOutputChannel::OnMediaOutputStateChanged;
+
+FAvaOutputChannel::~FAvaOutputChannel()
+{
+	StopPlaceholderTick();
+
+#if WITH_EDITOR
+	for (UMediaOutput* MediaOutput : MediaOutputs)
+	{
+		MediaOutput->OnOutputModified().RemoveAll(this);
+	}
+#endif
+}
+
+void FAvaOutputChannel::ReleasePlaceholderResources()
+{
+	OverridePlaceholderWidget.Reset();
+	if (OverridePlaceholderUserWidget)
+	{
+		OverridePlaceholderUserWidget->ReleaseSlateResources(true);
+		OverridePlaceholderUserWidget = nullptr;
+	}
+}
+
+void FAvaOutputChannel::ReleasePlaceholderRenderTargets()
+{
+	PlaceholderRenderTarget = nullptr;
+	PlaceholderRenderTargetTmp = nullptr;
+}
+
+void FAvaOutputChannel::ReleaseOutputs()
+{
+	for (UMediaOutput* MediaOutput : MediaOutputs)
+	{
+		StopCaptureForOutput(MediaOutput);
+#if WITH_EDITOR
+		MediaOutput->OnOutputModified().RemoveAll(this);
+#endif
+	}
+	MediaOutputs.Reset();
+}
+
+void FAvaOutputChannel::DuplicateChannel(const FAvaOutputChannel& InSourceChannel, FAvaOutputChannel& OutTargetChannel)
+{
+	OutTargetChannel.MediaOutputs.Reset(InSourceChannel.MediaOutputs.Num());
+	for (const UMediaOutput* const SourceMediaOutput : InSourceChannel.MediaOutputs)
+	{
+		UMediaOutput* const TargetMediaOutput = DuplicateObject<UMediaOutput>(SourceMediaOutput
+			, SourceMediaOutput->GetOuter()
+			, NAME_None);
+
+		OutTargetChannel.MediaOutputs.Add(TargetMediaOutput);
+		OutTargetChannel.MediaOutputInfos.Add(InSourceChannel.GetMediaOutputInfo(SourceMediaOutput));
+		OutTargetChannel.MediaOutputInfos.Last().Guid = FGuid::NewGuid();	// Allocate a new guid for the duplicate.
+	}
+	OnChannelChanged.Broadcast(OutTargetChannel, EAvaChannelChange::MediaOutputs);
+	OutTargetChannel.UpdateChannelResources(false);
+}
+
+void FAvaOutputChannel::SetChannelIndex(int32 InIndex)
+{
+	ChannelIndex = InIndex;
+}
+
+FName FAvaOutputChannel::GetChannelName() const
+{
+	return GetProfile().GetBroadcast().GetOrAddChannelName(ChannelIndex);
+}
+
+EAvaBroadcastChannelType FAvaOutputChannel::GetChannelType() const
+{
+	UAvalancheBroadcast& Broadcast = GetProfile().GetBroadcast();
+	const FName ChannelName = Broadcast.GetOrAddChannelName(ChannelIndex);
+	return Broadcast.GetChannelType(ChannelName);
+}
+
+FName FAvaOutputChannel::GetProfileName() const
+{
+	// Note: during the migration towards a non-singleton UAvalancheBroadcast,
+	// a fallback path to the singleton is kept, but it has an ensure to try and catch
+	// the code paths that lead to this.
+	return ensure(Profile) ? Profile->GetName() : UAvalancheBroadcast::Get().GetCurrentProfile().GetName();
+}
+
+FAvaBroadcastProfile& FAvaOutputChannel::GetProfile() const
+{
+	// Note: during the migration towards a non-singleton UAvalancheBroadcast,
+	// a fallback path to the singleton is kept, but it has an ensure to try and catch
+	// the code paths that lead to this.
+	return ensure(Profile) ? *Profile : UAvalancheBroadcast::Get().GetCurrentProfile();
+}
+
+EAvaChannelState FAvaOutputChannel::RefreshState()
+{
+	// Resolve the channel state from current output states. (PROTOTYPE LOGIC)
+	EAvaChannelState NewChannelState = EAvaChannelState::Offline;
+	EAvaMediaIssueSeverity NewChannelIssueSeverity = EAvaMediaIssueSeverity::None;
+
+	for (const UMediaOutput* MediaOutput : GetMediaOutputs())
+	{
+		// Offline outputs don't contribute to state (i.e. channel will be offline too).
+		const EAvaMediaOutputState OutputState = GetMediaOutputState(MediaOutput);
+		if (OutputState == EAvaMediaOutputState::Live || OutputState == EAvaMediaOutputState::Preparing)
+		{
+			NewChannelState = EAvaChannelState::Live;
+		}
+		else if (OutputState == EAvaMediaOutputState::Idle && NewChannelState != EAvaChannelState::Live)
+		{
+			NewChannelState = EAvaChannelState::Idle;
+		}
+
+		const EAvaMediaIssueSeverity OutputIssueSeverity = GetMediaOutputIssueSeverity(OutputState, MediaOutput);
+		if (OutputIssueSeverity > NewChannelIssueSeverity)
+		{
+			NewChannelIssueSeverity = OutputIssueSeverity;
+		}
+	}
+	
+	// Note: this will broadcast to delegates which may then request states of media outputs.
+	// So the back store needs to be updated before calling this.
+	SetState(NewChannelState, NewChannelIssueSeverity);
+
+	return GetState();
+}
+
+EAvaMediaIssueSeverity FAvaOutputChannel::GetMediaOutputIssueSeverity(EAvaMediaOutputState InOutputState, const UMediaOutput* InMediaOutput) const
+{
+	if (InOutputState == EAvaMediaOutputState::Live || InOutputState == EAvaMediaOutputState::Preparing)
+	{
+		// If the output is broadcasting remote, fetch the status from the playback client (which is proxying that output's status).
+		if (IAvaMediaModule::Get().IsMediaPlaybackClientStarted() && IsMediaOutputRemote(InMediaOutput))
+		{
+			const FAvaMediaOutputInfo& OutputInfo = GetMediaOutputInfo(InMediaOutput);
+			const IAvaMediaPlaybackClient& PlaybackClient = IAvaMediaModule::Get().GetMediaPlaybackClient();
+			return PlaybackClient.GetMediaOutputIssueSeverity(GetMediaOutputServerName(InMediaOutput), GetChannelName().ToString(), OutputInfo.Guid);
+		}
+
+		const UMediaCapture* const MediaCapture = GetMediaCaptureForOutput(InMediaOutput);
+		
+		// For now, we don't have a way to determine if the capture has warnings.
+		// We only know if it worked or not.
+		if (IsValid(MediaCapture) && Avalanche::IsCapturing(MediaCapture))
+		{
+			return GetLocalSeverityForOutput(InMediaOutput);
+		}
+
+		// If the output is live but we don't have a media capture, that would be an error.
+		return EAvaMediaIssueSeverity::Errors;
+	}
+
+	// If the state is "error", we override severity with "errors", but the local severity should be "errors" as well
+	// and we expect error messages to have been recorded.
+	return (InOutputState == EAvaMediaOutputState::Error) ?  EAvaMediaIssueSeverity::Errors : GetLocalSeverityForOutput(InMediaOutput);
+}
+
+const TArray<FString>& FAvaOutputChannel::GetMediaOutputIssueMessages(const UMediaOutput* InMediaOutput) const
+{
+	if (IAvaMediaModule::Get().IsMediaPlaybackClientStarted() && IsMediaOutputRemote(InMediaOutput))
+	{
+		const FAvaMediaOutputInfo& OutputInfo = GetMediaOutputInfo(InMediaOutput);
+		const IAvaMediaPlaybackClient& PlaybackClient = IAvaMediaModule::Get().GetMediaPlaybackClient();
+		return PlaybackClient.GetMediaOutputIssueMessages(GetMediaOutputServerName(InMediaOutput), GetChannelName().ToString(), OutputInfo.Guid);
+	}
+	else
+	{
+		if (const FLocalMediaOutputStatus* LocalMediaOutputStatus = LocalMediaOutputStatuses.Find(InMediaOutput))
+		{
+			return LocalMediaOutputStatus->Messages;
+		}
+	}
+	static const TArray<FString> EmptyStringArray;
+	return EmptyStringArray;
+}
+
+EAvaMediaOutputState FAvaOutputChannel::GetMediaOutputState(const UMediaOutput* InMediaOutput) const
+{
+	if (IsMediaOutputRemote(InMediaOutput))
+	{
+		if (IAvaMediaModule::Get().IsMediaPlaybackClientStarted())
+		{
+			const FAvaMediaOutputInfo& OutputInfo = GetMediaOutputInfo(InMediaOutput);
+			if (!OutputInfo.IsValid())
+			{
+				UE_LOG(LogAvaOutputChannel, Error, TEXT("Channel \"%s\": MediaOutput \"%s\" has invalid info."),
+					*GetNameSafe(InMediaOutput), *GetChannelName().ToString());
+			}
+			
+			// If the output is broadcasting remote, fetch the status from the playback client (which is proxying that output's status).
+			const IAvaMediaPlaybackClient& PlaybackClient = IAvaMediaModule::Get().GetMediaPlaybackClient();
+			return PlaybackClient.GetMediaOutputState(GetMediaOutputServerName(InMediaOutput),GetChannelName().ToString(), OutputInfo.Guid);
+		}
+		else
+		{
+			// Remote output if offline if the client is not running.
+			return EAvaMediaOutputState::Offline;
+		}
+	}
+	
+	const UMediaCapture* const MediaCapture = GetMediaCaptureForOutput(InMediaOutput);
+
+	if (IsValid(MediaCapture))
+	{
+		// For now, we don't have a way to determine if the capture has warnings.
+		// We only know if it worked or not.
+		switch (MediaCapture->GetState())
+		{
+		case EMediaCaptureState::Capturing:
+		case EMediaCaptureState::StopRequested:
+			return EAvaMediaOutputState::Live;
+		case EMediaCaptureState::Preparing:
+			return EAvaMediaOutputState::Preparing;
+		case EMediaCaptureState::Error:
+			return EAvaMediaOutputState::Error;
+		default:
+			return EAvaMediaOutputState::Idle;
+		}
+	}
+	
+	// If no media capture is found for the given output and we are marked as broadcasting,
+	// this is an error state.
+	return bInternalStateBroadcasting ?  EAvaMediaOutputState::Error : EAvaMediaOutputState::Idle;
+}
+
+const FAvaMediaOutputInfo& FAvaOutputChannel::GetMediaOutputInfo(const UMediaOutput* InMediaOutput) const
+{
+	const int32 Index = MediaOutputs.Find(const_cast<UMediaOutput*>(InMediaOutput));
+	if (MediaOutputInfos.IsValidIndex(Index))
+	{
+		return MediaOutputInfos[Index];
+	}
+	static FAvaMediaOutputInfo Empty;
+	return Empty;
+}
+
+FAvaMediaOutputInfo* FAvaOutputChannel::GetMediaOutputInfoMutable(const UMediaOutput* InMediaOutput)
+{
+	const int32 Index = MediaOutputs.Find(const_cast<UMediaOutput*>(InMediaOutput));
+	return (MediaOutputInfos.IsValidIndex(Index)) ? &MediaOutputInfos[Index] : nullptr; 
+}
+
+inline void UpdatePlaceHolderRenderTarget(TObjectPtr<UTextureRenderTarget2D>& InOutRenderTarget, FName InBaseName, const FIntPoint& InSize, EPixelFormat InFormat, const FLinearColor& InClearColor)
+{
+	if (!InOutRenderTarget)
+	{
+		InOutRenderTarget = AvaRenderTargetUtils::CreateDefaultRenderTarget(InBaseName);
+	}
+
+	AvaRenderTargetUtils::UpdateRenderTarget(InOutRenderTarget.Get(), InSize, InFormat, InClearColor);
+}
+
+bool FAvaOutputChannel::IsEditingEnabled() const
+{
+	return ChannelState != EAvaChannelState::Live;
+}
+
+void FAvaOutputChannel::SetState(EAvaChannelState InState, EAvaMediaIssueSeverity InIssueSeverity)
+{
+	if (ChannelState != InState || ChannelIssueSeverity != InIssueSeverity)
+	{
+		ChannelState = InState;
+		ChannelIssueSeverity = InIssueSeverity;
+
+		if (InState == EAvaChannelState::Live)
+		{
+			StartPlaceholderTick();
+		}
+		else
+		{
+			StopPlaceholderTick();
+		}
+		OnChannelChanged.Broadcast(*this, EAvaChannelChange::State);
+	}
+}
+
+void FAvaOutputChannel::UpdatePlaceholder()
+{
+	if (!FSlateApplication::IsInitialized())
+	{
+		return;
+	}
+	
+	if (!PlaceholderWidget.IsValid())
+	{
+		PlaceholderWidget = SNew(SAvaPlaceholderWidget);
+	}
+
+	if (!WidgetRenderer.IsValid())
+	{
+		WidgetRenderer = MakeShared<FWidgetRenderer>();
+		
+		WidgetRenderer->SetIsPrepassNeeded(true);
+		
+		//Enable Gamma Correction for Widget Rendering
+		WidgetRenderer->SetUseGammaCorrection(true);
+		
+		//Clear Target before rendering
+		WidgetRenderer->SetShouldClearTarget(true);
+	}
+
+	UpdateOverridePlaceholder();
+
+	PlaceholderWidget->SetChannelName(ChannelNameText);
+	
+	UpdatePlaceholderRenderTargets();
+}
+
+void FAvaOutputChannel::UpdatePlaceholderRenderTargets()
+{
+	const FIntPoint RenderTargetSize = DetermineRenderTargetSize();
+	const EPixelFormat RenderTargetFormat = DetermineRenderTargetFormat();
+	const FLinearColor ClearColor = IAvaMediaModule::Get().GetBroadcastSettings().GetChannelClearColor();
+	static const FName PlaceholderRenderTargetName = TEXT("AvaOutputChannel_PlaceHolderRT");
+	static const FName PlaceholderRenderTargetTmpName = TEXT("AvaOutputChannel_TmpPlaceHolderRT");
+
+	UpdatePlaceHolderRenderTarget(PlaceholderRenderTarget, PlaceholderRenderTargetName, RenderTargetSize, RenderTargetFormat, ClearColor);
+	UpdatePlaceHolderRenderTarget(PlaceholderRenderTargetTmp, PlaceholderRenderTargetTmpName, RenderTargetSize, RenderTargetFormat, ClearColor);
+}
+
+void FAvaOutputChannel::UpdateOverridePlaceholder()
+{
+	const IAvalancheBroadcastSettings& BroadcastSettings = IAvaMediaModule::Get().GetBroadcastSettings();
+
+	if (!BroadcastSettings.IsDrawPlaceholderWidget())
+	{
+		OverridePlaceholderWidget.Reset();
+		return;
+	}
+	
+	UClass* const CurrentOverridePlaceholderClass = OverridePlaceholderUserWidget
+		? OverridePlaceholderUserWidget->GetClass()
+		: nullptr;
+
+	// Remark: on the server, if this settings comes from a client, the server may not
+	// have the asset and will need to sync it.
+	// TODO: use sync service to sync asset and use it when ready.
+	TSoftClassPtr<UUserWidget> PlaceholderWidgetClass(BroadcastSettings.GetPlaceholderWidgetClass());
+	UClass* const TargetOverridePlaceholderClass = PlaceholderWidgetClass.LoadSynchronous();
+	
+	//Update only if the Current Placeholder Class is different from the one in Media Settings
+	if (CurrentOverridePlaceholderClass != TargetOverridePlaceholderClass)
+	{
+		OverridePlaceholderUserWidget = nullptr;
+		if (TargetOverridePlaceholderClass)
+		{			
+			OverridePlaceholderUserWidget = NewObject<UUserWidget>(GetTransientPackage()
+				, TargetOverridePlaceholderClass
+				, NAME_None
+				, RF_Transient);
+		}
+
+		OverridePlaceholderWidget.Reset();
+		if (OverridePlaceholderUserWidget)
+		{
+			OverridePlaceholderWidget = OverridePlaceholderUserWidget->TakeWidget();
+
+			//Update Tick State
+			if (OverridePlaceholderWidget.IsValid())
+			{
+				// Default to never tick, only recompute for auto
+				bool bCanTick = false;
+
+				//These rules are pretty much the same as UUserWidget::UpdateCanTick but without the UWorld Checks
+				if (OverridePlaceholderUserWidget->GetDesiredTickFrequency() == EWidgetTickFrequency::Auto)
+				{
+					// Note: WidgetBPClass can be NULL in a cooked build.
+					UWidgetBlueprintGeneratedClass* const WidgetBPClass = Cast<UWidgetBlueprintGeneratedClass>(OverridePlaceholderUserWidget->GetClass());
+					bCanTick |= !WidgetBPClass || WidgetBPClass->ClassRequiresNativeTick();
+					bCanTick |= OverridePlaceholderUserWidget->bHasScriptImplementedTick;
+					if (!bCanTick)
+					{
+						for (UUserWidgetExtension* const Extension : OverridePlaceholderUserWidget->GetExtensions(UUserWidgetExtension::StaticClass()))
+						{
+							if (Extension->RequiresTick())
+							{
+								bCanTick = true;
+								break;
+							}
+						}
+					}
+				}
+
+				OverridePlaceholderWidget->SetCanTick(bCanTick);
+			}
+		}
+	}
+
+	if (OverridePlaceholderUserWidget && OverridePlaceholderUserWidget->Implements<UAvalancheOutputChannelViewInterface>())
+	{
+		IAvalancheOutputChannelViewInterface::Execute_SetChannelName(OverridePlaceholderUserWidget, ChannelNameText);
+	}
+}
+
+inline bool AreDimensionsValid(const UTextureRenderTarget2D* InRenderTarget, const int32 InMaxAllowedDrawSize)
+{
+	return (InRenderTarget->SizeX <= 0
+		|| InRenderTarget->SizeY <= 0
+		|| InRenderTarget->SizeX > InMaxAllowedDrawSize
+		|| InRenderTarget->SizeY > InMaxAllowedDrawSize) ? false : true;
+}
+
+void FAvaOutputChannel::DrawPlaceholderWidget() const
+{
+	if (!IsValid(PlaceholderRenderTarget) || !IsValid(PlaceholderRenderTargetTmp))
+	{
+		return;
+	}
+	
+	const int32 MaxAllowedDrawSize = GetMax2DTextureDimension();
+	
+	if (!AreDimensionsValid(PlaceholderRenderTarget.Get(), MaxAllowedDrawSize) 
+		|| !AreDimensionsValid(PlaceholderRenderTargetTmp.Get(), MaxAllowedDrawSize))
+	{
+		return;
+	}
+
+	const bool bHasWidgetResources = PlaceholderWidget.IsValid() && WidgetRenderer.IsValid();
+	
+	if (IAvaMediaModule::Get().GetBroadcastSettings().IsDrawPlaceholderWidget() && bHasWidgetResources)
+	{
+		TSharedRef<SWidget> Widget = OverridePlaceholderWidget.IsValid()
+			? OverridePlaceholderWidget.ToSharedRef()
+			: PlaceholderWidget.ToSharedRef();
+	
+		// Render the widget in the tmp render target.
+		WidgetRenderer->DrawWidget(PlaceholderRenderTargetTmp.Get()
+			, Widget
+			, FVector2D(PlaceholderRenderTargetTmp->SizeX, PlaceholderRenderTargetTmp->SizeY)
+			, FSlateApplication::Get().GetDeltaTime());
+
+		// Invert the Alpha on the Placeholder Widget
+		UTextureRenderTarget2D* SourceRT = PlaceholderRenderTargetTmp.Get();
+		UTextureRenderTarget2D* DestinationRT = PlaceholderRenderTarget.Get();
+
+		ENQUEUE_RENDER_COMMAND(FAvaOutputInvertAlpha)(
+			[SourceRT, DestinationRT](FRHICommandListImmediate& RHICmdList)
+			{
+				Avalanche::ExecuteInvertAlphaConversionPass(RHICmdList, SourceRT, DestinationRT);
+			});
+	}
+	else
+	{
+		UE::AvaRenderTargetMediaUtils::ClearRenderTarget(PlaceholderRenderTarget.Get());
+	}
+}
+
+bool FAvaOutputChannel::IsValidChannel() const
+{
+	return this != &FAvaOutputChannel::NullChannel && GetChannelName() != NAME_None;
+}
+
+bool FAvaOutputChannel::StartChannelBroadcast()
+{
+	bInternalStateBroadcasting = true;
+	
+	IAvaMediaModule& Module = IAvaMediaModule::Get();
+	if (Module.IsMediaPlaybackClientStarted())
+	{
+		const TArray<UMediaOutput*> RemoteMediaOutputs = GetRemoteMediaOutputs();
+		if (!RemoteMediaOutputs.IsEmpty())
+		{
+			const FString ProfileName = GetProfileName().ToString();
+			Module.GetMediaPlaybackClient().RequestBroadcast(ProfileName, GetChannelName(), RemoteMediaOutputs, EAvaMediaBroadcastAction::Start);
+		}
+	}
+
+	const FString ChannelName = GetChannelName().ToString();
+	
+	if (GetState() == EAvaChannelState::Live)
+	{
+		UE_LOG(LogAvaOutputChannel, Warning
+			, TEXT("Output Channel (%s) attempted to Start Channel Broadcast while already Running!")
+			, *ChannelName);
+		return true;
+	}
+
+	// Ensure placeholder render targets are compatible with media outputs.
+	UpdateChannelResources(true);
+	
+	UTextureRenderTarget2D* const RenderTarget = GetCurrentRenderTarget(true);
+	
+	for (UMediaOutput* MediaOutput : MediaOutputs)
+	{
+		if (!IsValid(MediaOutput))
+		{
+			UE_LOG(LogAvaOutputChannel, Warning
+				, TEXT("Output Channel (%s) has an Invalid Media Output.")
+				, *ChannelName);
+			continue;
+		}
+
+		ResetLocalMediaOutputStatus(MediaOutput);
+		
+		// Don't locally start remote media outputs.
+		if (IsMediaOutputRemote(MediaOutput))
+		{
+			continue;
+		}
+		
+		// Check if this MediaOutput can start with the given render target
+		if (!IsMediaOutputCompatible(MediaOutput, RenderTarget))
+		{
+			StopCaptureForOutput(MediaOutput);	// Ensures no stale capture is running.
+			continue;
+		}
+
+		UMediaCapture* const MediaCapture = MediaOutput->CreateMediaCapture();
+		
+		if (!IsValid(MediaCapture))
+		{
+			FString& Message = EmplaceMessageForOutput(MediaOutput, EAvaMediaIssueSeverity::Warnings);
+			Message = FString::Printf(
+				TEXT("Output Channel (%s) has a Media Output (%s) that could not create a valid Media Capture.")
+				, *ChannelName, *MediaOutput->GetName());
+
+			UE_LOG(LogAvaOutputChannel, Warning, TEXT("%s"), *Message);
+			
+			StopCaptureForOutput(MediaOutput);	// Ensures no stale capture is running.
+			continue;
+		}
+		
+		// TODO: Expose this along with MediaOutput.
+		FMediaCaptureOptions CaptureOptions;
+		
+		// Don't stop Capture when Throttling the Engine
+		CaptureOptions.bSkipFrameWhenRunningExpensiveTasks = false;
+		// Allow the formats to be converted if different.
+		CaptureOptions.bConvertToDesiredPixelFormat = true;
+		
+		check(RenderTarget);
+		if (MediaCapture->CaptureTextureRenderTarget2D(RenderTarget, CaptureOptions))
+		{
+			FAvaMediaCapture& AvaMediaCapture = MediaCaptures.Emplace(MediaOutput);
+			AvaMediaCapture.Reset(MediaCapture, MediaOutput, RenderTarget);
+			AvaMediaCapture.OnMediaCaptureStateChanged.AddRaw(this, &FAvaOutputChannel::OnAvaMediaCaptureStateChanged); 
+		}
+	}
+
+	RefreshState();
+	return true;
+}
+
+void FAvaOutputChannel::StopChannelBroadcast()
+{
+	bInternalStateBroadcasting = false;
+
+	StopPlaceholderTick();
+	
+	IAvaMediaModule& Module = IAvaMediaModule::Get();
+	if (Module.IsMediaPlaybackClientStarted())
+	{
+		const TArray<UMediaOutput*> RemoteMediaOutputs = GetRemoteMediaOutputs();
+		if (!RemoteMediaOutputs.IsEmpty())
+		{
+			const FString ProfileName = GetProfileName().ToString();
+			Module.GetMediaPlaybackClient().RequestBroadcast(ProfileName, GetChannelName(), RemoteMediaOutputs, EAvaMediaBroadcastAction::Stop);
+		}
+	}
+
+	for (const UMediaOutput* MediaOutput : MediaOutputs)
+	{
+		ResetLocalMediaOutputStatus(MediaOutput);
+		StopCaptureForOutput(MediaOutput);
+	}
+	
+	MediaCaptures.Reset();
+	LastActivePlayableGroup.Reset();
+	CurrentRenderTarget.Reset();
+	
+	if (IsValid(PlaceholderRenderTarget))
+	{
+		PlaceholderRenderTarget->UpdateResourceImmediate();
+	}
+
+	if (IsValid(PlaceholderRenderTargetTmp))
+	{
+		PlaceholderRenderTargetTmp->UpdateResourceImmediate();
+	}
+
+	SetState(EAvaChannelState::Idle, EAvaMediaIssueSeverity::None);
+}
+
+UMediaOutput* FAvaOutputChannel::AddMediaOutput(const UClass* InMediaOutputClass, const FAvaMediaOutputInfo& InOutputInfo)
+{
+	// Don't add a remote output to a local preview channel.
+	if (InOutputInfo.IsValid() && InOutputInfo.IsRemote() && GetChannelType() == EAvaBroadcastChannelType::Preview)
+	{
+		return nullptr;
+	}
+	
+	if (InMediaOutputClass && InMediaOutputClass->IsChildOf(UMediaOutput::StaticClass()))
+	{
+		UMediaOutput* const MediaOutput = NewObject<UMediaOutput>(UAvalancheBroadcast::GetAvalancheBroadcast()
+			, InMediaOutputClass
+			, NAME_None
+			, RF_Transactional);
+
+		// If this is a MediaOutput class that requires render targets to be resolved to cpu memory,
+		// we need more buffers to avoid having the flush rendering commands.
+		static const int32 MinimumNumberOfTextureBuffers = 3;
+		if (MediaOutput->NumberOfTextureBuffers < MinimumNumberOfTextureBuffers)
+		{
+			MediaOutput->NumberOfTextureBuffers = MinimumNumberOfTextureBuffers;
+		}
+
+		// Set bInvertKeyOutput property to true by default since we always output a scene render with inverted alpha.
+		{
+			// Known implementations that have it: BlackMagic, Aja, NDI.
+			// Ava Display and Rivermax don't have it, but don't support Key.
+			FBoolProperty* const Property = FindFProperty<FBoolProperty>(MediaOutput->GetClass(), TEXT("bInvertKeyOutput"));
+			if (Property)
+			{
+				Property->SetPropertyValue(Property->ContainerPtrToValuePtr<bool>(MediaOutput, 0), true);
+			}
+		}
+
+		AddMediaOutput(MediaOutput, InOutputInfo);
+		
+		return MediaOutput;
+	}
+	return nullptr;
+}
+
+void FAvaOutputChannel::AddMediaOutput(UMediaOutput* InMediaOutput, const FAvaMediaOutputInfo& InOutputInfo)
+{
+	if (IsValid(InMediaOutput))
+	{
+		// Check that we don't have a remote output on a local preview.
+		check(!(InOutputInfo.IsValid() && InOutputInfo.IsRemote() && GetChannelType() == EAvaBroadcastChannelType::Preview));
+#if WITH_EDITOR
+		InMediaOutput->OnOutputModified().AddRaw(this, &FAvaOutputChannel::OnMediaOutputModified);
+#endif
+		MediaOutputs.Add(InMediaOutput);
+		MediaOutputInfos.Add(InOutputInfo);
+		// Note: the media output may not be fully configured at that point.
+		OnChannelChanged.Broadcast(*this, EAvaChannelChange::MediaOutputs);
+		UpdateChannelResources(true);
+	}
+}
+
+int32 FAvaOutputChannel::RemoveMediaOutput(UMediaOutput* InMediaOutput)
+{
+	int32 RemoveCount = 0;
+	const int32 IndexToRemove = MediaOutputs.Find(InMediaOutput);
+	if (IndexToRemove != INDEX_NONE)
+	{
+#if WITH_EDITOR
+		InMediaOutput->OnOutputModified().RemoveAll(this);
+#endif
+		RemoveCount = 1;
+		MediaOutputs.RemoveAt(IndexToRemove);
+		MediaOutputInfos.RemoveAt(IndexToRemove);
+		OnChannelChanged.Broadcast(*this, EAvaChannelChange::MediaOutputs);
+		UpdateChannelResources(true);
+	}
+	return RemoveCount;
+}
+
+void FAvaOutputChannel::OnMediaOutputModified(UMediaOutput* InMediaOutput)
+{
+	// Update media output info's device name.
+	if (const FString DeviceName = UE::AvaMediaOutputUtils::GetDeviceName(InMediaOutput); !DeviceName.IsEmpty())
+	{
+		if (FAvaMediaOutputInfo* OutputInfo = GetMediaOutputInfoMutable(InMediaOutput))
+		{
+			OutputInfo->DeviceName = FName(DeviceName);
+		}
+	}
+	
+	IAvaMediaModule& AvaMediaModule = IAvaMediaModule::Get();
+	if (IsMediaOutputRemote(InMediaOutput) && AvaMediaModule.IsMediaPlaybackClientStarted())
+	{
+		const FString ProfileName = GetProfileName().ToString();
+		IAvaMediaPlaybackClient& PlaybackClient = AvaMediaModule.GetMediaPlaybackClient();
+		PlaybackClient.RequestBroadcast(ProfileName, GetChannelName(), GetRemoteMediaOutputs(), EAvaMediaBroadcastAction::UpdateConfig);
+	}
+	else
+	{
+		UpdateChannelResources(true);
+		RefreshState();
+	}
+}
+
+void FAvaOutputChannel::PostLoadMediaOutputs(bool bInIsProfileActive, FAvaBroadcastProfile* InProfile)
+{
+	Profile = InProfile;
+	
+	// Create missing output infos (legacy support)
+	if (MediaOutputInfos.Num() != MediaOutputs.Num())
+	{
+		MediaOutputInfos.Empty(MediaOutputs.Num());
+		for (const UMediaOutput* Output : MediaOutputs)
+		{
+			MediaOutputInfos.Add(Avalanche::MediaOutputInfoHelper::InitFrom(Output));
+		}
+	}
+
+	// Remove invalid outputs and corresponding infos.
+	for (int32 OutputIndex = 0; OutputIndex < MediaOutputs.Num();)
+	{
+		if (!IsValid(MediaOutputs[OutputIndex]))
+		{
+			MediaOutputs.RemoveAt(OutputIndex);
+			MediaOutputInfos.RemoveAt(OutputIndex);
+		}
+		else
+		{
+			++OutputIndex;
+		}
+	}
+
+	// Complete missing fields (legacy support)
+	for (FAvaMediaOutputInfo& MediaOutputInfo : MediaOutputInfos)
+	{
+		MediaOutputInfo.PostLoad();
+	}
+
+#if WITH_EDITOR
+	for (UMediaOutput* MediaOutput : MediaOutputs)
+	{
+		if (!MediaOutput->OnOutputModified().IsBoundToObject(this))
+		{
+			MediaOutput->OnOutputModified().AddRaw(this, &FAvaOutputChannel::OnMediaOutputModified);
+		}
+	}
+#endif
+	
+	OnChannelChanged.Broadcast(*this, EAvaChannelChange::MediaOutputs);
+	UpdateChannelResources(bInIsProfileActive);
+	RefreshState();
+}
+
+void FAvaOutputChannel::UpdateRenderTarget(UAvaMediaPlayableGroup* InPlayableGroup, UTextureRenderTarget2D* InRenderTarget)
+{
+	// In the current design, the render target is associated with a "playable group".
+	// It can only be associated with one group at a time.
+
+	const bool bRenderTargetChanged = CurrentRenderTarget.Get() != InRenderTarget;
+	
+	LastActivePlayableGroup = InPlayableGroup;
+	CurrentRenderTarget = InRenderTarget;
+
+	if (bRenderTargetChanged)
+	{
+		UpdateViewportTarget();
+	}
+}
+
+void FAvaOutputChannel::UpdateAudioDevice(const FAudioDeviceHandle& InAudioDeviceHandle)
+{
+	for (TPair<TObjectPtr<UMediaOutput>, FAvaMediaCapture>& MediaCapture : MediaCaptures)
+	{
+		UMediaCapture* MediaCapturePtr = MediaCapture.Value.MediaCapture;
+		if (MediaCapturePtr->GetCaptureAudioDevice().GetDeviceID() != InAudioDeviceHandle.GetDeviceID())
+		{
+			MediaCapturePtr->SetCaptureAudioDevice(InAudioDeviceHandle);
+		}
+	}
+}
+
+UAvaMediaPlayableGroup* FAvaOutputChannel::GetLastActivePlayableGroup() const
+{
+	return LastActivePlayableGroup.Get();
+}
+
+UTextureRenderTarget2D* FAvaOutputChannel::GetCurrentRenderTarget(bool bInFallbackToPlaceholder) const
+{
+	if (CurrentRenderTarget.IsValid())
+	{
+		return CurrentRenderTarget.Get();
+	}
+	
+	return bInFallbackToPlaceholder
+		? GetPlaceholderRenderTarget()
+		: nullptr;
+}
+
+UTextureRenderTarget2D* FAvaOutputChannel::GetPlaceholderRenderTarget() const
+{
+	return PlaceholderRenderTarget.Get();
+}
+
+bool FAvaOutputChannel::HasAnyLocalMediaOutputs() const
+{
+	for (const UMediaOutput* MediaOutput : MediaOutputs)
+	{
+		if (IsValid(MediaOutput) && !IsMediaOutputRemote(MediaOutput))
+		{
+			return true;
+		}
+	}
+	return false;
+}
+
+bool FAvaOutputChannel::HasAnyRemoteMediaOutputs() const
+{
+	for (const UMediaOutput* MediaOutput : MediaOutputs)
+	{
+		if (IsValid(MediaOutput) && IsMediaOutputRemote(MediaOutput))
+		{
+			return true;
+		}
+	}
+	return false;
+}
+
+
+bool FAvaOutputChannel::IsMediaOutputRemote(const UMediaOutput* InMediaOutput) const
+{
+	// Ideally, the output info is valid here, so we just need to see if the
+	// output is associated with a remote server.
+	const FAvaMediaOutputInfo& OutputInfo = GetMediaOutputInfo(InMediaOutput);
+	if (OutputInfo.IsValid())
+	{
+		return OutputInfo.IsRemote();
+	}
+
+	// Fallback/Legacy, in case the output info is not valid,
+	// if the playback client is started, we can try the legacy method
+	// of checking the remote device providers. This is much less reliable.
+	IAvaMediaModule& AvaMediaModule = IAvaMediaModule::Get();
+	if (AvaMediaModule.IsMediaPlaybackClientStarted())
+	{
+		return AvaMediaModule.GetMediaPlaybackClient().IsMediaOutputRemoteFallback(InMediaOutput);
+	}
+
+	return false;
+}
+
+const FString& FAvaOutputChannel::GetMediaOutputServerName(const UMediaOutput* InMediaOutput) const
+{
+	const FAvaMediaOutputInfo& OutputInfo = GetMediaOutputInfo(InMediaOutput);
+	if (OutputInfo.IsValid())
+	{
+		return OutputInfo.ServerName;
+	}
+
+	static const FString EmptyString;
+	return EmptyString;
+}
+
+TArray<UMediaOutput*> FAvaOutputChannel::GetRemoteMediaOutputs() const
+{
+	TArray<UMediaOutput*> RemoteMediaOutputs;
+	RemoteMediaOutputs.Reserve(MediaOutputs.Num());
+	for (UMediaOutput* MediaOutput : MediaOutputs)
+	{
+		if (IsValid(MediaOutput) && IsMediaOutputRemote(MediaOutput))
+		{
+			RemoteMediaOutputs.Add(MediaOutput);
+		}
+	}
+	return RemoteMediaOutputs;
+}
+
+bool FAvaOutputChannel::IsMediaOutputCompatible(UMediaOutput* InMediaOutput, UTextureRenderTarget2D* InRenderTarget) const
+{
+	// Check if this MediaOutput can start with the given render target
+	if (IsValid(InRenderTarget) && IsValid(InMediaOutput))
+	{
+		// Size Match Check
+		{
+			const FIntPoint MediaOutputSize = InMediaOutput->GetRequestedSize();
+			const FIntPoint RenderTargetSize = AvaRenderTargetUtils::GetRenderTargetSize(InRenderTarget);
+			if (MediaOutputSize != UMediaOutput::RequestCaptureSourceSize && MediaOutputSize != RenderTargetSize)
+			{
+				FString& Message = EmplaceMessageForOutput(InMediaOutput, EAvaMediaIssueSeverity::Errors);
+				Message = FString::Printf(
+					TEXT("Media Output \"%s\" requested size (%d, %d) is not compatible with Channel \"%s\" Render Target \"%s\" size (%d, %d).")
+					, *InMediaOutput->GetName(), MediaOutputSize.X, MediaOutputSize.Y
+					, *GetChannelName().ToString()
+					, *InRenderTarget->GetName(), RenderTargetSize.X, RenderTargetSize.Y);
+				
+				UE_LOG(LogAvaOutputChannel, Error, TEXT("%s"), *Message);
+				return false;
+			}
+		}
+
+		// Format Match Check
+		{
+			const EPixelFormat MediaOutputFormat = InMediaOutput->GetRequestedPixelFormat();
+			const EPixelFormat RenderTargetFormat = InRenderTarget->GetFormat();
+			if (MediaOutputFormat != PF_Unknown && MediaOutputFormat != RenderTargetFormat)
+			{
+				const int32 TargetNumColorBits = AvaRenderTargetUtils::GetNumColorChannelBits(RenderTargetFormat);
+				const int32 OutputNumColorBits = AvaRenderTargetUtils::GetNumColorChannelBits(MediaOutputFormat);
+				if (OutputNumColorBits > TargetNumColorBits)
+				{
+					FString& Message = EmplaceMessageForOutput(InMediaOutput, EAvaMediaIssueSeverity::Errors);
+					Message = FString::Printf(
+						TEXT("Media Output \"%s\" requested format %s is not compatible with Channel \"%s\" Render Target \"%s\" format %s")
+						TEXT(" because it has higher color precision (%d vs %d bits).")
+						, *InMediaOutput->GetName(), GetPixelFormatString(MediaOutputFormat)
+						, *GetChannelName().ToString()
+						, *InRenderTarget->GetName(), GetPixelFormatString(RenderTargetFormat)
+						, OutputNumColorBits, TargetNumColorBits);
+
+					UE_LOG(LogAvaOutputChannel, Error, TEXT("%s"), *Message);
+					return false;
+				}
+
+				// RGBA8 vs RGB10A2, we will detect a conflict here.
+				// RGB10A2 wins in color bits, but loose in alpha bits.
+				const int32 TargetNumAlphaBits = AvaRenderTargetUtils::GetNumAlphaChannelBits(RenderTargetFormat);
+				const int32 OutputNumAlphaBits = AvaRenderTargetUtils::GetNumAlphaChannelBits(MediaOutputFormat);
+				if (OutputNumAlphaBits > TargetNumAlphaBits)
+				{
+					FString& Message = EmplaceMessageForOutput(InMediaOutput, EAvaMediaIssueSeverity::Warnings);
+					Message = FString::Printf(
+						TEXT("Media Output \"%s\" requested format %s is not compatible with Channel \"%s\" Render Target \"%s\" format %s")
+						TEXT(" because it has higher alpha precision (%d vs %d bits).")
+						, *InMediaOutput->GetName(), GetPixelFormatString(MediaOutputFormat)
+						, *GetChannelName().ToString()
+						, *InRenderTarget->GetName(), GetPixelFormatString(RenderTargetFormat)
+						, OutputNumAlphaBits, TargetNumAlphaBits);
+
+					UE_LOG(LogAvaOutputChannel, Warning, TEXT("%s"), *Message);
+				}
+			}
+		}
+		return true;
+	}
+	return false;
+}
+
+void FAvaOutputChannel::UpdateViewportTarget()
+{
+	//Only need to update when Broadcasting, since when starting Broadcast it will Capture and update it anyways
+	if (GetState() == EAvaChannelState::Live)
+	{
+		UTextureRenderTarget2D* const RenderTarget = GetCurrentRenderTarget(true);
+
+		bool bNeedStateRefresh = false;
+
+		// Start all media outputs that have a compatible format 
+		// and are not already started on another game instance.
+		for (UMediaOutput* const MediaOutput : MediaOutputs)
+		{
+			if (!IsValid(MediaOutput))
+			{
+				continue;
+			}
+
+			if (IsMediaOutputRemote(MediaOutput))
+			{
+				continue;
+			}
+
+			// For now, we don't update devices that are not live.
+			if (GetMediaOutputState(MediaOutput) != EAvaMediaOutputState::Live)
+			{
+				continue;
+			}
+		
+			// Check if this MediaOutput can still run with the given render target.
+			// Remark: this test maybe optional. It is normally done when channel start
+			// and we don't allow configuration changes to the output while it is running.
+			if (!IsMediaOutputCompatible(MediaOutput, RenderTarget))
+			{
+				// Ensure no stale capture is running for this output.
+				StopCaptureForOutput(MediaOutput);
+				bNeedStateRefresh = true;
+				continue;
+			}
+
+			FAvaMediaCapture* AvaMediaCapture = MediaCaptures.Find(MediaOutput);
+			UMediaCapture* const MediaCapture = AvaMediaCapture ? AvaMediaCapture->MediaCapture : nullptr;
+			
+			// Double Check and Skip Updating SceneViewport / RenderTarget if not Capturing
+			// Remark: if the capture is an external/removable device, it can stop the capture.
+			// We can the state to be refreshed at least.
+			if (!Avalanche::IsCapturing(MediaCapture))
+			{
+				bNeedStateRefresh = true;
+				continue;
+			}
+			
+			check(RenderTarget);
+
+			// We go out of our way to avoid updating the render target (if it is the same)
+			// because UpdateTextureRenderTarget2D has side effects even if we are updating with the same render target.
+			if (AvaMediaCapture && AvaMediaCapture->RenderTarget.IsValid() && AvaMediaCapture->RenderTarget == RenderTarget)
+			{
+				continue;
+			}
+			
+			if (AvaMediaCapture && MediaCapture->UpdateTextureRenderTarget2D(RenderTarget))
+			{
+				AvaMediaCapture->RenderTarget = RenderTarget;
+			}
+		}
+		
+		if (bNeedStateRefresh)
+		{
+			RefreshState();
+		}
+	}
+	OnChannelChanged.Broadcast(*this, EAvaChannelChange::RenderTarget);
+}
+
+UMediaCapture* FAvaOutputChannel::GetMediaCaptureForOutput(const UMediaOutput* InMediaOutput) const
+{
+	const FAvaMediaCapture* AvaMediaCapture = MediaCaptures.Find(InMediaOutput);
+	return AvaMediaCapture ? AvaMediaCapture->MediaCapture : nullptr;
+}
+
+void FAvaOutputChannel::StopCaptureForOutput(const UMediaOutput* InMediaOutput)
+{
+	UMediaCapture* const MediaCapture = GetMediaCaptureForOutput(InMediaOutput);
+	if (IsValid(MediaCapture) && Avalanche::IsCapturing(MediaCapture))
+	{
+		MediaCapture->StopCapture(false);
+	}
+	MediaCaptures.Remove(InMediaOutput);
+}
+
+void FAvaOutputChannel::TickPlaceholder(float)
+{
+	// Only Draw if Render Target is Invalid 
+	if (GetState() == EAvaChannelState::Live && !IsValid(GetCurrentRenderTarget(false)))
+	{
+		DrawPlaceholderWidget();
+	}
+}
+
+void FAvaOutputChannel::StartPlaceholderTick()
+{
+	if (!PlaceholderTickHandle.IsValid() && FSlateApplication::IsInitialized())
+	{
+		PlaceholderTickHandle = FSlateApplication::Get().OnPreTick().AddRaw(this, &FAvaOutputChannel::TickPlaceholder);
+	}
+}
+
+void FAvaOutputChannel::StopPlaceholderTick()
+{
+	if (PlaceholderTickHandle.IsValid() && FSlateApplication::IsInitialized())
+	{
+		FSlateApplication::Get().OnPreTick().Remove(PlaceholderTickHandle);
+	}
+	PlaceholderTickHandle.Reset();
+	ReleasePlaceholderResources();
+}
+
+void FAvaOutputChannel::OnAvaMediaCaptureStateChanged(const UMediaCapture* InMediaCapture, const UMediaOutput* InMediaOutput)
+{
+	RefreshState();
+
+	// Note: the overall channel state may not change.
+	// So we need to broadcast a specific event for the
+	// media output state change.
+	OnMediaOutputStateChanged.Broadcast(*this, InMediaOutput);
+}
+
+void FAvaOutputChannel::UpdateChannelResources(bool bInIsProfileActive)
+{
+	ChannelNameText = FText::FromName(GetChannelName());
+	
+	if (bInIsProfileActive)
+	{
+		// Only update the placeholder resources (including render targets) if currently active profile.
+		UpdatePlaceholder();
+	}
+	else
+	{
+		// Release all resources.
+		ReleasePlaceholderResources();
+		ReleasePlaceholderRenderTargets();
+	}
+}
+
+FIntPoint FAvaOutputChannel::DetermineRenderTargetSize() const
+{
+	// Just go through the list of media output in order and returns the first
+	// one that specified a valid resolution.
+	for (UMediaOutput* const MediaOutput : MediaOutputs)
+	{
+		if (IsValid(MediaOutput))
+		{
+			const FIntPoint TargetSize = MediaOutput->GetRequestedSize();
+			if (TargetSize != UMediaOutput::RequestCaptureSourceSize && TargetSize.X > 0 && TargetSize.Y > 0)
+			{
+				return TargetSize;
+			}
+		}
+	}
+	return GetDefaultMediaOutputSize(GetChannelType());
+}
+
+EPixelFormat FAvaOutputChannel::DetermineRenderTargetFormat() const
+{
+	// Determine render target format with the most bits per pixel.
+	// Note: check for alpha bits (A2 vs A8).
+	EPixelFormat HighestBitsPerColorChannelFormat = PF_Unknown;
+	int32 HighestNumBitsPerColorChannel = 0;
+	int32 HighestNumColorComponents = 0;
+	EPixelFormat HighestAlphaBitsFormat = PF_Unknown;
+	int32 HighestNumAlphaBits = 0;
+	
+	for (const UMediaOutput* const MediaOutput : MediaOutputs)
+	{
+		if (IsValid(MediaOutput))
+		{
+			const EPixelFormat TargetFormat = MediaOutput->GetRequestedPixelFormat();			
+			if (TargetFormat != PF_Unknown)
+			{
+				const int32 NumBitsPerColorChannel = AvaRenderTargetUtils::GetNumColorChannelBits(TargetFormat);
+				const int32 NumColorComponents = AvaRenderTargetUtils::GetNumColorComponents(TargetFormat);
+				// Note: only select the format if it has at least as many color components.
+				// Otherwise, we stick with the format that has less bits, but more components.
+				if (NumBitsPerColorChannel > HighestNumBitsPerColorChannel && NumColorComponents >= HighestNumColorComponents)
+				{
+					HighestNumBitsPerColorChannel = NumBitsPerColorChannel;
+					HighestNumColorComponents = NumColorComponents;
+					HighestBitsPerColorChannelFormat = TargetFormat;
+				}
+				
+				const int32 NumAlphaBits = AvaRenderTargetUtils::GetNumAlphaChannelBits(TargetFormat);
+				if (NumAlphaBits > HighestNumAlphaBits)
+				{
+					HighestNumAlphaBits = NumAlphaBits;
+					HighestAlphaBitsFormat = TargetFormat;
+				}
+			}
+		}
+	}
+
+	// Note: this logic favors the number of color bits to the detriment of alpha bits.
+	// In a RGBA8 vs RGB10A2 case, if we want 8 bits alpha, to the detriment of color bits,
+	// the logic below will need to be changed.
+	
+	if (HighestNumAlphaBits > HighestNumBitsPerColorChannel && HighestAlphaBitsFormat != PF_Unknown)
+	{
+		return HighestAlphaBitsFormat;
+	}
+	
+	if (HighestBitsPerColorChannelFormat != PF_Unknown)
+	{
+		return HighestBitsPerColorChannelFormat;
+	}
+	
+	if (HighestAlphaBitsFormat != PF_Unknown)
+	{
+		return HighestAlphaBitsFormat;
+	}
+	
+	return GetDefaultMediaOutputFormat();
+}
+
+FIntPoint FAvaOutputChannel::GetDefaultMediaOutputSize(EAvaBroadcastChannelType InChannelType)
+{
+	// Special case for local preview channel.
+	if (InChannelType == EAvaBroadcastChannelType::Preview)
+	{
+		return UAvalancheMediaSettings::Get().PreviewDefaultResolution;
+	}
+
+	// Default for program channel.
+	return IAvaMediaModule::Get().GetBroadcastSettings().GetDefaultResolution();
+}
+
+EPixelFormat FAvaOutputChannel::GetDefaultMediaOutputFormat()
+{
+	return IAvaMediaModule::Get().GetBroadcastSettings().GetDefaultPixelFormat();
+}
+
+void FAvaOutputChannel::SetViewportQualitySettings(const FAvaViewportQualitySettings& InQualitySettings)
+{
+	QualitySettings = InQualitySettings;
+	OnChannelChanged.Broadcast(*this, EAvaChannelChange::Settings);
+}
+
+FString& FAvaOutputChannel::EmplaceMessageForOutput(const UMediaOutput* InMediaOutput, EAvaMediaIssueSeverity InSeverity) const
+{
+	FLocalMediaOutputStatus& LocalMediaOutputStatus = LocalMediaOutputStatuses.Emplace(InMediaOutput);
+	LocalMediaOutputStatus.Severity = InSeverity;
+	return LocalMediaOutputStatus.Messages.Emplace_GetRef();
+}
+
+EAvaMediaIssueSeverity FAvaOutputChannel::GetLocalSeverityForOutput(const UMediaOutput* InMediaOutput) const
+{
+	const FLocalMediaOutputStatus* LocalMediaOutputStatus = LocalMediaOutputStatuses.Find(InMediaOutput);
+	return LocalMediaOutputStatus ? LocalMediaOutputStatus->Severity : EAvaMediaIssueSeverity::None;
+}
+
+void FAvaOutputChannel::ResetLocalMediaOutputStatus(const UMediaOutput* InMediaOutput)
+{
+	if (IsValid(InMediaOutput))
+	{
+		LocalMediaOutputStatuses.Add(InMediaOutput,FLocalMediaOutputStatus());
+	}
+}
\ No newline at end of file
diff --git a/Engine/Plugins/Experimental/Avalanche/Source/AvalancheMedia/Private/Channel/Slate/SAvaPlaceholderWidget.cpp b/Engine/Plugins/Experimental/Avalanche/Source/AvalancheMedia/Private/Channel/Slate/SAvaPlaceholderWidget.cpp
new file mode 100644
--- /dev/null
+++ b/Engine/Plugins/Experimental/Avalanche/Source/AvalancheMedia/Private/Channel/Slate/SAvaPlaceholderWidget.cpp
@@ -0,0 +1,103 @@
+// Copyright Epic Games, Inc. All Rights Reserved.
+
+#include "Channel/Slate/SAvaPlaceholderWidget.h"
+#include "AvaMediaStyle.h"
+#include "Widgets/Images/SImage.h"
+#include "Widgets/Images/SThrobber.h"
+#include "Widgets/Layout/SScaleBox.h"
+#include "Widgets/Layout/SBox.h"
+#include "Widgets/SBoxPanel.h"
+#include "Widgets/Text/STextBlock.h"
+
+#define LOCTEXT_NAMESPACE "AvalancheBroadcast"
+
+void SAvaPlaceholderWidget::Construct(const FArguments& InArgs)
+{
+	ChildSlot
+	[
+		SNew(SVerticalBox)
+		+ SVerticalBox::Slot()
+		.FillHeight(0.1f)
+		[
+			SNew(SBox)
+		]
+		
+		+ SVerticalBox::Slot()
+		.FillHeight(0.2f)
+		.VAlign(EVerticalAlignment::VAlign_Fill)
+		.HAlign(EHorizontalAlignment::HAlign_Fill)
+		[
+			SNew(SScaleBox)
+			.Stretch(EStretch::ScaleToFit)
+			[
+				SAssignNew(ChannelText, STextBlock)
+				.Justification(ETextJustify::Center)
+			]
+		]
+		+ SVerticalBox::Slot()
+		.FillHeight(0.5f)
+		.VAlign(EVerticalAlignment::VAlign_Fill)
+		.HAlign(EHorizontalAlignment::HAlign_Fill)
+		.Padding(0.f, 25.f)
+		[
+			SNew(SScaleBox)
+			.Stretch(EStretch::ScaleToFit)
+			[
+				SNew(SImage)
+				.Image(FAvaMediaStyle::GetBrush(TEXT("AvalancheMedia.UnrealIcon")))
+			]
+		]
+
+		+ SVerticalBox::Slot()
+		.FillHeight(0.1f)
+		.VAlign(EVerticalAlignment::VAlign_Fill)
+		.HAlign(EHorizontalAlignment::HAlign_Fill)
+		[
+			SNew(SScaleBox)
+			.Stretch(EStretch::ScaleToFit)
+			[
+				SNew(SHorizontalBox)
+				+ SHorizontalBox::Slot()
+				.FillWidth(0.2f)
+				.Padding(0.f, 5.f)
+				[
+					SNew(SScaleBox)
+					.Stretch(EStretch::ScaleToFit)
+					[
+						SNew(SCircularThrobber)
+						.Period(1.f)
+						.NumPieces(3)
+					]
+				]
+				+ SHorizontalBox::Slot()
+				.FillWidth(0.8f)
+				.Padding(5.f, 0.f)
+				[
+					SNew(SScaleBox)
+					.Stretch(EStretch::ScaleToFit)
+					[
+						SNew(STextBlock)
+						.Text(LOCTEXT("WaitingForPlayback_Text", "Waiting for Playback"))
+					]
+				]
+			]
+		]
+		
+		+ SVerticalBox::Slot()
+		.FillHeight(0.1f)
+		[
+			SNew(SBox)
+		]
+	];
+}
+
+void SAvaPlaceholderWidget::SetChannelName(const FText& InChannelName)
+{
+	if (ChannelText.IsValid())
+	{
+		ChannelText->SetText(InChannelName);
+	}
+}
+
+
+#undef LOCTEXT_NAMESPACE
diff --git a/Engine/Plugins/Experimental/Avalanche/Source/AvalancheMedia/Private/Http/AvaMediaHttpServer.cpp b/Engine/Plugins/Experimental/Avalanche/Source/AvalancheMedia/Private/Http/AvaMediaHttpServer.cpp
new file mode 100644
--- /dev/null
+++ b/Engine/Plugins/Experimental/Avalanche/Source/AvalancheMedia/Private/Http/AvaMediaHttpServer.cpp
@@ -0,0 +1,164 @@
+// Copyright Epic Games, Inc. All Rights Reserved.
+
+#include "Http/AvaMediaHttpServer.h"
+
+#include "AvaMediaModule.h"
+#include "HttpServerModule.h"
+#include "HttpServerResponse.h"
+#include "Playback/AvaMediaPlaybackServer.h"
+#include "Serialization/JsonWriter.h"
+
+DEFINE_LOG_CATEGORY(LogAvaMediaHttpServer)
+
+static const FString& GetHttpRouteVerbString(EHttpServerRequestVerbs InVerbs)
+{
+	static TMap<EHttpServerRequestVerbs, FString> VerbStrings =
+	{
+		{ EHttpServerRequestVerbs::VERB_POST, TEXT("POST") },
+		{ EHttpServerRequestVerbs::VERB_PUT, TEXT("PUT") },
+		{ EHttpServerRequestVerbs::VERB_GET, TEXT("GET") },
+		{ EHttpServerRequestVerbs::VERB_PATCH, TEXT("PATCH") },
+		{ EHttpServerRequestVerbs::VERB_DELETE, TEXT("DELETE") },
+		{ EHttpServerRequestVerbs::VERB_NONE, TEXT("NONE") },
+	};
+
+	if (const FString* FoundEntry = VerbStrings.Find(InVerbs))
+	{
+		return *FoundEntry;		
+	}
+
+	static const FString Unknown = TEXT("UNKNOWN");
+	return Unknown;
+}
+
+FHttpRouteBuilder::FHttpRouteBuilder(
+	const TSharedPtr<FAvaMediaHttpServer>& InHttpServer,
+	const TSharedPtr<FAvaMediaPlaybackServer>& InPlaybackServer)
+	: HttpServer(InHttpServer)
+	, PlaybackServer(InPlaybackServer)
+{
+}
+
+void FAvaMediaHttpServer::Start(int32 InPortToUse)
+{
+	PortToUse = InPortToUse;
+	HttpRouter = FHttpServerModule::Get().GetHttpRouter(PortToUse);
+
+	RegisterRoutes();
+
+	FHttpServerModule::Get().StartAllListeners();
+}
+
+void FAvaMediaHttpServer::RegisterRoutes()
+{
+	const TSharedPtr<FAvaMediaPlaybackServer> MediaPlaybackServer = FModuleManager::GetModulePtr<FAvaMediaModule>(TEXT("AvalancheMedia"))->GetMediaPlaybackServer();
+	check(MediaPlaybackServer.IsValid());
+	
+	// Map as per FAvaMediaPlaybackServer::Init(const FString& AssignedServerName)
+	FAvaMediaHttpServer::Builder(AsShared(), MediaPlaybackServer)
+	.Route<FAvaMediaPlaybackPing>(TEXT("/playback/ping"), EHttpServerRequestVerbs::VERB_POST, &FAvaMediaPlaybackServer::HandlePlaybackPing)
+	.Route<FAvaDeviceProviderDataRequest>(TEXT("/playback/devices"), EHttpServerRequestVerbs::VERB_POST, &FAvaMediaPlaybackServer::HandleDeviceProviderDataRequest)
+	.Route<FAvaMediaPlaybackRequest>(TEXT("/playback"), EHttpServerRequestVerbs::VERB_POST, &FAvaMediaPlaybackServer::HandlePlaybackRequest)
+	.Route<FAvaMediaAnimPlaybackRequest>(TEXT("/playback/animation"), EHttpServerRequestVerbs::VERB_POST, &FAvaMediaPlaybackServer::HandleAnimPlaybackRequest)
+	.Route<FAvaMediaBroadcastRequest>(TEXT("/broadcast"), EHttpServerRequestVerbs::VERB_POST, &FAvaMediaPlaybackServer::HandleBroadcastRequest)
+	.Route<FAvaMediaBroadcastStatusRequest>(TEXT("/broadcast/status"), EHttpServerRequestVerbs::VERB_POST, &FAvaMediaPlaybackServer::HandleBroadcastStatusRequest);
+
+	// We always want the ListRegisteredRoutes route bound, no matter what.
+	RegisterNewRoute(TEXT("ListRegisteredRoutes"), FHttpPath("/listroutes"), EHttpServerRequestVerbs::VERB_GET,
+		FHttpRequestHandler::CreateSP(this, &FAvaMediaHttpServer::HttpListOpenRoutes), true);
+}
+
+bool FAvaMediaHttpServer::GetRegisteredRoute(FName RouteName, FAvaMediaHttpRouteInfo& OutRouteInfo)
+{
+	if (RegisteredRoutes.Find(RouteName))
+	{
+		OutRouteInfo.RouteName = RouteName;
+		OutRouteInfo.RoutePath = RegisteredRoutes[RouteName].Handle->Path;
+		OutRouteInfo.RequestVerbs = RegisteredRoutes[RouteName].Handle->Verbs;
+		OutRouteInfo.InputContentType = RegisteredRoutes[RouteName].InputContentType;
+		OutRouteInfo.InputExpectedFormat = RegisteredRoutes[RouteName].InputExpectedFormat;
+		return true;
+	}
+	return false;
+}
+
+void FAvaMediaHttpServer::RegisterNewRoute(FName RouteName, const FHttpPath& HttpPath, const EHttpServerRequestVerbs& RequestVerbs, const FHttpRequestHandler& Handler, bool bOverrideIfBound /* = false */, FString OptionalContentType /* = TEXT("")*/, FString OptionalExpectedFormat /*= TEXT("")*/)
+{
+	FAvaMediaHttpRouteInfo InRouteInfo;
+	InRouteInfo.RouteName = RouteName;
+	InRouteInfo.RoutePath = HttpPath;
+	InRouteInfo.RequestVerbs = RequestVerbs;
+	InRouteInfo.InputContentType = OptionalContentType;
+	InRouteInfo.InputExpectedFormat = OptionalExpectedFormat;
+	RegisterNewRoute(InRouteInfo, Handler, bOverrideIfBound);
+}
+
+void FAvaMediaHttpServer::RegisterNewRoute(FAvaMediaHttpRouteInfo InRouteInfo, const FHttpRequestHandler& Handler, bool bOverrideIfBound /* = false */)
+{
+	if (HttpRouter.IsValid())
+	{
+		if (RegisteredRoutes.Find(InRouteInfo.RouteName))
+		{
+			if (!bOverrideIfBound)
+			{
+				UE_LOG(LogAvaMediaHttpServer, Error, TEXT("Failed to bind route with friendly key %s - a route at location %s already exists."), *InRouteInfo.RouteName.ToString(), *InRouteInfo.RoutePath.GetPath());
+				return;
+			}
+			UE_LOG(LogAvaMediaHttpServer, Log, TEXT("Overwriting route at friendly key %s - from %s to %s "), *InRouteInfo.RouteName.ToString(), *RegisteredRoutes[InRouteInfo.RouteName].Handle->Path, *InRouteInfo.RoutePath.GetPath());
+			HttpRouter->UnbindRoute(RegisteredRoutes[InRouteInfo.RouteName].Handle);
+		}
+		FAvaMediaHttpRouteDesc RouteDesc;
+		RouteDesc.Handle = HttpRouter->BindRoute(InRouteInfo.RoutePath, InRouteInfo.RequestVerbs, Handler);
+		RouteDesc.InputContentType = InRouteInfo.InputContentType;
+		RouteDesc.InputExpectedFormat = InRouteInfo.InputExpectedFormat;
+		RegisteredRoutes.Add(InRouteInfo.RouteName, RouteDesc);
+	}
+}
+
+void FAvaMediaHttpServer::CleanUpRoute(FName RouteName, bool bFailIfUnbound /* = false */)
+{
+	if (HttpRouter.IsValid())
+	{
+		if (RegisteredRoutes.Find(RouteName))
+		{
+			HttpRouter->UnbindRoute(RegisteredRoutes[RouteName].Handle);
+			RegisteredRoutes.Remove(RouteName);
+			UE_LOG(LogAvaMediaHttpServer, Log, TEXT("Route name %s was unbound!"), *RouteName.ToString());
+		}
+		else
+		{
+			UE_LOG(LogAvaMediaHttpServer, Warning, TEXT("Route name %s does not exist, could not unbind."), *RouteName.ToString());
+			check(!bFailIfUnbound);
+		}
+	}
+}
+
+bool FAvaMediaHttpServer::HttpListOpenRoutes(const FHttpServerRequest& Request, const FHttpResultCallback& OnComplete)
+{
+	FString ResponseStr;
+	TArray<FName> OutRouteKeys;
+	RegisteredRoutes.GetKeys(OutRouteKeys);
+	TSharedRef<TJsonWriter<>> JsonWriter = TJsonWriterFactory<>::Create(&ResponseStr);
+	JsonWriter->WriteArrayStart();
+	for (const FName& RouteKey : OutRouteKeys)
+	{
+		JsonWriter->WriteObjectStart();
+		JsonWriter->WriteValue(TEXT("name"), RouteKey.ToString());
+		JsonWriter->WriteValue(TEXT("route"), RegisteredRoutes[RouteKey].Handle->Path);
+		JsonWriter->WriteValue(TEXT("verb"), GetHttpRouteVerbString(RegisteredRoutes[RouteKey].Handle->Verbs));
+		if (!RegisteredRoutes[RouteKey].InputContentType.IsEmpty())
+		{
+			JsonWriter->WriteValue(TEXT("inputContentType"), RegisteredRoutes[RouteKey].InputContentType);
+		}
+		if (!RegisteredRoutes[RouteKey].InputExpectedFormat.IsEmpty())
+		{
+			JsonWriter->WriteValue(TEXT("inputExpectedFormat"), RegisteredRoutes[RouteKey].InputExpectedFormat);
+		}
+		JsonWriter->WriteObjectEnd();
+	}
+	JsonWriter->WriteArrayEnd();
+	JsonWriter->Close();
+	TUniquePtr<FHttpServerResponse> Response = FHttpServerResponse::Create(ResponseStr, TEXT("application/json"));
+	OnComplete(MoveTemp(Response));
+	return true;
+}
\ No newline at end of file
diff --git a/Engine/Plugins/Experimental/Avalanche/Source/AvalancheMedia/Private/Http/AvaMediaHttpServer.h b/Engine/Plugins/Experimental/Avalanche/Source/AvalancheMedia/Private/Http/AvaMediaHttpServer.h
new file mode 100644
--- /dev/null
+++ b/Engine/Plugins/Experimental/Avalanche/Source/AvalancheMedia/Private/Http/AvaMediaHttpServer.h
@@ -0,0 +1,243 @@
+// Copyright Epic Games, Inc. All Rights Reserved.
+
+#pragma once
+
+#include "IMessageContext.h"
+#include "IMessageAttachment.h"
+#include "HttpResultCallback.h"
+#include "HttpRouteHandle.h"
+#include "HttpServerRequest.h"
+#include "IHttpRouter.h"
+#include "StructDeserializer.h"
+#include "Backends/JsonStructDeserializerBackend.h"
+#include "Serialization/MemoryReader.h"
+
+#include "AvaMediaHttpServer.generated.h"
+
+class FAvaMediaPlaybackServer;
+
+namespace UE::Ava::Web::Private
+{
+	template <typename MessageType>
+	[[nodiscard]] bool DeserializeRequest(const FHttpServerRequest& InRequest, const FHttpResultCallback* InCompleteCallback, MessageType& OutDeserializedRequest);
+}
+
+// Required by Messaging API which this wraps, so a dummy context is required
+class FEmptyMessageContext
+	: public IMessageContext
+{
+public:
+	FEmptyMessageContext() = default;
+	virtual ~FEmptyMessageContext() override = default;
+
+public:
+	//~ IMessageContext interface
+	virtual const TMap<FName, FString>& GetAnnotations() const override { return EmptyMap; }
+	virtual TSharedPtr<IMessageAttachment, ESPMode::ThreadSafe> GetAttachment() const override { return nullptr; }
+	virtual const FDateTime& GetExpiration() const override { return EmptyDate; }
+	virtual const void* GetMessage() const override { return nullptr; }
+	virtual const TWeakObjectPtr<UScriptStruct>& GetMessageTypeInfo() const override { return EmptyTypeInfo; }
+	virtual TSharedPtr<IMessageContext, ESPMode::ThreadSafe> GetOriginalContext() const override { return nullptr; }
+	virtual const TArray<FMessageAddress>& GetRecipients() const override { return EmptyRecipients; }
+	virtual EMessageScope GetScope() const override { return {}; }
+	virtual EMessageFlags GetFlags() const override { return {}; }
+	virtual const FMessageAddress& GetSender() const override { return EmptyAddress; }
+	virtual const FMessageAddress& GetForwarder() const override { return EmptyAddress; }
+	virtual ENamedThreads::Type GetSenderThread() const override  { return {}; }
+	virtual const FDateTime& GetTimeForwarded() const override { return EmptyDate; }
+	virtual const FDateTime& GetTimeSent() const override { return EmptyDate; }
+
+private:
+	TMap<FName, FString> EmptyMap;
+	FDateTime EmptyDate;
+	TWeakObjectPtr<UScriptStruct> EmptyTypeInfo;
+	TArray<FMessageAddress> EmptyRecipients;
+	FMessageAddress EmptyAddress;
+};
+
+USTRUCT()
+struct FAvaMediaHttpRouteInfo
+{
+	GENERATED_BODY()
+public:
+	FName RouteName;
+	FHttpPath RoutePath;
+	EHttpServerRequestVerbs RequestVerbs;
+	FString InputContentType;
+	FString InputExpectedFormat;
+	FAvaMediaHttpRouteInfo()
+	{
+		RouteName = FName(TEXT(""));
+		RoutePath = FHttpPath();
+		RequestVerbs = EHttpServerRequestVerbs::VERB_NONE;
+		InputContentType = TEXT("");
+		InputExpectedFormat = TEXT("");
+	}
+	FAvaMediaHttpRouteInfo(FName InRouteName, FHttpPath InRoutePath, EHttpServerRequestVerbs InRequestVerbs, FString InContentType = TEXT(""), FString InExpectedFormat = TEXT(""))
+	{
+		RouteName = InRouteName;
+		RoutePath = InRoutePath;
+		RequestVerbs = InRequestVerbs;
+		InputContentType = InContentType;
+		InputExpectedFormat = InExpectedFormat;
+	}
+};
+
+USTRUCT()
+struct FAvaMediaHttpRouteDesc
+{
+	GENERATED_BODY()
+public:
+	FHttpRouteHandle Handle;
+	FString InputContentType;
+	FString InputExpectedFormat;
+	FAvaMediaHttpRouteDesc()
+	{
+	
+	}
+	FAvaMediaHttpRouteDesc(FHttpRouteHandle InHandle, FString InContentType, FString InExpectedFormat)
+	{
+		Handle = InHandle;
+		InputContentType = InContentType;
+		InputExpectedFormat = InExpectedFormat;
+	}
+};
+
+DECLARE_LOG_CATEGORY_EXTERN(LogAvaMediaHttpServer, Log, All);
+
+class FAvaMediaHttpServer;
+
+struct FHttpRouteBuilder
+{
+public:
+	template <typename MessageType, typename HandlerType>
+	struct TMessageHandler
+	{
+		typedef void (HandlerType::*FuncType)(const MessageType&, const TSharedRef<IMessageContext, ESPMode::ThreadSafe>&);
+	};
+
+public:
+	FHttpRouteBuilder(const TSharedPtr<FAvaMediaHttpServer>& InHttpServer, const TSharedPtr<FAvaMediaPlaybackServer>& InPlaybackServer);
+	
+public:	
+	template <typename MessageType>
+	FHttpRouteBuilder& Route(const FString& InPath, const EHttpServerRequestVerbs& InVerb, typename TMessageHandler<MessageType, FAvaMediaPlaybackServer>::FuncType InHandlerFunc);
+	
+private:
+	TSharedPtr<FAvaMediaHttpServer> HttpServer;
+	TSharedPtr<FAvaMediaPlaybackServer> PlaybackServer;
+};
+
+class FAvaMediaHttpServer : public TSharedFromThis<FAvaMediaHttpServer>
+{
+public:
+	typedef struct FHttpRouteBuilder Builder;
+	
+	FAvaMediaHttpServer() = default;
+	virtual ~FAvaMediaHttpServer() = default;
+
+	void Start(int32 InPortToUse = 10123);
+
+	void RegisterRoutes();
+
+	/**
+	 * Try to get a route registered under given friendly name. Returns false if could not be found.
+	 */
+	bool GetRegisteredRoute(FName RouteName, FAvaMediaHttpRouteInfo& OutRouteInfo);
+
+	void RegisterNewRoute(FAvaMediaHttpRouteInfo InRouteInfo, const FHttpRequestHandler& Handler, bool bOverrideIfBound = false);
+
+	/**
+	 * Register a new route.
+	 * Will override existing routes if option is set, otherwise will error and fail to bind.
+	 */
+	void RegisterNewRoute(FName RouteName, const FHttpPath& HttpPath, const EHttpServerRequestVerbs& RequestVerbs, const FHttpRequestHandler& Handler, bool bOverrideIfBound = false, FString OptionalContentType = TEXT(""), FString OptionalExpectedFormat = TEXT(""));
+
+	/**
+	 * Clean up a route.
+	 * Can be set to fail if trying to unbind an unbound route.
+	 */
+	void CleanUpRoute(FName RouteName, bool bFailIfUnbound = false);
+
+	/**
+	 * Default Route Listing http call. Spits out all registered routes and describes them via a REST API call.
+	 * Always registered at /listroutes GET by default
+	 */
+	bool HttpListOpenRoutes(const FHttpServerRequest& Request, const FHttpResultCallback& OnComplete);
+	
+	bool IsRunning() const { return HttpRouter.IsValid(); }
+
+protected:
+	int PortToUse = 10123;
+	TSharedPtr<IHttpRouter> HttpRouter;
+	TMap<FName, FAvaMediaHttpRouteDesc> RegisteredRoutes;
+};
+
+template <typename MessageType>
+bool UE::Ava::Web::Private::DeserializeRequest(
+	const FHttpServerRequest& InRequest,
+	const FHttpResultCallback* InCompleteCallback,
+	MessageType& OutDeserializedRequest)
+{
+	TArray<uint8> TCHARPayload;
+
+	const int32 StartIndex = TCHARPayload.Num();
+	TCHARPayload.AddUninitialized(FUTF8ToTCHAR_Convert::ConvertedLength((ANSICHAR*)InRequest.Body.GetData(), InRequest.Body.Num() / sizeof(ANSICHAR)) * sizeof(TCHAR));
+	FUTF8ToTCHAR_Convert::Convert((TCHAR*)(TCHARPayload.GetData() + StartIndex), (TCHARPayload.Num() - StartIndex) / sizeof(TCHAR), (ANSICHAR*)InRequest.Body.GetData(), InRequest.Body.Num() / sizeof(ANSICHAR));
+
+	FMemoryReaderView Reader(TCHARPayload);
+	FJsonStructDeserializerBackend DeserializerBackend(Reader);
+	if (!FStructDeserializer::Deserialize(&OutDeserializedRequest, *MessageType::StaticStruct(), DeserializerBackend, FStructDeserializerPolicies()))
+	{
+		if (InCompleteCallback)
+		{
+			TUniquePtr<FHttpServerResponse> Response = MakeUnique<FHttpServerResponse>();
+			Response->Code = EHttpServerResponseCodes::BadRequest;
+			(*InCompleteCallback)(MoveTemp(Response));
+		}
+		return false;
+	}
+
+	return true;
+}
+
+template <typename MessageType>
+FHttpRouteBuilder& FHttpRouteBuilder::Route(
+	const FString& InPath,
+	const EHttpServerRequestVerbs& InVerb,
+	typename TMessageHandler<MessageType, FAvaMediaPlaybackServer>::FuncType InHandlerFunc)
+{
+	static_assert(TModels<CStaticStructProvider, MessageType>::Value, "MessageType must be a UStruct");
+
+	TWeakPtr<FAvaMediaPlaybackServer> WeakPlaybackServerPtr = TWeakPtr<FAvaMediaPlaybackServer>(PlaybackServer);
+	
+	HttpServer->RegisterNewRoute(MessageType::StaticStruct()->GetFName(), FHttpPath(InPath), InVerb,
+	FHttpRequestHandler::CreateLambda([WeakPlaybackServerPtr, HandlerFunc = MoveTemp(InHandlerFunc)](const FHttpServerRequest& Request, const FHttpResultCallback& OnComplete)
+	{
+		if (WeakPlaybackServerPtr.IsValid())
+		{
+			MessageType DeserializedMessage;
+			if (!UE::Ava::Web::Private::DeserializeRequest(Request, &OnComplete, DeserializedMessage))
+			{
+				return false;
+			}
+
+			static TSharedRef<IMessageContext> EmptyContext = MakeShared<FEmptyMessageContext>();
+			(WeakPlaybackServerPtr.Pin().Get()->*HandlerFunc)(DeserializedMessage, EmptyContext);
+
+			if (OnComplete)
+			{
+				TUniquePtr<FHttpServerResponse> Response = MakeUnique<FHttpServerResponse>();
+				Response->Headers.Add(TEXT("content-type"), { TEXT("application/json") });
+				Response->Code = EHttpServerResponseCodes::Ok;
+
+				OnComplete(MoveTemp(Response));
+			}
+			
+			return true;
+		}
+		return false;
+	}), true);
+
+	return *this;
+}
diff --git a/Engine/Plugins/Experimental/Avalanche/Source/AvalancheMedia/Private/ModularFeature/AvaSyncProviderFeatureTypes.cpp b/Engine/Plugins/Experimental/Avalanche/Source/AvalancheMedia/Private/ModularFeature/AvaSyncProviderFeatureTypes.cpp
new file mode 100644
--- /dev/null
+++ b/Engine/Plugins/Experimental/Avalanche/Source/AvalancheMedia/Private/ModularFeature/AvaSyncProviderFeatureTypes.cpp
@@ -0,0 +1,99 @@
+// Copyright Epic Games, Inc. All Rights Reserved.
+
+#include "ModularFeature/AvaSyncProviderFeatureTypes.h"
+
+#include "Engine/Engine.h"
+#include "Misc/App.h"
+#include "Misc/Paths.h"
+
+#define LOCTEXT_NAMESPACE "AvaSyncProviderFeatureTypes"
+
+FAvaSyncConnectionInfo::FAvaSyncConnectionInfo()
+	: EngineVersion(FNetworkVersion::GetLocalNetworkVersion())
+	, InstanceId(FApp::GetInstanceId())
+	, SessionId(FApp::GetSessionId())
+	, HostName(FPlatformProcess::ComputerName())
+	, ProjectName(FApp::GetProjectName())
+{
+	// We only want the dirname, not absolute path
+	ProjectDir = GetBasename(FPaths::ProjectDir());
+
+	// Figure out which type of instance it is
+	if (GEngine == nullptr)
+	{
+		InstanceType = EAvaSyncEngineType::Unknown;
+	}
+	else if (IsRunningDedicatedServer())
+	{
+		InstanceType = EAvaSyncEngineType::Server;
+	}
+	else if (IsRunningCommandlet())
+	{
+		InstanceType = EAvaSyncEngineType::Commandlet;
+	}
+	else if (GEngine->IsEditor())
+	{
+		InstanceType = EAvaSyncEngineType::Editor;
+	}
+	else if (IsRunningGame())
+	{
+		InstanceType = EAvaSyncEngineType::Game;
+	}
+	else
+	{
+		InstanceType = EAvaSyncEngineType::Other;
+	}	
+}
+
+FString FAvaSyncConnectionInfo::ToString() const
+{
+	return FString::Printf(
+		TEXT("EngineVersion: %d, HostName: %s, InstanceId: %s, InstanceType: %s, SessionId: %s, ProjectName: %s, ProjectDir: %s"),
+		EngineVersion,
+		*HostName,
+		*InstanceId.ToString(),
+		*UEnum::GetValueAsString(InstanceType),
+		*SessionId.ToString(),
+		*ProjectName,
+		*ProjectDir
+	);
+}
+
+FString FAvaSyncConnectionInfo::GetBasename(const FString& InPath)
+{
+	FString LocalPath = InPath;
+	LocalPath.RemoveFromEnd(TEXT("/"));
+	return FPaths::GetBaseFilename(LocalPath);
+}
+
+FText FAvaSyncConnectionInfo::GetHumanReadableInstanceType() const
+{
+	FText Result;
+
+	switch (InstanceType)
+	{
+	case EAvaSyncEngineType::Server:
+		Result = LOCTEXT("Instance_Type_Server", "Server");
+		break;
+	case EAvaSyncEngineType::Commandlet:
+		Result = LOCTEXT("Instance_Type_Commandlet", "Commandlet");
+		break;
+	case EAvaSyncEngineType::Editor:
+		Result = LOCTEXT("Instance_Type_Editor", "Editor");
+		break;
+	case EAvaSyncEngineType::Game:
+		Result = LOCTEXT("Instance_Type_Game", "Game");
+		break;
+	case EAvaSyncEngineType::Other:
+		Result = LOCTEXT("Instance_Type_Other", "Other");
+		break;
+	case EAvaSyncEngineType::Unknown:
+		Result = LOCTEXT("Instance_Type_Unknown", "Unknown");
+		break;
+	default: break;
+	}
+	
+	return Result;
+}
+
+#undef LOCTEXT_NAMESPACE
diff --git a/Engine/Plugins/Experimental/Avalanche/Source/AvalancheMedia/Private/OutputDevices/AvaDeviceProviderData.cpp b/Engine/Plugins/Experimental/Avalanche/Source/AvalancheMedia/Private/OutputDevices/AvaDeviceProviderData.cpp
new file mode 100644
--- /dev/null
+++ b/Engine/Plugins/Experimental/Avalanche/Source/AvalancheMedia/Private/OutputDevices/AvaDeviceProviderData.cpp
@@ -0,0 +1,168 @@
+// Copyright Epic Games, Inc. All Rights Reserved.
+
+#include "OutputDevices/AvaDeviceProviderData.h"
+
+#include "MediaOutput.h"
+#include "IMediaIOCoreModule.h"
+#include "StructDeserializer.h"
+#include "StructSerializer.h"
+#include "Backends/JsonStructDeserializerBackend.h"
+#include "Backends/JsonStructSerializerBackend.h"
+#include "Backends/XmlStructSerializerBackend.h"
+#include "HAL/FileManager.h"
+#include "Misc/Paths.h"
+#include "UObject/UObjectIterator.h"
+
+void FAvaDeviceProviderData::PopulateFrom(const FName& InName, const IMediaIOCoreDeviceProvider* InDeviceProvider)
+{
+	Name = InName;
+	Connections = InDeviceProvider->GetConnections();
+	Configurations = InDeviceProvider->GetConfigurations();
+	Devices = InDeviceProvider->GetDevices();
+	//Modes = InDeviceProvider->GetModes(); TODO
+	InputConfigurations = InDeviceProvider->GetInputConfigurations();
+	OutputConfigurations = InDeviceProvider->GetOutputConfigurations();
+	TimecodeConfigurations = InDeviceProvider->GetTimecodeConfigurations();
+	DefaultConfiguration = InDeviceProvider->GetDefaultConfiguration();
+	DefaultMode = InDeviceProvider->GetDefaultMode();
+	DefaultInputConfiguration = InDeviceProvider->GetDefaultInputConfiguration();
+	DefaultOutputConfiguration = InDeviceProvider->GetDefaultOutputConfiguration();
+	DefaultTimecodeConfiguration = InDeviceProvider->GetDefaultTimecodeConfiguration();
+	
+#if WITH_EDITOR
+	bShowInputTransportInSelector = InDeviceProvider->ShowInputTransportInSelector();
+	bShowOutputTransportInSelector = InDeviceProvider->ShowOutputTransportInSelector();
+	bShowInputKeyInSelector = InDeviceProvider->ShowInputKeyInSelector();
+	bShowOutputKeyInSelector = InDeviceProvider->ShowOutputKeyInSelector();
+	bShowReferenceInSelector = InDeviceProvider->ShowReferenceInSelector();
+#endif
+}
+namespace FDeviceProviderDataHelper
+{
+	inline void ApplyServerName(const FString& InServerName, FMediaIODevice& OutDevice)
+	{
+		FString NewDeviceName = FString::Printf(TEXT("%s: %s"), *InServerName, * OutDevice.DeviceName.ToString());
+		OutDevice.DeviceName = *NewDeviceName;
+	}
+	inline void ApplyServerName(const FString& InServerName, FMediaIOConnection& OutConnection)
+	{
+		ApplyServerName(InServerName, OutConnection.Device);
+	}
+	inline void ApplyServerName(const FString& InServerName, FMediaIOConfiguration& OutConfig)
+	{
+		ApplyServerName(InServerName, OutConfig.MediaConnection);
+	}
+	inline void ApplyServerName(const FString& InServerName, FMediaIOInputConfiguration& OutConfig)
+	{
+		ApplyServerName(InServerName, OutConfig.MediaConfiguration);
+	}
+	inline void ApplyServerName(const FString& InServerName, FMediaIOOutputConfiguration& OutConfig)
+	{
+		ApplyServerName(InServerName, OutConfig.MediaConfiguration);
+	}
+	inline void ApplyServerName(const FString& InServerName, FMediaIOVideoTimecodeConfiguration& OutConfig)
+	{
+		ApplyServerName(InServerName, OutConfig.MediaConfiguration);
+	}
+}
+
+void FAvaDeviceProviderData::ApplyServerName(const FString& InServerName)
+{
+	for (FMediaIOConnection& Connection : Connections)
+	{
+		FDeviceProviderDataHelper::ApplyServerName(InServerName,Connection);
+	}
+	for (FMediaIOConfiguration& Config : Configurations)
+	{
+		FDeviceProviderDataHelper::ApplyServerName(InServerName,Config);
+	}
+	for (FMediaIODevice& Device : Devices)
+	{
+		FDeviceProviderDataHelper::ApplyServerName(InServerName,Device);
+	}
+	for (FMediaIOInputConfiguration& Config : InputConfigurations)
+	{
+		FDeviceProviderDataHelper::ApplyServerName(InServerName,Config);
+	}
+	for (FMediaIOOutputConfiguration& Config : OutputConfigurations)
+	{
+		FDeviceProviderDataHelper::ApplyServerName(InServerName,Config);
+	}
+	for (FMediaIOVideoTimecodeConfiguration& Config : TimecodeConfigurations)
+	{
+		FDeviceProviderDataHelper::ApplyServerName(InServerName,Config);
+	}
+
+	FDeviceProviderDataHelper::ApplyServerName(InServerName,DefaultConfiguration);
+	FDeviceProviderDataHelper::ApplyServerName(InServerName,DefaultInputConfiguration);
+	FDeviceProviderDataHelper::ApplyServerName(InServerName,DefaultOutputConfiguration);
+	FDeviceProviderDataHelper::ApplyServerName(InServerName,DefaultTimecodeConfiguration);
+}
+
+namespace FDeviceProviderDataListHelper
+{
+	inline const TCHAR* GetJsonSaveFilepath()
+	{
+		static const FString SaveFilepath = FPaths::ProjectConfigDir() / TEXT("DeviceProviders.json");
+		return *SaveFilepath;
+	}
+	inline const TCHAR* GetXmlSaveFilepath()
+	{
+		static const FString SaveFilepath = FPaths::ProjectConfigDir() / TEXT("DeviceProviders.xml");
+		return *SaveFilepath;
+	}
+}
+
+
+bool FAvaDeviceProviderDataList::SaveToJson()
+{
+	TUniquePtr<FArchive> FileWriter(
+		IFileManager::Get().CreateFileWriter(FDeviceProviderDataListHelper::GetJsonSaveFilepath()));
+	FJsonStructSerializerBackend Backend(*FileWriter, EStructSerializerBackendFlags::Default);
+	FStructSerializer::Serialize(*this, Backend);
+
+	FileWriter->Close();
+	return !FileWriter->IsError();
+}
+
+bool FAvaDeviceProviderDataList::SaveToXml()
+{
+	TUniquePtr<FArchive> FileWriter(
+		IFileManager::Get().CreateFileWriter(FDeviceProviderDataListHelper::GetXmlSaveFilepath()));
+	FXmlStructSerializerBackend Backend(*FileWriter, EStructSerializerBackendFlags::Default);
+	FStructSerializer::Serialize(*this, Backend);
+	Backend.SaveDocument();
+
+	FileWriter->Close();
+	return !FileWriter->IsError();
+}
+
+bool FAvaDeviceProviderDataList::LoadFromJson()
+{
+	const TUniquePtr<FArchive> FileReader = TUniquePtr<FArchive>(
+		IFileManager::Get().CreateFileReader(FDeviceProviderDataListHelper::GetJsonSaveFilepath()));
+	FJsonStructDeserializerBackend Backend(*FileReader);
+	FStructDeserializer::Deserialize(*this, Backend);
+
+	FileReader->Close();
+	return !FileReader->IsError();
+}
+
+void FAvaDeviceProviderDataList::Populate(const FString& InAssignedServerName)
+{
+	ServerName = InAssignedServerName;
+	
+	if (!IMediaIOCoreModule::IsAvailable())
+	{
+		UE_LOG(LogTemp, Error, TEXT("IMediaIOCoreModule is not available, sad face."));
+		return;
+	}
+	
+	const TConstArrayView<IMediaIOCoreDeviceProvider*> LocalDeviceProviders = IMediaIOCoreModule::Get().GetDeviceProviders();
+	for (IMediaIOCoreDeviceProvider* DeviceProvider : LocalDeviceProviders)
+	{
+		FAvaDeviceProviderData DeviceProviderData;
+		DeviceProviderData.PopulateFrom(DeviceProvider->GetFName(), DeviceProvider);
+		DeviceProviders.Add(DeviceProviderData);
+	}
+}
\ No newline at end of file
diff --git a/Engine/Plugins/Experimental/Avalanche/Source/AvalancheMedia/Private/OutputDevices/AvaDeviceProviderProxy.cpp b/Engine/Plugins/Experimental/Avalanche/Source/AvalancheMedia/Private/OutputDevices/AvaDeviceProviderProxy.cpp
new file mode 100644
--- /dev/null
+++ b/Engine/Plugins/Experimental/Avalanche/Source/AvalancheMedia/Private/OutputDevices/AvaDeviceProviderProxy.cpp
@@ -0,0 +1,557 @@
+// Copyright Epic Games, Inc. All Rights Reserved.
+
+#include "OutputDevices/AvaDeviceProviderProxy.h"
+
+#include "AvalancheBroadcast.h"
+#include "IAvaMediaModule.h"
+#include "IMediaIOCoreDeviceProvider.h"
+#include "IMediaIOCoreModule.h"
+#include "MediaOutput.h"
+#include "OutputDevices/AvaMediaOutputUtils.h"
+#include "UObject/UObjectIterator.h"
+
+//-----------------------------------------------------------------------------
+//----- FAvaDeviceProviderProxy -----------------------------------------------
+//-----------------------------------------------------------------------------
+
+FAvaDeviceProviderProxy::FAvaDeviceProviderProxy(FAvaDeviceProviderProxyManager* InManager, const FAvaDeviceProviderData& InData)
+: Manager(InManager)
+, Data(InData)
+{
+	
+}
+
+TArray<FMediaIOConfiguration> FAvaDeviceProviderProxy::GetConfigurations(bool bAllowInput, bool bAllowOutput) const
+{
+	TArray<FMediaIOConfiguration> Out;
+	for (const FMediaIOConfiguration& Config : Data.Configurations)
+	{
+		if ((Config.bIsInput && bAllowInput) || (!Config.bIsInput && bAllowOutput))
+		{
+			Out.Add(Config);
+		}
+	}
+	return Out;
+}
+
+TArray<FMediaIOMode> FAvaDeviceProviderProxy::GetModes(const FMediaIODevice& InDevice, bool bInOutput) const
+{
+		return Data.Modes; // TODO	
+}
+
+//-----------------------------------------------------------------------------
+//----- FAvaDeviceProviderWrapper ---------------------------------------------
+//-----------------------------------------------------------------------------
+
+FAvaDeviceProviderWrapper::FAvaDeviceProviderWrapper(FName InName, FAvaDeviceProviderProxyManager* InManager)
+	: Manager(InManager)
+	, Name(InName)
+	, DummyProxy(InManager, FAvaDeviceProviderData())
+{ }
+
+void FAvaDeviceProviderWrapper::AddProvider(const FString& InServerName, IMediaIOCoreDeviceProvider* InProvider, bool bInIsLocal)
+{
+	if (bInIsLocal)
+	{
+		LocalProvider = InProvider;
+	}
+
+	for (FProviderInfo& Info : Providers)
+	{
+		if (Info.ServerName == InServerName)
+		{
+			Info.Provider = InProvider;
+			return;
+		}
+	}
+	
+	Providers.Add(FProviderInfo{InServerName, InProvider, nullptr});
+}
+
+void FAvaDeviceProviderWrapper::AddProvider(const FString& InServerName, TUniquePtr<FAvaDeviceProviderProxy>&& InProviderProxy)
+{
+	for (FProviderInfo& Info : Providers)
+	{
+		if (Info.ServerName == InServerName)
+		{
+			Info.Proxy = MoveTemp(InProviderProxy);
+			Info.Provider = Info.Proxy.Get();
+			return;
+		}
+	}
+
+	Providers.Add(FProviderInfo{InServerName, InProviderProxy.Get(), MoveTemp(InProviderProxy)});
+}
+
+void FAvaDeviceProviderWrapper::RemoveProvider(IMediaIOCoreDeviceProvider* InProvider)
+{
+	for (int32 i = 0; i < Providers.Num(); ++i)
+	{
+		if (Providers[i].Provider == InProvider)
+		{
+			OnRemoveProvider(InProvider);
+			Providers.RemoveAt(i);
+			return;
+		}
+	}
+}
+
+void FAvaDeviceProviderWrapper::RemoveProvider(const FString& InServerName)
+{
+	for (int32 i = 0; i < Providers.Num(); ++i)
+	{
+		if (Providers[i].ServerName == InServerName)
+		{
+			OnRemoveProvider(Providers[i].Provider);
+			Providers.RemoveAt(i);
+			return;
+		}
+	}
+}
+
+const IMediaIOCoreDeviceProvider* FAvaDeviceProviderWrapper::GetProviderForDeviceName(const FName& InDeviceName) const
+{
+	const FString DeviceName = InDeviceName.ToString();
+	for (const FProviderInfo& ProviderInfo : Providers)
+	{
+		// For device from another server, the device name starts with the "ServerName:"
+		// See FAvaDeviceProviderData::ApplyServerName.
+		if (DeviceName.StartsWith(ProviderInfo.ServerName))
+		{
+			return (ProviderInfo.Proxy.IsValid()) ? ProviderInfo.Proxy.Get() : ProviderInfo.Provider; 
+		}
+	}
+	return GetLocalProvider();
+}
+
+const IMediaIOCoreDeviceProvider* FAvaDeviceProviderWrapper::GetProviderForServer(const FString& InServerName) const
+{
+	for (const FProviderInfo& ProviderInfo : Providers)
+	{
+		// For device from another server, the device name starts with the "ServerName:"
+		// See FAvaDeviceProviderData::ApplyServerName.
+		if (ProviderInfo.ServerName == InServerName)
+		{
+			return (ProviderInfo.Proxy.IsValid()) ? ProviderInfo.Proxy.Get() : ProviderInfo.Provider; 
+		}
+	}
+	return nullptr;
+}
+
+FString FAvaDeviceProviderWrapper::GetServerNameForDevice(const FName& InDeviceName) const
+{
+	// Search which provider proxy has the given device.
+	for (const FProviderInfo& ProviderInfo : Providers)
+	{
+		if (FAvaDeviceProviderProxyManager::IsDevicePresent(ProviderInfo.Provider, InDeviceName))
+		{
+			return ProviderInfo.ServerName;
+		}
+	}
+	return FString();
+}
+
+bool FAvaDeviceProviderWrapper::ShouldRemove() const
+{
+	// If it is either empty or just wrapping the local provider.
+	return (Providers.IsEmpty() || (Providers.Num() == 1 && Providers[0].Provider == LocalProvider));
+}
+
+TArray<FMediaIOConnection> FAvaDeviceProviderWrapper::GetConnections() const
+{
+	if (Providers.Num() == 1)
+	{
+		return Providers[0].Provider->GetConnections();
+	}
+
+	TArray<FMediaIOConnection> Connections;
+	for (const FProviderInfo& Info : Providers)
+	{
+		Connections.Append(Info.Provider->GetConnections());
+	}
+	return Connections;
+}
+
+TArray<FMediaIOConfiguration> FAvaDeviceProviderWrapper::GetConfigurations() const
+{
+	if (Providers.Num() == 1)
+	{
+		return Providers[0].Provider->GetConfigurations();
+	}
+
+	TArray<FMediaIOConfiguration> Configurations;
+	for (const FProviderInfo& Info : Providers)
+	{
+		Configurations.Append(Info.Provider->GetConfigurations());
+	}
+	return Configurations;
+}
+
+TArray<FMediaIOConfiguration> FAvaDeviceProviderWrapper::GetConfigurations(bool bInAllowInput, bool bInAllowOutput) const
+{
+	if (Providers.Num() == 1)
+	{
+		return Providers[0].Provider->GetConfigurations(bInAllowInput, bInAllowOutput);
+	}
+
+	TArray<FMediaIOConfiguration> Configurations;
+	for (const FProviderInfo& Info : Providers)
+	{
+		Configurations.Append(Info.Provider->GetConfigurations(bInAllowInput, bInAllowOutput));
+	}
+	return Configurations;
+}
+
+TArray<FMediaIODevice> FAvaDeviceProviderWrapper::GetDevices() const
+{
+	if (Providers.Num() == 1)
+	{
+		return Providers[0].Provider->GetDevices();
+	}
+
+	TArray<FMediaIODevice> Devices;
+	for (const FProviderInfo& Info : Providers)
+	{
+		Devices.Append(Info.Provider->GetDevices());
+	}
+	return Devices;
+}
+
+TArray<FMediaIOMode> FAvaDeviceProviderWrapper::GetModes(const FMediaIODevice& InDevice, bool bInOutput) const
+{
+	// SMELL: not sure this is correct.
+	// Can you accumulate the modes for a device for different providers? Don't think so.
+	if (Providers.Num() == 1)
+	{
+		return Providers[0].Provider->GetModes(InDevice, bInOutput);
+	}
+
+	TArray<FMediaIOMode> Modes;
+	for (const FProviderInfo& Info : Providers)
+	{
+		Modes.Append(Info.Provider->GetModes(InDevice, bInOutput));
+	}
+	return Modes;
+}
+
+TArray<FMediaIOInputConfiguration> FAvaDeviceProviderWrapper::GetInputConfigurations() const
+{
+	if (Providers.Num() == 1)
+	{
+		return Providers[0].Provider->GetInputConfigurations();
+	}
+
+	TArray<FMediaIOInputConfiguration> InputConfigurations;
+	for (const FProviderInfo& Info : Providers)
+	{
+		InputConfigurations.Append(Info.Provider->GetInputConfigurations());
+	}
+	return InputConfigurations;
+}
+
+TArray<FMediaIOOutputConfiguration> FAvaDeviceProviderWrapper::GetOutputConfigurations() const
+{
+	if (Providers.Num() == 1)
+	{
+		return Providers[0].Provider->GetOutputConfigurations();
+	}
+
+	TArray<FMediaIOOutputConfiguration> OutputConfigurations;
+	for (const FProviderInfo& Info : Providers)
+	{
+		OutputConfigurations.Append(Info.Provider->GetOutputConfigurations());
+	}
+	return OutputConfigurations;
+}
+
+TArray<FMediaIOVideoTimecodeConfiguration> FAvaDeviceProviderWrapper::GetTimecodeConfigurations() const
+{
+	if (Providers.Num() == 1)
+	{
+		return Providers[0].Provider->GetTimecodeConfigurations();
+	}
+
+	TArray<FMediaIOVideoTimecodeConfiguration> TimecodeConfigurations;
+	for (const FProviderInfo& Info : Providers)
+	{
+		TimecodeConfigurations.Append(Info.Provider->GetTimecodeConfigurations());
+	}
+	return TimecodeConfigurations;
+}
+
+//-----------------------------------------------------------------------------
+//----- FAvaDeviceProviderProxyManager ----------------------------------------
+//-----------------------------------------------------------------------------
+FString FAvaDeviceProviderProxyManager::LocalServerName(TEXT("Local (This Process)"));
+
+const FAvaDeviceProviderWrapper* FAvaDeviceProviderProxyManager::GetDeviceProviderWrapper(FName InDeviceProviderName) const
+{
+	const TUniquePtr<FAvaDeviceProviderWrapper>* DeviceProviderWrapper = DeviceProviderWrappers.Find(InDeviceProviderName);
+	return DeviceProviderWrapper ? DeviceProviderWrapper->Get() : nullptr;
+}
+
+TArray<const IMediaIOCoreDeviceProvider*> FAvaDeviceProviderProxyManager::GetDeviceProvidersForServer(const FString& InServerName) const
+{
+	TArray<const IMediaIOCoreDeviceProvider*> Providers;
+	for (const TPair<FName, TUniquePtr<FAvaDeviceProviderWrapper>>& WrapperEntry : DeviceProviderWrappers)
+	{
+		if (const IMediaIOCoreDeviceProvider* const Provider = WrapperEntry.Value->GetProviderForServer(InServerName))
+		{
+			Providers.Add(Provider);
+		}
+	}
+	return Providers;
+}
+
+FString FAvaDeviceProviderProxyManager::FindServerNameForDevice(const FName& InDeviceProviderName, const FName& InDeviceName) const
+{
+	const FAvaDeviceProviderWrapper* DeviceProviderWrapper = GetDeviceProviderWrapper(InDeviceProviderName);
+
+	if (DeviceProviderWrapper)
+	{
+		// If the device is local, i.e. on the current server, this will return "Local (This Process)".
+		return DeviceProviderWrapper->GetServerNameForDevice(InDeviceName);
+	}
+	else
+	{
+		// There may not be a wrapper if the device provider is local, or there is no render node
+		// connected. In this case, we just search the local provider directly.
+		const IMediaIOCoreDeviceProvider* LocalProvider = GetLocalDeviceProvider(InDeviceProviderName);
+		if (LocalProvider && IsDevicePresent(LocalProvider, InDeviceName))
+		{
+			return LocalServerName;
+		}
+	}
+
+	// Return an empty string if the device couldn't be found anywhere.
+	return FString();
+}
+
+bool FAvaDeviceProviderProxyManager::IsLocalDevice(const FName& InDeviceProviderName, const FName& InDeviceName) const
+{
+	const FAvaDeviceProviderWrapper* DeviceProviderWrapper = GetDeviceProviderWrapper(InDeviceProviderName);
+
+	if (DeviceProviderWrapper && DeviceProviderWrapper->HasLocalProvider())
+	{
+		return FAvaDeviceProviderProxyManager::IsDevicePresent(DeviceProviderWrapper->GetLocalProvider(), InDeviceName);
+	}
+	else if (const IMediaIOCoreDeviceProvider* LocalProvider = GetLocalDeviceProvider(InDeviceProviderName))
+	{
+		return FAvaDeviceProviderProxyManager::IsDevicePresent(LocalProvider, InDeviceName);
+	}
+	return false;
+}
+
+void FAvaDeviceProviderProxyManager::Install(const FString& InServerName, const FAvaDeviceProviderDataList& InDeviceProviderDataList)
+{
+	if (!IMediaIOCoreModule::IsAvailable())
+	{
+		return;
+	}
+
+	DeviceProviderDataListByServer.Add(InServerName, MakeShared<FAvaDeviceProviderDataList>(InDeviceProviderDataList));
+	
+	for (const FAvaDeviceProviderData& ProviderData : InDeviceProviderDataList.DeviceProviders)
+	{
+		FAvaDeviceProviderWrapper* Wrapper = nullptr;
+		if (DeviceProviderWrappers.Contains(ProviderData.Name))
+		{
+			Wrapper = (*DeviceProviderWrappers.Find(ProviderData.Name)).Get();
+		}
+		else
+		{
+			// --- Need to make a new wrapper ---
+			TUniquePtr<FAvaDeviceProviderWrapper> DeviceProviderWrapper = MakeUnique<FAvaDeviceProviderWrapper>(ProviderData.Name, this);
+
+			// The local provider, if any, needs to be kept in the wrapper.
+			IMediaIOCoreDeviceProvider* LocalProvider = IMediaIOCoreModule::Get().GetDeviceProvider(ProviderData.Name);
+			if (LocalProvider)
+			{
+				DeviceProviderWrapper->AddProvider(LocalServerName, LocalProvider, true);
+			}
+
+			IMediaIOCoreModule::Get().RegisterDeviceProvider(DeviceProviderWrapper.Get());
+
+			Wrapper = DeviceProviderWrapper.Get(); 
+			
+			DeviceProviderWrappers.Add(ProviderData.Name, MoveTemp(DeviceProviderWrapper));
+		}
+
+		if (Wrapper)
+		{
+			// Install a proxy for this server's devices in the wrapper.
+			TUniquePtr<FAvaDeviceProviderProxy> DeviceProviderProxy = MakeUnique<FAvaDeviceProviderProxy>(this, ProviderData);
+			DeviceProviderProxy->ApplyServerName(InServerName);
+			Wrapper->AddProvider(InServerName, MoveTemp(DeviceProviderProxy));
+		}
+	}
+	
+	// Ensure all the wrapped local providers are at the end of IMediaIOCoreModule's array.
+	ReorderWrappedLocalProviders();
+
+	UAvalancheBroadcast::Get().QueueNotifyChange(EAvaBroadcastChange::OutputDevices);
+}
+
+void FAvaDeviceProviderProxyManager::Uninstall(const FString& InServerName)
+{
+	if (!IMediaIOCoreModule::IsAvailable())
+	{
+		return;
+	}
+
+	TArray<FName> EntriesToRemove;
+	
+	// Remove the given server from all the wrappers.
+	for (TPair<FName, TUniquePtr<FAvaDeviceProviderWrapper>>& Entry : DeviceProviderWrappers )
+	{
+		Entry.Value->RemoveProvider(InServerName);
+		if (Entry.Value->ShouldRemove())
+		{
+			IMediaIOCoreModule::Get().UnregisterDeviceProvider(Entry.Value.Get());
+			EntriesToRemove.Add(Entry.Key);
+		}
+	}
+
+	for (const FName& Key : EntriesToRemove)
+	{
+		DeviceProviderWrappers.Remove(Key);
+	}
+
+	DeviceProviderDataListByServer.Remove(InServerName);
+	
+	UAvalancheBroadcast::Get().QueueNotifyChange(EAvaBroadcastChange::OutputDevices);
+}
+
+void FAvaDeviceProviderProxyManager::ListAllProviders()
+{
+	if (!IMediaIOCoreModule::IsAvailable())
+	{
+		UE_LOG(LogAvaMedia, Error, TEXT("IMediaIOCoreModule is not available."));
+		return;
+	}
+
+	// List all device providers.
+	const TConstArrayView<IMediaIOCoreDeviceProvider*> DeviceProviders = IMediaIOCoreModule::Get().GetDeviceProviders();
+	for (IMediaIOCoreDeviceProvider* DeviceProvider : DeviceProviders)
+	{
+		UE_LOG(LogAvaMedia, Log, TEXT("Found DeviceProvider: \"%s\"."), *DeviceProvider->GetFName().ToString());
+		// It seems the only way to get the output devices is to query all the output configurations and filter out the unique devices.
+		const TArray<FMediaIOConfiguration> OutputConfigs = DeviceProvider->GetConfigurations(false, true);
+
+		if (OutputConfigs.Num() > 0)
+		{
+			TSet<FString> OutputDeviceNames;
+			for (const FMediaIOConfiguration& Config : OutputConfigs)
+			{
+				//UE_LOG(LogAvaMedia, Log, TEXT("Found Output Config: \"%s\"."), *DeviceProvider->ToText(Config).ToString());
+				OutputDeviceNames.Add(Config.MediaConnection.Device.DeviceName.ToString());
+			}
+
+			for (const FString& OutputDeviceName : OutputDeviceNames)
+			{
+				UE_LOG(LogAvaMedia, Log, TEXT("Found Output Device: \"%s\"."), *OutputDeviceName);
+			}
+		}
+		else
+		{
+			UE_LOG(LogAvaMedia, Log, TEXT("No Output Devices Available for \"%s\"."), *DeviceProvider->GetFName().ToString());
+		}	
+	}
+
+#if WITH_EDITORONLY_DATA
+	for (const UClass* Class : TObjectRange<UClass>())
+	{
+		if (Class->IsChildOf(UMediaOutput::StaticClass()))
+		{
+			// Check that it has the ProviderName meta
+			if (FName DeviceProviderName = UE::AvaMediaOutputUtils::GetDeviceProviderName(Class); !DeviceProviderName.IsNone())
+			{
+				UE_LOG(LogAvaMedia, Log, TEXT("Found DeviceProviderName: \"%s\" for class \"%s\"."), *DeviceProviderName.ToString(), *Class->GetName());
+
+				if (IMediaIOCoreModule::Get().GetDeviceProvider(DeviceProviderName) == nullptr)
+				{
+					// This is possible, NDI has the meta data field, but it doesn't have a provider.
+					UE_LOG(LogAvaMedia, Log, TEXT("DeviceProvider \"%s\" not found in IMediaIOCoreModule."), *DeviceProviderName.ToString());
+				}
+			}
+		}
+	}
+#endif // WITH_EDITORONLY_DATA
+}
+
+void FAvaDeviceProviderProxyManager::TestInstall()
+{
+	FAvaDeviceProviderDataList OutProviders;
+	if (OutProviders.LoadFromJson())
+	{
+		Install(TEXT("RemoteTest"), OutProviders);
+	}
+}
+
+void FAvaDeviceProviderProxyManager::TestUninstall()
+{
+	Uninstall(TEXT("RemoteTest"));
+}
+
+bool FAvaDeviceProviderProxyManager::IsDevicePresent(const IMediaIOCoreDeviceProvider* InProvider, const FName& InDeviceName)
+{
+	TArray<FMediaIODevice> Devices = InProvider->GetDevices();
+	for (const FMediaIODevice& Device : Devices)
+	{
+		if (Device.DeviceName == InDeviceName)
+		{
+			return true;
+		}
+	}
+	return false;
+}
+
+void FAvaDeviceProviderProxyManager::ReorderWrappedLocalProviders()
+{
+	if (IMediaIOCoreModule::IsAvailable())
+	{
+		IMediaIOCoreModule& MediaIOCoreModule = IMediaIOCoreModule::Get();
+		
+		// Remote all local providers
+		for (TPair<FName, TUniquePtr<FAvaDeviceProviderWrapper>>& WrapperEntry : DeviceProviderWrappers)
+		{
+			if (WrapperEntry.Value->HasLocalProvider())
+			{
+				MediaIOCoreModule.UnregisterDeviceProvider(WrapperEntry.Value->GetLocalProvider());
+			}
+		}
+
+		// Add them at the back of the array so they are ignored, but not lost.
+		// Note: we keep them in there so we can detect if their respective module gets unloaded.
+		// To test: "Module Unload BlackmagicMedia", "Module Load BlackmagicMedia".
+		for (TPair<FName, TUniquePtr<FAvaDeviceProviderWrapper>>& WrapperEntry : DeviceProviderWrappers)
+		{
+			if (WrapperEntry.Value->HasLocalProvider())
+			{
+				MediaIOCoreModule.RegisterDeviceProvider(WrapperEntry.Value->GetLocalProvider());
+			}
+		}
+	}
+}
+
+const IMediaIOCoreDeviceProvider* FAvaDeviceProviderProxyManager::GetLocalDeviceProvider(const FName& InDeviceProviderName)
+{
+	if (IMediaIOCoreModule::IsAvailable())
+	{
+		return IMediaIOCoreModule::Get().GetDeviceProvider(InDeviceProviderName);
+	}
+	return nullptr;
+}
+
+TSet<FString> FAvaDeviceProviderProxyManager::GetServerNames() const
+{
+	TSet<FString> ServerNames;
+	DeviceProviderDataListByServer.GetKeys(ServerNames);
+	return (ServerNames);
+}
+
+const TSharedPtr<const FAvaDeviceProviderDataList> FAvaDeviceProviderProxyManager::GetDeviceProviderDataListForServer(const FString& InServerName) const
+{
+	const TSharedPtr<FAvaDeviceProviderDataList>* FoundDeviceProviderDataList = DeviceProviderDataListByServer.Find(InServerName);
+	return FoundDeviceProviderDataList ? *FoundDeviceProviderDataList : nullptr;
+}
