[MotionDesign]

diff --git a/Engine/Plugins/Experimental/Avalanche/Source/AvalancheMedia/Private/OutputDevices/AvaDeviceProviderProxy.h b/Engine/Plugins/Experimental/Avalanche/Source/AvalancheMedia/Private/OutputDevices/AvaDeviceProviderProxy.h
new file mode 100644
--- /dev/null
+++ b/Engine/Plugins/Experimental/Avalanche/Source/AvalancheMedia/Private/OutputDevices/AvaDeviceProviderProxy.h
@@ -0,0 +1,193 @@
+// Copyright Epic Games, Inc. All Rights Reserved.
+
+#pragma once
+
+#include "IMediaIOCoreDeviceProvider.h"
+#include "OutputDevices/AvaDeviceProviderData.h"
+#include "OutputDevices/IAvaDeviceProviderProxyManager.h"
+
+class FAvaDeviceProviderProxyManager;
+
+/**
+ * Implements a proxy for the device provider of another system.
+ */
+class FAvaDeviceProviderProxy: public IMediaIOCoreDeviceProvider
+{
+public:
+	FAvaDeviceProviderProxy(FAvaDeviceProviderProxyManager* InManager, const FAvaDeviceProviderData& InData);
+	virtual ~FAvaDeviceProviderProxy() = default;
+
+	// This will add the server name to all device names.
+	void ApplyServerName(const FString& InServerName) { Data.ApplyServerName(InServerName);}
+
+	//IMediaIOCoreDeviceProvider
+	virtual FName GetFName() override { return Data.Name; }
+
+	virtual TArray<FMediaIOConnection> GetConnections() const override { return Data.Connections; }
+	virtual TArray<FMediaIOConfiguration> GetConfigurations() const override { return Data.Configurations; }
+	virtual TArray<FMediaIOConfiguration> GetConfigurations(bool bAllowInput, bool bAllowOutput) const override;
+	virtual TArray<FMediaIODevice> GetDevices() const override { return Data.Devices; }
+	virtual TArray<FMediaIOMode> GetModes(const FMediaIODevice& InDevice, bool bInOutput) const override;
+	virtual TArray<FMediaIOInputConfiguration> GetInputConfigurations() const override { return Data.InputConfigurations; }
+	virtual TArray<FMediaIOOutputConfiguration> GetOutputConfigurations() const override { return Data.OutputConfigurations; }
+	virtual TArray<FMediaIOVideoTimecodeConfiguration> GetTimecodeConfigurations() const override { return Data.TimecodeConfigurations;}
+	
+	virtual FMediaIOConfiguration GetDefaultConfiguration() const override { return Data.DefaultConfiguration; }
+	virtual FMediaIOMode GetDefaultMode() const override { return Data.DefaultMode; }
+	virtual FMediaIOInputConfiguration GetDefaultInputConfiguration() const override { return Data.DefaultInputConfiguration; }
+	virtual FMediaIOOutputConfiguration GetDefaultOutputConfiguration() const override { return Data.DefaultOutputConfiguration; }
+	virtual FMediaIOVideoTimecodeConfiguration GetDefaultTimecodeConfiguration() const override { return Data.DefaultTimecodeConfiguration;}
+
+#if WITH_EDITOR
+	virtual bool ShowInputTransportInSelector() const override { return Data.bShowInputTransportInSelector; }
+	virtual bool ShowOutputTransportInSelector() const override { return Data.bShowOutputTransportInSelector; }
+	virtual bool ShowInputKeyInSelector() const override { return Data.bShowInputKeyInSelector; }
+	virtual bool ShowOutputKeyInSelector() const override { return Data.bShowOutputKeyInSelector; }
+	virtual bool ShowReferenceInSelector() const override { return Data.bShowReferenceInSelector; }
+#endif
+	//~IMediaIOCoreDeviceProvider
+	
+private:
+	FAvaDeviceProviderProxyManager* Manager;
+	FAvaDeviceProviderData Data;
+};
+
+/**
+ *	Implements a wrapper of multiple other providers of the same device type from different systems.
+ */
+class FAvaDeviceProviderWrapper : public IMediaIOCoreDeviceProvider
+{
+public:
+	/**
+	 *	@param InName Device Provider Name (ex: Blackmagic, AJA, etc).
+	 *	@param InManager Parent Manager.
+	 */
+	FAvaDeviceProviderWrapper(FName InName, FAvaDeviceProviderProxyManager* InManager);
+	virtual ~FAvaDeviceProviderWrapper() = default;
+
+	/** Add a provider for the given server name. */
+	void AddProvider(const FString& InServerName, IMediaIOCoreDeviceProvider* InProvider, bool bInIsLocal);
+	void AddProvider(const FString& InServerName, TUniquePtr<FAvaDeviceProviderProxy>&& InProvider);
+	
+	/** Removes the provider from the list. */
+	void RemoveProvider(IMediaIOCoreDeviceProvider* InProvider);
+	void RemoveProvider(const FString& InServerName);
+
+	bool HasLocalProvider() const { return LocalProvider != nullptr; }
+
+	IMediaIOCoreDeviceProvider* GetLocalProvider() const { return LocalProvider; }
+	const IMediaIOCoreDeviceProvider* GetProviderForDeviceName(const FName& InDeviceName) const;
+	const IMediaIOCoreDeviceProvider* GetProviderForServer(const FString& InServerName) const;
+	
+	/** Search in the list of all wrapped proxies to find the device and returns the corresponding
+	 * server name. If the provider is local, the server name is FAvaDeviceProviderProxyManager::LocalServerName.
+	 */
+	FString GetServerNameForDevice(const FName& InDeviceName) const;
+	
+	/** Check if the wrapper is still wrapping something. */
+	bool ShouldRemove() const;
+
+	//IMediaIOCoreDeviceProvider
+	virtual FName GetFName() override { return Name; }
+
+	virtual TArray<FMediaIOConnection> GetConnections() const override;
+	virtual TArray<FMediaIOConfiguration> GetConfigurations() const override;
+	virtual TArray<FMediaIOConfiguration> GetConfigurations(bool bInAllowInput, bool bInAllowOutput) const override;
+	virtual TArray<FMediaIODevice> GetDevices() const override;
+	virtual TArray<FMediaIOMode> GetModes(const FMediaIODevice& InDevice, bool bInOutput) const override;
+	virtual TArray<FMediaIOInputConfiguration> GetInputConfigurations() const override;
+	virtual TArray<FMediaIOOutputConfiguration> GetOutputConfigurations() const override;
+	virtual TArray<FMediaIOVideoTimecodeConfiguration> GetTimecodeConfigurations() const override;
+
+
+	virtual FMediaIOConfiguration GetDefaultConfiguration() const override { return GetDefault()->GetDefaultConfiguration(); }
+	virtual FMediaIOMode GetDefaultMode() const override { return GetDefault()->GetDefaultMode(); }
+	virtual FMediaIOInputConfiguration GetDefaultInputConfiguration() const override { return GetDefault()->GetDefaultInputConfiguration(); }
+	virtual FMediaIOOutputConfiguration GetDefaultOutputConfiguration() const override { return GetDefault()->GetDefaultOutputConfiguration(); }
+	virtual FMediaIOVideoTimecodeConfiguration GetDefaultTimecodeConfiguration() const override { return GetDefault()->GetDefaultTimecodeConfiguration(); }
+
+#if WITH_EDITOR
+	virtual bool ShowInputTransportInSelector() const override { return GetDefault()->ShowInputTransportInSelector(); }
+	virtual bool ShowOutputTransportInSelector() const override { return GetDefault()->ShowOutputTransportInSelector(); }
+	virtual bool ShowInputKeyInSelector() const override { return GetDefault()->ShowInputKeyInSelector(); }
+	virtual bool ShowOutputKeyInSelector() const override { return GetDefault()->ShowOutputKeyInSelector(); }
+	virtual bool ShowReferenceInSelector() const override { return GetDefault()->ShowReferenceInSelector(); }
+#endif
+	//~IMediaIOCoreDeviceProvider
+
+private:
+	const IMediaIOCoreDeviceProvider* GetDefault() const
+	{
+		return !Providers.IsEmpty() ? Providers[0].Provider : &DummyProxy;
+	}
+
+	void OnRemoveProvider(IMediaIOCoreDeviceProvider* InProviderToRemove)
+	{
+		if (InProviderToRemove == LocalProvider)
+		{
+			LocalProvider = nullptr;
+		}
+	}
+	
+private:
+	FAvaDeviceProviderProxyManager* Manager = nullptr;
+
+	FName Name;
+
+	// We need to keep track of the local provider because it is a special snowflake.
+	IMediaIOCoreDeviceProvider* LocalProvider = nullptr;
+	
+	FAvaDeviceProviderProxy DummyProxy;
+	
+	struct FProviderInfo
+	{
+		FString ServerName;
+		IMediaIOCoreDeviceProvider* Provider;
+		TUniquePtr<FAvaDeviceProviderProxy> Proxy;
+	};
+	TArray<FProviderInfo> Providers;
+};
+
+class FAvaDeviceProviderProxyManager : public IAvaDeviceProviderProxyManager
+{
+public:
+	virtual ~FAvaDeviceProviderProxyManager() {}
+
+	//~ Begin IAvaDeviceProviderProxyManager
+	const FAvaDeviceProviderWrapper* GetDeviceProviderWrapper(FName InDeviceProviderName) const override;
+	TArray<const IMediaIOCoreDeviceProvider*> GetDeviceProvidersForServer(const FString& InServerName) const override;
+	FString FindServerNameForDevice(const FName& InDeviceProviderName, const FName& InDeviceName) const override;
+	bool IsLocalDevice(const FName& InDeviceProviderName, const FName& InDeviceName) const override;
+
+	virtual void Install(const FString& InServerName, const FAvaDeviceProviderDataList& InDeviceProviderDataList) override;
+	virtual void Uninstall(const FString& InServerName) override;
+
+	virtual TSet<FString> GetServerNames() const override;
+	virtual const TSharedPtr<const FAvaDeviceProviderDataList> GetDeviceProviderDataListForServer(const FString& InServerName) const override;
+	virtual FString GetLocalServerName() const override { return LocalServerName; }
+	//~ End IAvaDeviceProviderProxyManager
+
+	/** Lists all the devices from all device providers from all the media output classes. */
+	static void ListAllProviders();
+
+	void TestInstall();
+	void TestUninstall();
+
+	/** Returns true if the given device name can be found in the device provider. */
+	static bool IsDevicePresent(const IMediaIOCoreDeviceProvider* InProvider, const FName& InDeviceName);
+
+	static FString LocalServerName;
+	
+private:
+	/** Ensure all the wrapped local providers are at the end of IMediaIOCoreModule's array. */
+	void ReorderWrappedLocalProviders();
+
+	/** Returns the given local provider. */
+	static const IMediaIOCoreDeviceProvider* GetLocalDeviceProvider(const FName& InDeviceProviderName);
+
+private:
+	TMap<FName, TUniquePtr<FAvaDeviceProviderWrapper>> DeviceProviderWrappers;
+
+	/** Keep a copy of the device provider data list last received for each server. */
+	TMap<FString, TSharedPtr<FAvaDeviceProviderDataList>> DeviceProviderDataListByServer;
+};
diff --git a/Engine/Plugins/Experimental/Avalanche/Source/AvalancheMedia/Private/OutputDevices/AvaDisplayDeviceManager.cpp b/Engine/Plugins/Experimental/Avalanche/Source/AvalancheMedia/Private/OutputDevices/AvaDisplayDeviceManager.cpp
new file mode 100644
--- /dev/null
+++ b/Engine/Plugins/Experimental/Avalanche/Source/AvalancheMedia/Private/OutputDevices/AvaDisplayDeviceManager.cpp
@@ -0,0 +1,261 @@
+// Copyright Epic Games, Inc. All Rights Reserved.
+
+#include "OutputDevices/AvaDisplayDeviceManager.h"
+
+#include "Templates/RefCounting.h"
+
+DEFINE_LOG_CATEGORY_STATIC(LogAvaDisplayDeviceManager, Log, All);
+
+#if PLATFORM_WINDOWS
+
+THIRD_PARTY_INCLUDES_START
+#include "Windows/AllowWindowsPlatformTypes.h"
+#include "dxgi1_3.h"
+#include "dxgi1_4.h"
+#include "dxgi1_6.h"
+#include "Windows/HideWindowsPlatformTypes.h"
+THIRD_PARTY_INCLUDES_END
+
+namespace AvaDXGIHelper
+{
+	// I haven't found a way to get the actual pixel format of an output.
+	// So we try to guess it.
+	static DXGI_FORMAT GuessOutputPixelFormat(IDXGIOutput* InOutput)
+	{
+		DXGI_FORMAT OutputPixelFormat = DXGI_FORMAT_R8G8B8A8_UNORM;
+		
+		TRefCountPtr<IDXGIOutput6> Output6;
+		if (SUCCEEDED(
+			InOutput->QueryInterface(__uuidof(IDXGIOutput6), (void**)Output6.GetInitReference())))
+		{
+			DXGI_OUTPUT_DESC1 OutputDesc1;
+			Output6->GetDesc1(&OutputDesc1);
+			switch (OutputDesc1.BitsPerColor)
+			{
+			case 8: OutputPixelFormat = DXGI_FORMAT_R8G8B8A8_UNORM;
+				break;
+			case 10: OutputPixelFormat = DXGI_FORMAT_R10G10B10A2_UNORM;
+				break;
+			case 16: OutputPixelFormat = DXGI_FORMAT_R16G16B16A16_FLOAT;
+				break;
+			default: OutputPixelFormat = DXGI_FORMAT_R8G8B8A8_UNORM;
+				break;
+			}
+		}
+		return OutputPixelFormat;
+	}
+	
+	static TArray<DXGI_MODE_DESC> GetOutputModesForFormat(IDXGIOutput* InOutput, DXGI_FORMAT InFormat)
+	{
+		TArray<DXGI_MODE_DESC> ModeDescriptors;
+		UINT Num = 0;
+		const UINT Flags = DXGI_ENUM_MODES_INTERLACED;
+		InOutput->GetDisplayModeList(InFormat, Flags, &Num, 0);
+		if (Num)
+		{
+			ModeDescriptors.SetNumUninitialized(Num);
+			InOutput->GetDisplayModeList(InFormat, Flags, &Num, ModeDescriptors.GetData());
+		}
+		return ModeDescriptors;
+	}
+
+	inline FFrameRate ToFramerate(const DXGI_RATIONAL& InFrameRate)
+	{
+		return FFrameRate(InFrameRate.Numerator, InFrameRate.Denominator);
+	}
+
+	// We have an approximate frame rate in DevMode, but we can try to guess more
+	// accurately by checking the modes returned by GetDisplayModeList.
+	// Unfortunately, there is no way to get the current display mode directly, so again
+	// we have to first guess the pixel format.
+	static FFrameRate GuessFrameRate(IDXGIOutput* InOutput, const DEVMODE& InDevMode)
+	{
+		// Because this API returns an integer display frequency, it is not accurate.
+		FFrameRate ApproxFrameRate(InDevMode.dmDisplayFrequency, 1);
+
+		// Step 1 of the guessing game: guess the pixel format.
+		const DXGI_FORMAT GuessedOutputPixelFormat = GuessOutputPixelFormat(InOutput);
+
+		{
+			DXGI_MODE_DESC ModeToMatch;
+			ModeToMatch.Format = GuessedOutputPixelFormat;
+			ModeToMatch.Height = InDevMode.dmPelsHeight;
+			ModeToMatch.Width = InDevMode.dmPelsWidth;
+			ModeToMatch.RefreshRate.Numerator = InDevMode.dmDisplayFrequency;
+			ModeToMatch.RefreshRate.Denominator = 1;
+			ModeToMatch.Scaling = DXGI_MODE_SCALING_UNSPECIFIED;
+			DXGI_MODE_DESC ClosestMatch;
+			if (SUCCEEDED(InOutput->FindClosestMatchingMode(&ModeToMatch, &ClosestMatch, nullptr)))
+			{
+				return ToFramerate(ClosestMatch.RefreshRate);
+			}
+		}
+		
+		TArray<DXGI_MODE_DESC> ModeDescriptors = GetOutputModesForFormat(InOutput, GuessedOutputPixelFormat);
+		if (!ModeDescriptors.IsEmpty())
+		{
+			FFrameRate ClosestFrameRate(0,0);
+			double ClosestDist = 0.0;
+
+			for (DXGI_MODE_DESC& ModeDesc : ModeDescriptors)
+			{
+				if (ModeDesc.Width == InDevMode.dmPelsWidth && ModeDesc.Height == InDevMode.dmPelsHeight)
+				{
+					FFrameRate FrameRate = ToFramerate(ModeDesc.RefreshRate);
+					double Dist = FGenericPlatformMath::Abs(ApproxFrameRate.AsDecimal() - FrameRate.AsDecimal());
+					if (!ClosestFrameRate.IsValid() || Dist < ClosestDist)
+					{
+						ClosestFrameRate = FrameRate;
+						ClosestDist = Dist;
+					}
+				}
+			}
+
+			if (ClosestFrameRate.IsValid())
+			{
+				ApproxFrameRate = ClosestFrameRate;
+			}							
+		}
+		return ApproxFrameRate;
+	}
+	
+	static void EnumMonitorsWindowsDXGI(TArray<FAvaMonitorInfo>& OutMonitorInfo)
+	{
+		TRefCountPtr<IDXGIFactory1> DXGIFactory1;
+		if (CreateDXGIFactory1(__uuidof(IDXGIFactory1), (void**)DXGIFactory1.GetInitReference()) != S_OK || !DXGIFactory1)
+		{
+			return;
+		}
+
+		OutMonitorInfo.Empty(2); // Reserve two slots, as that will be the most common maximum
+	
+		TRefCountPtr<IDXGIAdapter> TempAdapter;
+		for (uint32 AdapterIndex = 0; DXGIFactory1->EnumAdapters(AdapterIndex, TempAdapter.GetInitReference()) != DXGI_ERROR_NOT_FOUND; ++AdapterIndex)
+		{
+			if (TempAdapter)
+			{
+				DXGI_ADAPTER_DESC AdapterDesc;
+				TempAdapter->GetDesc(&AdapterDesc);
+
+				FAvaDisplayAdapterInfo AdapterInfo;
+				AdapterInfo.Description = AdapterDesc.Description;
+				AdapterInfo.VendorId = AdapterDesc.VendorId;
+				AdapterInfo.DeviceId = AdapterDesc.DeviceId;
+				AdapterInfo.SubSysId = AdapterDesc.SubSysId;
+				AdapterInfo.Revision = AdapterDesc.Revision;
+				AdapterInfo.DedicatedVideoMemory = AdapterDesc.DedicatedVideoMemory;
+			
+				TRefCountPtr<IDXGIOutput> TempOutput;
+				for (uint32 OutputIndex =0; TempAdapter->EnumOutputs(OutputIndex, TempOutput.GetInitReference()) != DXGI_ERROR_NOT_FOUND; ++OutputIndex)
+				{
+					DXGI_OUTPUT_DESC OutputDesc;
+					TempOutput->GetDesc(&OutputDesc);
+				
+					// Get Resolution and Frame rate
+					DEVMODE DevMode;
+					ZeroMemory(&DevMode, sizeof(DevMode));
+					DevMode.dmSize = sizeof(DevMode);
+					EnumDisplaySettings(OutputDesc.DeviceName, ENUM_CURRENT_SETTINGS, &DevMode);
+
+					FFrameRate ApproxFrameRate = GuessFrameRate(TempOutput.GetReference(), DevMode);
+					UE_LOG(LogAvaDisplayDeviceManager, Verbose, TEXT("Mode: %d x %d @ %f Hz"),
+						DevMode.dmPelsWidth, DevMode.dmPelsHeight, ApproxFrameRate.AsDecimal());
+					
+					// Get Display and Working rect.
+					MONITORINFOEX MonitorInfo;
+					ZeroMemory(&MonitorInfo, sizeof(MonitorInfo));
+					MonitorInfo.cbSize = sizeof(MonitorInfo);
+					GetMonitorInfo(OutputDesc.Monitor, &MonitorInfo);
+				
+					// UE_LOG(LogAvaDisplayDeviceManager, Log, TEXT("Monitor: %s on %s (%d x %d @ %d Hz)"),
+					// 	OutputDesc.DeviceName, AdapterDesc.Description,
+					// 	DevMode.dmPelsWidth, DevMode.dmPelsHeight, DevMode.dmDisplayFrequency);
+
+					FAvaMonitorInfo AvaInfo;
+					AvaInfo.Name = OutputDesc.DeviceName;
+					AvaInfo.Name.RemoveFromStart(TEXT("\\\\.\\"));
+					AvaInfo.AdapterInfo = AdapterInfo;
+					AvaInfo.Width = DevMode.dmPelsWidth;
+					AvaInfo.Height = DevMode.dmPelsHeight;
+					AvaInfo.DisplayFrequency = ApproxFrameRate;
+					AvaInfo.DisplayRect.Bottom = MonitorInfo.rcMonitor.bottom;
+					AvaInfo.DisplayRect.Left = MonitorInfo.rcMonitor.left;
+					AvaInfo.DisplayRect.Right = MonitorInfo.rcMonitor.right;
+					AvaInfo.DisplayRect.Top = MonitorInfo.rcMonitor.top;
+					AvaInfo.WorkArea.Bottom = MonitorInfo.rcWork.bottom;
+					AvaInfo.WorkArea.Left = MonitorInfo.rcWork.left;
+					AvaInfo.WorkArea.Right = MonitorInfo.rcWork.right;
+					AvaInfo.WorkArea.Top = MonitorInfo.rcWork.top;
+					AvaInfo.bIsPrimary = (MonitorInfo.dwFlags & MONITORINFOF_PRIMARY) > 0;
+					OutMonitorInfo.Add(AvaInfo);
+				}
+			}
+		}
+	}
+}
+#endif
+
+static void EnumMonitorsGeneric(TArray<FAvaMonitorInfo>& OutMonitorInfo)
+{
+	FDisplayMetrics DisplayMetrics;
+
+	FDisplayMetrics::RebuildDisplayMetrics(DisplayMetrics);
+
+	OutMonitorInfo.Empty(DisplayMetrics.MonitorInfo.Num());
+	
+	for (const FMonitorInfo& Info : DisplayMetrics.MonitorInfo)
+	{
+		FAvaMonitorInfo AvaInfo;
+		AvaInfo.Name = Info.Name;
+		//AvaInfo.AdapterDesc -- Don't have the adapter name.
+		AvaInfo.Width = Info.DisplayRect.Right - Info.DisplayRect.Left;
+		AvaInfo.Height = Info.DisplayRect.Bottom - Info.DisplayRect.Top;
+		//AvaInfo.DisplayFrequency -- Don't have display frequency.
+		AvaInfo.DisplayRect = Info.DisplayRect;
+		AvaInfo.WorkArea = Info.WorkArea;
+		AvaInfo.bIsPrimary = Info.bIsPrimary;
+		OutMonitorInfo.Add(AvaInfo);
+	}
+}
+
+void FAvaDisplayDeviceManager::EnumMonitors(TArray<FAvaMonitorInfo>& OutMonitorInfo)
+{
+#if PLATFORM_WINDOWS
+	static bool bUseDXGIEnumeration = true;
+	if (bUseDXGIEnumeration)
+	{
+		AvaDXGIHelper::EnumMonitorsWindowsDXGI(OutMonitorInfo);
+		return;
+	}
+#endif
+	EnumMonitorsGeneric(OutMonitorInfo);
+}
+
+TArray<FAvaMonitorInfo> FAvaDisplayDeviceManager::CachedMonitorInfo;
+
+const TArray<FAvaMonitorInfo>&  FAvaDisplayDeviceManager::GetCachedMonitors(bool bInForceUpdate)
+{
+	if (bInForceUpdate || CachedMonitorInfo.Num() == 0)
+	{
+		EnumMonitors(CachedMonitorInfo);
+	}
+	return CachedMonitorInfo;
+}
+
+FString FAvaDisplayDeviceManager::GetMonitorDisplayName(const FAvaMonitorInfo& InMonitorInfo)
+{
+	FString DisplayName = FString::Format(TEXT("{0}: {1}x{2} @ {3},{4}"),
+		{
+			InMonitorInfo.Name,
+			InMonitorInfo.Width,
+			InMonitorInfo.Height,
+			InMonitorInfo.DisplayRect.Left,
+			InMonitorInfo.DisplayRect.Top
+		});
+
+	if (InMonitorInfo.bIsPrimary)
+	{
+		DisplayName += TEXT(" (Primary Monitor)");
+	}
+	return DisplayName;
+}
diff --git a/Engine/Plugins/Experimental/Avalanche/Source/AvalancheMedia/Private/OutputDevices/AvaDisplayDeviceProvider.cpp b/Engine/Plugins/Experimental/Avalanche/Source/AvalancheMedia/Private/OutputDevices/AvaDisplayDeviceProvider.cpp
new file mode 100644
--- /dev/null
+++ b/Engine/Plugins/Experimental/Avalanche/Source/AvalancheMedia/Private/OutputDevices/AvaDisplayDeviceProvider.cpp
@@ -0,0 +1,222 @@
+// Copyright Epic Games, Inc. All Rights Reserved.
+
+#include "AvaDisplayDeviceProvider.h"
+#include "MediaIOCoreCommonDisplayMode.h"
+#include "OutputDevices/AvaDisplayDeviceManager.h"
+
+#define LOCTEXT_NAMESPACE "AvaDisplayDeviceProvider"
+
+FName FAvaDisplayDeviceProvider::GetProviderName()
+{
+	static FName NAME_Provider = "AvaDisplay";
+	return NAME_Provider;
+}
+
+
+FName FAvaDisplayDeviceProvider::GetProtocolName()
+{
+	static FName NAME_Protocol = "avadisplay";
+	return NAME_Protocol;
+}
+
+FName FAvaDisplayDeviceProvider::GetFName()
+{
+	return GetProviderName();
+}
+
+TArray<FMediaIOConnection> FAvaDisplayDeviceProvider::GetConnections() const
+{
+	return TArray<FMediaIOConnection>();
+}
+
+TArray<FMediaIOConfiguration> FAvaDisplayDeviceProvider::GetConfigurations() const
+{
+	return GetConfigurations(true, true);
+}
+
+TArray<FMediaIOConfiguration> FAvaDisplayDeviceProvider::GetConfigurations(bool /*bInAllowInput*/, bool bInAllowOutput) const
+{
+	TArray<FMediaIOConfiguration> Results;
+	if (bInAllowOutput)
+	{
+		TArray<FAvaMonitorInfo> MonitorInfos = FAvaDisplayDeviceManager::GetCachedMonitors();
+
+		for (int32 i = 0; i < MonitorInfos.Num(); ++i)
+		{
+			const FAvaMonitorInfo& MonitorInfo = MonitorInfos[i];
+
+			FMediaIOConfiguration MediaConfiguration = GetDefaultConfiguration();
+			MediaConfiguration.bIsInput = false;
+			MediaConfiguration.MediaMode.Resolution = FIntPoint(MonitorInfo.Width, MonitorInfo.Height);
+			if (MonitorInfo.DisplayFrequency.IsValid())
+			{
+				MediaConfiguration.MediaMode.FrameRate = MonitorInfo.DisplayFrequency;
+			}
+
+			MediaConfiguration.MediaConnection.Device.DeviceName = *FAvaDisplayDeviceManager::GetMonitorDisplayName(MonitorInfo);
+			MediaConfiguration.MediaConnection.Device.DeviceIdentifier = i;
+			MediaConfiguration.MediaConnection.Protocol = GetProtocolName();
+			MediaConfiguration.MediaConnection.PortIdentifier = 0;
+
+			Results.Add(MediaConfiguration);
+		}
+	}
+	return Results;
+}
+
+TArray<FMediaIOInputConfiguration> FAvaDisplayDeviceProvider::GetInputConfigurations() const
+{
+	return TArray<FMediaIOInputConfiguration>();
+}
+
+TArray<FMediaIOOutputConfiguration> FAvaDisplayDeviceProvider::GetOutputConfigurations() const
+{
+	TArray<FMediaIOOutputConfiguration> Results;
+
+	TArray<FMediaIOConfiguration> Configs = GetConfigurations(false, true);
+
+	FMediaIOOutputConfiguration DefaultOutputConfiguration = GetDefaultOutputConfiguration();
+	DefaultOutputConfiguration.KeyPortIdentifier = 0;
+	DefaultOutputConfiguration.OutputType = EMediaIOOutputType::Fill;	// HDMI/Display port don't support alpha.
+	DefaultOutputConfiguration.OutputReference = EMediaIOReferenceType::FreeRun;
+
+	for (const FMediaIOConfiguration& Config : Configs)
+	{
+		DefaultOutputConfiguration.MediaConfiguration = Config;
+		Results.Add(DefaultOutputConfiguration);
+	}
+	return Results;
+}
+
+TArray<FMediaIOVideoTimecodeConfiguration> FAvaDisplayDeviceProvider::GetTimecodeConfigurations() const
+{
+	TArray<FMediaIOVideoTimecodeConfiguration> MediaConfigurations;
+	return MediaConfigurations;
+}
+
+TArray<FMediaIODevice> FAvaDisplayDeviceProvider::GetDevices() const
+{
+	TArray<FMediaIODevice> Results;
+
+	TArray<FAvaMonitorInfo> MonitorInfos = FAvaDisplayDeviceManager::GetCachedMonitors();
+	
+	for (int32 i = 0; i < MonitorInfos.Num(); ++i)
+	{
+		const FAvaMonitorInfo& MonitorInfo = MonitorInfos[i];
+		FMediaIODevice Device;
+		Device.DeviceName = *FAvaDisplayDeviceManager::GetMonitorDisplayName(MonitorInfo);
+		Device.DeviceIdentifier = i;
+		Results.Add(Device);
+	}
+
+	return Results;
+}
+
+
+TArray<FMediaIOMode> FAvaDisplayDeviceProvider::GetModes(const FMediaIODevice& InDevice, bool bInOutput) const
+{
+	TArray<FMediaIOMode> Results;
+	return Results;
+}
+
+FMediaIOConfiguration FAvaDisplayDeviceProvider::GetDefaultConfiguration() const
+{
+	FMediaIOConfiguration Configuration;
+	Configuration.bIsInput = true;
+	Configuration.MediaConnection.Device.DeviceIdentifier = 1;
+	Configuration.MediaConnection.Protocol = GetProtocolName();
+	Configuration.MediaConnection.PortIdentifier = 0;
+	Configuration.MediaMode = GetDefaultMode();
+	return Configuration;
+}
+
+
+FMediaIOMode FAvaDisplayDeviceProvider::GetDefaultMode() const
+{
+	FMediaIOMode Mode;
+	Mode.DeviceModeIdentifier = 0;	// Unused, but can't be invalid.
+	Mode.FrameRate = FFrameRate(30, 1);
+	Mode.Resolution = FIntPoint(1920, 1080);
+	Mode.Standard = EMediaIOStandardType::Progressive;
+	return Mode;
+}
+
+
+FMediaIOInputConfiguration FAvaDisplayDeviceProvider::GetDefaultInputConfiguration() const
+{
+	FMediaIOInputConfiguration Configuration;
+	Configuration.MediaConfiguration = GetDefaultConfiguration();
+	Configuration.MediaConfiguration.bIsInput = true;
+	Configuration.InputType = EMediaIOInputType::Fill;
+	return Configuration;
+}
+
+FMediaIOOutputConfiguration FAvaDisplayDeviceProvider::GetDefaultOutputConfiguration() const
+{
+	FMediaIOOutputConfiguration Configuration;
+	Configuration.MediaConfiguration = GetDefaultConfiguration();
+	Configuration.MediaConfiguration.bIsInput = false;
+	Configuration.OutputReference = EMediaIOReferenceType::FreeRun;
+	Configuration.OutputType = EMediaIOOutputType::Fill;
+	return Configuration;
+}
+
+FMediaIOVideoTimecodeConfiguration FAvaDisplayDeviceProvider::GetDefaultTimecodeConfiguration() const
+{
+	FMediaIOVideoTimecodeConfiguration Configuration;
+	Configuration.MediaConfiguration = GetDefaultConfiguration();
+	return Configuration;
+}
+
+FText FAvaDisplayDeviceProvider::ToText(const FMediaIOConfiguration& InConfiguration, bool bInIsAutoDetected) const
+{
+	if (bInIsAutoDetected)
+	{
+		return FText::Format(LOCTEXT("FMediaIOConfigurationToText", "{0} - {1} [device{2}/auto]")
+				, InConfiguration.bIsInput ? LOCTEXT("In", "In") : LOCTEXT("Out", "Out")
+				, FText::FromName(InConfiguration.MediaConnection.Device.DeviceName)
+				, FText::AsNumber(InConfiguration.MediaConnection.Device.DeviceIdentifier)
+				);
+	}
+	if (InConfiguration.IsValid())
+	{
+		return FText::Format(LOCTEXT("FMediaIOConfigurationToText", "[{0}] - {1} [device{2}/{3}]")
+			, InConfiguration.bIsInput ? LOCTEXT("In", "In") : LOCTEXT("Out", "Out")
+			, FText::FromName(InConfiguration.MediaConnection.Device.DeviceName)
+			, FText::AsNumber(InConfiguration.MediaConnection.Device.DeviceIdentifier)
+			, InConfiguration.MediaMode.GetModeName()
+		);
+	}
+	return LOCTEXT("Invalid", "<Invalid>");
+}
+
+
+FText FAvaDisplayDeviceProvider::ToText(const FMediaIOConnection& InConnection) const
+{
+	if (InConnection.IsValid())
+	{
+		return FText::Format(LOCTEXT("FMediaIOConnectionToText", "{0} [device{1}]")
+			, FText::FromName(InConnection.Device.DeviceName)
+			, LOCTEXT("Device", "device")
+			, FText::AsNumber(InConnection.Device.DeviceIdentifier)
+		);
+	}
+	return LOCTEXT("Invalid", "<Invalid>");
+}
+
+FText FAvaDisplayDeviceProvider::ToText(const FMediaIOOutputConfiguration& InConfiguration) const
+{
+	if (InConfiguration.IsValid())
+	{
+		return FText::Format(LOCTEXT("FMediaIOOutputConfigurationToText", "{0} - {1} [device{2}/{3}/{4}]")
+			, InConfiguration.OutputType == EMediaIOOutputType::Fill ? LOCTEXT("Fill", "Fill") : LOCTEXT("FillAndKey", "Fill&Key")
+			, FText::FromName(InConfiguration.MediaConfiguration.MediaConnection.Device.DeviceName)
+			, FText::AsNumber(InConfiguration.MediaConfiguration.MediaConnection.Device.DeviceIdentifier)
+			, GetTransportName(InConfiguration.MediaConfiguration.MediaConnection.TransportType, InConfiguration.MediaConfiguration.MediaConnection.QuadTransportType)
+			, InConfiguration.MediaConfiguration.MediaMode.GetModeName()
+		);
+	}
+	return LOCTEXT("Invalid", "<Invalid>");
+}
+
+#undef LOCTEXT_NAMESPACE
diff --git a/Engine/Plugins/Experimental/Avalanche/Source/AvalancheMedia/Private/OutputDevices/AvaDisplayDeviceProvider.h b/Engine/Plugins/Experimental/Avalanche/Source/AvalancheMedia/Private/OutputDevices/AvaDisplayDeviceProvider.h
new file mode 100644
--- /dev/null
+++ b/Engine/Plugins/Experimental/Avalanche/Source/AvalancheMedia/Private/OutputDevices/AvaDisplayDeviceProvider.h
@@ -0,0 +1,47 @@
+// Copyright Epic Games, Inc. All Rights Reserved.
+
+#pragma once
+
+#include "CoreMinimal.h"
+#include "IMediaIOCoreDeviceProvider.h"
+#include "UObject/ObjectMacros.h"
+
+//#include "AvaDisplayDeviceProvider.generated.h"
+
+/**
+ * Implementation of IMediaIOCoreDeviceProvider for UAvaDisplayMediaOutput.
+ */
+class FAvaDisplayDeviceProvider : public IMediaIOCoreDeviceProvider
+{
+public:
+	static FName GetProviderName();
+	static FName GetProtocolName();
+	
+public:
+	virtual FName GetFName() override;
+
+	virtual TArray<FMediaIOConnection> GetConnections() const override;
+	virtual TArray<FMediaIOConfiguration> GetConfigurations() const override;
+	virtual TArray<FMediaIOConfiguration> GetConfigurations(bool bInAllowInput, bool bInAllowOutput) const override;
+	virtual TArray<FMediaIODevice> GetDevices() const override;
+	virtual TArray<FMediaIOMode> GetModes(const FMediaIODevice& InDevice, bool bInOutput) const override;
+	virtual TArray<FMediaIOInputConfiguration> GetInputConfigurations() const override;
+	virtual TArray<FMediaIOOutputConfiguration> GetOutputConfigurations() const override;
+	virtual TArray<FMediaIOVideoTimecodeConfiguration> GetTimecodeConfigurations() const override;
+	
+	virtual FMediaIOConfiguration GetDefaultConfiguration() const override;
+	virtual FMediaIOMode GetDefaultMode() const override;
+	virtual FMediaIOInputConfiguration GetDefaultInputConfiguration() const override;
+	virtual FMediaIOOutputConfiguration GetDefaultOutputConfiguration() const override;
+	virtual FMediaIOVideoTimecodeConfiguration GetDefaultTimecodeConfiguration() const override;
+	
+	virtual FText ToText(const FMediaIOConfiguration& InConfiguration, bool bInIsAutoDetected = false) const override;
+	virtual FText ToText(const FMediaIOConnection& InConnection) const override;
+	virtual FText ToText(const FMediaIOOutputConfiguration& InConfiguration) const override;
+
+#if WITH_EDITOR
+	virtual bool ShowInputTransportInSelector() const override { return false; }
+	virtual bool ShowInputKeyInSelector() const override { return false; }
+	virtual bool ShowOutputKeyInSelector() const override { return false; }
+#endif
+};
diff --git a/Engine/Plugins/Experimental/Avalanche/Source/AvalancheMedia/Private/OutputDevices/AvaDisplayMediaCapture.cpp b/Engine/Plugins/Experimental/Avalanche/Source/AvalancheMedia/Private/OutputDevices/AvaDisplayMediaCapture.cpp
new file mode 100644
--- /dev/null
+++ b/Engine/Plugins/Experimental/Avalanche/Source/AvalancheMedia/Private/OutputDevices/AvaDisplayMediaCapture.cpp
@@ -0,0 +1,224 @@
+// Copyright Epic Games, Inc. All Rights Reserved.
+
+#include "OutputDevices/AvaDisplayMediaCapture.h"
+
+#include "AvaRenderTargetUtils.h"
+#include "Brushes/SlateImageBrush.h"
+#include "Engine/TextureRenderTarget2D.h"
+#include "Framework/Application/SlateApplication.h"
+#include "GenericPlatform/GenericApplication.h"
+#include "GenericPlatform/GenericWindow.h"
+#include "OutputDevices/AvaDisplayDeviceManager.h"
+#include "OutputDevices/AvaDisplayMediaOutput.h"
+#include "OutputDevices/AvaRenderTargetMediaUtils.h"
+#include "OutputDevices/Slate/SAvaCaptureImage.h"
+#include "Slate/SceneViewport.h"
+#include "TextureResource.h"
+#include "Widgets/SWindow.h"
+
+DEFINE_LOG_CATEGORY(LogAvaDisplayMedia);
+
+class UAvaDisplayMediaCapture::FCaptureInstance
+{
+public:
+	FCaptureInstance(UAvaDisplayMediaOutput* InMediaOutput, UTextureRenderTarget2D* InRenderTarget)
+	{
+		if (!InRenderTarget)
+		{
+			return;
+		}
+		
+		TArray<FAvaMonitorInfo> MonitorInfos = FAvaDisplayDeviceManager::GetCachedMonitors();
+
+		FString DeviceName = InMediaOutput->OutputConfiguration.MediaConfiguration.MediaConnection.Device.DeviceName.ToString();
+
+		FVector2D Size(InRenderTarget->SizeX, InRenderTarget->SizeY);
+		FVector2D Position(0, 0);
+		EAutoCenter CenterRule = EAutoCenter::PreferredWorkArea;
+
+		for (int32 i = 0; i < MonitorInfos.Num(); ++i)
+		{
+			const FAvaMonitorInfo& MonitorInfo = MonitorInfos[i];
+	
+			if (DeviceName.Contains(MonitorInfo.Name))
+			{
+				Size = FVector2D(MonitorInfo.Width, MonitorInfo.Height);
+				Position = FVector2D(MonitorInfo.DisplayRect.Left, MonitorInfo.DisplayRect.Top);
+				CenterRule = EAutoCenter::None;	// Since we found the display, don't try to auto center.
+				break;
+			}
+		}
+
+		ImageBrush = MakeShared<FSlateImageBrush>(InRenderTarget, FVector2D(InRenderTarget->SizeX, InRenderTarget->SizeY));
+
+		// Source: FPIEPreviewDeviceModule
+
+		static FWindowStyle BackgroundlessStyle = FCoreStyle::Get().GetWidgetStyle<FWindowStyle>("Window");
+		BackgroundlessStyle.SetBackgroundBrush(FSlateNoResource());
+
+		// Attempt at creating a borderless game window.
+		TSharedRef<SWindow> NewWindow = SNew(SWindow)
+			.Type(EWindowType::GameWindow)
+			.Style(&BackgroundlessStyle)
+			.ScreenPosition(Position)
+			.ClientSize(Size)
+			.UseOSWindowBorder(false)
+			.LayoutBorder(FMargin(0))
+			.SizingRule(ESizingRule::FixedSize)
+			.CreateTitleBar(false)
+			.HasCloseButton(false)
+			.SupportsMinimize(false)
+			.SupportsMaximize(false)
+			.AutoCenter(CenterRule)	// Don't mess with requested position/size
+			.AdjustInitialSizeAndPositionForDPIScale(false)	// Don't mess with requested position/size
+			[
+				SNew(SAvaCaptureImage)
+				.ImageArgs(SImage::FArguments()
+					.Image(ImageBrush.Get()))
+				.EnableGammaCorrection(false)
+				.EnableBlending(false)
+			];
+
+		// This will make the window and show it.
+		FSlateApplication::Get().AddWindow(NewWindow, true);
+
+		NewWindow->SetWindowMode(EWindowMode::Type::WindowedFullscreen);
+
+		Window = NewWindow;
+
+		void* ViewportResource = FSlateApplicationBase::Get().GetRenderer()->GetViewportResource(*Window);
+		if (ViewportResource)
+		{
+			ViewportRHI = *((FViewportRHIRef*)ViewportResource);
+		}
+	}
+
+	~FCaptureInstance()
+	{
+		if (Window.IsValid() && FSlateApplicationBase::IsInitialized())
+		{
+			Window->RequestDestroyWindow();
+		}
+	}
+
+public:
+	TSharedPtr<SWindow> Window;
+	FViewportRHIRef ViewportRHI;		// SWindow's render viewport.
+
+	TSharedPtr<FSlateImageBrush> ImageBrush;
+};
+
+UAvaDisplayMediaCapture::~UAvaDisplayMediaCapture()
+{
+	delete CaptureInstance;
+}
+
+void UAvaDisplayMediaCapture::OnRHIResourceCaptured_RenderingThread(const FCaptureBaseData& InBaseData, TSharedPtr<FMediaCaptureUserData, ESPMode::ThreadSafe> InUserData, FTextureRHIRef InTexture)
+{
+	FScopeLock ScopeLock(&CaptureInstanceCriticalSection);
+	if (CaptureInstance)
+	{
+		FTextureRHIRef Target;
+
+		// Ideally, we would just copy the texture right into the window's back buffer and present.
+		// But if we do that, it will be overwritten by the Slate renderer. We would need a way to
+		// disable Slate renderer for that window.
+		//if (CaptureInstance->ViewportRHI.IsValid())
+		//{
+		//	Target = RHIGetViewportBackBuffer(CaptureInstance->ViewportRHI.GetReference());
+		//}
+
+		// In the mean time, we use an intermediate render target that will be then rendered by Slate in the window.
+		// This means an additional avoidable render pass so it is less optimal (but works).
+		if (RenderTarget)
+		{
+			Target = RenderTarget->GetRenderTargetResource() ? RenderTarget->GetRenderTargetResource()->GetRenderTargetTexture() : nullptr;
+		}
+
+		if (Target.IsValid())
+		{
+			FRHICommandListImmediate& RHICmdList = FRHICommandListExecutor::GetImmediateCommandList();
+			UE::AvaRenderTargetMediaUtils::CopyTexture(RHICmdList, InTexture, Target);
+		}
+	}
+}
+
+bool UAvaDisplayMediaCapture::InitializeCapture()
+{
+	return true;
+}
+
+bool UAvaDisplayMediaCapture::PostInitializeCaptureViewport(TSharedPtr<FSceneViewport>& InSceneViewport)
+{
+	bool bSuccess = false;
+	const FTexture2DRHIRef& BackBuffer = InSceneViewport->GetRenderTargetTexture();
+	if (BackBuffer.IsValid())
+	{
+		const FRHITextureDesc& Desc = BackBuffer->GetDesc();
+		bSuccess = StartNewCapture(Desc.Extent, Desc.Format);
+	}
+	else
+	{
+		UE_LOG(LogAvaDisplayMedia, Error, TEXT("Can't start capture of a scene viewport with no back buffer."));
+	}
+	return bSuccess;
+}
+
+bool UAvaDisplayMediaCapture::PostInitializeCaptureRenderTarget(UTextureRenderTarget2D* InRenderTarget)
+{
+	bool bSuccess = StartNewCapture(FIntPoint(InRenderTarget->SizeX, InRenderTarget->SizeY), InRenderTarget->GetFormat());
+	return bSuccess;
+}
+
+void UAvaDisplayMediaCapture::StopCaptureImpl(bool /*bAllowPendingFrameToBeProcess*/)
+{
+	TRACE_BOOKMARK(TEXT("UAvaDisplayMediaCapture::StopCapture"));
+
+	FScopeLock ScopeLock(&CaptureInstanceCriticalSection);
+
+	delete CaptureInstance;
+	CaptureInstance = nullptr;
+}
+
+bool UAvaDisplayMediaCapture::StartNewCapture(const FIntPoint& InSourceTargetSize, EPixelFormat InSourceTargetFormat)
+{
+	TRACE_BOOKMARK(TEXT("UAvaDisplayMediaCapture::StartNewCapture"));
+	{
+		FScopeLock ScopeLock(&CaptureInstanceCriticalSection);
+
+		delete CaptureInstance;
+		CaptureInstance = nullptr;
+
+		UAvaDisplayMediaOutput* AvaDisplayMediaOutput = CastChecked<UAvaDisplayMediaOutput>(MediaOutput);
+		if (AvaDisplayMediaOutput)
+		{
+			static const FName RenderTargetBaseName = TEXT("AvaDisplayCapture_RenderTarget");
+			RenderTarget = AvaRenderTargetUtils::CreateDefaultRenderTarget(RenderTargetBaseName);
+			check(RenderTarget);
+
+			FIntPoint TargetSize = AvaDisplayMediaOutput->GetRequestedSize();
+			if (TargetSize == UMediaOutput::RequestCaptureSourceSize)
+			{
+				TargetSize = InSourceTargetSize;
+			}
+
+			EPixelFormat TargetFormat = AvaDisplayMediaOutput->GetRequestedPixelFormat();
+			if (TargetFormat == EPixelFormat::PF_Unknown)
+			{
+				TargetFormat = InSourceTargetFormat;
+			}
+
+			RenderTarget->OverrideFormat = TargetFormat;
+			RenderTarget->ResizeTarget(TargetSize.X, TargetSize.Y);
+
+			CaptureInstance = new FCaptureInstance(AvaDisplayMediaOutput, RenderTarget);
+		}
+		else
+		{
+			UE_LOG(LogAvaDisplayMedia, Error, TEXT("Internal Error: Media Capture's associated Media Output is not of type \"UAvaDisplayMediaOutput\"."));
+		}
+	}
+
+	SetState(EMediaCaptureState::Capturing);
+	return true;
+}
diff --git a/Engine/Plugins/Experimental/Avalanche/Source/AvalancheMedia/Private/OutputDevices/AvaDisplayMediaOutput.cpp b/Engine/Plugins/Experimental/Avalanche/Source/AvalancheMedia/Private/OutputDevices/AvaDisplayMediaOutput.cpp
new file mode 100644
--- /dev/null
+++ b/Engine/Plugins/Experimental/Avalanche/Source/AvalancheMedia/Private/OutputDevices/AvaDisplayMediaOutput.cpp
@@ -0,0 +1,45 @@
+// Copyright Epic Games, Inc. All Rights Reserved.
+
+#include "OutputDevices/AvaDisplayMediaOutput.h"
+#include "OutputDevices/AvaDisplayMediaCapture.h"
+
+UAvaDisplayMediaOutput::UAvaDisplayMediaOutput(const FObjectInitializer& ObjectInitializer)
+	: Super(ObjectInitializer)
+{
+}
+
+bool UAvaDisplayMediaOutput::Validate(FString& OutFailureReason) const
+{
+	if (!Super::Validate(OutFailureReason))
+	{
+		return false;
+	}
+
+	return true;
+}
+
+FIntPoint UAvaDisplayMediaOutput::GetRequestedSize() const
+{
+	return OutputConfiguration.MediaConfiguration.MediaMode.Resolution;
+}
+
+EPixelFormat UAvaDisplayMediaOutput::GetRequestedPixelFormat() const
+{
+	return PF_A2B10G10R10;
+}
+
+EMediaCaptureConversionOperation UAvaDisplayMediaOutput::GetConversionOperation(EMediaCaptureSourceType /*InSourceType*/) const
+{
+	return EMediaCaptureConversionOperation::NONE;
+}
+
+UMediaCapture* UAvaDisplayMediaOutput::CreateMediaCaptureImpl()
+{
+	UMediaCapture* Result = NewObject<UAvaDisplayMediaCapture>();
+	if (Result)
+	{
+		UE_LOG(LogAvaDisplayMedia, Log, TEXT("Created Avalanche Display Media Capture"));
+		Result->SetMediaOutput(this);
+	}
+	return Result;
+}
diff --git a/Engine/Plugins/Experimental/Avalanche/Source/AvalancheMedia/Private/OutputDevices/AvaMediaOutputUtils.cpp b/Engine/Plugins/Experimental/Avalanche/Source/AvalancheMedia/Private/OutputDevices/AvaMediaOutputUtils.cpp
new file mode 100644
--- /dev/null
+++ b/Engine/Plugins/Experimental/Avalanche/Source/AvalancheMedia/Private/OutputDevices/AvaMediaOutputUtils.cpp
@@ -0,0 +1,113 @@
+// Copyright Epic Games, Inc. All Rights Reserved.
+
+#include "OutputDevices/AvaMediaOutputUtils.h"
+
+#include "AvaMediaModule.h"
+#include "MediaOutput.h"
+#include "Playback/AvaMediaPlaybackMessages.h"
+#include "Serialization/MemoryReader.h"
+#include "Serialization/ObjectAndNameAsStringProxyArchive.h"
+
+namespace UE::AvaMediaOutputUtils::Private
+{
+	// This is a meta data that is added to all Media Output classes (that have a device provider)
+	// giving information on what DeviceProvider it is using.
+	static const FName NAME_MediaIOCustomLayout("MediaIOCustomLayout");
+}
+
+FString UE::AvaMediaOutputUtils::GetDeviceName(const UMediaOutput* InMediaOutput)
+{
+	const UClass* const MediaOutputClass = InMediaOutput->GetClass();
+	{
+		const FStructProperty* const Property = FindFProperty<FStructProperty>(MediaOutputClass, TEXT("OutputConfiguration"));
+		if (Property && Property->Struct->IsChildOf(FMediaIOOutputConfiguration::StaticStruct()))
+		{
+			const FMediaIOOutputConfiguration* const OutputConfig = Property->ContainerPtrToValuePtr<FMediaIOOutputConfiguration>(InMediaOutput);
+			return OutputConfig->MediaConfiguration.MediaConnection.Device.DeviceName.ToString();
+		}
+	}
+	
+	if (const FProperty* const SourceNameProperty = FindFProperty<FProperty>(MediaOutputClass, TEXT("SourceName")))
+	{
+		return *SourceNameProperty->ContainerPtrToValuePtr<FString>(InMediaOutput);
+	}
+
+	if (const FProperty* const StreamerIdProperty = FindFProperty<FProperty>(MediaOutputClass, TEXT("StreamerId")))
+	{
+		return *StreamerIdProperty->ContainerPtrToValuePtr<FString>(InMediaOutput);
+	}
+	
+	return TEXT("");
+}
+
+bool UE::AvaMediaOutputUtils::HasDeviceProviderName(const UClass* InMediaOutputClass)
+{
+#if WITH_EDITORONLY_DATA
+	return IsValid(InMediaOutputClass)
+		&& InMediaOutputClass->IsChildOf(UMediaOutput::StaticClass())
+		&& InMediaOutputClass->HasMetaData(Private::NAME_MediaIOCustomLayout);
+#else
+	return false;
+#endif
+}
+
+FName UE::AvaMediaOutputUtils::GetDeviceProviderName(const UClass* InMediaOutputClass)
+{
+#if WITH_EDITORONLY_DATA
+	if (HasDeviceProviderName(InMediaOutputClass))
+	{
+		return FName(InMediaOutputClass->GetMetaData(Private::NAME_MediaIOCustomLayout));
+	}
+#endif
+	return FName();
+}
+
+FName UE::AvaMediaOutputUtils::GetDeviceProviderName(const UMediaOutput* InMediaOutput)
+{
+#if WITH_EDITORONLY_DATA
+	const UClass* const MediaOutputClass = IsValid(InMediaOutput) ? InMediaOutput->GetClass() : nullptr;
+	if (MediaOutputClass
+		&& IsValid(MediaOutputClass)
+		&& MediaOutputClass->HasMetaData(Private::NAME_MediaIOCustomLayout))
+	{
+		return FName(MediaOutputClass->GetMetaData(Private::NAME_MediaIOCustomLayout));
+	}
+#endif
+	return FName();
+}
+
+FAvaMediaOutputData UE::AvaMediaOutputUtils::CreateMediaOutputData(UMediaOutput* InMediaOutput)
+{
+	FAvaMediaOutputData MediaOutputData;
+	MediaOutputData.MediaOutputClass = InMediaOutput->GetClass();
+			
+	TArray<uint8> SerializedData;
+	FMemoryWriter MemoryWriter(SerializedData, true); 
+	FObjectAndNameAsStringProxyArchive ObjectWriter = FObjectAndNameAsStringProxyArchive(MemoryWriter, false);
+	InMediaOutput->Serialize(ObjectWriter);
+
+	MediaOutputData.SerializedData = MoveTemp(SerializedData);
+	MediaOutputData.ObjectFlags = InMediaOutput->GetFlags();
+
+	return MediaOutputData;
+}
+
+UMediaOutput* UE::AvaMediaOutputUtils::CreateMediaOutput(const FAvaMediaOutputData& InMediaOutputData, UObject* InOuter)
+{
+	UClass* const MediaOutputClass = InMediaOutputData.MediaOutputClass.LoadSynchronous();
+	UMediaOutput* MediaOutput = nullptr;
+	if (MediaOutputClass)
+	{
+		MediaOutput = NewObject<UMediaOutput>(
+			InOuter
+			, MediaOutputClass
+			, NAME_None
+			, InMediaOutputData.GetObjectFlags());
+
+		FMemoryReader MemoryReader(InMediaOutputData.SerializedData, true); 
+		FObjectAndNameAsStringProxyArchive ObjectReader = FObjectAndNameAsStringProxyArchive(MemoryReader, false);
+		MediaOutput->Serialize(ObjectReader);
+		MediaOutput->PostLoad();
+	}
+	return MediaOutput;
+}
\ No newline at end of file
diff --git a/Engine/Plugins/Experimental/Avalanche/Source/AvalancheMedia/Private/OutputDevices/AvaRenderTargetMediaCapture.cpp b/Engine/Plugins/Experimental/Avalanche/Source/AvalancheMedia/Private/OutputDevices/AvaRenderTargetMediaCapture.cpp
new file mode 100644
--- /dev/null
+++ b/Engine/Plugins/Experimental/Avalanche/Source/AvalancheMedia/Private/OutputDevices/AvaRenderTargetMediaCapture.cpp
@@ -0,0 +1,121 @@
+// Copyright Epic Games, Inc. All Rights Reserved.
+
+#include "OutputDevices/AvaRenderTargetMediaCapture.h"
+
+#include "AvaRenderTargetUtils.h"
+#include "Engine/Engine.h"
+#include "Engine/TextureRenderTarget2D.h"
+#include "OutputDevices/AvaRenderTargetMediaOutput.h"
+#include "OutputDevices/AvaRenderTargetMediaUtils.h"
+#include "Slate/SceneViewport.h"
+#include "TextureResource.h"
+
+DEFINE_LOG_CATEGORY(LogAvaRenderTargetMedia);
+
+void UAvaRenderTargetMediaCapture::OnRHIResourceCaptured_RenderingThread(const FCaptureBaseData& InBaseData, TSharedPtr<FMediaCaptureUserData, ESPMode::ThreadSafe> InUserData, FTextureRHIRef InTexture)
+{
+	FScopeLock ScopeLock(&RenderTargetCriticalSection);
+
+	if (!RenderTarget || !RenderTarget->GetRenderTargetResource())
+	{
+		return;
+	}
+	
+	const FTextureRHIRef RenderTargetRHIRef = RenderTarget->GetRenderTargetResource()->GetRenderTargetTexture();
+
+	if (RenderTargetRHIRef.IsValid())
+	{
+		FRHICommandListImmediate& RHICmdList = FRHICommandListExecutor::GetImmediateCommandList();
+
+		// Note: we possibly don't want to undo this for RGB8_SRGB texture. Need to test.
+		constexpr bool bSRGBToLinear = true;
+
+		float Gamma = 1.0f;
+		
+		if (const FTextureRenderTarget2DResource* Resource = static_cast<FTextureRenderTarget2DResource*>(RenderTarget->GetResource()))
+		{
+			// The render targets for rendering the ava channel has bForceLinearGamma = false
+			// and TargetGamma = 2.2 for float textures (to force output to gamma space).
+			// Those parameters are used to determine the render target's display gamma, so
+			// we need to match the same display gamma that was used for rendering in order to invert it.
+			// See UAvalanchePlayback::UpdateRenderTarget.			
+			const bool bBackupLinearGamma = RenderTarget->bForceLinearGamma;
+			const float BackupTargetGamma = RenderTarget->TargetGamma;
+			RenderTarget->bForceLinearGamma = false;
+			if (AvaRenderTargetUtils::IsFloatFormat(RenderTarget))
+			{
+				RenderTarget->TargetGamma = GEngine ? GEngine->GetDisplayGamma() : 2.2f;
+			}
+			// Must invert the gamma done in post process.
+			// See GetTonemapperOutputDeviceParameters (PostProcessTonemap.cpp).
+			Gamma = Resource ? Resource->GetDisplayGamma()/2.2f : 1.0f;
+			RenderTarget->bForceLinearGamma = bBackupLinearGamma;
+			RenderTarget->TargetGamma = BackupTargetGamma;
+		}
+
+		// This will convert the color space from SCS_FinalColorSDR back to SCS_FinalToneCurveHDR.
+		UE::AvaRenderTargetMediaUtils::ConvertTextureRGBGamma(RHICmdList, InTexture, RenderTargetRHIRef, bSRGBToLinear, Gamma);
+	}
+}
+
+bool UAvaRenderTargetMediaCapture::InitializeCapture()
+{
+	return true;
+}
+
+bool UAvaRenderTargetMediaCapture::PostInitializeCaptureViewport(TSharedPtr<FSceneViewport>& InSceneViewport)
+{
+	bool bSuccess = false;
+	const FTexture2DRHIRef& BackBuffer = InSceneViewport->GetRenderTargetTexture();
+	if (BackBuffer.IsValid())
+	{
+		const FRHITextureDesc& Desc = BackBuffer->GetDesc();
+		bSuccess = StartNewCapture(Desc.Extent, Desc.Format);
+	}
+	else
+	{
+		UE_LOG(LogAvaRenderTargetMedia, Error, TEXT("Can't start capture of a scene viewport with no back buffer."));
+	}
+	return bSuccess;
+}
+
+bool UAvaRenderTargetMediaCapture::PostInitializeCaptureRenderTarget(UTextureRenderTarget2D* InRenderTarget)
+{
+	return StartNewCapture(FIntPoint(InRenderTarget->SizeX, InRenderTarget->SizeY), InRenderTarget->GetFormat());
+}
+
+void UAvaRenderTargetMediaCapture::StopCaptureImpl(bool /*bAllowPendingFrameToBeProcess*/)
+{
+	TRACE_BOOKMARK(TEXT("UAvaRenderTargetMediaCapture::StopCapture"));
+}
+
+bool UAvaRenderTargetMediaCapture::StartNewCapture(const FIntPoint& InSourceTargetSize, EPixelFormat InSourceTargetFormat)
+{
+	TRACE_BOOKMARK(TEXT("UAvaRenderTargetMediaCapture::StartNewCapture"));
+	{
+		FScopeLock ScopeLock(&RenderTargetCriticalSection);
+
+		RenderTarget = nullptr;
+		
+		const UAvaRenderTargetMediaOutput* AvaRenderTargetMediaOutput = CastChecked<UAvaRenderTargetMediaOutput>(MediaOutput);
+		if (AvaRenderTargetMediaOutput)
+		{
+			RenderTarget = AvaRenderTargetMediaOutput->RenderTarget.LoadSynchronous();
+			if (RenderTarget)
+			{
+				SetState(EMediaCaptureState::Capturing);
+				return true;
+			}
+			else
+			{
+				UE_LOG(LogAvaRenderTargetMedia, Error, TEXT("Missing render target object: \"%s\"."), *AvaRenderTargetMediaOutput->RenderTarget.ToString());
+			}
+		}
+		else
+		{
+			UE_LOG(LogAvaRenderTargetMedia, Error, TEXT("Internal Error: Media Capture's associated Media Output is not of type \"UAvaRenderTargetMediaOutput\"."));
+		}
+	}
+	SetState(EMediaCaptureState::Error);
+	return false;
+}
diff --git a/Engine/Plugins/Experimental/Avalanche/Source/AvalancheMedia/Private/OutputDevices/AvaRenderTargetMediaOutput.cpp b/Engine/Plugins/Experimental/Avalanche/Source/AvalancheMedia/Private/OutputDevices/AvaRenderTargetMediaOutput.cpp
new file mode 100644
--- /dev/null
+++ b/Engine/Plugins/Experimental/Avalanche/Source/AvalancheMedia/Private/OutputDevices/AvaRenderTargetMediaOutput.cpp
@@ -0,0 +1,67 @@
+// Copyright Epic Games, Inc. All Rights Reserved.
+
+#include "OutputDevices/AvaRenderTargetMediaOutput.h"
+
+#include "Engine/TextureRenderTarget2D.h"
+#include "OutputDevices/AvaRenderTargetMediaCapture.h"
+
+bool UAvaRenderTargetMediaOutput::Validate(FString& OutFailureReason) const
+{
+	if (!Super::Validate(OutFailureReason))
+	{
+		return false;
+	}
+
+	const UTextureRenderTarget2D* LoadedRenderTarget = RenderTarget.LoadSynchronous();
+	if (!LoadedRenderTarget)
+	{
+		OutFailureReason = FString::Printf(TEXT("Missing render target object: \"%s\"."), *RenderTarget.ToString());
+		return false;
+	}
+	return true;
+}
+
+FIntPoint UAvaRenderTargetMediaOutput::GetRequestedSize() const
+{
+	const UTextureRenderTarget2D* LoadedRenderTarget = RenderTarget.LoadSynchronous();
+	return (LoadedRenderTarget) ? FIntPoint(LoadedRenderTarget->SizeX, LoadedRenderTarget->SizeY) : RequestCaptureSourceSize;
+}
+
+
+EPixelFormat UAvaRenderTargetMediaOutput::GetRequestedPixelFormat() const
+{
+	const UTextureRenderTarget2D* LoadedRenderTarget = RenderTarget.LoadSynchronous();
+	return (LoadedRenderTarget) ? LoadedRenderTarget->GetFormat() : EPixelFormat::PF_Unknown;
+}
+
+EMediaCaptureConversionOperation UAvaRenderTargetMediaOutput::GetConversionOperation(EMediaCaptureSourceType /*InSourceType*/) const
+{
+	return bInvertKeyOutput ? EMediaCaptureConversionOperation::INVERT_ALPHA : EMediaCaptureConversionOperation::NONE;
+}
+
+#if WITH_EDITOR
+void UAvaRenderTargetMediaOutput::PostEditChangeProperty(struct FPropertyChangedEvent& PropertyChangedEvent)
+{
+	// The source name is used to display as the "device name" in the broadcast editor.
+	SourceName = RenderTarget.GetAssetName();
+	
+	// Fallback in case the path is not set.
+	if (SourceName.IsEmpty())
+	{
+		SourceName = GetFName().ToString();
+	}
+	
+	Super::PostEditChangeProperty(PropertyChangedEvent);
+}
+#endif
+
+UMediaCapture* UAvaRenderTargetMediaOutput::CreateMediaCaptureImpl()
+{
+	UMediaCapture* Result = NewObject<UAvaRenderTargetMediaCapture>();
+	if (Result)
+	{
+		UE_LOG(LogAvaRenderTargetMedia, Log, TEXT("Created Avalanche Render Target Media Capture"));
+		Result->SetMediaOutput(this);
+	}
+	return Result;
+}
diff --git a/Engine/Plugins/Experimental/Avalanche/Source/AvalancheMedia/Private/OutputDevices/AvaRenderTargetMediaUtils.cpp b/Engine/Plugins/Experimental/Avalanche/Source/AvalancheMedia/Private/OutputDevices/AvaRenderTargetMediaUtils.cpp
new file mode 100644
--- /dev/null
+++ b/Engine/Plugins/Experimental/Avalanche/Source/AvalancheMedia/Private/OutputDevices/AvaRenderTargetMediaUtils.cpp
@@ -0,0 +1,153 @@
+// Copyright Epic Games, Inc. All Rights Reserved.
+
+#include "OutputDevices/AvaRenderTargetMediaUtils.h"
+#include "DataDrivenShaderPlatformInfo.h"
+#include "Engine/TextureRenderTarget2D.h"
+#include "MediaShaders.h"
+#include "RHIStaticStates.h"
+#include "RenderGraphUtils.h"
+#include "ScreenPass.h"
+#include "ScreenRendering.h"
+#include "TextureResource.h"
+
+void UE::AvaRenderTargetMediaUtils::ClearRenderTarget(UTextureRenderTarget2D* InRenderTarget)
+{
+	FLinearColor ClearColor = InRenderTarget->ClearColor;
+	FTextureRenderTargetResource* RenderTargetResource = InRenderTarget->GameThread_GetRenderTargetResource();
+	ENQUEUE_RENDER_COMMAND(FAvaClearRenderTarget)(
+		[RenderTargetResource, ClearColor, RenderTargetName = InRenderTarget->GetName()](FRHICommandListImmediate& RHICmdList)
+		{
+			FRDGBuilder GraphBuilder(RHICmdList, RDG_EVENT_NAME("ClearRenderTarget(%s)", *RenderTargetName));
+			const FRDGTextureRef OutputTexture = GraphBuilder.RegisterExternalTexture(CreateRenderTarget(RenderTargetResource->TextureRHI, *RenderTargetName));
+			AddClearRenderTargetPass(GraphBuilder, OutputTexture, ClearColor);
+			GraphBuilder.Execute();
+		});
+}
+
+void UE::AvaRenderTargetMediaUtils::CopyTexture(FRHICommandListImmediate& InRHICmdList, FTextureRHIRef InSourceTexture, FTextureRHIRef InDestTarget)
+{
+	if (!InSourceTexture.IsValid() || !InDestTarget.IsValid())
+	{
+		return;
+	}
+	
+	FRDGBuilder GraphBuilder(InRHICmdList);
+
+	FRDGTexture* InputTexture = GraphBuilder.RegisterExternalTexture(CreateRenderTarget(InSourceTexture, TEXT("SourceTexture")));
+	FRDGTexture* OutputTexture = GraphBuilder.RegisterExternalTexture(CreateRenderTarget(InDestTarget, TEXT("DestTarget")));
+	const FGlobalShaderMap* GlobalShaderMap = GetGlobalShaderMap(GMaxRHIFeatureLevel);
+	
+	if (InputTexture->Desc.Format == OutputTexture->Desc.Format &&
+		InputTexture->Desc.Extent.X == OutputTexture->Desc.Extent.X &&
+		InputTexture->Desc.Extent.Y == OutputTexture->Desc.Extent.Y)
+	{
+		// The formats are the same and size are the same. simple copy
+		// Note: this deals with differing formats.
+		// Note2: also deals with mipmaps if any.
+		AddDrawTexturePass(GraphBuilder, GlobalShaderMap, InputTexture, OutputTexture, FRDGDrawTextureInfo());
+	}
+	else
+	{
+		// The formats or sizes differ to pixel shader stuff
+		// Configure source/output viewport to get the right UV scaling from source texture to output texture
+		const FScreenPassTextureViewport InputViewport(InputTexture);
+		const FScreenPassTextureViewport OutputViewport(OutputTexture);
+
+		const TShaderMapRef<FScreenPassVS> VertexShader(GlobalShaderMap);
+
+		// In cases where texture is converted from a format that doesn't have A channel, we want to force set it to 1.
+		constexpr int32 ConversionOperation = 0; // None
+		FModifyAlphaSwizzleRgbaPS::FPermutationDomain PermutationVector;
+		PermutationVector.Set<FModifyAlphaSwizzleRgbaPS::FConversionOp>(ConversionOperation);
+
+		// Rectangle area to use from source
+		const FIntRect ViewRect(FIntPoint(0, 0), InputTexture->Desc.Extent);
+
+		//Dummy ViewFamily/ViewInfo created to use built in Draw Screen/Texture Pass
+		const FSceneViewFamily ViewFamily(FSceneViewFamily::ConstructionValues(nullptr, nullptr, FEngineShowFlags(ESFIM_Game))
+			.SetTime(FGameTime()));
+		FSceneViewInitOptions ViewInitOptions;
+		ViewInitOptions.ViewFamily = &ViewFamily;
+		ViewInitOptions.SetViewRectangle(ViewRect);
+		ViewInitOptions.ViewOrigin = FVector::ZeroVector;
+		ViewInitOptions.ViewRotationMatrix = FMatrix::Identity;
+		ViewInitOptions.ProjectionMatrix = FMatrix::Identity;
+		FSceneView ViewInfo(ViewInitOptions);
+
+		const TShaderMapRef<FModifyAlphaSwizzleRgbaPS> PixelShader(GlobalShaderMap, PermutationVector);
+		FModifyAlphaSwizzleRgbaPS::FParameters* Parameters = PixelShader->AllocateAndSetParameters(GraphBuilder, InputTexture, OutputTexture);
+		AddDrawScreenPass(GraphBuilder, RDG_EVENT_NAME("AvaRenderTargetMediaSwizzle"), ViewInfo, OutputViewport, InputViewport, VertexShader, PixelShader, Parameters);
+	}
+	GraphBuilder.Execute();
+}
+
+namespace UE::AvaRenderTargetMediaUtils::Private
+{
+	class FAvaRGBGammaConvertPS : public FGlobalShader
+	{
+		DECLARE_SHADER_TYPE(FAvaRGBGammaConvertPS, Global);
+		SHADER_USE_PARAMETER_STRUCT(FAvaRGBGammaConvertPS, FGlobalShader);
+
+		BEGIN_SHADER_PARAMETER_STRUCT(FParameters, )
+		SHADER_PARAMETER(int, SrgbToLinear)
+		SHADER_PARAMETER(float, Gamma)
+		SHADER_PARAMETER_RDG_TEXTURE(Texture2D, InputTexture)
+		SHADER_PARAMETER_SAMPLER(SamplerState, InputSampler)
+		RENDER_TARGET_BINDING_SLOTS()
+		END_SHADER_PARAMETER_STRUCT()
+	
+	public:
+		static bool ShouldCompilePermutation(const FGlobalShaderPermutationParameters& InParameters)
+		{
+			return IsFeatureLevelSupported(InParameters.Platform, ERHIFeatureLevel::ES3_1);
+		}
+		
+		FParameters* AllocateAndSetParameters(FRDGBuilder& InGraphBuilder, FRDGTextureRef InRGBATexture, bool bInSrgbToLinear, float InGamma, FRDGTextureRef InOutputTexture)
+		{
+			FParameters* Parameters = InGraphBuilder.AllocParameters<FParameters>();
+			Parameters->InputTexture = InRGBATexture;
+			Parameters->InputSampler = TStaticSamplerState<SF_Bilinear>::GetRHI(); // or  SF_Point
+			Parameters->SrgbToLinear = bInSrgbToLinear ? 1 : 0;
+			Parameters->Gamma = InGamma;
+			Parameters->RenderTargets[0] = FRenderTargetBinding{ InOutputTexture, ERenderTargetLoadAction::ENoAction };
+			return Parameters;
+		}
+	};
+	IMPLEMENT_SHADER_TYPE(, FAvaRGBGammaConvertPS, TEXT("/Plugin/Avalanche/AvaRGBGammaConvert.usf"), TEXT("MainPS"), SF_Pixel);
+}
+
+void UE::AvaRenderTargetMediaUtils::ConvertTextureRGBGamma(FRHICommandListImmediate& InRHICmdList, FTextureRHIRef InSourceTexture, FTextureRHIRef InDestTarget, bool bInSrgbToLinear, float InGamma)
+{
+	FRDGBuilder GraphBuilder(InRHICmdList);
+
+	FRDGTexture* InputTexture = GraphBuilder.RegisterExternalTexture(CreateRenderTarget(InSourceTexture, TEXT("SourceTexture")));
+	FRDGTexture* OutputTexture = GraphBuilder.RegisterExternalTexture(CreateRenderTarget(InDestTarget, TEXT("DestTarget")));
+	const FGlobalShaderMap* GlobalShaderMap = GetGlobalShaderMap(GMaxRHIFeatureLevel);
+	
+	// The formats or sizes differ to pixel shader stuff
+	// Configure source/output viewport to get the right UV scaling from source texture to output texture
+	const FScreenPassTextureViewport InputViewport(InputTexture);
+	const FScreenPassTextureViewport OutputViewport(OutputTexture);
+
+	const TShaderMapRef<FScreenPassVS> VertexShader(GlobalShaderMap);
+	
+	// Rectangle area to use from source
+	const FIntRect ViewRect(FIntPoint(0, 0), InputTexture->Desc.Extent);
+
+	//Dummy ViewFamily/ViewInfo created to use built in Draw Screen/Texture Pass
+	const FSceneViewFamily ViewFamily(FSceneViewFamily::ConstructionValues(nullptr, nullptr, FEngineShowFlags(ESFIM_Game))
+		.SetTime(FGameTime()));
+	FSceneViewInitOptions ViewInitOptions;
+	ViewInitOptions.ViewFamily = &ViewFamily;
+	ViewInitOptions.SetViewRectangle(ViewRect);
+	ViewInitOptions.ViewOrigin = FVector::ZeroVector;
+	ViewInitOptions.ViewRotationMatrix = FMatrix::Identity;
+	ViewInitOptions.ProjectionMatrix = FMatrix::Identity;
+	FSceneView ViewInfo(ViewInitOptions);
+
+	const TShaderMapRef<Private::FAvaRGBGammaConvertPS> PixelShader(GlobalShaderMap);
+	Private::FAvaRGBGammaConvertPS::FParameters* Parameters = PixelShader->AllocateAndSetParameters(GraphBuilder, InputTexture, bInSrgbToLinear, InGamma, OutputTexture);
+	AddDrawScreenPass(GraphBuilder, RDG_EVENT_NAME("ConvertTextureRGBGamma"), ViewInfo, OutputViewport, InputViewport, VertexShader, PixelShader, Parameters);
+	
+	GraphBuilder.Execute();
+}
diff --git a/Engine/Plugins/Experimental/Avalanche/Source/AvalancheMedia/Private/OutputDevices/Slate/SAvaCaptureImage.cpp b/Engine/Plugins/Experimental/Avalanche/Source/AvalancheMedia/Private/OutputDevices/Slate/SAvaCaptureImage.cpp
new file mode 100644
--- /dev/null
+++ b/Engine/Plugins/Experimental/Avalanche/Source/AvalancheMedia/Private/OutputDevices/Slate/SAvaCaptureImage.cpp
@@ -0,0 +1,73 @@
+// Copyright Epic Games, Inc. All Rights Reserved.
+
+#include "OutputDevices/Slate/SAvaCaptureImage.h"
+
+void SAvaCaptureImage::Construct(const FArguments& InArgs)
+{
+	ShouldInvertAlpha = InArgs._ShouldInvertAlpha;
+	bEnableGammaCorrection = InArgs._EnableGammaCorrection;
+	EnableBlending = InArgs._EnableBlending;
+	
+	SImage::Construct(InArgs._ImageArgs);
+}
+
+int32 SAvaCaptureImage::OnPaint(const FPaintArgs& Args
+	, const FGeometry& AllottedGeometry
+	, const FSlateRect& MyCullingRect
+	, FSlateWindowElementList& OutDrawElements
+	, int32 LayerId
+	, const FWidgetStyle& InWidgetStyle
+	, bool bParentEnabled) const
+{
+	const FSlateBrush* const ImageBrush = GetImageAttribute().Get();
+	
+	if (ImageBrush && ImageBrush->DrawAs != ESlateBrushDrawType::NoDrawType)
+	{
+		const bool bIsEnabled = ShouldBeEnabled(bParentEnabled);
+		
+		ESlateDrawEffect DrawEffects = bIsEnabled ? ESlateDrawEffect::None : ESlateDrawEffect::DisabledEffect;
+
+		if (ShouldInvertAlpha.IsSet() && ShouldInvertAlpha.Get())
+		{
+			DrawEffects |= ESlateDrawEffect::InvertAlpha;
+		}
+
+		if (EnableBlending.IsSet() && EnableBlending.Get() == false)
+		{
+			DrawEffects |= ESlateDrawEffect::NoBlending;
+		}
+
+		if (!bEnableGammaCorrection)
+		{
+			DrawEffects |= ESlateDrawEffect::NoGamma;
+		}
+		
+		const FLinearColor FinalColorAndOpacity(InWidgetStyle.GetColorAndOpacityTint()
+			* GetColorAndOpacityAttribute().Get().GetColor(InWidgetStyle)
+			* ImageBrush->GetTint(InWidgetStyle));
+
+		if (bFlipForRightToLeftFlowDirection && GSlateFlowDirection == EFlowDirection::RightToLeft)
+		{
+			const FGeometry FlippedGeometry = AllottedGeometry.MakeChild(FSlateRenderTransform(FScale2D(-1, 1)));
+			
+			FSlateDrawElement::MakeBox(OutDrawElements
+				, LayerId
+				, FlippedGeometry.ToPaintGeometry()
+				, ImageBrush
+				, DrawEffects
+				, FinalColorAndOpacity);
+		}
+		else
+		{
+			FSlateDrawElement::MakeBox(OutDrawElements
+				, LayerId
+				, AllottedGeometry.ToPaintGeometry()
+				, ImageBrush
+				, DrawEffects
+				, FinalColorAndOpacity);
+		}
+	}
+
+	return LayerId;
+}
+
diff --git a/Engine/Plugins/Experimental/Avalanche/Source/AvalancheMedia/Private/Playback/AvaMediaPlayableGroup.cpp b/Engine/Plugins/Experimental/Avalanche/Source/AvalancheMedia/Private/Playback/AvaMediaPlayableGroup.cpp
new file mode 100644
--- /dev/null
+++ b/Engine/Plugins/Experimental/Avalanche/Source/AvalancheMedia/Private/Playback/AvaMediaPlayableGroup.cpp
@@ -0,0 +1,623 @@
+// Copyright Epic Games, Inc. All Rights Reserved.
+
+#include "Playback/AvaMediaPlayableGroup.h"
+
+#include "Engine/ViewportStatsSubsystem.h"
+#include "Framework/AvalancheGameInstance.h"
+#include "Playback/AvaMediaPlayableGroupManager.h"
+#include "Playback/AvaMediaPlaybackUtils.h"
+#include "Playback/AvalanchePlayable.h"
+#include "Playback/Transition/AvalanchePlayableTransition.h"
+#include "UObject/Package.h"
+
+#define LOCTEXT_NAMESPACE "AvaMediaPlayableGroup"
+
+namespace UE::AvalancheMedia::PlayableGroup::Private
+{
+	template<typename InElementType>
+	void CleanStaleKeys(TSet<TObjectKey<InElementType>>& InOutSetToClean)
+	{
+		for (typename TSet<TObjectKey<InElementType>>::TIterator It(InOutSetToClean); It; ++It)
+		{
+			if (!It->ResolveObjectPtr())
+			{
+				It.RemoveCurrent();
+			}
+		}
+	}
+}
+
+UPackage* UAvaMediaPlayableGroup::MakeGameInstancePackage(const FSoftObjectPath& InSourceAssetPath, const FName& InChannelName)
+{
+	// Remote Control Preset will be registered with the package name.
+	// We want a package name unique to the game instance and that will be
+	// human readable since it will show up in the Web Remote Control page.
+	
+	FString InstancePackageName = TEXT("/Temp/");	// Keep it short for web page.
+
+	// Using channel name, since we should have one instance per channel.
+	InstancePackageName += InChannelName.ToString();
+
+	// In order to keep things short, we remove "/Game" since we added the channel name instead.
+	FString InstanceSubPath;
+	if (!InSourceAssetPath.GetLongPackageName().Split(TEXT("/Game"), nullptr, &InstanceSubPath))
+	{
+		InstanceSubPath = InSourceAssetPath.GetLongPackageName();
+	}
+
+	// This may happen if the original asset path is not specified.
+	if (InstanceSubPath.IsEmpty())
+	{
+		// Add something to get a valid path at least.
+		InstanceSubPath = TEXT("InvalidAssetName");
+	}
+	
+	InstancePackageName += InstanceSubPath;
+
+	return MakeInstancePackage(InstancePackageName);
+}
+
+UPackage* UAvaMediaPlayableGroup::MakeSharedInstancePackage(const FName& InChannelName)
+{
+	// Remote Control Preset will be registered with the package name.
+	// We want a package name unique to the game instance and that will be
+	// human readable since it will show up in the Web Remote Control page.
+	
+	FString SharedPackageName = TEXT("/Temp/");	// Keep it short for web page.
+
+	// Using channel name, since we should have one instance per channel.
+	SharedPackageName += InChannelName.ToString();
+
+	// Shared for all levels.
+	SharedPackageName += TEXT("/SharedLevels");
+	
+	return MakeInstancePackage(SharedPackageName);
+}
+
+UPackage* UAvaMediaPlayableGroup::MakeInstancePackage(const FString& InInstancePackageName)
+{
+	UPackage* InstancePackage = CreatePackage(*InInstancePackageName);
+	if (InstancePackage)
+	{
+		InstancePackage->SetFlags(RF_Transient);
+	}
+	else
+	{
+		// Note: The outer will fallback to GEngine in that case.
+		UE_LOG(LogAvalanchePlayable, Error, TEXT("Unable to create package \"%s\" for Avalanche Game Instance."), *InInstancePackageName);
+	}
+	return InstancePackage;
+}
+
+UAvaMediaPlayableGroup* UAvaMediaPlayableGroup::MakePlayableGroup(UObject* InOuter, const FPlayableGroupCreationInfo& InPlayableGroupInfo)
+{
+	UObject* Outer = InOuter ? InOuter : GetTransientPackage();
+	
+	UAvaMediaPlayableGroup* GameInstanceGroup;
+	if (InPlayableGroupInfo.bIsRemoteProxy)
+	{
+		// Remote Proxy group doesn't have a game instance.
+		GameInstanceGroup = NewObject<UAvaMediaRemoteProxyPlayableGroup>(Outer);
+		GameInstanceGroup->ParentPlayableGroupManagerWeak = InPlayableGroupInfo.PlayableGroupManager;
+	}
+	else
+	{
+		GameInstanceGroup = NewObject<UAvaMediaPlayableGroup>(Outer);
+		GameInstanceGroup->ParentPlayableGroupManagerWeak = InPlayableGroupInfo.PlayableGroupManager;
+		
+		if (InPlayableGroupInfo.bIsSharedGroup)
+		{
+			GameInstanceGroup->GameInstancePackage = MakeSharedInstancePackage(InPlayableGroupInfo.ChannelName);
+		}
+		else
+		{
+			// We can create the package even if the name is null, it will have a generic name. But that will be considered an error.
+			if (!InPlayableGroupInfo.SourceAssetPath.IsNull())
+			{
+				UE_LOG(LogAvalanchePlayable, Error, TEXT("Creating game instance package for asset with unspecified name."));
+			}
+			GameInstanceGroup->GameInstancePackage = MakeGameInstancePackage(InPlayableGroupInfo.SourceAssetPath, InPlayableGroupInfo.ChannelName);
+		}
+		
+		GameInstanceGroup->GameInstance = UAvalancheGameInstance::Create(GameInstanceGroup->GameInstancePackage);
+	}
+	return GameInstanceGroup;
+}
+
+void UAvaMediaPlayableGroup::RegisterPlayable(UAvalanchePlayable* InPlayable)
+{
+	// Prevent accumulation of stale keys.
+	UE::AvalancheMedia::PlayableGroup::Private::CleanStaleKeys(Playables);
+
+	if (InPlayable)
+	{
+		InPlayable->OnPlayableStatusChanged().AddUObject(this, &UAvaMediaPlayableGroup::OnPlayableStatusChanged);
+		Playables.Add(InPlayable);
+	}
+}
+
+/** Unregister a playable when it is about to be deleted. */
+void UAvaMediaPlayableGroup::UnregisterPlayable(UAvalanchePlayable* InPlayable)
+{
+	if (InPlayable)
+	{
+		InPlayable->OnPlayableStatusChanged().RemoveAll(this);
+		Playables.Remove(InPlayable);
+	}
+}
+
+bool UAvaMediaPlayableGroup::HasPlayables() const
+{
+	for (const TObjectKey<UAvalanchePlayable>& PlayableKey : Playables)
+	{
+		if (PlayableKey.ResolveObjectPtr())
+		{
+			return true;
+		}
+	}
+	return false;
+}
+
+bool UAvaMediaPlayableGroup::HasPlayingPlayables() const
+{
+	for (const TObjectKey<UAvalanchePlayable>& PlayableKey : Playables)
+	{
+		const UAvalanchePlayable* Playable = PlayableKey.ResolveObjectPtr();
+		if (Playable && Playable->IsPlaying())
+		{
+			return true;
+		}
+	}
+	return false;
+}
+
+UAvalanchePlayable* UAvaMediaPlayableGroup::FindFirstPlayableBySourceAssetPath(const FSoftObjectPath& InSourceAssetPath) const
+{
+	for (const TObjectKey<UAvalanchePlayable>& PlayableKey : Playables)
+	{
+		UAvalanchePlayable* Playable = PlayableKey.ResolveObjectPtr();
+		if (Playable && Playable->GetSourceAssetPath() == InSourceAssetPath)
+		{
+			return Playable;
+		}
+	}
+	return nullptr;
+}
+
+void UAvaMediaPlayableGroup::RegisterPlayableTransition(UAvalanchePlayableTransition* InPlayableTransition)
+{
+	if (InPlayableTransition)
+	{
+		// Protect PlayableTransitions iterator.
+		if (bIsTickingTransitions)
+		{
+			PlayableTransitionsToRemove.Remove(InPlayableTransition);
+			PlayableTransitionsToAdd.Add(InPlayableTransition);
+			return;
+		}
+
+		// Prevent accumulation of stale keys.
+		UE::AvalancheMedia::PlayableGroup::Private::CleanStaleKeys(PlayableTransitions);
+
+		PlayableTransitions.Add(InPlayableTransition);
+		
+		if (UAvaMediaPlayableGroupManager* PlayableGroupManager = GetPlayableGroupManager())
+		{
+			PlayableGroupManager->RegisterForTransitionTicking(this);
+		}
+	}
+}
+	
+void UAvaMediaPlayableGroup::UnregisterPlayableTransition(UAvalanchePlayableTransition* InPlayableTransition)
+{
+	if (InPlayableTransition)
+	{
+		// Protect PlayableTransitions iterator.
+		if (bIsTickingTransitions)
+		{
+			PlayableTransitionsToAdd.Remove(InPlayableTransition);
+			PlayableTransitionsToRemove.Add(InPlayableTransition);
+			return;
+		}
+		
+		PlayableTransitions.Remove(InPlayableTransition);
+
+		UE::AvalancheMedia::PlayableGroup::Private::CleanStaleKeys(PlayableTransitions);
+		
+		if (PlayableTransitions.IsEmpty())
+		{
+			if (UAvaMediaPlayableGroupManager* PlayableGroupManager = GetPlayableGroupManager())
+			{
+				PlayableGroupManager->UnregisterFromTransitionTicking(this);
+			}
+		}
+	}
+}
+
+void UAvaMediaPlayableGroup::TickTransitions(double InDeltaSeconds)
+{
+	// Ticking the transitions will lead to some transitions being removed.
+	// So we need to protect the iterator with a scope and do the operations when iteration is done.
+	{
+		TGuardValue TickGuard(bIsTickingTransitions, true);
+	
+		for (TSet<TObjectKey<UAvalanchePlayableTransition>>::TIterator TickIt(PlayableTransitions); TickIt; ++TickIt)
+		{
+			if (UAvalanchePlayableTransition* TransitionToTick = TickIt->ResolveObjectPtr())
+			{
+				TransitionToTick->Tick(InDeltaSeconds);
+			}
+			else
+			{
+				// Prevent accumulation of stale keys.
+				// Remark: if there is a stale key, it means the transition was not stopped properly.
+				TickIt.RemoveCurrent();
+			}
+		}
+	}
+	
+	for (const TObjectKey<UAvalanchePlayableTransition>& ToRemove : PlayableTransitionsToRemove)
+	{
+		UnregisterPlayableTransition(ToRemove.ResolveObjectPtr());
+	}
+	PlayableTransitionsToRemove.Reset();
+	
+	for (const TObjectKey<UAvalanchePlayableTransition>& ToAdd : PlayableTransitionsToAdd)
+	{
+		RegisterPlayableTransition(ToAdd.ResolveObjectPtr());
+	}
+	PlayableTransitionsToAdd.Reset();
+}
+
+bool UAvaMediaPlayableGroup::HasTransitions() const
+{
+	return !PlayableTransitions.IsEmpty();
+}
+
+bool UAvaMediaPlayableGroup::ConditionalCreateWorld()
+{
+	if (!GameInstance)
+	{
+		return false;
+	}
+
+	bool bWorldWasCreated = false;
+
+	if (!GameInstance->IsWorldCreated())
+	{
+		bWorldWasCreated = GameInstance->CreateWorld();
+	}
+	
+	// Make sure we register our delegates to this world.
+	if (GameInstance->GetPlayWorld())
+	{
+		ConditionalRegisterWorldDelegates(GameInstance->GetPlayWorld());
+	}
+	
+	return bWorldWasCreated;
+}
+
+bool UAvaMediaPlayableGroup::ConditionalBeginPlay(const FAvalancheInstancePlaySettings& InWorldPlaySettings)
+{
+	bool bHasBegunPlay = false;
+	if (!GameInstance)
+	{
+		return bHasBegunPlay;
+	}
+	
+	// Make sure we don't have pending unload or stop requests left over in the game instance.
+	GameInstance->CancelWorldRequests();
+
+	if (!GameInstance->IsWorldPlaying())
+	{
+		bHasBegunPlay = GameInstance->BeginPlayWorld(InWorldPlaySettings);
+	}
+	else
+	{
+		GameInstance->UpdateRenderTarget(InWorldPlaySettings.RenderTarget);
+		GameInstance->UpdateSceneViewportSize(InWorldPlaySettings.ViewportSize);
+	}
+	return bHasBegunPlay;
+}
+
+void UAvaMediaPlayableGroup::RequestEndPlayWorld(bool bInForceImmediate)
+{
+	if (GameInstance)
+	{
+		GameInstance->RequestEndPlayWorld(bInForceImmediate);
+	}
+}
+
+void UAvaMediaPlayableGroup::SetLastAppliedCameraPlayable(UAvalanchePlayable* InPlayable)
+{
+	LastAppliedCameraPlayableWeak = InPlayable;
+}
+
+bool UAvaMediaPlayableGroup::UpdateCameraSetup()
+{
+	// With rigs in sub-playables, it may still be valid and won't need updating.
+	if (const UAvalanchePlayable* LastAppliedCameraPlayable = LastAppliedCameraPlayableWeak.Get())
+	{
+		if (LastAppliedCameraPlayable->GetPlayableStatus() == EAvalanchePlayableStatus::Visible
+			&& LastAppliedCameraPlayable->GetShouldBeVisible())
+		{
+			return true;	// Camera setup is still valid.
+		}
+	}
+
+	// This is called when a playable from the group is stopped. We need to
+	// select a playable in the group that is still playing and will use it's camera.
+	for (const TObjectKey<UAvalanchePlayable>& PlayableKey : Playables )
+	{
+		UAvalanchePlayable* Playable = PlayableKey.ResolveObjectPtr();
+		
+		if (Playable && Playable->IsPlaying())
+		{
+			if (Playable->ApplyCamera())
+			{
+				return true;
+			}
+		}
+	}
+	return false;
+}
+
+bool UAvaMediaPlayableGroup::ConditionalRequestUnloadWorld(bool bForceImmediate)
+{
+	if (!GameInstance)
+	{
+		return false;
+	}
+	
+	if (!HasPlayables())
+	{
+		UnregisterWorldDelegates(GameInstance->GetPlayWorld());
+		GameInstance->RequestUnloadWorld(bForceImmediate);
+		return true;
+	}
+	return false;
+}
+
+void UAvaMediaPlayableGroup::QueueCameraCut()
+{
+	if (GameInstance)
+	{
+		if (UAvalancheGameViewportClient* GameViewportClient = GameInstance->GetAvalancheGameViewportClient())
+		{
+			GameViewportClient->SetCameraCutThisFrame();
+		}
+	}
+}
+
+void UAvaMediaPlayableGroup::NotifyLevelStreaming(UAvalanchePlayable* InPlayable)
+{
+	// If the world is not playing, we need to make sure the level streaming still updates.
+	if (!IsWorldPlaying())
+	{
+		if (UAvaMediaPlayableGroupManager* PlayableGroupManager = GetPlayableGroupManager())
+		{
+			PlayableGroupManager->RegisterForLevelStreamingUpdate(this);
+		}
+	}
+}
+
+void UAvaMediaPlayableGroup::FVisibilityRequest::Execute(const UAvaMediaPlayableGroup* InPlayableGroup) const
+{
+	UAvalanchePlayable* Playable = PlayableWeak.Get();
+	
+	if (!Playable)
+	{
+		using namespace UE::AvaMediaPlayback::Utils;
+		UE_LOG(LogAvalanchePlayable, Error,
+			TEXT("%s Failed to Set Visibility to \"%s\" because the playable has become stale. Playable Group: \"%s\"."),
+			*GetBriefFrameInfo(), bShouldBeVisible ? TEXT("true") : TEXT("false"), *InPlayableGroup->GetFullName());
+	}
+
+	Playable->SetShouldBeVisible(bShouldBeVisible);
+}
+
+void UAvaMediaPlayableGroup::RegisterVisibilityConstraint(const TWeakInterfacePtr<IAvaPlayableVisibilityConstraint>& InVisibilityConstraint)
+{
+	if (!VisibilityConstraints.Contains(InVisibilityConstraint))
+	{
+		VisibilityConstraints.RemoveAll([](const TWeakInterfacePtr<IAvaPlayableVisibilityConstraint>& InConstraint) { return InConstraint.IsStale();});
+		VisibilityConstraints.Add(InVisibilityConstraint);
+	}
+}
+
+void UAvaMediaPlayableGroup::UnregisterVisibilityConstraint(const IAvaPlayableVisibilityConstraint* InVisibilityConstraint)
+{
+	VisibilityConstraints.RemoveAll([InVisibilityConstraint](const TWeakInterfacePtr<IAvaPlayableVisibilityConstraint>& InConstraint)
+	{
+		// also remove stale pointers.
+		return InConstraint.IsStale() || InConstraint.Get() == InVisibilityConstraint;
+	});
+}
+
+void UAvaMediaPlayableGroup::RequestSetVisibility(UAvalanchePlayable* InPlayable, bool bInShouldBeVisible)
+{
+	FVisibilityRequest Request = { InPlayable, bInShouldBeVisible };
+
+	if (IsVisibilityConstrained(InPlayable))
+	{
+		VisibilityRequests.Add(MoveTemp(Request));
+	}
+	else
+	{
+		Request.Execute(this);
+	}
+}
+
+bool UAvaMediaPlayableGroup::IsVisibilityConstrained(const UAvalanchePlayable* InPlayable) const
+{
+	for (const TWeakInterfacePtr<IAvaPlayableVisibilityConstraint>& ConstraintWeak : VisibilityConstraints)
+	{
+		if (const IAvaPlayableVisibilityConstraint* Constraint = ConstraintWeak.Get())
+		{
+			if (Constraint->IsVisibilityConstrained(InPlayable))
+			{
+				return true;
+			}
+		}
+	}
+	return false;
+}
+
+void UAvaMediaPlayableGroup::OnPlayableStatusChanged(UAvalanchePlayable* InPlayable)
+{
+	// Evaluate the playable visibility requests
+	for (TArray<FVisibilityRequest>::TIterator RequestIt(VisibilityRequests); RequestIt; ++RequestIt)
+	{
+		const UAvalanchePlayable* Playable = RequestIt->PlayableWeak.Get();
+		if (Playable && IsVisibilityConstrained(Playable))
+		{
+			continue;
+		}
+		
+		RequestIt->Execute(this);
+		RequestIt.RemoveCurrent();
+	}
+}
+
+void UAvaMediaPlayableGroup::ConditionalRegisterWorldDelegates(UWorld* InWorld)
+{
+	if (!DisplayDelegateIndices.IsEmpty() && LastWorldBoundToDisplayDelegates.Get() == InWorld)
+	{
+		return;
+	}
+
+	if (!DisplayDelegateIndices.IsEmpty() && LastWorldBoundToDisplayDelegates.IsValid())
+	{
+		UnregisterWorldDelegates(LastWorldBoundToDisplayDelegates.Get());
+	}
+
+	if (UViewportStatsSubsystem* ViewportSubsystem = InWorld->GetSubsystem<UViewportStatsSubsystem>())
+	{
+		DisplayDelegateIndices.Add(ViewportSubsystem->AddDisplayDelegate([this](FText& OutText, FLinearColor& OutColor)
+		{
+			return DisplayLoadedAssets(OutText, OutColor);
+		}));
+		DisplayDelegateIndices.Add(ViewportSubsystem->AddDisplayDelegate([this](FText& OutText, FLinearColor& OutColor)
+		{
+			return DisplayPlayingAssets(OutText, OutColor);
+		}));
+		DisplayDelegateIndices.Add(ViewportSubsystem->AddDisplayDelegate([this](FText& OutText, FLinearColor& OutColor)
+		{
+			return DisplayTransitions(OutText, OutColor);
+		}));
+
+		LastWorldBoundToDisplayDelegates = InWorld;
+	}
+}
+
+void UAvaMediaPlayableGroup::UnregisterWorldDelegates(UWorld* InWorld)
+{
+	if (!DisplayDelegateIndices.IsEmpty())
+	{
+		// Check that we are indeed registered to the world
+		check(LastWorldBoundToDisplayDelegates.Get() == InWorld);
+	
+		if (UViewportStatsSubsystem* ViewportSubsystem = InWorld->GetSubsystem<UViewportStatsSubsystem>())
+		{
+			// Remark: removing them in the reverse order they where added. It is using RemoveAtSwap()
+			// so removing the first one would change the index of the next one. Removing a delegate in the
+			// middle of the array will invalidate the indices above.
+			for (int32 Index = DisplayDelegateIndices.Num()-1; Index >= 0; --Index)
+			{
+				ViewportSubsystem->RemoveDisplayDelegate(DisplayDelegateIndices[Index]);
+			}
+		}
+
+		DisplayDelegateIndices.Reset();
+		LastWorldBoundToDisplayDelegates.Reset();
+	}
+
+	check(LastWorldBoundToDisplayDelegates.Get() == nullptr);
+}
+
+bool UAvaMediaPlayableGroup::DisplayLoadedAssets(FText& OutText, FLinearColor& OutColor)
+{
+	FString AssetList;
+	for (const TObjectKey<UAvalanchePlayable>& PlayableKey : Playables)
+	{
+		if (const UAvalanchePlayable* Playable = PlayableKey.ResolveObjectPtr())
+		{
+			AssetList += AssetList.IsEmpty() ? TEXT("") : TEXT(", ");
+			AssetList += Playable->GetSourceAssetPath().GetAssetName();
+		}
+	}
+
+	if (!AssetList.IsEmpty())
+	{
+		OutText = FText::Format(LOCTEXT("DisplayLoadedGraphics", "Loaded Graphic(s): {0}"), FText::FromString(AssetList));
+		OutColor = FLinearColor::Red;
+		return true;
+	}
+	return false;
+}
+
+bool UAvaMediaPlayableGroup::DisplayPlayingAssets(FText& OutText, FLinearColor& OutColor)
+{
+	FString AssetList;
+	for (const TObjectKey<UAvalanchePlayable>& PlayableKey : Playables)
+	{
+		const UAvalanchePlayable* Playable = PlayableKey.ResolveObjectPtr();
+		if (Playable && Playable->IsPlaying())
+		{
+			AssetList += AssetList.IsEmpty() ? TEXT("") : TEXT(", ");
+			AssetList += Playable->GetSourceAssetPath().GetAssetName();
+
+			if (!Playable->GetUserData().IsEmpty())
+			{
+				AssetList += FString::Printf(TEXT(" (%s)"), *Playable->GetUserData());
+			}
+		}
+	}
+
+	if (!AssetList.IsEmpty())
+	{
+		OutText = FText::Format(LOCTEXT("DisplayPlayingGraphics", "Playing Graphic(s): {0}"), FText::FromString(AssetList));
+		OutColor = FLinearColor::Green;
+		return true;
+	}
+	return false;
+}
+
+bool UAvaMediaPlayableGroup::DisplayTransitions(FText& OutText, FLinearColor& OutColor)
+{
+	FString TransitionList;
+	for (const TObjectKey<UAvalanchePlayableTransition>& TransitionKey : PlayableTransitions)
+	{
+		const UAvalanchePlayableTransition* Transition = TransitionKey.ResolveObjectPtr();
+		if (Transition && Transition->IsRunning())
+		{
+			TransitionList += TransitionList.IsEmpty() ? TEXT("") : TEXT(", ");
+			TransitionList += Transition->GetPrettyInfo();
+		}
+	}
+
+	if (!TransitionList.IsEmpty())
+	{
+		OutText = FText::Format(LOCTEXT("DisplayTransitions", "Transition(s): {0}"), FText::FromString(TransitionList));
+		OutColor = FLinearColor::Green;
+		return true;
+	}
+	return false;
+}
+
+bool UAvaMediaRemoteProxyPlayableGroup::ConditionalBeginPlay(const FAvalancheInstancePlaySettings& InWorldPlaySettings)
+{
+	if (!bIsPlaying)
+	{
+		bIsPlaying = true;
+		return true;
+	}
+	return false;
+}
+
+void UAvaMediaRemoteProxyPlayableGroup::RequestEndPlayWorld(bool bInForceImmediate)
+{
+	bIsPlaying = false;
+}
+
+#undef LOCTEXT_NAMESPACE
\ No newline at end of file
diff --git a/Engine/Plugins/Experimental/Avalanche/Source/AvalancheMedia/Private/Playback/AvaMediaPlayableGroupManager.cpp b/Engine/Plugins/Experimental/Avalanche/Source/AvalancheMedia/Private/Playback/AvaMediaPlayableGroupManager.cpp
new file mode 100644
--- /dev/null
+++ b/Engine/Plugins/Experimental/Avalanche/Source/AvalancheMedia/Private/Playback/AvaMediaPlayableGroupManager.cpp
@@ -0,0 +1,225 @@
+// Copyright Epic Games, Inc. All Rights Reserved.
+
+#include "Playback/AvaMediaPlayableGroupManager.h"
+
+#include "AvalancheBroadcast.h"
+#include "Framework/AvalancheGameInstance.h"
+#include "Misc/CoreDelegates.h"
+#include "Misc/TimeGuard.h"
+#include "OutputDevices/AvaRenderTargetMediaUtils.h"
+#include "Playback/AvaMediaPlayableGroup.h"
+#include "UObject/Package.h"
+
+#define LOCTEXT_NAMESPACE "AvaMediaPlayableGroupManager"
+
+UAvaMediaPlayableGroup* UAvaMediaChannelPlayableGroupManager::GetOrCreateSharedLevelGroup(bool bInIsRemoteProxy)
+{
+	if (UAvaMediaPlayableGroup* ExistingPlayableGroup = bInIsRemoteProxy ? SharedRemoteProxyLevelGroupWeak.Get() : SharedLevelGroupWeak.Get())
+	{
+		return ExistingPlayableGroup;
+	}
+	
+	UAvaMediaPlayableGroup::FPlayableGroupCreationInfo PlayableGroupCreationInfo;
+	PlayableGroupCreationInfo.PlayableGroupManager = GetPlayableGroupManager();
+	PlayableGroupCreationInfo.ChannelName = ChannelName;
+	PlayableGroupCreationInfo.bIsRemoteProxy = bInIsRemoteProxy;
+	PlayableGroupCreationInfo.bIsSharedGroup = true;
+
+	UAvaMediaPlayableGroup* NewPlayableGroup = UAvaMediaPlayableGroup::MakePlayableGroup(GetPlayableGroupManager(), PlayableGroupCreationInfo);
+
+	// Keep track of the shared playable group.
+	if (bInIsRemoteProxy)
+	{
+		SharedRemoteProxyLevelGroupWeak = NewPlayableGroup;
+	}
+	else
+	{
+		SharedLevelGroupWeak = NewPlayableGroup;
+	}
+	
+	return NewPlayableGroup;
+}
+
+UAvaMediaPlayableGroupManager* UAvaMediaChannelPlayableGroupManager::GetPlayableGroupManager() const
+{
+	return Cast<UAvaMediaPlayableGroupManager>(GetOuter());	
+}
+
+void UAvaMediaChannelPlayableGroupManager::Shutdown()
+{
+	SharedRemoteProxyLevelGroupWeak.Reset();
+	SharedLevelGroupWeak.Reset();
+}
+
+void UAvaMediaChannelPlayableGroupManager::BeginDestroy()
+{
+	Shutdown();
+	Super::BeginDestroy();
+}
+
+void UAvaMediaPlayableGroupManager::Init()
+{
+	if (!UAvalancheGameInstance::GetOnEndPlay().IsBoundToObject(this))
+	{
+		UAvalancheGameInstance::GetOnEndPlay().AddUObject(this, &UAvaMediaPlayableGroupManager::OnGameInstanceEndPlay);
+	}
+}
+
+void UAvaMediaPlayableGroupManager::Shutdown()
+{
+	for (TPair<FName, TObjectPtr<UAvaMediaChannelPlayableGroupManager>>& Pair : ChannelManagers)
+	{
+		if (UAvaMediaChannelPlayableGroupManager* const ChannelManager = Pair.Value)
+		{
+			ChannelManager->Shutdown();
+		}
+	}
+	UAvalancheGameInstance::GetOnEndPlay().RemoveAll(this);
+}
+
+void UAvaMediaPlayableGroupManager::Tick(double InDeltaSeconds)
+{
+	SCOPE_TIME_GUARD(TEXT("UAvaMediaPlayableGroupManager::Tick"));
+	UpdateLevelStreaming();
+	TickTransitions(InDeltaSeconds);
+}
+
+UAvaMediaChannelPlayableGroupManager* UAvaMediaPlayableGroupManager::FindOrAddChannelManager(const FName& InChannelName)
+{
+	UAvaMediaChannelPlayableGroupManager* ChannelManager = FindChannelManager(InChannelName);
+	if (!ChannelManager)
+	{
+		ChannelManager = NewObject<UAvaMediaChannelPlayableGroupManager>(this);
+		ChannelManager->ChannelName = InChannelName;
+		ChannelManagers.Add(InChannelName, ChannelManager);
+	}
+	return ChannelManager;
+}
+
+void UAvaMediaPlayableGroupManager::RegisterForLevelStreamingUpdate(UAvaMediaPlayableGroup* InPlayableGroup)
+{
+	if (ensure(!bIsUpdatingStreaming))
+	{
+		GroupsToUpdateStreaming.Add(InPlayableGroup);
+	}
+}
+
+void UAvaMediaPlayableGroupManager::UnregisterFromLevelStreamingUpdate(UAvaMediaPlayableGroup* InPlayableGroup)
+{
+	if (!bIsUpdatingStreaming)
+	{
+		GroupsToUpdateStreaming.Remove(InPlayableGroup);
+	}
+}
+
+void UAvaMediaPlayableGroupManager::RegisterForTransitionTicking(UAvaMediaPlayableGroup* InPlayableGroup)
+{
+	if (ensure(!bIsTickingTransitions))
+	{
+		GroupsToTickTransitions.Add(InPlayableGroup);
+	}
+}
+
+void UAvaMediaPlayableGroupManager::UnregisterFromTransitionTicking(UAvaMediaPlayableGroup* InPlayableGroup)
+{
+	if (!bIsTickingTransitions)
+	{
+		GroupsToTickTransitions.Remove(InPlayableGroup);
+	}
+}
+
+void UAvaMediaPlayableGroupManager::BeginDestroy()
+{
+	Shutdown();
+	Super::BeginDestroy();
+}
+
+void UAvaMediaPlayableGroupManager::OnGameInstanceEndPlay(UAvalancheGameInstance* InGameInstance, FName InChannelName)
+{
+	FAvaOutputChannel& Channel = UAvalancheBroadcast::Get().GetCurrentProfile().GetChannelMutable(InChannelName);
+	if (!Channel.IsValidChannel())
+	{
+		return;
+	}
+
+	// In the current design, there is only one playable group active on a channel, i.e.
+	// the graph of playables and groups must end on a root "playable group".
+	// This is resolved by the playback graph(s) running on that channel.
+	//
+	// When we receive an EndPlay event from a game instance for a given channel, we need to
+	// check if that is the current root playable group, and if it is, clear out the
+	// channels association with it.
+	//
+	// If the channel is live, it will then start rendering the placeholder graphic, if
+	// configured to do so, on the next slate tick.
+	
+	if (const UAvaMediaPlayableGroup* PlayableGroup = Channel.GetLastActivePlayableGroup())
+	{
+		if (PlayableGroup->GetGameInstance() == InGameInstance)
+		{
+			Channel.UpdateRenderTarget(nullptr, nullptr);
+			Channel.UpdateAudioDevice(FAudioDeviceHandle());
+		}
+	}
+
+	// If the channel is not live, the placeholder graphic doesn't render
+	// so we need to explicitly clear the channel here.
+	if (Channel.GetState() != EAvaChannelState::Live)
+	{
+		if (UTextureRenderTarget2D* const RenderTarget = Channel.GetCurrentRenderTarget(true))
+		{
+			UE::AvaRenderTargetMediaUtils::ClearRenderTarget(RenderTarget);
+		}
+	}
+}
+
+void UAvaMediaPlayableGroupManager::UpdateLevelStreaming()
+{
+	TGuardValue TransitionsTickGuard(bIsUpdatingStreaming, true);		
+	for (TSet<TWeakObjectPtr<UAvaMediaPlayableGroup>>::TIterator GroupIterator(GroupsToUpdateStreaming); GroupIterator; ++GroupIterator)
+	{
+		const UAvaMediaPlayableGroup* GroupToUpdate = GroupIterator->Get();
+		// We only update streaming if the group is not playing.
+		if (GroupToUpdate && !GroupToUpdate->IsWorldPlaying())
+		{
+			// World may not be created yet.
+			if (UWorld* PlayWorld = GroupToUpdate->GetPlayWorld())
+			{
+				// (This is normally updated by the game viewport client when the world is playing.)
+				PlayWorld->UpdateLevelStreaming();
+
+				// Check if still has streaming. If not, from the list.
+				if (!PlayWorld->HasStreamingLevelsToConsider())
+				{
+					GroupIterator.RemoveCurrent();
+				}
+			}
+		}
+		else
+		{
+			GroupIterator.RemoveCurrent();
+		}
+	}
+}
+
+void UAvaMediaPlayableGroupManager::TickTransitions(double InDeltaSeconds)
+{
+	TGuardValue TransitionsTickGuard(bIsTickingTransitions, true);
+	for (TSet<TWeakObjectPtr<UAvaMediaPlayableGroup>>::TIterator GroupIterator(GroupsToTickTransitions); GroupIterator; ++GroupIterator)
+	{
+		bool bHasTransitions = false;
+		if (UAvaMediaPlayableGroup* GroupToTick = GroupIterator->Get())
+		{
+			GroupToTick->TickTransitions(InDeltaSeconds);
+			bHasTransitions = GroupToTick->HasTransitions();
+		}
+
+		// Group automatically deregister if stale or they don't have active transitions.
+		if (!bHasTransitions)
+		{
+			GroupIterator.RemoveCurrent();
+		}
+	}
+}
+
+#undef LOCTEXT_NAMESPACE
\ No newline at end of file
