[MotionDesign]

diff --git a/Engine/Plugins/Experimental/Avalanche/Source/AvalancheMedia/Private/Playback/AvaMediaPlaybackClient.cpp b/Engine/Plugins/Experimental/Avalanche/Source/AvalancheMedia/Private/Playback/AvaMediaPlaybackClient.cpp
new file mode 100644
--- /dev/null
+++ b/Engine/Plugins/Experimental/Avalanche/Source/AvalancheMedia/Private/Playback/AvaMediaPlaybackClient.cpp
@@ -0,0 +1,1935 @@
+// Copyright Epic Games, Inc. All Rights Reserved.
+
+#include "Playback/AvaMediaPlaybackClient.h"
+
+#include "AssetRegistry/AssetData.h"
+#include "AssetRegistry/IAssetRegistry.h"
+#include "AvaMediaMessageUtils.h"
+#include "AvaMediaModule.h"
+#include "AvalancheBroadcast.h"
+#include "AvalancheMediaSettings.h"
+#include "IAvaModule.h"
+#include "IMediaIOCoreModule.h"
+#include "MediaOutput.h"
+#include "MessageEndpointBuilder.h"
+#include "Misc/CoreDelegates.h"
+#include "Misc/Paths.h"
+#include "OutputDevices/AvaMediaOutputUtils.h"
+#include "UObject/Package.h"
+
+DEFINE_LOG_CATEGORY_STATIC(LogAvaPlaybackClient, Log, All);
+
+namespace UE::AvaMediaPlaybackClient::Private
+{
+	template<typename InEnumType>
+	FString EnumToString(InEnumType InValue)
+	{
+		return StaticEnum<InEnumType>()->GetNameStringByValue(static_cast<int64>(InValue));
+	}
+}
+
+FAvaMediaPlaybackClient::FAvaMediaPlaybackClient(FAvaMediaModule* InParentModule)
+	: ParentModule(InParentModule)
+{
+	UPackage::PreSavePackageWithContextEvent.AddRaw(this, &FAvaMediaPlaybackClient::OnPreSavePackage);
+	UPackage::PackageSavedWithContextEvent.AddRaw(this, &FAvaMediaPlaybackClient::OnPackageSaved);
+	if (IAssetRegistry* AssetRegistry = IAssetRegistry::Get())
+	{
+		AssetRegistry->OnAssetRemoved().AddRaw(this, &FAvaMediaPlaybackClient::OnAssetRemoved);
+	}
+}
+
+FAvaMediaPlaybackClient::~FAvaMediaPlaybackClient()
+{
+	if (BroadcastChangedHandle.IsValid())
+	{
+		UAvalancheBroadcast::Get().RemoveChangeListener(BroadcastChangedHandle);
+	}
+
+	FAvaOutputChannel::GetOnChannelChanged().RemoveAll(this);
+	
+	UPackage::PreSavePackageWithContextEvent.RemoveAll(this);
+	UPackage::PackageSavedWithContextEvent.RemoveAll(this);
+	if (IAssetRegistry* AssetRegistry = IAssetRegistry::Get())
+	{
+		AssetRegistry->OnAssetRemoved().RemoveAll(this);
+	}
+	
+	FTSTicker::GetCoreTicker().RemoveTicker(PingTickDelegateHandle);
+	FCoreDelegates::OnEndFrame.RemoveAll(this);
+	FMessageEndpoint::SafeRelease(MessageEndpoint);
+
+	for (IConsoleObject* ConsoleCmd : ConsoleCommands)
+	{
+		IConsoleManager::Get().UnregisterConsoleObject(ConsoleCmd);
+	}
+	ConsoleCommands.Empty();
+
+#if WITH_EDITOR
+	if (UObjectInitialized())
+	{
+		UAvalancheMediaSettings& AvaMediaSettings = UAvalancheMediaSettings::GetMutable();
+		AvaMediaSettings.OnSettingChanged().RemoveAll(this);
+	}
+#endif
+}
+
+void FAvaMediaPlaybackClient::Init()
+{
+	ComputerName = FPlatformProcess::ComputerName();
+	ProjectContentPath = FPaths::ConvertRelativePathToFull(FPaths::ProjectContentDir());
+	ProcessId = FPlatformProcess::GetCurrentProcessId();
+	
+	RegisterCommands();
+
+	// Create our end point.
+	MessageEndpoint = FMessageEndpoint::Builder("AvaMediaPlaybackClient")
+		.Handling<FAvaMediaPlaybackPong>(this, &FAvaMediaPlaybackClient::HandlePlaybackPongMessage)
+		.Handling<FAvaMediaPlaybackLog>(this, &FAvaMediaPlaybackClient::HandlePlaybackLogMessage)
+		.Handling<FAvaMediaUpdateServerUserData>(this, &FAvaMediaPlaybackClient::HandleUpdateServerUserData)
+		.Handling<FAvaMediaStatStatus>(this, &FAvaMediaPlaybackClient::HandleStatStatus)
+		.Handling<FAvaDeviceProviderDataList>(this, &FAvaMediaPlaybackClient::HandleDeviceProviderDataListMessage)
+		.Handling<FAvaMediaBroadcastStatus>(this, &FAvaMediaPlaybackClient::HandleBroadcastStatusMessage)
+		.Handling<FAvaMediaPlaybackAssetStatus>(this, &FAvaMediaPlaybackClient::HandlePlaybackAssetStatusMessage)
+		.Handling<FAvaMediaPlaybackStatus>(this, &FAvaMediaPlaybackClient::HandlePlaybackStatusMessage)
+		.Handling<FAvaMediaPlaybackStatuses>(this, &FAvaMediaPlaybackClient::HandlePlaybackStatusesMessage)
+		.Handling<FAvaMediaPlayableSequenceEvent>(this, &FAvaMediaPlaybackClient::HandlePlaybackSequenceEventMessage)
+		.Handling<FAvaMediaPlayableTransitionEvent>(this, &FAvaMediaPlaybackClient::HandlePlaybackTransitionEventMessage);
+	
+	if (MessageEndpoint.IsValid())
+	{
+		PingTickDelegate = FTickerDelegate::CreateRaw(this, &FAvaMediaPlaybackClient::HandlePingTicker);
+		PingTickDelegateHandle = FTSTicker::GetCoreTicker().AddTicker(PingTickDelegate, UAvalancheMediaSettings::Get().PingInterval);
+		FCoreDelegates::OnEndFrame.AddRaw(this, &FAvaMediaPlaybackClient::Tick);
+		
+		UE_LOG(LogAvaPlaybackClient, Log, TEXT("Motion Design Playback Client \"%s\" Started."), *ComputerName);
+	}
+
+	BroadcastChangedHandle = UAvalancheBroadcast::Get().AddChangeListener(
+		FOnAvaBroadcastChanged::FDelegate::CreateRaw(this, &FAvaMediaPlaybackClient::OnBroadcastChanged));
+
+	FAvaOutputChannel::GetOnChannelChanged().AddRaw(this, &FAvaMediaPlaybackClient::OnChannelChanged);
+	
+#if WITH_EDITOR
+	UAvalancheMediaSettings& AvaMediaSettings = UAvalancheMediaSettings::GetMutable();
+	AvaMediaSettings.OnSettingChanged().AddRaw(this, &FAvaMediaPlaybackClient::OnAvaMediaSettingsChanged);
+#endif
+
+	ApplyAvaMediaSettings();
+}
+
+int32 FAvaMediaPlaybackClient::GetNumConnectedServers() const
+{
+	return Servers.Num();
+}
+
+TArray<FString> FAvaMediaPlaybackClient::GetServerNames() const
+{
+	TArray<FString> ServerNames;
+	ServerNames.Empty(Servers.Num());
+	ForAllServers([&ServerNames](const FServerInfo& InServerInfo)
+	{
+		ServerNames.Add(InServerInfo.ServerName);
+	});
+	return ServerNames;
+}
+
+FMessageAddress FAvaMediaPlaybackClient::GetServerAddress(const FString& InServerName) const
+{
+	if (const FServerInfo* ServerInfo = GetServerInfo(InServerName))
+	{
+		return ServerInfo->Address;
+	}
+	FMessageAddress InvalidAddress;
+	InvalidAddress.Invalidate();
+	return InvalidAddress;
+}
+
+FString FAvaMediaPlaybackClient::GetServerProjectContentPath(const FString& InServerName) const
+{
+	if (const FServerInfo* ServerInfo = GetServerInfo(InServerName))
+	{
+		return ServerInfo->ProjectContentPath;
+	}
+	return FString();
+}
+
+uint32 FAvaMediaPlaybackClient::GetServerProcessId(const FString& InServerName) const
+{
+	if (const FServerInfo* ServerInfo = GetServerInfo(InServerName))
+	{
+		return ServerInfo->ProcessId;
+	}
+	return 0;
+}
+
+bool FAvaMediaPlaybackClient::HasServerUserData(const FString& InServerName, const FString& InKey) const
+{
+	if (const FServerInfo* ServerInfo = GetServerInfo(InServerName))
+	{
+		return ServerInfo->UserDataEntries.Contains(InKey);
+	}
+	return false;
+}
+
+const FString& FAvaMediaPlaybackClient::GetServerUserData(const FString& InServerName, const FString& InKey) const
+{
+	if (const FServerInfo* ServerInfo = GetServerInfo(InServerName))
+	{
+		if (const FString* Data = ServerInfo->UserDataEntries.Find(InKey))
+		{
+			return *Data;
+		}
+	}
+	static FString EmptyData;
+	return EmptyData;
+}
+
+bool FAvaMediaPlaybackClient::HasUserData(const FString& InKey) const
+{
+	return UserDataEntries.Contains(InKey);
+}
+
+const FString& FAvaMediaPlaybackClient::GetUserData(const FString& InKey) const
+{
+	if (const FString* Data = UserDataEntries.Find(InKey))
+	{
+		return *Data;
+	}
+	static FString EmptyString;
+	return EmptyString;
+}
+
+void FAvaMediaPlaybackClient::SetUserData(const FString& InKey, const FString& InData)
+{
+	UserDataEntries.Add(InKey, InData);
+	SendUserDataUpdate(AllServerAddresses);
+}
+
+void FAvaMediaPlaybackClient::RemoveUserData(const FString& InKey)
+{
+	UserDataEntries.Remove(InKey);
+	SendUserDataUpdate(AllServerAddresses);
+}
+
+void FAvaMediaPlaybackClient::BroadcastStatCommand(const FString& InCommand, bool bInBroadcastLocalState)
+{
+	SendStatCommand(InCommand, bInBroadcastLocalState, AllServerAddresses);
+}
+
+void FAvaMediaPlaybackClient::RequestPlaybackAssetStatus(const FSoftObjectPath& InAssetPath, const FString& InChannelOrServerName, bool bInForceRefresh)
+{
+	if (!InAssetPath.IsValid())
+	{
+		UE_LOG(LogAvaPlaybackClient, Warning, TEXT("Invalid asset path for Playback Asset Status Request."));
+		return;
+	}
+	
+	TArray<FString> ServerNames = Servers.Contains(InChannelOrServerName) ? TArray<FString>{InChannelOrServerName} : GetServerNamesForChannel(FName(InChannelOrServerName));
+	
+	for (const FString& ServerName : ServerNames)
+	{
+		RequestPlaybackAssetStatusForServer(InAssetPath, ServerName, bInForceRefresh);
+	}
+}
+
+void FAvaMediaPlaybackClient::RequestPlayback(const FGuid& InInstanceId, const FSoftObjectPath& InAssetPath, const FString& InChannelName,
+	EAvaMediaPlaybackAction InAction, const FString& InArguments)
+{
+	// Possibly TMP. Maybe find a generic way to deal with pending requests.
+	// Status requests can be spammed every frame by the UI so we block them
+	// here to avoid spamming the message bus and servers.
+	if (InAction == EAvaMediaPlaybackAction::Status)
+	{
+		const FDateTime CurrentTime = FDateTime::UtcNow();
+		const FString RequestKey = GetPlaybackStatusRequestKey(InInstanceId, InAssetPath, InChannelName);
+		const FDateTime* PendingRequest = PendingPlaybackStatusRequests.Find(RequestKey);
+		// Check if we already have a pending request and that it hasn't expired.
+		if (PendingRequest && CurrentTime < *PendingRequest)
+		{
+			return;
+		}
+
+		// Keep track of the request and the time it was made so we have a timeout.
+		const FDateTime ExpirationTime = CurrentTime + FTimespan::FromSeconds(UAvalancheMediaSettings::Get().ClientPendingStatusRequestTimeout);
+		PendingPlaybackStatusRequests.Add(RequestKey, ExpirationTime);
+	}
+
+	// Special case where the command is sent on all the assets and channels on all servers.
+	if (InChannelName.IsEmpty() || InAssetPath.IsNull())
+	{
+		FAvaMediaPlaybackRequest* Request = FMessageEndpoint::MakeMessage<FAvaMediaPlaybackRequest>();
+		Request->Commands.Reserve(1);
+		Request->Commands.Add({InInstanceId, InAssetPath, InChannelName, InAction, InArguments});
+		SendRequest(Request, AllServerAddresses);
+		return;
+	}
+	
+	if (InAssetPath.IsValid())
+	{
+		// Update the status of the playback for each server that we made a request to.
+		TArray<FString> ServerNames = GetServerNamesForChannel(FName(InChannelName));
+		for (const FString& ServerName : ServerNames)
+		{
+			if (FServerInfo* ServerInfo = GetServerInfo(ServerName))
+			{
+				// The command will be sent in a batch on the next Tick.
+				ServerInfo->PendingPlaybackCommands.Add({InInstanceId, InAssetPath, InChannelName, InAction, InArguments});
+				
+				switch (InAction)
+				{
+				case EAvaMediaPlaybackAction::None:
+				case EAvaMediaPlaybackAction::Status:
+				case EAvaMediaPlaybackAction::GetUserData:
+				case EAvaMediaPlaybackAction::SetUserData:
+					break;
+				case EAvaMediaPlaybackAction::Load:
+					ServerInfo->SetInstanceStatus(InChannelName, InInstanceId, InAssetPath, EAvaMediaPlaybackStatus::Loading);
+					break;
+				case EAvaMediaPlaybackAction::Start:
+					ServerInfo->SetInstanceStatus(InChannelName, InInstanceId, InAssetPath, EAvaMediaPlaybackStatus::Starting);
+					break;
+				case EAvaMediaPlaybackAction::Stop:
+					ServerInfo->SetInstanceStatus(InChannelName, InInstanceId, InAssetPath, EAvaMediaPlaybackStatus::Stopping);
+					break;
+				case EAvaMediaPlaybackAction::Unload:
+					ServerInfo->SetInstanceStatus(InChannelName, InInstanceId, InAssetPath, EAvaMediaPlaybackStatus::Unloading);
+					break;
+				}
+			}
+		}
+	}
+	else
+	{
+		// We will wait for feedback from server, if it managed to play the asset correctly.
+		UE_LOG(LogAvaPlaybackClient, Warning,
+			   TEXT("Playback request \"%s\" on channel \"%s\" with no asset specified. Unable to update playback status."),
+			   *UE::AvaMediaPlaybackClient::Private::EnumToString(InAction),
+			   *InChannelName);
+	}
+}
+
+void FAvaMediaPlaybackClient::RequestAnimPlayback(const FGuid& InInstanceId, const FSoftObjectPath& InAssetPath, const FString& InChannelName,
+												  const FAnimPlaySettings& InAnimSettings)
+{
+	FAvaMediaAnimPlaybackRequest* Request = FMessageEndpoint::MakeMessage<FAvaMediaAnimPlaybackRequest>();
+	Request->InstanceId = InInstanceId;
+	Request->AssetPath = InAssetPath;
+	Request->ChannelName = InChannelName;
+	Request->AnimPlaySettings.Add(InAnimSettings);
+
+	if (!InAssetPath.IsValid())
+	{
+		UE_LOG(LogAvaPlaybackClient, Warning, TEXT("Animation \"%s\" request with invalid asset path on channel \"%s\"."),
+			*InAnimSettings.AnimationName.ToString(), *InChannelName);
+	}
+		
+	SendRequest(Request, GetServerAddressesForChannel(InChannelName));
+}
+
+void FAvaMediaPlaybackClient::RequestAnimAction(const FGuid& InInstanceId, const FSoftObjectPath& InAssetPath, const FString& InChannelName,
+												  const FString& InAnimationName, EAvaMediaAnimAction InAction)
+{
+	FAvaMediaAnimPlaybackRequest* Request = FMessageEndpoint::MakeMessage<FAvaMediaAnimPlaybackRequest>();
+	Request->InstanceId = InInstanceId;
+	Request->AssetPath = InAssetPath;
+	Request->ChannelName = InChannelName;
+	Request->AnimActionInfos.Add({InAnimationName, InAction});
+
+	if (!InAssetPath.IsValid())
+	{
+		UE_LOG(LogAvaPlaybackClient, Warning, TEXT("Animation \"%s\" %s request with invalid asset path on channel \"%s\"."),
+			*InAnimationName, *StaticEnum<EAvaMediaAnimAction>()->GetValueAsString(InAction), *InChannelName);
+	}
+	
+	SendRequest(Request, GetServerAddressesForChannel(InChannelName));
+}
+
+void FAvaMediaPlaybackClient::RequestRemoteControlUpdate(const FGuid& InInstanceId, const FSoftObjectPath& InAssetPath,
+														 const FString& InChannelName,
+														 const FAvalancheRemoteControlValues& InRemoteControlValues)
+{
+	FAvaMediaRemoteControlUpdateRequest* Request = FMessageEndpoint::MakeMessage<FAvaMediaRemoteControlUpdateRequest>();
+	Request->InstanceId = InInstanceId;
+	Request->AssetPath = InAssetPath;
+	Request->ChannelName = InChannelName;
+	Request->RemoteControlValues = InRemoteControlValues;
+
+	if (!InAssetPath.IsValid())
+	{
+		UE_LOG(LogAvaPlaybackClient, Warning, TEXT("Remote Control Values Update request with invalid asset path on channel \"%s\"."), *InChannelName);
+	}
+	
+	SendRequest(Request, GetServerAddressesForChannel(InChannelName));
+}
+
+void FAvaMediaPlaybackClient::RequestPlayableTransitionStart(const FGuid& InTransitionId, TArray<FGuid>&& InEnterInstanceIds, TArray<FGuid>&& InPlayingInstanceIds, TArray<FGuid>&& InExitInstanceIds, TArray<FAvalancheRemoteControlValues>&& InEnterValues, const FName& InChannelName, EAvalanchePlayableTransitionFlags InTransitionFlags)
+{
+	FAvaMediaTransitionStartRequest* Request = FMessageEndpoint::MakeMessage<FAvaMediaTransitionStartRequest>();
+	Request->TransitionId = InTransitionId;
+	Request->ChannelName = InChannelName.ToString();
+	Request->EnterInstanceIds = MoveTemp(InEnterInstanceIds);
+	Request->PlayingInstanceIds = MoveTemp(InPlayingInstanceIds);
+	Request->ExitInstanceIds = MoveTemp(InExitInstanceIds);
+	Request->EnterValues = MoveTemp(InEnterValues);
+	Request->bUnloadDiscardedInstances = !UAvalancheMediaSettings::Get().bKeepPagesLoaded;
+	Request->SetTransitionFlags(InTransitionFlags);
+	SendRequest(Request, GetServerAddressesForChannel(InChannelName));
+}
+
+void FAvaMediaPlaybackClient::RequestPlayableTransitionStop(const FGuid& InTransitionId, const FName& InChannelName)
+{
+	FAvaMediaTransitionStopRequest* Request = FMessageEndpoint::MakeMessage<FAvaMediaTransitionStopRequest>();
+	Request->TransitionId = InTransitionId;
+	Request->ChannelName = InChannelName.ToString();
+	Request->bUnloadDiscardedInstances = !UAvalancheMediaSettings::Get().bKeepPagesLoaded;
+	SendRequest(Request, GetServerAddressesForChannel(InChannelName));
+}
+
+void FAvaMediaPlaybackClient::RequestBroadcast(const FString& InProfile, const FName& InChannel,
+											   const TArray<UMediaOutput*>& InRemoteMediaOutputs,
+											   EAvaMediaBroadcastAction InAction)
+{
+	FAvaMediaBroadcastRequest* Request = FMessageEndpoint::MakeMessage<FAvaMediaBroadcastRequest>();
+	Request->Profile = InProfile;
+	Request->Channel = InChannel.ToString();
+	Request->Action = InAction;
+
+	bool bNeedMediaOutputs = (InAction == EAvaMediaBroadcastAction::Start || InAction == EAvaMediaBroadcastAction::UpdateConfig);
+	if (Request->Channel.IsEmpty())
+	{
+		// If no channel is specified, no need to send outputs. 
+		bNeedMediaOutputs = false;
+		
+		if (InAction == EAvaMediaBroadcastAction::UpdateConfig)
+		{
+			UE_LOG(LogAvaPlaybackClient, Error, TEXT("Invalid Broadcast Request: UpdateConfig requires a channel to be specified."));
+			return;
+		}
+	}
+	
+	if (bNeedMediaOutputs)
+	{
+		// For now, we send the media output objects in the request.
+		uint32 TotalOutputDataSize = 0;
+		for (UMediaOutput* const MediaOutput : InRemoteMediaOutputs)
+		{
+			if (IsValid(MediaOutput))
+			{
+				FAvaMediaOutputData MediaOutputData = UE::AvaMediaOutputUtils::CreateMediaOutputData(MediaOutput);
+				
+				// Also propagate output info. Necessary to send the Guid and Server.
+				MediaOutputData.OutputInfo
+					= UAvalancheBroadcast::Get().GetCurrentProfile().GetChannel(InChannel).GetMediaOutputInfo(MediaOutput);
+				
+				if (!MediaOutputData.OutputInfo.IsValid())
+				{
+					UE_LOG(LogAvaPlaybackClient, Warning, TEXT("MediaOutput information is not valid for channel \"%s\"."),
+						   *InChannel.ToString());
+				}
+				TotalOutputDataSize += MediaOutputData.SerializedData.Num();
+				Request->MediaOutputs.Add(MoveTemp(MediaOutputData));
+			}
+		}
+		
+		// Adding a warning here, if we hit this warning, it may be necessary
+		// to send the data through some other transport.
+		// Data size is about 1k per output. It is unlikely we will hit this limit.
+		const uint32 SafeMessageSizeLimit = UE::AvaMediaMessageUtils::GetSafeMessageSizeLimit();	
+		if (TotalOutputDataSize > SafeMessageSizeLimit)
+		{
+			UE_LOG(LogAvaPlaybackClient, Warning,
+				TEXT("The broadcast request (DataSize: %d) is larger that the safe message size limit (%d)."),
+				TotalOutputDataSize, SafeMessageSizeLimit);
+		}
+	}
+
+	const TArray<FMessageAddress> ServerAddressesForChannel = GetServerAddressesForChannel(InChannel);
+
+	// Also update the channel settings.
+	if (InAction == EAvaMediaBroadcastAction::Start || InAction == EAvaMediaBroadcastAction::UpdateConfig)
+	{
+		SendBroadcastChannelSettingsUpdate(ServerAddressesForChannel, UAvalancheBroadcast::Get().GetCurrentProfile().GetChannel(InChannel));
+	}
+
+	// Send to server(s) that have output for this channel.
+	SendRequest(Request, ServerAddressesForChannel);
+}
+
+bool FAvaMediaPlaybackClient::IsMediaOutputRemoteFallback(const UMediaOutput* InMediaOutput)
+{
+	const FString DeviceName = UE::AvaMediaOutputUtils::GetDeviceName(InMediaOutput);
+	if (!DeviceName.IsEmpty())
+	{
+		// See if it begins with any of the remote server names that are connected.
+		// See FAvaDeviceProviderData::ApplyServerName. All MediaOutput from replicated Device Provider
+		// Have the server name in the beginning of the device name.
+		for (const TPair<FString, TSharedPtr<FServerInfo>>& Server : Servers)
+		{
+			if (DeviceName.StartsWith(Server.Key))
+			{
+				return true;
+			}
+		}
+	}
+
+	// We may fall in this case if the device is from a remote server that is offline
+	// right now. So we still want to double check if the device is a local one.
+
+	// Search in the list of device providers for the local server only.
+	const FName DeviceProviderName = UE::AvaMediaOutputUtils::GetDeviceProviderName(InMediaOutput);
+	if (!DeviceProviderName.IsNone() && !DeviceName.IsEmpty())
+	{
+		// First check if we have a wrapper for this provider. This means
+		// we have some remote devices currently online on that provider.
+		const FAvaDeviceProviderWrapper* DeviceProviderWrapper =
+			ParentModule->GetDeviceProviderProxyManager().GetDeviceProviderWrapper(DeviceProviderName);
+
+		IMediaIOCoreDeviceProvider* LocalDeviceProvider;
+
+		if (DeviceProviderWrapper && DeviceProviderWrapper->HasLocalProvider())
+		{
+			LocalDeviceProvider = DeviceProviderWrapper->GetLocalProvider();
+		}
+		else
+		{
+			// If we don't have a wrapper, then directly fetch the concrete local device provider.
+			LocalDeviceProvider = IMediaIOCoreModule::Get().GetDeviceProvider(DeviceProviderName);
+		}
+
+		if (LocalDeviceProvider)
+		{
+			const FName DeviceFName(DeviceName);
+			TArray<FMediaIODevice> LocalDevices = LocalDeviceProvider->GetDevices();
+			for (const FMediaIODevice& LocalDevice : LocalDevices)
+			{
+				if (DeviceFName == LocalDevice.DeviceName)
+				{
+					return false; // found as local device, so not remote.
+				}
+			}
+			// If we didn't find the device in local provider (and it has a local provider),
+			// then we can safely consider it a remote device, but it's server is currently offline.
+			return true;
+		}
+	}
+
+	// By default the device will be local. But this is not a good default.
+	return false;
+}
+
+EAvaMediaIssueSeverity FAvaMediaPlaybackClient::GetMediaOutputIssueSeverity(
+	const FString& InServerName, const FString& InChannelName, const FGuid& InOutputGuid) const
+{
+	if (const FBroadcastChannelInfo* ChannelInfo = GetChannelInfo(InServerName, InChannelName))
+	{
+		if (const FAvaMediaOutputStatus* MediaOutputStatus = ChannelInfo->MediaOutputStatuses.Find(InOutputGuid))
+		{
+			return MediaOutputStatus->MediaIssueSeverity;
+		}
+	}
+	return EAvaMediaIssueSeverity::None;
+}
+
+const TArray<FString>& FAvaMediaPlaybackClient::GetMediaOutputIssueMessages(
+	const FString& InServerName, const FString& InChannelName, const FGuid& InOutputGuid) const
+{
+	if (const FBroadcastChannelInfo* ChannelInfo = GetChannelInfo(InServerName, InChannelName))
+	{
+		if (const FAvaMediaOutputStatus* MediaOutputStatus = ChannelInfo->MediaOutputStatuses.Find(InOutputGuid))
+		{
+			return MediaOutputStatus->MediaIssueMessages;
+		}
+	}
+	static const TArray<FString> EmptyStringArray;
+	return EmptyStringArray;
+}
+
+EAvaMediaOutputState FAvaMediaPlaybackClient::GetMediaOutputState(
+	const FString& InServerName, const FString& InChannelName, const FGuid& InOutputGuid) const
+{
+	if (const FBroadcastChannelInfo* ChannelInfo = GetChannelInfo(InServerName, InChannelName))
+	{
+		if (const FAvaMediaOutputStatus* MediaOutputStatus = ChannelInfo->MediaOutputStatuses.Find(InOutputGuid))
+		{
+			return MediaOutputStatus->MediaOutputState;
+		}
+		else
+		{
+			// If the server is connected but doesn't have a status yet, we consider it "idle".
+			return EAvaMediaOutputState::Idle;
+		}
+	}
+	// If we couldn't find the channel state in the list of servers,
+	// then consider this output as offline.
+	return EAvaMediaOutputState::Offline;
+}
+
+// Note: same logic as GetServerNamesForChannel.
+bool FAvaMediaPlaybackClient::HasAnyServerOnlineForChannel(const FName& InChannelName) const
+{
+	if (InChannelName.IsNone())
+	{
+		return false;
+	}
+	
+	const FAvaOutputChannel& Channel = UAvalancheBroadcast::Get().GetCurrentProfile().GetChannel(InChannelName);
+	if (!Channel.IsValidChannel())
+	{
+		return false;
+	}
+
+	const TArray<UMediaOutput*>& RemoteOutputs = Channel.GetRemoteMediaOutputs();
+	
+	for (const UMediaOutput* RemoteOutput : RemoteOutputs)
+	{
+		const FAvaMediaOutputInfo& OutputInfo = Channel.GetMediaOutputInfo(RemoteOutput);
+		if (OutputInfo.IsValid())
+		{
+			if (Servers.Contains(OutputInfo.ServerName))
+			{
+				return true;
+			}
+		}
+		else
+		{
+			UE_LOG(LogAvaPlaybackClient, Warning, TEXT("MediaOutputInfo invalid for channel \"%s\"."), *InChannelName.ToString());
+			
+			// Try to find the server name from the device name.
+			const FString ServerName = GetServerNameForMediaOutputFallback(RemoteOutput);
+			if (!ServerName.IsEmpty())
+			{
+				return true;
+			}
+		}
+	}
+	
+	return false;
+}
+
+TOptional<EAvaMediaPlaybackStatus> FAvaMediaPlaybackClient::GetRemotePlaybackStatus(
+	const FGuid& InInstanceId, const FSoftObjectPath& InAssetPath, const FString& InChannelName, const FString& InServerName) const
+{
+	TOptional<EAvaMediaPlaybackStatus> PlaybackStatus;
+
+	if (!InServerName.IsEmpty())
+	{
+		PlaybackStatus = GetPlaybackStatusForServer(InServerName, InChannelName, InInstanceId, InAssetPath);
+	}
+	else
+	{
+		// Because of "forked" channels, we could actually have more than one server playing this channel.
+		TArray<FString> ServerNames = GetServerNamesForChannel(FName(InChannelName));
+
+		// In this case, we return the first one found. It is not accurate though.
+		// Todo: Ideally, we should retire this code path.
+		for (const FString& ServerName : ServerNames)
+		{
+			PlaybackStatus = GetPlaybackStatusForServer(ServerName, InChannelName, InInstanceId, InAssetPath);
+			if (PlaybackStatus.IsSet())
+			{
+				break;
+			}
+		}
+	}
+
+	return PlaybackStatus;
+}
+
+const FString* FAvaMediaPlaybackClient::GetRemotePlaybackUserData(
+	const FGuid& InInstanceId, const FSoftObjectPath& InAssetPath, const FString& InChannelName, const FString& InServerName) const
+{
+	if (!InServerName.IsEmpty())
+	{
+		return GetPlaybackUserDataForServer(InServerName, InChannelName, InInstanceId, InAssetPath);
+	}
+
+	// Because of "forked" channels, we could actually have more than one server playing this channel.
+	TArray<FString> ServerNames = GetServerNamesForChannel(FName(InChannelName));
+
+	// In this case, we return the first one found. It is not accurate though.
+	// Todo: Ideally, we should retire this code path.
+	for (const FString& ServerName : ServerNames)
+	{
+		if (const FString* FoundUserData = GetPlaybackUserDataForServer(ServerName, InChannelName, InInstanceId, InAssetPath))
+		{
+			return FoundUserData;
+		}
+	}
+
+	return nullptr;
+}
+
+TOptional<EAvaMediaPlaybackAssetStatus> FAvaMediaPlaybackClient::GetRemotePlaybackAssetStatus(
+	const FSoftObjectPath& InAssetPath, const FString& InServerName) const
+{
+	if (const FServerInfo* ServerInfo = GetServerInfo(InServerName))
+	{
+		return ServerInfo->GetPlaybackAssetStatus(InAssetPath);
+	}
+	return TOptional<EAvaMediaPlaybackAssetStatus>();
+}
+
+TOptional<EAvaMediaPlaybackStatus> FAvaMediaPlaybackClient::GetPlaybackStatusForServer(
+	const FString& InServerName, const FString& InChannelName, const FGuid& InInstanceId, const FSoftObjectPath& InAssetPath) const
+{
+	if (const FServerInfo* ServerInfo = GetServerInfo(InServerName))
+	{
+		return ServerInfo->GetInstanceStatus(InChannelName, InInstanceId, InAssetPath);
+	}
+	return TOptional<EAvaMediaPlaybackStatus>();
+}
+
+const FString* FAvaMediaPlaybackClient::GetPlaybackUserDataForServer(
+	const FString& InServerName, const FString& InChannelName, const FGuid& InInstanceId, const FSoftObjectPath& InAssetPath) const
+{
+	if (const FServerInfo* ServerInfo = GetServerInfo(InServerName))
+	{
+		return ServerInfo->GetInstanceUserData(InChannelName, InInstanceId, InAssetPath);
+	}
+	return nullptr;
+}
+
+void FAvaMediaPlaybackClient::RequestPlaybackAssetStatusForServer(const FSoftObjectPath& InAssetPath, const FString& InServerName, bool bInForceRefresh)
+{
+	FServerInfo* ServerInfo = GetServerInfo(InServerName);
+	if (!ServerInfo)
+	{
+		UE_LOG(LogAvaPlaybackClient, Warning, TEXT("Specified server \"%s\" not found for Playback Asset Status Request."), *InServerName);
+		return;
+	}
+	
+	// Status requests can be spammed every frame by the UI so we block them
+	// here to avoid spamming the message bus and servers.
+	const FDateTime CurrentTime = FDateTime::UtcNow();
+	if (const FPendingPlaybackAssetStatusRequest* PendingRequest = ServerInfo->PendingPlaybackAssetStatusRequests.Find(InAssetPath))
+	{
+		// Check if the current request hasn't expired.
+		if (CurrentTime < PendingRequest->ExpirationTime)
+		{
+			// A "force refresh" request will override a non-"force refresh" one.
+			if (PendingRequest->bForceRefresh >= bInForceRefresh)
+			{
+				return;
+			}
+		}
+	}
+
+	// Keep track of the request and the time it was made so we have a timeout.
+	const FDateTime ExpirationTime = CurrentTime + FTimespan::FromSeconds(UAvalancheMediaSettings::Get().ClientPendingStatusRequestTimeout);
+	ServerInfo->PendingPlaybackAssetStatusRequests.Add(InAssetPath, {ExpirationTime, bInForceRefresh} );
+
+	FAvaMediaPlaybackAssetStatusRequest* Request = FMessageEndpoint::MakeMessage<FAvaMediaPlaybackAssetStatusRequest>();
+	Request->AssetPath = InAssetPath;
+	Request->bForceRefresh = bInForceRefresh;
+	SendRequest(Request, ServerInfo->Address);
+}
+
+void FAvaMediaPlaybackClient::Tick()
+{
+	for (const TPair<FString, TSharedPtr<FServerInfo>>& Server : Servers)
+	{
+		if (!Server.Value->PendingPlaybackCommands.IsEmpty())
+		{
+			// Batched per server for now.
+			FAvaMediaPlaybackRequest* Request = FMessageEndpoint::MakeMessage<FAvaMediaPlaybackRequest>();
+			Request->Commands = MoveTemp(Server.Value->PendingPlaybackCommands);
+			Server.Value->PendingPlaybackCommands.Reset();
+			SendRequest(Request, Server.Value->Address);
+		}
+	}
+}
+
+bool FAvaMediaPlaybackClient::HandlePingTicker(float InDeltaTime)
+{
+	const FDateTime CurrentTime = FDateTime::UtcNow();
+	PublishPlaybackPing(CurrentTime, true);
+	RemoveDeadServers(CurrentTime);
+	return true;
+}
+
+void FAvaMediaPlaybackClient::HandlePlaybackPongMessage(const FAvaMediaPlaybackPong& InMessage,
+														const TSharedRef<IMessageContext, ESPMode::ThreadSafe>&	InContext)
+{
+	// Remark: the server will send a user data update before sending the pong message.
+	// The server info will have been created already by this point. But it won't have all the information
+	// complete yet.
+	
+	FServerInfo& ServerInfo = GetOrCreateServerInfo(InMessage.ServerName, InContext->GetSender());
+	ServerInfo.ResetPingTimeout();
+
+	// If the server info is new, the process id will not be set yet.
+	const bool bIsNewServer = ServerInfo.ProcessId == 0 ? true : false;
+	
+	if (!bIsNewServer && ServerInfo.ProcessId != InMessage.ProcessId)
+	{
+		UE_LOG(LogAvaPlaybackClient, Warning, TEXT("Received server \"%s\" process Id has changed from %d to %d."),
+			*InMessage.ServerName, ServerInfo.ProcessId, InMessage.ProcessId);
+	}
+	
+	ServerInfo.ProcessId = InMessage.ProcessId;
+	ServerInfo.ProjectContentPath = InMessage.ProjectContentPath;
+	
+	// The server may have requested the client info, and it has to be sent unless
+	// this is a new server, in which case it has already been sent by OnServerAdded().
+	if (InMessage.bRequestClientInfo && !bIsNewServer)
+	{
+		SendClientInfo(InContext->GetSender());
+	}
+
+	// We want to propagate the connection event after all the information has been set
+	// in the server info.
+	if (bIsNewServer)
+	{
+		using namespace UE::AvaMediaPlaybackClient::Delegates;
+		GetOnConnectionEvent().Broadcast(*this, {InMessage.ServerName, EConnectionEvent::ServerConnected});
+	}
+}
+
+void FAvaMediaPlaybackClient::HandlePlaybackLogMessage(const FAvaMediaPlaybackLog& InMessage, const TSharedRef<IMessageContext, ESPMode::ThreadSafe>& InContext)
+{
+	if (GLog && InMessage.Verbosity != 0)
+	{
+		// Basic message formatting: adding the server name to identify the origin of the message.
+		const FString FormattedText = InMessage.ServerName + TEXT(" >> ") + InMessage.Text;
+
+		// Relay to local GLog.
+		GLog->Serialize(*FormattedText, static_cast<ELogVerbosity::Type>(InMessage.Verbosity), InMessage.Category, InMessage.Time);
+	}
+}
+
+void FAvaMediaPlaybackClient::HandleUpdateServerUserData(const FAvaMediaUpdateServerUserData& InMessage, const TSharedRef<IMessageContext, ESPMode::ThreadSafe>& InContext)
+{
+	FServerInfo& ServerInfo = GetOrCreateServerInfo(InMessage.ServerName, InContext->GetSender());
+	ServerInfo.ResetPingTimeout();
+	ServerInfo.UserDataEntries = InMessage.UserDataEntries;
+
+	// Logging when user data is updated (for debugging).
+	UE_LOG(LogAvaPlaybackClient, Verbose, TEXT("Received new user data for server \"%s\"."), *InMessage.ServerName);
+	for (const TPair<FString, FString>& UserData : ServerInfo.UserDataEntries)
+	{
+		UE_LOG(LogAvaPlaybackClient, Verbose, TEXT("User data \"%s\":\"%s\"."), *UserData.Key, *UserData.Value);
+	}
+}
+
+void FAvaMediaPlaybackClient::HandleStatStatus(const FAvaMediaStatStatus& InMessage, const TSharedRef<IMessageContext, ESPMode::ThreadSafe>& InContext)
+{
+	// If the client couldn't run the command, but the serve did, we are going
+	// to replicate the stats state on the client instead.
+	if (InMessage.bClientStateReliable == false && InMessage.bCommandSucceeded)
+	{
+		UE_LOG(LogAvaPlaybackClient, Verbose,
+			TEXT("Received reliable enabled runtime stats from server \"%s\"."), *InMessage.ServerName);
+		IAvaModule::Get().OverwriteEnabledRuntimeStats(InMessage.EnabledRuntimeStats);
+	}
+}
+
+void FAvaMediaPlaybackClient::HandleDeviceProviderDataListMessage(const FAvaDeviceProviderDataList& InMessage,
+																  const TSharedRef<IMessageContext, ESPMode::ThreadSafe>& InContext)
+{
+	UE_LOG(LogAvaPlaybackClient, Verbose, TEXT("Received Device Provider Data for server \"%s\"."),
+		   *InMessage.ServerName);
+	IAvaDeviceProviderProxyManager& Manager = ParentModule->GetDeviceProviderProxyManager();
+	Manager.Install(InMessage.ServerName, InMessage);
+
+	{
+		FServerInfo& ServerInfo = GetOrCreateServerInfo(InMessage.ServerName, InContext->GetSender());
+		ServerInfo.bDeviceProviderInstalled = true;
+	}
+
+	// Check if we have installed all the connected server's device providers.
+	bool bAllDeviceProvidersInstalled = true;
+	{
+		ForAllServers([&bAllDeviceProvidersInstalled](const FServerInfo& InServerInfo)
+		{
+			if (InServerInfo.bDeviceProviderInstalled == false)
+			{
+				UE_LOG(LogAvaPlaybackClient, Verbose,
+					   TEXT("Device Provider Data for server \"%s\" hasn't been received yet."), *InServerInfo.ServerName);
+				bAllDeviceProvidersInstalled = false;
+			}
+		});
+	}
+
+	if (bAllDeviceProvidersInstalled)
+	{
+		UE_LOG(LogAvaPlaybackClient, Verbose,
+			   TEXT("All Device Providers Proxies are installed. Requesting broadcast status update..."));
+		// Remark: We are waiting for all the device providers from all online servers to be
+		// installed before requesting a broadcast status.
+		// This is because handling the broadcast status update will load the broadcast
+		// configuration. When the broadcast object is loaded, there is the potential for
+		// legacy code to convert the config and it needs the device provider proxies to
+		// be installed to resolve the device names and corresponding servers.
+		ForAllServers([this](const FServerInfo& InServerInfo)
+		{
+			UE_LOG(LogAvaPlaybackClient, Verbose, TEXT("Requesting full broadcast status update for server \"%s\"."),
+				   *InServerInfo.ServerName);
+			// Also request broadcast channels status to update state of local channels if required.
+			FAvaMediaBroadcastStatusRequest* StatusRequest = FMessageEndpoint::MakeMessage<
+				FAvaMediaBroadcastStatusRequest>();
+			StatusRequest->bIncludeMediaOutputData = true;
+			// We will want to ensure our local config is the same as the servers.
+
+			SendRequest(StatusRequest, InServerInfo.Address);
+		});
+	}
+}
+
+void FAvaMediaPlaybackClient::HandleBroadcastStatusMessage(const FAvaMediaBroadcastStatus& InMessage,
+														   const TSharedRef<IMessageContext, ESPMode::ThreadSafe>& InContext)
+{
+	using namespace UE::AvaMediaPlaybackClient::Private;
+	UE_LOG(LogAvaPlaybackClient, Verbose,
+		   TEXT("Received broadcast update from server \"%s\" for channel \"%s\": Status: \"%s\"."),
+		   *InMessage.ServerName, *InMessage.ChannelName,
+		   *EnumToString(InMessage.ChannelState));
+
+	// Log details for the output status.
+	for (const TPair<FGuid, FAvaMediaOutputStatus>& OutputStatus : InMessage.MediaOutputStatuses)
+	{
+		UE_LOG(LogAvaPlaybackClient, Verbose,
+			   TEXT("Playback Server: \"%s\" Channel: \"%s\" OutputId \"%s\" Status: \"%s\" Severity: \"%s\"."),
+			   *InMessage.ServerName, *InMessage.ChannelName, *OutputStatus.Key.ToString(),
+			   *EnumToString(OutputStatus.Value.MediaOutputState),
+			   *EnumToString(OutputStatus.Value.MediaIssueSeverity));
+		
+		for (const FString& Message : OutputStatus.Value.MediaIssueMessages)
+		{
+			UE_LOG(LogAvaPlaybackClient, Verbose, TEXT("Output Issue Message: \"%s\"."), *Message);
+		}
+	}
+
+	FServerInfo& ServerInfo = GetOrCreateServerInfo(InMessage.ServerName, InContext->GetSender());
+	ServerInfo.ResetPingTimeout();
+
+	const FString CurrentProfileName = UAvalancheBroadcast::Get().GetCurrentProfileName().ToString();
+
+	{
+		// Keep a back store of this locally, mainly so we can respond to UI requests
+		// for the status of the media output objects.
+		FBroadcastChannelInfo& ChannelInfo = ServerInfo.GetOrCreateBroadcastChannelInfo(InMessage.ChannelName);
+		ChannelInfo.ChannelState = InMessage.ChannelState;
+		ChannelInfo.ChannelIssueSeverity = InMessage.ChannelIssueSeverity;		
+		ChannelInfo.MediaOutputStatuses = InMessage.MediaOutputStatuses;
+	}
+
+	{
+		const FName ChannelName = *InMessage.ChannelName;
+		FAvaOutputChannel& Channel = UAvalancheBroadcast::Get().GetCurrentProfile().GetChannelMutable(ChannelName);
+		if (Channel.IsValidChannel())
+		{
+			// Purpose of this section is to ensure the server and the client have
+			// the same configuration. We want to perform this in the less intrusive
+			// way possible, i.e. if the server is running and the client connects to it
+			// we would like the server to be left completely undisturbed if it's configuration
+			// is already synced.
+
+			// Verify if the server has the corresponding outputs already.
+			bool bServerIsMissingOutputs = false;
+			int32 NumOutputsForThisServer = 0;
+			TArray<UMediaOutput*> RemoteOutputs = Channel.GetRemoteMediaOutputs();
+			for (const UMediaOutput* RemoteOutput : RemoteOutputs)
+			{
+				const FAvaMediaOutputInfo& OutputInfo = Channel.GetMediaOutputInfo(RemoteOutput);
+				if (OutputInfo.IsValid() && OutputInfo.ServerName == InMessage.ServerName)
+				{
+					++NumOutputsForThisServer;
+					if (!InMessage.MediaOutputStatuses.Contains(OutputInfo.Guid))
+					{
+						bServerIsMissingOutputs = true;
+					}
+					else if (InMessage.bIncludeMediaOutputData)
+					{
+						// The server already has this output, but we need to check
+						// if it is the same. TODO.
+					}
+				}
+			}
+
+			if ((InMessage.bIncludeMediaOutputData && NumOutputsForThisServer > 0 && InMessage.MediaOutputs.Num() == 0)
+				|| bServerIsMissingOutputs)
+			{
+				UE_LOG(LogAvaPlaybackClient, Log,
+					   TEXT("Playback Server: \"%s\" Channel: \"%s\" is missing outputs. Requesting configuration update."),
+					   *InMessage.ServerName, *InMessage.ChannelName);
+
+				RequestBroadcast(CurrentProfileName, ChannelName, RemoteOutputs, EAvaMediaBroadcastAction::UpdateConfig);
+			}
+
+			// Note: this will broadcast to delegates which may then request states of media outputs.
+			// So the back store needs to be updated before calling this.
+			Channel.RefreshState();
+
+			// Explicitly call the broadcast of the channel state change to propagate the status of the output,
+			// which is not fully reflected in the channel state.
+			FAvaOutputChannel::GetOnChannelChanged().Broadcast(Channel, EAvaChannelChange::State);
+		}
+		else
+		{
+			UE_LOG(LogAvaPlaybackClient, Error,
+				   TEXT("Received broadcast update from server \"%s\" for locally invalid channel \"%s\"."),
+				   *InMessage.ServerName, *InMessage.ChannelName);
+
+			// Request this channel be deleted.
+			RequestBroadcast(CurrentProfileName, ChannelName, {}, EAvaMediaBroadcastAction::DeleteChannel);
+		}
+	}
+
+	// Check for missing channels (is this logic valid?)
+	if (InMessage.ChannelIndex == InMessage.NumChannels - 1)
+	{
+		// We have received the last channel this server has, so we can check in the server state
+		// if it has all the channels it is supposed to have.
+		const TArray<FAvaOutputChannel*>& Channels = UAvalancheBroadcast::Get().GetCurrentProfile().GetChannels();
+		for (const FAvaOutputChannel* Channel : Channels)
+		{
+			bool bHasOutputsForThisServer = false;
+			TArray<UMediaOutput*> RemoteOutputs = Channel->GetRemoteMediaOutputs();
+			for (const UMediaOutput* RemoteOutput : RemoteOutputs)
+			{
+				const FAvaMediaOutputInfo& OutputInfo = Channel->GetMediaOutputInfo(RemoteOutput);
+				if (OutputInfo.IsValid() && OutputInfo.ServerName == InMessage.ServerName)
+				{
+					bHasOutputsForThisServer = true;
+					break;
+				}
+			}
+			if (bHasOutputsForThisServer)
+			{
+				if (ServerInfo.GetBroadcastChannelInfo(Channel->GetChannelName().ToString()) == nullptr)
+				{
+					// The server is missing this channel, we need to update it's config.
+					RequestBroadcast(CurrentProfileName, Channel->GetChannelName(), RemoteOutputs, EAvaMediaBroadcastAction::UpdateConfig);
+				}
+			}
+		}
+	}
+}
+
+void FAvaMediaPlaybackClient::HandlePlaybackAssetStatusMessage(
+	const FAvaMediaPlaybackAssetStatus& InMessage, const TSharedRef<IMessageContext, ESPMode::ThreadSafe>& InContext)
+{
+	using namespace UE::AvaMediaPlaybackClient::Private;
+	using namespace UE::AvaMediaPlaybackClient::Delegates;
+	if (!InMessage.AssetPath.IsValid())
+	{
+		UE_LOG(LogAvaPlaybackClient, Error,
+			   TEXT("Received Playback Asset Status \"%s\" from server \"%s\" with invalid asset path, can't update status."),
+			   *EnumToString(InMessage.Status), *InMessage.ServerName);
+		return;
+	}
+	
+	FServerInfo& ServerInfo = GetOrCreateServerInfo(InMessage.ServerName, InContext->GetSender());
+	ServerInfo.ResetPingTimeout();	
+	ServerInfo.SetPlaybackAssetStatus(InMessage.AssetPath, InMessage.Status);
+		
+	// Clear the pending request (if any).
+	ServerInfo.PendingPlaybackAssetStatusRequests.Remove(InMessage.AssetPath);
+
+	GetOnPlaybackAssetStatusChanged().Broadcast(*this,{InMessage.AssetPath, InMessage.ServerName, InMessage.Status});
+
+	// TODO - Further refactoring needed to untangle playback state and asset state. 
+	// Because some of the states of the playback status reflect the state of the asset on disk,
+	// we need to make sure the playback state properly reflects the asset state for those cases.
+	const bool bAssetIsMissing = InMessage.Status == EAvaMediaPlaybackAssetStatus::Missing;
+	const bool bAssetIsAvailable = InMessage.Status == EAvaMediaPlaybackAssetStatus::Available
+									|| InMessage.Status == EAvaMediaPlaybackAssetStatus::MissingDependencies;
+
+	for (const TPair<FString, TUniquePtr<FPlaybackChannelInfo>>& ChannelInfo : ServerInfo.PlaybackChannelInfosByName)
+	{
+		check(ChannelInfo.Value.IsValid());
+		if (FPlaybackAssetInfo* AssetInfo = ChannelInfo.Value->GetAssetInfo(InMessage.AssetPath))
+		{
+			auto ConditionalChangeStatus = [AssetInfo, &InMessage, &ChannelInfo, bAssetIsMissing, bAssetIsAvailable, this](const FGuid& InInstanceId, EAvaMediaPlaybackStatus InPlaybackStatus)
+			{
+				// A missing asset leads to a missing playback.
+				if (InPlaybackStatus == EAvaMediaPlaybackStatus::Available && bAssetIsMissing)
+				{
+					AssetInfo->SetInstanceStatus(InInstanceId, EAvaMediaPlaybackStatus::Missing);
+					GetOnPlaybackStatusChanged().Broadcast(*this, {InInstanceId, InMessage.AssetPath, ChannelInfo.Key, InMessage.ServerName, InPlaybackStatus,EAvaMediaPlaybackStatus::Missing});
+				}
+				// If the playback was missing, and the asset becomes available, update the playback to available too.
+				else if (InPlaybackStatus == EAvaMediaPlaybackStatus::Missing && bAssetIsAvailable)
+				{
+					AssetInfo->SetInstanceStatus(InInstanceId, EAvaMediaPlaybackStatus::Available);
+					GetOnPlaybackStatusChanged().Broadcast(*this, {InInstanceId, InMessage.AssetPath, ChannelInfo.Key, InMessage.ServerName, InPlaybackStatus,EAvaMediaPlaybackStatus::Available});
+				}
+			};
+
+			for (TPair<FGuid, FPlaybackInstanceInfo>& InstanceInfo : AssetInfo->InstanceByIds)
+			{
+				ConditionalChangeStatus(InstanceInfo.Key, InstanceInfo.Value.Status);
+			}
+		}
+	}
+	
+	UE_LOG(LogAvaPlaybackClient, Verbose,
+		   TEXT("Received Playback Asset Status \"%s\" from server \"%s\" for \"%s\"."),
+		   *EnumToString(InMessage.Status),
+		   *InMessage.ServerName, *InMessage.AssetPath.ToString());
+}
+
+void FAvaMediaPlaybackClient::HandlePlaybackStatusMessage(const FAvaMediaPlaybackStatus& InMessage,
+														  const TSharedRef<IMessageContext, ESPMode::ThreadSafe>& InContext)
+{
+	FServerInfo& ServerInfo = GetOrCreateServerInfo(InMessage.ServerName, InContext->GetSender());
+	ServerInfo.ResetPingTimeout();
+	
+	if (InMessage.AssetPath.IsValid())
+	{
+		TOptional<EAvaMediaPlaybackStatus> PrevPlaybackStatus =
+			ServerInfo.GetInstanceStatus(InMessage.ChannelName, InMessage.InstanceId, InMessage.AssetPath);
+		if (!PrevPlaybackStatus.IsSet())
+		{
+			PrevPlaybackStatus = EAvaMediaPlaybackStatus::Unknown;
+		}
+		
+		ServerInfo.SetInstanceStatus(InMessage.ChannelName, InMessage.InstanceId, InMessage.AssetPath, InMessage.Status);
+		
+		if (InMessage.bValidUserData)
+		{
+			ServerInfo.SetInstanceUserData(InMessage.ChannelName, InMessage.InstanceId, InMessage.AssetPath, InMessage.UserData);
+		}
+
+		// Clear the pending request (if any).
+		PendingPlaybackStatusRequests.Remove(GetPlaybackStatusRequestKey(InMessage.InstanceId, InMessage.AssetPath, InMessage.ChannelName));
+
+		UE::AvaMediaPlaybackClient::Delegates::GetOnPlaybackStatusChanged().Broadcast(
+			*this, {InMessage.InstanceId, InMessage.AssetPath, InMessage.ChannelName, InMessage.ServerName, PrevPlaybackStatus.GetValue(), InMessage.Status});
+		
+		UE_LOG(LogAvaPlaybackClient, Verbose,
+			   TEXT("Received Playback Status \"%s\" from server.channel \"%s.%s\" for \"%s\"."),
+			   *StaticEnum<EAvaMediaPlaybackStatus>()->GetValueAsString(InMessage.Status),
+			   *InMessage.ServerName, *InMessage.ChannelName, *InMessage.AssetPath.ToString());
+	}
+	else
+	{
+		UE_LOG(LogAvaPlaybackClient, Error,
+			   TEXT(
+				   "Received Playback Status \"%s\" from server.channel \"%s.%s\" with invalid asset path, can't update status."
+			   ),
+			   *StaticEnum<EAvaMediaPlaybackStatus>()->GetValueAsString(InMessage.Status), *InMessage.ServerName,
+			   *InMessage.ChannelName);
+	}
+}
+
+void FAvaMediaPlaybackClient::HandlePlaybackStatusesMessage(const FAvaMediaPlaybackStatuses& InMessage,
+															const TSharedRef<IMessageContext, ESPMode::ThreadSafe>& InContext)
+{
+	FServerInfo& ServerInfo = GetOrCreateServerInfo(InMessage.ServerName, InContext->GetSender());
+	ServerInfo.ResetPingTimeout();
+
+	const int32 NumStatuses = FMath::Max(InMessage.AssetPaths.Num(), InMessage.InstanceIds.Num());
+	for (int32 StatusIndex = 0; StatusIndex < NumStatuses; StatusIndex++)
+	{
+		const FSoftObjectPath AssetPath = InMessage.AssetPaths.IsValidIndex(StatusIndex) ? InMessage.AssetPaths[StatusIndex] : FSoftObjectPath();
+		const FGuid InstanceId = InMessage.InstanceIds.IsValidIndex(StatusIndex) ? InMessage.InstanceIds[StatusIndex] : FGuid();
+		
+		if (AssetPath.IsValid())
+		{
+			TOptional<EAvaMediaPlaybackStatus> PrevPlaybackStatus = ServerInfo.GetInstanceStatus(InMessage.ChannelName, InstanceId, AssetPath);
+			if (!PrevPlaybackStatus.IsSet())
+			{
+				PrevPlaybackStatus = EAvaMediaPlaybackStatus::Unknown;
+			}
+
+			ServerInfo.SetInstanceStatus(InMessage.ChannelName, InstanceId, AssetPath, InMessage.Status);
+
+			// Clear the pending request (if any).
+			PendingPlaybackStatusRequests.Remove(GetPlaybackStatusRequestKey(InstanceId, AssetPath, InMessage.ChannelName));
+
+			UE::AvaMediaPlaybackClient::Delegates::GetOnPlaybackStatusChanged().Broadcast(
+				*this, {InstanceId, AssetPath, InMessage.ChannelName, InMessage.ServerName, PrevPlaybackStatus.GetValue(), InMessage.Status});
+			
+			UE_LOG(LogAvaPlaybackClient, Verbose,
+				   TEXT("Received Playback Status \"%s\" from server.channel \"%s.%s\" for \"%s\"."),
+				   *StaticEnum<EAvaMediaPlaybackStatus>()->GetValueAsString(InMessage.Status),
+				   *InMessage.ServerName, *InMessage.ChannelName, *AssetPath.ToString());
+		}
+		else
+		{
+			UE_LOG(LogAvaPlaybackClient, Error,
+				   TEXT(
+					   "Received Playback Status \"%s\" from server.channel \"%s.%s\" with invalid asset path, can't update status."
+				   ),
+				   *StaticEnum<EAvaMediaPlaybackStatus>()->GetValueAsString(InMessage.Status), *InMessage.ServerName,
+				   *InMessage.ChannelName);
+		}
+	}
+}
+
+void FAvaMediaPlaybackClient::HandlePlaybackSequenceEventMessage(const FAvaMediaPlayableSequenceEvent& InMessage, const TSharedRef<IMessageContext, ESPMode::ThreadSafe>& InContext)
+{
+	using namespace UE::AvaMediaPlaybackClient::Delegates;
+	const FPlaybackSequenceEventArgs Args =
+	{
+		InMessage.InstanceId,
+		InMessage.AssetPath,
+		InMessage.ChannelName,
+		InMessage.ServerName,
+		InMessage.SequenceName,
+		InMessage.EventType
+	};
+	GetOnPlaybackSequenceEvent().Broadcast(*this, Args);
+}
+
+void FAvaMediaPlaybackClient::HandlePlaybackTransitionEventMessage(const FAvaMediaPlayableTransitionEvent& InMessage, const TSharedRef<IMessageContext, ESPMode::ThreadSafe>& InContext)
+{
+	using namespace UE::AvaMediaPlaybackClient::Delegates;
+	const FPlaybackTransitionEventArgs Args =
+	{
+		InMessage.TransitionId,
+		InMessage.InstanceId,
+		InMessage.ChannelName,
+		InMessage.ServerName,
+		InMessage.GetEventFlags()
+	};
+	GetOnPlaybackTransitionEvent().Broadcast(*this, Args);
+}
+
+void FAvaMediaPlaybackClient::RegisterCommands()
+{
+	if (ConsoleCommands.Num() != 0)
+	{
+		return;
+	}
+
+	ConsoleCommands.Add(IConsoleManager::Get().RegisterConsoleCommand(
+		TEXT("AvaMediaClient.PingServers"),
+		TEXT("Requests servers to give their information."),
+		FConsoleCommandWithArgsDelegate::CreateRaw(this, &FAvaMediaPlaybackClient::PingServersCommand),
+		ECVF_Default
+	));
+	ConsoleCommands.Add(IConsoleManager::Get().RegisterConsoleCommand(
+		TEXT("AvaMediaClient.RefreshAssetStatus"),
+		TEXT("Request a refresh of the asset status."),
+		FConsoleCommandWithArgsDelegate::CreateRaw(this, &FAvaMediaPlaybackClient::RequestPlaybackAssetStatusCommand),
+		ECVF_Default
+	));
+	ConsoleCommands.Add(IConsoleManager::Get().RegisterConsoleCommand(
+		TEXT("AvaMediaClient.PlaybackRequest"),
+		TEXT("Request Connected Servers to execute a playback request."),
+		FConsoleCommandWithArgsDelegate::CreateRaw(this, &FAvaMediaPlaybackClient::RequestPlaybackCommand),
+		ECVF_Default
+	));
+	ConsoleCommands.Add(IConsoleManager::Get().RegisterConsoleCommand(
+		TEXT("AvaMediaClient.BroadcastRequest"),
+		TEXT("Request Connected Servers to execute a broadcast request."),
+		FConsoleCommandWithArgsDelegate::CreateRaw(this, &FAvaMediaPlaybackClient::RequestBroadcastCommand),
+		ECVF_Default
+	));
+	ConsoleCommands.Add(IConsoleManager::Get().RegisterConsoleCommand(
+		TEXT("AvaMediaClient.SetUserData"),
+		TEXT("Set Replicated User Data Entry (Key, Value)."),
+		FConsoleCommandWithArgsDelegate::CreateRaw(this, &FAvaMediaPlaybackClient::SetUserDataCommand),
+		ECVF_Default
+	));
+	ConsoleCommands.Add(IConsoleManager::Get().RegisterConsoleCommand(
+		TEXT("AvaMediaClient.Status"),
+		TEXT("Display current status of all server info."),
+		FConsoleCommandWithArgsDelegate::CreateRaw(this, &FAvaMediaPlaybackClient::ShowStatusCommand),
+		ECVF_Default
+	));
+}
+
+namespace UE::AvaMediaPlaybackClient::Private
+{
+	inline FString GetCommaSeparatedList(const UEnum* InEnum)
+	{
+		FString List;
+		const int32 NumEnums = InEnum->ContainsExistingMax() ? InEnum->NumEnums() - 1 : InEnum->NumEnums();
+		for (int32 EnumIndex = 0; EnumIndex < NumEnums; ++EnumIndex)
+		{
+			if (!List.IsEmpty())
+			{
+				List += TEXT(", ");
+			}
+			List += InEnum->GetNameStringByIndex(EnumIndex);
+		}
+		return List;
+	}
+}
+
+void FAvaMediaPlaybackClient::RequestPlaybackAssetStatusCommand(const TArray<FString>& InArgs)
+{
+	using namespace UE::AvaMediaPlaybackClient::Private;
+	if (Servers.Num() == 0)
+	{
+		UE_LOG(LogAvaPlaybackClient, Log, TEXT("No servers available. Do AvaMediaClient.PingServers."));
+		return;
+	}
+
+	if (InArgs.Num() == 0)
+	{
+		UE_LOG(LogAvaPlaybackClient, Log,
+			   TEXT("Arguments: [AssetPath] [ChannelOrServer]. Ex: \"/Game/AvaPlayback.AvaPlayback Channel0\""));
+	}
+
+	FString ChannelOrServerName;
+	FSoftObjectPath AssetPath;
+
+	if (InArgs.Num() > 0)
+	{
+		AssetPath = FSoftObjectPath(FTopLevelAssetPath(InArgs[0]));
+	}
+	
+	if (InArgs.Num() > 1)
+	{
+		ChannelOrServerName = InArgs[1];
+	}
+	else
+	{
+		ChannelOrServerName = UAvalancheBroadcast::Get().GetChannelName(0).ToString();
+		UE_LOG(LogAvaPlaybackClient, Log, TEXT("No Channel specified, using \"%s\"."), *ChannelOrServerName);
+	}
+
+	if (AssetPath.IsValid())
+	{
+		RequestPlaybackAssetStatus(AssetPath, ChannelOrServerName, true);
+	}
+	else
+	{
+		// Request a refresh of all cached assets for the specified server/channel.
+		ForAllServers([&ChannelOrServerName, this](const FServerInfo& InServerInfo)
+		{			
+			if (InServerInfo.ServerName ==  ChannelOrServerName || InServerInfo.BroadcastChannelInfosByName.Contains(ChannelOrServerName))
+			{
+				for (const TPair<FSoftObjectPath, EAvaMediaPlaybackAssetStatus>& AssetStatus : InServerInfo.PlaybackAssetStatuses)
+				{
+					RequestPlaybackAssetStatus(AssetStatus.Key, ChannelOrServerName, true);
+				}
+			}
+		});
+	}
+}
+
+void FAvaMediaPlaybackClient::RequestPlaybackCommand(const TArray<FString>& InArgs)
+{
+	using namespace UE::AvaMediaPlaybackClient::Private;
+	if (Servers.Num() == 0)
+	{
+		UE_LOG(LogAvaPlaybackClient, Log, TEXT("No servers available. Do AvaMediaClient.PingServers."));
+		return;
+	}
+
+	if (InArgs.Num() == 0)
+	{
+		UE_LOG(LogAvaPlaybackClient, Log,
+			   TEXT("Arguments: Action [Package] [AssetName]. Ex: \"Start /Game/AvaPlayback AvaPlayback\""));
+		UE_LOG(LogAvaPlaybackClient, Log, TEXT("Action: %s"), *GetCommaSeparatedList(StaticEnum<EAvaMediaPlaybackAction>()));
+		return;
+	}
+
+	const int64 PlaybackActionValue = StaticEnum<EAvaMediaPlaybackAction>()->GetValueByNameString(InArgs[0]);
+	if (PlaybackActionValue == INDEX_NONE)
+	{
+		UE_LOG(LogAvaPlaybackClient, Log, TEXT("Action: \"%s\" is not valid. Possible actions: %s"),
+			   *InArgs[0], *GetCommaSeparatedList(StaticEnum<EAvaMediaPlaybackAction>()));
+		return;
+	}
+
+	const EAvaMediaPlaybackAction PlaybackAction = static_cast<EAvaMediaPlaybackAction>(PlaybackActionValue);
+
+	if (InArgs.Num() >= 3)
+	{
+		//Concatenate all Args (starting from the 3rd) into one String with spaces in between each arg
+		FString ConcatenatedCommands;
+		for (int32 Index = 3; Index < InArgs.Num(); ++Index)
+		{
+			ConcatenatedCommands += InArgs[Index] + TEXT(" ");
+		}
+
+		FString ChannelName;
+		FParse::Value(*ConcatenatedCommands, TEXT("Channel="), ChannelName);
+
+		FString InstanceId;
+		FParse::Value(*ConcatenatedCommands, TEXT("InstandId="), InstanceId);
+		
+		const FSoftObjectPath AssetPath(InArgs[1] + TEXT(".") + InArgs[2]);
+		RequestPlayback(FGuid(InstanceId), AssetPath, ChannelName, PlaybackAction);
+	}
+	else
+	{
+		RequestPlayback(FGuid(), FSoftObjectPath(), TEXT(""), PlaybackAction);
+	}
+}
+
+void FAvaMediaPlaybackClient::RequestBroadcastCommand(const TArray<FString>& InArgs)
+{
+	using namespace UE::AvaMediaPlaybackClient::Private;
+
+	if (Servers.Num() == 0)
+	{
+		UE_LOG(LogAvaPlaybackClient, Log, TEXT("No servers available."));
+		return;
+	}
+
+	if (InArgs.Num() == 0 || InArgs.Num() > 1)
+	{
+		UE_LOG(LogAvaPlaybackClient, Log, TEXT("Arguments: Action.  Possible actions: %s"),
+			   *GetCommaSeparatedList(StaticEnum<EAvaMediaBroadcastAction>()));
+		return;
+	}
+
+	const int64 BroadcastActionValue = StaticEnum<EAvaMediaBroadcastAction>()->GetValueByNameString(InArgs[0]);
+	if (BroadcastActionValue == INDEX_NONE)
+	{
+		UE_LOG(LogAvaPlaybackClient, Log, TEXT("Action: \"%s\" is not valid. Possible actions: %s"),
+			   *InArgs[0], *GetCommaSeparatedList(StaticEnum<EAvaMediaBroadcastAction>()));
+		return;
+	}
+	
+	RequestBroadcast(TEXT(""), FName(), TArray<UMediaOutput*>(), static_cast<EAvaMediaBroadcastAction>(BroadcastActionValue));
+}
+
+void FAvaMediaPlaybackClient::SetUserDataCommand(const TArray<FString>& InArgs)
+{
+	if (InArgs.Num() >= 2)
+	{
+		UE_LOG(LogAvaPlaybackClient, Log, TEXT("Setting User Data Key \"%s\" to Value: \"%s\"."), *InArgs[0], *InArgs[1]);
+		SetUserData(InArgs[0], InArgs[1]);
+	}
+	else if (InArgs.Num() == 1)
+	{
+		// One argument means to remove that user data entry.
+		if (HasUserData(InArgs[0]))
+		{
+			UE_LOG(LogAvaPlaybackClient, Log, TEXT("Removing User Data Key \"%s\"."), *InArgs[0]);
+			RemoveUserData(InArgs[0]);
+		}
+		else
+		{
+			UE_LOG(LogAvaPlaybackClient, Error, TEXT("User Data Key \"%s\" not found."), *InArgs[0]);
+		}
+	}
+}
+
+void FAvaMediaPlaybackClient::ShowStatusCommand(const TArray<FString>& InArgs)
+{
+	using namespace UE::AvaMediaPlaybackClient::Private;
+	UE_LOG(LogAvaPlaybackClient, Display, TEXT("Playback Client: \"%s\""), *ComputerName);
+	UE_LOG(LogAvaPlaybackClient, Display, TEXT("- Endpoint Bus Address: \"%s\""), MessageEndpoint.IsValid() ? *MessageEndpoint->GetAddress().ToString() : TEXT("Invalid"));
+	UE_LOG(LogAvaPlaybackClient, Display, TEXT("- ProcessId: %d"), ProcessId);
+	UE_LOG(LogAvaPlaybackClient, Display, TEXT("- Content Path: \"%s\""), *ProjectContentPath);
+
+	for (const TPair<FString, FString>& UserData : UserDataEntries)
+	{
+		UE_LOG(LogAvaPlaybackClient, Display, TEXT("- User data \"%s\":\"%s\"."), *UserData.Key, *UserData.Value);
+	}
+	
+	for (const TPair<FString, TSharedPtr<FServerInfo>>& Server : Servers)
+	{
+		check(Server.Value.IsValid());
+		const FServerInfo& ServerInfo = *Server.Value;
+		UE_LOG(LogAvaPlaybackClient, Display, TEXT("Connected Server: \"%s\""), *ServerInfo.ServerName);
+		
+		UE_LOG(LogAvaPlaybackClient, Display, TEXT("   - Endpoint Bus Address: \"%s\""), *ServerInfo.Address.ToString());
+		UE_LOG(LogAvaPlaybackClient, Display, TEXT("   - ProcessId: %d"), ServerInfo.ProcessId);
+		UE_LOG(LogAvaPlaybackClient, Display, TEXT("   - Content Path: \"%s\""), *ServerInfo.ProjectContentPath);
+		
+		for (const TPair<FString, FString>& UserData : ServerInfo.UserDataEntries)
+		{
+			UE_LOG(LogAvaPlaybackClient, Display, TEXT("   - User data \"%s\":\"%s\"."), *UserData.Key, *UserData.Value);
+		}
+		for (const TPair<FString, TUniquePtr<FBroadcastChannelInfo>>& ChannelInfo : ServerInfo.BroadcastChannelInfosByName)
+		{
+			check(ChannelInfo.Value.IsValid());
+			UE_LOG(LogAvaPlaybackClient, Display, TEXT("   - Channel(\"%s\") Channel State: \"%s\"."),
+				*ChannelInfo.Key,
+				*EnumToString(ChannelInfo.Value->ChannelState));
+			UE_LOG(LogAvaPlaybackClient, Display, TEXT("   - Channel(\"%s\") Channel Issue Severity: \"%s\"."),
+				*ChannelInfo.Key,
+				*EnumToString(ChannelInfo.Value->ChannelIssueSeverity));
+			for (const TPair<FGuid, FAvaMediaOutputStatus>& MediaOutputStatus : ChannelInfo.Value->MediaOutputStatuses)
+			{
+				UE_LOG(LogAvaPlaybackClient, Display, TEXT("   - Channel(\"%s\") Media Output \"%s\" State: \"%s\"."),
+					*ChannelInfo.Key, *MediaOutputStatus.Key.ToString(),
+					*EnumToString(MediaOutputStatus.Value.MediaOutputState));
+				UE_LOG(LogAvaPlaybackClient, Display, TEXT("   - Channel(\"%s\") Media Output \"%s\" Issue Severity: \"%s\"."),
+					*ChannelInfo.Key, *MediaOutputStatus.Key.ToString(),
+					*EnumToString(MediaOutputStatus.Value.MediaIssueSeverity));
+				for (const FString& MediaIssueMessage : MediaOutputStatus.Value.MediaIssueMessages)
+				{
+					UE_LOG(LogAvaPlaybackClient, Display, TEXT("   - Channel(\"%s\") Media Output \"%s\" Message: \"%s\"."),
+						*ChannelInfo.Key, *MediaOutputStatus.Key.ToString(), *MediaIssueMessage);
+				}
+			}
+		}
+		for (const TPair<FString, TUniquePtr<FPlaybackChannelInfo>>& ChannelInfoEntry : ServerInfo.PlaybackChannelInfosByName)
+		{
+			check(ChannelInfoEntry.Value.IsValid());
+			const FPlaybackChannelInfo& ChannelInfo = *ChannelInfoEntry.Value;
+			for (const TPair<FSoftObjectPath, TUniquePtr<FPlaybackAssetInfo>>& AssetInfo : ChannelInfo.AssetInfoByPaths)
+			{
+				// Dump per instance statuses.
+				for (const TPair<FGuid, FPlaybackInstanceInfo>& InstanceInfo : AssetInfo.Value->InstanceByIds)
+				{
+					// Check if we have user data.
+					FString PrettyPrintUserData = InstanceInfo.Value.UserData.IsSet() ?
+						FString::Printf(TEXT(" - UserData: %s"), *InstanceInfo.Value.UserData.GetValue()) : FString();
+					
+					UE_LOG(LogAvaPlaybackClient, Display, TEXT("   - Channel(\"%s\") Playback \"%s\"(instance \"%s\") Status:\"%s\"%s."),
+						*ChannelInfoEntry.Key,
+						*AssetInfo.Key.ToString(),
+						*InstanceInfo.Key.ToString(),
+						*EnumToString(InstanceInfo.Value.Status),
+						*PrettyPrintUserData);
+				}
+			}
+		}
+		for (const TPair<FSoftObjectPath, EAvaMediaPlaybackAssetStatus>& AssetStatus : ServerInfo.PlaybackAssetStatuses)
+		{
+			UE_LOG(LogAvaPlaybackClient, Display, TEXT("   - Asset (\"%s\") Status \"%s\"."),
+				*AssetStatus.Key.ToString(),
+				*EnumToString(AssetStatus.Value));
+		}
+	}
+}
+
+void FAvaMediaPlaybackClient::OnBroadcastChanged(EAvaBroadcastChange InChange)
+{
+	if (EnumHasAnyFlags(InChange, EAvaBroadcastChange::CurrentProfile))
+	{
+		// Propagate the new output configuration to the connected servers of this channel.
+		// We can do this because the profile can only be switched if channels are idle.
+		FAvaBroadcastProfile& CurrentProfile = UAvalancheBroadcast::Get().GetCurrentProfile();
+		const TArray<FAvaOutputChannel*>& Channels = CurrentProfile.GetChannels();
+		for (const FAvaOutputChannel* Channel : Channels)
+		{
+			TArray<UMediaOutput*> RemoteOutputs = Channel->GetRemoteMediaOutputs();
+			if (!RemoteOutputs.IsEmpty())
+			{
+				const FString ProfileName = CurrentProfile.GetName().ToString();
+				RequestBroadcast(ProfileName, Channel->GetChannelName(), RemoteOutputs, EAvaMediaBroadcastAction::UpdateConfig);
+			}
+		}
+	}
+}
+
+void FAvaMediaPlaybackClient::OnChannelChanged(const FAvaOutputChannel& InChannel, EAvaChannelChange InChange)
+{
+	if (EnumHasAnyFlags(InChange, EAvaChannelChange::Settings))
+	{
+		SendBroadcastChannelSettingsUpdate(GetServerAddressesForChannel(InChannel.GetChannelName()), InChannel);
+	}
+}
+
+void FAvaMediaPlaybackClient::OnAvaMediaSettingsChanged(UObject*, struct FPropertyChangedEvent&)
+{
+	ApplyAvaMediaSettings();
+	SendBroadcastSettingsUpdate(AllServerAddresses);
+	SendAvalancheInstanceSettingsUpdate(AllServerAddresses);
+}
+
+void FAvaMediaPlaybackClient::OnPreSavePackage(UPackage* InPackage, FObjectPreSaveContext InObjectSaveContext)
+{
+	// Only execute if this is a user save
+	if (InObjectSaveContext.IsProceduralSave())
+	{
+		return;
+	}
+
+	// Early return is no servers are connected.
+	if (Servers.IsEmpty())
+	{
+		return;
+	}
+
+	// Propagate this to inform local servers to flush loaders (and unlock the files).
+	// Remark: The server may already flush the playback asset's package loaders (if enabled) but this path will
+	// cover any other assets (not playback) that might have been loaded by both the server and client. 
+	SendPackageEvent(AllServerAddresses, InPackage->GetFName(), EAvaMediaPlaybackPackageEvent::PreSave);
+}
+
+void FAvaMediaPlaybackClient::OnPackageSaved(const FString& InPackageFileName, UPackage* InPackage, FObjectPostSaveContext InObjectSaveContext)
+{
+	// Only execute if this is a user save
+	if (InObjectSaveContext.IsProceduralSave())
+	{
+		return;
+	}
+
+	// Early return is no servers are connected.
+	if (Servers.IsEmpty())
+	{
+		return;
+	}
+
+	// Propagate this to inform local servers to reload the package.
+	SendPackageEvent(AllServerAddresses, InPackage->GetFName(), EAvaMediaPlaybackPackageEvent::PostSave);
+	
+	TArray<FAssetData> AssetsInPackage;
+	if (const IAssetRegistry* AssetRegistry = IAssetRegistry::Get())
+	{
+		AssetsInPackage.Reserve(16);
+		AssetRegistry->GetAssetsByPackageName(InPackage->GetFName(), AssetsInPackage);
+		
+		if (AssetsInPackage.IsEmpty())
+		{
+			UE_LOG(LogAvaPlaybackClient, Warning, TEXT("Asset Registry returns no asset data for package \"%s\"."), *InPackage->GetName());
+		}
+	}
+	
+	// If a local package (of a playback asset) was saved, we will refresh the status of the assets with the servers.
+	ForAllServers([&AssetsInPackage, this](const FServerInfo& InServerInfo)
+	{
+		for (const FAssetData& Assets : AssetsInPackage)
+		{
+			if (InServerInfo.PlaybackAssetStatuses.Contains(Assets.ToSoftObjectPath()))
+			{
+				// Note: we force a refresh (i.e. disregard the server's cached value) because the asset has changed on client side.
+				RequestPlaybackAssetStatusForServer(Assets.ToSoftObjectPath(), InServerInfo.ServerName, true);
+			}
+		}
+	});
+}
+
+void FAvaMediaPlaybackClient::OnAssetRemoved(const FAssetData& InAssetData)
+{
+	// Early return is no servers are connected.
+	if (Servers.IsEmpty())
+	{
+		return;
+	}
+
+	// Propagate this to inform local servers to purge the package.
+	SendPackageEvent(AllServerAddresses, InAssetData.PackageName, EAvaMediaPlaybackPackageEvent::AssetDeleted);
+}
+
+void FAvaMediaPlaybackClient::ApplyAvaMediaSettings()
+{
+	const UAvalancheMediaSettings& Settings = UAvalancheMediaSettings::Get();
+#if !NO_LOGGING
+	if (Settings.bVerbosePlaybackClientLogging)
+	{
+		LogAvaPlaybackClient.SetVerbosity(ELogVerbosity::Verbose);
+	}
+	else
+	{
+		LogAvaPlaybackClient.SetVerbosity(ELogVerbosity::Log);
+	}
+#endif
+}
+
+void FAvaMediaPlaybackClient::PublishPlaybackPing(const FDateTime& InCurrentTime, bool bInAutoPing)
+{
+	if (MessageEndpoint.IsValid())
+	{
+		const UAvalancheMediaSettings& AvalancheMediaSettings = UAvalancheMediaSettings::Get();
+		// Add a timeout to all known servers
+		const FDateTime Timeout = InCurrentTime + FTimespan::FromSeconds(AvalancheMediaSettings.PingTimeoutInterval);
+		ForAllServers([&Timeout](FServerInfo& InServerInfo)
+		{
+			InServerInfo.AddTimeout(Timeout);
+		});
+
+		FAvaMediaPlaybackPing* PlaybackPingMessage = FMessageEndpoint::MakeMessage<FAvaMediaPlaybackPing>();
+		PlaybackPingMessage->bAutoPing = bInAutoPing;
+		PlaybackPingMessage->PingIntervalSeconds = AvalancheMediaSettings.PingInterval;
+		PublishRequest(PlaybackPingMessage);
+	}
+}
+
+void FAvaMediaPlaybackClient::SendUserDataUpdate(const TArray<FMessageAddress>& InRecipients)
+{
+	FAvaMediaUpdateClientUserData* UserDataUpdate = FMessageEndpoint::MakeMessage<FAvaMediaUpdateClientUserData>();
+	UserDataUpdate->UserDataEntries = UserDataEntries;
+	SendRequest(UserDataUpdate, InRecipients, EMessageFlags::Reliable);
+}
+
+void FAvaMediaPlaybackClient::SendBroadcastSettingsUpdate(const TArray<FMessageAddress>& InRecipients)
+{
+	const IAvalancheBroadcastSettings& BroadcastSettings = IAvaMediaModule::Get().GetBroadcastSettings();
+	FAvaMediaBroadcastSettingsUpdate* BroadcastSettingsUpdate = FMessageEndpoint::MakeMessage<FAvaMediaBroadcastSettingsUpdate>();
+	BroadcastSettingsUpdate->BroadcastSettings.ChannelClearColor = BroadcastSettings.GetChannelClearColor();
+	BroadcastSettingsUpdate->BroadcastSettings.ChannelDefaultPixelFormat = BroadcastSettings.GetDefaultPixelFormat();
+	BroadcastSettingsUpdate->BroadcastSettings.ChannelDefaultResolution = BroadcastSettings.GetDefaultResolution();
+	BroadcastSettingsUpdate->BroadcastSettings.bDrawPlaceholderWidget = BroadcastSettings.IsDrawPlaceholderWidget();
+	BroadcastSettingsUpdate->BroadcastSettings.PlaceholderWidgetClass = BroadcastSettings.GetPlaceholderWidgetClass();
+	SendRequest(BroadcastSettingsUpdate, InRecipients, EMessageFlags::Reliable);
+}
+
+void FAvaMediaPlaybackClient::SendAvalancheInstanceSettingsUpdate(const TArray<FMessageAddress>& InRecipients)
+{
+	FAvalancheInstanceSettingsUpdate* InstanceSettingsUpdate = FMessageEndpoint::MakeMessage<FAvalancheInstanceSettingsUpdate>();
+	InstanceSettingsUpdate->InstanceSettings = UAvalancheMediaSettings::Get().AvalancheInstanceSettings;
+	SendRequest(InstanceSettingsUpdate, InRecipients, EMessageFlags::Reliable);
+}
+
+void FAvaMediaPlaybackClient::SendBroadcastChannelSettingsUpdate(const TArray<FMessageAddress>& InRecipients, const FAvaOutputChannel& InChannel)
+{
+	FAvaMediaBroadcastChannelSettingsUpdate* ChannelSettingsUpdate = FMessageEndpoint::MakeMessage<FAvaMediaBroadcastChannelSettingsUpdate>();
+	ChannelSettingsUpdate->Profile = InChannel.GetProfileName().ToString();
+	ChannelSettingsUpdate->Channel = InChannel.GetChannelName().ToString();
+	ChannelSettingsUpdate->QualitySettings = InChannel.GetViewportQualitySettings();
+	SendRequest(ChannelSettingsUpdate, InRecipients, EMessageFlags::Reliable);
+}
+
+void FAvaMediaPlaybackClient::SendPackageEvent(const TArray<FMessageAddress>& InRecipients, const FName& InPackageName, EAvaMediaPlaybackPackageEvent InEvent)
+{
+	FAvaMediaPlaybackPackageEvent* PackageEvent = FMessageEndpoint::MakeMessage<FAvaMediaPlaybackPackageEvent>();
+	PackageEvent->PackageName = InPackageName;
+	PackageEvent->Event = InEvent;
+	SendRequest(PackageEvent, InRecipients, EMessageFlags::Reliable);
+}
+
+void FAvaMediaPlaybackClient::SendStatCommand(const FString& InCommand, bool bInBroadcastLocalState, const TArray<FMessageAddress>& InRecipients)
+{
+	FAvaMediaStatCommand* Request = FMessageEndpoint::MakeMessage<FAvaMediaStatCommand>();
+	Request->Command = InCommand;
+	Request->bClientStateReliable = bInBroadcastLocalState;
+	Request->ClientEnabledRuntimeStats = IAvaModule::Get().GetEnabledRuntimeStats();
+	SendRequest(Request, InRecipients);
+}
+
+void FAvaMediaPlaybackClient::SendClientInfo(const FMessageAddress& InRecipient)
+{
+	FAvaMediaUpdateClientInfo* ClientInfo = FMessageEndpoint::MakeMessage<FAvaMediaUpdateClientInfo>();
+	ClientInfo->ComputerName = ComputerName;
+	ClientInfo->ProcessId = ProcessId;
+	ClientInfo->ProjectContentPath = ProjectContentPath;
+	SendRequest(ClientInfo, InRecipient);
+	
+	SendUserDataUpdate({InRecipient});
+	SendBroadcastSettingsUpdate({InRecipient});
+	SendAvalancheInstanceSettingsUpdate({InRecipient});
+	SendStatCommand(FString(), true, {InRecipient});	// Send empty stat command, will just send current states.
+}
+
+void FAvaMediaPlaybackClient::RemoveDeadServers(const FDateTime& InCurrentTime)
+{
+	bool bServerRemoved = false;
+
+	for (TMap<FString, TSharedPtr<FServerInfo>>::TIterator ServerIter = Servers.CreateIterator(); ServerIter; ++
+		 ServerIter)
+	{
+		check(ServerIter.Value().IsValid());
+		if (ServerIter.Value()->HasTimedOut(InCurrentTime))
+		{
+			UE_LOG(LogAvaPlaybackClient, Log, TEXT("Server \"%s\" is not longer responding to pings. Removing."),
+				   *ServerIter.Key());
+			const TSharedPtr<FServerInfo> RemovedServer = ServerIter.Value();
+			ServerIter.RemoveCurrent();
+			OnServerRemoved(*RemovedServer);
+			bServerRemoved = true;
+		}
+	}
+
+	if (bServerRemoved)
+	{
+		UpdateServerAddresses();
+	}
+}
+
+void FAvaMediaPlaybackClient::UpdateServerAddresses()
+{
+	AllServerAddresses.Empty(Servers.Num());
+	ForAllServers([this](const FServerInfo& InServerInfo)
+	{
+		AllServerAddresses.Add(InServerInfo.Address);
+	});
+}
+
+TArray<FMessageAddress> FAvaMediaPlaybackClient::GetServerAddressesForChannel(const FName& InChannelName) const
+{
+	if (!InChannelName.IsNone())
+	{
+		TArray<FString> ServerNames = GetServerNamesForChannel(InChannelName);
+		TArray<FMessageAddress> ServerAddresses;
+		ServerAddresses.Reserve(ServerNames.Num());
+		for (const FString& ServerName : ServerNames)
+		{
+			if (const FServerInfo* ServerInfo = GetServerInfo(ServerName))
+			{
+				ServerAddresses.Add(ServerInfo->Address);
+			}
+		}
+		return ServerAddresses;
+	}
+
+	// If no channel is specified, we will return all the servers that are connected.
+	return AllServerAddresses;
+}
+
+TArray<FString> FAvaMediaPlaybackClient::GetServerNamesForChannel(const FName& InChannelName) const
+{
+	TArray<FString> ServerNames;
+
+	if (!InChannelName.IsNone())
+	{
+		const FAvaOutputChannel& Channel = UAvalancheBroadcast::Get().GetCurrentProfile().GetChannel(InChannelName);
+		const TArray<UMediaOutput*>& RemoteOutputs = Channel.GetRemoteMediaOutputs();
+		ServerNames.Reserve(RemoteOutputs.Num());
+		for (const UMediaOutput* RemoteOutput : RemoteOutputs)
+		{
+			const FAvaMediaOutputInfo& OutputInfo = Channel.GetMediaOutputInfo(RemoteOutput);
+			if (OutputInfo.IsValid())
+			{
+				ServerNames.Add(OutputInfo.ServerName);
+			}
+			else
+			{
+				UE_LOG(LogAvaPlaybackClient, Warning, TEXT("MediaOutputInfo invalid for channel \"%s\"."),
+					   *InChannelName.ToString());
+
+				// Try to find the server name from the device name. The server has to be online for this to work.
+				const FString ServerName = GetServerNameForMediaOutputFallback(RemoteOutput);
+				if (!ServerName.IsEmpty())
+				{
+					ServerNames.Add(ServerName);
+				}
+				else
+				{
+					UE_LOG(LogAvaPlaybackClient, Error,
+						   TEXT("Unable to find server name for remote MediaOutput for channel \"%s\"."),
+						   *InChannelName.ToString());
+				}
+			}
+		}
+	}
+	else
+	{
+		// If no channel is specified, we will return all the servers that are connected.
+		ServerNames = GetServerNames();
+	}
+	return ServerNames;
+}
+
+FString FAvaMediaPlaybackClient::GetServerNameForMediaOutputFallback(const UMediaOutput* InMediaOutput) const
+{
+	const FString DeviceName = UE::AvaMediaOutputUtils::GetDeviceName(InMediaOutput);
+	if (!DeviceName.IsEmpty())
+	{
+		// See if it begins with any of the remote server names that are connected.
+		// See FAvaDeviceProviderData::ApplyServerName. All MediaOutput from replicated Device Provider
+		// Have the server name in the beginning of the device name.
+		for (const TPair<FString, TSharedPtr<FServerInfo>>& Server : Servers)
+		{
+			if (DeviceName.StartsWith(Server.Key))
+			{
+				return Server.Key;
+			}
+		}
+	}
+	return FString();
+}
+
+FAvaMediaPlaybackClient::FServerInfo& FAvaMediaPlaybackClient::GetOrCreateServerInfo(
+	const FString& InServerName, const FMessageAddress& InSenderAddress, bool* bOutCreated)
+{
+	if (FServerInfo* ServerInfo = GetServerInfo(InServerName))
+	{
+		if (ServerInfo->Address != InSenderAddress)
+		{
+			// This is suspicious though. It may also indicate a collision with multiple servers
+			// with the same name on the same computer host.
+			UE_LOG(LogAvaPlaybackClient, Warning, TEXT("Server \"%s\" Address changed."), *InServerName);
+			ServerInfo->Address = InSenderAddress;
+		}
+		return *ServerInfo;
+	}
+
+	const TSharedPtr<FServerInfo> ServerInfo = MakeShared<FServerInfo>();
+	ServerInfo->Address = InSenderAddress;
+	ServerInfo->ServerName = InServerName;
+	Servers.Add(InServerName, ServerInfo);
+	
+	UpdateServerAddresses();
+	OnServerAdded(*ServerInfo);
+
+	if (bOutCreated)
+	{
+		*bOutCreated = true;
+	}
+	return *ServerInfo;
+}
+
+void FAvaMediaPlaybackClient::ForAllServers(TFunctionRef<void(FServerInfo& /*InServerInfo*/)> InFunction)
+{
+	for (const TPair<FString, TSharedPtr<FServerInfo>>& ServerInfo : Servers)
+	{
+		check(ServerInfo.Value.IsValid());
+		InFunction(*ServerInfo.Value);
+	}
+}
+
+void FAvaMediaPlaybackClient::ForAllServers(TFunctionRef<void(const FServerInfo& /*InServerInfo*/)> InFunction) const
+{
+	for (const TPair<FString, TSharedPtr<FServerInfo>>& ServerInfo : Servers)
+	{
+		check(ServerInfo.Value.IsValid());
+		InFunction(*ServerInfo.Value);
+	}
+}
+
+void FAvaMediaPlaybackClient::OnServerAdded(const FServerInfo& InServerInfo)
+{
+	UE_LOG(LogAvaPlaybackClient, Log, TEXT("Registering new playback server \"%s\"."), *InServerInfo.ServerName);
+		
+	FAvaDeviceProviderDataRequest* DataRequest = FMessageEndpoint::MakeMessage<FAvaDeviceProviderDataRequest>();
+	SendRequest(DataRequest, InServerInfo.Address);
+	
+	SendClientInfo(InServerInfo.Address);
+}
+
+void FAvaMediaPlaybackClient::OnServerRemoved(const FServerInfo& InRemovedServer)
+{
+	if (ParentModule)
+	{
+		IAvaDeviceProviderProxyManager& Manager = ParentModule->GetDeviceProviderProxyManager();
+		Manager.Uninstall(InRemovedServer.ServerName);
+	}
+
+	// Update the status (i.e. go offline) of channels for this server.
+	TArray<FString> AffectedChannelNames;
+	for (const TPair<FString, TUniquePtr<FBroadcastChannelInfo>>& ChannelInfo : InRemovedServer.BroadcastChannelInfosByName)
+	{
+		check(ChannelInfo.Value.IsValid());
+		AffectedChannelNames.Add(ChannelInfo.Key);
+	}
+
+	FAvaBroadcastProfile& CurrentProfile = UAvalancheBroadcast::Get().GetCurrentProfile();
+	
+	// Try to reconcile channel state from remaining outputs (if any), i.e. channel is still partially online.
+	for (const FString& AffectedChannelName : AffectedChannelNames)
+	{
+		FAvaOutputChannel& Channel = CurrentProfile.GetChannelMutable(FName(AffectedChannelName));
+		if (Channel.IsValidChannel())
+		{
+			// Note: this will broadcast to delegates which may then request states of media outputs.
+			// So the back store needs to be updated before calling this.
+			Channel.RefreshState();
+		}
+	}
+
+	using namespace UE::AvaMediaPlaybackClient::Delegates;
+	GetOnConnectionEvent().Broadcast(*this, {InRemovedServer.ServerName, EConnectionEvent::ServerDisconnected});
+}
+
+const FAvaMediaPlaybackClient::FBroadcastChannelInfo* FAvaMediaPlaybackClient::GetChannelInfo(
+	const FString& InServerName, const FString& InChannelName) const
+{
+	if (!InServerName.IsEmpty())
+	{
+		if (const FServerInfo* ServerInfo = GetServerInfo(InServerName))
+		{
+			return ServerInfo->GetBroadcastChannelInfo(InChannelName);
+		}
+	}
+	else
+	{
+		// Legacy support: If the server name is not specified,
+		// Find the first server with the given channel name.
+		// Note: This doesn't work for "forked" channels.	
+		for (const TPair<FString, TSharedPtr<FServerInfo>>& Server : Servers)
+		{
+			check(Server.Value.IsValid());
+			if (const FBroadcastChannelInfo* ChannelInfo = Server.Value->GetBroadcastChannelInfo(InChannelName))
+			{
+				return ChannelInfo;
+			}
+		}
+	}
+	return nullptr;
+}
diff --git a/Engine/Plugins/Experimental/Avalanche/Source/AvalancheMedia/Private/Playback/AvaMediaPlaybackClientDelegates.cpp b/Engine/Plugins/Experimental/Avalanche/Source/AvalancheMedia/Private/Playback/AvaMediaPlaybackClientDelegates.cpp
new file mode 100644
--- /dev/null
+++ b/Engine/Plugins/Experimental/Avalanche/Source/AvalancheMedia/Private/Playback/AvaMediaPlaybackClientDelegates.cpp
@@ -0,0 +1,40 @@
+// Copyright Epic Games, Inc. All Rights Reserved.
+
+#include "Playback/AvaMediaPlaybackClientDelegates.h"
+
+namespace UE::AvaMediaPlaybackClient::Delegates
+{
+	namespace Private
+	{
+		static FOnConnectionEvent OnConnectionEvent;
+		static FOnPlaybackStatusChanged OnPlaybackStatusChanged;
+		static FOnPlaybackAssetStatusChanged OnPlaybackAssetStatusChanged;
+		static FOnPlaybackSequenceEvent OnPlaybackSequenceEvent;
+		static FOnPlaybackTransitionEvent OnPlaybackTransitionEvent;
+	}
+	
+	FOnConnectionEvent& GetOnConnectionEvent()
+	{
+		return Private::OnConnectionEvent;
+	}
+	
+	FOnPlaybackStatusChanged& GetOnPlaybackStatusChanged()
+	{
+		return Private::OnPlaybackStatusChanged;
+	}
+	
+	FOnPlaybackAssetStatusChanged& GetOnPlaybackAssetStatusChanged()
+	{
+		return Private::OnPlaybackAssetStatusChanged;	
+	}
+	
+	FOnPlaybackSequenceEvent& GetOnPlaybackSequenceEvent()
+	{
+		return Private::OnPlaybackSequenceEvent;
+	}
+
+	FOnPlaybackTransitionEvent& GetOnPlaybackTransitionEvent()
+	{
+		return Private::OnPlaybackTransitionEvent;
+	}
+}
