[MotionDesign]

diff --git a/Engine/Plugins/Experimental/Avalanche/Source/AvalancheMedia/Private/Playback/AvaMediaPlaybackManager.cpp b/Engine/Plugins/Experimental/Avalanche/Source/AvalancheMedia/Private/Playback/AvaMediaPlaybackManager.cpp
new file mode 100644
--- /dev/null
+++ b/Engine/Plugins/Experimental/Avalanche/Source/AvalancheMedia/Private/Playback/AvaMediaPlaybackManager.cpp
@@ -0,0 +1,1101 @@
+// Copyright Epic Games, Inc. All Rights Reserved.
+
+#include "Playback/AvaMediaPlaybackManager.h"
+
+#include "AssetRegistry/AssetData.h"
+#include "AssetRegistry/IAssetRegistry.h"
+#include "AvaBlueprint.h"
+#include "AvalancheBroadcast.h"
+#include "Engine/Engine.h"
+#include "Framework/AvalancheGameInstance.h"
+#include "IAvaMediaModule.h"
+#include "IAvaModule.h"
+#include "Playback/AvaMediaPlayableGroup.h"
+#include "Playback/AvaMediaPlayableGroupManager.h"
+#include "Playback/AvaMediaPlaybackUtils.h"
+#include "Playback/AvalanchePlayable.h"
+#include "Playback/Nodes/AvaPlaybackNodeBlueprintPlayer.h"
+#include "Playback/Nodes/AvaPlaybackNodeLevelPlayer.h"
+#include "Playback/Transition/AvaMediaPlaybackTransition.h"
+#include "UObject/ObjectSaveContext.h"
+#include "UObject/Package.h"
+#include "UnrealClient.h"
+
+DEFINE_LOG_CATEGORY_STATIC(LogAvaPlaybackManager, Log, All);
+
+namespace UE::AvaMediaPlaybackManager::Private
+{
+	inline bool ShouldForcePurgePackage(EAvaMediaPlaybackPackageEventFlags InEventFlags)
+	{
+		// There is one asset per package, if it is deleted, the package will be deleted too,
+		// eventually. When the event is received, the package is unlikely to be deleted yet.
+		// If the only cause for this event is an asset deletion, we will force unload the package
+		// even if the file may still be there.
+		return EnumHasAnyFlags(InEventFlags, EAvaMediaPlaybackPackageEventFlags::AssetDeleted)
+		 && !EnumHasAnyFlags(InEventFlags, EAvaMediaPlaybackPackageEventFlags::Saved);
+	}
+	
+	// Used to check validity of FAvaMediaPlaybackInstance's playback object to ensure
+	// it is not pending kill or destroyed (might happen in Engine Shutdown).
+	inline bool IsPlaybackValid(const TObjectPtr<UAvalanchePlayback>& InPlayback)
+	{
+		return IsValid(InPlayback.Get())
+				&& InPlayback->HasAnyFlags(RF_BeginDestroyed | RF_FinishDestroyed) == false;
+	}
+
+	static bool IsPackageIgnored(const FString& InPackageName)
+	{
+		// Filter out any references outside of /Game.
+		return !InPackageName.StartsWith(TEXT("/Game"));
+	}
+
+	/** Utility to walk a package's dependencies. */
+	struct FPackageDependencyWalker
+	{
+		int MaxRecursion;
+		bool bTrackMissingDependencies;
+		TSet<FName> AllDependencies;
+		TArray<FName> MissingDependencies;
+		
+		FPackageDependencyWalker(int InMaxRecursion, bool bInTrackMissingDependencies)
+			: MaxRecursion(InMaxRecursion)
+			, bTrackMissingDependencies(bInTrackMissingDependencies)
+		{
+			AllDependencies.Reserve(64);
+			if (bInTrackMissingDependencies)
+			{
+				MissingDependencies.Reserve(8);
+			}
+		}
+
+		/** Returns true if the asset has all of it's (non-ignored) dependencies. */
+		bool CheckDependencies(const IAssetRegistry& InAssetRegistry, const FName& InPackageName)
+		{
+			AllDependencies.Reset();
+			MissingDependencies.Reset();
+			return CheckDependencies(InAssetRegistry, InPackageName, 0);
+		}
+
+	private:
+		/** Returns true if the asset has all of it's (non-ignored) dependencies. */
+		bool CheckDependencies(const IAssetRegistry& InAssetRegistry, const FName& InPackageName, int InRecursion)
+		{
+			TArray<FName> Dependencies;
+			InAssetRegistry.GetDependencies(InPackageName, Dependencies);
+		
+			// We do a breath first pass, determine if any of the direct dependencies are missing.
+			for (const FName& Dependency : Dependencies)
+			{
+				if (!IsPackageIgnored(Dependency.ToString()) && !FPackageName::DoesPackageExist(Dependency.ToString()))
+				{
+					if (bTrackMissingDependencies)
+					{
+						MissingDependencies.AddUnique(Dependency);
+					}
+					
+					UE_LOG(LogAvaPlaybackManager, Verbose, TEXT("Package \"%s\" doesn't exist (dependency of \"%s\")."),
+						*Dependency.ToString(), *InPackageName.ToString());
+
+					return false;
+				}
+			}
+
+			if (InRecursion < MaxRecursion)
+			{
+				// Depth pass
+				for (const FName& Dependency : Dependencies)
+				{
+					// Note: we keep track of the dependencies we already visited and don't check them again.
+					// Not doing this can lead to infinite recursion because assets have circular dependencies (apparently).
+					if (!IsPackageIgnored(Dependency.ToString()) && !AllDependencies.Contains(Dependency))
+					{
+						AllDependencies.Add(Dependency);
+						if (!CheckDependencies(InAssetRegistry, Dependency, InRecursion + 1))
+						{
+							UE_LOG(LogAvaPlaybackManager, Verbose, TEXT("Package \"%s\" (dependency of \"%s\") has missing dependencies."),
+								*Dependency.ToString(), *InPackageName.ToString());
+							return false;
+						}
+					}
+				}
+			}
+			else
+			{
+				if (MaxRecursion != 0)
+				{
+					UE_LOG(LogAvaPlaybackManager, Warning, TEXT("Reached maximum recursion (%d) while evaluating dependencies on package \"%s\"."),
+						InRecursion, *InPackageName.ToString());
+				}
+			}
+			return true;
+		}
+	};
+}
+
+FAvaMediaPlaybackManager::FAvaMediaPlaybackManager()
+	: PlayableGroupManager(NewObject<UAvaMediaPlayableGroupManager>())
+{
+	PlayableGroupManager->Init();
+	UPackage::PackageSavedWithContextEvent.AddRaw(this, &FAvaMediaPlaybackManager::OnPackageSaved);
+	if (IAssetRegistry* AssetRegistry = IAssetRegistry::Get())
+	{
+		AssetRegistry->OnAssetRemoved().AddRaw(this, &FAvaMediaPlaybackManager::OnAssetRemoved);
+	}
+	IAvaMediaModule::Get().GetOnAvaMediaSyncPackageModified().AddRaw(this, &FAvaMediaPlaybackManager::OnAvaSyncPackageModified);
+	FCoreDelegates::OnEndFrame.AddRaw(this, &FAvaMediaPlaybackManager::Tick);
+}
+
+FAvaMediaPlaybackManager::~FAvaMediaPlaybackManager()
+{
+	UPackage::PackageSavedWithContextEvent.RemoveAll(this);
+
+	if (IAssetRegistry* AssetRegistry = IAssetRegistry::Get())
+	{
+		AssetRegistry->OnAssetRemoved().RemoveAll(this);
+	}
+	IAvaMediaModule::Get().GetOnAvaMediaSyncPackageModified().RemoveAll(this);
+	FCoreDelegates::OnEndFrame.RemoveAll(this);
+	
+	StopAllPlaybacks(true);
+	PlayableGroupManager->Shutdown();
+}
+
+void FAvaMediaPlaybackManager::Tick()
+{
+	double DeltaSeconds = FApp::GetDeltaTime();
+
+	OnBeginTick.Broadcast(DeltaSeconds);
+	
+	// Update the status of the playback objects, propagate any changed status to listeners.
+	ForAllPlaybackInstances([DeltaSeconds, this](FAvaMediaPlaybackInstance& InPlaybackInstance)
+	{
+		if (InPlaybackInstance.IsPlaying())
+		{
+			InPlaybackInstance.GetPlayback()->Tick(DeltaSeconds);
+		}
+		
+		if (InPlaybackInstance.UpdateStatus())
+		{
+			OnPlaybackInstanceStatusChanged.Broadcast(InPlaybackInstance);
+		}
+	});
+
+	// Tick the playable groups after the playback graphs so that playable groups are created for this tick.
+	if (GetPlayableGroupManager())
+	{
+		GetPlayableGroupManager()->Tick(DeltaSeconds);
+	}
+
+	if (PendingStartTransitions.Num())
+	{
+		TArray<TWeakObjectPtr<UAvaMediaPlaybackTransition>> StillPendingStartTransitions;
+		StillPendingStartTransitions.Reserve(PendingStartTransitions.Num());
+
+		for (TWeakObjectPtr<UAvaMediaPlaybackTransition>& TransitionToStart : PendingStartTransitions)
+		{
+			if (TransitionToStart.IsValid())
+			{
+				bool bShouldDiscard = false;
+				if (TransitionToStart->CanStart(bShouldDiscard))
+				{
+					TransitionToStart->Start();
+				}
+				else if (!bShouldDiscard)
+				{
+					StillPendingStartTransitions.Add(TransitionToStart);
+				}
+			}
+		}
+		PendingStartTransitions = StillPendingStartTransitions;
+	}
+	
+}
+
+TSharedPtr<FAvaMediaPlaybackInstance> FAvaMediaPlaybackManager::AcquirePlaybackInstance(const FSoftObjectPath& InAssetPath, const FString& InChannelName) const
+{
+	const TSharedPtr<FAvaMediaPlaybackSourceAssetEntry>* AssetEntry = PlaybackAssetEntries.Find(InAssetPath);
+	if (AssetEntry && *AssetEntry)
+	{
+		return (*AssetEntry)->AcquirePlaybackInstance(InChannelName);
+	}
+	return TSharedPtr<FAvaMediaPlaybackInstance>();
+}
+
+// Package name "/Game/AvaPlayback"
+// Asset name "AvaPlayback"
+TSharedPtr<FAvaMediaPlaybackInstance> FAvaMediaPlaybackManager::LoadPlaybackInstance(const FSoftObjectPath& InAssetPath, const FString& InChannelName)
+{
+	TSharedPtr<FAvaMediaPlaybackInstance> PlaybackInstance;
+	
+	if (UAvalanchePlayback* Playback = LoadPlaybackObject(InAssetPath, InChannelName))
+	{
+		FGuid NewInstanceId = FGuid::NewGuid();
+		
+		PlaybackInstance = MakeShared<FAvaMediaPlaybackInstance>(NewInstanceId, InAssetPath, InChannelName, Playback);
+		if (const TSharedPtr<FAvaMediaPlaybackSourceAssetEntry> AssetEntry = GetPlaybackAssetEntry(InAssetPath))
+		{
+			Playback->SetPlaybackManager(SharedThis(this));
+			PlaybackInstance->AssetEntryWeak = AssetEntry.ToWeakPtr();
+			AssetEntry->UsedInstances.Add(PlaybackInstance);
+		}
+	}
+	return PlaybackInstance;
+}
+
+TSharedPtr<FAvaMediaPlaybackInstance> FAvaMediaPlaybackManager::AcquireOrLoadPlaybackInstance(const FSoftObjectPath& InAssetPath, const FString& InChannelName)
+{
+	TSharedPtr<FAvaMediaPlaybackInstance> PlaybackInstance = AcquirePlaybackInstance(InAssetPath, InChannelName);
+	if (!PlaybackInstance)
+	{
+		PlaybackInstance = LoadPlaybackInstance(InAssetPath, InChannelName);
+	}
+	return PlaybackInstance;
+}
+
+TSharedPtr<FAvaMediaPlaybackInstance> FAvaMediaPlaybackManager::FindPlaybackInstance(const FGuid& InInstanceId, const FSoftObjectPath& InAssetPath, const FString& InChannelName) const
+{
+	if (const TSharedPtr<FAvaMediaPlaybackSourceAssetEntry> AssetEntry = FindPlaybackAssetEntry(InAssetPath))
+	{
+		if (InInstanceId.IsValid())
+		{
+			return AssetEntry->FindPlaybackInstance(InInstanceId);
+		}
+		// Only use the channel if the instance id is not specified.
+		return AssetEntry->FindPlaybackInstanceForChannel(InChannelName);
+	}
+	return TSharedPtr<FAvaMediaPlaybackInstance>();
+}
+
+bool FAvaMediaPlaybackManager::UnloadPlaybackInstances(const FSoftObjectPath& InAssetPath, const FString& InChannelName)
+{
+	bool bFoundInstance = false;
+	if (const TSharedPtr<FAvaMediaPlaybackSourceAssetEntry> AssetEntry = FindPlaybackAssetEntry(InAssetPath))
+	{
+		for (TArray<TSharedPtr<FAvaMediaPlaybackInstance>>::TIterator InstanceIt(AssetEntry->AvailableInstances); InstanceIt; ++InstanceIt)
+		{
+			if ( (*InstanceIt)->GetChannelName() == InChannelName || InChannelName.IsEmpty())
+			{
+				(*InstanceIt)->AssetEntryWeak.Reset();
+				(*InstanceIt)->Unload();
+				InstanceIt.RemoveCurrent();
+				bFoundInstance = true;
+			}
+		}
+	}
+	return bFoundInstance;
+}
+
+void FAvaMediaPlaybackManager::ForAllPlaybackInstances(TFunctionRef<void(FAvaMediaPlaybackInstance& /*InPlaybackInstance*/)> InFunction)
+{
+	using namespace UE::AvaMediaPlaybackManager::Private;
+	for (TPair<FSoftObjectPath, TSharedPtr<FAvaMediaPlaybackSourceAssetEntry>>& AssetEntry : PlaybackAssetEntries)
+	{
+		if (AssetEntry.Value)
+		{
+			AssetEntry.Value->ForAllPlaybackInstances(InFunction);
+		}
+	}
+}
+
+EAvaMediaPlaybackAssetStatus FAvaMediaPlaybackManager::GetLocalAssetStatus(const FName& InPackageName)
+{
+	// Fast check in cached results
+	if (const EAvaMediaPlaybackAssetStatus* CachedStatusEntry = CachedAssetStatus.Find(InPackageName))
+	{
+		return (*CachedStatusEntry);
+	}
+
+	EAvaMediaPlaybackAssetStatus AssetStatus = EAvaMediaPlaybackAssetStatus::Missing;
+	
+	if (FPackageName::DoesPackageExist(InPackageName.ToString()))
+	{
+		// Remark: Checking for dependencies might be optional. It is not particularly reliable.
+		if (const IAssetRegistry* AssetRegistry = IAssetRegistry::Get())
+		{
+			constexpr int MaxRecursion = 4;						// Todo: expose in config?
+			constexpr bool bTrackMissingDependencies = false;	// This is used for debugging only (for now).
+			UE::AvaMediaPlaybackManager::Private::FPackageDependencyWalker DependencyWalker(MaxRecursion, bTrackMissingDependencies);
+			const bool bHasAllDependencies = DependencyWalker.CheckDependencies(*AssetRegistry, InPackageName);
+			AssetStatus = bHasAllDependencies ? EAvaMediaPlaybackAssetStatus::Available : EAvaMediaPlaybackAssetStatus::MissingDependencies;
+		}
+		else
+		{
+			AssetStatus = EAvaMediaPlaybackAssetStatus::Available;
+		}
+	}
+
+	CachedAssetStatus.Add(InPackageName, AssetStatus);
+	return AssetStatus;
+}
+
+void FAvaMediaPlaybackManager::InvalidateCachedLocalAssetStatus(const FName& InPackageName)
+{
+	CachedAssetStatus.Remove(InPackageName);
+}
+
+void FAvaMediaPlaybackManager::InvalidatePlaybackAssetEntry(const FSoftObjectPath& InAssetPath)
+{
+	if (const TSharedPtr<FAvaMediaPlaybackSourceAssetEntry> AssetEntry = GetPlaybackAssetEntry(InAssetPath))
+	{
+		AssetEntry->ForAllPlaybackInstances([this](FAvaMediaPlaybackInstance& InPlaybackInstance)
+		{
+			OnPlaybackInstanceInvalidated.Broadcast(InPlaybackInstance);
+		});
+	}
+	PlaybackAssetEntries.Remove(InAssetPath);
+}
+
+UAvalanchePlayback*  FAvaMediaPlaybackManager::LoadPlaybackObject(const FSoftObjectPath& InAssetPath, const FString& InChannelName) const
+{
+	TRACE_CPUPROFILER_EVENT_SCOPE(FAvaMediaPlaybackManager::LoadPlaybackObject);
+	
+	const FString PackageName = InAssetPath.GetLongPackageName();
+	const FString AssetName = InAssetPath.GetAssetName();
+	
+	// First check if the package is already loaded.
+	UPackage* TempPackage = FindPackage(nullptr, *PackageName);
+
+	if (!TempPackage)
+	{
+		// Short cut: avoid sync package load for maps.
+		// The playback object will load it async using level streaming.
+		if (FAvaMediaPlaybackUtils::IsMapAsset(PackageName))
+		{
+			UAvalanchePlayback* const AvaPlayback = BuildPlaybackFromWorld(TSoftObjectPtr<UWorld>(InAssetPath), InChannelName);
+			check(AvaPlayback);
+			return AvaPlayback;
+		}
+
+		// Todo: Investigate LoadPackageAsync.
+		// For now, we tolerate a sync load here because there will be hitch from converting the
+		// avalanche blueprint to a world.
+		TempPackage = LoadPackage(nullptr, *PackageName, LOAD_None );
+	}
+	
+	if (TempPackage)
+	{
+		if (UObject* FoundObject = FindObject<UObject>(TempPackage, *AssetName))
+		{
+			// When the asset is an Avalanche Playback, it is loaded directly.
+			if (UAvalanchePlayback* const AvaPlayback = Cast<UAvalanchePlayback>(FoundObject))
+			{
+				return AvaPlayback;
+			}
+			
+			// When the asset is an Avalanche Blueprint, it is wrapped in a playback object.
+			if (const UAvalancheBlueprint* const AvaBlueprint = Cast<UAvalancheBlueprint>(FoundObject))
+			{
+				UAvalanchePlayback* const AvaPlayback = BuildPlaybackFromBlueprint(AvaBlueprint, InChannelName);
+				check(AvaPlayback);
+				return AvaPlayback;
+			}
+
+			if (const UWorld* const World = Cast<UWorld>(FoundObject))
+			{
+				UAvalanchePlayback* const AvaPlayback = BuildPlaybackFromWorld(World, InChannelName);
+				check(AvaPlayback);
+				return AvaPlayback;
+			}
+
+			UE_LOG(LogAvaPlaybackManager, Error,
+				TEXT("Asset \"%s\" in package \"%s\" is not a supported Avalanche playback asset (\"%s\")."),
+				*AssetName, *PackageName, *FoundObject->GetClass()->GetFullName());
+		}
+		else
+		{
+			UE_LOG(LogAvaPlaybackManager, Error,
+				TEXT("Failed to find asset \"%s\" in package \"%s\""),
+				*AssetName, *PackageName);
+		}
+	}
+	else
+	{
+		UE_LOG(LogAvaPlaybackManager, Error,
+			TEXT("Failed to load package \"%s\""),
+			*PackageName);
+	}
+	return nullptr;
+}
+
+UAvalanchePlayback* FAvaMediaPlaybackManager::BuildPlaybackFromBlueprint(const UAvalancheBlueprint* InBlueprint, const FString& InChannelName) const
+{
+	using namespace UE::AvaMediaPlaybackManager::Private;
+	FAvalanchePlaybackGraphBuilder GraphBuilder(GetPlayableGroupManager());
+
+	// Construct Player Node and assign the Blueprint that we passed in
+	UAvaPlaybackNodeBlueprintPlayer* const PlayerNode = GraphBuilder.ConstructPlaybackNode<UAvaPlaybackNodeBlueprintPlayer>();
+	PlayerNode->SetAvalancheAsset(InBlueprint);
+	
+	GraphBuilder.ConnectToRoot(InChannelName, PlayerNode);
+	return GraphBuilder.FinishBuilding();
+}
+
+UAvalanchePlayback* FAvaMediaPlaybackManager::BuildPlaybackFromWorld(const TSoftObjectPtr<UWorld>& InWorld, const FString& InChannelName) const
+{
+	using namespace UE::AvaMediaPlaybackManager::Private;
+	FAvalanchePlaybackGraphBuilder GraphBuilder(GetPlayableGroupManager());
+
+	// Construct Player Node and assign the Blueprint that we passed in
+	UAvaPlaybackNodeLevelPlayer* const PlayerNode = GraphBuilder.ConstructPlaybackNode<UAvaPlaybackNodeLevelPlayer>();
+	PlayerNode->SetAvalancheAsset(InWorld);
+	
+	GraphBuilder.ConnectToRoot(InChannelName, PlayerNode);
+	return GraphBuilder.FinishBuilding();
+}
+
+TArray<FSoftObjectPath> FAvaMediaPlaybackManager::StopAllPlaybacks(bool bInUnload)
+{
+	using namespace UE::AvaMediaPlaybackManager::Private;
+	TArray<FSoftObjectPath> StoppedPlaybackObjectPaths;
+	
+	if (UObjectInitialized())
+	{
+		const EAvaPlaybackStopOptions PlaybackStopOptions = GetPlaybackStopOptions(bInUnload);
+		const EAvaPlaybackUnloadOptions PlaybackUnloadOptions = GetPlaybackUnloadOptions();
+		for (TPair<FSoftObjectPath, TSharedPtr<FAvaMediaPlaybackSourceAssetEntry>> AssetEntry : PlaybackAssetEntries)
+		{
+			if (AssetEntry.Value)
+			{
+				StoppedPlaybackObjectPaths.Add(AssetEntry.Key);
+				AssetEntry.Value->ForAllPlaybackInstances([PlaybackStopOptions, PlaybackUnloadOptions, bInUnload](FAvaMediaPlaybackInstance& InPlaybackInstance)
+				{
+					if (InPlaybackInstance.IsPlaying())
+					{
+						InPlaybackInstance.Playback->Stop(PlaybackStopOptions);
+						InPlaybackInstance.Status = EAvaMediaPlaybackStatus::Loaded;
+					}
+					else if (bInUnload)
+					{
+						InPlaybackInstance.Playback->UnloadInstances(PlaybackUnloadOptions);
+						InPlaybackInstance.Status = EAvaMediaPlaybackStatus::Available;
+					}
+				});
+			}
+		}
+	}
+
+	if (bInUnload)
+	{
+		PlaybackAssetEntries.Empty();
+	}
+	return StoppedPlaybackObjectPaths;
+}
+
+bool FAvaMediaPlaybackManager::PushAnimationCommand(const FGuid& InInstanceId, const FSoftObjectPath& InSourcePath, const FString& InChannelName, EAvaMediaAnimAction InAction, const FAnimPlaySettings& InAnimSettings)
+{
+	bool bAnimationCommandPushed = false;
+	if (InInstanceId.IsValid())
+	{
+		// Check if we have the entry directly in the map. Playlist path.
+		if (const TSharedPtr<FAvaMediaPlaybackInstance> PlaybackInstance = FindPlaybackInstance(InInstanceId, InSourcePath, InChannelName))
+		{
+			PlaybackInstance->GetPlayback()->PushAnimationCommand(InSourcePath, InChannelName, InAction, InAnimSettings);
+			bAnimationCommandPushed = true;
+		}
+	}
+	else
+	{
+		// If we had a playback asset (not from the playlist path), check if it is present in one of the playback instances.
+		ForAllPlaybackInstances([InSourcePath, InChannelName, InAction, InAnimSettings, &bAnimationCommandPushed](FAvaMediaPlaybackInstance& InPlaybackInstance)
+		{
+			if (InPlaybackInstance.GetPlayback()->HasPlayerNodeForSourceAvalancheAsset(InSourcePath))
+			{
+				InPlaybackInstance.GetPlayback()->PushAnimationCommand(InSourcePath, InChannelName, InAction, InAnimSettings);
+				bAnimationCommandPushed = true;
+			}
+		});
+	}
+	
+	if (!bAnimationCommandPushed && bEnablePlaybackCommandsBuffering)
+	{
+		FString CommandBufferKey = MakeCommandBufferKey(InInstanceId, InSourcePath, InChannelName);
+		GetOrCreatePlaybackCommandBuffers(MoveTemp(CommandBufferKey)).AnimationCommands.Add({InAction, InAnimSettings});
+	}
+	return bAnimationCommandPushed;
+}
+
+bool FAvaMediaPlaybackManager::PushRemoteControlCommand(const FGuid& InInstanceId, const FSoftObjectPath& InSourcePath, const FString& InChannelName, const TSharedRef<FAvalancheRemoteControlValues>& InRemoteControlValues)
+{
+	bool bRemoteControlValuesPushed = false;
+	if (InInstanceId.IsValid())
+	{
+		// Check if we have the entry directly in the map. Playlist path.
+		if (const TSharedPtr<FAvaMediaPlaybackInstance> PlaybackInstance = FindPlaybackInstance(InInstanceId, InSourcePath, InChannelName))
+		{
+			PlaybackInstance->GetPlayback()->PushRemoteControlValues(InSourcePath, InChannelName, InRemoteControlValues);
+			bRemoteControlValuesPushed = true;
+		}
+	}
+	else
+	{
+		// If we had a playback asset (not from the playlist path), check if it is present in one of the playback instances.
+		ForAllPlaybackInstances([InSourcePath, InChannelName, InRemoteControlValues, &bRemoteControlValuesPushed](FAvaMediaPlaybackInstance& InPlaybackInstance)
+		{
+			if (InPlaybackInstance.GetPlayback()->HasPlayerNodeForSourceAvalancheAsset(InSourcePath))
+			{
+				InPlaybackInstance.GetPlayback()->PushRemoteControlValues(InSourcePath, InChannelName, InRemoteControlValues);
+				bRemoteControlValuesPushed = true;
+			}
+		});
+	}
+	
+	if (!bRemoteControlValuesPushed && bEnablePlaybackCommandsBuffering)
+	{
+		FString CommandBufferKey = MakeCommandBufferKey(InInstanceId, InSourcePath, InChannelName);
+		GetOrCreatePlaybackCommandBuffers(MoveTemp(CommandBufferKey)).RemoteControlCommands.Add({InRemoteControlValues});
+	}
+	return bRemoteControlValuesPushed;
+}
+
+bool FAvaMediaPlaybackManager::PushPlaybackTransitionStartCommand(UAvaMediaPlaybackTransition* InTransitionToStart)
+{
+	PendingStartTransitions.AddUnique(InTransitionToStart);
+	return true;
+}
+
+void FAvaMediaPlaybackManager::ApplyPendingCommands(UAvalanchePlayback* InPlaybackObject, const FGuid& InInstanceId, const FSoftObjectPath& InSourcePath, const FString& InChannelName)
+{
+	if (!InPlaybackObject || PlaybackObjectCommandBuffers.IsEmpty())
+	{
+		return;
+	}
+	
+	const FString CommandBufferKey = MakeCommandBufferKey(InInstanceId, InSourcePath, InChannelName);
+	
+	if (const FPlaybackObjectCommandBuffers* Commands = GetPlaybackCommandBuffers(CommandBufferKey))
+	{
+		for (const FPlaybackObjectCommandBuffers::FAnimationCommand& AnimCommand : Commands->AnimationCommands)
+		{
+			InPlaybackObject->PushAnimationCommand(InSourcePath, InChannelName, AnimCommand.AnimAction, AnimCommand.AnimPlaySettings);
+		}
+		for (const FPlaybackObjectCommandBuffers::FRemoteControlCommand& RCCommand : Commands->RemoteControlCommands)
+		{
+			InPlaybackObject->PushRemoteControlValues(InSourcePath, InChannelName, RCCommand.Values);
+		}
+		PlaybackObjectCommandBuffers.Remove(CommandBufferKey);
+	}
+}
+
+void FAvaMediaPlaybackManager::OnPackageModified(const FName& InPackageName, EAvaMediaPlaybackPackageEventFlags InFlags)
+{
+	InvalidateCachedLocalAssetStatus(InPackageName);
+	
+	TArray<UAvaMediaPlayableGroup*> PlayableGroupsToFlush;
+	
+	// Invalidate corresponding assets from that package.
+	for (TPair<FSoftObjectPath, TSharedPtr<FAvaMediaPlaybackSourceAssetEntry>>& AssetEntry : PlaybackAssetEntries)
+	{
+		if (AssetEntry.Value)
+		{
+			const FName SourcePackageFName = AssetEntry.Key.GetLongPackageFName();
+			if (SourcePackageFName == InPackageName)
+			{
+				AssetEntry.Value->ForAllPlaybackInstances([this, SourcePackageFName](FAvaMediaPlaybackInstance& InPlaybackInstance)
+				{
+					OnPlaybackInstanceInvalidated.Broadcast(InPlaybackInstance);
+					UE_LOG(LogAvaPlaybackManager, Log,
+						TEXT("Package \"%s\" being touched caused asset \"%s\" to be invalidated."),
+						*SourcePackageFName.ToString(), *InPlaybackInstance.SourcePath.ToString());
+
+					// Disconnect from the cache so it doesn't get recycled.
+					InPlaybackInstance.AssetEntryWeak.Reset();
+				});
+				
+				// Get rid of all "available" instances. Instances in use will not be flushed yet.
+				AssetEntry.Value->ForAllPlaybackInstances([this, SourcePackageFName, &PlayableGroupsToFlush](FAvaMediaPlaybackInstance& InPlaybackInstance)
+				{
+					InPlaybackInstance.AssetEntryWeak.Reset();
+					if (InPlaybackInstance.GetPlayback())
+					{
+						TArray<UAvalanchePlayable*> Playables;
+						InPlaybackInstance.GetPlayback()->GetAllPlayables(Playables);
+						for (const UAvalanchePlayable* Playable : Playables)
+						{
+							PlayableGroupsToFlush.AddUnique(Playable->GetPlayableGroup());
+						}
+					}
+					InPlaybackInstance.Unload();	// this should lead to calling UnloadAsset().
+				}, true, false);
+				AssetEntry.Value->AvailableInstances.Reset();
+			}
+		}
+	}
+
+	// The following steps are necessary to get rid of the level streaming in time
+	// for the corresponding level package to properly unload. In the "available" pool,
+	// playables and their group may not be ticking, thus the need to do this here.
+	for (const UAvaMediaPlayableGroup* PlayableGroup : PlayableGroupsToFlush)
+	{
+		// Unloading the playables should set the level streaming to "should unload", but
+		// we need to update parent world's level streaming in order for that to happen.
+		if (PlayableGroup->GetPlayWorld())
+		{
+			PlayableGroup->GetPlayWorld()->UpdateLevelStreaming();	// Attempt at getting the levels unloaded.
+		}
+	}
+
+	if (PlayableGroupsToFlush.Num() > 0)
+	{
+		// This will call FLevelStreamingGCHelper::PrepareStreamedOutLevelForGC to really get rid of the level instances.
+		CollectGarbage(GARBAGE_COLLECTION_KEEPFLAGS);
+	}
+
+	// Modified package by external process will need to be reloaded.
+	if (EnumHasAnyFlags(InFlags, EAvaMediaPlaybackPackageEventFlags::External))
+	{		
+		if (UPackage* ExistingPackage = FindPackage(nullptr, *InPackageName.ToString()))
+		{
+			// TODO: Investigate if it is better to batch the package purges/reloads to reduce overhead (gc once).
+			
+			using namespace UE::AvaMediaPlaybackManager::Private;
+			
+			// Check if the package was deleted, if not, reload it.
+			if (FAvaMediaPlaybackUtils::IsPackageDeleted(ExistingPackage) || ShouldForcePurgePackage(InFlags))
+			{
+				UE_LOG(LogAvaPlaybackManager, Log, TEXT("Loaded package \"%s\" being deleted on an external process requires a purge."), *InPackageName.ToString());
+				FAvaMediaPlaybackUtils::PurgePackages({ExistingPackage});
+			}
+			else
+			{
+				UE_LOG(LogAvaPlaybackManager, Log, TEXT("Loaded package \"%s\" being touched on an external process requires a reload."), *InPackageName.ToString());
+				FAvaMediaPlaybackUtils::ReloadPackages({ExistingPackage});
+			}
+		}
+	}	
+}
+
+void FAvaMediaPlaybackManager::OnParentWorldBeginTearDown()
+{
+	ForAllPlaybackInstances([](FAvaMediaPlaybackInstance& InPlaybackInstance)
+	{
+		if (InPlaybackInstance.Playback->IsPlaying())
+		{
+			InPlaybackInstance.Playback->Stop(EAvaPlaybackStopOptions::ForceImmediate | EAvaPlaybackStopOptions::Unload);
+		}
+		else
+		{
+			InPlaybackInstance.Playback->UnloadInstances(EAvaPlaybackUnloadOptions::ForceImmediate);
+		}
+	});
+	
+	PlaybackAssetEntries.Empty();
+	PlaybackObjectCommandBuffers.Empty();
+	PlayableGroupManager->Shutdown();
+}
+
+bool FAvaMediaPlaybackManager::HandleStatCommand(const TArray<FString>& InArgs)
+{
+	check(!InArgs.IsEmpty());
+
+	// We want to replicate the normal viewport client logic, only fallback to avalanche if nothing else is available.
+	FCommonViewportClient* ViewportClient = GStatProcessingViewportClient ? GStatProcessingViewportClient : GEngine->GameViewport;
+	// Use a delegate to fetch the editor viewport from the corresponding editor module (if present).
+	IAvaMediaModule::Get().GetEditorViewportClientDelegate().ExecuteIfBound(&ViewportClient);
+
+	bool bPreviousAvaModuleRuntimeStatProcessingEnabled = false;
+	
+	// Fallback to a viewport from Avalanche playback.
+	if (!ViewportClient)
+	{
+		bool bFound = false;
+		// Look in the current playback objects.
+		for (TPair<FSoftObjectPath, TSharedPtr<FAvaMediaPlaybackSourceAssetEntry>>& AssetEntry : PlaybackAssetEntries)
+		{
+			ForAllPlaybackInstances([&ViewportClient](FAvaMediaPlaybackInstance& InPlaybackInstance)
+			{
+				const TArray<UAvalancheGameInstance*> ActiveGameInstances = InPlaybackInstance.Playback->GetActiveGameInstances();
+				for (const UAvalancheGameInstance* GameInstance : ActiveGameInstances)
+				{
+					if (IsValid(GameInstance->GetWorld()) && IsValid(GameInstance->GetAvalancheGameViewportClient()))
+					{
+						ViewportClient = GameInstance->GetAvalancheGameViewportClient();
+						break;						
+					}
+				}
+			});
+			if (ViewportClient)
+			{
+				// Indicate that the stats are managed by the AvaModule since it is an Avalanche viewport client.
+				bPreviousAvaModuleRuntimeStatProcessingEnabled = IAvaModule::Get().SetRuntimeStatProcessingEnabled(true);
+				bFound = true;
+				break;
+			}
+		}
+		if (!bFound)
+		{
+			UE_LOG(LogAvaPlaybackManager, Warning, TEXT("No Active Avalanche Game Instances found to apply the stat '%s' command on."), *InArgs[0]);
+		}
+	}
+	
+	if (ViewportClient)
+	{
+		FCommonViewportClient* PreviousStatProcessingViewportClient = GStatProcessingViewportClient;
+		GStatProcessingViewportClient = ViewportClient;
+
+		FString StatCommand = TEXT("STAT ");
+		StatCommand += InArgs[0];
+		GEngine->Exec(ViewportClient->GetWorld(), *StatCommand, *GLog);
+
+		GStatProcessingViewportClient = PreviousStatProcessingViewportClient;
+		IAvaModule::Get().SetRuntimeStatProcessingEnabled(bPreviousAvaModuleRuntimeStatProcessingEnabled);
+		return true;
+	}
+
+	// This is going to work for simple stats, but not stat groups such as "detailed".
+	IAvaModule& AvaModule = IAvaModule::Get();
+	AvaModule.SetRuntimeStatEnabled(*InArgs[0], !AvaModule.IsRuntimeStatEnabled(InArgs[0]));
+
+	// Indicate we didn't find a viewport client to execute the stat command.
+	// The current enabled stats is thus probably not the desired state and deemed unreliable.
+	// If a connected server can run the command with an active viewport, the results will
+	// be back propagated to the client and will be used instead,
+	// see FAvaMediaPlaybackClient::HandleStatStatus.
+	return false;
+}
+
+bool FAvaMediaPlaybackManager::IsPlaybackAsset(const FAssetData& InAssetData)
+{
+	return FAvaMediaPlaybackUtils::IsPlayableAsset(InAssetData) || InAssetData.IsInstanceOf(StaticClass<UAvalanchePlayback>());
+}
+
+void FAvaMediaPlaybackManager::OnPackageSaved(const FString& InPackageFileName, UPackage* InPackage, FObjectPostSaveContext InObjectSaveContext)
+{
+	// Only execute if this is a user save
+	if (InObjectSaveContext.IsProceduralSave())
+	{
+		return;
+	}
+
+	OnPackageModified(InPackage->GetFName());
+}
+
+void FAvaMediaPlaybackManager::OnAvaSyncPackageModified(IAvaMediaSyncProvider* InAvaMediaSyncProvider, const FName& InPackageName)
+{
+	UE_LOG(LogAvaPlaybackManager, Verbose,
+		TEXT("A sync operation has touched the package \"%s\" on disk. Playback manager notified."),
+		*InPackageName.ToString());
+
+	OnPackageModified(InPackageName);
+}
+
+void FAvaMediaPlaybackManager::OnAssetRemoved(const FAssetData& InAssetData)
+{
+	// Invalidate the internal cache for the given package.
+	OnPackageModified(InAssetData.PackageName);
+
+	// If the asset removed is a playback asset, broadcast event.
+	if (IsPlaybackAsset(InAssetData))
+	{
+		OnLocalPlaybackAssetRemoved.Broadcast(InAssetData.ToSoftObjectPath());
+	}
+}
+
+FAvaMediaPlaybackManager::FPlaybackObjectCommandBuffers& FAvaMediaPlaybackManager::GetOrCreatePlaybackCommandBuffers(FString&& InPlaybackKey)
+{
+	if (FPlaybackObjectCommandBuffers* ObjectCommands = PlaybackObjectCommandBuffers.Find(InPlaybackKey))
+	{
+		return *ObjectCommands;
+	}
+	return PlaybackObjectCommandBuffers.Emplace(MoveTemp(InPlaybackKey));
+}
+
+const FAvaMediaPlaybackManager::FPlaybackObjectCommandBuffers* FAvaMediaPlaybackManager::GetPlaybackCommandBuffers(const FString& InPlaybackKey) const
+{
+	return PlaybackObjectCommandBuffers.Find(InPlaybackKey);
+}
+
+FAvaMediaPlaybackInstance::FAvaMediaPlaybackInstance(const FGuid& InInstanceId, const FSoftObjectPath& InSourcePath, const FString& InChannelName, UAvalanchePlayback* InPlayback)
+	: InstanceId(InInstanceId)
+	, ChannelName(InChannelName)
+	, ChannelFName(InChannelName)
+	, SourcePath(InSourcePath)
+	, Playback(InPlayback)
+	, Status(EAvaMediaPlaybackStatus::Loaded)
+{
+	if (Playback)
+	{
+		// Playables may not be created yet, in which case, InstanceId will be propagated by OnPlayableCreated. 
+		if (UAvalanchePlayable* Playable = Playback->FindPlayable(SourcePath, ChannelFName))
+		{
+			Playable->SetInstanceId(InstanceId);
+		}
+		Playback->OnPlayableCreated.AddRaw(this, &FAvaMediaPlaybackInstance::OnPlayableCreated);
+	}
+}
+
+FAvaMediaPlaybackInstance::~FAvaMediaPlaybackInstance()
+{
+	if (Playback)
+	{
+		Playback->OnPlayableCreated.RemoveAll(this);
+	}
+}
+
+void FAvaMediaPlaybackInstance::AddReferencedObjects( FReferenceCollector& Collector )
+{
+	Collector.AddReferencedObject( Playback );
+}
+
+FString FAvaMediaPlaybackInstance::GetReferencerName() const
+{
+	return TEXT("FAvaMediaPlaybackInstance");
+}
+
+void FAvaMediaPlaybackInstance::SetInstanceId(const FGuid& InInstanceId)
+{
+	InstanceId = InInstanceId;
+	if (Playback)
+	{
+		// Playables may not be created yet, in which case, InstanceId will be propagated by OnPlayableCreated. 
+		if (UAvalanchePlayable* Playable = Playback->FindPlayable(SourcePath, ChannelFName))
+		{
+			Playable->SetInstanceId(InstanceId);
+		}
+	}
+}
+
+void FAvaMediaPlaybackInstance::SetInstanceUserData(const FString& InUserData)
+{
+	InstanceUserData = InUserData;
+	if (Playback)
+	{
+		// Playables may not be created yet, in which case, InstanceUserData will be propagated by OnPlayableCreated. 
+		if (UAvalanchePlayable* Playable = Playback->FindPlayable(SourcePath, ChannelFName))
+		{
+			Playable->SetUserData(InUserData);
+		}
+	}
+}
+
+bool FAvaMediaPlaybackInstance::UpdateStatus()
+{
+	using namespace UE::AvaMediaPlaybackManager::Private;
+	if (IsPlaybackValid(Playback))
+	{
+		EAvaMediaPlaybackStatus NewStatus = Status;
+
+		if (Playback->IsPlaying())
+		{
+			if (const UAvalanchePlayable* Playable = Playback->FindPlayable(SourcePath, ChannelFName))
+			{
+				switch (Playable->GetPlayableStatus())
+				{
+				case EAvalanchePlayableStatus::Unloaded:
+					NewStatus = EAvaMediaPlaybackStatus::Available;
+					break;
+				case EAvalanchePlayableStatus::Loading:
+					NewStatus = EAvaMediaPlaybackStatus::Loading;
+					break;
+				case EAvalanchePlayableStatus::Loaded:
+					NewStatus = EAvaMediaPlaybackStatus::Loaded;
+					break;
+				case EAvalanchePlayableStatus::Visible:
+					NewStatus = Playable->GetPlayableGroup()->IsRenderTargetReady() ? EAvaMediaPlaybackStatus::Started : EAvaMediaPlaybackStatus::Starting;
+					break;
+				}
+			}
+			else
+			{
+				// The game instance may not be created yet. The RefreshPlayback is done on the next tick.
+				NewStatus = EAvaMediaPlaybackStatus::Loading;
+			}
+		}
+		else
+		{
+			// Even if not playing, we could have a game instance already, it can be preloaded now.
+			if (const UAvalanchePlayable* Playable = Playback->FindPlayable(SourcePath, ChannelFName))
+			{
+				switch (Playable->GetPlayableStatus())
+				{
+				case EAvalanchePlayableStatus::Unloaded:
+					NewStatus = EAvaMediaPlaybackStatus::Available;
+					break;
+				case EAvalanchePlayableStatus::Loading:
+					NewStatus = EAvaMediaPlaybackStatus::Loading;
+					break;
+				case EAvalanchePlayableStatus::Loaded:
+				case EAvalanchePlayableStatus::Visible:
+					NewStatus = EAvaMediaPlaybackStatus::Loaded;
+					break;
+				}
+			}
+			else
+			{
+				NewStatus = EAvaMediaPlaybackStatus::Loading;
+			}
+		}
+
+		if (NewStatus != Status)
+		{
+			Status = NewStatus;
+			return true;
+		}
+	}
+	return false;
+}
+
+void FAvaMediaPlaybackInstance::Unload()
+{
+	const TSharedPtr<FAvaMediaPlaybackManager> Manager = GetManager();
+
+	if (Playback->IsPlaying())
+	{
+		Playback->Stop(Manager ? Manager->GetPlaybackStopOptions(true) : EAvaPlaybackStopOptions::Unload);
+	}
+	else
+	{
+		Playback->UnloadInstances(Manager ? Manager->GetPlaybackUnloadOptions() : EAvaPlaybackUnloadOptions::None);
+	}
+
+	// Detach from cache.
+	if (const TSharedPtr<FAvaMediaPlaybackSourceAssetEntry> AssetEntry = AssetEntryWeak.Pin())
+	{
+		AssetEntry->DiscardInstance(this);
+	}
+}
+
+void FAvaMediaPlaybackInstance::Recycle()
+{
+	if (const TSharedPtr<FAvaMediaPlaybackSourceAssetEntry> AssetEntry = AssetEntryWeak.Pin())
+	{
+		AssetEntry->RecycleInstance(this);
+	}
+}
+
+
+TSharedPtr<FAvaMediaPlaybackManager> FAvaMediaPlaybackInstance::GetManager() const
+{
+	const TSharedPtr<FAvaMediaPlaybackSourceAssetEntry> AssetEntry = AssetEntryWeak.Pin();
+	return AssetEntry ? AssetEntry->GetManager() :  TSharedPtr<FAvaMediaPlaybackManager>();
+}
+
+void FAvaMediaPlaybackInstance::OnPlayableCreated(UAvalanchePlayback* InPlayback, UAvalanchePlayable* InPlayable)
+{
+	// Remark: The playable's source asset is not yet set at this point.
+	if (InPlayable)
+	{
+		InPlayable->SetInstanceId(InstanceId);
+		InPlayable->SetUserData(InstanceUserData);
+	}
+}
+
+TSharedPtr<FAvaMediaPlaybackInstance> FAvaMediaPlaybackSourceAssetEntry::FindPlaybackInstance(const FGuid& InInstanceId,
+	bool bInAvailableInstances, bool bInUsedInstances) const
+{
+	return FindPlaybackInstanceByPredicate([InInstanceId](const FAvaMediaPlaybackInstance& InPlaybackInstance) ->bool
+	{
+		return InPlaybackInstance.GetInstanceId() == InInstanceId;
+	}, bInAvailableInstances, bInUsedInstances);
+}
+
+TSharedPtr<FAvaMediaPlaybackInstance> FAvaMediaPlaybackSourceAssetEntry::FindPlaybackInstanceForChannel(const FString& InChannelName,
+	bool bInAvailableInstances, bool bInUsedInstances) const
+{
+	return FindPlaybackInstanceByPredicate([InChannelName](const FAvaMediaPlaybackInstance& InPlaybackInstance) ->bool
+	{
+		return InPlaybackInstance.GetChannelName() == InChannelName;
+	}, bInAvailableInstances, bInUsedInstances);
+}
+
+TSharedPtr<FAvaMediaPlaybackInstance> FAvaMediaPlaybackSourceAssetEntry::FindPlaybackInstanceByPredicate(TFunctionRef<bool(const FAvaMediaPlaybackInstance& /*InPlaybackInstance*/)> InPredicate,
+	bool bInAvailableInstances, bool bInUsedInstances) const
+{
+	using namespace UE::AvaMediaPlaybackManager::Private;
+	if (bInAvailableInstances)
+	{
+		for (const TSharedPtr<FAvaMediaPlaybackInstance>& PlaybackInstance : AvailableInstances)
+		{
+			if (PlaybackInstance && IsPlaybackValid(PlaybackInstance->GetPlayback()))
+			{
+				if (InPredicate(*PlaybackInstance))
+				{
+					return PlaybackInstance;
+				}
+			}
+		}
+	}
+	if (bInUsedInstances)
+	{
+		for (const TWeakPtr<FAvaMediaPlaybackInstance>& PlaybackInstanceWeak : UsedInstances)
+		{
+			TSharedPtr<FAvaMediaPlaybackInstance> PlaybackInstance = PlaybackInstanceWeak.Pin();
+			if (PlaybackInstance && IsPlaybackValid(PlaybackInstance->GetPlayback()))
+			{
+				if (InPredicate(*PlaybackInstance))
+				{
+					return PlaybackInstance;
+				}
+			}
+		}
+	}
+	return TSharedPtr<FAvaMediaPlaybackInstance>();
+}
+
+void FAvaMediaPlaybackSourceAssetEntry::ForAllPlaybackInstances(TFunctionRef<void(FAvaMediaPlaybackInstance& /*InPlaybackInstance*/)> InFunction,
+	bool bInAvailableInstances, bool bInUsedInstances)
+{
+	using namespace UE::AvaMediaPlaybackManager::Private;
+	if (bInAvailableInstances)
+	{
+		for (TSharedPtr<FAvaMediaPlaybackInstance>& PlaybackInstance : AvailableInstances)
+		{
+			if (PlaybackInstance && IsPlaybackValid(PlaybackInstance->GetPlayback()))
+			{
+				InFunction(*PlaybackInstance);
+			}
+		}
+	}
+	if (bInUsedInstances)
+	{
+		for (TWeakPtr<FAvaMediaPlaybackInstance>& PlaybackInstanceWeak : UsedInstances)
+		{
+			TSharedPtr<FAvaMediaPlaybackInstance> PlaybackInstance = PlaybackInstanceWeak.Pin();
+			if (PlaybackInstance && IsPlaybackValid(PlaybackInstance->GetPlayback()))
+			{
+				InFunction(*PlaybackInstance);
+			}
+		}
+	}
+}
+
+void FAvaMediaPlaybackSourceAssetEntry::DiscardInstance(FAvaMediaPlaybackInstance* InInstanceToRemove)
+{
+	{
+		const int32 Index = AvailableInstances.IndexOfByPredicate([InInstanceToRemove](const TSharedPtr<FAvaMediaPlaybackInstance>& InInstance)
+		{
+			return InInstance.Get() == InInstanceToRemove;
+		});
+		if (Index != INDEX_NONE)
+		{
+			AvailableInstances.RemoveAtSwap(Index);
+			return;
+		}
+	}
+	{
+		const int32 Index = UsedInstances.IndexOfByPredicate([InInstanceToRemove](const TWeakPtr<FAvaMediaPlaybackInstance>& InInstanceWeak)
+		{
+			const TSharedPtr<FAvaMediaPlaybackInstance> Instance = InInstanceWeak.Pin();
+			return Instance.Get() == InInstanceToRemove;
+		});
+		if (Index != INDEX_NONE)
+		{
+			UsedInstances.RemoveAtSwap(Index);
+		}
+	}
+}
+
+void FAvaMediaPlaybackSourceAssetEntry::RecycleInstance(FAvaMediaPlaybackInstance* InInstanceToRecycle)
+{
+	const int32 Index = UsedInstances.IndexOfByPredicate([InInstanceToRecycle](const TWeakPtr<FAvaMediaPlaybackInstance>& InInstanceWeak)
+	{
+		const TSharedPtr<FAvaMediaPlaybackInstance> Instance = InInstanceWeak.Pin();
+		return Instance.Get() == InInstanceToRecycle;
+	});
+	if (Index != INDEX_NONE)
+	{
+		const TSharedPtr<FAvaMediaPlaybackInstance> InstanceToRecycle = UsedInstances[Index].Pin();
+		UsedInstances.RemoveAtSwap(Index);
+		AvailableInstances.Add(InstanceToRecycle);
+	}
+}
+
