[MotionDesign]

diff --git a/Engine/Plugins/Experimental/Avalanche/Source/AvalancheMedia/Private/Playback/AvaMediaPlaybackServer.cpp b/Engine/Plugins/Experimental/Avalanche/Source/AvalancheMedia/Private/Playback/AvaMediaPlaybackServer.cpp
new file mode 100644
--- /dev/null
+++ b/Engine/Plugins/Experimental/Avalanche/Source/AvalancheMedia/Private/Playback/AvaMediaPlaybackServer.cpp
@@ -0,0 +1,1775 @@
+// Copyright Epic Games, Inc. All Rights Reserved.
+
+#include "Playback/AvaMediaPlaybackServer.h"
+
+#include "Async/Async.h"
+#include "AvaMediaMessageUtils.h"
+#include "AvaMediaSyncManager.h"
+#include "AvalancheBroadcast.h"
+#include "AvalancheMediaSettings.h"
+#include "IAvaModule.h"
+#include "MediaOutput.h"
+#include "MessageEndpointBuilder.h"
+#include "Misc/CoreDelegates.h"
+#include "Misc/Paths.h"
+#include "OutputDevices/AvaDeviceProviderData.h"
+#include "OutputDevices/AvaDeviceProviderProxy.h"
+#include "OutputDevices/AvaMediaOutputUtils.h"
+#include "Playback/AvaMediaPlaybackUtils.h"
+#include "Playback/AvalanchePlayable.h"
+#include "Playback/Transition/AvaMediaServerPlaybackTransition.h"
+
+DEFINE_LOG_CATEGORY(LogAvaPlaybackServer);
+
+namespace UE::AvaMediaPlaybackServer::Private
+{
+	FAvaMediaPlaybackStatus* MakePlaybackStatusMessage(const FGuid& InInstanceId, const FString& InChannelName,
+		const FSoftObjectPath& InAssetPath, EAvaMediaPlaybackStatus InStatus)
+	{
+		FAvaMediaPlaybackStatus* Message = FMessageEndpoint::MakeMessage<FAvaMediaPlaybackStatus>();
+		Message->InstanceId = InInstanceId;
+		Message->ChannelName = InChannelName;
+		Message->AssetPath = InAssetPath;
+		Message->Status = InStatus;
+		Message->bValidUserData = false;
+		return Message;
+	}
+	
+	FAvaMediaPlaybackStatus* MakePlaybackStatusMessage(const FGuid& InInstanceId, const FString& InChannelName,
+		const FSoftObjectPath& InAssetPath, EAvaMediaPlaybackStatus InStatus, const FString& InUserData)
+	{
+		FAvaMediaPlaybackStatus* Message = FMessageEndpoint::MakeMessage<FAvaMediaPlaybackStatus>();
+		Message->InstanceId = InInstanceId;
+		Message->ChannelName = InChannelName;
+		Message->AssetPath = InAssetPath;
+		Message->Status = InStatus;
+		Message->bValidUserData = true;
+		Message->UserData = InUserData;
+		return Message;
+	}
+}
+
+/**
+ * Container for the active playback instance transitions.
+ * We use composition for the FGCObject interface.
+ */
+class FAvaMediaPlaybackServer::FServerPlaybackInstanceTransitionCollection : public FGCObject
+{
+public:
+	FServerPlaybackInstanceTransitionCollection() = default;
+	virtual ~FServerPlaybackInstanceTransitionCollection() override = default;
+		
+	//~ Begin FGCObject
+	virtual void AddReferencedObjects( FReferenceCollector& Collector ) override
+	{
+		Collector.AddReferencedObjects(Transitions);
+	}
+		
+	virtual FString GetReferencerName() const override { return TEXT("FServerPlaybackInstanceTransitions"); }
+	//~ End FGCObject
+
+	TMap<FGuid, TObjectPtr<UAvaMediaServerPlaybackTransition>> Transitions;
+};
+
+FAvaMediaPlaybackServer::FAvaMediaPlaybackServer()
+	: Manager(MakeShared<FAvaMediaPlaybackManager>())
+	, PlaybackInstanceTransitions(MakeUnique<FServerPlaybackInstanceTransitionCollection>())
+{
+}
+
+FAvaMediaPlaybackServer::~FAvaMediaPlaybackServer()
+{
+	FAvaOutputChannel::GetOnMediaOutputStateChanged().RemoveAll(this);
+	FAvaOutputChannel::GetOnChannelChanged().RemoveAll(this);
+	UAvalanchePlayable::OnSequenceEvent().RemoveAll(this);
+
+	Manager->OnPlaybackInstanceInvalidated.RemoveAll(this);
+	Manager->OnPlaybackInstanceStatusChanged.RemoveAll(this);
+	Manager->OnLocalPlaybackAssetRemoved.RemoveAll(this);
+
+	FMessageEndpoint::SafeRelease(MessageEndpoint);
+	
+	StopPlaybacks();
+	
+	for (IConsoleObject* ConsoleCommand : ConsoleCommands)
+	{
+		IConsoleManager::Get().UnregisterConsoleObject(ConsoleCommand);
+	}
+	ConsoleCommands.Empty();
+
+	FCoreDelegates::OnEndFrame.RemoveAll(this);
+
+#if WITH_EDITOR
+	if (UObjectInitialized())
+	{
+		UAvalancheMediaSettings* AvaMediaSettings = GetMutableDefault<UAvalancheMediaSettings>();
+		AvaMediaSettings->OnSettingChanged().RemoveAll(this);
+	}
+#endif
+}
+
+void FAvaMediaPlaybackServer::Init(const FString& InAssignedServerName)
+{
+	Manager->SetEnablePlaybackCommandsBuffering(true);
+	Manager->OnPlaybackInstanceInvalidated.AddRaw(this, &FAvaMediaPlaybackServer::OnPlaybackInstanceInvalidated);
+	Manager->OnPlaybackInstanceStatusChanged.AddRaw(this, &FAvaMediaPlaybackServer::OnPlaybackInstanceStatusChanged);
+	Manager->OnLocalPlaybackAssetRemoved.AddRaw(this, &FAvaMediaPlaybackServer::OnPlaybackAssetRemoved);
+
+	ComputerName = FPlatformProcess::ComputerName();
+	ProjectContentPath = FPaths::ConvertRelativePathToFull(FPaths::ProjectContentDir());
+	ProcessId = FPlatformProcess::GetCurrentProcessId();
+	ServerName = InAssignedServerName.IsEmpty() ? ComputerName : InAssignedServerName;
+	RegisterCommands();
+	
+	// Create our end point. Note that these handlers are used by other services and the Context may not be valid
+	MessageEndpoint = FMessageEndpoint::Builder("AvaMediaPlaybackServer")
+	.Handling<FAvaMediaPlaybackPing>(this, &FAvaMediaPlaybackServer::HandlePlaybackPing)
+	.Handling<FAvaMediaUpdateClientUserData>(this, &FAvaMediaPlaybackServer::HandleUpdateClientUserData)
+	.Handling<FAvaMediaStatCommand>(this, &FAvaMediaPlaybackServer::HandleStatCommand)
+	.Handling<FAvaDeviceProviderDataRequest>(this, &FAvaMediaPlaybackServer::HandleDeviceProviderDataRequest)
+	.Handling<FAvaMediaUpdateClientInfo>(this, &FAvaMediaPlaybackServer::HandleUpdateClientInfo)
+	.Handling<FAvalancheInstanceSettingsUpdate>(this, &FAvaMediaPlaybackServer::HandleAvalancheInstanceSettingsUpdate)
+	.Handling<FAvaMediaPlaybackPackageEvent>(this, &FAvaMediaPlaybackServer::HandlePackageEvent)
+	.Handling<FAvaMediaPlaybackAssetStatusRequest>(this, &FAvaMediaPlaybackServer::HandlePlaybackAssetStatusRequest)
+	.Handling<FAvaMediaPlaybackRequest>(this, &FAvaMediaPlaybackServer::HandlePlaybackRequest)
+	.Handling<FAvaMediaAnimPlaybackRequest>(this, &FAvaMediaPlaybackServer::HandleAnimPlaybackRequest)
+	.Handling<FAvaMediaRemoteControlUpdateRequest>(this, &FAvaMediaPlaybackServer::HandleRemoteControlUpdateRequest)
+	.Handling<FAvaMediaTransitionStartRequest>(this, &FAvaMediaPlaybackServer::HandlePlayableTransitionStartRequest)
+	.Handling<FAvaMediaTransitionStopRequest>(this, &FAvaMediaPlaybackServer::HandlePlayableTransitionStopRequest)
+	.Handling<FAvaMediaBroadcastSettingsUpdate>(this, &FAvaMediaPlaybackServer::HandleBroadcastSettingsUpdate)
+	.Handling<FAvaMediaBroadcastRequest>(this, &FAvaMediaPlaybackServer::HandleBroadcastRequest)
+	.Handling<FAvaMediaBroadcastChannelSettingsUpdate>(this, &FAvaMediaPlaybackServer::HandleBroadcastChannelSettingsUpdate)
+	.Handling<FAvaMediaBroadcastStatusRequest>(this, &FAvaMediaPlaybackServer::HandleBroadcastStatusRequest);
+
+	if (MessageEndpoint.IsValid())
+	{
+		// Subscribe to the server listing requests
+		MessageEndpoint->Subscribe<FAvaMediaPlaybackPing>();
+
+		UE_LOG(LogAvaPlaybackServer, Log, TEXT("Motion Design Playback Server \"%s\" Started."), *ServerName);
+	}
+
+	// Prevent throttling and idling.	
+	if (IConsoleVariable* IdleWhenNotForeground = IConsoleManager::Get().FindConsoleVariable(TEXT("t.IdleWhenNotForeground")))
+	{
+		IdleWhenNotForeground->Set(0);
+	}
+
+	FCoreDelegates::OnEndFrame.AddSP(this, &FAvaMediaPlaybackServer::Tick);
+
+	{
+		FString LogReplicationVerbosity;
+		if (FParse::Value(FCommandLine::Get(),TEXT("AvaMediaPlaybackServerLogReplication="), LogReplicationVerbosity))
+		{
+			LogReplicationVerbosityFromCommandLine = ParseLogVerbosityFromString(LogReplicationVerbosity);
+		}
+	}
+
+#if WITH_EDITOR
+	UAvalancheMediaSettings* AvaMediaSettings = GetMutableDefault<UAvalancheMediaSettings>();
+	AvaMediaSettings->OnSettingChanged().AddSP(this, &FAvaMediaPlaybackServer::OnAvaMediaSettingsChanged);
+#endif
+	ApplyAvaMediaSettings();
+
+	FAvaOutputChannel::GetOnMediaOutputStateChanged().AddSP(this, &FAvaMediaPlaybackServer::OnMediaOutputStateChanged);
+	FAvaOutputChannel::GetOnChannelChanged().AddSP(this, &FAvaMediaPlaybackServer::OnChannelChanged);
+	UAvalanchePlayable::OnSequenceEvent().AddSP(this, &FAvaMediaPlaybackServer::OnPlayableSequenceEvent);
+}
+
+TArray<FAvaMediaPlaybackServer::FPlaybackInstanceReference> FAvaMediaPlaybackServer::StopPlaybacks(const FString& InChannelName, const FSoftObjectPath& InAssetPath, bool bInUnload)
+{
+	TArray<FPlaybackInstanceReference> StoppedPlaybackInstances;
+	
+	if (UObjectInitialized())
+	{
+		const EAvaPlaybackStopOptions PlaybackStopOptions = Manager->GetPlaybackStopOptions(bInUnload);
+		StoppedPlaybackInstances.Reserve(ActivePlaybackInstances.Num());
+		
+		for (const TPair<FGuid, TSharedPtr<FAvaMediaPlaybackInstance>>& PlaybackInstance : ActivePlaybackInstances)
+		{
+			if (!PlaybackInstance.Value)
+			{
+				continue;
+			}
+
+			// channel filtering.
+			if (!InChannelName.IsEmpty() && PlaybackInstance.Value->GetChannelName() != InChannelName)
+			{
+				continue;
+			}
+
+			// Asset path filtering
+			if (!InAssetPath.IsNull() && PlaybackInstance.Value->GetSourcePath() != InAssetPath)
+			{
+				continue;
+			}
+
+			// If we just stopping and the playback is already stopped, skip.
+			if (!bInUnload && PlaybackInstance.Value->IsPlaying() == false)
+			{
+				continue;
+			}
+			
+			StoppedPlaybackInstances.Add({PlaybackInstance.Key, PlaybackInstance.Value->GetSourcePath()});
+			PlaybackInstance.Value->GetPlayback()->Stop(PlaybackStopOptions);
+			PlaybackInstance.Value->SetStatus(EAvaMediaPlaybackStatus::Loaded);
+
+			if (bInUnload)
+			{
+				PlaybackInstance.Value->Unload();
+			}
+			else
+			{
+				PlaybackInstance.Value->Recycle();
+			}
+		}
+	}
+	
+	if (bInUnload)
+	{
+		for (const FPlaybackInstanceReference& StoppedInstance : StoppedPlaybackInstances)
+		{
+			ActivePlaybackInstances.Remove(StoppedInstance.Id);
+		}
+	}
+	return StoppedPlaybackInstances;
+}
+
+TArray<FAvaMediaPlaybackServer::FPlaybackInstanceReference> FAvaMediaPlaybackServer::StartPlaybacks()
+{
+	TArray<FPlaybackInstanceReference> Instances;
+	Instances.Reserve(ActivePlaybackInstances.Num());
+	
+	// Start all the loaded playback.
+	for (const TPair<FGuid, TSharedPtr<FAvaMediaPlaybackInstance>>& PlaybackInstance : ActivePlaybackInstances)
+	{
+		if (PlaybackInstance.Value->GetPlayback() && PlaybackInstance.Value->GetPlayback()->IsPlaying() == false)
+		{
+			Instances.Add({PlaybackInstance.Value->GetInstanceId(), PlaybackInstance.Value->GetSourcePath()});
+			PlaybackInstance.Value->GetPlayback()->Play();
+			PlaybackInstance.Value->SetStatus(EAvaMediaPlaybackStatus::Starting);
+		}
+	}
+	return Instances;
+}
+
+TArray<FString> FAvaMediaPlaybackServer::GetAllChannelsFromPlayingPlaybacks(const FSoftObjectPath& InAssetPath) const
+{
+	TSet<FString> Channels;
+	for (const TPair<FGuid, TSharedPtr<FAvaMediaPlaybackInstance>>& PlaybackInstance : ActivePlaybackInstances)
+	{
+		if (!PlaybackInstance.Value)
+		{
+			continue;
+		}
+
+		// filter with asset path.
+		if (!InAssetPath.IsNull() && PlaybackInstance.Value->GetSourcePath() == InAssetPath)
+		{
+			continue;
+		}
+
+		// Ignore stopped instances.
+		if (!PlaybackInstance.Value->IsPlaying())
+		{
+			continue;
+		}
+		
+		Channels.Add(PlaybackInstance.Value->GetChannelName());
+	}
+	return Channels.Array();
+}
+
+void FAvaMediaPlaybackServer::StartShuttingDown()
+{
+	Manager->StartShuttingDown();
+}
+
+void FAvaMediaPlaybackServer::StartBroadcast()
+{
+	UAvalancheBroadcast::Get().StartBroadcast();
+}
+
+void FAvaMediaPlaybackServer::StopBroadcast()
+{
+	UAvalancheBroadcast::Get().StopBroadcast();
+}
+
+const FString& FAvaMediaPlaybackServer::GetUserData(const FString& InKey) const
+{
+	if (const FString* Data = UserDataEntries.Find(InKey))
+	{
+		return *Data;
+	}
+	static FString EmptyString;
+	return EmptyString;
+}
+
+void FAvaMediaPlaybackServer::SetUserData(const FString& InKey, const FString& InData)
+{
+	UserDataEntries.Add(InKey, InData);
+	SendUserDataUpdate(GetAllClientAddresses());
+}
+
+void FAvaMediaPlaybackServer::RemoveUserData(const FString& InKey)
+{
+	UserDataEntries.Remove(InKey);
+	SendUserDataUpdate(GetAllClientAddresses());
+}
+
+TArray<FString> FAvaMediaPlaybackServer::GetClientNames() const
+{
+	TArray<FString> ClientNames;
+	ClientNames.Empty(Clients.Num());
+	for (const TPair<FString, TSharedPtr<FClientInfo>>& ClientInfo : Clients)
+	{
+		ClientNames.Add(ClientInfo.Value->ClientName);
+	}
+	return ClientNames;
+}
+
+FMessageAddress FAvaMediaPlaybackServer::GetClientAddress(const FString& InClientName) const
+{
+	if (const TSharedPtr<FClientInfo>* ClientInfo = Clients.Find(InClientName))
+	{
+		return (*ClientInfo)->Address;
+	}
+	FMessageAddress InvalidAddress;
+	InvalidAddress.Invalidate();
+	return InvalidAddress;
+}
+
+bool FAvaMediaPlaybackServer::HasClientUserData(const FString& InClientName, const FString& InKey) const
+{
+	if (const TSharedPtr<FClientInfo>* ClientInfo = Clients.Find(InClientName))
+	{
+		return (*ClientInfo)->UserDataEntries.Contains(InKey);
+	}
+	return false;
+}
+
+const FString& FAvaMediaPlaybackServer::GetClientUserData(const FString& InClientName, const FString& InKey) const
+{
+	if (const TSharedPtr<FClientInfo>* ClientInfo = Clients.Find(InClientName))
+	{
+		if (const FString* Data = (*ClientInfo)->UserDataEntries.Find(InKey))
+		{
+			return *Data;
+		}
+	}
+	static FString EmptyData;
+	return EmptyData;
+}
+
+const IAvalancheBroadcastSettings* FAvaMediaPlaybackServer::GetBroadcastSettings() const
+{
+	// Returns the first client we have.
+	// Todo: In case we have multiple clients, we will need a smarter way to handle this.
+	for (const TPair<FString, TSharedPtr<FClientInfo>>& Client : Clients)
+	{
+		return &Client.Value->BroadcastSettings;
+	}
+	return nullptr;
+}
+
+const FAvalancheInstanceSettings* FAvaMediaPlaybackServer::GetAvalancheInstanceSettings() const
+{
+	// Returns the first client we have.
+	// Todo: In case we have multiple clients, we will need a smarter way to handle this.
+	for (const TPair<FString, TSharedPtr<FClientInfo>>& Client : Clients)
+	{
+		return &Client.Value->AvalancheInstanceSettings;
+	}
+	return nullptr;
+}
+
+bool FAvaMediaPlaybackServer::RemovePlaybackInstanceTransition(const FGuid& InTransitionId)
+{
+	if (PlaybackInstanceTransitions)
+	{
+		return PlaybackInstanceTransitions->Transitions.Remove(InTransitionId) > 0;
+	}
+	return false;
+}
+
+void FAvaMediaPlaybackServer::SendPlayableTransitionEvent(
+	const FGuid& InTransitionId, const FGuid& InInstanceId, EAvalanchePlayableTransitionEventFlags InFlags,
+	const FName& InChannelName, const FString& InClientName)
+{
+	FAvaMediaPlayableTransitionEvent* Message = FMessageEndpoint::MakeMessage<FAvaMediaPlayableTransitionEvent>();
+	Message->ChannelName = InChannelName.ToString();
+	Message->TransitionId = InTransitionId;
+	Message->InstanceId = InInstanceId;
+	Message->SetEventFlags(InFlags);
+	SendResponse(Message, GetClientAddressSafe(InClientName));
+}
+
+void FAvaMediaPlaybackServer::HandlePlaybackPing(const FAvaMediaPlaybackPing& InMessage, const TSharedRef<IMessageContext, ESPMode::ThreadSafe>& InContext)
+{
+	if (!InMessage.bAutoPing)
+	{
+		UE_LOG(LogAvaPlaybackServer, Log, TEXT("Received Manual Ping from %s"), *InContext->GetSender().ToString());
+	}
+	
+	FClientInfo& ClientInfo = GetOrCreateClientInfo(InMessage.ClientName, InContext->GetSender());
+	ClientInfo.ResetPingTimeout();
+	// The client announce it's ping interval. We can expect a new ping around that interval. For tolerance
+	// we allow up to 3 ping intervals before declaring the client non responsive.
+	ClientInfo.AddTimeout(FDateTime::UtcNow() + FTimespan::FromSeconds(3 * InMessage.PingIntervalSeconds));
+
+	// Reply to the ping.
+	{
+		FAvaMediaPlaybackPong* ReplyMessage = FMessageEndpoint::MakeMessage<FAvaMediaPlaybackPong>();
+		ReplyMessage->bAutoPong = InMessage.bAutoPing;
+		ReplyMessage->bRequestClientInfo = !ClientInfo.bClientInfoReceived;
+		ReplyMessage->ProjectContentPath = ProjectContentPath;
+		ReplyMessage->ProcessId = ProcessId; 
+		SendResponse(ReplyMessage, InContext->GetSender());
+	}
+}
+
+void FAvaMediaPlaybackServer::HandleUpdateClientUserData(const FAvaMediaUpdateClientUserData& InMessage, const TSharedRef<IMessageContext, ESPMode::ThreadSafe>& InContext)
+{
+	FClientInfo& ClientInfo = GetOrCreateClientInfo(InMessage.ClientName, InContext->GetSender());
+	ClientInfo.ResetPingTimeout();
+	ClientInfo.UserDataEntries = InMessage.UserDataEntries;
+
+	// Logging when user data is updated (for debugging).
+	UE_LOG(LogAvaPlaybackServer, Verbose, TEXT("Received new user data for client \"%s\"."), *InMessage.ClientName);
+	for (const TPair<FString, FString>& UserData : ClientInfo.UserDataEntries)
+	{
+		UE_LOG(LogAvaPlaybackServer, Verbose, TEXT("User data \"%s\":\"%s\"."), *UserData.Key, *UserData.Value);
+	}
+}
+
+void FAvaMediaPlaybackServer::HandleStatCommand(const FAvaMediaStatCommand& InMessage, const TSharedRef<IMessageContext, ESPMode::ThreadSafe>& InContext)
+{
+	IAvaModule& AvaModule = IAvaModule::Get();
+
+	bool bLocalCommandSucceeded = false;
+
+	// Note: the client sends an empty command to sync it's state with server in connection handshake.
+	if (!InMessage.Command.IsEmpty())
+	{
+		bLocalCommandSucceeded = Manager->HandleStatCommand({InMessage.Command});
+	}
+
+	// If the enabled state from the client where reliable, we ensure
+	// that the server's state is sync'd to it.
+	if (InMessage.bClientStateReliable)
+	{
+		AvaModule.OverwriteEnabledRuntimeStats(InMessage.ClientEnabledRuntimeStats);
+	}
+
+	// The server replies with it's current status.
+	FAvaMediaStatStatus* ReplyMessage = FMessageEndpoint::MakeMessage<FAvaMediaStatStatus>();
+	ReplyMessage->bClientStateReliable = InMessage.bClientStateReliable;
+	ReplyMessage->bCommandSucceeded = bLocalCommandSucceeded;
+	ReplyMessage->EnabledRuntimeStats = AvaModule.GetEnabledRuntimeStats();
+	SendResponse(ReplyMessage, InContext->GetSender());
+}
+
+void FAvaMediaPlaybackServer::HandleDeviceProviderDataRequest(const FAvaDeviceProviderDataRequest& InMessage, const TSharedRef<IMessageContext, ESPMode::ThreadSafe>& InContext)
+{
+	FAvaDeviceProviderDataList* ReplyMessage = FMessageEndpoint::MakeMessage<FAvaDeviceProviderDataList>();
+	ReplyMessage->Populate(ServerName);
+	SendResponse(ReplyMessage, InContext->GetSender());
+}
+
+void FAvaMediaPlaybackServer::HandleUpdateClientInfo(const FAvaMediaUpdateClientInfo& InMessage, const TSharedRef<IMessageContext, ESPMode::ThreadSafe>& InContext)
+{
+	FClientInfo& ClientInfo = GetOrCreateClientInfo(InMessage.ClientName, InContext->GetSender());
+	ClientInfo.ComputerName = InMessage.ComputerName;
+	ClientInfo.ProjectContentPath = InMessage.ProjectContentPath;
+	ClientInfo.ProcessId = InMessage.ProcessId;
+	ClientInfo.bClientInfoReceived = true;
+
+	// Sync Manager is not needed if the server instance is a local instance from the same project directory.
+	const bool bShouldEnableSyncManager = !IsLocalClient(ClientInfo);
+	ClientInfo.MediaSyncManager->SetEnable(bShouldEnableSyncManager);
+}
+
+void FAvaMediaPlaybackServer::HandleAvalancheInstanceSettingsUpdate(const FAvalancheInstanceSettingsUpdate& InMessage, const TSharedRef<IMessageContext, ESPMode::ThreadSafe>& InContext)
+{
+	FClientInfo& ClientInfo = GetOrCreateClientInfo(InMessage.ClientName, InContext->GetSender());
+	ClientInfo.AvalancheInstanceSettings = InMessage.InstanceSettings;
+	
+	UE_LOG(LogAvaPlaybackServer, Verbose, TEXT("Received new instance settings from client \"%s\"."), *InMessage.ClientName);
+}
+
+void FAvaMediaPlaybackServer::HandlePackageEvent(const FAvaMediaPlaybackPackageEvent& InMessage, const TSharedRef<IMessageContext, ESPMode::ThreadSafe>& InContext)
+{
+	const FClientInfo* ClientInfo = GetClientInfo(InContext->GetSender());
+	
+	// Ignore package events if the client is not local.
+	// Also ignore package events if the client is in the same process.
+	if (ClientInfo && (!IsLocalClient(*ClientInfo) || IsClientOnLocalProcess(*ClientInfo)))
+	{
+		return;
+	}
+		
+	switch (InMessage.Event)
+	{
+	case EAvaMediaPlaybackPackageEvent::None:
+		break;
+	case EAvaMediaPlaybackPackageEvent::PostSave:
+		UE_LOG(LogAvaPlaybackServer, Verbose, TEXT("Client modified package \"%s\"."), *InMessage.PackageName.ToString());
+		// External event will trigger a reload.
+		Manager->OnPackageModified(InMessage.PackageName, EAvaMediaPlaybackPackageEventFlags::External | EAvaMediaPlaybackPackageEventFlags::Saved);
+		break;
+	case EAvaMediaPlaybackPackageEvent::PreSave:
+		// On demand flush package loading.
+		if (!UE::IsSavingPackage(nullptr) && !IsGarbageCollectingAndLockingUObjectHashTables())
+		{
+			if (UPackage* ExistingPackage = FindPackage(nullptr, *InMessage.PackageName.ToString()))
+			{
+				FAvaMediaPlaybackUtils::FlushPackageLoading(ExistingPackage);
+			}
+		}
+		break;
+	case EAvaMediaPlaybackPackageEvent::AssetDeleted:
+		UE_LOG(LogAvaPlaybackServer, Verbose, TEXT("Client deleted asset in package \"%s\"."), *InMessage.PackageName.ToString());
+		Manager->OnPackageModified(InMessage.PackageName, EAvaMediaPlaybackPackageEventFlags::External | EAvaMediaPlaybackPackageEventFlags::AssetDeleted);
+		break;
+	}
+}
+
+void FAvaMediaPlaybackServer::HandlePlaybackAssetStatusRequest(const FAvaMediaPlaybackAssetStatusRequest& InMessage, const TSharedRef<IMessageContext, ESPMode::ThreadSafe>& InContext)
+{
+	if (!InMessage.AssetPath.IsValid())
+	{
+		return;
+	}
+
+	EAvaMediaPlaybackAssetStatus PlaybackAssetStatus = EAvaMediaPlaybackAssetStatus::Missing;
+
+	if (Manager->IsLocalAssetAvailable(InMessage.AssetPath))
+	{
+		// Consider the asset available, unless we find that it is out of date below.
+		PlaybackAssetStatus = EAvaMediaPlaybackAssetStatus::Available;
+
+		// If there is a corresponding client info, we need to compare the asset to make sure it is up to date.
+		if (const FClientInfo* ClientInfo = GetClientInfo(InContext->GetSender()))
+		{
+			// Note: if the value is not available, we still send a reply with an "available" status, it is better than
+			// no status. The status will be updated again if OnAvaAssetSyncStatusReceived is called.
+			const TOptional<bool> NeedsSync = ClientInfo->MediaSyncManager->GetAssetSyncStatus(InMessage.AssetPath, InMessage.bForceRefresh);
+			if (NeedsSync.IsSet() && NeedsSync.GetValue())
+			{
+				PlaybackAssetStatus = EAvaMediaPlaybackAssetStatus::NeedsSync;
+			}
+		}
+	}
+	
+	SendPlaybackAssetStatus(InContext->GetSender(), InMessage.AssetPath, PlaybackAssetStatus);
+}
+
+void FAvaMediaPlaybackServer::HandlePlaybackRequest(const FAvaMediaPlaybackRequest& InMessage, const TSharedRef<IMessageContext, ESPMode::ThreadSafe>& InContext)
+{
+	for (const FAvaMediaPlaybackCommand& Command : InMessage.Commands)
+	{
+		PendingPlaybackCommands.Add({InContext->GetSender(), Command});
+	}
+}
+
+void FAvaMediaPlaybackServer::HandleAnimPlaybackRequest(const FAvaMediaAnimPlaybackRequest& InMessage, const TSharedRef<IMessageContext, ESPMode::ThreadSafe>& InContext)
+{
+	for (const FAnimPlaySettings& AnimSettings : InMessage.AnimPlaySettings)
+	{
+		Manager->PushAnimationCommand(InMessage.InstanceId, InMessage.AssetPath, InMessage.ChannelName, AnimSettings.Action, AnimSettings);
+	}
+
+	for (const FAvaMediaAnimActionInfo& ActionInfo : InMessage.AnimActionInfos)
+	{
+		FAnimPlaySettings AnimSettings;
+		AnimSettings.AnimationName = !ActionInfo.AnimationName.IsEmpty() ? *ActionInfo.AnimationName : FName();
+		Manager->PushAnimationCommand(InMessage.InstanceId, InMessage.AssetPath, InMessage.ChannelName, ActionInfo.AnimationAction, AnimSettings);
+	}
+}
+
+void FAvaMediaPlaybackServer::HandleRemoteControlUpdateRequest(const FAvaMediaRemoteControlUpdateRequest& InMessage, const TSharedRef<IMessageContext, ESPMode::ThreadSafe>& InContext)
+{
+	Manager->PushRemoteControlCommand(InMessage.InstanceId, InMessage.AssetPath, InMessage.ChannelName, MakeShared<FAvalancheRemoteControlValues>(InMessage.RemoteControlValues));
+}
+
+void FAvaMediaPlaybackServer::HandlePlayableTransitionStartRequest(const FAvaMediaTransitionStartRequest& InMessage, const TSharedRef<IMessageContext, ESPMode::ThreadSafe>& InContext)
+{
+	check(PlaybackInstanceTransitions);
+	
+	UAvaMediaServerPlaybackTransition* Transition = NewObject<UAvaMediaServerPlaybackTransition>();
+	Transition->SetTransitionId(InMessage.TransitionId);
+	Transition->SetChannelName(FName(InMessage.ChannelName));
+	Transition->SetClientName(GetClientNameSafe(InContext->GetSender()));
+	Transition->SetUnloadDiscardedInstances(InMessage.bUnloadDiscardedInstances);
+	Transition->SetTransitionFlags(InMessage.GetTransitionFlags());
+
+	// Enter Instances are likely not loaded yet.
+	Transition->SetEnterInstanceIds(InMessage.EnterInstanceIds);
+	Transition->SetEnterValues(InMessage.EnterValues);
+
+	// We can resolve the playing instances since they should be loaded.
+	for (const FGuid& PlayingInstanceId : InMessage.PlayingInstanceIds)
+	{
+		if (TSharedPtr<FAvaMediaPlaybackInstance> Instance = FindActivePlaybackInstance(PlayingInstanceId))
+		{
+			Transition->AddPlayingInstance(Instance);
+		}
+		else
+		{
+			UE_LOG(LogAvaPlaybackServer, Error,
+				TEXT("Transition \"%s\" from client \"%s\": \"Playing\" Instance Id \"%s\" was not found in active playback instances."),
+				*InMessage.TransitionId.ToString(), *GetClientNameSafe(InContext->GetSender()), *PlayingInstanceId.ToString());
+		}
+	}
+	
+	// We can resolve the exit instances since they should be loaded.
+	for (const FGuid& ExitInstanceId : InMessage.ExitInstanceIds)
+	{
+		if (TSharedPtr<FAvaMediaPlaybackInstance> Instance = FindActivePlaybackInstance(ExitInstanceId))
+		{
+			Transition->AddExitInstance(Instance);
+		}
+		else
+		{
+			UE_LOG(LogAvaPlaybackServer, Error,
+				TEXT("Transition \"%s\" from client \"%s\": \"Exit\" Instance Id \"%s\" was not found in active playback instances."),
+				*InMessage.TransitionId.ToString(), *GetClientNameSafe(InContext->GetSender()), *ExitInstanceId.ToString());
+		}
+	}
+	
+	PlaybackInstanceTransitions->Transitions.Add(InMessage.TransitionId, Transition);
+	GetPlaybackManager().PushPlaybackTransitionStartCommand(Transition);
+}
+
+void FAvaMediaPlaybackServer::HandlePlayableTransitionStopRequest(const FAvaMediaTransitionStopRequest& InMessage, const TSharedRef<IMessageContext, ESPMode::ThreadSafe>& InContext)
+{
+	check(PlaybackInstanceTransitions);
+	
+	if (const TObjectPtr<UAvaMediaServerPlaybackTransition>* FoundTransition = PlaybackInstanceTransitions->Transitions.Find(InMessage.TransitionId))
+	{
+		if (*FoundTransition)
+		{
+			(*FoundTransition)->Stop();
+			
+			// That should have removed the transition from the list.
+			if (PlaybackInstanceTransitions->Transitions.Contains(InMessage.TransitionId))
+			{
+				UE_LOG(LogAvaPlaybackServer, Error,
+					TEXT("Stopping Transition \"%s\" didn't remove it from the active list."), *InMessage.TransitionId.ToString());
+			}
+		}
+	}
+	else
+	{
+		UE_LOG(LogAvaPlaybackServer, Warning,
+			TEXT("Stop Transition request: Transition \"%s\" was not found in the active list."), *InMessage.TransitionId.ToString());
+	}
+}
+
+void FAvaMediaPlaybackServer::HandleBroadcastSettingsUpdate(const FAvaMediaBroadcastSettingsUpdate& InMessage, const TSharedRef<IMessageContext, ESPMode::ThreadSafe>& InContext)
+{
+	FClientInfo& ClientInfo = GetOrCreateClientInfo(InMessage.ClientName, InContext->GetSender());
+	ClientInfo.ResetPingTimeout();
+
+	// TODO: check if assets are required, if so request a sync. Need the DataSync API for this.
+	ClientInfo.BroadcastSettings.Settings = InMessage.BroadcastSettings;
+	
+	UE_LOG(LogAvaPlaybackServer, Verbose, TEXT("Received new broadcast settings from client \"%s\"."), *InMessage.ClientName);
+}
+
+void FAvaMediaPlaybackServer::HandleBroadcastChannelSettingsUpdate(const FAvaMediaBroadcastChannelSettingsUpdate& InMessage, const TSharedRef<IMessageContext, ESPMode::ThreadSafe>& InContext)
+{
+	const FName ChannelName = *InMessage.Channel;
+	if (!InMessage.Channel.IsEmpty())
+	{
+		FAvaOutputChannel& Channel = UAvalancheBroadcast::Get().GetCurrentProfile().GetOrAddChannel(ChannelName);
+		Channel.SetViewportQualitySettings(InMessage.QualitySettings);
+	}
+}
+
+void FAvaMediaPlaybackServer::HandleBroadcastRequest(const FAvaMediaBroadcastRequest& InMessage, const TSharedRef<IMessageContext, ESPMode::ThreadSafe>& InContext)
+{
+	const FName ChannelName = *InMessage.Channel;
+	if (InMessage.Action == EAvaMediaBroadcastAction::Start)
+	{
+		if (InMessage.Channel.IsEmpty())
+		{
+			StartBroadcast();
+		}
+		else
+		{
+			FAvaOutputChannel& Channel = UAvalancheBroadcast::Get().GetCurrentProfile().GetOrAddChannel(ChannelName);
+			if (UpdateChannelOutputConfig(Channel, InMessage.MediaOutputs, false))
+			{
+				Channel.StartChannelBroadcast();
+			}
+		}
+	}
+	else if (InMessage.Action == EAvaMediaBroadcastAction::UpdateConfig && !InMessage.Channel.IsEmpty())
+	{
+		FAvaOutputChannel& Channel = UAvalancheBroadcast::Get().GetCurrentProfile().GetOrAddChannel(ChannelName);
+		UpdateChannelOutputConfig(Channel, InMessage.MediaOutputs, true);
+	}
+	else if (InMessage.Action == EAvaMediaBroadcastAction::Stop)
+	{
+		if (ChannelName.IsNone())
+		{
+			StopBroadcast();
+		}
+		else
+		{
+			FAvaBroadcastProfile& Profile = UAvalancheBroadcast::Get().GetCurrentProfile();
+			FAvaOutputChannel& Channel = Profile.GetChannelMutable(ChannelName);
+			if (Channel.IsValidChannel())
+			{
+				Channel.StopChannelBroadcast();
+				SendChannelStatusUpdate(InMessage.Channel, Channel, InContext->GetSender());
+			}
+		}
+	}
+	else if (InMessage.Action == EAvaMediaBroadcastAction::DeleteChannel)
+	{
+		if (!InMessage.Channel.IsEmpty())
+		{
+			FAvaBroadcastProfile& Profile = UAvalancheBroadcast::Get().GetCurrentProfile();
+			if (Profile.RemoveChannel(FName(InMessage.Channel)))
+			{
+				SendAllChannelStatusUpdate(InContext->GetSender(), false);
+			}
+			else
+			{
+				// TODO: need to inform client that operation failed.
+				UE_LOG(LogAvaPlaybackServer, Error, TEXT("Failed to remove channel \"%s\"."), *InMessage.Channel);
+			}
+		}
+		else
+		{
+			UE_LOG(LogAvaPlaybackServer, Error, TEXT("Received a \"Delete Channel\" command with an empty channel name."));
+		}
+	}
+}
+
+void FAvaMediaPlaybackServer::HandleBroadcastStatusRequest(const FAvaMediaBroadcastStatusRequest& InMessage, const TSharedRef<IMessageContext, ESPMode::ThreadSafe>& InContext)
+{
+	UE_LOG(LogAvaPlaybackServer, Verbose, TEXT("Received a broadcast status request from \"%s\""), *InMessage.ClientName);
+	// Make sure all the channel status are refreshed
+	{
+		// Block channel status update while we refresh, we want to send one clean update at the end.
+		TGuardValue BlockChannelStatusUpdate(bBlockChannelStatusUpdate, true);
+		for (FAvaOutputChannel* Channel : UAvalancheBroadcast::Get().GetCurrentProfile().GetChannels())
+		{
+			Channel->RefreshState();
+		}
+	}
+	
+	SendAllChannelStatusUpdate(InContext->GetSender(), InMessage.bIncludeMediaOutputData);
+}
+
+FString FAvaMediaPlaybackServer::GetMessageEndpointAddressId() const
+{
+	return MessageEndpoint.IsValid() && MessageEndpoint->IsEnabled() ? MessageEndpoint->GetAddress().ToString() : TEXT("");
+}
+
+void FAvaMediaPlaybackServer::Tick()
+{
+	const FDateTime CurrentTime = FDateTime::UtcNow();
+	RemoveDeadClients(CurrentTime);
+
+	for (const TPair<FString, TSharedPtr<FClientInfo>>& Client : Clients)
+	{
+		if (Client.Value->MediaSyncManager.IsValid() && Client.Value->MediaSyncManager->IsEnabled())
+		{
+			Client.Value->MediaSyncManager->Tick();
+		}
+	}
+
+	// Execute the pending commands in batch for this tick.
+	ExecutePendingPlaybackCommands();
+
+	// Try to resolve the instance for loaded transitions.
+	for (const TPair<FGuid, TObjectPtr<UAvaMediaServerPlaybackTransition>>& Transition : PlaybackInstanceTransitions->Transitions)
+	{
+		Transition.Value->TryResolveInstances(*this);
+	}
+	
+	// TODO:
+	// - Check status of outputs and send to client(s). In particular, watch the media capture's transient states.
+	// - Send telemetry if client subscribed to the stream.
+}
+
+void FAvaMediaPlaybackServer::RegisterCommands()
+{
+	if (ConsoleCommands.Num() != 0)
+	{
+		return;
+	}
+		
+	ConsoleCommands.Add(IConsoleManager::Get().RegisterConsoleCommand(
+			TEXT("AvaMediaServer.StartPlayback"),
+			TEXT("Starts the playback of the given playback object."),
+			FConsoleCommandWithArgsDelegate::CreateRaw(this, &FAvaMediaPlaybackServer::StartPlaybackCommand),
+			ECVF_Default
+			));
+	ConsoleCommands.Add(IConsoleManager::Get().RegisterConsoleCommand(
+			TEXT("AvaMediaServer.StopPlayback"),
+			TEXT("Stops the playback of the given playback object."),
+			FConsoleCommandWithArgsDelegate::CreateRaw(this, &FAvaMediaPlaybackServer::StopPlaybackCommand),
+			ECVF_Default
+			));
+	ConsoleCommands.Add(IConsoleManager::Get().RegisterConsoleCommand(
+			TEXT("AvaMediaServer.StartBroacast"),
+			TEXT("Starts the broacast on specified (or all) channel(s)."),
+			FConsoleCommandWithArgsDelegate::CreateRaw(this, &FAvaMediaPlaybackServer::StartBroadcastCommand),
+			ECVF_Default
+			));
+	ConsoleCommands.Add(IConsoleManager::Get().RegisterConsoleCommand(
+			TEXT("AvaMediaServer.StopBroadcast"),
+			TEXT("Stops the broadcast of the specified (or all) channel(s)."),
+			FConsoleCommandWithArgsDelegate::CreateRaw(this, &FAvaMediaPlaybackServer::StopBroadcastCommand),
+			ECVF_Default
+			));
+	ConsoleCommands.Add(IConsoleManager::Get().RegisterConsoleCommand(
+			TEXT("AvaMediaServer.SetUserData"),
+			TEXT("Set Replicated User Data Entry (Key, Value)."),
+			FConsoleCommandWithArgsDelegate::CreateRaw(this, &FAvaMediaPlaybackServer::SetUserDataCommand),
+			ECVF_Default
+			));
+	ConsoleCommands.Add(IConsoleManager::Get().RegisterConsoleCommand(
+			TEXT("AvaMediaServer.Status"),
+			TEXT("Display current status of all server info."),
+			FConsoleCommandWithArgsDelegate::CreateRaw(this, &FAvaMediaPlaybackServer::ShowStatusCommand),
+			ECVF_Default
+			));
+}
+
+void FAvaMediaPlaybackServer::StartPlaybackCommand(const TArray<FString>& InArgs)
+{
+	if (InArgs.Num() >= 1)
+	{
+		//Concatenate all Args (starting from the 2) into one String with spaces in between each arg
+		FString ConcatenatedCommands;
+		for (int32 Index = 1; Index < InArgs.Num(); ++Index)
+		{
+			ConcatenatedCommands += InArgs[Index] + TEXT(" ");
+		}
+
+		FString ChannelName;
+		FParse::Value(*ConcatenatedCommands, TEXT("Channel="), ChannelName);
+
+		const FSoftObjectPath AssetPath(InArgs[0]);
+
+		// Uncached load.
+		if (UAvalanchePlayback* PlaybackObject = Manager->LoadPlaybackObject(AssetPath, ChannelName))
+		{
+			PlaybackObject->Play();
+			StartBroadcast();
+		}
+	}
+	else
+	{
+		UE_LOG(LogAvaPlaybackServer, Log, TEXT("Arguments: Package.AssetName. Ex: \"/Game/AvaPlayback.AvaPlayback\""));
+	}
+}
+
+void FAvaMediaPlaybackServer::StopPlaybackCommand(const TArray<FString>& InArgs)
+{
+	StopPlaybacks();
+}
+
+void FAvaMediaPlaybackServer::StartBroadcastCommand(const TArray<FString>& InArgs)
+{
+	StartBroadcast();
+}
+
+void FAvaMediaPlaybackServer::StopBroadcastCommand(const TArray<FString>& InArgs)
+{
+	StopBroadcast();
+}
+
+void FAvaMediaPlaybackServer::SetUserDataCommand(const TArray<FString>& InArgs)
+{
+	if (InArgs.Num() >= 2)
+	{
+		UE_LOG(LogAvaPlaybackServer, Log, TEXT("Setting User Data Key \"%s\" to Value: \"%s\"."), *InArgs[0], *InArgs[1]);
+		SetUserData(InArgs[0], InArgs[1]);
+	}
+	else if (InArgs.Num() == 1)
+	{
+		// One argument means to remove that user data entry.
+		if (HasUserData(InArgs[0]))
+		{
+			UE_LOG(LogAvaPlaybackServer, Log, TEXT("Removing User Data Key \"%s\"."), *InArgs[0]);
+			RemoveUserData(InArgs[0]);
+		}
+		else
+		{
+			UE_LOG(LogAvaPlaybackServer, Error, TEXT("User Data Key \"%s\" not found."), *InArgs[0]);
+		}
+	}
+}
+
+void FAvaMediaPlaybackServer::ShowStatusCommand(const TArray<FString>& InArgs)
+{
+	UE_LOG(LogAvaPlaybackServer, Display, TEXT("Playback Server: \"%s\""), *ServerName);
+	UE_LOG(LogAvaPlaybackServer, Display, TEXT("- Endpoint Bus Address: \"%s\""), MessageEndpoint.IsValid() ? *MessageEndpoint->GetAddress().ToString() : TEXT("Invalid"));
+	UE_LOG(LogAvaPlaybackServer, Display, TEXT("- Computer: \"%s\""), *ComputerName);
+	UE_LOG(LogAvaPlaybackServer, Display, TEXT("- ProcessId: %d"), ProcessId);
+	UE_LOG(LogAvaPlaybackServer, Display, TEXT("- Content Path: \"%s\""), *ProjectContentPath);
+
+	for (const TPair<FString, FString>& UserData : UserDataEntries)
+	{
+		UE_LOG(LogAvaPlaybackServer, Display, TEXT("- User data \"%s\":\"%s\"."), *UserData.Key, *UserData.Value);
+	}
+	
+	for (const TPair<FString, TSharedPtr<FClientInfo>>& Client : Clients)
+	{
+		const FClientInfo& ClientInfo = *Client.Value;
+		UE_LOG(LogAvaPlaybackServer, Display, TEXT("Connected Client: \"%s\""), *ClientInfo.ClientName);
+		UE_LOG(LogAvaPlaybackServer, Display, TEXT("   - Endpoint Bus Address: \"%s\""), *ClientInfo.Address.ToString());
+		UE_LOG(LogAvaPlaybackServer, Display, TEXT("   - Computer: \"%s\""), *ClientInfo.ComputerName);
+		UE_LOG(LogAvaPlaybackServer, Display, TEXT("   - ProcessId: %d"), ClientInfo.ProcessId);
+		UE_LOG(LogAvaPlaybackServer, Display, TEXT("   - Content Path: \"%s\""), *ClientInfo.ProjectContentPath);
+
+		for (const TPair<FString, FString>& UserData : ClientInfo.UserDataEntries)
+		{
+			UE_LOG(LogAvaPlaybackServer, Display, TEXT("   - User data \"%s\":\"%s\"."), *UserData.Key, *UserData.Value);
+		}
+		
+		UE_LOG(LogAvaPlaybackServer, Display, TEXT("   - BroadcastSettings.ChannelClearColor: (%f, %f, %f, %f)"),
+			ClientInfo.BroadcastSettings.Settings.ChannelClearColor.R,
+			ClientInfo.BroadcastSettings.Settings.ChannelClearColor.G,
+			ClientInfo.BroadcastSettings.Settings.ChannelClearColor.B,
+			ClientInfo.BroadcastSettings.Settings.ChannelClearColor.A);
+		UE_LOG(LogAvaPlaybackServer, Display, TEXT("   - BroadcastSettings.ChannelDefaultPixelFormat: (%s)"),
+			GetPixelFormatString(ClientInfo.BroadcastSettings.Settings.ChannelDefaultPixelFormat));
+		UE_LOG(LogAvaPlaybackServer, Display, TEXT("   - BroadcastSettings.ChannelDefaultResolution: (%d, %d)"),
+			ClientInfo.BroadcastSettings.Settings.ChannelDefaultResolution.X,
+			ClientInfo.BroadcastSettings.Settings.ChannelDefaultResolution.Y);
+		UE_LOG(LogAvaPlaybackServer, Display, TEXT("   - BroadcastSettings.bDrawPlaceholderWidget: %s"),
+			ClientInfo.BroadcastSettings.Settings.bDrawPlaceholderWidget ? TEXT("true") : TEXT("false"));
+		UE_LOG(LogAvaPlaybackServer, Display, TEXT("   - BroadcastSettings.PlaceholderWidgetClass: "),
+			*ClientInfo.BroadcastSettings.Settings.PlaceholderWidgetClass.ToString());
+
+		UE_LOG(LogAvaPlaybackServer, Display, TEXT("   - MediaSyncManager: %s."),
+			ClientInfo.MediaSyncManager->IsEnabled() ? TEXT("enabled") : TEXT("disabled"));		
+		ClientInfo.MediaSyncManager->EnumerateAllTrackedPackages([](const FName& InPackageName, const TOptional<bool>& bInNeedSync)
+		{
+			if (bInNeedSync.IsSet())
+			{
+				UE_LOG(LogAvaPlaybackServer, Display, TEXT("   - Package Sync Status \"%s\":\"%s\"."),
+					*InPackageName.ToString(), bInNeedSync.GetValue() ? TEXT("Need Sync") : TEXT("Up To Date"));
+			}
+		});
+		for (const FName& PendingPackage : ClientInfo.MediaSyncManager->GetPendingRequests())
+		{
+			UE_LOG(LogAvaPlaybackServer, Display, TEXT("   - Pending Sync Status Request \"%s\"."), *PendingPackage.ToString());
+		}
+	}
+
+	UE_LOG(LogAvaPlaybackServer, Display, TEXT("Active Playback Instances:"));
+	for (const TPair<FGuid, TSharedPtr<FAvaMediaPlaybackInstance>>& ActivePlaybackInstance : ActivePlaybackInstances)
+	{
+		const TSharedPtr<FAvaMediaPlaybackInstance>& Instance = ActivePlaybackInstance.Value;
+		UE_LOG(LogAvaPlaybackServer, Display, TEXT("   - Id:%s, Channel: %s, Asset: %s, Status: %s, UserData: %s ."),
+			*Instance->GetInstanceId().ToString(), *Instance->GetChannelName(), *Instance->GetSourcePath().ToString(),
+			*StaticEnum<EAvaMediaPlaybackStatus>()->GetNameByValue(static_cast<int32>(Instance->GetStatus())).ToString(),
+			*Instance->GetInstanceUserData());
+	}
+	
+	UE_LOG(LogAvaPlaybackServer, Display, TEXT("Active Playback Transitions:"));
+	for (const TPair<FGuid, TObjectPtr<UAvaMediaServerPlaybackTransition>>& TransitionEntry : PlaybackInstanceTransitions->Transitions)
+	{
+		if (const UAvaMediaServerPlaybackTransition* Transition = TransitionEntry.Value)
+		{
+			UE_LOG(LogAvaPlaybackServer, Display, TEXT("   - %s: %s"),
+				*Transition->GetPrettyTransitionInfo(), *Transition->GetBriefTransitionDescription());
+		}
+	}
+}
+
+void FAvaMediaPlaybackServer::OnAvaMediaSettingsChanged(UObject*, struct FPropertyChangedEvent&)
+{
+	ApplyAvaMediaSettings();
+}
+
+void FAvaMediaPlaybackServer::OnChannelChanged(const FAvaOutputChannel& InChannel, EAvaChannelChange InChange)
+{
+	// Propagate channel state change event to all clients (unless blocked).
+	if (!bBlockChannelStatusUpdate && EnumHasAnyFlags(InChange, EAvaChannelChange::State))
+	{
+		const FString ChannelName = InChannel.GetChannelName().ToString();
+		for (const TPair<FString, TSharedPtr<FClientInfo>>& Client : Clients)
+		{
+			SendChannelStatusUpdate(ChannelName, InChannel, Client.Value->Address);
+		}
+	}
+}
+
+void FAvaMediaPlaybackServer::OnMediaOutputStateChanged(const FAvaOutputChannel& InChannel, const UMediaOutput* InMediaOutput)
+{
+	// Remark: the channel's state has already been refreshed.
+	const FString ChannelName = InChannel.GetChannelName().ToString();
+	for (const TPair<FString, TSharedPtr<FClientInfo>>& Client : Clients)
+	{
+		SendChannelStatusUpdate(ChannelName, InChannel, Client.Value->Address);
+	}
+}
+
+void FAvaMediaPlaybackServer::OnAvaAssetSyncStatusReceived(const FAvaAssetSyncStatusReceivedParams& InParams)
+{
+	if (const FClientInfo* ClientInfo = GetClientInfo(InParams.RemoteName))
+	{
+		// If it needs sync, the status becomes "NeedSync" otherwise, we consider it available.
+		const EAvaMediaPlaybackAssetStatus PlaybackAssetStatus = (InParams.bNeedsSync) ? EAvaMediaPlaybackAssetStatus::NeedsSync : EAvaMediaPlaybackAssetStatus::Available;
+		SendPlaybackAssetStatus(ClientInfo->Address, InParams.AssetPath, PlaybackAssetStatus);
+	}
+}
+
+void FAvaMediaPlaybackServer::OnPlaybackInstanceInvalidated(const FAvaMediaPlaybackInstance& InPlaybackInstance)
+{
+	// If the entry was not playing (i.e. just loaded), we update it's status to one of the
+	// unloaded status. When a non-playing playback entry is invalidated, it is as if it is unloaded.
+	if (!InPlaybackInstance.GetPlayback()->IsPlaying())
+	{
+		SendPlaybackStatus(GetAllClientAddresses(), InPlaybackInstance.GetInstanceId(), InPlaybackInstance.GetChannelName(), InPlaybackInstance.GetSourcePath(), GetUnloadedPlaybackStatus(InPlaybackInstance.GetSourcePath()));
+	}
+}
+
+void FAvaMediaPlaybackServer::OnPlaybackInstanceStatusChanged(const FAvaMediaPlaybackInstance& InPlaybackInstance)
+{
+	SendPlaybackStatus(GetAllClientAddresses(), InPlaybackInstance.GetInstanceId(), InPlaybackInstance.GetChannelName(), InPlaybackInstance.GetSourcePath(), InPlaybackInstance.GetStatus());
+}
+
+void FAvaMediaPlaybackServer::OnPlaybackAssetRemoved(const FSoftObjectPath& InAssetPath)
+{
+	for (const FMessageAddress& ClientAddress : GetAllClientAddresses())
+	{
+		SendPlaybackAssetStatus(ClientAddress, InAssetPath, EAvaMediaPlaybackAssetStatus::Missing);
+	}
+}
+
+void FAvaMediaPlaybackServer::OnPlayableSequenceEvent(UAvalanchePlayable* InPlayable, const FName& SequenceName, EAvalanchePlayableSequenceEventType InEventType)
+{
+	if (!InPlayable)
+	{
+		return;
+	}
+
+	// Use the instance id to trace back which playback instance this event belongs to.
+	const TSharedPtr<FAvaMediaPlaybackInstance> PlaybackInstance = FindActivePlaybackInstance(InPlayable->GetInstanceId());
+	if (!PlaybackInstance)
+	{
+		return;
+	}
+
+	// Filter out clients on the same process, there is no need to replicate playable events in that case.
+	constexpr bool bExcludeClientOnLocalProcess = true;
+	const TArray<FMessageAddress> ClientAddresses = GetAllClientAddresses(bExcludeClientOnLocalProcess);
+
+	if (ClientAddresses.IsEmpty())
+	{
+		return;
+	}
+	
+	FAvaMediaPlayableSequenceEvent* Message = FMessageEndpoint::MakeMessage<FAvaMediaPlayableSequenceEvent>();
+	Message->InstanceId = InPlayable->GetInstanceId();
+	Message->AssetPath = PlaybackInstance->GetSourcePath();
+	Message->ChannelName = PlaybackInstance->GetChannelName();
+	Message->SequenceName = SequenceName.ToString();
+	Message->EventType = InEventType;
+	SendResponse(Message, ClientAddresses);
+}
+
+void FAvaMediaPlaybackServer::ApplyAvaMediaSettings()
+{
+	const UAvalancheMediaSettings& Settings = UAvalancheMediaSettings::Get();
+
+	const ELogVerbosity::Type LogReplicationVerbosity = LogReplicationVerbosityFromCommandLine.IsSet() ?
+		LogReplicationVerbosityFromCommandLine.GetValue() : UAvalancheMediaSettings::ToLogVerbosity(Settings.PlaybackServerLogReplicationVerbosity);
+	
+	if (LogReplicationVerbosity > ELogVerbosity::NoLogging && LogReplicationVerbosity < ELogVerbosity::NumVerbosity)
+	{
+		if (!ReplicationOutputDevice.IsValid())
+		{
+			ReplicationOutputDevice = MakeUnique<FReplicationOutputDevice>(this);
+		}
+		ReplicationOutputDevice->SetVerbosityThreshold(LogReplicationVerbosity);
+	}
+	else
+	{
+		ReplicationOutputDevice.Reset();
+	}
+
+#if !NO_LOGGING
+	if (Settings.bVerbosePlaybackServerLogging)
+	{
+		LogAvaPlaybackServer.SetVerbosity(ELogVerbosity::Verbose);
+	}
+	else
+	{
+		LogAvaPlaybackServer.SetVerbosity(ELogVerbosity::Log);
+	}
+#endif
+}
+
+void FAvaMediaPlaybackServer::SendUserDataUpdate(const TArray<FMessageAddress>& InRecipients)
+{
+	FAvaMediaUpdateServerUserData* UserDataUpdate = FMessageEndpoint::MakeMessage<FAvaMediaUpdateServerUserData>();
+	UserDataUpdate->UserDataEntries = UserDataEntries;
+	SendResponse(UserDataUpdate, InRecipients, EMessageFlags::Reliable);
+}
+
+void FAvaMediaPlaybackServer::SendChannelStatusUpdate(const FString& InChannelName, const FAvaOutputChannel& InChannel, const FMessageAddress& InSender, bool bInIncludeOutputData)
+{
+	FAvaMediaBroadcastStatus* Response = FMessageEndpoint::MakeMessage<FAvaMediaBroadcastStatus>();
+	Response->ChannelName = InChannelName;
+	// Remark: The channel index and number of channels is used to know if the client has received all the channel's statuses
+	// for the current profile. Since the profile may have a sub-set of all channels, we use the index of the channel
+	// in the profile itself and the number of channels in the profile.
+	Response->ChannelIndex = UAvalancheBroadcast::Get().GetCurrentProfile().GetChannelIndexInProfile(InChannel.GetChannelName());
+	Response->NumChannels = UAvalancheBroadcast::Get().GetCurrentProfile().GetChannels().Num();
+	Response->ChannelState = InChannel.GetState();	// Assumes state has already been refreshed.
+	Response->ChannelIssueSeverity = InChannel.GetIssueSeverity();
+	Response->bIncludeMediaOutputData = bInIncludeOutputData;
+
+	uint32 TotalOutputDataSize = 0;
+	const TArray<UMediaOutput*>& MediaOutputs = InChannel.GetMediaOutputs();
+	for (UMediaOutput* MediaOutput : MediaOutputs)
+	{
+		const FAvaMediaOutputInfo& OutputInfo = InChannel.GetMediaOutputInfo(MediaOutput);
+		FAvaMediaOutputStatus& OutputStatus = Response->MediaOutputStatuses.Emplace(OutputInfo.Guid);
+		const EAvaMediaOutputState OutputState = InChannel.GetMediaOutputState(MediaOutput);
+		OutputStatus.MediaOutputState = OutputState;
+		OutputStatus.MediaIssueSeverity = InChannel.GetMediaOutputIssueSeverity(OutputState, MediaOutput);
+		OutputStatus.MediaIssueMessages = InChannel.GetMediaOutputIssueMessages(MediaOutput);
+		
+		if (bInIncludeOutputData)
+		{
+			FAvaMediaOutputData MediaOutputData = UE::AvaMediaOutputUtils::CreateMediaOutputData(MediaOutput);
+			MediaOutputData.OutputInfo = InChannel.GetMediaOutputInfo(MediaOutput);
+			MediaOutputData.OutputInfo.ServerName = ServerName;	// Restore server name (was Local).
+			TotalOutputDataSize += MediaOutputData.SerializedData.Num();
+			Response->MediaOutputs.Add(MoveTemp(MediaOutputData));
+		}
+	}
+
+	// Adding a warning here, if we hit this warning, it may be necessary to send the
+	// data through some other transport.
+	const uint32 SafeMessageSizeLimit = UE::AvaMediaMessageUtils::GetSafeMessageSizeLimit();	
+	if (TotalOutputDataSize > SafeMessageSizeLimit)
+	{
+		UE_LOG(LogAvaPlaybackServer, Warning,
+			TEXT("The requested channel status update (DataSize: %d) is larger that the safe message size limit (%d)."),
+			TotalOutputDataSize, SafeMessageSizeLimit);
+	}
+	
+	SendResponse(Response, InSender);
+}
+
+void FAvaMediaPlaybackServer::SendAllChannelStatusUpdate(const FMessageAddress& InSender, bool bInIncludeOutputData)
+{
+	for (const FAvaOutputChannel* Channel : UAvalancheBroadcast::Get().GetCurrentProfile().GetChannels())
+	{
+		SendChannelStatusUpdate(Channel->GetChannelName().ToString(), *Channel, InSender, bInIncludeOutputData);
+	}
+}
+
+void FAvaMediaPlaybackServer::SendLogMessage(const TCHAR* InText, ELogVerbosity::Type InVerbosity, const FName& InCategory, double InTime)
+{
+	// Filter out clients on the same process, there is no need to replicate logs in that case.
+	constexpr bool bExcludeClientOnLocalProcess = true;
+	const TArray<FMessageAddress> ClientAddresses = GetAllClientAddresses(bExcludeClientOnLocalProcess);
+
+	if (ClientAddresses.IsEmpty())
+	{
+		return;
+	}
+	
+	FAvaMediaPlaybackLog* Message = FMessageEndpoint::MakeMessage<FAvaMediaPlaybackLog>();
+	Message->Text = InText;
+	Message->Verbosity = InVerbosity;
+	Message->Category = InCategory;
+	Message->Time = InTime;
+	SendResponse(Message, ClientAddresses);
+}
+
+void FAvaMediaPlaybackServer::ExecutePendingPlaybackCommands()
+{
+	for (const FPendingPlaybackCommand& PendingCommand : PendingPlaybackCommands)
+	{
+		const FAvaMediaPlaybackCommand& Command = PendingCommand.Command;
+		switch (Command.Action)
+		{
+		case EAvaMediaPlaybackAction::None:
+			break;
+			
+		case EAvaMediaPlaybackAction::Load:
+			LoadPlayback(PendingCommand.ReplyTo, Command.InstanceId, Command.ChannelName, Command.AssetPath);
+			break;
+			
+		case EAvaMediaPlaybackAction::Start:
+			StartPlayback(PendingCommand.ReplyTo, Command.InstanceId, Command.ChannelName, Command.AssetPath);
+			break;
+			
+		case EAvaMediaPlaybackAction::Stop:
+			StopPlayback(PendingCommand.ReplyTo, Command.InstanceId, Command.ChannelName, Command.AssetPath);
+			break;
+			
+		case EAvaMediaPlaybackAction::Unload:
+			UnloadPlayback(PendingCommand.ReplyTo, Command.InstanceId, Command.ChannelName, Command.AssetPath);
+			break;
+			
+		case EAvaMediaPlaybackAction::Status:
+			SendPlaybackStatus(PendingCommand.ReplyTo, Command.InstanceId, Command.ChannelName, Command.AssetPath);
+			break;
+			
+		case EAvaMediaPlaybackAction::SetUserData:	
+			SetPlaybackUserData(PendingCommand.ReplyTo, Command.InstanceId, Command.Arguments);
+			break;
+
+		case EAvaMediaPlaybackAction::GetUserData:
+			SendPlaybackUserData(PendingCommand.ReplyTo, Command.InstanceId);
+			break;
+		}
+	}
+
+	PendingPlaybackCommands.Reset();
+}
+
+TSharedPtr<FAvaMediaPlaybackInstance> FAvaMediaPlaybackServer::GetOrLoadPlaybackInstance(const FGuid& InInstanceId, const FString& InChannelName, const FSoftObjectPath& InAssetPath)
+{
+	// Check if loaded locally by the server under the same Id.
+	TSharedPtr<FAvaMediaPlaybackInstance> PlaybackInstance = FindActivePlaybackInstance(InInstanceId);
+	if (PlaybackInstance.IsValid())
+	{
+		bool bInstanceIsValid = true;
+		
+		// Validate channel name and asset path.
+		// If client has reassigned the Id for some reason, it will cause a reload of a new instance. 
+		if (PlaybackInstance->GetChannelName() != InChannelName)
+		{
+			UE_LOG(LogAvaPlaybackServer, Error
+				, TEXT("Existing Playback InstanceId \"%s\" for asset \"%s\" has the wrong channel \"%s\", requested \"%s\".")
+				, *InInstanceId.ToString(), *InAssetPath.ToString(), *PlaybackInstance->GetChannelName(), *InChannelName);
+			bInstanceIsValid = false;
+		}
+		
+		if (PlaybackInstance->GetSourcePath() != InAssetPath)
+		{
+			UE_LOG(LogAvaPlaybackServer, Error
+				, TEXT("Existing Playback InstanceId \"%s\" has wrong source asset path \"%s\", requested \"%s\".")
+				, *InInstanceId.ToString(), *PlaybackInstance->GetSourcePath().ToString(), *InAssetPath.ToString());
+			bInstanceIsValid = false;
+		}
+
+		if (bInstanceIsValid)
+		{
+			return PlaybackInstance;
+		}
+	}
+
+	// Load it or acquire a cached recycled asset.
+	PlaybackInstance = Manager->AcquireOrLoadPlaybackInstance(InAssetPath, InChannelName);
+	
+	if (PlaybackInstance.IsValid())
+	{
+		PlaybackInstance->SetInstanceId(InInstanceId); // set the instance id provided by the client.
+		PlaybackInstance->SetStatus(EAvaMediaPlaybackStatus::Loading);
+		Manager->ApplyPendingCommands(PlaybackInstance->GetPlayback(), InInstanceId, InAssetPath, InChannelName);
+		ActivePlaybackInstances.Add(InInstanceId, PlaybackInstance);
+	}
+	
+	return PlaybackInstance;	
+}
+
+void FAvaMediaPlaybackServer::LoadPlayback(const FMessageAddress& InReplyToAddress, const FGuid& InInstanceId, const FString& InChannelName, const FSoftObjectPath& InAssetPath)
+{
+	if (!InAssetPath.IsValid())
+	{
+		// Not supported.
+		UE_LOG(LogAvaPlaybackServer, Error, TEXT("Specifying invalid path for load command is not supported."));
+	}
+	else
+	{
+		if (const TSharedPtr<FAvaMediaPlaybackInstance> PlaybackInstance = GetOrLoadPlaybackInstance(InInstanceId, InChannelName, InAssetPath))
+		{
+			if (PlaybackInstance->GetPlayback())
+			{
+				PlaybackInstance->GetPlayback()->LoadInstances();
+			}
+			SendPlaybackStatus(InReplyToAddress, InInstanceId, InChannelName, InAssetPath, EAvaMediaPlaybackStatus::Loading);
+		}
+		else
+		{
+			// There was an error loading, we could either send Error or Missing as status. Sending missing for now.
+			SendPlaybackStatus(InReplyToAddress, InInstanceId, InChannelName, InAssetPath, EAvaMediaPlaybackStatus::Missing);
+		}
+	}
+}
+
+void FAvaMediaPlaybackServer::StartPlayback(const FMessageAddress& InReplyToAddress, const FGuid& InInstanceId, const FString& InChannelName, const FSoftObjectPath& InAssetPath)
+{
+	if (!InAssetPath.IsValid())
+	{
+		// Start all loaded playback
+		const TArray<FPlaybackInstanceReference> StartedInstances = StartPlaybacks();
+		SendPlaybackStatuses(InReplyToAddress, InChannelName, StartedInstances, EAvaMediaPlaybackStatus::Starting);
+	}
+	else
+	{
+		if (const TSharedPtr<FAvaMediaPlaybackInstance> PlaybackInstance = GetOrLoadPlaybackInstance(InInstanceId, InChannelName, InAssetPath))
+		{
+			PlaybackInstance->GetPlayback()->Play();
+			PlaybackInstance->SetStatus(EAvaMediaPlaybackStatus::Starting);
+			SendPlaybackStatus(InReplyToAddress, InInstanceId, InChannelName, InAssetPath, EAvaMediaPlaybackStatus::Starting);
+		}
+		else
+		{
+			// There was an error loading, we could either send Error or Missing as status. Sending missing for now.
+			SendPlaybackStatus(InReplyToAddress, InInstanceId, InChannelName, InAssetPath, EAvaMediaPlaybackStatus::Missing);
+		}	
+	}
+}
+
+void FAvaMediaPlaybackServer::StopPlayback(const FMessageAddress& InReplyToAddress, const FGuid& InInstanceId, const FString& InChannelName, const FSoftObjectPath& InAssetPath)
+{
+	// Instance Id is specified.
+	if (InInstanceId.IsValid())
+	{
+		if (const TSharedPtr<FAvaMediaPlaybackInstance> PlaybackInstance = FindActivePlaybackInstance(InInstanceId))
+		{
+			if (PlaybackInstance->GetPlayback()->IsPlaying())
+			{
+				PlaybackInstance->GetPlayback()->Stop(EAvaPlaybackStopOptions::Default);
+			}
+			PlaybackInstance->SetStatus(EAvaMediaPlaybackStatus::Loaded);
+			SendPlaybackStatus(InReplyToAddress, InInstanceId, PlaybackInstance->GetChannelName(), PlaybackInstance->GetSourcePath(), EAvaMediaPlaybackStatus::Loaded);
+		}
+		else
+		{
+			SendPlaybackStatus(InReplyToAddress, InInstanceId, InChannelName, InAssetPath, GetUnloadedPlaybackStatus(InAssetPath));
+		}
+		return;
+	}
+
+	constexpr bool bUnload = false; // don't unload.
+
+	// Channel is specified.
+	if (!InChannelName.IsEmpty())
+	{
+		const TArray<FPlaybackInstanceReference> StoppedInstances = StopPlaybacks(InChannelName, InAssetPath, bUnload);
+		SendPlaybackStatuses(InReplyToAddress, InChannelName, StoppedInstances, EAvaMediaPlaybackStatus::Loaded);
+	}
+	else
+	{
+		// If there is no channel specified, we want to stop all playbacks,
+		// but group them by channel because that is how reply messages are grouped.
+		const TArray<FString> ChannelNames = GetAllChannelsFromPlayingPlaybacks(InAssetPath);
+		for (const FString& ChannelName : ChannelNames)
+		{
+			const TArray<FPlaybackInstanceReference> StoppedInstances = StopPlaybacks(ChannelName, InAssetPath, bUnload);
+			if (!StoppedInstances.IsEmpty())
+			{
+				SendPlaybackStatuses(InReplyToAddress, ChannelName, StoppedInstances, EAvaMediaPlaybackStatus::Loaded);
+			}
+		}
+	}
+}
+
+void FAvaMediaPlaybackServer::UnloadPlayback(const FMessageAddress& InReplyToAddress, const FGuid& InInstanceId, const FString& InChannelName, const FSoftObjectPath& InAssetPath)
+{
+	// Instance Id is specified.
+	if (InInstanceId.IsValid())
+	{
+		if (const TSharedPtr<FAvaMediaPlaybackInstance> Instance = FindActivePlaybackInstance(InInstanceId))
+		{
+			Instance->Unload();
+			ActivePlaybackInstances.Remove(InInstanceId);
+			SendPlaybackStatus(InReplyToAddress, InInstanceId, Instance->GetChannelName(), Instance->GetSourcePath(), GetUnloadedPlaybackStatus(Instance->GetSourcePath()));
+		}
+		else
+		{
+			SendPlaybackStatus(InReplyToAddress, InInstanceId, InChannelName, InAssetPath, GetUnloadedPlaybackStatus(InAssetPath));
+		}
+		return;
+	}
+
+	constexpr bool bUnload = true; // Stop and unload.
+
+	// Channel is specified.
+	if (!InChannelName.IsEmpty())
+	{
+		// Will filter on asset path if specified.
+		const TArray<FPlaybackInstanceReference> StoppedInstances = StopPlaybacks(InChannelName, InAssetPath, bUnload);
+		SendPlaybackStatuses(InReplyToAddress, InChannelName, StoppedInstances, EAvaMediaPlaybackStatus::Available);
+	}
+	else
+	{
+		// If there is no channel specified, we want to unload all playbacks,
+		// but group them by channel because that is how reply messages are grouped.
+		const FAvaBroadcastProfile& Profile = UAvalancheBroadcast::Get().GetCurrentProfile();
+		for (const FAvaOutputChannel* Channel : Profile.GetChannels())	// Note: using all channels, not just playing ones.
+		{
+			const FString ChannelName = Channel->GetChannelName().ToString();
+			// Will filter on asset path if specified.
+			const TArray<FPlaybackInstanceReference> StoppedInstances = StopPlaybacks(ChannelName, InAssetPath, bUnload);
+			if (!StoppedInstances.IsEmpty())
+			{
+				SendPlaybackStatuses(InReplyToAddress, ChannelName, StoppedInstances, EAvaMediaPlaybackStatus::Available);
+			}
+		}
+	}
+}
+
+void FAvaMediaPlaybackServer::SetPlaybackUserData(const FMessageAddress& InReplyToAddress, const FGuid& InInstanceId, const FString& InUserData)
+{
+	if (InInstanceId.IsValid())
+	{
+		if (const TSharedPtr<FAvaMediaPlaybackInstance> Instance = FindActivePlaybackInstance(InInstanceId))
+		{
+			Instance->SetInstanceUserData(InUserData);
+			
+			using namespace UE::AvaMediaPlaybackServer::Private;
+			SendResponse( MakePlaybackStatusMessage(InInstanceId,
+				Instance->GetChannelName(), Instance->GetSourcePath(),
+				Instance->GetStatus(), Instance->GetInstanceUserData()), InReplyToAddress);
+		}
+	}
+}
+
+void FAvaMediaPlaybackServer::SendPlaybackUserData(const FMessageAddress& InReplyToAddress, const FGuid& InInstanceId)
+{
+	if (InInstanceId.IsValid())
+	{
+		if (const TSharedPtr<FAvaMediaPlaybackInstance> Instance = FindActivePlaybackInstance(InInstanceId))
+		{
+			using namespace UE::AvaMediaPlaybackServer::Private;
+			SendResponse( MakePlaybackStatusMessage(InInstanceId,
+				Instance->GetChannelName(), Instance->GetSourcePath(),
+				Instance->GetStatus(), Instance->GetInstanceUserData()), InReplyToAddress);
+		}
+	}
+}
+
+void FAvaMediaPlaybackServer::SendPlaybackStatus(const FMessageAddress& InReplyToAddress, const FGuid& InInstanceId, const FString& InChannelName, const FSoftObjectPath& InAssetPath)
+{
+	if (!InAssetPath.IsValid())
+	{
+		// Send all loaded/playing objects.
+		if (!InChannelName.IsEmpty())
+		{
+			SendAllPlaybackStatusesForChannelAndAssetPath(InReplyToAddress, InChannelName, FSoftObjectPath());
+		}
+		else
+		{
+			const FAvaBroadcastProfile& Profile = UAvalancheBroadcast::Get().GetCurrentProfile();
+			for (const FAvaOutputChannel* Channel : Profile.GetChannels())
+			{
+				SendAllPlaybackStatusesForChannelAndAssetPath(InReplyToAddress, Channel->GetChannelName().ToString(), FSoftObjectPath());
+			}
+		}
+	}
+	else
+	{
+		if (InInstanceId.IsValid())
+		{
+			// Possibilities: Missing, Syncing, Available, Loaded, Started.
+			// TODO: For syncing asset, we need to query the status of the transfer from StormSync, but we probably can do that on the client.
+			if (const TSharedPtr<FAvaMediaPlaybackInstance> PlaybackInstance = FindActivePlaybackInstance(InInstanceId))
+			{
+				PlaybackInstance->UpdateStatus();
+				SendPlaybackStatus(InReplyToAddress, InInstanceId, InChannelName, InAssetPath, PlaybackInstance->GetStatus());
+			}
+			else
+			{
+				SendPlaybackStatus(InReplyToAddress, InInstanceId, InChannelName, InAssetPath, GetUnloadedPlaybackStatus(InAssetPath));
+			}
+		}
+		else
+		{
+			// In case the instance id is not specified, we send the status of any playback we have for the given asset.
+			SendAllPlaybackStatusesForChannelAndAssetPath(InReplyToAddress, InChannelName, InAssetPath);
+		}
+	}
+}
+
+void FAvaMediaPlaybackServer::SendPlaybackStatus(const FMessageAddress& InSendTo, const FGuid& InInstanceId, const FString& InChannelName, const FSoftObjectPath& InAssetPath, EAvaMediaPlaybackStatus InStatus)
+{
+	using namespace UE::AvaMediaPlaybackServer::Private;
+	SendResponse(MakePlaybackStatusMessage(InInstanceId, InChannelName, InAssetPath, InStatus), InSendTo);
+}
+
+void FAvaMediaPlaybackServer::SendPlaybackStatus(const TArray<FMessageAddress>& InRecipients, const FGuid& InInstanceId, const FString& InChannelName, const FSoftObjectPath& InAssetPath, EAvaMediaPlaybackStatus InStatus)
+{
+	using namespace UE::AvaMediaPlaybackServer::Private;
+	SendResponse(MakePlaybackStatusMessage(InInstanceId, InChannelName, InAssetPath, InStatus), InRecipients);
+}
+
+void FAvaMediaPlaybackServer::SendPlaybackStatuses(const FMessageAddress& InSendTo, const FString& InChannelName, const TArray<FPlaybackInstanceReference>& InInstances, EAvaMediaPlaybackStatus InStatus)
+{
+	FAvaMediaPlaybackStatuses* Response = FMessageEndpoint::MakeMessage<FAvaMediaPlaybackStatuses>();
+	Response->ChannelName = InChannelName;
+	Response->AssetPaths.Reserve(InInstances.Num());
+	Response->InstanceIds.Reserve(InInstances.Num());
+	for (const FPlaybackInstanceReference& Instance : InInstances)
+	{
+		Response->AssetPaths.Add(Instance.Path);
+		Response->InstanceIds.Add(Instance.Id);
+	}
+	Response->Status = InStatus;
+	SendResponse(Response, InSendTo);
+}
+
+void FAvaMediaPlaybackServer::SendAllPlaybackStatusesForChannelAndAssetPath(const FMessageAddress& InSendTo, const FString& InChannelName, const FSoftObjectPath& InAssetPath)
+{
+	// Group all the playback objects per status.
+	TMap<EAvaMediaPlaybackStatus, TArray<FPlaybackInstanceReference>> InstancesPerStatus;
+
+	for (const TPair<FGuid, TSharedPtr<FAvaMediaPlaybackInstance>>& InstanceEntry : ActivePlaybackInstances)
+	{
+		if (const TSharedPtr<FAvaMediaPlaybackInstance> Instance = InstanceEntry.Value)
+		{
+			if (!InChannelName.IsEmpty() && Instance->GetChannelName() != InChannelName)
+			{
+				continue;
+			}
+
+			if (!InAssetPath.IsNull() && Instance->GetSourcePath() != InAssetPath)
+			{
+				continue;
+			}
+
+			TArray<FPlaybackInstanceReference>* Instances = InstancesPerStatus.Find(Instance->GetStatus());
+			if (!Instances)
+			{
+				Instances = &InstancesPerStatus.Emplace(Instance->GetStatus());
+				Instances->Reserve(32);
+			}
+			Instances->Add({Instance->GetInstanceId(), Instance->GetSourcePath()});
+		}
+	}
+
+	for (TPair<EAvaMediaPlaybackStatus, TArray<FPlaybackInstanceReference>>& StatusEntry :  InstancesPerStatus)
+	{
+		SendPlaybackStatuses(InSendTo, InChannelName, StatusEntry.Value, StatusEntry.Key);
+	}
+}
+
+void FAvaMediaPlaybackServer::SendPlaybackAssetStatus(const FMessageAddress& InSendTo, const FSoftObjectPath& InAssetPath, EAvaMediaPlaybackAssetStatus InStatus)
+{
+	FAvaMediaPlaybackAssetStatus* Response = FMessageEndpoint::MakeMessage<FAvaMediaPlaybackAssetStatus>();
+	Response->AssetPath = InAssetPath;
+	Response->Status = InStatus;
+	SendResponse(Response, InSendTo);
+}
+
+bool FAvaMediaPlaybackServer::UpdateChannelOutputConfig(FAvaOutputChannel& InChannel,
+	const TArray<FAvaMediaOutputData>& InMediaOutputs, bool bInRefreshState)
+{
+	if (InChannel.GetState() == EAvaChannelState::Live)
+	{
+		UE_LOG(LogAvaPlaybackServer, Error, TEXT("Failed to update output config on channel \"%s\". Channel is live."),
+			   *InChannel.GetChannelName().ToString());
+		return false;
+	}
+
+	if (!InMediaOutputs.IsEmpty())
+	{
+		TArray<TStrongObjectPtr<UMediaOutput>> NewOutputs;
+		NewOutputs.Reserve(InMediaOutputs.Num());
+		TArray<FAvaMediaOutputInfo> NewOutputInfos;
+		NewOutputInfos.Reserve(InMediaOutputs.Num());
+
+		for (const FAvaMediaOutputData& MediaOutputData : InMediaOutputs)
+		{
+			// Important: don't add outputs not destined for this server.
+			if (MediaOutputData.OutputInfo.ServerName == ServerName)
+			{
+				NewOutputs.Emplace_GetRef().Reset(
+					UE::AvaMediaOutputUtils::CreateMediaOutput(MediaOutputData, &UAvalancheBroadcast::Get()));
+				NewOutputInfos.Add(MediaOutputData.OutputInfo);
+			}
+			else
+			{
+				UE_LOG(LogAvaPlaybackServer, Warning,
+					   TEXT("Channel \"%s\" received an output for another server (\"%s\")"),
+					   *InChannel.GetChannelName().ToString(), *MediaOutputData.OutputInfo.ServerName);
+			}
+		}
+
+		if (!NewOutputs.IsEmpty())
+		{
+			{
+				TArray<UMediaOutput*> MediaOutputs = InChannel.GetMediaOutputs();
+				for (UMediaOutput* MediaOutput : MediaOutputs)
+				{
+					InChannel.RemoveMediaOutput(MediaOutput);
+				}
+			}
+
+			for (int32 Index = 0; Index < NewOutputs.Num(); ++Index)
+			{
+				// Make the device info "local" for this server.
+				NewOutputInfos[Index].ServerName = FAvaDeviceProviderProxyManager::LocalServerName;
+				InChannel.AddMediaOutput(NewOutputs[Index].Get(), NewOutputInfos[Index]);
+			}
+
+			// We may not desired refresh state here to avoid spurious states if
+			// we are in the middle of a series of commands.
+			if (bInRefreshState)
+			{
+				InChannel.RefreshState();
+			}
+			return true;
+		}
+	}
+	return false;
+}
+
+FAvaMediaPlaybackServer::FClientInfo& FAvaMediaPlaybackServer::GetOrCreateClientInfo(
+	const FString& InClientName, const FMessageAddress& InClientAddress)
+{
+	if (const TSharedPtr<FClientInfo>* ClientInfo = Clients.Find(InClientName))
+	{
+		if ((*ClientInfo)->Address != InClientAddress)
+		{
+			// This is suspicious though. It may also indicate a collision with multiple clients
+			// with the same name on the same computer host. This is a case we don't support for now.
+			UE_LOG(LogAvaPlaybackServer, Warning, TEXT("Client \"%s\" Address changed, possible collision between clients with same name."), *InClientName);
+			(*ClientInfo)->Address = InClientAddress;
+		}
+		return *(ClientInfo->Get());
+	}
+
+	const TSharedPtr<FClientInfo> ClientInfo = MakeShared<FClientInfo>(InClientAddress, InClientName);
+	ClientInfo->MediaSyncManager->OnAvaAssetSyncStatusReceived.AddRaw(this, &FAvaMediaPlaybackServer::OnAvaAssetSyncStatusReceived);
+	Clients.Add(InClientName, ClientInfo);
+
+	OnClientAdded(*ClientInfo);
+	
+	return *ClientInfo;
+}
+
+FAvaMediaPlaybackServer::FClientInfo* FAvaMediaPlaybackServer::GetClientInfo(const FMessageAddress& InClientAddress) const
+{
+	for (const TPair<FString, TSharedPtr<FClientInfo>>& Client : Clients )
+	{
+		if (Client.Value->Address == InClientAddress)
+		{
+			return Client.Value.Get();
+		}
+	}
+	return nullptr;
+}
+
+const FString& FAvaMediaPlaybackServer::GetClientNameSafe(const FMessageAddress& InClientAddress) const
+{
+	if (const FClientInfo* ClientInfo = GetClientInfo(InClientAddress))
+	{
+		return ClientInfo->ClientName;
+	}
+	static FString ClientNotFoundString(TEXT("[ClientNotFound]"));
+	return ClientNotFoundString;
+}
+
+const FMessageAddress& FAvaMediaPlaybackServer::GetClientAddressSafe(const FString& InClientName) const
+{
+	if (const FClientInfo* ClientInfo = GetClientInfo(InClientName))
+	{
+		return ClientInfo->Address;
+	}
+	static FMessageAddress Invalid;
+	return Invalid;
+}
+
+
+TArray<FMessageAddress> FAvaMediaPlaybackServer::GetAllClientAddresses(bool bInExcludeClientOnLocalProcess) const
+{
+	TArray<FMessageAddress> OutAddresses;
+	OutAddresses.Reserve(Clients.Num());
+	for (const TPair<FString, TSharedPtr<FClientInfo>>& Client : Clients )
+	{
+		if (bInExcludeClientOnLocalProcess && IsClientOnLocalProcess(*Client.Value))
+		{
+			continue;
+		}
+		
+		OutAddresses.Add(Client.Value->Address);
+	}
+	return OutAddresses;
+}
+
+void FAvaMediaPlaybackServer::RemoveDeadClients(const FDateTime& InCurrentTime)
+{
+	for (TMap<FString, TSharedPtr<FClientInfo>>::TIterator ServerIter(Clients); ServerIter; ++ServerIter)
+	{
+		if (ServerIter.Value()->HasTimedOut(InCurrentTime))
+		{
+			UE_LOG(LogAvaPlaybackServer, Log, TEXT("Client \"%s\" is not longer sending pings. Removing."),
+				   *ServerIter.Key());
+			const TSharedPtr<FClientInfo> RemovedClient = ServerIter.Value();
+			ServerIter.RemoveCurrent();
+			OnClientRemoved(*RemovedClient);
+		}
+	}
+}
+
+void FAvaMediaPlaybackServer::OnClientAdded(const FClientInfo& InClientInfo)
+{
+	UE_LOG(LogAvaPlaybackServer, Log, TEXT("Registering new playback client \"%s\"."), *InClientInfo.ClientName);
+	// We send user data update on connection only (reliable send).
+	SendUserDataUpdate({InClientInfo.Address});
+}
+
+void FAvaMediaPlaybackServer::OnClientRemoved(const FClientInfo& InRemovedClient)
+{
+	// TODO	
+}
+
+bool FAvaMediaPlaybackServer::IsLocalClient(const FClientInfo& ClientInfo) const
+{
+	return ClientInfo.ComputerName == ComputerName && ClientInfo.ProjectContentPath == ProjectContentPath;
+}
+
+FAvaMediaPlaybackServer::FClientInfo::FClientInfo(const FMessageAddress& InClientAddress, const FString& InClientName)
+	: Address(InClientAddress)
+	, ClientName(InClientName)
+{
+	MediaSyncManager = MakeShared<FAvaMediaSyncManager>(InClientName);
+}
+
+FAvaMediaPlaybackServer::FClientInfo::~FClientInfo() = default;
+
+FAvaMediaPlaybackServer::FReplicationOutputDevice::FReplicationOutputDevice(FAvaMediaPlaybackServer* InServer)
+		: Server(InServer)
+{
+	GLog->AddOutputDevice(this);
+	GLog->SerializeBacklog(this);
+}
+		
+FAvaMediaPlaybackServer::FReplicationOutputDevice::~FReplicationOutputDevice()
+{
+	// At shutdown, GLog may already be null
+	if (GLog != nullptr)
+	{
+		GLog->RemoveOutputDevice(this);
+	}
+}
+
+void FAvaMediaPlaybackServer::FReplicationOutputDevice::SetVerbosityThreshold(ELogVerbosity::Type InVerbosityThreshold)
+{
+	VerbosityThreshold = InVerbosityThreshold;
+}
+
+void FAvaMediaPlaybackServer::FReplicationOutputDevice::Serialize(const TCHAR* InText, ELogVerbosity::Type InVerbosity, const FName& InCategory)
+{
+	Serialize(InText, InVerbosity, InCategory, {});
+}
+	
+void FAvaMediaPlaybackServer::FReplicationOutputDevice::Serialize(const TCHAR* InText, ELogVerbosity::Type InVerbosity, const FName& InCategory, double InTime)
+{
+	if (Server && InVerbosity <= VerbosityThreshold)
+	{
+		Server->SendLogMessage(InText, InVerbosity, InCategory, InTime);
+	}
+}
\ No newline at end of file
diff --git a/Engine/Plugins/Experimental/Avalanche/Source/AvalancheMedia/Private/Playback/AvaMediaPlaybackServerProcess.cpp b/Engine/Plugins/Experimental/Avalanche/Source/AvalancheMedia/Private/Playback/AvaMediaPlaybackServerProcess.cpp
new file mode 100644
--- /dev/null
+++ b/Engine/Plugins/Experimental/Avalanche/Source/AvalancheMedia/Private/Playback/AvaMediaPlaybackServerProcess.cpp
@@ -0,0 +1,183 @@
+// Copyright Epic Games, Inc. All Rights Reserved.
+
+#include "AvaMediaPlaybackServerProcess.h"
+
+#include "AvalancheMediaSettings.h"
+#include "IAvaMediaModule.h"
+#include "Misc/App.h"
+#include "Misc/Paths.h"
+#include "Playback/AvaMediaPlaybackClient.h"
+
+namespace UE::AvaMediaPlaybackServerProcess::Private
+{
+	FString GetLocalServerName()
+	{
+		return UAvalancheMediaSettings::Get().LocalPlaybackServerSettings.ServerName;
+	}
+
+	const TCHAR* GetPlaybackServerLogReplicationVerbosityString()
+	{
+		const UAvalancheMediaSettings& Settings = UAvalancheMediaSettings::Get(); 
+		return ToString(UAvalancheMediaSettings::ToLogVerbosity(Settings.PlaybackServerLogReplicationVerbosity));
+	}
+}
+
+FAvaMediaPlaybackServerProcess::FAvaMediaPlaybackServerProcess(FProcHandle&& InProcessHandle)
+	: ProcessHandle(MoveTemp(InProcessHandle))
+{}
+
+bool FAvaMediaPlaybackServerProcess::IsLaunched()
+{
+	return ProcessHandle.IsValid() && FPlatformProcess::IsProcRunning(ProcessHandle);
+}
+
+bool FAvaMediaPlaybackServerProcess::Launch()
+{
+	if (IsLaunched())
+	{
+		UE_LOG(LogAvaMedia, Warning, TEXT("Local playback server is already running."));
+		return true;
+	}
+
+	using namespace UE::AvaMediaPlaybackServerProcess::Private;
+
+	FString GameNameOrProjectFile;
+	if (FPaths::IsProjectFilePathSet())
+	{
+		GameNameOrProjectFile = FString::Printf(TEXT("\"%s\""), *FPaths::GetProjectFilePath());
+	}
+	else
+	{
+		GameNameOrProjectFile = FApp::GetProjectName();
+	}
+
+	FString CommandLine;
+	CommandLine += GameNameOrProjectFile;
+
+	// Todo: It will likely need a more complete set of parameters. Reference: UDisplayClusterLaunchEditorProjectSettings.
+	const FLocalPlaybackServerSettings& Settings = UAvalancheMediaSettings::Get().LocalPlaybackServerSettings;
+
+	// Standard arguments
+	const FIntPoint Resolution = Settings.Resolution;
+	CommandLine += FString::Printf(TEXT(" -game -windowed -ResX=%d -ResY=%d -messaging"), Resolution.X, Resolution.Y);
+	if (Settings.bEnableLogConsole)
+	{
+		CommandLine += TEXT(" -log");
+	}
+	// Avalanche arguments
+	CommandLine += FString::Printf(TEXT(" -AvaMediaPlaybackServerStart=%s"), *GetLocalServerName());
+	CommandLine += TEXT(" -AvaMediaPlaybackClientSuppress");
+	CommandLine += FString::Printf(TEXT(" -AvaMediaPlaybackServerLogReplication=%s"), GetPlaybackServerLogReplicationVerbosityString());
+	// Storm Sync arguments
+	CommandLine += TEXT(" -NoStormSyncServerAutoStart");
+
+	if (!Settings.ExtraCommandLineArguments.IsEmpty())
+	{
+		CommandLine += TEXT(" ");
+		CommandLine += Settings.ExtraCommandLineArguments;
+	}
+
+	FString LogCommands;
+	
+	for (const FAvaPlaybackServerLoggingEntry& LoggingEntry : Settings.Logging)
+	{
+		if (!LoggingEntry.Category.IsNone())
+		{
+			LogCommands += FString::Printf(TEXT("%s%s %s"),
+					LogCommands.IsEmpty() ? TEXT("") : TEXT(", "), *LoggingEntry.Category.ToString(),
+					ToString(UAvalancheMediaSettings::ToLogVerbosity(LoggingEntry.VerbosityLevel)));
+		}
+	}
+
+	if (!LogCommands.IsEmpty())
+	{
+		CommandLine += FString::Printf(TEXT(" -LogCmds=\"%s\""), *LogCommands);		
+	}
+
+	const FString ExecutablePath = FPlatformProcess::ExecutablePath();
+
+	uint32 ProcessID = 0;
+	constexpr bool bLaunchDetached = true;
+	constexpr bool bLaunchMinimized = false;
+	constexpr bool bLaunchWindowHidden = false;
+	constexpr uint32 PriorityModifier = 0;
+
+	UE_LOG(LogAvaMedia, Log, TEXT("Launching a playback server in game mode in a new process with the following command line:"));
+	UE_LOG(LogAvaMedia, Log, TEXT("%s %s"), *ExecutablePath, *CommandLine);
+
+	ProcessHandle = FPlatformProcess::CreateProc(
+		*ExecutablePath, *CommandLine, bLaunchDetached,
+		bLaunchMinimized, bLaunchWindowHidden, &ProcessID,
+		PriorityModifier, nullptr, nullptr, nullptr);
+	
+	return ProcessHandle.IsValid();
+}
+
+void FAvaMediaPlaybackServerProcess::Stop()
+{
+	FPlatformProcess::TerminateProc(ProcessHandle);
+}
+
+TSharedPtr<FAvaMediaPlaybackServerProcess> FAvaMediaPlaybackServerProcess::Find(const FAvaMediaPlaybackClient& InPlaybackClient)
+{
+	using namespace UE::AvaMediaPlaybackServerProcess::Private;
+
+	// We use the playback client to find connected server processes that are considered local.
+	// Local means that it is in the same project content path and the same machine. We figure out
+	// if it is the same machine if the process Id can be opened.
+
+	const FString ProjectContentPath = InPlaybackClient.GetProjectContentPath();
+	const uint32 ClientProcessId = FPlatformProcess::GetCurrentProcessId();
+	const FString LocalServerName = GetLocalServerName();
+	
+	// First try with the configured local server name.
+	if (InPlaybackClient.GetServerProjectContentPath(LocalServerName) == ProjectContentPath)
+	{
+		const uint32 ServerProcessId = InPlaybackClient.GetServerProcessId(LocalServerName);
+		if (ServerProcessId != 0 && ClientProcessId != ServerProcessId)
+		{
+			FProcHandle ProcessHandle = FPlatformProcess::OpenProcess(ServerProcessId);
+			if (ProcessHandle.IsValid())
+			{
+				return MakeShared<FAvaMediaPlaybackServerProcess>(MoveTemp(ProcessHandle));
+			}
+
+			// Warn about the server name collision issue.
+			UE_LOG(LogAvaMedia, Warning, TEXT("Found Connected Server \"%s\" but it is not a local process."), *LocalServerName);
+			UE_LOG(LogAvaMedia, Warning, TEXT("This indicates a name collision between the servers. The local server should be renamed."));
+		}
+	}
+
+	// Fallback
+	// Try to find any connected server that is "local", i.e. on the same machine and
+	// in the same project content folder, but on a different process id.
+	
+	TArray<FString> ServerNames = InPlaybackClient.GetServerNames();
+	for (const FString& ServerName : ServerNames)
+	{
+		if (InPlaybackClient.GetServerProjectContentPath(ServerName) == ProjectContentPath)
+		{
+			const uint32 ServerProcessId = InPlaybackClient.GetServerProcessId(ServerName);
+			if (ClientProcessId != ServerProcessId)
+			{
+				FProcHandle ProcessHandle = FPlatformProcess::OpenProcess(ServerProcessId);
+				if (ProcessHandle.IsValid())
+				{
+					return MakeShared<FAvaMediaPlaybackServerProcess>(MoveTemp(ProcessHandle));
+				}
+			}
+		}
+	}
+	
+	return nullptr;	
+}
+
+TSharedPtr<FAvaMediaPlaybackServerProcess> FAvaMediaPlaybackServerProcess::FindOrCreate(const FAvaMediaPlaybackClient& InPlaybackClient)
+{
+	TSharedPtr<FAvaMediaPlaybackServerProcess> ServerProcess = Find(InPlaybackClient);
+	if (!ServerProcess.IsValid())
+	{
+		ServerProcess = MakeShared<FAvaMediaPlaybackServerProcess>();
+	}
+	return ServerProcess;
+}
\ No newline at end of file
diff --git a/Engine/Plugins/Experimental/Avalanche/Source/AvalancheMedia/Private/Playback/AvaMediaPlaybackServerProcess.h b/Engine/Plugins/Experimental/Avalanche/Source/AvalancheMedia/Private/Playback/AvaMediaPlaybackServerProcess.h
new file mode 100644
--- /dev/null
+++ b/Engine/Plugins/Experimental/Avalanche/Source/AvalancheMedia/Private/Playback/AvaMediaPlaybackServerProcess.h
@@ -0,0 +1,48 @@
+// Copyright Epic Games, Inc. All Rights Reserved.
+
+#pragma once
+
+#include "HAL/PlatformProcess.h"
+#include "Templates/SharedPointer.h"
+
+class FAvaMediaPlaybackClient;
+
+/**
+ * Helper class wrapping and managing a process handle to the game mode local playback server process. 
+ */
+class FAvaMediaPlaybackServerProcess
+{
+public:
+	FAvaMediaPlaybackServerProcess() = default;
+
+	/**
+	 * @brief Construct from an already running process handle.
+	 * @param InProcessHandle Existing process handle.
+	 */
+	explicit FAvaMediaPlaybackServerProcess(FProcHandle&& InProcessHandle);
+	
+	bool IsLaunched();
+	bool Launch();
+	void Stop();
+
+	/**
+	 * @brief Returns a playback server process object wrapping the process handle of an existing server process.
+	 * @param InPlaybackClient Playback client instance to use for finding the server.
+	 * @return A new object wrapping the process handle of the existing server process. It will return null if no existing server is found.
+	 */
+	static TSharedPtr<FAvaMediaPlaybackServerProcess> Find(const FAvaMediaPlaybackClient& InPlaybackClient);
+	
+	/**
+	 * @brief Returns a playback server process object that wraps the process handle of either the existing server process or a newly created one.
+	 * 
+	 * The playback server may be running in a separate process already. If so, it will be detected by the playback client.
+	 * The playback client will also know it's process id so a handle to the process can be opened.
+	 *
+	 * @param InPlaybackClient Playback client instance to use for finding the server.
+	 * @return A new object wrapping the process handle of the server process.
+	 */
+	static TSharedPtr<FAvaMediaPlaybackServerProcess> FindOrCreate(const FAvaMediaPlaybackClient& InPlaybackClient);
+
+private:
+	FProcHandle ProcessHandle;
+};
\ No newline at end of file
diff --git a/Engine/Plugins/Experimental/Avalanche/Source/AvalancheMedia/Private/Playback/AvaMediaPlaybackUtils.cpp b/Engine/Plugins/Experimental/Avalanche/Source/AvalancheMedia/Private/Playback/AvaMediaPlaybackUtils.cpp
new file mode 100644
--- /dev/null
+++ b/Engine/Plugins/Experimental/Avalanche/Source/AvalancheMedia/Private/Playback/AvaMediaPlaybackUtils.cpp
@@ -0,0 +1,191 @@
+// Copyright Epic Games, Inc. All Rights Reserved.
+
+#include "Playback/AvaMediaPlaybackUtils.h"
+
+#include "AssetRegistry/AssetData.h"
+#include "AvaAssetTags.h"
+#include "AvaMediaModule.h"
+#include "HAL/PlatformFileManager.h"
+#include "Misc/Paths.h"
+#include "RenderingThread.h"
+#include "UObject/Linker.h"
+#include "UObject/LinkerLoad.h"
+#include "UObject/Package.h"
+
+#if WITH_EDITOR
+#include "Editor.h"
+#include "PackageTools.h"
+#include "Selection.h"
+#include "Subsystems/AssetEditorSubsystem.h"
+#endif
+
+void FAvaMediaPlaybackUtils::FlushPackageLoading(UPackage* InPackage)
+{
+	if (!InPackage->IsFullyLoaded())
+	{
+		FlushAsyncLoading();
+		InPackage->FullyLoad();
+	}
+	ResetLoaders(InPackage);
+}
+
+bool FAvaMediaPlaybackUtils::IsPackageDeleted(const UPackage* InExistingPackage)
+{
+	if (!InExistingPackage)
+	{
+		return false;
+	}
+	
+	const FString PackageExtension = InExistingPackage->ContainsMap() ? FPackageName::GetMapPackageExtension() : FPackageName::GetAssetPackageExtension();
+	const FString PackageFilename = FPackageName::LongPackageNameToFilename(InExistingPackage->GetName(), PackageExtension);
+			
+	return !FPaths::FileExists(PackageFilename);
+}
+
+namespace UE::AvaMediaPlaybackUtils::Private
+{
+#if WITH_EDITOR
+	void CollectObjectToPurge(UObject* InObject, TArray<UObject*>& OutObjectsToPurge)
+	{
+		if (InObject->IsAsset() && GIsEditor)
+		{
+			GEditor->GetEditorSubsystem<UAssetEditorSubsystem>()->CloseAllEditorsForAsset(InObject);
+			GEditor->GetSelectedObjects()->Deselect(InObject);
+		}
+		OutObjectsToPurge.Add(InObject);
+	}
+#endif
+}
+
+// Adapted from: PurgePackages in ConcertSyncClientUtil.cpp
+// Notes:
+// - The method used in USourceControlHelpers::ApplyOperationAndReloadPackages with the
+//   asset registry and ObjectTools::DeleteObjectsUnchecked can't be used because we want
+//   this to work in game mode.
+// - We assume the assets purged will not be the edited world (current level editor world)
+//   so we can skip the special case from the original code.
+void FAvaMediaPlaybackUtils::PurgePackages(const TArray<UPackage*>& InExistingPackages)
+{
+#if WITH_EDITOR
+	using namespace UE::AvaMediaPlaybackUtils::Private;
+
+	TArray<UObject*> ObjectsToPurge;
+
+	for (UPackage* ExistingPackage : InExistingPackages)
+	{
+		if (!IsValid(ExistingPackage))
+		{
+			continue;
+		}
+		
+		// Prevent any message from the editor saying a package is not saved or doesn't exist on disk.
+		ExistingPackage->SetDirtyFlag(false);
+
+		CollectObjectToPurge(ExistingPackage, ObjectsToPurge);
+		ForEachObjectWithPackage(ExistingPackage, [&ObjectsToPurge](UObject* InObject)
+		{
+			CollectObjectToPurge(InObject, ObjectsToPurge);
+			return true;
+		});
+	}
+	
+	// Broadcast the eminent objects destruction (ex. tell BlueprintActionDatabase to release its reference(s) on Blueprint(s) right now)
+	FEditorDelegates::OnAssetsPreDelete.Broadcast(ObjectsToPurge);
+
+	// Mark objects as purgeable.
+	for (UObject* Object : ObjectsToPurge)
+	{
+		if (Object->IsRooted())
+		{
+			Object->RemoveFromRoot();
+		}
+		Object->ClearFlags(RF_Public | RF_Standalone);
+	}
+
+	if (ObjectsToPurge.Num() > 0)
+	{
+		CollectGarbage(GARBAGE_COLLECTION_KEEPFLAGS);
+	}
+#endif
+}
+
+// Adapted from HotReloadPackages in ConcertSyncClientUtil.cpp
+bool FAvaMediaPlaybackUtils::ReloadPackages(const TArray<UPackage*>& InExistingPackages)
+{
+	FlushAsyncLoading();
+	{
+		bool bRunGC = false;
+		for (UPackage* Package : InExistingPackages)
+		{
+			if (FLinkerLoad::RemoveKnownMissingPackage(Package->GetFName()))
+			{
+				UE_LOG(LogAvaMedia, Verbose, TEXT("Package \"%s\" was removed from known missing."), *Package->GetName());
+				bRunGC = true;
+			}
+			if (Package->HasAnyPackageFlags(PKG_NewlyCreated))
+			{
+				UE_LOG(LogAvaMedia, Verbose, TEXT("Clearing Newly Created flag on Package \"%s\"."), *Package->GetName());
+				Package->ClearPackageFlags(PKG_NewlyCreated);
+			}
+		}
+		if (bRunGC)
+		{
+			UE_LOG(LogAvaMedia, Log, TEXT("Some packages where removed from known missing, garbage collecting ..."));
+			CollectGarbage(GARBAGE_COLLECTION_KEEPFLAGS);
+		}
+	}
+	FlushRenderingCommands();
+
+	FText ErrorMessage;
+#if WITH_EDITOR
+	UPackageTools::ReloadPackages(InExistingPackages, ErrorMessage, UPackageTools::EReloadPackagesInteractionMode::AssumePositive);
+#endif
+	
+	if (!ErrorMessage.IsEmpty())
+	{
+		UE_LOG(LogAvaMedia, Error, TEXT("%s"), *ErrorMessage.ToString());
+		return false;
+	}
+	return true;
+}
+
+bool FAvaMediaPlaybackUtils::IsMapAsset(const FString& InPackageName)
+{
+	FString FullPath;
+	if (FPackageName::TryConvertLongPackageNameToFilename(InPackageName, FullPath, FPackageName::GetMapPackageExtension()))
+	{
+		if (FPlatformFileManager::Get().GetPlatformFile().FileExists(*FullPath))
+		{
+			return true;
+		}
+	}
+	return false;
+}
+
+bool FAvaMediaPlaybackUtils::IsPlayableAsset(const FAssetData& InAssetData)
+{
+	const EAvalancheAssetType AssetType = FAvaSoftAssetPath::GetAssetTypeFromClass(InAssetData.GetClass(), true);
+	if (AssetType == EAvalancheAssetType::Unknown)
+	{
+		return false;
+	}
+	// For world type, we need to check the tags.
+	if (AssetType == EAvalancheAssetType::World)
+	{
+		const FAssetTagValueRef AvalancheSceneTag = InAssetData.TagsAndValues.FindTag(UE::Ava::AssetTags::AvalancheScene);
+		if (!AvalancheSceneTag.IsSet() || !AvalancheSceneTag.Equals(UE::Ava::AssetTags::Values::Enabled))
+		{
+			return false;
+		}
+	}
+	return true;
+}
+
+
+namespace UE::AvaMediaPlayback::Utils
+{
+	FString GetBriefFrameInfo()
+	{
+		return FString::Printf(TEXT("[%d]"), GFrameNumber);
+	}
+}
\ No newline at end of file
diff --git a/Engine/Plugins/Experimental/Avalanche/Source/AvalancheMedia/Private/Playback/AvaMediaSyncManager.cpp b/Engine/Plugins/Experimental/Avalanche/Source/AvalancheMedia/Private/Playback/AvaMediaSyncManager.cpp
new file mode 100644
--- /dev/null
+++ b/Engine/Plugins/Experimental/Avalanche/Source/AvalancheMedia/Private/Playback/AvaMediaSyncManager.cpp
@@ -0,0 +1,271 @@
+// Copyright Epic Games, Inc. All Rights Reserved.
+
+#include "AvaMediaSyncManager.h"
+
+#include "AssetRegistry/AssetData.h"
+#include "AssetRegistry/IAssetRegistry.h"
+#include "AvalancheMediaSettings.h"
+#include "IAvaMediaModule.h"
+#include "ModularFeature/IAvaMediaSyncProvider.h"
+#include "UObject/Package.h"
+#include "Playback/AvaMediaPlaybackServer.h"
+
+FAvaMediaSyncManager::FAvaMediaSyncManager(const FString& InRemoteName)
+	: RemoteName(InRemoteName)
+{
+	if (bEnabled)
+	{
+		RegisterEventHandlers();
+	}
+}
+
+FAvaMediaSyncManager::~FAvaMediaSyncManager()
+{
+	UnregisterEventHandlers();
+}
+
+void FAvaMediaSyncManager::SetEnable(bool bInEnabled)
+{
+	if (bInEnabled == bEnabled)
+	{
+		return;
+	}
+
+	bEnabled = bInEnabled;
+
+	if (bEnabled)
+	{
+		RegisterEventHandlers();
+	}
+	else
+	{
+		UnregisterEventHandlers();
+		PackageStatuses.Reset();
+		CompareRequests.Reset();
+	}
+	UE_LOG(LogAvaPlaybackServer, Verbose, TEXT("Server's Media Sync Manager is %s."),
+		   bEnabled ? TEXT("enabled") : TEXT("disabled"));
+}
+
+void FAvaMediaSyncManager::Tick()
+{
+	const FDateTime CurrentTime = FDateTime::UtcNow();
+	for (TMap<FName, FDateTime>::TIterator PendingRequestIterator(CompareRequests); PendingRequestIterator; ++PendingRequestIterator)
+	{
+		if (CurrentTime > PendingRequestIterator.Value())
+		{
+			HandleExpiredRequest(PendingRequestIterator.Key());
+			PendingRequestIterator.RemoveCurrent();
+		}
+	}
+}
+
+TOptional<bool> FAvaMediaSyncManager::GetAssetSyncStatus(const FSoftObjectPath& InAssetPath, bool bInForceRefresh)
+{
+	// AvaMediaSyncProvider is a modular feature, it may not be available.
+	if (!bEnabled || !IAvaMediaModule::Get().GetAvaMediaSyncProvider())
+	{
+		// Assumes the asset is locally present. Returning "false" means it doesn't need sync.
+		return TOptional<bool>(false);
+	}
+	
+	if (bInForceRefresh)
+	{
+		PackageStatuses.Remove(InAssetPath.GetLongPackageFName());
+	}
+	else if (const bool *bNeedSynchronization = PackageStatuses.Find(InAssetPath.GetLongPackageFName()))
+	{
+		return TOptional<bool>(*bNeedSynchronization);
+	}
+	
+	RequestPackageSyncStatus(InAssetPath.GetLongPackageFName());
+	return TOptional<bool>();
+}
+
+void FAvaMediaSyncManager::EnumerateAllTrackedPackages(TFunctionRef<void(const FName& /*InPackageName*/, bool /*bInNeedsSync*/)> InFunction) const
+{
+	for (const TPair<FName, bool>& PackageStatusEntry : PackageStatuses)
+	{
+		InFunction(PackageStatusEntry.Key, PackageStatusEntry.Value);
+	}
+}
+
+TArray<FName> FAvaMediaSyncManager::GetPendingRequests() const
+{
+	TArray<FName> Requests;
+	Requests.Reserve(CompareRequests.Num());
+	for (const TPair<FName, FDateTime>& Entry : CompareRequests)
+	{
+		Requests.Add(Entry.Key);
+	}
+	return Requests;
+}
+
+void FAvaMediaSyncManager::RegisterEventHandlers()
+{
+	UPackage::PackageSavedWithContextEvent.AddRaw(this, &FAvaMediaSyncManager::HandlePackageSaved);
+	IAvaMediaModule::Get().GetOnAvaMediaSyncPackageModified().AddRaw(this, &FAvaMediaSyncManager::HandleAvaSyncPackageModified);
+	if (IAssetRegistry* AssetRegistry = IAssetRegistry::Get())
+	{
+		AssetRegistry->OnAssetRemoved().AddRaw(this, &FAvaMediaSyncManager::HandleAssetRemoved);
+	}
+}
+
+void FAvaMediaSyncManager::UnregisterEventHandlers()
+{
+	UPackage::PackageSavedWithContextEvent.RemoveAll(this);
+	IAvaMediaModule::Get().GetOnAvaMediaSyncPackageModified().RemoveAll(this);
+	if (IAssetRegistry* AssetRegistry = IAssetRegistry::Get())
+	{
+		AssetRegistry->OnAssetRemoved().RemoveAll(this);
+	}
+}
+
+void FAvaMediaSyncManager::RequestPackageSyncStatus(const FName& InPackageName)
+{
+	// Check if we already have a pending request and that it hasn't expired.
+	const FDateTime CurrentTime = FDateTime::UtcNow();
+	const FDateTime* PendingRequest = CompareRequests.Find(InPackageName);
+	if (PendingRequest && CurrentTime < *PendingRequest)
+	{
+		UE_LOG(LogAvaPlaybackServer, Verbose,
+			TEXT("Sync Compare of package \"%s\" with remote \"%s\" is already pending."),
+			*InPackageName.ToString(), *RemoteName);
+		return;
+	}
+
+	if (PendingRequest)
+	{
+		HandleExpiredRequest(InPackageName);
+		CompareRequests.Remove(InPackageName);
+		return;
+	}
+	
+	IAvaMediaSyncProvider* AvaMediaSyncProvider = IAvaMediaModule::Get().GetAvaMediaSyncProvider();
+	if (!AvaMediaSyncProvider)
+	{
+		return;
+	}
+	
+	const FDateTime ExpirationTime = CurrentTime + FTimespan::FromSeconds(UAvalancheMediaSettings::Get().ServerPendingStatusRequestTimeout);
+	CompareRequests.Add(InPackageName, ExpirationTime);
+
+	TWeakPtr<FAvaMediaSyncManager> LocalWeakThis = SharedThis(this);
+	AvaMediaSyncProvider->CompareWithRemote(RemoteName,{InPackageName},
+		FOnAvaSyncCompareResponse::CreateLambda([LocalWeakThis, InPackageName](const TSharedPtr<FAvaSyncCompareResponse>& Response)
+		{				
+			if (const TSharedPtr<FAvaMediaSyncManager> LocalThis = LocalWeakThis.Pin())
+			{
+				LocalThis->CompareRequests.Remove(InPackageName);
+				if (Response->IsValid())
+				{
+					LocalThis->HandleSyncStatusReceived(InPackageName, Response->bNeedsSynchronization);
+				}
+				else
+				{
+					LocalThis->HandleFailedRequest(InPackageName, Response->HasError() ? Response->ToString() : FString(TEXT("Unknown error")));
+				}
+			}
+		}));
+}
+
+void FAvaMediaSyncManager::HandleExpiredRequest(const FName& InPackageName)
+{
+	// An expired request means the storm sync layer didn't respond, and is unlikely to respond even
+	// if the request is made again. The best course of action is to consider it non operational and
+	// work as if the assets are up to date.
+	UE_LOG(LogAvaPlaybackServer, Warning,
+		TEXT("Sync Compare of package \"%s\" with remote \"%s\" timed out. Asset will be considered up to date until sync layer can resume operation."),
+		*InPackageName.ToString(), *RemoteName);
+
+	// We want the client to receive a response.
+	HandleSyncStatusReceived(InPackageName, false);
+}
+
+void FAvaMediaSyncManager::HandleFailedRequest(const FName& InPackageName, const FString& InErrorMessage)
+{
+	// For a failed request, we could try again, but since it is not likely to work, the same logic
+	// will be applied as an expired request. The most likely event when an error occurs is for the
+	// request to time out anyway. Haven't encountered an actual error.
+	UE_LOG(LogAvaPlaybackServer, Error,
+		TEXT("Sync Compare of package \"%s\" with remote \"%s\" failed: %s. Asset will be considered up to date until sync layer can resume operation."),
+		*InPackageName.ToString(), *RemoteName, *InErrorMessage);
+
+	// We want the client to receive a response.
+	HandleSyncStatusReceived(InPackageName, false);
+}
+
+void FAvaMediaSyncManager::HandleSyncStatusReceived(const FName& InPackageName, bool bInNeedsSynchronization)
+{
+	PackageStatuses.Add(InPackageName, bInNeedsSynchronization);
+	
+	// Broadcast an event for each asset in the package.
+	if (const IAssetRegistry* AssetRegistry = IAssetRegistry::Get())
+	{
+		IAssetRegistry::FLoadPackageRegistryData RegistryData;
+		if (AssetRegistry->GetAssetsByPackageName(InPackageName, RegistryData.Data))
+		{
+			// Asset Data is not yet in registry's cache, need to load it directly.
+			if (RegistryData.Data.Num() == 0)
+			{
+				FString PackageFilename;
+				if (FPackageName::DoesPackageExist(InPackageName.ToString(), &PackageFilename))
+				{
+					AssetRegistry->LoadPackageRegistryData(PackageFilename, RegistryData);
+				}
+			}
+			
+			for (const FAssetData& Asset : RegistryData.Data)
+			{
+				if (FAvaMediaPlaybackManager::IsPlaybackAsset(Asset))
+				{
+					OnAvaAssetSyncStatusReceived.Broadcast({RemoteName, Asset.GetSoftObjectPath(), bInNeedsSynchronization});
+				}
+			}
+		}
+	}
+}
+
+void FAvaMediaSyncManager::HandleAvaSyncPackageModified(IAvaMediaSyncProvider* InAvaMediaSyncProvider, const FName& InPackageName)
+{
+	if (IAssetRegistry* AssetRegistry = IAssetRegistry::Get())
+	{
+		// If this package just got received from sync provider, it may not be loaded yet.
+		AssetRegistry->WaitForPackage(InPackageName.ToString());
+	}
+	
+	// We had a sync operation for the given package. We could assume it is now up to date.
+	// It would be better if we knew the sync operation came from the remote but this is not available
+	// from the IAvaMediaSyncProvider (should be added).
+	HandleSyncStatusReceived(InPackageName, false);
+}
+
+void FAvaMediaSyncManager::HandlePackageSaved(const FString& InPackageFileName, UPackage* InPackage, FObjectPostSaveContext InObjectSaveContext)
+{
+	const FName PackageName = InPackage->GetFName();
+	if (PackageStatuses.Contains(PackageName))
+	{
+		PackageStatuses.Remove(PackageName);
+
+		// We don't know what the status is anymore, so we request it.
+		RequestPackageSyncStatus(PackageName);
+	}
+}
+
+void FAvaMediaSyncManager::HandleAssetRemoved(const FAssetData& InAssetData)
+{
+	if (PackageStatuses.Contains(InAssetData.PackageName))
+	{
+		// Invalidate status.
+		PackageStatuses.Remove(InAssetData.PackageName);
+
+		// If there are remaining assets in the package, we will ask for a status update.
+		if (const IAssetRegistry* AssetRegistry = IAssetRegistry::Get())
+		{
+			if (AssetRegistry->HasAssets(InAssetData.PackageName))
+			{
+				RequestPackageSyncStatus(InAssetData.PackageName);
+			}
+		}
+	}
+}
\ No newline at end of file
