[MotionDesign]

diff --git a/Engine/Plugins/Experimental/Avalanche/Source/AvalancheMedia/Private/Playback/AvaMediaSyncManager.h b/Engine/Plugins/Experimental/Avalanche/Source/AvalancheMedia/Private/Playback/AvaMediaSyncManager.h
new file mode 100644
--- /dev/null
+++ b/Engine/Plugins/Experimental/Avalanche/Source/AvalancheMedia/Private/Playback/AvaMediaSyncManager.h
@@ -0,0 +1,83 @@
+// Copyright Epic Games, Inc. All Rights Reserved.
+
+#pragma once
+
+#include "Delegates/Delegate.h"
+#include "Misc/DateTime.h"
+#include "UObject/ObjectSaveContext.h"
+
+class IAvaMediaSyncProvider;
+struct FAssetData;
+
+struct FAvaAssetSyncStatusReceivedParams
+{
+	const FString& RemoteName;
+	const FSoftObjectPath& AssetPath;
+	bool bNeedsSync;
+};
+
+/**
+ * Implements tracking and caching the status of local assets compared to the corresponding assets
+ * on a given remote. The remote assets are considered to be the reference. For each tracked
+ * assets, a status of "need sync" (true) or "up to date" (false) is kept.
+ */
+class FAvaMediaSyncManager : public TSharedFromThis<FAvaMediaSyncManager>
+{
+public:
+	FAvaMediaSyncManager(const FString& InRemoteName);
+	~FAvaMediaSyncManager();
+
+	bool IsEnabled() const { return bEnabled; }
+	void SetEnable(bool bInEnabled);
+	
+	void Tick();
+	
+	/**
+	 * Returns the last cached asset sync status for the given asset if available.
+	 * The sync status can be true (i.e. need sync) or false (i.e. up to date).
+	 * If it is not available, a sync compare request is made with the specified
+	 * remote and the result is posted through OnAvaAssetSyncStatusReceived event
+	 * when it is received.
+	 * If bInForceRefresh is true, a request will be made even if the status was already cached.
+	 */
+	TOptional<bool> GetAssetSyncStatus(const FSoftObjectPath& InAssetPath, bool bInForceRefresh);
+
+	/**
+	 * Enumerate all tracked packages in the manager.
+	 */
+	void EnumerateAllTrackedPackages(TFunctionRef<void(const FName& /*InPackageName*/, bool /*bInNeedsSync*/)> InFunction) const;
+
+	TArray<FName> GetPendingRequests() const;
+	
+	/**
+	 *	Delegate called when an asset sync status changed.
+	 */
+	DECLARE_MULTICAST_DELEGATE_OneParam(FOnAvaAssetSyncStatusReceived, const FAvaAssetSyncStatusReceivedParams&);
+	FOnAvaAssetSyncStatusReceived OnAvaAssetSyncStatusReceived;
+
+private:
+	void RegisterEventHandlers();
+	void UnregisterEventHandlers();
+	
+	void RequestPackageSyncStatus(const FName& InPackageName);
+	void HandleExpiredRequest(const FName& InPackageName);
+	void HandleFailedRequest(const FName& InPackageName, const FString& InErrorMessage);
+	void HandleSyncStatusReceived(const FName& InPackageName, bool bInNeedsSynchronization);
+	
+	// Event handlers
+	void HandleAvaSyncPackageModified(IAvaMediaSyncProvider* InAvaMediaSyncProvider, const FName& InPackageName);
+	void HandlePackageSaved(const FString& InPackageFileName, UPackage* InPackage, FObjectPostSaveContext InObjectSaveContext);
+	void HandleAssetRemoved(const FAssetData& InAssetData);
+
+	/** Corresponding remote to be comparing with. */
+	FString RemoteName;
+
+	/** Determines if the sync manager is going to be used for the corresponding remote. */
+	bool bEnabled = true;
+	
+	/** Maps package names and sync status. */
+	TMap<FName, bool> PackageStatuses;
+
+	/** Currently pending sync compare requests. Key is long package names. */
+	TMap<FName, FDateTime> CompareRequests;
+};
\ No newline at end of file
diff --git a/Engine/Plugins/Experimental/Avalanche/Source/AvalancheMedia/Private/Playback/AvalanchePlayable.cpp b/Engine/Plugins/Experimental/Avalanche/Source/AvalancheMedia/Private/Playback/AvalanchePlayable.cpp
new file mode 100644
--- /dev/null
+++ b/Engine/Plugins/Experimental/Avalanche/Source/AvalancheMedia/Private/Playback/AvalanchePlayable.cpp
@@ -0,0 +1,386 @@
+// Copyright Epic Games, Inc. All Rights Reserved.
+
+#include "Playback/AvalanchePlayable.h"
+
+#include "AvaSequence.h"
+#include "AvaSequencePlayer.h"
+#include "AvalancheBroadcast.h"
+#include "Engine/Engine.h"
+#include "Framework/AvaSoftAssetPtr.h"
+#include "IAvaSceneInterface.h"
+#include "IAvaSequenceProvider.h"
+#include "Playback/AvaMediaPlayableGroup.h"
+#include "Playback/AvaMediaPlaybackUtils.h"
+#include "Playback/AvalancheRemoteControlValues.h"
+#include "Playback/Playables/AvalancheBlueprintPlayable.h"
+#include "Playback/Playables/AvalancheLevelStreamingPlayable.h"
+#include "Playback/Playables/AvalancheRemoteProxyPlayable.h"
+#include "RemoteControlPreset.h"
+
+DEFINE_LOG_CATEGORY(LogAvalanchePlayable);
+
+#define LOCTEXT_NAMESPACE "AvalanchePlayable"
+
+namespace UE::AvalanchePlayable::Private
+{
+	bool ShouldCreateLocalPlayable(const FName& InChannelName, const UAvalancheBroadcast& InBroadcast)
+	{
+		const FAvaOutputChannel& Channel = InBroadcast.GetCurrentProfile().GetChannel(InChannelName);
+
+		// If there is no broadcast channel defined, like for preview (by default), then this is a local playable.
+		if (!Channel.IsValidChannel())
+		{
+			return true;
+		}
+
+		// For non-preview, the commands will be executed locally if the channel has at least one local outputs or no outputs.
+		// The "no outputs" condition is considered valid. Empty channels run locally.
+		if (Channel.HasAnyLocalMediaOutputs() || Channel.GetMediaOutputs().IsEmpty())
+		{
+			return true;
+		}
+		
+		return false;
+	}
+
+	bool HasRemoteOutputs(const FName& InChannelName, const UAvalancheBroadcast& InBroadcast)
+	{
+		const FAvaOutputChannel& Channel = InBroadcast.GetCurrentProfile().GetChannel(InChannelName);
+		return Channel.IsValidChannel() && Channel.HasAnyRemoteMediaOutputs();
+	}
+
+	FString GetPrettyPlayableInfo(const UAvalanchePlayable* InPlayable)
+	{
+		if (InPlayable)
+		{
+			return FString::Printf(TEXT("Id:%s, Asset:%s, Status:%s"),
+				*InPlayable->GetInstanceId().ToString(),
+				*InPlayable->GetSourceAssetPath().ToString(),
+				*StaticEnum<EAvalanchePlayableStatus>()->GetNameByValue(static_cast<int32>(InPlayable->GetPlayableStatus())).ToString());
+		}
+		return TEXT("(nullptr)");
+	}
+
+	FString GetBriefFrameInfo()
+	{
+		return UE::AvaMediaPlayback::Utils::GetBriefFrameInfo();
+	}
+	
+	FString GetPrettySequenceCommandInfo(EAvaMediaAnimAction InAnimAction, const FAnimPlaySettings& InAnimPlaySettings)
+	{
+		return FString::Printf(TEXT("Action:%s, Name:%s"),
+			*StaticEnum<EAvaMediaAnimAction>()->GetNameByValue(static_cast<int32>(InAnimAction)).ToString(),
+			*InAnimPlaySettings.AnimationName.ToString());
+	}
+}
+
+UAvalanchePlayable::FOnSequenceEvent UAvalanchePlayable::OnSequenceEventDelegate;
+UAvalanchePlayable::FOnTransitionEvent UAvalanchePlayable::OnTransitionEventDelegate;
+
+UAvalanchePlayable* UAvalanchePlayable::Create(UObject* InOuter, const FPlayableCreationInfo& InPlayableInfo)
+{
+	using namespace UE::AvalanchePlayable::Private;
+
+	const UAvalancheBroadcast& Broadcast = UAvalancheBroadcast::Get();
+	
+	UAvalanchePlayable* NewPlayable;
+	
+	// Forked channels considerations:
+	// - The case of forked remote channels is/will be handled internally to the RemoteProxy playable.
+	// - The case of forked local and remote channels will lead to a local playable and a remote proxy playable.
+	//   It would require wrapping the playables in a composite (or facade?) proxy. TODO
+	
+	if (ShouldCreateLocalPlayable(InPlayableInfo.ChannelName, Broadcast))
+	{
+		// For the moment, remote outputs will be ignored.
+		if (HasRemoteOutputs(InPlayableInfo.ChannelName, Broadcast))
+		{
+			UE_LOG(LogAvalanchePlayable, Error, TEXT("Forked Channels with both local and remote outputs are not supported in this version. Only local instance will be created."));
+		}
+		
+		NewPlayable = CreateLocalPlayable(InOuter, InPlayableInfo);
+	}
+	else
+	{
+		// Purely remote channel.
+		NewPlayable = CreateRemoteProxyPlayable(InOuter, InPlayableInfo);
+	}
+
+	// Finish the setup.
+	if (NewPlayable && !NewPlayable->InitPlayable(InPlayableInfo))
+	{
+		// final setup may fail, in this case the playable is discarded.
+		return nullptr;
+	}
+
+	return NewPlayable;
+}
+
+const FSoftObjectPath& UAvalanchePlayable::GetSourceAssetPath() const
+{
+	static const FSoftObjectPath Empty;
+	return Empty;
+}
+
+EAvalanchePlayableCommandResult UAvalanchePlayable::ExecuteAnimationCommand(EAvaMediaAnimAction InAnimAction, const FAnimPlaySettings& InAnimPlaySettings)
+{
+	using namespace UE::AvalanchePlayable::Private;
+
+	const EAvalanchePlayableStatus PlayableStatus = GetPlayableStatus();
+
+	if (PlayableStatus == EAvalanchePlayableStatus::Unknown
+		|| PlayableStatus == EAvalanchePlayableStatus::Error
+		|| PlayableStatus == EAvalanchePlayableStatus::Unloaded)
+	{
+		UE_LOG(LogAvalanchePlayable, Verbose,
+			TEXT("%s Playable {%s} -> Discarding Sequence Command: {%s}."),
+			*GetBriefFrameInfo(), *GetPrettyPlayableInfo(this), *GetPrettySequenceCommandInfo(InAnimAction, InAnimPlaySettings));
+	
+		// Discard the command
+		return EAvalanchePlayableCommandResult::ErrorDiscard;
+	}
+
+	// Asset status must be visible to run the animation commands.
+	// If not visible, the components are not yet added to the world and animations won't execute.
+	if (PlayableStatus != EAvalanchePlayableStatus::Visible)
+	{
+		UE_LOG(LogAvalanchePlayable, Verbose,
+			TEXT("%s Playable {%s} -> ReQueueing Sequence Command: {%s}."),
+			*GetBriefFrameInfo(), *GetPrettyPlayableInfo(this), *GetPrettySequenceCommandInfo(InAnimAction, InAnimPlaySettings));
+		
+		// Keep the command in the queue for next tick.
+		return EAvalanchePlayableCommandResult::KeepPending;
+	}
+
+	const IAvaSceneInterface* AvalancheScene = GetSceneInterface();
+	if (!AvalancheScene)
+	{
+		return EAvalanchePlayableCommandResult::ErrorDiscard;
+	}
+	
+	IAvaSequencePlaybackObject* PlaybackObject = AvalancheScene->GetPlaybackObject();
+	
+	if (!PlaybackObject)
+	{
+		return EAvalanchePlayableCommandResult::ErrorDiscard;
+	}
+
+	const IAvaSequenceProvider* SequenceProvider = AvalancheScene->GetSequenceProvider();
+
+	if (!SequenceProvider)
+	{
+		return EAvalanchePlayableCommandResult::ErrorDiscard;
+	}
+
+	UE_LOG(LogAvalanchePlayable, Verbose,
+		TEXT("%s Playable {%s} -> Executing Sequence Command: {%s}."),
+		*GetBriefFrameInfo(), *GetPrettyPlayableInfo(this), *GetPrettySequenceCommandInfo(InAnimAction, InAnimPlaySettings));
+	
+	for (const TObjectPtr<UAvaSequence>& Sequence : SequenceProvider->GetSequences())
+	{
+		// Remark: if the command doesn't specify the sequence name, we run the command on all the sequences.
+		if (Sequence && (Sequence->GetFName() == InAnimPlaySettings.AnimationName || InAnimPlaySettings.AnimationName.IsNone()))
+		{
+			if (InAnimAction == EAvaMediaAnimAction::Play || InAnimAction == EAvaMediaAnimAction::PreviewFrame)
+			{
+				UAvaSequencePlayer* const SequencePlayer = PlaybackObject->PlaySequence(Sequence, InAnimPlaySettings.AsPlayParams());
+				if (InAnimAction == EAvaMediaAnimAction::PreviewFrame && SequencePlayer)
+				{
+					SequencePlayer->PreviewFrame();
+				}
+			}
+			else if (InAnimAction == EAvaMediaAnimAction::Continue)
+			{
+				PlaybackObject->ContinueSequence(Sequence);
+			}
+			else if (InAnimAction == EAvaMediaAnimAction::Stop)
+			{
+				PlaybackObject->StopSequence(Sequence);
+			}
+		}
+	}
+	
+	return EAvalanchePlayableCommandResult::Executed;
+}
+
+EAvalanchePlayableCommandResult UAvalanchePlayable::UpdateRemoteControlCommand(const TSharedRef<FAvalancheRemoteControlValues>& InRemoteControlValues)
+{
+	const EAvalanchePlayableStatus PlayableStatus = GetPlayableStatus();
+
+	if (PlayableStatus == EAvalanchePlayableStatus::Unknown
+		|| PlayableStatus == EAvalanchePlayableStatus::Error
+		|| PlayableStatus == EAvalanchePlayableStatus::Unloaded)
+	{
+		using namespace UE::AvalanchePlayable::Private;
+		UE_LOG(LogAvalanchePlayable, Verbose, TEXT("%s Playable {%s} -> Discarding RC Update."), *GetBriefFrameInfo(), *GetPrettyPlayableInfo(this));
+
+		// Discard the command
+		return EAvalanchePlayableCommandResult::ErrorDiscard;
+	}
+
+	// Asset status must be visible to run the command.
+	// If not visible, the components are not yet added to the world.
+	if (PlayableStatus != EAvalanchePlayableStatus::Visible)
+	{
+		using namespace UE::AvalanchePlayable::Private;
+		UE_LOG(LogAvalanchePlayable, Verbose, TEXT("%s Playable {%s} -> ReQueueing RC Update."), *GetBriefFrameInfo(), *GetPrettyPlayableInfo(this));
+
+		// Keep the command in the queue for next tick.
+		return EAvalanchePlayableCommandResult::KeepPending;
+	}
+
+	const IAvaSceneInterface* AvalancheScene = GetSceneInterface();
+	if (!AvalancheScene)
+	{
+		return EAvalanchePlayableCommandResult::ErrorDiscard;
+	}
+	
+	URemoteControlPreset* RemoteControlPreset = AvalancheScene->GetRemoteControlPreset();
+	
+	if (!IsValid(RemoteControlPreset))
+	{
+		UE_LOG(LogAvalanchePlayable, Error,
+			TEXT("Remote Control command for asset \"%s\": Remote Control Preset is null."),
+			*GetSourceAssetPath().ToString());
+		return EAvalanchePlayableCommandResult::ErrorDiscard;
+	}
+
+	LatestRemoteControlValues = InRemoteControlValues;
+
+	using namespace UE::AvalanchePlayable::Private;
+	UE_LOG(LogAvalanchePlayable, Verbose, TEXT("%s Playable {%s} -> Executing RC Update."), *GetBriefFrameInfo(), *GetPrettyPlayableInfo(this));
+
+	// WYSIWYG (Solution): For the runtime/playback RCP, we don't apply the controllers.
+	// We assume the controller actions are already executed in the playlist's managed RCP
+	// during page edition and the resulting entity values are already captured.
+	InRemoteControlValues->ApplyEntityValuesToRemoteControlPreset(RemoteControlPreset);
+
+	return EAvalanchePlayableCommandResult::Executed;
+}
+
+void UAvalanchePlayable::BeginPlay(const FAvalancheInstancePlaySettings& InWorldPlaySettings)
+{
+	if (!PlayableGroup)
+	{
+		return;
+	}
+
+	const bool bGroupHasBegunPlay = PlayableGroup->ConditionalBeginPlay(InWorldPlaySettings);
+	
+	if (!bIsPlaying || bGroupHasBegunPlay)
+	{
+		bIsPlaying = true;
+		
+		// Playable events need to transit through playback events to reach the playlist for proper impl layer separation.
+		UAvaSequencePlayer::OnSequenceStarted().AddUObject(this, &UAvalanchePlayable::HandleOnSequenceStarted);
+		UAvaSequencePlayer::OnSequenceFinished().AddUObject(this, &UAvalanchePlayable::HandleOnSequenceFinished);
+
+		OnPlay();
+	}
+}
+
+void UAvalanchePlayable::EndPlay(EAvalanchePlayableEndPlayOptions InOptions)
+{
+	if (!bIsPlaying)
+	{
+		return;
+	}
+
+	bIsPlaying = false;
+	UAvaSequencePlayer::OnSequenceStarted().RemoveAll(this);
+	UAvaSequencePlayer::OnSequenceFinished().RemoveAll(this);
+	OnEndPlay();
+
+	if (PlayableGroup)
+	{
+		PlayableGroup->UpdateCameraSetup();
+
+		if (EnumHasAnyFlags(InOptions, EAvalanchePlayableEndPlayOptions::ConditionalEndPlayWorld) && !PlayableGroup->HasPlayingPlayables())
+		{
+			const bool bForceImmediate = EnumHasAnyFlags(InOptions, EAvalanchePlayableEndPlayOptions::ForceImmediate);
+			PlayableGroup->RequestEndPlayWorld(bForceImmediate);
+		}
+	}
+}
+
+bool UAvalanchePlayable::HasSequence(const UAvaSequence* InSequence) const
+{
+	const IAvaSceneInterface* SceneInterface = GetSceneInterface();
+	if (!SceneInterface)
+	{
+		return false;
+	}
+	
+	const IAvaSequenceProvider* SequenceProvider = SceneInterface->GetSequenceProvider();
+	if (!SequenceProvider)
+	{
+		return false;
+	}
+	
+	for (const TObjectPtr<UAvaSequence>& Sequence : SequenceProvider->GetSequences())
+	{
+		if (Sequence == InSequence)
+		{
+			return true;
+		}
+	}
+	return false;
+}
+
+bool UAvalanchePlayable::InitPlayable(const FPlayableCreationInfo& InPlayableInfo)
+{
+	if (PlayableGroup)
+	{
+		// Register this playable in the instance group.
+		// This is necessary to determine what is playing in what group.
+		PlayableGroup->RegisterPlayable(this);
+		return true;
+	}
+	
+	// Currently, playables must have a playable group otherwise they are unplayable.
+	UE_LOG(LogAvalanchePlayable, Error, TEXT("Failed to create or acquire a playable group for \"%s\". Playable will be discarded."), *InPlayableInfo.SourceAsset.ToSoftObjectPath().ToString());
+	return false;
+}
+
+void UAvalanchePlayable::HandleOnSequenceStarted(UAvaSequencePlayer* InSequencePlayer, UAvaSequence* InSequence)
+{
+	if (HasSequence(InSequence))
+	{
+		using namespace UE::AvalanchePlayable::Private;
+		UE_LOG(LogAvalanchePlayable, Verbose, TEXT("%s Playable {%s}: Sequence \"%s\" started."),
+			*GetBriefFrameInfo(), *GetPrettyPlayableInfo(this), *InSequence->GetFName().ToString());
+		OnSequenceEventDelegate.Broadcast(this, InSequence->GetFName(), EAvalanchePlayableSequenceEventType::Started);
+	}
+}
+
+void UAvalanchePlayable::HandleOnSequenceFinished(UAvaSequencePlayer* InSequencePlayer, UAvaSequence* InSequence)
+{
+	if (HasSequence(InSequence))
+	{
+		using namespace UE::AvalanchePlayable::Private;
+		UE_LOG(LogAvalanchePlayable, Verbose, TEXT("%s Playable {%s}: Sequence \"%s\" finished."),
+			*GetBriefFrameInfo(), *GetPrettyPlayableInfo(this), *InSequence->GetFName().ToString());
+		OnSequenceEventDelegate.Broadcast(this, InSequence->GetFName(), EAvalanchePlayableSequenceEventType::Finished);
+	}
+}
+
+UAvalanchePlayable* UAvalanchePlayable::CreateLocalPlayable(UObject* InOuter, const FPlayableCreationInfo& InPlayableInfo)
+{
+	switch (InPlayableInfo.SourceAsset.GetAssetType())
+	{
+	case EAvalancheAssetType::Blueprint:
+		return NewObject<UAvalancheBlueprintPlayable>(InOuter ? InOuter : GEngine);
+	case EAvalancheAssetType::World:
+		return NewObject<UAvalancheLevelStreamingPlayable>(InOuter ? InOuter : GEngine);
+	default:
+		UE_LOG(LogAvalanchePlayable, Error, TEXT("Asset \"%s\" is an unsupported type."), *InPlayableInfo.SourceAsset.ToSoftObjectPath().ToString());
+		return nullptr;
+	}
+}
+
+UAvalanchePlayable* UAvalanchePlayable::CreateRemoteProxyPlayable(UObject* InOuter, const FPlayableCreationInfo& InPlayableInfo)
+{
+	return NewObject<UAvalancheRemoteProxyPlayable>(InOuter ? InOuter : GEngine);
+}
+
+#undef LOCTEXT_NAMESPACE
diff --git a/Engine/Plugins/Experimental/Avalanche/Source/AvalancheMedia/Private/Playback/AvalanchePlayback.cpp b/Engine/Plugins/Experimental/Avalanche/Source/AvalancheMedia/Private/Playback/AvalanchePlayback.cpp
new file mode 100644
--- /dev/null
+++ b/Engine/Plugins/Experimental/Avalanche/Source/AvalancheMedia/Private/Playback/AvalanchePlayback.cpp
@@ -0,0 +1,1023 @@
+// Copyright Epic Games, Inc. All Rights Reserved.
+
+#include "Playback/AvalanchePlayback.h"
+
+#include "Async/Async.h"
+#include "AvaRenderTargetUtils.h"
+#include "AvalancheBroadcast.h"
+#include "AvalancheMediaSettings.h"
+#include "Engine/TextureRenderTarget2D.h"
+#include "Framework/AvalancheGameInstance.h"
+#include "IAvaMediaModule.h"
+#include "IAvalancheBroadcastSettings.h"
+#include "Misc/CoreDelegates.h"
+#include "Playback/AvaMediaPlayableGroup.h"
+#include "Playback/AvaMediaPlayableGroupManager.h"
+#include "Playback/AvaMediaPlaybackManager.h"
+#include "Playback/AvalanchePlayable.h"
+#include "Playback/IAvaPlaybackEditor.h"
+#include "Playback/Nodes/AvaPlaybackNode.h"
+#include "Playback/Playables/AvalancheBlueprintPlayable.h"
+
+#if WITH_EDITOR
+#include "EdGraph/EdGraph.h"
+#include "EdGraph/EdGraphSchema.h"
+#endif
+
+DEFINE_LOG_CATEGORY(LogAvalanchePlayback);
+
+UAvalanchePlayback::~UAvalanchePlayback()
+{
+	FAvaOutputChannel::GetOnChannelChanged().RemoveAll(this);
+}
+
+void UAvalanchePlayback::Play()
+{
+	if (!FAvaOutputChannel::GetOnChannelChanged().IsBoundToObject(this))
+	{
+		FAvaOutputChannel::GetOnChannelChanged().AddUObject(this, &UAvalanchePlayback::OnChannelChanged);
+	}
+
+	if (!IsPlaying())
+	{
+		SetIsPlaying(true);
+	}	
+}
+
+void UAvalanchePlayback::LoadInstances()
+{
+	// Gather the asset references from the player nodes.
+	ConditionalResolvePlaybackSettings();
+
+	for (TPair<FName, FAvaPlaybackChannelParameters>& Playback : PlaybackSettings)
+	{
+		// Load the assets
+		for (const FAvaSoftAssetPtr& AvalancheAsset : Playback.Value.AvalancheAssets)
+		{
+			FindOrLoadPlayable(AvalancheAsset, ResolveChannelName(Playback.Key));
+		}
+	}
+}
+
+// Unload playables.
+void UAvalanchePlayback::UnloadInstances(EAvaPlaybackUnloadOptions InUnloadOptions)
+{
+	// We need to unload all the playables, not just the ones currently traversed by the graph.
+	// We keep track of all the local playables in ChannelPlayableGroups. This now includes the
+	// remote playables.
+	
+	const bool bIsForceImmediate = EnumHasAnyFlags(InUnloadOptions, EAvaPlaybackUnloadOptions::ForceImmediate);
+
+	for (const TPair<FName, FAvalanchePlaybackPlayableGroup>& ChannelPlayableGroup : ChannelPlayableGroups)
+	{
+		TArray<UAvalanchePlayable*> Playables;
+		ChannelPlayableGroup.Value.GetAllPlayables(Playables);
+		
+		for (UAvalanchePlayable* Playable : Playables)
+		{
+			FSoftObjectPath SourceAsset = Playable->GetSourceAssetPath();
+			UnloadAndRemovePlayable(Playable, SourceAsset, ChannelPlayableGroup.Key, bIsForceImmediate);
+		}
+	}
+}
+
+void UAvalanchePlayback::Stop(EAvaPlaybackStopOptions InStopOptions)
+{
+	const bool bIsForceImmediate = EnumHasAnyFlags(InStopOptions, EAvaPlaybackStopOptions::ForceImmediate);
+	const bool bIsUnload = EnumHasAnyFlags(InStopOptions, EAvaPlaybackStopOptions::Unload);
+
+	// We need to unload all the playables, not just the ones currently traversed by the graph.
+	// We keep track of all the local playables in ChannelPlayableGroups. This now includes the
+	// remote playables.
+
+	for (const TPair<FName, FAvalanchePlaybackPlayableGroup>& ChannelPlayableGroup : ChannelPlayableGroups)
+	{
+		TArray<UAvalanchePlayable*> Playables;
+		ChannelPlayableGroup.Value.GetAllPlayables(Playables);
+		
+		for (UAvalanchePlayable* Playable : Playables)
+		{
+			// EndPlayWorld is no longer issued here as some of the engine systems don't seem
+			// to support a EndPlay, BeginPlay sequence anymore. It is unknown if this situation
+			// will change in the future, so we keep the option.
+			Playable->EndPlay(EAvalanchePlayableEndPlayOptions::None);
+
+			if (bIsUnload)
+			{
+				FSoftObjectPath SourceAsset = Playable->GetSourceAssetPath();
+				UnloadAndRemovePlayable(Playable, SourceAsset, ChannelPlayableGroup.Key, bIsForceImmediate);
+			}
+		}
+	}
+
+	PlaybackSettings.Reset();
+	PreviousPlaybackSettings.Reset();
+	SetIsPlaying(false);
+}
+
+void UAvalanchePlayback::SetPlaybackManager(const TSharedPtr<FAvaMediaPlaybackManager>& InPlaybackManager)
+{
+	PlaybackManagerWeak = InPlaybackManager;
+
+	if (IsManaged())
+	{
+		// Since the playback manager is going to manage ticking,
+		// make sure to unregister the tick delegate.
+		UnregisterTickDelegate();
+	}
+	else if (IsPlaying() && !TickDelegateHandle.IsValid())
+	{
+		RegisterTickDelegate();
+	}
+}
+
+UAvalanchePlayable* UAvalanchePlayback::GetFirstPlayable() const
+{
+	for (const TPair<FName, FAvalanchePlaybackPlayableGroup>& ChannelPlayableGroup : ChannelPlayableGroups)
+	{
+		if (UAvalanchePlayable* Playable = ChannelPlayableGroup.Value.GetFirstPlayable())
+		{
+			return Playable;
+		}
+	}
+	return nullptr;
+}
+
+void UAvalanchePlayback::GetAllPlayables(TArray<UAvalanchePlayable*>& OutPlayables) const
+{
+	for (const TPair<FName, FAvalanchePlaybackPlayableGroup>& ChannelPlayableGroup : ChannelPlayableGroups)
+	{
+		ChannelPlayableGroup.Value.GetAllPlayables(OutPlayables);
+	}
+}
+
+void UAvalanchePlayback::ForEachPlayable(TFunctionRef<void(const UAvalanchePlayable*)> InFunction) const
+{
+	for (const TPair<FName, FAvalanchePlaybackPlayableGroup>& ChannelPlayableGroup : ChannelPlayableGroups)
+	{
+		ChannelPlayableGroup.Value.ForEachPlayable(InFunction);
+	}
+}
+
+bool UAvalanchePlayback::HasPlayable(const UAvalanchePlayable* InPlayable) const
+{
+	for (const TPair<FName, FAvalanchePlaybackPlayableGroup>& ChannelPlayableGroup : ChannelPlayableGroups)
+	{
+		if (ChannelPlayableGroup.Value.HasPlayable(InPlayable))
+		{
+			return true;
+		}
+	}
+	return false;
+}
+
+UAvalanchePlayable* UAvalanchePlayback::FindPlayable(const FSoftObjectPath& InSourceAssetPath, const FName& InChannelName) const
+{
+	const FAvalanchePlaybackPlayableGroup* ChannelPlayableGroup = ChannelPlayableGroups.Find(InChannelName);
+	return ChannelPlayableGroup ? ChannelPlayableGroup->FindPlayable(InSourceAssetPath) : nullptr;
+}
+
+UAvalanchePlayable* UAvalanchePlayback::FindOrLoadPlayable(const FAvaSoftAssetPtr& InSourceAsset, const FName& InChannelName)
+{
+	FAvalanchePlaybackPlayableGroup* ChannelPlayableGroup = ChannelPlayableGroups.Find(InChannelName);
+	if (!ChannelPlayableGroup)
+	{
+		ChannelPlayableGroup = &ChannelPlayableGroups.Add(InChannelName, FAvalanchePlaybackPlayableGroup());
+	}
+	
+	if (const TObjectPtr<UAvalanchePlayable>* FoundPlayable = ChannelPlayableGroup->Playables.Find(InSourceAsset.ToSoftObjectPath()))
+	{
+		return *FoundPlayable;
+	}
+
+	UAvalanchePlayable* NewPlayable = UAvalanchePlayable::Create(this, {GetPlayableGroupManager(), InSourceAsset, InChannelName});
+	if (NewPlayable)
+	{
+		ChannelPlayableGroup->Playables.Add(InSourceAsset.ToSoftObjectPath(), NewPlayable);
+		OnPlayableCreated.Broadcast(this, NewPlayable);
+		NewPlayable->LoadAsset(InSourceAsset, false);
+	}
+	return NewPlayable;
+}
+
+void UAvalanchePlayback::RemovePlayable(const FSoftObjectPath& InSourceAssetPath, const FName& InChannelName)
+{
+	if (FAvalanchePlaybackPlayableGroup* ChannelPlayableGroup = ChannelPlayableGroups.Find(InChannelName))
+	{
+		ChannelPlayableGroup->Playables.Remove(InSourceAssetPath);
+	}
+}
+
+bool UAvalanchePlayback::UnloadAndRemovePlayable(UAvalanchePlayable* InPlayable, const FSoftObjectPath& InSourceAssetPath, const FName& InChannelName, bool bInForceImmediate)
+{
+	bool bWorldUnloaded = false;
+	if (InPlayable)
+	{
+		InPlayable->UnloadAsset(); // bIsForceImmediate is not used anymore.
+		InPlayable->GetPlayableGroup()->UnregisterPlayable(InPlayable);
+		// If no more playables in the group, we can request unload of the world too.
+		bWorldUnloaded = InPlayable->GetPlayableGroup()->ConditionalRequestUnloadWorld(bInForceImmediate);
+	}	
+	RemovePlayable(InSourceAssetPath, InChannelName);
+	return bWorldUnloaded;
+}
+
+bool UAvalanchePlayback::UnloadAndRemovePlayable(const FSoftObjectPath& InSourceAssetPath, const FName& InChannelName, bool bInForceImmediate)
+{
+	return UnloadAndRemovePlayable(FindPlayable(InSourceAssetPath, InChannelName), InSourceAssetPath, InChannelName, bInForceImmediate);
+}
+
+TArray<UAvalancheGameInstance*> UAvalanchePlayback::GetActiveGameInstances() const
+{
+	TArray<UAvalancheGameInstance*> OutGameInstances;
+
+	if (IsPlaying())
+	{
+		// TODO: revisit this. We have Active and Loaded.
+		for (const TPair<FName, FAvaPlaybackChannelParameters>& Playback : PlaybackSettings)
+		{
+			FName ChannelName = ResolveChannelName(Playback.Key);
+			for (const FAvaSoftAssetPtr& AvalancheAsset : Playback.Value.AvalancheAssets)
+			{
+				const UAvalanchePlayable* Playable = FindPlayable(AvalancheAsset.ToSoftObjectPath(), ChannelName);
+				if (Playable && Playable->GetPlayableGroup())
+				{
+					if (UAvalancheGameInstance* AvaGameInstance = Cast<UAvalancheGameInstance>(Playable->GetPlayableGroup()->GetGameInstance()))
+					{
+						OutGameInstances.AddUnique(AvaGameInstance);
+					}
+				}
+			}
+		}
+	}
+	return OutGameInstances;
+}
+
+bool UAvalanchePlayback::HasPlayerNodeForSourceAvalancheAsset(const FSoftObjectPath& InSourceAssetPath) const
+{
+	for (const TObjectPtr<UAvaPlaybackNodePlayer>& PlayerNode : PlayerNodes)
+	{
+		if (PlayerNode && PlayerNode->GetAvalancheAssetPath() == InSourceAssetPath)
+		{
+			return true;
+		}
+	}
+	return false;
+}
+
+void UAvalanchePlayback::SetPlayableGroupManager(UAvaMediaPlayableGroupManager* InPlayableGroupManager)
+{
+	PlayableGroupManager = InPlayableGroupManager;
+}
+
+UAvaMediaPlayableGroupManager* UAvalanchePlayback::GetGlobalPlayableGroupManager()
+{
+	return IAvaMediaModule::Get().GetLocalPlaybackManager().GetPlayableGroupManager();	
+}
+
+void UAvalanchePlayback::DryRunGraph(bool bDeferredExecution)
+{
+	if (bDeferredExecution)
+	{
+		//Only async dry run if we have not requested it yet, and it's not dry running already.
+		if (!bAsyncDryRunRequested && !bIsDryRunningGraph)
+		{
+			bAsyncDryRunRequested = true;
+			
+			TWeakObjectPtr<UAvalanchePlayback> ThisWeak(this);
+			AsyncTask(ENamedThreads::GameThread, [ThisWeak]()
+			{
+				if (ThisWeak.IsValid())
+				{
+					ThisWeak->DryRunGraphInternal();
+				}
+			});
+		}
+	}
+	else
+	{
+		DryRunGraphInternal();
+	}
+}
+
+TArray<FName> UAvalanchePlayback::GetChannelNamesForIndices(const TArray<int32>& InChannelIndices) const
+{
+	TArray<FName> OutChannelNames;
+	if (IsPreviewOnly())
+	{
+		OutChannelNames.Reserve(1);
+		OutChannelNames.Add(PreviewChannelName);
+	}
+	else if (InChannelIndices.Num() > 0)
+	{
+		const UAvalancheBroadcast& AvaBroadcast = UAvalancheBroadcast::Get();
+		OutChannelNames.Reserve(InChannelIndices.Num());
+		for (const int32 ChannelIndex : InChannelIndices)
+		{
+			const FName ChannelFName = AvaBroadcast.GetChannelName(ChannelIndex);
+			if (!ChannelFName.IsNone())
+			{
+				OutChannelNames.Add(ChannelFName);
+			}
+		}
+	}
+	return OutChannelNames;
+}
+
+void UAvalanchePlayback::DryRunGraphInternal()
+{
+	//Everytime we run this, the Task has been complete.
+	bAsyncDryRunRequested = false;
+	
+	if (!bIsDryRunningGraph)
+	{
+		TGuardValue<bool> Guard(bIsDryRunningGraph, true);
+		if (RootNode)
+		{
+			TSet<UAvaPlaybackNode*> SeenNodes;
+			TArray<UAvaPlaybackNode*> RemainingNodes;
+			
+			auto NotifyDryRun = [this, &SeenNodes, &RemainingNodes](void(UAvaPlaybackNode::*ExecFunc)())
+				{
+					//Reset but not deallocate the Memory that was used when previous NotifyDryRuns were called
+					RemainingNodes.Reset();
+					SeenNodes.Reset();
+					
+					RemainingNodes.Add(RootNode);
+					
+					while (RemainingNodes.Num() > 0)
+					{
+						UAvaPlaybackNode* const Node = RemainingNodes.Pop();
+						if (Node && !SeenNodes.Contains(Node))
+						{
+							SeenNodes.Add(Node);
+							(*Node.*ExecFunc)();
+							RemainingNodes.Append(Node->GetChildNodes());
+						}
+					}
+				};
+
+			TArray<UAvaPlaybackNode*> Ancestors;
+			
+			NotifyDryRun(&UAvaPlaybackNode::PreDryRun);
+			RootNode->DryRunNode(Ancestors);
+			NotifyDryRun(&UAvaPlaybackNode::PostDryRun);
+		}
+	}
+}
+
+void UAvalanchePlayback::RegisterTickDelegate()
+{
+	if (ensure(!TickDelegateHandle.IsValid()))
+	{
+		TickDelegateHandle = FCoreDelegates::OnEndFrame.AddUObject(this, &UAvalanchePlayback::OnEndFrameTick);
+	}
+}
+
+void UAvalanchePlayback::UnregisterTickDelegate()
+{
+	if (TickDelegateHandle.IsValid())
+	{
+		FCoreDelegates::OnEndFrame.Remove(TickDelegateHandle);
+		TickDelegateHandle.Reset();
+	}
+}
+
+void UAvalanchePlayback::OnEndFrameTick()
+{
+	Tick(FApp::GetDeltaTime());
+}
+
+void UAvalanchePlayback::ConditionalResolvePlaybackSettings()
+{
+	if (PlaybackSettings.IsEmpty() && IsValid(RootNode))
+	{
+		RootNode->TickRoot(0.016f, PlaybackSettings);
+	}
+}
+
+void UAvalanchePlayback::SetRootNode(UAvaPlaybackNodeRoot* InRoot)
+{
+	if (IsValid(InRoot) && !IsValid(RootNode))
+	{
+		RootNode = InRoot;
+	}
+}
+
+void UAvalanchePlayback::AddPlayerNode(UAvaPlaybackNodePlayer* InPlayer)
+{
+	if (IsValid(InPlayer))
+	{
+		PlayerNodes.AddUnique(InPlayer);
+	}
+}
+
+UAvaPlaybackNodeRoot* UAvalanchePlayback::GetRootNode() const
+{
+	return RootNode;
+}
+
+void UAvalanchePlayback::AddPlaybackNode(UAvaPlaybackNode* Node)
+{
+	//Note: Ensure that all the Nodes added here have Outer set to this Playback Object.
+	Node->Rename(nullptr, this);
+
+#if WITH_EDITORONLY_DATA
+	PlaybackNodes.Add(Node);
+#endif
+}
+
+void UAvalanchePlayback::RemovePlaybackNode(UAvaPlaybackNode* Node)
+{
+#if WITH_EDITORONLY_DATA
+	PlaybackNodes.Remove(Node);
+#endif
+}
+
+void UAvalanchePlayback::StopPlayable(const FSoftObjectPath& InSourceAssetPath, const FName& InChannelName)
+{
+	if (UAvalanchePlayable* const Playable = FindPlayable(InSourceAssetPath, InChannelName))
+	{
+		Playable->EndPlay(EAvalanchePlayableEndPlayOptions::None);
+	}
+}
+
+void UAvalanchePlayback::LoadAsset(const FAvaSoftAssetPtr& InSourceAsset, const FName& InChannelName)
+{
+	FindOrLoadPlayable(InSourceAsset, InChannelName);
+}
+
+void UAvalanchePlayback::Tick(float DeltaTime)
+{
+	if (bIsTicking)
+	{
+		return;
+	}
+	TGuardValue TickGuard(bIsTicking, true);
+
+	if (!IsValid(RootNode))
+	{
+		return;
+	}
+
+	// Backup previous settings so we can detect if an asset has switched and needs to be stopped.
+	PreviousPlaybackSettings = PlaybackSettings;
+	
+	//Pass 1: Tick from Root (Channels Node) up to the Player Nodes.
+	for (TPair<FName, FAvaPlaybackChannelParameters>& Playback : PlaybackSettings)
+	{
+		Playback.Value.AvalancheAssets.Reset();
+	}
+	RootNode->TickRoot(DeltaTime, PlaybackSettings);
+
+	//Post Pass 1: Refresh any Playback from the Settings gotten from pass 1
+	for (const TPair<FName, FAvaPlaybackChannelParameters>& Playback : PlaybackSettings)
+	{
+		const FName ChannelName = ResolveChannelName(Playback.Key);
+		FAvaOutputChannel& Channel = UAvalancheBroadcast::Get().GetCurrentProfile().GetChannelMutable(ChannelName);
+
+		for (const FAvaSoftAssetPtr& AvalancheAsset : Playback.Value.AvalancheAssets)
+		{
+			if (IsPreviewOnly())
+			{
+				RefreshPreview(AvalancheAsset, Channel, ChannelName);
+			}
+			else if (Channel.IsValidChannel())
+			{
+				RefreshPlayback(AvalancheAsset, Channel, ChannelName);
+			}
+		}
+
+		// Check if the asset playing has changed (this may be the result of a switch/combiner node).
+		if (const FAvaPlaybackChannelParameters* PreviousParameters = PreviousPlaybackSettings.Find(Playback.Key))
+		{
+			for (const FAvaSoftAssetPtr& PreviousAsset : PreviousParameters->AvalancheAssets)
+			{
+				if (Playback.Value.AvalancheAssets.Find(PreviousAsset) == INDEX_NONE)
+				{
+					StopPlayable(PreviousAsset.ToSoftObjectPath(), ChannelName);
+				}
+			}
+		}
+	}
+
+	PreviousPlaybackSettings.Reset();
+
+	//Pass 2: Tick Events from each Player Node to the Connected Event Triggers/Actions 
+	for (TArray<TObjectPtr<UAvaPlaybackNodePlayer>>::TIterator Iter(PlayerNodes); Iter; ++Iter)
+	{
+		TObjectPtr<UAvaPlaybackNodePlayer> PlayerNode = *Iter;
+		if (IsValid(PlayerNode))
+		{
+			PlayerNode->TickEventFeed(DeltaTime);
+		}
+		else
+		{
+			Iter.RemoveCurrent();
+		}
+	}
+
+	//Post Pass 2: Reset the Event Triggers 
+	for (UAvaPlaybackNodePlayer* PlayerNode : PlayerNodes)
+	{
+		if (IsValid(PlayerNode))
+		{
+			PlayerNode->ResetEvents();
+		}
+	}
+
+	ExecutePendingAnimationCommands();
+	ExecutePendingRemoteControlCommands();
+}
+
+void UAvalanchePlayback::PushAnimationCommand(const FSoftObjectPath& InSourceAssetPath, const FString& InChannelName, EAvaMediaAnimAction InAnimAction, const FAnimPlaySettings& InAnimSettings)
+{
+	const FDateTime Timeout = FDateTime::UtcNow() + FTimespan::FromSeconds(10);
+	PendingAnimationCommands.Add({{Timeout, InSourceAssetPath, InChannelName, FName(InChannelName)}, InAnimAction, InAnimSettings});
+}
+
+void UAvalanchePlayback::PushRemoteControlValues(const FSoftObjectPath& InSourceAssetPath, const FString& InChannelName, const TSharedRef<FAvalancheRemoteControlValues>& InRemoteControlValues)
+{
+	const FDateTime Timeout = FDateTime::UtcNow() + FTimespan::FromSeconds(10);
+	PendingRemoteControlCommands.Add({{Timeout, InSourceAssetPath, InChannelName, FName(InChannelName)}, InRemoteControlValues});
+}
+
+void UAvalanchePlayback::SetupPlaybackNode(UAvaPlaybackNode* InPlaybackNode, bool bSelectNewNode)
+{
+	check(InPlaybackNode);
+
+	InPlaybackNode->CreateStartingConnectors();
+
+#if WITH_EDITOR
+	// Create the graph node
+	check(InPlaybackNode->GetGraphNode() == nullptr);
+	if (TSharedPtr<IAvaPlaybackGraphEditor> GraphEditor = GetGraphEditor())
+	{
+		GraphEditor->SetupPlaybackNode(EdGraph, InPlaybackNode, bSelectNewNode);
+	}
+#endif
+}
+
+#if WITH_EDITOR
+const TArray<TObjectPtr<UAvaPlaybackNode>>& UAvalanchePlayback::GetPlaybackNodes() const
+{
+	return PlaybackNodes;
+}
+
+void UAvalanchePlayback::CreateGraph()
+{
+	TSharedPtr<IAvaPlaybackGraphEditor> GraphEditor = GetGraphEditor();
+	if (EdGraph == nullptr && GraphEditor.IsValid())
+	{
+		EdGraph = GraphEditor->CreatePlaybackGraph(this);
+        EdGraph->bAllowDeletion = false;
+        
+        // Give the schema a chance to fill out any required nodes (like the results node)
+        const UEdGraphSchema* const Schema = EdGraph->GetSchema();
+        Schema->CreateDefaultNodesForGraph(*EdGraph);
+	}
+}
+
+void UAvalanchePlayback::ClearGraph()
+{
+	if (EdGraph)
+	{
+		EdGraph->Nodes.Empty();
+		
+		// Give the schema a chance to fill out any required nodes (like the results node)
+		const UEdGraphSchema* const Schema = EdGraph->GetSchema();
+		Schema->CreateDefaultNodesForGraph(*EdGraph);
+	}
+}
+
+void UAvalanchePlayback::RefreshPlaybackNode(UAvaPlaybackNode* InPlaybackNode)
+{
+	if (TSharedPtr<IAvaPlaybackGraphEditor> GraphEditor = GetGraphEditor())
+	{
+		if (InPlaybackNode && InPlaybackNode->GetGraphNode())
+		{
+			GraphEditor->RefreshNode(*InPlaybackNode->GetGraphNode());
+		}
+	}
+}
+
+void UAvalanchePlayback::CompilePlaybackNodesFromGraphNodes()
+{
+	if (TSharedPtr<IAvaPlaybackGraphEditor> GraphEditor = GetGraphEditor())
+	{
+		GraphEditor->CompilePlaybackNodesFromGraphNodes(this);
+	}
+}
+
+UEdGraph* UAvalanchePlayback::GetGraph()
+{
+	return EdGraph;
+}
+
+void UAvalanchePlayback::ResetGraph()
+{
+	for (const UAvaPlaybackNode* const PlaybackNode : PlaybackNodes)
+	{
+		EdGraph->RemoveNode(PlaybackNode->GetGraphNode());
+	}
+
+	PlaybackNodes.Reset();
+}
+
+void UAvalanchePlayback::SetGraphEditor(TSharedPtr<IAvaPlaybackGraphEditor> InGraphEditor)
+{
+	GraphEditorWeak = InGraphEditor;
+}
+
+TSharedPtr<IAvaPlaybackGraphEditor> UAvalanchePlayback::GetGraphEditor() const
+{
+	return GraphEditorWeak.Pin();
+}
+#endif
+
+void UAvalanchePlayback::ExecutePendingAnimationCommands()
+{
+	if (PendingAnimationCommands.IsEmpty())
+	{
+		return;
+	}
+
+	const UAvalancheBroadcast& AvaBroadcast = UAvalancheBroadcast::Get();
+
+	// If commands can't be executed because the assets are still
+	// streaming, we will accumulate them in this array for next time.
+	// Note: preserves the order of the commands (ie. not a map or set).
+	TArray<FAnimationCommand> StillPendingAnimationCommands;
+
+	const FDateTime CurrentTime = FDateTime::UtcNow();
+
+	for (FAnimationCommand& AnimCommand : PendingAnimationCommands)
+	{
+		if (AnimCommand.HasTimedOut(CurrentTime))
+		{
+			UE_LOG(LogAvalanchePlayback, Warning,
+				TEXT("Animation command for asset \"%s\" on channel \"%s\" timed out and is discarded."),
+				*AnimCommand.SourcePath.ToString(), *AnimCommand.ChannelName);
+			continue;
+		}
+
+		FAnimPlaySettings& AnimSettings = AnimCommand.AnimPlaySettings;
+		AnimSettings.Action = AnimCommand.AnimAction; // TODO: get rid of the AnimAction in 2 places.
+		
+		UAvalanchePlayable* Playable = FindPlayable(AnimCommand.SourcePath, AnimCommand.ChannelFName);
+		if (!Playable)
+		{
+			continue;
+		}
+		
+		const EAvalanchePlayableCommandResult Result = Playable->ExecuteAnimationCommand(AnimCommand.AnimAction, AnimCommand.AnimPlaySettings);
+		if (Result == EAvalanchePlayableCommandResult::KeepPending)
+		{
+			StillPendingAnimationCommands.Add(AnimCommand);
+		}
+	}
+	PendingAnimationCommands = StillPendingAnimationCommands;
+}
+
+void UAvalanchePlayback::ExecutePendingRemoteControlCommands()
+{
+	if (PendingRemoteControlCommands.IsEmpty())
+	{
+		return;
+	}
+
+	// If commands can't be executed because the assets are still
+	// streaming, we will accumulate them in this array for next time.
+	// Note: preserves the order of the commands (ie. not a map or set).
+	TArray<FRemoteControlCommand> StillPendingRemoteControlCommands;
+
+	const FDateTime CurrentTime = FDateTime::UtcNow();
+	
+	for (const FRemoteControlCommand& RemoteControlCommand : PendingRemoteControlCommands)
+	{
+		if (RemoteControlCommand.HasTimedOut(CurrentTime))
+		{
+			UE_LOG(LogAvalanchePlayback, Warning,
+				TEXT("Remote Control command for asset \"%s\" on channel \"%s\" timed out and is discarded."),
+				*RemoteControlCommand.SourcePath.ToString(), *RemoteControlCommand.ChannelName);
+			continue;
+		}
+		
+		UAvalanchePlayable* Playable = FindPlayable(RemoteControlCommand.SourcePath, RemoteControlCommand.ChannelFName);
+		if (!Playable)
+		{
+			continue;
+		}
+
+		const EAvalanchePlayableCommandResult Result = Playable->UpdateRemoteControlCommand(RemoteControlCommand.Values);
+		if (Result == EAvalanchePlayableCommandResult::KeepPending)
+		{
+			StillPendingRemoteControlCommands.Add(RemoteControlCommand);
+		}
+	}
+	PendingRemoteControlCommands = StillPendingRemoteControlCommands;
+}
+
+void UAvalanchePlayback::OnChannelChanged(const FAvaOutputChannel& InChannel, EAvaChannelChange InChange)
+{
+	if (EnumHasAnyFlags(InChange, EAvaChannelChange::State))
+	{
+		OnChannelBroadcastStateChanged(InChannel);
+	}
+	
+	if (EnumHasAnyFlags(InChange, EAvaChannelChange::Settings))
+	{
+		// Need to propagate the channel settings to loaded/playing game instances.
+		if (const FAvalanchePlaybackPlayableGroup* ChannelPlayableGroup = ChannelPlayableGroups.Find(InChannel.GetChannelName()))
+		{
+			TSet<const UAvaMediaPlayableGroup*> PlayableGroups; 
+			ChannelPlayableGroup->ForEachPlayable([&PlayableGroups](const UAvalanchePlayable* InPlayable)
+			{
+				if (const UAvaMediaPlayableGroup* PlayableGroup = InPlayable->GetPlayableGroup())
+				{
+					PlayableGroups.Add(PlayableGroup);
+				}
+			});
+			for (const UAvaMediaPlayableGroup* PlayableGroup : PlayableGroups)
+			{
+				if (const UAvalancheGameInstance* GameInstance = Cast<UAvalancheGameInstance>(PlayableGroup->GetGameInstance()))
+				{
+					if (UAvalancheGameViewportClient* ViewportClient = GameInstance->GetAvalancheGameViewportClient())
+					{
+						FAvaViewportQualitySettings QualitySettingsMutable = InChannel.GetViewportQualitySettings(); 
+						QualitySettingsMutable.Apply(ViewportClient->EngineShowFlags);
+					}
+				}
+			}
+		}
+	}
+}
+
+void UAvalanchePlayback::OnChannelBroadcastStateChanged(const FAvaOutputChannel& InChannel)
+{
+	if (IsPlaying() && InChannel.IsValidChannel() && InChannel.GetState() == EAvaChannelState::Live)
+	{
+		const FName ChangedChannelName = InChannel.GetChannelName();
+		FAvaBroadcastProfile& Profile = UAvalancheBroadcast::Get().GetCurrentProfile();
+		for (const TPair<FName, FAvaPlaybackChannelParameters>& Playback : PlaybackSettings)
+		{
+			const FName ChannelName = ResolveChannelName(Playback.Key);
+			
+			// Ensure to route the preview into it's channel regardless of pin connections.
+			// The preview channel(s) don't have corresponding pins.
+			if (ChannelName == ChangedChannelName)
+			{
+				FAvaOutputChannel& Channel = Profile.GetChannelMutable(ChannelName);
+				for (const FAvaSoftAssetPtr& AvalancheAsset : Playback.Value.AvalancheAssets)
+				{
+					RefreshPlayback(AvalancheAsset, Channel, ChannelName);
+				}
+			}
+		}
+	}
+}
+
+void UAvalanchePlayback::SetIsPlaying(bool bInIsPlaying)
+{
+	if (bIsPlaying != bInIsPlaying)
+	{
+		bIsPlaying = bInIsPlaying;
+		if (bInIsPlaying && !IsManaged())
+		{
+			RegisterTickDelegate();
+		}
+#if WITH_EDITOR
+		for (const TObjectPtr<UAvaPlaybackNode>& Node : PlaybackNodes)
+		{
+			if (Node)
+			{
+				Node->NotifyPlaybackStateChanged(bInIsPlaying);
+			}
+		}
+#else
+		RootNode->NotifyPlaybackStateChanged(bInIsPlaying);
+#endif
+		
+		OnPlaybackStateChanged.Broadcast(bInIsPlaying);
+	}
+
+	if (!bInIsPlaying)
+	{
+		UnregisterTickDelegate();
+	}
+}
+
+// Remark: very similar to FAvaOutputChannel::IsChannelAvailableForPlayback, could we possibly merge the logic?
+bool UAvalanchePlayback::CanRefreshPlayback(const FAvaSoftAssetPtr& InSourceAsset, const FAvaOutputChannel& InChannel) const
+{
+	return IsPlaying() && InChannel.IsValidChannel() && !InSourceAsset.IsNull();
+}
+
+bool UAvalanchePlayback::RefreshPlayback(const FAvaSoftAssetPtr& InSourceAsset, FAvaOutputChannel& InChannel, const FName& InChannelName)
+{
+	if (!CanRefreshPlayback(InSourceAsset, InChannel))
+	{
+		InChannel.UpdateRenderTarget(nullptr, nullptr);
+		InChannel.UpdateAudioDevice(FAudioDeviceHandle());
+		return false;
+	}
+	
+	if (UAvalanchePlayable* const Playable = FindOrLoadPlayable(InSourceAsset, InChannelName))
+	{
+		UTextureRenderTarget2D* const RenderTarget = UpdatePlaybackRenderTarget(Playable, InChannel);
+
+		const FIntPoint ViewportSize = RenderTarget
+			? AvaRenderTargetUtils::GetRenderTargetSize(RenderTarget)
+			: InChannel.DetermineRenderTargetSize();
+
+		const FAvalancheInstancePlaySettings WorldPlaySettings =
+		{ IAvaMediaModule::Get().GetAvalancheInstanceSettings(), InChannelName, RenderTarget, ViewportSize, InChannel.GetViewportQualitySettings() };
+
+		Playable->BeginPlay(WorldPlaySettings);
+
+		InChannel.UpdateAudioDevice(Playable->GetPlayWorld() ? Playable->GetPlayWorld()->GetAudioDevice() : FAudioDeviceHandle());
+		InChannel.UpdateRenderTarget(Playable->GetPlayableGroup(), RenderTarget);
+
+		return Playable->IsPlaying();
+	}
+
+	// Fallback
+	InChannel.UpdateRenderTarget(nullptr, nullptr);
+	InChannel.UpdateAudioDevice(FAudioDeviceHandle());
+	return false;
+}
+
+bool UAvalanchePlayback::CanRefreshPreview(const FAvaSoftAssetPtr& InSourceAsset) const
+{
+	return IsPlaying() && !InSourceAsset.IsNull();
+}
+
+bool UAvalanchePlayback::RefreshPreview(const FAvaSoftAssetPtr& InSourceAsset, FAvaOutputChannel& InChannel, const FName& InChannelName)
+{
+	// If the preview channel has been setup, go through the normal channel refresh.
+	if (InChannel.IsValidChannel())
+	{
+		return RefreshPlayback(InSourceAsset, InChannel, InChannelName);
+	}
+
+	if (!CanRefreshPreview(InSourceAsset))
+	{
+		return false;
+	}
+
+	UAvalanchePlayable* const Playable = FindOrLoadPlayable(InSourceAsset, InChannelName);
+
+	if (!Playable)
+	{
+		return false;
+	}
+	
+	UTextureRenderTarget2D* RenderTarget = Playable->GetPlayableGroup()->RenderTarget;
+
+	if (!RenderTarget)
+	{
+		// The render target is created here for now.
+		// TODO: adjust to preview window's size.
+		static const FName PreviewRenderTargetBaseName = TEXT("AvaPlayback_PreviewRenderTarget");
+		RenderTarget = AvaRenderTargetUtils::CreateDefaultRenderTarget(PreviewRenderTargetBaseName);
+		AvaRenderTargetUtils::UpdateRenderTarget(RenderTarget,
+			UAvalancheMediaSettings::Get().PreviewDefaultResolution,
+			FAvaOutputChannel::GetDefaultMediaOutputFormat(),
+			IAvaMediaModule::Get().GetBroadcastSettings().GetChannelClearColor());
+
+		Playable->GetPlayableGroup()->RenderTarget = RenderTarget;
+	}
+	
+	const FIntPoint ViewportSize = AvaRenderTargetUtils::GetRenderTargetSize(RenderTarget);
+
+	static const FAvaViewportQualitySettings DefaultPreviewQualitySettings;
+	const FAvalancheInstancePlaySettings WorldPlaySettings =
+		{ IAvaMediaModule::Get().GetAvalancheInstanceSettings(), InChannelName, RenderTarget, ViewportSize, DefaultPreviewQualitySettings };
+
+	Playable->BeginPlay(WorldPlaySettings);
+	
+	return Playable->IsPlaying();
+}
+
+UTextureRenderTarget2D* UAvalanchePlayback::UpdatePlaybackRenderTarget(const UAvalanchePlayable* InPlayable
+	, const FAvaOutputChannel& InChannel)
+{
+	UTextureRenderTarget2D* RenderTarget = nullptr;
+	
+	if (InChannel.IsValidChannel())
+	{
+		// Preferably use the channel's placeholder RT to avoid having the MediaCapture
+		// switch and possibly cause rendering glitches.
+		RenderTarget = InChannel.GetPlaceholderRenderTarget();
+
+		// Only fallback to internal RTs if we can't get one from the channels.
+		if (!IsValid(RenderTarget))
+		{
+			const FName ChannelName = InChannel.GetChannelName();
+
+			// See if we have a managed render target for this asset in this channel.
+			RenderTarget = InPlayable->GetPlayableGroup()->RenderTarget;
+			
+			if (!IsValid(RenderTarget))
+			{
+				static const FName PlaybackRenderTargetBaseName = TEXT("AvaPlayback_RenderTarget");
+				RenderTarget = AvaRenderTargetUtils::CreateDefaultRenderTarget(PlaybackRenderTargetBaseName);
+
+				// This render target is now managed by the playable's instance group.
+				InPlayable->GetPlayableGroup()->RenderTarget = RenderTarget;
+			}
+		}
+		
+		check(RenderTarget);
+
+		AvaRenderTargetUtils::UpdateRenderTarget(RenderTarget,
+			InChannel.DetermineRenderTargetSize(),
+			InChannel.DetermineRenderTargetFormat(),
+			IAvaMediaModule::Get().GetBroadcastSettings().GetChannelClearColor());
+	}
+	
+	return RenderTarget;
+}
+
+FAvalanchePlaybackGraphBuilder::FAvalanchePlaybackGraphBuilder(UAvaMediaPlayableGroupManager* InPlayableGroupManager)
+{
+	Playback = NewObject<UAvalanchePlayback>();
+	Playback->SetPlayableGroupManager(InPlayableGroupManager);
+	
+	// Construct Root
+	RootNode = Playback->ConstructPlaybackNode<UAvaPlaybackNodeRoot>();
+}
+
+FAvalanchePlaybackGraphBuilder::~FAvalanchePlaybackGraphBuilder()
+{
+	FinishBuilding();
+}
+
+bool FAvalanchePlaybackGraphBuilder::ConnectToRoot(const FString& InChannelName, UAvaPlaybackNode* InNodeToConnect)
+{
+	if (!InNodeToConnect)
+	{
+		return false;
+	}
+	
+	// Connect Player Node to Root Node
+	TArray<UAvaPlaybackNode*> NewRootChildNodes;
+	NewRootChildNodes.AddZeroed(RootNode->GetMaxChildNodes());
+
+	// Copy existing children.
+	const TArray<TObjectPtr<UAvaPlaybackNode>>& ExistingChildren = RootNode->GetChildNodes();
+	for (int32 PinIndex = 0; PinIndex < ExistingChildren.Num(); ++PinIndex)
+	{
+		if (NewRootChildNodes.IsValidIndex(PinIndex))
+		{
+			NewRootChildNodes[PinIndex] = ExistingChildren[PinIndex];
+		}
+	}
+
+	// Get the pin index corresponding to the given channel name.
+	const int32 ConnectionPinIndex = GetPinIndexForChannel(InChannelName);
+
+	// Check that Index is valid. RootNode Max Child Nodes should be the Channel Names Count,
+	// and the Connection Pin Index should be >= 0.
+	if (NewRootChildNodes.IsValidIndex(ConnectionPinIndex))
+	{
+		NewRootChildNodes[ConnectionPinIndex] = InNodeToConnect;
+		RootNode->SetChildNodes(MoveTemp(NewRootChildNodes));
+		bFinished = false;
+		return true;
+	}
+	else
+	{
+		UE_LOG(LogAvalanchePlayback, Error
+			, TEXT("Using index %d for channel \"%s\" is invalid. Root node has %d children. The node \"%s\" could not be added.")
+			, ConnectionPinIndex, *InChannelName, NewRootChildNodes.Num(), *InNodeToConnect->GetNodeDisplayNameText().ToString());
+	}
+	return false;	
+}
+
+int32 FAvalanchePlaybackGraphBuilder::GetPinIndexForChannel(const FString& InChannelName) const
+{
+	int32 ConnectionPinIndex = UAvalancheBroadcast::Get().GetChannelIndex(*InChannelName);
+
+	// We could return nullptr here, but given that this is a "GetOrCreate",
+	// the most sensible thing is to just connect to first index and continue
+	if (ConnectionPinIndex == INDEX_NONE)
+	{
+		UE_LOG(LogAvalanchePlayback, Warning
+			, TEXT("No channel with name (%s) was found. Connecting player node to first channel available")
+			, *InChannelName);
+	
+		ConnectionPinIndex = 0;
+	}
+	return ConnectionPinIndex;
+}
+
+UAvalanchePlayback* FAvalanchePlaybackGraphBuilder::FinishBuilding()
+{
+	if (Playback && !bFinished)
+	{
+		Playback->DryRunGraph();
+		bFinished = true;
+	}
+	return Playback;
+}
diff --git a/Engine/Plugins/Experimental/Avalanche/Source/AvalancheMedia/Private/Playback/AvalancheRemoteControl.cpp b/Engine/Plugins/Experimental/Avalanche/Source/AvalancheMedia/Private/Playback/AvalancheRemoteControl.cpp
new file mode 100644
--- /dev/null
+++ b/Engine/Plugins/Experimental/Avalanche/Source/AvalancheMedia/Private/Playback/AvalancheRemoteControl.cpp
@@ -0,0 +1,312 @@
+// Copyright Epic Games, Inc. All Rights Reserved.
+
+#include "Playback/AvalancheRemoteControl.h"
+
+#include "Action/RCAction.h"
+#include "Action/RCActionContainer.h"
+#include "Action/RCPropertyIdAction.h"
+#include "AvaSerializationUtils.h"
+#include "Backends/JsonStructDeserializerBackend.h"
+#include "Backends/JsonStructSerializerBackend.h"
+#include "Behaviour/RCBehaviour.h"
+#include "Controller/RCController.h"
+#include "IRemoteControlModule.h"
+#include "RCVirtualProperty.h"
+#include "RemoteControlPropertyIdRegistry.h"
+
+DEFINE_LOG_CATEGORY(LogAvaMediaRemoteControl);
+
+namespace UE::AvalancheRemoteControl::Private
+{
+	bool ResolveObjectPropertyForReadOnly(UObject* Object, FRCFieldPathInfo PropertyPath, FRCObjectReference& OutObjectRef, FString& OutErrorText)
+	{
+		if (!Object)
+		{
+			OutErrorText = FString::Printf(TEXT("Invalid object to resolve property '%s'"), *PropertyPath.GetFieldName().ToString());
+			return false;
+		}
+
+		bool bSuccess = true;
+		
+		if (PropertyPath.GetSegmentCount() != 0)
+		{
+			if (PropertyPath.Resolve(Object))
+			{
+				OutObjectRef = FRCObjectReference{ERCAccess::READ_ACCESS, Object, MoveTemp(PropertyPath)};
+			}
+			else
+			{
+				OutErrorText = FString::Printf(TEXT("Object property: %s could not be resolved on object: %s"), *PropertyPath.GetFieldName().ToString(), *Object->GetPathName());
+				bSuccess = false;
+			}
+		}
+		else
+		{
+			OutObjectRef = FRCObjectReference{ERCAccess::READ_ACCESS, Object};
+		}
+		return bSuccess;
+	}
+	
+	inline bool GetObjectRef(const TSharedPtr<const FRemoteControlProperty>& InField, const ERCAccess InAccess, FRCObjectReference& OutObjectRef)
+	{
+		if (!InField.IsValid())
+		{
+			return false;
+		}
+		if (UObject* FieldBoundObject = InField->GetBoundObject(); IsValid(FieldBoundObject))
+		{
+			FString ErrorText;
+			bool bSuccess = false;
+			
+			if (InAccess == ERCAccess::READ_ACCESS)
+			{
+				// Fix for private/protected properties: we allow reading the property even if private/protected.
+				// This is necessary for the case of private/protect props that have a setter.
+				// RC doesn't have a path to allow reading those (it doesn't support getters).
+				bSuccess = ResolveObjectPropertyForReadOnly(FieldBoundObject, InField->FieldPathInfo, OutObjectRef, ErrorText);
+			}
+			else
+			{
+				bSuccess = IRemoteControlModule::Get().ResolveObjectProperty(InAccess, FieldBoundObject, InField->FieldPathInfo, OutObjectRef, &ErrorText);
+			}
+
+			if (bSuccess)
+			{
+				return true;
+			}
+			
+			UE_LOG(LogAvaMediaRemoteControl, Error,
+				TEXT("Couldn\'t resolve object property \"%s\" in object \"%s\": %s"),
+				*InField->FieldName.ToString(), *FieldBoundObject->GetPathName(), *ErrorText);
+		}
+		else
+		{
+			UE_LOG(LogAvaMediaRemoteControl, Error,
+				TEXT("Couldn\'t resolve object property \"%s\": Invalid Field Bound Object."),
+				*InField->FieldName.ToString());
+		}
+		return false;
+	}
+}
+
+EAvaRemoteControlResult UE::AvalancheRemoteControl::GetValueOfEntity(const TSharedPtr<const FRemoteControlEntity>& RemoteControlEntity, TArray<uint8>& OutValue)
+{
+	using namespace UE::AvalancheRemoteControl::Private;
+	OutValue.Reset();
+
+	TSharedPtr<const FRemoteControlProperty> Field = StaticCastSharedPtr<const FRemoteControlProperty>(RemoteControlEntity);
+	if (!Field.IsValid())
+	{
+		return EAvaRemoteControlResult::InvalidParameter;
+	}
+
+	FRCObjectReference ObjectRef;
+	if (!GetObjectRef(Field, ERCAccess::READ_ACCESS, ObjectRef))
+	{
+		return EAvaRemoteControlResult::ReadAccessDenied;
+	}
+	FMemoryWriter Writer = FMemoryWriter(OutValue);
+	FJsonStructSerializerBackend WriterBackend = FJsonStructSerializerBackend(Writer, EStructSerializerBackendFlags::Default);
+	return IRemoteControlModule::Get().GetObjectProperties(ObjectRef, WriterBackend) ? EAvaRemoteControlResult::Completed : EAvaRemoteControlResult::ReadPropertyFailed;
+}
+
+EAvaRemoteControlResult UE::AvalancheRemoteControl::GetValueOfEntity(const TSharedPtr<const FRemoteControlEntity>& InRemoteControlEntity, FString& OutValue)
+{
+	TArray<uint8> ValueAsBytes;
+	const EAvaRemoteControlResult Result = GetValueOfEntity(InRemoteControlEntity, ValueAsBytes);
+	AvaSerializationUtils::JsonValueConversion::BytesToString(ValueAsBytes, OutValue);
+	return Result;
+}
+
+EAvaRemoteControlResult UE::AvalancheRemoteControl::SetValueOfEntity(const TSharedPtr<FRemoteControlEntity>& RemoteControlEntity, const TArrayView<const uint8>& InValue)
+{
+	using namespace UE::AvalancheRemoteControl::Private;
+
+	TSharedPtr<FRemoteControlProperty> Field = StaticCastSharedPtr<FRemoteControlProperty>(RemoteControlEntity);
+	if (!Field.IsValid())
+	{
+		return EAvaRemoteControlResult::InvalidParameter;
+	}
+	
+	FRCObjectReference ObjectRefRead;
+	if (GetObjectRef(Field, ERCAccess::READ_ACCESS, ObjectRefRead))
+	{
+		TArray<uint8> CurrentValueAsBytes;
+		FMemoryWriter Writer(CurrentValueAsBytes);
+		FJsonStructSerializerBackend WriterBackend = FJsonStructSerializerBackend(Writer, EStructSerializerBackendFlags::Default);
+		if (IRemoteControlModule::Get().GetObjectProperties(ObjectRefRead, WriterBackend))
+		{
+			if (CurrentValueAsBytes == InValue)
+			{
+				// if the given value is already set, don't do anything
+				return EAvaRemoteControlResult::UpToDate;
+			}
+		}
+	}
+
+	FRCObjectReference ObjectRefWrite;
+	if (!GetObjectRef(Field, ERCAccess::WRITE_ACCESS, ObjectRefWrite))
+	{
+		return EAvaRemoteControlResult::WriteAccessDenied;
+	}
+
+	FMemoryReaderView Reader(InValue);
+	FJsonStructDeserializerBackend ReaderBackend = FJsonStructDeserializerBackend(Reader);
+
+	// Notes:
+	// - if RemoteControl.EnableOngoingChangeOptimization is enabled, PostEditChangeProperty is not called right away
+	// there might be a delay (of 0.2 seconds) before it is called.
+	// - OnPropertyChangedDelegate (OnExposedPropertiesModified()) is a "per frame" event and is broadcast from URemoteControlPreset::OnEndFrame().
+	const bool bDeserializationSucceeded = IRemoteControlModule::Get().SetObjectProperties(ObjectRefWrite, ReaderBackend, ERCPayloadType::Json);
+
+#if WITH_EDITOR
+	UObject* const Object = ObjectRefWrite.Object.Get();
+	if (bDeserializationSucceeded && IsValid(Object))
+	{
+		FEditPropertyChain EditPropertyChain;
+		ObjectRefWrite.PropertyPathInfo.ToEditPropertyChain(EditPropertyChain);
+
+		// Note: Only PostEditChangeChainProperty is called here because PostEditChangeProperty is already dealt with in RC.
+		// Ideally, this should be in RC so that PostEditChangeProperty does not get called twice.
+		if (!EditPropertyChain.IsEmpty())
+		{
+			FPropertyChangedEvent PropertyEvent = ObjectRefWrite.PropertyPathInfo.ToPropertyChangedEvent(EPropertyChangeType::ValueSet);
+
+			FPropertyChangedChainEvent ChainEvent(EditPropertyChain, PropertyEvent);
+
+			TArray<TMap<FString, int32>> ArrayIndicesPerObject;
+			{
+				TMap<FString, int32> ArrayIndices;
+				ArrayIndices.Reserve(ObjectRefWrite.PropertyPathInfo.Segments.Num());
+
+				for (FRCFieldPathSegment& Segment : ObjectRefWrite.PropertyPathInfo.Segments)
+				{
+					ArrayIndices.Add(Segment.Name.ToString(), Segment.ArrayIndex);
+				}
+
+				ArrayIndicesPerObject.Add(MoveTemp(ArrayIndices));
+			}
+
+			ChainEvent.ObjectIteratorIndex = 0;
+			ChainEvent.SetArrayIndexPerObject(ArrayIndicesPerObject);
+
+			Object->PostEditChangeChainProperty(ChainEvent);
+		}
+	}
+#endif
+
+	return bDeserializationSucceeded ? EAvaRemoteControlResult::Completed : EAvaRemoteControlResult::WritePropertyFailed;
+}
+
+EAvaRemoteControlResult UE::AvalancheRemoteControl::SetValueOfEntity(const TSharedPtr<FRemoteControlEntity>& InRemoteControlEntity, const FString& InValue)
+{
+	return SetValueOfEntity(InRemoteControlEntity, AvaSerializationUtils::JsonValueConversion::ValueToConstBytesView(InValue));
+}
+
+EAvaRemoteControlResult UE::AvalancheRemoteControl::GetValueOfController(URCVirtualPropertyBase* InController, TArray<uint8>& OutValue)
+{
+	OutValue.Reset();
+
+	if (!InController)
+	{
+		return EAvaRemoteControlResult::InvalidParameter;
+	}
+
+	FMemoryWriter Writer = FMemoryWriter(OutValue);
+	FJsonStructSerializerBackend WriterBackend = FJsonStructSerializerBackend(Writer, EStructSerializerBackendFlags::Default);
+	InController->SerializeToBackend(WriterBackend);
+	return EAvaRemoteControlResult::Completed;
+}
+
+EAvaRemoteControlResult UE::AvalancheRemoteControl::GetValueOfController(URCVirtualPropertyBase* InController, FString& OutValue)
+{
+	TArray<uint8> ValueAsBytes;
+	const EAvaRemoteControlResult Result = GetValueOfController(InController, ValueAsBytes);
+
+	AvaSerializationUtils::JsonValueConversion::BytesToString(ValueAsBytes, OutValue);
+	return Result;
+}
+
+EAvaRemoteControlResult UE::AvalancheRemoteControl::SetValueOfController(URCVirtualPropertyBase* InController, const TArrayView<const uint8>& InValue)
+{
+	if (!InController)
+	{
+		return EAvaRemoteControlResult::InvalidParameter;
+	}
+
+	// TODO: Use the RC error callback to catch errors in the controller actions.
+	FMemoryReaderView Reader(InValue);
+	FJsonStructDeserializerBackend ReaderBackend = FJsonStructDeserializerBackend(Reader);
+	return InController->DeserializeFromBackend(ReaderBackend) ? EAvaRemoteControlResult::Completed : EAvaRemoteControlResult::WritePropertyFailed;
+}
+
+EAvaRemoteControlResult UE::AvalancheRemoteControl::SetValueOfController(URCVirtualPropertyBase* InController, const FString& InValue)
+{
+	return SetValueOfController(InController, AvaSerializationUtils::JsonValueConversion::ValueToConstBytesView(InValue));
+}
+
+bool UE::AvalancheRemoteControl::GetEntitiesControlledByController(const URemoteControlPreset* InRemoteControlPreset, const URCVirtualPropertyBase* InVirtualProperty, TSet<FGuid>& OutEntityIds)
+{
+	const URCController* Controller = Cast<URCController>(InVirtualProperty);
+	if (!IsValid(Controller))
+	{
+		return false;
+	}
+
+	for (const TObjectPtr<URCBehaviour>& Behaviour : Controller->Behaviours)
+	{
+		const TSet<TObjectPtr<URCAction>>& Actions = Behaviour->ActionContainer->GetActions();
+		for (const TObjectPtr<URCAction>& Action : Actions)
+		{
+			if (Action)
+			{
+				if (const URCPropertyIdAction* IdentityAction = Cast<URCPropertyIdAction>(Action.Get()))
+				{
+					if (const URemoteControlPropertyIdRegistry* IdentityRegistry = InRemoteControlPreset->GetPropertyIdRegistry())
+					{
+						OutEntityIds.Append(IdentityRegistry->GetEntityIdsForPropertyId(IdentityAction->PropertyId));
+					}
+				}
+				else
+				{
+					OutEntityIds.Add(Action->ExposedFieldId);
+				}
+			}
+		}
+	}
+	return true;
+}
+
+FString UE::AvalancheRemoteControl::EnumToString(EAvaRemoteControlResult InValue)
+{
+	return StaticEnum<EAvaRemoteControlResult>()->GetNameStringByValue(static_cast<int64>(InValue));
+}
+
+namespace UE::AvalancheRemoteControl
+{
+	FScopedPushControllerBehavioursEnable::FScopedPushControllerBehavioursEnable(URCVirtualPropertyBase* InVirtualProperty, bool bInBehavioursEnabled)
+		: VirtualProperty(InVirtualProperty)
+	{
+		if (URCController* const Controller = Cast<URCController>(InVirtualProperty))
+		{
+			PreviousBehavioursEnabled.Reserve(Controller->Behaviours.Num());			
+			for (URCBehaviour* const Behavior : Controller->Behaviours)
+			{
+				PreviousBehavioursEnabled.Add(Behavior->bIsEnabled);
+				Behavior->bIsEnabled = bInBehavioursEnabled;
+			}
+		}
+	}
+
+	FScopedPushControllerBehavioursEnable::~FScopedPushControllerBehavioursEnable()
+	{
+		if (URCController* const Controller = Cast<URCController>(VirtualProperty))
+		{
+			int32 BehaviourIndex = 0;
+			for (URCBehaviour* const Behavior : Controller->Behaviours)
+			{
+				Behavior->bIsEnabled = PreviousBehavioursEnabled[BehaviourIndex++];
+			}
+		}
+	}
+}
\ No newline at end of file
diff --git a/Engine/Plugins/Experimental/Avalanche/Source/AvalancheMedia/Private/Playback/AvalancheRemoteControlValues.cpp b/Engine/Plugins/Experimental/Avalanche/Source/AvalancheMedia/Private/Playback/AvalancheRemoteControlValues.cpp
new file mode 100644
--- /dev/null
+++ b/Engine/Plugins/Experimental/Avalanche/Source/AvalancheMedia/Private/Playback/AvalancheRemoteControlValues.cpp
@@ -0,0 +1,358 @@
+// Copyright Epic Games, Inc. All Rights Reserved.
+
+#include "Playback/AvalancheRemoteControlValues.h"
+
+#include "AvaSerializationUtils.h"
+#include "Playback/AvalancheRemoteControl.h"
+#include "Playback/AvalancheRemoteControlValuesPrivate.h"
+#include "RCVirtualProperty.h"
+#include "Serialization/CustomVersion.h"
+
+namespace UE::AvalancheRemoteControlValues::Private
+{
+	bool PruneValues(const TMap<FGuid, FAvalancheRemoteControlValue>& InValues, TMap<FGuid, FAvalancheRemoteControlValue>& OutValues)
+	{
+		bool bModified = false;
+		for (TMap<FGuid, FAvalancheRemoteControlValue>::TIterator ValueIterator(OutValues); ValueIterator; ++ValueIterator)
+		{
+			if (!InValues.Contains(ValueIterator->Key))
+			{
+				ValueIterator.RemoveCurrent();
+				bModified = true;
+			}
+		}
+		return bModified;
+	}
+	
+	bool UpdateValues(const TMap<FGuid, FAvalancheRemoteControlValue>& InValues, TMap<FGuid, FAvalancheRemoteControlValue>& OutValues, bool bInUpdateDefaults)
+	{
+		// Remove property values that are no longer exposed.
+		bool bModified = PruneValues(InValues, OutValues);
+	
+		// Add missing property values and optionally update the default values.
+		for (const TPair<FGuid, FAvalancheRemoteControlValue>& SourceValue : InValues)
+		{
+			FAvalancheRemoteControlValue* ExistingValue = OutValues.Find(SourceValue.Key);
+			if (!ExistingValue)
+			{
+				// Remark: IsDefault flag follows along.
+				OutValues.Add(SourceValue.Key, SourceValue.Value);
+				bModified = true;
+			}
+			else if (bInUpdateDefaults && ExistingValue->bIsDefault && !ExistingValue->IsSameValueAs(SourceValue.Value))
+			{
+				ExistingValue->SetValueFrom(SourceValue.Value);
+				bModified = true;
+			}
+		}
+		return bModified;
+	}
+
+	bool HasSameValues(const TMap<FGuid, FAvalancheRemoteControlValue>& InValues, const TMap<FGuid, FAvalancheRemoteControlValue>& InOtherValues)
+	{
+		// If values count differ, consider as different
+		if (InValues.Num() != InOtherValues.Num())
+		{
+			return false;
+		}
+
+		// Both Value maps have the same count, so one cannot be a subset of another, therefore, a single find pass should determine equality 
+		for (const TPair<FGuid, FAvalancheRemoteControlValue>& Pair : InValues)
+		{
+			const FAvalancheRemoteControlValue* FoundValue = InOtherValues.Find(Pair.Key);
+			if (!FoundValue || !FoundValue->IsSameValueAs(Pair.Value))
+			{
+				// Other's Value wasn't found, or was different from the value of this
+				return false;
+			}
+		}
+
+		return true;
+	}
+
+	EAvaRemoteControlChanges ToRemoteControlChanges(bool bInModified, EAvaRemoteControlChanges InModifiedChanges)
+	{
+		return bInModified ? InModifiedChanges : EAvaRemoteControlChanges::None;
+	}
+}
+
+const FGuid FAvalancheRemoteControlValueCustomVersion::Key(0x85218F83, 0xEDF141CA, 0x800EF947, 0x2F14CB06);
+FCustomVersionRegistration GRegisterAvalancheRemoteControlValueCustomVersion(FAvalancheRemoteControlValueCustomVersion::Key
+	, FAvalancheRemoteControlValueCustomVersion::LatestVersion
+	, TEXT("AvalancheRemoteControlValueVersion"));
+
+bool FAvalancheRemoteControlValue::Serialize(FArchive& Ar)
+{
+	Ar.UsingCustomVersion(FAvalancheRemoteControlValueCustomVersion::Key);
+	
+	if (Ar.CustomVer(FAvalancheRemoteControlValueCustomVersion::Key) >= FAvalancheRemoteControlValueCustomVersion::ValueAsString)
+	{
+		UScriptStruct* Struct = FAvalancheRemoteControlValue::StaticStruct();
+		Struct->SerializeTaggedProperties(Ar, reinterpret_cast<uint8*>(this), Struct, nullptr);
+	}
+	else
+	{
+		FAvalancheRemoteControlValueAsBytes_Legacy LegacyValue;
+		UScriptStruct* Struct = FAvalancheRemoteControlValueAsBytes_Legacy::StaticStruct();
+		Struct->SerializeTaggedProperties(Ar, reinterpret_cast<uint8*>(&LegacyValue), Struct, nullptr);
+
+		UE::AvaSerializationUtils::JsonValueConversion::BytesToString(LegacyValue.Bytes, Value);
+		bIsDefault = LegacyValue.bIsDefault;
+	}
+
+	return true;
+}
+
+void FAvalancheRemoteControlValues::RefreshControlledEntities(const URemoteControlPreset* InRemoteControlPreset)
+{
+	EntitiesControlledByController.Reset();
+
+	if (IsValid(InRemoteControlPreset))
+	{
+		const TArray<URCVirtualPropertyBase*> Controllers = InRemoteControlPreset->GetControllers();
+		for (const URCVirtualPropertyBase* PropertyBase : Controllers)
+		{
+			if (!UE::AvalancheRemoteControl::GetEntitiesControlledByController(InRemoteControlPreset, PropertyBase, EntitiesControlledByController))
+			{
+				UE_LOG(LogAvaMediaRemoteControl, Warning, TEXT("Failed to get controlled entities for controller \"%s\" (id:%s)."),
+					*PropertyBase->DisplayName.ToString(), *PropertyBase->Id.ToString());
+			}
+		}
+	}
+}
+
+void FAvalancheRemoteControlValues::CopyFrom(const URemoteControlPreset* InRemoteControlPreset, bool bInIsDefault)
+{
+	using namespace UE::AvalancheRemoteControl;
+
+	EntityValues.Reset();
+	ControllerValues.Reset();
+
+	RefreshControlledEntities(InRemoteControlPreset);	// will reset ids if invalid.
+	
+	if (IsValid(InRemoteControlPreset))
+	{
+		FString ValueAsString;
+		for (const TWeakPtr<const FRemoteControlEntity>& EntityWeakPtr : InRemoteControlPreset->GetExposedEntities<FRemoteControlEntity>())
+		{
+			const TSharedPtr<const FRemoteControlEntity> Entity = EntityWeakPtr.Pin();
+			if (!Entity.IsValid())
+			{
+				continue;
+			}
+
+			const EAvaRemoteControlResult Result = GetValueOfEntity(Entity, ValueAsString);
+
+			if (Failed(Result))
+			{
+				UE_LOG(LogAvaMediaRemoteControl, Error,
+					TEXT("Failed to read value of entity \"%s\" (id:%s) from RemoteControlPreset \"%s\": %s."),
+					*Entity->GetLabel().ToString(), *Entity->GetId().ToString(), *InRemoteControlPreset->GetName(), *EnumToString(Result));
+				continue;
+			}
+			
+			EntityValues.Add(Entity->GetId(), FAvalancheRemoteControlValue(ValueAsString, bInIsDefault));
+		}
+
+		TArray<URCVirtualPropertyBase*> Controllers = InRemoteControlPreset->GetControllers();
+		for (URCVirtualPropertyBase* Controller : Controllers)
+		{
+			const EAvaRemoteControlResult Result = GetValueOfController(Controller, ValueAsString);
+	
+			if (Failed(Result))
+			{
+				UE_LOG(LogAvaMediaRemoteControl, Error,
+					TEXT("Failed to read value of controller \"%s\" (id:%s) from RemoteControlPreset \"%s\": %s."),
+					*Controller->DisplayName.ToString(), *Controller->Id.ToString(), *InRemoteControlPreset->GetName(), *EnumToString(Result));
+				continue;
+			}
+			
+			ControllerValues.Add(Controller->Id, FAvalancheRemoteControlValue(ValueAsString, bInIsDefault));
+		}
+	}
+}
+
+bool FAvalancheRemoteControlValues::HasSameEntityValues(const FAvalancheRemoteControlValues& InOther) const
+{
+	return UE::AvalancheRemoteControlValues::Private::HasSameValues(EntityValues, InOther.EntityValues);
+}
+
+bool FAvalancheRemoteControlValues::HasSameControllerValues(const FAvalancheRemoteControlValues& InOther) const
+{
+	return UE::AvalancheRemoteControlValues::Private::HasSameValues(ControllerValues, InOther.ControllerValues);
+}
+
+EAvaRemoteControlChanges FAvalancheRemoteControlValues::PruneRemoteControlValues(const FAvalancheRemoteControlValues& InRemoteControlValues)
+{
+	using namespace UE::AvalancheRemoteControlValues::Private;
+	return ToRemoteControlChanges(PruneValues(InRemoteControlValues.EntityValues, EntityValues), EAvaRemoteControlChanges::EntityValues) 
+		| ToRemoteControlChanges(PruneValues(InRemoteControlValues.ControllerValues, ControllerValues), EAvaRemoteControlChanges::ControllerValues); 
+}
+
+EAvaRemoteControlChanges FAvalancheRemoteControlValues::UpdateRemoteControlValues(const FAvalancheRemoteControlValues& InRemoteControlValues, bool bInUpdateDefaults)
+{
+	using namespace UE::AvalancheRemoteControlValues::Private;
+	return ToRemoteControlChanges(UpdateValues(InRemoteControlValues.EntityValues, EntityValues, bInUpdateDefaults), EAvaRemoteControlChanges::EntityValues) 
+		| ToRemoteControlChanges(UpdateValues(InRemoteControlValues.ControllerValues, ControllerValues, bInUpdateDefaults), EAvaRemoteControlChanges::ControllerValues); 
+}
+
+bool FAvalancheRemoteControlValues::SetEntityValue(const FGuid& InId, const URemoteControlPreset* InRemoteControlPreset, bool bInIsDefault)
+{
+	const TSharedPtr<const FRemoteControlEntity> Entity = InRemoteControlPreset->GetExposedEntity<FRemoteControlEntity>(InId).Pin();
+
+	if (!Entity)
+	{
+		UE_LOG(LogAvaMediaRemoteControl, Error,
+			TEXT("Requested entity id \"%s\" was not found in RemoteControlPreset \"%s\"."),
+			*InId.ToString(), *InRemoteControlPreset->GetName());
+		return false;
+	}
+
+	using namespace UE::AvalancheRemoteControl;
+	FAvalancheRemoteControlValue Value;
+	Value.bIsDefault = bInIsDefault;
+
+	const EAvaRemoteControlResult Result = GetValueOfEntity(Entity, Value.Value);
+
+	if (Failed(Result))
+	{
+		UE_LOG(LogAvaMediaRemoteControl, Error,
+			TEXT("Failed to read value of entity \"%s\" (id:%s) from RemoteControlPreset \"%s\": %s."),
+			*Entity->GetLabel().ToString(), *InId.ToString(), *InRemoteControlPreset->GetName(), *EnumToString(Result));
+		return false;
+	}
+
+	EntityValues.Add(Entity->GetId(), MoveTemp(Value));
+	return true;
+}
+
+bool FAvalancheRemoteControlValues::SetControllerValue(const FGuid& InId, const URemoteControlPreset* InRemoteControlPreset, bool bInIsDefault)
+{
+	URCVirtualPropertyBase* Controller = InRemoteControlPreset->GetController(InId);
+
+	if (!Controller)
+	{
+		UE_LOG(LogAvaMediaRemoteControl, Error,
+			TEXT("Requested controller id \"%s\" was not found in RemoteControlPreset \"%s\"."),
+			*InId.ToString(), *InRemoteControlPreset->GetName());
+		return false;
+	}
+
+	using namespace UE::AvalancheRemoteControl;
+	FAvalancheRemoteControlValue Value;
+	Value.bIsDefault = bInIsDefault;
+
+	const EAvaRemoteControlResult Result = GetValueOfController(Controller, Value.Value);
+	
+	if (Failed(Result))
+	{
+		UE_LOG(LogAvaMediaRemoteControl, Error,
+			TEXT("Failed to read value of controller \"%s\" (id:%s) from RemoteControlPreset \"%s\": %s."),
+			*Controller->DisplayName.ToString(), *InId.ToString(), *InRemoteControlPreset->GetName(), *EnumToString(Result));
+		return false;
+	}
+		
+	ControllerValues.Add(InId, MoveTemp(Value));
+	return true;
+}
+
+void FAvalancheRemoteControlValues::ApplyEntityValuesToRemoteControlPreset(URemoteControlPreset* InRemoteControlPreset) const
+{
+	if (!InRemoteControlPreset)
+	{
+		return;
+	}
+	using namespace UE::AvalancheRemoteControl;
+	for (const TWeakPtr<FRemoteControlEntity>& EntityWeakPtr : InRemoteControlPreset->GetExposedEntities<FRemoteControlEntity>())
+	{
+		if (const TSharedPtr<FRemoteControlEntity> Entity = EntityWeakPtr.Pin())
+		{
+			if (const FAvalancheRemoteControlValue* Value = GetEntityValue(Entity->GetId()))
+			{
+				const EAvaRemoteControlResult Result = SetValueOfEntity(Entity, Value->Value); 
+				if (Failed(Result))
+				{
+					UE_LOG(LogAvaMediaRemoteControl, Error, TEXT("Failed to set value of exposed entity \"%s\" (id:%s): %s."),
+						*Entity->GetLabel().ToString(), *Entity->GetId().ToString(), *EnumToString(Result));
+				}
+			}
+			else
+			{
+				UE_LOG(LogAvaMediaRemoteControl, Error, TEXT("Exposed entity \"%s\" (id:%s): value not found in page."),
+					*Entity->GetLabel().ToString(), *Entity->GetId().ToString());
+			}
+		}
+	}
+}
+
+void FAvalancheRemoteControlValues::ApplyControllerValuesToRemoteControlPreset(URemoteControlPreset* InRemoteControlPreset, bool bInForceDisableBehaviors) const
+{
+	if (!InRemoteControlPreset)
+	{
+		return;
+	}
+	TArray<URCVirtualPropertyBase*> Controllers = InRemoteControlPreset->GetControllers();
+	for (URCVirtualPropertyBase* Controller : Controllers)
+	{
+		if (const FAvalancheRemoteControlValue* Value = GetControllerValue(Controller->Id))
+		{
+			using namespace UE::AvalancheRemoteControl;
+			EAvaRemoteControlResult Result;
+			if (bInForceDisableBehaviors)
+			{
+				FScopedPushControllerBehavioursEnable PushBehavioursEnable(Controller, false);
+				Result = SetValueOfController(Controller, Value->Value);
+			}
+			else
+			{
+				Result = SetValueOfController(Controller, Value->Value);
+			}
+			if (Failed(Result))
+			{
+				UE_LOG(LogAvaMediaRemoteControl, Error, TEXT("Failed to set virtual value of controller \"%s\" (id:%s): %s."),
+					*Controller->DisplayName.ToString(), *Controller->Id.ToString(), *EnumToString(Result));
+			}
+		}
+		else
+		{
+			UE_LOG(LogAvaMediaRemoteControl, Error, TEXT("Controller \"%s\" (id:%s): value not found in page."),
+				*Controller->DisplayName.ToString(), *Controller->Id.ToString());
+		}
+	}
+}
+
+bool FAvalancheRemoteControlValues::HasIdCollisions(const FAvalancheRemoteControlValues& InOtherValues) const
+{
+	const bool bHasControllerIdCollisions = HasIdCollisions(ControllerValues, InOtherValues.ControllerValues);
+	const bool bHasEntityIdCollisions = HasIdCollisions(EntityValues, InOtherValues.EntityValues);
+	return bHasControllerIdCollisions || bHasEntityIdCollisions; 
+}
+
+bool FAvalancheRemoteControlValues::Merge(const FAvalancheRemoteControlValues& InOtherValues)
+{
+	const bool bHasIdCollisions = HasIdCollisions(InOtherValues);
+
+	ControllerValues.Append(InOtherValues.ControllerValues);
+	EntityValues.Append(InOtherValues.EntityValues);		
+	EntitiesControlledByController.Append(InOtherValues.EntitiesControlledByController);
+
+	return !bHasIdCollisions;
+}
+
+bool FAvalancheRemoteControlValues::HasIdCollisions(const TMap<FGuid, FAvalancheRemoteControlValue>& InValues, const TMap<FGuid, FAvalancheRemoteControlValue>& InOtherValues)
+{
+	for (const TPair<FGuid, FAvalancheRemoteControlValue>& ValueEntry : InValues)
+	{
+		if (InOtherValues.Contains(ValueEntry.Key))
+		{
+			return true;
+		}
+	}
+	return false;
+}
+
+const FAvalancheRemoteControlValues& FAvalancheRemoteControlValues::GetDefaultEmpty()
+{
+	const static FAvalancheRemoteControlValues Empty;
+	return Empty;
+}
diff --git a/Engine/Plugins/Experimental/Avalanche/Source/AvalancheMedia/Private/Playback/AvalancheRemoteControlValuesPrivate.h b/Engine/Plugins/Experimental/Avalanche/Source/AvalancheMedia/Private/Playback/AvalancheRemoteControlValuesPrivate.h
new file mode 100644
--- /dev/null
+++ b/Engine/Plugins/Experimental/Avalanche/Source/AvalancheMedia/Private/Playback/AvalancheRemoteControlValuesPrivate.h
@@ -0,0 +1,41 @@
+// Copyright Epic Games, Inc. All Rights Reserved.
+
+#pragma once
+
+#include "CoreTypes.h"
+#include "AvalancheRemoteControlValuesPrivate.generated.h"
+
+/** Custom serialization version for FAvalancheRemoteControlValue */
+struct FAvalancheRemoteControlValueCustomVersion
+{
+	enum Type
+	{
+		// Initial version had the values stored as raw bytes.
+		BeforeCustomVersionWasAdded = 0,
+
+		// Values are stored as strings (json formatted).
+		ValueAsString,
+
+		// -----<new versions can be added above this line>-------------------------------------------------
+		VersionPlusOne,
+		LatestVersion = VersionPlusOne - 1
+	};
+
+	// The GUID for this custom version number
+	static const FGuid Key;
+	
+private:
+	FAvalancheRemoteControlValueCustomVersion() = delete;
+};
+
+USTRUCT()
+struct FAvalancheRemoteControlValueAsBytes_Legacy
+{
+	GENERATED_BODY()
+	
+	UPROPERTY()
+	TArray<uint8> Bytes;
+
+	UPROPERTY()
+	bool bIsDefault = false;
+};
\ No newline at end of file
diff --git a/Engine/Plugins/Experimental/Avalanche/Source/AvalancheMedia/Private/Playback/Nodes/AvaPlaybackNodeBlueprintPlayer.cpp b/Engine/Plugins/Experimental/Avalanche/Source/AvalancheMedia/Private/Playback/Nodes/AvaPlaybackNodeBlueprintPlayer.cpp
new file mode 100644
--- /dev/null
+++ b/Engine/Plugins/Experimental/Avalanche/Source/AvalancheMedia/Private/Playback/Nodes/AvaPlaybackNodeBlueprintPlayer.cpp
@@ -0,0 +1,57 @@
+// Copyright Epic Games, Inc. All Rights Reserved.
+
+#include "Playback/Nodes/AvaPlaybackNodeBlueprintPlayer.h"
+#include "AvaBlueprint.h"
+#include "Framework/AvalancheGameInstance.h"
+#include "Http/AvaMediaHttpServer.h"
+#include "Playback/AvalanchePlayback.h"
+
+#define LOCTEXT_NAMESPACE "AvalanchePlayback"
+
+UAvaPlaybackNodeBlueprintPlayer::UAvaPlaybackNodeBlueprintPlayer()
+{
+	//Update to Default Text
+	UpdateDisplayNameText();
+}
+
+void UAvaPlaybackNodeBlueprintPlayer::RefreshNode(bool bDryRunGraph)
+{
+	UpdateDisplayNameText();
+	Super::RefreshNode(bDryRunGraph);
+}
+
+void UAvaPlaybackNodeBlueprintPlayer::PostLoad()
+{
+	Super::PostLoad();
+	UpdateDisplayNameText();
+}
+
+void UAvaPlaybackNodeBlueprintPlayer::SetAvalancheAsset(const TSoftObjectPtr<UAvalancheBlueprint>& InAsset)
+{
+	BlueprintAsset = InAsset;
+}
+
+void UAvaPlaybackNodeBlueprintPlayer::UpdateDisplayNameText()
+{
+	const FString AssetName = BlueprintAsset.GetAssetName();
+	
+	if (AssetName.IsEmpty())
+	{
+		DisplayNameText = LOCTEXT("AvaPlaybackNode_BlueprintPlayerNoName", "Motion Design Blueprint Player");
+	}
+	else
+	{
+		DisplayNameText = FText::Format(LOCTEXT("AvaPlaybackNode_BlueprintPlayerName", "Motion Design Blueprint Player\n{0}")
+			, FText::FromString(AssetName));
+	}
+}
+
+FAvaSoftAssetPtr UAvaPlaybackNodeBlueprintPlayer::GetAvalancheAssetPtr() const
+{
+	FAvaSoftAssetPtr OutAvalancheAsset;
+	OutAvalancheAsset.AssetClassPath = FSoftClassPath(UAvalancheBlueprint::StaticClass());
+	OutAvalancheAsset.AssetPtr = BlueprintAsset;
+	return OutAvalancheAsset;
+}
+
+#undef LOCTEXT_NAMESPACE
diff --git a/Engine/Plugins/Experimental/Avalanche/Source/AvalancheMedia/Private/Playback/Nodes/AvaPlaybackNodeLevelPlayer.cpp b/Engine/Plugins/Experimental/Avalanche/Source/AvalancheMedia/Private/Playback/Nodes/AvaPlaybackNodeLevelPlayer.cpp
new file mode 100644
--- /dev/null
+++ b/Engine/Plugins/Experimental/Avalanche/Source/AvalancheMedia/Private/Playback/Nodes/AvaPlaybackNodeLevelPlayer.cpp
@@ -0,0 +1,56 @@
+// Copyright Epic Games, Inc. All Rights Reserved.
+
+#include "Playback/Nodes/AvaPlaybackNodeLevelPlayer.h"
+
+#include "Engine/World.h"
+#include "Playback/AvalanchePlayback.h"
+
+#define LOCTEXT_NAMESPACE "AvalanchePlayback"
+
+UAvaPlaybackNodeLevelPlayer::UAvaPlaybackNodeLevelPlayer()
+{
+	//Update to Default Text
+	UpdateDisplayNameText();
+}
+
+void UAvaPlaybackNodeLevelPlayer::RefreshNode(bool bDryRunGraph)
+{
+	UpdateDisplayNameText();
+	Super::RefreshNode(bDryRunGraph);
+}
+
+void UAvaPlaybackNodeLevelPlayer::PostLoad()
+{
+	Super::PostLoad();
+	UpdateDisplayNameText();
+}
+
+void UAvaPlaybackNodeLevelPlayer::SetAvalancheAsset(const TSoftObjectPtr<UWorld>& InAsset)
+{
+	LevelAsset = InAsset;
+}
+
+void UAvaPlaybackNodeLevelPlayer::UpdateDisplayNameText()
+{
+	const FString AssetName = LevelAsset.GetAssetName();
+	
+	if (AssetName.IsEmpty())
+	{
+		DisplayNameText = LOCTEXT("AvaPlaybackNode_LevelPlayerNoName", "Motion Design Level Player");
+	}
+	else
+	{
+		DisplayNameText = FText::Format(LOCTEXT("AvaPlaybackNode_LevelPlayerName", "Motion Design Level Player\n{0}")
+			, FText::FromString(AssetName));
+	}
+}
+
+FAvaSoftAssetPtr UAvaPlaybackNodeLevelPlayer::GetAvalancheAssetPtr() const
+{
+	FAvaSoftAssetPtr OutAvalancheAsset;
+	OutAvalancheAsset.AssetClassPath = FSoftClassPath(UWorld::StaticClass());
+	OutAvalancheAsset.AssetPtr = LevelAsset;
+	return OutAvalancheAsset;
+}
+
+#undef LOCTEXT_NAMESPACE
diff --git a/Engine/Plugins/Experimental/Avalanche/Source/AvalancheMedia/Private/Playback/Nodes/AvaPlaybackNodeRoot.cpp b/Engine/Plugins/Experimental/Avalanche/Source/AvalancheMedia/Private/Playback/Nodes/AvaPlaybackNodeRoot.cpp
new file mode 100644
--- /dev/null
+++ b/Engine/Plugins/Experimental/Avalanche/Source/AvalancheMedia/Private/Playback/Nodes/AvaPlaybackNodeRoot.cpp
@@ -0,0 +1,105 @@
+// Copyright Epic Games, Inc. All Rights Reserved.
+
+#include "Playback/Nodes/AvaPlaybackNodeRoot.h"
+#include "AvalancheBroadcast.h"
+#include "Playback/AvalanchePlayback.h"
+
+#define LOCTEXT_NAMESPACE "AvalanchePlayback"
+
+void UAvaPlaybackNodeRoot::PostAllocateNode()
+{
+	Super::PostAllocateNode();
+	
+	UAvalancheBroadcast& Broadcast = UAvalancheBroadcast::Get();
+	
+	BroadcastChangedHandle = Broadcast.AddChangeListener(FOnAvaBroadcastChanged::FDelegate::CreateUObject(this
+		, &UAvaPlaybackNodeRoot::OnBroadcastChanged));
+		
+	ChannelChangedHandle = FAvaOutputChannel::GetOnChannelChanged().AddUObject(this
+		, &UAvaPlaybackNodeRoot::OnChannelChanged);
+
+	if (UAvalanchePlayback* const Playback = GetPlayback())
+	{
+		Playback->SetRootNode(this);
+	}
+}
+
+void UAvaPlaybackNodeRoot::BeginDestroy()
+{
+	if (BroadcastChangedHandle.IsValid())
+	{
+		UAvalancheBroadcast::Get().RemoveChangeListener(BroadcastChangedHandle);
+	}
+	if (ChannelChangedHandle.IsValid())
+	{
+		FAvaOutputChannel::GetOnChannelChanged().Remove(ChannelChangedHandle);
+	}
+	Super::BeginDestroy();
+}
+
+void UAvaPlaybackNodeRoot::TickRoot(float InDeltaTime, TMap<FName, FAvaPlaybackChannelParameters>& OutPlaybackSettings)
+{
+	OutPlaybackSettings.Empty(ChildNodes.Num());
+	LastValidTick = FApp::GetCurrentTime();
+	
+	for (int32 Index = 0; Index < ChildNodes.Num(); ++Index)
+	{
+		FName ChannelName = GetChannelName(Index);
+		if (!ChannelName.IsNone())
+		{
+			FAvaPlaybackChannelParameters& PlaybackParameters = OutPlaybackSettings.Add(ChannelName);
+			PlaybackParameters.ChannelIndex = Index;
+			TickChild(InDeltaTime, Index, PlaybackParameters);
+		}
+	}
+}
+
+FText UAvaPlaybackNodeRoot::GetNodeDisplayNameText() const
+{
+	return LOCTEXT("RootNode_Title", "Channels");
+}
+
+FText UAvaPlaybackNodeRoot::GetNodeTooltipText() const
+{
+	return LOCTEXT("RootNode_Tooltip", "Wire the Playback Nodes into this Final Node for Broadcast");
+}
+
+void UAvaPlaybackNodeRoot::OnBroadcastChanged(EAvaBroadcastChange InChange)
+{
+	//TODO: Changing Profiles should not trigger reconstruct, rather the Pins should react to whether the Channel is Available or not
+	ReconstructNode();
+}
+
+void UAvaPlaybackNodeRoot::OnChannelChanged(const FAvaOutputChannel& InChannel, EAvaChannelChange InChange)
+{
+	if (EnumHasAnyFlags(InChange, EAvaChannelChange::State))
+	{
+		//TODO: Instead of Reconstructing, have the Pins just react to each Individual Channel State Change
+		ReconstructNode();
+	}
+}
+
+int32 UAvaPlaybackNodeRoot::GetMinChildNodes() const
+{
+	return UAvalancheBroadcast::Get().GetChannelNameCount();
+}
+
+int32 UAvaPlaybackNodeRoot::GetMaxChildNodes() const
+{
+	return UAvalancheBroadcast::Get().GetChannelNameCount();
+}
+
+#if WITH_EDITOR
+FName UAvaPlaybackNodeRoot::GetInputPinName(int32 InputPinIndex) const
+{
+	return UAvalancheBroadcast::Get().GetChannelName(InputPinIndex);
+}
+#endif
+
+FName UAvaPlaybackNodeRoot::GetChannelName(int32 InChannelNameIndex) const
+{
+	const UAvalancheBroadcast& AvalancheBroadcast = UAvalancheBroadcast::Get();
+	return AvalancheBroadcast.GetChannelName(InChannelNameIndex);
+}
+
+#undef LOCTEXT_NAMESPACE
diff --git a/Engine/Plugins/Experimental/Avalanche/Source/AvalancheMedia/Private/Playback/Nodes/Events/Actions/AvaPlaybackNode_PlayAnim.cpp b/Engine/Plugins/Experimental/Avalanche/Source/AvalancheMedia/Private/Playback/Nodes/Events/Actions/AvaPlaybackNode_PlayAnim.cpp
new file mode 100644
--- /dev/null
+++ b/Engine/Plugins/Experimental/Avalanche/Source/AvalancheMedia/Private/Playback/Nodes/Events/Actions/AvaPlaybackNode_PlayAnim.cpp
@@ -0,0 +1,160 @@
+// Copyright Epic Games, Inc. All Rights Reserved.
+
+#include "AvaPlaybackNode_PlayAnim.h"
+#include "Async/Async.h"
+#include "AvaBlueprint.h"
+#include "AvaScene.h"
+#include "AvalancheBroadcast.h"
+#include "Engine/Level.h"
+#include "Playback/AvalanchePlayback.h"
+#include "Playback/Nodes/AvaPlaybackNodeBlueprintPlayer.h"
+#include "Playback/Nodes/AvaPlaybackNodeLevelPlayer.h"
+#include "AvaSequence.h"
+
+#define LOCTEXT_NAMESPACE "AvalanchePlayback"
+
+namespace Avalanche
+{
+	FName GetSequenceName(const TObjectPtr<UAvaSequence>& InSequence)
+	{
+		return InSequence->GetFName();
+	}
+}
+
+FText UAvaPlaybackNode_PlayAnim::GetNodeDisplayNameText() const
+{
+	FFormatNamedArguments Args;
+	if (HasAnyFlags(RF_ClassDefaultObject))
+	{
+		Args.Add("AnimationTree", FText::GetEmpty());
+	}
+	else
+	{
+		FString AnimationTree;
+		for (TMap<FSoftObjectPath, FAvalancheAnimations>::TConstIterator Iter(AnimationMap); Iter; ++Iter)
+		{
+			bool bAction = false;
+			for (const FAnimPlaySettings& PlaySettings : Iter->Value.AvailableAnimations)
+			{
+				if (PlaySettings.Action != EAvaMediaAnimAction::None)
+				{
+					bAction = true;
+					break;
+				}
+			}
+			
+			if (bAction)
+			{
+				AnimationTree.Append(TEXT("\n") + Iter->Key.GetAssetName());
+			}
+		}		
+		Args.Add("AnimationTree", FText::FromString(AnimationTree));
+	}
+	return FText::Format(LOCTEXT("PlayAnimNode_Title", "Play Animation{AnimationTree}"), Args);
+}
+
+FText UAvaPlaybackNode_PlayAnim::GetNodeTooltipText() const
+{
+	return LOCTEXT("PlayAnimNode_Tooltip", "Plays an Animation from the Avalanche Asset");
+}
+
+void UAvaPlaybackNode_PlayAnim::OnEventTriggered(const FAvaPlaybackEventParameters& InEventParameters)
+{
+	UAvalanchePlayback* const Playback = GetPlayback();
+	if (!Playback)
+	{
+		return;
+	}
+
+	const FSoftObjectPath& AvalancheAssetPath = InEventParameters.AvalancheAsset.ToSoftObjectPath();
+	TArray<FName> ChannelNames = Playback->GetChannelNamesForIndices(InEventParameters.ChannelIndices);
+	
+	// Gather the Animations to Play
+	if (FAvalancheAnimations* const FoundAnimations = AnimationMap.Find(AvalancheAssetPath))
+	{
+		for (const FAnimPlaySettings& PlaySettings : FoundAnimations->AvailableAnimations)
+		{
+			if (PlaySettings.Action != EAvaMediaAnimAction::None)
+			{
+				for (const FName& ChannelName : ChannelNames)
+				{
+					Playback->PushAnimationCommand(AvalancheAssetPath, ChannelName.ToString(), PlaySettings.Action, PlaySettings);
+				}
+			}
+		}
+	}
+}
+
+void UAvaPlaybackNode_PlayAnim::PreDryRun()
+{
+	SeenAssetsInDryRun.Empty();
+}
+
+void UAvaPlaybackNode_PlayAnim::DryRun(const TArray<UAvaPlaybackNode*>& InAncestors)
+{
+	for (UAvaPlaybackNode* const PlaybackNode : InAncestors)
+	{
+		if (UAvaPlaybackNodeBlueprintPlayer* const PlayerNode = Cast<UAvaPlaybackNodeBlueprintPlayer>(PlaybackNode))
+		{
+			TSoftObjectPtr<UAvalancheBlueprint> AvalancheAsset = PlayerNode->GetAvalancheAsset();
+			if (UAvalancheBlueprint* const AvalancheBlueprint = AvalancheAsset.LoadSynchronous())
+			{
+				SeenAssetsInDryRun.Add(PlayerNode->GetAvalancheAssetPath());
+				FAvalancheAnimations& Animations = AnimationMap.FindOrAdd(PlayerNode->GetAvalancheAssetPath());
+				for (const TObjectPtr<UAvaSequence>& Animation : AvalancheBlueprint->GetSequences())
+				{
+					if (Animation)
+					{
+						Animations.AvailableAnimations.FindOrAdd(Avalanche::GetSequenceName(Animation));
+					}
+				}
+			}
+		}
+		if (UAvaPlaybackNodeLevelPlayer* const PlayerNode = Cast<UAvaPlaybackNodeLevelPlayer>(PlaybackNode))
+		{
+			TSoftObjectPtr<UWorld> AvalancheAsset = PlayerNode->GetAvalancheAsset();
+			if (UWorld* const AvalancheWorld = AvalancheAsset.LoadSynchronous())
+			{
+				SeenAssetsInDryRun.Add(PlayerNode->GetAvalancheAssetPath());
+				FAvalancheAnimations& Animations = AnimationMap.FindOrAdd(PlayerNode->GetAvalancheAssetPath());
+
+				AAvaScene* Scene = nullptr;
+				AvalancheWorld->PersistentLevel->Actors.FindItemByClass(&Scene);
+				if (IsValid(Scene))
+				{
+					for (const TObjectPtr<UAvaSequence>& Animation : Scene->GetSequences())
+					{
+						if (Animation)
+						{
+							Animations.AvailableAnimations.FindOrAdd(Avalanche::GetSequenceName(Animation));
+						}
+					}
+				}
+			}
+		}
+	}
+}
+
+void UAvaPlaybackNode_PlayAnim::PostDryRun()
+{
+	bool bRefreshNode = false;
+	
+	//Remove all the Assets that were not Seen.
+	for (TMap<FSoftObjectPath, FAvalancheAnimations>::TIterator Iter(AnimationMap); Iter; ++Iter)
+	{
+		if (!SeenAssetsInDryRun.Contains(Iter->Key))
+		{
+			Iter.RemoveCurrent();
+			bRefreshNode = true;
+		}
+	}
+
+	if (bRefreshNode)
+	{
+		RefreshNode(false);
+	}
+
+	SeenAssetsInDryRun.Empty();
+}
+
+#undef LOCTEXT_NAMESPACE
diff --git a/Engine/Plugins/Experimental/Avalanche/Source/AvalancheMedia/Private/Playback/Nodes/Events/Actions/AvaPlaybackNode_PlayAnim.h b/Engine/Plugins/Experimental/Avalanche/Source/AvalancheMedia/Private/Playback/Nodes/Events/Actions/AvaPlaybackNode_PlayAnim.h
new file mode 100644
--- /dev/null
+++ b/Engine/Plugins/Experimental/Avalanche/Source/AvalancheMedia/Private/Playback/Nodes/Events/Actions/AvaPlaybackNode_PlayAnim.h
@@ -0,0 +1,30 @@
+// Copyright Epic Games, Inc. All Rights Reserved.
+
+#pragma once
+
+#include "CoreMinimal.h"
+#include "Playback/Nodes/Events/Actions/AvalancheAnimations.h"
+#include "AvaPlaybackNode_PlayAnim.generated.h"
+
+UCLASS()
+class AVALANCHEMEDIA_API UAvaPlaybackNode_PlayAnim : public UAvaPlaybackNodeAction
+{
+	GENERATED_BODY()
+
+public:
+	
+	virtual FText GetNodeDisplayNameText() const override;
+	virtual FText GetNodeTooltipText() const override;
+	virtual void OnEventTriggered(const FAvaPlaybackEventParameters& InEventParameters) override;
+
+	virtual void PreDryRun() override;
+	virtual void DryRun(const TArray<UAvaPlaybackNode*>& InAncestors) override;
+	virtual void PostDryRun() override;
+protected:
+	
+	UPROPERTY(VisibleAnywhere, Category = "Motion Design")
+	TMap<FSoftObjectPath, FAvalancheAnimations> AnimationMap;
+
+	//All the Blueprints connected to this Node from the Dry Run (only populated while Dry Running!)
+	TSet<FSoftObjectPath> SeenAssetsInDryRun;
+};
diff --git a/Engine/Plugins/Experimental/Avalanche/Source/AvalancheMedia/Private/Playback/Nodes/Events/Actions/AvaPlaybackNode_PreloadPlayer.cpp b/Engine/Plugins/Experimental/Avalanche/Source/AvalancheMedia/Private/Playback/Nodes/Events/Actions/AvaPlaybackNode_PreloadPlayer.cpp
new file mode 100644
--- /dev/null
+++ b/Engine/Plugins/Experimental/Avalanche/Source/AvalancheMedia/Private/Playback/Nodes/Events/Actions/AvaPlaybackNode_PreloadPlayer.cpp
@@ -0,0 +1,34 @@
+// Copyright Epic Games, Inc. All Rights Reserved.
+
+#include "AvaPlaybackNode_PreloadPlayer.h"
+#include "AvalancheBroadcast.h"
+#include "Playback/AvalanchePlayback.h"
+
+#define LOCTEXT_NAMESPACE "AvalanchePlayback"
+
+FText UAvaPlaybackNode_PreloadPlayer::GetNodeDisplayNameText() const
+{
+	return LOCTEXT("PreloadPlayerNode_Title", "Preload Player");
+}
+
+FText UAvaPlaybackNode_PreloadPlayer::GetNodeTooltipText() const
+{
+	return LOCTEXT("PreloadPlayerNode_Tooltip", "Loads the Avalanche World Preemptively and Begins Play on that World");
+}
+
+void UAvaPlaybackNode_PreloadPlayer::OnEventTriggered(const FAvaPlaybackEventParameters& InEventParameters)
+{
+	UAvalanchePlayback* const Playback = GetPlayback();
+	
+	//Make sure we have a valid Avalanche Asset
+	if (Playback && InEventParameters.IsAssetValid())
+	{
+		TArray<FName> ChannelNames = Playback->GetChannelNamesForIndices(InEventParameters.ChannelIndices);
+		for (const FName& ChannelName : ChannelNames)
+		{
+			Playback->LoadAsset(InEventParameters.AvalancheAsset, ChannelName);
+		}
+	}
+}
+
+#undef LOCTEXT_NAMESPACE 
\ No newline at end of file
