[MotionDesign]

diff --git a/Engine/Plugins/Experimental/Avalanche/Source/AvalancheMedia/Private/Playback/Playables/AvalancheBlueprintPlayable.cpp b/Engine/Plugins/Experimental/Avalanche/Source/AvalancheMedia/Private/Playback/Playables/AvalancheBlueprintPlayable.cpp
new file mode 100644
--- /dev/null
+++ b/Engine/Plugins/Experimental/Avalanche/Source/AvalancheMedia/Private/Playback/Playables/AvalancheBlueprintPlayable.cpp
@@ -0,0 +1,183 @@
+// Copyright Epic Games, Inc. All Rights Reserved.
+
+#include "Playback/Playables/AvalancheBlueprintPlayable.h"
+
+#include "AvaBlueprint.h"
+#include "AvaRemoteControlRebind.h"
+#include "AvaSequencePlaybackObject.h"
+#include "Framework/AvaSoftAssetPtr.h"
+#include "Playback/AvaMediaPlayableGroup.h"
+#include "Playback/AvaMediaPlayableGroupManager.h"
+#include "Viewport/AvaCameraManager.h"
+
+#define LOCTEXT_NAMESPACE "AvalancheBlueprintPlayable"
+
+bool UAvalancheBlueprintPlayable::LoadAsset(const FAvaSoftAssetPtr& InAvalancheSourceAsset, bool bInInitiallyVisible)
+{
+	if (!PlayableGroup)
+	{
+		return false;
+	}
+
+	// Ensure world is created. Does nothing if already created.
+	PlayableGroup->ConditionalCreateWorld(); 
+
+	check(InAvalancheSourceAsset.GetAssetType() == EAvalancheAssetType::Blueprint);
+	return LoadAvalancheBlueprintInternal(TSoftObjectPtr<UAvalancheBlueprint>(InAvalancheSourceAsset.ToSoftObjectPath()));
+}
+
+bool UAvalancheBlueprintPlayable::UnloadAsset()
+{
+	if (IsValid(ManagedAvalancheBlueprint.Get()))
+	{
+		ManagedAvalancheBlueprint->UnregisterRemoteControlPreset();
+	}
+	SourceAvalancheBlueprint.Reset();
+	ManagedAvalancheBlueprint = nullptr;
+	return true;
+}
+
+const FSoftObjectPath& UAvalancheBlueprintPlayable::GetSourceAssetPath() const
+{
+	return SourceAvalancheBlueprint.ToSoftObjectPath();
+}
+
+EAvalanchePlayableStatus UAvalancheBlueprintPlayable::GetPlayableStatus() const
+{
+	// The blueprint assets don't stream. We can't even make them visible/invisible.
+	return ManagedAvalancheBlueprint ? EAvalanchePlayableStatus::Visible : EAvalanchePlayableStatus::Unloaded;
+}
+
+IAvaSceneInterface* UAvalancheBlueprintPlayable::GetSceneInterface() const
+{
+	if (ManagedAvalancheBlueprint)
+	{
+		return static_cast<IAvaSceneInterface*>(ManagedAvalancheBlueprint);
+	}
+	return nullptr;
+}
+
+EAvalanchePlayableCommandResult UAvalancheBlueprintPlayable::ExecuteAnimationCommand(EAvaMediaAnimAction InAnimAction, const FAnimPlaySettings& InAnimPlaySettings)
+{
+	if (InAnimAction ==  EAvaMediaAnimAction::CameraCut)
+	{
+		if (GetPlayableGroup())
+		{
+			GetPlayableGroup()->QueueCameraCut();
+		}
+		return EAvalanchePlayableCommandResult::Executed;
+	}
+	
+	return Super::ExecuteAnimationCommand(InAnimAction, InAnimPlaySettings);
+}
+
+bool UAvalancheBlueprintPlayable::InitPlayable(const FPlayableCreationInfo& InPlayableInfo)
+{
+	// Avalanche blueprint requires a separate playable group for each asset.
+	
+	UAvaMediaPlayableGroup::FPlayableGroupCreationInfo PlayableGroupCreationInfo;
+	PlayableGroupCreationInfo.PlayableGroupManager = InPlayableInfo.PlayableGroupManager;
+	PlayableGroupCreationInfo.SourceAssetPath = InPlayableInfo.SourceAsset.ToSoftObjectPath();
+	PlayableGroupCreationInfo.ChannelName = InPlayableInfo.ChannelName;
+	PlayableGroupCreationInfo.bIsRemoteProxy = false;
+	PlayableGroupCreationInfo.bIsSharedGroup = false;
+
+	PlayableGroup = InPlayableInfo.PlayableGroup ?
+		InPlayableInfo.PlayableGroup : UAvaMediaPlayableGroup::MakePlayableGroup(InPlayableInfo.PlayableGroupManager, PlayableGroupCreationInfo);
+	
+	return Super::InitPlayable(InPlayableInfo);
+}
+
+void UAvalancheBlueprintPlayable::OnPlay()
+{
+	if (!GetPlayableGroup())
+	{
+		return;
+	}
+
+	const UAvalancheGameInstance* AvalancheGameInstance = Cast<UAvalancheGameInstance>(GetPlayableGroup()->GetGameInstance());
+	if (!AvalancheGameInstance)
+	{
+		return;
+	}
+	
+	UAvalancheGameViewportClient* ViewportClient = AvalancheGameInstance->GetAvalancheGameViewportClient();
+
+	// Verify if this can be changed on the fly.
+	ManagedAvalancheBlueprint->GetViewportQualitySettings().Apply(ViewportClient->EngineShowFlags);
+	
+	// Old code using ava camera manager. To retire.
+	constexpr bool bIsCanvasController = false;
+	ViewportClient->GetCameraManager()->Init(ManagedAvalancheBlueprint->GetPlaybackObject(), bIsCanvasController);
+#if WITH_EDITOR
+	ViewportClient->GetCameraManager()->SetDefaultViewTarget(GetPlayWorld(), ManagedAvalancheBlueprint->GetStartupCameraName());
+#endif
+}
+
+void UAvalancheBlueprintPlayable::OnEndPlay()
+{
+	if (ManagedAvalancheBlueprint)
+	{
+		if (IAvaSequencePlaybackObject* PlaybackObject = ManagedAvalancheBlueprint->GetPlaybackObject())
+		{
+			PlaybackObject->CleanupPlayers();
+		}
+	}
+}
+
+bool UAvalancheBlueprintPlayable::LoadAvalancheBlueprintInternal(const TSoftObjectPtr<UAvalancheBlueprint>& InSourceAvalancheBlueprint)
+{
+	if (!GetPlayableGroup())
+	{
+		return false;
+	}
+
+	UAvalancheGameInstance* AvalancheGameInstance = Cast<UAvalancheGameInstance>(GetPlayableGroup()->GetGameInstance());
+	if (!AvalancheGameInstance)
+	{
+		return false;
+	}
+	
+	UAvalancheBlueprint* Source;
+	{
+		TRACE_CPUPROFILER_EVENT_SCOPE(UAvalancheBlueprintPlayable::LoadAsset::LoadSourceAvaBp);
+		Source = InSourceAvalancheBlueprint.LoadSynchronous();
+		check(Source);
+	}
+
+#if WITH_EDITORONLY_DATA
+	FGuardValue_Bitfield(Source->bDuplicatingReadOnly, true);
+#endif
+
+	{
+		TRACE_CPUPROFILER_EVENT_SCOPE(UAvalancheBlueprintPlayable::LoadAsset::DupAvaBp);
+		const FName AvaInstanceName = MakeUniqueObjectName(AvalancheGameInstance, UAvalancheBlueprint::StaticClass(), Source->GetFName());
+		ManagedAvalancheBlueprint = Cast<UAvalancheBlueprint>(StaticDuplicateObject(Source, AvalancheGameInstance, AvaInstanceName, RF_Transient));
+	}
+
+	if (!ManagedAvalancheBlueprint)
+	{
+		return false;
+	}
+	
+	{
+		TRACE_CPUPROFILER_EVENT_SCOPE(UAvalancheBlueprintPlayable::LoadAsset::LoadAvaBpWorld);
+		ManagedAvalancheBlueprint->SetAvalancheWorld(AvalancheGameInstance->GetPlayWorld());
+		ManagedAvalancheBlueprint->LoadAvalancheWorld();
+	}
+
+	ManagedAvalancheBlueprint->UpdatePlaceholderActor();
+	
+	{
+		TRACE_CPUPROFILER_EVENT_SCOPE(UAvalancheBlueprintPlayable::LoadAsset::RebindRCP);
+		ManagedAvalancheBlueprint->RegisterRemoteControlPreset();
+		FAvaRemoteControlRebind::RebindUnboundEntities(ManagedAvalancheBlueprint->GetRemoteControlPreset());
+	}
+
+	AvalancheGameInstance->MarkSynchronousAssetLoadingThisFrame();
+	
+	SourceAvalancheBlueprint = InSourceAvalancheBlueprint;
+	return true;
+}
+
+#undef LOCTEXT_NAMESPACE
diff --git a/Engine/Plugins/Experimental/Avalanche/Source/AvalancheMedia/Private/Playback/Playables/AvalancheLevelStreamingPlayable.cpp b/Engine/Plugins/Experimental/Avalanche/Source/AvalancheMedia/Private/Playback/Playables/AvalancheLevelStreamingPlayable.cpp
new file mode 100644
--- /dev/null
+++ b/Engine/Plugins/Experimental/Avalanche/Source/AvalancheMedia/Private/Playback/Playables/AvalancheLevelStreamingPlayable.cpp
@@ -0,0 +1,780 @@
+// Copyright Epic Games, Inc. All Rights Reserved.
+
+#include "Playback/Playables/AvalancheLevelStreamingPlayable.h"
+
+#include "AvaRemoteControlRebind.h"
+#include "AvaRemoteControlUtils.h"
+#include "AvaScene.h"
+#include "Camera/CameraActor.h"
+#include "Engine/Level.h"
+#include "Engine/LevelStreamingDynamic.h"
+#include "Engine/LocalPlayer.h"
+#include "Framework/AvaSoftAssetPtr.h"
+#include "Framework/AvalancheInstanceSettings.h"
+#include "IAvaMediaModule.h"
+#include "Playback/AvaMediaPlayableGroup.h"
+#include "Playback/AvaMediaPlayableGroupManager.h"
+#include "Playback/AvaMediaPlaybackUtils.h"
+#include "Streaming/LevelStreamingDelegates.h"
+#include "UObject/Package.h"
+#include "Viewport/AvaCameraManager.h"
+
+#define LOCTEXT_NAMESPACE "AvalancheLevelStreamingPlayable"
+
+namespace UE::AvaMedia::LevelStreamingPlayable::Private
+{
+	AAvaScene* FindAvaScene(const ULevel* InLevel)
+	{
+		AAvaScene* AvaScene = nullptr;
+		InLevel->Actors.FindItemByClass(&AvaScene);
+		return AvaScene;
+	}
+
+	AActor* FindActorByName(const ULevel* InLevel, const FName& InActorName)
+	{
+		for (AActor* Actor : InLevel->Actors)
+		{
+			if (Actor && Actor->GetFName() == InActorName)
+			{
+				return Actor;
+			}
+		}
+		return nullptr;
+	}
+
+	bool IsCameraActor(const AActor* InActor)
+	{
+		return InActor->IsA<ACameraActor>();
+	}
+	
+	AActor* FindFirstCameraActor(const ULevel* InLevel)
+	{
+		for (AActor* Actor : InLevel->Actors)
+		{
+			if (Actor && IsCameraActor(Actor))
+			{
+				return Actor;
+			}
+		}
+		return nullptr;
+	}
+
+	FName GetStartupCameraName(const AAvaScene* InAvaScene)
+	{
+		return InAvaScene ? InAvaScene->GetStartupCameraName() : NAME_None;
+	}
+
+	AActor* FindCameraActor(const ULevel* InLevel, FName InStartupCameraName)
+	{
+		if (InStartupCameraName != NAME_None)
+		{
+			AActor* CameraActor = FindActorByName(InLevel, InStartupCameraName);
+			if (CameraActor && IsCameraActor(CameraActor))
+			{
+				return CameraActor;
+			}
+		}
+		return nullptr;
+	}
+
+	AActor* FindStartupCameraActor(const ULevel* InLevel, FName InStartupCameraName, const AAvaScene* InAvaScene)
+	{
+		if (!IsValid(InLevel))
+		{
+			return nullptr;
+		}
+
+		// 1- Find by name from the parent scene.
+		if (AActor* FoundCameraActor = FindCameraActor(InLevel, InStartupCameraName))
+		{
+			return FoundCameraActor;
+		}
+
+		// 2- Find by name from dependent scene (if any). Dependent levels may not have an ava scene.
+		if (AActor* FoundCameraActor = FindCameraActor(InLevel, GetStartupCameraName(InAvaScene)))
+		{
+			return FoundCameraActor;
+		}
+
+		// 3- For now we fallback to looking for the first camera actor in the level.
+		return FindFirstCameraActor(InLevel);
+	}
+
+	EAvalanchePlayableStatus GetPlayableStatusFromLevelStreaming(const ULevelStreaming* InLevelStreaming)
+	{
+		if (!InLevelStreaming)
+		{
+			return EAvalanchePlayableStatus::Unloaded;
+		}
+
+		switch (InLevelStreaming->GetLevelStreamingState())
+		{
+			case ELevelStreamingState::Removed:
+				return EAvalanchePlayableStatus::Unloaded;
+			
+			case ELevelStreamingState::Unloaded:
+				// If the LevelStreaming was not loaded and has just been made to be loading, the status will still be "unloaded" but
+				// we consider it is loading.
+				return InLevelStreaming->ShouldBeLoaded() ? EAvalanchePlayableStatus::Loading : EAvalanchePlayableStatus::Unloaded;
+			
+			case ELevelStreamingState::FailedToLoad:
+				return EAvalanchePlayableStatus::Error;
+			
+			case ELevelStreamingState::Loading:
+				return EAvalanchePlayableStatus::Loading;
+			
+			case ELevelStreamingState::LoadedNotVisible:
+			case ELevelStreamingState::MakingVisible:
+			case ELevelStreamingState::MakingInvisible:
+				return EAvalanchePlayableStatus::Loaded;
+			
+			case ELevelStreamingState::LoadedVisible:
+				return EAvalanchePlayableStatus::Visible;
+			
+			default:
+				return EAvalanchePlayableStatus::Error;
+		}
+	}
+	
+	static int32 GetPlayableStatusRelevanceShouldBeUnloaded(EAvalanchePlayableStatus InStatus)
+	{
+		switch (InStatus)
+		{
+			case EAvalanchePlayableStatus::Unknown:
+				return 0;
+			case EAvalanchePlayableStatus::Error:
+				return 1;
+			case EAvalanchePlayableStatus::Unloaded:
+				return 2; // Weakest (desired)
+			case EAvalanchePlayableStatus::Loading:
+				return 3;
+			case EAvalanchePlayableStatus::Loaded:
+				return 4;
+			case EAvalanchePlayableStatus::Visible:
+				return 5;
+			default:
+				return 0;
+		}
+	}
+
+	static int32 GetPlayableStatusRelevanceShouldBeLoadedOrVisible(EAvalanchePlayableStatus InStatus)
+	{
+		switch (InStatus)
+		{
+			case EAvalanchePlayableStatus::Unknown:
+				return 0;
+			case EAvalanchePlayableStatus::Error:
+				return 1;
+			case EAvalanchePlayableStatus::Unloaded:
+				return 5;
+			case EAvalanchePlayableStatus::Loading:
+				return 4;
+			case EAvalanchePlayableStatus::Loaded:
+				return 3;
+			case EAvalanchePlayableStatus::Visible:
+				return 2; // Weakest (desired)
+			default:
+				return 0;
+		}
+	}
+	
+	/**
+	 * Returns the relevance of the status for comparison with another status
+	 * in order to determine the combined status of a playable according to what it should be.
+	 */
+	static int32 GetPlayableStatusRelevance(EAvalanchePlayableStatus InStatus, bool bInShouldBeLoaded, bool bInShouldBeVisible)
+	{
+		// The "desired" status is the weakest of the valid statues.
+		if (bInShouldBeLoaded || bInShouldBeVisible)
+		{
+			return GetPlayableStatusRelevanceShouldBeLoadedOrVisible(InStatus);
+		}
+		
+		return GetPlayableStatusRelevanceShouldBeUnloaded(InStatus);
+	}
+}
+
+bool UAvalancheLevelStreamingPlayable::LoadAsset(const FAvaSoftAssetPtr& InAvalancheSourceAsset, bool bInInitiallyVisible)
+{
+	if (!PlayableGroup)
+	{
+		return false;
+	}
+
+	// Ensure world is created. Does nothing if already created.
+	PlayableGroup->ConditionalCreateWorld();
+
+	const FAvalancheInstanceSettings& PlaybackInstanceSettings = IAvaMediaModule::Get().GetAvalancheInstanceSettings();
+	bLoadSubPlayables = PlaybackInstanceSettings.bEnableLoadSubPlayables;
+	
+	check(InAvalancheSourceAsset.GetAssetType() == EAvalancheAssetType::World);
+	const bool bAssetLoading = LoadAvalancheLevel(TSoftObjectPtr<UWorld>(InAvalancheSourceAsset.ToSoftObjectPath()), bInInitiallyVisible); 
+	if (bAssetLoading)
+	{
+		PlayableGroup->NotifyLevelStreaming(this);
+	}
+	return bAssetLoading;
+}
+
+bool UAvalancheLevelStreamingPlayable::UnloadAsset()
+{
+	if (AvalancheScene)
+	{
+		FAvaRemoteControlUtils::UnregisterRemoteControlPreset(AvalancheScene->GetRemoteControlPreset());
+	}
+
+	UnloadSubPlayables();
+	
+	// Ripped from UPocketLevelInstance::BeginDestroy()
+	if (LevelStreaming)
+	{
+		LevelStreaming->bShouldBlockOnUnload = false;
+		LevelStreaming->SetShouldBeVisible(false);
+		LevelStreaming->SetShouldBeLoaded(false);
+		LevelStreaming->SetIsRequestingUnloadAndRemoval(true);
+		LevelStreaming->OnLevelShown.RemoveAll(this);
+		LevelStreaming->OnLevelLoaded.RemoveAll(this);
+
+		const ULevel* const Level = LevelStreaming->GetLoadedLevel();
+		if (IsValid(Level) && Level->GetPackage())
+		{
+			// Hack so that FLevelStreamingGCHelper::PrepareStreamedOutLevelForGC unloads the level.
+			Level->GetPackage()->SetPackageFlags(PKG_PlayInEditor);
+		}
+	}
+	LevelStreaming = nullptr;
+	AvalancheScene = nullptr;
+	SourceAvalancheLevel.Reset();
+	return true;
+}
+
+EAvalanchePlayableStatus UAvalancheLevelStreamingPlayable::GetPlayableStatus() const
+{
+	using namespace UE::AvaMedia::LevelStreamingPlayable::Private;
+
+	bool bShouldBeLoaded = LevelStreaming ? LevelStreaming->ShouldBeLoaded() : false;
+	bool bShouldBeVisible = LevelStreaming ? LevelStreaming->ShouldBeVisible() : false;
+
+	EAvalanchePlayableStatus MostRelevantStatus = EAvalanchePlayableStatus::Unknown;
+	int32 HighestRelevance = -1;
+
+	auto CompareAndSetMostRelevantStatus = [&MostRelevantStatus, &HighestRelevance, bShouldBeLoaded, bShouldBeVisible](EAvalanchePlayableStatus InStatus)
+	{
+		const int32 CurrentRelevance = GetPlayableStatusRelevance(InStatus, bShouldBeLoaded, bShouldBeVisible);
+		if (CurrentRelevance > HighestRelevance)
+		{
+			MostRelevantStatus = InStatus;
+			HighestRelevance = CurrentRelevance;
+		}
+	};
+	
+	// Check sub playables.
+	for (const UAvalancheLevelStreamingPlayable* SubPlayable : SubPlayables)
+	{
+		CompareAndSetMostRelevantStatus(SubPlayable->GetPlayableStatus());
+	}
+
+	CompareAndSetMostRelevantStatus(GetPlayableStatusFromLevelStreaming(LevelStreaming));
+
+	return MostRelevantStatus;
+}
+
+IAvaSceneInterface* UAvalancheLevelStreamingPlayable::GetSceneInterface() const
+{
+	if (AvalancheScene)
+	{
+		return static_cast<IAvaSceneInterface*>(AvalancheScene);
+	}
+	return nullptr;
+}
+
+bool UAvalancheLevelStreamingPlayable::ApplyCamera()
+{
+	using namespace UE::AvaMedia::LevelStreamingPlayable;
+
+	if (!AvalancheScene)
+	{
+		return false;
+	}
+
+	bool bSetupDone = false;
+	UAvalanchePlayable* PlayableOwningCamera = nullptr;
+
+	// Camera setup.
+	if (AActor* CameraActor = FindStartupCameraActor(NAME_None, &PlayableOwningCamera))
+	{
+		// New code using player controller.
+		// Note: doesn't work yet because there is no player controller in avalanche game instance, but it works in PIE.
+		if (APlayerController* PlayerController = GetPlayableGroup()->GetPlayWorld()->GetFirstPlayerController())
+		{
+			PlayerController->SetViewTargetWithBlend(CameraActor);
+			GetPlayableGroup()->SetLastAppliedCameraPlayable(PlayableOwningCamera);
+			bSetupDone = true;
+		}
+		else
+		{
+			// If there is no player controller, fallback to avalanche camera manager, if the game instance has it.
+			const UAvalancheGameInstance* AvalancheGameInstance = Cast<UAvalancheGameInstance>(GetPlayableGroup()->GetGameInstance());			
+			if (const UAvalancheGameViewportClient* ViewportClient = AvalancheGameInstance ? AvalancheGameInstance->GetAvalancheGameViewportClient() : nullptr)
+			{
+				// Old code using ava camera manager. To retire.
+				constexpr bool bIsCanvasController = false;
+				ViewportClient->GetCameraManager()->Init(AvalancheScene->GetPlaybackObject(), bIsCanvasController);
+#if WITH_EDITOR
+				ViewportClient->GetCameraManager()->SetViewTarget(CameraActor);
+#endif
+				GetPlayableGroup()->SetLastAppliedCameraPlayable(PlayableOwningCamera);
+				bSetupDone = true;
+				UE_LOG(LogAvalanchePlayable, Log, TEXT("No player controller found in \"%s\" - Using Ava Camera Manager instead."), *LevelStreaming->PackageNameToLoad.ToString())
+			}
+			else
+			{
+				UE_LOG(LogAvalanchePlayable, Warning, TEXT("No player controller found in \"%s\" - Ava Camera Manager not available either."), *LevelStreaming->PackageNameToLoad.ToString());				
+			}
+		}
+	}
+	else
+	{
+		UE_LOG(LogAvalanchePlayable, Error, TEXT("Failed to find camera \"%s\", or any camera actor, in loaded level."),
+			*Private::GetStartupCameraName(AvalancheScene).ToString());
+	}
+	return bSetupDone;
+}
+
+bool UAvalancheLevelStreamingPlayable::GetShouldBeVisible() const
+{
+	return LevelStreaming ? LevelStreaming->GetShouldBeVisibleFlag() : false;
+}
+
+void UAvalancheLevelStreamingPlayable::SetShouldBeVisible(bool bInShouldBeVisible)
+{
+	if (LevelStreaming)
+	{
+		LevelStreaming->SetShouldBeVisible(bInShouldBeVisible);
+	}
+
+	for (UAvalancheLevelStreamingPlayable* SubPlayable : SubPlayables)
+	{
+		SubPlayable->UpdateVisibilityFromParents();
+	}
+}
+
+bool UAvalancheLevelStreamingPlayable::LoadAvalancheLevel(const TSoftObjectPtr<UWorld>& InSourceAvalancheLevel, bool bInInitiallyVisible)
+{
+	if (SourceAvalancheLevel == InSourceAvalancheLevel)
+	{
+		// Already started loading.
+		return false;
+	}
+
+	if (!PlayableGroup)
+	{
+		return false;
+	}
+	
+	bool bSuccess = false;
+	const FVector Origin(0,0,0);
+
+	ULevelStreamingDynamic::FLoadLevelInstanceParams Params(PlayableGroup->GetPlayWorld()
+			, InSourceAvalancheLevel.GetLongPackageName()
+			, FTransform(FRotator::ZeroRotator, Origin));
+	Params.bLoadAsTempPackage = true;
+	Params.bInitiallyVisible  = bInInitiallyVisible;
+
+	LevelStreaming = ULevelStreamingDynamic::LoadLevelInstance(Params, bSuccess);
+
+	if (!bSuccess || !LevelStreaming)
+	{
+		UE_LOG(LogAvalanchePlayable, Error, TEXT("[%s]: Failed to load level instance `%s`."), *GetPathNameSafe(this), *InSourceAvalancheLevel.ToString());
+		return false;
+	}
+
+	SourceAvalancheLevel = InSourceAvalancheLevel;
+	LevelStreaming->SetShouldBeLoaded(true);
+	LevelStreaming->SetShouldBeVisible(bInInitiallyVisible);
+	return true;
+}
+
+void UAvalancheLevelStreamingPlayable::OnLevelStreamingStateChanged(UWorld* InWorld
+	, const ULevelStreaming* InLevelStreaming
+	, ULevel* InLevelIfLoaded
+	, ELevelStreamingState InPreviousState
+	, ELevelStreamingState InNewState)
+{
+	// Filter out levels we don't care about.
+	if (InLevelStreaming != LevelStreaming)
+	{
+		return;
+	}
+	
+	if (InNewState == ELevelStreamingState::FailedToLoad)
+	{
+		UE_LOG(LogAvalanchePlayable, Error, TEXT("Level \"%s\" failed to load."), *InLevelStreaming->PackageNameToLoad.ToString());		
+	}
+	else if (InNewState == ELevelStreamingState::LoadedNotVisible || InNewState == ELevelStreamingState::LoadedVisible)
+	{
+		const ULevel* const Level = LevelStreaming->GetLoadedLevel();
+		if (!IsValid(Level))
+		{
+			return;
+		}
+
+		if (UWorld* OuterWorld = Level->GetTypedOuter<UWorld>())
+		{
+			// Workaround to avoid UEditorEngine::CheckForWorldGCLeaks killing the editor.
+			// Change the sub-level's world to be a "persistent" world type.
+			OuterWorld->WorldType = EWorldType::GamePreview;
+
+			if (bLoadSubPlayables)
+			{
+				LoadSubPlayables(OuterWorld);
+			}
+		}
+
+		// Workaround to destroy the Linker Load so that it does not keep the underlying File Opened
+		if (Level->GetPackage())
+		{
+			FAvaMediaPlaybackUtils::FlushPackageLoading(Level->GetPackage());
+		}
+		
+		// Resolve the ava scene for the other operations.
+		ResolveAvalancheScene(Level);
+
+		OnPlayableStatusChanged().Broadcast(this);
+
+		for (const TObjectKey<UAvalancheLevelStreamingPlayable>& ParentPlayableKey : ParentPlayables)
+		{
+			if (UAvalancheLevelStreamingPlayable* ParentPlayable = ParentPlayableKey.ResolveObjectPtr())
+			{
+				ParentPlayable->OnLevelStreamingPlayableStatusChanged(this);
+			}
+		}
+
+		OnLevelStreamingPlayableStatusChanged(this);
+	}
+}
+
+void UAvalancheLevelStreamingPlayable::OnLevelStreamingPlayableStatusChanged(UAvalancheLevelStreamingPlayable* InPlayable)
+{
+	// OnPlay (camera setup, animations, etc) can only be done when the level is visible (components must be active).
+	// With camera rig, we also need to make sure the rig level is loaded and visible.
+
+	if (GetPlayableStatus() == EAvalanchePlayableStatus::Visible)
+	{
+		if (bOnPlayQueued)
+		{
+			bOnPlayQueued = false;
+			OnPlay();
+		}
+	}
+}
+
+void UAvalancheLevelStreamingPlayable::BindDelegates()
+{
+	if (!FLevelStreamingDelegates::OnLevelStreamingStateChanged.IsBoundToObject(this))
+	{
+		FLevelStreamingDelegates::OnLevelStreamingStateChanged.AddUObject(this, &UAvalancheLevelStreamingPlayable::OnLevelStreamingStateChanged);
+	}
+}
+
+void UAvalancheLevelStreamingPlayable::UnbindDelegates()
+{
+	FLevelStreamingDelegates::OnLevelStreamingStateChanged.RemoveAll(this);
+}
+
+ULevel* UAvalancheLevelStreamingPlayable::GetLoadedLevel() const
+{
+	if (!LevelStreaming)
+	{
+		return nullptr;	
+	}
+	
+	ULevel* Level = LevelStreaming->GetLoadedLevel();
+	return IsValid(Level) ? Level : nullptr;
+}
+
+void UAvalancheLevelStreamingPlayable::ResolveAvalancheScene(const ULevel* InLevel)
+{
+	if (!IsValid(AvalancheScene))
+	{
+		AvalancheScene = UE::AvaMedia::LevelStreamingPlayable::Private::FindAvaScene((InLevel));
+		if (AvalancheScene)
+		{
+			FAvaRemoteControlUtils::RegisterRemoteControlPreset(AvalancheScene->GetRemoteControlPreset(), /*bInEnsureUniqueId*/ true);
+			FAvaRemoteControlRebind::RebindUnboundEntities(AvalancheScene->GetRemoteControlPreset(), InLevel);
+		}
+		else
+		{
+			UE_LOG(LogAvalanchePlayable, Error, TEXT("Loaded level \"%s\" is not an Avalanche level."), *LevelStreaming->PackageNameToLoad.ToString());
+		}
+	}
+}
+
+AActor* UAvalancheLevelStreamingPlayable::FindStartupCameraActor(FName InStartupCameraName, UAvalanchePlayable** OutParentPlayable)
+{
+	using namespace UE::AvaMedia::LevelStreamingPlayable;
+
+	// TODO: Configurable priorities for camera selection.
+	
+	if (InStartupCameraName == NAME_None && AvalancheScene)
+	{
+		InStartupCameraName = AvalancheScene->GetStartupCameraName();
+	}
+	
+	if (AActor* FoundCameraActor = Private::FindStartupCameraActor(GetLoadedLevel(), InStartupCameraName, AvalancheScene))
+	{
+		if (OutParentPlayable)
+		{
+			*OutParentPlayable = this;
+		}
+		return FoundCameraActor;
+	}
+
+	for (UAvalancheLevelStreamingPlayable* SubPlayable : SubPlayables)
+	{
+		if (AActor* FoundCameraActor = SubPlayable->FindStartupCameraActor(InStartupCameraName, OutParentPlayable))
+		{
+			return FoundCameraActor;
+		}
+	}
+	
+	return nullptr;
+}
+
+bool UAvalancheLevelStreamingPlayable::InitPlayable(const FPlayableCreationInfo& InPlayableInfo)
+{
+	// For now, we share all the levels in the same instance group. We may do sub-grouping later.
+	PlayableGroup = InPlayableInfo.PlayableGroup ?
+		InPlayableInfo.PlayableGroup : InPlayableInfo.PlayableGroupManager->GetOrCreateSharedLevelGroup(InPlayableInfo.ChannelName, false);
+
+	const bool bInitSuccess = Super::InitPlayable(InPlayableInfo);
+	
+	if (bInitSuccess)
+	{
+		BindDelegates();
+	}
+	return bInitSuccess;
+}
+
+void UAvalancheLevelStreamingPlayable::OnPlay()
+{
+	if (!LevelStreaming || !GetPlayableGroup())
+	{
+		return;
+	}
+	
+	if (!LevelStreaming->GetShouldBeVisibleFlag())
+	{
+		// Can't make visible immediately if part of a transition with other playables and the others are not ready.
+		GetPlayableGroup()->RequestSetVisibility(this, true);
+	}
+
+	const ULevel* const Level = LevelStreaming->GetLoadedLevel();
+	if (!IsValid(Level))
+	{
+		// Level is not yet loaded, queue the action when it gets loaded.
+		// Remark: we could either do this when the event is received or poll on the next tick.
+		const ELevelStreamingState LevelStreamingState = LevelStreaming->GetLevelStreamingState();
+		ensure(LevelStreaming->ShouldBeLoaded());
+		if (LevelStreamingState != ELevelStreamingState::FailedToLoad)
+		{
+			bOnPlayQueued = true;
+		}
+		else
+		{
+			UE_LOG(LogAvalanchePlayable, Error, TEXT("Level \"%s\" is not loading. Current Streaming State: \"%s\"."),
+				*LevelStreaming->PackageNameToLoad.ToString(), EnumToString(LevelStreamingState));			
+		}
+		return;
+	}
+
+	// Check if the level is visible. We can't do the actual camera setup, or animation
+	// if the level is not yet visible as the components are inactive.
+	const ELevelStreamingState LevelStreamingState = LevelStreaming->GetLevelStreamingState();
+	if (LevelStreamingState != ELevelStreamingState::LoadedVisible)
+	{
+		bOnPlayQueued = true;
+		return;
+	}
+
+	// Ensure scene is resolved.
+	ResolveAvalancheScene(Level);
+
+	if (!AvalancheScene)
+	{
+		return;
+	}
+
+	ApplyCamera();
+}
+
+void UAvalancheLevelStreamingPlayable::OnEndPlay()
+{
+	// Ensure the level is hidden and clear dirty flag because those are transient and shouldn't be saved.
+	if (LevelStreaming)
+	{
+		const ULevel* const Level = LevelStreaming->GetLoadedLevel();
+		if (IsValid(Level))
+		{
+			if (UPackage* LevelPackage = Level->GetPackage())
+			{
+				LevelPackage->ClearDirtyFlag();
+			}
+		}
+
+		LevelStreaming->SetShouldBeVisible(false);
+	}
+}
+
+void UAvalancheLevelStreamingPlayable::BeginDestroy()
+{
+	UnbindDelegates();
+	Super::BeginDestroy();
+}
+
+void UAvalancheLevelStreamingPlayable::LoadSubPlayables(const UWorld* InLevelInstanceWorld)
+{
+	check(InLevelInstanceWorld);
+	
+	for (const ULevelStreaming* SubLevelStreaming : InLevelInstanceWorld->GetStreamingLevels())
+	{
+		if (SubLevelStreaming)
+		{
+			GetOrLoadSubPlayable(SubLevelStreaming);
+		}
+	}
+}
+
+void UAvalancheLevelStreamingPlayable::UnloadSubPlayables()
+{
+	for (UAvalancheLevelStreamingPlayable* SubPlayable : SubPlayables)
+	{
+		if (SubPlayable)
+		{
+			SubPlayable->ParentPlayables.Remove(this);
+
+			// Shared Sub-Playables will be unloaded if they no longer have any parent
+			// playables to keep them alive.
+			if (!SubPlayable->HasParentPlayables())
+			{
+				SubPlayable->UnloadAsset();
+				if (UAvaMediaPlayableGroup* ParentPlayableGroup = SubPlayable->GetPlayableGroup())
+				{
+					ParentPlayableGroup->UnregisterPlayable(SubPlayable);
+				}
+			}
+		}
+	}
+	
+	SubPlayables.Reset();
+}
+
+void UAvalancheLevelStreamingPlayable::GetOrLoadSubPlayable(const ULevelStreaming* InLevelStreaming)
+{
+	if (!PlayableGroup)
+	{
+		return;
+	}
+
+	const FSoftObjectPath SourceAssetPath = InLevelStreaming->GetWorldAsset().ToSoftObjectPath();
+	
+	// Assume sub playables are shared, i.e. unique instance per group. For now.
+	// Need to see with Brad how to determine instancing scope (i.e. global vs local).	
+	if (UAvalancheLevelStreamingPlayable* ExistingPlayable = Cast<UAvalancheLevelStreamingPlayable>(PlayableGroup->FindFirstPlayableBySourceAssetPath(SourceAssetPath)))
+	{
+		AddSubPlayable(ExistingPlayable);
+		return;
+	}
+
+	if (UAvalancheLevelStreamingPlayable* NewPlayable = CreateSubPlayable(PlayableGroup, SourceAssetPath))
+	{
+		// TODO: Propagate more stuff from InLevelStreaming. Needs to reach LoadAvalancheLevel.
+		const FAvaSoftAssetPtr AssetPtr = { UWorld::StaticClass(), TSoftObjectPtr(SourceAssetPath)};
+		if (NewPlayable->LoadAsset(AssetPtr, GetShouldBeVisible()))
+		{
+			AddSubPlayable(NewPlayable);
+		}
+		else
+		{
+			PlayableGroup->UnregisterPlayable(NewPlayable);
+		}
+	}
+}
+
+UAvalancheLevelStreamingPlayable* UAvalancheLevelStreamingPlayable::CreateSubPlayable(UAvaMediaPlayableGroup* InPlayableGroup, const FSoftObjectPath& InSourceAssetPath)
+{
+	if (!InPlayableGroup)
+	{
+		return nullptr;
+	}
+
+	UAvalancheLevelStreamingPlayable* NewPlayable = NewObject<UAvalancheLevelStreamingPlayable>(GEngine);
+	
+	const FPlayableCreationInfo PlayableCreationInfo =
+		{
+			InPlayableGroup->GetPlayableGroupManager(),
+			{ UWorld::StaticClass(), TSoftObjectPtr(InSourceAssetPath) },
+			FName(),
+			InPlayableGroup
+		};
+	
+	if (NewPlayable && !NewPlayable->InitPlayable(PlayableCreationInfo))
+	{
+		// final setup may fail, in this case the playable is discarded.
+		return nullptr;
+	}
+
+	return NewPlayable;
+}
+
+void UAvalancheLevelStreamingPlayable::AddSubPlayable(UAvalancheLevelStreamingPlayable* InSubPlayable)
+{
+	if (InSubPlayable)
+	{
+		SubPlayables.AddUnique(InSubPlayable);
+		InSubPlayable->ParentPlayables.Add(this);
+	}
+}
+
+void UAvalancheLevelStreamingPlayable::RemoveSubPlayable(UAvalancheLevelStreamingPlayable* InSubPlayable)
+{
+	if (InSubPlayable)
+	{
+		SubPlayables.Remove(InSubPlayable);
+		InSubPlayable->ParentPlayables.Remove(this);
+	}
+}
+
+bool UAvalancheLevelStreamingPlayable::HasParentPlayables() const
+{
+	for (const TObjectKey<UAvalancheLevelStreamingPlayable>& ParentPlayableKey : ParentPlayables)
+	{
+		if (ParentPlayableKey.ResolveObjectPtr())
+		{
+			return true;
+		}
+	}
+	return false;
+}
+
+void UAvalancheLevelStreamingPlayable::UpdateVisibilityFromParents()
+{
+	bool bShouldBeVisible = false;
+	
+	for (const TObjectKey<UAvalancheLevelStreamingPlayable>& ParentPlayableKey : ParentPlayables)
+	{
+		if (const UAvalancheLevelStreamingPlayable* ParentPlayable = ParentPlayableKey.ResolveObjectPtr())
+		{
+			if (ParentPlayable->GetShouldBeVisible())
+			{
+				bShouldBeVisible = true;
+				break;
+			}
+		}
+	}
+
+	SetShouldBeVisible(bShouldBeVisible);
+}
+
+
+#undef LOCTEXT_NAMESPACE
diff --git a/Engine/Plugins/Experimental/Avalanche/Source/AvalancheMedia/Private/Playback/Playables/AvalancheRemoteProxyPlayable.cpp b/Engine/Plugins/Experimental/Avalanche/Source/AvalancheMedia/Private/Playback/Playables/AvalancheRemoteProxyPlayable.cpp
new file mode 100644
--- /dev/null
+++ b/Engine/Plugins/Experimental/Avalanche/Source/AvalancheMedia/Private/Playback/Playables/AvalancheRemoteProxyPlayable.cpp
@@ -0,0 +1,299 @@
+// Copyright Epic Games, Inc. All Rights Reserved.
+
+#include "Playback/Playables/AvalancheRemoteProxyPlayable.h"
+
+#include "AvalancheBroadcast.h"
+#include "Framework/AvaSoftAssetPtr.h"
+#include "IAvaMediaModule.h"
+#include "Playback/AvaMediaPlayableGroupManager.h"
+#include "Playback/AvaMediaPlaybackClientDelegates.h"
+#include "Playback/IAvaMediaPlaybackClient.h"
+
+#define LOCTEXT_NAMESPACE "AvalancheRemoteProxyPlayable"
+
+bool UAvalancheRemoteProxyPlayable::LoadAsset(const FAvaSoftAssetPtr& InAvalancheSourceAsset, bool bInInitiallyVisible)
+{
+	if (!PlayableGroup)
+	{
+		return false;
+	}
+	
+	SourceAssetPath = InAvalancheSourceAsset.ToSoftObjectPath();
+
+	IAvaMediaPlaybackClient& PlaybackClient = IAvaMediaModule::Get().GetMediaPlaybackClient();
+	
+	if (PlaybackClient.HasAnyServerOnlineForChannel(PlayingChannelFName))
+	{
+		const FString ChannelName = PlayingChannelName;
+		const FSoftObjectPath& AssetPath = InAvalancheSourceAsset.ToSoftObjectPath();
+		const TOptional<EAvaMediaPlaybackStatus> RemoteStatusOpt = PlaybackClient.GetRemotePlaybackStatus(InstanceId, AssetPath, ChannelName);
+		const EAvaMediaPlaybackStatus RemoteStatus = RemoteStatusOpt.IsSet() ? RemoteStatusOpt.GetValue() : EAvaMediaPlaybackStatus::Unknown;
+
+		const bool bCanLoad = RemoteStatus == EAvaMediaPlaybackStatus::Available
+			|| RemoteStatus == EAvaMediaPlaybackStatus::Unknown;
+		const bool bIsLoaded = RemoteStatus == EAvaMediaPlaybackStatus::Loading
+			|| RemoteStatus == EAvaMediaPlaybackStatus::Loaded
+			|| RemoteStatus == EAvaMediaPlaybackStatus::Starting
+			|| RemoteStatus == EAvaMediaPlaybackStatus::Started;
+
+		if (bCanLoad && !bIsLoaded)
+		{
+			PlaybackClient.RequestPlayback(InstanceId, AssetPath, ChannelName, EAvaMediaPlaybackAction::Load);
+			PlaybackClient.RequestPlayback(InstanceId, AssetPath, ChannelName, EAvaMediaPlaybackAction::SetUserData, UserData);
+		}
+	}
+	return true;
+}
+
+bool UAvalancheRemoteProxyPlayable::UnloadAsset()
+{
+	IAvaMediaPlaybackClient& PlaybackClient = IAvaMediaModule::Get().GetMediaPlaybackClient();
+	if (PlaybackClient.HasAnyServerOnlineForChannel(PlayingChannelFName))
+	{
+		PlaybackClient.RequestPlayback(InstanceId, SourceAssetPath, PlayingChannelName, EAvaMediaPlaybackAction::Unload);
+	}
+	return true;
+}
+
+namespace UE::AvalancheRemoteProxyPlayable::Private
+{
+	EAvalanchePlayableStatus GetPlayableStatus(EAvaMediaPlaybackStatus InPlaybackStatus)
+	{
+		switch (InPlaybackStatus)
+		{
+		case EAvaMediaPlaybackStatus::Unknown:
+			return EAvalanchePlayableStatus::Unknown;
+		case EAvaMediaPlaybackStatus::Missing:
+		case EAvaMediaPlaybackStatus::Syncing:
+		case EAvaMediaPlaybackStatus::Available:
+			return EAvalanchePlayableStatus::Unloaded;
+		case EAvaMediaPlaybackStatus::Loading:
+			return EAvalanchePlayableStatus::Loading;
+		case EAvaMediaPlaybackStatus::Loaded:
+			return EAvalanchePlayableStatus::Loaded;
+		case EAvaMediaPlaybackStatus::Starting:
+			return EAvalanchePlayableStatus::Loaded;
+		case EAvaMediaPlaybackStatus::Started:
+			return EAvalanchePlayableStatus::Visible;
+		case EAvaMediaPlaybackStatus::Stopping:
+			return EAvalanchePlayableStatus::Loaded;
+		case EAvaMediaPlaybackStatus::Unloading:
+			return EAvalanchePlayableStatus::Unloaded;
+		case EAvaMediaPlaybackStatus::Error:
+		default:
+			return EAvalanchePlayableStatus::Error;
+		}
+	}
+}
+
+EAvalanchePlayableStatus UAvalancheRemoteProxyPlayable::GetPlayableStatus() const
+{
+	IAvaMediaPlaybackClient& PlaybackClient = IAvaMediaModule::Get().GetMediaPlaybackClient();
+	const TArray<FString> OnlineServers = GetOnlineServerForChannel(PlayingChannelFName);
+	
+	for (const FString& Server : OnlineServers)
+	{
+		TOptional<EAvaMediaPlaybackStatus> PlaybackStatus = PlaybackClient.GetRemotePlaybackStatus(InstanceId, SourceAssetPath, PlayingChannelName, Server);
+
+		if (!PlaybackStatus.IsSet())
+		{
+			PlaybackClient.RequestPlayback(InstanceId, SourceAssetPath, PlayingChannelName, EAvaMediaPlaybackAction::Status);
+			PlaybackStatus = EAvaMediaPlaybackStatus::Unknown;
+		}
+
+		// TODO: reconcile forked channels.
+		return UE::AvalancheRemoteProxyPlayable::Private::GetPlayableStatus(PlaybackStatus.GetValue());
+	}
+	
+	return EAvalanchePlayableStatus::Unknown;
+}
+
+IAvaSceneInterface* UAvalancheRemoteProxyPlayable::GetSceneInterface() const
+{
+	return nullptr;
+}
+
+EAvalanchePlayableCommandResult UAvalancheRemoteProxyPlayable::ExecuteAnimationCommand(EAvaMediaAnimAction InAnimAction, const FAnimPlaySettings& InAnimPlaySettings)
+{
+	IAvaMediaPlaybackClient& PlaybackClient = IAvaMediaModule::Get().GetMediaPlaybackClient();
+	
+	// If an animation event was locally scheduled on a remote playable,
+	// we need to propagate the event.
+	if (PlaybackClient.HasAnyServerOnlineForChannel(PlayingChannelFName))
+	{
+		switch (InAnimAction)
+		{
+		case EAvaMediaAnimAction::None:
+			break;
+		case EAvaMediaAnimAction::Play:
+		case EAvaMediaAnimAction::PreviewFrame:
+			PlaybackClient.RequestAnimPlayback(InstanceId, SourceAssetPath, PlayingChannelName, InAnimPlaySettings);
+			break;
+
+		case EAvaMediaAnimAction::Continue:
+		case EAvaMediaAnimAction::Stop:
+		case EAvaMediaAnimAction::CameraCut:
+			PlaybackClient.RequestAnimAction(InstanceId, SourceAssetPath, PlayingChannelName, InAnimPlaySettings.AnimationName.ToString(), InAnimAction);
+			break;
+
+		default:
+			UE_LOG(LogAvalanchePlayable, Warning,
+				TEXT("Animation command action \"%s\" for asset \"%s\" on channel \"%s\" is not implemented."),
+				*StaticEnum<EAvaMediaAnimAction>()->GetValueAsString(InAnimAction),
+				*SourceAssetPath.ToString(), *PlayingChannelName);
+			break;
+		}
+	}
+	return EAvalanchePlayableCommandResult::Executed;
+}
+
+EAvalanchePlayableCommandResult UAvalancheRemoteProxyPlayable::UpdateRemoteControlCommand(const TSharedRef<FAvalancheRemoteControlValues>& InRemoteControlValues)
+{
+	IAvaMediaPlaybackClient& PlaybackClient = IAvaMediaModule::Get().GetMediaPlaybackClient();
+	if (PlaybackClient.HasAnyServerOnlineForChannel(PlayingChannelFName))
+	{
+		PlaybackClient.RequestRemoteControlUpdate(InstanceId, SourceAssetPath, PlayingChannelName, *InRemoteControlValues);
+	}
+	return EAvalanchePlayableCommandResult::Executed;
+}
+
+bool UAvalancheRemoteProxyPlayable::ApplyCamera()
+{
+	return false;
+}
+
+void UAvalancheRemoteProxyPlayable::SetUserData(const FString& InUserData)
+{
+	if (UserData != InUserData)
+	{
+		IAvaMediaPlaybackClient& PlaybackClient = IAvaMediaModule::Get().GetMediaPlaybackClient();
+		if (PlaybackClient.HasAnyServerOnlineForChannel(PlayingChannelFName))
+		{
+			PlaybackClient.RequestPlayback(InstanceId, GetSourceAssetPath(), PlayingChannelName, EAvaMediaPlaybackAction::SetUserData, InUserData);
+		}
+	}
+	
+	Super::SetUserData(InUserData);
+}
+
+bool UAvalancheRemoteProxyPlayable::InitPlayable(const FPlayableCreationInfo& InPlayableInfo)
+{
+	// We keep track of the channel this playable is part of.
+	PlayingChannelFName = InPlayableInfo.ChannelName;
+	PlayingChannelName = InPlayableInfo.ChannelName.ToString();
+	
+	constexpr bool bIsRemoteProxy = true;
+
+	// Remote playables have proxy playable groups imitating the same setup as local ones.
+	switch (InPlayableInfo.SourceAsset.GetAssetType())
+	{
+	case EAvalancheAssetType::Blueprint:
+		{
+			UAvaMediaPlayableGroup::FPlayableGroupCreationInfo PlayableGroupCreationInfo;
+			PlayableGroupCreationInfo.PlayableGroupManager = InPlayableInfo.PlayableGroupManager;
+			PlayableGroupCreationInfo.SourceAssetPath = InPlayableInfo.SourceAsset.ToSoftObjectPath();
+			PlayableGroupCreationInfo.ChannelName = InPlayableInfo.ChannelName;
+			PlayableGroupCreationInfo.bIsRemoteProxy = bIsRemoteProxy;
+			PlayableGroupCreationInfo.bIsSharedGroup = false;
+			
+			PlayableGroup = UAvaMediaPlayableGroup::MakePlayableGroup(InPlayableInfo.PlayableGroupManager, PlayableGroupCreationInfo);
+		}
+		break;
+	case EAvalancheAssetType::World:
+		PlayableGroup = InPlayableInfo.PlayableGroupManager->GetOrCreateSharedLevelGroup(InPlayableInfo.ChannelName, bIsRemoteProxy);
+		break;
+	default:
+		UE_LOG(LogAvalanchePlayable, Error, TEXT("Asset \"%s\" is an unsupported type."), *InPlayableInfo.SourceAsset.ToSoftObjectPath().ToString());
+		break;
+	}
+	
+	const bool bInitSucceeded = Super::InitPlayable(InPlayableInfo);
+	if (bInitSucceeded)
+	{
+		RegisterClientEventHandlers();
+	}
+	return bInitSucceeded;
+}
+
+void UAvalancheRemoteProxyPlayable::OnPlay()
+{
+	IAvaMediaModule& AvaMediaModule = IAvaMediaModule::Get();
+	if (!AvaMediaModule.IsMediaPlaybackClientStarted())
+	{
+		return;
+	}
+	
+	IAvaMediaPlaybackClient& Client = AvaMediaModule.GetMediaPlaybackClient();
+	
+	if (Client.HasAnyServerOnlineForChannel(PlayingChannelFName))
+	{
+		const FString ChannelName = PlayingChannelName;
+		const TOptional<EAvaMediaPlaybackStatus> RemoteStatusOpt = Client.GetRemotePlaybackStatus(InstanceId, SourceAssetPath, ChannelName);
+		const EAvaMediaPlaybackStatus RemoteStatus = RemoteStatusOpt.IsSet() ? RemoteStatusOpt.GetValue() : EAvaMediaPlaybackStatus::Unknown;
+		// TODO: rework this logic.
+		if (RemoteStatus == EAvaMediaPlaybackStatus::Available
+			|| RemoteStatus == EAvaMediaPlaybackStatus::Loading || RemoteStatus == EAvaMediaPlaybackStatus::Loaded
+			|| RemoteStatus == EAvaMediaPlaybackStatus::Unknown || RemoteStatus == EAvaMediaPlaybackStatus::Stopping
+			|| RemoteStatus == EAvaMediaPlaybackStatus::Unloading)
+		{
+			Client.RequestPlayback(InstanceId, SourceAssetPath, ChannelName, EAvaMediaPlaybackAction::Start);
+		}
+	}
+}
+
+void UAvalancheRemoteProxyPlayable::OnEndPlay()
+{
+	IAvaMediaPlaybackClient& PlaybackClient = IAvaMediaModule::Get().GetMediaPlaybackClient();
+	if (PlaybackClient.HasAnyServerOnlineForChannel(PlayingChannelFName))
+	{
+		PlaybackClient.RequestPlayback(InstanceId, SourceAssetPath, PlayingChannelName, EAvaMediaPlaybackAction::Stop);
+	}
+}
+
+void UAvalancheRemoteProxyPlayable::BeginDestroy()
+{
+	UnregisterClientEventHandlers();
+	Super::BeginDestroy();
+}
+
+void UAvalancheRemoteProxyPlayable::RegisterClientEventHandlers()
+{
+	using namespace UE::AvaMediaPlaybackClient::Delegates;
+	GetOnPlaybackSequenceEvent().RemoveAll(this);
+	GetOnPlaybackSequenceEvent().AddUObject(this, &UAvalancheRemoteProxyPlayable::HandleAvaMediaPlaybackSequenceEvent);
+}
+
+void UAvalancheRemoteProxyPlayable::UnregisterClientEventHandlers() const
+{
+	using namespace UE::AvaMediaPlaybackClient::Delegates;
+	GetOnPlaybackSequenceEvent().RemoveAll(this);
+}
+
+TArray<FString> UAvalancheRemoteProxyPlayable::GetOnlineServerForChannel(const FName& InChannelName)
+{
+	const FAvaOutputChannel& Channel = UAvalancheBroadcast::Get().GetCurrentProfile().GetChannel(InChannelName);
+	const TArray<UMediaOutput*>& Outputs = Channel.GetMediaOutputs();
+	TArray<FString> OnlineServers;
+	OnlineServers.Reserve(Outputs.Num());
+	
+	for (const UMediaOutput* Output : Outputs)
+	{
+		if (Channel.IsMediaOutputRemote(Output) && Channel.GetMediaOutputState(Output) != EAvaMediaOutputState::Offline)
+		{
+			OnlineServers.AddUnique(Channel.GetMediaOutputServerName(Output));
+		}
+	}
+	return OnlineServers;
+}
+
+void UAvalancheRemoteProxyPlayable::HandleAvaMediaPlaybackSequenceEvent(IAvaMediaPlaybackClient& InPlaybackClient,
+	const UE::AvaMediaPlaybackClient::Delegates::FPlaybackSequenceEventArgs& InEventArgs)
+{
+	if (InEventArgs.InstanceId == InstanceId && InEventArgs.ChannelName == PlayingChannelName)
+	{
+		const FName SequenceName(InEventArgs.SequenceName);
+		OnSequenceEventDelegate.Broadcast(this, SequenceName, InEventArgs.EventType);
+	}
+}
+
+#undef LOCTEXT_NAMESPACE
diff --git a/Engine/Plugins/Experimental/Avalanche/Source/AvalancheMedia/Private/Playback/Transition/AvaLocalPlayableTransition.cpp b/Engine/Plugins/Experimental/Avalanche/Source/AvalancheMedia/Private/Playback/Transition/AvaLocalPlayableTransition.cpp
new file mode 100644
--- /dev/null
+++ b/Engine/Plugins/Experimental/Avalanche/Source/AvalancheMedia/Private/Playback/Transition/AvaLocalPlayableTransition.cpp
@@ -0,0 +1,485 @@
+// Copyright Epic Games, Inc. All Rights Reserved.
+
+#include "Playback/Transition/AvaLocalPlayableTransition.h"
+
+#include "AvaSequencePlayer.h"
+#include "AvaTransitionSubsystem.h"
+#include "AvaTransitionTree.h"
+#include "Behavior/IAvaTransitionBehavior.h"
+#include "Execution/AvaTransitionExecutorBuilder.h"
+#include "Execution/IAvaTransitionExecutor.h"
+#include "IAvaSceneInterface.h"
+#include "IAvaSequenceProvider.h"
+#include "Playback/AvaMediaPlaybackUtils.h"
+#include "Playback/AvalanchePlayable.h"
+#include "Playback/Transition/AvaPlayableTransitionPrivate.h"
+#include "Playback/Transition/AvaTransitionPlayableScene.h"
+
+namespace UE::AvaPlayableTransition::Private
+{
+	UAvaTransitionSubsystem* GetTransitionSubsystem(const UAvalanchePlayable& InPlayable)
+	{
+		if (const UAvaMediaPlayableGroup* PlayableGroup = InPlayable.GetPlayableGroup())
+		{
+			if (const UWorld* World = PlayableGroup->GetPlayWorld())
+			{
+				return World->GetSubsystem<UAvaTransitionSubsystem>();
+			}
+		}
+		return nullptr;
+	}
+
+	IAvaTransitionBehavior* GetTransitionBehavior(const UAvalanchePlayable& InPlayable, const UAvaTransitionSubsystem& InTransitionSubsystem)
+	{
+		if (const IAvaSceneInterface* SceneInterface = InPlayable.GetSceneInterface())
+		{
+			if (ULevel* const Level = SceneInterface->GetSceneLevel())
+			{
+				return InTransitionSubsystem.GetTransitionBehavior(Level);
+			}
+		}
+		return nullptr;		
+	}
+
+	IAvaTransitionBehavior* GetTransitionBehavior(const UAvalanchePlayable& InPlayable)
+	{
+		if (const UAvaTransitionSubsystem* TransitionSubsystem = GetTransitionSubsystem(InPlayable))
+		{
+			return GetTransitionBehavior(InPlayable, *TransitionSubsystem);
+		}
+		return nullptr;		
+	}
+
+	bool MakeNullTransitionBehaviorInstance(UAvalanchePlayableTransition* InPlayableTransition, const FAvaTagHandle& InTransitionLayer, FAvaTransitionBehaviorInstance& OutBehaviorInstance)
+	{
+		IAvaTransitionBehavior* const TransitionBehavior = nullptr;
+		OutBehaviorInstance = FAvaTransitionBehaviorInstance();
+		OutBehaviorInstance.SetBehavior(TransitionBehavior);
+		OutBehaviorInstance.CreateScene<FAvaTransitionPlayableScene>(InPlayableTransition, InTransitionLayer, InPlayableTransition);
+		return true;
+	}
+
+	FAvaTagHandle GetTransitionLayer(const FAvaTransitionBehaviorInstance& InBehaviorInstance)
+	{
+		// Remark: InBehaviorInstance.GetTransitionLayer() doesn't return a valid layer.
+		if (const IAvaTransitionBehavior* Behavior = InBehaviorInstance.GetBehavior())
+		{
+			if (const UAvaTransitionTree* TransitionTree = Behavior->GetTransitionTree())
+			{
+				return TransitionTree->GetTransitionLayer();
+			}
+		}
+		return FAvaTagHandle();
+	}
+	
+	class FBuilderHelper
+	{
+	public:
+		FBuilderHelper(FString&& InContextName, UAvalanchePlayableTransition* InPlayableTransition)
+			: ContextName(MoveTemp(InContextName))
+			, PlayableTransition(InPlayableTransition)
+		{
+			ExecutorBuilder.SetContextName(ContextName);
+		}
+
+		const FString& GetContextName() const { return ContextName; }
+
+		bool HasBehaviorInstances() const
+		{
+			return NumEnterInstance > 0 || NumExitInstance > 0;
+		}
+
+		bool AddTransitionBehaviorInstance(UAvalanchePlayable* InPlayable, EAvaMediaPlayableTransitionEntryRole InPlayableRole)
+		{
+			if (!InPlayable)
+			{
+				UE_LOG(LogAvalanchePlayable, Error, TEXT("Playable Transition \"%s\" setup error: Invalid playable."), *ContextName);
+				return false;
+			}
+
+			UAvaTransitionSubsystem* TransitionSubsystem = GetTransitionSubsystem(*InPlayable); 
+
+			if (!TransitionSubsystem)
+			{
+				UE_LOG(LogAvalanchePlayable, Error,
+					TEXT("Playable Transition \"%s\" setup error: Can't retrieve transition subsystem for playable {%s}."),
+					*ContextName, *GetPrettyPlayableInfo(InPlayable));
+				return false;
+			}
+
+			if (LastTransitionSubsystem != nullptr && LastTransitionSubsystem != TransitionSubsystem)
+			{
+				// Todo: if we hit this error, we need to create multiple playable transitions batched per subsystem.
+				UE_LOG(LogAvalanchePlayable, Error,
+					TEXT("Playable Transition \"%s\" setup error: Playable {%s} is in a different subsystem."),
+					*ContextName, *GetPrettyPlayableInfo(InPlayable));
+				return false;
+			}
+			
+			LastTransitionSubsystem = TransitionSubsystem;
+
+			FAvaTransitionBehaviorInstance BehaviorInstance;
+			if (MakeTransitionBehaviorInstance(InPlayable, TransitionSubsystem, BehaviorInstance))
+			{
+				if (InPlayableRole == EAvaMediaPlayableTransitionEntryRole::Enter)
+				{
+					ExecutorBuilder.AddEnterInstance(BehaviorInstance);
+					++NumEnterInstance;
+				}
+				else
+				{
+					// Either "playing" or "exit" are considered Exit in that implementation layer.
+					ExecutorBuilder.AddExitInstance(BehaviorInstance);
+					++NumExitInstance;
+					
+					// Special case of transition out (no enter instances).
+					// The way this works, we need to make a dummy "Enter" behavior instance
+					// to provide the transition layer for the logic take out.
+					if (NumEnterInstance == 0 && InPlayableRole == EAvaMediaPlayableTransitionEntryRole::Exit)
+					{
+						FAvaTransitionBehaviorInstance NullEnterInstance;
+						if (MakeNullTransitionBehaviorInstance(PlayableTransition, GetTransitionLayer(BehaviorInstance), NullEnterInstance))
+						{
+							ExecutorBuilder.AddEnterInstance(NullEnterInstance);
+						}
+					}
+				}
+				return true;
+			}
+			return false;
+		}
+		
+		bool MakeTransitionBehaviorInstance(UAvalanchePlayable* InPlayable, const UAvaTransitionSubsystem* InTransitionSubsystem, FAvaTransitionBehaviorInstance& OutBehaviorInstance) const
+		{
+			if (IAvaTransitionBehavior* TransitionBehavior = GetTransitionBehavior(*InPlayable, *InTransitionSubsystem))
+			{
+				if (const UAvaTransitionTree* TransitionTree = TransitionBehavior->GetTransitionTree())
+				{
+					if (TransitionTree->IsEnabled())
+					{
+						OutBehaviorInstance = FAvaTransitionBehaviorInstance();
+						OutBehaviorInstance.SetBehavior(TransitionBehavior);
+						OutBehaviorInstance.CreateScene<FAvaTransitionPlayableScene>(InPlayable, InPlayable, PlayableTransition);
+						return true;
+					}
+				}
+			}
+			return false;
+		}
+
+	private:
+		FString ContextName;
+		UAvalanchePlayableTransition* PlayableTransition = nullptr;
+		int32 NumEnterInstance = 0;
+		int32 NumExitInstance = 0;
+
+	public:
+		UAvaTransitionSubsystem* LastTransitionSubsystem = nullptr;
+		FAvaTransitionExecutorBuilder ExecutorBuilder;
+	};
+
+	struct FSequenceHelper
+	{
+		const IAvaSequenceProvider* SequenceProvider = nullptr;
+		IAvaSequencePlaybackObject* SequencePlayback = nullptr;
+
+		FSequenceHelper(const UAvalanchePlayable* InPlayable)
+		{
+			if (InPlayable)
+			{
+				if (IAvaSceneInterface* SceneInterface = InPlayable->GetSceneInterface())
+				{
+					SequenceProvider = SceneInterface->GetSequenceProvider();
+					SequencePlayback = SceneInterface->GetPlaybackObject();
+				}
+			}
+		}
+
+		bool IsValid() const
+		{
+			return SequenceProvider && SequencePlayback;
+		}
+	};
+
+	// Source: FAvaTransitionInitializeSequence::ExecuteSequenceTask
+	void InitializeSequences(const UAvalanchePlayable* InPlayable)
+	{
+		const FSequenceHelper Helper(InPlayable);
+		if (Helper.IsValid())
+		{
+			FAvaSequencePlayParams PlaySettings;
+			PlaySettings.Start = PlaySettings.End = FAvaSequenceTime(0.0);
+			PlaySettings.PlayMode = EAvaSequencePlayMode::Forward;
+
+			for (const TObjectPtr<UAvaSequence>& Sequence : Helper.SequenceProvider->GetSequences())
+			{
+				Helper.SequencePlayback->PlaySequence(Sequence.Get(), PlaySettings);
+			}
+		}
+	}	
+
+	// Source: FAvaTransitionSequenceUtils::UpdatePlayerRunStatus
+	bool HasActiveSequences(const UAvalanchePlayable* InPlayable)
+	{
+		const FSequenceHelper Helper(InPlayable);
+		if (Helper.IsValid())
+		{
+			for (const TObjectPtr<UAvaSequence>& Sequence : Helper.SequenceProvider->GetSequences())
+			{
+				if (const UAvaSequencePlayer* SequencePlayer = Helper.SequencePlayback->GetSequencePlayer(Sequence.Get()))
+				{
+					// Considering paused as not active.
+					if (SequencePlayer->GetPlaybackStatus() != EMovieScenePlayerStatus::Paused)
+					{
+						return true;
+					}
+				}
+			}
+		}
+		return false;
+	}
+
+	bool HasActiveSequences(const TArray<TWeakObjectPtr<UAvalanchePlayable>>& InPlayablesWeak)
+	{
+		for (const TWeakObjectPtr<UAvalanchePlayable> PlayableWeak : InPlayablesWeak)
+		{
+			if (HasActiveSequences(PlayableWeak.Get()))
+			{
+				return true;
+			}
+		}
+		return false;
+	}
+}
+
+bool UAvaLocalPlayableTransition::Start()
+{
+	if (!Super::Start())
+	{
+		return false;
+	}
+
+	using namespace UE::AvaPlayableTransition::Private;
+
+	FBuilderHelper Helper(GetFullName(), this);
+
+	TArray<UAvalanchePlayable*> EnterPlayables = Pin(EnterPlayablesWeak);
+	TArray<UAvalanchePlayable*> PlayingPlayables = Pin(PlayingPlayablesWeak);
+	TArray<UAvalanchePlayable*> ExitPlayables = Pin(ExitPlayablesWeak);
+	
+	if (EnterPlayables.IsEmpty() && PlayingPlayables.IsEmpty() && ExitPlayables.IsEmpty())
+	{
+		UE_LOG(LogAvalanchePlayable, Error,
+			TEXT("Playable Transition \"%s\" setup error: no playables specified, either as enter or exit. Nothing to do transition on."),
+			*Helper.GetContextName());
+		return false;
+	}
+
+	int32 ArrayIndex = 0;
+	for (UAvalanchePlayable* Playable : EnterPlayables)
+	{
+		if (EnterPlayableValues.IsValidIndex(ArrayIndex) && EnterPlayableValues[ArrayIndex].IsValid())
+		{
+			Playable->UpdateRemoteControlCommand(EnterPlayableValues[ArrayIndex].ToSharedRef());	
+		}
+		
+		if (!Helper.AddTransitionBehaviorInstance(Playable, EAvaMediaPlayableTransitionEntryRole::Enter))
+		{
+			// Handle Playables with no transition tree.
+			InitializeSequences(Playable);
+			PostExecutorSequencePlayablesWeak.Add(Playable);
+		}
+		++ArrayIndex;
+	}
+
+	for (UAvalanchePlayable* Playable : PlayingPlayables)
+	{
+		if (!Helper.AddTransitionBehaviorInstance(Playable, EAvaMediaPlayableTransitionEntryRole::Playing))
+		{
+			// No Transition Tree:
+			// Stop exit/playing playable without any transition.
+			// We can't integrate that with a transition tree for enter pages atm.
+			UAvalanchePlayable::OnTransitionEvent().Broadcast(Playable, this, EAvalanchePlayableTransitionEventFlags::StopPlayable);
+		}
+	}
+
+	for (UAvalanchePlayable* Playable : ExitPlayables)
+	{
+		if (!Helper.AddTransitionBehaviorInstance(Playable, EAvaMediaPlayableTransitionEntryRole::Exit))
+		{
+			// No Transition Tree:
+			// Stop exit/playing playable without any transition.
+			// We can't integrate that with a transition tree for enter pages atm.
+			UAvalanchePlayable::OnTransitionEvent().Broadcast(Playable, this, EAvalanchePlayableTransitionEventFlags::StopPlayable);
+		}
+	}
+
+	bool bTransitionExecutorStarted = false;
+	if (Helper.HasBehaviorInstances() && Helper.LastTransitionSubsystem)
+	{
+		using namespace UE::AvaMediaPlayback::Utils;
+		UE_LOG(LogAvalanchePlayable, Verbose, TEXT("%s Playable Transition \"%s\" starting."), *GetBriefFrameInfo(), *Helper.GetContextName());
+
+		Helper.ExecutorBuilder.SetOnFinished(FSimpleDelegate::CreateUObject(this, &UAvaLocalPlayableTransition::OnTransitionExecutorEnded));
+		
+		TransitionExecutor = Helper.ExecutorBuilder.Build(*Helper.LastTransitionSubsystem);
+		if (TransitionExecutor)
+		{
+			TransitionExecutor->Start();
+			bTransitionExecutorStarted = true;
+		}
+	}
+
+	// If no transition tree, skip to next phase of the transition.
+	if (!bTransitionExecutorStarted)
+	{
+		PostTransitionExecutorPhase();
+	}
+
+	// Todo: reconsider if this is needed and if so where it should go. Not used in code for now.
+	UAvalanchePlayable::OnTransitionEvent().Broadcast(nullptr, this, EAvalanchePlayableTransitionEventFlags::Starting);
+
+	return true;
+}
+
+void UAvaLocalPlayableTransition::Stop()
+{
+	if (TransitionExecutor.IsValid())
+	{
+		TransitionExecutor->Stop();
+		
+		// Stop should've called OnTransitionTreeEnded, which resets the ptr.
+		ensure(!TransitionExecutor.IsValid());
+	}
+
+	UAvalanchePlayable::OnSequenceEvent().RemoveAll(this);
+
+	Super::Stop();
+}
+
+bool UAvaLocalPlayableTransition::IsRunning() const
+{
+	// The transition can be considered "running" either if it is executing
+	// a transition tree or if it has "post transition tree" sequences going.
+	return TransitionExecutor.IsValid() || bWaitOnPostExecutorSequences;
+}
+
+void UAvaLocalPlayableTransition::Tick(double InDeltaSeconds)
+{
+	using namespace UE::AvaPlayableTransition::Private;
+	
+	// We need to poll the sequences. Sequence Events are not reliable.
+	if (bWaitOnPostExecutorSequences)
+	{
+		if (!HasActiveSequences(PostExecutorSequencePlayablesWeak))
+		{
+			FinishWaitOnPostExecutorSequences();
+			NotifyTransitionFinished();
+		}
+	}
+}
+
+void UAvaLocalPlayableTransition::OnTransitionExecutorEnded()
+{
+	using namespace UE::AvaMediaPlayback::Utils;
+	
+	if (TransitionExecutor.IsValid())
+	{
+		// Notify all the playable that have been discarded.
+		auto StopDiscardedPlayables = [this](const TArray<TWeakObjectPtr<UAvalanchePlayable>>& InPlayablesWeak)
+		{
+			for (const TWeakObjectPtr<UAvalanchePlayable>& PlayableWeak : InPlayablesWeak)
+			{
+				if (UAvalanchePlayable* Playable = PlayableWeak.Get())
+				{
+					if (DiscardPlayablesWeak.Contains(Playable))
+					{
+						UAvalanchePlayable::OnTransitionEvent().Broadcast(Playable, this, EAvalanchePlayableTransitionEventFlags::StopPlayable);
+					}
+				}
+			}
+		};
+
+		StopDiscardedPlayables(PlayingPlayablesWeak);
+		StopDiscardedPlayables(ExitPlayablesWeak);
+
+		TransitionExecutor.Reset();
+		
+		UE_LOG(LogAvalanchePlayable, Verbose, TEXT("%s Transition Executor \"%s\" ended."), *GetBriefFrameInfo(), *GetFullName());
+	}
+
+	PostTransitionExecutorPhase();
+}
+
+void UAvaLocalPlayableTransition::PostTransitionExecutorPhase()
+{
+	bool bSequenceStarted = false;
+	// Deal with non-TL enter playables. After the TL pages have been taken out,
+	// we need to start the sequences.
+	for (const TWeakObjectPtr<UAvalanchePlayable>& EnterPlayableWeak : PostExecutorSequencePlayablesWeak)
+	{
+		if (UAvalanchePlayable* Playable = EnterPlayableWeak.Get())
+		{
+			// Default settings will play all sequences.
+			if (Playable->ExecuteAnimationCommand(EAvaMediaAnimAction::Play, FAnimPlaySettings()) == EAvalanchePlayableCommandResult::Executed)
+			{
+				bSequenceStarted = true;
+			}
+		}
+	}
+
+	if (bSequenceStarted)
+	{
+		StartWaitOnPostExecutorSequences();
+	}
+	else
+	{
+		NotifyTransitionFinished();
+	}
+}
+
+void UAvaLocalPlayableTransition::NotifyTransitionFinished()
+{
+	using namespace UE::AvaMediaPlayback::Utils;
+	UE_LOG(LogAvalanchePlayable, Verbose, TEXT("%s Playable Transition \"%s\" ended."), *GetBriefFrameInfo(), *GetFullName());
+
+	// This will indicate the playable transition is completed and can be cleaned up.
+	// Combo templates break the one page one playable rule, so we need a dedicated event to signal the end of the playable transition.
+	UAvalanchePlayable::OnTransitionEvent().Broadcast(nullptr, this, EAvalanchePlayableTransitionEventFlags::Finished);
+}
+
+void UAvaLocalPlayableTransition::StartWaitOnPostExecutorSequences()
+{
+	bWaitOnPostExecutorSequences = true;
+
+	// Make sure we listen to sequence events.
+	if (UAvalanchePlayable::OnSequenceEvent().IsBoundToObject(this))
+	{
+		UAvalanchePlayable::OnSequenceEvent().AddUObject(this, &UAvaLocalPlayableTransition::OnPlayableSequenceEvent);
+	}
+}
+
+void UAvaLocalPlayableTransition::FinishWaitOnPostExecutorSequences()
+{
+	bWaitOnPostExecutorSequences = false;
+	
+	// Remark: This is the only task of playable transition requiring to listen to sequence events, *for now*.
+	UAvalanchePlayable::OnSequenceEvent().RemoveAll(this);
+}
+
+void UAvaLocalPlayableTransition::OnPlayableSequenceEvent(UAvalanchePlayable* InPlayable, const FName& SequenceName, EAvalanchePlayableSequenceEventType InEventType)
+{
+	// Remark: the sequence events are not entirely reliable.
+	// The transitions are also "ticked" to poll this condition.
+	using namespace UE::AvaPlayableTransition::Private;
+
+	if (PostExecutorSequencePlayablesWeak.Contains(InPlayable) && InEventType == EAvalanchePlayableSequenceEventType::Finished)
+	{
+		// Check if this was the last active sequence of this transition.
+		if (!HasActiveSequences(PostExecutorSequencePlayablesWeak))
+		{
+			FinishWaitOnPostExecutorSequences();
+			NotifyTransitionFinished();
+		}
+	}
+}
\ No newline at end of file
diff --git a/Engine/Plugins/Experimental/Avalanche/Source/AvalancheMedia/Private/Playback/Transition/AvaLocalPlayableTransition.h b/Engine/Plugins/Experimental/Avalanche/Source/AvalancheMedia/Private/Playback/Transition/AvaLocalPlayableTransition.h
new file mode 100644
--- /dev/null
+++ b/Engine/Plugins/Experimental/Avalanche/Source/AvalancheMedia/Private/Playback/Transition/AvaLocalPlayableTransition.h
@@ -0,0 +1,55 @@
+// Copyright Epic Games, Inc. All Rights Reserved.
+
+#pragma once
+
+#include "AvaMediaDefines.h"
+#include "Playback/Transition/AvalanchePlayableTransition.h"
+
+#include "AvaLocalPlayableTransition.generated.h"
+
+class IAvaTransitionExecutor;
+
+/**
+ * Implements the transition tree execution for all playables with one.
+ * It also implements the "hardcoded" TL for playables with no transition tree.
+ * The hardcoded TL rules are:
+ * - "exit" playables are stopped with "out" sequences defines in transition tree
+ *   if there is one, otherwise nothing is done.
+ * - "enter" playables have their animation triggered after the transition tree is finished, letting the
+ *  exit pages animated out as they should.
+ */
+UCLASS()
+class UAvaLocalPlayableTransition : public UAvalanchePlayableTransition
+{
+	GENERATED_BODY()
+	
+public:
+	//~ Begin UAvalanchePlayableTransition
+	virtual bool Start() override;
+	virtual void Stop() override;
+	virtual bool IsRunning() const override;
+	virtual void Tick(double InDeltaSeconds) override;
+	//~ End UAvalanchePlayableTransition
+
+protected:
+	void OnTransitionExecutorEnded();
+	
+	void PostTransitionExecutorPhase();
+
+	void OnPlayableSequenceEvent(UAvalanchePlayable* InPlayable, const FName& SequenceName, EAvalanchePlayableSequenceEventType InEventType);
+	
+	void NotifyTransitionFinished();
+
+	void StartWaitOnPostExecutorSequences();
+	void FinishWaitOnPostExecutorSequences();
+	
+protected:
+	/** Transition tree executor. */
+	TSharedPtr<IAvaTransitionExecutor> TransitionExecutor;
+
+	/** List of enter playables that have no TL and need to play the animations after the transition tree. */
+	TArray<TWeakObjectPtr<UAvalanchePlayable>> PostExecutorSequencePlayablesWeak;
+
+	/** Indicate the transition is waiting on sequences. */
+	bool bWaitOnPostExecutorSequences = false;
+};
\ No newline at end of file
diff --git a/Engine/Plugins/Experimental/Avalanche/Source/AvalancheMedia/Private/Playback/Transition/AvaMediaServerPlaybackTransition.cpp b/Engine/Plugins/Experimental/Avalanche/Source/AvalancheMedia/Private/Playback/Transition/AvaMediaServerPlaybackTransition.cpp
new file mode 100644
--- /dev/null
+++ b/Engine/Plugins/Experimental/Avalanche/Source/AvalancheMedia/Private/Playback/Transition/AvaMediaServerPlaybackTransition.cpp
@@ -0,0 +1,492 @@
+// Copyright Epic Games, Inc. All Rights Reserved.
+
+#include "Playback/Transition/AvaMediaServerPlaybackTransition.h"
+
+#include "IAvaMediaModule.h"
+#include "Playback/AvaMediaPlaybackManager.h"
+#include "Playback/AvaMediaPlaybackServer.h"
+#include "Playback/Transition/AvalanchePlayableTransition.h"
+
+namespace UE::AvaMediaServerPlaybackTransition::Private
+{
+	FString GetPrettyPlaybackInstanceInfo(const FAvaMediaPlaybackInstance* InPlaybackInstance)
+	{
+		if (InPlaybackInstance)
+		{
+			return FString::Printf(TEXT("Id:%s, Asset:%s, Channel:%s, UserData:\"%s\""),
+				*InPlaybackInstance->GetInstanceId().ToString(),
+				*InPlaybackInstance->GetSourcePath().ToString(),
+				*InPlaybackInstance->GetChannelName(),
+				*InPlaybackInstance->GetInstanceUserData());
+		}
+		return TEXT("");
+	}
+
+	UAvalanchePlayable* GetPlayable(const FAvaMediaPlaybackInstance* InPlaybackInstance)
+	{
+		if (InPlaybackInstance && InPlaybackInstance->GetPlayback())
+		{
+			return InPlaybackInstance->GetPlayback()->GetFirstPlayable();
+		}
+		return nullptr;
+	}
+
+	TSharedPtr<FAvaMediaPlaybackInstance> FindInstanceForPlayable(const TArray<TWeakPtr<FAvaMediaPlaybackInstance>>& InPlaybackInstancesWeak, const UAvalanchePlayable* InPlayable)
+	{
+		for (const TWeakPtr<FAvaMediaPlaybackInstance>& InstanceWeak : InPlaybackInstancesWeak)
+		{
+			TSharedPtr<FAvaMediaPlaybackInstance> Instance = InstanceWeak.Pin();
+			if (GetPlayable(Instance.Get()) == InPlayable)
+			{
+				return Instance;
+			}
+		}
+		return nullptr;
+	}
+}
+
+void UAvaMediaServerPlaybackTransition::SetEnterValues(const TArray<FAvalancheRemoteControlValues>& InEnterValues)
+{
+	EnterValues.Reserve(InEnterValues.Num());
+	for (const FAvalancheRemoteControlValues& Values : InEnterValues)
+	{
+		EnterValues.Add(MakeShared<FAvalancheRemoteControlValues>(Values));
+	}
+}
+
+bool UAvaMediaServerPlaybackTransition::AddEnterInstance(const TSharedPtr<FAvaMediaPlaybackInstance>& InPlaybackInstance)
+{
+	if (!InPlaybackInstance)
+	{
+		return false;
+	}
+	
+	// Register this transition as a visibility constraint.
+	using namespace UE::AvaMediaServerPlaybackTransition::Private;
+	if (const UAvalanchePlayable* Playable = GetPlayable(InPlaybackInstance.Get()))
+	{
+		if (UAvaMediaPlayableGroup* PlayableGroup = Playable->GetPlayableGroup())
+		{
+			PlayableGroup->RegisterVisibilityConstraint(this);
+		}
+	}
+	else if (UAvalanchePlayback* Playback = InPlaybackInstance->GetPlayback())
+	{
+		// If the playable is not created yet, register to the creation event.
+		Playback->OnPlayableCreated.AddUObject(this, &UAvaMediaServerPlaybackTransition::OnPlayableCreated);
+	}
+	
+	return AddPlaybackInstance(InPlaybackInstance, EnterPlaybackInstancesWeak);
+}
+
+bool UAvaMediaServerPlaybackTransition::AddPlayingInstance(const TSharedPtr<FAvaMediaPlaybackInstance>& InPlaybackInstance)
+{
+	return AddPlaybackInstance(InPlaybackInstance, PlayingPlaybackInstancesWeak);
+}
+
+bool UAvaMediaServerPlaybackTransition::AddExitInstance(const TSharedPtr<FAvaMediaPlaybackInstance>& InPlaybackInstance)
+{
+	return AddPlaybackInstance(InPlaybackInstance, ExitPlaybackInstancesWeak);
+}
+
+void UAvaMediaServerPlaybackTransition::TryResolveInstances(const FAvaMediaPlaybackServer& InPlaybackServer)
+{
+	if (EnterPlaybackInstancesWeak.Num() < EnterInstanceIds.Num())
+	{
+		for (const FGuid& InstanceId : EnterInstanceIds)
+		{
+			if (TSharedPtr<FAvaMediaPlaybackInstance> Instance = InPlaybackServer.FindActivePlaybackInstance(InstanceId))
+			{
+				if (!EnterPlaybackInstancesWeak.Contains(Instance.ToWeakPtr()))
+				{
+					AddEnterInstance(Instance);
+				}
+			}
+		}
+	}	
+}
+
+bool UAvaMediaServerPlaybackTransition::IsVisibilityConstrained(const UAvalanchePlayable* InPlayable) const
+{
+	using namespace UE::AvaMediaServerPlaybackTransition::Private;
+	bool bAllPlayablesLoaded = true;
+	bool bIsPlayableInThisTransition = false;
+	
+	for (const TWeakPtr<FAvaMediaPlaybackInstance>& InstanceWeak : EnterPlaybackInstancesWeak)
+	{
+		if (const TSharedPtr<FAvaMediaPlaybackInstance> Instance = InstanceWeak.Pin())
+		{
+			if (const UAvalanchePlayable* Playable = GetPlayable(Instance.Get()))
+			{
+				if (Playable == InPlayable)
+				{
+					bIsPlayableInThisTransition = true;
+				}
+				const EAvalanchePlayableStatus PlayableStatus = Playable->GetPlayableStatus();
+				if (PlayableStatus != EAvalanchePlayableStatus::Loaded && PlayableStatus != EAvalanchePlayableStatus::Visible)
+				{
+					bAllPlayablesLoaded = false;
+				}
+			}
+		}
+	}
+	return bIsPlayableInThisTransition && !bAllPlayablesLoaded;
+}
+
+
+bool UAvaMediaServerPlaybackTransition::CanStart(bool& bOutShouldDiscard) const
+{
+	using namespace UE::AvaMediaServerPlaybackTransition::Private;
+
+	if (EnterPlaybackInstancesWeak.Num() < EnterInstanceIds.Num())
+	{
+		bOutShouldDiscard = false;
+		return false;
+	}
+
+	for (const TWeakPtr<FAvaMediaPlaybackInstance>& InstanceWeak : EnterPlaybackInstancesWeak)
+	{
+		if (const TSharedPtr<FAvaMediaPlaybackInstance> Instance = InstanceWeak.Pin())
+		{
+			if (const UAvalanchePlayable* Playable = GetPlayable(Instance.Get()))
+			{
+				const EAvalanchePlayableStatus PlayableStatus = Playable->GetPlayableStatus();
+
+				if (PlayableStatus == EAvalanchePlayableStatus::Unknown
+					|| PlayableStatus == EAvalanchePlayableStatus::Error)
+				{
+					// Discard the command
+					bOutShouldDiscard = true;
+					return false;
+				}
+
+				// todo: this might cause commands to become stale and fill the pending command list
+				if (PlayableStatus == EAvalanchePlayableStatus::Unloaded)
+				{
+					return false;
+				}
+
+				// Asset status must be visible to run the command.
+				// If not visible, the components are not yet added to the world.
+				if (PlayableStatus != EAvalanchePlayableStatus::Visible)
+				{
+					// Keep the command in the queue for next tick.
+					bOutShouldDiscard = false;
+					return false;
+				}
+			}
+		}
+		else
+		{
+			bOutShouldDiscard = true;
+			return false;
+		}
+	}
+	
+	bOutShouldDiscard = true;
+	return true;
+}
+
+void UAvaMediaServerPlaybackTransition::Start()
+{
+	RegisterToPlayableTransitionEvent();
+
+	// May fail if playables are not loaded yet. Playables are loaded
+	// when the playback object has ticked at least one.
+	MakePlayableTransition();
+
+	bool bTransitionStarted = false;
+	
+	if (PlayableTransition)
+	{
+		LogDetailedTransitionInfo();
+
+		// Todo: validate the level streaming playables are finished streaming the asset.
+		// Otherwise, transition start must be queued on playable streaming events. 
+		bTransitionStarted = PlayableTransition->Start();
+	}
+
+	if (!bTransitionStarted)
+	{
+		Stop();
+	}
+}
+
+void UAvaMediaServerPlaybackTransition::Stop()
+{
+	if (PlayableTransition)
+	{
+		PlayableTransition->Stop();
+		PlayableTransition = nullptr;
+	}
+
+	for (const TWeakPtr<FAvaMediaPlaybackInstance>& InstanceWeak : EnterPlaybackInstancesWeak)
+	{
+		if (const TSharedPtr<FAvaMediaPlaybackInstance> Instance = InstanceWeak.Pin())
+		{
+			if (UAvalanchePlayback* Playback = Instance->GetPlayback())
+			{
+				Playback->OnPlayableCreated.RemoveAll(this);
+			}
+		}
+	}
+	
+	UnregisterFromPlayableTransitionEvent();
+	
+	// Remove transition from server.
+	if (const TSharedPtr<FAvaMediaPlaybackServer> PlaybackServer = IAvaMediaModule::Get().GetMediaPlaybackServer())
+	{
+		if (!PlaybackServer->RemovePlaybackInstanceTransition(TransitionId))
+		{
+			UE_LOG(LogAvaPlaybackServer, Error,
+				TEXT("Playback Transition {%s} Error: Was not found in server's active transitions. "), *GetPrettyTransitionInfo());
+		}
+	}
+}
+
+bool UAvaMediaServerPlaybackTransition::IsRunning() const
+{
+	return PlayableTransition ? PlayableTransition->IsRunning() : false;
+}
+
+FString UAvaMediaServerPlaybackTransition::GetPrettyTransitionInfo() const
+{
+	return FString::Printf(TEXT("Id:%s, Channel:%s, Client:%s"),
+		*TransitionId.ToString(), *ChannelName.ToString(), *ClientName);
+}
+
+FString UAvaMediaServerPlaybackTransition::GetBriefTransitionDescription() const
+{
+	auto MakeInstanceIdList = [](const TArray<TWeakPtr<FAvaMediaPlaybackInstance>>& InPlaybackInstancesWeak) -> FString
+	{
+		FString InstanceIdList;
+		for (const TWeakPtr<FAvaMediaPlaybackInstance>& InstanceWeak : InPlaybackInstancesWeak)
+		{
+			if (const TSharedPtr<FAvaMediaPlaybackInstance> Instance = InstanceWeak.Pin())
+			{
+				InstanceIdList += FString::Printf(TEXT("%s%s"), InstanceIdList.IsEmpty() ? TEXT("") : TEXT(", "), *Instance->GetInstanceId().ToString());
+			}
+		}
+		return InstanceIdList.IsEmpty() ? TEXT("None") : InstanceIdList;
+	};
+	
+	const FString EnterInstanceList = MakeInstanceIdList(EnterPlaybackInstancesWeak);
+	const FString PlayingInstanceList = MakeInstanceIdList(PlayingPlaybackInstancesWeak);
+	const FString ExitInstanceList = MakeInstanceIdList(ExitPlaybackInstancesWeak);
+
+	return FString::Printf(TEXT("Enter Instance(s): [%s], Playing Instance(s): [%s], Exit Instance(s): [%s]."), *EnterInstanceList, *PlayingInstanceList, *ExitInstanceList);
+}
+
+TSharedPtr<FAvaMediaPlaybackInstance> UAvaMediaServerPlaybackTransition::FindInstanceForPlayable(const UAvalanchePlayable* InPlayable)
+{
+	using namespace UE::AvaMediaServerPlaybackTransition;
+	if (!InPlayable)
+	{
+		return nullptr;
+	}
+
+	TSharedPtr<FAvaMediaPlaybackInstance> Instance =  Private::FindInstanceForPlayable(EnterPlaybackInstancesWeak, InPlayable);
+	if (Instance)
+	{
+		return Instance;
+	}
+	
+	Instance =  Private::FindInstanceForPlayable(PlayingPlaybackInstancesWeak, InPlayable);
+	if (Instance)
+	{
+		return Instance;
+	}
+
+	Instance =  Private::FindInstanceForPlayable(ExitPlaybackInstancesWeak, InPlayable);
+	if (Instance)
+	{
+		return Instance;
+	}
+	return nullptr;
+}
+
+void UAvaMediaServerPlaybackTransition::OnTransitionEvent(UAvalanchePlayable* InPlayable, UAvalanchePlayableTransition* InTransition, EAvalanchePlayableTransitionEventFlags InTransitionFlags)
+{
+	using namespace UE::AvaMediaServerPlaybackTransition::Private;
+	
+	// not this transition.
+	if (InTransition != PlayableTransition || PlayableTransition == nullptr)
+	{
+		return;
+	}
+
+	const TSharedPtr<FAvaMediaPlaybackServer> PlaybackServer = IAvaMediaModule::Get().GetMediaPlaybackServer();
+
+	// Find the page player for this playable
+	if (const TSharedPtr<FAvaMediaPlaybackInstance> Instance = FindInstanceForPlayable(InPlayable))
+	{
+		// Relay the transition event back to the client
+		if (PlaybackServer)
+		{
+			PlaybackServer->SendPlayableTransitionEvent(TransitionId, InPlayable->GetInstanceId(), InTransitionFlags, ChannelName, ClientName);
+		}
+		
+		if (EnumHasAnyFlags(InTransitionFlags, EAvalanchePlayableTransitionEventFlags::StopPlayable))
+		{
+			// Validating that we are not removing an "enter" playable.
+			if (PlayableTransition->IsEnterPlayable(InPlayable))
+			{
+				UE_LOG(LogAvaPlaybackServer, Error,
+					TEXT("Playback Transition {%s} Error: An \"enter\" playable is being discarded for instance {%s}."),
+					*GetPrettyTransitionInfo(), *GetPrettyPlaybackInstanceInfo(Instance.Get()));
+			}
+
+			// See UAvalanchePagePlayer::Stop()
+			const EAvaPlaybackStopOptions PlaybackStopOptions = bUnloadDiscardedInstances ?
+				EAvaPlaybackStopOptions::Default | EAvaPlaybackStopOptions::Unload : EAvaPlaybackStopOptions::Default;	
+			Instance->GetPlayback()->Stop(PlaybackStopOptions);
+			
+			if (bUnloadDiscardedInstances)
+			{
+				Instance->Unload();
+				// Remove instance from the server.
+				if (PlaybackServer)
+				{
+					if (!PlaybackServer->RemoveActivePlaybackInstance(Instance->GetInstanceId()))
+					{
+						UE_LOG(LogAvaPlaybackServer, Error,
+							TEXT("Playback Transition {%s} Error: \"exit\" instance {%s} was not found in server's active instances. "),
+							*GetPrettyTransitionInfo(), *GetPrettyPlaybackInstanceInfo(Instance.Get()));
+					}
+				}
+			}
+			else
+			{
+				Instance->Recycle();
+			}
+		}
+	}
+
+	if (EnumHasAnyFlags(InTransitionFlags, EAvalanchePlayableTransitionEventFlags::Finished))
+	{
+		if (PlaybackServer)
+		{
+			PlaybackServer->SendPlayableTransitionEvent(TransitionId, FGuid(), InTransitionFlags, ChannelName, ClientName);
+		}
+		
+		Stop();
+	}
+}
+
+void UAvaMediaServerPlaybackTransition::OnPlayableCreated(UAvalanchePlayback* InPlayback, UAvalanchePlayable* InPlayable)
+{
+	if (UAvaMediaPlayableGroup* PlayableGroup = InPlayable->GetPlayableGroup())
+	{
+		PlayableGroup->RegisterVisibilityConstraint(this);
+	}
+}
+
+void UAvaMediaServerPlaybackTransition::MakePlayableTransition()
+{
+	using namespace UE::AvaMediaServerPlaybackTransition::Private;
+
+	FAvaPlayableTransitionBuilder TransitionBuilder;
+
+	auto AddInstancesToBuilder = [&TransitionBuilder, this](const TArray<TWeakPtr<FAvaMediaPlaybackInstance>>& InPlaybackInstancesWeak, const TCHAR* InCategory, EAvaMediaPlayableTransitionEntryRole InEntryRole)
+	{
+		using namespace UE::AvaMediaServerPlaybackTransition::Private;
+		int32 ArrayIndex = 0;
+		for (const TWeakPtr<FAvaMediaPlaybackInstance>& InstanceWeak : InPlaybackInstancesWeak)
+		{
+			if (const TSharedPtr<FAvaMediaPlaybackInstance> Instance = InstanceWeak.Pin())
+			{
+				if (UAvalanchePlayable* Playable = GetPlayable(Instance.Get()))
+				{
+					const bool bPlayableAdded = TransitionBuilder.AddPlayable(Playable, InEntryRole);
+					if (InEntryRole == EAvaMediaPlayableTransitionEntryRole::Enter && bPlayableAdded)
+					{
+						TransitionBuilder.AddEnterPlayableValues(EnterValues.IsValidIndex(ArrayIndex) ? EnterValues[ArrayIndex] : nullptr);	
+					}
+				}
+				else
+				{
+					// If this happens, likely the playable is not yet loaded.
+					UE_LOG(LogAvaPlaybackServer, Error,
+						TEXT("Playback Transition {%s} Error: Failed to retrieve \"%s\" playable for instance {%s}."),
+						*GetPrettyTransitionInfo(), InCategory, *GetPrettyPlaybackInstanceInfo(Instance.Get()));
+				}
+			}
+			++ArrayIndex;
+		}
+	};
+ 
+	AddInstancesToBuilder(EnterPlaybackInstancesWeak, TEXT("Enter"), EAvaMediaPlayableTransitionEntryRole::Enter);
+	AddInstancesToBuilder(PlayingPlaybackInstancesWeak, TEXT("Playing"), EAvaMediaPlayableTransitionEntryRole::Playing);
+	AddInstancesToBuilder(ExitPlaybackInstancesWeak, TEXT("Exit"), EAvaMediaPlayableTransitionEntryRole::Exit);
+	PlayableTransition = TransitionBuilder.MakeTransition(this);
+
+	if (PlayableTransition)
+	{
+		PlayableTransition->SetTransitionFlags(TransitionFlags);
+	}
+}
+
+void UAvaMediaServerPlaybackTransition::LogDetailedTransitionInfo() const
+{
+	UE_LOG(LogAvaPlaybackServer, Verbose, TEXT("Playback Transition {%s}:"), *GetPrettyTransitionInfo());
+
+	auto LogInstances = [](const TArray<TWeakPtr<FAvaMediaPlaybackInstance>>& InPlaybackInstancesWeak, const TCHAR* InCategory)
+	{
+		using namespace UE::AvaMediaServerPlaybackTransition::Private;
+
+		for (const TWeakPtr<FAvaMediaPlaybackInstance>& InstanceWeak : InPlaybackInstancesWeak)
+		{
+			if (const TSharedPtr<FAvaMediaPlaybackInstance> Instance = InstanceWeak.Pin())
+			{
+				UE_LOG(LogAvaPlaybackServer, Verbose, TEXT("- %s Instance: {%s}."), InCategory, *GetPrettyPlaybackInstanceInfo(Instance.Get()));
+			}
+		}
+	};
+
+	LogInstances(EnterPlaybackInstancesWeak, TEXT("Enter"));
+	LogInstances(PlayingPlaybackInstancesWeak, TEXT("Playing"));
+	LogInstances(ExitPlaybackInstancesWeak, TEXT("Exit"));
+}
+
+void UAvaMediaServerPlaybackTransition::RegisterToPlayableTransitionEvent()
+{
+	UAvalanchePlayable::OnTransitionEvent().RemoveAll(this);
+	UAvalanchePlayable::OnTransitionEvent().AddUObject(this, &UAvaMediaServerPlaybackTransition::OnTransitionEvent);
+}
+
+void UAvaMediaServerPlaybackTransition::UnregisterFromPlayableTransitionEvent() const
+{
+	UAvalanchePlayable::OnTransitionEvent().RemoveAll(this);
+}
+
+bool UAvaMediaServerPlaybackTransition::AddPlaybackInstance(const TSharedPtr<FAvaMediaPlaybackInstance>& InPlaybackInstance, TArray<TWeakPtr<FAvaMediaPlaybackInstance>>& OutPlaybackInstancesWeak)
+{
+	if (!InPlaybackInstance)
+	{
+		return false;	
+	}
+
+	OutPlaybackInstancesWeak.Add(InPlaybackInstance);
+	UpdateChannelName(InPlaybackInstance.Get());
+	return true;
+}
+
+void UAvaMediaServerPlaybackTransition::UpdateChannelName(const FAvaMediaPlaybackInstance* InPlaybackInstance)
+{
+	if (ChannelName.IsNone())
+	{
+		ChannelName = InPlaybackInstance->GetChannelFName();
+	}
+	else
+	{
+		using namespace UE::AvaMediaServerPlaybackTransition::Private;
+
+		// Validate the channel is the same.
+		if (ChannelName != InPlaybackInstance->GetChannelFName())
+		{
+			UE_LOG(LogAvaPlaybackServer, Error,
+				TEXT("Playback Transition {%s}: Adding Playback Instance {%s} in a different channel than previous playback instance (\"%s\")."),
+				*GetPrettyTransitionInfo(), *GetPrettyPlaybackInstanceInfo(InPlaybackInstance), *ChannelName.ToString());
+		}
+	}
+}
\ No newline at end of file
diff --git a/Engine/Plugins/Experimental/Avalanche/Source/AvalancheMedia/Private/Playback/Transition/AvaMediaServerPlaybackTransition.h b/Engine/Plugins/Experimental/Avalanche/Source/AvalancheMedia/Private/Playback/Transition/AvaMediaServerPlaybackTransition.h
new file mode 100644
--- /dev/null
+++ b/Engine/Plugins/Experimental/Avalanche/Source/AvalancheMedia/Private/Playback/Transition/AvaMediaServerPlaybackTransition.h
@@ -0,0 +1,91 @@
+// Copyright Epic Games, Inc. All Rights Reserved.
+
+#pragma once
+
+#include "Playback/AvalanchePlayable.h"
+#include "Playback/Transition/AvaMediaPlaybackTransition.h"
+#include "UObject/Object.h"
+
+#include "AvaMediaServerPlaybackTransition.generated.h"
+
+class FAvaMediaPlaybackInstance;
+class FAvaMediaPlaybackServer;
+class UAvalanchePlayableTransition;
+
+/**
+ * Class for creating and tracking playback graph instance transitions on the server.
+ * It is responsible for creating the playable transition object when requested from
+ * the playback graphs.
+ *
+ * This class handles each playback graph instance as a single playable. It is
+ * meant to be used by the playback server primarily.
+ */
+UCLASS()
+class UAvaMediaServerPlaybackTransition : public UAvaMediaPlaybackTransition
+{
+	GENERATED_BODY()
+	
+public:
+	void SetChannelName(const FName& InChannelName) { ChannelName = InChannelName; }
+	void SetTransitionId(const FGuid& InTransitionId) { TransitionId = InTransitionId; }
+	void SetClientName(const FString& InClientName) { ClientName = InClientName; }
+	void SetUnloadDiscardedInstances(bool bInUnloadDiscardedInstances) { bUnloadDiscardedInstances = bInUnloadDiscardedInstances; }
+	void SetTransitionFlags(EAvalanchePlayableTransitionFlags InTransitionFlags) { TransitionFlags = InTransitionFlags; }
+	void SetEnterInstanceIds(const TArray<FGuid>& InInstanceIds) { EnterInstanceIds = InInstanceIds; }
+	void SetEnterValues(const TArray<FAvalancheRemoteControlValues>& InEnterValues);
+	bool AddEnterInstance(const TSharedPtr<FAvaMediaPlaybackInstance>& InPlaybackInstance);
+	bool AddPlayingInstance(const TSharedPtr<FAvaMediaPlaybackInstance>& InPlaybackInstance);
+	bool AddExitInstance(const TSharedPtr<FAvaMediaPlaybackInstance>& InPlaybackInstance);
+
+	void TryResolveInstances(const FAvaMediaPlaybackServer& InPlaybackServer);
+	
+	//~ Begin IAvaPlayableVisibilityConstraint
+	virtual bool IsVisibilityConstrained(const UAvalanchePlayable* InPlayable) const override;
+	//~ End IAvaPlayableVisibilityConstraint
+	
+	//~ Begin UAvaMediaPlaybackTransition
+	virtual bool CanStart(bool& bOutShouldDiscard) const override;
+	virtual void Start() override;
+	virtual void Stop() override;
+	virtual bool IsRunning() const override;
+	//~ End UAvaMediaPlaybackTransition
+
+	/** Returns the channel this transition is happening in. A transition can only have instances within the same channel. */
+	FName GetChannelName() const { return ChannelName; }
+
+	FString GetPrettyTransitionInfo() const;
+	FString GetBriefTransitionDescription() const;
+	
+protected:
+	TSharedPtr<FAvaMediaPlaybackInstance> FindInstanceForPlayable(const UAvalanchePlayable* InPlayable);
+
+	void OnTransitionEvent(UAvalanchePlayable* InPlayable, UAvalanchePlayableTransition* InTransition, EAvalanchePlayableTransitionEventFlags InTransitionFlags);
+	void OnPlayableCreated(UAvalanchePlayback* InPlayback, UAvalanchePlayable* InPlayable);
+	
+	void MakePlayableTransition();
+
+	void LogDetailedTransitionInfo() const;
+
+	void RegisterToPlayableTransitionEvent();
+	void UnregisterFromPlayableTransitionEvent() const;
+
+	bool AddPlaybackInstance(const TSharedPtr<FAvaMediaPlaybackInstance>& InPlaybackInstance, TArray<TWeakPtr<FAvaMediaPlaybackInstance>>& OutPlaybackInstancesWeak);	
+	void UpdateChannelName(const FAvaMediaPlaybackInstance* InPlaybackInstance);
+
+protected:
+	FString ClientName;
+	FName ChannelName;
+	FGuid TransitionId;
+	bool bUnloadDiscardedInstances = false;
+	EAvalanchePlayableTransitionFlags TransitionFlags = EAvalanchePlayableTransitionFlags::None;
+	
+	TArray<FGuid> EnterInstanceIds;
+
+	TArray<TWeakPtr<FAvaMediaPlaybackInstance>> EnterPlaybackInstancesWeak;
+	TArray<TWeakPtr<FAvaMediaPlaybackInstance>> PlayingPlaybackInstancesWeak;
+	TArray<TWeakPtr<FAvaMediaPlaybackInstance>> ExitPlaybackInstancesWeak;
+	TArray<TSharedPtr<FAvalancheRemoteControlValues>> EnterValues;
+	
+	UPROPERTY(Transient)
+	TObjectPtr<UAvalanchePlayableTransition> PlayableTransition;
+};
\ No newline at end of file
diff --git a/Engine/Plugins/Experimental/Avalanche/Source/AvalancheMedia/Private/Playback/Transition/AvaPlayableTransitionPrivate.cpp b/Engine/Plugins/Experimental/Avalanche/Source/AvalancheMedia/Private/Playback/Transition/AvaPlayableTransitionPrivate.cpp
new file mode 100644
--- /dev/null
+++ b/Engine/Plugins/Experimental/Avalanche/Source/AvalancheMedia/Private/Playback/Transition/AvaPlayableTransitionPrivate.cpp
@@ -0,0 +1,81 @@
+// Copyright Epic Games, Inc. All Rights Reserved.
+
+#include "Playback/Transition/AvaPlayableTransitionPrivate.h"
+
+#include "Playback/AvalanchePlayable.h"
+#include "Playback/Playables/AvalancheRemoteProxyPlayable.h"
+
+namespace UE::AvaPlayableTransition::Private
+{
+	/**
+	 * Convert the array of weak objects to an array of objects than can be used (hasn't expired).
+	 */
+	TArray<UAvalanchePlayable*> Pin(const TArray<TWeakObjectPtr<UAvalanchePlayable>>& InPlayablesWeak)
+	{
+		TArray<UAvalanchePlayable*> Playables;
+		Playables.Reserve(InPlayablesWeak.Num());
+		for (const TWeakObjectPtr<UAvalanchePlayable>& PlayableWeak : InPlayablesWeak)
+		{
+			if (UAvalanchePlayable* Playable = PlayableWeak.Get())
+			{
+				Playables.Add(Playable);
+			}
+		}
+		return Playables;
+	}
+
+	TArray<FGuid> GetInstanceIds(const TArray<TWeakObjectPtr<UAvalanchePlayable>>& InPlayablesWeak)
+	{
+		TArray<FGuid> InstanceIds;
+		InstanceIds.Reserve(InPlayablesWeak.Num());
+		for (const TWeakObjectPtr<UAvalanchePlayable>& PlayableWeak : InPlayablesWeak)
+		{
+			if (const UAvalanchePlayable* Playable = PlayableWeak.Get())
+			{
+				InstanceIds.Add(Playable->GetInstanceId());
+			}
+		}
+		return InstanceIds;
+	}
+
+	bool AreAllPlayablesRemote(const TArray<TWeakObjectPtr<UAvalanchePlayable>>& InPlayablesWeak)
+	{
+		for (const TWeakObjectPtr<UAvalanchePlayable>& PlayableWeak : InPlayablesWeak)
+		{
+			if (const UAvalanchePlayable* Playable = PlayableWeak.Get())
+			{
+				if (!Playable->IsA<UAvalancheRemoteProxyPlayable>())
+				{
+					return false;
+				}
+			}
+		}
+		return true;
+	}
+
+	void GetChannelNamesFromPlayables(const TArray<TWeakObjectPtr<UAvalanchePlayable>>& InPlayablesWeak, TArray<FName>& OutChannelNames)
+	{
+		for (const TWeakObjectPtr<UAvalanchePlayable>& PlayableWeak : InPlayablesWeak)
+		{
+			if (const UAvalanchePlayable* Playable = PlayableWeak.Get())
+			{
+				if (const UAvalancheRemoteProxyPlayable* RemotePlayable = Cast<UAvalancheRemoteProxyPlayable>(Playable))
+				{
+					OutChannelNames.AddUnique(RemotePlayable->GetPlayingChannelFName());
+				}
+			}
+		}
+	}
+
+	FString GetPrettyPlayableInfo(const UAvalanchePlayable* InPlayable)
+	{
+		if (InPlayable)
+		{
+			return FString::Printf(TEXT("Id:%s, Asset:%s, Status:%s"),
+				*InPlayable->GetInstanceId().ToString(),
+				*InPlayable->GetSourceAssetPath().ToString(),
+				*StaticEnum<EAvalanchePlayableStatus>()->GetNameByValue(static_cast<int32>(InPlayable->GetPlayableStatus())).ToString());
+		}
+		return TEXT("(nullptr)");
+	}
+}
diff --git a/Engine/Plugins/Experimental/Avalanche/Source/AvalancheMedia/Private/Playback/Transition/AvaPlayableTransitionPrivate.h b/Engine/Plugins/Experimental/Avalanche/Source/AvalancheMedia/Private/Playback/Transition/AvaPlayableTransitionPrivate.h
new file mode 100644
--- /dev/null
+++ b/Engine/Plugins/Experimental/Avalanche/Source/AvalancheMedia/Private/Playback/Transition/AvaPlayableTransitionPrivate.h
@@ -0,0 +1,23 @@
+// Copyright Epic Games, Inc. All Rights Reserved.
+
+#pragma once
+
+#include "Playback/AvalanchePlayable.h"
+
+namespace UE::AvaPlayableTransition::Private
+{
+	/** Convert the array of weak objects to an array of objects than can be used (hasn't expired). */
+	TArray<UAvalanchePlayable*> Pin(const TArray<TWeakObjectPtr<UAvalanchePlayable>>& InPlayablesWeak);
+
+	/** Returns the array of playable instance ids. */
+	TArray<FGuid> GetInstanceIds(const TArray<TWeakObjectPtr<UAvalanchePlayable>>& InPlayablesWeak);
+
+	/** Returns true if all the given playables are remote. */
+	bool AreAllPlayablesRemote(const TArray<TWeakObjectPtr<UAvalanchePlayable>>& InPlayablesWeak);
+
+	/** Collects all the channel names from the given playables. Only remote playables have channel. */
+	void GetChannelNamesFromPlayables(const TArray<TWeakObjectPtr<UAvalanchePlayable>>& InPlayablesWeak, TArray<FName>& OutChannelNames);
+	
+	/** Returns a string with prettified information about the playable suitable for logging. */
+	FString GetPrettyPlayableInfo(const UAvalanchePlayable* InPlayable);
+}
\ No newline at end of file
diff --git a/Engine/Plugins/Experimental/Avalanche/Source/AvalancheMedia/Private/Playback/Transition/AvaRemotePlayableTransition.cpp b/Engine/Plugins/Experimental/Avalanche/Source/AvalancheMedia/Private/Playback/Transition/AvaRemotePlayableTransition.cpp
new file mode 100644
--- /dev/null
+++ b/Engine/Plugins/Experimental/Avalanche/Source/AvalancheMedia/Private/Playback/Transition/AvaRemotePlayableTransition.cpp
@@ -0,0 +1,118 @@
+// Copyright Epic Games, Inc. All Rights Reserved.
+
+#include "Playback/Transition/AvaRemotePlayableTransition.h"
+
+#include "IAvaMediaModule.h"
+#include "Playback/AvalanchePlayable.h"
+#include "Playback/AvaMediaPlaybackClientDelegates.h"
+#include "Playback/AvaMediaPlaybackUtils.h"
+#include "Playback/IAvaMediaPlaybackClient.h"
+#include "Playback/Transition/AvaPlayableTransitionPrivate.h"
+
+UAvaRemotePlayableTransition::~UAvaRemotePlayableTransition()
+{
+	UnregisterFromPlaybackClientDelegates();
+}
+
+bool UAvaRemotePlayableTransition::Start()
+{
+	IAvaMediaModule& AvaMediaModule = IAvaMediaModule::Get();
+	
+	if (!AvaMediaModule.IsMediaPlaybackClientStarted())
+	{
+		TransitionId.Invalidate();
+		return false;
+	}
+
+	if (!Super::Start())
+	{
+		return false;
+	}
+
+	TransitionId = FGuid::NewGuid();
+	
+	using namespace UE::AvaPlayableTransition::Private;
+	TArray<FGuid> EnterInstanceIds = GetInstanceIds(EnterPlayablesWeak);
+	TArray<FGuid> PlayingInstanceIds = GetInstanceIds(PlayingPlayablesWeak);
+	TArray<FGuid> ExitInstanceIds = GetInstanceIds(ExitPlayablesWeak);
+	TArray<FAvalancheRemoteControlValues> EnterValues;
+	EnterValues.Reserve(EnterPlayableValues.Num());
+	for (const TSharedPtr<FAvalancheRemoteControlValues>& Values : EnterPlayableValues)
+	{
+		EnterValues.Add(Values.IsValid() ? *Values : FAvalancheRemoteControlValues::GetDefaultEmpty());
+	}
+	
+	IAvaMediaPlaybackClient& PlaybackClient = AvaMediaModule.GetMediaPlaybackClient();
+	PlaybackClient.RequestPlayableTransitionStart(TransitionId, MoveTemp(EnterInstanceIds), MoveTemp(PlayingInstanceIds), MoveTemp(ExitInstanceIds), MoveTemp(EnterValues), ChannelName, TransitionFlags);
+	RegisterToPlaybackClientDelegates();	// to get the transition events from the server side.
+
+	using namespace UE::AvaMediaPlayback::Utils;
+	UE_LOG(LogAvalanchePlayable, Verbose, TEXT("%s Remote Playable Transition \"%s\" Id:\"%s\" starting."), *GetBriefFrameInfo(), *GetFullName(), *TransitionId.ToString());
+	return true;
+}
+
+void UAvaRemotePlayableTransition::Stop()
+{
+	if (TransitionId.IsValid())
+	{
+		IAvaMediaModule& AvaMediaModule = IAvaMediaModule::Get();
+	
+		if (AvaMediaModule.IsMediaPlaybackClientStarted())
+		{
+			AvaMediaModule.GetMediaPlaybackClient().RequestPlayableTransitionStop(TransitionId, ChannelName);
+		}
+	}
+
+	Super::Stop();
+}
+
+void UAvaRemotePlayableTransition::RegisterToPlaybackClientDelegates()
+{
+	UE::AvaMediaPlaybackClient::Delegates::GetOnPlaybackTransitionEvent().RemoveAll(this);
+	UE::AvaMediaPlaybackClient::Delegates::GetOnPlaybackTransitionEvent().AddUObject(this, &UAvaRemotePlayableTransition::HandlePlaybackTransitionEvent);
+}
+
+void UAvaRemotePlayableTransition::UnregisterFromPlaybackClientDelegates() const
+{
+	UE::AvaMediaPlaybackClient::Delegates::GetOnPlaybackTransitionEvent().RemoveAll(this);
+}
+
+bool UAvaRemotePlayableTransition::IsRunning() const
+{
+	return TransitionId.IsValid();	
+}
+
+void UAvaRemotePlayableTransition::HandlePlaybackTransitionEvent(IAvaMediaPlaybackClient& InPlaybackClient,
+	const UE::AvaMediaPlaybackClient::Delegates::FPlaybackTransitionEventArgs& InArgs)
+{
+	if (InArgs.TransitionId != TransitionId)
+	{
+		return;
+	}
+
+	if (InArgs.InstanceId.IsValid())
+	{
+		if (UAvalanchePlayable* Playable = FindPlayable(InArgs.InstanceId))
+		{
+			// Relay locally through playable event.
+			UAvalanchePlayable::OnTransitionEvent().Broadcast(Playable, this, InArgs.EventFlags);
+		}
+		else
+		{
+			UE_LOG(LogAvalanchePlayable, Error,
+				TEXT("Remote Playable Transition \"%s\" doesn't have playable instance Id \"%s\"."),
+				*TransitionId.ToString(), *InArgs.InstanceId.ToString());
+		}
+	}
+	else
+	{
+		UAvalanchePlayable::OnTransitionEvent().Broadcast(nullptr, this, InArgs.EventFlags);
+	}
+
+	if (EnumHasAnyFlags(InArgs.EventFlags, EAvalanchePlayableTransitionEventFlags::Finished))
+	{
+		using namespace UE::AvaMediaPlayback::Utils;
+		UE_LOG(LogAvalanchePlayable, Verbose, TEXT("%s Remote Playable Transition \"%s\" Id:\"%s\" ended."), *GetBriefFrameInfo(), *GetFullName(), *TransitionId.ToString());
+	}
+}
+
diff --git a/Engine/Plugins/Experimental/Avalanche/Source/AvalancheMedia/Private/Playback/Transition/AvaRemotePlayableTransition.h b/Engine/Plugins/Experimental/Avalanche/Source/AvalancheMedia/Private/Playback/Transition/AvaRemotePlayableTransition.h
new file mode 100644
--- /dev/null
+++ b/Engine/Plugins/Experimental/Avalanche/Source/AvalancheMedia/Private/Playback/Transition/AvaRemotePlayableTransition.h
@@ -0,0 +1,41 @@
+// Copyright Epic Games, Inc. All Rights Reserved.
+
+#pragma once
+
+#include "Playback/Transition/AvalanchePlayableTransition.h"
+
+#include "AvaRemotePlayableTransition.generated.h"
+
+class IAvaMediaPlaybackClient;
+namespace UE::AvaMediaPlaybackClient::Delegates
+{
+	struct FPlaybackTransitionEventArgs;
+}
+
+UCLASS()
+class UAvaRemotePlayableTransition : public UAvalanchePlayableTransition
+{
+	GENERATED_BODY()
+	
+public:
+	virtual ~UAvaRemotePlayableTransition() override;
+
+	void SetChannelName(const FName& InChannelName) { ChannelName = InChannelName; }
+	
+	//~ Begin UAvalanchePlayableTransition
+	virtual bool Start() override;
+	virtual void Stop() override;
+	virtual bool IsRunning() const override;
+	//~ End UAvalanchePlayableTransition
+
+protected:
+	void HandlePlaybackTransitionEvent(IAvaMediaPlaybackClient& InPlaybackClient,
+		const UE::AvaMediaPlaybackClient::Delegates::FPlaybackTransitionEventArgs& InArgs);
+
+	void RegisterToPlaybackClientDelegates();
+	void UnregisterFromPlaybackClientDelegates() const;
+
+protected:
+	FGuid TransitionId;
+	FName ChannelName;
+};
\ No newline at end of file
diff --git a/Engine/Plugins/Experimental/Avalanche/Source/AvalancheMedia/Private/Playback/Transition/AvaTransitionPlayableScene.cpp b/Engine/Plugins/Experimental/Avalanche/Source/AvalancheMedia/Private/Playback/Transition/AvaTransitionPlayableScene.cpp
new file mode 100644
--- /dev/null
+++ b/Engine/Plugins/Experimental/Avalanche/Source/AvalancheMedia/Private/Playback/Transition/AvaTransitionPlayableScene.cpp
@@ -0,0 +1,103 @@
+// Copyright Epic Games, Inc. All Rights Reserved.
+
+#include "Playback/Transition/AvaTransitionPlayableScene.h"
+#include "IAvaSceneInterface.h"
+#include "Playback/AvalanchePlayable.h"
+#include "Playback/Transition/AvalanchePlayableTransition.h"
+#include "Transition/Extensions/IAvaTransitionRCExtension.h"
+
+namespace UE::AvaMedia::Private
+{
+	/** Controllers aren't applied to the Preset, instead, this compares the latest remote control values for a given playable */
+	class FAvaRCTransitionPlayableExtension : public IAvaRCTransitionExtension
+	{
+		virtual EAvaTransitionComparisonResult CompareControllers(const FGuid& InControllerId, const FAvaTransitionScene& InMyScene, const FAvaTransitionScene& InOtherScene) const override
+		{
+			const UAvalanchePlayable* MyPlayable = InMyScene.GetDataView().GetPtr<UAvalanchePlayable>();
+			const UAvalanchePlayable* OtherPlayable = InOtherScene.GetDataView().GetPtr<UAvalanchePlayable>();
+			if (!MyPlayable || !OtherPlayable)
+			{
+				return EAvaTransitionComparisonResult::None;
+			}
+
+			const FAvalancheRemoteControlValue* MyValue = MyPlayable->GetLatestRemoteControlValues().ControllerValues.Find(InControllerId);
+			const FAvalancheRemoteControlValue* OtherValue = OtherPlayable->GetLatestRemoteControlValues().ControllerValues.Find(InControllerId);
+			if (!MyValue || !OtherValue)
+			{
+				return EAvaTransitionComparisonResult::None;
+			}
+
+			return MyValue->IsSameValueAs(*OtherValue)
+				? EAvaTransitionComparisonResult::Same
+				: EAvaTransitionComparisonResult::Different;
+		}
+	};
+}
+
+FAvaTransitionPlayableScene::FAvaTransitionPlayableScene(UAvalanchePlayable* InPlayable, UAvalanchePlayableTransition* InPlayableTransition)
+	: FAvaTransitionScene(InPlayable)
+	, PlayableTransitionWeak(InPlayableTransition)
+{
+	AddExtension<UE::AvaMedia::Private::FAvaRCTransitionPlayableExtension>();
+}
+
+FAvaTransitionPlayableScene::FAvaTransitionPlayableScene(const FAvaTagHandle& InTransitionLayer, UAvalanchePlayableTransition* InPlayableTransition)
+	: FAvaTransitionPlayableScene(nullptr, InPlayableTransition)
+{
+	OverrideTransitionLayer = InTransitionLayer;
+}
+
+EAvaTransitionComparisonResult FAvaTransitionPlayableScene::Compare(const FAvaTransitionScene& InOther) const
+{
+	const UAvalanchePlayable* MyPlayable    = GetDataView().GetPtr<UAvalanchePlayable>();
+	const UAvalanchePlayable* OtherPlayable = InOther.GetDataView().GetPtr<UAvalanchePlayable>();
+
+	if (!MyPlayable || !OtherPlayable)
+	{
+		return EAvaTransitionComparisonResult::None;
+	}
+
+	// Determine if Template is the same via the Package Name To Load (i.e. Source Level)
+	if (MyPlayable->GetSourceAssetPath() == OtherPlayable->GetSourceAssetPath())
+	{
+		return EAvaTransitionComparisonResult::Same;
+	}
+
+	return EAvaTransitionComparisonResult::Different;
+}
+
+ULevel* FAvaTransitionPlayableScene::GetLevel() const
+{
+	const UAvalanchePlayable* Playable = GetDataView().GetPtr<UAvalanchePlayable>();
+	const IAvaSceneInterface* SceneInterface = Playable ? Playable->GetSceneInterface() : nullptr;
+	return SceneInterface ? SceneInterface->GetSceneLevel() : nullptr;
+}
+
+void FAvaTransitionPlayableScene::GetOverrideTransitionLayer(FAvaTagHandle& OutTransitionLayer) const
+{
+	if (OverrideTransitionLayer.IsSet())
+	{
+		OutTransitionLayer = *OverrideTransitionLayer;
+	}
+}
+
+void FAvaTransitionPlayableScene::OnFlagsChanged()
+{
+	UAvalanchePlayable* Playable = GetDataView().GetMutablePtr<UAvalanchePlayable>();
+	UAvalanchePlayableTransition* PlayableTransition = PlayableTransitionWeak.Get();
+	if (!Playable || !PlayableTransition)
+	{
+		return;
+	}
+
+	// Event received when the playable can be discarded/recycled.
+	if (HasAnyFlags(EAvaTransitionSceneFlags::NeedsDiscard))
+	{
+		// Do some error checking.
+		if (PlayableTransition->IsEnterPlayable(Playable))
+		{
+			UE_LOG(LogAvalanchePlayable, Error, TEXT("Playable Transition \"%s\" Error: An \"enter\" playable is being discarded."), *PlayableTransition->GetFullName());
+		}
+		PlayableTransition->MarkPlayableAsDiscard(Playable);
+	}
+}
