[MotionDesign]

diff --git a/Engine/Plugins/Experimental/Avalanche/Source/AvalancheMedia/Private/Playback/Transition/AvalanchePlayableTransition.cpp b/Engine/Plugins/Experimental/Avalanche/Source/AvalancheMedia/Private/Playback/Transition/AvalanchePlayableTransition.cpp
new file mode 100644
--- /dev/null
+++ b/Engine/Plugins/Experimental/Avalanche/Source/AvalancheMedia/Private/Playback/Transition/AvalanchePlayableTransition.cpp
@@ -0,0 +1,308 @@
+// Copyright Epic Games, Inc. All Rights Reserved.
+
+#include "Playback/Transition/AvalanchePlayableTransition.h"
+
+#include "Algo/Accumulate.h"
+#include "Playback/AvalanchePlayable.h"
+#include "Playback/Transition/AvaLocalPlayableTransition.h"
+#include "Playback/Transition/AvaPlayableTransitionPrivate.h"
+#include "Playback/Transition/AvaRemotePlayableTransition.h"
+
+bool UAvalanchePlayableTransition::Start()
+{
+	using namespace UE::AvaPlayableTransition::Private;
+
+	// Accumulate all the playable groups this transition is part of.
+	PlayableGroupsWeak.Reset();
+	auto AccumulatePlayableGroups = [this](const TArray<TWeakObjectPtr<UAvalanchePlayable>>& InPlayablesWeak)
+	{
+		for (const TWeakObjectPtr<UAvalanchePlayable>& PlayableWeak : InPlayablesWeak)
+		{
+			if (const UAvalanchePlayable* Playable = PlayableWeak.Get())
+			{
+				PlayableGroupsWeak.Add(Playable->GetPlayableGroup());
+			}
+		}
+	};
+
+	AccumulatePlayableGroups(EnterPlayablesWeak);
+	AccumulatePlayableGroups(ExitPlayablesWeak);
+
+	for (const TWeakObjectPtr<UAvaMediaPlayableGroup>& PlayableGroupWeak : PlayableGroupsWeak)
+	{
+		if (UAvaMediaPlayableGroup* PlayableGroup = PlayableGroupWeak.Get())
+		{
+			PlayableGroup->RegisterPlayableTransition(this);
+		}
+	}
+	return true;
+}
+
+void UAvalanchePlayableTransition::Stop()
+{
+	for (const TWeakObjectPtr<UAvaMediaPlayableGroup>& PlayableGroupWeak : PlayableGroupsWeak)
+	{
+		if (UAvaMediaPlayableGroup* PlayableGroup = PlayableGroupWeak.Get())
+		{
+			PlayableGroup->UnregisterPlayableTransition(this);
+		}
+	}
+}
+
+void UAvalanchePlayableTransition::SetTransitionFlags(EAvalanchePlayableTransitionFlags InFlags)
+{
+	TransitionFlags = InFlags;
+}
+
+void UAvalanchePlayableTransition::SetEnterPlayables(TArray<TWeakObjectPtr<UAvalanchePlayable>>&& InPlayablesWeak)
+{
+	EnterPlayablesWeak = MoveTemp(InPlayablesWeak);
+}
+
+void UAvalanchePlayableTransition::SetPlayingPlayables(TArray<TWeakObjectPtr<UAvalanchePlayable>>&& InPlayablesWeak)
+{
+	PlayingPlayablesWeak = MoveTemp(InPlayablesWeak);
+}
+
+void UAvalanchePlayableTransition::SetExitPlayables(TArray<TWeakObjectPtr<UAvalanchePlayable>>&& InPlayablesWeak)
+{
+	ExitPlayablesWeak = MoveTemp(InPlayablesWeak);
+}
+
+bool UAvalanchePlayableTransition::IsEnterPlayable(UAvalanchePlayable* InPlayable) const
+{
+	return EnterPlayablesWeak.Contains(InPlayable);
+}
+
+bool UAvalanchePlayableTransition::IsPlayingPlayable(UAvalanchePlayable* InPlayable) const
+{
+	return PlayingPlayablesWeak.Contains(InPlayable);
+}
+
+bool UAvalanchePlayableTransition::IsExitPlayable(UAvalanchePlayable* InPlayable) const
+{
+	return ExitPlayablesWeak.Contains(InPlayable);
+}
+
+void UAvalanchePlayableTransition::SetEnterPlayableValues(TArray<TSharedPtr<FAvalancheRemoteControlValues>>&& InPlayableValues)
+{
+	EnterPlayableValues = MoveTemp(InPlayableValues);
+}
+
+void UAvalanchePlayableTransition::MarkPlayableAsDiscard(UAvalanchePlayable* InPlayable)
+{
+	DiscardPlayablesWeak.AddUnique(InPlayable);
+	UAvalanchePlayable::OnTransitionEvent().Broadcast(InPlayable, this, EAvalanchePlayableTransitionEventFlags::MarkPlayableDiscard);
+}
+
+FString UAvalanchePlayableTransition::GetPrettyInfo() const
+{
+	auto MakePrettyPlayableList = [](const TArray<TWeakObjectPtr<UAvalanchePlayable>>& InPlayablesWeak) -> FString
+	{
+		FString List;
+		for (const TWeakObjectPtr<UAvalanchePlayable>& PlayableWeak : InPlayablesWeak)
+		{
+			if (const UAvalanchePlayable* Playable = PlayableWeak.Get())
+			{
+				List += List.IsEmpty() ? TEXT("") : TEXT(", ");
+				List += Playable->GetUserData();
+			}
+		}
+		return List;
+	};
+
+	const FString InList = MakePrettyPlayableList(EnterPlayablesWeak);
+	const FString OutList = MakePrettyPlayableList(ExitPlayablesWeak);
+
+	if (InList.IsEmpty())
+	{
+		return FString::Printf(TEXT("{Out:{%s}}"), *OutList);
+	}
+
+	if (OutList.IsEmpty())
+	{
+		return FString::Printf(TEXT("{In:{%s}}"), *InList);
+	}
+
+	return FString::Printf(TEXT("{In:{%s}, Out:{%s}}"), *InList, *OutList);
+}
+
+UAvalanchePlayable* UAvalanchePlayableTransition::FindPlayable(const FGuid& InInstanceId) const
+{
+	auto IsPlayablePredicate = [&InInstanceId](const TWeakObjectPtr<UAvalanchePlayable>& InPlayableWeak) -> bool
+	{
+		if (const UAvalanchePlayable* Playable = InPlayableWeak.Get())
+		{
+			return Playable->GetInstanceId() == InInstanceId;		
+		}
+		return false;
+	};
+	
+	const TWeakObjectPtr<UAvalanchePlayable>* PlayableWeak = EnterPlayablesWeak.FindByPredicate(IsPlayablePredicate);
+	if (PlayableWeak && PlayableWeak->IsValid())
+	{
+		return PlayableWeak->Get();
+	}
+
+	PlayableWeak = PlayingPlayablesWeak.FindByPredicate(IsPlayablePredicate);
+	if (PlayableWeak && PlayableWeak->IsValid())
+	{
+		return PlayableWeak->Get();
+	}
+
+	PlayableWeak = ExitPlayablesWeak.FindByPredicate(IsPlayablePredicate);
+	if (PlayableWeak && PlayableWeak->IsValid())
+	{
+		return PlayableWeak->Get();
+	}
+	
+	return nullptr;
+}
+
+FAvaPlayableTransitionBuilder::FAvaPlayableTransitionBuilder() = default;
+
+void FAvaPlayableTransitionBuilder::AddEnterPlayableValues(const TSharedPtr<FAvalancheRemoteControlValues>& InValues)
+{
+	EnterPlayableValues.Add(InValues);
+}
+
+bool FAvaPlayableTransitionBuilder::AddEnterPlayable(UAvalanchePlayable* InPlayable)
+{
+	if (ExitPlayablesWeak.Contains(InPlayable))
+	{
+		using namespace UE::AvaPlayableTransition::Private;
+		UE_LOG(LogAvalanchePlayable, Error,
+			TEXT("Playable Transition setup error: Playable {%s} can't added as an \"enter\" playable because it is already in the \"exit\" list."),
+			*GetPrettyPlayableInfo(InPlayable));
+		return false;
+	}
+
+	if (PlayingPlayablesWeak.Contains(InPlayable))
+	{
+		using namespace UE::AvaPlayableTransition::Private;
+		UE_LOG(LogAvalanchePlayable, Error,
+			TEXT("Playable Transition setup error: Playable {%s} can't added as an \"enter\" playable because it is already in the \"playing\" list."),
+			*GetPrettyPlayableInfo(InPlayable));
+		return false;
+	}
+	
+	EnterPlayablesWeak.AddUnique(InPlayable);
+	return true;
+}
+
+bool FAvaPlayableTransitionBuilder::AddPlayingPlayable(UAvalanchePlayable* InPlayable)
+{
+	if (EnterPlayablesWeak.Contains(InPlayable))
+	{
+		using namespace UE::AvaPlayableTransition::Private;
+		UE_LOG(LogAvalanchePlayable, Error,
+			TEXT("Playable Transition setup error: Playable {%s} can't added as an \"exit\" playable because it is already in the \"enter\" list."),
+			*GetPrettyPlayableInfo(InPlayable));
+		return false;
+	}
+
+	if (ExitPlayablesWeak.Contains(InPlayable))
+	{
+		using namespace UE::AvaPlayableTransition::Private;
+		UE_LOG(LogAvalanchePlayable, Error,
+			TEXT("Playable Transition setup error: Playable {%s} can't added as an \"exit\" playable because it is already in the \"exit\" list."),
+			*GetPrettyPlayableInfo(InPlayable));
+		return false;
+	}
+
+	
+	PlayingPlayablesWeak.AddUnique(InPlayable);
+	return true;
+}
+
+bool FAvaPlayableTransitionBuilder::AddExitPlayable(UAvalanchePlayable* InPlayable)
+{
+	if (EnterPlayablesWeak.Contains(InPlayable))
+	{
+		using namespace UE::AvaPlayableTransition::Private;
+		UE_LOG(LogAvalanchePlayable, Error,
+			TEXT("Playable Transition setup error: Playable {%s} can't added as an \"exit\" playable because it is already in the \"enter\" list."),
+			*GetPrettyPlayableInfo(InPlayable));
+		return false;
+	}
+
+	if (PlayingPlayablesWeak.Contains(InPlayable))
+	{
+		using namespace UE::AvaPlayableTransition::Private;
+		UE_LOG(LogAvalanchePlayable, Error,
+			TEXT("Playable Transition setup error: Playable {%s} can't added as an \"enter\" playable because it is already in the \"playing\" list."),
+			*GetPrettyPlayableInfo(InPlayable));
+		return false;
+	}
+
+	ExitPlayablesWeak.AddUnique(InPlayable);
+	return true;
+}
+
+UAvalanchePlayableTransition* FAvaPlayableTransitionBuilder::MakeTransition(UObject* InOuter)
+{
+	// Determine if we have remote playables.
+	using namespace UE::AvaPlayableTransition::Private;
+
+	// A transition that has only "playing" playables does nothing.
+	if (EnterPlayablesWeak.IsEmpty() && ExitPlayablesWeak.IsEmpty())
+	{
+		return nullptr;
+	}
+
+	// A remote transition is created only if all playables are remote.
+	const bool bCreateRemoteTransition = AreAllPlayablesRemote(EnterPlayablesWeak)
+		&& AreAllPlayablesRemote(ExitPlayablesWeak)
+		&& AreAllPlayablesRemote(PlayingPlayablesWeak);
+	
+	UAvalanchePlayableTransition* PlayableTransition;
+	if (bCreateRemoteTransition)
+	{
+		UAvaRemotePlayableTransition* RemotePlayableTransition = NewObject<UAvaRemotePlayableTransition>(InOuter);
+		TArray<FName> ChannelNames;
+		ChannelNames.Reserve(1);
+		GetChannelNamesFromPlayables(EnterPlayablesWeak, ChannelNames);
+		GetChannelNamesFromPlayables(ExitPlayablesWeak, ChannelNames);
+		GetChannelNamesFromPlayables(PlayingPlayablesWeak, ChannelNames);
+		if (ChannelNames.Num() > 1)
+		{
+			const FString ChannelNameList = Algo::Accumulate(ChannelNames, FString(), [](FString InResult, const FName& InName)
+			{
+				InResult = InResult.IsEmpty() ? InName.ToString() : InResult + ", " + InName.ToString();
+				return InResult;
+			});
+		
+			UE_LOG(LogAvalanchePlayable, Warning,
+				TEXT("Playable Transition setup warning: Playables from different channels (%s) are in the same transition."),
+				*ChannelNameList);
+		}
+		
+		if (ChannelNames.Num() > 0)
+		{
+			RemotePlayableTransition->SetChannelName(ChannelNames[0]);
+		}
+		else
+		{
+			UE_LOG(LogAvalanchePlayable, Error,
+				TEXT("Playable Transition setup error: Couldn't get a channel name from the list of playables."));
+		}
+		
+		PlayableTransition = RemotePlayableTransition;
+	}
+	else
+	{
+		PlayableTransition = NewObject<UAvaLocalPlayableTransition>(InOuter);
+	}
+
+	// There must be the same number of entries then there are playables.
+	if (ensure(EnterPlayableValues.Num() == EnterPlayablesWeak.Num()))
+	{
+		PlayableTransition->SetEnterPlayableValues(MoveTemp(EnterPlayableValues));
+	}
+
+	PlayableTransition->SetEnterPlayables(MoveTemp(EnterPlayablesWeak));
+	PlayableTransition->SetPlayingPlayables(MoveTemp(PlayingPlayablesWeak));
+	PlayableTransition->SetExitPlayables(MoveTemp(ExitPlayablesWeak));
+	
+	return PlayableTransition;
+}
diff --git a/Engine/Plugins/Experimental/Avalanche/Source/AvalancheMedia/Private/Playlist/AvaPlaylistPageLoadingManager.cpp b/Engine/Plugins/Experimental/Avalanche/Source/AvalancheMedia/Private/Playlist/AvaPlaylistPageLoadingManager.cpp
new file mode 100644
--- /dev/null
+++ b/Engine/Plugins/Experimental/Avalanche/Source/AvalancheMedia/Private/Playlist/AvaPlaylistPageLoadingManager.cpp
@@ -0,0 +1,253 @@
+// Copyright Epic Games, Inc. All Rights Reserved.
+
+#include "Playlist/AvaPlaylistPageLoadingManager.h"
+
+#include "AvalancheBroadcast.h"
+#include "Playback/AvaMediaPlaybackManager.h"
+
+namespace UE::AvaPlaylistPageLoadingManager::Private
+{
+	bool IsLoaded(EAvaMediaPlaybackStatus InStatus)
+	{
+		switch (InStatus)
+		{
+		case EAvaMediaPlaybackStatus::Unknown:
+		case EAvaMediaPlaybackStatus::Missing:
+		case EAvaMediaPlaybackStatus::Syncing:
+		case EAvaMediaPlaybackStatus::Available:
+		case EAvaMediaPlaybackStatus::Loading:
+			return false;
+		case EAvaMediaPlaybackStatus::Loaded:
+		case EAvaMediaPlaybackStatus::Starting:
+		case EAvaMediaPlaybackStatus::Stopping:
+		case EAvaMediaPlaybackStatus::Unloading:
+			return true;
+		case EAvaMediaPlaybackStatus::Error:
+		default:
+			return false;
+		}
+	}
+
+	bool IsError(EAvaMediaPlaybackStatus InStatus)
+	{
+		return InStatus == EAvaMediaPlaybackStatus::Error;
+	}
+
+	bool IsLoading(EAvaMediaPlaybackStatus InStatus)
+	{
+		return InStatus == EAvaMediaPlaybackStatus::Loading;
+	}
+}
+
+FAvaPlaylistPageLoadingManager::FAvaPlaylistPageLoadingManager(UAvalanchePlaylist* InPlaylist)
+	: Playlist(InPlaylist)
+{
+	PendingRequestSet.Reserve(32);
+	LoadingInstances.Reserve(4);
+
+	if (Playlist)
+	{
+		Playlist->GetPlaybackManager().OnPlaybackInstanceStatusChanged.AddRaw(this, &FAvaPlaylistPageLoadingManager::HandlePlaybackInstanceStatusChanged);
+		Playlist->GetPlaybackManager().OnBeginTick.AddRaw(this, &FAvaPlaylistPageLoadingManager::HandlePlaybackManagerBeginTick);
+		PlaybackManagerWeak = Playlist->GetPlaybackManager().AsShared();
+	}
+}
+
+FAvaPlaylistPageLoadingManager::~FAvaPlaylistPageLoadingManager()
+{
+	if (const TSharedPtr<FAvaMediaPlaybackManager> PlaybackManager = PlaybackManagerWeak.Pin())
+	{
+		PlaybackManager->OnPlaybackInstanceStatusChanged.RemoveAll(this);
+		PlaybackManager->OnBeginTick.RemoveAll(this);
+	}
+}
+
+bool FAvaPlaylistPageLoadingManager::RequestLoadPage(int32 InPageId, bool bInIsPreview, const FName& InPreviewChannelName)
+{
+	using namespace UE::AvaPlaylistPageLoadingManager::Private;
+	
+	if (!Playlist)
+	{
+		return false;
+	}
+
+	const FAvalanchePage& Page = Playlist->GetPage(InPageId);
+	if (!Page.IsValidPage())
+	{
+		return false;
+	}
+	
+	if (LoadingInstances.Num() < MaxLoadingInstances)
+	{
+		return RequestLoadPageInternal(Page, bInIsPreview, InPreviewChannelName);
+	}
+
+	const FPageLoadRequest Request = {InPageId, bInIsPreview, InPreviewChannelName};
+
+	// Prevent the same request from being added more than once.
+	// It may be necessary to do that to prevent request spamming.
+	if (!PendingRequestSet.Contains(Request))
+	{
+		PendingRequestQueue.Enqueue(Request);
+		PendingRequestSet.Add(Request);
+	}
+	else
+	{
+		UE_LOG(LogAvaPlaylist, Warning,
+			TEXT("Page Loading Manager: Request for page id \"%d\" type: \"%s\" Channel \"%s\" is already in the queue."),
+			InPageId, bInIsPreview ? TEXT("Preview") : TEXT("Program"), *InPreviewChannelName.ToString());
+	}
+	return true;
+}
+
+bool FAvaPlaylistPageLoadingManager::RequestLoadPageInternal(const FAvalanchePage& InPage, bool bInIsPreview, const FName& InPreviewChannelName)
+{
+	using namespace UE::AvaPlaylistPageLoadingManager::Private;
+	
+	const TArray<UAvalanchePlaylist::FLoadedInstanceInfo> LoadedInstances = Playlist->LoadPage(InPage.GetPageId(), bInIsPreview, InPreviewChannelName);
+	
+	if (LoadedInstances.IsEmpty())
+	{
+		return false;
+	}
+	
+	for (const UAvalanchePlaylist::FLoadedInstanceInfo& LoadedInstance : LoadedInstances)
+	{
+		FInstanceInfo InstanceInfo;
+		InstanceInfo.InstanceId = LoadedInstance.InstanceId;
+		InstanceInfo.ChannelName = bInIsPreview ? InPreviewChannelName.ToString() : InPage.GetChannelName().ToString();
+		InstanceInfo.AssetPath = LoadedInstance.AssetPath;
+		
+		// Check if the instance is already loaded
+		const TSharedPtr<FAvaMediaPlaybackInstance> PlaybackInstance = FindInstance(InstanceInfo);
+
+		if (!PlaybackInstance || IsError(PlaybackInstance->GetStatus()))
+		{
+			continue;
+		}
+
+		// If it is not loaded yet, wait for it.
+		if (!IsLoaded(PlaybackInstance->GetStatus()))
+		{
+			// Keep track of the current status so we can better track what is going on.
+			InstanceInfo.Status = PlaybackInstance->GetStatus();
+			// Keep track of submit time so we can time it out if it takes too long.
+			InstanceInfo.SubmitTime = FDateTime::UtcNow();
+			LoadingInstances.Add(InstanceInfo);
+		}
+	}
+	
+	return !LoadedInstances.IsEmpty();
+}
+
+void FAvaPlaylistPageLoadingManager::HandlePlaybackInstanceStatusChanged(const FAvaMediaPlaybackInstance& InPlaybackInstance)
+{
+	using namespace UE::AvaPlaylistPageLoadingManager::Private;
+	for (TArray<FInstanceInfo>::TIterator InstanceIterator = LoadingInstances.CreateIterator(); InstanceIterator; ++InstanceIterator  )
+	{
+		FInstanceInfo& InstanceInfo = *InstanceIterator;
+		if (InstanceInfo.InstanceId == InPlaybackInstance.GetInstanceId())
+		{
+			if (IsError(InPlaybackInstance.GetStatus()) || IsLoaded(InPlaybackInstance.GetStatus()))
+			{
+				InstanceIterator.RemoveCurrent();
+			}
+		}
+	}
+}
+
+void FAvaPlaylistPageLoadingManager::HandlePlaybackManagerBeginTick(float InDeltaSeconds)
+{
+	using namespace UE::AvaPlaylistPageLoadingManager::Private;
+
+	const FDateTime CurrentTime = FDateTime::UtcNow();
+	const FTimespan TimeoutSpan = FTimespan::FromSeconds(2);
+
+	// Attempt to prune the waiting list.
+	for (TArray<FInstanceInfo>::TIterator InstanceIterator = LoadingInstances.CreateIterator(); InstanceIterator; ++InstanceIterator  )
+	{
+		FInstanceInfo& InstanceInfo = *InstanceIterator;
+		const TSharedPtr<FAvaMediaPlaybackInstance> PlaybackInstance = FindInstance(InstanceInfo);
+
+		if (!PlaybackInstance)
+		{
+			UE_LOG(LogAvaPlaylist, Error,
+				TEXT("Page Loading Manager: Failed to find playback instance Id \"%s\" Path \"%s\" Channel \"%s\"."),
+				*InstanceInfo.InstanceId.ToString(), *InstanceInfo.AssetPath.ToString(), *InstanceInfo.ChannelName);
+			
+			InstanceIterator.RemoveCurrent();
+			continue;
+		}
+
+		if (IsError(PlaybackInstance->GetStatus()))
+		{
+			UE_LOG(LogAvaPlaylist, Error,
+				TEXT("Page Loading Manager: Failed to load playback instance Id \"%s\" Path \"%s\" Channel \"%s\"."),
+				*InstanceInfo.InstanceId.ToString(), *InstanceInfo.AssetPath.ToString(), *InstanceInfo.ChannelName);
+			InstanceIterator.RemoveCurrent();
+			continue;
+		}
+
+		// Success condition
+		if (IsLoaded(PlaybackInstance->GetStatus()))
+		{
+			InstanceIterator.RemoveCurrent();
+			continue;
+		}
+
+		// Handle an edge case:
+		// for the remote server, the load command might have been lost. If the playback instance returns to a status
+		// of "Available", something went wrong and the asset loading has been interrupted. Server went offline or
+		// was externally reset.
+		if (IsLoading(InstanceInfo.Status) && !IsLoading(PlaybackInstance->GetStatus()))
+		{
+			UE_LOG(LogAvaPlaylist, Warning,
+				TEXT("Page Loading Manager: Playback instance Id \"%s\" Path \"%s\" Channel \"%s\" stopped loading."),
+					*InstanceInfo.InstanceId.ToString(), *InstanceInfo.AssetPath.ToString(), *InstanceInfo.ChannelName);
+			// TODO: special recovery? Might push the request again back in the queue and try again.
+			InstanceIterator.RemoveCurrent();
+			continue;
+		}
+
+		// Finally check for timeout.
+		// Remark: load time may be super long for things that need to import animated skeletal mesh (alembic) those will timeout for sure.
+		// If a load request times out, we just keep going and request the next page, it will finish eventually.
+		if (CurrentTime - InstanceInfo.SubmitTime > TimeoutSpan)
+		{
+			UE_LOG(LogAvaPlaylist, Warning,
+				TEXT("Page Loading Manager: Playback instance Id \"%s\" Path \"%s\" Channel \"%s\" loading request timed out."),
+					*InstanceInfo.InstanceId.ToString(), *InstanceInfo.AssetPath.ToString(), *InstanceInfo.ChannelName);
+			InstanceIterator.RemoveCurrent();
+		}
+
+		// Update status so we can detect deltas.
+		// This would normally be detected and propagated as status change event in playback manager.
+		InstanceInfo.Status = PlaybackInstance->GetStatus();
+	}
+
+	if (Playlist)
+	{
+		while (LoadingInstances.Num() < MaxLoadingInstances && !PendingRequestQueue.IsEmpty())
+		{
+			FPageLoadRequest Request;
+			if (PendingRequestQueue.Dequeue(Request))
+			{
+				PendingRequestSet.Remove(Request);
+				const FAvalanchePage& Page = Playlist->GetPage(Request.PageId);
+				if (Page.IsValidPage())
+				{
+					RequestLoadPageInternal(Page, Request.bIsPreview, Request.PreviewChannelName);
+				}
+			}
+		}
+	}
+}
+
+TSharedPtr<FAvaMediaPlaybackInstance> FAvaPlaylistPageLoadingManager::FindInstance(const FInstanceInfo& InInstanceInfo) const
+{
+	if (Playlist)
+	{
+		return Playlist->GetPlaybackManager().FindPlaybackInstance(InInstanceInfo.InstanceId, InInstanceInfo.AssetPath, InInstanceInfo.ChannelName);
+	}
+	return nullptr;
+}
\ No newline at end of file
diff --git a/Engine/Plugins/Experimental/Avalanche/Source/AvalancheMedia/Private/Playlist/AvaPlaylistPageLoadingManager.h b/Engine/Plugins/Experimental/Avalanche/Source/AvalancheMedia/Private/Playlist/AvaPlaylistPageLoadingManager.h
new file mode 100644
--- /dev/null
+++ b/Engine/Plugins/Experimental/Avalanche/Source/AvalancheMedia/Private/Playlist/AvaPlaylistPageLoadingManager.h
@@ -0,0 +1,77 @@
+// Copyright Epic Games, Inc. All Rights Reserved.
+
+#pragma once
+
+#include "Containers/Queue.h"
+#include "Playlist/AvalanchePlaylist.h"
+
+/**
+ * Manager for page loading requests.
+ *
+ * The main problem this class is solving right now is the throttling of page loading.
+ * We don't want to start loading all the pages at the same time because it leads to main thread
+ * bottleneck both when the requests are made and when the levels are loaded. To alleviate that
+ * we make the operations more granular by loading a subset of levels at any given time.
+ */
+class FAvaPlaylistPageLoadingManager final : public IAvaPlaylistPageLoadingManager
+{
+public:
+	explicit FAvaPlaylistPageLoadingManager(UAvalanchePlaylist* InPlaylist);
+	virtual ~FAvaPlaylistPageLoadingManager() override;
+
+	//~ Begin IAvaPlaylistPageLoadingManager
+	virtual bool RequestLoadPage(int32 InPageId, bool bInIsPreview, const FName& InPreviewChannelName) override;
+	//~ End IAvaPlaylistPageLoadingManager
+
+private:
+	bool RequestLoadPageInternal(const FAvalanchePage& InPage, bool bInIsPreview, const FName& InPreviewChannelName);
+	
+	void HandlePlaybackInstanceStatusChanged(const FAvaMediaPlaybackInstance& InPlaybackInstance);
+	void HandlePlaybackManagerBeginTick(float InDeltaSeconds);
+
+	struct FInstanceInfo;
+	TSharedPtr<FAvaMediaPlaybackInstance> FindInstance(const FInstanceInfo& InInstanceInfo) const;
+	
+private:
+	UAvalanchePlaylist* Playlist = nullptr;
+	TWeakPtr<FAvaMediaPlaybackManager> PlaybackManagerWeak;
+
+	struct FPageLoadRequest
+	{
+		int32 PageId;
+		bool bIsPreview;
+		FName PreviewChannelName;
+
+		bool operator==(FPageLoadRequest const& Other) const
+		{
+			return PageId == Other.PageId && bIsPreview == Other.bIsPreview && PreviewChannelName == Other.PreviewChannelName;
+		}
+
+		FORCEINLINE friend uint32 GetTypeHash(FPageLoadRequest const& This)
+		{
+			uint32 Hash = 0;
+			Hash = HashCombine(Hash, GetTypeHash(This.PageId));
+			Hash = HashCombine(Hash, GetTypeHash(This.PreviewChannelName));
+			Hash = HashCombine(Hash, GetTypeHash(This.bIsPreview));
+			return Hash;
+		}
+	};
+	/** Ordered queue. Requests are executed in order. */
+	TQueue<FPageLoadRequest> PendingRequestQueue;
+	/** Unordered set to ensure uniqueness of requests. */
+	TSet<FPageLoadRequest> PendingRequestSet;
+
+	/** Instances currently loading. */
+	struct FInstanceInfo
+	{
+		FGuid InstanceId;
+		FSoftObjectPath AssetPath;
+		FString ChannelName;
+		EAvaMediaPlaybackStatus Status = EAvaMediaPlaybackStatus::Unknown;
+		FDateTime SubmitTime;	// Time stamp when the load request was submitted.
+	};
+	TArray<FInstanceInfo> LoadingInstances;
+
+	/** Maximum concurrent instance to load. */
+	int32 MaxLoadingInstances = 1;
+};
diff --git a/Engine/Plugins/Experimental/Avalanche/Source/AvalancheMedia/Private/Playlist/AvaPlaylistPlaybackClientWatcher.cpp b/Engine/Plugins/Experimental/Avalanche/Source/AvalancheMedia/Private/Playlist/AvaPlaylistPlaybackClientWatcher.cpp
new file mode 100644
--- /dev/null
+++ b/Engine/Plugins/Experimental/Avalanche/Source/AvalancheMedia/Private/Playlist/AvaPlaylistPlaybackClientWatcher.cpp
@@ -0,0 +1,129 @@
+// Copyright Epic Games, Inc. All Rights Reserved.
+
+#include "Playlist/AvaPlaylistPlaybackClientWatcher.h"
+
+#include "AvalancheBroadcast.h"
+#include "Playback/IAvaMediaPlaybackClient.h"
+#include "Playlist/AvalanchePlaylist.h"
+
+FAvaPlaylistPlaybackClientWatcher::FAvaPlaylistPlaybackClientWatcher(UAvalanchePlaylist* InPlaylist)
+	: Playlist(InPlaylist)
+{
+	using namespace UE::AvaMediaPlaybackClient::Delegates;
+	GetOnPlaybackStatusChanged().AddRaw(this, &FAvaPlaylistPlaybackClientWatcher::HandlePlaybackStatusChanged);
+	
+}
+FAvaPlaylistPlaybackClientWatcher::~FAvaPlaylistPlaybackClientWatcher()
+{
+	using namespace UE::AvaMediaPlaybackClient::Delegates;
+	GetOnPlaybackStatusChanged().RemoveAll(this);
+}
+
+void FAvaPlaylistPlaybackClientWatcher::TryRestorePlaySubPage(int InPageId, const UE::AvaMediaPlaybackClient::Delegates::FPlaybackStatusChangedArgs& InEventArgs) const
+{
+	// Restore page player and local playback proxies.
+	// (Need to specify the InstanceId from the server for everything to match.)
+	const FName ChannelFName(InEventArgs.ChannelName);
+
+	// Ensure the specified channel exists.
+	if (UAvalancheBroadcast::Get().GetChannelIndex(ChannelFName) == INDEX_NONE)
+	{
+		UE_LOG(LogAvaPlaylist, Error,
+			TEXT("Received a playback object on channel \"%s\" which doesn't exist locally. Playback Server should be reset."),
+			*InEventArgs.ChannelName);
+		return;
+	}
+	
+	const bool bIsPreview = UAvalancheBroadcast::Get().GetChannelType(ChannelFName) == EAvaBroadcastChannelType::Preview ? true : false;
+
+	const FAvalanchePage& PageToRestore = Playlist->GetPage(InPageId);
+
+	const int32 SubPageIndex = PageToRestore.GetAvalancheAssetPaths(Playlist).Find(InEventArgs.AssetPath);
+
+	if (SubPageIndex == INDEX_NONE)
+	{
+		UE_LOG(LogAvaPlaylist, Error,
+			TEXT("Asset mismatch (expected (any of): \"%s\", received: \"%s\") for restoring page %d. Playback Server should be reset."),
+			*FString::JoinBy(PageToRestore.GetAvalancheAssetPaths(Playlist), TEXT(","), [](const FSoftObjectPath& Path){ return Path.ToString();}),
+			*InEventArgs.AssetPath.ToString(), InPageId);
+		return;
+	}
+
+	if (!Playlist->RestorePlaySubPage(InPageId, SubPageIndex, InEventArgs.InstanceId, bIsPreview, ChannelFName))
+	{
+		UE_LOG(LogAvaPlaylist, Error, TEXT("Failed to restore page %d. Playback Server should be reset."), InPageId);
+	}
+}
+
+void FAvaPlaylistPlaybackClientWatcher::HandlePlaybackStatusChanged(IAvaMediaPlaybackClient& InPlaybackClient,
+	const UE::AvaMediaPlaybackClient::Delegates::FPlaybackStatusChangedArgs& InEventArgs)
+{
+	if (!Playlist)
+	{
+		return;
+	}
+
+	static const TArray<EAvaMediaPlaybackStatus> RunningStates =
+	{
+		//EAvaMediaPlaybackStatus::Starting, // Starting is not reliable, it may also mean "loading". FIXME.
+		EAvaMediaPlaybackStatus::Started
+	};
+
+	// Try to determine if a playback has started or stopped.
+	const bool bWasRunning = IsAnyOf(InEventArgs.PrevStatus, RunningStates);
+	const bool bIsRunning = IsAnyOf(InEventArgs.NewStatus, RunningStates);
+	
+	// If a playback instance is stopping, stop corresponding page (if any).
+	if (bWasRunning && !bIsRunning)
+	{
+		for (UAvalanchePagePlayer* PagePlayer : Playlist->PagePlayers)
+		{
+			// Search for a match with the event:				
+			if (PagePlayer && PagePlayer->ChannelName == InEventArgs.ChannelName)
+			{
+				PagePlayer->ForEachInstancePlayer([&InEventArgs](UAvaRundownPlaybackInstancePlayer* InInstancePlayer)
+				{
+					if (InInstancePlayer
+						&& InInstancePlayer->SourceAssetPath == InEventArgs.AssetPath
+						&& InInstancePlayer->GetPlaybackInstanceId() != InEventArgs.InstanceId)
+					{
+						InInstancePlayer->Stop();	
+					}
+				});
+
+				// If we stopped all the instance players, stop the page (to broadcast events).
+				if (!PagePlayer->IsPlaying())
+				{
+					PagePlayer->Stop();	
+				}
+			}
+		}
+		Playlist->RemoveStoppedPagePlayers();
+	}
+
+	// Note: execute this even if not on rising transition because it may be a user data update following the "GetUserData" request.
+	if (bIsRunning)
+	{
+		// We need to figure out which page it is.
+		const FString* UserData = InPlaybackClient.GetRemotePlaybackUserData(InEventArgs.InstanceId, InEventArgs.AssetPath, InEventArgs.ChannelName);
+
+		// We haven't received the user data for this playback. So we request it.
+		// This event will be received again with user data next time.
+		if (!UserData)
+		{
+			InPlaybackClient.RequestPlayback(InEventArgs.InstanceId, InEventArgs.AssetPath, InEventArgs.ChannelName, EAvaMediaPlaybackAction::GetUserData);
+		}
+		else
+		{
+			const int32 PageId = UAvalanchePagePlayer::GetPageIdFromInstanceUserData(*UserData);
+			if (PageId != FAvalanchePage::InvalidPageId)
+			{
+				const UAvalanchePagePlayer* PagePlayer = Playlist->FindPlayerForProgramPage(PageId);
+				if (!PagePlayer || !PagePlayer->FindInstancePlayerByInstanceId(InEventArgs.InstanceId))
+				{
+					TryRestorePlaySubPage(PageId, InEventArgs);
+				}
+			}
+		}
+	}
+}
\ No newline at end of file
diff --git a/Engine/Plugins/Experimental/Avalanche/Source/AvalancheMedia/Private/Playlist/AvaPlaylistPlaybackClientWatcher.h b/Engine/Plugins/Experimental/Avalanche/Source/AvalancheMedia/Private/Playlist/AvaPlaylistPlaybackClientWatcher.h
new file mode 100644
--- /dev/null
+++ b/Engine/Plugins/Experimental/Avalanche/Source/AvalancheMedia/Private/Playlist/AvaPlaylistPlaybackClientWatcher.h
@@ -0,0 +1,43 @@
+// Copyright Epic Games, Inc. All Rights Reserved.
+
+#pragma once
+
+#include "AvaMediaDefines.h"
+#include "Playback/AvaMediaPlaybackClientDelegates.h"
+
+class IAvaMediaPlaybackClient;
+class UAvalanchePlaylist;
+
+/**
+ * Playback Client Watcher ensures external playback events are reconciled.
+ * Cases covered:
+ * 1 - if a remote playable is stopped (server reset), the corresponding page in the playlist will be stopped
+ * 2 - if a remote playable is playing (server side), the corresponding page's state will be restored, i.e. rebuilding the proxies.
+ * to be continued...
+ */
+class FAvaPlaylistPlaybackClientWatcher
+{
+public:
+	FAvaPlaylistPlaybackClientWatcher(UAvalanchePlaylist* InPlaylist);
+	~FAvaPlaylistPlaybackClientWatcher();
+
+	void TryRestorePlaySubPage(int InPageId, const UE::AvaMediaPlaybackClient::Delegates::FPlaybackStatusChangedArgs& InEventArgs) const;
+
+	static bool IsAnyOf(EAvaMediaPlaybackStatus InStatus, const TArray<EAvaMediaPlaybackStatus>& InStatuses)
+	{
+		for (const EAvaMediaPlaybackStatus Status : InStatuses)
+		{
+			if (Status == InStatus)
+			{
+				return true;
+			}
+		}
+		return false;
+	}
+	
+	void HandlePlaybackStatusChanged(IAvaMediaPlaybackClient& InPlaybackClient,
+		const UE::AvaMediaPlaybackClient::Delegates::FPlaybackStatusChangedArgs& InEventArgs);
+	
+private:
+	UAvalanchePlaylist* Playlist;
+};
diff --git a/Engine/Plugins/Experimental/Avalanche/Source/AvalancheMedia/Private/Playlist/AvalancheManagedInstance.cpp b/Engine/Plugins/Experimental/Avalanche/Source/AvalancheMedia/Private/Playlist/AvalancheManagedInstance.cpp
new file mode 100644
--- /dev/null
+++ b/Engine/Plugins/Experimental/Avalanche/Source/AvalancheMedia/Private/Playlist/AvalancheManagedInstance.cpp
@@ -0,0 +1,137 @@
+// Copyright Epic Games, Inc. All Rights Reserved.
+
+#include "Playlist/AvalancheManagedInstance.h"
+
+#include "Playlist/AvalancheManagedInstanceCache.h"
+#include "RemoteControlPreset.h"
+
+FAvalancheManagedInstance::FAvalancheManagedInstance(FAvalancheManagedInstanceCache* InParentCache, const FSoftObjectPath& InAssetPath)
+	: ParentCache(InParentCache)
+	, SourceAssetPath(InAssetPath)
+{}
+
+FAvalancheManagedInstance::~FAvalancheManagedInstance()
+{}
+
+void FAvalancheManagedInstance::InvalidateFromCache()
+{
+	if (ParentCache)
+	{
+		// The managed blueprint is not deleted immediately to avoid issues
+		// with destroying delegates while they are being called.
+		ParentCache->InvalidateNoDelete(SourceAssetPath);
+		ParentCache = nullptr;
+	}
+}
+
+void FAvalancheManagedInstance::RegisterSourceRemoteControlPresetDelegates(URemoteControlPreset* InSourceRemoteControlPreset)
+{
+	if (!::IsValid(InSourceRemoteControlPreset))
+	{
+		return;
+	}
+	
+	InSourceRemoteControlPreset->OnEntityExposed().AddRaw(
+		this, &FAvalancheManagedInstance::OnSourceRemoteControlEntitiesExposed);
+	InSourceRemoteControlPreset->OnEntityUnexposed().AddRaw(
+		this, &FAvalancheManagedInstance::OnSourceRemoteControlEntitiesUnexposed);
+	InSourceRemoteControlPreset->OnEntitiesUpdated().AddRaw(
+		this, &FAvalancheManagedInstance::OnSourceRemoteControlEntitiesUpdated);
+	InSourceRemoteControlPreset->OnExposedPropertiesModified().AddRaw(
+		this, &FAvalancheManagedInstance::OnSourceRemoteControlExposedPropertiesModified);
+	InSourceRemoteControlPreset->OnControllerAdded().AddRaw(
+		this, &FAvalancheManagedInstance::OnSourceRemoteControlControllerAdded);
+	InSourceRemoteControlPreset->OnControllerRemoved().AddRaw(
+		this, &FAvalancheManagedInstance::OnSourceRemoteControlControllerRemoved);
+	InSourceRemoteControlPreset->OnControllerRenamed().AddRaw(
+		this, &FAvalancheManagedInstance::OnSourceRemoteControlControllerRenamed);
+	InSourceRemoteControlPreset->OnControllerModified().AddRaw(
+		this, &FAvalancheManagedInstance::OnSourceRemoteControlControllerModified);
+}
+
+void FAvalancheManagedInstance::UnregisterSourceRemoteControlPresetDelegates(URemoteControlPreset* InSourceRemoteControlPreset) const
+{
+	if (!::IsValid(InSourceRemoteControlPreset))
+	{
+		return;
+	}
+	
+	InSourceRemoteControlPreset->OnEntityExposed().RemoveAll(this);
+	InSourceRemoteControlPreset->OnEntityUnexposed().RemoveAll(this);
+	InSourceRemoteControlPreset->OnEntitiesUpdated().RemoveAll(this);
+	InSourceRemoteControlPreset->OnExposedPropertiesModified().RemoveAll(this);
+	InSourceRemoteControlPreset->OnControllerAdded().RemoveAll(this);
+	InSourceRemoteControlPreset->OnControllerRemoved().RemoveAll(this);
+	InSourceRemoteControlPreset->OnControllerRenamed().RemoveAll(this);
+	InSourceRemoteControlPreset->OnControllerModified().RemoveAll(this);
+}
+
+void FAvalancheManagedInstance::OnSourceRemoteControlEntitiesExposed(URemoteControlPreset* InSourcePreset, const FGuid& InEntityId)
+{
+	InvalidateFromCache();
+}
+
+void FAvalancheManagedInstance::OnSourceRemoteControlEntitiesUnexposed(URemoteControlPreset* InSourcePreset, const FGuid& InEntityId)
+{
+	InvalidateFromCache();
+}
+
+void FAvalancheManagedInstance::OnSourceRemoteControlEntitiesUpdated(URemoteControlPreset* InSourcePreset, const TSet<FGuid>& InModifiedEntities)
+{
+	// Info: Delegate called when the exposed entity wrapper itself is updated (ie. binding change, rename).
+	// Not sure if we should invalidate the cache for this one.
+}
+
+void FAvalancheManagedInstance::OnSourceRemoteControlExposedPropertiesModified(URemoteControlPreset* InSourcePreset, const TSet<FGuid>& InModifiedProperties)
+{
+	// This will be called if values are modified from the Remote Control tab in the Ava BP editor.
+	// Experimental: Propagate the default value immediately if it exists.
+	// We don't want to apply those to the managed remote control because we don't want to overwrite the page's values.
+	// We probably need an extra event to propagate a change in default value so pages with no values for that property
+	// can be updated.
+
+	// There is no event called when entities are unbound to actions, so rescan all actions just in case.
+	DefaultRemoteControlValues.RefreshControlledEntities(InSourcePreset);
+
+	for (const FGuid& PropertyId : InModifiedProperties)
+	{
+		// Update existing values only.
+		// Remark: added values will go through OnSourceRemoteControlControllerAdded instead.
+		if (DefaultRemoteControlValues.HasEntityValue(PropertyId))
+		{
+			DefaultRemoteControlValues.SetEntityValue(PropertyId, InSourcePreset, true);
+		}
+	}
+}
+
+void FAvalancheManagedInstance::OnSourceRemoteControlControllerAdded(URemoteControlPreset* InSourcePreset, const FName NewControllerName, const FGuid& InControllerId)
+{
+	InvalidateFromCache();
+}
+
+void FAvalancheManagedInstance::OnSourceRemoteControlControllerRemoved(URemoteControlPreset* InSourcePreset, const FGuid& InControllerId)
+{
+	InvalidateFromCache();
+}
+
+void FAvalancheManagedInstance::OnSourceRemoteControlControllerRenamed(URemoteControlPreset* InSourcePreset, const FName InOldLabel, const FName InNewLabel)
+{
+	InvalidateFromCache();
+}
+
+void FAvalancheManagedInstance::OnSourceRemoteControlControllerModified(URemoteControlPreset* InSourcePreset, const TSet<FGuid>& InModifiedControllerIds)
+{
+	// This will be called if values are modified from the Remote Control tab in the Ava BP editor.
+	// Experimental: Propagate the default value immediately if it exists.
+	// We don't want to apply those to the managed remote control because we don't want to overwrite the page's values.
+	// We probably need an extra event to propagate a change in default value so pages with no values for that property
+	// can be updated.
+	for (const FGuid& ControllerId : InModifiedControllerIds)
+	{
+		// Update existing values only.
+		if (DefaultRemoteControlValues.HasControllerValue(ControllerId))
+		{
+			DefaultRemoteControlValues.SetControllerValue(ControllerId, InSourcePreset, true);
+		}
+	}
+}
\ No newline at end of file
diff --git a/Engine/Plugins/Experimental/Avalanche/Source/AvalancheMedia/Private/Playlist/AvalancheManagedInstanceBlueprint.cpp b/Engine/Plugins/Experimental/Avalanche/Source/AvalancheMedia/Private/Playlist/AvalancheManagedInstanceBlueprint.cpp
new file mode 100644
--- /dev/null
+++ b/Engine/Plugins/Experimental/Avalanche/Source/AvalancheMedia/Private/Playlist/AvalancheManagedInstanceBlueprint.cpp
@@ -0,0 +1,143 @@
+// Copyright Epic Games, Inc. All Rights Reserved.
+
+#include "Playlist/AvalancheManagedInstanceBlueprint.h"
+
+#include "AssetRegistry/AssetData.h"
+#include "AvaBlueprint.h"
+#include "AvaRemoteControlRebind.h"
+#include "AvaRemoteControlUtils.h"
+#include "Engine/Engine.h"
+#include "Engine/World.h"
+#include "IAvaMediaModule.h"
+#include "Playlist/AvalancheManagedInstanceUtils.h"
+#include "RemoteControlPreset.h"
+#include "UObject/Package.h"
+
+namespace UE::AvalancheManagedInstanceBlueprint::Private
+{
+	UWorld* CreateManagedInstanceWorld(UPackage* InOuter, EWorldType::Type InWorldType)
+	{
+		const FName AvalancheWorldName = MakeUniqueObjectName(InOuter, UWorld::StaticClass(), TEXT("AvalancheManagedInstanceWorld"));
+		
+
+		// Note: this will initialize the world.
+		UWorld* World = UWorld::CreateWorld(InWorldType, false, AvalancheWorldName, InOuter, false);
+
+		if (InOuter == nullptr)
+		{
+			UPackage* const Package = World->GetPackage();
+			Package->SetFlags(RF_Transient | RF_Public);
+		}
+
+		FWorldContext& WorldContext = GEngine->CreateNewWorldContext(InWorldType);
+		WorldContext.SetCurrentWorld(World);
+		return World;
+	}
+	
+	UAvalancheBlueprint* LoadAvalancheBlueprint(const FSoftObjectPath& InAssetPath)
+	{
+		return Cast<UAvalancheBlueprint>(InAssetPath.TryLoad());
+	}
+}
+
+FAvalancheManagedInstanceBlueprint::FAvalancheManagedInstanceBlueprint(FAvalancheManagedInstanceCache* InParentCache, const FSoftObjectPath& InAssetPath)
+	: FAvalancheManagedInstance(InParentCache, InAssetPath)
+{
+	using namespace UE::AvalancheManagedInstanceBlueprint::Private;
+	SourceAvalancheBlueprint = LoadAvalancheBlueprint(InAssetPath);
+	if (!SourceAvalancheBlueprint)
+	{
+		UE_LOG(LogAvaMedia, Error, TEXT("Unable to load Source Avalanche Blueprint: %s"), *InAssetPath.ToString());
+		return;
+	}
+	
+	RegisterSourceRemoteControlPresetDelegates(SourceAvalancheBlueprint->GetRemoteControlPreset());
+
+#if WITH_EDITORONLY_DATA
+	FGuardValue_Bitfield(SourceAvalancheBlueprint->bDuplicatingReadOnly, true);
+#endif
+
+	ManagedAvalancheBlueprintPackage = FAvalancheManagedInstanceUtils::MakeManagedInstancePackage(InAssetPath);
+	if (!ManagedAvalancheBlueprintPackage)
+	{
+		UE_LOG(LogAvaMedia, Error, TEXT("Unable to create a Managed Avalanche Blueprint Package for %s"), *InAssetPath.ToString());
+		return;
+	}
+
+	// Remark: using PIE duplication mode to avoid entity and controller Ids from being renewed.
+	ManagedAvalancheBlueprint = Cast<UAvalancheBlueprint>(StaticDuplicateObject(SourceAvalancheBlueprint.Get(), ManagedAvalancheBlueprintPackage.Get(), NAME_None, RF_NoFlags, nullptr, EDuplicateMode::PIE));
+	if (!ManagedAvalancheBlueprint)
+	{
+		UE_LOG(LogAvaMedia, Error, TEXT("Unable to duplicate Source Avalanche Blueprint: %s"), *InAssetPath.ToString());
+		return;
+	}
+
+	ManagedAvalancheBlueprint->SetFlags(RF_Public | RF_Transient);
+
+	ManagedWorld = CreateManagedInstanceWorld(ManagedAvalancheBlueprintPackage.Get(), EWorldType::Type::Game);
+	ManagedAvalancheBlueprint->SetAvalancheWorld(ManagedWorld.Get());
+	ManagedAvalancheBlueprint->LoadAvalancheWorld();
+	FAvaRemoteControlRebind::RebindUnboundEntities(ManagedAvalancheBlueprint->GetRemoteControlPreset());
+
+	// Backup the remote control values from the source.
+	constexpr bool bIsDefault = true;	// Flag the values as "default".
+	DefaultRemoteControlValues.CopyFrom(ManagedAvalancheBlueprint->GetRemoteControlPreset(), bIsDefault);
+
+	FAvalancheManagedInstanceUtils::PreventWorldFromBeingSeenAsLeakingByLevelEditor(ManagedWorld.Get());
+}
+
+FAvalancheManagedInstanceBlueprint::~FAvalancheManagedInstanceBlueprint()
+{
+	if (ManagedAvalancheBlueprintPackage)
+	{
+		ManagedAvalancheBlueprintPackage->ClearDirtyFlag();
+	}
+
+	// Simplified version of UAvalancheGameInstance::EndPlay()
+	if (ManagedWorld)
+	{
+		GEngine->DestroyWorldContext(ManagedWorld.Get());
+		ManagedWorld->DestroyWorld(true);
+	}
+
+	ManagedWorld = nullptr;
+	if (ManagedAvalancheBlueprint && ManagedAvalancheBlueprint->GetRemoteControlPreset())
+	{
+		FAvaRemoteControlUtils::UnregisterRemoteControlPreset(ManagedAvalancheBlueprint->GetRemoteControlPreset());
+	}
+	ManagedAvalancheBlueprint = nullptr;
+
+	if (SourceAvalancheBlueprint)
+	{
+		UnregisterSourceRemoteControlPresetDelegates(SourceAvalancheBlueprint->GetRemoteControlPreset());
+	}
+	SourceAvalancheBlueprint = nullptr;
+}
+
+void FAvalancheManagedInstanceBlueprint::AddReferencedObjects( FReferenceCollector& Collector )
+{
+	Collector.AddReferencedObject( SourceAvalancheBlueprint );
+	Collector.AddReferencedObject( ManagedWorld );
+	Collector.AddReferencedObject( ManagedAvalancheBlueprint );
+	Collector.AddReferencedObject( ManagedAvalancheBlueprintPackage );
+}
+
+FString FAvalancheManagedInstanceBlueprint::GetReferencerName() const
+{
+	return TEXT("FAvalancheManagedInstanceBlueprint");
+}
+
+IAvaSceneInterface* FAvalancheManagedInstanceBlueprint::GetSceneInterface() const
+{
+	return ManagedAvalancheBlueprint ? static_cast<IAvaSceneInterface*>(ManagedAvalancheBlueprint.Get()) : nullptr;
+}
+
+bool FAvalancheManagedInstanceBlueprint::IsValid() const
+{
+	return ::IsValid(ManagedAvalancheBlueprint) && ::IsValid(ManagedWorld);
+}
+
+URemoteControlPreset* FAvalancheManagedInstanceBlueprint::GetRemoteControlPreset() const
+{
+	return IsValid() ? ManagedAvalancheBlueprint->GetRemoteControlPreset() : nullptr;
+}
\ No newline at end of file
diff --git a/Engine/Plugins/Experimental/Avalanche/Source/AvalancheMedia/Private/Playlist/AvalancheManagedInstanceBlueprint.h b/Engine/Plugins/Experimental/Avalanche/Source/AvalancheMedia/Private/Playlist/AvalancheManagedInstanceBlueprint.h
new file mode 100644
--- /dev/null
+++ b/Engine/Plugins/Experimental/Avalanche/Source/AvalancheMedia/Private/Playlist/AvalancheManagedInstanceBlueprint.h
@@ -0,0 +1,39 @@
+// Copyright Epic Games, Inc. All Rights Reserved.
+
+#pragma once
+
+#include "Playlist/AvalancheManagedInstance.h"
+#include "UObject/StrongObjectPtr.h"
+
+class UAvalancheBlueprint;
+
+/**
+ * Container for a fully instanced Avalanche Blueprint.
+ * The World is not fully initialized, but contains all loaded actors.
+ * This is required for RemoteControl bindings to be resolved.
+ */
+class FAvalancheManagedInstanceBlueprint : public FAvalancheManagedInstance
+{
+public:
+	UE_AVA_INHERITS(FAvalancheManagedInstanceBlueprint, FAvalancheManagedInstance);
+	
+	FAvalancheManagedInstanceBlueprint(FAvalancheManagedInstanceCache* InParentCache, const FSoftObjectPath& InAssetPath);
+	virtual ~FAvalancheManagedInstanceBlueprint() override;
+
+	//~ Begin FGCObject
+	virtual void AddReferencedObjects( FReferenceCollector& Collector ) override;
+	virtual FString GetReferencerName() const override;
+	virtual IAvaSceneInterface* GetSceneInterface() const override;
+	//~ End FGCObject
+
+	//~ Begin FAvalancheManagedInstance
+	virtual bool IsValid() const override;
+	virtual URemoteControlPreset* GetRemoteControlPreset() const override;
+	//~ End FAvalancheManagedInstance
+
+private:
+	TObjectPtr<UAvalancheBlueprint> SourceAvalancheBlueprint;
+	TObjectPtr<UWorld> ManagedWorld;
+	TObjectPtr<UAvalancheBlueprint> ManagedAvalancheBlueprint;
+	TObjectPtr<UPackage> ManagedAvalancheBlueprintPackage;
+};
diff --git a/Engine/Plugins/Experimental/Avalanche/Source/AvalancheMedia/Private/Playlist/AvalancheManagedInstanceCache.cpp b/Engine/Plugins/Experimental/Avalanche/Source/AvalancheMedia/Private/Playlist/AvalancheManagedInstanceCache.cpp
new file mode 100644
--- /dev/null
+++ b/Engine/Plugins/Experimental/Avalanche/Source/AvalancheMedia/Private/Playlist/AvalancheManagedInstanceCache.cpp
@@ -0,0 +1,234 @@
+// Copyright Epic Games, Inc. All Rights Reserved.
+
+#include "Playlist/AvalancheManagedInstanceCache.h"
+
+#include "AssetRegistry/AssetData.h"
+#include "AssetRegistry/IAssetRegistry.h"
+#include "AvalancheMediaSettings.h"
+#include "Engine/Engine.h"
+#include "Engine/World.h"
+#include "IAvaMediaModule.h"
+#include "Playback/AvaMediaPlaybackUtils.h"
+#include "Playlist/AvalancheManagedInstanceBlueprint.h"
+#include "Playlist/AvalancheManagedInstanceLevel.h"
+#include "UObject/Package.h"
+
+FAvalancheManagedInstanceCache::FAvalancheManagedInstanceCache()
+{
+	UPackage::PackageSavedWithContextEvent.AddRaw(this, &FAvalancheManagedInstanceCache::OnPackageSaved);
+	if (IAssetRegistry* AssetRegistry = IAssetRegistry::Get())
+	{
+		AssetRegistry->OnAssetRemoved().AddRaw(this, &FAvalancheManagedInstanceCache::OnAssetRemoved);
+	}
+	IAvaMediaModule::Get().GetOnAvaMediaSyncPackageModified().AddRaw(this, &FAvalancheManagedInstanceCache::OnAvaSyncPackageModified);
+	
+#if WITH_EDITOR
+	UAvalancheMediaSettings::GetMutable().OnSettingChanged().AddRaw(this, &FAvalancheManagedInstanceCache::OnSettingChanged);	
+#endif
+}
+
+FAvalancheManagedInstanceCache::~FAvalancheManagedInstanceCache()
+{
+	UPackage::PackageSavedWithContextEvent.RemoveAll(this);
+	if (IAssetRegistry* AssetRegistry = IAssetRegistry::Get())
+	{
+		AssetRegistry->OnAssetRemoved().RemoveAll(this);
+	}
+	IAvaMediaModule::Get().GetOnAvaMediaSyncPackageModified().RemoveAll(this);
+	
+#if WITH_EDITOR
+	UAvalancheMediaSettings::GetMutable().OnSettingChanged().RemoveAll(this);
+#endif
+}
+
+TSharedPtr<FAvalancheManagedInstance> FAvalancheManagedInstanceCache::GetOrLoadAvalancheInstance(const FSoftObjectPath& InAssetPath)
+{
+	FinishPendingActions();
+
+	if (InAssetPath.IsNull())
+	{
+		return nullptr;
+	}
+	
+	OrderQueue.Remove(InAssetPath); // Removes preserves the order. O(n)
+	OrderQueue.Add(InAssetPath); // Most recent is at the end of the array.
+
+	if (const TSharedPtr<FAvalancheManagedInstance>* ExistingEntry = AvalancheInstances.Find(InAssetPath))
+	{
+		if (ExistingEntry->IsValid())
+		{
+			return *ExistingEntry;
+		}
+	}
+
+	TSharedPtr<FAvalancheManagedInstance> NewEntry;
+
+	const FString PackageName = InAssetPath.GetLongPackageName();
+	if (FAvaMediaPlaybackUtils::IsMapAsset(PackageName))
+	{
+		NewEntry = MakeShared<FAvalancheManagedInstanceLevel>(this, InAssetPath);
+	}
+	else
+	{
+		NewEntry = MakeShared<FAvalancheManagedInstanceBlueprint>(this, InAssetPath);
+	}
+
+	if (NewEntry.IsValid() && NewEntry->IsValid())
+	{
+		AvalancheInstances.Add(InAssetPath, NewEntry);
+		TrimCache();
+	}
+
+	return NewEntry;
+}
+
+void FAvalancheManagedInstanceCache::InvalidateNoDelete(const FSoftObjectPath& InAssetPath)
+{
+	if (AvalancheInstances.Contains(InAssetPath) && !PendingInvalidatedPaths.Contains(InAssetPath))
+	{
+		PendingInvalidatedPaths.Add(InAssetPath);
+		OnEntryInvalidated.Broadcast(InAssetPath);
+	}
+}
+
+void FAvalancheManagedInstanceCache::Invalidate(const FSoftObjectPath& InAssetPath)
+{
+	// Ensure no pending actions. For instance, current path could already be pending
+	// and we don't want the events to be fired multiple time.
+	FinishPendingActions();
+
+	if (AvalancheInstances.Contains(InAssetPath))
+	{
+		RemoveEntry(InAssetPath);
+		OnEntryInvalidated.Broadcast(InAssetPath);
+	}
+}
+
+int32 FAvalancheManagedInstanceCache::GetMaximumCacheSize() const
+{
+	return UAvalancheMediaSettings::Get().ManagedAvalancheInstanceCacheMaximumSize;
+}
+
+void FAvalancheManagedInstanceCache::Flush(const FSoftObjectPath& InAssetPath)
+{
+	if (const TSharedPtr<FAvalancheManagedInstance>* ExistingEntry = AvalancheInstances.Find(InAssetPath))
+	{
+		if (ExistingEntry->GetSharedReferenceCount() <= 1)
+		{
+			RemoveEntry(InAssetPath);
+		}
+	}
+}
+
+void FAvalancheManagedInstanceCache::Flush()
+{
+	RemoveEntries([](const FSoftObjectPath& InAssetPath, const TSharedPtr<FAvalancheManagedInstance>& InManagedAvalancheBlueprint)
+	{
+		return (InManagedAvalancheBlueprint.GetSharedReferenceCount() <= 1) ? true : false;
+	}, false);
+}
+
+void FAvalancheManagedInstanceCache::TrimCache()
+{
+	const int32 MaximumCacheSize = GetMaximumCacheSize(); 
+	if (MaximumCacheSize > 0)
+	{
+		while (OrderQueue.Num() > MaximumCacheSize)
+		{
+			// LRU: oldest is at the start of the array.
+			AvalancheInstances.Remove(OrderQueue[0]);
+			OrderQueue.RemoveAt(0);
+		}
+	}
+}
+
+void FAvalancheManagedInstanceCache::FinishPendingActions()
+{
+	RemovePendingInvalidatedPaths();
+}
+
+void FAvalancheManagedInstanceCache::RemovePendingInvalidatedPaths()
+{
+	for (const FSoftObjectPath& InvalidatedPath : PendingInvalidatedPaths)
+	{
+		RemoveEntry(InvalidatedPath);
+	}
+	PendingInvalidatedPaths.Reset();
+}
+
+void FAvalancheManagedInstanceCache::RemoveEntry(const FSoftObjectPath& InAssetPath)
+{
+	OrderQueue.Remove(InAssetPath);
+	AvalancheInstances.Remove(InAssetPath);
+}
+
+void FAvalancheManagedInstanceCache::RemoveEntries(TFunctionRef<bool(const FSoftObjectPath&, const TSharedPtr<FAvalancheManagedInstance>&)> InRemovePredicate, bool bInNotify)
+{
+	TArray<FSoftObjectPath> AssetsToNotify;
+	
+	for (TMap<FSoftObjectPath, TSharedPtr<FAvalancheManagedInstance>>::TIterator EntryIt = AvalancheInstances.CreateIterator(); EntryIt; ++EntryIt)
+	{
+		if (InRemovePredicate(EntryIt.Key(), EntryIt.Value()))
+		{
+			if (bInNotify)
+			{
+				AssetsToNotify.Add(EntryIt.Key());
+			}
+			OrderQueue.Remove(EntryIt.Key());
+			EntryIt.RemoveCurrent();
+		}
+	}
+
+	// Notify outside of the clean up loop since it is very likely the result of this
+	// will be a reload of the asset that got invalidated.
+	for (const FSoftObjectPath& AssetPath : AssetsToNotify)
+	{
+		OnEntryInvalidated.Broadcast(AssetPath);
+	}
+}
+
+void FAvalancheManagedInstanceCache::OnPackageSaved(const FString& InPackageFileName, UPackage* InPackage, FObjectPostSaveContext InObjectSaveContext)
+{
+	if (InObjectSaveContext.IsProceduralSave())
+	{
+		return;
+	}
+
+	OnPackageModified(InPackage->GetFName());
+}
+
+void FAvalancheManagedInstanceCache::OnAvaSyncPackageModified(IAvaMediaSyncProvider* InAvaMediaSyncProvider, const FName& InPackageName)
+{
+	UE_LOG(LogAvaMedia, Verbose,
+		TEXT("A sync operation has touched the package \"%s\" on disk. ManagedAvalancheBlueprintCache notified."),
+		*InPackageName.ToString());
+
+	OnPackageModified(InPackageName);
+}
+
+void FAvalancheManagedInstanceCache::OnAssetRemoved(const FAssetData& InAssetData)
+{
+	// Invalidate the internal cache for the given package.
+	OnPackageModified(InAssetData.PackageName);
+}
+
+void FAvalancheManagedInstanceCache::OnPackageModified(const FName& InPackageName)
+{
+	// Invalidate corresponding assets from that package.
+	RemoveEntries([InPackageName](const FSoftObjectPath& InAssetPath, const TSharedPtr<FAvalancheManagedInstance>&)
+		{
+			if (InAssetPath.GetLongPackageFName() == InPackageName)
+			{
+				UE_LOG(LogAvaMedia, Log,
+					TEXT("Managed Avalanche Blueprint Cache: Package \"%s\" being touched caused asset \"%s\" to be invalidated."),
+					*InPackageName.ToString(), *InAssetPath.ToString());
+				return true;
+			}
+			return false;
+		}, true);	// Notify of asset being invalidated to trigger a UI refresh.
+}
+
+void FAvalancheManagedInstanceCache::OnSettingChanged(UObject* , struct FPropertyChangedEvent&)
+{
+	TrimCache();
+}
\ No newline at end of file
diff --git a/Engine/Plugins/Experimental/Avalanche/Source/AvalancheMedia/Private/Playlist/AvalancheManagedInstanceLevel.cpp b/Engine/Plugins/Experimental/Avalanche/Source/AvalancheMedia/Private/Playlist/AvalancheManagedInstanceLevel.cpp
new file mode 100644
--- /dev/null
+++ b/Engine/Plugins/Experimental/Avalanche/Source/AvalancheMedia/Private/Playlist/AvalancheManagedInstanceLevel.cpp
@@ -0,0 +1,150 @@
+// Copyright Epic Games, Inc. All Rights Reserved.
+
+#include "Playlist/AvalancheManagedInstanceLevel.h"
+
+#include "AssetRegistry/AssetData.h"
+#include "AvaMediaDefines.h"
+#include "AvaRemoteControlRebind.h"
+#include "AvaRemoteControlUtils.h"
+#include "AvaScene.h"
+#include "Engine/Engine.h"
+#include "Engine/World.h"
+#include "IAvaMediaModule.h"
+#include "Playlist/AvalancheManagedInstanceUtils.h"
+#include "UObject/Package.h"
+
+namespace UE::AvalancheManagedInstanceLevel::Private
+{
+	UWorld* LoadAvalancheLevel(const FSoftObjectPath& InAssetPath)
+	{
+		return Cast<UWorld>(InAssetPath.TryLoad());
+	}
+
+	URemoteControlPreset* FindRemoteControlPreset(ULevel* InLevel)
+	{
+		const AAvaScene* AvaScene = AAvaScene::GetScene(InLevel, false);
+		return AvaScene ? AvaScene->GetRemoteControlPreset() : nullptr;
+	}
+}
+
+FAvalancheManagedInstanceLevel::FAvalancheManagedInstanceLevel(FAvalancheManagedInstanceCache* InParentCache, const FSoftObjectPath& InAssetPath)
+	: FAvalancheManagedInstance(InParentCache, InAssetPath)
+{
+	using namespace UE::AvalancheManagedInstanceLevel::Private;
+	// We need to load the source avalanche level.
+	UWorld* SourceAvalancheLevel = LoadAvalancheLevel(InAssetPath);
+	if (!SourceAvalancheLevel)
+	{
+		UE_LOG(LogAvaMedia, Error, TEXT("Unable to load Source Avalanche Level: %s"), *InAssetPath.ToString());
+		return;
+	}
+
+	// Keep a weak pointer
+	SourceAvalancheLevelWeak = SourceAvalancheLevel;
+
+	// Register the delegates on the source RCP just in case the level is currently being edited.
+	RegisterSourceRemoteControlPresetDelegates(FindRemoteControlPreset(SourceAvalancheLevel->PersistentLevel));
+
+	ManagedAvalancheLevelPackage = FAvalancheManagedInstanceUtils::MakeManagedInstancePackage(InAssetPath);
+	if (!ManagedAvalancheLevelPackage)
+	{
+		UE_LOG(LogAvaMedia, Error, TEXT("Unable to create a Managed Avalanche Level Package for %s"), *InAssetPath.ToString());
+		return;
+	}
+
+	// Remark: using PIE duplication mode to avoid entity and controller Ids from being renewed.
+	ManagedAvalancheLevel = Cast<UWorld>(StaticDuplicateObject(SourceAvalancheLevel, ManagedAvalancheLevelPackage.Get(), NAME_None, RF_NoFlags, nullptr, EDuplicateMode::PIE));
+	if (!ManagedAvalancheLevel)
+	{
+		UE_LOG(LogAvaMedia, Error, TEXT("Unable to duplicate Source Avalanche Level: %s"), *InAssetPath.ToString());
+		return;
+	}
+
+	ManagedAvalancheLevel->SetFlags(RF_Public | RF_Transient);
+	
+	ManagedRemoteControlPreset = FindRemoteControlPreset(ManagedAvalancheLevel->PersistentLevel);
+	FAvaRemoteControlRebind::RebindUnboundEntities(ManagedRemoteControlPreset, ManagedAvalancheLevel->PersistentLevel);
+	
+	// Backup the remote control values from the source asset.
+	constexpr bool bIsDefault = true;	// Flag the values as "default".
+	DefaultRemoteControlValues.CopyFrom(ManagedRemoteControlPreset, bIsDefault);
+
+	FAvalancheManagedInstanceUtils::PreventWorldFromBeingSeenAsLeakingByLevelEditor(ManagedAvalancheLevel.Get());
+
+	IAvaMediaModule::Get().GetOnMapChangedEvent().AddRaw(this, &FAvalancheManagedInstanceLevel::OnMapChangedEvent);
+}
+
+FAvalancheManagedInstanceLevel::~FAvalancheManagedInstanceLevel()
+{
+	IAvaMediaModule::Get().GetOnMapChangedEvent().RemoveAll(this);
+
+	if (ManagedRemoteControlPreset)
+	{
+		FAvaRemoteControlUtils::UnregisterRemoteControlPreset(ManagedRemoteControlPreset);
+	}
+
+	if (ManagedAvalancheLevelPackage)
+	{
+		ManagedAvalancheLevelPackage->ClearDirtyFlag();
+	}
+	
+	ManagedAvalancheLevel = nullptr;
+	ManagedAvalancheLevelPackage = nullptr;
+	ManagedRemoteControlPreset = nullptr;
+	DiscardSourceAvalancheLevel();
+}
+
+void FAvalancheManagedInstanceLevel::AddReferencedObjects( FReferenceCollector& Collector )
+{
+	Collector.AddReferencedObject( ManagedAvalancheLevel );
+	Collector.AddReferencedObject( ManagedAvalancheLevelPackage );
+	Collector.AddReferencedObject( ManagedRemoteControlPreset );
+}
+
+FString FAvalancheManagedInstanceLevel::GetReferencerName() const
+{
+	return TEXT("FAvalancheManagedInstanceLevel");
+}
+
+IAvaSceneInterface* FAvalancheManagedInstanceLevel::GetSceneInterface() const
+{
+	return ManagedAvalancheLevel ? static_cast<IAvaSceneInterface*>(AAvaScene::GetScene(ManagedAvalancheLevel->PersistentLevel, false)) : nullptr;
+}
+
+void FAvalancheManagedInstanceLevel::DiscardSourceAvalancheLevel()
+{
+	if (const UWorld* SourceAvalancheLevel = SourceAvalancheLevelWeak.Get())
+	{
+		using namespace UE::AvalancheManagedInstanceLevel::Private;
+		UnregisterSourceRemoteControlPresetDelegates(FindRemoteControlPreset(SourceAvalancheLevel->PersistentLevel));
+	}
+	SourceAvalancheLevelWeak.Reset();
+}
+
+void FAvalancheManagedInstanceLevel::OnMapChangedEvent(UWorld* InWorld, EAvaMediaMapChangeType InEventType)
+{
+	if (!InWorld->GetPackage())
+	{
+		return;
+	}
+
+	if (SourceAssetPath.GetLongPackageFName() != InWorld->GetPackage()->GetFName())
+	{
+		return;
+	}
+	
+	if (InEventType == EAvaMediaMapChangeType::LoadMap)
+	{
+		using namespace UE::AvalancheManagedInstanceLevel::Private;
+		// This should be fast given the level has been loaded in the editor.
+		if (UWorld* SourceAvalancheLevel = LoadAvalancheLevel(SourceAssetPath))
+		{
+			SourceAvalancheLevelWeak = SourceAvalancheLevel;
+			RegisterSourceRemoteControlPresetDelegates(FindRemoteControlPreset(SourceAvalancheLevel->PersistentLevel));
+		}
+	}
+	else if (InEventType == EAvaMediaMapChangeType::TearDownWorld)
+	{
+		DiscardSourceAvalancheLevel();
+	}
+}
\ No newline at end of file
diff --git a/Engine/Plugins/Experimental/Avalanche/Source/AvalancheMedia/Private/Playlist/AvalancheManagedInstanceLevel.h b/Engine/Plugins/Experimental/Avalanche/Source/AvalancheMedia/Private/Playlist/AvalancheManagedInstanceLevel.h
new file mode 100644
--- /dev/null
+++ b/Engine/Plugins/Experimental/Avalanche/Source/AvalancheMedia/Private/Playlist/AvalancheManagedInstanceLevel.h
@@ -0,0 +1,42 @@
+// Copyright Epic Games, Inc. All Rights Reserved.
+
+#pragma once
+
+#include "Engine/World.h"
+#include "Playlist/AvalancheManagedInstance.h"
+#include "UObject/Package.h"
+#include "UObject/StrongObjectPtr.h"
+
+enum class EAvaMediaMapChangeType : uint8;
+
+class FAvalancheManagedInstanceLevel : public FAvalancheManagedInstance
+{
+public:
+	UE_AVA_INHERITS(FAvalancheManagedInstanceLevel, FAvalancheManagedInstance);
+
+	FAvalancheManagedInstanceLevel(FAvalancheManagedInstanceCache* InParentCache, const FSoftObjectPath& InAssetPath);
+	virtual ~FAvalancheManagedInstanceLevel() override;
+
+	//~ Begin FGCObject
+	virtual void AddReferencedObjects( FReferenceCollector& Collector ) override;
+	virtual FString GetReferencerName() const override;
+	//~ End FGCObject
+	
+	//~ Begin FAvalancheManagedInstance
+	virtual bool IsValid() const override { return ManagedAvalancheLevel != nullptr;}
+	virtual URemoteControlPreset* GetRemoteControlPreset() const override { return ManagedRemoteControlPreset;}
+	virtual IAvaSceneInterface* GetSceneInterface() const override;
+	//~ End FAvalancheManagedInstance
+
+private:
+	void DiscardSourceAvalancheLevel();
+	void OnMapChangedEvent(UWorld* InWorld, EAvaMediaMapChangeType InEventType);
+
+private:
+	/** Keep a weak pointer to the loaded source level. This is used to unregister delegates if needed. */
+	TWeakObjectPtr<UWorld> SourceAvalancheLevelWeak;
+	
+	TObjectPtr<UWorld> ManagedAvalancheLevel;
+	TObjectPtr<UPackage> ManagedAvalancheLevelPackage;
+	TObjectPtr<URemoteControlPreset> ManagedRemoteControlPreset;
+};
\ No newline at end of file
diff --git a/Engine/Plugins/Experimental/Avalanche/Source/AvalancheMedia/Private/Playlist/AvalancheManagedInstanceUtils.cpp b/Engine/Plugins/Experimental/Avalanche/Source/AvalancheMedia/Private/Playlist/AvalancheManagedInstanceUtils.cpp
new file mode 100644
--- /dev/null
+++ b/Engine/Plugins/Experimental/Avalanche/Source/AvalancheMedia/Private/Playlist/AvalancheManagedInstanceUtils.cpp
@@ -0,0 +1,71 @@
+// Copyright Epic Games, Inc. All Rights Reserved.
+
+#include "Playlist/AvalancheManagedInstanceUtils.h"
+
+#include "Engine/Engine.h"
+#include "Engine/World.h"
+#include "IAvaMediaModule.h"
+#include "UObject/Package.h"
+
+UPackage* FAvalancheManagedInstanceUtils::MakeManagedInstancePackage(const FSoftObjectPath& InAssetPath)
+{
+	// Remote Control Preset will be registered with the package name.
+	// We want a package name unique to the game instance and that will be
+	// human readable since it will show up in the Web Remote Control page.
+
+	FString InstancePackageName = TEXT("/Temp/Managed");	// Keep it short for web page.
+
+	// In order to keep things short, we remove "/Game".
+	FString InstanceSubPath;
+	if (!InAssetPath.GetLongPackageName().Split(TEXT("/Game"), nullptr, &InstanceSubPath))
+	{
+		InstanceSubPath = InAssetPath.GetLongPackageName();
+	}
+	InstancePackageName += InstanceSubPath;
+
+	UPackage* ManagedAvalancheBlueprintPackage = nullptr;
+
+	// Make sure the package name is unique. We don't want stomping of existing instances.
+	if (FindObject<UPackage>( nullptr, *InstancePackageName))
+	{
+		FString UniqueInstancePackageName = MakeUniqueObjectName( nullptr, UPackage::StaticClass(), FName(*InstancePackageName)).ToString();
+		
+		UE_LOG(LogAvaMedia, Warning, 
+			TEXT("Package \"%s\" for Avalanche Managed Instance already exists. Will use \"%s\" instead."),
+			*InstancePackageName, *UniqueInstancePackageName);
+
+		InstancePackageName = UniqueInstancePackageName;
+	}
+	
+	ManagedAvalancheBlueprintPackage = CreatePackage(*InstancePackageName);
+	if (ManagedAvalancheBlueprintPackage)
+	{
+		ManagedAvalancheBlueprintPackage->SetFlags(RF_Transient | RF_Public);
+	}
+	else
+	{
+		// Note: The outer will fallback to GEngine in that case.
+		UE_LOG(LogAvaMedia, Error, 
+			TEXT("Unable to create package \"%s\" for Avalanche Managed Instance."),
+			*InstancePackageName);
+	}
+	return ManagedAvalancheBlueprintPackage;
+}
+
+void FAvalancheManagedInstanceUtils::PreventWorldFromBeingSeenAsLeakingByLevelEditor(UWorld* InWorld)
+{
+	// Appease the level editor clean up when it sees worlds, it must
+	// not consider it as a leaking world.
+	InWorld->WorldType = EWorldType::Inactive;
+
+	if (UPackage* Package = InWorld->GetPackage())
+	{
+#if WITH_EDITOR
+		// Prevents "unloading" of map package by marking it as new.
+		Package->MarkAsNewlyCreated();
+#endif
+	
+		// Prevents "unload" to try to save it (just in case).
+		Package->ClearDirtyFlag();
+	}
+}
diff --git a/Engine/Plugins/Experimental/Avalanche/Source/AvalancheMedia/Private/Playlist/AvalancheManagedInstanceUtils.h b/Engine/Plugins/Experimental/Avalanche/Source/AvalancheMedia/Private/Playlist/AvalancheManagedInstanceUtils.h
new file mode 100644
--- /dev/null
+++ b/Engine/Plugins/Experimental/Avalanche/Source/AvalancheMedia/Private/Playlist/AvalancheManagedInstanceUtils.h
@@ -0,0 +1,17 @@
+// Copyright Epic Games, Inc. All Rights Reserved.
+
+#pragma once
+
+#include "UObject/StrongObjectPtr.h"
+
+class UPackage;
+class UWorld;
+struct FSoftObjectPath;
+
+/** Common utility functions for implementing Managed Instance. */ 
+class FAvalancheManagedInstanceUtils
+{
+public:
+	static UPackage* MakeManagedInstancePackage(const FSoftObjectPath& InAssetPath);
+	static void PreventWorldFromBeingSeenAsLeakingByLevelEditor(UWorld* InWorld);
+};
\ No newline at end of file
diff --git a/Engine/Plugins/Experimental/Avalanche/Source/AvalancheMedia/Private/Playlist/AvalanchePage.cpp b/Engine/Plugins/Experimental/Avalanche/Source/AvalancheMedia/Private/Playlist/AvalanchePage.cpp
new file mode 100644
--- /dev/null
+++ b/Engine/Plugins/Experimental/Avalanche/Source/AvalancheMedia/Private/Playlist/AvalanchePage.cpp
@@ -0,0 +1,646 @@
+// Copyright Epic Games, Inc. All Rights Reserved.
+
+#include "Playlist/AvalanchePage.h"
+
+#include "AssetRegistry/AssetData.h"
+#include "AvaScene.h"
+#include "AvaTransitionTree.h"
+#include "AvalancheBroadcast.h"
+#include "IAvaMediaModule.h"
+#include "Playback/AvaMediaPlaybackClient.h"
+#include "Playback/AvaMediaPlaybackManager.h"
+#include "Playback/AvalancheRemoteControl.h"
+#include "Playlist/AvalancheManagedInstanceCache.h"
+#include "Playlist/AvalanchePageAssetUtils.h"
+#include "Playlist/AvalanchePlaylist.h"
+#include "RCVirtualProperty.h"
+
+#define LOCTEXT_NAMESPACE "AvalanchePage"
+
+FAvalanchePage FAvalanchePage::NullPage;
+const int32 FAvalanchePage::InvalidPageId = -1;
+
+FAvalanchePage::FAvalanchePage(int32 InPageId, int32 InTemplateId)
+	: PageId(InPageId)
+	, TemplateId(InTemplateId)
+	, PageName(TEXT("New Page"))
+{
+}
+
+bool FAvalanchePage::IsValidPage() const
+{
+	return this != &FAvalanchePage::NullPage
+		&& PageId != InvalidPageId;
+}
+
+void FAvalanchePage::Rename(const FString& InNewName)
+{
+	PageName = InNewName;
+}
+
+void FAvalanchePage::RenameFriendlyName(const FString& InNewName)
+{
+	FriendlyName = FText::FromString(InNewName);
+}
+
+namespace UE::AvalanchePage::Private
+{
+	// Determines the page status for a given playback status and secondary states.
+	FAvalanchePageStatus GetPageStatus(EAvaBroadcastChannelType InType, EAvaMediaPlaybackStatus InPlaybackStatus, bool bInPagePlaying, bool bInAssetNeedSync)
+	{
+		switch (InPlaybackStatus)
+		{
+		case EAvaMediaPlaybackStatus::Unknown:
+			return {InType,!bInAssetNeedSync ? EAvalanchePageStatus::Unknown : EAvalanchePageStatus::NeedsSync, bInAssetNeedSync};
+		case EAvaMediaPlaybackStatus::Missing:			
+			return {InType, EAvalanchePageStatus::Missing, false};
+		case EAvaMediaPlaybackStatus::Syncing:
+			return {InType, EAvalanchePageStatus::Syncing, false};
+		case EAvaMediaPlaybackStatus::Available:
+			// There is an explicit "needs sync" page status, along with the flag. This is just to make it more explicit.
+			return {InType, !bInAssetNeedSync ? EAvalanchePageStatus::Available : EAvalanchePageStatus::NeedsSync, bInAssetNeedSync};
+		case EAvaMediaPlaybackStatus::Loading:
+			return {InType, EAvalanchePageStatus::Loading, bInAssetNeedSync};
+		case EAvaMediaPlaybackStatus::Loaded:
+			return {InType, EAvalanchePageStatus::Loaded, bInAssetNeedSync};
+		case EAvaMediaPlaybackStatus::Starting:
+			return {InType, EAvalanchePageStatus::Loading, bInAssetNeedSync};
+		case EAvaMediaPlaybackStatus::Started:
+			return {InType, bInPagePlaying ? EAvalanchePageStatus::Playing : EAvalanchePageStatus::Loaded, bInAssetNeedSync};
+		case EAvaMediaPlaybackStatus::Stopping:
+		case EAvaMediaPlaybackStatus::Unloading:
+			return {InType, EAvalanchePageStatus::Available, bInAssetNeedSync};
+		case EAvaMediaPlaybackStatus::Error:
+		default:
+			return {InType, EAvalanchePageStatus::Error, bInAssetNeedSync};
+		}
+	}
+
+	FAvalanchePageStatus GetProgramPageStatus(EAvaMediaPlaybackStatus InPlaybackStatus, bool bInPagePlaying, bool bInAssetNeedSync)
+	{
+		return GetPageStatus(EAvaBroadcastChannelType::Program, InPlaybackStatus, bInPagePlaying, bInAssetNeedSync);
+	}
+	
+	EAvaMediaPlaybackStatus GetPlaybackStatus(EAvaMediaPlaybackAssetStatus InAssetStatus)
+	{
+		switch (InAssetStatus)
+		{
+		case EAvaMediaPlaybackAssetStatus::Unknown:
+			return EAvaMediaPlaybackStatus::Unknown;
+		case EAvaMediaPlaybackAssetStatus::Missing:
+			return EAvaMediaPlaybackStatus::Missing;
+		case EAvaMediaPlaybackAssetStatus::MissingDependencies:
+		case EAvaMediaPlaybackAssetStatus::NeedsSync:
+		case EAvaMediaPlaybackAssetStatus::Available:
+			return EAvaMediaPlaybackStatus::Available;
+		default:
+			return EAvaMediaPlaybackStatus::Unknown;
+		}
+	}
+}
+
+
+FText FAvalanchePage::GetPageDescription() const
+{
+	if (HasPageFriendlyName())
+	{
+		return FriendlyName;
+	}
+
+	if (HasPageSummary())
+	{
+		return PageSummary;
+	}
+
+	return FText::FromString(PageName);
+}
+
+bool FAvalanchePage::UpdatePageSummary(const UAvalanchePlaylist* InPlaylist)
+{
+	const bool bShouldUpdateDescription = !HasPageSummary();
+	if (!IsValidPage() || IsTemplate() || !bShouldUpdateDescription)
+	{
+		return false;
+	}
+
+	const TArray<FSoftObjectPath> AvaAssetPaths = GetAvalancheAssetPaths(InPlaylist);
+	TArray<const URemoteControlPreset*> Presets;
+	Presets.Reserve(AvaAssetPaths.Num());
+
+	FAvalancheManagedInstanceCache& ManagedInstanceCache = IAvaMediaModule::Get().GetManagedInstanceCache();
+	
+	for (const FSoftObjectPath& AvaAssetPath : AvaAssetPaths)
+	{
+		if (const TSharedPtr<FAvalancheManagedInstance> ManagedInstance = ManagedInstanceCache.GetOrLoadAvalancheInstance(AvaAssetPath))
+		{
+			if (const URemoteControlPreset* Preset = ManagedInstance->GetRemoteControlPreset())
+			{
+				Presets.Add(Preset);
+			}
+		}
+	}
+	
+	return UpdatePageSummary(Presets, false);
+}
+
+bool FAvalanchePage::UpdatePageSummary(const TArray<const URemoteControlPreset*>& InPresets, bool bInIsPresetChanged)
+{
+	const bool bShouldUpdateDescription = bInIsPresetChanged ? true : !HasPageSummary();
+	if (!IsValidPage() || IsTemplate() || InPresets.IsEmpty() || !bShouldUpdateDescription)
+	{
+		return false;
+	}
+	
+	TSortedMap<int32, FText> OrderedDescription;
+
+	int32 DisplayIndexOffset = 0;
+	for (const URemoteControlPreset* Preset : InPresets)
+	{
+		for (const URCVirtualPropertyBase* VirtualProperty : Preset->GetControllers())
+		{
+			FString StringValue = VirtualProperty->GetDisplayValueAsString();
+			const FText TextToAdd = FText::FromString(StringValue);
+			if (!TextToAdd.IsEmptyOrWhitespace())
+			{
+				OrderedDescription.FindOrAdd(VirtualProperty->DisplayIndex + DisplayIndexOffset) = TextToAdd;
+			}
+		}
+		DisplayIndexOffset += Preset->GetNumControllers();
+	}
+	
+	if (!OrderedDescription.IsEmpty())
+	{
+		FFormatOrderedArguments InitialDescription;
+		InitialDescription.Reserve(OrderedDescription.Num());
+
+		constexpr int32 FinishDescriptionWhenExceedLimit = 50;
+		int32 CurrentLenght = 0;
+		for (const TPair<int32, FText>& OrderedValue : OrderedDescription)
+		{
+			InitialDescription.Add(OrderedValue.Value);
+			CurrentLenght += OrderedValue.Value.ToString().Len();
+			if (CurrentLenght >= FinishDescriptionWhenExceedLimit)
+			{
+				InitialDescription.Add(LOCTEXT("EndOfDescription", "..."));
+				break;
+			}
+		}
+
+		PageSummary = FText::Join(LOCTEXT("DescriptionDelimiter", " / "), InitialDescription);
+		return true;
+	}
+
+	return false;
+}
+
+bool FAvalanchePage::UpdateTransitionLogic()
+{
+	check(IsTemplate());
+	if (const UObject* LoadedSourceAsset = AssetPath.TryLoad())
+	{
+		if (const IAvaSceneInterface* SceneInterface = FAvalanchePageAssetUtils::GetSceneInterface(LoadedSourceAsset))
+		{
+			const UAvaTransitionTree* TransitionTree = FAvalanchePageAssetUtils::FindTransitionTree(SceneInterface);
+			bHasTransitionLogic = TransitionTree ? TransitionTree->IsEnabled() : false;
+			TransitionLayerTag = FAvalanchePageAssetUtils::GetTransitionLayerTag(TransitionTree);
+		}
+	}
+	return false;
+}
+
+bool FAvalanchePage::HasTransitionLogic(const UAvalanchePlaylist* InPlaylist) const
+{	
+	const FAvalanchePage& Template = ResolveTemplate(InPlaylist);
+	if (Template.IsValidPage())
+	{
+		// Combo templates have Transition Logic.
+		return Template.IsComboTemplate() || Template.bHasTransitionLogic;
+	}
+
+	return bHasTransitionLogic;
+}
+
+FAvaTagHandle FAvalanchePage::GetTransitionLayer(const UAvalanchePlaylist* InPlaylist, int32 InTemplateIndex) const
+{
+	const FAvalanchePage& Template = GetTemplate(InPlaylist, InTemplateIndex);
+	return Template.IsValidPage() ?  Template.TransitionLayerTag : TransitionLayerTag;
+}
+
+TArray<FAvaTagHandle> FAvalanchePage::GetTransitionLayers(const UAvalanchePlaylist* InPlaylist) const
+{
+	TArray<FAvaTagHandle> TransitionLayers;
+
+	const int32 NumTemplates = GetNumTemplates(InPlaylist);
+	TransitionLayers.Reserve(NumTemplates);
+
+	for(int32 TemplateIndex = 0; TemplateIndex < NumTemplates; ++TemplateIndex)
+	{
+		const FAvalanchePage& Template = GetTemplate(InPlaylist, TemplateIndex);
+		if (Template.IsValidPage())
+		{
+			TransitionLayers.Add(Template.TransitionLayerTag);
+		}
+	}
+
+	return TransitionLayers;
+}
+
+int32 FAvalanchePage::AppendPageProgramStatuses(const UAvalanchePlaylist* InParentPlaylist, TArray<FAvalanchePageStatus>& OutPageStatuses) const
+{
+	if (!InParentPlaylist || IsTemplate())
+	{
+		return 0;
+	}
+	
+	const int32 PreviousNumStatuses = OutPageStatuses.Num();
+	
+	IAvaMediaModule& AvaMediaModule = IAvaMediaModule::Get();
+	FAvaMediaPlaybackManager& PlaybackManager = InParentPlaylist->GetPlaybackManager();
+	const FSoftObjectPath ResolvedAssetPath = GetAvalancheAssetPath(InParentPlaylist);
+	const FString ChannelName = GetChannelName().ToString();
+	const UAvalanchePagePlayer* ProgramPagePlayer = InParentPlaylist->FindPlayerForProgramPage(GetPageId());
+	const bool bIsPagePlaying = ProgramPagePlayer ? ProgramPagePlayer->IsPlaying() : false;
+	const FAvaMediaPlaybackInstance* LocalPlaybackInstance = ProgramPagePlayer ? ProgramPagePlayer->GetPlaybackInstance() : nullptr;
+	const FGuid LocalPlaybackInstanceId = LocalPlaybackInstance ? LocalPlaybackInstance->GetInstanceId() : FGuid();
+	
+	bool bLocalStatusAdded = false;
+
+	auto AddLocalProgramStatusOnce = [LocalPlaybackInstance, &PlaybackManager, &bLocalStatusAdded, &OutPageStatuses, bIsPagePlaying, &ResolvedAssetPath, &ChannelName]()
+	{
+		using namespace UE::AvalanchePage::Private;
+
+		// Local status is added only once.
+		if (bLocalStatusAdded)
+		{
+			return;
+		}
+		
+		// if page is playing, access the status from the page player's playback instance.
+		if (LocalPlaybackInstance)
+		{
+			OutPageStatuses.Add(GetProgramPageStatus(LocalPlaybackInstance->GetStatus(), bIsPagePlaying, false));
+		}
+		else
+		{
+			// If the page is not playing, try to find a cached playback instance of the asset for the given channel.
+			constexpr FGuid InvalidId;	// Providing an invalid instance id will fallback to searching by channel. 
+			if (const TSharedPtr<FAvaMediaPlaybackInstance> PlaybackInstance = PlaybackManager.FindPlaybackInstance(InvalidId, ResolvedAssetPath, ChannelName))
+			{
+				OutPageStatuses.Add(GetProgramPageStatus(PlaybackInstance->GetStatus(), bIsPagePlaying, false));
+			}
+			else
+			{
+				OutPageStatuses.Add(GetProgramPageStatus(PlaybackManager.GetUnloadedPlaybackStatus(ResolvedAssetPath), bIsPagePlaying, false));
+			}
+		}
+		bLocalStatusAdded = true;
+	};
+	
+	const FAvaOutputChannel& Channel = UAvalancheBroadcast::Get().GetCurrentProfile().GetChannel(GetChannelName());
+	const TArray<UMediaOutput*>& Outputs = Channel.GetMediaOutputs();
+	IAvaMediaPlaybackClient& PlaybackClient = AvaMediaModule.GetMediaPlaybackClient();
+
+	TSet<FString> AddedServers;
+
+	for (const UMediaOutput* Output : Outputs)
+	{
+		if (Channel.IsMediaOutputRemote(Output))
+		{
+			if (Channel.GetMediaOutputState(Output) != EAvaMediaOutputState::Offline)
+			{
+				const FString& ServerForOutput = Channel.GetMediaOutputServerName(Output);
+
+				if (AddedServers.Contains(ServerForOutput))
+				{
+					continue; // We have already added that server.
+				}
+
+				AddedServers.Add(ServerForOutput);	// Keep track of our servers so we add it only once.
+
+				// There is only one playback/asset status per server, even if it has many outputs.
+				TOptional<EAvaMediaPlaybackStatus> PlaybackStatus = PlaybackClient.GetRemotePlaybackStatus(LocalPlaybackInstanceId, ResolvedAssetPath, ChannelName, ServerForOutput);
+				const TOptional<EAvaMediaPlaybackAssetStatus> PlaybackAssetStatus = PlaybackClient.GetRemotePlaybackAssetStatus(ResolvedAssetPath, ServerForOutput);
+
+				if (!PlaybackAssetStatus.IsSet())
+				{
+					PlaybackClient.RequestPlaybackAssetStatus(ResolvedAssetPath, ServerForOutput, false);
+				}
+
+				using namespace UE::AvalanchePage::Private;
+				if (!PlaybackStatus.IsSet())
+				{
+					PlaybackClient.RequestPlayback(LocalPlaybackInstanceId, ResolvedAssetPath, ChannelName, EAvaMediaPlaybackAction::Status);
+					// Derive playback status from the asset status.
+					PlaybackStatus = PlaybackAssetStatus.IsSet() ? GetPlaybackStatus(PlaybackAssetStatus.GetValue()) : EAvaMediaPlaybackStatus::Unknown;
+				}
+
+				const bool bAssetNeedsSync = PlaybackAssetStatus.IsSet() && PlaybackAssetStatus.GetValue() == EAvaMediaPlaybackAssetStatus::NeedsSync;
+				OutPageStatuses.Add(GetProgramPageStatus(PlaybackStatus.GetValue(), bIsPagePlaying, bAssetNeedsSync));
+			}
+			else
+			{
+				OutPageStatuses.Add({EAvaBroadcastChannelType::Program, EAvalanchePageStatus::Offline, false});
+			}
+		}
+		else
+		{
+			// All local outputs lead to only one status, the local one.
+			AddLocalProgramStatusOnce();
+		}
+	}
+
+	// If there are no outputs defined, use the local status.
+	if (Outputs.IsEmpty())
+	{
+		AddLocalProgramStatusOnce();
+	}
+	
+	return OutPageStatuses.Num() - PreviousNumStatuses;
+}
+
+int32 FAvalanchePage::AppendPagePreviewStatuses(const UAvalanchePlaylist* InParentPlaylist, const FName& InPreviewChannelName, TArray<FAvalanchePageStatus>& OutPageStatuses) const
+{
+	if (!InParentPlaylist)
+	{
+		return 0;
+	}
+	
+	FAvaMediaPlaybackManager& PlaybackManager = InParentPlaylist->GetPlaybackManager();
+	const FSoftObjectPath ResolvedAssetPath = GetAvalancheAssetPath(InParentPlaylist);
+
+	const int32 PreviousNumStatuses = OutPageStatuses.Num();
+	
+	// For the preview, we only add the status if it is playing.
+	// We are not really interested if the preview is loaded.
+	if (InParentPlaylist->IsPagePreviewing(PageId))
+	{
+		OutPageStatuses.Add({EAvaBroadcastChannelType::Preview, EAvalanchePageStatus::Previewing, false});
+	}
+	else
+	{
+		using namespace UE::AvalanchePage::Private;
+		const FString PreviewChannelName = InPreviewChannelName.IsNone() ? InParentPlaylist->GetDefaultPreviewChannelName().ToString() : InPreviewChannelName.ToString();
+
+		// If the page is not previewing, try to find a cached playback instance of the asset for the given preview channel.
+		constexpr FGuid InvalidId;	// Providing an invalid instance id will fallback to searching by channel. 
+		if (const TSharedPtr<FAvaMediaPlaybackInstance> PlaybackInstance = PlaybackManager.FindPlaybackInstance(InvalidId, ResolvedAssetPath, PreviewChannelName))
+		{
+			OutPageStatuses.Add(GetPageStatus(EAvaBroadcastChannelType::Preview, PlaybackInstance->GetStatus(), false, false));
+		}
+		else
+		{
+			OutPageStatuses.Add(GetPageStatus(EAvaBroadcastChannelType::Preview, PlaybackManager.GetUnloadedPlaybackStatus(ResolvedAssetPath), false, false));
+		}
+	}
+
+	return OutPageStatuses.Num() - PreviousNumStatuses;
+}
+
+TArray<FAvalanchePageStatus> FAvalanchePage::GetPageStatuses(const UAvalanchePlaylist* InParentPlaylist) const
+{
+	TArray<FAvalanchePageStatus> Statuses;
+	if (InParentPlaylist)
+	{
+		// Typical instanced page: 1 program channel + 1 preview channel.
+		// Template page only has a preview channel.
+		Statuses.Reserve(IsTemplate() ? 1 : 2);
+		
+		if (!IsTemplate())
+		{
+			AppendPageProgramStatuses(InParentPlaylist, Statuses);
+		}
+		AppendPagePreviewStatuses(InParentPlaylist, InParentPlaylist->GetDefaultPreviewChannelName(), Statuses);
+	}
+	return Statuses;
+}
+
+TArray<FAvalanchePageStatus> FAvalanchePage::GetPageContextualStatuses(const UAvalanchePlaylist* InParentPlaylist) const
+{
+	TArray<FAvalanchePageStatus> Statuses;
+	if (InParentPlaylist)
+	{
+		Statuses.Reserve(1);
+		
+		if (!IsTemplate())
+		{
+			AppendPageProgramStatuses(InParentPlaylist, Statuses);
+		}
+		else
+		{
+			AppendPagePreviewStatuses(InParentPlaylist, InParentPlaylist->GetDefaultPreviewChannelName(), Statuses);
+		}
+	}
+	return Statuses;
+}
+
+TArray<FAvalanchePageStatus> FAvalanchePage::GetPageProgramStatuses(const UAvalanchePlaylist* InParentPlaylist) const
+{
+	TArray<FAvalanchePageStatus> Statuses;
+	if (InParentPlaylist && !IsTemplate())
+	{
+		Statuses.Reserve(1);
+		AppendPageProgramStatuses(InParentPlaylist, Statuses);
+	}
+	return Statuses;
+}
+
+TArray<FAvalanchePageStatus> FAvalanchePage::GetPagePreviewStatuses(const UAvalanchePlaylist* InParentPlaylist, const FName& InPreviewChannelName) const
+{
+	TArray<FAvalanchePageStatus> Statuses;
+	if (InParentPlaylist)
+	{
+		Statuses.Reserve(1);
+		AppendPagePreviewStatuses(InParentPlaylist, InPreviewChannelName, Statuses);
+	}
+	return Statuses;
+}
+
+FSoftObjectPath FAvalanchePage::GetAvalancheAssetPath(const UAvalanchePlaylist* InPlaylist, int32 InTemplateIndex) const
+{
+	const FAvalanchePage& Template = GetTemplate(InPlaylist, InTemplateIndex);
+	return Template.IsValidPage() ? Template.AssetPath : AssetPath;
+}
+
+TArray<FSoftObjectPath> FAvalanchePage::GetAvalancheAssetPaths(const UAvalanchePlaylist* InPlaylist) const
+{
+	TArray<FSoftObjectPath> AssetPaths;
+
+	const int32 NumTemplates = GetNumTemplates(InPlaylist);
+	AssetPaths.Reserve(NumTemplates);
+
+	for(int32 TemplateIndex = 0; TemplateIndex < NumTemplates; ++TemplateIndex)
+	{
+		const FAvalanchePage& Template = GetTemplate(InPlaylist, TemplateIndex);
+		if (Template.IsValidPage())
+		{
+			AssetPaths.Add(Template.AssetPath);
+		}
+	}
+
+	return AssetPaths;
+}
+
+bool FAvalanchePage::UpdateAvalancheAsset(const FSoftObjectPath& InAssetPath, bool bInReimportPage)
+{
+	if (IsComboTemplate())
+	{
+		UE_LOG(LogAvaPlaylist, Error, TEXT("Can't update asset on a combo page directly."));
+		return false;
+	}
+	
+	if (InAssetPath != AssetPath || bInReimportPage)
+	{
+		AssetPath = InAssetPath;
+
+		if (!HasPageFriendlyName())
+		{
+			// Is GetName() the same as the asset name from soft object path?
+			PageName = InAssetPath.GetAssetName();
+		}
+
+		UpdateTransitionLogic();
+		return true;
+	}
+	return false;
+}
+
+FName FAvalanchePage::GetChannelName() const
+{
+	return UAvalancheBroadcast::Get().GetChannelName(OutputChannel);
+}
+
+void FAvalanchePage::SetChannelName(FName InChannelName)
+{
+	OutputChannel = UAvalancheBroadcast::Get().GetChannelIndex(InChannelName);
+	if (OutputChannel == INDEX_NONE)
+	{
+		UE_LOG(LogAvaBroadcast, Error, TEXT("Channel %s was not found in broadcast channels, using %s instead."),
+			*InChannelName.ToString(), *UAvalancheBroadcast::Get().GetChannelName(0).ToString());
+		OutputChannel = 0;
+	}
+}
+
+EAvaRemoteControlChanges FAvalanchePage::PruneRemoteControlValues(const FAvalancheRemoteControlValues& InRemoteControlValues)
+{
+	return RemoteControlValues.PruneRemoteControlValues(InRemoteControlValues);
+}
+
+EAvaRemoteControlChanges FAvalanchePage::UpdateRemoteControlValues(const FAvalancheRemoteControlValues& InRemoteControlValues, bool bInUpdateDefaults)
+{
+	return RemoteControlValues.UpdateRemoteControlValues(InRemoteControlValues, bInUpdateDefaults);
+}
+
+void FAvalanchePage::SetRemoteControlEntityValue(const FGuid& InId, const FAvalancheRemoteControlValue& InValue)
+{
+	RemoteControlValues.SetEntityValue(InId, InValue);
+}
+
+void FAvalanchePage::SetRemoteControlControllerValue(const FGuid& InId, const FAvalancheRemoteControlValue& InValue)
+{
+	RemoteControlValues.SetControllerValue(InId, InValue);
+}
+
+void FAvalanchePage::PostLoad()
+{
+	if (!AvalancheBlueprint_DEPRECATED.IsNull() && !AssetPath.IsValid())
+	{
+		AssetPath = AvalancheBlueprint_DEPRECATED.ToSoftObjectPath();
+	}
+	AvalancheBlueprint_DEPRECATED.Reset();
+}
+
+int32 FAvalanchePage::GetNumTemplates(const UAvalanchePlaylist* InPlaylist) const
+{
+	const FAvalanchePage& Template = ResolveTemplate(InPlaylist);
+	if (Template.IsValidPage())
+	{
+		const int32 NumCombinedTemplates = Template.CombinedTemplateIds.Num();
+		return NumCombinedTemplates > 0 ? NumCombinedTemplates : 1;
+	}
+	return 0;
+}
+
+const FAvalanchePage& FAvalanchePage::GetTemplate(const UAvalanchePlaylist* InPlaylist, int32 InIndex) const
+{
+	const FAvalanchePage& Template = ResolveTemplate(InPlaylist);
+	
+	if (Template.IsValidPage() && Template.CombinedTemplateIds.Num())
+	{
+		if (Template.CombinedTemplateIds.IsValidIndex(InIndex))
+		{
+			// Remark: not supporting recursive template combos.
+			const FAvalanchePage& OtherTemplate  = InPlaylist->GetPage(Template.CombinedTemplateIds[InIndex]);
+			if (OtherTemplate.IsValidPage())
+			{
+				return OtherTemplate;
+			}
+			
+			UE_LOG(LogAvaPlaylist, Error,
+				TEXT("Internal error while accessing page %d's template (%d index %d): reference to template Id %d is not valid."),
+				GetPageId(), Template.GetPageId(), InIndex, Template.CombinedTemplateIds[InIndex]);
+		}
+		
+		UE_LOG(LogAvaPlaylist, Error,
+			TEXT("Internal error while accessing page %d's template (%d): specified index %d is not valid."),
+			GetPageId(), Template.GetPageId(), InIndex);
+	}
+	
+	return Template;
+}
+
+const FAvalanchePage& FAvalanchePage::ResolveTemplate(const UAvalanchePlaylist* InPlaylist) const
+{
+	if (IsTemplate())
+	{
+		return *this;
+	}
+
+	if (IsValid(InPlaylist))
+	{
+		if (InPlaylist->GetTemplatePages().PageIndices.Contains(GetPageId()))
+		{
+			UE_LOG(LogAvaPlaylist, Warning,
+				TEXT("PageId %d is in the template list but has \"IsTemplate\" flag to false."), GetPageId());
+
+			// We're obviously a template, but there's been a mix-up...
+			return *this;
+		}
+
+		if (InPlaylist->GetInstancedPages().PageIndices.Contains(GetPageId()))
+		{
+			if (const int32* TemplateIdxPtr = InPlaylist->GetTemplatePages().PageIndices.Find(GetTemplateId()))
+			{
+				return InPlaylist->GetTemplatePages().Pages[*TemplateIdxPtr];
+			}
+			
+			UE_LOG(LogAvaPlaylist, Error,
+				TEXT("PageId [%d] is an instanced page, has template id [%d], but that template doesn't exist."),
+				GetPageId(), GetTemplateId());
+		}
+	}
+
+	return NullPage;
+}
+
+bool FAvalanchePage::IsTemplateMatchingByValue(const FAvalanchePage& InTemplatePage) const
+{
+	if (!IsTemplate() || !InTemplatePage.IsTemplate())
+	{
+		return false;
+	}
+
+	if (AssetPath != InTemplatePage.AssetPath)
+	{
+		return false;
+	}
+
+	if (!RemoteControlValues.HasSameEntityValues(InTemplatePage.RemoteControlValues)
+		|| !RemoteControlValues.HasSameControllerValues(InTemplatePage.RemoteControlValues))
+	{
+		return false;
+	}
+	
+	return true;
+}
+
+#undef LOCTEXT_NAMESPACE
diff --git a/Engine/Plugins/Experimental/Avalanche/Source/AvalancheMedia/Private/Playlist/AvalanchePageAssetUtils.cpp b/Engine/Plugins/Experimental/Avalanche/Source/AvalancheMedia/Private/Playlist/AvalanchePageAssetUtils.cpp
new file mode 100644
--- /dev/null
+++ b/Engine/Plugins/Experimental/Avalanche/Source/AvalancheMedia/Private/Playlist/AvalanchePageAssetUtils.cpp
@@ -0,0 +1,96 @@
+// Copyright Epic Games, Inc. All Rights Reserved.
+
+#include "Playlist/AvalanchePageAssetUtils.h"
+
+#include "AvaBlueprint.h"
+#include "AvaScene.h"
+#include "AvaTagHandle.h"
+#include "AvaTransitionSubsystem.h"
+#include "AvaTransitionTree.h"
+#include "Behavior/IAvaTransitionBehavior.h"
+#include "Engine/Level.h"
+#include "Engine/World.h"
+#include "IAvaSceneInterface.h"
+
+namespace UE::AvalanchePageAssetUtils::Private
+{
+	template<typename InInterfaceType>
+	const InInterfaceType* GetInterfaceFromAvaAsset(const UObject* InLoadedAsset)
+	{
+		if (const InInterfaceType* const Interface = Cast<InInterfaceType>(InLoadedAsset))
+		{
+			return Interface;
+		}
+		if (const UWorld* const AvalancheWorld = Cast<UWorld>(InLoadedAsset))
+		{
+			if (const AAvaScene* AvaScene = AAvaScene::GetScene(AvalancheWorld->PersistentLevel, false))
+			{
+				return static_cast<const InInterfaceType*>(AvaScene);
+			}
+		}
+		return nullptr;
+	}
+}
+
+const IAvaSceneInterface* FAvalanchePageAssetUtils::GetSceneInterface(const UObject* InLoadedAsset)
+{
+	using namespace UE::AvalanchePageAssetUtils::Private;
+	return GetInterfaceFromAvaAsset<IAvaSceneInterface>(InLoadedAsset);
+}
+
+const UAvaTransitionTree* FAvalanchePageAssetUtils::FindTransitionTree(const IAvaSceneInterface* InSceneInterface)
+{
+	if (!InSceneInterface)
+	{
+		return nullptr;
+	}
+	
+	ULevel* SceneLevel = InSceneInterface->GetSceneLevel();
+	if (!SceneLevel)
+	{
+		return nullptr;
+	}
+
+	// Method 1 - Using the UAvaTransitionSubsystem
+	// This method may not work if the world is not initialized (inactive, as in the managed instance).
+	if (const UWorld* SceneWorld = SceneLevel->GetWorld())
+	{
+		if (const UAvaTransitionSubsystem* TransitionSubsystem = SceneWorld->GetSubsystem<UAvaTransitionSubsystem>())
+		{
+			if (const IAvaTransitionBehavior* TransitionBehavior = TransitionSubsystem->GetTransitionBehavior(SceneLevel))
+			{
+				if (const UAvaTransitionTree* TransitionTree = TransitionBehavior->GetTransitionTree())
+				{
+					return TransitionTree;
+				}
+			}
+			// If we had a subsystem and it didn't find the behavior then there was no state tree either and we are done.
+			return nullptr;
+		}
+	}
+	
+	// Method 2 - Fallback using direct lookup of the private actor class.
+	// This is needed to lookup in managed instances and source asset.
+
+	static FString AvaTransitionBehaviorActor(TEXT("AvaTransitionBehaviorActor"));
+	for (const AActor* Actor : SceneLevel->Actors)
+	{
+		const UClass* ActorClass = Actor->GetClass();
+		if (ActorClass && ActorClass->GetName() == AvaTransitionBehaviorActor)
+		{
+			const FObjectPropertyBase* const Property = FindFProperty<FObjectPropertyBase>(ActorClass, TEXT("TransitionTree"));
+			if (Property
+				&& Property->PropertyClass
+				&& Property->PropertyClass->IsChildOf(UAvaTransitionTree::StaticClass()))
+			{
+				return Cast<UAvaTransitionTree>(Property->GetObjectPropertyValue_InContainer(Actor));
+			}
+		}
+	}
+	return nullptr;
+}
+
+FAvaTagHandle FAvalanchePageAssetUtils::GetTransitionLayerTag(const UAvaTransitionTree* InTransitionTree)
+{
+	return InTransitionTree ? InTransitionTree->GetTransitionLayer() : FAvaTagHandle();
+}
\ No newline at end of file
diff --git a/Engine/Plugins/Experimental/Avalanche/Source/AvalancheMedia/Private/Playlist/AvalanchePageAssetUtils.h b/Engine/Plugins/Experimental/Avalanche/Source/AvalancheMedia/Private/Playlist/AvalanchePageAssetUtils.h
new file mode 100644
--- /dev/null
+++ b/Engine/Plugins/Experimental/Avalanche/Source/AvalancheMedia/Private/Playlist/AvalanchePageAssetUtils.h
@@ -0,0 +1,25 @@
+// Copyright Epic Games, Inc. All Rights Reserved.
+
+#pragma once
+
+#include "UObject/StrongObjectPtr.h"
+
+class IAvaSceneInterface;
+class IAvaSequenceProvider;
+class UAvaTransitionTree;
+class UObject;
+struct FAvaTagHandle;
+
+/**
+ * Common utility functions to query information about loaded Avalanche assets.
+ * This is intended to work on "inactive" assets, i.e. either source assets
+ * directly or managed ones that don't have an active world.
+ * For "runtime" (fully initialized) instances, UAvalanchePlayable should be used instead.
+ */ 
+class FAvalanchePageAssetUtils
+{
+public:
+	static const IAvaSceneInterface* GetSceneInterface(const UObject* InLoadedAsset);
+	static const UAvaTransitionTree* FindTransitionTree(const IAvaSceneInterface* InSceneInterface);
+	static FAvaTagHandle GetTransitionLayerTag(const UAvaTransitionTree* InTransitionTree);
+};
\ No newline at end of file
