[MotionDesign]

diff --git a/Engine/Plugins/Experimental/Avalanche/Source/AvalancheMedia/Private/Playlist/AvalanchePagePlayer.cpp b/Engine/Plugins/Experimental/Avalanche/Source/AvalancheMedia/Private/Playlist/AvalanchePagePlayer.cpp
new file mode 100644
--- /dev/null
+++ b/Engine/Plugins/Experimental/Avalanche/Source/AvalancheMedia/Private/Playlist/AvalanchePagePlayer.cpp
@@ -0,0 +1,429 @@
+// Copyright Epic Games, Inc. All Rights Reserved.
+
+#include "Playlist/AvalanchePagePlayer.h"
+
+#include "AvalancheMediaSettings.h"
+#include "Playback/AvaMediaPlaybackManager.h"
+#include "Playback/AvalanchePlayable.h"
+#include "Playback/AvaMediaPlaybackUtils.h"
+#include "Playlist/AvalanchePlaylist.h"
+
+namespace UE::AvaMedia::Rundown::PagePlayer::Private
+{
+	static void PushAllAnimations(UAvalanchePlayback* InPlaybackObject, const UAvalanchePlaylist* InPlaylist, const FAvalanchePage& InPage
+		, const FString& InChannelName, EAvaMediaAnimAction InAnimAction)
+	{
+		const int32 NumTemplates = InPage.GetNumTemplates(InPlaylist);
+		for (int32 TemplateIndex = 0; TemplateIndex < NumTemplates; ++TemplateIndex)
+		{
+			// Note: default FAnimPlaySettings (with NAME_None) will play all animations.
+			InPlaybackObject->PushAnimationCommand(InPage.GetAvalancheAssetPath(InPlaylist, TemplateIndex), InChannelName,
+				InAnimAction, FAnimPlaySettings());
+		}
+	}
+
+	static void PushCameraCut(UAvalanchePlayback* InPlaybackObject, const UAvalanchePlaylist* InPlaylist, const FAvalanchePage& InPage
+		, const FString& InChannelName)
+	{
+		InPlaybackObject->PushAnimationCommand(InPage.GetAvalancheAssetPath(InPlaylist), InChannelName,
+			EAvaMediaAnimAction::CameraCut, FAnimPlaySettings());
+	}
+}
+
+UAvaRundownPlaybackInstancePlayer::UAvaRundownPlaybackInstancePlayer() = default;
+UAvaRundownPlaybackInstancePlayer::~UAvaRundownPlaybackInstancePlayer() = default;
+
+bool UAvaRundownPlaybackInstancePlayer::Load(const UAvalanchePagePlayer& InPagePlayer, const UAvalanchePlaylist* InPlaylist, const FAvalanchePage& InPage, int32 InSubPageIndex, const FGuid& InInstanceId)
+{
+	SourceAssetPath = InPage.GetAvalancheAssetPath(InPlaylist, InSubPageIndex);
+	TransitionLayer = InPage.GetTransitionLayer(InPlaylist, InSubPageIndex);
+	
+	FAvaMediaPlaybackManager& Manager = InPlaylist->GetPlaybackManager();
+	PlaybackInstance = Manager.AcquireOrLoadPlaybackInstance(SourceAssetPath, InPagePlayer.ChannelName);
+	Playback = PlaybackInstance ? PlaybackInstance->GetPlayback() : nullptr;
+
+	// If restoring from a remote instance.
+	if (InInstanceId.IsValid())
+	{
+		PlaybackInstance->SetInstanceId(InInstanceId);
+	}
+
+	// Setup user instance data to be able to track this page.
+	if (PlaybackInstance)
+	{
+		UAvalanchePagePlayer::SetInstanceUserDataFromPage(*PlaybackInstance, InPage);
+	}
+
+	if (Playback && InPagePlayer.bIsPreview)
+	{
+		Playback->SetPreviewChannelName(InPagePlayer.ChannelFName);
+	}
+
+	return IsLoaded();
+}
+
+bool UAvaRundownPlaybackInstancePlayer::IsLoaded() const
+{
+	return Playback != nullptr;
+}
+
+void UAvaRundownPlaybackInstancePlayer::Play(const UAvalanchePagePlayer& InPagePlayer, const UAvalanchePlaylist* InPlaylist, EAvaPlayType InPlayType, bool bInIsUsingTransitionLogic)
+{
+	using namespace UE::AvaMedia::Rundown::PagePlayer::Private;
+	
+	if (!Playback || !InPlaylist)
+	{
+		return;
+	}
+
+	const bool bPlaybackObjectWasPlaying = Playback->IsPlaying();
+	
+	if (!Playback->IsPlaying())
+	{
+		Playback->Play();
+	}
+
+	const FAvalanchePage& Page = InPlaylist->GetPage(InPagePlayer.PageId);
+
+	// When not using transition logic, we start all animations manually.
+	if (!bInIsUsingTransitionLogic)
+	{
+		Playback->PushRemoteControlValues(SourceAssetPath, InPagePlayer.ChannelName, MakeShared<FAvalancheRemoteControlValues>(Page.GetRemoteControlValues()));
+
+		// TODO: play type should be transmitted to the transition behavior.
+		const EAvaMediaAnimAction AnimAction = (InPlayType == EAvaPlayType::PreviewFromFrame)
+			? EAvaMediaAnimAction::PreviewFrame
+			: EAvaMediaAnimAction::Play;
+	
+		// Play all the animations with the page's animation settings.
+		PushAllAnimations(Playback, InPlaylist, Page, InPagePlayer.ChannelName, AnimAction);
+	}
+	
+	if (bPlaybackObjectWasPlaying)
+	{
+		PushCameraCut(Playback, InPlaylist, Page, InPagePlayer.ChannelName);
+	}
+}
+
+bool UAvaRundownPlaybackInstancePlayer::IsPlaying() const
+{
+	return Playback != nullptr && Playback->IsPlaying();
+}
+
+bool UAvaRundownPlaybackInstancePlayer::Continue(const FString& InChannelName)
+{
+	if (Playback && Playback->IsPlaying())
+	{
+		// Animation command, within this playback, needs channel for now.
+		const FAnimPlaySettings AnimSettings;	// Note: Leaving the name to None means the action apply to all animations.
+		Playback->PushAnimationCommand(SourceAssetPath, InChannelName, EAvaMediaAnimAction::Continue, AnimSettings);
+		return true;
+	}
+	return false;
+}
+
+bool UAvaRundownPlaybackInstancePlayer::Stop()
+{
+	if (!Playback)
+	{
+		return false;
+	}
+	
+	const bool bUnload = !UAvalancheMediaSettings::Get().bKeepPagesLoaded;
+	
+	if (Playback->IsPlaying())
+	{
+		// Propagate the unload options in case this object is playing remote.
+		const EAvaPlaybackStopOptions PlaybackStopOptions = bUnload ?
+			EAvaPlaybackStopOptions::Default | EAvaPlaybackStopOptions::Unload : EAvaPlaybackStopOptions::Default;
+		Playback->Stop(PlaybackStopOptions);
+	}
+
+	if (PlaybackInstance)
+	{
+		// Unload the local object as well.
+		if (bUnload)
+		{
+			PlaybackInstance->Unload();
+		}
+		else
+		{
+			PlaybackInstance->Recycle();
+		}
+	}
+
+	Playback = nullptr;
+	PlaybackInstance.Reset();
+	return true;
+}
+
+bool UAvaRundownPlaybackInstancePlayer::HasPlayable(const UAvalanchePlayable* InPlayable) const
+{
+	return Playback && Playback->HasPlayable(InPlayable);
+}
+
+UAvalanchePlayable* UAvaRundownPlaybackInstancePlayer::GetFirstPlayable() const
+{
+	return Playback ? Playback->GetFirstPlayable() : nullptr;
+}
+
+UAvalanchePagePlayer* UAvaRundownPlaybackInstancePlayer::GetPagePlayer() const
+{
+	return Cast<UAvalanchePagePlayer>(GetOuter());
+}
+
+void UAvaRundownPlaybackInstancePlayer::SetPagePlayer(UAvalanchePagePlayer* InPagePlayer)
+{
+	LowLevelRename(GetFName(), InPagePlayer);
+}
+
+UAvalanchePagePlayer::UAvalanchePagePlayer()
+{
+	UAvalanchePlayable::OnSequenceEvent().AddUObject(this, &UAvalanchePagePlayer::HandleOnPlayableSequenceEvent);
+}
+
+UAvalanchePagePlayer::~UAvalanchePagePlayer()
+{
+	UAvalanchePlayable::OnSequenceEvent().RemoveAll(this);
+}
+
+bool UAvalanchePagePlayer::Initialize(UAvalanchePlaylist* InPlaylist, const FAvalanchePage& InPage, bool bInIsPreview, const FName& InPreviewChannel)
+{
+	if (!InPlaylist)
+	{
+		UE_LOG(LogAvaPlaylist, Error, TEXT("UAvalanchePagePlayer::Initialize: Invalid playlist."));
+		return false;
+	}
+
+	if (!InPage.IsValidPage())
+	{
+		UE_LOG(LogAvaPlaylist, Error, TEXT("UAvalanchePagePlayer::Initialize: Invalid page."));
+		return false;
+	}
+	
+	checkf(InstancePlayers.IsEmpty(), TEXT("Can't initialize a page player if already loaded or playing."));
+
+	PlaylistWeak = InPlaylist;
+	PageId = InPage.GetPageId();
+	bIsPreview = bInIsPreview;
+	ChannelFName = bIsPreview ? InPreviewChannel : InPage.GetChannelName();
+	ChannelName = ChannelFName.ToString();
+	return true;
+}
+
+bool UAvalanchePagePlayer::InitializeAndLoad(UAvalanchePlaylist* InPlaylist, const FAvalanchePage& InPage, bool bInIsPreview, const FName& InPreviewChannel)
+{
+	if (!Initialize(InPlaylist, InPage, bInIsPreview, InPreviewChannel))
+	{
+		return false;
+	}
+	
+	const int32 NumTemplates = InPage.GetNumTemplates(InPlaylist);
+	for (int32 SubPageIndex = 0; SubPageIndex < NumTemplates; ++SubPageIndex)
+	{
+		CreateAndLoadInstancePlayer(InPlaylist, InPage, SubPageIndex, FGuid());
+	}
+	return InstancePlayers.Num() > 0;
+}
+
+UAvaRundownPlaybackInstancePlayer* UAvalanchePagePlayer::LoadInstancePlayer(int32 InSubPageIndex, const FGuid& InInstanceId)
+{
+	const UAvalanchePlaylist* Playlist = PlaylistWeak.Get();
+	if (!Playlist)
+	{
+		UE_LOG(LogAvaPlaylist, Error, TEXT("UAvalanchePagePlayer::LoadSubPage: Playlist is no longuer valid."));
+		return nullptr;
+	}
+
+	const FAvalanchePage& Page = Playlist->GetPage(PageId);
+	if (!Page.IsValidPage())
+	{
+		UE_LOG(LogAvaPlaylist, Error, TEXT("UAvalanchePagePlayer::LoadSubPage: Invalid pageId %d."), PageId);
+		return nullptr;
+	}
+
+	return CreateAndLoadInstancePlayer(Playlist, Page, InSubPageIndex, InInstanceId);
+}
+
+void UAvalanchePagePlayer::AddInstancePlayer(UAvaRundownPlaybackInstancePlayer* InExistingInstancePlayer)
+{
+	// Remove from previous player.
+	if (UAvalanchePagePlayer* PreviousPagePlayer = InExistingInstancePlayer->GetPagePlayer())
+	{
+		PreviousPagePlayer->RemoveInstancePlayer(InExistingInstancePlayer);
+	}
+
+	InstancePlayers.Add(InExistingInstancePlayer);
+	InExistingInstancePlayer->SetPagePlayer(this);
+}
+
+bool UAvalanchePagePlayer::IsLoaded() const
+{
+	for (const TObjectPtr<UAvaRundownPlaybackInstancePlayer>& InstancePlayer : InstancePlayers)
+	{
+		if (InstancePlayer && InstancePlayer->IsLoaded())
+		{
+			return true;
+		}
+	}
+	return false;
+}
+
+bool UAvalanchePagePlayer::Play(EAvaPlayType InPlayType, bool bInIsUsingTransitionLogic)
+{
+	const UAvalanchePlaylist* Playlist = PlaylistWeak.Get();
+	if (!Playlist)
+	{
+		return false;
+	}
+
+	bool bIsPlaying = false;
+	for (const TObjectPtr<UAvaRundownPlaybackInstancePlayer>& InstancePlayer : InstancePlayers)
+	{
+		InstancePlayer->Play(*this, Playlist, InPlayType, bInIsUsingTransitionLogic);
+		bIsPlaying |= InstancePlayer->IsPlaying();
+	}
+
+	return bIsPlaying;
+}
+
+bool UAvalanchePagePlayer::IsPlaying() const
+{
+	for (const TObjectPtr<UAvaRundownPlaybackInstancePlayer>& InstancePlayer : InstancePlayers)
+	{
+		if (InstancePlayer && InstancePlayer->IsPlaying())
+		{
+			return true;
+		}
+	}
+	return false;
+}
+
+bool UAvalanchePagePlayer::Continue()
+{
+	bool bSuccess = false;
+	for (const TObjectPtr<UAvaRundownPlaybackInstancePlayer>& InstancePlayer : InstancePlayers)
+	{
+		bSuccess |= InstancePlayer->Continue(ChannelName);
+	}
+	return bSuccess;
+}
+
+bool UAvalanchePagePlayer::Stop()
+{
+	bool bSuccess = false;
+	for (const TObjectPtr<UAvaRundownPlaybackInstancePlayer>& InstancePlayer : InstancePlayers)
+	{
+		bSuccess |= InstancePlayer->Stop();
+	}
+	if (PlaylistWeak.IsValid())
+	{
+		PlaylistWeak->NotifyPageStopped(PageId);
+	}
+	return bSuccess;
+}
+
+int32 UAvalanchePagePlayer::GetPageIdFromInstanceUserData(const FString& InUserData)
+{
+	FString PageIdString;
+	if (FParse::Value(*InUserData, TEXT("PageId="), PageIdString))
+	{
+		return FCString::Atoi(*PageIdString);
+	}
+	return FAvalanchePage::InvalidPageId;
+}
+
+void UAvalanchePagePlayer::SetInstanceUserDataFromPage(FAvaMediaPlaybackInstance& InPlaybackInstance, const FAvalanchePage& InPage)
+{
+	InPlaybackInstance.SetInstanceUserData(FString::Printf(TEXT("PageId=%d"), InPage.GetPageId()));
+}
+
+bool UAvalanchePagePlayer::HasPlayable(const UAvalanchePlayable* InPlayable) const
+{
+	for (const TObjectPtr<UAvaRundownPlaybackInstancePlayer>& InstancePlayer : InstancePlayers)
+	{
+		if (InstancePlayer->HasPlayable(InPlayable))
+		{
+			return true;
+		}
+	}
+	return false;
+}
+
+UAvaRundownPlaybackInstancePlayer* UAvalanchePagePlayer::FindInstancePlayerForPlayable(const UAvalanchePlayable* InPlayable) const
+{
+	for (const TObjectPtr<UAvaRundownPlaybackInstancePlayer>& InstancePlayer : InstancePlayers)
+	{
+		if (InstancePlayer && InstancePlayer->HasPlayable(InPlayable))
+		{
+			return InstancePlayer;
+		}
+	}
+	return nullptr;
+}
+
+UAvaRundownPlaybackInstancePlayer* UAvalanchePagePlayer::FindInstancePlayerByInstanceId(const FGuid& InInstanceId) const
+{
+	for (const TObjectPtr<UAvaRundownPlaybackInstancePlayer>& InstancePlayer : InstancePlayers)
+	{
+		if (InstancePlayer && InstancePlayer->GetPlaybackInstanceId() == InInstanceId)
+		{
+			return InstancePlayer;
+		}
+	}
+	return nullptr;
+}
+
+UAvaRundownPlaybackInstancePlayer* UAvalanchePagePlayer::FindInstancePlayerByAssetPath(const FSoftObjectPath& InAssetPath) const
+{
+	for (const TObjectPtr<UAvaRundownPlaybackInstancePlayer>& InstancePlayer : InstancePlayers)
+	{
+		if (InstancePlayer && InstancePlayer->SourceAssetPath == InAssetPath)
+		{
+			return InstancePlayer;
+		}
+	}
+	return nullptr;
+}
+
+UAvaRundownPlaybackInstancePlayer* UAvalanchePagePlayer::CreateAndLoadInstancePlayer(const UAvalanchePlaylist* InPlaylist, const FAvalanchePage& InPage, int32 InSubPageIndex, const FGuid& InInstanceId)
+{
+	UAvaRundownPlaybackInstancePlayer* InstancePlayer = NewObject<UAvaRundownPlaybackInstancePlayer>(this);
+	if (InstancePlayer->Load(*this, InPlaylist, InPage, InSubPageIndex, InInstanceId))
+	{
+		InstancePlayers.Add(InstancePlayer);
+		return InstancePlayer;
+	}
+
+	return nullptr;
+}
+
+void UAvalanchePagePlayer::RemoveInstancePlayer(UAvaRundownPlaybackInstancePlayer* InInstancePlayer)
+{
+	InstancePlayers.Remove(InInstancePlayer);
+}
+
+void UAvalanchePagePlayer::HandleOnPlayableSequenceEvent(UAvalanchePlayable* InPlayable, const FName& SequenceName, EAvalanchePlayableSequenceEventType InEventType)
+{
+	// Check that this is the playable for this page player.
+	if (!HasPlayable(InPlayable))
+	{
+		return;
+	}
+	
+	// Notify the playlist.
+	if (PlaylistWeak.IsValid())
+	{
+		using namespace UE::AvaMediaPlayback::Utils;
+		if (InEventType == EAvalanchePlayableSequenceEventType::Started)
+		{
+			UE_LOG(LogAvaPlaylist, Verbose, TEXT("%s Playlist Page %d: Sequence Started \"%s\"."), *GetBriefFrameInfo(), PageId, *SequenceName.ToString());
+		}
+
+		if (InEventType == EAvalanchePlayableSequenceEventType::Finished)
+		{
+			UE_LOG(LogAvaPlaylist, Verbose, TEXT("%s Playlist Page %d: Sequence Finished \"%s\"."), *GetBriefFrameInfo(), PageId, *SequenceName.ToString());
+			PlaylistWeak->NotifyPageSequenceFinished(PageId);
+		}
+	}
+}
diff --git a/Engine/Plugins/Experimental/Avalanche/Source/AvalancheMedia/Private/Playlist/AvalanchePlaylist.cpp b/Engine/Plugins/Experimental/Avalanche/Source/AvalancheMedia/Private/Playlist/AvalanchePlaylist.cpp
new file mode 100644
--- /dev/null
+++ b/Engine/Plugins/Experimental/Avalanche/Source/AvalancheMedia/Private/Playlist/AvalanchePlaylist.cpp
@@ -0,0 +1,2053 @@
+// Copyright Epic Games, Inc. All Rights Reserved.
+
+#include "Playlist/AvalanchePlaylist.h"
+
+#include "AvalancheBroadcast.h"
+#include "AvalancheMediaSettings.h"
+#include "IAvaMediaModule.h"
+#include "Playback/AvaMediaPlaybackManager.h"
+#include "Playback/AvalanchePlayable.h"
+#include "Playback/AvalanchePlayback.h"
+#include "Playlist/AvaPlaylistPageLoadingManager.h"
+#include "Playlist/AvaPlaylistPlaybackClientWatcher.h"
+#include "Playlist/AvalancheManagedInstanceCache.h"
+#include "Playlist/Transition/AvalanchePageTransition.h"
+#include "Playlist/Transition/AvalanchePageTransitionBuilder.h"
+
+#if WITH_EDITOR
+#include "Editor.h"
+#endif
+
+DEFINE_LOG_CATEGORY(LogAvaPlaylist);
+
+#define LOCTEXT_NAMESPACE "AvalanchePlaylist"
+
+void FAvalanchePageCollection::Empty(UAvalanchePlaylist* InPlaylist)
+{
+	TArray<int32> PageIds;
+	if (OnPageListChanged.IsBound() && Pages.IsEmpty() == false)
+	{
+		PageIds.Reserve(Pages.Num());
+		for (const FAvalanchePage& Page : Pages)
+		{
+			PageIds.Add(Page.GetPageId());
+		}
+	}
+	
+	Pages.Empty();
+	PageIndices.Empty();
+	
+	if (PageIds.IsEmpty() == false)
+	{
+		OnPageListChanged.Broadcast({InPlaylist, EAvaPageListChange::RemovedPages, PageIds});
+	}
+}
+
+const FAvaPageListReference UAvalanchePlaylist::TemplatePageList = {EAvaPageListType::Template, -2};
+const FAvaPageListReference UAvalanchePlaylist::InstancePageList = {EAvaPageListType::Instance, -1};
+const FAvalancheSubList UAvalanchePlaylist::InvalidSubList;
+
+UAvalanchePlaylist::UAvalanchePlaylist()
+{
+	if (!HasAnyFlags(RF_ClassDefaultObject))
+	{
+#if WITH_EDITOR
+		FEditorDelegates::PrePIEEnded.AddUObject(this, &UAvalanchePlaylist::NotifyPIEEnded);
+#endif
+	}
+}
+
+UAvalanchePlaylist::~UAvalanchePlaylist() = default;
+
+int32 UAvalanchePlaylist::GenerateUniquePageId(int32 InReferencePageId, int32 InIncrement) const
+{
+	if (InIncrement == 0)
+	{
+		InIncrement = 1;
+	}
+
+	// Search space must be zero-positive.
+	int32 UniquePageId = FMath::Max(0, InReferencePageId);
+
+	// Search a unique id in the given direction.
+	while (!IsPageIdUnique(UniquePageId))
+	{
+		UniquePageId += InIncrement;
+		
+		// End of the search space is reached, start in the other direction from initial value.
+		if (UniquePageId < 0 && InIncrement < 0)
+		{
+			return GenerateUniquePageId(InReferencePageId, -InIncrement);
+		}
+	}
+
+	return UniquePageId;
+}
+
+int32 UAvalanchePlaylist::GenerateUniquePageId(const FAvaPageIdGeneratorParams& InParams) const
+{
+	return GenerateUniquePageId(InParams.ReferenceId, InParams.Increment);
+}
+
+void UAvalanchePlaylist::RefreshPageIndices()
+{	
+	TemplatePages.RefreshPageIndices();
+	InstancedPages.RefreshPageIndices();
+}
+
+namespace UE::AvalancheMedia::Playlist::Private
+{
+	void OnPostLoadPages(TArray<FAvalanchePage>& InPages)
+	{
+		for (FAvalanchePage& Page : InPages)
+		{
+			Page.PostLoad();
+		}
+	}
+}
+
+void UAvalanchePlaylist::BeginDestroy()
+{
+	Super::BeginDestroy();
+
+#if WITH_EDITOR
+	FEditorDelegates::PrePIEEnded.RemoveAll(this);
+#endif
+}
+
+void UAvalanchePlaylist::PostLoad()
+{
+	UObject::PostLoad();
+
+	if (!Pages_DEPRECATED.IsEmpty())
+	{
+		AddPagesFromTemplates(AddTemplates(Pages_DEPRECATED));
+		Pages_DEPRECATED.Empty();
+	}
+
+	UE::AvalancheMedia::Playlist::Private::OnPostLoadPages(TemplatePages.Pages);
+	UE::AvalancheMedia::Playlist::Private::OnPostLoadPages(InstancedPages.Pages);
+	
+	RefreshPageIndices();
+}
+
+#if WITH_EDITOR
+void UAvalanchePlaylist::PostEditUndo()
+{
+	UObject::PostEditUndo();
+	
+	//Force Refresh for any Undo
+	RefreshPageIndices();
+	GetOnTemplatePageListChanged().Broadcast({this, EAvaPageListChange::All, {}});
+	GetOnInstancedPageListChanged().Broadcast({this, EAvaPageListChange::All, {}});
+}
+#endif
+
+bool UAvalanchePlaylist::IsEmpty() const
+{
+	return TemplatePages.Pages.IsEmpty() && InstancedPages.Pages.IsEmpty();
+}
+
+bool UAvalanchePlaylist::Empty()
+{
+	if (IsPlaying())
+	{
+		return false;
+	}
+
+	// Since we are about to delete the sublists,
+	// make sure we return the active page list to something not deleted.
+	if (HasActiveSubList())
+	{
+		SetActivePageList(InstancePageList);
+	}
+	
+	SubLists.Empty();
+	InstancedPages.Empty(this);
+	TemplatePages.Empty(this);
+	
+	return true;
+}
+
+int32 UAvalanchePlaylist::AddTemplateInternal(const FAvaPageIdGeneratorParams& InIdGeneratorParams, const TFunctionRef<bool(FAvalanchePage&)> InSetupTemplateFunction)
+{
+	if (!CanAddPage())
+	{
+		return FAvalanchePage::InvalidPageId;
+	}
+	const int32 TemplateId = GenerateUniquePageId(InIdGeneratorParams);
+	FAvalanchePage NewTemplate(TemplateId);
+	
+	if (!InSetupTemplateFunction(NewTemplate))
+	{
+		return FAvalanchePage::InvalidPageId;
+	}
+	
+	TemplatePages.Pages.Emplace(MoveTemp(NewTemplate));
+	
+	RefreshPageIndices();
+
+	GetOnTemplatePageListChanged().Broadcast({this, EAvaPageListChange::AddedPages, {TemplateId}});
+
+	return TemplateId;
+}
+
+int32 UAvalanchePlaylist::AddTemplate(const FAvaPageIdGeneratorParams& InIdGeneratorParams)
+{
+	return AddTemplateInternal(InIdGeneratorParams, [](FAvalanchePage&){return true;});
+}
+
+int32 UAvalanchePlaylist::AddComboTemplate(const TArray<int32>& InTemplateIds, const FAvaPageIdGeneratorParams& InIdGeneratorParams)
+{
+	return AddTemplateInternal(InIdGeneratorParams, [&InTemplateIds](FAvalanchePage& InNewTemplate)
+	{
+		InNewTemplate.CombinedTemplateIds = InTemplateIds;
+		return true;
+	});
+}
+
+TArray<int32> UAvalanchePlaylist::AddTemplates(const TArray<FAvalanchePage>& InSourceTemplates)
+{
+	if (!CanAddPage() || InSourceTemplates.IsEmpty())
+	{
+		return {};
+	}
+
+	TArray<int32> OutTemplateIds;
+	OutTemplateIds.Reserve(InSourceTemplates.Num());
+
+	for (const FAvalanchePage& SourceTemplate : InSourceTemplates)
+	{
+		// Try to preserve the template id from the source.
+		int32 NewTemplateId = GenerateUniquePageId(SourceTemplate.GetPageId());
+
+		// Add to template list.
+		const int32 Index = TemplatePages.Pages.Add(SourceTemplate);
+		TemplatePages.Pages[Index].PageId = NewTemplateId;
+		TemplatePages.Pages[Index].TemplateId = FAvalanchePage::InvalidPageId;
+		TemplatePages.PageIndices.Emplace(NewTemplateId, Index);
+		
+		OutTemplateIds.Add(NewTemplateId);
+	}
+
+	if (!OutTemplateIds.IsEmpty())
+	{
+		GetOnTemplatePageListChanged().Broadcast({this, EAvaPageListChange::AddedPages, OutTemplateIds});
+	}
+
+	return OutTemplateIds;
+}
+
+TArray<int32> UAvalanchePlaylist::AddPagesFromTemplates(const TArray<int32>& InTemplateIds)
+{
+	TArray<int32> OutPageIds;
+	OutPageIds.Reserve(InTemplateIds.Num());
+	
+	FAvaPageIdGeneratorParams IdGeneratorParams;
+
+	// Special id generation case: start from the last page id.
+	{
+		int32 LastInstancedPageId = 0;
+		for (const FAvalanchePage& Page : InstancedPages.Pages)
+		{
+			LastInstancedPageId = FMath::Max(LastInstancedPageId, Page.GetPageId());
+		}
+		IdGeneratorParams.ReferenceId = LastInstancedPageId;
+	}
+
+	for (const int32& TemplateId : InTemplateIds)
+	{
+		const int32 NewId = AddPageFromTemplateInternal(TemplateId, IdGeneratorParams);
+		if (NewId != FAvalanchePage::InvalidPageId)
+		{
+			OutPageIds.Add(NewId);
+			IdGeneratorParams.ReferenceId = NewId;
+		}
+	}
+
+	RefreshPageIndices();
+	GetOnInstancedPageListChanged().Broadcast({this, EAvaPageListChange::AddedPages, OutPageIds});
+
+	return OutPageIds;
+}
+
+int32 UAvalanchePlaylist::AddPageFromTemplate(int32 InTemplateId, const FAvaPageIdGeneratorParams& InIdGeneratorParams, const FAvaPageInsertPosition& InInsertAt)
+{
+	const int32 NewId = AddPageFromTemplateInternal(InTemplateId, InIdGeneratorParams, InInsertAt);
+	
+	if (NewId != FAvalanchePage::InvalidPageId)
+	{
+		RefreshPageIndices();
+		GetOnInstancedPageListChanged().Broadcast({this, EAvaPageListChange::AddedPages, {NewId}});
+	}
+
+	return NewId;
+}
+
+bool UAvalanchePlaylist::CanAddPage() const
+{
+	// Pages can always be added. This is needed for live editing playlists.
+	// Because of this, having pointers to pages is risky.
+	// Pages should always be referred to by page id and de-referenced only when needed.
+	return true;
+}
+
+bool UAvalanchePlaylist::CanChangePageOrder() const
+{
+	return !IsPlaying();
+}
+
+bool UAvalanchePlaylist::ChangePageOrder(const FAvaPageListReference& InPageListReference, const TArray<int32>& InPageIndices)
+{
+	TSet<int32> MovedIndices;
+
+	// Templates & Instances
+	if (InPageListReference.Type != EAvaPageListType::View)
+	{
+		FAvalanchePageCollection& Collection = InPageListReference.Type == EAvaPageListType::Template ? TemplatePages : InstancedPages;
+		TArray<FAvalanchePage> NewPages;
+		NewPages.Reserve(Collection.Pages.Num());
+
+		for (int32 PageIndex : InPageIndices)
+		{
+			NewPages.Add(MoveTemp(Collection.Pages[PageIndex]));
+			MovedIndices.Add(PageIndex);
+		}
+
+		// Make sure all pages were moved.
+		for (int32 PageIndex = 0; PageIndex < Collection.Pages.Num(); ++PageIndex)
+		{
+			if (!MovedIndices.Contains(PageIndex))
+			{
+				NewPages.Add(MoveTemp(Collection.Pages[PageIndex]));
+			}
+		}
+
+		Collection.Pages = NewPages;
+		Collection.PageIndices.Empty();
+		RefreshPageIndices();
+
+		Collection.OnPageListChanged.Broadcast({this, EAvaPageListChange::ReorderedPageView, {}});
+
+		return true;
+	}
+
+	if (SubLists.IsValidIndex(InPageListReference.SubListIndex))
+	{
+		FAvalancheSubList& SubList = SubLists[InPageListReference.SubListIndex];
+		TArray<int32> NewIndices;
+		NewIndices.Reserve(SubList.PageIds.Num());
+
+		for (int32 PageIndex : InPageIndices)
+		{
+			NewIndices.Add(SubList.PageIds[PageIndex]);
+			MovedIndices.Add(PageIndex);
+		}
+
+		// Make sure all pages were moved.
+		for (int32 PageIndex = 0; PageIndex < SubList.PageIds.Num(); ++PageIndex)
+		{
+			if (!MovedIndices.Contains(PageIndex))
+			{
+				NewIndices.Add(SubList.PageIds[PageIndex]);
+			}
+		}
+
+		SubList.PageIds = NewIndices;
+		SubList.OnPageListChanged.Broadcast({this, EAvaPageListChange::ReorderedPageView, {}});
+
+		return true;
+	}
+
+	return false;
+}
+
+bool UAvalanchePlaylist::RemovePage(int32 InPageId)
+{
+	return RemovePages({InPageId}) > 0;
+}
+
+bool UAvalanchePlaylist::CanRemovePage(int32 InPageId) const
+{
+	return CanRemovePages({InPageId});
+}
+
+int32 UAvalanchePlaylist::RemovePages(const TArray<int32>& InPageIds)
+{
+	if (!CanRemovePages(InPageIds))
+	{
+		return 0;
+	}
+	
+	TArray<int32> SortedPageIds(InPageIds);
+	SortedPageIds.Sort();
+
+	int32 RemovedCount = 0;
+	bool bRemovedTemplate = false;
+	bool bRemovedInstanced = false;
+
+	// Find the instanced page ids to remove
+	TArray<int32> TemplatesIndicesToRemove;
+	TSet<int32> InstancesToRemove;
+
+	for (int32 PageIdx = SortedPageIds.Num() - 1; PageIdx >= 0; --PageIdx)
+	{
+		const int32 PageId = SortedPageIds[PageIdx];
+		const int32* TemplateIdx = TemplatePages.PageIndices.Find(PageId);
+		const int32* InstanceIdx = InstancedPages.PageIndices.Find(PageId);
+
+		if (TemplateIdx)
+		{
+			if (TemplatePages.Pages[*TemplateIdx].Instances.IsEmpty() == false)
+			{
+				InstancesToRemove.Append(TemplatePages.Pages[*TemplateIdx].Instances);
+			}
+
+			// Double check, just in case.
+			for (const FAvalanchePage& Page : InstancedPages.Pages)
+			{
+				if (Page.TemplateId == PageId)
+				{
+					InstancesToRemove.Add(Page.GetPageId());
+				}
+			}
+
+			TemplatesIndicesToRemove.Add(*TemplateIdx);
+			TemplatePages.PageIndices.Remove(PageId);
+		}
+		else if (InstanceIdx)
+		{
+			InstancesToRemove.Add(PageId);
+
+			const FAvalanchePage& Instance = InstancedPages.Pages[*InstanceIdx];
+
+			if (const int32* InstanceTemplateIdx = TemplatePages.PageIndices.Find(Instance.TemplateId))
+			{
+				TemplatePages.Pages[*InstanceTemplateIdx].Instances.Remove(PageId);
+			}
+		}
+	}
+
+	if (TemplatesIndicesToRemove.IsEmpty() == false)
+	{
+		TemplatesIndicesToRemove.Sort();
+
+		for (int32 TemplateIdx = TemplatesIndicesToRemove.Num() - 1; TemplateIdx >= 0; --TemplateIdx)
+		{
+			TemplatePages.Pages.RemoveAt(TemplatesIndicesToRemove[TemplateIdx]);
+			++RemovedCount;
+			bRemovedTemplate = true;
+		}
+	}
+
+	if (InstancesToRemove.IsEmpty() == false)
+	{
+		// Process list from highest to lowest so we don't change future indices while removing.
+		TArray<int32> InstancesToRemoveIndices;
+		InstancesToRemoveIndices.Reserve(InstancesToRemove.Num());
+
+		for (int32 InstanceToRemove : InstancesToRemove)
+		{
+			if (int32* InstanceIndexPtr = InstancedPages.PageIndices.Find(InstanceToRemove))
+			{
+				InstancesToRemoveIndices.Add(*InstanceIndexPtr);
+				InstancedPages.PageIndices.Remove(InstanceToRemove);
+			}
+		}
+
+		InstancesToRemoveIndices.Sort();
+
+		for (int32 RemoveIdx = InstancesToRemoveIndices.Num() - 1; RemoveIdx >= 0; --RemoveIdx)
+		{
+			const int32 InstanceToRemoveIdx = InstancesToRemoveIndices[RemoveIdx];
+			InstancedPages.Pages.RemoveAt(InstanceToRemoveIdx);
+			++RemovedCount;
+			bRemovedInstanced = true;
+		}
+	}
+		
+	if (RemovedCount == 0)
+	{
+		return 0;
+	}
+
+	RefreshPageIndices();
+
+	if (bRemovedTemplate)
+	{
+		GetOnTemplatePageListChanged().Broadcast({this, EAvaPageListChange::RemovedPages, {InPageIds}});
+	}
+
+	if (bRemovedInstanced)
+	{
+		GetOnInstancedPageListChanged().Broadcast({this, EAvaPageListChange::RemovedPages, {InPageIds}});
+	}
+
+	for (FAvalancheSubList& SubList : SubLists)
+	{
+		bool bFoundInstance = false;
+
+		for (TArray<int32>::TIterator Iter(SubList.PageIds); Iter; ++Iter)
+		{
+			if (InstancesToRemove.Contains(*Iter))
+			{
+				Iter.RemoveCurrent();
+				bFoundInstance = true;
+			}
+		}
+
+		if (bFoundInstance)
+		{
+			SubList.OnPageListChanged.Broadcast({this, EAvaPageListChange::RemovedPages, {InPageIds}});
+		}
+	}
+
+	return RemovedCount;
+}
+
+bool UAvalanchePlaylist::CanRemovePages(const TArray<int32>& InPageIds) const
+{
+	for (int32 PageId : InPageIds)
+	{
+		if (IsPagePlayingOrPreviewing(PageId))
+		{
+			return false;
+		}
+	}
+	return InPageIds.Num() > 0;
+}
+
+bool UAvalanchePlaylist::RenumberPageId(int32 InPageId, int32 InNewPageId)
+{
+	if (!CanRenumberPageId(InPageId, InNewPageId))
+	{
+		return false;
+	}
+
+	FAvalanchePage& Page = GetPage(InPageId);
+	check(Page.IsValidPage());
+
+	Page.PageId = InNewPageId;
+	
+	const int32* TemplateIdx = TemplatePages.PageIndices.Find(InPageId);
+	const int32* InstanceIdx = InstancedPages.PageIndices.Find(InPageId);
+
+	if (TemplateIdx)
+	{
+		bool bFoundInstanceOfTemplate = false;
+
+		for (const int32& InstancePageId : TemplatePages.Pages[*TemplateIdx].Instances)
+		{
+			if (const int32* InstancePageIdx = InstancedPages.PageIndices.Find(InstancePageId))
+			{
+				InstancedPages.Pages[*InstancePageIdx].TemplateId = InNewPageId;
+				bFoundInstanceOfTemplate = true;
+			}
+		}
+
+		// Double check, just in case.
+		for (FAvalanchePage& InstancedPage : InstancedPages.Pages)
+		{
+			if (InstancedPage.TemplateId == InPageId)
+			{
+				InstancedPage.TemplateId = InNewPageId;
+				bFoundInstanceOfTemplate = true;
+
+				// Has become desynced somehow so add it to the template page instance set.
+				TemplatePages.Pages[*TemplateIdx].Instances.Add(InstancedPage.GetPageId());
+			}
+		}
+
+		GetOnTemplatePageListChanged().Broadcast({this, EAvaPageListChange::RenumberedPageId, {InNewPageId}});
+
+		if (bFoundInstanceOfTemplate)
+		{
+			GetOnInstancedPageListChanged().Broadcast({this, EAvaPageListChange::RenumberedPageId, {InNewPageId}});
+		}
+	}
+	else if (InstanceIdx)
+	{
+		GetOnInstancedPageListChanged().Broadcast({this, EAvaPageListChange::RenumberedPageId, {InNewPageId}});
+
+		for (FAvalancheSubList& SubList : SubLists)
+		{
+			int32 Index = SubList.PageIds.Find(InPageId);
+
+			if (Index != INDEX_NONE)
+			{
+				SubList.PageIds[Index] = InNewPageId;
+				SubList.OnPageListChanged.Broadcast({this, EAvaPageListChange::RenumberedPageId, {InNewPageId}});
+			}
+		}
+	}
+	else
+	{
+		return false;
+	}
+
+	RefreshPageIndices();
+
+	return true;
+}
+
+bool UAvalanchePlaylist::CanRenumberPageId(int32 InPageId) const
+{
+	//There must be a valid Page that we will be renumbering
+	const bool bPageIdValid = GetPage(InPageId).IsValidPage();
+
+	return !IsPagePlayingOrPreviewing(InPageId)
+		&& bPageIdValid;
+}
+
+bool UAvalanchePlaylist::CanRenumberPageId(int32 InPageId, int32 InNewPageId) const
+{
+	//There must be a valid Page that we will be renumbering
+	const bool bPageIdValid = GetPage(InPageId).IsValidPage();
+	
+	//Make sure that if we get a Page with New Page Id, it returns a Null Page
+	const bool bNewPageIdAvailable = !GetPage(InNewPageId).IsValidPage();
+
+	return !IsPagePlayingOrPreviewing(InPageId)
+		&& InPageId != InNewPageId
+		&& bPageIdValid
+		&& bNewPageIdAvailable;
+}
+
+bool UAvalanchePlaylist::SetRemoteControlEntityValue(int32 InPageId, const FGuid& InId, const FAvalancheRemoteControlValue& InValue)
+{
+	FAvalanchePage& Page = GetPage(InPageId);
+	if (Page.IsValidPage())
+	{
+		Page.SetRemoteControlEntityValue(InId, InValue);
+		NotifyPageRemoteControlValueChanged(InPageId, EAvaRemoteControlChanges::EntityValues);
+		return true;
+	}
+	return false;
+}
+
+bool UAvalanchePlaylist::SetRemoteControlControllerValue(int32 InPageId, const FGuid& InId, const FAvalancheRemoteControlValue& InValue)
+{
+	FAvalanchePage& Page = GetPage(InPageId);
+	if (Page.IsValidPage())
+	{
+		Page.SetRemoteControlControllerValue(InId, InValue);
+		NotifyPageRemoteControlValueChanged(InPageId, EAvaRemoteControlChanges::ControllerValues);
+		return true;
+	}
+	return false;
+}
+
+EAvaRemoteControlChanges UAvalanchePlaylist::UpdateRemoteControlValues(int32 InPageId, const FAvalancheRemoteControlValues& InRemoteControlValues, bool bInUpdateDefaults)
+{
+	FAvalanchePage& Page = GetPage(InPageId);
+	if (Page.IsValidPage())
+	{
+		const EAvaRemoteControlChanges Changes = Page.UpdateRemoteControlValues(InRemoteControlValues, bInUpdateDefaults);
+		if (Changes != EAvaRemoteControlChanges::None)
+		{
+			NotifyPageRemoteControlValueChanged(InPageId, Changes);
+		}
+		return Changes;
+	}
+	return EAvaRemoteControlChanges::None;
+}
+
+void UAvalanchePlaylist::InvalidateManagedInstanceCacheForPages(const TArray<int32>& InPageIds) const
+{
+	if (!IAvaMediaModule::IsModuleLoaded())
+	{
+		return;
+	}
+	
+	FAvalancheManagedInstanceCache& ManagedInstanceCache = IAvaMediaModule::Get().GetManagedInstanceCache();
+	
+	for (const int32 PageId : InPageIds)
+	{
+		const FAvalanchePage& Page = GetPage(PageId);
+
+		if (Page.IsValidPage())
+		{
+			ManagedInstanceCache.InvalidateNoDelete(Page.GetAvalancheAssetPath(this));
+		}
+	}
+
+	// Delete all invalidated entries immediately.
+	ManagedInstanceCache.FinishPendingActions();
+}
+
+void UAvalanchePlaylist::UpdateAvalancheAssetForPages(const TArray<int32>& InPageIds, bool bInReimportPage)
+{
+	for (const int32 SelectedPageId : InPageIds)
+	{
+		FAvalanchePage& Page = GetPage(SelectedPageId);
+		if (!Page.IsValidPage())
+		{
+			UE_LOG(LogAvaPlaylist, Error, TEXT("Reimport asset failed: page id %d is not valid."), SelectedPageId);
+			continue;
+		}
+		if (!Page.IsTemplate())
+		{
+			UE_LOG(LogAvaPlaylist, Error, TEXT("Reimport asset failed: page id %d is not a template."), SelectedPageId);
+			continue;
+		}				
+		Page.UpdateAvalancheAsset(Page.GetAvalancheAssetPath(this), bInReimportPage);
+	}
+}
+
+const FAvalanchePage& UAvalanchePlaylist::GetPage(int32 InPageId) const
+{
+	if (const int32* TemplateIdx = TemplatePages.PageIndices.Find(InPageId))
+	{
+		return TemplatePages.Pages[*TemplateIdx];
+	}
+
+	if (const int32* InstancedIdx = InstancedPages.PageIndices.Find(InPageId))
+	{
+		return InstancedPages.Pages[*InstancedIdx];
+	}
+
+	return FAvalanchePage::NullPage;
+}
+
+FAvalanchePage& UAvalanchePlaylist::GetPage(int32 InPageId)
+{
+	if (const int32* TemplateIdx = TemplatePages.PageIndices.Find(InPageId))
+	{
+		return TemplatePages.Pages[*TemplateIdx];
+	}
+
+	if (const int32* InstancedIdx = InstancedPages.PageIndices.Find(InPageId))
+	{
+		return InstancedPages.Pages[*InstancedIdx];
+	}
+
+	return FAvalanchePage::NullPage;
+}
+
+const FAvalanchePage& UAvalanchePlaylist::GetNextPage(int32 InPageId, const FAvaPageListReference& InPageListReference) const
+{
+	if (const int32* TemplateIdx = TemplatePages.PageIndices.Find(InPageId))
+	{
+		if (TemplatePages.Pages.IsValidIndex((*TemplateIdx) + 1))
+		{
+			return TemplatePages.Pages[(*TemplateIdx) + 1];
+		}
+		else if (TemplatePages.Pages.IsEmpty() == false)
+		{
+			return TemplatePages.Pages[0];
+		}
+	}
+
+	if (const int32* InstancedIdx = InstancedPages.PageIndices.Find(InPageId))
+	{
+		if (InPageListReference.Type == EAvaPageListType::View && SubLists.IsValidIndex(InPageListReference.SubListIndex))
+		{
+			int32 Index = SubLists[InPageListReference.SubListIndex].PageIds.Find(InPageId);
+
+			if (Index != INDEX_NONE)
+			{
+				return GetNextFromSubList(SubLists[InPageListReference.SubListIndex].PageIds, Index);
+			}
+
+			if (SubLists[InPageListReference.SubListIndex].PageIds.IsEmpty())
+			{
+				return GetNextFromPages(InstancedPages.Pages, (*InstancedIdx));
+			}
+		}
+
+		if (InPageListReference.Type == EAvaPageListType::Instance)
+		{
+			return GetNextFromPages(InstancedPages.Pages, (*InstancedIdx));
+		}
+	}
+
+	return FAvalanchePage::NullPage;
+}
+
+FAvalanchePage& UAvalanchePlaylist::GetNextPage(int32 InPageId, const FAvaPageListReference& InPageListReference)
+{
+	if (const int32* TemplateIdx = TemplatePages.PageIndices.Find(InPageId))
+	{
+		if (TemplatePages.Pages.IsValidIndex((*TemplateIdx) + 1))
+		{
+			return TemplatePages.Pages[(*TemplateIdx) + 1];
+		}
+		else if (TemplatePages.Pages.IsEmpty() == false)
+		{
+			return TemplatePages.Pages[0];
+		}
+	}
+
+	if (const int32* InstancedIdx = InstancedPages.PageIndices.Find(InPageId))
+	{
+		if (InPageListReference.Type == EAvaPageListType::View && SubLists.IsValidIndex(InPageListReference.SubListIndex))
+		{
+			const int32 Index = SubLists[InPageListReference.SubListIndex].PageIds.Find(InPageId);
+
+			if (Index != INDEX_NONE)
+			{
+				return GetNextFromSubList(SubLists[InPageListReference.SubListIndex].PageIds, Index);
+			}
+
+			if (SubLists[InPageListReference.SubListIndex].PageIds.IsEmpty())
+			{
+				return GetNextFromPages(InstancedPages.Pages, (*InstancedIdx));
+			}
+		}
+
+		if (InPageListReference.Type == EAvaPageListType::Instance)
+		{
+			return GetNextFromPages(InstancedPages.Pages, (*InstancedIdx));
+		}
+	}
+
+	return FAvalanchePage::NullPage;
+}
+
+void UAvalanchePlaylist::InitializePlaybackContext()
+{
+	if (!PlaybackClientWatcher)
+	{
+		PlaybackClientWatcher = MakePimpl<FAvaPlaylistPlaybackClientWatcher>(this);
+	}
+}
+
+void UAvalanchePlaylist::ClosePlaybackContext(bool bInStopAllPages)
+{
+	if (bInStopAllPages)
+	{
+		for (UAvalanchePagePlayer* PagePlayer : PagePlayers)
+		{
+			if (PagePlayer)
+			{
+				PagePlayer->Stop();
+			}
+		}
+		RemoveStoppedPagePlayers();
+	}
+	
+	PlaybackClientWatcher.Reset();
+}
+
+bool UAvalanchePlaylist::IsPlaying() const
+{
+	return PagePlayers.Num() > 0;
+}
+
+bool UAvalanchePlaylist::IsPagePreviewing(int32 InPageId) const
+{
+	return PagePlayers.ContainsByPredicate([InPageId](const UAvalanchePagePlayer* InPagePlayer)
+	{
+		return InPagePlayer->bIsPreview && InPagePlayer->PageId == InPageId && InPagePlayer->IsPlaying();
+	});
+}
+
+bool UAvalanchePlaylist::IsPagePlaying(int32 InPageId) const
+{
+	return PagePlayers.ContainsByPredicate([InPageId](const UAvalanchePagePlayer* InPagePlayer)
+	{
+		return !InPagePlayer->bIsPreview && InPagePlayer->PageId == InPageId && InPagePlayer->IsPlaying();
+	});
+}
+
+bool UAvalanchePlaylist::IsPagePlayingOrPreviewing(int32 InPageId) const
+{
+	return PagePlayers.ContainsByPredicate([InPageId](const UAvalanchePagePlayer* InPagePlayer)
+	{
+		return InPagePlayer->PageId == InPageId && InPagePlayer->IsPlaying();
+	});
+}
+
+bool UAvalanchePlaylist::UnloadPage(int32 InPageId, const FString& InChannelName)
+{
+	FAvaMediaPlaybackManager& Manager = GetPlaybackManager();
+	
+	if (const FAvalanchePage& SelectedPage = GetPage(InPageId); SelectedPage.IsValidPage())
+	{
+		// Ensure all players for this page have stopped.
+		for (UAvalanchePagePlayer* PagePlayer : PagePlayers)
+		{
+			if (PagePlayer->PageId == InPageId)
+			{
+				PagePlayer->Stop();
+			}
+		}
+		RemoveStoppedPagePlayers();
+
+		bool bSuccess = false;
+		for (const FSoftObjectPath& AssetPath : SelectedPage.GetAvalancheAssetPaths(this))
+		{
+			// This will unload all the "available" (i.e. not used) instances of that asset on that channel.
+			bSuccess |= Manager.UnloadPlaybackInstances(AssetPath, InChannelName);
+		}
+		return bSuccess;
+	}
+	return false;
+}
+
+TArray<UAvalanchePlaylist::FLoadedInstanceInfo> UAvalanchePlaylist::LoadPage(int32 InPageId,  bool bInPreview, const FName& InPreviewChannelName)
+{
+	const FAvalanchePage& Page = GetPage(InPageId);
+	if (!Page.IsValidPage())
+	{
+		return {};
+	}
+
+	const FName ChannelName = bInPreview ? (InPreviewChannelName.IsNone() ? GetDefaultPreviewChannelName() : InPreviewChannelName) : Page.GetChannelName();
+	const TArray<FSoftObjectPath> AssetPaths = Page.GetAvalancheAssetPaths(this);
+
+	TArray<FLoadedInstanceInfo> LoadedInstances;
+	LoadedInstances.Reserve(AssetPaths.Num());
+
+	for (const FSoftObjectPath& AssetPath : AssetPaths)
+	{
+		const TSharedPtr<FAvaMediaPlaybackInstance> PlaybackInstance = GetPlaybackManager().AcquireOrLoadPlaybackInstance(AssetPath, ChannelName.ToString());
+		if (!PlaybackInstance || !PlaybackInstance->GetPlayback())
+		{
+			continue;
+		}
+
+		UAvalanchePagePlayer::SetInstanceUserDataFromPage(*PlaybackInstance, Page);
+		PlaybackInstance->GetPlayback()->LoadInstances();
+		PlaybackInstance->UpdateStatus();
+		PlaybackInstance->Recycle();
+		LoadedInstances.Add({PlaybackInstance->GetInstanceId(), AssetPath});
+	}
+	return LoadedInstances;
+}
+
+TArray<int32> UAvalanchePlaylist::PlayPages(const TArray<int32>& InPageIds, EAvaPlayType InPlayType)
+{
+	return PlayPages(InPageIds, InPlayType, UE::AvalanchePlaylist::IsPreviewPlayType(InPlayType) ? GetDefaultPreviewChannelName() : NAME_None);
+}
+
+TArray<int32> UAvalanchePlaylist::PlayPages(const TArray<int32>& InPageIds, EAvaPlayType InPlayType, const FName& InPreviewChannelName)
+{
+	TArray<int32> PlayedPageIds;
+	PlayedPageIds.Reserve(InPageIds.Num());
+
+	FAvalanchePageTransitionBuilder TransitionBuilder(this);
+	
+	for (const int32 PageId : InPageIds)
+	{
+		const FAvalanchePage& SelectedPage = GetPage(PageId);
+		if (SelectedPage.IsValidPage() && SelectedPage.IsEnabled())
+		{
+			const bool bIsPreview = UE::AvalanchePlaylist::IsPreviewPlayType(InPlayType);
+
+			if (!IsChannelTypeCompatibleForRequest(SelectedPage, bIsPreview, InPreviewChannelName, true))
+			{
+				continue;
+			}
+
+			// Only preview from frame still uses the no transition path.
+			// Everything else uses transitions, even when no transition tree.
+			const bool bPagePlayed = (InPlayType != EAvaPlayType::PreviewFromFrame)
+				? PlayPageWithTransition(TransitionBuilder, SelectedPage, InPlayType, bIsPreview, InPreviewChannelName)
+				: PlayPageNoTransition(SelectedPage, InPlayType, bIsPreview, InPreviewChannelName);
+		
+			if (bPagePlayed)
+			{
+				GetOrCreatePageListPlaybackContextCollection().GetOrCreateContext(bIsPreview, InPreviewChannelName).PlayHeadPageId = PageId;
+				PlayedPageIds.Add(PageId);
+			}
+		}
+	}
+	return PlayedPageIds;
+}
+
+bool UAvalanchePlaylist::RestorePlaySubPage(int32 InPageId, int32 InSubPageIndex, const FGuid& InExistingInstanceId, bool bInIsPreview, const FName& InPreviewChannelName)
+{
+	auto LogError = [InPageId, InPreviewChannelName](const FString& InReason)
+	{
+		UE_LOG(LogAvaPlaylist, Error,
+			TEXT("Couldn't restore playback state of page %d on channel \"%s\": %s."),
+			InPageId, *InPreviewChannelName.ToString(), *InReason);
+	};
+
+	const FAvalanchePage& Page = GetPage(InPageId);
+	if (!Page.IsValidPage() || !Page.IsEnabled())
+	{
+		LogError(TEXT("Page is either not valid or disabled"));
+		return false;
+	}
+
+	if (!IsChannelTypeCompatibleForRequest(Page, bInIsPreview, InPreviewChannelName, true))
+	{
+		LogError(TEXT("Channel Type is not compatible"));
+		return false;
+	}
+
+	if (!InExistingInstanceId.IsValid())
+	{
+		LogError(TEXT("Specified instance id is invalid"));
+		return false;
+	}
+	
+	bool bPagePlayerCreated = false;
+	UAvalanchePagePlayer* PagePlayer = FindPlayerForPage(InPageId, bInIsPreview, InPreviewChannelName);
+
+	if (!PagePlayer)
+	{
+		bPagePlayerCreated = true;
+		PagePlayer = NewObject<UAvalanchePagePlayer>(this);
+		if (!PagePlayer->Initialize(this, Page, bInIsPreview, InPreviewChannelName))
+		{
+			return false;
+		}
+	}
+	
+	if (const UAvaRundownPlaybackInstancePlayer* LoadedInstancePlayer = PagePlayer->LoadInstancePlayer(InSubPageIndex, InExistingInstanceId))
+	{
+		if (bPagePlayerCreated)
+		{
+			AddPagePlayer(PagePlayer);
+			GetOrCreatePageListPlaybackContextCollection().GetOrCreateContext(bInIsPreview, InPreviewChannelName).PlayHeadPageId = InPageId;
+		}
+		
+		UAvalanchePlayback* Playback = LoadedInstancePlayer->Playback;
+		if (Playback && !Playback->IsPlaying())
+		{
+			Playback->Play();
+		}
+		return true;
+	}
+	
+	LogError(TEXT("Unable to acquire or load playback object"));
+	return false;
+}
+
+bool UAvalanchePlaylist::CanPlayPage(int32 InPageId, bool bInPreview) const
+{
+	return CanPlayPage(InPageId, bInPreview, bInPreview ? GetDefaultPreviewChannelName() : NAME_None);
+}
+
+bool UAvalanchePlaylist::CanPlayPage(int32 InPageId, bool bInPreview, const FName& InPreviewChannelName) const
+{
+	// Check if the page is valid and enabled.
+	const FAvalanchePage& SelectedPage = GetPage(InPageId);
+	if (!SelectedPage.IsValidPage() || !SelectedPage.IsEnabled())
+	{
+		return false;
+	}
+
+	// Check channel validity and type compatibility.
+	if (!IsChannelTypeCompatibleForRequest(SelectedPage, bInPreview, InPreviewChannelName, false))
+	{
+		return false;
+	}
+
+	// Check that if it is a template page it is meant to preview
+	if (SelectedPage.IsTemplate() && !bInPreview)
+	{
+		return false;
+	}
+
+	// Check if the asset path is valid
+	if (SelectedPage.GetAvalancheAssetPath(this).IsNull())
+	{
+		return false;
+	}
+
+	// For page with TL, need to see if a transition can be started for that page.
+	if (SelectedPage.HasTransitionLogic(this) && !CanStartTransitionForPage(SelectedPage, bInPreview, InPreviewChannelName))
+	{
+		return false;
+	}
+
+	// Remark:
+	// No longer checks if the playback object is already playing because
+	// a "playing" page can be played again, it means the animation will be restarted.
+	return true;
+}
+
+TArray<int32> UAvalanchePlaylist::StopPages(const TArray<int32>& InPageIds, EAvaPlaylistPageStopOptions InOptions, bool bInPreview)
+{
+	return StopPages(InPageIds, InOptions, bInPreview, bInPreview ? GetDefaultPreviewChannelName() : NAME_None);
+}
+
+TArray<int32> UAvalanchePlaylist::StopPages(const TArray<int32>& InPageIds, EAvaPlaylistPageStopOptions InOptions, bool bInPreview, const FName& InPreviewChannelName)
+{
+	TArray<int32> StoppedPageIds;
+	StoppedPageIds.Reserve(InPageIds.Num());
+
+	FAvalanchePageTransitionBuilder TransitionBuilder(this);
+
+	for (const int32 PageId : InPageIds)
+	{
+		const FAvalanchePage& SelectedPage = GetPage(PageId);
+
+		if (!SelectedPage.IsValidPage())
+		{
+			continue;
+		}
+
+		if (SelectedPage.HasTransitionLogic(this))
+		{
+			if (!EnumHasAnyFlags(InOptions, EAvaPlaylistPageStopOptions::ForceNoTransition))
+			{
+				if (StopPageWithTransition(TransitionBuilder, SelectedPage, bInPreview, InPreviewChannelName))
+				{
+					StoppedPageIds.Add(PageId);
+				}
+				continue;
+			}
+
+			// Force stop all page transitions for the given page, then stop the page without TL.
+			StopPageTransitionsForPage(SelectedPage, bInPreview, InPreviewChannelName);
+		}
+
+		if (StopPageNoTransition(SelectedPage, bInPreview, InPreviewChannelName))
+		{
+			StoppedPageIds.Add(PageId);
+		}
+	}
+	return StoppedPageIds;
+}
+
+bool UAvalanchePlaylist::CanStopPage(int32 InPageId, EAvaPlaylistPageStopOptions InOptions, bool bInPreview) const
+{
+	return CanStopPage(InPageId, InOptions, bInPreview, bInPreview ? GetDefaultPreviewChannelName() : NAME_None);
+}
+
+bool UAvalanchePlaylist::CanStopPage(int32 InPageId, EAvaPlaylistPageStopOptions InOptions, bool bInPreview, const FName& InPreviewChannelName) const
+{
+	const FAvalanchePage& SelectedPage = GetPage(InPageId);
+
+	if (!SelectedPage.IsValidPage())
+	{
+		return false;
+	}
+	
+	// For page with TL, need to see if a transition can be started for that page.
+	if (!EnumHasAnyFlags(InOptions, EAvaPlaylistPageStopOptions::ForceNoTransition)
+		&& SelectedPage.HasTransitionLogic(this)
+		&& !CanStartTransitionForPage(SelectedPage, bInPreview, InPreviewChannelName))
+	{
+		return false;
+	}
+	
+	const UAvalanchePagePlayer* Player = FindPlayerForPage(InPageId, bInPreview, InPreviewChannelName);
+	return Player && Player->IsPlaying();
+}
+
+bool UAvalanchePlaylist::StopChannel(const FString& InChannelName)
+{
+	const FName ChannelName(InChannelName);
+	int32 NumStoppedPages = 0;
+	for (UAvalanchePagePlayer* PagePlayer : PagePlayers)
+	{
+		// Don't let something else play on this channel.
+		if (PagePlayer->ChannelName == ChannelName)
+		{
+			if (PagePlayer->Stop())
+			{
+				++NumStoppedPages;
+			}
+		}
+	}
+	RemoveStoppedPagePlayers();
+	return NumStoppedPages > 0 ? true : false;
+}
+
+bool UAvalanchePlaylist::CanStopChannel(const FString& InChannelName) const
+{
+	const FName ChannelName(InChannelName);
+	for (const UAvalanchePagePlayer* PagePlayer : PagePlayers)
+	{
+		if (PagePlayer->ChannelName == ChannelName && PagePlayer->IsPlaying())
+		{
+			return true;
+		}
+	}
+	return false;
+}
+
+bool UAvalanchePlaylist::ContinuePage(int32 InPageId, bool bInPreview)
+{
+	return ContinuePage(InPageId, bInPreview, bInPreview ? GetDefaultPreviewChannelName() : NAME_None);	
+}
+
+bool UAvalanchePlaylist::ContinuePage(int32 InPageId, bool bInPreview, const FName& InPreviewChannelName)
+{
+	const FAvalanchePage& SelectedPage = GetPage(InPageId);
+
+	if (SelectedPage.IsValidPage() && SelectedPage.IsEnabled())
+	{
+		if (UAvalanchePagePlayer* Player = FindPlayerForPage(InPageId, bInPreview, InPreviewChannelName))
+		{
+			return Player->Continue();
+		}
+	}
+	return false;
+}
+
+bool UAvalanchePlaylist::CanContinuePage(int32 InPageId, bool bInPreview) const
+{
+	return CanContinuePage(InPageId, bInPreview, bInPreview ? GetDefaultPreviewChannelName() : NAME_None);	
+}
+
+bool UAvalanchePlaylist::CanContinuePage(int32 InPageId, bool bInPreview, const FName& InPreviewChannelName) const
+{
+	const FAvalanchePage& SelectedPage = GetPage(InPageId);
+
+	if (SelectedPage.IsValidPage() && SelectedPage.IsEnabled())
+	{
+		const UAvalanchePagePlayer* Player = FindPlayerForPage(InPageId, bInPreview, InPreviewChannelName);
+		return Player && Player->IsPlaying();
+	}
+
+	return false;
+}
+
+int32 UAvalanchePlaylist::AddSubList()
+{
+	const int32 SubListIdx = SubLists.Add(FAvalancheSubList());
+	SetActivePageList(CreateSubListReference(SubListIdx));
+
+	return SubListIdx;
+}
+
+TArray<int32> UAvalanchePlaylist::GetPlayingPageIds(const FName InProgramChannelName) const
+{
+	TArray<int32> OutPlayedIds;
+	OutPlayedIds.Reserve(PagePlayers.Num());
+	for (const UAvalanchePagePlayer* PagePlayer : PagePlayers)
+	{
+		if (PagePlayer->bIsPreview || !PagePlayer->IsPlaying())
+		{
+			continue;
+		}
+
+		if (!InProgramChannelName.IsNone() && PagePlayer->ChannelFName != InProgramChannelName)
+		{
+			continue;
+		}
+
+		OutPlayedIds.AddUnique(PagePlayer->PageId);
+	}
+	return OutPlayedIds;
+}
+
+TArray<int32> UAvalanchePlaylist::GetPreviewingPageIds(const FName InPreviewChannelName) const
+{
+	TArray<int32> OutPreviewingIds;
+	OutPreviewingIds.Reserve(PagePlayers.Num());
+	for (const UAvalanchePagePlayer* PagePlayer : PagePlayers)
+	{
+		if (!PagePlayer->bIsPreview || !PagePlayer->IsPlaying())
+		{
+			continue;
+		}
+
+		if (!InPreviewChannelName.IsNone() && PagePlayer->ChannelFName != InPreviewChannelName)
+		{
+			continue;
+		}
+		
+		OutPreviewingIds.AddUnique(PagePlayer->PageId);
+	}
+	return OutPreviewingIds;
+}
+
+bool UAvalanchePlaylist::SetActivePageList(const FAvaPageListReference& InPageListReference)
+{
+	if (InPageListReference.Type == EAvaPageListType::Instance)
+	{
+		ActivePageList = InstancePageList;
+		OnActiveListChanged.Broadcast();
+		return true;
+	}
+
+	if (InPageListReference.Type == EAvaPageListType::View && SubLists.IsValidIndex(InPageListReference.SubListIndex))
+	{
+		ActivePageList = InPageListReference;
+		OnActiveListChanged.Broadcast();
+		return true;
+	}
+
+	return false;
+}
+
+bool UAvalanchePlaylist::HasActiveSubList() const
+{
+	return (ActivePageList.Type == EAvaPageListType::View && SubLists.IsValidIndex(ActivePageList.SubListIndex));
+}
+
+const FAvalancheSubList& UAvalanchePlaylist::GetSubList(int32 InSubListIndex) const
+{
+	if (SubLists.IsValidIndex(InSubListIndex))
+	{
+		return SubLists[InSubListIndex];
+	}
+
+	return InvalidSubList;
+}
+
+FAvalancheSubList& UAvalanchePlaylist::GetSubList(int32 InSubListIndex)
+{
+	check(SubLists.IsValidIndex(InSubListIndex));
+
+	return SubLists[InSubListIndex];
+}
+
+bool UAvalanchePlaylist::IsValidSubList(const FAvaPageListReference& InPageListReference) const
+{
+	return (InPageListReference.Type == EAvaPageListType::View && SubLists.IsValidIndex(InPageListReference.SubListIndex));
+}
+
+bool UAvalanchePlaylist::AddPageToSubList(int32 InSubListIndex, int32 InPageId, const FAvaPageInsertPosition& InInsertPosition)
+{
+	if (SubLists.IsValidIndex(InSubListIndex) && InstancedPages.PageIndices.Contains(InPageId) 
+		&& !SubLists[InSubListIndex].PageIds.Contains(InPageId))
+	{
+		int32 ExistingPageIndex = INDEX_NONE;
+
+		if (InInsertPosition.IsValid())
+		{
+			ExistingPageIndex = SubLists[InSubListIndex].PageIds.IndexOfByKey(InInsertPosition.AdjacentId);
+		}
+
+		if (InInsertPosition.IsAddBelow() && SubLists[InSubListIndex].PageIds.IsValidIndex(ExistingPageIndex))
+		{
+			++ExistingPageIndex;
+		}
+
+		if (SubLists[InSubListIndex].PageIds.IsValidIndex(ExistingPageIndex))
+		{
+			SubLists[InSubListIndex].PageIds.Insert(InPageId, ExistingPageIndex);
+		}
+		else
+		{
+			SubLists[InSubListIndex].PageIds.Add(InPageId);
+		}
+		
+		SubLists[InSubListIndex].OnPageListChanged.Broadcast({this, EAvaPageListChange::AddedPages, {InPageId}});
+		return true;
+	}
+
+	return false;
+}
+
+bool UAvalanchePlaylist::AddPagesToSubList(int32 InSubListIndex, const TArray<int32>& InPages)
+{
+	if (SubLists.IsValidIndex(InSubListIndex))
+	{
+		bool bAddedPage = false;
+
+		// Super inefficient for now.
+		for (int32 PageId : InPages)
+		{
+			if (InstancedPages.PageIndices.Contains(PageId) && !SubLists[InSubListIndex].PageIds.Contains(PageId))
+			{
+				SubLists[InSubListIndex].PageIds.Add(PageId);
+				bAddedPage = true;
+			}
+		}
+
+		if (bAddedPage)
+		{
+			SubLists[InSubListIndex].OnPageListChanged.Broadcast({this, EAvaPageListChange::AddedPages, InPages});
+			return true;
+		}
+	}
+
+	return false;
+}
+
+int32 UAvalanchePlaylist::RemovePagesFromSubList(int32 InSubListIndex, const TArray<int32>& InPages)
+{
+	if (SubLists.IsValidIndex(InSubListIndex))
+	{
+		const int32 Removed = SubLists[InSubListIndex].PageIds.RemoveAll([InPages](const int32& PageId)
+			{
+				return InPages.Contains(PageId);
+			});
+
+		if (Removed > 0)
+		{
+			SubLists[InSubListIndex].OnPageListChanged.Broadcast({this, EAvaPageListChange::RemovedPages, InPages});
+		}
+
+		return Removed;
+	}
+
+	return 0;
+}
+
+namespace UE::AvalancheMedia::Playlist::Private
+{
+	const UAvalanchePagePlayer* FindPagePlayerForPreviewChannel(const FName& InPreviewChannel, const TArray<TObjectPtr<UAvalanchePagePlayer>>& InPagePlayers)
+	{
+		const TObjectPtr<UAvalanchePagePlayer>* PreviewingPagePlayerPtr = InPagePlayers.FindByPredicate([InPreviewChannel](const UAvalanchePagePlayer* InPagePlayer)
+		{
+			return InPagePlayer->bIsPreview && InPagePlayer->ChannelName == InPreviewChannel;
+		});
+		return PreviewingPagePlayerPtr ? *PreviewingPagePlayerPtr : nullptr;
+	}
+}
+
+UTextureRenderTarget2D* UAvalanchePlaylist::GetPreviewRenderTarget(const FName& InPreviewChannel) const
+{
+	// For preview, there can be an output channel or not.
+	// If there is one, we will prefer getting the render target directly from the channel.
+	
+	UAvalancheBroadcast& Broadcast = UAvalancheBroadcast::Get();
+	const FAvaOutputChannel& OutputChannel = Broadcast.GetCurrentProfile().GetChannel(InPreviewChannel); 
+	if (OutputChannel.IsValidChannel())
+	{
+		return OutputChannel.GetCurrentRenderTarget(true);
+	}
+
+	// If there is no channel, we can get the render target from the playable group of a previewing page's playable
+	// in the given channel. When playable group composition is implemented, this may have to change. 
+	
+	using namespace UE::AvalancheMedia::Playlist::Private;
+	const UAvalanchePagePlayer* PreviewingPagePlayer = FindPagePlayerForPreviewChannel(InPreviewChannel, PagePlayers);
+
+	if (!PreviewingPagePlayer)
+	{
+		return nullptr;
+	}
+	
+	if (const UAvalanchePlayback* Playback = PreviewingPagePlayer->GetPlayback())
+	{
+		if (const UAvalanchePlayable* const Playable = Playback->GetFirstPlayable())
+		{
+			if (const UAvaMediaPlayableGroup* const PlayableGroup = Playable->GetPlayableGroup())
+			{
+				return PlayableGroup->IsRenderTargetReady() ? PlayableGroup->GetRenderTarget() : nullptr;
+			}
+		}
+	}
+	return nullptr;
+}
+
+FName UAvalanchePlaylist::GetDefaultPreviewChannelName()
+{
+	// Even if the user selected preview channel is empty, we need a default
+	// name as a key for the playback manager.
+	static FName DefaultPreviewChannelFName = TEXT("_Preview");
+	const UAvalancheMediaSettings& Settings = UAvalancheMediaSettings::Get();
+	return !Settings.PreviewChannelName.IsEmpty() ? FName(Settings.PreviewChannelName) : DefaultPreviewChannelFName;
+}
+
+void UAvalanchePlaylist::OnParentWordBeginTearDown()
+{
+	PagePlayers.Reset();
+}
+
+bool UAvalanchePlaylist::PushRuntimeRemoteControlValues(int32 InPageId, bool bInIsPreview, const FName& InPreviewChannelName) const
+{
+	const FAvalanchePage& Page = GetPage(InPageId);
+	if (!Page.IsValidPage())
+	{
+		return false;
+	}
+
+	bool bValuesPushed = false;
+	for (const UAvalanchePagePlayer* PagePlayer : PagePlayers)
+	{
+		if (PagePlayer->PageId != InPageId || PagePlayer->bIsPreview != bInIsPreview)
+		{
+			continue;
+		}
+		
+		// Filter on preview channel if provided. 
+		if (bInIsPreview && !InPreviewChannelName.IsNone() && PagePlayer->ChannelFName != InPreviewChannelName)
+		{
+			continue;
+		}
+
+		TSharedRef<FAvalancheRemoteControlValues> SharedRCValues = MakeShared<FAvalancheRemoteControlValues>(Page.GetRemoteControlValues());
+		for (int32 InstanceIndex = 0; InstanceIndex < PagePlayer->GetNumInstancePlayers(); ++InstanceIndex)
+		{
+			if (UAvalanchePlayback* Playback = PagePlayer->GetPlayback(InstanceIndex))
+			{
+				Playback->PushRemoteControlValues(PagePlayer->GetSourceAssetPath(InstanceIndex), PagePlayer->ChannelName, SharedRCValues);
+			}
+		}
+		bValuesPushed = true;
+	}
+	return bValuesPushed;
+}
+
+// Note: This can be called if
+// - RC entity values are either added or modified.
+// - RC controller values are either added or modified.
+void UAvalanchePlaylist::NotifyPageRemoteControlValueChanged(int32 InPageId, EAvaRemoteControlChanges InRemoteControlChanges)
+{
+	// For the previewed page, we automatically update the playback object's RC values live.
+	// Note: Only the entity values are updated in the runtime (playback) RCP. No need to push controller values to runtime.
+	if (EnumHasAnyFlags(InRemoteControlChanges, EAvaRemoteControlChanges::EntityValues))
+	{
+		// For now, potentially pushing all values multiple time (per frame) is mitigated by the
+		// optimization in FAvalancheRemoteControlUtils::SetValueOfEntity that
+		// will only set the value of the entity if it changed.
+		PushRuntimeRemoteControlValues(InPageId, true);
+	}
+	OnPagesChanged.Broadcast(this, {InPageId}, EAvaPageChanges::RemoteControlValues);
+}
+
+#if WITH_EDITOR
+void UAvalanchePlaylist::NotifyPIEEnded(const bool)
+{
+	// When PIE Ends, all worlds should be forcibly destroyed
+	OnParentWordBeginTearDown();
+}
+#endif
+
+FAvaMediaPlaybackManager& UAvalanchePlaylist::GetPlaybackManager() const
+{
+	return IAvaMediaModule::Get().GetLocalPlaybackManager();
+}
+
+int32 UAvalanchePlaylist::AddPageFromTemplateInternal(int32 InTemplateId, const FAvaPageIdGeneratorParams& InIdGeneratorParams, const FAvaPageInsertPosition& InInsertAt)
+{
+	const int32* TemplateIndex = TemplatePages.PageIndices.Find(InTemplateId);
+
+	if (!TemplateIndex)
+	{
+		return FAvalanchePage::InvalidPageId;
+	}
+
+	const int32 NewId = GenerateUniquePageId(InIdGeneratorParams);
+
+	int32 ExistingPageIndex = INDEX_NONE;
+
+	if (InInsertAt.IsValid())
+	{
+		if (const int32* ExistingPageIndexPtr = InstancedPages.PageIndices.Find(InInsertAt.AdjacentId))
+		{
+			ExistingPageIndex = *ExistingPageIndexPtr;
+		}
+	}
+
+	if (InInsertAt.IsAddBelow() && InstancedPages.Pages.IsValidIndex(ExistingPageIndex))
+	{
+		++ExistingPageIndex;
+	}
+
+	int32 NewIndex = INDEX_NONE;
+
+	if (InstancedPages.Pages.IsValidIndex(ExistingPageIndex))
+	{
+		InstancedPages.Pages.Insert(TemplatePages.Pages[*TemplateIndex], ExistingPageIndex);
+		NewIndex = ExistingPageIndex;
+		
+		// Need to update page indices after insertion.
+		InstancedPages.PostInsertRefreshPageIndices(ExistingPageIndex + 1);
+	}
+	else
+	{
+		InstancedPages.Pages.Emplace(TemplatePages.Pages[*TemplateIndex]);
+		NewIndex = InstancedPages.Pages.Num() - 1;
+	}
+
+	InstancedPages.PageIndices.Emplace(NewId, NewIndex);
+	TemplatePages.Pages[*TemplateIndex].Instances.Add(NewId);
+	InitializePage(InstancedPages.Pages[NewIndex], NewId, InTemplateId);
+
+	return NewId;
+}
+
+void UAvalanchePlaylist::InitializePage(FAvalanchePage& InOutPage, int32 InPageId, int32 InTemplateId) const
+{
+	InOutPage.PageId = InPageId;
+	InOutPage.TemplateId = InTemplateId;
+	InOutPage.CombinedTemplateIds.Empty();
+	InOutPage.SetPageFriendlyName(FText::GetEmpty());
+	InOutPage.UpdatePageSummary(this);
+}
+
+bool UAvalanchePlaylist::IsChannelTypeCompatibleForRequest(const FAvalanchePage& InSelectedPage, bool bInIsPreview, const FName& InPreviewChannelName, bool bInLogFailureReason) const
+{
+	// Check channel validity and type compatibility.
+	UAvalancheBroadcast& Broadcast = UAvalancheBroadcast::Get();
+	if (bInIsPreview)
+	{
+		// The incoming preview channel name may not exist, that is allowed.
+		if (Broadcast.GetCurrentProfile().GetChannel(InPreviewChannelName).IsValidChannel()
+			&& Broadcast.GetChannelType(InPreviewChannelName) != EAvaBroadcastChannelType::Preview)
+		{
+			if (bInLogFailureReason)
+			{
+				UE_LOG(LogAvaPlaylist, Error, TEXT("Preview request failed. Channel \"%s\" is not a \"preview\" channel in profile \"%s\"."),
+					*InPreviewChannelName.ToString(), *Broadcast.GetCurrentProfileName().ToString());
+			}
+			return false;
+		}
+	}
+	else
+	{
+		const FAvaOutputChannel& Channel = Broadcast.GetCurrentProfile().GetChannel(InSelectedPage.GetChannelName()); 
+		if (!Channel.IsValidChannel())
+		{
+			if (bInLogFailureReason)
+			{
+				UE_LOG(LogAvaPlaylist, Error, TEXT("Playback request failed. Channel \"%s\" is not a valid channel in \"%s\" profile."),
+					*InSelectedPage.GetChannelName().ToString(), *Broadcast.GetCurrentProfileName().ToString());
+			}
+			return false;
+		}
+		if (Broadcast.GetChannelType(InSelectedPage.GetChannelName()) != EAvaBroadcastChannelType::Program)
+		{
+			if (bInLogFailureReason)
+			{
+				UE_LOG(LogAvaPlaylist, Error, TEXT("Playback request failed. Channel \"%s\" is not a \"program\" channel in profile \"%s\"."),
+					*InSelectedPage.GetChannelName().ToString(), *Broadcast.GetCurrentProfileName().ToString());
+			}
+			return false;
+		}
+
+		// Check if the channel is offline.
+		bool bHasOfflineOutput = false;
+		bool bHasLocalOutput = false;
+		const TArray<UMediaOutput*>& Outputs = Channel.GetMediaOutputs();
+		for (const UMediaOutput* Output : Outputs)
+		{
+			if (Channel.IsMediaOutputRemote(Output) && Channel.GetMediaOutputState(Output) == EAvaMediaOutputState::Offline)
+			{
+				bHasOfflineOutput = true;
+			}
+			else
+			{
+				bHasLocalOutput = true;
+				break;	// If a local output is detected all is good.
+			}
+		}
+
+		// A channel is considered offline only if it doesn't have any local outputs since
+		// the local outputs take priority (for now at least).
+		if (bHasOfflineOutput && !bHasLocalOutput)
+		{
+			if (bInLogFailureReason)
+			{
+				UE_LOG(LogAvaPlaylist, Error, TEXT("Playback request failed. Channel \"%s\" is offline."),
+					*InSelectedPage.GetChannelName().ToString());
+			}
+			return false;
+		}
+	}
+	return true;
+}
+
+void UAvalanchePlaylist::AddPagePlayer(UAvalanchePagePlayer* InPagePlayer)
+{
+	PagePlayers.Add(InPagePlayer);
+	OnPagePlayerAdded.Broadcast(this, InPagePlayer);
+}
+
+IAvaPlaylistPageLoadingManager& UAvalanchePlaylist::MakePageLoadingManager()
+{
+	PageLoadingManager = MakeUnique<FAvaPlaylistPageLoadingManager>(this);
+	return *PageLoadingManager;
+}
+
+bool UAvalanchePlaylist::PlayPageNoTransition(const FAvalanchePage& InPage, EAvaPlayType InPlayType, bool bInIsPreview, const FName& InPreviewChannelName)
+{
+	UAvalanchePagePlayer* PagePlayer = FindPlayerForPage(InPage.GetPageId(), bInIsPreview, InPreviewChannelName);
+	if (!PagePlayer)
+	{
+		PagePlayer = NewObject<UAvalanchePagePlayer>(this);
+		if (PagePlayer->InitializeAndLoad(this, InPage, bInIsPreview, InPreviewChannelName))
+		{
+			AddPagePlayer(PagePlayer);
+		}
+	}
+
+	if (PagePlayer && PagePlayer->IsLoaded())
+	{
+		PagePlayer->Play(InPlayType, false);	// not using transition logic: animations commands are going to be pushed.
+			
+		// Don't let something else play on this channel.
+		// FIXME: this is not perfect, on the remote render node, it may cause a frame of "place holder"
+		// between the end of the previous page and start of the next one.
+
+		for (UAvalanchePagePlayer* OtherPagePlayer : PagePlayers)
+		{
+			// Don't let something else play on this channel.
+			if (OtherPagePlayer != PagePlayer && OtherPagePlayer->ChannelFName == PagePlayer->ChannelFName)
+			{
+				OtherPagePlayer->Stop();
+			}
+		}
+
+		RemoveStoppedPagePlayers();
+				
+		return true;
+	}
+	return false;
+}
+
+namespace UE::AvalancheMedia::Playlist::Private
+{
+	bool ArePageRCValuesEqualForSubTemplate(const FAvalanchePage& InSubTemplate, const FAvalanchePage& InPage, const FAvalanchePage& InOtherPage)
+	{
+		// Comparing only the entity values for now. For playback, this is what determines if the values are the
+		// same or not. The controllers are for editing only.
+		for (const TPair<FGuid, FAvalancheRemoteControlValue>& EntityEntry : InSubTemplate.GetRemoteControlValues().EntityValues)
+		{
+			const FAvalancheRemoteControlValue* Value = InPage.GetRemoteControlEntityValue(EntityEntry.Key);
+			const FAvalancheRemoteControlValue* OtherValue = InOtherPage.GetRemoteControlEntityValue(EntityEntry.Key);
+			if (!Value || !OtherValue || !Value->IsSameValueAs(*OtherValue))
+			{
+				return false;
+			}
+		}
+		return true;
+	}
+
+	/**
+	 * @brief For special template transition logic, search for an existing instance player with same RC values.
+	 * @param InPlaylist Playlist
+	 * @param InPageToPlay New Page to be played.
+	 * @param InTemplate Template to be played. Should be direct template of the page.
+	 * @param InSubPageIndex Index of the sub-template.
+	 * @param bInIsPreview True if the channel is a preview channel.
+	 * @param InPreviewChannelName Name of the preview channel (if it is a preview channel).
+	 * @return Pointer to found instance player.
+	 */
+	UAvaRundownPlaybackInstancePlayer* FindExistingInstancePlayer(
+		const UAvalanchePlaylist* InPlaylist,
+		const FAvalanchePage& InPageToPlay,
+		const FAvalanchePage& InTemplate,
+		int32 InSubPageIndex,
+		bool bInIsPreview,
+		const FName& InPreviewChannelName)
+	{
+		for (const TObjectPtr<UAvalanchePagePlayer>& PagePlayer : InPlaylist->GetPagePlayers())
+		{
+			// Early filter on preview/channel.
+			if (!PagePlayer
+				|| PagePlayer->bIsPreview != bInIsPreview
+				|| (bInIsPreview && PagePlayer->ChannelName != InPreviewChannelName) )
+			{
+				continue;
+			}
+			
+			const FAvalanchePage& OtherPage = InPlaylist->GetPage(PagePlayer->PageId);
+
+			// Check if same template.
+			if (!OtherPage.IsValidPage() || OtherPage.GetTemplateId() != InTemplate.GetPageId())
+			{
+				continue;
+			}
+			
+			// Check if same RC values (of the sub-template).
+			const FAvalanchePage& SubTemplate = InTemplate.GetTemplate(InPlaylist, InSubPageIndex);
+			if (SubTemplate.IsValidPage() && ArePageRCValuesEqualForSubTemplate(SubTemplate, InPageToPlay, OtherPage))
+			{
+				return PagePlayer->FindInstancePlayerByAssetPath(SubTemplate.GetAvalancheAssetPath(InPlaylist));
+			}
+		}
+
+		return nullptr;
+	}
+}
+
+bool UAvalanchePlaylist::PlayPageWithTransition(FAvalanchePageTransitionBuilder& InBuilder, const FAvalanchePage& InPage, EAvaPlayType InPlayType, bool bInIsPreview, const FName& InPreviewChannelName)
+{
+	using namespace UE::AvalancheMedia::Playlist::Private;
+
+	// For now, we always start a new page player, loading a new instance.
+	UAvalanchePagePlayer* NewPagePlayer = NewObject<UAvalanchePagePlayer>(this);
+
+	if (!NewPagePlayer || !NewPagePlayer->Initialize(this, InPage, bInIsPreview, InPreviewChannelName))
+	{
+		return false;
+	}
+	
+	// -- TL vs No-TL pages mutual exclusion rule.
+	// The way it is resolved for now, it is first come first serve. The pages that are entered first in
+	// the transition will win. We may want to have say, the first no-TL page win. Will have to see what
+	// is the best rule after some testing.
+	if (const UAvalanchePageTransition* ExistingPageTransition = InBuilder.FindTransition(NewPagePlayer))
+	{
+		// Don't add a page with TL in a transition that has a non-TL page. Can't co-exist.
+		if (InPage.HasTransitionLogic(this))
+		{
+			if (ExistingPageTransition->HasEnterPagesWithNoTransitionLogic())
+			{
+				return false;	
+			}
+
+			// This is the layer exclusion rule. Rejects the page if any of the layers
+			// are already in the transition. This is to prevent combo pages to start
+			// "on top" of another page with same layer (in same transition only).
+			for (const FAvaTagHandle& TagHandle : InPage.GetTransitionLayers(this))
+			{
+				if (ExistingPageTransition->ContainsTransitionLayer(TagHandle.TagId))
+				{
+					return false;
+				}
+			}
+		}
+		// Don't add a page with no-TL in a transition that has enter pages already (any page).
+		else if (ExistingPageTransition->HasEnterPages())
+		{
+			return false;
+		}
+	}
+	
+	// Load or Recycle Instance Players.
+	const int32 NumTemplates = InPage.GetNumTemplates(this);
+	NewPagePlayer->InstancePlayers.Reserve(NumTemplates);
+
+	const FAvalanchePage& DirectTemplate = InPage.ResolveTemplate(this);
+
+	for (int32 SubPageIndex = 0; SubPageIndex < NumTemplates; ++SubPageIndex)
+	{
+		bool bFoundExistingInstancePlayer = false;
+
+		const UAvalancheMediaSettings& AvaMediaSettings = UAvalancheMediaSettings::Get();
+		const bool bUseSpecialTransitionLogic = DirectTemplate.IsComboTemplate()
+			? AvaMediaSettings.bEnableComboTemplateSpecialLogic
+			: AvaMediaSettings.bEnableSingleTemplateSpecialLogic;
+		
+		// -- Special Transition Logic --
+		if (bUseSpecialTransitionLogic)
+		{
+			// Try to find an existing instance player of the same combo template, sub-template and RC values.
+			UAvaRundownPlaybackInstancePlayer* InstancePlayer =
+				FindExistingInstancePlayer(this, InPage, DirectTemplate, SubPageIndex, bInIsPreview, InPreviewChannelName);
+			
+			if (InstancePlayer && InstancePlayer->PlaybackInstance)
+			{
+				NewPagePlayer->AddInstancePlayer(InstancePlayer);
+
+				// Setup user instance data to be able to track this page.
+				UAvalanchePagePlayer::SetInstanceUserDataFromPage(*InstancePlayer->PlaybackInstance, InPage);
+				
+				// This instance has to be excluded from the next playable transition.
+				NewPagePlayer->InstancesExcludedFromTransition.Add(InstancePlayer->GetPlaybackInstanceId());
+				bFoundExistingInstancePlayer = true;
+			}
+		}
+
+		if (!bFoundExistingInstancePlayer)
+		{
+			NewPagePlayer->LoadInstancePlayer(SubPageIndex, FGuid());
+		}
+	}
+
+	if (NewPagePlayer->IsLoaded())
+	{
+		if (UAvalanchePageTransition* PageTransition = InBuilder.FindOrAddTransition(NewPagePlayer))
+		{
+			if (PageTransition->AddEnterPage(NewPagePlayer))
+			{
+				AddPagePlayer(NewPagePlayer);
+
+				// Start the playback, will only actually start on next tick.
+				// Animation command will not be pushed, relying on TL to start the appropriate animations.
+				NewPagePlayer->Play(InPlayType, true);
+				return true;
+			}
+		}
+	}
+	return false;
+}
+
+bool UAvalanchePlaylist::StopPageNoTransition(const FAvalanchePage& InPage, bool bInPreview, const FName& InPreviewChannelName)
+{
+	if (UAvalanchePagePlayer* PagePlayer = FindPlayerForPage(InPage.GetPageId(), bInPreview, InPreviewChannelName))
+	{
+		const bool bPlayerStopped = PagePlayer->Stop();
+		RemoveStoppedPagePlayers();
+		return bPlayerStopped;
+	}
+	return false;
+}
+
+bool UAvalanchePlaylist::StopPageWithTransition(FAvalanchePageTransitionBuilder& InBuilder, const FAvalanchePage& InPage, bool bInPreview, const FName& InPreviewChannelName)
+{
+	if (UAvalanchePagePlayer* PagePlayer = FindPlayerForPage(InPage.GetPageId(), bInPreview, InPreviewChannelName))
+	{
+		if (UAvalanchePageTransition* PageTransition = InBuilder.FindOrAddTransition(PagePlayer))
+		{
+			PageTransition->AddExitPage(PagePlayer);
+			return true;
+		}
+	}
+	return false;
+}
+
+const FAvalanchePage& UAvalanchePlaylist::GetNextFromPages(const TArray<FAvalanchePage>& InPages, int32 InStartingIndex) const
+{
+	if (InPages.IsEmpty())
+	{
+		return FAvalanchePage::NullPage;
+	}
+
+	int32 NextIndex = InStartingIndex;
+	const FAvalanchePage& CurrentPage = InPages[InStartingIndex];
+	do
+	{
+		if (InPages.IsValidIndex(++NextIndex))
+		{
+			const FAvalanchePage& NextPage = InPages[NextIndex];
+			if (NextPage.GetChannelName() == CurrentPage.GetChannelName())
+			{
+				return NextPage;
+			}
+		}
+		else
+		{
+			NextIndex = FAvalanchePage::InvalidPageId;
+		}
+	}
+	while (NextIndex != InStartingIndex);
+
+	return FAvalanchePage::NullPage;
+}
+
+FAvalanchePage& UAvalanchePlaylist::GetNextFromPages(TArray<FAvalanchePage>& InPages, int32 InStartingIndex) const
+{
+	if (InPages.IsEmpty())
+	{
+		return FAvalanchePage::NullPage;
+	}
+
+	int32 NextIndex = InStartingIndex;
+	const FAvalanchePage& CurrentPage = InPages[InStartingIndex];
+	do
+	{
+		if (InPages.IsValidIndex(++NextIndex))
+		{
+			FAvalanchePage& NextPage = InPages[NextIndex];
+			if (NextPage.GetChannelName() == CurrentPage.GetChannelName())
+			{
+				return NextPage;
+			}
+		}
+		else
+		{
+			NextIndex = FAvalanchePage::InvalidPageId;
+		}
+	}
+	while (NextIndex != InStartingIndex);
+
+	return FAvalanchePage::NullPage;
+}
+
+const FAvalanchePage& UAvalanchePlaylist::GetNextFromSubList(const TArray<int32>& InSubListIds, int32 InStartingIndex) const
+{
+	if (InSubListIds.IsEmpty())
+	{
+		return FAvalanchePage::NullPage;
+	}
+
+	int32 NextIndex = InStartingIndex;
+	const FAvalanchePage& CurrentPage = GetPage(InSubListIds[InStartingIndex]);
+	do
+	{
+		if (InSubListIds.IsValidIndex(++NextIndex))
+		{
+			const FAvalanchePage& NextPage = GetPage(InSubListIds[NextIndex]);
+			if (NextPage.GetChannelName() == CurrentPage.GetChannelName())
+			{
+				return NextPage;
+			}
+		}
+		else
+		{
+			NextIndex = FAvalanchePage::InvalidPageId;
+		}
+	}
+	while (NextIndex != InStartingIndex);
+
+	return FAvalanchePage::NullPage;
+}
+
+FAvalanchePage& UAvalanchePlaylist::GetNextFromSubList(TArray<int32>& InSubListIds, int32 InStartingIndex)
+{
+	if (InSubListIds.IsEmpty())
+	{
+		return FAvalanchePage::NullPage;
+	}
+
+	int32 NextIndex = InStartingIndex;
+	const FAvalanchePage& CurrentPage = GetPage(InSubListIds[InStartingIndex]);
+	do
+	{
+		if (InSubListIds.IsValidIndex(++NextIndex))
+		{
+			FAvalanchePage& NextPage = GetPage(InSubListIds[NextIndex]);
+			if (NextPage.GetChannelName() == CurrentPage.GetChannelName())
+			{
+				return NextPage;
+			}
+		}
+		else
+		{
+			NextIndex = FAvalanchePage::InvalidPageId;
+		}
+	}
+	while (NextIndex != InStartingIndex);
+
+	return FAvalanchePage::NullPage;
+}
+
+bool UAvalanchePlaylist::CanStartTransitionForPage(const FAvalanchePage& InPage, bool bInIsPreview, const FName& InPreviewChannelName) const
+{
+	// Current constraint: There can only be one transition (running properly) at a time in a world.
+	// Given that, for now (and the foreseeable future until we support more playable groups per channels),
+	// we can equate a "channel" to a "world", this is hardcoded for the level streaming playables.
+	// So, we can just check the channels for now.
+	const FName ChannelName = bInIsPreview ? InPreviewChannelName : InPage.GetChannelName();
+	for (const TObjectPtr<UAvalanchePageTransition>& PageTransition : PageTransitions)
+	{
+		if (PageTransition && PageTransition->GetChannelName() == ChannelName)
+		{
+			return false;
+		}
+	}
+	return true;
+}
+
+void UAvalanchePlaylist::StopPageTransitionsForPage(const FAvalanchePage& InPage, bool bInIsPreview, const FName& InPreviewChannelName)
+{
+	TArray<UAvalanchePageTransition*> TransitionsToStop;
+	TransitionsToStop.Reserve(PageTransitions.Num());
+
+	// Note: we build a separate list because stopping the transitions should
+	// lead to the transitions being removed from PageTransitions (through the events).
+	const FName ChannelName = bInIsPreview ? InPreviewChannelName : InPage.GetChannelName();
+	for (TObjectPtr<UAvalanchePageTransition>& PageTransition : PageTransitions)
+	{
+		if (PageTransition && PageTransition->GetChannelName() == ChannelName)
+		{
+			TransitionsToStop.Add(PageTransition);	
+		}
+	}
+
+	for (UAvalanchePageTransition* Transition : TransitionsToStop)
+	{
+		Transition->Stop();
+
+		// Normal course of events should have removed the transition, but
+		// if something is wrong with the events, we double check it is indeed removed.
+		if (PageTransitions.Contains(Transition))
+		{
+			UE_LOG(LogAvaPlaylist, Warning, TEXT("A page transition was not properly cleaned up."));
+			PageTransitions.Remove(Transition);
+		}
+	}
+}
+
+void UAvalanchePlaylist::RemoveStoppedPagePlayers()
+{
+	for (UAvalanchePagePlayer* PagePlayer : PagePlayers)
+	{
+		if (PagePlayer && !PagePlayer->IsPlaying())
+		{
+			OnPagePlayerRemoving.Broadcast(this, PagePlayer);
+		}
+	}
+	
+	PagePlayers.RemoveAll([](const UAvalanchePagePlayer* InPagePlayer) { return !InPagePlayer || InPagePlayer->IsPlaying() == false;});
+}
+
+
+#undef LOCTEXT_NAMESPACE
diff --git a/Engine/Plugins/Experimental/Avalanche/Source/AvalancheMedia/Private/Playlist/AvalanchePlaylistComponent.cpp b/Engine/Plugins/Experimental/Avalanche/Source/AvalancheMedia/Private/Playlist/AvalanchePlaylistComponent.cpp
new file mode 100644
--- /dev/null
+++ b/Engine/Plugins/Experimental/Avalanche/Source/AvalancheMedia/Private/Playlist/AvalanchePlaylistComponent.cpp
@@ -0,0 +1,73 @@
+// Copyright Epic Games, Inc. All Rights Reserved.
+
+#include "Playlist/AvalanchePlaylistComponent.h"
+
+#include "Engine/World.h"
+#include "IAvaMediaModule.h"
+#include "Playback/AvaMediaPlaybackManager.h"
+#include "Playlist/AvalanchePlaylist.h"
+
+UAvalanchePlaylistComponent::UAvalanchePlaylistComponent(const FObjectInitializer& ObjectInitializer)
+	: Super(ObjectInitializer)
+{
+	bWantsInitializeComponent = true;
+}
+
+bool UAvalanchePlaylistComponent::PlayPage(int32 InPageId)
+{
+	if (Playlist && !Playlist->IsPagePlaying(InPageId))
+	{
+		return Playlist->PlayPage(InPageId, EAvaPlayType::PlayFromStart);
+	}
+	return false;
+}
+
+bool UAvalanchePlaylistComponent::StopPage(int32 InPageId)
+{
+	if (Playlist && Playlist->IsPagePlaying(InPageId))
+	{
+		return Playlist->StopPage(InPageId, EAvaPlaylistPageStopOptions::Default, false);
+	}
+	return false;
+}
+
+int32 UAvalanchePlaylistComponent::GetNumberOfPages() const
+{
+	return Playlist ? Playlist->GetInstancedPages().Pages.Num() : 0;
+}
+
+int32 UAvalanchePlaylistComponent::GetPageIdForIndex(int32 InPageIndex) const
+{
+	if (Playlist && Playlist->GetInstancedPages().Pages.IsValidIndex(InPageIndex))
+	{
+		return Playlist->GetInstancedPages().Pages[InPageIndex].GetPageId();
+	}
+	return 0;
+}
+
+void UAvalanchePlaylistComponent::InitializeComponent()
+{
+	Super::InitializeComponent();
+	FWorldDelegates::OnWorldBeginTearDown.AddUObject(this, &UAvalanchePlaylistComponent::OnWorldBeginTearDown);
+}
+
+void UAvalanchePlaylistComponent::UninitializeComponent()
+{
+	FWorldDelegates::OnWorldBeginTearDown.RemoveAll(this);
+	Super::UninitializeComponent();
+}
+
+void UAvalanchePlaylistComponent::OnWorldBeginTearDown(UWorld* InWorld)
+{
+	if (InWorld == GetWorld())
+	{
+		IAvaMediaModule::Get().GetLocalPlaybackManager().OnParentWorldBeginTearDown();
+		if (Playlist)
+		{
+			// Since we have forcibly teared down the manager, update the
+			// status of the playlist to reflect the appropriate state.
+			Playlist->OnParentWordBeginTearDown();
+		}
+	}
+}
+
