[MotionDesign]

diff --git a/Engine/Plugins/Experimental/Avalanche/Source/AvalancheMedia/Private/Playlist/Transition/AvalanchePageTransition.cpp b/Engine/Plugins/Experimental/Avalanche/Source/AvalancheMedia/Private/Playlist/Transition/AvalanchePageTransition.cpp
new file mode 100644
--- /dev/null
+++ b/Engine/Plugins/Experimental/Avalanche/Source/AvalancheMedia/Private/Playlist/Transition/AvalanchePageTransition.cpp
@@ -0,0 +1,642 @@
+// Copyright Epic Games, Inc. All Rights Reserved.
+
+#include "Playlist/Transition/AvalanchePageTransition.h"
+
+#include "Playback/AvaMediaPlaybackManager.h"
+#include "Playback/AvaMediaPlaybackUtils.h"
+#include "Playback/Transition/AvalanchePlayableTransition.h"
+#include "Playlist/AvalanchePagePlayer.h"
+#include "Playlist/AvalanchePlaylist.h"
+
+namespace UE::AvaMedia::Rundown::PageTransition::Private
+{
+	UAvalanchePlayable* GetPlayable(const UAvaRundownPlaybackInstancePlayer* InInstancePlayer)
+	{
+		return InInstancePlayer ? InInstancePlayer->GetFirstPlayable() : nullptr;
+	}
+	
+	UAvalanchePlayable* GetPlayable(const UAvalanchePagePlayer* InPagePlayer, int32 InIndex)
+	{
+		return InPagePlayer ? GetPlayable(InPagePlayer->GetInstancePlayer(InIndex)) : nullptr;
+	}
+
+	UAvalanchePlayback* GetPlayback(const UAvaRundownPlaybackInstancePlayer* InInstancePlayer)
+	{
+		return InInstancePlayer ? InInstancePlayer->Playback : nullptr; 
+	}
+	
+	UAvalanchePlayback* GetPlayback(const UAvalanchePagePlayer* InPagePlayer, int32 InIndex)
+	{
+		return InPagePlayer ? GetPlayback(InPagePlayer->GetInstancePlayer(InIndex)) : nullptr;
+	}
+
+	TSharedPtr<FAvalancheRemoteControlValues> GetRemoteControlValues(const UAvalanchePagePlayer* InPagePlayer)
+	{
+		if (UAvalanchePlaylist* Playlist = InPagePlayer->GetPlaylist())
+		{
+			const FAvalanchePage& Page = Playlist->GetPage(InPagePlayer->PageId);
+			if (Page.IsValidPage())
+			{
+				return MakeShared<FAvalancheRemoteControlValues>(Page.GetRemoteControlValues());
+			}
+		}
+		return nullptr;
+	}
+
+	FString GetPrettyPageInfo(const UAvalanchePagePlayer* InPagePlayer)
+	{
+		FString Info;
+		Info += FString(TEXT("Channel: ")) + InPagePlayer->ChannelName;
+		InPagePlayer->ForEachInstancePlayer([&Info](const UAvaRundownPlaybackInstancePlayer* InInstancePlayer)
+		{
+			Info += FString(TEXT(", Instance Id: ")) + InInstancePlayer->GetPlaybackInstanceId().ToString();
+			Info += FString(TEXT(", Asset: ")) + InInstancePlayer->SourceAssetPath.ToString();
+		});
+		return Info;
+	}
+	
+	UAvaRundownPlaybackInstancePlayer* FindInstancePlayerForPlayable(const UAvalanchePlayable* InPlayable, const TArray<TWeakObjectPtr<UAvalanchePagePlayer>>& InPagePlayersWeak)
+	{
+		for (const TWeakObjectPtr<UAvalanchePagePlayer>& PlayerWeak : InPagePlayersWeak)
+		{
+			if (const UAvalanchePagePlayer* PagePlayer = PlayerWeak.Get())
+			{
+				if (UAvaRundownPlaybackInstancePlayer* InstancePlayer = PagePlayer->FindInstancePlayerForPlayable(InPlayable))
+				{
+					return InstancePlayer;
+				}
+			}
+		}
+		return nullptr;
+	}
+
+	UAvalanchePagePlayer* FindPagePlayerForPlayable(const UAvalanchePlayable* InPlayable, const TArray<TWeakObjectPtr<UAvalanchePagePlayer>>& InPagePlayersWeak)
+	{
+		for (const TWeakObjectPtr<UAvalanchePagePlayer>& PlayerWeak : InPagePlayersWeak)
+		{
+			UAvalanchePagePlayer* PagePlayer = PlayerWeak.Get();
+			if (PagePlayer && PagePlayer->HasPlayable(InPlayable))
+			{
+				return PagePlayer;
+			}
+		}
+		return nullptr;
+	}
+}
+
+bool UAvalanchePageTransition::AddEnterPage(UAvalanchePagePlayer* InPagePlayer)
+{
+	using namespace UE::AvaMedia::Rundown::PageTransition::Private;
+
+	if (!InPagePlayer)
+	{
+		return false;	
+	}
+
+	// Multi-page constraint: Prevent pages on the same layer
+	for (const TWeakObjectPtr<UAvalanchePagePlayer>& PlayerWeak : EnterPlayersWeak)
+	{
+		if (const UAvalanchePagePlayer* Player = PlayerWeak.Get())
+		{
+			bool bHasLayer = false;
+			InPagePlayer->ForEachInstancePlayer([this, &bHasLayer](const UAvaRundownPlaybackInstancePlayer* InInstancePlayer)
+			{
+				if (CachedTransitionLayers.Contains(InInstancePlayer->TransitionLayer.TagId))
+				{
+					bHasLayer = true;
+				}
+			});
+			
+			if (bHasLayer)
+			{
+				UE_LOG(LogAvaPlaylist, Error,
+						TEXT("Page Transition \"%s\" Error: page %d can't be played with page %d because they are on the same layer."),
+						*GetFullName(), InPagePlayer->PageId, Player->PageId);
+				return false;
+			}
+		}
+	}
+
+	RegisterEnterPagePlayerEvents(InPagePlayer);
+	
+	return AddPagePlayer(InPagePlayer, EnterPlayersWeak);
+}
+
+bool UAvalanchePageTransition::AddPlayingPage(UAvalanchePagePlayer* InPagePlayer)
+{
+	return AddPagePlayer(InPagePlayer, PlayingPlayersWeak);
+}
+
+bool UAvalanchePageTransition::AddExitPage(UAvalanchePagePlayer* InPagePlayer)
+{
+	return AddPagePlayer(InPagePlayer, ExitPlayersWeak);
+}
+
+bool UAvalanchePageTransition::IsVisibilityConstrained(const UAvalanchePlayable* InPlayable) const
+{
+	using namespace UE::AvaMedia::Rundown::PageTransition::Private;
+	bool bAllPlayablesLoaded = true;
+	bool bIsPlayableInThisTransition = false;
+	
+	for (const TWeakObjectPtr<UAvalanchePagePlayer>& PlayerWeak : EnterPlayersWeak)
+	{
+		if (const UAvalanchePagePlayer* Player = PlayerWeak.Get())
+		{
+			Player->ForEachInstancePlayer([InPlayable, &bIsPlayableInThisTransition, &bAllPlayablesLoaded](const UAvaRundownPlaybackInstancePlayer* InInstancePlayer)
+			{
+				if (const UAvalanchePlayable* Playable = GetPlayable(InInstancePlayer))
+				{
+					if (Playable == InPlayable)
+					{
+						bIsPlayableInThisTransition = true;
+					}
+					const EAvalanchePlayableStatus PlayableStatus = Playable->GetPlayableStatus();
+					if (PlayableStatus != EAvalanchePlayableStatus::Loaded && PlayableStatus != EAvalanchePlayableStatus::Visible)
+					{
+						bAllPlayablesLoaded = false;
+					}
+				}
+			});
+		}
+	}
+	
+	return bIsPlayableInThisTransition && !bAllPlayablesLoaded;
+}
+
+bool UAvalanchePageTransition::CanStart(bool& bOutShouldDiscard) const
+{
+	using namespace UE::AvaMedia::Rundown::PageTransition::Private;
+	for (const TWeakObjectPtr<UAvalanchePagePlayer>& PlayerWeak : EnterPlayersWeak)
+	{
+		if (const UAvalanchePagePlayer* Player = PlayerWeak.Get())
+		{
+			const int32 NumInstancePlayers = Player->GetNumInstancePlayers();
+			for (int32 InstancePlayerIndex = 0; InstancePlayerIndex < NumInstancePlayers; ++InstancePlayerIndex)
+			{
+				if (const UAvalanchePlayable* Playable = GetPlayable(Player, InstancePlayerIndex))
+				{
+					const EAvalanchePlayableStatus PlayableStatus = Playable->GetPlayableStatus();
+
+					if (PlayableStatus == EAvalanchePlayableStatus::Unknown
+						|| PlayableStatus == EAvalanchePlayableStatus::Error)
+					{
+						// Discard the command
+						bOutShouldDiscard = true;
+						return false;
+					}
+
+					// todo: this might cause commands to become stale and fill the pending command list
+					if (PlayableStatus == EAvalanchePlayableStatus::Unloaded)
+					{
+						return false;
+					}
+
+					// Asset status must be visible to run the command locally.
+					// If not visible, the components are not yet added to the world.
+					// For remote proxies, we can run the command immediately, it will wait for the asset to be visible on the server instead. 
+					if (!Playable->IsRemoteProxy() && PlayableStatus != EAvalanchePlayableStatus::Visible)
+					{
+						// Keep the command in the queue for next tick.
+						bOutShouldDiscard = false;
+						return false;
+					}
+				}
+				else
+				{
+					// Playables not created - Keep the command in the queue for next tick.
+					bOutShouldDiscard = false;
+					return false;
+				}
+			}
+		}
+		else
+		{
+			bOutShouldDiscard = true;
+			return false;
+		}
+	}
+	
+	bOutShouldDiscard = true;
+	return true;
+}
+
+void UAvalanchePageTransition::Start()
+{
+	RegisterToPlayableTransitionEvent();
+
+	// May fail if playables are not loaded yet. Playables are loaded
+	// when the playback object has ticked at least one.
+	MakePlayableTransition();
+
+	bool bTransitionStarted = false;
+
+	if (PlayableTransition)
+	{
+		LogDetailedTransitionInfo();
+
+		// Todo: validate the level streaming playables are finished streaming the asset.
+		// Otherwise, transition start must be queued on playable streaming events. 
+		bTransitionStarted = PlayableTransition->Start();
+	}
+	
+	if (!bTransitionStarted)
+	{
+		Stop();
+	}
+}
+
+void UAvalanchePageTransition::Stop()
+{
+	if (PlayableTransition)
+	{
+		PlayableTransition->Stop();
+		PlayableTransition = nullptr;
+	}
+
+	for (TWeakObjectPtr<UAvalanchePagePlayer>& PagePlayerWeak : EnterPlayersWeak)
+	{
+		if (UAvalanchePagePlayer* PagePlayer = PagePlayerWeak.Get())
+		{
+			PagePlayer->InstancesExcludedFromTransition.Reset();
+			UnregisterEnterPagePlayerEvents(PagePlayer);
+		}
+	}
+
+	UnregisterFromPlayableTransitionEvent();
+
+	if (UAvalanchePlaylist* Playlist = GetPlaylist())
+	{
+		Playlist->RemovePageTransition(this);
+		Playlist->RemoveStoppedPagePlayers();
+	}
+	else
+	{
+		UE_LOG(LogAvaPlaylist, Error, TEXT("Page Transition \"%s\" Failed to remove transition: No playlist specified."), *GetFullName());
+	}
+}
+
+bool UAvalanchePageTransition::IsRunning() const
+{
+	return 	PlayableTransition ? PlayableTransition->IsRunning() : false;
+}
+
+bool UAvalanchePageTransition::HasEnterPagesWithNoTransitionLogic() const
+{
+	const UAvalanchePlaylist* Playlist = GetPlaylist();
+	if (!Playlist)
+	{
+		return false;
+	}
+	
+	for (const TWeakObjectPtr<UAvalanchePagePlayer>& PlayerWeak : EnterPlayersWeak)
+	{
+		if (const UAvalanchePagePlayer* Player = PlayerWeak.Get())
+		{
+			// Don't rely on instance players at this point. Assets may not be loaded.
+			const FAvalanchePage& Page = Playlist->GetPage(Player->PageId);
+			if (Page.IsValidPage() && !Page.HasTransitionLogic(Playlist))
+			{
+				return true;
+			}
+		}
+	}
+	return false;
+}
+
+bool UAvalanchePageTransition::HasPagePlayer(const UAvalanchePagePlayer* InPagePlayer) const
+{
+	return EnterPlayersWeak.Contains(InPagePlayer) || PlayingPlayersWeak.Contains(InPagePlayer) || ExitPlayersWeak.Contains(InPagePlayer);
+}
+
+bool UAvalanchePageTransition::ContainsTransitionLayer(const FAvaTagId& InTagId) const
+{
+	return CachedTransitionLayers.Contains(InTagId);
+}
+
+UAvalanchePlaylist* UAvalanchePageTransition::GetPlaylist() const
+{
+	return Cast<UAvalanchePlaylist>(GetOuter());
+}
+
+UAvaRundownPlaybackInstancePlayer* UAvalanchePageTransition::FindInstancePlayerForPlayable(const UAvalanchePlayable* InPlayable) const
+{
+	using namespace UE::AvaMedia::Rundown::PageTransition;
+	if (!InPlayable)
+	{
+		return nullptr;
+	}
+	
+	if (UAvaRundownPlaybackInstancePlayer* FoundInstancePlayer = Private::FindInstancePlayerForPlayable(InPlayable, EnterPlayersWeak))
+	{
+		return FoundInstancePlayer;
+	}
+	if (UAvaRundownPlaybackInstancePlayer* FoundInstancePlayer = Private::FindInstancePlayerForPlayable(InPlayable, PlayingPlayersWeak))
+	{
+		return FoundInstancePlayer;
+	}
+	if (UAvaRundownPlaybackInstancePlayer* FoundInstancePlayer = Private::FindInstancePlayerForPlayable(InPlayable, ExitPlayersWeak))
+	{
+		return FoundInstancePlayer;
+	}
+	return nullptr;
+}
+
+UAvalanchePagePlayer* UAvalanchePageTransition::FindPagePlayerForPlayable(const UAvalanchePlayable* InPlayable) const
+{
+	using namespace UE::AvaMedia::Rundown::PageTransition;
+	if (UAvalanchePagePlayer* FoundPagePlayer = Private::FindPagePlayerForPlayable(InPlayable, EnterPlayersWeak))
+	{
+		return FoundPagePlayer;
+	}
+	if (UAvalanchePagePlayer* FoundPagePlayer = Private::FindPagePlayerForPlayable(InPlayable, PlayingPlayersWeak))
+	{
+		return FoundPagePlayer;
+	}
+	if (UAvalanchePagePlayer* FoundPagePlayer = Private::FindPagePlayerForPlayable(InPlayable, ExitPlayersWeak))
+	{
+		return FoundPagePlayer;
+	}
+	return nullptr;
+}
+
+void UAvalanchePageTransition::OnTransitionEvent(UAvalanchePlayable* InPlayable, UAvalanchePlayableTransition* InTransition, EAvalanchePlayableTransitionEventFlags InTransitionFlags)
+{
+	using namespace UE::AvaMediaPlayback::Utils;
+
+	// not this transition.
+	if (InTransition != PlayableTransition || PlayableTransition == nullptr)
+	{
+		return;
+	}
+
+	// Find the page player for this playable
+	UAvaRundownPlaybackInstancePlayer* InstancePlayer = FindInstancePlayerForPlayable(InPlayable);
+	UAvalanchePagePlayer* PagePlayer = InstancePlayer ? InstancePlayer->GetPagePlayer() : nullptr;
+
+	if (InstancePlayer && EnumHasAnyFlags(InTransitionFlags, EAvalanchePlayableTransitionEventFlags::MarkPlayableDiscard))
+	{
+		UE_LOG(LogAvaPlaylist, Verbose, TEXT("%s Instance Player Marked for Discard: Id:%s, Asset:\"%s\""),
+			*GetBriefFrameInfo(), *InstancePlayer->GetPlaybackInstanceId().ToString(), *InstancePlayer->SourceAssetPath.ToString());
+	}
+	
+	if (InstancePlayer && EnumHasAnyFlags(InTransitionFlags, EAvalanchePlayableTransitionEventFlags::StopPlayable))
+	{
+		// Validating that we are not removing an "enter" playable.
+		if (PlayableTransition->IsEnterPlayable(InPlayable))
+		{
+			UE_LOG(LogAvaPlaylist, Error,
+				TEXT("Page Transition \"%s\" Error: An \"enter\" playable is being discarded for page %d."),
+				*GetFullName(), (PagePlayer ? PagePlayer->PageId : -1));
+		}
+
+		UE_LOG(LogAvaPlaylist, Verbose, TEXT("%s Stopping Instance Player: Id:%s, Asset:\"%s\""),
+			*GetBriefFrameInfo(), *InstancePlayer->GetPlaybackInstanceId().ToString(), *InstancePlayer->SourceAssetPath.ToString());
+
+		// With combo-templates, page players can be partially stopped.
+		InstancePlayer->Stop();
+
+		// Check if the page player is still playing (i.e. if all instance players have been stopped)
+		if (PagePlayer && !PagePlayer->IsPlaying())
+		{
+			UE_LOG(LogAvaPlaylist, Verbose, TEXT("%s Stopping Page Player: PageId:%d"), *GetBriefFrameInfo(), PagePlayer->PageId);
+
+			// Stop the whole page player and propagate page events.
+			PagePlayer->Stop();
+
+			if (UAvalanchePlaylist* Playlist = PagePlayer->GetPlaylist())
+			{
+				Playlist->RemoveStoppedPagePlayers();
+			}
+			else
+			{
+				UE_LOG(LogAvaPlaylist, Error, TEXT("Page Transition \"%s\" failed to remove stopped players: No playlist specified."), *GetFullName());
+			}
+		}
+	}
+	
+	if (EnumHasAnyFlags(InTransitionFlags, EAvalanchePlayableTransitionEventFlags::Finished))
+	{
+		using namespace UE::AvaMediaPlayback::Utils;
+		UE_LOG(LogAvaPlaylist, Verbose, TEXT("%s Finishing Page Transition: %s"), *GetBriefFrameInfo(), *GetBriefTransitionDescription());
+
+		Stop();
+	}
+}
+
+void UAvalanchePageTransition::OnPlayableCreated(UAvalanchePlayback* InPlayback, UAvalanchePlayable* InPlayable)
+{
+	if (UAvaMediaPlayableGroup* PlayableGroup = InPlayable->GetPlayableGroup())
+	{
+		PlayableGroup->RegisterVisibilityConstraint(this);
+	}
+}
+
+void UAvalanchePageTransition::AddPlayersToBuilder(
+	FAvaPlayableTransitionBuilder& InOutBuilder, const TArray<TWeakObjectPtr<UAvalanchePagePlayer>>& InPlayersWeak, const TCHAR* InCategory, EAvaMediaPlayableTransitionEntryRole InEntryRole) const
+{
+	for (const TWeakObjectPtr<UAvalanchePagePlayer>& PlayerWeak : InPlayersWeak)
+	{
+		if (UAvalanchePagePlayer* Player = PlayerWeak.Get())
+		{
+			AddPlayablesToBuilder(InOutBuilder, Player, InCategory, InEntryRole);
+			
+			// Excluded instances are only for current transition.
+			// Todo: Ideally, it would be stored in the transition itself to avoid the clean up.
+			Player->InstancesExcludedFromTransition.Reset();
+		}
+	}
+}
+
+void UAvalanchePageTransition::AddPlayablesToBuilder(FAvaPlayableTransitionBuilder& InOutBuilder, const UAvalanchePagePlayer* InPlayer, const TCHAR* InCategory, EAvaMediaPlayableTransitionEntryRole InEntryRole) const
+{
+	using namespace UE::AvaMedia::Rundown::PageTransition::Private;
+
+	for (const TObjectPtr<UAvaRundownPlaybackInstancePlayer>& InstancePlayer : InPlayer->InstancePlayers)
+	{
+		if (InPlayer->InstancesExcludedFromTransition.Contains(InstancePlayer->GetPlaybackInstanceId()))
+		{
+			continue;
+		}
+	
+		if (UAvalanchePlayable* Playable = GetPlayable(InstancePlayer))
+		{
+			const bool bPlayableAdded = InOutBuilder.AddPlayable(Playable, InEntryRole);
+			if (InEntryRole == EAvaMediaPlayableTransitionEntryRole::Enter && bPlayableAdded)
+			{
+				InOutBuilder.AddEnterPlayableValues(GetRemoteControlValues(InPlayer));
+			}
+		}
+		else
+		{
+			// If this happens, likely the playable is not yet loaded.
+			UE_LOG(LogAvaPlaylist, Error,
+				TEXT("Page Transition \"%s\" Error: Failed to retrieve \"%s\" playable for page %d."),
+				*GetFullName(), InCategory, InPlayer->PageId);
+		}
+	}
+}
+
+void UAvalanchePageTransition::MakePlayableTransition()
+{
+	FAvaPlayableTransitionBuilder TransitionBuilder;
+	AddPlayersToBuilder(TransitionBuilder, EnterPlayersWeak, TEXT("Enter"), EAvaMediaPlayableTransitionEntryRole::Enter);
+	AddPlayersToBuilder(TransitionBuilder, PlayingPlayersWeak, TEXT("Playing"), EAvaMediaPlayableTransitionEntryRole::Playing);
+	AddPlayersToBuilder(TransitionBuilder, ExitPlayersWeak, TEXT("Exit"), EAvaMediaPlayableTransitionEntryRole::Exit);
+	PlayableTransition = TransitionBuilder.MakeTransition(this);
+
+	// Mark exit only transitions to ensure they properly create the null behavior instances. 
+	if (PlayableTransition && EnterPlayersWeak.IsEmpty())
+	{
+		PlayableTransition->SetTransitionFlags(EAvalanchePlayableTransitionFlags::ExitOnly);
+	}
+}
+
+void UAvalanchePageTransition::LogDetailedTransitionInfo() const
+{
+	if (!UE_LOG_ACTIVE(LogAvaPlaylist, Verbose))
+	{
+		return;
+	}
+	
+	using namespace UE::AvaMediaPlayback::Utils;
+	UE_LOG(LogAvaPlaylist, Verbose, TEXT("%s Starting Page Transition \"%s\":"), *GetBriefFrameInfo(), *GetFullName());
+	
+	auto LogPlayers = [this](const TArray<TWeakObjectPtr<UAvalanchePagePlayer>>& InPlayersWeak, const TCHAR* InCategory)
+	{
+		using namespace UE::AvaMedia::Rundown::PageTransition::Private;
+		for (const TWeakObjectPtr<UAvalanchePagePlayer>& PlayerWeak : InPlayersWeak)
+		{
+			if (const UAvalanchePagePlayer* Player = PlayerWeak.Get())
+			{
+				UE_LOG(LogAvaPlaylist, Verbose, TEXT("- %s Page: %d, %s."), InCategory, Player->PageId, *GetPrettyPageInfo(Player));
+			}
+		}
+	};
+
+	LogPlayers(EnterPlayersWeak, TEXT("Enter"));
+	LogPlayers(PlayingPlayersWeak, TEXT("Playing"));
+	LogPlayers(ExitPlayersWeak, TEXT("Exit"));
+}
+
+FString UAvalanchePageTransition::GetBriefTransitionDescription() const
+{
+	auto MakePageIdList = [](const TArray<TWeakObjectPtr<UAvalanchePagePlayer>>& InPlayersWeak) -> FString
+	{
+		FString PageIdList;
+		for (const TWeakObjectPtr<UAvalanchePagePlayer>& PlayerWeak : InPlayersWeak)
+		{
+			if (const UAvalanchePagePlayer* Player = PlayerWeak.Get())
+			{
+				PageIdList += FString::Printf(TEXT("%s%d"), PageIdList.IsEmpty() ? TEXT("") : TEXT(", "), Player->PageId);
+			}
+		}
+		return PageIdList.IsEmpty() ? TEXT("None") : PageIdList;
+	};
+
+	const FString EnterPageList = MakePageIdList(EnterPlayersWeak);
+	const FString PlayingPageList = MakePageIdList(PlayingPlayersWeak);
+	const FString ExitPageList = MakePageIdList(ExitPlayersWeak);
+	return FString::Printf(TEXT("Page Transition \"%s\": Enter Page(s): [%s], Playing Page(s): [%s], Exit Page(s): [%s]."),
+		*GetFullName(), *EnterPageList, *PlayingPageList, *ExitPageList);
+}
+
+void UAvalanchePageTransition::RegisterToPlayableTransitionEvent()
+{
+	UAvalanchePlayable::OnTransitionEvent().RemoveAll(this);
+	UAvalanchePlayable::OnTransitionEvent().AddUObject(this, &UAvalanchePageTransition::OnTransitionEvent);
+}
+
+void UAvalanchePageTransition::UnregisterFromPlayableTransitionEvent() const
+{
+	UAvalanchePlayable::OnTransitionEvent().RemoveAll(this);
+}
+
+void UAvalanchePageTransition::RegisterEnterPagePlayerEvents(UAvalanchePagePlayer* InPagePlayer)
+{
+	using namespace UE::AvaMedia::Rundown::PageTransition::Private;
+
+	if (!InPagePlayer)
+	{
+		return;
+	}
+	
+	for (const UAvaRundownPlaybackInstancePlayer* InstancePlayer : InPagePlayer->InstancePlayers)
+	{
+		CachedTransitionLayers.Add(InstancePlayer->TransitionLayer.TagId);
+
+		// Register this transition as a visibility constraint for the playable group.
+		if (const UAvalanchePlayable* Playable = GetPlayable(InstancePlayer))
+		{
+			if (UAvaMediaPlayableGroup* PlayableGroup = Playable->GetPlayableGroup())
+			{
+				PlayableGroup->RegisterVisibilityConstraint(this);
+			}
+		}
+		else if (UAvalanchePlayback* Playback = GetPlayback(InstancePlayer))
+		{
+			// If the playable is not created yet, register to the creation event.
+			Playback->OnPlayableCreated.AddUObject(this, &UAvalanchePageTransition::OnPlayableCreated);
+		}
+	}
+}
+
+void UAvalanchePageTransition::UnregisterEnterPagePlayerEvents(UAvalanchePagePlayer* InPagePlayer)
+{
+	using namespace UE::AvaMedia::Rundown::PageTransition::Private;
+	
+	if (!InPagePlayer)
+	{
+		return;
+	}
+	
+	InPagePlayer->ForEachInstancePlayer([this](const UAvaRundownPlaybackInstancePlayer* InInstancePlayer)
+	{
+		if (UAvalanchePlayback* Playback = GetPlayback(InInstancePlayer))
+		{
+			Playback->OnPlayableCreated.RemoveAll(this);
+			Playback->ForEachPlayable([this](const UAvalanchePlayable* InPlayable)
+			{
+				if (UAvaMediaPlayableGroup* PlayableGroup = InPlayable->GetPlayableGroup())
+				{
+					PlayableGroup->UnregisterVisibilityConstraint(this);
+				}
+			});
+		}
+	});
+}
+
+bool UAvalanchePageTransition::AddPagePlayer(UAvalanchePagePlayer* InPagePlayer, TArray<TWeakObjectPtr<UAvalanchePagePlayer>>& OutPagePlayersWeak)
+{
+	if (!InPagePlayer)
+	{
+		return false;	
+	}
+
+	OutPagePlayersWeak.Add(InPagePlayer);
+
+	UpdateChannelName(InPagePlayer);
+
+	return true;
+}
+
+void UAvalanchePageTransition::UpdateChannelName(const UAvalanchePagePlayer* InPagePlayer)
+{
+	check(InPagePlayer);
+	
+	if (ChannelName.IsNone())
+	{
+		ChannelName = InPagePlayer->ChannelFName;
+	}
+	else
+	{
+		using namespace UE::AvaMedia::Rundown::PageTransition::Private;
+
+		// Validate the channel is the same.
+		if (ChannelName != InPagePlayer->ChannelFName)
+		{
+			UE_LOG(LogAvaPlaylist, Error, TEXT("Page Transition \"%s\": Adding Page: %d, {%s} in a different channel than previous pages (\"%s\")."),
+				*GetFullName(), InPagePlayer->PageId, *GetPrettyPageInfo(InPagePlayer), *ChannelName.ToString());
+		}
+	}
+}
+
diff --git a/Engine/Plugins/Experimental/Avalanche/Source/AvalancheMedia/Private/Playlist/Transition/AvalanchePageTransitionBuilder.cpp b/Engine/Plugins/Experimental/Avalanche/Source/AvalancheMedia/Private/Playlist/Transition/AvalanchePageTransitionBuilder.cpp
new file mode 100644
--- /dev/null
+++ b/Engine/Plugins/Experimental/Avalanche/Source/AvalancheMedia/Private/Playlist/Transition/AvalanchePageTransitionBuilder.cpp
@@ -0,0 +1,68 @@
+// Copyright Epic Games, Inc. All Rights Reserved.
+
+#include "AvalanchePageTransitionBuilder.h"
+
+#include "Playlist/AvalanchePlaylist.h"
+#include "Playlist/Transition/AvalanchePageTransition.h"
+
+FAvalanchePageTransitionBuilder::~FAvalanchePageTransitionBuilder()
+{
+	if (Playlist)
+	{
+		for (UAvalanchePageTransition* PageTransition : PageTransitions)
+		{
+			// Add any remaining playing pages in the channel that are not either exit or enter page already.
+			for (UAvalanchePagePlayer* PagePlayer : Playlist->GetPagePlayers())
+			{
+				if (PagePlayer->ChannelFName == PageTransition->GetChannelName())
+				{
+					if (!PageTransition->HasPagePlayer(PagePlayer))
+					{
+						PageTransition->AddPlayingPage(PagePlayer);
+					}
+				}
+			}
+			
+			Playlist->AddPageTransition(PageTransition);
+
+			bool bShouldDiscard = false;
+			if (PageTransition->CanStart(bShouldDiscard))
+			{
+				PageTransition->Start();
+			}
+			else
+			{
+				// Some playables are still loading, push the command for later execution.
+				Playlist->GetPlaybackManager().PushPlaybackTransitionStartCommand(PageTransition);
+			}
+		}
+	}
+}
+
+UAvalanchePageTransition* FAvalanchePageTransitionBuilder::FindTransition(const UAvalanchePagePlayer* InPlayer) const
+{
+	// Currently, the only batching criteria is the channel.
+	for (UAvalanchePageTransition* PageTransition : PageTransitions)
+	{
+		if (PageTransition->GetChannelName() == InPlayer->ChannelFName)
+		{
+			return PageTransition;
+		}
+	}
+	return nullptr;	
+}
+
+UAvalanchePageTransition* FAvalanchePageTransitionBuilder::FindOrAddTransition(const UAvalanchePagePlayer* InPlayer)
+{
+	if (UAvalanchePageTransition* PageTransition = FindTransition(InPlayer))
+	{
+		return PageTransition;
+	}
+		
+	UAvalanchePageTransition* PageTransition = NewObject<UAvalanchePageTransition>(Playlist);
+	if (PageTransition)
+	{
+		PageTransitions.Add(PageTransition);
+	}
+	return PageTransition;	
+}
\ No newline at end of file
diff --git a/Engine/Plugins/Experimental/Avalanche/Source/AvalancheMedia/Private/Playlist/Transition/AvalanchePageTransitionBuilder.h b/Engine/Plugins/Experimental/Avalanche/Source/AvalancheMedia/Private/Playlist/Transition/AvalanchePageTransitionBuilder.h
new file mode 100644
--- /dev/null
+++ b/Engine/Plugins/Experimental/Avalanche/Source/AvalancheMedia/Private/Playlist/Transition/AvalanchePageTransitionBuilder.h
@@ -0,0 +1,61 @@
+// Copyright Epic Games, Inc. All Rights Reserved.
+
+#pragma once
+
+#include "Containers/Array.h"
+
+class UAvalanchePlaylist;
+class UAvalanchePageTransition;
+class UAvalanchePagePlayer;
+
+/**
+ * Page Transition Builder to batch pages appropriately in page transitions
+ * to provide support for multiple pages being played at the same time.
+ *
+ * Summary of the page batching rules (as of yet):
+ * - For enter pages, the batching requires all enter pages (of any TL Layer) of a given channel
+ * to be in the same transition.
+ * - (*) There can only be one enter page of a particular layer in the transition. This rule should be
+ * implemented by TL itself, but it is not atm. It is thus a hard coded rule by the playlist
+ * transition system for now. The rule for now is to keep only the first page being added.
+ * - For exit pages, when they are part of a transition that has enter pages, we can add all
+ * the currently playing pages (in the given channel) as "exit" pages.
+ * - If the transition is a "take out" only transition, then we need to batch the exit pages
+ * per layers (on top of per channel) since the whole transition has an override that allows only
+ * one layer to be specified. This is not done at the moment. For take out, there is no batching
+ * at all atm, all pages are added in a separate transition.
+ *
+ * Note: for now, FAvalanchePageTransitionBuilder is not used for stopping pages. This may change
+ * depending on refactor of the underlying TL and transition system. Some traces of this
+ * case are left in this class for that reason, so it is easy to adapt it for the "take out"
+ * case.
+ */
+class FAvalanchePageTransitionBuilder
+{
+public:
+	FAvalanchePageTransitionBuilder(UAvalanchePlaylist* InPlaylist) : Playlist(InPlaylist) {}
+
+	~FAvalanchePageTransitionBuilder();
+
+	/**
+	 * @brief Finds a suitable page transition object for the given page player.
+	 * @param InPlayer Page player we are requesting the transition object for.
+	 * @return the requested page transition object or null if not found.
+	 */
+	UAvalanchePageTransition* FindTransition(const UAvalanchePagePlayer* InPlayer) const;
+
+	/**
+	 * @brief Retrieves a suitable page transition object for the given page player.
+	 * @param InPlayer Page player we are requesting the transition object for.
+	 * @return the requested page transition object.
+	 *
+	 * This is where the batching occurs according to criteria.
+	 */
+	UAvalanchePageTransition* FindOrAddTransition(const UAvalanchePagePlayer* InPlayer);
+
+private:
+	UAvalanchePlaylist* Playlist;
+
+	/** There will be a transition for each channel. */
+	TArray<UAvalanchePageTransition*> PageTransitions;
+};
\ No newline at end of file
diff --git a/Engine/Plugins/Experimental/Avalanche/Source/AvalancheMedia/Public/AvaMediaDefines.h b/Engine/Plugins/Experimental/Avalanche/Source/AvalancheMedia/Public/AvaMediaDefines.h
new file mode 100644
--- /dev/null
+++ b/Engine/Plugins/Experimental/Avalanche/Source/AvalancheMedia/Public/AvaMediaDefines.h
@@ -0,0 +1,324 @@
+// Copyright Epic Games, Inc. All Rights Reserved.
+
+#pragma once
+
+#include "Misc/EnumClassFlags.h"
+#include "Framework/AvaSoftAssetPtr.h"
+#include "AvaMediaDefines.generated.h"
+
+UENUM()
+enum class EAvaChannelState : uint8
+{
+	Offline,
+	Idle,
+	Live,
+
+	Max,
+};
+
+/**
+ * The channel type defines what it is used for in the broadcast framework.
+ *
+ * Primarily, the channel type is intended to resolve channel collisions between
+ * simultaneous "program" and "preview" playbacks on a given system. In other words:
+ * - Channel selection for playlist pages is restricted to "program" channels.
+ * - Channel selection for preview is restricted to "preview" channels.
+ *
+ * It is thus not possible for a user to mistakenly select the same channel for both preview and program.
+ * 
+ * Some additional restrictions are applied according to channel type:
+ * - preview channels must only have outputs local to the process. "Remote" previews are not supported.
+ * - [backend] playback request type (program or preview) must match with the channel type. This is a safety
+ *	 net for any other extended code paths that are not in the Avalanche plugin.
+ */
+UENUM()
+enum class EAvaBroadcastChannelType : uint8
+{
+	Program,
+	Preview
+};
+
+/**
+ *	Status of the media output.
+ */
+UENUM()
+enum class EAvaMediaOutputState : uint8
+{
+	/** Invalid/Uninitialized state. */
+	Invalid,
+	/** For remote output that is not connected, output disabled. */
+	Offline,
+	/** Server Connected or Local (MediaCapture Status: Stopped) */
+	Idle,
+	/** MediaCapture Status: Preparing */
+	Preparing,
+	/** Broadcasting (MediaCapture Status: Capturing) */
+	Live,
+	/** MediaCapture Error (unrecoverable) */
+	Error
+};
+
+UENUM()
+enum class EAvaMediaIssueSeverity : uint8
+{
+	None,
+	Warnings,
+	Errors,
+
+	Max,
+};
+
+//An enum indicating what changed in Broadcast
+enum class EAvaBroadcastChange : uint8
+{
+	None            = 0,
+	CurrentProfile  = 1 << 0, //When the Current Profile Changed
+	ChannelGrid     = 1 << 1, //When a Channel has been Added/Removed
+	ChannelRename   = 1 << 2, //When a Channel has been Renamed
+	ChannelType     = 1 << 3, //When a Channel has changed type.
+	OutputDevices   = 1 << 4, //When output devices have changed
+	All             = 0xFF
+};
+ENUM_CLASS_FLAGS(EAvaBroadcastChange);
+
+//An enum indicating what changed in Channel
+enum class EAvaChannelChange : uint8
+{
+	None         = 0,
+	State        = 1 << 0,
+	RenderTarget = 1 << 1, 
+	MediaOutputs = 1 << 2,
+	Settings     = 1 << 3,
+};
+ENUM_CLASS_FLAGS(EAvaChannelChange);
+
+/**
+ * The status of a playable object.
+ *
+ * This is related to the status of the transient playable object which is
+ * duplicated from the source asset. The transient playable object is also
+ * referred to as the "runtime" asset instance.
+ */
+UENUM()
+enum class EAvalanchePlayableStatus : uint8
+{
+	Unknown,
+	Error,
+	Unloaded,
+	Loading,
+	Loaded,
+	Visible
+};
+
+UENUM()
+enum class EAvalanchePlayableSequenceEventType : uint8
+{
+	None,
+	Started,
+	Finished
+};
+
+/**
+ * Transition flags carry additional context to help resolve the behaviors.
+ */
+enum class EAvalanchePlayableTransitionFlags : uint8
+{
+	None = 0,
+	/** This transition contains only exit playables. */
+	ExitOnly = 1 << 0,
+};
+ENUM_CLASS_FLAGS(EAvalanchePlayableTransitionFlags);
+
+enum class EAvalanchePlayableTransitionEventFlags : uint8
+{
+	None = 0,
+	/** The transition is starting. */
+	Starting = 1 << 0,
+	/** The playable needs to be stopped. */
+	StopPlayable = 1 << 1,
+	/** The playable needs to be discarded at the end of the transition. */
+	MarkPlayableDiscard = 1 << 2,
+	/** The transition is finished and can be cleaned up. */
+	Finished = 1 << 3,
+};
+ENUM_CLASS_FLAGS(EAvalanchePlayableTransitionEventFlags);
+
+/*
+ * Parameters used to tell a Channel its Playback Settings
+ */
+struct FAvaPlaybackChannelParameters
+{
+	/** Channel index for the current playback setting. */
+	int32 ChannelIndex = -1;
+
+	/**
+	 * If a playback node is reached during tick traversal,
+	 * it is going to set this value.
+	 */
+	TArray<FAvaSoftAssetPtr> AvalancheAssets;
+
+	bool HasAssets() const
+	{
+		return !AvalancheAssets.IsEmpty();
+	}
+};
+
+/*
+ * Parameters for Event/Action
+ */
+struct FAvaPlaybackEventParameters
+{
+	/** List of channel indices that this event originates from. */
+	TArray<int32> ChannelIndices;
+
+	FAvaSoftAssetPtr AvalancheAsset;
+
+	bool IsAssetValid() const
+	{
+		return !AvalancheAsset.IsNull();
+	}
+
+	const FSoftObjectPath& GetAssetPath() const
+	{
+		return AvalancheAsset.ToSoftObjectPath();
+	}
+	
+	void RequestTriggerEventAction()
+	{
+		bTriggerEventAction = true;
+	}
+
+	bool ShouldTriggerEventAction() const
+	{
+		return bTriggerEventAction;
+	}
+	
+protected:
+	
+	bool bTriggerEventAction = false;
+};
+
+
+enum class EAvaPageListChange : uint8
+{
+	None = 0,
+
+	AddedPages        = 1 << 0,
+	RemovedPages      = 1 << 1,
+	RenumberedPageId  = 1 << 2,
+	RenamedPageView   = 1 << 3,
+	ReorderedPageView = 1 << 4,
+
+	All              = 0xFF,
+};
+ENUM_CLASS_FLAGS(EAvaPageListChange);
+
+enum class EAvaPageChanges : uint8
+{
+	None = 0,
+
+	AnimationSettings       = 1 << 0,
+	RemoteControlValues     = 1 << 1,
+	Blueprint				= 1 << 2,
+	Channel					= 1 << 3,
+	Status					= 1 << 4,
+
+	All              = 0xFF,
+};
+ENUM_CLASS_FLAGS(EAvaPageChanges);
+
+/**
+ *	The status of the playback asset on disk.
+ *
+ *	This is also sometimes referred to as the "source" asset, by contrast to the
+ *	runtime (or managed) transient asset for runtime playback. It will also refer to
+ *	the source Avalanche blueprint, rather than a playback object.
+ */
+UENUM()
+enum class EAvaMediaPlaybackAssetStatus
+{
+	/** Invalid status. */
+	Unknown,
+	/** Missing asset */
+	Missing,
+	/**
+	 * Missing asset dependencies
+	 * Note: An Avalanche asset can run even with missing dependencies.
+	 */
+	 MissingDependencies,
+	 /** Asset is out of date from compare with remote. */
+	 NeedsSync,
+	 /** Asset is fully available and up to date. */
+	 Available
+};
+
+/**
+ * The status of the playback object.
+ *
+ * This is related to the status of the transient playback object which is
+ * duplicated from the source asset. The transient playback object is also
+ * referred to as the "managed" or "runtime" asset.
+ */
+UENUM()
+enum class EAvaMediaPlaybackStatus
+{
+	/** Invalid status. */
+	Unknown,
+	/** Missing asset (or out of date?). */
+	Missing,
+	/** Asset is being downloaded. */
+	Syncing,
+	/** Asset is available, not loaded. */
+	Available,
+	/** Load has been requested. */
+	Loading,
+	/** Asset is loaded in memory and ready to play. */
+	Loaded,
+	/** Start has been requested. */
+	Starting,
+	/** Currently playing. */
+	Started,
+	/** Stop has been requested. After stopping, status goes back to Loaded. */
+	Stopping,
+	/** Unload has been requested. After unloaded, status goes back to Available. */
+	Unloading,
+	/** Something bad happened. */
+	Error
+};
+
+UENUM(BlueprintType)
+enum class EAvaPageListType : uint8
+{
+	Template,
+	Instance,
+	View
+};
+
+USTRUCT(BlueprintType)
+struct FAvaPageListReference
+{
+	GENERATED_BODY()
+
+	UPROPERTY()
+	EAvaPageListType Type = EAvaPageListType::Instance;
+
+	UPROPERTY()
+	int32 SubListIndex = INDEX_NONE;
+
+	bool operator==(const FAvaPageListReference& InOther) const
+	{
+		if (Type != InOther.Type)
+		{
+			return false;
+		}
+
+		if (Type == EAvaPageListType::View && SubListIndex != InOther.SubListIndex)
+		{
+			return false;
+		}
+
+		return true;
+	}
+
+	bool operator!=(const FAvaPageListReference& InOther) const { return !(*this == InOther); }
+};
diff --git a/Engine/Plugins/Experimental/Avalanche/Source/AvalancheMedia/Public/AvaMediaDevice.h b/Engine/Plugins/Experimental/Avalanche/Source/AvalancheMedia/Public/AvaMediaDevice.h
new file mode 100644
--- /dev/null
+++ b/Engine/Plugins/Experimental/Avalanche/Source/AvalancheMedia/Public/AvaMediaDevice.h
@@ -0,0 +1,56 @@
+// Copyright Epic Games, Inc. All Rights Reserved.
+
+#pragma once
+
+#include "CoreMinimal.h"
+#include "MediaIOCoreDefinitions.h"
+
+struct FAvaMediaDevice
+{
+	FAvaMediaDevice(const FMediaIODevice& InDevice, const FName& InDeviceProviderName, const FString& InServerName)
+		: Device(InDevice)
+		, DeviceProviderName(InDeviceProviderName)
+		, ServerName(InServerName)
+	{
+	}
+	
+	bool IsValid() const
+	{
+		return Device.DeviceName != NAME_None && Device.IsValid();
+	}
+	
+	bool operator==(const FAvaMediaDevice& Other) const
+	{
+		return Device.DeviceName == Other.Device.DeviceName
+			&& Device.DeviceIdentifier == Other.Device.DeviceIdentifier
+			&& ServerName == Other.ServerName;
+	}
+	
+	FText GetDisplayNameText() const
+	{
+		return FText::Format(NSLOCTEXT("AvalancheMedia", "AvaMediaDeviceDisplayName", "{0}")
+			, FText::FromName(Device.DeviceName));
+	}
+
+	friend uint32 GetTypeHash(const FAvaMediaDevice& InDevice)
+	{
+		uint32 Hash = GetTypeHash(InDevice.Device.DeviceName);
+		Hash = HashCombine(Hash, InDevice.Device.DeviceIdentifier);
+		Hash = HashCombine(Hash, GetTypeHash(InDevice.ServerName));
+		return Hash;
+	}
+
+	const FMediaIODevice& GetDevice() const { return Device; }
+	const FName& GetDeviceProviderName() const { return DeviceProviderName;}
+	const FString& GetServerName() const { return ServerName;}
+	
+protected:
+	
+	FMediaIODevice Device;
+
+	/** Name of the device provider for this device. */
+	FName DeviceProviderName;
+	
+	/** Name of the server hosting this device. */
+	FString ServerName;
+};
diff --git a/Engine/Plugins/Experimental/Avalanche/Source/AvalancheMedia/Public/AvaRenderTargetUtils.h b/Engine/Plugins/Experimental/Avalanche/Source/AvalancheMedia/Public/AvaRenderTargetUtils.h
new file mode 100644
--- /dev/null
+++ b/Engine/Plugins/Experimental/Avalanche/Source/AvalancheMedia/Public/AvaRenderTargetUtils.h
@@ -0,0 +1,45 @@
+// Copyright Epic Games, Inc. All Rights Reserved.
+
+#pragma once
+
+#include "CoreMinimal.h"
+#include "Math/MathFwd.h"
+#include "PixelFormat.h"
+
+class UTextureRenderTarget2D;
+
+namespace AvaRenderTargetUtils
+{
+	/** Creates a Render Target with default settings like bForceLinearGamma to be false. */
+	AVALANCHEMEDIA_API UTextureRenderTarget2D* CreateDefaultRenderTarget(FName InBaseName);
+
+	/** Update the render target with the given parameters. */
+	AVALANCHEMEDIA_API void UpdateRenderTarget(UTextureRenderTarget2D* InRenderTarget, const FIntPoint& InSize, EPixelFormat InFormat, const FLinearColor& InClearColor);
+	
+	/**
+	 * Setup the render target gamma appropriately depending on the format.
+	 * @return true if the TargetGamma was modified, false otherwise.
+	 */
+	AVALANCHEMEDIA_API bool SetupRenderTargetGamma(UTextureRenderTarget2D* InRenderTarget);
+	
+	/** Returns the size of the render target. */ 
+	AVALANCHEMEDIA_API FIntPoint GetRenderTargetSize(const UTextureRenderTarget2D* InRenderTarget);
+
+	/** Returns true if render target's format is floating point. */
+	AVALANCHEMEDIA_API bool IsFloatFormat(const UTextureRenderTarget2D* InRenderTarget);
+
+	/** Returns the number of color components for the given pixel format. */
+	AVALANCHEMEDIA_API int32 GetNumColorComponents(EPixelFormat InFormat);
+
+	/** Returns true if the given format has an alpha channel. */
+	AVALANCHEMEDIA_API bool HasAlpha(EPixelFormat InFormat);
+
+	/**
+	 * Returns the number of bits per color component (assuming all the same) for the given pixel format.
+	 * Known exceptions: PF_R5G6B5_UNORM, PF_FloatR11G11B10. But those are not render target formats.
+	 */
+	AVALANCHEMEDIA_API int32 GetNumColorChannelBits(EPixelFormat InFormat);
+
+	/** Returns the number of bits for the alpha channel. */
+	AVALANCHEMEDIA_API int32 GetNumAlphaChannelBits(EPixelFormat InFormat);
+}
diff --git a/Engine/Plugins/Experimental/Avalanche/Source/AvalancheMedia/Public/AvaSerializationUtils.h b/Engine/Plugins/Experimental/Avalanche/Source/AvalancheMedia/Public/AvaSerializationUtils.h
new file mode 100644
--- /dev/null
+++ b/Engine/Plugins/Experimental/Avalanche/Source/AvalancheMedia/Public/AvaSerializationUtils.h
@@ -0,0 +1,37 @@
+// Copyright Epic Games, Inc. All Rights Reserved.
+
+#pragma once
+
+#include "Containers/ArrayView.h"
+#include "Containers/UnrealString.h"
+#include "Templates/Function.h"
+
+class FProperty;
+class FStructuredArchiveFormatter;
+class UObject;
+
+namespace UE::AvaSerializationUtils
+{
+	/** Serialize the given object with the given formatter. */
+	AVALANCHEMEDIA_API void SerializeObject(FStructuredArchiveFormatter& Formatter, UObject* InObject, const TFunction<bool(FProperty*)>& InShouldSkipProperty = TFunction<bool(FProperty*)>());
+
+	/**
+	 * The following functions are meant to be used to convert the serialized values of FJsonStruct(De)SerializerBackend
+	 * to/from FString.
+	 * 
+	 * @remark FJsonStructSerializerBackend and FJsonStructDeserializerBackend are currently hardcoded to use UCS2CHAR.
+	 */
+	namespace JsonValueConversion
+	{
+		/** Converts a serialized value stored as raw bytes into a string. In place version. */
+		AVALANCHEMEDIA_API void BytesToString(const TArray<uint8>& InValueAsBytes, FString& OutValueAsString);
+
+		/** This ArrayView version is meant to work with FMemoryReaderView. */
+		inline TArrayView<const uint8> ValueToConstBytesView(const FString& InValueAsString)
+		{
+			// Len() excludes the terminating character, and it is desired so it matches
+			// what FJsonStructSerializerBackend does.
+			return TArrayView<const uint8>(reinterpret_cast<const uint8*>(*InValueAsString), InValueAsString.Len() * sizeof(TCHAR));
+		}
+	}
+}
diff --git a/Engine/Plugins/Experimental/Avalanche/Source/AvalancheMedia/Public/AvalancheBroadcast.h b/Engine/Plugins/Experimental/Avalanche/Source/AvalancheMedia/Public/AvalancheBroadcast.h
new file mode 100644
--- /dev/null
+++ b/Engine/Plugins/Experimental/Avalanche/Source/AvalancheMedia/Public/AvalancheBroadcast.h
@@ -0,0 +1,202 @@
+// Copyright Epic Games, Inc. All Rights Reserved.
+
+#pragma once
+
+#include "AvaBroadcastProfile.h"
+#include "AvaMediaDefines.h"
+#include "CoreMinimal.h"
+#include "UObject/Object.h"
+#include "AvalancheBroadcast.generated.h"
+
+DECLARE_MULTICAST_DELEGATE_OneParam(FOnAvaBroadcastChanged, EAvaBroadcastChange)
+DECLARE_MULTICAST_DELEGATE_OneParam(FOnChannelsListChanged, const FAvaBroadcastProfile& /*InProfile*/);
+
+DECLARE_LOG_CATEGORY_EXTERN(LogAvaBroadcast, Log, All);
+
+/**
+ * Single Instance Class that manages all the Output Channels
+ */
+UCLASS(NotBlueprintable)
+class AVALANCHEMEDIA_API UAvalancheBroadcast : public UObject
+{
+	GENERATED_BODY()
+	
+public:
+
+	static UAvalancheBroadcast& Get();
+
+	virtual void BeginDestroy() override;
+	
+	UFUNCTION(BlueprintCallable, Category = "Motion Design Broadcast")
+	static UAvalancheBroadcast* GetAvalancheBroadcast();
+	
+	UFUNCTION(BlueprintCallable, Category = "Motion Design Broadcast")
+	void StartBroadcast();
+
+	UFUNCTION(BlueprintCallable, Category = "Motion Design Broadcast")
+	void StopBroadcast();
+	
+	UFUNCTION(BlueprintCallable, Category = "Motion Design Broadcast")
+	bool IsBroadcastingAnyChannel() const;
+
+	UFUNCTION(BlueprintCallable, Category = "Motion Design Broadcast")
+	bool IsBroadcastingAllChannels() const;
+
+	/**
+	 * @brief Starts the channel if not started. Does nothing if already started.
+	 * @param InChannelName Name of the channel to conditionally start.
+	 * @return true if the channel is started (including if it was already started), false if the channel couldn't be started. 
+	 */
+	bool ConditionalStartBroadcastChannel(const FName& InChannelName);
+	
+	TArray<FName> GetProfileNames() const;
+	
+	const TMap<FName, FAvaBroadcastProfile>& GetProfiles() const;
+
+	//Returns the actual name of the profile created. Defers in Number from InProfileName if it is an already existing Profile
+	FName CreateProfile(FName InProfileName, bool bMakeCurrentProfile = true);
+
+	bool DuplicateProfile(FName InNewProfile, FName InTemplateProfile, bool bMakeCurrentProfile = true);
+	
+	bool DuplicateCurrentProfile(FName InNewProfileName = NAME_None, bool bMakeCurrentProfile = true);
+	
+	bool RemoveProfile(FName InProfileName);
+
+	bool CanRenameProfile(FName InProfileName, FName InNewProfileName, FText* OutErrorMessage = nullptr) const;
+	
+	bool RenameProfile(FName InProfileName, FName InNewProfileName);
+	
+	bool SetCurrentProfile(FName InProfileName);
+
+	FName GetCurrentProfileName() const { return CurrentProfile; }
+
+	FAvaBroadcastProfile& GetProfile(FName InProfileName);
+
+	const FAvaBroadcastProfile& GetProfile(FName InProfileName) const;
+
+	FAvaBroadcastProfile& GetCurrentProfile() { return GetProfile(CurrentProfile); }
+
+	const FAvaBroadcastProfile& GetCurrentProfile() const { return GetProfile(CurrentProfile); }
+
+#if WITH_EDITOR
+	void LoadBroadcast();
+	void SaveBroadcast();
+	FString GetBroadcastSaveFilepath() const;
+#endif
+
+	//Rather than Immediately calling the Broadcast Event change, it will queue all the multiple calls and dispatch it once to avoid redundancy
+	void QueueNotifyChange(EAvaBroadcastChange InChange);
+	
+	FDelegateHandle AddChangeListener(FOnAvaBroadcastChanged::FDelegate&& InDelegate);
+	void RemoveChangeListener(FDelegateHandle InDelegateHandle);
+	void RemoveChangeListener(const void* InUserObject);
+
+public:
+
+	int32 GetChannelNameCount() const;
+
+	/**
+	 *	Returns the index of the channel in the channel name array.
+	 *	Remark: don't use this index to lookup in a profile's channel array.
+	 */
+	int32 GetChannelIndex(FName InChannelName) const;
+	
+	FName GetChannelName(int32 ChannelIndex) const;
+	FName GetOrAddChannelName(int32 ChannelIndex);
+
+	/**
+	 *	Add the given channel name to the array of channel names.
+	 *	The channel name must be unique, if it already exists, the existing index is returned instead.
+	 *	@return the index in the array where the name was added.
+	 */
+	int32 AddChannelName(FName InChannelName);
+	
+	void UpdateChannelNames();
+	
+	bool CanRenameChannel(FName InChannelName, FName InNewChannelName) const;
+	bool RenameChannel(FName InChannelName, FName InNewChannelName);
+
+	/**
+	 * Sets the type (ex: Program or Preview) of the given channel across all profiles.
+	 * Channel type for a given channel is the same in all profiles.
+	 */
+	void SetChannelType(FName InChannelName, EAvaBroadcastChannelType InChannelType);
+	
+	EAvaBroadcastChannelType GetChannelType(FName InChannelName) const;
+
+	/**
+	 * Pins the given channel from the given profile to persist across all the profiles.
+	 */
+	void PinChannel(FName InChannelName, FName InProfileName);
+
+	/**
+	 * Unpins the given channel (if it was pinned) so that is no longer persists across all profiles.
+	 */
+	void UnpinChannel(FName InChannelName);
+	
+	bool IsChannelPinned(FName InChannelName) const { return PinnedChannels.Contains(InChannelName); }
+
+	/**
+	 * If the given channel was pinned, returns the profile it was pinned in, returns Name_None otherwise.
+	 */
+	FName GetPinnedChannelProfileName(FName InChannelName) const;
+
+	/**
+	 * Rebuilds the channel lists that might include pinned channels from other profiles.
+	 * This is necessary for functions such as GetChannels().
+	 * @remark This relies on the channel names being updated (UpdateChannelNames) prior to this call.
+	 */
+	void RebuildProfiles();
+	
+	FOnChannelsListChanged& GetOnChannelsListChanged() { return OnChannelsListChanged; }
+
+protected:
+
+	// Todo: Proposal - make the events static and pass in the UAvalancheBroadcast that caused the event.
+	// Reason: Attempt at gradually getting rid of the global UAvalancheBroadcast object. Channel events are also global.
+	FOnAvaBroadcastChanged OnBroadcastChanged;
+	FOnChannelsListChanged OnChannelsListChanged;
+	
+	EAvaBroadcastChange QueuedBroadcastChanges;
+public:
+
+#if WITH_EDITOR
+	void SetCanShowPreview(bool bInCanShowPreview) { bCanShowPreview = bInCanShowPreview; }
+	bool CanShowPreview() const { return bCanShowPreview; }
+	virtual void PostEditUndo() override;
+#endif
+	
+protected:
+
+	FAvaBroadcastProfile& CreateProfileInternal(FName InProfileName);
+	
+	void EnsureValidCurrentProfile();
+	void UpdateProfileNames();
+	TArray<int32> BuildChannelIndices() const;
+
+protected:
+	
+	UPROPERTY()
+	FName CurrentProfile;
+
+	UPROPERTY()
+	TArray<FName> ChannelNames;
+	
+	UPROPERTY()
+	TMap<FName, FAvaBroadcastProfile> Profiles;
+
+	/**
+	 * List of channels types.
+	 * Backward compatibility: if not present in the map, defaults to "program" type.
+	 */
+	UPROPERTY()
+	TMap<FName, EAvaBroadcastChannelType> ChannelTypes;
+
+	/** Maps a channel name to a profile. Indicating which profile to use for a pinned channel. */
+	UPROPERTY()
+	TMap<FName, FName> PinnedChannels;
+	
+#if WITH_EDITOR
+	bool bCanShowPreview = true;
+#endif
+};
diff --git a/Engine/Plugins/Experimental/Avalanche/Source/AvalancheMedia/Public/AvalancheBroadcastComponent.h b/Engine/Plugins/Experimental/Avalanche/Source/AvalancheMedia/Public/AvalancheBroadcastComponent.h
new file mode 100644
--- /dev/null
+++ b/Engine/Plugins/Experimental/Avalanche/Source/AvalancheMedia/Public/AvalancheBroadcastComponent.h
@@ -0,0 +1,45 @@
+// Copyright Epic Games, Inc. All Rights Reserved.
+
+#pragma once
+
+#include "Components/ActorComponent.h"
+#include "AvalancheBroadcastComponent.generated.h"
+
+/**
+ * Add this actor component to blueprint actor to expose the API to control the broadcasting
+ * of channels in game.
+ */
+UCLASS(BlueprintType, Blueprintable, ClassGroup = "Motion Design Media", meta = (BlueprintSpawnableComponent))
+class AVALANCHEMEDIA_API UAvalancheBroadcastComponent : public UActorComponent
+{
+	GENERATED_BODY()
+
+public:
+	UAvalancheBroadcastComponent(const FObjectInitializer& ObjectInitializer);
+
+	/**
+	 *	Automatically stop broadcast when the world this component is part of is torn down.
+	 */
+	UPROPERTY(EditDefaultsOnly, Category = "Motion Design Media")
+	bool bStopBroadcastOnTearDown = true;
+
+	/**
+	 *	Start broadcasting all channels.
+	 *	Returns true on partial success, i.e. will be true even if some channels didn't start.
+	 *	Outputs error messages related to outputs that caused problems.
+	 */
+	UFUNCTION(BlueprintCallable, Category = "Motion Design Media")
+	bool StartBroadcasting(FString& OutErrorMessage);
+
+	/**
+	*	Stop broadcasting all channels.
+	*/
+	UFUNCTION(BlueprintCallable, Category = "Motion Design Media")
+	bool StopBroadcasting();
+	
+	virtual void InitializeComponent() override;
+	virtual void UninitializeComponent() override;
+
+protected:
+	void OnWorldBeginTearDown(UWorld* InWorld);
+};
diff --git a/Engine/Plugins/Experimental/Avalanche/Source/AvalancheMedia/Public/AvalancheMediaSettings.h b/Engine/Plugins/Experimental/Avalanche/Source/AvalancheMedia/Public/AvalancheMediaSettings.h
new file mode 100644
--- /dev/null
+++ b/Engine/Plugins/Experimental/Avalanche/Source/AvalancheMedia/Public/AvalancheMediaSettings.h
@@ -0,0 +1,202 @@
+// Copyright Epic Games, Inc. All Rights Reserved.
+
+#pragma once
+
+#include "CoreMinimal.h"
+#include "Engine/DeveloperSettings.h"
+#include "Framework/AvalancheInstanceSettings.h"
+#include "Logging/LogVerbosity.h"
+#include "PixelFormat.h"
+#include "AvalancheMediaSettings.generated.h"
+
+class UUserWidget;
+
+/**
+ * Defines the verbosity level of the logging system.
+ * This enum mirrors ELogVerbosity but can be used directly as a configuration property.
+ */
+UENUM()
+enum class EAvaMediaLogVerbosity : uint8
+{
+	NoLogging = 0,
+	Fatal,
+	Error,
+	Warning,
+	Display,
+	Log,
+	Verbose,
+	VeryVerbose
+};
+
+USTRUCT()
+struct FAvaPlaybackServerLoggingEntry
+{
+	GENERATED_BODY()
+
+	UPROPERTY(config, EditAnywhere, Category=Settings)
+	FName Category;
+
+	UPROPERTY(config, EditAnywhere, Category=Settings)
+	EAvaMediaLogVerbosity VerbosityLevel = EAvaMediaLogVerbosity::VeryVerbose;
+};
+
+USTRUCT()
+struct FLocalPlaybackServerSettings
+{
+	GENERATED_BODY()
+
+	/** Name given to the game mode local playback server started as a separate process. */
+	UPROPERTY(config, EditAnywhere, Category = Settings)
+	FString ServerName = TEXT("LocalServer");
+
+	/** Main window resolution of the local playback server. */
+	UPROPERTY(config, EditAnywhere, Category = Settings)
+	FIntPoint Resolution = FIntPoint(960, 540);
+
+	/** Enable a log console for the local server process. */
+	UPROPERTY(config, EditAnywhere, Category = Settings)
+	bool bEnableLogConsole = false;
+
+	/** Extra command line arguments. */
+	UPROPERTY(config, EditAnywhere, Category = Settings)
+	FString ExtraCommandLineArguments;
+
+	/**
+	 * Which logs to include and with which verbosity level
+	*/
+	UPROPERTY(Config, EditAnywhere, Category="Settings")
+	TArray<FAvaPlaybackServerLoggingEntry> Logging;
+};
+
+UCLASS(config=Engine, meta=(DisplayName="Playback & Broadcast"))
+class AVALANCHEMEDIA_API UAvalancheMediaSettings : public UDeveloperSettings
+{
+	GENERATED_BODY()
+
+public:
+	UAvalancheMediaSettings();
+	
+	static const UAvalancheMediaSettings& Get() { return *GetSingletonInstance();}
+	static UAvalancheMediaSettings& GetMutable() {return *GetSingletonInstance();}
+
+	static ELogVerbosity::Type ToLogVerbosity(EAvaMediaLogVerbosity InAvaMediaLogVerbosity);
+
+	/** Specifies the background clear color for the channel. */
+	UPROPERTY(config, EditAnywhere, Category = "Broadcast")
+	FLinearColor ChannelClearColor = FLinearColor::Black;
+
+	/** Pixel format used if no media output has specific format requirement. */
+	UPROPERTY(config, EditAnywhere, Category = "Broadcast")
+	TEnumAsByte<EPixelFormat> ChannelDefaultPixelFormat = EPixelFormat::PF_B8G8R8A8;
+
+	/** Resolution used if no media output has specific resolution requirement. */
+	UPROPERTY(config, EditAnywhere, Category = "Broadcast")
+	FIntPoint ChannelDefaultResolution = FIntPoint(1920, 1080);
+
+	/**
+	 * Enables drawing the placeholder widget when there is no avalanche asset playing.
+	 * If false, the channel is cleared to the background color.
+	 */
+	UPROPERTY(config, EditAnywhere, Category = "Broadcast")
+	bool bDrawPlaceholderWidget = false;
+	
+	/** Specify a place holder widget to render when no avalanche asset is playing. */
+	UPROPERTY(config, EditAnywhere, Category = "Broadcast")
+	TSoftClassPtr<UUserWidget> PlaceholderWidgetClass;
+
+	/**
+	 * Default resolution for playlist preview.
+	 * This resolution may be lower than the broadcast resolution to improve gpu performance.
+	 */
+	UPROPERTY(config, EditAnywhere, Category = "Playlist")
+	FIntPoint PreviewDefaultResolution = FIntPoint(960, 540);
+
+	UPROPERTY(config, EditAnywhere, Category = "Playlist")
+	FString PreviewChannelName;
+
+	/**
+	 * Special logic will not play any transition if the RC values are the same.
+	 * This applies to combo templates.
+	 */
+	UPROPERTY(config, EditAnywhere, Category = "Playlist")
+	bool bEnableComboTemplateSpecialLogic = true;
+
+	/**
+	 * Special logic will not play any transition if the RC values are the same.
+	 * This applies to single (non-combo) templates.
+	 */
+	UPROPERTY(config, EditAnywhere, Category = "Playlist")
+	bool bEnableSingleTemplateSpecialLogic = false;
+	
+	/** Whether playback client is started automatically. */
+	UPROPERTY(config, EditAnywhere, Category = "Playback Client")
+	bool bAutoStartPlaybackClient = false;
+
+	/** Enable verbose logging for playback client. */
+	UPROPERTY(config, EditAnywhere, Category = "Playback Client")
+	bool bVerbosePlaybackClientLogging = false;
+
+	/** Defines the interval in seconds in which servers are being pinged by the client. */
+	UPROPERTY(config, EditAnywhere, Category = "Playback Client")
+	float PingInterval = 2.5f;
+
+	/** If servers do not respond after the ping timeout interval in seconds, they are disconnected. */
+	UPROPERTY(config, EditAnywhere, Category = "Playback Client")
+	float PingTimeoutInterval = 3*2.5f;
+
+	/** Defines the timeout, in seconds, after which a pending status request is dropped and issued again. */
+	UPROPERTY(config, EditAnywhere, Category = "Playback Client")
+	float ClientPendingStatusRequestTimeout = 5.0f;
+	
+	/** Whether playback server is started automatically. */
+	UPROPERTY(config, EditAnywhere, Category = "Playback Server")
+	bool bAutoStartPlaybackServer = false;
+
+	/** Name given to the playback server. If empty, the server name will be the computer name. */
+	UPROPERTY(config, EditAnywhere, Category = "Playback Server")
+	FString PlaybackServerName;
+
+	/** Enable verbose logging for playback server. */
+	UPROPERTY(config, EditAnywhere, Category = "Playback Server")
+	bool bVerbosePlaybackServerLogging = false;
+
+	/**
+	 * Determines the verbosity level of the playback server's log replication.
+	 * The server is not going to replicate any log event that is below this log level.
+	 */
+	UPROPERTY(config, EditAnywhere, Category = "Playback Server")
+	EAvaMediaLogVerbosity PlaybackServerLogReplicationVerbosity = EAvaMediaLogVerbosity::Error;
+
+	/** Defines the timeout, in seconds, after which a pending status request is dropped and issued again. */
+	UPROPERTY(config, EditAnywhere, Category = "Playback Server")
+	float ServerPendingStatusRequestTimeout = 5.0f;
+
+	/** Settings for the local playback server process. See "Launch Local Server" in the broadcast editor toolbar. */
+	UPROPERTY(config, EditAnywhere, Category = "Playback Server")
+	FLocalPlaybackServerSettings LocalPlaybackServerSettings;
+	
+	/** If true, the playback objects are kept in memory after pages are stopped. They are unloaded otherwise.*/
+	UPROPERTY(Config, EditAnywhere, Category = "Playback Manager")
+	bool bKeepPagesLoaded = false;
+
+	UPROPERTY(Config, EditAnywhere, Category = "Playback Manager")
+	FAvalancheInstanceSettings AvalancheInstanceSettings;
+
+	/**
+	 * Maximum cached Managed Avalanche Blueprints used for playlist editor's page details.
+	 * A value of 0 indicate the cache will grow without limit.
+	 */
+	UPROPERTY(Config, EditAnywhere, Category = "Managed Avalanche Instance Cache", meta = (DisplayName = "Maximum Cache Size"))
+	int32 ManagedAvalancheInstanceCacheMaximumSize = 20;
+	
+	/** Whether web server is started automatically. */
+	UPROPERTY(config, EditAnywhere, Category = "Web Server")
+	bool bAutoStartWebServer = true;
+
+	/** The web remote control HTTP server's port. */
+	UPROPERTY(config, EditAnywhere, Category = "Web Server")
+	uint32 HttpServerPort = 10123;
+
+private:
+	static UAvalancheMediaSettings* GetSingletonInstance();
+};
diff --git a/Engine/Plugins/Experimental/Avalanche/Source/AvalancheMedia/Public/Channel/AvaMediaOutputInfo.h b/Engine/Plugins/Experimental/Avalanche/Source/AvalancheMedia/Public/Channel/AvaMediaOutputInfo.h
new file mode 100644
--- /dev/null
+++ b/Engine/Plugins/Experimental/Avalanche/Source/AvalancheMedia/Public/Channel/AvaMediaOutputInfo.h
@@ -0,0 +1,57 @@
+// Copyright Epic Games, Inc. All Rights Reserved.
+
+#pragma once
+
+#include "CoreMinimal.h"
+#include "AvaMediaOutputInfo.generated.h"
+
+/**
+ *	Extra information about the Media Output object.
+ *	This is used to determine the status of server hosting the device.
+ */
+USTRUCT()
+struct AVALANCHEMEDIA_API FAvaMediaOutputInfo
+{
+	GENERATED_BODY()
+	
+public:
+	/**
+	 * Unique identifier for this output.
+	 * Allows easier management fo client/server status and configuration replication.
+	 */
+	UPROPERTY()
+	FGuid Guid;
+	
+	/**
+	 * The server name if the media output was from a remote server.
+	 * This will be empty if the device was local.
+	 */
+	UPROPERTY()
+	FString ServerName;
+	
+	/**
+	 * The device provider name, ex: BlackMagic, for this device (if any).
+	 */
+	UPROPERTY()
+	FName DeviceProviderName;
+
+	/**
+	 *	Device name from the Device Provider.
+	 *	For device that have no provider (like NDI for instance), this is
+	 *	the name of the source or equivalent.
+	 */
+	UPROPERTY()
+	FName DeviceName;
+	
+	bool IsValid() const
+	{
+		return !DeviceName.IsNone();
+	}
+
+	bool IsRemote() const { return IsRemote(ServerName);}
+
+	void PostLoad();
+
+	/** Determines if the given server name is considered "remote", i.e. from a different instance. */
+	static bool IsRemote(const FString& InServerName);
+};
\ No newline at end of file
diff --git a/Engine/Plugins/Experimental/Avalanche/Source/AvalancheMedia/Public/Channel/AvaOutputChannel.h b/Engine/Plugins/Experimental/Avalanche/Source/AvalancheMedia/Public/Channel/AvaOutputChannel.h
new file mode 100644
--- /dev/null
+++ b/Engine/Plugins/Experimental/Avalanche/Source/AvalancheMedia/Public/Channel/AvaOutputChannel.h
@@ -0,0 +1,289 @@
+// Copyright Epic Games, Inc. All Rights Reserved.
+
+#pragma once
+
+#include "AvaMediaDefines.h"
+#include "AvaMediaOutputInfo.h"
+#include "CoreMinimal.h"
+#include "Engine/TextureRenderTarget2D.h"
+#include "Slate/SAvaPlaceholderWidget.h"
+#include "Viewport/AvaViewportQualitySettings.h"
+
+#include "AvaOutputChannel.generated.h"
+
+class FAudioDeviceHandle;
+class FWidgetRenderer;
+class SVirtualWindow;
+class UAvaMediaPlayableGroup;
+class UMediaCapture;
+class UMediaOutput;
+class UUserWidget;
+struct FAvaBroadcastProfile;
+
+/**
+ *	Creates a bridge between the UMediaCapture's event
+ *	and our events.
+ */
+USTRUCT()
+struct AVALANCHEMEDIA_API FAvaMediaCapture
+{
+	GENERATED_BODY()
+	
+	UPROPERTY()
+	TObjectPtr<UMediaCapture> MediaCapture = nullptr;
+
+	UPROPERTY()
+	TObjectPtr<UMediaOutput> MediaOutput = nullptr;
+
+	// Keep track of the render target the MediaCapture is capturing.
+	// We need to do this because the UMediaCapture doesn't have a public function to know
+	// what RT it is capturing.
+	TWeakObjectPtr<UTextureRenderTarget2D> RenderTarget;
+	
+	DECLARE_MULTICAST_DELEGATE_TwoParams(FOnMediaCaptureStateChanged, const UMediaCapture* InMediaCapture, const UMediaOutput* InMediaOutput);
+	FOnMediaCaptureStateChanged OnMediaCaptureStateChanged;
+
+	FAvaMediaCapture() = default;
+	~FAvaMediaCapture()
+	{
+		Reset();
+	}
+	
+	void Reset(UMediaCapture* InMediaCapture = nullptr, UMediaOutput* InMediaOutput = nullptr, UTextureRenderTarget2D* InRenderTarget = nullptr);
+	
+private:
+	/** Handler for MediaCapture's OnStateChangedNative event. */
+	void OnStateChangedNative();
+};
+
+USTRUCT()
+struct AVALANCHEMEDIA_API FAvaOutputChannel
+{
+	GENERATED_BODY()
+
+public:
+	FAvaOutputChannel() = default;
+	explicit FAvaOutputChannel(FAvaBroadcastProfile* InProfile) : Profile(InProfile) {}
+	~FAvaOutputChannel();
+
+	void ReleasePlaceholderResources();
+	void ReleasePlaceholderRenderTargets();
+	void ReleaseOutputs();
+	
+	static FAvaOutputChannel NullChannel;
+	
+	static void DuplicateChannel(const FAvaOutputChannel& InSourceChannel, FAvaOutputChannel& OutTargetChannel);
+
+	/**
+	 * Set this channel's index in the broadcast's channel names array.
+	 * Remark: don't use this index to lookup directly in the profile's channel array.
+	 */
+	void SetChannelIndex(int32 InIndex);
+
+	/**
+	 * Index of this channel in the broadcast's channel names array.
+	 * Remark: don't use this index to lookup directly in the profile's channel array.
+	 */
+	int32 GetChannelIndex() const { return ChannelIndex; }
+	
+	FName GetChannelName() const;
+	
+	EAvaBroadcastChannelType GetChannelType() const;
+
+	FName GetProfileName() const;
+
+	FAvaBroadcastProfile& GetProfile() const;
+
+	EAvaChannelState RefreshState();
+	
+	EAvaChannelState GetState() const { return ChannelState; }
+	EAvaMediaIssueSeverity GetIssueSeverity() const { return ChannelIssueSeverity; }
+
+	EAvaMediaIssueSeverity GetMediaOutputIssueSeverity(EAvaMediaOutputState InOutputState, const UMediaOutput* InMediaOutput) const;
+	const TArray<FString>& GetMediaOutputIssueMessages(const UMediaOutput* InMediaOutput) const;
+	EAvaMediaOutputState GetMediaOutputState(const UMediaOutput* InMediaOutput) const;
+	const FAvaMediaOutputInfo& GetMediaOutputInfo(const UMediaOutput* InMediaOutput) const;
+	FAvaMediaOutputInfo* GetMediaOutputInfoMutable(const UMediaOutput* InMediaOutput);
+	
+	bool IsEditingEnabled() const;
+	
+protected:
+	void SetState(EAvaChannelState InState, EAvaMediaIssueSeverity InIssueSeverity);
+
+	void UpdatePlaceholder();
+	void UpdatePlaceholderRenderTargets();
+	void UpdateOverridePlaceholder();	
+	void DrawPlaceholderWidget() const;
+	
+public:
+	bool IsValidChannel() const;
+
+	bool StartChannelBroadcast();
+	void StopChannelBroadcast();
+	
+	UMediaOutput* AddMediaOutput(const UClass* InMediaOutputClass, const FAvaMediaOutputInfo& InOutputInfo);
+	void AddMediaOutput(UMediaOutput* InMediaOutput, const FAvaMediaOutputInfo& InOutputInfo);
+	int32 RemoveMediaOutput(UMediaOutput* InMediaOutput);
+
+	/**
+	 * Called when the output is modified.
+	 * Propagates the configuration change accordingly.
+	 */
+	void OnMediaOutputModified(UMediaOutput* InMediaOutput);
+	
+	void PostLoadMediaOutputs(bool bInIsProfileActive, FAvaBroadcastProfile* InProfile);
+
+	void UpdateRenderTarget(UAvaMediaPlayableGroup* InPlayableGroup, UTextureRenderTarget2D* InRenderTarget);
+	void UpdateAudioDevice(const FAudioDeviceHandle& InAudioDeviceHandle);
+	UAvaMediaPlayableGroup* GetLastActivePlayableGroup() const;
+	
+	UTextureRenderTarget2D* GetCurrentRenderTarget(bool bInFallbackToPlaceholder) const;
+	UTextureRenderTarget2D* GetPlaceholderRenderTarget() const;
+
+	const TArray<UMediaOutput*>& GetMediaOutputs() const { return MediaOutputs; }
+
+	/** Returns true if this channel has a least one local media output. */
+	bool HasAnyLocalMediaOutputs() const;
+
+	/** Returns true if this channel has a least one remote media output. */
+	bool HasAnyRemoteMediaOutputs() const;
+	
+	/** Returns true if the Media Output object is associated to a remote server. */
+	bool IsMediaOutputRemote(const UMediaOutput* InMediaOutput) const;
+
+	/** Returns the server name associated with the Media Output object. */
+	const FString& GetMediaOutputServerName(const UMediaOutput* InMediaOutput) const;
+	
+	/** Returns the Media Output objects that are associated to a remote server. */
+	TArray<UMediaOutput*> GetRemoteMediaOutputs() const;
+
+	DECLARE_MULTICAST_DELEGATE_TwoParams(FOnAvaChannelChanged, const FAvaOutputChannel&, EAvaChannelChange);
+	static FOnAvaChannelChanged& GetOnChannelChanged() { return OnChannelChanged; }
+
+	/**
+	 * Called when the state of the capture changed.
+	 * The callback is called on the game thread. The change may occur on the rendering thread.
+	 */
+	DECLARE_MULTICAST_DELEGATE_TwoParams(FMediaOutputStateChanged, const FAvaOutputChannel& InChannel, const UMediaOutput* InMediaOutput);
+	static FMediaOutputStateChanged& GetOnMediaOutputStateChanged() { return OnMediaOutputStateChanged; }
+
+protected:
+	
+	static FOnAvaChannelChanged OnChannelChanged;
+	static FMediaOutputStateChanged OnMediaOutputStateChanged;
+
+protected:
+
+	bool IsMediaOutputCompatible(UMediaOutput* InMediaOutput, UTextureRenderTarget2D* InRenderTarget) const;
+	void UpdateViewportTarget();
+
+	UMediaCapture* GetMediaCaptureForOutput(const UMediaOutput* InMediaOutput) const;
+	void StopCaptureForOutput(const UMediaOutput* InMediaOutput);
+
+	void TickPlaceholder(float);
+	void StartPlaceholderTick();
+	void StopPlaceholderTick();
+
+	void OnAvaMediaCaptureStateChanged(const UMediaCapture* InMediaCapture, const UMediaOutput* InMediaOutput);
+	
+public:
+	
+	/**
+	* Updates Channel's Data and Resources.
+	* The resources are going to be allocated if the profile is active or released if not.
+	* This will also update the render targets according to the current configuration.
+	* @param bInIsProfileActive Indicate if the parent profile is active.
+	*/
+	void UpdateChannelResources(bool bInIsProfileActive);
+
+	FIntPoint DetermineRenderTargetSize() const;
+	EPixelFormat DetermineRenderTargetFormat() const;
+
+	static FIntPoint GetDefaultMediaOutputSize(EAvaBroadcastChannelType InChannelType);
+	static EPixelFormat GetDefaultMediaOutputFormat();
+
+	const FAvaViewportQualitySettings& GetViewportQualitySettings() const { return QualitySettings; }
+	void SetViewportQualitySettings(const FAvaViewportQualitySettings& InQualitySettings);
+
+protected:
+
+	/** Index of this channel in the broadcast's channel names array. */
+	UPROPERTY()
+	int32 ChannelIndex = INDEX_NONE;
+	
+	UPROPERTY()
+	TArray<TObjectPtr<UMediaOutput>> MediaOutputs;
+
+	UPROPERTY()
+	TArray<FAvaMediaOutputInfo> MediaOutputInfos;
+
+	UPROPERTY(EditAnywhere, Category = "Motion Design")
+	FAvaViewportQualitySettings QualitySettings;
+
+	/*
+	 * A Map of Media Outputs with their Media Captures when Broadcast Starts, and deleted when Broadcast Stops.
+	 * Only succeeded Captures are added. 
+	 */
+	UPROPERTY(Transient)
+	TMap<TObjectPtr<UMediaOutput>, FAvaMediaCapture> MediaCaptures;
+
+	/*
+	* Optional UMG Widget that is used to replace the Default Placeholder Widget
+	*/
+	UPROPERTY(Transient)
+	TObjectPtr<UUserWidget> OverridePlaceholderUserWidget;
+
+	UPROPERTY(Transient)
+	TObjectPtr<UTextureRenderTarget2D> PlaceholderRenderTarget;
+
+	// Extra render target to perform alpha inversion needed to mimic scene rendering.
+	UPROPERTY(Transient)
+	TObjectPtr<UTextureRenderTarget2D> PlaceholderRenderTargetTmp;
+	
+	FAvaBroadcastProfile* Profile = nullptr;
+	
+	/**
+	 * Last Active Playable group associated with the current render target.
+	 * In the current design, there can be only one playable group associated
+	 * to the currently active render target for a channel.
+	 */
+	TWeakObjectPtr<UAvaMediaPlayableGroup> LastActivePlayableGroup;
+
+	/**
+	 * Keep track of the active render target associated with the current playable group.
+	 * This is most likely the internal channel's "placeholder" render target, but
+	 * it could also be one owned by the given playable group.
+	 */
+	TWeakObjectPtr<UTextureRenderTarget2D> CurrentRenderTarget;
+	
+	TSharedPtr<SAvaPlaceholderWidget> PlaceholderWidget;
+
+	TSharedPtr<SWidget> OverridePlaceholderWidget;
+	
+	FDelegateHandle PlaceholderTickHandle;
+	
+	TSharedPtr<FWidgetRenderer> WidgetRenderer;
+
+	FText ChannelNameText;
+	
+	EAvaChannelState ChannelState = EAvaChannelState::Idle;	
+	EAvaMediaIssueSeverity ChannelIssueSeverity = EAvaMediaIssueSeverity::None;
+
+	// Keep track of the media capture/output status while
+	// starting the broadcast.
+	struct FLocalMediaOutputStatus
+	{
+		TArray<FString> Messages;
+		EAvaMediaIssueSeverity Severity = EAvaMediaIssueSeverity::None;
+	};
+	mutable TMap<const UMediaOutput*, FLocalMediaOutputStatus> LocalMediaOutputStatuses;
+
+	FString& EmplaceMessageForOutput(const UMediaOutput* InMediaOutput, EAvaMediaIssueSeverity InSeverity) const;
+	EAvaMediaIssueSeverity GetLocalSeverityForOutput(const UMediaOutput* InMediaOutput) const;
+	void ResetLocalMediaOutputStatus(const UMediaOutput* InMediaOutput);
+	
+	// This internal state is used to mark the channel as "should be broadcasting".
+	// This changes the interpretation of the output statuses with regard to the
+	// overall channel state.
+	bool bInternalStateBroadcasting = false;
+};
diff --git a/Engine/Plugins/Experimental/Avalanche/Source/AvalancheMedia/Public/Channel/AvalancheOutputChannelViewInterface.h b/Engine/Plugins/Experimental/Avalanche/Source/AvalancheMedia/Public/Channel/AvalancheOutputChannelViewInterface.h
new file mode 100644
--- /dev/null
+++ b/Engine/Plugins/Experimental/Avalanche/Source/AvalancheMedia/Public/Channel/AvalancheOutputChannelViewInterface.h
@@ -0,0 +1,22 @@
+// Copyright Epic Games, Inc. All Rights Reserved.
+
+#pragma once
+
+#include "UObject/Interface.h"
+#include "AvalancheOutputChannelViewInterface.generated.h"
+
+UINTERFACE(MinimalAPI, Blueprintable)
+class UAvalancheOutputChannelViewInterface : public UInterface
+{
+	GENERATED_BODY()
+};
+
+class IAvalancheOutputChannelViewInterface
+{
+	GENERATED_BODY()
+	
+public:
+	
+	UFUNCTION(BlueprintNativeEvent, Category = "Motion Design Output Channel")
+	void SetChannelName(const FText& InChannelName);
+};
diff --git a/Engine/Plugins/Experimental/Avalanche/Source/AvalancheMedia/Public/Channel/Slate/SAvaPlaceholderWidget.h b/Engine/Plugins/Experimental/Avalanche/Source/AvalancheMedia/Public/Channel/Slate/SAvaPlaceholderWidget.h
new file mode 100644
--- /dev/null
+++ b/Engine/Plugins/Experimental/Avalanche/Source/AvalancheMedia/Public/Channel/Slate/SAvaPlaceholderWidget.h
@@ -0,0 +1,25 @@
+// Copyright Epic Games, Inc. All Rights Reserved.
+
+#pragma once
+
+#include "CoreMinimal.h"
+#include "Widgets/SCompoundWidget.h"
+
+class STextBlock;
+
+class AVALANCHEMEDIA_API SAvaPlaceholderWidget : public SCompoundWidget
+{
+public:
+	
+	SLATE_BEGIN_ARGS(SAvaPlaceholderWidget) {}
+	SLATE_END_ARGS()
+
+	/** Constructs this widget with InArgs */
+	void Construct(const FArguments& InArgs);
+
+	void SetChannelName(const FText& InChannelName);
+	
+protected:
+
+	TSharedPtr<STextBlock> ChannelText;
+};
diff --git a/Engine/Plugins/Experimental/Avalanche/Source/AvalancheMedia/Public/IAvaMediaModule.h b/Engine/Plugins/Experimental/Avalanche/Source/AvalancheMedia/Public/IAvaMediaModule.h
new file mode 100644
--- /dev/null
+++ b/Engine/Plugins/Experimental/Avalanche/Source/AvalancheMedia/Public/IAvaMediaModule.h
@@ -0,0 +1,177 @@
+// Copyright Epic Games, Inc. All Rights Reserved.
+
+#pragma once
+
+#include "CoreMinimal.h"
+#include "Modules/ModuleInterface.h"
+#include "Modules/ModuleManager.h"
+
+AVALANCHEMEDIA_API DECLARE_LOG_CATEGORY_EXTERN(LogAvaMedia, Log, All);
+
+class FAvaMediaPlaybackManager;
+class FAvaMediaPlaybackServer;
+class FAvalancheManagedInstanceCache;
+class FCommonViewportClient;
+class IAvaDeviceProviderProxyManager;
+class IAvaMediaPlaybackClient;
+class IAvaMediaSyncProvider;
+class IAvalancheBroadcastSettings;
+class IMediaIOCoreDeviceProvider;
+class UWorld;
+struct FAvalancheInstanceSettings;
+struct FMediaIOOutputConfiguration;
+
+/** Maps one to one with the editor's map changed type (for now). */
+enum class EAvaMediaMapChangeType : uint8
+{
+	None,
+	LoadMap,
+	SaveMap,
+	NewMap,
+	TearDownWorld
+};
+
+class AVALANCHEMEDIA_API IAvaMediaModule : public IModuleInterface
+{
+public:
+	static bool IsModuleLoaded()
+	{
+		static const FName ModuleName = TEXT( "AvalancheMedia" );
+		return FModuleManager::Get().IsModuleLoaded(ModuleName);
+	}
+	
+	static IAvaMediaModule& Get()
+	{
+		static const FName ModuleName = TEXT( "AvalancheMedia" );
+		return FModuleManager::LoadModuleChecked<IAvaMediaModule>(ModuleName);
+	};
+
+	/**
+	 * @brief Returns true if the playback client is started.
+	 */
+	virtual bool IsMediaPlaybackClientStarted() const = 0;
+
+	/**
+	 * @brief Starts the playback client (if not already started).
+	 */
+	virtual void StartMediaPlaybackClient() = 0;
+
+	/**
+	 * @brief Stops the playback client.
+	 */
+	virtual void StopMediaPlaybackClient() = 0;
+
+	/**
+	 * @brief Returns true if the playback server is started.
+	 */
+	virtual bool IsMediaPlaybackServerStarted() const = 0;
+
+	/**
+	 * @brief Starts the playback server (if not already started).
+	 * @param InPlaybackServerName Optional server name. If empty, the host (computer) name will be used.
+	 */
+	virtual void StartMediaPlaybackServer(const FString& InPlaybackServerName) = 0;
+
+	/**
+	 * @brief Stops the playback server.
+	 */
+	virtual void StopMediaPlaybackServer() = 0;
+
+	virtual IAvaMediaPlaybackClient& GetMediaPlaybackClient() = 0;
+	virtual TSharedPtr<FAvaMediaPlaybackServer> GetMediaPlaybackServer() const = 0;
+	virtual const IMediaIOCoreDeviceProvider* GetDeviceProvider(FName InProviderName, const FMediaIOOutputConfiguration* InMediaIOOutputConfiguration) const = 0;
+	virtual TArray<const IMediaIOCoreDeviceProvider*> GetDeviceProvidersForServer(const FString& InServerName) const = 0;
+	virtual FString GetServerNameForDevice(const FName& InDeviceProviderName, const FName& InDeviceName) const = 0;
+	virtual bool IsLocalDevice(const FName& InDeviceProviderName, const FName& InDeviceName) const = 0;
+
+	/**
+	 * @brief Launches a separate process in game mode to run a local playback server.
+	 */
+	virtual void LaunchGameModeLocalPlaybackServer() = 0;
+
+	/**
+	 * @brief Stops currently running game mode local playback server.
+	 */
+	virtual void StopGameModeLocalPlaybackServer() = 0;
+
+	/**
+	 * @brief Returns true if the game mode local playback server process is launched. 
+	 */
+	virtual bool IsGameModeLocalPlaybackServerLaunched() const = 0;
+
+	/**
+	 * Access the global broadcast settings.
+	 */
+	virtual const IAvalancheBroadcastSettings& GetBroadcastSettings() const = 0;
+
+	/**
+	 * Access the global Avalanche instance settings.
+	 * 
+	 * Remark: lifetime of the returned reference is not guaranteed beyond the current call context.
+	 * If the settings are replicated from a client, it could get deleted if the client disconnects.
+	 * If the use of the settings is deferred, the caller must make a local copy of the settings or
+	 * call GetAvalancheInstanceSettings() in the deferred call instead.
+	 */
+	virtual const FAvalancheInstanceSettings& GetAvalancheInstanceSettings() const = 0;
+	
+	/**
+	 *	This is the backend for playing avalanche blueprints locally.
+	 */
+	virtual FAvaMediaPlaybackManager& GetLocalPlaybackManager() const = 0;
+
+	/**
+	 *	Access the "managed" Avalanche Asset Instance cache.
+	 */
+	virtual FAvalancheManagedInstanceCache& GetManagedInstanceCache() const = 0;
+
+	/**
+	 * Returns true if the AvaMediaSyncProvider modular feature is available.
+	 */
+	virtual bool IsAvaMediaSyncProviderFeatureAvailable() const = 0;
+	
+	/**
+	 * Access the currently used Ava Media Sync Provider.
+	 */
+	virtual IAvaMediaSyncProvider* GetAvaMediaSyncProvider() const = 0;
+
+	/**
+	 * @brief Propagate a map changed event (from the level editor).
+	 */
+	virtual void NotifyMapChangedEvent(UWorld* InWorld, EAvaMediaMapChangeType InEventType) = 0;
+
+	/**
+	* Delegate called when a new IAvaMediaSyncProvider modular feature is used.
+	*/
+	DECLARE_MULTICAST_DELEGATE_OneParam(FOnAvaMediaSyncProviderChanged, IAvaMediaSyncProvider* /*InNewAvaSyncProvider*/);
+	virtual FOnAvaMediaSyncProviderChanged& GetOnAvaMediaSyncProviderChanged() = 0;
+	
+	/**
+	* Delegate called when a package has been touched by a sync operation from the given sync provider.
+	*/
+	DECLARE_MULTICAST_DELEGATE_TwoParams(FOnAvaMediaSyncPackageModified, IAvaMediaSyncProvider* /*InAvaSyncProvider*/, const FName& /*InPackageName*/);
+	virtual FOnAvaMediaSyncPackageModified& GetOnAvaMediaSyncPackageModified() = 0;
+
+	DECLARE_MULTICAST_DELEGATE_TwoParams(FOnMapChangedEvent, UWorld* /*InWorld*/, EAvaMediaMapChangeType /*InEventType*/);
+	virtual FOnMapChangedEvent& GetOnMapChangedEvent() = 0;
+	
+	DECLARE_MULTICAST_DELEGATE(FOnAvaMediaPlaybackClientStarted);
+	virtual FOnAvaMediaPlaybackClientStarted& GetOnAvaMediaPlaybackClientStarted() = 0;
+
+	DECLARE_MULTICAST_DELEGATE(FOnAvaMediaPlaybackClientStopped);
+	virtual FOnAvaMediaPlaybackClientStopped& GetAvaMediaPlaybackClientStopped() = 0;
+
+	DECLARE_MULTICAST_DELEGATE(FOnAvaMediaPlaybackServerStarted);
+	virtual FOnAvaMediaPlaybackServerStarted& GetOnAvaMediaPlaybackServerStarted() = 0;
+
+	DECLARE_MULTICAST_DELEGATE(FOnAvaMediaPlaybackServerStopped);
+	virtual FOnAvaMediaPlaybackServerStopped& GetAvaMediaPlaybackServerStopped() = 0;
+	
+	/** Use to query the current editor viewport from the corresponding editor module. */
+	DECLARE_DELEGATE_OneParam(FGetEditorViewportClient, FCommonViewportClient** );
+	virtual FGetEditorViewportClient& GetEditorViewportClientDelegate() = 0;
+
+	/**
+	 * Access the device provider proxy manager.
+	 */
+	virtual IAvaDeviceProviderProxyManager& GetDeviceProviderProxyManager() = 0;
+};
diff --git a/Engine/Plugins/Experimental/Avalanche/Source/AvalancheMedia/Public/IAvalancheBroadcastSettings.h b/Engine/Plugins/Experimental/Avalanche/Source/AvalancheMedia/Public/IAvalancheBroadcastSettings.h
new file mode 100644
--- /dev/null
+++ b/Engine/Plugins/Experimental/Avalanche/Source/AvalancheMedia/Public/IAvalancheBroadcastSettings.h
@@ -0,0 +1,26 @@
+// Copyright Epic Games, Inc. All Rights Reserved.
+
+#pragma once
+
+#include "Math/MathFwd.h"
+
+enum EPixelFormat : uint8;
+struct FSoftObjectPath;
+
+/**
+ * Interface for the broadcast settings.
+ * Broadcast settings can either be local to the client, or if running
+ * on a playback server, they may be replicated from a connected remote client.
+ */
+class AVALANCHEMEDIA_API IAvalancheBroadcastSettings
+{
+public:
+	virtual const FLinearColor& GetChannelClearColor() const = 0;
+	virtual EPixelFormat GetDefaultPixelFormat() const = 0;
+	virtual const FIntPoint& GetDefaultResolution() const = 0;
+	virtual bool IsDrawPlaceholderWidget() const = 0;
+	virtual const FSoftObjectPath& GetPlaceholderWidgetClass() const = 0;
+	
+protected:
+	virtual ~IAvalancheBroadcastSettings() = default;
+};
diff --git a/Engine/Plugins/Experimental/Avalanche/Source/AvalancheMedia/Public/ModularFeature/AvaSyncProviderFeatureTypes.h b/Engine/Plugins/Experimental/Avalanche/Source/AvalancheMedia/Public/ModularFeature/AvaSyncProviderFeatureTypes.h
new file mode 100644
--- /dev/null
+++ b/Engine/Plugins/Experimental/Avalanche/Source/AvalancheMedia/Public/ModularFeature/AvaSyncProviderFeatureTypes.h
@@ -0,0 +1,140 @@
+// Copyright Epic Games, Inc. All Rights Reserved.
+
+#pragma once
+
+#include "AvaSyncProviderFeatureTypes.generated.h"
+
+UENUM()
+enum class EAvaSyncResponseResult : uint8
+{
+	/** Indicates an Error happened */
+	Error,
+
+	/** Indicates successful completion */
+	Success,
+
+	/** Should not happen */
+	Unknown
+};
+
+/** Engine type values for a sync connected device */
+UENUM()
+enum class EAvaSyncEngineType : uint8
+{
+	Server,
+	Commandlet,
+	Editor,
+	Game,
+	Other,
+	Unknown
+};
+
+/** Holds data about local project (such as Message Address Ids, Project Name, Hostname, InstanceId, etc.) */
+USTRUCT()
+struct AVALANCHEMEDIA_API FAvaSyncConnectionInfo
+{
+	GENERATED_BODY()
+
+	/** Holds the engine version checksum */
+	int32 EngineVersion = 0;
+
+	/** Holds the instance identifier */
+	FGuid InstanceId;
+
+	/** Holds the type of the engine instance */
+	EAvaSyncEngineType InstanceType = EAvaSyncEngineType::Unknown;
+
+	/** Holds the identifier of the session that the application belongs to */
+	FGuid SessionId;
+	
+	/** The hostname this message was generated from */
+	FString HostName;
+
+	/** The unreal project name this message was generated from */
+	FString ProjectName;
+	
+	/** The unreal project directory this message was generated from */
+	FString ProjectDir;
+
+	/** Default constructor */
+	FAvaSyncConnectionInfo();
+
+	/** Returns debug string for this message */
+	FString ToString() const;
+
+	/** Internal helper to return the last portion of a path, similar to the Unix basename command. Trailing directory separators are ignored */
+	static FString GetBasename(const FString& InPath);
+
+	/** Returns user friendly instance type for display */
+	FText GetHumanReadableInstanceType() const;
+};
+
+/** Base sync response with Status and Error information */
+USTRUCT()
+struct AVALANCHEMEDIA_API FAvaSyncResponse
+{
+	GENERATED_BODY()
+
+	/** Status of the response, error or success */
+	EAvaSyncResponseResult Status = EAvaSyncResponseResult::Unknown;
+
+	/** In case of error, a localized text describing what went wrong */
+	FText ErrorText;
+
+	/** In case of success, an optional localized text describing status of the operation */
+	FText StatusText;
+	
+	/** Holds data about local project and send to remote for it to display as result status */
+	FAvaSyncConnectionInfo ConnectionInfo;
+
+	/** Default constructor */
+	FAvaSyncResponse() = default;
+
+	/** Returns whether Status is in Success state */
+	bool IsValid() const
+	{
+		return Status == EAvaSyncResponseResult::Success;
+	}
+
+	/** Returns whether Status is in Error State and ErrorText is filled with error information */
+	bool HasError() const
+	{
+		return Status == EAvaSyncResponseResult::Error && !ErrorText.IsEmpty();
+	}
+
+	/** Returns string representation */
+	FString ToString() const
+	{
+		return FString::Printf(
+			TEXT("Status: %s, ErrorText: %s, StatusText: %s (ConnectionInfo: %s)"),
+			*UEnum::GetValueAsString(Status),
+			*ErrorText.ToString(),
+			*StatusText.ToString(),
+			*ConnectionInfo.ToString()
+		);
+	}
+};
+
+/** Payload for a comparison request */
+USTRUCT()
+struct AVALANCHEMEDIA_API FAvaSyncCompareResponse : public FAvaSyncResponse
+{
+	GENERATED_BODY()
+
+	/** Whether the two remotes have differing state */
+	bool bNeedsSynchronization = false;
+
+	/** Default constructor */
+	FAvaSyncCompareResponse() = default;
+
+	/** Returns string representation */
+	FString ToString() const
+	{
+		const FString SyncResponseText = Super::ToString();
+		return FString::Printf(
+			TEXT("%s, bNeedsSynchronization: %s"),
+			*SyncResponseText,
+			bNeedsSynchronization ? TEXT("true") : TEXT("false")
+		);
+	}
+};
diff --git a/Engine/Plugins/Experimental/Avalanche/Source/AvalancheMedia/Public/ModularFeature/IAvaMediaSyncProvider.h b/Engine/Plugins/Experimental/Avalanche/Source/AvalancheMedia/Public/ModularFeature/IAvaMediaSyncProvider.h
new file mode 100644
--- /dev/null
+++ b/Engine/Plugins/Experimental/Avalanche/Source/AvalancheMedia/Public/ModularFeature/IAvaMediaSyncProvider.h
@@ -0,0 +1,88 @@
+// Copyright Epic Games, Inc. All Rights Reserved.
+
+#pragma once
+
+#include "AvaSyncProviderFeatureTypes.h"
+#include "Features/IModularFeature.h"
+
+/** Generic delegate response handler */
+DECLARE_DELEGATE_OneParam(FOnAvaSyncResponse, const TSharedPtr<FAvaSyncResponse>&);
+
+/** Delegate with a comparison response. Payload is a FAvaSyncCompareResponse shared ptr */
+DECLARE_DELEGATE_OneParam(FOnAvaSyncCompareResponse, const TSharedPtr<FAvaSyncCompareResponse>&);
+
+/** Delegate called when a package is modified by a sync operation, either added or modified. */
+DECLARE_MULTICAST_DELEGATE_OneParam(FOnAvaSyncPackageModified, const FName& /*PackageName*/);
+
+/**
+ * Base Sync feature interface
+ *
+ * TODO:
+ * - Proper completion delegates handling for push and pull
+ * - Handle errors and respond with error state (with error case either from local or remote)
+ */
+class AVALANCHEMEDIA_API IAvaMediaSyncProvider : public IModularFeature
+{
+public:
+	virtual ~IAvaMediaSyncProvider() = default;
+
+	static FName GetModularFeatureName()
+	{
+		static const FName FeatureName = TEXT("Avalanche_Feature_Sync");
+		return FeatureName;
+	}
+
+	/** Returns the first registered implementation based on GetModularFeatureName(), if there is one available */
+    static IAvaMediaSyncProvider* Get();
+
+	/** Returns underlying feature implementation name */
+	virtual FName GetName() const = 0;
+
+	/**
+	 * Broadcasts a "sync" request to all connected remotes on the local network.
+	 * 
+	 * @param InPackageNames The list of "top level" package names (without inner dependencies / referencers)
+	 *	to base the synchronization on
+	 */
+	virtual void SyncToAll(const TArray<FName>& InPackageNames) = 0;
+
+	/**
+	 * Sends a "push" request to a specific remote (either a playback server or client name)
+	 * 
+	 * The remote name can be either a playback server name or a playback client name, depending on where the request is being triggered.
+	 *
+	 * @param InRemoteName The playback client or server name to push to
+	 * @param InPackageNames The list of "top level" package names (without inner dependencies / referencers) to send on remote
+	 * @param DoneDelegate A completion delegate called when receiving a response with a `FOnAvaSyncResponse` payload
+	 */
+	virtual void PushToRemote(const FString& InRemoteName, const TArray<FName>& InPackageNames, const FOnAvaSyncResponse& DoneDelegate = FOnAvaSyncResponse()) = 0;
+
+	/**
+	 * Sends a "pull" request to a specific remote (either a playback server or client name)
+	 * 
+	 * The remote name can be either a playback server name or a playback client name, depending on where the request is being triggered.
+	 *
+	 * @param InRemoteName The playback client or server name to push to
+	 * @param InPackageNames The list of "top level" package names (without inner dependencies / referencers) to send on remote
+	 * @param DoneDelegate A completion delegate called when receiving a response with a `FOnAvaSyncResponse` payload
+	 */
+	virtual void PullFromRemote(const FString& InRemoteName, const TArray<FName>& InPackageNames, const FOnAvaSyncResponse& DoneDelegate = FOnAvaSyncResponse()) = 0;
+
+	/**
+	 * Issue a "compare" request to a given remote, and calls back the completion delegate with status payload.
+	 *
+	 * The remote name can be either a playback server name or a playback client name, depending on where the request is being triggered, eg.
+	 *
+	 * - When local editor instance is a playback client, remote name is considered to be a server name.
+	 * - When local editor instance is a playback server, remote name is considered to be a client name.
+	 * - If local editor instance if both a playback client and server, remote name is considered to be a server name.
+	 *
+	 * @param InRemoteName The playback client or server name to compare with
+	 * @param InPackageNames The list of "top level" package names (without inner dependencies / referencers) to base the comparison on
+	 * @param DoneDelegate A completion delegate called when receiving a response with a `FAvaSyncCompareResponse` payload
+	 */
+	virtual void CompareWithRemote(const FString& InRemoteName, const TArray<FName>& InPackageNames, const FOnAvaSyncCompareResponse& DoneDelegate = FOnAvaSyncCompareResponse()) = 0;
+
+	/** Delegate called when a package is modified by a sync write operation, either being added or modified. */
+	virtual FOnAvaSyncPackageModified& GetOnAvaSyncPackageModified() = 0;
+};
diff --git a/Engine/Plugins/Experimental/Avalanche/Source/AvalancheMedia/Public/OutputDevices/AvaDeviceProviderData.h b/Engine/Plugins/Experimental/Avalanche/Source/AvalancheMedia/Public/OutputDevices/AvaDeviceProviderData.h
new file mode 100644
--- /dev/null
+++ b/Engine/Plugins/Experimental/Avalanche/Source/AvalancheMedia/Public/OutputDevices/AvaDeviceProviderData.h
@@ -0,0 +1,89 @@
+// Copyright Epic Games, Inc. All Rights Reserved.
+
+#pragma once
+
+#include "IMediaIOCoreDeviceProvider.h"
+#include "AvaDeviceProviderData.generated.h"
+
+USTRUCT()
+struct AVALANCHEMEDIA_API FAvaDeviceProviderData
+{
+	GENERATED_BODY()
+
+public:
+	UPROPERTY()
+	FName Name;
+
+	UPROPERTY()
+	TArray<FMediaIOConnection> Connections;
+
+	UPROPERTY()
+	TArray<FMediaIOConfiguration> Configurations;
+
+	UPROPERTY()
+	TArray<FMediaIODevice> Devices;
+
+	UPROPERTY()
+	TArray<FMediaIOMode> Modes;
+
+	UPROPERTY()
+	TArray<FMediaIOInputConfiguration> InputConfigurations;
+
+	UPROPERTY()
+	TArray<FMediaIOOutputConfiguration> OutputConfigurations;
+
+	UPROPERTY()
+	TArray<FMediaIOVideoTimecodeConfiguration> TimecodeConfigurations;
+	
+	UPROPERTY()
+	FMediaIOConfiguration DefaultConfiguration;
+
+	UPROPERTY()
+	FMediaIOMode DefaultMode;
+
+	UPROPERTY()
+	FMediaIOInputConfiguration DefaultInputConfiguration;
+
+	UPROPERTY()
+	FMediaIOOutputConfiguration DefaultOutputConfiguration;
+
+	UPROPERTY()
+	FMediaIOVideoTimecodeConfiguration DefaultTimecodeConfiguration;
+	
+	UPROPERTY()
+	bool bShowInputTransportInSelector = true;
+
+	UPROPERTY()
+	bool bShowOutputTransportInSelector = true;
+
+	UPROPERTY()
+	bool bShowInputKeyInSelector = true;
+
+	UPROPERTY()
+	bool bShowOutputKeyInSelector = true;
+
+	UPROPERTY()
+	bool bShowReferenceInSelector = true;
+
+public:
+	void PopulateFrom(const FName& InName, const IMediaIOCoreDeviceProvider* InDeviceProvider);
+	void ApplyServerName(const FString& InServerName);
+};
+
+USTRUCT()
+struct FAvaDeviceProviderDataList
+{
+	GENERATED_BODY()
+
+public:
+	UPROPERTY()
+	FString ServerName;
+	
+	UPROPERTY()
+	TArray<FAvaDeviceProviderData> DeviceProviders;
+
+	bool SaveToJson();
+	bool SaveToXml();
+	bool LoadFromJson();
+	void Populate(const FString& InAssignedServerName);
+};
\ No newline at end of file
diff --git a/Engine/Plugins/Experimental/Avalanche/Source/AvalancheMedia/Public/OutputDevices/AvaDisplayDeviceManager.h b/Engine/Plugins/Experimental/Avalanche/Source/AvalancheMedia/Public/OutputDevices/AvaDisplayDeviceManager.h
new file mode 100644
--- /dev/null
+++ b/Engine/Plugins/Experimental/Avalanche/Source/AvalancheMedia/Public/OutputDevices/AvaDisplayDeviceManager.h
@@ -0,0 +1,58 @@
+// Copyright Epic Games, Inc. All Rights Reserved.
+
+#pragma once
+
+#include "CoreMinimal.h"
+#include "GenericPlatform/GenericApplication.h"
+#include "Misc/FrameRate.h"
+
+/**
+ * Contains information on physical display adapter,
+ * a.k.a. graphics card.
+ **/
+struct FAvaDisplayAdapterInfo
+{
+	FString Description;			// ex: "Nvidia GeForce RTX ..."
+	uint32 VendorId = 0;
+	uint32 DeviceId = 0;
+	uint32 SubSysId = 0;
+	uint32 Revision = 0;
+	uint64 DedicatedVideoMemory = 0;
+};
+
+/**
+ * Contains info on a physical monitor connected to the display device.
+ * Similar to FMonitorInfo, except we have the adapter description
+ * and display frequency.
+ */
+struct FAvaMonitorInfo
+{
+	FAvaDisplayAdapterInfo AdapterInfo;
+	FString Name;					// ex: DISPLAY1, DISPLAY2, etc (not manufacturer name)
+	int32 Width = 0;				// Horizontal resolution in pixels.
+	int32 Height = 0;				// Vertical resolution in pixels.
+	FFrameRate DisplayFrequency = FFrameRate(0, 0);	// Default invalid.
+	FPlatformRect DisplayRect;
+	FPlatformRect WorkArea;
+	bool bIsPrimary = false;
+};
+
+/**
+ * Provides additional information on the display devices that is not available
+ * in FDisplayMetrics.
+ */
+class AVALANCHEMEDIA_API FAvaDisplayDeviceManager
+{
+public:
+	static void EnumMonitors(TArray<FAvaMonitorInfo>& OutMonitorInfo);
+
+	static const TArray<FAvaMonitorInfo>&  GetCachedMonitors(bool bForceUpdate = false);
+
+	/**
+	 *	Returns a string with a display name for the monitor.
+	 **/
+	static FString GetMonitorDisplayName(const FAvaMonitorInfo& InMonitorInfo);
+	
+private:
+	static TArray<FAvaMonitorInfo> CachedMonitorInfo;
+};
\ No newline at end of file
diff --git a/Engine/Plugins/Experimental/Avalanche/Source/AvalancheMedia/Public/OutputDevices/AvaDisplayMediaCapture.h b/Engine/Plugins/Experimental/Avalanche/Source/AvalancheMedia/Public/OutputDevices/AvaDisplayMediaCapture.h
new file mode 100644
--- /dev/null
+++ b/Engine/Plugins/Experimental/Avalanche/Source/AvalancheMedia/Public/OutputDevices/AvaDisplayMediaCapture.h
@@ -0,0 +1,40 @@
+// Copyright Epic Games, Inc. All Rights Reserved.
+
+#pragma once
+
+#include "HAL/CriticalSection.h"
+#include "MediaCapture.h"
+
+#include "AvaDisplayMediaCapture.generated.h"
+
+DECLARE_LOG_CATEGORY_EXTERN(LogAvaDisplayMedia, Log, All);
+
+UCLASS()
+class AVALANCHEMEDIA_API UAvaDisplayMediaCapture : public UMediaCapture
+{
+	GENERATED_BODY()
+
+public:
+	virtual ~UAvaDisplayMediaCapture();	// so we can pimpl.
+
+protected:
+	virtual bool ShouldCaptureRHIResource() const override { return true; }
+	virtual void OnRHIResourceCaptured_RenderingThread(const FCaptureBaseData& InBaseData, TSharedPtr<FMediaCaptureUserData, ESPMode::ThreadSafe> InUserData, FTextureRHIRef InTexture) override;
+
+	virtual bool InitializeCapture() override;
+	virtual bool PostInitializeCaptureViewport(TSharedPtr<FSceneViewport>& InSceneViewport) override;
+	virtual bool PostInitializeCaptureRenderTarget(UTextureRenderTarget2D* InRenderTarget) override;
+	virtual void StopCaptureImpl(bool bAllowPendingFrameToBeProcess) override;
+
+	UPROPERTY(Transient)
+	TObjectPtr<UTextureRenderTarget2D> RenderTarget;
+
+private:
+	bool StartNewCapture(const FIntPoint& InSourceTargetSize, EPixelFormat InSourceTargetFormat);
+
+	FCriticalSection CaptureInstanceCriticalSection;
+
+	// Private capture instance
+	class FCaptureInstance;
+	FCaptureInstance* CaptureInstance = nullptr; // lazy pimpl
+};
diff --git a/Engine/Plugins/Experimental/Avalanche/Source/AvalancheMedia/Public/OutputDevices/AvaDisplayMediaOutput.h b/Engine/Plugins/Experimental/Avalanche/Source/AvalancheMedia/Public/OutputDevices/AvaDisplayMediaOutput.h
new file mode 100644
--- /dev/null
+++ b/Engine/Plugins/Experimental/Avalanche/Source/AvalancheMedia/Public/OutputDevices/AvaDisplayMediaOutput.h
@@ -0,0 +1,35 @@
+// Copyright Epic Games, Inc. All Rights Reserved.
+
+#pragma once
+
+#include "MediaIOCoreDefinitions.h"
+#include "MediaOutput.h"
+
+#include "AvaDisplayMediaOutput.generated.h"
+
+/**
+ * Output Media to a display adapter.
+ */
+UCLASS(BlueprintType, meta = (MediaIOCustomLayout = "AvaDisplay"))
+class AVALANCHEMEDIA_API UAvaDisplayMediaOutput : public UMediaOutput
+{
+	GENERATED_BODY()
+
+public:
+	UAvaDisplayMediaOutput(const FObjectInitializer& ObjectInitializer);
+
+	/** The device, port and video settings that correspond to the output. */
+	UPROPERTY(EditAnywhere, Category = "AvaDisplay", meta = (DisplayName = "Configuration"))
+	FMediaIOOutputConfiguration OutputConfiguration;
+
+	//~ UMediaOutput interface
+public:
+	virtual bool Validate(FString& FailureReason) const override;
+	virtual FIntPoint GetRequestedSize() const override;
+	virtual EPixelFormat GetRequestedPixelFormat() const override;
+	virtual EMediaCaptureConversionOperation GetConversionOperation(EMediaCaptureSourceType InSourceType) const override;
+
+protected:
+	virtual UMediaCapture* CreateMediaCaptureImpl() override;
+	//~ End UMediaOutput interface
+};
diff --git a/Engine/Plugins/Experimental/Avalanche/Source/AvalancheMedia/Public/OutputDevices/AvaMediaOutputUtils.h b/Engine/Plugins/Experimental/Avalanche/Source/AvalancheMedia/Public/OutputDevices/AvaMediaOutputUtils.h
new file mode 100644
--- /dev/null
+++ b/Engine/Plugins/Experimental/Avalanche/Source/AvalancheMedia/Public/OutputDevices/AvaMediaOutputUtils.h
@@ -0,0 +1,32 @@
+// Copyright Epic Games, Inc. All Rights Reserved.
+
+#pragma once
+
+#include "Containers/UnrealString.h"
+#include "UObject/NameTypes.h"
+
+class UMediaOutput;
+class UObject;
+struct FAvaMediaOutputData;
+
+namespace UE::AvaMediaOutputUtils
+{
+	/** Get the device's name from the given MediaOutput. */
+	AVALANCHEMEDIA_API FString GetDeviceName(const UMediaOutput* InMediaOutput);
+
+	/** Returns true if the class has MediaIOCustomLayout meta data. */
+	AVALANCHEMEDIA_API bool HasDeviceProviderName(const UClass* InMediaOutputClass);
+
+	/** Retrieves the device provider name from the MediaIOCustomLayout class meta data. */
+	AVALANCHEMEDIA_API FName GetDeviceProviderName(const UClass* InMediaOutputClass);
+
+	/** Retrieves the device provider name from the MediaIOCustomLayout class meta data. */
+	AVALANCHEMEDIA_API FName GetDeviceProviderName(const UMediaOutput* InMediaOutput);
+
+	/** Serializes the given MediaOutput object into a FAvaMediaOutputData. */ 
+	AVALANCHEMEDIA_API FAvaMediaOutputData CreateMediaOutputData(UMediaOutput* InMediaOutput);
+
+	/** Create a MediaOutput object from the given data. */
+	AVALANCHEMEDIA_API UMediaOutput* CreateMediaOutput(const FAvaMediaOutputData& InMediaOutputData, UObject* InOuter);
+}
+
