[MotionDesign]

diff --git a/Engine/Plugins/Experimental/Avalanche/Source/AvalancheMedia/Public/OutputDevices/AvaRenderTargetMediaCapture.h b/Engine/Plugins/Experimental/Avalanche/Source/AvalancheMedia/Public/OutputDevices/AvaRenderTargetMediaCapture.h
new file mode 100644
--- /dev/null
+++ b/Engine/Plugins/Experimental/Avalanche/Source/AvalancheMedia/Public/OutputDevices/AvaRenderTargetMediaCapture.h
@@ -0,0 +1,33 @@
+// Copyright Epic Games, Inc. All Rights Reserved.
+
+#pragma once
+
+#include "MediaCapture.h"
+#include "HAL/CriticalSection.h"
+
+#include "AvaRenderTargetMediaCapture.generated.h"
+
+DECLARE_LOG_CATEGORY_EXTERN(LogAvaRenderTargetMedia, Log, All);
+
+UCLASS()
+class AVALANCHEMEDIA_API UAvaRenderTargetMediaCapture : public UMediaCapture
+{
+	GENERATED_BODY()
+
+protected:
+	virtual bool ShouldCaptureRHIResource() const override { return true; }
+	virtual void OnRHIResourceCaptured_RenderingThread(const FCaptureBaseData& InBaseData, TSharedPtr<FMediaCaptureUserData, ESPMode::ThreadSafe> InUserData, FTextureRHIRef InTexture) override;
+
+	virtual bool InitializeCapture() override;
+	virtual bool PostInitializeCaptureViewport(TSharedPtr<FSceneViewport>& InSceneViewport) override;
+	virtual bool PostInitializeCaptureRenderTarget(UTextureRenderTarget2D* InRenderTarget) override;
+	virtual void StopCaptureImpl(bool bAllowPendingFrameToBeProcess) override;
+
+	UPROPERTY(Transient)
+	TObjectPtr<UTextureRenderTarget2D> RenderTarget;
+
+private:
+	bool StartNewCapture(const FIntPoint& InSourceTargetSize, EPixelFormat InSourceTargetFormat);
+	
+	FCriticalSection RenderTargetCriticalSection;
+};
diff --git a/Engine/Plugins/Experimental/Avalanche/Source/AvalancheMedia/Public/OutputDevices/AvaRenderTargetMediaOutput.h b/Engine/Plugins/Experimental/Avalanche/Source/AvalancheMedia/Public/OutputDevices/AvaRenderTargetMediaOutput.h
new file mode 100644
--- /dev/null
+++ b/Engine/Plugins/Experimental/Avalanche/Source/AvalancheMedia/Public/OutputDevices/AvaRenderTargetMediaOutput.h
@@ -0,0 +1,51 @@
+// Copyright Epic Games, Inc. All Rights Reserved.
+
+#pragma once
+
+#include "MediaOutput.h"
+#include "UObject/SoftObjectPtr.h"
+#include "AvaRenderTargetMediaOutput.generated.h"
+
+class UTextureRenderTarget2D;
+
+/**
+ * Capture a Media to a render target.
+ */
+UCLASS(BlueprintType, meta = (MediaIOCustomLayout = "AvaRenderTarget"))
+class AVALANCHEMEDIA_API UAvaRenderTargetMediaOutput : public UMediaOutput
+{
+	GENERATED_BODY()
+
+public:
+	/** Invert the key (alpha channel). */
+	UPROPERTY(BlueprintReadWrite, EditAnywhere, Category="Media")
+	bool bInvertKeyOutput = true;
+
+	/**
+	 * Specify the render target to be capturing to.
+	 */
+	UPROPERTY(BlueprintReadWrite, EditAnywhere, Category="Media")
+	TSoftObjectPtr<UTextureRenderTarget2D> RenderTarget;
+
+	/**
+	 * The source name is a property that exists to provide the "device name" for
+	 * displaying in the broadcast editor. It is derived from the specified render target
+	 * path. */
+	UPROPERTY()
+	FString SourceName;
+	
+	//~ UMediaOutput interface
+public:
+	virtual bool Validate(FString& FailureReason) const override;
+	virtual FIntPoint GetRequestedSize() const override;
+	virtual EPixelFormat GetRequestedPixelFormat() const override;
+	virtual EMediaCaptureConversionOperation GetConversionOperation(EMediaCaptureSourceType InSourceType) const override;
+
+#if WITH_EDITOR
+	virtual void PostEditChangeProperty(struct FPropertyChangedEvent& PropertyChangedEvent) override;
+#endif
+
+protected:
+	virtual UMediaCapture* CreateMediaCaptureImpl() override;
+	//~ End UMediaOutput interface
+};
diff --git a/Engine/Plugins/Experimental/Avalanche/Source/AvalancheMedia/Public/OutputDevices/AvaRenderTargetMediaUtils.h b/Engine/Plugins/Experimental/Avalanche/Source/AvalancheMedia/Public/OutputDevices/AvaRenderTargetMediaUtils.h
new file mode 100644
--- /dev/null
+++ b/Engine/Plugins/Experimental/Avalanche/Source/AvalancheMedia/Public/OutputDevices/AvaRenderTargetMediaUtils.h
@@ -0,0 +1,27 @@
+// Copyright Epic Games, Inc. All Rights Reserved.
+
+#pragma once
+
+#include "RHI.h"
+
+class FRHICommandListImmediate;
+class UTextureRenderTarget2D;
+
+namespace UE::AvaRenderTargetMediaUtils
+{
+	/**
+	 * Clears the given render target to it's specified clear color.
+	 * @remark called from the main thread.
+	 */
+	AVALANCHEMEDIA_API void ClearRenderTarget(UTextureRenderTarget2D* InRenderTarget);
+
+	/**
+	 * Copy a texture from source to destination.
+	 * Handles different sizes, will rescale.
+	 * Handles different formats, will convert.
+	 */
+	AVALANCHEMEDIA_API void CopyTexture(FRHICommandListImmediate& InRHICmdList, FTextureRHIRef InSourceTexture, FTextureRHIRef InDestTarget);
+	
+	/** Implementation of texture conversion from SCS_FinalColorSDR back to SCS_FinalToneCurveHDR. */
+	AVALANCHEMEDIA_API void ConvertTextureRGBGamma(FRHICommandListImmediate& InRHICmdList, FTextureRHIRef InSourceTexture, FTextureRHIRef InDestTarget, bool bInSrgbToLinear, float InGamma);
+}
\ No newline at end of file
diff --git a/Engine/Plugins/Experimental/Avalanche/Source/AvalancheMedia/Public/OutputDevices/IAvaDeviceProviderProxyManager.h b/Engine/Plugins/Experimental/Avalanche/Source/AvalancheMedia/Public/OutputDevices/IAvaDeviceProviderProxyManager.h
new file mode 100644
--- /dev/null
+++ b/Engine/Plugins/Experimental/Avalanche/Source/AvalancheMedia/Public/OutputDevices/IAvaDeviceProviderProxyManager.h
@@ -0,0 +1,42 @@
+// Copyright Epic Games, Inc. All Rights Reserved.
+
+#pragma once
+
+#include "Templates/SharedPointer.h"
+#include "UObject/NameTypes.h"
+
+class FAvaDeviceProviderWrapper;
+class IMediaIOCoreDeviceProvider;
+struct FAvaDeviceProviderDataList;
+
+class AVALANCHEMEDIA_API IAvaDeviceProviderProxyManager
+{
+public:
+	virtual ~IAvaDeviceProviderProxyManager() = default;
+	
+	virtual const FAvaDeviceProviderWrapper* GetDeviceProviderWrapper(FName InDeviceProviderName) const = 0;
+
+	/** Collects all the device providers related to the given server. */
+	virtual TArray<const IMediaIOCoreDeviceProvider*> GetDeviceProvidersForServer(const FString& InServerName) const = 0;
+
+	/** Finds the server name of the given device. */
+	virtual FString FindServerNameForDevice(const FName& InDeviceProviderName, const FName& InDeviceName) const = 0;
+
+	/** Returns true if the device can be found in a local provider. */
+	virtual bool IsLocalDevice(const FName& InDeviceProviderName, const FName& InDeviceName) const = 0;
+
+	/** Install the given server's proxy data, installing wrappers as necessary. */
+	virtual void Install(const FString& InServerName, const FAvaDeviceProviderDataList& InDeviceProviderDataList) = 0;
+
+	/** Remove the given server from all the wrappers. */
+	virtual void Uninstall(const FString& InServerName) = 0;
+
+	/** Returns the set of all server names. */
+	virtual TSet<FString> GetServerNames() const = 0;
+
+	/** Returns device provider data list (all device providers) for a specific server. */
+	virtual const TSharedPtr<const FAvaDeviceProviderDataList> GetDeviceProviderDataListForServer(const FString& InServerName) const = 0;
+
+	/** Returns the default local server name. */
+	virtual FString GetLocalServerName() const = 0;
+};
diff --git a/Engine/Plugins/Experimental/Avalanche/Source/AvalancheMedia/Public/OutputDevices/Slate/SAvaCaptureImage.h b/Engine/Plugins/Experimental/Avalanche/Source/AvalancheMedia/Public/OutputDevices/Slate/SAvaCaptureImage.h
new file mode 100644
--- /dev/null
+++ b/Engine/Plugins/Experimental/Avalanche/Source/AvalancheMedia/Public/OutputDevices/Slate/SAvaCaptureImage.h
@@ -0,0 +1,35 @@
+// Copyright Epic Games, Inc. All Rights Reserved.
+
+#pragma once
+#include "CoreMinimal.h"
+#include "Widgets/Images/SImage.h"
+
+class AVALANCHEMEDIA_API SAvaCaptureImage : public SImage
+{
+public:
+	SLATE_BEGIN_ARGS(SAvaCaptureImage)
+		: _ImageArgs()
+		, _EnableGammaCorrection(true)
+	{}
+		SLATE_ARGUMENT(SImage::FArguments, ImageArgs)
+		SLATE_ARGUMENT(bool, EnableGammaCorrection)
+		SLATE_ATTRIBUTE(bool, EnableBlending)
+		SLATE_ATTRIBUTE(bool, ShouldInvertAlpha)
+	SLATE_END_ARGS()
+
+	/** Constructs this widget with InArgs */
+	void Construct(const FArguments& InArgs);
+	
+	virtual int32 OnPaint(const FPaintArgs& Args
+		, const FGeometry& AllottedGeometry
+		, const FSlateRect& MyCullingRect
+		, FSlateWindowElementList& OutDrawElements
+		, int32 LayerId
+		, const FWidgetStyle& InWidgetStyle
+		, bool bParentEnabled) const override;
+
+protected:
+	bool bEnableGammaCorrection = true;
+	TAttribute<bool> EnableBlending;
+	TAttribute<bool> ShouldInvertAlpha;
+};
diff --git a/Engine/Plugins/Experimental/Avalanche/Source/AvalancheMedia/Public/Playback/AvaMediaPlayableGroup.h b/Engine/Plugins/Experimental/Avalanche/Source/AvalancheMedia/Public/Playback/AvaMediaPlayableGroup.h
new file mode 100644
--- /dev/null
+++ b/Engine/Plugins/Experimental/Avalanche/Source/AvalancheMedia/Public/Playback/AvaMediaPlayableGroup.h
@@ -0,0 +1,222 @@
+// Copyright Epic Games, Inc. All Rights Reserved.
+
+#pragma once
+
+#include "Framework/AvalancheGameInstance.h"
+#include "Playback/IAvaPlayableVisibilityConstraint.h"
+#include "UObject/Object.h"
+#include "UObject/ObjectKey.h"
+#include "UObject/ObjectPtr.h"
+#include "UObject/SoftObjectPtr.h"
+
+#include "AvaMediaPlayableGroup.generated.h"
+
+class UAvaMediaPlayableGroupManager;
+class UAvalanchePlayable;
+class UAvalanchePlayableTransition;
+class UTextureRenderTarget2D;
+
+/**
+ * This defines a game instance playable group.
+ * It tracks and manage the game instance and playables state.
+ *
+ * The design goal of this class is to remove all the playback management
+ * from UAvalancheGameInstance and move it to the playable framework. This
+ * should allow us to hook the playback framework to any game instance, including PIE
+ * so it can work with any work flow (editor, PIE, game, nDisplay, etc).
+ *
+ * Also, ideally, the playable class itself should be "game instance" agnostic
+ * and do all it's bidding on it's container through this class.
+ */
+UCLASS()
+class AVALANCHEMEDIA_API UAvaMediaPlayableGroup : public UObject
+{
+	GENERATED_BODY()
+
+public:
+	/**
+	 * PlayableGroup creation information contains the necessary information to
+	 * create an instance of a playable group.
+	 */
+	struct FPlayableGroupCreationInfo
+	{
+		/** Container for shared playable groups. */
+		UAvaMediaPlayableGroupManager* PlayableGroupManager = nullptr;
+		/** Source asset path. */
+		FSoftObjectPath SourceAssetPath;
+		/** Channel name this playable will be instanced in. */
+		FName ChannelName;
+		/** Indicate if the group is for remote proxy playables. */
+		bool bIsRemoteProxy = false;
+		/** Indicate if the group is shared. */
+		bool bIsSharedGroup = false;
+	};
+	
+	static UAvaMediaPlayableGroup* MakePlayableGroup(UObject* InOuter, const FPlayableGroupCreationInfo& InPlayableGroupInfo);
+
+	/** Register a playable to this group when it is created. */
+	void RegisterPlayable(UAvalanchePlayable* InPlayable);
+
+	/** Unregister a playable when it is about to be deleted. */
+	void UnregisterPlayable(UAvalanchePlayable* InPlayable);
+
+	/** Returns true if there are any valid registered playables. */
+	bool HasPlayables() const;
+
+	/** Returns true if there are any valid registered playables that are currently playing. */
+	bool HasPlayingPlayables() const;
+
+	UAvalanchePlayable* FindFirstPlayableBySourceAssetPath(const FSoftObjectPath& InSourceAssetPath) const;
+
+	void RegisterPlayableTransition(UAvalanchePlayableTransition* InPlayableTransition);
+	
+	void UnregisterPlayableTransition(UAvalanchePlayableTransition* InPlayableTransition);
+
+	/** Tick transitions that have been registered. Returns the number of transitions that where ticked. */
+	void TickTransitions(double InDeltaSeconds);
+
+	bool HasTransitions() const;
+
+	/**
+	 * Creates the game instance's world if it wasn't already.
+	 * @return true if the world was created. false if nothing was done.
+	 */
+	bool ConditionalCreateWorld();
+
+	/**
+	 * Begin playing the game instance's world if it wasn't already.
+	 * @return true if the BeginPlay was done (i.e. on the state transition only), false otherwise.
+	 */
+	virtual bool ConditionalBeginPlay(const FAvalancheInstancePlaySettings& InWorldPlaySettings);
+
+	virtual void RequestEndPlayWorld(bool bInForceImmediate);
+
+	/** Keep track of the last playable that applied it's camera in the viewport/controller. */
+	void SetLastAppliedCameraPlayable(UAvalanchePlayable* InPlayable);
+	
+	/** Search for the first remaining playing playable and use it's camera. */
+	bool UpdateCameraSetup();
+
+	bool IsWorldPlaying() const { return GameInstance ? GameInstance->IsWorldPlaying() : false;}
+
+	bool IsRenderTargetReady() const { return GameInstance ? GameInstance->IsRenderTargetReady() : false;}
+
+	/**
+	 * Current logic for the render target, use the game instance's if there, fallback to internal one if not.
+	 */
+	UTextureRenderTarget2D* GetRenderTarget() const { return GameInstance ? GameInstance->GetRenderTarget() : RenderTarget.Get(); }
+
+	/** Return a vanilla game instance, we want to eventually support any game instance. */
+	UGameInstance* GetGameInstance() const { return GameInstance;}
+	
+	UWorld* GetPlayWorld() const { return GameInstance ? GameInstance->GetPlayWorld() : nullptr; }
+	
+	/**
+	 * Unloads the game instance's world if no more playables are loaded.
+	 * @return true if the world was unloaded. false if nothing was done.
+	 */ 
+	bool ConditionalRequestUnloadWorld(bool bForceImmediate);
+
+	/**
+	 * @brief Queue a camera cut for the next rendered frame.
+	 */
+	void QueueCameraCut();
+
+	/**
+	 * @brief Notify the playable group that a playable is loading an asset.
+	 * @param InPlayable Source of the event.
+	 */
+	void NotifyLevelStreaming(UAvalanchePlayable* InPlayable);
+
+	UAvaMediaPlayableGroupManager* GetPlayableGroupManager() const { return ParentPlayableGroupManagerWeak.Get(); }
+
+	void RegisterVisibilityConstraint(const TWeakInterfacePtr<IAvaPlayableVisibilityConstraint>& InVisibilityConstraint);
+	void UnregisterVisibilityConstraint(const IAvaPlayableVisibilityConstraint* InVisibilityConstraint);
+
+	void RequestSetVisibility(UAvalanchePlayable* InPlayable, bool bInShouldBeVisible);
+
+protected:
+	bool IsVisibilityConstrained(const UAvalanchePlayable* InPlayable) const;
+
+	void OnPlayableStatusChanged(UAvalanchePlayable* InPlayable);
+	
+	void ConditionalRegisterWorldDelegates(UWorld* InWorld);
+	void UnregisterWorldDelegates(UWorld* InWorld);
+
+	bool DisplayLoadedAssets(FText& OutText, FLinearColor& OutColor);
+	bool DisplayPlayingAssets(FText& OutText, FLinearColor& OutColor);
+	bool DisplayTransitions(FText& OutText, FLinearColor& OutColor);
+
+	static UPackage* MakeGameInstancePackage(const FSoftObjectPath& InSourceAssetPath, const FName& InChannelName);
+	static UPackage* MakeSharedInstancePackage(const FName& InChannelName);
+	static UPackage* MakeInstancePackage(const FString& InInstancePackageName);
+
+public:	
+	UPROPERTY(Transient)
+	TObjectPtr<UTextureRenderTarget2D> RenderTarget;	// Optional
+	
+	UPROPERTY(Transient)
+	TObjectPtr<UAvalancheGameInstance> GameInstance;
+
+	UPROPERTY(Transient)
+	TObjectPtr<UPackage> GameInstancePackage;
+
+protected:
+	/** PlayableGroup Manager handling this playable group. */
+	TWeakObjectPtr<UAvaMediaPlayableGroupManager> ParentPlayableGroupManagerWeak;
+	
+	/** List of playables for this group. */
+	TSet<TObjectKey<UAvalanchePlayable>> Playables;
+
+	/** Last playable that applied a camera. */
+	TWeakObjectPtr<UAvalanchePlayable> LastAppliedCameraPlayableWeak;
+
+	/** Set of registered playable transitions for this group. Remark: used for ticking. */
+	TSet<TObjectKey<UAvalanchePlayableTransition>> PlayableTransitions;
+
+	/** If transitions are added or removed while ticking, we need to protect transition iterator. */
+	bool bIsTickingTransitions = false;
+
+	/** Set of transitions to remove accumulated during transition ticking. */
+	TSet<TObjectKey<UAvalanchePlayableTransition>> PlayableTransitionsToRemove;
+
+	/** Set of transitions to add accumulated during transition ticking. */
+	TSet<TObjectKey<UAvalanchePlayableTransition>> PlayableTransitionsToAdd;
+
+	/**
+	 * Since the UViewportStatsSubsystem delegation mechanism does not allow us to verify
+	 * if it is bound to this world or another one. We need this auxiliary binding
+	 * tracking to compensate.
+	 */
+	TWeakObjectPtr<UWorld> LastWorldBoundToDisplayDelegates;
+
+	/** UViewportStatsSubsystem display delegate indices. */
+	TArray<int32> DisplayDelegateIndices;
+
+	struct FVisibilityRequest
+	{
+		TWeakObjectPtr<UAvalanchePlayable> PlayableWeak;
+		bool bShouldBeVisible;
+
+		void Execute(const UAvaMediaPlayableGroup* InPlayableGroup) const;
+	};
+	
+	TArray<FVisibilityRequest> VisibilityRequests;
+
+	TArray<TWeakInterfacePtr<IAvaPlayableVisibilityConstraint>> VisibilityConstraints;
+};
+
+/**
+ *	Remote Proxy Playable Group doesn't have a game instance.
+ */
+UCLASS()
+class AVALANCHEMEDIA_API UAvaMediaRemoteProxyPlayableGroup : public UAvaMediaPlayableGroup
+{
+	GENERATED_BODY()
+	
+public:
+	virtual bool ConditionalBeginPlay(const FAvalancheInstancePlaySettings& InWorldPlaySettings) override;
+	virtual void RequestEndPlayWorld(bool bInForceImmediate) override;
+	
+	bool bIsPlaying = false;
+};
\ No newline at end of file
diff --git a/Engine/Plugins/Experimental/Avalanche/Source/AvalancheMedia/Public/Playback/AvaMediaPlayableGroupManager.h b/Engine/Plugins/Experimental/Avalanche/Source/AvalancheMedia/Public/Playback/AvaMediaPlayableGroupManager.h
new file mode 100644
--- /dev/null
+++ b/Engine/Plugins/Experimental/Avalanche/Source/AvalancheMedia/Public/Playback/AvaMediaPlayableGroupManager.h
@@ -0,0 +1,114 @@
+// Copyright Epic Games, Inc. All Rights Reserved.
+
+#pragma once
+
+#include "UObject/Object.h"
+#include "UObject/ObjectPtr.h"
+#include "UObject/SoftObjectPtr.h"
+
+#include "AvaMediaPlayableGroupManager.generated.h"
+
+class UAvaMediaPlayableGroup;
+class UAvaMediaPlayableGroupManager;
+class UAvalancheGameInstance;
+
+/**
+ * Manager for the shared playable groups per channel.
+ */
+UCLASS()
+class AVALANCHEMEDIA_API UAvaMediaChannelPlayableGroupManager : public UObject
+{
+	GENERATED_BODY()
+
+public:
+	UAvaMediaPlayableGroup* GetOrCreateSharedLevelGroup(bool bInIsRemoteProxy);
+
+	UAvaMediaPlayableGroupManager* GetPlayableGroupManager() const;
+	
+protected:
+	//~ Begin UObject
+	virtual void BeginDestroy() override;
+	//~ End UObject
+
+	void Shutdown();
+	
+protected:
+	FName ChannelName;
+
+	/**
+	 * This is the shared playable group for this channel.
+	 * 
+	 * For now all the levels go in the same group for a given channel.
+	 * This a weak ptr because it is only a cache. The ownership
+	 * of the group is with the playable objects.
+	 */
+	TWeakObjectPtr<UAvaMediaPlayableGroup> SharedLevelGroupWeak;
+
+	/**
+	 * For remote proxy playables, we need a specialized group
+	 * that doesn't implement the local game instance but still
+	 * provides the correct logic to emulate functionality.
+	 */
+	TWeakObjectPtr<UAvaMediaPlayableGroup> SharedRemoteProxyLevelGroupWeak;
+
+	friend class UAvaMediaPlayableGroupManager;
+};
+
+/**
+ * Manager for the shared playable groups.
+ * The scope of this manager is either global (in the global playback manager)
+ * or for a given playback manager.
+ */
+UCLASS()
+class AVALANCHEMEDIA_API UAvaMediaPlayableGroupManager : public UObject
+{
+	GENERATED_BODY()
+	
+public:
+	void Init();
+
+	void Shutdown();
+
+	void Tick(double InDeltaSeconds);
+
+	UAvaMediaChannelPlayableGroupManager* FindChannelManager(const FName& InChannelName) const
+	{
+		const TObjectPtr<UAvaMediaChannelPlayableGroupManager>* ChannelManager = ChannelManagers.Find(InChannelName);
+		return ChannelManager ? *ChannelManager : nullptr;
+	}
+
+	UAvaMediaChannelPlayableGroupManager* FindOrAddChannelManager(const FName& InChannelName);
+
+	UAvaMediaPlayableGroup* GetOrCreateSharedLevelGroup(const FName& InChannelName, bool bInIsRemoteProxy)
+	{
+		UAvaMediaChannelPlayableGroupManager* ChannelManager = FindOrAddChannelManager(InChannelName);
+		return ChannelManager ? ChannelManager->GetOrCreateSharedLevelGroup(bInIsRemoteProxy) : nullptr;
+	}
+
+	void RegisterForLevelStreamingUpdate(UAvaMediaPlayableGroup* InPlayableGroup);
+	void UnregisterFromLevelStreamingUpdate(UAvaMediaPlayableGroup* InPlayableGroup);
+
+	void RegisterForTransitionTicking(UAvaMediaPlayableGroup* InPlayableGroup);
+	void UnregisterFromTransitionTicking(UAvaMediaPlayableGroup* InPlayableGroup);
+
+protected:
+	//~ Begin UObject
+	virtual void BeginDestroy() override;
+	//~ End UObject
+
+	void OnGameInstanceEndPlay(UAvalancheGameInstance* InGameInstance, FName InChannelName);
+
+	void UpdateLevelStreaming();
+	
+	void TickTransitions(double InDeltaSeconds);
+	
+protected:
+	UPROPERTY(Transient)
+	TMap<FName, TObjectPtr<UAvaMediaChannelPlayableGroupManager>> ChannelManagers;
+
+	bool bIsUpdatingStreaming = false;
+	TSet<TWeakObjectPtr<UAvaMediaPlayableGroup>> GroupsToUpdateStreaming;
+
+	bool bIsTickingTransitions = false;
+	TSet<TWeakObjectPtr<UAvaMediaPlayableGroup>> GroupsToTickTransitions;
+};
\ No newline at end of file
diff --git a/Engine/Plugins/Experimental/Avalanche/Source/AvalancheMedia/Public/Playback/AvaMediaPlaybackClient.h b/Engine/Plugins/Experimental/Avalanche/Source/AvalancheMedia/Public/Playback/AvaMediaPlaybackClient.h
new file mode 100644
--- /dev/null
+++ b/Engine/Plugins/Experimental/Avalanche/Source/AvalancheMedia/Public/Playback/AvaMediaPlaybackClient.h
@@ -0,0 +1,491 @@
+// Copyright Epic Games, Inc. All Rights Reserved.
+
+#pragma once
+
+#include "Containers/Ticker.h"
+#include "IMessageContext.h"
+#include "MessageEndpoint.h"
+#include "OutputDevices/AvaDeviceProviderData.h"
+#include "Playback/AvaMediaPlaybackMessages.h"
+#include "Playback/IAvaMediaPlaybackClient.h"
+#include "UObject/ObjectSaveContext.h"
+
+class FAvaMediaModule;
+struct FAvaOutputChannel;
+
+// Listens to command to play and broadcast
+class AVALANCHEMEDIA_API FAvaMediaPlaybackClient : public IAvaMediaPlaybackClient
+{
+public:
+	FAvaMediaPlaybackClient(FAvaMediaModule* InParentModule);
+	virtual ~FAvaMediaPlaybackClient() override;
+
+	void Init();
+
+	//IAvaMediaPlaybackClient
+	virtual int32 GetNumConnectedServers() const override;
+	virtual TArray<FString> GetServerNames() const override;
+	virtual FMessageAddress GetServerAddress(const FString& InServerName) const override;
+	virtual bool HasServerUserData(const FString& InServerName, const FString& InKey) const override;
+	virtual const FString& GetServerUserData(const FString& InServerName, const FString& InKey) const override;
+	virtual bool HasUserData(const FString& InKey) const override;
+	virtual const FString& GetUserData(const FString& InKey) const override;
+	virtual void SetUserData(const FString& InKey, const FString& InData) override;
+	virtual void RemoveUserData(const FString& InKey) override;
+	virtual void BroadcastStatCommand(const FString& InCommand, bool bInBroadcastLocalState) override;
+	virtual void RequestPlaybackAssetStatus(const FSoftObjectPath& InAssetPath, const FString& InChannelOrServerName, bool bInForceRefresh) override;
+	virtual void RequestPlayback(const FGuid& InInstanceId, const FSoftObjectPath& InAssetPath, const FString& InChannelName, EAvaMediaPlaybackAction InAction, const FString& InArguments = FString()) override;
+	virtual void RequestAnimPlayback(const FGuid& InInstanceId, const FSoftObjectPath& InAssetPath, const FString& InChannelName, const FAnimPlaySettings& InAnimSettings) override;
+	virtual void RequestAnimAction(const FGuid& InInstanceId, const FSoftObjectPath& InAssetPath, const FString& InChannelName, const FString& InAnimationName, EAvaMediaAnimAction InAction) override;
+	virtual void RequestRemoteControlUpdate(const FGuid& InInstanceId, const FSoftObjectPath& InAssetPath, const FString& InChannelName, const FAvalancheRemoteControlValues& InRemoteControlValues) override;
+	virtual void RequestPlayableTransitionStart(const FGuid& InTransitionId, TArray<FGuid>&& InEnterInstanceIds, TArray<FGuid>&& InPlayingInstanceIds, TArray<FGuid>&& InExitInstanceIds, TArray<FAvalancheRemoteControlValues>&& InEnterValues, const FName& InChannelName, EAvalanchePlayableTransitionFlags InTransitionFlags) override;
+	virtual void RequestPlayableTransitionStop(const FGuid& InTransitionId, const FName& InChannelName) override;
+	virtual void RequestBroadcast(const FString& InProfile, const FName& InChannel, const TArray<UMediaOutput*>& InRemoteMediaOutputs, EAvaMediaBroadcastAction InAction) override;
+	virtual bool IsMediaOutputRemoteFallback(const UMediaOutput* InMediaOutput) override;
+	virtual EAvaMediaIssueSeverity GetMediaOutputIssueSeverity(const FString& InServerName, const FString& InChannelName, const FGuid& InOutputGuid) const override;
+	virtual const TArray<FString>& GetMediaOutputIssueMessages(const FString& InServerName, const FString& InChannelName, const FGuid& InOutputGuid) const override;
+	virtual EAvaMediaOutputState GetMediaOutputState(const FString& InServerName, const FString& InChannelName, const FGuid& InOutputGuid) const override;
+	virtual bool HasAnyServerOnlineForChannel(const FName& InChannelName) const override;
+	virtual TOptional<EAvaMediaPlaybackStatus> GetRemotePlaybackStatus(const FGuid& InInstanceId, const FSoftObjectPath& InAssetPath, const FString& InChannelName, const FString& InServerName) const override;
+	virtual const FString* GetRemotePlaybackUserData(const FGuid& InInstanceId, const FSoftObjectPath& InAssetPath, const FString& InChannelName, const FString& InServerName) const override;
+	virtual TOptional<EAvaMediaPlaybackAssetStatus> GetRemotePlaybackAssetStatus(const FSoftObjectPath& InAssetPath, const FString& InServerName) const override;
+	//~IAvaMediaPlaybackClient
+
+	FString GetServerProjectContentPath(const FString& InServerName) const;
+	uint32 GetServerProcessId(const FString& InServerName) const;
+
+	FString GetProjectContentPath() const { return ProjectContentPath; }
+
+protected:
+	TOptional<EAvaMediaPlaybackStatus> GetPlaybackStatusForServer(const FString& InServerName, const FString& InChannelName, const FGuid& InInstanceId, const FSoftObjectPath& InAssetPath) const;
+	const FString* GetPlaybackUserDataForServer(const FString& InServerName, const FString& InChannelName, const FGuid& InInstanceId, const FSoftObjectPath& InAssetPath) const;
+	void RequestPlaybackAssetStatusForServer(const FSoftObjectPath& InAssetPath, const FString& InServerName, bool bInForceRefresh);
+
+	void Tick();
+
+	bool HandlePingTicker(float InDeltaTime);
+
+	// Message handlers
+	void HandlePlaybackPongMessage(const FAvaMediaPlaybackPong& InMessage, const TSharedRef<IMessageContext, ESPMode::ThreadSafe>& InContext);
+	void HandlePlaybackLogMessage(const FAvaMediaPlaybackLog& InMessage, const TSharedRef<IMessageContext, ESPMode::ThreadSafe>& InContext);
+	void HandleUpdateServerUserData(const FAvaMediaUpdateServerUserData& InMessage, const TSharedRef<IMessageContext, ESPMode::ThreadSafe>& InContext);
+	void HandleStatStatus(const FAvaMediaStatStatus& InMessage, const TSharedRef<IMessageContext, ESPMode::ThreadSafe>& InContext);
+	void HandleDeviceProviderDataListMessage(const FAvaDeviceProviderDataList& InMessage, const TSharedRef<IMessageContext, ESPMode::ThreadSafe>& InContext);
+	void HandleBroadcastStatusMessage(const FAvaMediaBroadcastStatus& InMessage, const TSharedRef<IMessageContext, ESPMode::ThreadSafe>& InContext);
+	void HandlePlaybackAssetStatusMessage(const FAvaMediaPlaybackAssetStatus& InMessage, const TSharedRef<IMessageContext, ESPMode::ThreadSafe>& InContext);
+	void HandlePlaybackStatusMessage(const FAvaMediaPlaybackStatus& InMessage, const TSharedRef<IMessageContext, ESPMode::ThreadSafe>& InContext);
+	void HandlePlaybackStatusesMessage(const FAvaMediaPlaybackStatuses& InMessage, const TSharedRef<IMessageContext, ESPMode::ThreadSafe>& InContext);
+	void HandlePlaybackSequenceEventMessage(const FAvaMediaPlayableSequenceEvent& InMessage, const TSharedRef<IMessageContext, ESPMode::ThreadSafe>& InContext);
+	void HandlePlaybackTransitionEventMessage(const FAvaMediaPlayableTransitionEvent& InMessage, const TSharedRef<IMessageContext, ESPMode::ThreadSafe>& InContext);
+	
+	void RegisterCommands();
+	
+	// Command handlers
+	void PingServersCommand(const TArray<FString>& InArgs) { PublishPlaybackPing(FDateTime::UtcNow(), false);}
+	void RequestPlaybackAssetStatusCommand(const TArray<FString>& InArgs);
+	void RequestPlaybackCommand(const TArray<FString>& InArgs);
+	void RequestBroadcastCommand(const TArray<FString>& InArgs);
+	void SetUserDataCommand(const TArray<FString>& InArgs);
+	void ShowStatusCommand(const TArray<FString>& InArgs);
+
+	// Event handlers
+	void OnBroadcastChanged(EAvaBroadcastChange InChange);
+	void OnChannelChanged(const FAvaOutputChannel& InChannel, EAvaChannelChange InChange);
+	void OnAvaMediaSettingsChanged(UObject*, struct FPropertyChangedEvent&);
+	void OnPreSavePackage(UPackage* InPackage, FObjectPreSaveContext InObjectSaveContext);
+	void OnPackageSaved(const FString& InPackageFileName, UPackage* InPackage, FObjectPostSaveContext InObjectSaveContext);
+	void OnAssetRemoved(const FAssetData& InAssetData);
+	
+	void ApplyAvaMediaSettings();
+
+protected:
+	void FillClientInfo(FAvaMediaPlaybackClientMessageBase& InMessage)
+	{
+		// Remark: for now, client name and computer name are the same.
+		// We don't support having many clients per computer at this point.
+		InMessage.ClientName =  ComputerName;
+	}
+
+	template<typename MessageType>
+	void PublishRequest(MessageType* InMessage)
+	{
+		FillClientInfo(*InMessage);
+		MessageEndpoint->Publish(InMessage);
+	}
+
+	template<typename MessageType>
+	void SendRequest(MessageType* InMessage, const FMessageAddress& InRecipient, EMessageFlags flags = EMessageFlags::None)
+	{
+		FillClientInfo(*InMessage);
+		MessageEndpoint->Send(InMessage, MessageType::StaticStruct(), flags, nullptr,
+			TArrayBuilder<FMessageAddress>().Add(InRecipient), FTimespan::Zero(), FDateTime::MaxValue());
+	}
+	
+	template<typename MessageType>
+	void SendRequest(MessageType* InMessage, const TArray<FMessageAddress>& InRecipients, EMessageFlags flags = EMessageFlags::None)
+	{
+		FillClientInfo(*InMessage);
+		MessageEndpoint->Send(InMessage, MessageType::StaticStruct(), flags, nullptr,
+			InRecipients, FTimespan::Zero(), FDateTime::MaxValue());
+	}
+
+	void PublishPlaybackPing(const FDateTime& InCurrentTime, bool bInAutoPing);
+	void SendUserDataUpdate(const TArray<FMessageAddress>& InRecipients);
+	void SendBroadcastSettingsUpdate(const TArray<FMessageAddress>& InRecipients);
+	void SendAvalancheInstanceSettingsUpdate(const TArray<FMessageAddress>& InRecipients);
+	void SendBroadcastChannelSettingsUpdate(const TArray<FMessageAddress>& InRecipients, const FAvaOutputChannel& InChannel);
+	void SendPackageEvent(const TArray<FMessageAddress>& InRecipients, const FName& InPackageName, EAvaMediaPlaybackPackageEvent InEvent);
+	void SendStatCommand(const FString& InCommand, bool bInBroadcastLocalState, const TArray<FMessageAddress>& InRecipients);
+	void SendClientInfo(const FMessageAddress& InRecipient);	
+	void RemoveDeadServers(const FDateTime& InCurrentTime);
+	void UpdateServerAddresses();
+	
+	/**
+	 * Get the server addresses for the given channel in the current broadcast profile.
+	 * This is based on the channel's MediaOutputInfos. It will properly work even
+	 * if the channel is idle (but not offline).
+	 */
+	TArray<FMessageAddress> GetServerAddressesForChannel(const FName& InChannelName) const;
+	TArray<FMessageAddress> GetServerAddressesForChannel(const FString& InChannelName) const
+	{
+		return GetServerAddressesForChannel(FName(InChannelName));
+	}
+
+	/**
+	 * Get the server names for the given channel in the current broadcast profile.
+	 * This is based on the channel's MediaOutputInfos. It will properly work
+	 * regardless of the channel's status (idle or even offline).
+	 */
+	TArray<FString> GetServerNamesForChannel(const FName& InChannelName) const;
+
+	FString GetServerNameForMediaOutputFallback(const UMediaOutput* InMediaOutput) const;
+
+private:
+	FAvaMediaModule* ParentModule;
+	
+	TArray<IConsoleObject*> ConsoleCommands;
+	
+	/** Holds the messaging endpoint. */
+	TSharedPtr<FMessageEndpoint, ESPMode::ThreadSafe> MessageEndpoint;
+	
+	/** Holds a delegate to be invoked when the widget ticks. */
+	FTickerDelegate PingTickDelegate;
+
+	/** Handle to the registered TickDelegate. */
+	FTSTicker::FDelegateHandle PingTickDelegateHandle;
+	
+	FDelegateHandle BroadcastChangedHandle;
+	
+	FString ComputerName;
+	FString ProjectContentPath;
+	uint32 ProcessId = 0;
+	TMap<FString, FString> UserDataEntries;
+	
+	/** Structure holding information on pending playback asset status requests.*/
+	struct FPendingPlaybackAssetStatusRequest
+	{
+		FDateTime ExpirationTime;
+		bool bForceRefresh = false;
+	};
+	
+	/** Remote Broadcast Channel Info */
+	struct FBroadcastChannelInfo
+	{
+		EAvaChannelState ChannelState;
+		EAvaMediaIssueSeverity ChannelIssueSeverity;
+		TMap<FGuid, FAvaMediaOutputStatus> MediaOutputStatuses;
+	};
+
+	/** Remote Playback Instance Info */
+	struct FPlaybackInstanceInfo
+	{
+		EAvaMediaPlaybackStatus  Status;
+		TOptional<FString> UserData;
+	};
+	
+	/** Remote playback for multiple instances of a given asset. */
+	struct FPlaybackAssetInfo
+	{
+		TMap<FGuid, FPlaybackInstanceInfo> InstanceByIds;
+
+		/*
+		 * Return a relevance factor for status reporting, based on what we would like to see in the page status.
+		 */
+		static int GetStatusRelevance(EAvaMediaPlaybackStatus InPlaybackStatus)
+		{
+			switch (InPlaybackStatus)
+			{
+			case EAvaMediaPlaybackStatus::Unknown: return 0;
+			case EAvaMediaPlaybackStatus::Missing: return 1;
+			case EAvaMediaPlaybackStatus::Syncing: return 2;
+			case EAvaMediaPlaybackStatus::Available: return 3;
+			case EAvaMediaPlaybackStatus::Loading: return 4;
+			case EAvaMediaPlaybackStatus::Loaded: return 5;
+			case EAvaMediaPlaybackStatus::Starting: return 6;
+			case EAvaMediaPlaybackStatus::Started: return 7;
+			case EAvaMediaPlaybackStatus::Unloading: return 3;
+			case EAvaMediaPlaybackStatus::Error: return 1;
+			default: return 0;
+			}
+		}
+		
+		TOptional<EAvaMediaPlaybackStatus> GetInstanceStatus(const FGuid& InInstanceId) const
+		{
+			// If the request is for a specific instance, check if we have it.
+			if (InInstanceId.IsValid())
+			{
+				const FPlaybackInstanceInfo* Instance = InstanceByIds.Find(InInstanceId);
+				return Instance ? TOptional(Instance->Status) : TOptional<EAvaMediaPlaybackStatus>();
+			}
+			
+			// If the instanceId is not specified, return the most relevant status we have already received.
+			// Todo: investigate why this happens and if there is better way to handle this (like better request handling).
+			TOptional<EAvaMediaPlaybackStatus> MostRelevantStatus;
+			for (const TPair<FGuid, FPlaybackInstanceInfo>& PlaybackInfo :  InstanceByIds)
+			{
+				if (!MostRelevantStatus.IsSet() || GetStatusRelevance(PlaybackInfo.Value.Status) > GetStatusRelevance(MostRelevantStatus.GetValue()))
+				{
+					MostRelevantStatus = PlaybackInfo.Value.Status;
+				}
+			}
+			return MostRelevantStatus;
+		}
+
+		const FPlaybackInstanceInfo* GetInstanceInfo(const FGuid& InInstanceId) const
+		{
+			return InInstanceId.IsValid() ? InstanceByIds.Find(InInstanceId) : nullptr;
+		}
+
+		FPlaybackInstanceInfo* GetOrAddInstanceInfo(const FGuid& InInstanceId)
+		{
+			if (InInstanceId.IsValid())
+			{
+				FPlaybackInstanceInfo* Instance = InstanceByIds.Find(InInstanceId);
+				return Instance ? Instance : &InstanceByIds.Emplace(InInstanceId);
+			}
+			return nullptr;
+		}
+		
+		void SetInstanceStatus(const FGuid& InInstanceId, EAvaMediaPlaybackStatus InStatus)
+		{
+			if (FPlaybackInstanceInfo* Instance = GetOrAddInstanceInfo(InInstanceId))
+			{
+				Instance->Status = InStatus;
+			}
+		}
+
+		const FString* GetInstanceUserData(const FGuid& InInstanceId) const
+		{
+			// If the request is for a specific instance, check if we have it.
+			if (const FPlaybackInstanceInfo* Instance = GetInstanceInfo(InInstanceId))
+			{
+				return Instance->UserData.IsSet() ? &Instance->UserData.GetValue() : nullptr;
+			}
+			return nullptr;
+		}
+
+		void SetInstanceUserData(const FGuid& InInstanceId, const FString& InUserData)
+		{
+			if (FPlaybackInstanceInfo* Instance = GetOrAddInstanceInfo(InInstanceId))
+			{
+				Instance->UserData = InUserData;
+			}
+		}
+	};
+
+	/** Playback Asset Info for a given channel.  */
+	struct FPlaybackChannelInfo
+	{
+		TMap<FSoftObjectPath, TUniquePtr<FPlaybackAssetInfo>> AssetInfoByPaths;
+
+		FPlaybackAssetInfo* GetAssetInfo(const FSoftObjectPath& InAssetPath) const
+		{
+			const TUniquePtr<FPlaybackAssetInfo>* AssetInfo = AssetInfoByPaths.Find(InAssetPath);
+			return AssetInfo ? (*AssetInfo).Get() : nullptr;
+		}
+
+		FPlaybackAssetInfo* GetOrAddAssetInfo(const FSoftObjectPath& InAssetPath)
+		{
+			if (const TUniquePtr<FPlaybackAssetInfo>* AssetInfo = AssetInfoByPaths.Find(InAssetPath))
+			{
+				return (*AssetInfo).Get();
+			}
+			return AssetInfoByPaths.Add(InAssetPath, MakeUnique<FPlaybackAssetInfo>()).Get();
+		}
+		
+		TOptional<EAvaMediaPlaybackStatus> GetInstanceStatus(const FGuid& InInstanceId, const FSoftObjectPath& InAssetPath) const
+		{
+			const FPlaybackAssetInfo* AssetInfo = GetAssetInfo(InAssetPath);
+			return AssetInfo ? AssetInfo->GetInstanceStatus(InInstanceId) : TOptional<EAvaMediaPlaybackStatus>();
+		}
+		
+		void SetInstanceStatus(const FGuid& InInstanceId, const FSoftObjectPath& InAssetPath, EAvaMediaPlaybackStatus InStatus)
+		{
+			if (FPlaybackAssetInfo* AssetInfo = GetOrAddAssetInfo(InAssetPath))
+			{
+				AssetInfo->SetInstanceStatus(InInstanceId, InStatus);
+			}
+		}
+
+		const FString* GetInstanceUserData(const FGuid& InInstanceId, const FSoftObjectPath& InAssetPath) const
+		{
+			const FPlaybackAssetInfo* AssetInfo = GetAssetInfo(InAssetPath);
+			return AssetInfo ? AssetInfo->GetInstanceUserData(InInstanceId) : nullptr;
+		}
+		
+		void SetInstanceUserData(const FGuid& InInstanceId, const FSoftObjectPath& InAssetPath, const FString& InUserData)
+		{
+			if (FPlaybackAssetInfo* PlaybackInfos = GetOrAddAssetInfo(InAssetPath))
+			{
+				PlaybackInfos->SetInstanceUserData(InInstanceId, InUserData);
+			}
+		}
+	};
+	
+	/** Remote Server Info */
+	class FServerInfo
+	{
+	public:
+		FMessageAddress Address;
+		FString ServerName;
+		FString ProjectContentPath;
+		uint32 ProcessId = 0;
+		TMap<FString, FString> UserDataEntries;
+		TOptional<FDateTime> PingTimeout;
+		bool bDeviceProviderInstalled = false;
+
+		// Key ChannelName.
+		TMap<FString, TUniquePtr<FBroadcastChannelInfo>> BroadcastChannelInfosByName;
+		TMap<FString, TUniquePtr<FPlaybackChannelInfo>> PlaybackChannelInfosByName;
+
+		// Asset statuses (on disk) are not per-channel.
+		TMap<FSoftObjectPath, EAvaMediaPlaybackAssetStatus> PlaybackAssetStatuses;
+
+		/** Map using "playback asset status request" key to keep track of pending playback asset status requests. */
+		TMap<FSoftObjectPath, FPendingPlaybackAssetStatusRequest> PendingPlaybackAssetStatusRequests;
+
+		/** Accumulate all the playback commands and send them all in one message on the next tick. */
+		TArray<FAvaMediaPlaybackCommand> PendingPlaybackCommands;
+		
+		FBroadcastChannelInfo& GetOrCreateBroadcastChannelInfo(const FString& InChannelName)
+		{
+			const TUniquePtr<FBroadcastChannelInfo>* ChannelInfo = BroadcastChannelInfosByName.Find(InChannelName);
+			return ChannelInfo ? *(*ChannelInfo).Get() : *BroadcastChannelInfosByName.Add(InChannelName, MakeUnique<FBroadcastChannelInfo>()).Get();
+		}
+
+		const FBroadcastChannelInfo* GetBroadcastChannelInfo(const FString& InChannelName) const
+		{
+			const TUniquePtr<FBroadcastChannelInfo>* ChannelInfo = BroadcastChannelInfosByName.Find(InChannelName);
+			return ChannelInfo ? ChannelInfo->Get() : nullptr;
+		}
+
+		FPlaybackChannelInfo* GetOrCreatePlaybackChannelInfo(const FString& InChannelName)
+		{
+			const TUniquePtr<FPlaybackChannelInfo>* ChannelInfo = PlaybackChannelInfosByName.Find(InChannelName);
+			return ChannelInfo ? (*ChannelInfo).Get() : PlaybackChannelInfosByName.Add(InChannelName, MakeUnique<FPlaybackChannelInfo>()).Get();
+		}
+
+		const FPlaybackChannelInfo* GetPlaybackChannelInfo(const FString& InChannelName) const
+		{
+			const TUniquePtr<FPlaybackChannelInfo>* ChannelInfo = PlaybackChannelInfosByName.Find(InChannelName);
+			return ChannelInfo ? (*ChannelInfo).Get() : nullptr;
+		}
+		
+		TOptional<EAvaMediaPlaybackStatus> GetInstanceStatus(
+			const FString& InChannelName, const FGuid& InInstanceId, const FSoftObjectPath& InAssetPath) const
+		{
+			const FPlaybackChannelInfo* ChannelInfo = GetPlaybackChannelInfo(InChannelName);
+			return ChannelInfo ? ChannelInfo->GetInstanceStatus(InInstanceId, InAssetPath) : TOptional<EAvaMediaPlaybackStatus>();
+		}
+		
+		void SetInstanceStatus(
+			const FString& InChannelName, const FGuid& InInstanceId, const FSoftObjectPath& InAssetPath, EAvaMediaPlaybackStatus InStatus)
+		{
+			if (FPlaybackChannelInfo* ChannelInfo = GetOrCreatePlaybackChannelInfo(InChannelName))
+			{
+				ChannelInfo->SetInstanceStatus(InInstanceId, InAssetPath, InStatus);
+			}
+		}
+
+		const FString* GetInstanceUserData(const FString& InChannelName, const FGuid& InInstanceId, const FSoftObjectPath& InAssetPath) const
+		{
+			const FPlaybackChannelInfo* ChannelInfo = GetPlaybackChannelInfo(InChannelName);
+			return ChannelInfo ? ChannelInfo->GetInstanceUserData(InInstanceId, InAssetPath) : nullptr;
+		}
+
+		void SetInstanceUserData(const FString& InChannelName, const FGuid& InInstanceId, const FSoftObjectPath& InAssetPath, const FString& InUserData)
+		{
+			if (FPlaybackChannelInfo* ChannelInfo = GetOrCreatePlaybackChannelInfo(InChannelName))
+			{
+				ChannelInfo->SetInstanceUserData(InInstanceId, InAssetPath, InUserData);
+			}
+		}
+		
+		TOptional<EAvaMediaPlaybackAssetStatus> GetPlaybackAssetStatus(const FSoftObjectPath& InAssetPath) const
+		{
+			const EAvaMediaPlaybackAssetStatus* FoundStatus = PlaybackAssetStatuses.Find(InAssetPath);
+			return FoundStatus ? TOptional(*FoundStatus) : TOptional<EAvaMediaPlaybackAssetStatus>();
+		}
+		
+		void SetPlaybackAssetStatus(const FSoftObjectPath& InAssetPath, EAvaMediaPlaybackAssetStatus InStatus)
+		{
+			PlaybackAssetStatuses.Add(InAssetPath, InStatus);
+		}
+		
+		/**
+		* Add a time by which a ping response should arrive
+		*/
+		void AddTimeout(const FDateTime& InNewTimeout)
+		{
+			if (!PingTimeout || InNewTimeout < PingTimeout.GetValue())
+			{
+				PingTimeout = InNewTimeout;
+			}
+		}
+
+		/**
+		 * Check if ping response arrived on time
+		*/
+		bool HasTimedOut(const FDateTime& InNow) const
+		{
+			return PingTimeout && InNow > PingTimeout.GetValue();
+		}
+
+		/**
+		 * Resets the ping timeout.
+		 * This is done when the ping response is received.
+		 */
+		void ResetPingTimeout()
+		{
+			PingTimeout.Reset();
+		}
+	};
+	TMap<FString, TSharedPtr<FServerInfo>> Servers;	// Key is ServerName
+	TArray<FMessageAddress> AllServerAddresses;
+
+	FServerInfo& GetOrCreateServerInfo(const FString& InServerName, const FMessageAddress& InSenderAddress, bool* bOutCreated = nullptr);
+	FServerInfo* GetServerInfo(const FString& InServerName)
+	{
+		const TSharedPtr<FServerInfo>* ServerInfo = Servers.Find(InServerName);
+		return ServerInfo ? (*ServerInfo).Get() : nullptr;
+	}
+	const FServerInfo* GetServerInfo(const FString& InServerName) const
+	{
+		const TSharedPtr<FServerInfo>* ServerInfo = Servers.Find(InServerName);
+		return ServerInfo ? (*ServerInfo).Get() : nullptr;
+	}
+	
+	void ForAllServers(TFunctionRef<void(FServerInfo& /*InServerInfo*/)> InFunction);
+	void ForAllServers(TFunctionRef<void(const FServerInfo& /*InServerInfo*/)> InFunction) const;
+	void OnServerAdded(const FServerInfo& InServerInfo);
+	void OnServerRemoved(const FServerInfo& InRemovedServer);
+	const FBroadcastChannelInfo* GetChannelInfo(const FString& InServerName, const FString& InChannelName) const;
+	
+	/** Map using "playback status request" key to keep track of pending playback status requests. */
+	TMap<FString, FDateTime> PendingPlaybackStatusRequests;
+	
+	static FString GetPlaybackStatusRequestKey(const FGuid& InInstanceId, const FSoftObjectPath& InAssetPath, const FString& InChannelName)
+	{
+		return InChannelName + InInstanceId.ToString() + InAssetPath.ToString();
+	}
+};
diff --git a/Engine/Plugins/Experimental/Avalanche/Source/AvalancheMedia/Public/Playback/AvaMediaPlaybackClientDelegates.h b/Engine/Plugins/Experimental/Avalanche/Source/AvalancheMedia/Public/Playback/AvaMediaPlaybackClientDelegates.h
new file mode 100644
--- /dev/null
+++ b/Engine/Plugins/Experimental/Avalanche/Source/AvalancheMedia/Public/Playback/AvaMediaPlaybackClientDelegates.h
@@ -0,0 +1,97 @@
+// Copyright Epic Games, Inc. All Rights Reserved.
+
+#pragma once
+
+#include "AvaMediaDefines.h"
+#include "Delegates/Delegate.h"
+
+class IAvaMediaPlaybackClient;
+
+/**
+ * Delegates related to IAvaMediaPlaybackClient.
+ * 
+ * The delegates have been separated in their own global scope to avoid issues with delegates disconnecting
+ * if the playback client is stopped or started during the normal course of operations. Having
+ * the delegates outside of IAvaMediaPlaybackClient's implementation simplify the delegate handler
+ * management.
+ */
+namespace UE::AvaMediaPlaybackClient::Delegates
+{
+	struct FPlaybackStatusChangedArgs
+	{
+		const FGuid& InstanceId;
+		const FSoftObjectPath& AssetPath;
+		const FString& ChannelName;
+		const FString& ServerName;
+		EAvaMediaPlaybackStatus PrevStatus;
+		EAvaMediaPlaybackStatus NewStatus;
+	};
+	
+	struct FPlaybackAssetStatusChangedArgs
+	{
+		const FSoftObjectPath& AssetPath;
+		const FString& ServerName;
+		EAvaMediaPlaybackAssetStatus Status;
+	};
+
+	struct FPlaybackSequenceEventArgs
+	{
+		const FGuid& InstanceId;
+		const FSoftObjectPath& AssetPath;
+		const FString& ChannelName;
+		const FString& ServerName;
+		const FString& SequenceName;
+		EAvalanchePlayableSequenceEventType EventType;
+	};
+
+	struct FPlaybackTransitionEventArgs
+	{
+		const FGuid& TransitionId;
+		const FGuid& InstanceId;
+		const FString& ChannelName;
+		const FString& ServerName;
+		EAvalanchePlayableTransitionEventFlags EventFlags;
+	};
+
+	enum class EConnectionEvent : uint8
+	{
+		ServerConnected = 0,
+		ServerDisconnected
+	};
+
+	struct FConnectionEventArgs
+	{
+		const FString& ServerName;
+		EConnectionEvent Event;
+	};
+
+	/**
+	 * Delegate called when a server connection event occurs. 
+	 */
+	DECLARE_MULTICAST_DELEGATE_TwoParams(FOnConnectionEvent, IAvaMediaPlaybackClient&, const FConnectionEventArgs&);
+	AVALANCHEMEDIA_API FOnConnectionEvent& GetOnConnectionEvent();
+	
+	/**
+	* Delegate called when a playback status changed message is received.
+	*/
+	DECLARE_MULTICAST_DELEGATE_TwoParams(FOnPlaybackStatusChanged, IAvaMediaPlaybackClient&, const FPlaybackStatusChangedArgs&);
+	AVALANCHEMEDIA_API FOnPlaybackStatusChanged& GetOnPlaybackStatusChanged();
+
+	/**
+	* Delegate called when a playback asset status changed message is received.
+	*/
+	DECLARE_MULTICAST_DELEGATE_TwoParams(FOnPlaybackAssetStatusChanged, IAvaMediaPlaybackClient&, const FPlaybackAssetStatusChangedArgs&);
+	AVALANCHEMEDIA_API FOnPlaybackAssetStatusChanged& GetOnPlaybackAssetStatusChanged();
+
+	/**
+	* Delegate called when a playback status changed message is received.
+	*/
+	DECLARE_MULTICAST_DELEGATE_TwoParams(FOnPlaybackSequenceEvent, IAvaMediaPlaybackClient&, const FPlaybackSequenceEventArgs&);
+	AVALANCHEMEDIA_API FOnPlaybackSequenceEvent& GetOnPlaybackSequenceEvent();
+
+	/**
+	* Delegate called when a playback status changed message is received.
+	*/
+	DECLARE_MULTICAST_DELEGATE_TwoParams(FOnPlaybackTransitionEvent, IAvaMediaPlaybackClient&, const FPlaybackTransitionEventArgs&);
+	AVALANCHEMEDIA_API FOnPlaybackTransitionEvent& GetOnPlaybackTransitionEvent();
+};
diff --git a/Engine/Plugins/Experimental/Avalanche/Source/AvalancheMedia/Public/Playback/AvaMediaPlaybackManager.h b/Engine/Plugins/Experimental/Avalanche/Source/AvalancheMedia/Public/Playback/AvaMediaPlaybackManager.h
new file mode 100644
--- /dev/null
+++ b/Engine/Plugins/Experimental/Avalanche/Source/AvalancheMedia/Public/Playback/AvaMediaPlaybackManager.h
@@ -0,0 +1,382 @@
+// Copyright Epic Games, Inc. All Rights Reserved.
+
+#pragma once
+
+#include "UObject/StrongObjectPtr.h"
+#include "Playback/AvalanchePlayback.h"
+#include "AvaMediaPlaybackManager.generated.h"
+
+class FAvaMediaPlaybackManager;
+class FAvaMediaPlaybackSourceAssetEntry;
+class IAvaMediaSyncProvider;
+class UAvaMediaPlayableGroupManager;
+class UAvaMediaPlaybackTransition;
+class UAvalancheBlueprint;
+struct FAnimPlaySettings;	// private
+
+
+/** Handle to a playback instance for recycling. */
+class AVALANCHEMEDIA_API FAvaMediaPlaybackInstance : public FGCObject
+{
+public:
+	FAvaMediaPlaybackInstance() = default;
+	FAvaMediaPlaybackInstance(const FGuid& InInstanceId, const FSoftObjectPath& InSourcePath, const FString& InChannelName, UAvalanchePlayback* InPlayback);
+	virtual ~FAvaMediaPlaybackInstance() override;
+
+	//~ Begin FGCObject
+	virtual void AddReferencedObjects( FReferenceCollector& Collector ) override;
+	virtual FString GetReferencerName() const override;
+	//~ End FGCObject
+
+	/**
+	 * This is used by the server to track the client's id.
+	 * It can also be used when reconciling the client state from the server.
+	 */
+	void SetInstanceId(const FGuid& InInstanceId);
+
+	/**
+	 * Set the instance's user data. This is propagate to the server and can be
+	 * used when reconciling the client state.
+	 */
+	void SetInstanceUserData(const FString& InUserData);
+	
+	const FGuid& GetInstanceId() const { return InstanceId; }
+	const FString& GetInstanceUserData() const { return InstanceUserData;}
+	const FString& GetChannelName() const { return ChannelName; }
+	const FName& GetChannelFName() const { return ChannelFName; }
+	const FSoftObjectPath& GetSourcePath() const { return SourcePath; }
+	UAvalanchePlayback* GetPlayback() const { return Playback; }
+	EAvaMediaPlaybackStatus GetStatus() const { return Status; }
+	bool IsPlaying() const { return Playback ? Playback->IsPlaying() : false;}
+
+	bool UpdateStatus();
+	void SetStatus(EAvaMediaPlaybackStatus InStatus) { Status = InStatus; }
+
+	void Unload();
+	void Recycle();
+
+	TSharedPtr<FAvaMediaPlaybackManager> GetManager() const;
+
+private:
+	void OnPlayableCreated(UAvalanchePlayback* InPlayback, UAvalanchePlayable* InPlayable);
+	
+private:
+	/** If the cache slot becomes invalid, it means this instance will be discarded instead of being recycled. */
+	TWeakPtr<FAvaMediaPlaybackSourceAssetEntry> AssetEntryWeak;
+
+	FGuid InstanceId;
+	FString ChannelName;
+	FName ChannelFName;
+	FSoftObjectPath SourcePath;
+	TObjectPtr<UAvalanchePlayback> Playback;
+	EAvaMediaPlaybackStatus Status = EAvaMediaPlaybackStatus::Unknown;
+	FString InstanceUserData;
+
+	friend class FAvaMediaPlaybackManager;
+};
+
+/**
+ * For each playback instance, a source asset entry is kept track of to allow the entry to be
+ * invalidated and prevent the instance to be recycled.
+ */
+class AVALANCHEMEDIA_API FAvaMediaPlaybackSourceAssetEntry
+{
+public:
+	FAvaMediaPlaybackSourceAssetEntry(const TSharedPtr<FAvaMediaPlaybackManager>& InParentManager) : ParentManagerWeak(InParentManager.ToWeakPtr()) {}
+	~FAvaMediaPlaybackSourceAssetEntry() = default;
+
+	TSharedPtr<FAvaMediaPlaybackInstance> AcquirePlaybackInstance(const FString& InChannelName)
+	{
+		TSharedPtr<FAvaMediaPlaybackInstance> AcquiredInstance;
+		const int32 InstanceIndex = AvailableInstances.IndexOfByPredicate([InChannelName](const TSharedPtr<FAvaMediaPlaybackInstance>& InInstance)
+		{
+			return InInstance && InInstance->GetChannelName() == InChannelName;
+		});
+
+		if (InstanceIndex != INDEX_NONE)
+		{
+			AcquiredInstance = AvailableInstances[InstanceIndex];
+			AvailableInstances.RemoveAtSwap(InstanceIndex, 1);
+			UsedInstances.Add(AcquiredInstance.ToWeakPtr());
+		}
+		return AcquiredInstance;
+	}
+
+	TSharedPtr<FAvaMediaPlaybackInstance> FindPlaybackInstance(const FGuid& InInstanceId,
+		bool bInAvailableInstances = true, bool bInUsedInstances = true) const;
+	
+	TSharedPtr<FAvaMediaPlaybackInstance> FindPlaybackInstanceForChannel(const FString& InChannelName,
+		bool bInAvailableInstances = true, bool bInUsedInstances = true) const;
+
+	TSharedPtr<FAvaMediaPlaybackInstance> FindPlaybackInstanceByPredicate(TFunctionRef<bool(const FAvaMediaPlaybackInstance& /*InPlaybackInstance*/)> InPredicate,
+		bool bInAvailableInstances = true, bool bInUsedInstances = true) const;
+	
+	void ForAllPlaybackInstances(TFunctionRef<void(FAvaMediaPlaybackInstance& /*InPlaybackInstance*/)> InFunction,
+		bool bInAvailableInstances = true, bool bInUsedInstances = true);
+
+	TSharedPtr<FAvaMediaPlaybackManager> GetManager() const { return ParentManagerWeak.Pin(); }
+
+	void DiscardInstance(FAvaMediaPlaybackInstance* InInstanceToRemove);
+	void RecycleInstance(FAvaMediaPlaybackInstance* InInstanceToRecycle);
+
+private:
+	/** Keeping a weak reference to the manager to recycle the instance. */
+	TWeakPtr<FAvaMediaPlaybackManager> ParentManagerWeak;
+	
+	TArray<TSharedPtr<FAvaMediaPlaybackInstance>> AvailableInstances;
+	TArray<TWeakPtr<FAvaMediaPlaybackInstance>> UsedInstances;
+
+	friend class FAvaMediaPlaybackManager;
+};
+
+/** Flags indicating what changed in the package event. */
+UENUM()
+enum class EAvaMediaPlaybackPackageEventFlags : uint8
+{
+	None			= 0,
+	External		= 1 << 0,
+	Saved			= 1 << 1,
+	AssetDeleted	= 1 << 2,
+	All				= 0xFF
+};
+ENUM_CLASS_FLAGS(EAvaMediaPlaybackPackageEventFlags);
+
+class AVALANCHEMEDIA_API FAvaMediaPlaybackManager : public TSharedFromThis<FAvaMediaPlaybackManager>
+{
+public:
+	FAvaMediaPlaybackManager();
+	virtual ~FAvaMediaPlaybackManager();
+
+	void Tick();
+
+	void SetEnablePlaybackCommandsBuffering(bool bInEnable) { bEnablePlaybackCommandsBuffering = bInEnable;}
+
+	UAvaMediaPlayableGroupManager* GetPlayableGroupManager() const { return PlayableGroupManager.Get(); }
+
+	/** Acquire a cached playback instance. Will return null if none available in the cache. */
+	TSharedPtr<FAvaMediaPlaybackInstance> AcquirePlaybackInstance(const FSoftObjectPath& InAssetPath, const FString& InChannelName) const;
+	
+	/** Load a new playback instance. */
+	TSharedPtr<FAvaMediaPlaybackInstance> LoadPlaybackInstance(const FSoftObjectPath& InAssetPath, const FString& InChannelName);
+
+	/** Attempts to acquire existing instance. If none available, will load a new one. */  
+	TSharedPtr<FAvaMediaPlaybackInstance> AcquireOrLoadPlaybackInstance(const FSoftObjectPath& InAssetPath, const FString& InChannelName);
+
+	/**
+	 * Finds an existing (either available or used) playback instance.
+	 * This will not acquire it (i.e. if available, it will remain so).
+	 */
+	TSharedPtr<FAvaMediaPlaybackInstance> FindPlaybackInstance(const FGuid& InInstanceId, const FSoftObjectPath& InAssetPath, const FString& InChannelName) const;
+
+	/**
+	 * Unload and discard any available (i.e. not used) instances for this asset/channel entry.
+	 * If specified channel name is empty, it will discard all instances of the asset.
+	 */
+	bool UnloadPlaybackInstances(const FSoftObjectPath& InAssetPath, const FString& InChannelName);
+	
+	void ForAllPlaybackInstances(TFunctionRef<void(FAvaMediaPlaybackInstance& /*InPlaybackInstance*/)> InFunction);
+	
+	/**
+	 *	Determines the local status of the given asset on this local instance of the playback manager.
+	 */
+	EAvaMediaPlaybackAssetStatus GetLocalAssetStatus(const FName& InPackageName);
+	
+	/**
+	 * Invalidates the cached local asset status.
+	 */
+	void InvalidateCachedLocalAssetStatus(const FName& InPackageName);
+
+	/**
+	 * Utility function to determine if an asset is locally available.
+	 * Remark: Determining the presence of dependencies is somewhat unreliable i.e.
+	 * an asset can playback fine even with some dependencies missing and it is hard to figure it out.
+	 * Because of that, for now, we consider the asset available even if it is missing some dependencies.
+	 */
+	bool IsLocalAssetAvailable(const FName& InPackageName)
+	{
+		const EAvaMediaPlaybackAssetStatus LocalAssetStatus = GetLocalAssetStatus(InPackageName);
+		return (LocalAssetStatus == EAvaMediaPlaybackAssetStatus::Available || LocalAssetStatus == EAvaMediaPlaybackAssetStatus::MissingDependencies);
+	}
+
+	bool IsLocalAssetAvailable(const FSoftObjectPath& InAssetPath)
+	{
+		return IsLocalAssetAvailable(InAssetPath.GetLongPackageFName());
+	}
+
+	/** Utility function to determine the playback status of an unloaded asset. */
+	EAvaMediaPlaybackStatus GetUnloadedPlaybackStatus(const FSoftObjectPath& InAssetPath)
+	{
+		// When the playback entry is unloaded, we rely on the local asset status to determine the playback status.
+		// Note: since there is now an independent asset status, we could remove all the playback states related to the asset.
+		return IsLocalAssetAvailable(InAssetPath) ? EAvaMediaPlaybackStatus::Available : EAvaMediaPlaybackStatus::Missing;
+	}
+
+	/**
+	 *	Invalidate the asset entry. All cached instances will be invalidated along with it.
+	 */
+	void InvalidatePlaybackAssetEntry(const FSoftObjectPath& InAssetPath);
+	
+	UAvalanchePlayback* LoadPlaybackObject(const FSoftObjectPath& InAssetPath, const FString& InChannelName) const;
+	UAvalanchePlayback* BuildPlaybackFromBlueprint(const UAvalancheBlueprint* InBlueprint, const FString& InChannelName) const;
+	UAvalanchePlayback* BuildPlaybackFromWorld(const TSoftObjectPtr<UWorld>& InWorld, const FString& InChannelName) const;
+	
+	/**
+	 * Stops all currently playing playback objects.
+	 * @param bInUnload if true, will also unload the objects.
+	 * @return Returns the list of all source blueprint assets that where stopped.
+	 **/
+	TArray<FSoftObjectPath> StopAllPlaybacks(bool bInUnload);
+
+	bool PushAnimationCommand(const FGuid& InInstanceId, const FSoftObjectPath& InSourcePath, const FString& InChannelName, EAvaMediaAnimAction InAction, const FAnimPlaySettings& InAnimSettings);
+	bool PushRemoteControlCommand(const FGuid& InInstanceId, const FSoftObjectPath& InSourcePath, const FString& InChannelName, const TSharedRef<FAvalancheRemoteControlValues>& InRemoteControlValues);
+	bool PushPlaybackTransitionStartCommand(UAvaMediaPlaybackTransition* InTransitionToStart);
+
+	/** This is used on the playback server to apply any pending commands to a playback instance. */
+	void ApplyPendingCommands(UAvalanchePlayback* InPlaybackObject, const FGuid& InInstanceId, const FSoftObjectPath& InSourcePath, const FString& InChannelName);
+
+	/**
+	 *	Indicate the manager is in a shutdown sequence and will force game instances to destroy worlds right away.
+	 */
+	void StartShuttingDown() { bIsShuttingDown = true;}
+
+	EAvaPlaybackStopOptions GetPlaybackStopOptions(bool bInUnload) const
+	{
+		EAvaPlaybackStopOptions Options = bIsShuttingDown ? EAvaPlaybackStopOptions::ForceImmediate : EAvaPlaybackStopOptions::None;
+		Options |= (bInUnload || bIsShuttingDown) ? EAvaPlaybackStopOptions::Unload : EAvaPlaybackStopOptions::None;
+		return Options;
+	}
+	
+	EAvaPlaybackUnloadOptions GetPlaybackUnloadOptions() const
+	{
+		return bIsShuttingDown ? EAvaPlaybackUnloadOptions::ForceImmediate : EAvaPlaybackUnloadOptions::None;
+	}
+
+	/** Let the Playback manager know that a package has been modified. */
+	void OnPackageModified(const FName& InPackageName, EAvaMediaPlaybackPackageEventFlags InFlags = EAvaMediaPlaybackPackageEventFlags::None);
+
+	/**
+	 *	TearDown the whole Avalanche Playback system.
+	 *	
+	 *	When Avalanche Playback is used within a game, we need to tear down everything
+	 *	as the parent world is being teared down. The game tear down process will forcibly
+	 *	mark as garbage (and GCs) all the GameInstances, including those held by the playback objects,
+	 *	despite playback objects holding strong references to them. To avoid issues we preemptively
+	 *	destroy all the playback objects.
+	 */
+	void OnParentWorldBeginTearDown();
+
+	/**
+	 * Implements a similar command to Engine::HandleStatCommand, except it
+	 * will fetch avalanche's game viewport client if everything else fails.
+	 */
+	bool HandleStatCommand(const TArray<FString>& InArgs);
+
+	/** Returns true if the given asset is a playback asset, i.e. either a "playable" asset or a playback graph. */
+	static bool IsPlaybackAsset(const FAssetData& InAssetData);
+	
+	DECLARE_MULTICAST_DELEGATE_OneParam(FOnPlaybackInstanceInvalidated, const FAvaMediaPlaybackInstance&);
+	FOnPlaybackInstanceInvalidated OnPlaybackInstanceInvalidated;
+
+	DECLARE_MULTICAST_DELEGATE_OneParam(FOnPlaybackInstanceStatusChanged, const FAvaMediaPlaybackInstance&);
+	FOnPlaybackInstanceStatusChanged OnPlaybackInstanceStatusChanged;
+
+	DECLARE_MULTICAST_DELEGATE_OneParam(FOnLocalPlaybackAssetRemoved, const FSoftObjectPath&);
+	FOnLocalPlaybackAssetRemoved OnLocalPlaybackAssetRemoved;
+
+	DECLARE_MULTICAST_DELEGATE_OneParam(FOnBeginTick, float);
+	FOnBeginTick OnBeginTick;
+
+protected:
+	// Non-copyable (because copy of CachedAssetStatusExtra is not defined)
+	FAvaMediaPlaybackManager(const FAvaMediaPlaybackManager&) = delete;
+	FAvaMediaPlaybackManager& operator=(const FAvaMediaPlaybackManager&) = delete;
+
+	void OnPackageSaved(const FString& InPackageFileName, UPackage* InPackage, FObjectPostSaveContext InObjectSaveContext);
+	void OnAvaSyncPackageModified(IAvaMediaSyncProvider* InAvaMediaSyncProvider, const FName& InPackageName);
+	void OnAssetRemoved(const FAssetData& InAssetData);
+
+	TSharedPtr<FAvaMediaPlaybackSourceAssetEntry> FindPlaybackAssetEntry(const FSoftObjectPath& InAssetPath) const
+	{
+		const TSharedPtr<FAvaMediaPlaybackSourceAssetEntry>* Existing = PlaybackAssetEntries.Find(InAssetPath);
+		return Existing ? *Existing : TSharedPtr<FAvaMediaPlaybackSourceAssetEntry>();
+	}
+	
+	TSharedPtr<FAvaMediaPlaybackSourceAssetEntry> GetPlaybackAssetEntry(const FSoftObjectPath& InAssetPath)
+	{
+		TSharedPtr<FAvaMediaPlaybackSourceAssetEntry>* Existing = PlaybackAssetEntries.Find(InAssetPath);
+		if (Existing)
+		{
+			return *Existing;
+		}
+
+		TSharedPtr<FAvaMediaPlaybackSourceAssetEntry> NewAssetEntry = MakeShared<FAvaMediaPlaybackSourceAssetEntry>(SharedThis(this));
+		
+		PlaybackAssetEntries.Add(InAssetPath, NewAssetEntry);
+		return NewAssetEntry;
+	}
+	
+private:
+	bool bIsShuttingDown = false;
+
+	/** This is the shared pool of shared playable groups for all the playback objects. */
+	TStrongObjectPtr<UAvaMediaPlayableGroupManager> PlayableGroupManager; 
+	
+	// TODO: Refactor this to support cache capacity with eviction (LRU) like FAvalancheManagedInstanceCache.
+	TMap<FSoftObjectPath, TSharedPtr<FAvaMediaPlaybackSourceAssetEntry>> PlaybackAssetEntries;
+	
+	/** Cached asset status. */
+	TMap<FName, EAvaMediaPlaybackAssetStatus> CachedAssetStatus;
+
+	TArray<TWeakObjectPtr<UAvaMediaPlaybackTransition>> PendingStartTransitions;
+	
+	/**
+	 * Enables the playback command buffering.
+	 * 
+	 * This is enabled on the playback server to handle the Remote control and animation
+	 * commands being received/processed before the playback command itself (where the playback object
+	 * is created). If the commands can't be executed because the object is not yet created the
+	 * playback manager will buffer the commands and apply them to the object once it is loaded.
+	 */
+	bool bEnablePlaybackCommandsBuffering = false;
+	
+	struct FPlaybackObjectCommandBuffers
+	{
+		struct FAnimationCommand
+		{
+			EAvaMediaAnimAction AnimAction;
+			FAnimPlaySettings AnimPlaySettings;
+		};
+		struct FRemoteControlCommand
+		{
+			TSharedRef<FAvalancheRemoteControlValues> Values;
+		};
+	
+		TArray<FAnimationCommand> AnimationCommands;
+		TArray<FRemoteControlCommand> RemoteControlCommands;
+	};
+	// Use FAvaMediaPlaybackManager::MakePlaybackKey for Map key.
+	TMap<FString, FPlaybackObjectCommandBuffers> PlaybackObjectCommandBuffers;
+
+	FString MakeCommandBufferKey(const FSoftObjectPath& InAssetPath, const FString& InChannelName) const
+	{
+		FString CommandBufferKey = InAssetPath.ToString();
+		if (!InChannelName.IsEmpty())
+		{
+			CommandBufferKey += TEXT("_") + InChannelName;
+		}
+		return CommandBufferKey;
+	}
+
+	FString MakeCommandBufferKey(const FGuid& InInstanceId) const
+	{
+		return InInstanceId.ToString();
+	}
+
+	FString MakeCommandBufferKey(const FGuid& InInstanceId, const FSoftObjectPath& InAssetPath, const FString& InChannelName) const
+	{
+		return InInstanceId.IsValid() ? MakeCommandBufferKey(InInstanceId) : MakeCommandBufferKey(InAssetPath, InChannelName);
+	}
+	
+	FPlaybackObjectCommandBuffers& GetOrCreatePlaybackCommandBuffers(FString&& InCommandBufferKey);
+	const FPlaybackObjectCommandBuffers* GetPlaybackCommandBuffers(const FString& InCommandBufferKey) const;
+};
diff --git a/Engine/Plugins/Experimental/Avalanche/Source/AvalancheMedia/Public/Playback/AvaMediaPlaybackMessages.h b/Engine/Plugins/Experimental/Avalanche/Source/AvalancheMedia/Public/Playback/AvaMediaPlaybackMessages.h
new file mode 100644
--- /dev/null
+++ b/Engine/Plugins/Experimental/Avalanche/Source/AvalancheMedia/Public/Playback/AvaMediaPlaybackMessages.h
@@ -0,0 +1,738 @@
+// Copyright Epic Games, Inc. All Rights Reserved.
+
+#pragma once
+
+#include "AvaMediaDefines.h"
+#include "Channel/AvaMediaOutputInfo.h"
+#include "Framework/AvalancheInstanceSettings.h"
+#include "PixelFormat.h"
+#include "Playback/AvalancheRemoteControlValues.h"
+#include "Playback/Nodes/Events/Actions/AvalancheAnimations.h"
+#include "Viewport/AvaViewportQualitySettings.h"
+
+
+#include "AvaMediaPlaybackMessages.generated.h"
+
+USTRUCT()
+struct AVALANCHEMEDIA_API FAvaMediaPlaybackClientMessageBase
+{
+	GENERATED_BODY()
+
+	UPROPERTY()
+	FString ClientName;
+};
+
+/**
+ * Notes:
+ * - ServerName vs HostName
+ * The Host name is the physical computer device the playback server runs on.
+ * A Server name is what the playback server is called and it is not necessarily the Host name.
+ * If there are multiple servers on the same host, then each server must have a different name.
+ * The Server name is a unique key to identify a server, while the host name is not.
+ */
+USTRUCT()
+struct AVALANCHEMEDIA_API FAvaMediaPlaybackServerMessageBase
+{
+	GENERATED_BODY()
+
+	UPROPERTY()
+	FString ServerName;
+};
+
+/**
+ * Request published by client to discover servers.
+ */
+USTRUCT()
+struct AVALANCHEMEDIA_API FAvaMediaPlaybackPing : public FAvaMediaPlaybackClientMessageBase
+{
+	GENERATED_BODY()
+
+	/** Indicate if the ping was generated on a timer (auto) or manually sent through a user command. */
+	UPROPERTY()
+	bool bAutoPing = true;
+
+	/** Defines the interval in seconds between each client's pings. */
+	UPROPERTY()
+	float PingIntervalSeconds = 0.0f;
+};
+
+/**
+ *	Response sent by server to client to be discovered.
+ */
+USTRUCT()
+struct AVALANCHEMEDIA_API FAvaMediaPlaybackPong : public FAvaMediaPlaybackServerMessageBase
+{
+	GENERATED_BODY()
+
+	/** Indicate if this is a reply from an auto ping. */
+	UPROPERTY()
+	bool bAutoPong = true;
+
+	/** Server may request client info in case of a reconnection event. */
+	UPROPERTY()
+	bool bRequestClientInfo = false;
+
+	/** Server's project content path. */
+	UPROPERTY()
+	FString ProjectContentPath;
+
+	/** Server's process id on the current host. */
+	UPROPERTY()
+	uint32 ProcessId = 0;
+};
+
+/**
+ *	Replicate server's log messages.
+ */
+USTRUCT()
+struct AVALANCHEMEDIA_API FAvaMediaPlaybackLog : public FAvaMediaPlaybackServerMessageBase
+{
+	GENERATED_BODY()
+
+	UPROPERTY()
+	FString Text;
+
+	UPROPERTY()
+	uint8 Verbosity = 0;
+
+	UPROPERTY()
+	FName Category;
+
+	UPROPERTY()
+	double Time = 0.0;
+};
+
+/**
+ * Request sent by client to replicate it's information on the destination server.
+ */
+USTRUCT()
+struct AVALANCHEMEDIA_API FAvaMediaUpdateClientInfo : public FAvaMediaPlaybackClientMessageBase
+{
+	GENERATED_BODY()
+
+	UPROPERTY()
+	FString ComputerName;
+
+	UPROPERTY()
+	FString ProjectContentPath;
+
+	UPROPERTY()
+	uint32 ProcessId = 0;
+};
+
+/**
+ * Request sent by client to replicate it's user data on the destination server.
+ */
+USTRUCT()
+struct AVALANCHEMEDIA_API FAvaMediaUpdateClientUserData : public FAvaMediaPlaybackClientMessageBase
+{
+	GENERATED_BODY()
+
+	UPROPERTY()
+	TMap<FString, FString> UserDataEntries;
+};
+
+/**
+ * Request sent by server to replicate it's user data to the client.
+ */
+USTRUCT()
+struct AVALANCHEMEDIA_API FAvaMediaUpdateServerUserData : public FAvaMediaPlaybackServerMessageBase
+{
+	GENERATED_BODY()
+
+	UPROPERTY()
+	TMap<FString, FString> UserDataEntries;
+};
+
+USTRUCT()
+struct AVALANCHEMEDIA_API FAvaMediaStatCommand : public FAvaMediaPlaybackClientMessageBase
+{
+	GENERATED_BODY()
+
+	UPROPERTY()
+	FString Command;
+
+	/**
+	 * If the local client couldn't get a viewport client to run the stat command, the
+	 * local state will not be reliable. In this case, the server state will be used.
+	 */
+	UPROPERTY()
+	bool bClientStateReliable = false;
+	
+	/**
+	 * Because the commands are toggles, there may be desync of states between client and server(s).
+	 * In order to correct that we also send a list of enabled states on the client so that servers can
+	 * ensure they have the same states enabled.
+	 */
+	UPROPERTY()
+	TArray<FString> ClientEnabledRuntimeStats;
+};
+
+/**
+ * Server's response when a stat command is received.
+ */
+USTRUCT()
+struct AVALANCHEMEDIA_API FAvaMediaStatStatus : public FAvaMediaPlaybackServerMessageBase
+{
+	GENERATED_BODY()
+
+	/**
+	 * Replication of the same field from the command.
+	 * This is used by the client to replicate the server state in case
+	 * the command succeeded on the server and it had an unreliable state.
+	 */
+	UPROPERTY()
+	bool bClientStateReliable = false;
+	
+	/** Indicate if the server could execute the command. */
+	UPROPERTY()
+	bool bCommandSucceeded = false;
+
+	/** Resulting enabled stats from the command. */
+	UPROPERTY()
+	TArray<FString> EnabledRuntimeStats;
+};
+
+/**
+ * Request sent by client to obtain the device provider data from destination server.
+ * The device provider data contains all available devices and their configuration
+ * installed on the server.
+ */
+USTRUCT()
+struct AVALANCHEMEDIA_API FAvaDeviceProviderDataRequest : public FAvaMediaPlaybackClientMessageBase
+{
+	GENERATED_BODY()
+};
+
+/**
+ *	Request for the client to replicate it's Avalanche instance settings to the server.
+ **/
+USTRUCT()
+struct AVALANCHEMEDIA_API FAvalancheInstanceSettingsUpdate : public FAvaMediaPlaybackClientMessageBase
+{
+	GENERATED_BODY()
+
+	UPROPERTY()
+	FAvalancheInstanceSettings InstanceSettings;
+};
+
+UENUM()
+enum class EAvaMediaPlaybackPackageEvent
+{
+	None,
+	PreSave,
+	PostSave,
+	AssetDeleted
+};
+
+/**
+ *	Message sent by the client to inform servers that a local package has been modified by an event.
+ **/
+USTRUCT()
+struct AVALANCHEMEDIA_API FAvaMediaPlaybackPackageEvent : public FAvaMediaPlaybackClientMessageBase
+{
+	GENERATED_BODY()
+
+	UPROPERTY()
+	FName PackageName;
+
+	UPROPERTY()
+	EAvaMediaPlaybackPackageEvent Event = EAvaMediaPlaybackPackageEvent::None;
+};
+
+/**
+ * Request by a client to obtain the status of an asset with the given path.
+ * Note: since this is an asset on disk, there is no channel name.
+ */
+USTRUCT()
+struct AVALANCHEMEDIA_API FAvaMediaPlaybackAssetStatusRequest : public FAvaMediaPlaybackClientMessageBase
+{
+	GENERATED_BODY()
+	
+	UPROPERTY()
+	FSoftObjectPath AssetPath;
+	
+	/**
+	 * If this is true, the server will perform a full asset status compare with the client.
+	 * Otherwise, the server may use locally cached status for the given asset.
+	 */
+	UPROPERTY()
+	bool bForceRefresh = false;
+};
+
+/** Response from the server to a playback asset status request. */
+USTRUCT()
+struct AVALANCHEMEDIA_API FAvaMediaPlaybackAssetStatus : public FAvaMediaPlaybackServerMessageBase
+{
+	GENERATED_BODY()
+	
+	UPROPERTY()
+	FSoftObjectPath AssetPath;
+	
+	UPROPERTY()
+	EAvaMediaPlaybackAssetStatus Status = EAvaMediaPlaybackAssetStatus::Unknown;
+};
+
+UENUM()
+enum class EAvaMediaPlaybackAction
+{
+	/** No op */
+	None,
+	/** Load the given asset. Used for pre-loading assets. */
+	Load,
+	/** Start (i.e. start ticking world and rendering) the given asset, loading it if not pre-loaded. */
+	Start,
+	/** Stop ticking and rendering the world. */
+	Stop,
+	/** Unload the given asset, i.e. destroy the world, etc. */
+	Unload,
+	/** Request the status of the asset. No action is actually performed on the asset. */
+	Status,
+	/** Request to set the user data of the playback instance.*/
+	SetUserData,
+	/** Request the user data of the playback instance. No action is actually performed on the asset. */
+	GetUserData,
+};
+
+/**
+ * Playback command for a given asset.
+ */
+USTRUCT()
+struct AVALANCHEMEDIA_API FAvaMediaPlaybackCommand
+{
+	GENERATED_BODY()
+
+	/** Playback's instance Id. */
+	UPROPERTY()
+	FGuid InstanceId;
+
+	/** Reference to the asset to playback. */
+	UPROPERTY()
+	FSoftObjectPath AssetPath;
+
+	/** Channel to Play, only considered when the Asset in question is an Avalanche Playable (Blueprint or Level). */
+	UPROPERTY()
+	FString ChannelName;
+	
+	/** Action to do: Load, Play, Stop, Unload, etc. */
+	UPROPERTY()
+	EAvaMediaPlaybackAction Action = EAvaMediaPlaybackAction::None;
+
+	/** Command additional arguments. */
+	UPROPERTY()
+	FString Arguments;
+};
+
+/**
+ * Request by a client to execute a batch playback of commands on the connected servers.
+ */
+USTRUCT()
+struct AVALANCHEMEDIA_API FAvaMediaPlaybackRequest : public FAvaMediaPlaybackClientMessageBase
+{
+	GENERATED_BODY()
+	
+	UPROPERTY()
+	TArray<FAvaMediaPlaybackCommand> Commands;
+};
+
+/** Response from the server to a playback request. */
+USTRUCT()
+struct AVALANCHEMEDIA_API FAvaMediaPlaybackStatus : public FAvaMediaPlaybackServerMessageBase
+{
+	GENERATED_BODY()
+
+	/** Playback's instance Id. */
+	UPROPERTY()
+	FGuid InstanceId;
+	
+	UPROPERTY()
+	FSoftObjectPath AssetPath;
+
+	UPROPERTY()
+	FString ChannelName;
+	
+	UPROPERTY()
+	EAvaMediaPlaybackStatus Status = EAvaMediaPlaybackStatus::Unknown;
+	
+	UPROPERTY()
+	bool bValidUserData = false;
+	
+	UPROPERTY()
+	FString UserData;
+};
+
+/**
+ * Message batching the status of many assets per channel/host.
+ * Used to reduce message overhead.
+ */
+USTRUCT()
+struct AVALANCHEMEDIA_API FAvaMediaPlaybackStatuses : public FAvaMediaPlaybackServerMessageBase
+{
+	GENERATED_BODY()
+	
+	UPROPERTY()
+	FString ChannelName;
+
+	/** Playback's instance Id. */
+	UPROPERTY()
+	TArray<FGuid> InstanceIds;
+	
+	UPROPERTY()
+	TArray<FSoftObjectPath> AssetPaths;
+
+	/** Batches all assets with the same status. */
+	UPROPERTY()
+	EAvaMediaPlaybackStatus Status = EAvaMediaPlaybackStatus::Unknown;
+};
+
+USTRUCT()
+struct FAvaMediaAnimActionInfo
+{
+	GENERATED_BODY()
+    
+	UPROPERTY()
+	FString AnimationName;
+
+	UPROPERTY()
+	EAvaMediaAnimAction AnimationAction = EAvaMediaAnimAction::None;
+};
+
+/**
+ * Request by a client to execute an animation on a playback asset
+ * on a connected server.
+ */
+USTRUCT()
+struct FAvaMediaAnimPlaybackRequest : public FAvaMediaPlaybackClientMessageBase
+{
+	GENERATED_BODY()
+
+	/** Playback's instance Id. */
+	UPROPERTY()
+	FGuid InstanceId;
+	
+	UPROPERTY()
+	FSoftObjectPath AssetPath;
+	
+	/** Channel to Play, only considered when the Asset in question is an Avalanche Blueprint */
+	UPROPERTY()
+	FString ChannelName;
+	
+	UPROPERTY()
+	TArray<FAnimPlaySettings> AnimPlaySettings;
+
+	UPROPERTY()
+	TArray<FAvaMediaAnimActionInfo> AnimActionInfos;
+};
+
+/** Server replication of playback sequence events. */
+USTRUCT()
+struct AVALANCHEMEDIA_API FAvaMediaPlayableSequenceEvent : public FAvaMediaPlaybackServerMessageBase
+{
+	GENERATED_BODY()
+
+	/** Playback's instance Id. */
+	UPROPERTY()
+	FGuid InstanceId;
+	
+	UPROPERTY()
+	FSoftObjectPath AssetPath;
+
+	UPROPERTY()
+	FString ChannelName;
+	
+	UPROPERTY()
+	FString SequenceName;
+
+	UPROPERTY()
+	EAvalanchePlayableSequenceEventType EventType = EAvalanchePlayableSequenceEventType::None;
+};
+
+USTRUCT()
+struct FAvaMediaRemoteControlUpdateRequest : public FAvaMediaPlaybackClientMessageBase
+{
+	GENERATED_BODY()
+
+	/** Playback's instance Id. */
+	UPROPERTY()
+	FGuid InstanceId;
+	
+	UPROPERTY()
+	FSoftObjectPath AssetPath;
+	
+	/** Channel to Play, only considered when the Asset in question is an Avalanche Blueprint. */
+	UPROPERTY()
+	FString ChannelName;
+	
+	UPROPERTY()
+	FAvalancheRemoteControlValues RemoteControlValues;
+};
+
+USTRUCT()
+struct FAvaMediaTransitionStartRequest : public FAvaMediaPlaybackClientMessageBase
+{
+	GENERATED_BODY()
+
+	UPROPERTY()
+	FString ChannelName;
+
+	UPROPERTY()
+	FGuid TransitionId;
+	
+	UPROPERTY()
+	TArray<FGuid> EnterInstanceIds;
+
+	UPROPERTY()
+	TArray<FGuid> PlayingInstanceIds;
+
+	UPROPERTY()
+	TArray<FGuid> ExitInstanceIds;
+
+	UPROPERTY()
+	TArray<FAvalancheRemoteControlValues> EnterValues;
+
+	UPROPERTY()
+	bool bUnloadDiscardedInstances = false;
+
+	/** See EAvalanchePlayableTransitionFlags. */
+	UPROPERTY()
+	uint8 TransitionFlags = 0;
+
+	EAvalanchePlayableTransitionFlags GetTransitionFlags() const { return static_cast<EAvalanchePlayableTransitionFlags>(TransitionFlags);}
+	void SetTransitionFlags(EAvalanchePlayableTransitionFlags InTransitionFlags) { TransitionFlags = static_cast<uint8>(InTransitionFlags);}
+
+};
+
+USTRUCT()
+struct FAvaMediaTransitionStopRequest : public FAvaMediaPlaybackClientMessageBase
+{
+	GENERATED_BODY()
+
+	UPROPERTY()
+	FString ChannelName;
+
+	UPROPERTY()
+	FGuid TransitionId;
+	
+	UPROPERTY()
+	bool bUnloadDiscardedInstances = false;
+};
+
+USTRUCT()
+struct FAvaMediaPlayableTransitionEvent : public FAvaMediaPlaybackServerMessageBase
+{
+	GENERATED_BODY()
+
+	UPROPERTY()
+	FString ChannelName;
+
+	UPROPERTY()
+	FGuid TransitionId;
+	
+	UPROPERTY()
+	FGuid InstanceId;
+
+	UPROPERTY()
+	uint8 EventFlags = static_cast<uint8>(EAvalanchePlayableTransitionEventFlags::None);
+
+	// Because the enum is used as flags, we need to convert to uint8 manually. (Can't use TEnumAsByte for this apparently.)
+	EAvalanchePlayableTransitionEventFlags GetEventFlags() const { return static_cast<EAvalanchePlayableTransitionEventFlags>(EventFlags);}
+	void SetEventFlags(EAvalanchePlayableTransitionEventFlags InFlags) { EventFlags = static_cast<uint8>(InFlags);}
+};
+
+UENUM()
+enum class EAvaMediaBroadcastAction
+{
+	None,
+	Start,
+	Stop,
+	UpdateConfig,
+	DeleteChannel
+};
+
+class UMediaOutput;
+
+/**
+ * Encapsulates a UMediaOutput object in binary form to be able to serialize it
+ * nested within the broadcast request.
+ */
+USTRUCT()
+struct FAvaMediaOutputData
+{
+	GENERATED_BODY()
+
+	UPROPERTY()
+	FAvaMediaOutputInfo OutputInfo;
+
+	UPROPERTY()
+	TSoftClassPtr<UMediaOutput> MediaOutputClass;
+	
+	UPROPERTY()
+	TArray<uint8> SerializedData;
+	
+	UPROPERTY()
+	uint64 ObjectFlags = 0;
+
+	EObjectFlags GetObjectFlags() const { return static_cast<EObjectFlags>(ObjectFlags); }
+};
+
+/** Request by a client to execute an action on broadcast channel(s). */
+USTRUCT()
+struct AVALANCHEMEDIA_API FAvaMediaBroadcastRequest : public FAvaMediaPlaybackClientMessageBase
+{
+	GENERATED_BODY()
+	
+	/** Profile to select for broadcast. */
+	UPROPERTY()
+	FString Profile;
+
+	/** Channel name. If empty, command is for all channels. */
+	UPROPERTY()
+	FString Channel;
+
+	UPROPERTY()
+	TArray<FAvaMediaOutputData> MediaOutputs;
+
+	/** Action to do on channel: Start, Stop, Disable, etc. */
+	UPROPERTY()
+	EAvaMediaBroadcastAction Action = EAvaMediaBroadcastAction::None;
+};
+
+/** Request by a client to update a broadcast channel's settings. */
+USTRUCT()
+struct AVALANCHEMEDIA_API FAvaMediaBroadcastChannelSettingsUpdate : public FAvaMediaPlaybackClientMessageBase
+{
+	GENERATED_BODY()
+	
+	/** Profile to select for broadcast. */
+	UPROPERTY()
+	FString Profile;
+
+	/** Channel name. If empty, command is for all channels. */
+	UPROPERTY()
+	FString Channel;
+
+	UPROPERTY()
+	FAvaViewportQualitySettings QualitySettings;
+};
+
+/**
+ * Client sends this request to have a status of broadcast.
+ *
+ * This is used by the client when it connects to a server which
+ * might be already broadcasting. This happens if the client
+ * disconnects and reconnects. It needs to query the full
+ * status of the server upon connection.
+ */
+USTRUCT()
+struct AVALANCHEMEDIA_API FAvaMediaBroadcastStatusRequest : public FAvaMediaPlaybackClientMessageBase
+{
+	GENERATED_BODY()
+
+	UPROPERTY()
+	bool bIncludeMediaOutputData = false;
+};
+
+USTRUCT()
+struct AVALANCHEMEDIA_API FAvaMediaBroadcastSettings
+{
+	GENERATED_BODY()
+
+	/** Specifies the background clear color for the channel. */
+	UPROPERTY()
+	FLinearColor ChannelClearColor = FLinearColor::Black;
+
+	/** Pixel format used if no media output has specific format requirement. */
+	UPROPERTY()
+	TEnumAsByte<EPixelFormat> ChannelDefaultPixelFormat = EPixelFormat::PF_B8G8R8A8;
+
+	/** Resolution used if no media output has specific resolution requirement. */
+	UPROPERTY()
+	FIntPoint ChannelDefaultResolution = FIntPoint(1920, 1080);
+
+	/**
+	 * Enables drawing the placeholder widget when there is no avalanche asset playing.
+	 * If false, the channel is cleared to the background color.
+	 */
+	UPROPERTY()
+	bool bDrawPlaceholderWidget = false;
+	
+	/** Specify a place holder widget to render when no avalanche asset is playing. */
+	UPROPERTY()
+	FSoftObjectPath PlaceholderWidgetClass;
+};
+
+/**
+ *	Request for the client to replicate it's broadcast settings to the server.
+ **/
+USTRUCT()
+struct AVALANCHEMEDIA_API FAvaMediaBroadcastSettingsUpdate : public FAvaMediaPlaybackClientMessageBase
+{
+	GENERATED_BODY()
+
+	UPROPERTY()
+	FAvaMediaBroadcastSettings BroadcastSettings;
+};
+
+/**
+ * Encapsulate the complete status of a media output.
+ */
+USTRUCT()
+struct AVALANCHEMEDIA_API FAvaMediaOutputStatus
+{
+	GENERATED_BODY()
+
+	UPROPERTY()
+	EAvaMediaOutputState MediaOutputState = EAvaMediaOutputState::Invalid;
+	
+	UPROPERTY()
+	EAvaMediaIssueSeverity MediaIssueSeverity = EAvaMediaIssueSeverity::None;
+
+	UPROPERTY()
+	TArray<FString> MediaIssueMessages;
+};
+
+/**
+ * Server's response when channel's broadcast status changes
+ */
+USTRUCT()
+struct AVALANCHEMEDIA_API FAvaMediaBroadcastStatus : public FAvaMediaPlaybackServerMessageBase
+{
+	GENERATED_BODY()
+	
+	/** Profile to select for broadcast. */
+	UPROPERTY()
+	FString Profile;
+	
+	/** Channel name. If empty, command is for all channels of the profile. */
+	UPROPERTY()
+	FString ChannelName;
+
+	/** Index of the channel in the profile. */
+	UPROPERTY()
+	int32 ChannelIndex = 0;
+
+	/** Number of channels in the profile. */
+	UPROPERTY()
+	int32 NumChannels = 0;
+
+	UPROPERTY()
+	EAvaChannelState ChannelState = EAvaChannelState::Idle;
+
+	UPROPERTY()
+	EAvaMediaIssueSeverity ChannelIssueSeverity = EAvaMediaIssueSeverity::None;
+
+	UPROPERTY()
+	TMap<FGuid, FAvaMediaOutputStatus> MediaOutputStatuses;
+	
+	/**
+	 * This is required to know if the MediaOutputs is empty because the data
+	 * was not included or if it is because the data was included and there is
+	 * actually no outputs defined on that channel.
+	 */
+	UPROPERTY()
+	bool bIncludeMediaOutputData = false;
+	
+	/** Remark: this is only included if bIncludeMediaOutputData == true in the request. */
+	UPROPERTY()
+	TArray<FAvaMediaOutputData> MediaOutputs;
+};
