[MotionDesign]

diff --git a/Engine/Plugins/Experimental/Avalanche/Source/AvalancheMedia/Public/Playback/AvaMediaPlaybackServer.h b/Engine/Plugins/Experimental/Avalanche/Source/AvalancheMedia/Public/Playback/AvaMediaPlaybackServer.h
new file mode 100644
--- /dev/null
+++ b/Engine/Plugins/Experimental/Avalanche/Source/AvalancheMedia/Public/Playback/AvaMediaPlaybackServer.h
@@ -0,0 +1,384 @@
+// Copyright Epic Games, Inc. All Rights Reserved.
+
+#pragma once
+
+#include "Framework/AvalancheInstanceSettings.h"
+#include "IAvalancheBroadcastSettings.h"
+#include "IMessageContext.h"
+#include "MessageEndpoint.h"
+#include "Playback/AvaMediaPlaybackManager.h"
+#include "Playback/AvaMediaPlaybackMessages.h"
+#include "Templates/SharedPointer.h"
+
+class FAvaMediaPlaybackInstance;
+class FAvaMediaSyncManager;
+class UAvaMediaServerPlaybackTransition;
+struct FAvaAssetSyncStatusReceivedParams;
+
+DECLARE_LOG_CATEGORY_EXTERN(LogAvaPlaybackServer, Log, All);
+
+/**
+ * Playback (and Broadcast) Server
+ *
+ * The playback server implements the commands for broadcast (channels and outputs) and playback.
+ * Playback assets are either "playables" or "playback graphs", however the playback server will
+ * create playback graphs for everything. If the asset to play is a "playable" (either a level or
+ * a ava blueprint), it will create a transient playback graph for it.
+ *
+ * A lot of the playback commands are geared toward running a playback graph with a single playable node,
+ * mostly because this system is only used with playlists or playback graphs run on the client side.
+ * So, it will emulate a client-side playable with a local transient playback graph with one player node.
+ * The use case of running a more complex playback graph asset on the server side has not occured yet.
+ */
+class AVALANCHEMEDIA_API FAvaMediaPlaybackServer : public TSharedFromThis<FAvaMediaPlaybackServer>
+{
+public:
+	FAvaMediaPlaybackServer();
+	virtual ~FAvaMediaPlaybackServer();
+
+	void Init(const FString& InAssignedServerName);
+
+	struct FPlaybackInstanceReference
+	{
+		FGuid Id;
+		FSoftObjectPath Path;
+	};
+	
+	TArray<FPlaybackInstanceReference> StopPlaybacks(const FString& InChannelName = FString(), const FSoftObjectPath& InAssetPath = FSoftObjectPath(), bool bInUnload = true);
+	TArray<FPlaybackInstanceReference> StartPlaybacks();
+
+	/**
+	 * Returns a list of channel names from all the playing playback instances.
+	 * Optionally filter for a given asset.
+	 */
+	TArray<FString> GetAllChannelsFromPlayingPlaybacks(const FSoftObjectPath& InAssetPath = FSoftObjectPath()) const;
+	
+	void StartBroadcast();
+	void StopBroadcast();
+
+	/**
+	 *	Indicate the manager is in a shutdown sequence and will force game instances to destroy worlds right away.
+	*/
+	void StartShuttingDown();
+
+	/** Returns the server's name. */
+	const FString& GetName() const { return ServerName;}
+
+	bool HasUserData(const FString& InKey) const { return UserDataEntries.Contains(InKey);}
+	
+	const FString& GetUserData(const FString& InKey) const;
+
+	/** Add user data to this server. This is replicated and accessible to the client. */
+	void SetUserData(const FString& InKey, const FString& InData);
+
+	/** Remove the server's user data entry from the given key. */
+	void RemoveUserData(const FString& InKey);
+
+	/** Returns the list of connected clients. */
+	TArray<FString> GetClientNames() const;
+
+	/** Returns the client address. */
+	FMessageAddress GetClientAddress(const FString& InClientName) const;
+
+	/** Returns true if the corresponding client user data for the given client name is found. */
+	bool HasClientUserData(const FString& InClientName, const FString& InKey) const;
+
+	/**
+	 * Returns the corresponding client user data for the given client name and key.
+	 * Returns empty string if not found.
+	 */
+	const FString& GetClientUserData(const FString& InClientName, const FString& InKey) const;
+
+	/**
+	 * Access broadcast settings replicated from connected client.
+	 * Will return nullptr if no clients are connected.
+	 */
+	const IAvalancheBroadcastSettings* GetBroadcastSettings() const;
+
+	/**
+	 * Access Avalanche Instance settings replicated from connected client(s).
+	 * Will return nullptr if no clients are connected.
+	 */
+	const FAvalancheInstanceSettings* GetAvalancheInstanceSettings() const;
+	
+	/** Access the server's playback manager. */
+	const FAvaMediaPlaybackManager& GetPlaybackManager() const { check(Manager); return *Manager; }
+	FAvaMediaPlaybackManager& GetPlaybackManager() { check(Manager); return *Manager; }
+
+	TSharedPtr<FAvaMediaPlaybackInstance> FindActivePlaybackInstance(const FGuid& InInstanceId) const
+	{
+		const TSharedPtr<FAvaMediaPlaybackInstance>* FoundInstance = ActivePlaybackInstances.Find(InInstanceId);
+		return FoundInstance ? *FoundInstance : TSharedPtr<FAvaMediaPlaybackInstance>();
+	}
+
+	bool RemoveActivePlaybackInstance(const FGuid& InInstanceId)
+	{
+		return ActivePlaybackInstances.Remove(InInstanceId) > 0;	
+	}
+
+	bool RemovePlaybackInstanceTransition(const FGuid& InTransitionId);
+	
+	void SendPlayableTransitionEvent(
+		const FGuid& InTransitionId, const FGuid& InInstanceId, EAvalanchePlayableTransitionEventFlags InFlags,
+		const FName& InChannelName, const FString& InClientName);
+	
+public:
+	// Message handlers
+	void HandlePlaybackPing(const FAvaMediaPlaybackPing& InMessage, const TSharedRef<IMessageContext, ESPMode::ThreadSafe>& InContext);
+	void HandleUpdateClientUserData(const FAvaMediaUpdateClientUserData& InMessage, const TSharedRef<IMessageContext, ESPMode::ThreadSafe>& InContext);
+	void HandleStatCommand(const FAvaMediaStatCommand& InMessage, const TSharedRef<IMessageContext, ESPMode::ThreadSafe>& InContext);
+	void HandleDeviceProviderDataRequest(const FAvaDeviceProviderDataRequest& InMessage, const TSharedRef<IMessageContext, ESPMode::ThreadSafe>& InContext);
+	void HandleUpdateClientInfo(const FAvaMediaUpdateClientInfo& InMessage, const TSharedRef<IMessageContext, ESPMode::ThreadSafe>& InContext);
+	void HandleAvalancheInstanceSettingsUpdate(const FAvalancheInstanceSettingsUpdate& InMessage, const TSharedRef<IMessageContext, ESPMode::ThreadSafe>& InContext);
+	void HandlePackageEvent(const FAvaMediaPlaybackPackageEvent& InMessage, const TSharedRef<IMessageContext, ESPMode::ThreadSafe>& InContext);
+	void HandlePlaybackAssetStatusRequest(const FAvaMediaPlaybackAssetStatusRequest& InMessage, const TSharedRef<IMessageContext, ESPMode::ThreadSafe>& InContext);
+	void HandlePlaybackRequest(const FAvaMediaPlaybackRequest& InMessage, const TSharedRef<IMessageContext, ESPMode::ThreadSafe>& InContext);
+	void HandleAnimPlaybackRequest(const FAvaMediaAnimPlaybackRequest& InMessage, const TSharedRef<IMessageContext, ESPMode::ThreadSafe>& InContext);
+	void HandleRemoteControlUpdateRequest(const FAvaMediaRemoteControlUpdateRequest& InMessage, const TSharedRef<IMessageContext, ESPMode::ThreadSafe>& InContext);
+	void HandlePlayableTransitionStartRequest(const FAvaMediaTransitionStartRequest& InMessage, const TSharedRef<IMessageContext, ESPMode::ThreadSafe>& InContext);
+	void HandlePlayableTransitionStopRequest(const FAvaMediaTransitionStopRequest& InMessage, const TSharedRef<IMessageContext, ESPMode::ThreadSafe>& InContext);
+	void HandleBroadcastSettingsUpdate(const FAvaMediaBroadcastSettingsUpdate& InMessage, const TSharedRef<IMessageContext, ESPMode::ThreadSafe>& InContext);
+	void HandleBroadcastRequest(const FAvaMediaBroadcastRequest& InMessage, const TSharedRef<IMessageContext, ESPMode::ThreadSafe>& InContext);
+	void HandleBroadcastChannelSettingsUpdate(const FAvaMediaBroadcastChannelSettingsUpdate& InMessage, const TSharedRef<IMessageContext, ESPMode::ThreadSafe>& InContext);
+	void HandleBroadcastStatusRequest(const FAvaMediaBroadcastStatusRequest& InMessage, const TSharedRef<IMessageContext, ESPMode::ThreadSafe>& InContext);
+
+	/** Returns underlying message bus endpoint address id  */
+	FString GetMessageEndpointAddressId() const;
+
+protected:
+	void Tick();
+
+	void RegisterCommands();
+	
+	// Command handlers
+	void StartPlaybackCommand(const TArray<FString>& InArgs);
+	void StopPlaybackCommand(const TArray<FString>& InArgs);
+	void StartBroadcastCommand(const TArray<FString>& InArgs);
+	void StopBroadcastCommand(const TArray<FString>& InArgs);
+	void SetUserDataCommand(const TArray<FString>& InArgs);
+	void ShowStatusCommand(const TArray<FString>& InArgs);
+	
+	// Event handlers
+	void OnAvaMediaSettingsChanged(UObject*, struct FPropertyChangedEvent&);
+	void OnChannelChanged(const FAvaOutputChannel& InChannel, EAvaChannelChange InChange);
+	void OnMediaOutputStateChanged(const FAvaOutputChannel& InChannel, const UMediaOutput* InMediaOutput);
+	void OnAvaAssetSyncStatusReceived(const FAvaAssetSyncStatusReceivedParams& InParams);
+	void OnPlaybackInstanceInvalidated(const FAvaMediaPlaybackInstance& InPlaybackInstance);
+	void OnPlaybackInstanceStatusChanged(const FAvaMediaPlaybackInstance& InPlaybackInstance);
+	void OnPlaybackAssetRemoved(const FSoftObjectPath& InAssetPath);
+	void OnPlayableSequenceEvent(UAvalanchePlayable* InPlayable, const FName& SequenceName, EAvalanchePlayableSequenceEventType InEventType);
+	
+	void ApplyAvaMediaSettings();
+
+	template<typename MessageType>
+	void FillServerInfo(MessageType* InMessage)
+	{
+		InMessage->ServerName = ServerName;
+	}
+	
+	template<typename MessageType>
+	void SendResponse(MessageType* InMessage, const FMessageAddress& InRecipient, EMessageFlags flags = EMessageFlags::None)
+	{
+		FillServerInfo(InMessage);
+		MessageEndpoint->Send(InMessage, MessageType::StaticStruct(), flags, nullptr,
+			TArrayBuilder<FMessageAddress>().Add(InRecipient), FTimespan::Zero(), FDateTime::MaxValue());
+	}
+
+	template<typename MessageType>
+	void SendResponse(MessageType* InMessage, const TArray<FMessageAddress>& InRecipients, EMessageFlags flags = EMessageFlags::None)
+	{
+		FillServerInfo(InMessage);
+		MessageEndpoint->Send(InMessage, MessageType::StaticStruct(), flags, nullptr,
+			InRecipients, FTimespan::Zero(), FDateTime::MaxValue());
+	}
+
+	void SendUserDataUpdate(const TArray<FMessageAddress>& InRecipients);
+
+	void SendChannelStatusUpdate(const FString& InChannelName, const FAvaOutputChannel& InChannel, const FMessageAddress& InSender, bool bInIncludeOutputData = false);
+	void SendAllChannelStatusUpdate(const FMessageAddress& InSender, bool bInIncludeOutputData = false);
+	void SendLogMessage(const TCHAR* V, ELogVerbosity::Type Verbosity, const FName& Category, double Time);
+
+	// Playback Commands
+	void ExecutePendingPlaybackCommands();
+
+	TSharedPtr<FAvaMediaPlaybackInstance> GetOrLoadPlaybackInstance(const FGuid& InInstanceId, const FString& InChannelName, const FSoftObjectPath& InAssetPath);
+	void LoadPlayback(const FMessageAddress& InReplyToAddress, const FGuid& InInstanceId, const FString& InChannelName, const FSoftObjectPath& InAssetPath);
+	void StartPlayback(const FMessageAddress& InReplyToAddress, const FGuid& InInstanceId, const FString& InChannelName, const FSoftObjectPath& InAssetPath);
+	void StopPlayback(const FMessageAddress& InReplyToAddress, const FGuid& InInstanceId, const FString& InChannelName, const FSoftObjectPath& InAssetPath);	
+	void UnloadPlayback(const FMessageAddress& InReplyToAddress, const FGuid& InInstanceId, const FString& InChannelName, const FSoftObjectPath& InAssetPath);
+	void SetPlaybackUserData(const FMessageAddress& InReplyToAddress, const FGuid& InInstanceId, const FString& InUserData);
+	void SendPlaybackUserData(const FMessageAddress& InReplyToAddress, const FGuid& InInstanceId);
+	void SendPlaybackStatus(const FMessageAddress& InReplyToAddress, const FGuid& InInstanceId, const FString& InChannelName, const FSoftObjectPath& InAssetPath);
+	
+	void SendPlaybackStatus(const FMessageAddress& InSendTo, const FGuid& InInstanceId, const FString& InChannelName, const FSoftObjectPath& InAssetPath, EAvaMediaPlaybackStatus InStatus);
+	void SendPlaybackStatus(const TArray<FMessageAddress>& InRecipients, const FGuid& InInstanceId, const FString& InChannelName, const FSoftObjectPath& InAssetPath, EAvaMediaPlaybackStatus InStatus);
+	void SendPlaybackStatuses(const FMessageAddress& InSendTo, const FString& InChannelName, const TArray<FPlaybackInstanceReference>& InInstances, EAvaMediaPlaybackStatus InStatus);
+	void SendAllPlaybackStatusesForChannelAndAssetPath(const FMessageAddress& InSendTo, const FString& InChannelName, const FSoftObjectPath& InAssetPath);
+	void SendPlaybackAssetStatus(const FMessageAddress& InSendTo, const FSoftObjectPath& InAssetPath, EAvaMediaPlaybackAssetStatus InStatus);
+
+	EAvaMediaPlaybackStatus GetUnloadedPlaybackStatus(const FSoftObjectPath& InAssetPath)
+	{
+		return Manager->GetUnloadedPlaybackStatus(InAssetPath);
+	}
+
+	bool UpdateChannelOutputConfig(FAvaOutputChannel& InChannel, const TArray<FAvaMediaOutputData>& InMediaOutputs, bool bInRefreshState);
+
+private:
+	FString ComputerName;
+	FString ServerName;
+	FString ProjectContentPath;
+	uint32 ProcessId = 0;
+	TMap<FString, FString> UserDataEntries;
+	
+	/** Holds the messaging endpoint. */
+	TSharedPtr<FMessageEndpoint, ESPMode::ThreadSafe> MessageEndpoint;
+
+	TArray<IConsoleObject*> ConsoleCommands;
+
+	/** The playback server has its own playback manager to not interfere with the local one. */
+	TSharedPtr<FAvaMediaPlaybackManager> Manager;
+
+	// This is used to block sending status update from the channel event while the channels
+	// are refreshing state on media output state changes. This avoid sending spurious channel states
+	// while the update is not completed for all outputs.
+	bool bBlockChannelStatusUpdate = false;
+
+	struct FPendingPlaybackCommand
+	{
+		FMessageAddress ReplyTo;
+		FAvaMediaPlaybackCommand Command;
+	};
+
+	/** Accumulate all the playback commands and execute them all in one batch on the next tick. */
+	TArray<FPendingPlaybackCommand> PendingPlaybackCommands;
+
+	/** Keep an map of active instances per id for fast lookup. */
+	TMap<FGuid, TSharedPtr<FAvaMediaPlaybackInstance>> ActivePlaybackInstances;
+	
+	class FServerPlaybackInstanceTransitionCollection;
+	TUniquePtr<FServerPlaybackInstanceTransitionCollection> PlaybackInstanceTransitions;
+	
+	/**
+	 * Routes the log messages to the server for replication over the message bus.
+	 */
+	class FReplicationOutputDevice : public FOutputDevice
+	{
+	public:
+		FReplicationOutputDevice(FAvaMediaPlaybackServer* InServer);
+		virtual ~FReplicationOutputDevice() override;
+
+		/** Set the minimum verbosity that will be replicated to the client. */
+		void SetVerbosityThreshold(ELogVerbosity::Type InVerbosityThreshold);
+
+	protected:
+		//~ FOutputDevice interface.
+		virtual void Serialize(const TCHAR* InText, ELogVerbosity::Type InVerbosity, const FName& InCategory) override;
+		virtual void Serialize(const TCHAR* InText, ELogVerbosity::Type InVerbosity, const FName& InCategory, double InTime) override;
+
+	private:
+		FAvaMediaPlaybackServer* Server;
+		ELogVerbosity::Type VerbosityThreshold = ELogVerbosity::Type::NoLogging; 
+	};
+	
+	TUniquePtr<FReplicationOutputDevice> ReplicationOutputDevice;
+
+	TOptional<ELogVerbosity::Type> LogReplicationVerbosityFromCommandLine;
+
+	class FClientBroadcastSettings : public IAvalancheBroadcastSettings
+	{
+	public:
+		FAvaMediaBroadcastSettings Settings;
+		virtual ~FClientBroadcastSettings() override = default;
+
+		// IAvalancheBroadcastSettings
+		virtual const FLinearColor& GetChannelClearColor() const override { return Settings.ChannelClearColor; }
+		virtual EPixelFormat GetDefaultPixelFormat() const override { return Settings.ChannelDefaultPixelFormat; }
+		virtual const FIntPoint& GetDefaultResolution() const override { return Settings.ChannelDefaultResolution; }
+		virtual bool IsDrawPlaceholderWidget() const override { return Settings.bDrawPlaceholderWidget; }
+		virtual const FSoftObjectPath& GetPlaceholderWidgetClass() const override { return Settings.PlaceholderWidgetClass; }
+		// ~IAvalancheBroadcastSettings
+	};
+	
+	/** Remote Client Info */
+	class FClientInfo
+	{
+	public:
+		FMessageAddress Address;
+		FString ClientName;
+		FString ComputerName;
+		FString ProjectContentPath;
+		uint32 ProcessId = 0;
+		TMap<FString, FString> UserDataEntries;
+		FClientBroadcastSettings BroadcastSettings;
+		FAvalancheInstanceSettings AvalancheInstanceSettings;
+
+		bool bClientInfoReceived = false;
+	
+		/**
+		 * The status of the server's asset compared to the remote is tracked
+		 * on the server side. For now, the remote client is considered to have the
+		 * reference asset. But in the future, a hub may hold the reference asset instead.
+		 */
+		TSharedPtr<FAvaMediaSyncManager> MediaSyncManager;
+
+		TOptional<FDateTime> PingTimeout;
+		
+		FClientInfo(const FMessageAddress& InClientAddress, const FString& InClientName);
+		~FClientInfo();
+
+		
+		/**
+		* Add a time by which a new ping should arrive
+		*/
+		void AddTimeout(const FDateTime& InNewTimeout)
+		{
+			if (!PingTimeout || InNewTimeout < PingTimeout.GetValue())
+			{
+				PingTimeout = InNewTimeout;
+			}
+		}
+
+		/**
+		 * Check if new ping arrived on time
+		*/
+		bool HasTimedOut(const FDateTime& InNow) const
+		{
+			return PingTimeout && InNow > PingTimeout.GetValue();
+		}
+
+		/**
+		 * Resets the ping timeout.
+		 * This is done when the ping response is received.
+		 */
+		void ResetPingTimeout()
+		{
+			PingTimeout.Reset();
+		}
+	};
+	TMap<FString, TSharedPtr<FClientInfo>> Clients;	// Key is ClientName
+
+	FClientInfo& GetOrCreateClientInfo(const FString& InClientName, const FMessageAddress& InClientAddress);
+	FClientInfo* GetClientInfo(const FMessageAddress& InClientAddress) const;
+	FClientInfo* GetClientInfo(const FString& InClientName) const
+	{
+		const TSharedPtr<FClientInfo>* ClientInfoPtr = Clients.Find(InClientName);
+		return ClientInfoPtr ? ClientInfoPtr->Get() : nullptr;
+	}
+	const FString& GetClientNameSafe(const FMessageAddress& InClientAddress) const;
+	const FMessageAddress& GetClientAddressSafe(const FString& InClientName) const;
+	TArray<FMessageAddress> GetAllClientAddresses(bool bInExcludeClientOnLocalProcess = false) const;
+	void RemoveDeadClients(const FDateTime& InCurrentTime);
+	void OnClientAdded(const FClientInfo& InClientInfo);
+	void OnClientRemoved(const FClientInfo& InRemovedClient);
+	/** A local client will run from the same computer and project folder as the current server. */
+	bool IsLocalClient(const FClientInfo& ClientInfo) const;
+	bool IsLocalClient(const FMessageAddress& InClientAddress) const
+	{
+		const FClientInfo* ClientInfo = GetClientInfo(InClientAddress);
+		return ClientInfo && IsLocalClient(*ClientInfo);
+	}
+	bool IsClientOnLocalProcess(const FClientInfo& ClientInfo) const
+	{
+		return ClientInfo.ComputerName == ComputerName && ClientInfo.ProcessId == ProcessId;
+	}
+};
diff --git a/Engine/Plugins/Experimental/Avalanche/Source/AvalancheMedia/Public/Playback/AvaMediaPlaybackUtils.h b/Engine/Plugins/Experimental/Avalanche/Source/AvalancheMedia/Public/Playback/AvaMediaPlaybackUtils.h
new file mode 100644
--- /dev/null
+++ b/Engine/Plugins/Experimental/Avalanche/Source/AvalancheMedia/Public/Playback/AvaMediaPlaybackUtils.h
@@ -0,0 +1,60 @@
+// Copyright Epic Games, Inc. All Rights Reserved.
+
+#pragma once
+
+#include "Containers/Array.h"
+
+class UPackage;
+struct FAssetData;
+
+class AVALANCHEMEDIA_API FAvaMediaPlaybackUtils
+{
+public:
+	/**
+	 * Reset the loaders of any loaded packages so the client can overwrite the files.
+	*/
+	static void FlushPackageLoading(UPackage* InPackage);
+
+	/**
+	 * Checks if the given packages has been deleted on disk.
+	 * @remark With Editor only.
+	 * @return true if the package was deleted, false if the package still exists on disk.
+	 */
+	static bool IsPackageDeleted(const UPackage* InExistingPackage);
+
+	/**
+	 *	Purge all the objects in memory owned by the given packages.
+	 */
+	static void PurgePackages(const TArray<UPackage*>& InExistingPackages);
+
+	/**
+	 * Reloads the given package.
+	 * @remark With Editor only.
+	 * @return true if the package was reloaded.
+	 */
+	static bool ReloadPackages(const TArray<UPackage*>& InExistingPackages);
+
+	/**
+	 * @brief Determines if the asset is a map by checking the file extension.
+	 * @remark The file must exist on disk.
+	 * @param InPackageName Package name.
+	 * @return true if the file on disk is a .umap file.
+	 */
+	static bool IsMapAsset(const FString& InPackageName);
+
+	/**
+	 * @brief Determines if the asset is a playable (can be used as template) asset, using the asset class.
+	 * @param InAssetData Asset Data
+	 * @return true if the asset is an ava playable.
+	 */
+	static bool IsPlayableAsset(const FAssetData& InAssetData);
+};
+
+namespace UE::AvaMediaPlayback::Utils
+{
+	/**
+	 *	Returns a compactly formatted time stamp information for the current frame.
+	 *	This is to used for logging and tracing.
+	 */
+	AVALANCHEMEDIA_API FString GetBriefFrameInfo();
+}
\ No newline at end of file
diff --git a/Engine/Plugins/Experimental/Avalanche/Source/AvalancheMedia/Public/Playback/AvalanchePlayable.h b/Engine/Plugins/Experimental/Avalanche/Source/AvalancheMedia/Public/Playback/AvalanchePlayable.h
new file mode 100644
--- /dev/null
+++ b/Engine/Plugins/Experimental/Avalanche/Source/AvalancheMedia/Public/Playback/AvalanchePlayable.h
@@ -0,0 +1,201 @@
+// Copyright Epic Games, Inc. All Rights Reserved.
+
+#pragma once
+
+#include "AvaMediaDefines.h"
+#include "AvaMediaPlayableGroup.h"
+#include "AvalancheRemoteControlValues.h"
+#include "Nodes/Events/Actions/AvalancheAnimations.h"
+#include "UObject/Object.h"
+#include "UObject/ObjectPtr.h"
+#include "UObject/SoftObjectPtr.h"
+
+#include "AvalanchePlayable.generated.h"
+
+class IAvaSceneInterface;
+class UAvaMediaPlayableGroup;
+class UAvaMediaPlayableGroupManager;
+class UAvaSequence;
+class UAvaSequencePlayer;
+class UAvalanchePlayableTransition;
+class URemoteControlPreset;
+struct FAvaSoftAssetPtr;
+struct FAvalancheInstancePlaySettings;
+struct FAvalancheRemoteControlValues;
+
+AVALANCHEMEDIA_API DECLARE_LOG_CATEGORY_EXTERN(LogAvalanchePlayable, Log, All);
+
+/** Option flags for the EndPlay function. */
+UENUM()
+enum class EAvalanchePlayableEndPlayOptions : uint8
+{
+	None = 0,
+	/** End play world if no more assets are playing. */
+	ConditionalEndPlayWorld = 1 << 0,
+	/** Perform the request immediately instead of waiting on the next tick. */
+	ForceImmediate = 1 << 1,
+};
+ENUM_CLASS_FLAGS(EAvalanchePlayableEndPlayOptions);
+
+enum class EAvalanchePlayableCommandResult : uint8
+{
+	/** The command was executed successfully. */
+	Executed = 0,
+	/** The command failed and should be discarded. */
+	ErrorDiscard = 1,
+	/** The command couldn't be executed but should be kept and attempted again. */
+	KeepPending = 2
+};
+
+/*
+ * Base class for an avalanche playable.
+ *
+ * A playable (a.k.a. graphic or page) is the basic element that can be rendered
+ * and controlled through the animations and remote control.
+ * 
+ * Design goal:
+ *
+ * The design goal is to abstract the implementation of a playable.
+ * So far we have 2 implementations:
+ * - Avalanche Blueprint that requires to be in it's own game instance/world.
+ * - Level Streaming that can be streamed with other levels in the same game instance.
+ *
+ * To support multiple playable in the same channel/output, there are 2 ways:
+ * - rendering in the same world which can only be done with the level streaming playables.
+ * - compositing different renders. This is not yet supported, but the playable abstraction should help.
+ *
+ * Future Work:
+ * There is also another abstraction desired which is the "remote" playable. The AvalanchePlayback
+ * has a lot of logic to manage remote replication of playable commands. It would probably
+ * be cleaner and more reusable to have a RemotePlayable class that implements all those instead. 
+ */
+UCLASS(Abstract, NotBlueprintable, BlueprintType)	// Todo: review class specifier
+class AVALANCHEMEDIA_API UAvalanchePlayable : public UObject
+{
+	GENERATED_BODY()
+
+public:
+	/**
+	 * Playable creation information contains the necessary information to
+	 * create an instance of a playable and it's playable group.
+	 */
+	struct FPlayableCreationInfo
+	{
+		/** Container for shared playable groups. */
+		UAvaMediaPlayableGroupManager* PlayableGroupManager = nullptr;
+		/** Information about the asset type and path. */
+		FAvaSoftAssetPtr SourceAsset;
+		/** Channel name this playable will be instanced in. */
+		FName ChannelName;
+		/** Provided Playable Group, if not provided, the channel name will be used. */
+		UAvaMediaPlayableGroup* PlayableGroup = nullptr;
+	};
+	
+	/**
+	 * Factory method for the playables.
+	 * This will create the appropriate playable object and initialize it (calling InitPlayable).
+	 * This will setup the playable group appropriately, but will not load the asset.
+	 */
+	static UAvalanchePlayable* Create(UObject* InOuter, const FPlayableCreationInfo& InPlayableInfo);
+
+	DECLARE_MULTICAST_DELEGATE_ThreeParams(FOnSequenceEvent, UAvalanchePlayable*, const FName& /*InSequenceName*/, EAvalanchePlayableSequenceEventType);
+	static FOnSequenceEvent& OnSequenceEvent() { return OnSequenceEventDelegate; }
+
+	DECLARE_MULTICAST_DELEGATE_ThreeParams(FOnTransitionEvent, UAvalanchePlayable*, UAvalanchePlayableTransition*, EAvalanchePlayableTransitionEventFlags);
+	static FOnTransitionEvent& OnTransitionEvent() { return OnTransitionEventDelegate; }
+
+	DECLARE_MULTICAST_DELEGATE_OneParam(FOnPlayableStatusChanged, UAvalanchePlayable*);
+	FOnPlayableStatusChanged& OnPlayableStatusChanged() { return OnPlayableStatusChangedDelegate; }
+	
+	virtual bool LoadAsset(const FAvaSoftAssetPtr& InAvalancheSourceAsset, bool bInInitiallyVisible) { return false;}
+	virtual bool UnloadAsset() {return false;}
+	virtual const FSoftObjectPath& GetSourceAssetPath() const;
+	virtual EAvalanchePlayableStatus GetPlayableStatus() const { return EAvalanchePlayableStatus::Unknown;}
+	virtual IAvaSceneInterface* GetSceneInterface() const {return nullptr;}
+	virtual EAvalanchePlayableCommandResult ExecuteAnimationCommand(EAvaMediaAnimAction InAnimAction, const FAnimPlaySettings& InAnimPlaySettings);
+	virtual EAvalanchePlayableCommandResult UpdateRemoteControlCommand(const TSharedRef<FAvalancheRemoteControlValues>& InRemoteControlValues);
+	virtual bool ApplyCamera() { return false; }
+	virtual bool IsRemoteProxy() const { return false; }
+	virtual bool GetShouldBeVisible() const { return true; }
+	virtual void SetShouldBeVisible(bool bInShouldBeVisible) {}
+	
+	/**
+	 * @brief Ensures the given asset is playing (visible) with the given parameters.
+	 * @param InWorldPlaySettings World settings (render target) to render this asset with.
+	 * @remark This doesn't trigger the animations.
+	 */
+	void BeginPlay(const FAvalancheInstancePlaySettings& InWorldPlaySettings);
+
+	/**
+	 * @brief Make this asset not play anymore (will end up hidden)
+	 */
+	void EndPlay(EAvalanchePlayableEndPlayOptions InOptions);
+
+	bool IsPlaying() const { return bIsPlaying; }
+	
+	UAvaMediaPlayableGroup* GetPlayableGroup() const { return PlayableGroup; }
+
+	UWorld* GetPlayWorld() const { return PlayableGroup ? PlayableGroup->GetPlayWorld() : nullptr; }
+
+	bool HasSequence(const UAvaSequence* InSequence) const;
+
+	virtual void SetInstanceId(const FGuid& InInstanceId ) { InstanceId = InInstanceId; }
+	const FGuid& GetInstanceId() const { return InstanceId; }
+
+	virtual void SetUserData(const FString& InUserData) { UserData = InUserData; }
+	const FString& GetUserData() const { return UserData; }
+
+	const FAvalancheRemoteControlValues& GetLatestRemoteControlValues() const
+	{
+		return LatestRemoteControlValues.IsValid() ? *LatestRemoteControlValues : FAvalancheRemoteControlValues::GetDefaultEmpty();
+	}
+
+protected:
+	/**
+	 * @brief Performs derived initialization of the playable. Does not load the asset.
+	 * @remark Playable Group is setup by the derived classes.
+	 * @param InPlayableInfo Necessary information to setup the playable.
+	 * @return true if the playable group and derived initialization was successful.
+	 */
+	virtual bool InitPlayable(const FPlayableCreationInfo& InPlayableInfo);
+
+	/** Called by BeginPlay for derived classes implementation. */
+	virtual void OnPlay() {}
+
+	/** Called by EndPlay for derived classes implementation. */
+	virtual void OnEndPlay() {}
+
+	void HandleOnSequenceStarted(UAvaSequencePlayer* InSequencePlayer, UAvaSequence* InSequence);
+	void HandleOnSequenceFinished(UAvaSequencePlayer* InSequencePlayer, UAvaSequence* InSequence);
+
+	static UAvalanchePlayable* CreateLocalPlayable(UObject* InOuter, const FPlayableCreationInfo& InPlayableInfo);
+	static UAvalanchePlayable* CreateRemoteProxyPlayable(UObject* InOuter, const FPlayableCreationInfo& InPlayableInfo);
+
+protected:
+	/**
+	 * Playable "instancing container" group.
+	 * Defines the interface for what a playable can do to it's container.
+	 */
+	UPROPERTY(Transient)
+	TObjectPtr<UAvaMediaPlayableGroup> PlayableGroup;
+	
+	bool bIsPlaying = false;
+
+	/**
+	 * Unique identifier for this playable instance.
+	 * This is used for the playback client and server's playable replication. 
+	 */
+	FGuid InstanceId;
+
+	/**
+	 * User data that is replicated to and from the playback server.
+	 * This can be used to transport additional information for this playable.
+	 */
+	FString UserData;
+
+	TSharedPtr<FAvalancheRemoteControlValues> LatestRemoteControlValues;
+
+	static FOnSequenceEvent OnSequenceEventDelegate;
+	static FOnTransitionEvent OnTransitionEventDelegate;
+	FOnPlayableStatusChanged OnPlayableStatusChangedDelegate;
+};
diff --git a/Engine/Plugins/Experimental/Avalanche/Source/AvalancheMedia/Public/Playback/AvalanchePlayback.h b/Engine/Plugins/Experimental/Avalanche/Source/AvalancheMedia/Public/Playback/AvalanchePlayback.h
new file mode 100644
--- /dev/null
+++ b/Engine/Plugins/Experimental/Avalanche/Source/AvalancheMedia/Public/Playback/AvalanchePlayback.h
@@ -0,0 +1,481 @@
+// Copyright Epic Games, Inc. All Rights Reserved.
+
+#pragma once
+
+#include "AvaMediaDefines.h"
+#include "AvalancheRemoteControlValues.h"
+#include "Nodes/AvaPlaybackNodePlayer.h"
+#include "Nodes/AvaPlaybackNodeRoot.h"
+#include "Nodes/Events/Actions/AvalancheAnimations.h"
+#include "UObject/Object.h"
+#include "UObject/ObjectPtr.h"
+#include "UObject/SoftObjectPtr.h"
+
+#include "AvalanchePlayback.generated.h"
+
+AVALANCHEMEDIA_API DECLARE_LOG_CATEGORY_EXTERN(LogAvalanchePlayback, Log, All);
+
+class FAvaMediaPlaybackManager;
+class IAvaPlaybackGraphEditor;
+class UAvaMediaPlayableGroupManager;
+class UAvaPlaybackNode;
+class UAvaPlaybackNodeRoot;
+class UAvalancheBroadcast;
+class UAvalancheGameInstance;
+class UAvalanchePlayable;
+class UAvalanchePlayableTransition;
+class UEdGraph;
+class UTextureRenderTarget2D;
+struct FAnimPlaySettings;
+struct FAvaOutputChannel;
+
+UENUM()
+enum class EAvaPlaybackStopOptions : uint8
+{
+	/**
+	 * Default option allows for deferred execution of the request when it is safe to do so.
+	 */
+	None			= 0,
+	/**
+	 * Forces the execution of the request when it is called. Typically during shut down.
+	 */
+	ForceImmediate	= 1 << 1,
+	/**
+	 *	Unload from memory after being stopped.
+	 */
+	Unload			= 1 << 2,
+	/**
+	 * Default option allows for deferred execution of the request when it is safe to do so.
+	 */
+	Default			= None
+};
+ENUM_CLASS_FLAGS(EAvaPlaybackStopOptions);
+
+UENUM()
+enum class EAvaPlaybackUnloadOptions : uint8
+{
+	/**
+	 * Default option allows for deferred execution of the request when it is safe to do so.
+	 */
+	None			= 0,
+	/**
+	 * Forces the execution of the request when it is called. Typically during shut down.
+	 */
+	ForceImmediate	= 1 << 1,
+	/**
+	 * Default option allows for deferred execution of the request when it is safe to do so.
+	 */
+	Default			= None
+};
+ENUM_CLASS_FLAGS(EAvaPlaybackUnloadOptions);
+
+/**
+ * Owns a group of playables.
+ * Used to group playables per channels, but can be extended to any conceptual grouping.
+ */
+USTRUCT()
+struct AVALANCHEMEDIA_API FAvalanchePlaybackPlayableGroup
+{
+	GENERATED_BODY()
+
+	// TODO: For transition logic within the playback graph, we can't simply index the
+	// playables by asset path since we may have to instance the same asset twice for a transition between pages
+	// using the same template. For now, playlist does transition with multiple playback graph instances instead.
+	UPROPERTY(Transient)
+	TMap<FSoftObjectPath, TObjectPtr<UAvalanchePlayable>> Playables;
+	
+	UAvalanchePlayable* FindPlayable(const FSoftObjectPath& InSourceAssetPath) const
+	{
+		const TObjectPtr<UAvalanchePlayable>* Playable = Playables.Find(InSourceAssetPath);
+		return Playable ? *Playable : nullptr;
+	}
+
+	bool HasPlayable(const UAvalanchePlayable* InPlayable) const
+	{
+		for (const TPair<FSoftObjectPath, TObjectPtr<UAvalanchePlayable>>& Playable : Playables)
+		{
+			if (Playable.Value == InPlayable)
+			{
+				return true;
+			}
+		}
+		return false;
+	}
+	
+	UAvalanchePlayable* GetFirstPlayable() const
+	{
+		for (const TPair<FSoftObjectPath, TObjectPtr<UAvalanchePlayable>>& Playable : Playables)
+		{
+			if (Playable.Value)
+			{
+				return Playable.Value;
+			}
+		}
+		return nullptr;
+	}
+	
+	void GetAllPlayables(TArray<UAvalanchePlayable*>& OutPlayables) const
+	{
+		for (const TPair<FSoftObjectPath, TObjectPtr<UAvalanchePlayable>>& Playable : Playables)
+		{
+			if (Playable.Value)
+			{
+				OutPlayables.Add(Playable.Value);
+			}
+		}
+	}
+
+	void ForEachPlayable(TFunctionRef<void(const UAvalanchePlayable*)> InFunction) const
+	{
+		for (const TPair<FSoftObjectPath, TObjectPtr<UAvalanchePlayable>>& Playable : Playables)
+		{
+			if (Playable.Value)
+			{
+				InFunction(Playable.Value);
+			}
+		}
+	}
+};
+
+UCLASS(NotBlueprintable, BlueprintType)
+class AVALANCHEMEDIA_API UAvalanchePlayback : public UObject
+{
+	GENERATED_BODY()
+
+public:
+	virtual ~UAvalanchePlayback() override;
+	
+	UFUNCTION(BlueprintCallable, Category = "Motion Design Playback")
+	void Play();
+	
+	UFUNCTION(BlueprintCallable, Category = "Motion Design Playback")
+	void Stop(EAvaPlaybackStopOptions InStopOptions);
+	
+	UFUNCTION(BlueprintCallable, Category = "Motion Design Playback")
+	void LoadInstances();
+
+	/** Unload all game instance's worlds from this playback. */
+	UFUNCTION(BlueprintCallable, Category = "Motion Design Playback")
+	void UnloadInstances(EAvaPlaybackUnloadOptions InUnloadOptions);
+	
+	UFUNCTION(BlueprintCallable, Category = "Motion Design Playback")
+	bool IsPlaying() const { return bIsPlaying; }
+
+	/**
+	 * Sets the parent playback manager that will handle ticking. 
+	 */
+	void SetPlaybackManager(const TSharedPtr<FAvaMediaPlaybackManager>& InPlaybackManager);
+
+	bool IsManaged() const { return PlaybackManagerWeak.IsValid(); }
+	
+	DECLARE_MULTICAST_DELEGATE_OneParam(FOnPlaybackStateChanged, bool);
+	FOnPlaybackStateChanged OnPlaybackStateChanged;
+
+	/** Quick access to the first playable in the playback graph. */ 
+	UAvalanchePlayable* GetFirstPlayable() const;
+
+	void GetAllPlayables(TArray<UAvalanchePlayable*>& OutPlayables) const;
+
+	void ForEachPlayable(TFunctionRef<void(const UAvalanchePlayable*)> InFunction) const;
+
+	/** Returns true if the given playable is part of this playback graph. */
+	bool HasPlayable(const UAvalanchePlayable* InPlayable) const;
+	
+	UAvalanchePlayable* FindPlayable(const FSoftObjectPath& InSourceAssetPath, const FName& InChannelName) const;
+	UAvalanchePlayable* FindOrLoadPlayable(const FAvaSoftAssetPtr& InSourceAsset, const FName& InChannelName);
+	void RemovePlayable(const FSoftObjectPath& InSourceAssetPath, const FName& InChannelName);
+	bool UnloadAndRemovePlayable(UAvalanchePlayable* InPlayable, const FSoftObjectPath& InSourceAssetPath, const FName& InChannelName, bool bInForceImmediate);
+	bool UnloadAndRemovePlayable(const FSoftObjectPath& InSourceAssetPath, const FName& InChannelName, bool bInForceImmediate);
+
+	TArray<UAvalancheGameInstance*> GetActiveGameInstances() const;
+	
+	/**
+	 *	Determines if the graph contains a node with the given source avalanche blueprint.
+	 *	This works even if the avalanche blueprint running.
+	 */
+	bool HasPlayerNodeForSourceAvalancheAsset(const FSoftObjectPath& InSourceAssetPath) const;
+
+	void SetPlayableGroupManager(UAvaMediaPlayableGroupManager* InPlayableGroupManager);
+
+	UAvaMediaPlayableGroupManager* GetPlayableGroupManager() const
+	{
+		return PlayableGroupManager ? PlayableGroupManager.Get() : GetGlobalPlayableGroupManager();
+	}
+
+	static UAvaMediaPlayableGroupManager* GetGlobalPlayableGroupManager();
+
+	/** Resolve the channel name for a playback settings entry. */
+	FName ResolveChannelName(const FName InPlaybackChannelName) const
+	{
+		// For a preview, route everything to the preview channel,
+		// otherwise use the resolved channel (from traversing the graph)
+		return IsPreviewOnly() ? GetPreviewChannelName() : InPlaybackChannelName;
+	}
+
+public:
+	bool IsDryRunningGraph() const { return bIsDryRunningGraph; }
+	void DryRunGraph(bool bDeferredExecution = false);
+
+	/**
+	 *	Sets the preview channel this playback object is dedicated to.
+	 *	By doing this, the playback object will render in only that channel regardless
+	 *	of the internal playback graph.
+	 *	The specified channel may not exist in the broadcast configuration, in which case the
+	 *	preview playback will fallback to a channel-less render target.
+	 */
+	void SetPreviewChannelName(const FName& InPreviewChannelName) { PreviewChannelName = InPreviewChannelName; }
+
+	FName GetPreviewChannelName() const { return PreviewChannelName; }
+
+	/**
+	 * Returns true if the playback object is dedicated to a local preview. 
+	 */
+	bool IsPreviewOnly() const { return !PreviewChannelName.IsNone(); }
+
+	/**
+	 * Helper function to implement some of the event nodes logic.
+	 * It will return the list of channel names that correspond this the indices.
+	 */
+	TArray<FName> GetChannelNamesForIndices(const TArray<int32>& InChannelIndices) const;
+
+private:
+	void DryRunGraphInternal();
+	void RegisterTickDelegate();
+	void UnregisterTickDelegate();
+	void OnEndFrameTick();
+	void ConditionalResolvePlaybackSettings();
+
+public:
+	void SetRootNode(UAvaPlaybackNodeRoot* InRoot);
+	void AddPlayerNode(UAvaPlaybackNodePlayer* InPlayer);
+	
+	UAvaPlaybackNodeRoot* GetRootNode() const;
+	
+	void AddPlaybackNode(UAvaPlaybackNode* Node);
+	void RemovePlaybackNode(UAvaPlaybackNode* Node);
+
+	/**
+	 * Stops the playable corresponding to the given asset in the given channel.
+	 */
+	void StopPlayable(const FSoftObjectPath& InSourceAssetPath, const FName& InChannelName);
+
+	/**
+	 * Start streaming the asset in a game instance.
+	 */
+	void LoadAsset(const FAvaSoftAssetPtr& InSourceAsset, const FName& InChannelName);
+	
+	virtual void Tick(float DeltaTime);
+	
+	void PushAnimationCommand(const FSoftObjectPath& InSourceAssetPath, const FString& InChannelName, EAvaMediaAnimAction InAnimAction, const FAnimPlaySettings& InAnimSettings);
+	void PushRemoteControlValues(const FSoftObjectPath& InSourceAssetPath, const FString& InChannelName, const TSharedRef<FAvalancheRemoteControlValues>& InRemoteControlValues);
+
+	void SetupPlaybackNode(UAvaPlaybackNode* InPlaybackNode, bool bSelectNewNode = true);
+
+#if WITH_EDITOR
+	const TArray<TObjectPtr<UAvaPlaybackNode>>& GetPlaybackNodes() const;
+
+	/** Create the basic sound graph */
+	void CreateGraph();
+
+	/** Clears all nodes from the graph */
+	void ClearGraph();
+
+	void RefreshPlaybackNode(UAvaPlaybackNode* InPlaybackNode);	
+
+	void CompilePlaybackNodesFromGraphNodes();
+
+	/** Get the Playback EdGraph  */
+	UEdGraph* GetGraph();
+
+	/** Resets all graph data and nodes */
+	void ResetGraph();
+
+	void SetGraphEditor(TSharedPtr<IAvaPlaybackGraphEditor> InGraphEditor);
+
+	TSharedPtr<IAvaPlaybackGraphEditor> GetGraphEditor() const;
+#endif
+
+	template<typename InPlaybackNodeType, typename = typename TEnableIf<TIsDerivedFrom<InPlaybackNodeType, UAvaPlaybackNode>::IsDerived>::Type>
+	InPlaybackNodeType* ConstructPlaybackNode(TSubclassOf<InPlaybackNodeType> NodeClass = nullptr, bool bSelectNewNode = true)
+	{
+		//Ensure Class is Valid
+		if (!NodeClass)
+		{
+			NodeClass = InPlaybackNodeType::StaticClass();
+		}
+
+		if (!NodeClass->IsChildOf(UAvaPlaybackNodeRoot::StaticClass()) || !IsValid(RootNode))
+		{
+			//Set flag to be transactional so it registers with undo system
+			InPlaybackNodeType* const PlaybackNode = NewObject<InPlaybackNodeType>(this
+				, NodeClass
+				, NAME_None
+				, RF_Transactional);
+
+#if WITH_EDITOR
+			PlaybackNodes.Add(PlaybackNode);
+#endif
+			
+			PlaybackNode->PostAllocateNode();
+			SetupPlaybackNode(PlaybackNode, bSelectNewNode);
+		
+			return PlaybackNode;
+		}
+
+		checkf(0, TEXT("Root Node has already been created!"));
+		return nullptr;
+	}
+
+	/** Delegate called when a playable is created in this playback graph. */
+	DECLARE_MULTICAST_DELEGATE_TwoParams(FOnPlayableCreated, UAvalanchePlayback* /*InPlayback*/, UAvalanchePlayable* /*InPlayable*/);
+	FOnPlayableCreated OnPlayableCreated;
+
+protected:
+	void ExecutePendingAnimationCommands();
+	void ExecutePendingRemoteControlCommands();
+	
+	void OnChannelChanged(const FAvaOutputChannel& InChannel, EAvaChannelChange InChange);
+	void OnChannelBroadcastStateChanged(const FAvaOutputChannel& InChannel);
+
+	void SetIsPlaying(bool bInIsPlaying);
+	
+	bool CanRefreshPlayback(const FAvaSoftAssetPtr& InSourceAsset, const FAvaOutputChannel& InChannel) const;
+
+	bool RefreshPlayback(const FAvaSoftAssetPtr& InSourceAsset, FAvaOutputChannel& InChannel, const FName& InChannelName);
+	
+	bool CanRefreshPreview(const FAvaSoftAssetPtr& InSourceAsset) const;
+	
+	bool RefreshPreview(const FAvaSoftAssetPtr& InSourceAsset, FAvaOutputChannel& InChannel, const FName& InChannelName);
+
+	static UTextureRenderTarget2D* UpdatePlaybackRenderTarget(const UAvalanchePlayable* InPlayable, const FAvaOutputChannel& InChannel);
+	
+protected:
+	/**
+	 * Keep track of the parent manager that handles ticking.
+	 * If the playback graph is an asset (not managed), it will handle it's own ticking.
+	 */
+	TWeakPtr<FAvaMediaPlaybackManager> PlaybackManagerWeak;
+
+	TMap<FName, FAvaPlaybackChannelParameters> PlaybackSettings;
+	TMap<FName, FAvaPlaybackChannelParameters> PreviousPlaybackSettings;
+	
+	UPROPERTY(Transient)
+	TObjectPtr<UAvaPlaybackNodeRoot> RootNode;
+
+	UPROPERTY(Transient)
+	TArray<TObjectPtr<UAvaPlaybackNodePlayer>> PlayerNodes;
+
+#if WITH_EDITORONLY_DATA
+	TWeakPtr<IAvaPlaybackGraphEditor> GraphEditorWeak;
+	
+	UPROPERTY()
+	TObjectPtr<UEdGraph> EdGraph;
+
+	/**
+	 * An array of all the Playback Nodes, Only available in Editor (Playback doesn't need to know disconnected nodes,
+	 * unless we support runtime changes as well).
+	 */
+	UPROPERTY()
+	TArray<TObjectPtr<UAvaPlaybackNode>> PlaybackNodes;
+#endif
+
+	/**
+	 * Playback Group Manager for this Playback object.
+	 *
+	 * This manager defines the scope with which this playback object's playables
+	 * can share their group with. In other words, the playables of this playback
+	 * (that support sharing a group) can have a shared group (for the given channel)
+	 * cached in this given manager.
+	 *
+	 * So far, the PlayableGroupManager has the same scope as the playable manager,
+	 * i.e. global (a.k.a local) or per playback server.
+	 */
+	UPROPERTY(Transient)
+	TObjectPtr<UAvaMediaPlayableGroupManager> PlayableGroupManager;
+
+	/**
+	 * Keep track of the playables per channel.
+	 * 
+	 * The Playback object owns it's playables. However, the playables
+	 * can share a playable group from other playback objects. This
+	 * is necessary for the playback server that is always creating just
+	 * one playback object per playable (atm).
+	 */
+	UPROPERTY(Transient)
+	TMap<FName, FAvalanchePlaybackPlayableGroup> ChannelPlayableGroups;
+	
+	bool bIsPlaying = false;
+	bool bIsDryRunningGraph = false;
+	bool bAsyncDryRunRequested = false;
+	bool bIsTicking = false;	// Tick reentrancy guard.
+	FDelegateHandle TickDelegateHandle;
+
+	/** Name of the preview channel this playback object is dedicated to. */
+	FName PreviewChannelName;
+	
+	struct FPlaybackCommand
+	{
+		/** If the command has been pending for too long, it will be discarded. */
+		FDateTime Timeout;
+		
+		FSoftObjectPath SourcePath;
+		FString ChannelName;
+		FName ChannelFName;
+
+		bool HasTimedOut(const FDateTime& InNow) const
+		{
+			return InNow > Timeout;	
+		}
+	};
+	struct FAnimationCommand : public FPlaybackCommand
+	{
+		EAvaMediaAnimAction AnimAction;
+		FAnimPlaySettings AnimPlaySettings;
+	};
+	struct FRemoteControlCommand : public FPlaybackCommand
+	{
+		TSharedRef<FAvalancheRemoteControlValues> Values;
+	};
+	
+	TArray<FAnimationCommand> PendingAnimationCommands;
+	TArray<FRemoteControlCommand> PendingRemoteControlCommands;
+};
+
+/**
+ * Somewhat reusable helper class to build a playback graph. 
+ */
+class FAvalanchePlaybackGraphBuilder
+{
+public:
+	FAvalanchePlaybackGraphBuilder(UAvaMediaPlayableGroupManager* InPlayableGroupManager);
+	~FAvalanchePlaybackGraphBuilder();
+
+	/**
+	 * @brief Connect the given node to the root node's pin corresponding to the given channel name.
+	 * @param InChannelName Channel name to connect to.
+	 * @param InNodeToConnect Node (already constructed) to connect.
+	 */
+	bool ConnectToRoot(const FString& InChannelName, UAvaPlaybackNode* InNodeToConnect);
+	
+	int32 GetPinIndexForChannel(const FString& InChannelName) const;
+
+	UAvalanchePlayback* FinishBuilding();
+
+	/** Helper function to construct a playback node. This will add the node in the playback's node list. */
+	template<typename InPlaybackNodeType, typename = typename TEnableIf<TIsDerivedFrom<InPlaybackNodeType, UAvaPlaybackNode>::IsDerived>::Type>
+		InPlaybackNodeType* ConstructPlaybackNode(TSubclassOf<InPlaybackNodeType> InNodeClass = nullptr, bool bInSelectNewNode = true)
+	{
+		if (Playback)
+		{
+			bFinished = false;
+			return Playback->ConstructPlaybackNode<InPlaybackNodeType>(InNodeClass, bInSelectNewNode);
+		}
+		checkf(0, TEXT("Playback object is not created."));
+		return nullptr;
+	}
+	
+private:
+	bool bFinished = false;
+	UAvalanchePlayback* Playback = nullptr;
+	UAvaPlaybackNodeRoot* RootNode = nullptr;
+};
\ No newline at end of file
diff --git a/Engine/Plugins/Experimental/Avalanche/Source/AvalancheMedia/Public/Playback/AvalancheRemoteControl.h b/Engine/Plugins/Experimental/Avalanche/Source/AvalancheMedia/Public/Playback/AvalancheRemoteControl.h
new file mode 100644
--- /dev/null
+++ b/Engine/Plugins/Experimental/Avalanche/Source/AvalancheMedia/Public/Playback/AvalancheRemoteControl.h
@@ -0,0 +1,76 @@
+// Copyright Epic Games, Inc. All Rights Reserved.
+
+#pragma once
+
+#include "CoreTypes.h"
+#include "RemoteControlPreset.h"
+#include "AvalancheRemoteControl.generated.h"
+
+AVALANCHEMEDIA_API DECLARE_LOG_CATEGORY_EXTERN(LogAvaMediaRemoteControl, Log, All);
+
+/**
+ * Return code for Remote Control operations.
+ */
+UENUM()
+enum class EAvaRemoteControlResult : int8
+{
+	/** Operation completed successfully */
+	Completed = 0,
+	/** Operation was not necessary. */
+	UpToDate = 1,	
+	/** Input parameter is invalid */ 
+	InvalidParameter = -1,
+	/** Object reference read access failed. */
+	ReadAccessDenied = -2,
+	/** Object reference write access failed. */
+	WriteAccessDenied = -3,
+	/** Serialization from property failed. */
+	ReadPropertyFailed = -4,
+	/** Serialization to property failed. */
+	WritePropertyFailed = -5
+};
+
+namespace UE::AvalancheRemoteControl
+{
+	/** Gets the value (as bytes) of the given property (remote control entity). Returns true if the operation was successful, false otherwise. */
+	AVALANCHEMEDIA_API EAvaRemoteControlResult GetValueOfEntity(const TSharedPtr<const FRemoteControlEntity>& InRemoteControlEntity, TArray<uint8>& OutValue);
+	AVALANCHEMEDIA_API EAvaRemoteControlResult GetValueOfEntity(const TSharedPtr<const FRemoteControlEntity>& InRemoteControlEntity, FString& OutValue);
+	
+	/** Set the value (as bytes) of the given property (remote control entity). Returns true if the operation was successful, false otherwise. */
+	AVALANCHEMEDIA_API EAvaRemoteControlResult SetValueOfEntity(const TSharedPtr<FRemoteControlEntity>& InRemoteControlEntity, const TArrayView<const uint8>& InValue);
+	AVALANCHEMEDIA_API EAvaRemoteControlResult SetValueOfEntity(const TSharedPtr<FRemoteControlEntity>& InRemoteControlEntity, const FString& InValue);
+
+	/** Gets the value (as bytes) of the given controller. Returns true if the operation was successful, false otherwise. */
+	AVALANCHEMEDIA_API EAvaRemoteControlResult GetValueOfController(URCVirtualPropertyBase* InController, TArray<uint8>& OutValue);
+	AVALANCHEMEDIA_API EAvaRemoteControlResult GetValueOfController(URCVirtualPropertyBase* InController, FString& OutValue);
+	
+	/** Set the value (as bytes) of the given controller. Returns true if the operation was successful, false otherwise. */
+	AVALANCHEMEDIA_API EAvaRemoteControlResult SetValueOfController(URCVirtualPropertyBase* InController, const TArrayView<const uint8>& InValue);
+	AVALANCHEMEDIA_API EAvaRemoteControlResult SetValueOfController(URCVirtualPropertyBase* InController, const FString& InValue);
+
+	/** Retrieve the entity ids that are controlled by the given controller. */
+	AVALANCHEMEDIA_API bool GetEntitiesControlledByController(const URemoteControlPreset* InRemoteControlPreset, const URCVirtualPropertyBase* InVirtualProperty, TSet<FGuid>& OutEntityIds);
+
+	inline bool Failed(EAvaRemoteControlResult InResult)
+	{
+		// Negative result codes are failure.
+		return InResult < EAvaRemoteControlResult::Completed;
+	}
+
+	AVALANCHEMEDIA_API FString EnumToString(EAvaRemoteControlResult InValue);
+	
+	/**
+	 * Utility class to temporarily push the enabled state of the controller's behaviours.
+	 * This class is intended to be used temporarily on local stack. Lifetime of virtual property is not managed.
+	 */
+	class AVALANCHEMEDIA_API FScopedPushControllerBehavioursEnable
+	{
+	public:
+		FScopedPushControllerBehavioursEnable(URCVirtualPropertyBase* InVirtualProperty, bool bInBehavioursEnabled);
+		~FScopedPushControllerBehavioursEnable();
+		
+	private:
+		URCVirtualPropertyBase* VirtualProperty = nullptr;
+		TArray<bool> PreviousBehavioursEnabled;
+	};
+}
\ No newline at end of file
diff --git a/Engine/Plugins/Experimental/Avalanche/Source/AvalancheMedia/Public/Playback/AvalancheRemoteControlValues.h b/Engine/Plugins/Experimental/Avalanche/Source/AvalancheMedia/Public/Playback/AvalancheRemoteControlValues.h
new file mode 100644
--- /dev/null
+++ b/Engine/Plugins/Experimental/Avalanche/Source/AvalancheMedia/Public/Playback/AvalancheRemoteControlValues.h
@@ -0,0 +1,170 @@
+// Copyright Epic Games, Inc. All Rights Reserved.
+
+#pragma once
+
+#include "CoreTypes.h"
+#include "RemoteControlPreset.h"
+#include "AvalancheRemoteControlValues.generated.h"
+
+/**
+ *	Flags indicating which component of the remote control values have been
+ *	modified by an operation.
+ */
+UENUM()
+enum class EAvaRemoteControlChanges : uint8
+{
+	None = 0,
+	EntityValues		= 1 << 0,
+	ControllerValues	= 1 << 1,
+
+	All              = 0xFF,
+};
+ENUM_CLASS_FLAGS(EAvaRemoteControlChanges);
+
+USTRUCT()
+struct AVALANCHEMEDIA_API FAvalancheRemoteControlValue
+{
+	GENERATED_BODY()
+
+	FAvalancheRemoteControlValue() = default;
+
+	FAvalancheRemoteControlValue(const FString& InValue, bool bInIsDefault = false)
+		: Value(InValue), bIsDefault(bInIsDefault) {}
+
+	void SetValueFrom(const FAvalancheRemoteControlValue& InOther)
+	{
+		Value = InOther.Value;
+	}
+
+	/** Returns true if the given value is the same. Ignores the default flag. */
+	bool IsSameValueAs(const FAvalancheRemoteControlValue& InOther) const
+	{
+		return Value.Equals(InOther.Value, ESearchCase::CaseSensitive);
+	}
+
+	bool Serialize(FArchive& Ar);
+
+	/**
+	 * The Remote Control Entity or Controller's Value stored as a Json formatted string.
+	 */
+	UPROPERTY()
+	FString Value;
+
+	/**
+	 * Indicate if the value is a default value from a template.
+	 * This is used to know which values to update when
+	 * updating the page's values from the template (reimport page).
+	 * This is set to true only when the values are from the template.
+	 * If values are modified by an edit operation, it will be set to false.
+	 */
+	UPROPERTY()
+	bool bIsDefault = false;
+};
+
+template<>
+struct TStructOpsTypeTraits<FAvalancheRemoteControlValue> : public TStructOpsTypeTraitsBase2<FAvalancheRemoteControlValue>
+{
+	enum
+	{
+		WithSerializer = true,
+	};
+};
+
+USTRUCT()
+struct AVALANCHEMEDIA_API FAvalancheRemoteControlValues
+{
+	GENERATED_BODY()
+
+	FAvalancheRemoteControlValues() = default;
+	
+	/** Value as a binary array of the Remote Control Entity. */
+	UPROPERTY()
+	TMap<FGuid, FAvalancheRemoteControlValue> EntityValues;
+
+	/** Controller values. */
+	UPROPERTY()
+	TMap<FGuid, FAvalancheRemoteControlValue> ControllerValues;
+
+	/** Contains a set of entity guids that are bound to a controller action. */
+	UPROPERTY()
+	TSet<FGuid> EntitiesControlledByController;
+	
+	/** Copies the values (properties and controllers) from the given RemoteControlPreset. */
+	void CopyFrom(const URemoteControlPreset* InRemoteControlPreset, bool bInIsDefault);
+
+	/**
+	 * Compares the remote control EntityValues with another instance
+	 * @return true if the other instance has the exact same EntityValues (count and value), false otherwise
+	 */
+	bool HasSameEntityValues(const FAvalancheRemoteControlValues& InOther) const;
+
+	/**
+	 * Compares the remote control ControllerValues with another instance
+	 * @return true if the other instance has the exact same ControllerValues (count and value), false otherwise
+	 */
+	bool HasSameControllerValues(const FAvalancheRemoteControlValues& InOther) const;
+
+	/**
+	 *	Removes the extra values compared to the given reference values.
+	 * @return flags indicating what changed.
+	 */
+	EAvaRemoteControlChanges PruneRemoteControlValues(const FAvalancheRemoteControlValues& InReferenceValues);
+	
+	/**
+	 * Update the property/controller values (i.e. add missing, remove extras) from the given reference values.
+	 * If bInUpdateDefaults is true, the existing values flagged as "default" will be updated, i.e. the reference values will be applied.
+	 * Otherwise, the existing values are not modified.
+	 * Also, when adding the missing values from reference default values, the default flag is also set in the destination value. 
+	 * For a full copy of all properties and controllers, use CopyFrom instead.
+	 * @return flags indicating what changed.
+	 */
+	EAvaRemoteControlChanges UpdateRemoteControlValues(const FAvalancheRemoteControlValues& InReferenceValues, bool bInUpdateDefaults);
+	
+	/**
+	 * Refreshes the EntitiesControlledByController set.
+	 */
+	void RefreshControlledEntities(const URemoteControlPreset* InRemoteControlPreset);
+
+	/**
+	 * Set the entity value from the given preset.
+	 * @return true if the value was set.
+	 */
+	bool SetEntityValue(const FGuid& InId, const URemoteControlPreset* InRemoteControlPreset, bool bInIsDefault);
+	bool HasEntityValue(const FGuid& InId) const { return EntityValues.Contains(InId); }
+	const FAvalancheRemoteControlValue* GetEntityValue(const FGuid& InId) const { return EntityValues.Find(InId); }
+	void SetEntityValue(const FGuid& InId, const FAvalancheRemoteControlValue& InValue) {EntityValues.Add(InId, InValue);}
+	
+	bool SetControllerValue(const FGuid& InId, const URemoteControlPreset* InRemoteControlPreset, bool bInIsDefault);
+	bool HasControllerValue(const FGuid& InId) const { return ControllerValues.Contains(InId); }
+	const FAvalancheRemoteControlValue* GetControllerValue(const FGuid& InId) const { return ControllerValues.Find(InId); }
+	void SetControllerValue(const FGuid& InId, const FAvalancheRemoteControlValue& InValue) {ControllerValues.Add(InId, InValue);}
+
+	/**
+	 *	Apply the entity values to the given remote control preset.
+	 */
+	void ApplyEntityValuesToRemoteControlPreset(URemoteControlPreset* InRemoteControlPreset) const;
+
+	/**
+	 *	Apply the controller values to the given remote control preset.
+	 *	Remark: controller actions are executed by this operation.
+	 */
+	void ApplyControllerValuesToRemoteControlPreset(URemoteControlPreset* InRemoteControlPreset, bool bInForceDisableBehaviors = false) const;
+
+
+	/**
+	 * Return true if there are key collisions with the other set of values. 
+	 */
+	bool HasIdCollisions(const FAvalancheRemoteControlValues& InOtherValues) const;
+	
+	/**
+	 * @brief Merge the other values with current ones, combining the keys.
+	 * @param InOtherValues Other values to merge with.
+	 * @return true if the merge was clean with no collisions. false indicate there was some key collisions and information is lost.
+	 */
+	bool Merge(const FAvalancheRemoteControlValues& InOtherValues);
+	
+	/** Returns true if the given maps have id collisions. */
+	static bool HasIdCollisions(const TMap<FGuid, FAvalancheRemoteControlValue>& InValues, const TMap<FGuid, FAvalancheRemoteControlValue>& InOtherValues);
+	
+	static const FAvalancheRemoteControlValues& GetDefaultEmpty();
+};
diff --git a/Engine/Plugins/Experimental/Avalanche/Source/AvalancheMedia/Public/Playback/IAvaMediaPlaybackClient.h b/Engine/Plugins/Experimental/Avalanche/Source/AvalancheMedia/Public/Playback/IAvaMediaPlaybackClient.h
new file mode 100644
--- /dev/null
+++ b/Engine/Plugins/Experimental/Avalanche/Source/AvalancheMedia/Public/Playback/IAvaMediaPlaybackClient.h
@@ -0,0 +1,105 @@
+// Copyright Epic Games, Inc. All Rights Reserved.
+
+#pragma once
+
+#include "IMessageContext.h"
+#include "Playback/AvaMediaPlaybackMessages.h"
+
+class UMediaOutput;
+struct FAnimPlaySettings;
+
+class AVALANCHEMEDIA_API IAvaMediaPlaybackClient
+{
+public:
+	IAvaMediaPlaybackClient() = default;
+	virtual ~IAvaMediaPlaybackClient() = default;
+
+	virtual int32 GetNumConnectedServers() const = 0;
+	
+	virtual TArray<FString> GetServerNames() const = 0;
+	
+	/** Returns the corresponding server address for the given server name. If server name is not found, returns invalid address. */
+	virtual FMessageAddress GetServerAddress(const FString& InServerName) const = 0;
+
+	/** Returns true if the corresponding server user data for the given server name is found. */
+	virtual bool HasServerUserData(const FString& InServerName, const FString& InUserDataKey) const = 0;
+
+	/** Returns the corresponding server user data for the given server name and key. Returns empty array if not found. */
+	virtual const FString& GetServerUserData(const FString& InServerName, const FString& InUserDataKey) const = 0;
+
+	/** Returns true if the corresponding client user data for the given key is found. */
+	virtual bool HasUserData(const FString& InKey) const = 0;
+
+	/** Returns the corresponding client user data for the given key. Returns empty container if not found. */
+	virtual const FString& GetUserData(const FString& InKey) const = 0;
+
+	/** Add user data to this client. This is replicated and accessible to the servers. */
+	virtual void SetUserData(const FString& InKey, const FString& InData) = 0;
+
+	/** Remove the client's user data entry from the given key. */
+	virtual void RemoveUserData(const FString& InKey) = 0;
+
+	/** Send a stat command to all connected servers. */
+	virtual void BroadcastStatCommand(const FString& InCommand, bool bInBroadcastLocalState) = 0;
+	
+	virtual void RequestPlaybackAssetStatus(const FSoftObjectPath& InAssetPath, const FString& InChannelOrServerName, bool bInForceRefresh) = 0;
+	virtual void RequestPlayback(const FGuid& InInstanceId, const FSoftObjectPath& InAssetPath, const FString& InChannelName, EAvaMediaPlaybackAction InAction, const FString& InArguments = FString()) = 0;
+	virtual void RequestAnimPlayback(const FGuid& InInstanceId, const FSoftObjectPath& InAssetPath, const FString& InChannelName, const FAnimPlaySettings& InAnimSettings) = 0;
+	virtual void RequestAnimAction(const FGuid& InInstanceId, const FSoftObjectPath& InAssetPath, const FString& InChannelName, const FString& InAnimationName, EAvaMediaAnimAction InAction) = 0;
+	virtual void RequestRemoteControlUpdate(const FGuid& InInstanceId, const FSoftObjectPath& InAssetPath, const FString& InChannelName, const FAvalancheRemoteControlValues& InRemoteControlValues) = 0;
+	virtual void RequestPlayableTransitionStart(const FGuid& InTransitionId, TArray<FGuid>&& InEnterInstanceIds, TArray<FGuid>&& InPlayingInstanceIds, TArray<FGuid>&& InExitInstanceIds, TArray<FAvalancheRemoteControlValues>&& InEnterValues, const FName& InChannelName, EAvalanchePlayableTransitionFlags InTransitionFlags) = 0;
+	virtual void RequestPlayableTransitionStop(const FGuid& InTransitionId, const FName& InChannelName) = 0;
+	virtual void RequestBroadcast(const FString& InProfile, const FName& InChannel, const TArray<UMediaOutput*>& InRemoteMediaOutputs, EAvaMediaBroadcastAction InAction) = 0;
+	/**
+	 * Use the device provider to determine if the media output is remote.
+	 * This function should only be used as fallback if the channel's media output info is not valid.
+	 */
+	virtual bool IsMediaOutputRemoteFallback(const UMediaOutput* InMediaOutput) = 0;
+	virtual EAvaMediaIssueSeverity GetMediaOutputIssueSeverity(const FString& InServerName, const FString& InChannelName, const FGuid& InOutputGuid) const = 0;
+	virtual const TArray<FString>& GetMediaOutputIssueMessages(const FString& InServerName, const FString& InChannelName, const FGuid& InOutputGuid) const = 0;
+	virtual EAvaMediaOutputState GetMediaOutputState(const FString& InServerName, const FString& InChannelName, const FGuid& InOutputGuid) const = 0;
+
+	/**
+	 * Returns true if the client has at least one connected server for the given channel.
+	 * This does not care if the channel is live or idle. This is needed to send preload
+	 * commands to server even if they are not live yet (but connected).
+	 */
+	virtual bool HasAnyServerOnlineForChannel(const FName& InChannelName) const = 0;
+	
+	/**
+	 * Returns the status of the playback asset on the given channel on the remote server.
+	 * If the return value is not set, it is because it is not yet available, in which case
+	 * a status request can be made to make it available.
+	 * 
+	 * This supports forked channels by allowing to specify the optional server to further precise
+	 * the request. In a forked channel situation, a channel will be hosted by multiple servers
+	 * and therefore, have a playback status per server. If the server name is left empty,
+	 * the playback status of the first server hosting the given channel will be returned.
+	 *
+	 * @param InInstanceId Specify instance. If not specified, it will look for the most relevant status for the given asset.
+	 * @param InAssetPath Full path of the playback asset to fetch the status of.
+	 * @param InChannelName Servers can be running different channels, so need to identify the channel as well.
+	 * @param InServerName Identifies the server to fetch the status from.
+	 */
+	virtual TOptional<EAvaMediaPlaybackStatus> GetRemotePlaybackStatus(const FGuid& InInstanceId, const FSoftObjectPath& InAssetPath, const FString& InChannelName, const FString& InServerName = FString()) const = 0;
+
+	/**
+	 * @brief Returns the user data of the specified playback instance.
+	* @param InInstanceId Specify instance of the given asset.
+	 * @param InAssetPath Full path of the playback asset. Must be valid.
+	 * @param InChannelName Servers can be running different channels, so need to identify the channel as well.
+	 * @param InServerName Identifies the server to fetch the status from.
+	 * @return 
+	 */
+	virtual const FString* GetRemotePlaybackUserData(const FGuid& InInstanceId, const FSoftObjectPath& InAssetPath, const FString& InChannelName, const FString& InServerName = FString()) const = 0;
+
+	/**
+	 *	Returns the status of the playback asset on the given remote server.
+	 *	If the return value is not set, it is because it is not yet available, in which case
+	 *	a status request must be made explicitly, using RequestPlaybackAssetStatus to make it available.
+	 *
+	 * @param InAssetPath Full path of the playback asset to fetch the status of.
+	 * @param InServerName Identifies the server to fetch the status from.
+	 */
+	virtual TOptional<EAvaMediaPlaybackAssetStatus> GetRemotePlaybackAssetStatus(const FSoftObjectPath& InAssetPath, const FString& InServerName) const = 0;
+};
diff --git a/Engine/Plugins/Experimental/Avalanche/Source/AvalancheMedia/Public/Playback/IAvaPlaybackEditor.h b/Engine/Plugins/Experimental/Avalanche/Source/AvalancheMedia/Public/Playback/IAvaPlaybackEditor.h
new file mode 100644
--- /dev/null
+++ b/Engine/Plugins/Experimental/Avalanche/Source/AvalancheMedia/Public/Playback/IAvaPlaybackEditor.h
@@ -0,0 +1,48 @@
+// Copyright Epic Games, Inc. All Rights Reserved.
+
+#pragma once
+
+#if WITH_EDITOR
+
+#include "CoreMinimal.h"
+
+class UAvalanchePlayback;
+class UAvaPlaybackNode;
+class UEdGraph;
+class UEdGraphNode;
+class FSlateRect;
+
+class IAvaPlaybackGraphEditor
+{
+public:
+	
+	virtual ~IAvaPlaybackGraphEditor() = default;
+	
+	virtual UEdGraph* CreatePlaybackGraph(UAvalanchePlayback* InPlayback) = 0;
+	
+	virtual void SetupPlaybackNode(UEdGraph* InGraph, UAvaPlaybackNode* InPlaybackNode, bool bSelectNewNode) = 0;
+	
+	/** Compiles Playback nodes from graph nodes. */
+	virtual void CompilePlaybackNodesFromGraphNodes(UAvalanchePlayback* InPlayback) = 0;
+	
+	virtual void CreateInputPin(UEdGraphNode* InGraphNode) = 0;
+
+	virtual void RefreshNode(UEdGraphNode& InGraphNode) = 0;
+
+	/** Get the bounding area for the currently selected nodes
+ 	*
+ 	* @param Rect Final output bounding area, including padding
+ 	* @param Padding An amount of padding to add to all sides of the bounds
+ 	*
+ 	* @return false if nothing is selected
+	*/
+	virtual bool GetBoundsForSelectedNodes(FSlateRect& Rect, float Padding) = 0;
+
+	virtual TSet<UObject*> GetSelectedNodes() const = 0;
+
+	virtual bool CanPasteNodes() const = 0;
+	virtual void PasteNodesHere(const FVector2D& Location) = 0;
+	
+};
+
+#endif
diff --git a/Engine/Plugins/Experimental/Avalanche/Source/AvalancheMedia/Public/Playback/Nodes/AvaPlaybackNodeBlueprintPlayer.h b/Engine/Plugins/Experimental/Avalanche/Source/AvalancheMedia/Public/Playback/Nodes/AvaPlaybackNodeBlueprintPlayer.h
new file mode 100644
--- /dev/null
+++ b/Engine/Plugins/Experimental/Avalanche/Source/AvalancheMedia/Public/Playback/Nodes/AvaPlaybackNodeBlueprintPlayer.h
@@ -0,0 +1,35 @@
+// Copyright Epic Games, Inc. All Rights Reserved.
+
+#pragma once
+
+#include "CoreMinimal.h"
+#include "AvaPlaybackNodePlayer.h"
+#include "AvaPlaybackNodeBlueprintPlayer.generated.h"
+
+class UAvalancheBlueprint;
+
+UCLASS()
+class AVALANCHEMEDIA_API UAvaPlaybackNodeBlueprintPlayer : public UAvaPlaybackNodePlayer
+{
+	GENERATED_BODY()
+
+public:
+	UAvaPlaybackNodeBlueprintPlayer();
+
+	virtual void RefreshNode(bool bDryRunGraph) override;
+	virtual void PostLoad() override;
+	
+	TSoftObjectPtr<UAvalancheBlueprint> GetAvalancheAsset() const { return BlueprintAsset; }
+	void SetAvalancheAsset(const TSoftObjectPtr<UAvalancheBlueprint>& InAsset);
+	
+	void UpdateDisplayNameText();
+
+	virtual const FSoftObjectPath& GetAvalancheAssetPath() const override { return BlueprintAsset.ToSoftObjectPath();}
+	
+	virtual FAvaSoftAssetPtr GetAvalancheAssetPtr() const override;
+	
+protected:
+	
+	UPROPERTY(EditAnywhere, Category = "Motion Design")
+	TSoftObjectPtr<UAvalancheBlueprint> BlueprintAsset;
+};
diff --git a/Engine/Plugins/Experimental/Avalanche/Source/AvalancheMedia/Public/Playback/Nodes/AvaPlaybackNodeLevelPlayer.h b/Engine/Plugins/Experimental/Avalanche/Source/AvalancheMedia/Public/Playback/Nodes/AvaPlaybackNodeLevelPlayer.h
new file mode 100644
--- /dev/null
+++ b/Engine/Plugins/Experimental/Avalanche/Source/AvalancheMedia/Public/Playback/Nodes/AvaPlaybackNodeLevelPlayer.h
@@ -0,0 +1,37 @@
+// Copyright Epic Games, Inc. All Rights Reserved.
+
+#pragma once
+
+#include "CoreMinimal.h"
+#include "AvaPlaybackNodePlayer.h"
+#include "AvaPlaybackNodeLevelPlayer.generated.h"
+
+class UWorld;
+
+UCLASS()
+class AVALANCHEMEDIA_API UAvaPlaybackNodeLevelPlayer : public UAvaPlaybackNodePlayer
+{
+	GENERATED_BODY()
+
+public:
+	
+	UAvaPlaybackNodeLevelPlayer();
+	
+	virtual void RefreshNode(bool bDryRunGraph) override;
+	virtual void PostLoad() override;
+	
+	TSoftObjectPtr<UWorld> GetAvalancheAsset() const { return LevelAsset; }
+	void SetAvalancheAsset(const TSoftObjectPtr<UWorld>& InAsset);
+	
+	void UpdateDisplayNameText();
+
+	virtual const FSoftObjectPath& GetAvalancheAssetPath() const override { return LevelAsset.ToSoftObjectPath();}
+
+	virtual FAvaSoftAssetPtr GetAvalancheAssetPtr() const override;
+
+protected:
+	
+	UPROPERTY(EditAnywhere, Category = "Motion Design")
+	TSoftObjectPtr<UWorld> LevelAsset;
+};
+
diff --git a/Engine/Plugins/Experimental/Avalanche/Source/AvalancheMedia/Public/Playback/Nodes/Events/Actions/AvalancheAnimations.h b/Engine/Plugins/Experimental/Avalanche/Source/AvalancheMedia/Public/Playback/Nodes/Events/Actions/AvalancheAnimations.h
new file mode 100644
--- /dev/null
+++ b/Engine/Plugins/Experimental/Avalanche/Source/AvalancheMedia/Public/Playback/Nodes/Events/Actions/AvalancheAnimations.h
@@ -0,0 +1,84 @@
+// Copyright Epic Games, Inc. All Rights Reserved.
+
+#pragma once
+
+#include "CoreMinimal.h"
+#include "AvaSequencePlaybackObject.h"
+#include "AvaSequenceShared.h"
+#include "Playback/Nodes/Events/AvaPlaybackNodeAction.h"
+#include "AvalancheAnimations.generated.h"
+
+UENUM()
+enum class EAvaMediaAnimAction
+{
+	None,
+	Play,
+	Continue,
+	Stop,
+	PreviewFrame,
+	CameraCut
+};
+
+
+USTRUCT()
+struct FAnimPlaySettings
+{
+	GENERATED_BODY()
+
+	FAnimPlaySettings() : FAnimPlaySettings(NAME_None) {} 
+	FAnimPlaySettings(FName InAnimationName) : AnimationName(InAnimationName) {}
+	
+	UPROPERTY(EditAnywhere, Category = "Motion Design")
+	EAvaMediaAnimAction Action = EAvaMediaAnimAction::None;
+	
+	UPROPERTY(VisibleAnywhere, Category = "Motion Design")
+	FName AnimationName = NAME_None;
+
+	UPROPERTY(EditAnywhere, Category = "Motion Design")
+	float StartAtTime = 0.0f;
+
+	UPROPERTY(EditAnywhere, Category = "Motion Design")
+	int32 LoopCount = 0;
+
+	UPROPERTY(EditAnywhere, Category = "Motion Design")
+	EAvaSequencePlayMode PlayMode = EAvaSequencePlayMode::Forward;
+
+	UPROPERTY(EditAnywhere, Category = "Motion Design")
+	float PlaybackSpeed = 1.0f;
+	
+	UPROPERTY(EditAnywhere, Category = "Motion Design")
+	bool bRestoreState = false;
+
+	friend uint32 GetTypeHash(const FAnimPlaySettings& InPlaySettings)
+	{
+		return GetTypeHash(InPlaySettings.AnimationName);
+	}
+
+	bool operator==(const FAnimPlaySettings& Other) const
+	{
+		return AnimationName == Other.AnimationName;
+	}
+
+	FAvaSequencePlayParams AsPlayParams() const
+	{
+		FAvaSequencePlayParams PlayParams;
+
+		PlayParams.Start    = FAvaSequenceTime(StartAtTime);
+		PlayParams.PlayMode = PlayMode;
+
+		PlayParams.AdvancedSettings.PlaybackSpeed = PlaybackSpeed;
+		PlayParams.AdvancedSettings.LoopCount     = LoopCount;
+		PlayParams.AdvancedSettings.bRestoreState = bRestoreState;
+
+		return PlayParams;
+	}
+};
+
+USTRUCT()
+struct FAvalancheAnimations
+{
+	GENERATED_BODY()
+	
+	UPROPERTY(EditAnywhere, Category = "Motion Design")
+	TSet<FAnimPlaySettings> AvailableAnimations;
+};
\ No newline at end of file
diff --git a/Engine/Plugins/Experimental/Avalanche/Source/AvalancheMedia/Public/Playback/Playables/AvalancheBlueprintPlayable.h b/Engine/Plugins/Experimental/Avalanche/Source/AvalancheMedia/Public/Playback/Playables/AvalancheBlueprintPlayable.h
new file mode 100644
--- /dev/null
+++ b/Engine/Plugins/Experimental/Avalanche/Source/AvalancheMedia/Public/Playback/Playables/AvalancheBlueprintPlayable.h
@@ -0,0 +1,46 @@
+// Copyright Epic Games, Inc. All Rights Reserved.
+
+#pragma once
+
+#include "AvaMediaDefines.h"
+#include "Playback/AvalanchePlayable.h"
+
+#include "AvalancheBlueprintPlayable.generated.h"
+
+UCLASS(NotBlueprintable, BlueprintType)
+class AVALANCHEMEDIA_API UAvalancheBlueprintPlayable : public UAvalanchePlayable
+{
+	GENERATED_BODY()
+
+public:
+	//~ Begin UAvalanchePlayable
+	virtual bool LoadAsset(const FAvaSoftAssetPtr& InAvalancheSourceAsset, bool bInInitiallyVisible) override;
+	virtual bool UnloadAsset() override;
+	virtual const FSoftObjectPath& GetSourceAssetPath() const override;
+	virtual EAvalanchePlayableStatus GetPlayableStatus() const override;
+	virtual IAvaSceneInterface* GetSceneInterface() const override;
+	virtual EAvalanchePlayableCommandResult ExecuteAnimationCommand(EAvaMediaAnimAction InAnimAction, const FAnimPlaySettings& InAnimPlaySettings) override;
+
+protected:
+	virtual bool InitPlayable(const FPlayableCreationInfo& InPlayableInfo) override;
+	virtual void OnPlay() override;
+	virtual void OnEndPlay() override;
+	//~ End UAvalanchePlayable
+
+public:
+	UAvalancheBlueprint* GetManagedAvalancheBlueprint() const { return ManagedAvalancheBlueprint; }
+	
+protected:
+	bool LoadAvalancheBlueprintInternal(const TSoftObjectPtr<UAvalancheBlueprint>& InSourceAvalancheBlueprint);
+
+protected:
+	UPROPERTY(Transient)
+	TSoftObjectPtr<UAvalancheBlueprint> SourceAvalancheBlueprint;
+
+	/*
+	 * The Managed Avalanche, a Duplication of the Source Avalanche when Begin Play was called.
+	 * This shouldn't be GCd while in play, as it contains Animation and other data used during play.
+	 */
+	UPROPERTY(Transient)
+	TObjectPtr<UAvalancheBlueprint> ManagedAvalancheBlueprint;
+};
\ No newline at end of file
diff --git a/Engine/Plugins/Experimental/Avalanche/Source/AvalancheMedia/Public/Playback/Playables/AvalancheLevelStreamingPlayable.h b/Engine/Plugins/Experimental/Avalanche/Source/AvalancheMedia/Public/Playback/Playables/AvalancheLevelStreamingPlayable.h
new file mode 100644
--- /dev/null
+++ b/Engine/Plugins/Experimental/Avalanche/Source/AvalancheMedia/Public/Playback/Playables/AvalancheLevelStreamingPlayable.h
@@ -0,0 +1,108 @@
+// Copyright Epic Games, Inc. All Rights Reserved.
+
+#pragma once
+
+#include "AvaMediaDefines.h"
+#include "Playback/AvalanchePlayable.h"
+
+#include "AvalancheLevelStreamingPlayable.generated.h"
+
+class AActor;
+class AAvaScene;
+class ULevel;
+class ULevelStreaming;
+class ULevelStreamingDynamic;
+enum class ELevelStreamingState : uint8;
+
+UCLASS(NotBlueprintable, BlueprintType)
+class AVALANCHEMEDIA_API UAvalancheLevelStreamingPlayable : public UAvalanchePlayable
+{
+	GENERATED_BODY()
+public:
+	//~ Begin UAvalanchePlayable
+	virtual bool LoadAsset(const FAvaSoftAssetPtr& InAvalancheSourceAsset, bool bInInitiallyVisible) override;
+	virtual bool UnloadAsset() override;
+	virtual const FSoftObjectPath& GetSourceAssetPath() const override { return SourceAvalancheLevel.ToSoftObjectPath(); }
+	virtual EAvalanchePlayableStatus GetPlayableStatus() const override;
+	virtual IAvaSceneInterface* GetSceneInterface() const override;
+	virtual bool ApplyCamera() override;
+	virtual bool GetShouldBeVisible() const override;
+	virtual void SetShouldBeVisible(bool bInShouldBeVisible) override;
+
+protected:
+	virtual bool InitPlayable(const FPlayableCreationInfo& InPlayableInfo) override;
+	virtual void OnPlay() override;
+	virtual void OnEndPlay() override;
+	//~ End UAvalanchePlayable
+
+	//~ Begin UObject
+	virtual void BeginDestroy() override;
+	//~ End UObject
+
+public:
+	ULevelStreamingDynamic* GetLevelStreaming() const { return LevelStreaming; }
+
+protected:
+	bool LoadAvalancheLevel(const TSoftObjectPtr<UWorld>& InSourceAvalancheLevel, bool bInInitiallyVisible);
+
+	void OnLevelStreamingStateChanged(UWorld* InWorld
+		, const ULevelStreaming* InLevelStreaming
+		, ULevel* InLevelIfLoaded
+		, ELevelStreamingState InPreviousState
+		, ELevelStreamingState InNewState);
+
+	void OnLevelStreamingPlayableStatusChanged(UAvalancheLevelStreamingPlayable* InSubPlayable);
+	
+	void BindDelegates();
+	void UnbindDelegates();
+	ULevel* GetLoadedLevel() const;
+	void ResolveAvalancheScene(const ULevel* InLevel);
+	AActor* FindStartupCameraActor(FName InStartupCameraName, UAvalanchePlayable** OutParentPlayable);
+
+	void LoadSubPlayables(const UWorld* InLevelInstance);
+	void UnloadSubPlayables();
+	
+	void GetOrLoadSubPlayable(const ULevelStreaming* InLevelStreaming);
+
+	/**
+	 * Creates a level streaming playable from the given level streaming information.
+	 * A new level streaming object is created, wrapped in the returned playable.
+	 */
+	static UAvalancheLevelStreamingPlayable* CreateSubPlayable(UAvaMediaPlayableGroup* InPlayableGroup, const FSoftObjectPath& InSourceAssetPath);
+
+	void AddSubPlayable(UAvalancheLevelStreamingPlayable* InSubPlayable);
+	void RemoveSubPlayable(UAvalancheLevelStreamingPlayable* InSubPlayable);
+
+	/** For shared playables (loaded through streaming dependencies), returns true if the playable is still part of other dependencies. */
+	bool HasParentPlayables() const;
+
+	void UpdateVisibilityFromParents();
+
+protected:
+	UPROPERTY(Transient)
+	TSoftObjectPtr<UWorld> SourceAvalancheLevel;
+	
+	UPROPERTY(Transient)
+	TObjectPtr<ULevelStreamingDynamic> LevelStreaming;
+
+	UPROPERTY(Transient)
+	TObjectPtr<AAvaScene> AvalancheScene; 
+
+	bool bLoadSubPlayables = false;
+	
+	/**
+	 * Dependent playables loaded from secondary streaming levels.
+	 * Those playables are shared by the parent playable(s) and will be unloaded
+	 * when the last parent playable is unloaded.
+	 */
+	UPROPERTY(Transient)
+	TArray<TObjectPtr<UAvalancheLevelStreamingPlayable>> SubPlayables;
+
+	/**
+	 * Keep track of the dependencies this playable is part of.
+	 * This helps determine when the playable should be unloaded.
+	 */
+	TSet<TObjectKey<UAvalancheLevelStreamingPlayable>> ParentPlayables;
+
+	bool bOnPlayQueued = false;
+};
diff --git a/Engine/Plugins/Experimental/Avalanche/Source/AvalancheMedia/Public/Playback/Playables/AvalancheRemoteProxyPlayable.h b/Engine/Plugins/Experimental/Avalanche/Source/AvalancheMedia/Public/Playback/Playables/AvalancheRemoteProxyPlayable.h
new file mode 100644
--- /dev/null
+++ b/Engine/Plugins/Experimental/Avalanche/Source/AvalancheMedia/Public/Playback/Playables/AvalancheRemoteProxyPlayable.h
@@ -0,0 +1,58 @@
+// Copyright Epic Games, Inc. All Rights Reserved.
+
+#pragma once
+
+#include "Playback/AvalanchePlayable.h"
+
+#include "AvalancheRemoteProxyPlayable.generated.h"
+
+class IAvaMediaPlaybackClient;
+namespace UE::AvaMediaPlaybackClient::Delegates
+{
+	struct FPlaybackSequenceEventArgs;
+}
+
+UCLASS(NotBlueprintable, BlueprintType)
+class AVALANCHEMEDIA_API UAvalancheRemoteProxyPlayable : public UAvalanchePlayable
+{
+	GENERATED_BODY()
+public:
+	FName GetPlayingChannelFName() const { return PlayingChannelFName; }
+	
+	//~ Begin UAvalanchePlayable
+	virtual bool LoadAsset(const FAvaSoftAssetPtr& InAvalancheSourceAsset, bool bInInitiallyVisible) override;
+	virtual bool UnloadAsset() override;
+	virtual const FSoftObjectPath& GetSourceAssetPath() const override { return SourceAssetPath; }
+	virtual EAvalanchePlayableStatus GetPlayableStatus() const override;
+	virtual IAvaSceneInterface* GetSceneInterface() const override;
+	virtual EAvalanchePlayableCommandResult ExecuteAnimationCommand(EAvaMediaAnimAction InAnimAction, const FAnimPlaySettings& InAnimPlaySettings) override;
+	virtual EAvalanchePlayableCommandResult UpdateRemoteControlCommand(const TSharedRef<FAvalancheRemoteControlValues>& InRemoteControlValues) override;
+	virtual bool ApplyCamera() override;
+	virtual bool IsRemoteProxy() const override { return true; }
+	virtual void SetUserData(const FString& InUserData) override;
+
+protected:
+	virtual bool InitPlayable(const FPlayableCreationInfo& InPlayableInfo) override;
+	virtual void OnPlay() override;
+	virtual void OnEndPlay() override;
+	//~ End UAvalanchePlayable
+
+	//~ Begin UObject
+	virtual void BeginDestroy() override;
+	//~ End UObject
+
+	void RegisterClientEventHandlers();
+	void UnregisterClientEventHandlers() const;
+
+	static TArray<FString> GetOnlineServerForChannel(const FName& InChannelName);
+	
+	void HandleAvaMediaPlaybackSequenceEvent(IAvaMediaPlaybackClient& InPlaybackClient,
+		const UE::AvaMediaPlaybackClient::Delegates::FPlaybackSequenceEventArgs& InEventArgs);
+	
+protected:
+	/** Channel name this playable is playing on. */
+	FName PlayingChannelFName;
+	FString PlayingChannelName;
+
+	FSoftObjectPath SourceAssetPath;
+};
diff --git a/Engine/Plugins/Experimental/Avalanche/Source/AvalancheMedia/Public/Playback/Transition/AvaMediaPlaybackTransition.h b/Engine/Plugins/Experimental/Avalanche/Source/AvalancheMedia/Public/Playback/Transition/AvaMediaPlaybackTransition.h
new file mode 100644
--- /dev/null
+++ b/Engine/Plugins/Experimental/Avalanche/Source/AvalancheMedia/Public/Playback/Transition/AvaMediaPlaybackTransition.h
@@ -0,0 +1,36 @@
+// Copyright Epic Games, Inc. All Rights Reserved.
+
+#pragma once
+
+#include "UObject/Object.h"
+#include "Playback/IAvaPlayableVisibilityConstraint.h"
+
+#include "AvaMediaPlaybackTransition.generated.h"
+
+/**
+ * Abstract base class for playback transitions that can be queued in
+ * the playback manager's commands.
+ */
+UCLASS()
+class AVALANCHEMEDIA_API UAvaMediaPlaybackTransition : public UObject, public IAvaPlayableVisibilityConstraint
+{
+	GENERATED_BODY()
+	
+public:
+	/**
+	 * @brief Evaluate the status of loading playables to determine if the transition can start.
+	 * @param bOutShouldDiscard indicate if the pending start transition command should be discarded.
+	 * @return true if the transition can start, false otherwise.
+	 */
+	virtual bool CanStart(bool& bOutShouldDiscard) const { bOutShouldDiscard = true; return false; }
+
+	virtual void Start() {}
+
+	virtual void Stop() {}
+
+	virtual bool IsRunning() const { return false; }
+	
+	//~ Begin IAvaPlayableVisibilityConstraint
+	virtual bool IsVisibilityConstrained(const UAvalanchePlayable* InPlayable) const override { return false; }
+	//~ End IAvaPlayableVisibilityConstraint
+};
\ No newline at end of file
diff --git a/Engine/Plugins/Experimental/Avalanche/Source/AvalancheMedia/Public/Playback/Transition/AvaTransitionPlayableScene.h b/Engine/Plugins/Experimental/Avalanche/Source/AvalancheMedia/Public/Playback/Transition/AvaTransitionPlayableScene.h
new file mode 100644
--- /dev/null
+++ b/Engine/Plugins/Experimental/Avalanche/Source/AvalancheMedia/Public/Playback/Transition/AvaTransitionPlayableScene.h
@@ -0,0 +1,33 @@
+// Copyright Epic Games, Inc. All Rights Reserved.
+
+#pragma once
+
+#include "AvaTagHandle.h"
+#include "AvaTransitionScene.h"
+#include "AvaTransitionPlayableScene.generated.h"
+
+class UAvalanchePlayable;
+class UAvalanchePlayableTransition;
+
+USTRUCT()
+struct AVALANCHEMEDIA_API FAvaTransitionPlayableScene : public FAvaTransitionScene
+{
+	GENERATED_BODY()
+
+	FAvaTransitionPlayableScene() = default;
+
+	explicit FAvaTransitionPlayableScene(UAvalanchePlayable* InPlayable, UAvalanchePlayableTransition* InPlayableTransition);
+
+	explicit FAvaTransitionPlayableScene(const FAvaTagHandle& InTransitionLayer, UAvalanchePlayableTransition* InPlayableTransition);
+
+	//~ Begin FAvaTransitionScene
+	virtual EAvaTransitionComparisonResult Compare(const FAvaTransitionScene& InOther) const override;
+	virtual ULevel* GetLevel() const override;
+	virtual void GetOverrideTransitionLayer(FAvaTagHandle& OutTransitionLayer) const override;
+	virtual void OnFlagsChanged() override;
+	//~ End FAvaTransitionScene
+
+	TWeakObjectPtr<UAvalanchePlayableTransition> PlayableTransitionWeak;
+
+	TOptional<FAvaTagHandle> OverrideTransitionLayer;
+};
diff --git a/Engine/Plugins/Experimental/Avalanche/Source/AvalancheMedia/Public/Playback/Transition/AvalanchePlayableTransition.h b/Engine/Plugins/Experimental/Avalanche/Source/AvalancheMedia/Public/Playback/Transition/AvalanchePlayableTransition.h
new file mode 100644
--- /dev/null
+++ b/Engine/Plugins/Experimental/Avalanche/Source/AvalancheMedia/Public/Playback/Transition/AvalanchePlayableTransition.h
@@ -0,0 +1,108 @@
+// Copyright Epic Games, Inc. All Rights Reserved.
+
+#pragma once
+
+#include "AvaMediaDefines.h"
+#include "Playback/AvalancheRemoteControlValues.h"
+#include "UObject/Object.h"
+
+#include "AvalanchePlayableTransition.generated.h"
+
+class UAvalanchePlayable;
+class UAvaMediaPlayableGroup;
+
+/**
+ * @brief Defines the playable entry role in the transition.
+ */
+enum class EAvaMediaPlayableTransitionEntryRole : uint8
+{
+	/** The playable corresponding to this entry is a entering the scene. */
+	Enter,
+	/** The playable corresponding to this entry is already in the scene and may react to the transition, but is otherwise neutral. */
+	Playing,
+	/** The playable corresponding to this entry is already in the scene but is commanded to exit. */
+	Exit
+};
+
+UCLASS()
+class AVALANCHEMEDIA_API UAvalanchePlayableTransition : public UObject
+{
+	GENERATED_BODY()
+	
+public:
+	virtual bool Start();
+	virtual void Stop();
+	virtual bool IsRunning() const { return false; }
+	virtual void Tick(double InDeltaSeconds) {}
+
+	void SetTransitionFlags(EAvalanchePlayableTransitionFlags InFlags);
+	void SetEnterPlayables(TArray<TWeakObjectPtr<UAvalanchePlayable>>&& InPlayablesWeak);
+	void SetPlayingPlayables(TArray<TWeakObjectPtr<UAvalanchePlayable>>&& InPlayablesWeak);
+	void SetExitPlayables(TArray<TWeakObjectPtr<UAvalanchePlayable>>&& InPlayablesWeak);
+	
+	bool IsEnterPlayable(UAvalanchePlayable* InPlayable) const;
+	bool IsPlayingPlayable(UAvalanchePlayable* InPlayable) const;
+	bool IsExitPlayable(UAvalanchePlayable* InPlayable) const;
+
+	void SetEnterPlayableValues(TArray<TSharedPtr<FAvalancheRemoteControlValues>>&& InPlayableValues);
+	
+	/** This is called during the transition evaluation to indicate discarded playables. */
+	void MarkPlayableAsDiscard(UAvalanchePlayable* InPlayable);
+
+	/** Returns information on this transition suitable for logging. */
+	virtual FString GetPrettyInfo() const;
+
+	EAvalanchePlayableTransitionFlags GetTransitionFlags() const { return TransitionFlags; }
+	
+protected:
+	UAvalanchePlayable* FindPlayable(const FGuid& InstanceId) const;
+	
+protected:
+	EAvalanchePlayableTransitionFlags TransitionFlags = EAvalanchePlayableTransitionFlags::None;
+	
+	TArray<TSharedPtr<FAvalancheRemoteControlValues>> EnterPlayableValues;
+	
+	TArray<TWeakObjectPtr<UAvalanchePlayable>> EnterPlayablesWeak;
+	TArray<TWeakObjectPtr<UAvalanchePlayable>> PlayingPlayablesWeak;
+	TArray<TWeakObjectPtr<UAvalanchePlayable>> ExitPlayablesWeak;
+
+	/** Keep track of the discarded playables so events can be sent when the transition ends. */
+	TArray<TWeakObjectPtr<UAvalanchePlayable>> DiscardPlayablesWeak;
+
+	TSet<TWeakObjectPtr<UAvaMediaPlayableGroup>> PlayableGroupsWeak;
+};
+
+class AVALANCHEMEDIA_API FAvaPlayableTransitionBuilder
+{
+public:
+	FAvaPlayableTransitionBuilder();
+
+	void AddEnterPlayableValues(const TSharedPtr<FAvalancheRemoteControlValues>& InValues);
+	
+	bool AddEnterPlayable(UAvalanchePlayable* InPlayable);
+	bool AddPlayingPlayable(UAvalanchePlayable* InPlayable);
+	bool AddExitPlayable(UAvalanchePlayable* InPlayable);
+	
+	bool AddPlayable(UAvalanchePlayable* InPlayable, EAvaMediaPlayableTransitionEntryRole InPlayableRole)
+	{
+		switch(InPlayableRole)
+		{
+			case EAvaMediaPlayableTransitionEntryRole::Enter:
+				return AddEnterPlayable(InPlayable);
+			case EAvaMediaPlayableTransitionEntryRole::Playing:
+				return AddPlayingPlayable(InPlayable);
+			case EAvaMediaPlayableTransitionEntryRole::Exit:
+				return AddExitPlayable(InPlayable);
+		}
+		return false;
+	}
+
+	UAvalanchePlayableTransition* MakeTransition(UObject* InOuter);
+
+private:
+	TArray<TSharedPtr<FAvalancheRemoteControlValues>> EnterPlayableValues;
+
+	TArray<TWeakObjectPtr<UAvalanchePlayable>> EnterPlayablesWeak;
+	TArray<TWeakObjectPtr<UAvalanchePlayable>> PlayingPlayablesWeak;
+	TArray<TWeakObjectPtr<UAvalanchePlayable>> ExitPlayablesWeak;
+};
\ No newline at end of file
diff --git a/Engine/Plugins/Experimental/Avalanche/Source/AvalancheMedia/Public/Playlist/AvalancheManagedInstance.h b/Engine/Plugins/Experimental/Avalanche/Source/AvalancheMedia/Public/Playlist/AvalancheManagedInstance.h
new file mode 100644
--- /dev/null
+++ b/Engine/Plugins/Experimental/Avalanche/Source/AvalancheMedia/Public/Playlist/AvalancheManagedInstance.h
@@ -0,0 +1,54 @@
+// Copyright Epic Games, Inc. All Rights Reserved.
+
+#pragma once
+
+#include "AvaType.h"
+#include "Playback/AvalancheRemoteControlValues.h"
+#include "UObject/GCObject.h"
+#include "UObject/SoftObjectPath.h"
+
+class FAvalancheManagedInstanceCache;
+class IAvaSceneInterface;
+
+class AVALANCHEMEDIA_API FAvalancheManagedInstance : public IAvaTypeCastable, public FGCObject, public TSharedFromThis<FAvalancheManagedInstance>
+{
+public:
+	UE_AVA_INHERITS(FAvalancheManagedInstance, IAvaTypeCastable);
+	
+	FAvalancheManagedInstance(FAvalancheManagedInstanceCache* InParentCache, const FSoftObjectPath& InAssetPath);
+	virtual ~FAvalancheManagedInstance() override;
+
+	/** Returns true if the instanced asset is valid. */
+	virtual bool IsValid() const { return false; }
+	
+	/** Get the Managed Avalanche Asset's (top level) Remote Control Preset. */
+	virtual URemoteControlPreset* GetRemoteControlPreset() const { return nullptr; }
+	
+	/** Access a copy of the Remote Control values from the source (a.k.a. template). */
+	const FAvalancheRemoteControlValues& GetDefaultRemoteControlValues() const { return DefaultRemoteControlValues; }
+
+	/** Get the Ava scene interface from the asset. */
+	virtual IAvaSceneInterface* GetSceneInterface() const { return nullptr; }
+
+	const FSoftObjectPath& GetSourceAssetPath() const { return SourceAssetPath; }
+	
+protected:
+	void InvalidateFromCache();
+
+	void RegisterSourceRemoteControlPresetDelegates(URemoteControlPreset* InSourceRemoteControlPreset);
+	void UnregisterSourceRemoteControlPresetDelegates(URemoteControlPreset* InSourceRemoteControlPreset) const;
+	
+	void OnSourceRemoteControlEntitiesExposed(URemoteControlPreset* InSourcePreset, const FGuid& InEntityId);
+	void OnSourceRemoteControlEntitiesUnexposed(URemoteControlPreset* InSourcePreset, const FGuid& InEntityId);
+	void OnSourceRemoteControlEntitiesUpdated(URemoteControlPreset* InSourcePreset, const TSet<FGuid>& InModifiedEntities);
+	void OnSourceRemoteControlExposedPropertiesModified(URemoteControlPreset* InSourcePreset, const TSet<FGuid>& InModifiedProperties);
+	void OnSourceRemoteControlControllerAdded(URemoteControlPreset* InSourcePreset, const FName NewControllerName, const FGuid& InControllerId);
+	void OnSourceRemoteControlControllerRemoved(URemoteControlPreset* InSourcePreset, const FGuid& InControllerId);
+	void OnSourceRemoteControlControllerRenamed(URemoteControlPreset* InSourcePreset, const FName InOldLabel, const FName InNewLabel);
+	void OnSourceRemoteControlControllerModified(URemoteControlPreset* InSourcePreset, const TSet<FGuid>& InModifiedControllerIds);
+
+protected:
+	FAvalancheManagedInstanceCache* ParentCache;
+	FSoftObjectPath SourceAssetPath;
+	FAvalancheRemoteControlValues DefaultRemoteControlValues;
+};
\ No newline at end of file
diff --git a/Engine/Plugins/Experimental/Avalanche/Source/AvalancheMedia/Public/Playlist/AvalancheManagedInstanceCache.h b/Engine/Plugins/Experimental/Avalanche/Source/AvalancheMedia/Public/Playlist/AvalancheManagedInstanceCache.h
new file mode 100644
--- /dev/null
+++ b/Engine/Plugins/Experimental/Avalanche/Source/AvalancheMedia/Public/Playlist/AvalancheManagedInstanceCache.h
@@ -0,0 +1,87 @@
+// Copyright Epic Games, Inc. All Rights Reserved.
+
+#pragma once
+
+#include "Playlist/AvalancheManagedInstance.h"
+#include "UObject/ObjectSaveContext.h"
+
+class IAvaMediaSyncProvider;
+class UPackage;
+
+/**
+ * Cache for the managed Avalanche Blueprint instances.
+ *
+ * This is LRU cache. 
+ */
+class AVALANCHEMEDIA_API FAvalancheManagedInstanceCache
+{
+public:
+	FAvalancheManagedInstanceCache();
+	virtual ~FAvalancheManagedInstanceCache();
+
+	/**
+	 *	Retrieves elements from the cache or load a new one if not present.
+	 *	If the cache size is exceeded, the oldest elements are going to be flushed.
+	 */
+	TSharedPtr<FAvalancheManagedInstance> GetOrLoadAvalancheInstance(const FSoftObjectPath& InAssetPath);
+
+	/**
+	 * Invalidates a cached entry without deleting it immediately.
+	 * It will be deleted and reloaded on the next access query.
+	 */
+	void InvalidateNoDelete(const FSoftObjectPath& InAssetPath);
+
+	/**
+	 * Invalidates a cached entry. The entry may be deleted as a result.
+	 */
+	void Invalidate(const FSoftObjectPath& InAssetPath);
+
+	/**
+	 *	Delegate called when an entry is invalidated. Reason for an entry to be invalidated is if the source
+	 *	has been modified.
+	 */
+	DECLARE_MULTICAST_DELEGATE_OneParam(FOnEntryInvalidated, const FSoftObjectPath& /*InAssetPath*/);
+	FOnEntryInvalidated OnEntryInvalidated;
+	
+	/**
+	 * Get the maximum size of the cache beyond witch it will start flushing elements.
+	 */
+	int32 GetMaximumCacheSize() const;
+	
+	/**
+	 *	Flush specified entry from the cache.
+	 */
+	void Flush(const FSoftObjectPath& InAssetPath);
+
+	/**
+	 * Flush all unused entries from the cache.
+	 */
+	void Flush();
+
+	/**
+	 * Trim the cache elements that exceed the cache capacity according to LRU replacement policy.
+	 */
+	void TrimCache();
+
+	/**
+	 *	Perform any pending actions, such has removing invalidated entries. This may lead to
+	 *	objects being deleted.
+	 */
+	void FinishPendingActions();
+
+private:
+	void RemovePendingInvalidatedPaths();
+	void RemoveEntry(const FSoftObjectPath& InAssetPath);
+	void RemoveEntries(TFunctionRef<bool(const FSoftObjectPath&, const TSharedPtr<FAvalancheManagedInstance>&)> InRemovePredicate, bool bInNotify);
+
+	void OnPackageSaved(const FString& InPackageFileName, UPackage* InPackage, FObjectPostSaveContext InObjectSaveContext);
+	void OnAvaSyncPackageModified(IAvaMediaSyncProvider* InAvaMediaSyncProvider, const FName& InPackageName);
+	void OnAssetRemoved(const FAssetData& InAssetData);
+	void OnPackageModified(const FName& InPackageName);
+	void OnSettingChanged(UObject* , struct FPropertyChangedEvent&);
+
+private:
+	TMap<FSoftObjectPath, TSharedPtr<FAvalancheManagedInstance>> AvalancheInstances;
+	TSet<FSoftObjectPath> PendingInvalidatedPaths;
+	TArray<FSoftObjectPath> OrderQueue;
+};
\ No newline at end of file
