[MotionDesign]

diff --git a/Engine/Plugins/Experimental/Avalanche/Source/AvalancheMedia/Public/Playlist/AvalanchePage.h b/Engine/Plugins/Experimental/Avalanche/Source/AvalancheMedia/Public/Playlist/AvalanchePage.h
new file mode 100644
--- /dev/null
+++ b/Engine/Plugins/Experimental/Avalanche/Source/AvalancheMedia/Public/Playlist/AvalanchePage.h
@@ -0,0 +1,278 @@
+// Copyright Epic Games, Inc. All Rights Reserved.
+
+#pragma once
+
+#include "AvaMediaDefines.h"
+#include "AvaTagHandle.h"
+#include "Playback/AvalancheRemoteControlValues.h"
+#include "AvalanchePage.generated.h"
+
+class UAvalancheBlueprint;
+class UAvalanchePlaylist;
+
+UENUM()
+enum class EAvalanchePageStatus : uint8
+{
+	/** Invalid page status. */
+	Unknown = 0,
+	/** Output is offline. */
+	Offline,
+	/** When the page is not available, i.e. the asset is not present in the local content. */
+	Missing,
+	/** Out of date */
+	NeedsSync,
+	/** Asset is being downloaded. */
+	Syncing,
+	/** When the page is present in local content, but not loaded. */
+	Available,
+	/** Load/Start has been requested. */
+	Loading,
+	/** Page is loaded in memory and ready to play. */
+	Loaded,
+	/** Page is currently playing in an output channel.*/
+	Playing,
+	/** Page is currently playing as local preview.*/
+	Previewing,
+	/** Something bad happened. */
+	Error
+};
+
+USTRUCT()
+struct AVALANCHEMEDIA_API FAvalanchePageStatus
+{
+	GENERATED_BODY()
+	
+public:
+	UPROPERTY()
+	EAvaBroadcastChannelType Type = EAvaBroadcastChannelType::Program;
+	
+	UPROPERTY()
+	EAvalanchePageStatus Status = EAvalanchePageStatus::Unknown;
+	
+	UPROPERTY()
+	bool bNeedsSync = false;
+};
+
+USTRUCT(BlueprintType)
+struct AVALANCHEMEDIA_API FAvalanchePage
+{
+	friend UAvalanchePlaylist;
+	GENERATED_BODY()
+
+public:
+	
+	static FAvalanchePage NullPage;
+	static const int32 InvalidPageId;
+
+	static bool StatusesContainsStatus(const TArray<FAvalanchePageStatus>& InStatuses, const TArray<EAvalanchePageStatus>& InStatusEnums)
+	{
+		for (const FAvalanchePageStatus& Status : InStatuses)
+		{
+			if (InStatusEnums.Contains(Status.Status))
+			{
+				return true;
+			}
+		}
+
+		return false;
+	}
+	
+	FAvalanchePage(int32 InPageId = InvalidPageId, int32 InTemplateId = InvalidPageId);
+
+	bool IsValidPage() const;
+
+	void Rename(const FString& InNewName);
+	void RenameFriendlyName (const FString& InNewName);
+	
+	int32 GetPageId() const { return PageId; }
+	void SetPageId(int32 InPageId) { PageId = InPageId; } // Do not use lightly.
+
+	int32 GetTemplateId() const { return TemplateId; }
+	void SetTemplateId(int32 InTemplateId) { TemplateId = InTemplateId; } // Do not use lightly.
+	bool IsTemplate() const { return TemplateId == InvalidPageId; }
+
+	const TArray<int32>& GetCombinedTemplateIds() const { return CombinedTemplateIds; }
+	bool IsComboTemplate() const { return IsTemplate() && !CombinedTemplateIds.IsEmpty();}
+	
+	const TSet<int32>& GetInstancedIds() const { return Instances; }
+
+	const FString& GetPageName() const { return PageName; }
+	void SetPageName(const FString& InPageName) { PageName = InPageName; }
+
+	bool HasPageSummary() const { return !PageSummary.IsEmptyOrWhitespace(); }
+	FText GetPageSummary() const { return PageSummary; }
+
+	bool HasPageFriendlyName() const { return !FriendlyName.IsEmptyOrWhitespace(); }
+	void SetPageFriendlyName(const FText& InPageFriendlyName) { FriendlyName = InPageFriendlyName; }
+	FText GetPageFriendlyName() const { return FriendlyName; }
+
+	FText GetPageDescription() const;
+
+	bool UpdatePageSummary(const UAvalanchePlaylist* InPlaylist);
+	bool UpdatePageSummary(const TArray<const URemoteControlPreset*>& InPresets, bool bInIsPresetChanged = false);
+
+	bool UpdateTransitionLogic();
+	bool HasTransitionLogic(const UAvalanchePlaylist* InPlaylist) const;
+	
+	FAvaTagHandle GetTransitionLayer(const UAvalanchePlaylist* InPlaylist, int32 InTemplateIndex = 0) const;
+	TArray<FAvaTagHandle> GetTransitionLayers(const UAvalanchePlaylist* InPlaylist) const;
+	
+	/**
+	 * Appends the page's program status(es).
+	 * Supports multiple statuses per page for forked channels (one status per server including local).
+	 * The order of the statuses is the same as the outputs (but only once per server) in the channel.
+	 * @remark Nothing is added if the page is a template.
+	 * @return The number of program statuses added.
+	 */
+	int32 AppendPageProgramStatuses(const UAvalanchePlaylist* InParentPlaylist, TArray<FAvalanchePageStatus>& OutPageStatuses) const;
+	
+	/**
+	 * Appends the page's preview status(es) for the given preview channel.
+	 * Remote channels are not yet supported for preview channels, only local for now.
+	 * @remark If the preview channel is not specified, it will use the playlist's default preview channel.
+	 * @remark this api is intended to query extra preview channels (for playlist server).
+	 * @return The number of preview statuses added.
+	 */
+	int32 AppendPagePreviewStatuses(const UAvalanchePlaylist* InParentPlaylist, const FName& InPreviewChannelName, TArray<FAvalanchePageStatus>& OutPageStatuses) const;
+
+	/**
+	 * Returns all the page's "standard" (program and preview) playback statuses. 
+	 * The order of the status is the same as the outputs in the channel
+	 * and the preview status is always last.
+	 * Template pages only have the preview status added.
+	 */
+	TArray<FAvalanchePageStatus> GetPageStatuses(const UAvalanchePlaylist* InParentPlaylist) const;
+
+	/**
+	 * Returns either the program statuses or preview statuses only depending if the page is instanced (has a program channel),
+	 * or a template respectively.
+	 */
+	TArray<FAvalanchePageStatus> GetPageContextualStatuses(const UAvalanchePlaylist* InParentPlaylist) const;
+
+	/**
+	 * Returns the page's program statuses.
+	 * Returns nothing for template pages.
+	 */
+	TArray<FAvalanchePageStatus> GetPageProgramStatuses(const UAvalanchePlaylist* InParentPlaylist) const;
+
+	/**
+	 * Returns the page's preview statuses for the given preview channel.
+	 * If the preview channel is not specified, it will use the playlist's default preview channel.
+	 */
+	TArray<FAvalanchePageStatus> GetPagePreviewStatuses(const UAvalanchePlaylist* InParentPlaylist, const FName& InPreviewChannelName = NAME_None) const;
+	
+	bool IsEnabled() const { return bEnabled; }
+	void SetEnabled(bool bInEnabled) { bEnabled = bInEnabled; }
+
+	FSoftObjectPath GetAvalancheAssetPath(const UAvalanchePlaylist* InPlaylist, int32 InTemplateIndex = 0) const;
+	TArray<FSoftObjectPath> GetAvalancheAssetPaths(const UAvalanchePlaylist* InPlaylist) const;
+	FSoftObjectPath GetAvalancheAssetPathDirect() const { return AssetPath;}
+
+	bool UpdateAvalancheAsset(const FSoftObjectPath& InAssetPath, bool bInReimportPage = false);
+
+	FName GetChannelName() const;
+	int32 GetChannelIndex() const { return OutputChannel; }
+	void SetChannelName(FName InChannelName);
+
+	EAvaRemoteControlChanges PruneRemoteControlValues(const FAvalancheRemoteControlValues& InRemoteControlValues);
+	EAvaRemoteControlChanges UpdateRemoteControlValues(const FAvalancheRemoteControlValues& InRemoteControlValues, bool bInUpdateDefaults);
+	
+	bool HasRemoteControlEntityValue(const FGuid& InId) const { return RemoteControlValues.HasEntityValue(InId); }
+	const FAvalancheRemoteControlValue* GetRemoteControlEntityValue(const FGuid& InId) const { return RemoteControlValues.GetEntityValue(InId); }
+	void SetRemoteControlEntityValue(const FGuid& InId, const FAvalancheRemoteControlValue& InValue);
+
+	bool HasRemoteControlControllerValue(const FGuid& InId) const { return RemoteControlValues.HasControllerValue(InId); }
+	const FAvalancheRemoteControlValue* GetRemoteControlControllerValue(const FGuid& InId) const { return RemoteControlValues.GetControllerValue(InId); }
+	void SetRemoteControlControllerValue(const FGuid& InId, const FAvalancheRemoteControlValue& InValue);
+	
+	const FAvalancheRemoteControlValues& GetRemoteControlValues() const { return RemoteControlValues; }
+
+	friend FORCEINLINE uint32 GetTypeHash(const FAvalanchePage& InPage)
+	{
+		return GetTypeHash(InPage.PageId);
+	}
+
+	bool operator<(const FAvalanchePage& InOther) const
+	{
+		return this->PageId < InOther.PageId;
+	}
+
+	void PostLoad();
+
+	/** Returns the number of templates this page/combo template has. */
+	int32 GetNumTemplates(const UAvalanchePlaylist* InPlaylist) const;
+
+	/**
+	 * Returns the template at given index.
+	 * This will first resolve the direct template, then lookup combined templates.
+	 */
+	const FAvalanchePage& GetTemplate(const UAvalanchePlaylist* InPlaylist, int32 InIndex) const;
+
+	/**
+	 * Resolves a page's templates.
+	 * For a page, it returns its (direct) template.
+	 * For a template, it returns itself. 
+	 */
+	const FAvalanchePage& ResolveTemplate(const UAvalanchePlaylist* InPlaylist) const;
+	
+	/**
+	 * @brief Compare the values from another template and determines if there is a match.
+	 * @param InTemplatePage Other template page to compare to.
+	 * @return true if the templates are matching (ignores ids).
+	 */
+	bool IsTemplateMatchingByValue(const FAvalanchePage& InTemplatePage) const;
+	
+protected:
+
+	UPROPERTY(VisibleAnywhere, BlueprintReadOnly, Category = "Motion Design")
+	bool bEnabled = true;
+
+	UPROPERTY(VisibleAnywhere, BlueprintReadOnly, Category = "Motion Design")
+	int32 PageId = InvalidPageId;
+
+	/** Page Instance Property: Template Id for this page. */
+	UPROPERTY(VisibleAnywhere, BlueprintReadOnly, Category = "Motion Design")
+	int32 TemplateId;
+
+	/**
+	 * Template property: For combination template, lists the templates that are combined.
+	 * A combination template can only be created using transition logic templates.
+	 * In order to create a combination template, the templates must be in different transition layers.
+	 */ 
+	UPROPERTY(VisibleAnywhere, BlueprintReadOnly, Category = "Motion Design")
+	TArray<int32> CombinedTemplateIds;
+
+	UPROPERTY(VisibleAnywhere, BlueprintReadOnly, Category = "Motion Design")
+	FString PageName;
+
+	UPROPERTY(meta = (DeprecatedProperty, DeprecationMessage = "Use AssetPath."))
+	TSoftObjectPtr<UAvalancheBlueprint> AvalancheBlueprint_DEPRECATED;
+
+	/** Template property: path for this template. */
+	UPROPERTY(VisibleAnywhere, BlueprintReadOnly, Category = "Motion Design", meta=(AllowedClasses="/Script/Engine.World, /Script/Avalanche.AvalancheBlueprint"))
+	FSoftObjectPath AssetPath;
+
+	/** Template property: List the Ids of all instances. */
+	UPROPERTY(VisibleAnywhere, BlueprintReadOnly, Category = "Motion Design")
+	TSet<int32> Instances;
+
+	UPROPERTY(VisibleAnywhere, BlueprintReadOnly, Category = "Motion Design")
+	int32 OutputChannel = 0;
+
+	UPROPERTY(EditAnywhere, Category = "Motion Design")
+	FAvalancheRemoteControlValues RemoteControlValues;
+
+	UPROPERTY(VisibleAnywhere, BlueprintReadOnly, Category = "Motion Design")
+	FText PageSummary = FText::GetEmpty();
+
+	UPROPERTY(VisibleAnywhere, BlueprintReadOnly, Category = "Motion Design")
+	FText FriendlyName = FText::GetEmpty();
+
+	/** Indicate if the template asset has transition logic. */
+	UPROPERTY(VisibleAnywhere, BlueprintReadOnly, Category = "Motion Design")
+	bool bHasTransitionLogic = false;
+
+	/** Transition Layer Tag cached from the transition tree. Cached for fast display in page/template list. */
+	UPROPERTY(VisibleAnywhere, Category = "Motion Design")
+	FAvaTagHandle TransitionLayerTag;
+};
diff --git a/Engine/Plugins/Experimental/Avalanche/Source/AvalancheMedia/Public/Playlist/AvalanchePagePlayer.h b/Engine/Plugins/Experimental/Avalanche/Source/AvalancheMedia/Public/Playlist/AvalanchePagePlayer.h
new file mode 100644
--- /dev/null
+++ b/Engine/Plugins/Experimental/Avalanche/Source/AvalancheMedia/Public/Playlist/AvalanchePagePlayer.h
@@ -0,0 +1,196 @@
+// Copyright Epic Games, Inc. All Rights Reserved.
+
+#pragma once
+
+#include "AvaTagHandle.h"
+#include "Playback/AvalanchePlayable.h"
+#include "Playback/AvaMediaPlaybackManager.h"
+#include "UObject/Object.h"
+
+#include "AvalanchePagePlayer.generated.h"
+
+class FAvaMediaPlaybackInstance;
+class UAvalanchePagePlayer;
+class UAvalanchePlayback;
+class UAvalanchePlaylist;
+struct FAvalanchePage;
+
+UENUM(BlueprintType)
+enum class EAvaPlayType : uint8
+{
+	PlayFromStart,
+	PreviewFromStart,
+	PreviewFromFrame
+};
+
+/**
+ *	To support combo templates, a page is now potentially composed of a number of
+ *	playback instances that each have their own player.
+ */
+UCLASS()
+class AVALANCHEMEDIA_API UAvaRundownPlaybackInstancePlayer : public UObject
+{
+	GENERATED_BODY()
+
+public:
+	UAvaRundownPlaybackInstancePlayer();
+	virtual ~UAvaRundownPlaybackInstancePlayer() override;
+
+	bool Load(const UAvalanchePagePlayer& InPagePlayer, const UAvalanchePlaylist* InPlaylist, const FAvalanchePage& InPage, int32 InSubPageIndex, const FGuid& InInstanceId);
+	bool IsLoaded() const;
+
+	void Play(const UAvalanchePagePlayer& InPagePlayer, const UAvalanchePlaylist* InPlaylist, EAvaPlayType InPlayType, bool bInIsUsingTransitionLogic);
+	bool IsPlaying() const;
+
+	bool Continue(const FString& InChannelName);
+	bool Stop();
+	
+	FAvaMediaPlaybackInstance* GetPlaybackInstance() const { return PlaybackInstance ? PlaybackInstance.Get() : nullptr; }
+
+	FGuid GetPlaybackInstanceId() const { return PlaybackInstance ? PlaybackInstance->GetInstanceId() : FGuid(); }
+
+	bool HasPlayable(const UAvalanchePlayable* InPlayable) const;
+
+	UAvalanchePlayable* GetFirstPlayable() const;
+
+	UAvalanchePagePlayer* GetPagePlayer() const;
+
+	void SetPagePlayer(UAvalanchePagePlayer* InPagePlayer);
+
+public:	
+	UPROPERTY()
+	FAvaTagHandle TransitionLayer;
+
+	UPROPERTY()
+	FSoftObjectPath SourceAssetPath;
+	
+	UPROPERTY()
+	TObjectPtr<UAvalanchePlayback> Playback;
+
+	TSharedPtr<FAvaMediaPlaybackInstance> PlaybackInstance;
+};
+
+UCLASS()
+class AVALANCHEMEDIA_API UAvalanchePagePlayer : public UObject
+{
+	GENERATED_BODY()
+	
+public:
+	UAvalanchePagePlayer();
+	virtual ~UAvalanchePagePlayer() override;
+
+	/** Initialize Page player without loading any playback instances. */
+	bool Initialize(UAvalanchePlaylist* InPlaylist, const FAvalanchePage& InPage, bool bInIsPreview, const FName& InPreviewChannel);
+
+	/** Initialize the Page player and loads all playback instances. */
+	bool InitializeAndLoad(UAvalanchePlaylist* InPlaylist, const FAvalanchePage& InPage, bool bInIsPreview, const FName& InPreviewChannel);
+
+	/**
+	 * Load an instance player for the "sub page" at the given index.
+	 * The index is the template index in the combo template.
+	 */
+	UAvaRundownPlaybackInstancePlayer* LoadInstancePlayer(int32 InSubPageIndex, const FGuid& InInstanceId);
+
+	/** Add a pre-existing instance player to this page player. */
+	void AddInstancePlayer(UAvaRundownPlaybackInstancePlayer* InExistingInstancePlayer);
+
+	/** Returns true if at least one of the instance player is loaded, false otherwise. */
+	bool IsLoaded() const;
+	
+	bool Play(EAvaPlayType InPlayType, bool bInIsUsingTransitionLogic);
+
+	/** Returns true if at least one of the instance player is playing, false otherwise. */
+	bool IsPlaying() const;
+	
+	bool Continue();
+	bool Stop();
+
+	int32 GetNumInstancePlayers() const { return InstancePlayers.Num(); }
+
+	FAvaMediaPlaybackInstance* GetPlaybackInstance(int32 InIndex = 0) const
+	{
+		return InstancePlayers.IsValidIndex(InIndex) && IsValid(InstancePlayers[InIndex]) ? InstancePlayers[InIndex]->GetPlaybackInstance() : nullptr;
+	}
+
+	FGuid GetPlaybackInstanceId(int32 InIndex = 0) const
+	{
+		return InstancePlayers.IsValidIndex(InIndex) && IsValid(InstancePlayers[InIndex]) ? InstancePlayers[InIndex]->GetPlaybackInstanceId() : FGuid();
+	}
+
+	UAvalanchePlayback* GetPlayback(int32 InIndex = 0) const
+	{
+		return InstancePlayers.IsValidIndex(InIndex) && IsValid(InstancePlayers[InIndex]) ? InstancePlayers[InIndex]->Playback : nullptr;
+	}
+
+	FSoftObjectPath GetSourceAssetPath(int32 InIndex = 0) const
+	{
+		return InstancePlayers.IsValidIndex(InIndex) && IsValid(InstancePlayers[InIndex]) ? InstancePlayers[InIndex]->SourceAssetPath : FSoftObjectPath();
+	}
+
+	UAvaRundownPlaybackInstancePlayer* GetInstancePlayer(int32 InIndex) const
+	{
+		return InstancePlayers.IsValidIndex(InIndex)? InstancePlayers[InIndex] : nullptr;
+	}
+
+	void ForEachInstancePlayer(TFunctionRef<void(UAvaRundownPlaybackInstancePlayer*)> InFunction)
+	{
+		for (const TObjectPtr<UAvaRundownPlaybackInstancePlayer>& InstancePlayer : InstancePlayers)
+		{
+			InFunction(InstancePlayer.Get());
+		}
+	}
+
+	void ForEachInstancePlayer(TFunctionRef<void(const UAvaRundownPlaybackInstancePlayer*)> InFunction) const
+	{
+		for (const TObjectPtr<UAvaRundownPlaybackInstancePlayer>& InstancePlayer : InstancePlayers)
+		{
+			InFunction(InstancePlayer.Get());
+		}
+	}
+	
+	UAvalanchePlaylist* GetPlaylist() const { return PlaylistWeak.Get(); }
+
+	static int32 GetPageIdFromInstanceUserData(const FString& InUserData);
+
+	static void SetInstanceUserDataFromPage(FAvaMediaPlaybackInstance& InPlaybackInstance, const FAvalanchePage& InPage);
+
+	bool HasPlayable(const UAvalanchePlayable* InPlayable) const;
+
+	UAvaRundownPlaybackInstancePlayer* FindInstancePlayerForPlayable(const UAvalanchePlayable* InPlayable) const;
+
+	UAvaRundownPlaybackInstancePlayer* FindInstancePlayerByInstanceId(const FGuid& InInstanceId) const;
+
+	UAvaRundownPlaybackInstancePlayer* FindInstancePlayerByAssetPath(const FSoftObjectPath& InAssetPath) const;
+	
+protected:
+	UAvaRundownPlaybackInstancePlayer* CreateAndLoadInstancePlayer(const UAvalanchePlaylist* InPlaylist, const FAvalanchePage& InPage, int32 InSubPageIndex, const FGuid& InInstanceId);
+
+	void RemoveInstancePlayer(UAvaRundownPlaybackInstancePlayer* InInstancePlayer);
+	
+	void HandleOnPlayableSequenceEvent(UAvalanchePlayable* InPlayable, const FName& SequenceName, EAvalanchePlayableSequenceEventType InEventType);
+
+public:
+	UPROPERTY()
+	int32 PageId = INDEX_NONE;
+
+	UPROPERTY()
+	bool bIsPreview = false;
+
+	/** @remark For previews, the channel name will not be the one set in the page. */
+	UPROPERTY()
+	FName ChannelFName;
+
+	/** @remark For previews, the channel name will not be the one set in the page. */
+	UPROPERTY()
+	FString ChannelName;
+
+	UPROPERTY()
+	TArray<TObjectPtr<UAvaRundownPlaybackInstancePlayer>> InstancePlayers;
+
+	/** Instances that are excluded from the next transition. */
+	UPROPERTY(Transient)
+	TSet<FGuid> InstancesExcludedFromTransition;
+
+protected:
+	TWeakObjectPtr<UAvalanchePlaylist> PlaylistWeak;
+};
diff --git a/Engine/Plugins/Experimental/Avalanche/Source/AvalancheMedia/Public/Playlist/AvalanchePlaylist.h b/Engine/Plugins/Experimental/Avalanche/Source/AvalancheMedia/Public/Playlist/AvalanchePlaylist.h
new file mode 100644
--- /dev/null
+++ b/Engine/Plugins/Experimental/Avalanche/Source/AvalancheMedia/Public/Playlist/AvalanchePlaylist.h
@@ -0,0 +1,731 @@
+// Copyright Epic Games, Inc. All Rights Reserved.
+
+#pragma once
+
+#include "AvalanchePage.h"
+#include "AvaMediaDefines.h"
+#include "UObject/Object.h"
+#include "AvalanchePagePlayer.h"
+#include "AvalanchePlaylist.generated.h"
+
+AVALANCHEMEDIA_API DECLARE_LOG_CATEGORY_EXTERN(LogAvaPlaylist, Log, All);
+
+class FAvaMediaPlaybackManager;
+class FAvaPlaylistPlaybackClientWatcher;
+class FAvalanchePageTransitionBuilder;
+class UAvalanchePageTransition;
+class UAvalanchePlayback;
+class UTextureRenderTarget2D;
+
+/**
+ * @brief Defines the insertion position in a page list.
+ */
+struct FAvaPageInsertPosition
+{
+	/**
+	 * The position is defined with a page id because it is coupled with the id generation.
+	 * For the id generator, if an insertion position is defined, we also want the generated
+	 * ids to be in relation to that.
+	 */
+	int32 AdjacentId;
+		
+	/** Defines the insertion position relative to the reference page i.e. above or below the adjacent page. */
+	bool bAddBelow;
+
+	explicit FAvaPageInsertPosition(int32 InAdjacentId = FAvalanchePage::InvalidPageId, bool bInAddBelow = true) : AdjacentId(InAdjacentId), bAddBelow(bInAddBelow) {} 
+
+	bool IsValid() const { return AdjacentId != FAvalanchePage::InvalidPageId; }
+
+	bool IsAddAbove() const { return !bAddBelow;}
+	
+	bool IsAddBelow() const { return bAddBelow;}
+
+	/** Update the id only if it was initially valid. */
+	void ConditionalUpdateAdjacentId(int32 InNewAdjacentId)
+	{
+		if (IsValid())
+		{
+			AdjacentId = InNewAdjacentId;
+		}
+	}
+};
+
+/**
+ * Defines the parameters for the page id generator algorithm.
+ * The Id generator uses a sequence strategy to search for an unused id.
+ * It is defined by a starting id and a search direction.
+ */
+struct FAvaPageIdGeneratorParams
+{
+	/** Starting Id for the search. */
+	int32 ReferenceId;
+		
+	/**
+	 * @brief (Initial) Search increment.
+	 * @remark For negative increment search, the limit of the search space can be reached. If no unique id is found,
+	 *		   the search will continue in the positive direction instead.		   
+	 */
+	int32 Increment;
+		
+	explicit FAvaPageIdGeneratorParams(int32 InReferenceId = FAvalanchePage::InvalidPageId, int32 InIncrement = 1)
+		: ReferenceId(InReferenceId), Increment(InIncrement) {}
+
+	/** Operation helper: Determines id generation from the insert parameters. */
+	static FAvaPageIdGeneratorParams FromInsertPosition(const FAvaPageInsertPosition& InInsertPosition)
+	{
+		// When used with insertion in a page list, if the element is added above, we will first try to generate
+		// the id in decreasing order.
+		return FAvaPageIdGeneratorParams( InInsertPosition.AdjacentId, InInsertPosition.bAddBelow ? 1 : -1);
+	}
+
+	/** Operation helper: Id generation prefers using insertion parameters (if specified) over source id. */
+	static FAvaPageIdGeneratorParams FromInsertPositionOrSourceId(int32 InSourceId, const FAvaPageInsertPosition& InInsertPosition)
+	{
+		return InInsertPosition.IsValid() ? FromInsertPosition(InInsertPosition) : FAvaPageIdGeneratorParams(InSourceId);
+	}
+};
+
+struct FAvaPageListChangeParams
+{
+	UAvalanchePlaylist* Playlist;
+	EAvaPageListChange ChangeType;
+	TArray<int32> AffectedPages;
+};
+
+DECLARE_MULTICAST_DELEGATE_OneParam(FOnAvaPageListChanged, const FAvaPageListChangeParams&)
+DECLARE_MULTICAST_DELEGATE_ThreeParams(FOnAvaPagesChanged, const UAvalanchePlaylist*, const FAvalanchePage&, EAvaPageChanges)
+
+USTRUCT()
+struct FAvalanchePageCollection
+{
+	GENERATED_BODY()
+
+	UPROPERTY()
+	TArray<FAvalanchePage> Pages;
+
+	/** Cache mapping the Page Id to the index where the Page with such Page Id is at */
+	TMap<int32, int32> PageIndices;
+
+	FOnAvaPageListChanged OnPageListChanged;
+
+	int32 GetPageIndex(const int32 InPageId) const
+	{
+		if (InPageId != FAvalanchePage::InvalidPageId)
+		{
+			const int32* PageIndex = PageIndices.Find(InPageId);
+			return PageIndex ? *PageIndex : INDEX_NONE;
+		}
+		return INDEX_NONE;
+	}
+
+	AVALANCHEMEDIA_API void Empty(UAvalanchePlaylist* InPlaylist);
+
+	/** Complete refresh of the page indices. */
+	void RefreshPageIndices()
+	{
+		PageIndices.Empty(Pages.Num());
+		for (int32 Index = 0; Index < Pages.Num(); ++Index)
+		{
+			PageIndices.Add(Pages[Index].GetPageId(), Index);
+		}
+	}
+
+	/** Refresh page indices after a new page has been inserted at the given index. */
+	void PostInsertRefreshPageIndices(const int32 InStartAtIndex)
+	{
+		for (int32 Index = InStartAtIndex; Index < Pages.Num(); ++Index)
+		{
+			PageIndices.Add(Pages[Index].GetPageId(), Index);
+		}
+	}
+};
+
+USTRUCT()
+struct FAvalancheSubList
+{
+	GENERATED_BODY()
+
+	UPROPERTY()
+	TArray<int32> PageIds;
+
+	UPROPERTY()
+	FText Name;
+
+	FOnAvaPageListChanged OnPageListChanged;
+};
+
+namespace UE::AvalanchePlaylist
+{
+	inline bool IsPreviewPlayType(EAvaPlayType InPlayType)
+	{
+		return InPlayType == EAvaPlayType::PreviewFromStart || InPlayType == EAvaPlayType::PreviewFromFrame;
+	}
+}
+
+/**
+ * Carries the context for the playback of a page list.
+ * In particular, what is the current play head to be able to move to
+ * the next page in the page list.
+ * Note: a "playlist" has many page lists and each page lists has 2 play type (preview or program),
+ * therefore, we have many instances of this page list context for each page lists and play type within
+ * a "playlist" object.
+ */
+struct FAvaPageListPlaybackContext
+{
+	/** Keeps track of the page id the play head is on, i.e. this is the last played page. */
+	int32 PlayHeadPageId = FAvalanchePage::InvalidPageId;
+};
+
+/**
+ * This class is a container for all the page list contexts are playlist can have.
+ * The current implementation only keeps track of the play type, i.e. preview vs program.
+ * The design for this is not settle yet. The requirement for a page list context per preview channel comes
+ * from the playlist server as it may have a preview channel dedicated per client connection, which implies a page list
+ * context for each one of them. There is only one "program" page list context for now.
+ */
+struct FAvaPageListPlaybackContextCollection
+{
+	TSharedPtr<FAvaPageListPlaybackContext> GetContext(bool bInIsPreview, const FName& InPreviewChannelName) const
+	{
+		if (bInIsPreview)
+		{
+			const TSharedPtr<FAvaPageListPlaybackContext>* FoundContext = PreviewContexts.Find(InPreviewChannelName);
+			return FoundContext ? *FoundContext : TSharedPtr<FAvaPageListPlaybackContext>();
+		}
+		return ProgramContext;
+	}
+
+	FAvaPageListPlaybackContext& GetOrCreateContext(bool bInIsPreview, const FName& InPreviewChannelName)
+	{
+		const TSharedPtr<FAvaPageListPlaybackContext> ExistingContext = GetContext(bInIsPreview, InPreviewChannelName);
+		return ExistingContext ? *ExistingContext : *CreateContext(bInIsPreview, InPreviewChannelName);
+	}
+
+protected:
+	TSharedPtr<FAvaPageListPlaybackContext> CreateContext(bool bInIsPreview, const FName& InPreviewChannelName)
+	{
+		if (bInIsPreview)
+		{
+			TSharedPtr<FAvaPageListPlaybackContext> NewPreviewContext = MakeShared<FAvaPageListPlaybackContext>(); 
+			PreviewContexts.Add(InPreviewChannelName, NewPreviewContext);
+			return NewPreviewContext;
+		}
+		ProgramContext = MakeShared<FAvaPageListPlaybackContext>();
+		return ProgramContext;
+	}
+
+	TSharedPtr<FAvaPageListPlaybackContext> ProgramContext;
+	TMap<FName, TSharedPtr<FAvaPageListPlaybackContext>> PreviewContexts;
+};
+
+UENUM()
+enum class EAvaPlaylistPageStopOptions : uint8
+{
+	/**
+	 * Default option will stop the page with transitions if available.
+	 */
+	None				= 0,
+	/**
+	 * Forces the page to stop without transitions.
+	 */
+	ForceNoTransition	= 1 << 1,
+	/**
+	 * Default option will stop the page with transitions if available.
+	 */
+	Default				= None
+};
+ENUM_CLASS_FLAGS(EAvaPlaylistPageStopOptions);
+
+/**
+ * @brief Manages page pre-loading.
+ */
+class IAvaPlaylistPageLoadingManager
+{
+public:
+	IAvaPlaylistPageLoadingManager() = default;
+	virtual ~IAvaPlaylistPageLoadingManager() = default;
+
+	virtual bool RequestLoadPage(int32 InPageId, bool bInIsPreview, const FName& InPreviewChannelName) = 0;
+};
+
+/**
+ * This class is a container for what could be described as a "show" for broadcast purposes.
+ *
+ * It goes beyond a simple playlist of items. It contains the following:
+ * - a list of Avalanche Template Pages (or just Templates).
+ * - a list of Avalanche Instanced Pages (or just Pages).
+ * - a list of page views (or just Views).
+ *
+ * Workflow:
+ *
+ * 1- Templates
+ * 
+ * The first step in the work flow consist in importing templates. The source asset is not actually imported
+ * in the "show" container, it is just soft referenced. However, the import process will load and cache some information
+ * about the template (exposed properties, default values, animations, transition logic layer, etc).
+ * Given that this information is cached, it may become stale if the source asset is updated. Therefore, reimporting
+ * the templates may be necessary within the normal work flow.
+ * Todo: keep a hash of the source asset to determine if it has changed.
+ *
+ * 2- Pages
+ *
+ * The pages are instances of the templates, allowing to change the exposed properties and controllers, also selecting
+ * an output program channel for the given page. Only one program channel is allowed per page.
+ *
+ * 3- Page Views
+ *
+ * Separate page views can be made in order to create "playlists" for separate segments/parts of a show.
+ *
+ * "Page Groups" Discussion:
+ * "Page Groups" are not implemented. It would be different than page views, i.e.
+ * pages could be grouped in either of the page list or page views.
+ * Other applications support page grouping to emulate MOS's hierarchy.
+ * In the MOS/NCS hierarchies: Rundown -> Stories/Segments -> Parts -> Pieces/Items
+ * Although full emulation of MOS schema may not be necessary within the Avalanche playback framework.
+ *
+ */
+UCLASS(NotBlueprintable, BlueprintType)
+class AVALANCHEMEDIA_API UAvalanchePlaylist : public UObject
+{
+	GENERATED_BODY()
+
+public:
+	UAvalanchePlaylist();
+
+	virtual ~UAvalanchePlaylist() override;
+
+	static const FAvaPageListReference TemplatePageList;
+	static const FAvaPageListReference InstancePageList;
+
+	static FAvaPageListReference CreateSubListReference(int32 InSubListIndex) { return {EAvaPageListType::View, InSubListIndex}; }
+
+protected:
+	bool IsPageIdUnique(int32 InPageId) const { return !TemplatePages.PageIndices.Contains(InPageId) && !InstancedPages.PageIndices.Contains(InPageId); }
+	int32 GenerateUniquePageId(int32 InReferencePageId = FAvalanchePage::InvalidPageId, int32 InIncrement = 1) const;
+	int32 GenerateUniquePageId(const FAvaPageIdGeneratorParams& InParams) const;
+	
+	/** Caches the Page's Id to its Index in the Pages Array*/
+	void RefreshPageIndices();
+	
+public:
+	static const FAvalancheSubList InvalidSubList;
+
+	//~ Begin UObject Interface
+	virtual void BeginDestroy() override;
+	virtual void PostLoad() override;
+#if WITH_EDITOR
+	virtual void PostEditUndo() override;
+#endif
+	//~ End UObject Interface
+
+	/** Returns true if the playlist is empty, i.e. no pages and no templates. */
+	bool IsEmpty() const;
+
+	/**
+	 * Clear the playlist of all it's content.
+	 * @remark Will be prevented if the playlist is playing.
+	 */
+	bool Empty();
+
+	int32 AddTemplateInternal(const FAvaPageIdGeneratorParams& InIdGeneratorParams, const TFunctionRef<bool(FAvalanchePage&)> InSetupTemplateFunction);
+	
+	/** Add empty template. */
+	int32 AddTemplate(const FAvaPageIdGeneratorParams& InIdGeneratorParams = FAvaPageIdGeneratorParams());
+
+	int32 AddComboTemplate(const TArray<int32>& InTemplateIds, const FAvaPageIdGeneratorParams& InIdGeneratorParams = FAvaPageIdGeneratorParams());
+
+	/**
+	 * @brief Add templates from existing source.
+	 * @param InSourceTemplates Source template to add.
+	 * @return The new template Ids created.
+	 *
+	 * For the id generation, it will attempt to reuse the source ids, but
+	 * in case of collision, new ids are generated with the positive increment
+	 * sequence search method.
+	 */
+	TArray<int32> AddTemplates(const TArray<FAvalanchePage>& InSourceTemplates);
+
+	/**
+	 * @brief Create new pages in the page list for teh given template Ids.
+	 * @param InTemplateIds Templates to use for each page. A new page is created for each entry in that array.
+	 * @return The new page Ids created.
+	 */
+	TArray<int32> AddPagesFromTemplates(const TArray<int32>& InTemplateIds);
+
+	/**
+	 * @brief Creates a new page in the page list using the given template.
+	 * @param InTemplateId Reference to the template to use for that page.
+	 * @param InIdGeneratorParams Defines how the page id is going to be generated.
+	 * @param InInsertAt Specifies the insertion location in the page list (i.e. the index in the page list).
+	 * @return Returns the page Id of the created page.
+	 */
+	int32 AddPageFromTemplate(int32 InTemplateId, const FAvaPageIdGeneratorParams& InIdGeneratorParams = FAvaPageIdGeneratorParams(), const FAvaPageInsertPosition& InInsertAt = FAvaPageInsertPosition());
+
+	bool CanAddPage() const;
+
+	bool CanChangePageOrder() const;
+	/** Reorders the pages, swapping the old indices for the new ones. */
+	bool ChangePageOrder(const FAvaPageListReference& InPageListReference, const TArray<int32>& InPageIndices);
+
+	bool RemovePage(int32 InPageId);
+	bool CanRemovePage(int32 InPageId) const;
+
+	/** Remove all the Pages in the Array. Returns the number of Pages removed.*/
+	int32 RemovePages(const TArray<int32>& InPageIds);
+	bool CanRemovePages(const TArray<int32>& InPageIds) const;
+
+	bool RenumberPageId(int32 InPageId, int32 InNewPageId);
+	bool CanRenumberPageId(int32 InPageId) const;
+	bool CanRenumberPageId(int32 InPageId, int32 InNewPageId) const;
+
+	bool SetRemoteControlEntityValue(int32 InPageId, const FGuid& InId, const FAvalancheRemoteControlValue& InValue);
+	bool SetRemoteControlControllerValue(int32 InPageId, const FGuid& InId,const FAvalancheRemoteControlValue& InValue);
+	EAvaRemoteControlChanges UpdateRemoteControlValues(int32 InPageId, const FAvalancheRemoteControlValues& InRemoteControlValues, bool bInUpdateDefaults);
+
+	void InvalidateManagedInstanceCacheForPages(const TArray<int32>& InPageIds) const;
+
+	void UpdateAvalancheAssetForPages(const TArray<int32>& InPageIds, bool bInReimportPage);
+	
+	const FAvalanchePage& GetPage(int32 InPageId) const;
+	FAvalanchePage& GetPage(int32 InPageId);
+	const FAvalanchePageCollection& GetTemplatePages() const { return TemplatePages; }
+	const FAvalanchePageCollection& GetInstancedPages() const { return InstancedPages; }
+	static const FAvalanchePage& GetPageSafe(const UAvalanchePlaylist* InPlaylist, int32 InPageId)
+	{
+		return InPlaylist ? InPlaylist->GetPage(InPageId) : FAvalanchePage::NullPage;
+	}
+	/**
+	 * Gets the page following the page with the given page id in the given page list.
+	 * If the given page is not in the given list, the returned page is invalid.
+	 */
+	const FAvalanchePage& GetNextPage(int32 InPageId, const FAvaPageListReference& InPageListReference) const;
+	
+	/**
+	 * Gets the page following the page with the given page id in the given page list.
+	 * If the given page is not in the given list, the returned page is invalid.
+	 */
+	FAvalanchePage& GetNextPage(int32 InPageId, const FAvaPageListReference& InPageListReference);
+
+	/** Gets the page following the page with the given page id. Using current active page list. */
+	const FAvalanchePage& GetNextPage(int32 InPageId) const { return GetNextPage(InPageId, ActivePageList);}
+
+	/** Gets the page following the page with the given page id. Using current active page list. */
+	FAvalanchePage& GetNextPage(int32 InPageId)  { return GetNextPage(InPageId, ActivePageList);}
+	
+	FOnAvaPageListChanged& GetOnTemplatePageListChanged() { return TemplatePages.OnPageListChanged; }
+	FOnAvaPageListChanged& GetOnInstancedPageListChanged() { return InstancedPages.OnPageListChanged; }
+	FOnAvaPagesChanged& GetOnPagesChanged() { return OnPagesChanged; }
+
+	/**
+	 * Since the playback context is part of the asset for now, there is an explicit call to initialize it.
+	 * This would be done by the editor (or server).
+	 * A future refactor will extract the playlist "player" functionality in another class and it will
+	 * be possible to create multiple instance of a playlist player for the same playlist.
+	 */
+	void InitializePlaybackContext();
+
+	/**
+	 * Similarly, when the editor is done, it can close the playback context. 
+	 * This will clean up the internal structures for playback and optionally stop all the pages.
+	 */
+	void ClosePlaybackContext(bool bInStopAllPages);
+
+	/**
+	 * Returns true if the any page is either playing or previewing.
+	 */
+	bool IsPlaying() const;
+
+	/** Returns true if the page with the given page Id is being previewed (in any preview channel). */
+	bool IsPagePreviewing(int32 InPageId) const;
+
+	/** Returns true if the page is playing in it's assigned program channel. */
+	bool IsPagePlaying(int32 InPageId) const;
+	
+	bool IsPagePlaying(const FAvalanchePage& InPage) const { return IsPagePlaying(InPage.GetPageId()); }
+
+	bool IsPagePlayingOrPreviewing(int32 InPageId) const;
+
+	bool UnloadPage(int32 InPageId, const FString& InChannelName);
+
+	struct FLoadedInstanceInfo
+	{
+		FGuid InstanceId;
+		FSoftObjectPath AssetPath;
+	};
+	
+	/**
+	 * @brief Preload the given page so it has an asset ready for playback.
+	 * @return UUIDs (and asset paths) of the playback instances that where loaded (or are scheduled for loading).
+	 */
+	TArray<FLoadedInstanceInfo> LoadPage(int32 InPageId,  bool bInPreview, const FName& InPreviewChannelName);
+
+	/**
+	 * @brief Start the playback of the asset defined in the given pages.
+	 * @remark If the play type is a preview, the default preview channel is used.
+	 * @param InPageIds Playlist's pages to play out.
+	 * @param InPlayType Either play on program or preview
+	 * @return page Ids that where started.
+	 */
+	TArray<int32> PlayPages(const TArray<int32>& InPageIds, EAvaPlayType InPlayType);
+
+	/**
+	 * @brief Start the playback of the asset defined in the given pages.
+	 * @remark If the play type is a preview, the given preview channel is used.
+	 * @param InPageIds Playlist's pages to play out.
+	 * @param InPlayType Either play on program or preview
+	 * @param InPreviewChannelName Channel to use for preview. Only used if play type is preview.
+	 * @return page Ids that where started.
+	 */
+	TArray<int32> PlayPages(const TArray<int32>& InPageIds, EAvaPlayType InPlayType, const FName& InPreviewChannelName);
+
+	bool PlayPage(int32 InPageId, EAvaPlayType InPlayType) { return !PlayPages({InPageId}, InPlayType).IsEmpty(); }
+	bool PlayPage(int32 InPageId, EAvaPlayType InPlayType, const FName& InPreviewChannelName) { return !PlayPages({InPageId}, InPlayType, InPreviewChannelName).IsEmpty(); }
+
+	bool CanPlayPage(int32 InPageId, bool bInPreview) const;
+	bool CanPlayPage(int32 InPageId, bool bInPreview, const FName& InPreviewChannelName) const;
+
+	TArray<int32> StopPages(const TArray<int32>& InPageIds, EAvaPlaylistPageStopOptions InOptions, bool bInPreview);
+	TArray<int32> StopPages(const TArray<int32>& InPageIds, EAvaPlaylistPageStopOptions InOptions, bool bInPreview, const FName& InPreviewChannelName);
+
+	bool StopPage(int32 InPageId, EAvaPlaylistPageStopOptions InOptions, bool bInPreview) { return !StopPages({InPageId}, InOptions, bInPreview).IsEmpty(); }
+	bool StopPage(int32 InPageId, EAvaPlaylistPageStopOptions InOptions, bool bInPreview, const FName& InPreviewChannelName) { return !StopPages({InPageId}, InOptions, bInPreview, InPreviewChannelName).IsEmpty(); }
+	
+	bool CanStopPage(int32 InPageId, EAvaPlaylistPageStopOptions InOptions, bool bInPreview) const;
+	bool CanStopPage(int32 InPageId, EAvaPlaylistPageStopOptions InOptions, bool bInPreview, const FName& InPreviewChannelName) const;
+
+	bool StopChannel(const FString& InChannelName);
+	bool CanStopChannel(const FString& InChannelName) const;
+
+	bool ContinuePage(int32 InPageId, bool bInPreview);
+	bool ContinuePage(int32 InPageId, bool bInPreview, const FName& InPreviewChannelName);
+
+	bool CanContinuePage(int32 InPageId, bool bInPreview) const;
+	bool CanContinuePage(int32 InPageId, bool bInPreview, const FName& InPreviewChannelName) const;
+
+	/** Used to reconcile playing state with a remote playback if connection was lost. */
+	bool RestorePlaySubPage(int32 InPageId, int32 InSubPageIndex, const FGuid& InExistingInstanceId, bool bInIsPreview, const FName& InPreviewChannelName);
+
+	int32 AddSubList();
+
+	/**
+	 * Return the current playing Page Ids on the specified program channel.
+	 * If program channel is none, returns all playing pages on all program channels. 
+	 */
+	TArray<int32> GetPlayingPageIds(const FName InProgramChannelName = NAME_None) const;
+
+	/**
+	 * Return the current previewing Page Ids on the specified preview channel.
+	 * If preview channel is none, returns all previewing pages on all channels. 
+	 */
+	TArray<int32> GetPreviewingPageIds(const FName InPreviewChannelName = NAME_None) const;
+
+	const FAvaPageListReference& GetActivePageListReference() const { return ActivePageList; }
+
+	bool SetActivePageList(const FAvaPageListReference& InPageListReference);
+
+	/** Returns true only if a sub list, not the main list, is active. */
+	bool HasActiveSubList() const;
+
+	const FAvalancheSubList& GetActiveSubList() const { return GetSubList(ActivePageList.SubListIndex); }
+	FAvalancheSubList& GetActiveSubList() { return GetSubList(ActivePageList.SubListIndex); }
+
+	const FAvalancheSubList& GetSubList(int32 InSubListIndex) const;
+
+	FAvalancheSubList& GetSubList(int32 InSubListIndex);
+
+	bool IsValidSubListIndex(int32 InIndex) const { return SubLists.IsValidIndex(InIndex); }
+	bool IsValidSubList(const FAvaPageListReference& InPageListReference) const;
+
+	const TArray<FAvalancheSubList>& GetSubLists() const { return SubLists; }
+
+	bool AddPageToSubList(int32 InSubListIndex, int32 InPageId, const FAvaPageInsertPosition& InInsertPosition = FAvaPageInsertPosition());
+	bool AddPagesToSubList(int32 InSubListIndex, const TArray<int32>& InPages);
+
+	int32 RemovePagesFromSubList(int32 InSubListIndex, const TArray<int32>& InPages);
+
+	DECLARE_MULTICAST_DELEGATE(FOnActiveListChanged)
+	FOnActiveListChanged& GetOnActiveListChanged() { return OnActiveListChanged; }
+
+	UTextureRenderTarget2D* GetPreviewRenderTarget() const { return GetPreviewRenderTarget(GetDefaultPreviewChannelName());}
+	UTextureRenderTarget2D* GetPreviewRenderTarget(const FName& InPreviewChannel) const;
+
+	/** Returns the currently selected preview channel (from the settings). */
+	static FName GetDefaultPreviewChannelName();
+
+	/** Clean up playing status on a system tear down. */
+	void OnParentWordBeginTearDown();
+
+	/**
+	 * Push the page's RC values to the runtime playback instances.
+	 * @remark If updating preview, will update all preview channels by default. 
+	 */
+	bool PushRuntimeRemoteControlValues(int32 InPageId, bool bInIsPreview, const FName& InPreviewChannelName = NAME_None) const;
+	
+	void NotifyPageRemoteControlValueChanged(int32 InPageId, EAvaRemoteControlChanges InRemoteControlChanges);
+	
+	void NotifyPageStopped(int32 InPageId) const
+	{
+		OnPagesChanged.Broadcast(this, GetPage(InPageId), EAvaPageChanges::Status);
+	}
+	
+	void NotifyPageSequenceFinished(int32 InPageId)
+	{
+		OnPagesChanged.Broadcast(this, GetPage(InPageId), EAvaPageChanges::Status);
+	}
+
+#if WITH_EDITOR
+	void NotifyPIEEnded(const bool);
+#endif
+
+	FAvaMediaPlaybackManager& GetPlaybackManager() const;
+
+	/** Access the page loading manager for this playlist. */
+	IAvaPlaylistPageLoadingManager& GetPageLoadingManager()
+	{
+		return PageLoadingManager ? *PageLoadingManager : MakePageLoadingManager();
+	}
+
+protected:
+	int32 AddPageFromTemplateInternal(int32 InTemplateId, const FAvaPageIdGeneratorParams& InIdGeneratorParams = FAvaPageIdGeneratorParams(), const FAvaPageInsertPosition& InInsertAt = FAvaPageInsertPosition());
+	
+	void InitializePage(FAvalanchePage& InOutPage, int32 InPageId, int32 InTemplateId) const;
+
+	/**
+	 * Returns true if the selected page can play on the given channel.
+	 * This enforces that the channel type (program or preview) is compatible with the requested operation.
+	 */
+	bool IsChannelTypeCompatibleForRequest(const FAvalanchePage& InSelectedPage, bool bInIsPreview, const FName& InPreviewChannelName, bool bInLogFailureReason) const;
+
+private:
+	IAvaPlaylistPageLoadingManager& MakePageLoadingManager();
+
+	bool PlayPageNoTransition(const FAvalanchePage& InPage, EAvaPlayType InPlayType, bool bInIsPreview, const FName& InPreviewChannelName);
+	bool PlayPageWithTransition(FAvalanchePageTransitionBuilder& InBuilder, const FAvalanchePage& InPage, EAvaPlayType InPlayType, bool bInIsPreview, const FName& InPreviewChannelName);
+	bool StopPageNoTransition(const FAvalanchePage& InPage, bool bInPreview, const FName& InPreviewChannelName);
+	bool StopPageWithTransition(FAvalanchePageTransitionBuilder& InBuilder, const FAvalanchePage& InPage, bool bInPreview, const FName& InPreviewChannelName);
+	
+	const FAvalanchePage& GetNextFromPages(const TArray<FAvalanchePage>& InPages, int32 InStartingIndex) const;
+	FAvalanchePage& GetNextFromPages(TArray<FAvalanchePage>& InPages, int32 InStartingIndex) const;
+
+	const FAvalanchePage& GetNextFromSubList(const TArray<int32>& InSubListIds, int32 InStartingIndex) const;
+	FAvalanchePage& GetNextFromSubList(TArray<int32>& InSubListIds, int32 InStartingIndex);
+
+protected:
+	UPROPERTY()
+	TArray<FAvalanchePage> Pages_DEPRECATED;
+
+	UPROPERTY()
+	FAvalanchePageCollection TemplatePages;
+
+	UPROPERTY()
+	FAvalanchePageCollection InstancedPages;
+
+	UPROPERTY()
+	TArray<FAvalancheSubList> SubLists;
+
+	/** ==InstancePageList Indicates that the entire list is being played, rather than a specific view. */
+	UPROPERTY(Transient, DuplicateTransient, TextExportTransient)
+	FAvaPageListReference ActivePageList = InstancePageList;
+
+	/**
+	 * Keeping track of playing pages.
+	 * 
+	 * Note: For the playlist editor, we can only preview one page on the selected preview channel,
+	 * however, playlist server will eventually require that more than one page is previewed on different preview channels
+	 * for different connected clients.
+	 */
+	UPROPERTY(Transient, DuplicateTransient, TextExportTransient)
+	TArray<TObjectPtr<UAvalanchePagePlayer>> PagePlayers;
+	
+	UPROPERTY(Transient, DuplicateTransient, TextExportTransient)
+	TArray<TObjectPtr<UAvalanchePageTransition>> PageTransitions;
+
+	TUniquePtr<FAvaPageListPlaybackContextCollection> PageListPlaybackContextCollection;
+
+	TUniquePtr<IAvaPlaylistPageLoadingManager> PageLoadingManager;
+	
+	/**
+	 * Playback Client Watcher ensures external playback events are reconciled.
+	 */
+	friend class FAvaPlaylistPlaybackClientWatcher;
+	TPimplPtr<FAvaPlaylistPlaybackClientWatcher> PlaybackClientWatcher;
+
+protected:
+	void AddPagePlayer(UAvalanchePagePlayer* InPagePlayer);
+	
+public:
+	DECLARE_MULTICAST_DELEGATE_TwoParams(FOnPagePlayerEvent, UAvalanchePlaylist*, UAvalanchePagePlayer*);
+
+	/**
+	 * Playlist Player Event - Called when a new page player is added to the playback context.
+	 */
+	FOnPagePlayerEvent OnPagePlayerAdded;
+	
+	/**
+	 * Playlist Player Event - Called when a stopped page player is about to be removed from the playback context.
+	 */
+	FOnPagePlayerEvent OnPagePlayerRemoving;
+
+	FOnPagePlayerEvent& GetOnPagePlayerAdded() { return OnPagePlayerAdded; }
+	FOnPagePlayerEvent& GetOnPagePlayerRemoving() { return OnPagePlayerRemoving; }
+	
+	void AddPageTransition(UAvalanchePageTransition* InPageTransition)
+	{
+		PageTransitions.Add(InPageTransition);
+	}
+
+	void RemovePageTransition(UAvalanchePageTransition* InPageTransition)
+	{
+		PageTransitions.Remove(InPageTransition);
+	}
+
+	bool CanStartTransitionForPage(const FAvalanchePage& InPage, bool bInIsPreview, const FName& InPreviewChannelName) const;
+
+	void StopPageTransitionsForPage(const FAvalanchePage& InPage, bool bInIsPreview, const FName& InPreviewChannelName);
+
+	const TArray<TObjectPtr<UAvalanchePagePlayer>>& GetPagePlayers() const { return PagePlayers; }
+	
+	UAvalanchePagePlayer* FindPlayerForProgramPage(int32 InPageId) const
+	{
+		const TObjectPtr<UAvalanchePagePlayer>* FoundPlayer = PagePlayers.FindByPredicate([InPageId](const UAvalanchePagePlayer* InPagePlayer)
+		{
+			return InPagePlayer->PageId == InPageId && !InPagePlayer->bIsPreview;
+		});
+		return FoundPlayer ? *FoundPlayer : nullptr;
+	}
+
+	UAvalanchePagePlayer* FindPlayerForPreviewPage(int32 InPageId, const FName& InPreviewChannelFName) const
+	{
+		const TObjectPtr<UAvalanchePagePlayer>* FoundPlayer = PagePlayers.FindByPredicate([InPageId, InPreviewChannelFName](const UAvalanchePagePlayer* InPagePlayer)
+		{
+			return InPagePlayer->PageId == InPageId && InPagePlayer->bIsPreview && InPagePlayer->ChannelFName == InPreviewChannelFName;
+		});
+		return FoundPlayer ? *FoundPlayer : nullptr;
+	}
+	
+	UAvalanchePagePlayer* FindPlayerForPage(int32 InPageId, bool bInIsPreview, const FName& InPreviewChannelName) const
+	{
+		return bInIsPreview ?  FindPlayerForPreviewPage(InPageId, InPreviewChannelName) : FindPlayerForProgramPage(InPageId);
+	}
+
+	void RemoveStoppedPagePlayers();
+
+	FAvaPageListPlaybackContextCollection* GetPageListPlaybackContextCollection() const
+	{
+		return PageListPlaybackContextCollection.Get();
+	}
+
+	FAvaPageListPlaybackContextCollection& GetOrCreatePageListPlaybackContextCollection()
+	{
+		if (!PageListPlaybackContextCollection.IsValid())
+		{
+			PageListPlaybackContextCollection = MakeUnique<FAvaPageListPlaybackContextCollection>();
+		}
+		check(PageListPlaybackContextCollection.IsValid());
+		return *PageListPlaybackContextCollection;
+	}
+
+protected:
+	FOnAvaPagesChanged OnPagesChanged;
+	FOnActiveListChanged OnActiveListChanged;
+};
diff --git a/Engine/Plugins/Experimental/Avalanche/Source/AvalancheMedia/Public/Playlist/AvalanchePlaylistComponent.h b/Engine/Plugins/Experimental/Avalanche/Source/AvalancheMedia/Public/Playlist/AvalanchePlaylistComponent.h
new file mode 100644
--- /dev/null
+++ b/Engine/Plugins/Experimental/Avalanche/Source/AvalancheMedia/Public/Playlist/AvalanchePlaylistComponent.h
@@ -0,0 +1,42 @@
+// Copyright Epic Games, Inc. All Rights Reserved.
+
+#pragma once
+
+#include "Components/ActorComponent.h"
+#include "AvalanchePlaylistComponent.generated.h"
+
+class UAvalanchePlaylist;
+
+/**
+ * Add this actor component to blueprint actor to expose the API to control an
+ * Avalanche playlist in game.
+ */
+UCLASS(BlueprintType, Blueprintable, ClassGroup = "Motion Design Media", meta = (BlueprintSpawnableComponent))
+class AVALANCHEMEDIA_API UAvalanchePlaylistComponent : public UActorComponent
+{
+	GENERATED_BODY()
+
+public:
+	UAvalanchePlaylistComponent(const FObjectInitializer& ObjectInitializer);
+	
+	UPROPERTY(EditAnywhere, Category = "Motion Design Media", meta = (DisplayName = "Motion Design Playlist"))
+	TObjectPtr<UAvalanchePlaylist> Playlist = nullptr;
+	
+	UFUNCTION(BlueprintCallable, Category = "Motion Design Media")
+	bool PlayPage(int32 InPageId);
+
+	UFUNCTION(BlueprintCallable, Category = "Motion Design Media")
+	bool StopPage(int32 InPageId);
+
+	UFUNCTION(BlueprintCallable, Category = "Motion Design Media")
+	int32 GetNumberOfPages() const;
+	
+	UFUNCTION(BlueprintCallable, Category = "Motion Design Media")
+	int32 GetPageIdForIndex(int32 InPageIndex) const;
+
+	virtual void InitializeComponent() override;
+	virtual void UninitializeComponent() override;
+
+protected:
+	void OnWorldBeginTearDown(UWorld* InWorld);
+};
diff --git a/Engine/Plugins/Experimental/Avalanche/Source/AvalancheMedia/Public/Playlist/Transition/AvalanchePageTransition.h b/Engine/Plugins/Experimental/Avalanche/Source/AvalancheMedia/Public/Playlist/Transition/AvalanchePageTransition.h
new file mode 100644
--- /dev/null
+++ b/Engine/Plugins/Experimental/Avalanche/Source/AvalancheMedia/Public/Playlist/Transition/AvalanchePageTransition.h
@@ -0,0 +1,94 @@
+// Copyright Epic Games, Inc. All Rights Reserved.
+
+#pragma once
+
+#include "AvaTag.h"
+#include "Playback/AvalanchePlayable.h"
+#include "Playback/Transition/AvaMediaPlaybackTransition.h"
+#include "UObject/Object.h"
+
+#include "AvalanchePageTransition.generated.h"
+
+class FAvaPlayableTransitionBuilder;
+class UAvaRundownPlaybackInstancePlayer;
+class UAvalanchePagePlayer;
+class UAvalanchePlayableTransition;
+class UAvalanchePlayback;
+class UAvalanchePlaylist;
+enum class EAvaMediaPlayableTransitionEntryRole : uint8;
+
+/**
+ * Class for creating and tracking page transitions in the playlist.
+ * It is responsible for creating the playable transition object when requested from
+ * the playback graphs.
+ */
+UCLASS()
+class AVALANCHEMEDIA_API UAvalanchePageTransition : public UAvaMediaPlaybackTransition
+{
+	GENERATED_BODY()
+	
+public:
+	bool AddEnterPage(UAvalanchePagePlayer* InPagePlayer);
+	bool AddPlayingPage(UAvalanchePagePlayer* InPagePlayer);
+	bool AddExitPage(UAvalanchePagePlayer* InPagePlayer);
+	
+	//~ Begin IAvaPlayableVisibilityConstraint
+	virtual bool IsVisibilityConstrained(const UAvalanchePlayable* InPlayable) const override;
+	//~ End IAvaPlayableVisibilityConstraint
+	
+	//~ Begin UAvaMediaPlaybackTransition
+	virtual bool CanStart(bool& bOutShouldDiscard) const override;
+	virtual void Start() override;
+	virtual void Stop() override;
+	virtual bool IsRunning() const override;
+	//~ End UAvaMediaPlaybackTransition
+
+	/** Returns the channel this transition is happening in. A transition can only have pages within the same channel. */
+	FName GetChannelName() const { return ChannelName; }
+
+	bool HasEnterPages() const { return !EnterPlayersWeak.IsEmpty(); }
+
+	bool HasEnterPagesWithNoTransitionLogic() const;
+
+	bool HasPagePlayer(const UAvalanchePagePlayer* InPagePlayer) const;
+
+	bool ContainsTransitionLayer(const FAvaTagId& InTagId) const;
+
+	UAvalanchePlaylist* GetPlaylist() const;
+
+protected:
+	UAvaRundownPlaybackInstancePlayer* FindInstancePlayerForPlayable(const UAvalanchePlayable* InPlayable) const;
+	UAvalanchePagePlayer* FindPagePlayerForPlayable(const UAvalanchePlayable* InPlayable) const;
+
+	void OnTransitionEvent(UAvalanchePlayable* InPlayable, UAvalanchePlayableTransition* InTransition, EAvalanchePlayableTransitionEventFlags InTransitionFlags);
+	void OnPlayableCreated(UAvalanchePlayback* InPlayback, UAvalanchePlayable* InPlayable);
+
+	void AddPlayersToBuilder(FAvaPlayableTransitionBuilder& InOutBuilder, const TArray<TWeakObjectPtr<UAvalanchePagePlayer>>& InPlayersWeak, const TCHAR* InCategory, EAvaMediaPlayableTransitionEntryRole InEntryRole) const;
+	void AddPlayablesToBuilder(FAvaPlayableTransitionBuilder& InOutBuilder, const UAvalanchePagePlayer* InPlayer, const TCHAR* InCategory, EAvaMediaPlayableTransitionEntryRole InEntryRole) const;
+
+	void MakePlayableTransition();
+
+	void LogDetailedTransitionInfo() const;
+	FString GetBriefTransitionDescription() const;
+
+	void RegisterToPlayableTransitionEvent();
+	void UnregisterFromPlayableTransitionEvent() const;
+
+	void RegisterEnterPagePlayerEvents(UAvalanchePagePlayer* InPagePlayer);
+	void UnregisterEnterPagePlayerEvents(UAvalanchePagePlayer* InPagePlayer);
+
+	bool AddPagePlayer(UAvalanchePagePlayer* InPagePlayer, TArray<TWeakObjectPtr<UAvalanchePagePlayer>>& OutPagePlayersWeak);	
+	void UpdateChannelName(const UAvalanchePagePlayer* InPagePlayer);
+
+protected:
+	FName ChannelName;
+	
+	TArray<TWeakObjectPtr<UAvalanchePagePlayer>> EnterPlayersWeak;
+	TArray<TWeakObjectPtr<UAvalanchePagePlayer>> PlayingPlayersWeak;
+	TArray<TWeakObjectPtr<UAvalanchePagePlayer>> ExitPlayersWeak;
+
+	TSet<FAvaTagId> CachedTransitionLayers;
+	
+	UPROPERTY(Transient)
+	TObjectPtr<UAvalanchePlayableTransition> PlayableTransition;
+};
\ No newline at end of file
diff --git a/Engine/Plugins/Experimental/Avalanche/Source/AvalancheMediaEditor/Private/AvaMediaEditorModule.cpp b/Engine/Plugins/Experimental/Avalanche/Source/AvalancheMediaEditor/Private/AvaMediaEditorModule.cpp
new file mode 100644
--- /dev/null
+++ b/Engine/Plugins/Experimental/Avalanche/Source/AvalancheMediaEditor/Private/AvaMediaEditorModule.cpp
@@ -0,0 +1,469 @@
+// Copyright Epic Games, Inc. All Rights Reserved.
+
+#include "AvaMediaEditorModule.h"
+#include "AvaMediaEditorStyle.h"
+#include "Broadcast/AvaBroadcastEditor.h"
+#include "Editor.h"
+#include "Engine/GameViewportClient.h"
+#include "Framework/Application/SlateApplication.h"
+#include "IAvaMediaModule.h"
+#include "LevelEditor.h"
+#include "LevelEditor/AvaLevelEditorUtils.h"
+#include "LevelEditorViewport.h"
+#include "OutputDevices/AvaMediaIOOutputConfigurationCustomization.h"
+#include "Playback/AvaPlaybackCommands.h"
+#include "Playback/Graph/AvaPlaybackConnectionDrawingPolicy.h"
+#include "Playlist/AvaPlaylistCommands.h"
+#include "Playlist/AvaRundownEditorSettings.h"
+#include "Playlist/AvaRundownMacroCollection.h"
+#include "Playlist/Customization/AvaRundownMacroCommandCustomization.h"
+#include "Playlist/Customization/AvaRundownMacroKeyBindingCustomization.h"
+#include "Playlist/Factories/Filters/AvaPlaylistFactoriesUtils.h"
+#include "Playlist/Factories/Filters/AvaPlaylistFilterChannelExpressionFactory.h"
+#include "Playlist/Factories/Filters/AvaPlaylistFilterIdExpressionFactory.h"
+#include "Playlist/Factories/Filters/AvaPlaylistFilterNameExpressionFactory.h"
+#include "Playlist/Factories/Filters/AvaPlaylistFilterPathExpressionFactory.h"
+#include "Playlist/Factories/Filters/AvaPlaylistFilterStatusExpressionFactory.h"
+#include "Playlist/Factories/Filters/AvaPlaylistFilterTransitionLayerExpressionFactory.h"
+#include "Playlist/Factories/Filters/IAvaPlaylistFilterExpressionFactory.h"
+#include "Playlist/Factories/Filters/IAvaPlaylistFilterSuggestionFactory.h"
+#include "Playlist/Factories/Suggestions/AvaPlaylistFilterChannelSuggestionFactory.h"
+#include "Playlist/Factories/Suggestions/AvaPlaylistFilterIdSuggestionFactory.h"
+#include "Playlist/Factories/Suggestions/AvaPlaylistFilterNameSuggestionFactory.h"
+#include "Playlist/Factories/Suggestions/AvaPlaylistFilterPathSuggestionFactory.h"
+#include "Playlist/Factories/Suggestions/AvaPlaylistFilterStatusSuggestionFactory.h"
+#include "Playlist/Factories/Suggestions/AvaPlaylistFilterTransitionLayerSuggestionFactory.h"
+#include "ToolMenus.h"
+
+#define LOCTEXT_NAMESPACE "AvaMediaEditorModule"
+
+DEFINE_LOG_CATEGORY(LogAvaMediaEditor);
+
+namespace UE::AvaMediaEditorModule::Private
+{
+	namespace BroadcastEditorEntry
+	{
+		static const FName MenuName(TEXT("LevelEditor.StatusBar.ToolBar"));
+		static const FName SectionName(TEXT("AvalancheMedia"));
+	}
+	
+	// Command line parsing helper.
+	bool IsPlaylistServerManuallyStarted(FString& OutPlaylistServerName)
+	{
+		return FParse::Value(FCommandLine::Get(), TEXT("AvaMediaPlaylistServerStart="), OutPlaylistServerName) ||
+			FParse::Param(FCommandLine::Get(), TEXT("AvaMediaPlaylistServerStart"));
+	}
+
+	static void GetEditorViewportClient(FCommonViewportClient** OutViewportClient)
+	{
+		// Replicating the logic from UUnrealEdEngine::Exec.
+		FCommonViewportClient* ViewportClient = nullptr;
+		if (!GStatProcessingViewportClient && (!GEngine->GameViewport || GEngine->GameViewport->IsSimulateInEditorViewport() ))
+		{
+			ViewportClient = GLastKeyLevelEditingViewportClient ? GLastKeyLevelEditingViewportClient : GCurrentLevelEditingViewportClient;
+		}
+
+		if (ViewportClient)
+		{
+			*OutViewportClient = ViewportClient;
+		}
+	}
+}
+
+void FAvaMediaEditorModule::StartupModule()
+{
+	using namespace UE::AvaMediaEditorModule::Private;
+
+	InitExtensibilityManagers();
+
+	FAvaPlaybackCommands::Register();
+	FAvaPlaylistCommands::Register();
+
+	FAvaMediaEditorStyle::Initialize();
+
+	if (FSlateApplication::IsInitialized())
+	{
+		AddEditorToolbarButtons();
+	}
+
+	// Register the Avalanche Playback Graph connection policy with the graph editor
+	PlaybackConnectionFactory = MakeShared<FAvaPlaybackConnectionDrawingPolicyFactory>();
+	FEdGraphUtilities::RegisterVisualPinConnectionFactory(PlaybackConnectionFactory);
+
+	// MediaIOEditor module is loaded in PostEngineInit phase,
+	// so we need in order to have our customizations override theirs, we need to
+	// register ours after, i.e. once all modules are loaded.
+	FCoreDelegates::OnAllModuleLoadingPhasesComplete.AddRaw(this, &FAvaMediaEditorModule::RegisterCustomizations);
+	FCoreDelegates::OnPostEngineInit.AddRaw(this, &FAvaMediaEditorModule::PostEngineInit);
+	FCoreDelegates::OnEnginePreExit.AddRaw(this, &FAvaMediaEditorModule::EnginePreExit);
+
+	// Register Map Change Events
+	FLevelEditorModule& LevelEditor = FModuleManager::LoadModuleChecked<FLevelEditorModule>("LevelEditor");
+	LevelEditor.OnMapChanged().AddRaw(this, &FAvaMediaEditorModule::HandleMapChanged);
+
+	ConsoleCmds.Add(IConsoleManager::Get().RegisterConsoleCommand(
+		TEXT("AvaPlaylistServer.Start"),
+		TEXT("Starts the playlist server."),
+		FConsoleCommandWithArgsDelegate::CreateRaw(this, &FAvaMediaEditorModule::StartPlaylistServerCommand),
+		ECVF_Default
+	));
+	ConsoleCmds.Add(IConsoleManager::Get().RegisterConsoleCommand(
+		TEXT("AvaPlaylistServer.Stop"),
+		TEXT("Stops the playlist server."),
+		FConsoleCommandWithArgsDelegate::CreateRaw(this, &FAvaMediaEditorModule::StopPlaylistServerCommand),
+		ECVF_Default
+	));
+
+	IAvaMediaModule::Get().GetEditorViewportClientDelegate().BindStatic(&GetEditorViewportClient);
+
+	FString DummyServerName;
+	if (IsPlaylistServerManuallyStarted(DummyServerName))
+	{
+		// Prevent throttling when the server is started.
+		// This has to be done before any SLevelViewport are ticked since the cvar value is cached on first tick.
+		static const FSlateThrottleManager& ThrottleManager = FSlateThrottleManager::Get();
+		if (IConsoleVariable* AllowThrottling = IConsoleManager::Get().FindConsoleVariable(TEXT("Slate.bAllowThrottling")))
+		{
+			AllowThrottling->Set(0);
+			UE_LOG(LogAvaMediaEditor, Log, TEXT("Setting Slate.bAllowThrottling to false."));
+		}
+	}
+	RegisterPlaylistFilterExpressionFactories();
+	RegisterPlaylistFilterSuggestionFactories();
+}
+
+void FAvaMediaEditorModule::ShutdownModule()
+{
+	StopAllServices();
+
+	FCoreDelegates::OnAllModuleLoadingPhasesComplete.RemoveAll(this);
+	FCoreDelegates::OnPostEngineInit.RemoveAll(this);
+	FCoreDelegates::OnEnginePreExit.RemoveAll(this);
+
+	// Unregister Map Change Events
+	if (FLevelEditorModule* LevelEditor = FModuleManager::GetModulePtr<FLevelEditorModule>("LevelEditor"))
+	{
+		LevelEditor->OnMapChanged().RemoveAll(this);
+	}
+
+	ResetExtensibilityManagers();
+	if (FSlateApplication::IsInitialized())
+	{
+		RemoveEditorToolbarButtons();
+	}
+	if (UObjectInitialized() && !IsEngineExitRequested())
+	{
+		UnregisterCustomizations();
+	}
+
+	FAvaMediaEditorStyle::Shutdown();
+	FAvaPlaybackCommands::Unregister();
+	FAvaPlaylistCommands::Unregister();
+
+	if (PlaybackConnectionFactory.IsValid())
+	{
+		FEdGraphUtilities::UnregisterVisualPinConnectionFactory(PlaybackConnectionFactory);
+		PlaybackConnectionFactory.Reset();
+	}
+
+	for (IConsoleObject* ConsoleCmd : ConsoleCmds)
+	{
+		IConsoleManager::Get().UnregisterConsoleObject(ConsoleCmd);
+	}
+	ConsoleCmds.Empty();
+}
+
+TSharedPtr<FExtensibilityManager> FAvaMediaEditorModule::GetBroadcastToolBarExtensibilityManager()
+{
+	return BroadcastToolBarExtensibility;
+}
+
+TSharedPtr<FExtensibilityManager> FAvaMediaEditorModule::GetPlaybackToolBarExtensibilityManager()
+{
+	return PlaybackToolBarExtensibility;
+}
+
+TSharedPtr<FExtensibilityManager> FAvaMediaEditorModule::GetPlaylistToolBarExtensibilityManager()
+{
+	return PlaylistToolBarExtensibility;
+}
+
+TSharedPtr<FExtensibilityManager> FAvaMediaEditorModule::GetPlaylistMenuExtensibilityManager()
+{
+	return PlaylistMenuExtensibility;
+}
+
+FSlateIcon FAvaMediaEditorModule::GetToolbarBroadcastButtonIcon() const
+{
+	const IAvaMediaModule& MediaModule = IAvaMediaModule::Get();
+
+	if (MediaModule.IsMediaPlaybackClientStarted())
+	{
+		return FSlateIcon(FAvaMediaEditorStyle::GetStyleSetName(), "AvalancheMediaEditor.BroadcastClient", "AvalancheMediaEditor.BroadcastClient.Small");
+	}
+	else if (MediaModule.IsMediaPlaybackServerStarted())
+	{
+		return FSlateIcon(FAvaMediaEditorStyle::GetStyleSetName(), "AvalancheMediaEditor.BroadcastServer", "AvalancheMediaEditor.BroadcastServer.Small");
+	}
+	return FSlateIcon(FAvaMediaEditorStyle::GetStyleSetName(), "AvalancheMediaEditor.BroadcastIcon");
+}
+
+bool FAvaMediaEditorModule::CanFilterSupportComparisonOperation(const FName& InFilterKey, ETextFilterComparisonOperation InOperation, EAvaPlaylistSearchListType InPlaylistSearchListType) const
+{
+	if (const TSharedPtr<IAvaPlaylistFilterExpressionFactory>* FilterExpressionFactory = FilterExpressionFactories.Find(InFilterKey))
+	{
+		return FilterExpressionFactory->Get()->SupportsComparisonOperation(InOperation, InPlaylistSearchListType);
+	}
+	return false;
+}
+
+bool FAvaMediaEditorModule::FilterExpression(const FName& InFilterKey, const FAvalanchePage& InItem, const FAvaPlaylistTextFilterArgs& InArgs) const
+{
+	if (const TSharedPtr<IAvaPlaylistFilterExpressionFactory>* FilterExpressionFactory = FilterExpressionFactories.Find(InFilterKey))
+	{
+		return FilterExpressionFactory->Get()->FilterExpression(InItem, InArgs);
+	}
+	return false;
+}
+
+TArray<TSharedPtr<IAvaPlaylistFilterSuggestionFactory>> FAvaMediaEditorModule::GetSimpleSuggestions(EAvaPlaylistSearchListType InSuggestionType) const
+{
+	TArray<TSharedPtr<IAvaPlaylistFilterSuggestionFactory>> OutArray;
+
+	for (const TPair<FName, TSharedPtr<IAvaPlaylistFilterSuggestionFactory>>& Suggestion : FilterSuggestionFactories)
+	{
+		if (Suggestion.Value->SupportSuggestionType(InSuggestionType) && Suggestion.Value->IsSimpleSuggestion())
+		{
+			OutArray.Add(Suggestion.Value);
+		}
+	}
+
+	return OutArray;
+}
+
+TArray<TSharedPtr<IAvaPlaylistFilterSuggestionFactory>> FAvaMediaEditorModule::GetComplexSuggestions(EAvaPlaylistSearchListType InSuggestionType) const
+{
+	TArray<TSharedPtr<IAvaPlaylistFilterSuggestionFactory>> OutArray;
+
+	for (const TPair<FName, TSharedPtr<IAvaPlaylistFilterSuggestionFactory>>& Suggestion : FilterSuggestionFactories)
+	{
+		if (Suggestion.Value->SupportSuggestionType(InSuggestionType) && !Suggestion.Value->IsSimpleSuggestion())
+		{
+			OutArray.Add(Suggestion.Value);
+		}
+	}
+
+	return OutArray;
+}
+
+void FAvaMediaEditorModule::AddEditorToolbarButtons()
+{
+	using namespace UE::AvaMediaEditorModule::Private;
+
+	FToolMenuEntry OpenBroadcastButtonEntry = FToolMenuEntry::InitToolBarButton(TEXT("OpenBroadcastToolbarButton")
+		, FExecuteAction::CreateStatic(&FAvaBroadcastEditor::OpenBroadcastEditor)
+		, LOCTEXT("OpenBroadcast_Title", "Broadcast")
+		, LOCTEXT("OpenBroadcast_Tooltip", "Opens the Avalanche Broadcast Editor Window")
+		, TAttribute<FSlateIcon>::Create([]() { return IAvaMediaEditorModule::Get().GetToolbarBroadcastButtonIcon(); })
+	);
+	OpenBroadcastButtonEntry.StyleNameOverride = TEXT("CalloutToolbar"); // Display Labels
+	
+	if (UToolMenu* const Menu = UToolMenus::Get()->ExtendMenu(BroadcastEditorEntry::MenuName))
+	{
+		FToolMenuSection& Section = Menu->FindOrAddSection(BroadcastEditorEntry::SectionName);
+		Section.AddEntry(OpenBroadcastButtonEntry);
+	}
+}
+
+void FAvaMediaEditorModule::RemoveEditorToolbarButtons()
+{
+	using namespace UE::AvaMediaEditorModule::Private;
+
+	if (GIsEditor && UObjectInitialized())
+	{
+		UToolMenus::Get()->RemoveSection(BroadcastEditorEntry::MenuName, BroadcastEditorEntry::SectionName);
+	}
+}
+
+void FAvaMediaEditorModule::OpenBroadcastEditor(const TArray<FString>& InArguments)
+{
+	FAvaBroadcastEditor::OpenBroadcastEditor();
+}
+
+void FAvaMediaEditorModule::InitExtensibilityManagers()
+{
+	BroadcastToolBarExtensibility = MakeShared<FExtensibilityManager>();
+	PlaybackToolBarExtensibility  = MakeShared<FExtensibilityManager>();
+	PlaylistToolBarExtensibility  = MakeShared<FExtensibilityManager>();
+	PlaylistMenuExtensibility     = MakeShared<FExtensibilityManager>();
+}
+
+void FAvaMediaEditorModule::ResetExtensibilityManagers()
+{
+	BroadcastToolBarExtensibility = nullptr;
+	PlaybackToolBarExtensibility  = nullptr;
+	PlaylistToolBarExtensibility  = nullptr;
+	PlaylistMenuExtensibility     = nullptr;
+}
+
+void FAvaMediaEditorModule::RegisterCustomizations() const
+{
+	FPropertyEditorModule& PropertyModule = FModuleManager::LoadModuleChecked<FPropertyEditorModule>("PropertyEditor");
+	PropertyModule.RegisterCustomPropertyTypeLayout(FMediaIOOutputConfiguration::StaticStruct()->GetFName(), FOnGetPropertyTypeCustomizationInstance::CreateStatic(&FAvaMediaIOOutputConfigurationCustomization::MakeInstance));
+	PropertyModule.RegisterCustomPropertyTypeLayout(FAvaRundownMacroCommand::StaticStruct()->GetFName(), FOnGetPropertyTypeCustomizationInstance::CreateStatic(&FAvaRundownMacroCommandCustomization::MakeInstance));
+	PropertyModule.RegisterCustomPropertyTypeLayout(FAvaRundownMacroKeyBinding::StaticStruct()->GetFName(), FOnGetPropertyTypeCustomizationInstance::CreateStatic(&FAvaRundownMacroKeyBindingCustomization::MakeInstance));
+}
+
+void FAvaMediaEditorModule::UnregisterCustomizations() const
+{
+	FPropertyEditorModule& PropertyModule = FModuleManager::LoadModuleChecked<FPropertyEditorModule>("PropertyEditor");
+	PropertyModule.UnregisterCustomPropertyTypeLayout(FMediaIOOutputConfiguration::StaticStruct()->GetFName());
+	PropertyModule.UnregisterCustomPropertyTypeLayout(FAvaRundownMacroCommand::StaticStruct()->GetFName());
+	PropertyModule.UnregisterCustomPropertyTypeLayout(FAvaRundownMacroKeyBinding::StaticStruct()->GetFName());
+}
+
+void FAvaMediaEditorModule::StartPlaylistServerCommand(const TArray<FString>& Args)
+{
+	if (PlaylistServer)
+	{
+		UE_LOG(LogAvaMediaEditor, Log, TEXT("Playlist Server is already started."));
+		return;
+	}
+	
+	PlaylistServer = MakeShared<FAvaPlaylistServer>();
+	
+	// Remark: Only the module's playlist server register console commands to avoid
+	// conflicts with temporary servers (for testing).
+	PlaylistServer->RegisterConsoleCommands();
+
+	PlaylistServer->Init(Args.Num() > 0 ? Args[0] : TEXT(""));
+	OnPlaylistServerStarted.Broadcast();
+
+	UE_LOG(LogAvaMediaEditor, Log, TEXT("Playlist Server Started."));
+}
+
+void FAvaMediaEditorModule::StopPlaylistServerCommand(const TArray<FString>& Args)
+{
+	if (PlaylistServer)
+	{
+		UE_LOG(LogAvaMediaEditor, Log, TEXT("Stopping Playlist Server..."));
+		OnPlaylistServerStopped.Broadcast();
+	}
+	PlaylistServer.Reset();
+}
+
+void FAvaMediaEditorModule::PostEngineInit()
+{
+	using namespace UE::AvaMediaEditorModule::Private;
+
+	// This needs to happen late in the loading process, otherwise it fails.
+	const UAvaRundownEditorSettings* Settings = UAvaRundownEditorSettings::Get();
+
+	// Allow for specification of the server name in the command line.
+	// Command line has priority over project settings.
+	FString ServerName;
+	if (IsPlaylistServerManuallyStarted(ServerName))
+	{
+		StartPlaylistServerCommand({ServerName});
+	}
+	else if (Settings && Settings->bAutoStartRundownServer)
+	{
+		StartPlaylistServerCommand({Settings->RundownServerName});
+	}
+}
+
+void FAvaMediaEditorModule::EnginePreExit()
+{
+	StopAllServices();
+}
+
+void FAvaMediaEditorModule::StopAllServices()
+{
+	StopPlaylistServerCommand({});
+}
+
+void FAvaMediaEditorModule::HandleMapChanged(UWorld* InWorld, EMapChangeType InMapChangeType)
+{
+	EAvaMediaMapChangeType EventType;
+	switch (InMapChangeType)
+	{
+		case EMapChangeType::LoadMap:
+			EventType = EAvaMediaMapChangeType::LoadMap;
+			break;
+
+		case EMapChangeType::SaveMap:
+			EventType = EAvaMediaMapChangeType::SaveMap;
+			break;
+
+		case EMapChangeType::NewMap:
+			EventType = EAvaMediaMapChangeType::NewMap;
+			break; 
+
+		case EMapChangeType::TearDownWorld:
+			EventType = EAvaMediaMapChangeType::TearDownWorld;
+			break; 
+
+		default:
+			EventType = EAvaMediaMapChangeType::None;
+			break; 
+	}
+
+	// Propagate the editor event to the runtime module.
+	IAvaMediaModule::Get().NotifyMapChangedEvent(InWorld, EventType);
+}
+
+template <
+	typename InPlaylistFilterExpressionFactoryType,
+	typename ... InArgsType
+	UE_REQUIRES(TIsDerivedFrom<InPlaylistFilterExpressionFactoryType, IAvaPlaylistFilterExpressionFactory>::Value)
+>
+void FAvaMediaEditorModule::RegisterPlaylistFilterExpressionFactory(InArgsType&&... InArgs)
+{
+	const TSharedRef<IAvaPlaylistFilterExpressionFactory> FilterExpressionFactory =
+		IAvaPlaylistFilterExpressionFactory::MakeInstance<InPlaylistFilterExpressionFactoryType>(Forward<InArgsType>(InArgs)...);
+
+	if (!FilterExpressionFactories.Contains(FilterExpressionFactory->GetFilterIdentifier()))
+	{
+		FilterExpressionFactories.Add(FilterExpressionFactory->GetFilterIdentifier(), FilterExpressionFactory);
+	}
+}
+
+template <
+	typename InPlaylistSuggestionFactoryType,
+	typename ... InArgsType
+	UE_REQUIRES(TIsDerivedFrom<InPlaylistSuggestionFactoryType, IAvaPlaylistFilterSuggestionFactory>::Value)
+>
+void FAvaMediaEditorModule::RegisterPlaylistFilterSuggestionFactory(InArgsType&&... InArgs)
+{
+	const TSharedRef<IAvaPlaylistFilterSuggestionFactory> FilterSuggestionFactory =
+		IAvaPlaylistFilterSuggestionFactory::MakeInstance<InPlaylistSuggestionFactoryType>(Forward<InArgsType>(InArgs)...);
+
+	if (!FilterSuggestionFactories.Contains(FilterSuggestionFactory->GetSuggestionIdentifier()))
+	{
+		FilterSuggestionFactories.Add(FilterSuggestionFactory->GetSuggestionIdentifier(), FilterSuggestionFactory);
+	}
+}
+
+void FAvaMediaEditorModule::RegisterPlaylistFilterExpressionFactories()
+{
+	RegisterPlaylistFilterExpressionFactory<FAvaPlaylistFilterNameExpressionFactory>();
+	RegisterPlaylistFilterExpressionFactory<FAvaPlaylistFilterIdExpressionFactory>();
+	RegisterPlaylistFilterExpressionFactory<FAvaPlaylistFilterPathExpressionFactory>();
+	RegisterPlaylistFilterExpressionFactory<FAvaPlaylistFilterChannelExpressionFactory>();
+	RegisterPlaylistFilterExpressionFactory<FAvaPlaylistFilterStatusExpressionFactory>();
+	RegisterPlaylistFilterExpressionFactory<FAvaPlaylistFilterTransitionLayerExpressionFactory>();
+}
+
+void FAvaMediaEditorModule::RegisterPlaylistFilterSuggestionFactories()
+{
+	RegisterPlaylistFilterSuggestionFactory<FAvaPlaylistFilterNameSuggestionFactory>();
+	RegisterPlaylistFilterSuggestionFactory<FAvaPlaylistFilterIdSuggestionFactory>();
+	RegisterPlaylistFilterSuggestionFactory<FAvaPlaylistFilterPathSuggestionFactory>();
+	RegisterPlaylistFilterSuggestionFactory<FAvaPlaylistFilterChannelSuggestionFactory>();
+	RegisterPlaylistFilterSuggestionFactory<FAvaPlaylistFilterStatusSuggestionFactory>();
+	RegisterPlaylistFilterSuggestionFactory<FAvaPlaylistFilterTransitionLayerSuggestionFactory>();
+}
+
+#undef LOCTEXT_NAMESPACE
+
+IMPLEMENT_MODULE(FAvaMediaEditorModule, AvalancheMediaEditor)
diff --git a/Engine/Plugins/Experimental/Avalanche/Source/AvalancheMediaEditor/Private/AvaMediaEditorModule.h b/Engine/Plugins/Experimental/Avalanche/Source/AvalancheMediaEditor/Private/AvaMediaEditorModule.h
new file mode 100644
--- /dev/null
+++ b/Engine/Plugins/Experimental/Avalanche/Source/AvalancheMediaEditor/Private/AvaMediaEditorModule.h
@@ -0,0 +1,101 @@
+// Copyright Epic Games, Inc. All Rights Reserved.
+
+#pragma once
+
+#include "IAvaMediaEditorModule.h"
+#include "Playlist/AvaPlaylistServer.h"
+#include "Templates/UnrealTypeTraits.h"
+
+class IAvaPlaylistFilterExpressionFactory;
+class IAvaPlaylistFilterSuggestionFactory;
+struct FGraphPanelPinConnectionFactory;
+enum class EAvaPlaylistSearchListType : uint8;
+
+class FAvaMediaEditorModule : public IAvaMediaEditorModule
+{
+
+public:
+
+	//IModuleInterface
+	virtual void StartupModule() override;
+	virtual void ShutdownModule() override;
+	//~IModuleInterface
+
+	//IAvaMediaEditorModule
+	virtual TSharedPtr<FExtensibilityManager> GetBroadcastToolBarExtensibilityManager() override;
+	virtual TSharedPtr<FExtensibilityManager> GetPlaybackToolBarExtensibilityManager() override;
+	virtual TSharedPtr<FExtensibilityManager> GetPlaylistToolBarExtensibilityManager() override;
+	virtual TSharedPtr<FExtensibilityManager> GetPlaylistMenuExtensibilityManager() override;
+	virtual FOnPlaylistServerStarted& GetOnPlaylistServerStarted() override { return OnPlaylistServerStarted; }
+	virtual FOnPlaylistServerStopped& GetOnPlaylistServerStopped() override { return OnPlaylistServerStopped; }
+	virtual TSharedPtr<FAvaPlaylistServer> GetPlaylistServer() const override { return PlaylistServer; }
+	virtual bool CanFilterSupportComparisonOperation(const FName& InFilterKey, ETextFilterComparisonOperation InOperation, EAvaPlaylistSearchListType InPlaylistSearchListType) const override;
+	virtual bool FilterExpression(const FName& InFilterKey, const FAvalanchePage& InItem, const FAvaPlaylistTextFilterArgs& InArgs) const override;
+	virtual TArray<TSharedPtr<IAvaPlaylistFilterSuggestionFactory>> GetSimpleSuggestions(EAvaPlaylistSearchListType InSuggestionType) const override;
+	virtual TArray<TSharedPtr<IAvaPlaylistFilterSuggestionFactory>> GetComplexSuggestions(EAvaPlaylistSearchListType InSuggestionType) const override;
+	//~IAvaMediaEditorModule
+
+	void AddEditorToolbarButtons();
+	void RemoveEditorToolbarButtons();
+	virtual FSlateIcon GetToolbarBroadcastButtonIcon() const override;
+
+	static void OpenBroadcastEditor(const TArray<FString>& InArguments);
+
+protected:
+	void InitExtensibilityManagers();
+	void ResetExtensibilityManagers();
+
+	/** Register details view customizations. */
+	void RegisterCustomizations() const;
+
+	/** Unregister details view customizations. */
+	void UnregisterCustomizations() const;
+
+	void StartPlaylistServerCommand(const TArray<FString>& Args);
+	void StopPlaylistServerCommand(const TArray<FString>& Args);
+
+private:
+	void PostEngineInit();
+	void EnginePreExit();
+	void StopAllServices();
+	void HandleMapChanged(UWorld* InWorld, EMapChangeType InMapChangeType);
+	
+	template <
+		typename InPlaylistFilterExpressionFactoryType,
+		typename... InArgsType
+		UE_REQUIRES(TIsDerivedFrom<InPlaylistFilterExpressionFactoryType, IAvaPlaylistFilterExpressionFactory>::Value)
+	>
+	void RegisterPlaylistFilterExpressionFactory(InArgsType&&... InArgs);
+
+	template <
+		typename InPlaylistSuggestionFactoryType,
+		typename... InArgsType
+		UE_REQUIRES(TIsDerivedFrom<InPlaylistSuggestionFactoryType, IAvaPlaylistFilterSuggestionFactory>::Value)
+	>
+	void RegisterPlaylistFilterSuggestionFactory(InArgsType&&... InArgs);
+
+	void RegisterPlaylistFilterExpressionFactories();
+
+	void RegisterPlaylistFilterSuggestionFactories();
+
+private:
+	TSharedPtr<FExtensibilityManager> BroadcastToolBarExtensibility;
+	TSharedPtr<FExtensibilityManager> PlaybackToolBarExtensibility;
+	TSharedPtr<FExtensibilityManager> PlaylistToolBarExtensibility;
+	TSharedPtr<FExtensibilityManager> PlaylistMenuExtensibility;
+
+	TSharedPtr<FGraphPanelPinConnectionFactory> PlaybackConnectionFactory;
+
+	TSharedPtr<FAvaPlaylistServer> PlaylistServer;
+
+	FOnPlaylistServerStarted OnPlaylistServerStarted;
+	FOnPlaylistServerStopped OnPlaylistServerStopped;
+
+	TArray<IConsoleObject*> ConsoleCmds;
+
+	/** Holds all the PlaylistFilterExpressionFactory */
+	TMap<FName, TSharedPtr<IAvaPlaylistFilterExpressionFactory>> FilterExpressionFactories;
+
+	/** Holds all the PlaylistFilterSuggestionFactory */
+	TMap<FName, TSharedPtr<IAvaPlaylistFilterSuggestionFactory>> FilterSuggestionFactories;
+};
diff --git a/Engine/Plugins/Experimental/Avalanche/Source/AvalancheMediaEditor/Private/AvaMediaEditorStyle.cpp b/Engine/Plugins/Experimental/Avalanche/Source/AvalancheMediaEditor/Private/AvaMediaEditorStyle.cpp
new file mode 100644
--- /dev/null
+++ b/Engine/Plugins/Experimental/Avalanche/Source/AvalancheMediaEditor/Private/AvaMediaEditorStyle.cpp
@@ -0,0 +1,191 @@
+// Copyright Epic Games, Inc. All Rights Reserved.
+
+#include "AvaMediaEditorStyle.h"
+#include "Brushes/SlateImageBrush.h"
+#include "Interfaces/IPluginManager.h"
+#include "Styling/AppStyle.h"
+#include "Styling/SlateStyle.h"
+#include "Styling/SlateTypes.h"
+#include "Styling/SlateStyleRegistry.h"
+#include "Styling/StarshipCoreStyle.h"
+#include "Styling/ToolBarStyle.h"
+
+TSharedPtr<FSlateStyleSet> FAvaMediaEditorStyle::StyleInstance = nullptr;
+
+void FAvaMediaEditorStyle::Initialize()
+{
+	if (!StyleInstance.IsValid())
+	{
+		StyleInstance = Create();
+		FSlateStyleRegistry::RegisterSlateStyle(*StyleInstance);
+	}
+}
+
+void FAvaMediaEditorStyle::Shutdown()
+{
+	FSlateStyleRegistry::UnRegisterSlateStyle(*StyleInstance);
+	ensure(StyleInstance.IsUnique());
+	StyleInstance.Reset();
+}
+
+FName FAvaMediaEditorStyle::GetStyleSetName()
+{
+	static FName StyleSetName(TEXT("AvalancheMediaEditor"));
+	return StyleSetName;
+}
+
+const FLinearColor& FAvaMediaEditorStyle::GetColor(FName PropertyName, const ANSICHAR* Specifier)
+{
+	return StyleInstance->GetColor(PropertyName, Specifier);
+}
+
+const FSlateBrush* FAvaMediaEditorStyle::GetBrush(FName PropertyName, const ANSICHAR* Specifier)
+{
+	return StyleInstance->GetBrush(PropertyName, Specifier);
+}
+
+#define IMAGE_BRUSH(RelativePath, ...) FSlateImageBrush(Style->RootToContentDir(RelativePath, TEXT(".png")), __VA_ARGS__)
+#define IMAGE_BRUSH_SVG(RelativePath, ...) FSlateVectorImageBrush(Style->RootToContentDir(RelativePath, TEXT(".svg")), __VA_ARGS__)
+#define BOX_BRUSH(RelativePath, ...) FSlateBoxBrush(Style->RootToContentDir(RelativePath, TEXT(".png") ), __VA_ARGS__ )
+#define CORE_FONT(...) FSlateFontInfo(FCoreStyle::GetDefaultFont(), __VA_ARGS__)
+
+const FVector2D Icon12x12(12.0f, 12.0f);
+const FVector2D Icon16x16(16.0f, 16.0f);
+const FVector2D Icon20x20(20.0f, 20.0f);
+const FVector2D Icon24x24(24.0f, 24.0f);
+const FVector2D Icon40x40(40.0f, 40.0f);
+const FVector2D Icon64x64(64.0f, 64.0f);
+
+TSharedRef<FSlateStyleSet> FAvaMediaEditorStyle::Create()
+{
+	TSharedRef<FSlateStyleSet> Style = MakeShared<FSlateStyleSet>("AvalancheMediaEditor");
+
+	TSharedPtr<IPlugin> Plugin = IPluginManager::Get().FindPlugin(TEXT("Avalanche"));
+	check(Plugin.IsValid());
+	if (Plugin.IsValid())
+	{
+		Style->SetContentRoot(FPaths::Combine(Plugin->GetBaseDir(), TEXT("Resources")));
+	}
+
+	// Asset Class Icons
+	Style->Set("ClassIcon.AvalanchePlayback", new IMAGE_BRUSH("Icons/MediaIcons/AvaPlaybackIcon_16x"     , Icon16x16));
+	Style->Set("ClassThumbnail.AvalanchePlayback", new IMAGE_BRUSH("Icons/MediaIcons/AvaPlaybackIcon_64x", Icon64x64));
+	Style->Set("ClassIcon.AvalanchePlaylist", new IMAGE_BRUSH("Icons/MediaIcons/AvaRundownIcon_16x"     , Icon16x16));
+	Style->Set("ClassThumbnail.AvalanchePlaylist", new IMAGE_BRUSH("Icons/MediaIcons/AvaRundownIcon_64x", Icon64x64));
+	Style->Set("ClassIcon.AvaRundownMacroCollection", new IMAGE_BRUSH("Icons/MediaIcons/AvaRundownMacroCollectionIcon_16x"     , Icon16x16));
+	Style->Set("ClassThumbnail.AvaRundownMacroCollection", new IMAGE_BRUSH("Icons/MediaIcons/AvaRundownMacroCollectionIcon_64x", Icon64x64));
+
+	//Playlist Commands
+	Style->Set("AvaPlaylistCommands.AddPage"         , new IMAGE_BRUSH_SVG("Icons/MediaIcons/AddPage"      , Icon16x16));
+	Style->Set("AvaPlaylistCommands.AddTemplate"     , new IMAGE_BRUSH_SVG("Icons/MediaIcons/AddPage"      , Icon16x16));
+	Style->Set("AvaPlaylistCommands.CreatePageInstanceFromTemplate"  , new IMAGE_BRUSH_SVG("Icons/MediaIcons/AddPage"      , Icon16x16));
+	Style->Set("AvaPlaylistCommands.CreateComboTemplate", new IMAGE_BRUSH_SVG("Icons/MediaIcons/AddPage"   , Icon16x16));
+	Style->Set("AvaPlaylistCommands.RemovePage"      , new IMAGE_BRUSH_SVG("Icons/MediaIcons/RemovePage"   , Icon16x16));
+	Style->Set("AvaPlaylistCommands.RenumberPage"    , new IMAGE_BRUSH_SVG("Icons/MediaIcons/RenumberPage" , Icon16x16));
+	Style->Set("AvaPlaylistCommands.ReimportPage"    , new IMAGE_BRUSH_SVG("Icons/MediaIcons/ReimportPage" , Icon16x16));
+	Style->Set("AvaPlaylistCommands.EditPageSource"  , new IMAGE_BRUSH_SVG("Icons/MediaIcons/OpenAsset"	  , Icon16x16));
+	Style->Set("AvaPlaylistCommands.Play"            , new IMAGE_BRUSH_SVG("Icons/MediaIcons/Play"         , Icon16x16));
+	Style->Set("AvaPlaylistCommands.UpdateValues"    , new IMAGE_BRUSH_SVG("Icons/MediaIcons/UpdateValues" , Icon16x16));
+	Style->Set("AvaPlaylistCommands.Stop"            , new IMAGE_BRUSH_SVG("Icons/MediaIcons/Stop"         , Icon16x16));
+	Style->Set("AvaPlaylistCommands.ForceStop"       , new IMAGE_BRUSH_SVG("Icons/MediaIcons/Stop"         , Icon16x16));
+	Style->Set("AvaPlaylistCommands.Continue"        , new IMAGE_BRUSH_SVG("Icons/MediaIcons/Continue"     , Icon16x16));
+	Style->Set("AvaPlaylistCommands.PlayNext"        , new IMAGE_BRUSH_SVG("Icons/MediaIcons/PlayNext"     , Icon16x16));
+	Style->Set("AvaPlaylistCommands.PreviewFrame"    , new IMAGE_BRUSH_SVG("Icons/MediaIcons/PreviewFrame" , Icon16x16));
+	Style->Set("AvaPlaylistCommands.PreviewPlay"     , new IMAGE_BRUSH_SVG("Icons/MediaIcons/Play"         , Icon16x16));
+	Style->Set("AvaPlaylistCommands.PreviewStop"     , new IMAGE_BRUSH_SVG("Icons/MediaIcons/Stop"         , Icon16x16));
+	Style->Set("AvaPlaylistCommands.PreviewForceStop", new IMAGE_BRUSH_SVG("Icons/MediaIcons/Stop"         , Icon16x16));
+	Style->Set("AvaPlaylistCommands.PreviewContinue" , new IMAGE_BRUSH_SVG("Icons/MediaIcons/Continue"     , Icon16x16));
+	Style->Set("AvaPlaylistCommands.PreviewPlayNext" , new IMAGE_BRUSH_SVG("Icons/MediaIcons/PlayNext"     , Icon16x16));
+	Style->Set("AvaPlaylistCommands.TakeToProgram"   , new IMAGE_BRUSH_SVG("Icons/MediaIcons/Play"		   , Icon16x16));
+
+	//Broadcast
+	Style->Set("AvalancheMediaEditor.OutputIcon"           , new IMAGE_BRUSH("Icons/MediaIcons/MediaOutput"     , Icon20x20));
+	Style->Set("AvalancheMediaEditor.BroadcastIcon"        , new IMAGE_BRUSH("Icons/MediaIcons/MediaOutput"     , Icon20x20));
+	Style->Set("AvalancheMediaEditor.BroadcastClient"      , new IMAGE_BRUSH("Icons/MediaIcons/BroadcastClient" , Icon64x64));
+	Style->Set("AvalancheMediaEditor.BroadcastServer"      , new IMAGE_BRUSH("Icons/MediaIcons/BroadcastServer" , Icon64x64));
+	Style->Set("AvalancheMediaEditor.BroadcastClient.Small", new IMAGE_BRUSH("Icons/MediaIcons/BroadcastClient" , Icon20x20));
+	Style->Set("AvalancheMediaEditor.BroadcastServer.Small", new IMAGE_BRUSH("Icons/MediaIcons/BroadcastServer" , Icon20x20));
+	Style->Set("AvalancheMediaEditor.BroadcastOffline"     , new IMAGE_BRUSH_SVG("Icons/MediaIcons/BroadcastOffline"	, Icon16x16));
+	Style->Set("AvalancheMediaEditor.BroadcastError"       , new IMAGE_BRUSH_SVG("Icons/MediaIcons/BroadcastError"		, Icon16x16));
+	Style->Set("AvalancheMediaEditor.BroadcastWarning"     , new IMAGE_BRUSH_SVG("Icons/MediaIcons/BroadcastWarning"	, Icon16x16));
+	Style->Set("AvalancheMediaEditor.BroadcastLive"        , new IMAGE_BRUSH_SVG("Icons/MediaIcons/BroadcastLive"		, Icon16x16)); 
+	Style->Set("AvalancheMediaEditor.BroadcastIdle"        , new IMAGE_BRUSH_SVG("Icons/MediaIcons/BroadcastIdle"		, Icon16x16));
+
+	// Broadcast Channel Types
+	Style->Set("AvalancheMediaEditor.ChannelTypeProgram"	, new IMAGE_BRUSH_SVG("Icons/MediaIcons/ChannelTypeProgram"	, Icon16x16)); 
+	Style->Set("AvalancheMediaEditor.ChannelTypePreview"	, new IMAGE_BRUSH_SVG("Icons/MediaIcons/ChannelTypePreview"	, Icon16x16));
+	
+	// Media Status
+	/* For when we have SVGs
+	Style->Set("AvalancheMediaEditor.MediaAssetStatus", new IMAGE_BRUSH_SVG("Icons/MediaIcons/MediaLoaded", Icon16x16));
+	Style->Set("AvalancheMediaEditor.MediaSyncStatus",  new IMAGE_BRUSH_SVG("Icons/MediaIcons/MediaLoaded", Icon16x16));
+	Style->Set("AvalancheMediaEditor.MediaPreviewing",  new IMAGE_BRUSH_SVG("Icons/MediaIcons/MediaLoaded", Icon16x16));
+	Style->Set("AvalancheMediaEditor.MediaPlaying",     new IMAGE_BRUSH_SVG("Icons/MediaIcons/MediaLoaded", Icon16x16));
+	*/
+	Style->Set("AvalancheMediaEditor.MediaAssetStatus", new IMAGE_BRUSH("Icons/MediaIcons/MediaStatus",     Icon16x16));
+	Style->Set("AvalancheMediaEditor.MediaSyncStatus",  new IMAGE_BRUSH("Icons/MediaIcons/MediaSyncStatus", Icon16x16));
+	Style->Set("AvalancheMediaEditor.MediaPreviewing",  new IMAGE_BRUSH("Icons/MediaIcons/MediaPreviewing", Icon16x16));
+	Style->Set("AvalancheMediaEditor.MediaPlaying",     new IMAGE_BRUSH("Icons/MediaIcons/MediaPlaying",    Icon16x16));
+
+	// Media Output Status
+	Style->Set("AvalancheMediaEditor.MediaOutputOffline"   , new IMAGE_BRUSH_SVG("Icons/MediaIcons/MediaOutputOffline"	, Icon16x16));
+	Style->Set("AvalancheMediaEditor.MediaOutputIdle"      , new IMAGE_BRUSH_SVG("Icons/MediaIcons/MediaOutputIdle"		, Icon16x16));
+	Style->Set("AvalancheMediaEditor.MediaOutputPreparing" , new IMAGE_BRUSH_SVG("Icons/MediaIcons/MediaOutputPreparing", Icon16x16));
+	Style->Set("AvalancheMediaEditor.MediaOutputLive"      , new IMAGE_BRUSH_SVG("Icons/MediaIcons/MediaOutputLive"		, Icon16x16));
+	Style->Set("AvalancheMediaEditor.MediaOutputLiveWarn"  , new IMAGE_BRUSH_SVG("Icons/MediaIcons/MediaOutputLiveWarn"	, Icon16x16));
+	Style->Set("AvalancheMediaEditor.MediaOutputError"     , new IMAGE_BRUSH_SVG("Icons/MediaIcons/MediaOutputError"	, Icon16x16));
+
+	// Avalanche Preview
+	Style->Set("AvalancheMediaEditor.Checkerboard" , new IMAGE_BRUSH("Images/AvaPreviewCheckerboard", Icon16x16, FLinearColor::White, ESlateBrushTileType::Both));
+
+	// Button styles
+	const FButtonStyle& AppStyle_SimpleButton = FAppStyle::Get().GetWidgetStyle<FButtonStyle>("SimpleButton");
+
+	Style->Set("AvalancheMediaEditor.BorderlessButton", FButtonStyle(AppStyle_SimpleButton));
+
+	const FButtonStyle& MenuButtonStyleGreen = FAppStyle::Get().GetWidgetStyle<FButtonStyle>("FlatButton.Success");
+	Style->Set("AvalancheMediaEditor.ButtonGreen", MenuButtonStyleGreen);
+	const FButtonStyle& MenuButtonStyleRed = FAppStyle::Get().GetWidgetStyle<FButtonStyle>("FlatButton.Danger");
+	Style->Set("AvalancheMediaEditor.ButtonRed", MenuButtonStyleRed);
+
+	const FToolBarStyle& SlimToolbarStyle = FAppStyle::Get().GetWidgetStyle<FToolBarStyle>("SlimToolBar");
+	Style->Set("AvalancheMediaEditor.ToolBar", SlimToolbarStyle);
+	const FButtonStyle& MenuButtonStyle = SlimToolbarStyle.ButtonStyle;
+
+	FToolBarStyle SlimToolbarStyleOverrideRedButton = FAppStyle::Get().GetWidgetStyle<FToolBarStyle>("SlimToolBar");
+	SlimToolbarStyleOverrideRedButton.SetButtonStyle(Style->GetWidgetStyle<FButtonStyle>("AvalancheMediaEditor.ButtonRed"));
+	SlimToolbarStyleOverrideRedButton.ButtonStyle.SetDisabled(MenuButtonStyle.Disabled);
+	Style->Set("AvalancheMediaEditor.ToolBarRedButtonOverride", SlimToolbarStyleOverrideRedButton);
+
+	FToolBarStyle SlimToolbarStyleOverrideGreenButton = FAppStyle::Get().GetWidgetStyle<FToolBarStyle>("SlimToolBar");
+	SlimToolbarStyleOverrideGreenButton.SetButtonStyle(Style->GetWidgetStyle<FButtonStyle>("AvalancheMediaEditor.ButtonGreen"));
+	SlimToolbarStyleOverrideGreenButton.ButtonStyle.SetDisabled(MenuButtonStyle.Disabled);
+	Style->Set("AvalancheMediaEditor.ToolBarGreenButtonOverride", SlimToolbarStyleOverrideGreenButton);
+
+	const FToolBarStyle& CalloutToolbarStyleOverride = FAppStyle::Get().GetWidgetStyle<FToolBarStyle>("CalloutToolbar");
+	Style->Set("AvalancheMediaEditor.CalloutToolbar", CalloutToolbarStyleOverride);
+
+	Style->Set("TableView.ActivePageBrush", new FSlateColorBrush(FLinearColor(0.05f, 0.1f, 0.05f, 0.5f)));
+	Style->Set("TableView.Hovered.ActivePageBrush", new FSlateColorBrush(FLinearColor(0.05f, 0.2f, 0.05f, 0.5f)));
+	Style->Set("TableView.Selected.ActivePageBrush", new FSlateColorBrush(FLinearColor(0.05f, 0.3f, 0.05f, 0.5f)));
+	Style->Set("TableView.DisabledPageBrush", new FSlateColorBrush(FLinearColor(0.1f, 0.05f, 0.05f, 0.5f)));
+	Style->Set("TableView.Hovered.DisabledPageBrush", new FSlateColorBrush(FLinearColor(0.2f, 0.05f, 0.05f, 0.5f)));
+	Style->Set("TableView.Selected.DisabledPageBrush", new FSlateColorBrush(FLinearColor(0.3f, 0.05f, 0.05f, 0.5f)));
+
+	// Asset Colors - in the shades of yellow (hue = 36, saturation = 221)
+	Style->Set("AvalancheMediaEditor.AssetColors.Rundown", FLinearColor::MakeFromHSV8(36, 221, 170));
+	Style->Set("AvalancheMediaEditor.AssetColors.RundownMacroCollection", FLinearColor::MakeFromHSV8(36, 221, 124));
+	Style->Set("AvalancheMediaEditor.AssetColors.Playback", FLinearColor::MakeFromHSV8(36, 221, 105));
+
+	return Style;
+}
+
+#undef IMAGE_BRUSH
+#undef IMAGE_BRUSH_SVG
+#undef BOX_BRUSH
+#undef CORE_FONT
+
+const ISlateStyle& FAvaMediaEditorStyle::Get()
+{
+	return *StyleInstance;
+}
diff --git a/Engine/Plugins/Experimental/Avalanche/Source/AvalancheMediaEditor/Private/AvaMediaEditorStyle.h b/Engine/Plugins/Experimental/Avalanche/Source/AvalancheMediaEditor/Private/AvaMediaEditorStyle.h
new file mode 100644
--- /dev/null
+++ b/Engine/Plugins/Experimental/Avalanche/Source/AvalancheMediaEditor/Private/AvaMediaEditorStyle.h
@@ -0,0 +1,38 @@
+// Copyright Epic Games, Inc. All Rights Reserved.
+
+#pragma once
+
+#include "Styling/ISlateStyle.h"
+#include "Styling/SlateStyle.h"
+
+struct FSlateBrush;
+class UTexture2D;
+
+class FAvaMediaEditorStyle
+{
+public:
+
+	static void Initialize();
+
+	static void Shutdown();
+	
+	static const ISlateStyle& Get();
+
+	static FName GetStyleSetName();
+
+	static const FLinearColor& GetColor(FName PropertyName, const ANSICHAR* Specifier = nullptr);
+
+	static const FSlateBrush* GetBrush(FName PropertyName, const ANSICHAR* Specifier = nullptr);
+
+	template< typename WidgetStyleType >
+	static const WidgetStyleType& GetWidgetStyle(FName PropertyName, const ANSICHAR* Specifier = nullptr)
+	{
+		return StyleInstance->GetWidgetStyle<WidgetStyleType>(PropertyName, Specifier);
+	}
+
+private:
+
+	static TSharedRef<FSlateStyleSet> Create();
+	static TSharedPtr<FSlateStyleSet> StyleInstance;
+	static TMap<FName,UTexture2D*> Sprites;
+};
diff --git a/Engine/Plugins/Experimental/Avalanche/Source/AvalancheMediaEditor/Private/AvaMediaEditorUtils.cpp b/Engine/Plugins/Experimental/Avalanche/Source/AvalancheMediaEditor/Private/AvaMediaEditorUtils.cpp
new file mode 100644
--- /dev/null
+++ b/Engine/Plugins/Experimental/Avalanche/Source/AvalancheMediaEditor/Private/AvaMediaEditorUtils.cpp
@@ -0,0 +1,35 @@
+// Copyright Epic Games, Inc. All Rights Reserved.
+
+#include "AvaMediaEditorUtils.h"
+
+const FSlateBrush* FAvaMediaEditorUtils::GetChannelStatusBrush(EAvaChannelState InChannelState, EAvaMediaIssueSeverity InChannelIssueSeverity)
+{
+	if (InChannelState == EAvaChannelState::Live)
+	{
+		switch (InChannelIssueSeverity)
+		{
+		case EAvaMediaIssueSeverity::None:
+			return FAvaMediaEditorStyle::Get().GetBrush("AvalancheMediaEditor.BroadcastLive");
+		case EAvaMediaIssueSeverity::Warnings:
+			return FAvaMediaEditorStyle::Get().GetBrush("AvalancheMediaEditor.BroadcastWarning");
+		case EAvaMediaIssueSeverity::Errors:
+			return FAvaMediaEditorStyle::Get().GetBrush("AvalancheMediaEditor.BroadcastError");
+		}
+	}
+	else if (InChannelState == EAvaChannelState::Idle)
+	{
+		return FAvaMediaEditorStyle::Get().GetBrush("AvalancheMediaEditor.BroadcastIdle");
+	}
+	else
+	{
+		return FAvaMediaEditorStyle::Get().GetBrush("AvalancheMediaEditor.BroadcastOffline");
+	}
+	
+	return FAvaMediaEditorStyle::Get().GetBrush("AvalancheMediaEditor.BroadcastIdle");
+}
+
+FText FAvaMediaEditorUtils::GetChannelStatusText(EAvaChannelState InChannelState, EAvaMediaIssueSeverity InChannelIssueSeverity)
+{
+	//TODO: Possibly also add Issue Severity (if not none)
+	return StaticEnum<EAvaChannelState>()->GetDisplayNameTextByIndex(static_cast<int32>(InChannelState));
+}
diff --git a/Engine/Plugins/Experimental/Avalanche/Source/AvalancheMediaEditor/Private/AvaMediaEditorUtils.h b/Engine/Plugins/Experimental/Avalanche/Source/AvalancheMediaEditor/Private/AvaMediaEditorUtils.h
new file mode 100644
--- /dev/null
+++ b/Engine/Plugins/Experimental/Avalanche/Source/AvalancheMediaEditor/Private/AvaMediaEditorUtils.h
@@ -0,0 +1,13 @@
+// Copyright Epic Games, Inc. All Rights Reserved.
+
+#pragma once
+
+#include "AvaMediaDefines.h"
+#include "AvaMediaEditorStyle.h"
+
+class FAvaMediaEditorUtils
+{
+public:
+	static const FSlateBrush* GetChannelStatusBrush(EAvaChannelState InChannelState, EAvaMediaIssueSeverity InChannelIssueSeverity);
+	static FText GetChannelStatusText(EAvaChannelState InChannelState, EAvaMediaIssueSeverity InChannelIssueSeverity);
+};
diff --git a/Engine/Plugins/Experimental/Avalanche/Source/AvalancheMediaEditor/Private/AvalancheMediaEditorSettings.cpp b/Engine/Plugins/Experimental/Avalanche/Source/AvalancheMediaEditor/Private/AvalancheMediaEditorSettings.cpp
new file mode 100644
--- /dev/null
+++ b/Engine/Plugins/Experimental/Avalanche/Source/AvalancheMediaEditor/Private/AvalancheMediaEditorSettings.cpp
@@ -0,0 +1,27 @@
+// Copyright Epic Games, Inc. All Rights Reserved.
+
+#include "AvalancheMediaEditorSettings.h"
+
+UAvalancheMediaEditorSettings::UAvalancheMediaEditorSettings()
+{
+	CategoryName = TEXT("Motion Design");
+	SectionName  = TEXT("Playback & Broadcast");
+	
+	PlaybackDefaultNodeColor  = FLinearColor(1.f, 1.f, 1.f, 1.f);
+	PlaybackChannelsNodeColor = FLinearColor(1.f, 0.1f, 0.f, 1.f);
+	PlaybackPlayerNodeColor   = FLinearColor(0.f, 1.f, 0.f, 1.f);
+	PlaybackEventNodeColor    = FLinearColor(1.f, 0.f, 0.f, 1.f);
+	PlaybackActionNodeColor   = FLinearColor(0.190525f, 0.583898f, 1.0f, 1.0f);
+}
+
+UAvalancheMediaEditorSettings* UAvalancheMediaEditorSettings::GetSingletonInstance()
+{
+	UAvalancheMediaEditorSettings* const DefaultSettings = GetMutableDefault<UAvalancheMediaEditorSettings>();
+	static bool bInitialized = false;
+	if (!bInitialized)
+	{
+		bInitialized = true;
+		DefaultSettings->SetFlags(RF_Transactional);
+	}
+	return DefaultSettings;
+}
