[MotionDesign]

diff --git a/Engine/Plugins/Experimental/Avalanche/Source/AvalancheMediaEditor/Private/Broadcast/ChannelGrid/AvaOutputTileItem.cpp b/Engine/Plugins/Experimental/Avalanche/Source/AvalancheMediaEditor/Private/Broadcast/ChannelGrid/AvaOutputTileItem.cpp
new file mode 100644
--- /dev/null
+++ b/Engine/Plugins/Experimental/Avalanche/Source/AvalancheMediaEditor/Private/Broadcast/ChannelGrid/AvaOutputTileItem.cpp
@@ -0,0 +1,253 @@
+// Copyright Epic Games, Inc. All Rights Reserved.
+
+#include "Broadcast/ChannelGrid/AvaOutputTileItem.h"
+#include "AvaMediaEditorStyle.h"
+#include "AvalancheBroadcast.h"
+#include "ClassIconFinder.h"
+#include "DragDropOps/AvaOutputTileItemDragDropOp.h"
+#include "MediaOutput.h"
+#include "OutputDevices/AvaMediaOutputUtils.h"
+#include "Widgets/Layout/SScaleBox.h"
+
+#define LOCTEXT_NAMESPACE "AvaOutputTileItem"
+
+FAvaOutputTileItem::FAvaOutputTileItem(FName InChannelName, UMediaOutput* InMediaOutput)
+	: ChannelName(InChannelName)
+	, MediaOutput(InMediaOutput)
+{
+	FCoreUObjectDelegates::OnObjectPropertyChanged.AddRaw(this, &FAvaOutputTileItem::OnMediaOutputPropertyChanged);
+	FAvaOutputChannel::GetOnChannelChanged().AddRaw(this, &FAvaOutputTileItem::OnChannelChanged);
+	FAvaOutputChannel::GetOnMediaOutputStateChanged().AddRaw(this, &FAvaOutputTileItem::OnMediaOutputStateChanged);
+	BroadcastChangedHandle = UAvalancheBroadcast::Get().AddChangeListener(
+		FOnAvaBroadcastChanged::FDelegate::CreateRaw(this, &FAvaOutputTileItem::OnBroadcastChanged));
+
+	UpdateInfo();
+}
+
+FAvaOutputTileItem::~FAvaOutputTileItem()
+{
+	FCoreUObjectDelegates::OnObjectPropertyChanged.RemoveAll(this);
+	FAvaOutputChannel::GetOnChannelChanged().RemoveAll(this);
+	FAvaOutputChannel::GetOnMediaOutputStateChanged().RemoveAll(this);
+	UAvalancheBroadcast::Get().RemoveChangeListener(BroadcastChangedHandle);
+}
+
+const FAvaOutputChannel& FAvaOutputTileItem::GetChannel() const
+{
+	return UAvalancheBroadcast::Get().GetCurrentProfile().GetChannel(ChannelName);
+}
+
+FAvaOutputChannel& FAvaOutputTileItem::GetChannel()
+{
+	return UAvalancheBroadcast::Get().GetCurrentProfile().GetChannelMutable(ChannelName);
+}
+
+TSharedRef<SWidget> FAvaOutputTileItem::GenerateTile() const
+{
+	check(MediaOutput.IsValid());
+	
+	return SNew(SHorizontalBox)
+		.ToolTipText(this, &FAvaOutputTileItem::GetToolTipText)
+		+ SHorizontalBox::Slot()
+		.AutoWidth()
+		[
+			SNew(SImage)
+			.Image(GetMediaOutputIcon())
+		]
+		+ SHorizontalBox::Slot()
+		.AutoWidth()
+		.Padding(5.f, 0.f, 0.f, 0.f)
+		.HAlign(EHorizontalAlignment::HAlign_Fill)
+		.VAlign(EVerticalAlignment::VAlign_Fill)
+		[
+			SNew(SVerticalBox)
+			+ SVerticalBox::Slot()
+			.FillHeight(0.75f)
+			.HAlign(EHorizontalAlignment::HAlign_Fill)
+			.VAlign(EVerticalAlignment::VAlign_Fill)
+			[
+				SNew(SScaleBox)
+				.Stretch(EStretch::UserSpecified)
+				.UserSpecifiedScale(1.25f)
+				.VAlign(EVerticalAlignment::VAlign_Center)
+				[
+					SNew(STextBlock)
+					.Text(this, &FAvaOutputTileItem::GetDisplayText)
+					.Justification(ETextJustify::Center)
+				]
+			]
+			+ SVerticalBox::Slot()
+			.FillHeight(0.25f)
+			.HAlign(EHorizontalAlignment::HAlign_Fill)
+			.VAlign(EVerticalAlignment::VAlign_Fill)
+			[
+				SNew(SScaleBox)
+				.Stretch(EStretch::ScaleToFitY)
+				[
+					SNew(SHorizontalBox)
+					+SHorizontalBox::Slot()
+					.AutoWidth()
+					.VAlign(EVerticalAlignment::VAlign_Center)
+					[
+						SNew(SImage)
+						.Image(this, &FAvaOutputTileItem::GetMediaOutputStatusBrush)
+					]
+					+SHorizontalBox::Slot()
+					.Padding(2.f, 0.f, 0.f, 0.f)
+					.VAlign(EVerticalAlignment::VAlign_Center)
+					.AutoWidth()
+					[
+						SNew(STextBlock)
+						.Text(this, &FAvaOutputTileItem::GetMediaOutputStatusText)
+					]
+				]
+			]
+		];
+}
+
+FText FAvaOutputTileItem::GetDisplayText() const
+{
+	return MediaOutputDisplayText;
+}
+
+FText FAvaOutputTileItem::GetMediaOutputStatusText() const
+{
+	return MediaOutputStatusText;
+}
+
+FText FAvaOutputTileItem::GetToolTipText() const
+{
+	return MediaOutputToolTipText;
+}
+
+const FSlateBrush* FAvaOutputTileItem::GetMediaOutputIcon() const
+{
+	check(MediaOutput.IsValid());
+	return FClassIconFinder::FindThumbnailForClass(MediaOutput->GetClass());
+}
+
+const FSlateBrush* FAvaOutputTileItem::GetMediaOutputStatusBrush() const
+{
+	return MediaOutputStatusBrush;
+}
+
+FReply FAvaOutputTileItem::OnDragDetected(const FGeometry& MyGeometry, const FPointerEvent& MouseEvent)
+{
+	if (!UAvalancheBroadcast::Get().IsBroadcastingAnyChannel())
+	{
+		const bool bShouldDuplicate = MouseEvent.IsAltDown();
+		
+		const TSharedRef<FAvaOutputTileItemDragDropOp> DragDropOp = FAvaOutputTileItemDragDropOp::New(SharedThis(this)
+				, bShouldDuplicate);
+			
+		return FReply::Handled().BeginDragDrop(DragDropOp);
+	}
+	return FReply::Unhandled();
+}
+
+void FAvaOutputTileItem::OnMediaOutputPropertyChanged(UObject* InObject, FPropertyChangedEvent& PropertyChangedEvent)
+{
+	if (MediaOutput.IsValid() && InObject == MediaOutput)
+	{
+		UpdateInfo();
+	}
+}
+
+void FAvaOutputTileItem::OnChannelChanged(const FAvaOutputChannel& InChannel, EAvaChannelChange InChange)
+{
+	if (EnumHasAnyFlags(InChange, EAvaChannelChange::State)
+		&& InChannel.IsValidChannel()
+		&& InChannel.GetChannelName() == ChannelName)
+	{
+		UpdateInfo();
+	}
+}
+
+void FAvaOutputTileItem::OnMediaOutputStateChanged(const FAvaOutputChannel& InChannel, const UMediaOutput* InMediaOutput)
+{
+	if (InChannel.IsValidChannel() && InChannel.GetChannelName() == ChannelName)
+	{
+		UpdateInfo();
+	}
+}
+
+void FAvaOutputTileItem::OnBroadcastChanged(EAvaBroadcastChange InChange)
+{
+	// When output devices are updated, some of their status can change from offline/idle.
+	if (EnumHasAnyFlags(InChange, EAvaBroadcastChange::OutputDevices))
+	{
+		UpdateInfo();
+	}
+}
+
+
+void FAvaOutputTileItem::UpdateInfo()
+{
+	MediaOutputDisplayText = FindLatestDisplayText();
+	
+	const EAvaMediaOutputState OutputState = GetChannel().GetMediaOutputState(MediaOutput.Get());
+	const EAvaMediaIssueSeverity Severity = GetChannel().GetMediaOutputIssueSeverity(OutputState, MediaOutput.Get());
+	switch(OutputState)
+	{
+	case EAvaMediaOutputState::Offline:
+		MediaOutputStatusText = LOCTEXT("MediaOutput_NoIssues", "Offline");
+		MediaOutputStatusBrush = FAvaMediaEditorStyle::Get().GetBrush("AvalancheMediaEditor.MediaOutputOffline");
+		break;
+	case EAvaMediaOutputState::Idle:
+		MediaOutputStatusText = LOCTEXT("MediaOutput_NoIssues", "Idle");
+		MediaOutputStatusBrush = FAvaMediaEditorStyle::Get().GetBrush("AvalancheMediaEditor.MediaOutputIdle");
+		break;
+	case EAvaMediaOutputState::Preparing:
+		MediaOutputStatusText = LOCTEXT("MediaOutput_NoIssues", "Preparing");
+		MediaOutputStatusBrush = FAvaMediaEditorStyle::Get().GetBrush("AvalancheMediaEditor.MediaOutputPreparing");
+		break;
+	case EAvaMediaOutputState::Live:
+		MediaOutputStatusText = LOCTEXT("MediaOutput_NoIssues", "Live");
+		if (Severity == EAvaMediaIssueSeverity::Errors || Severity == EAvaMediaIssueSeverity::Warnings)
+		{
+			// If severity is warning or error, add a secondary icon (yellow or red exclamation mark)
+			// so the user knows to lookup the error in the tooltip.
+			MediaOutputStatusBrush = FAvaMediaEditorStyle::Get().GetBrush("AvalancheMediaEditor.MediaOutputLiveWarn");
+		}
+		else
+		{
+			MediaOutputStatusBrush = FAvaMediaEditorStyle::Get().GetBrush("AvalancheMediaEditor.MediaOutputLive");
+		}
+		break;
+	case EAvaMediaOutputState::Error:
+		MediaOutputStatusText = LOCTEXT("MediaOutput_Errors", "Error(s)");
+		MediaOutputStatusBrush = FAvaMediaEditorStyle::Get().GetBrush("AvalancheMediaEditor.MediaOutputError");
+		break;
+	}
+
+	FString AllMessages;
+	switch(Severity)
+	{
+	case EAvaMediaIssueSeverity::Errors:
+		AllMessages.Append("Errors: \n");
+		break;
+	case EAvaMediaIssueSeverity::Warnings:
+		AllMessages.Append(TEXT("Warning: \n"));
+		break;
+	default:
+		AllMessages.Append(TEXT("Healthy\n"));
+		break;
+	}
+	
+	const TArray<FString>& Messages = GetChannel().GetMediaOutputIssueMessages(MediaOutput.Get());
+	for (const FString& Message : Messages)
+	{
+		AllMessages.Append(Message);
+		AllMessages.Append(TEXT("\n"));
+	}
+	MediaOutputToolTipText = FText::FromString(AllMessages);
+}
+
+FText FAvaOutputTileItem::FindLatestDisplayText() const
+{
+	check(MediaOutput.IsValid());
+	const FString DeviceName = UE::AvaMediaOutputUtils::GetDeviceName(MediaOutput.Get());
+	return !DeviceName.IsEmpty() ? FText::FromString(DeviceName) : FText::FromName(MediaOutput->GetFName());
+}
+
+#undef LOCTEXT_NAMESPACE
diff --git a/Engine/Plugins/Experimental/Avalanche/Source/AvalancheMediaEditor/Private/Broadcast/ChannelGrid/DragDropOps/AvaOutputTileItemDragDropOp.cpp b/Engine/Plugins/Experimental/Avalanche/Source/AvalancheMediaEditor/Private/Broadcast/ChannelGrid/DragDropOps/AvaOutputTileItemDragDropOp.cpp
new file mode 100644
--- /dev/null
+++ b/Engine/Plugins/Experimental/Avalanche/Source/AvalancheMediaEditor/Private/Broadcast/ChannelGrid/DragDropOps/AvaOutputTileItemDragDropOp.cpp
@@ -0,0 +1,100 @@
+// Copyright Epic Games, Inc. All Rights Reserved.
+
+#include "AvaOutputTileItemDragDropOp.h"
+#include "AvalancheBroadcast.h"
+#include "MediaOutput.h"
+#include "ScopedTransaction.h"
+
+#define LOCTEXT_NAMESPACE "AvaOutputTileItemDragDropOp"
+
+TSharedRef<FAvaOutputTileItemDragDropOp> FAvaOutputTileItemDragDropOp::New(const FAvaOutputTileItemPtr& InOutputClassItem, bool bInIsDuplicating)
+{
+	TSharedRef<FAvaOutputTileItemDragDropOp> DragDropOp = MakeShared<FAvaOutputTileItemDragDropOp>();
+	DragDropOp->Init(InOutputClassItem, bInIsDuplicating);
+	return DragDropOp;
+}
+
+bool FAvaOutputTileItemDragDropOp::IsValidToDropInChannel(FName InTargetChannelName) const
+{
+	if (!OutputTileItem.IsValid())
+	{
+		return false;
+	}
+
+	// Don't allow dropping a remote output in a preview channel
+	if (UAvalancheBroadcast::Get().GetChannelType(InTargetChannelName) == EAvaBroadcastChannelType::Preview
+		&& OutputTileItem->GetChannel().IsMediaOutputRemote(OutputTileItem->GetMediaOutput()))
+	{
+		return false;
+	}
+	
+	// Allow Duplicate on Same Channel
+	return (bIsDuplicating || OutputTileItem->GetChannel().GetChannelName() != InTargetChannelName);
+}
+
+FReply FAvaOutputTileItemDragDropOp::OnChannelDrop(FName InTargetChannelName)
+{
+	if (TSharedPtr<FAvaOutputTileItem> Tile = GetOutputTileItem())
+	{
+		UAvalancheBroadcast& Broadcast = UAvalancheBroadcast::Get();
+		
+		FScopedTransaction Transaction(LOCTEXT("OnChannelDrop", "Drop Media Output"));
+		Broadcast.Modify();
+		
+		UMediaOutput* const SourceMediaOutput = Tile->GetMediaOutput();
+
+		if (!IsValid(SourceMediaOutput))
+		{
+			Transaction.Cancel();
+			return FReply::Unhandled();
+		}
+
+		FAvaOutputChannel& TargetChannel = Broadcast.GetCurrentProfile().GetChannelMutable(InTargetChannelName);
+		FAvaOutputChannel& SourceChannel = Tile->GetChannel();
+		FAvaMediaOutputInfo MediaOutputInfo = SourceChannel.GetMediaOutputInfo(SourceMediaOutput);
+
+		// Don't allow dropping a remote output in a preview channel
+		if (Broadcast.GetChannelType(InTargetChannelName) == EAvaBroadcastChannelType::Preview && MediaOutputInfo.IsRemote())
+		{
+			Transaction.Cancel();
+			return FReply::Unhandled();
+		}
+		
+		if (bIsDuplicating)
+		{
+			UMediaOutput* const NewMediaOutput = DuplicateObject<UMediaOutput>(SourceMediaOutput
+				, SourceMediaOutput->GetOuter()
+				, NAME_None);
+			
+			check(NewMediaOutput);
+			MediaOutputInfo.Guid = FGuid::NewGuid();	// Allocate a new guid for the duplicate.
+			TargetChannel.AddMediaOutput(NewMediaOutput, MediaOutputInfo);
+		}
+		else
+		{
+			SourceChannel.RemoveMediaOutput(SourceMediaOutput);
+			TargetChannel.AddMediaOutput(SourceMediaOutput, MediaOutputInfo);
+		}
+
+		return FReply::Handled();
+	}
+	return FReply::Unhandled();
+}
+
+
+void FAvaOutputTileItemDragDropOp::Init(const FAvaOutputTileItemPtr& InOutputTileItem, bool bInIsDuplicating)
+{
+	OutputTileItem = InOutputTileItem;
+	bIsDuplicating = bInIsDuplicating;
+	
+	CurrentHoverText = InOutputTileItem->GetDisplayText();
+	CurrentIconBrush = InOutputTileItem->GetMediaOutputIcon();
+	
+	CurrentIconColorAndOpacity = FSlateColor::UseForeground();
+	MouseCursor = EMouseCursor::GrabHandClosed;
+
+	SetupDefaults();
+	Construct();
+}
+
+#undef LOCTEXT_NAMESPACE
diff --git a/Engine/Plugins/Experimental/Avalanche/Source/AvalancheMediaEditor/Private/Broadcast/ChannelGrid/DragDropOps/AvaOutputTileItemDragDropOp.h b/Engine/Plugins/Experimental/Avalanche/Source/AvalancheMediaEditor/Private/Broadcast/ChannelGrid/DragDropOps/AvaOutputTileItemDragDropOp.h
new file mode 100644
--- /dev/null
+++ b/Engine/Plugins/Experimental/Avalanche/Source/AvalancheMediaEditor/Private/Broadcast/ChannelGrid/DragDropOps/AvaOutputTileItemDragDropOp.h
@@ -0,0 +1,31 @@
+// Copyright Epic Games, Inc. All Rights Reserved.
+
+#pragma once
+
+#include "Broadcast/ChannelGrid/AvaOutputTileItem.h"
+#include "CoreMinimal.h"
+#include "DragAndDrop/DecoratedDragDropOp.h"
+
+class FAvaOutputTileItemDragDropOp : public FDecoratedDragDropOp
+{
+public:
+
+	DRAG_DROP_OPERATOR_TYPE(FAvaOutputTileDragDropOp, FDecoratedDragDropOp)
+	
+	static TSharedRef<FAvaOutputTileItemDragDropOp> New(const FAvaOutputTileItemPtr& InOutputClassItem, bool bInIsDuplicating);
+
+	bool IsValidToDropInChannel(FName InTargetChannelName) const;
+	
+	FAvaOutputTileItemPtr GetOutputTileItem() const { return OutputTileItem; }
+
+	FReply OnChannelDrop(FName InTargetChannelName);
+	
+protected:
+	
+	void Init(const FAvaOutputTileItemPtr& InOutputTileItem, bool bInIsDuplicating);
+
+	//Keep Reference Count while Drag Dropping
+	FAvaOutputTileItemPtr OutputTileItem;
+
+	bool bIsDuplicating = false;
+};
diff --git a/Engine/Plugins/Experimental/Avalanche/Source/AvalancheMediaEditor/Private/Broadcast/ChannelGrid/Slate/SAvaChannel.cpp b/Engine/Plugins/Experimental/Avalanche/Source/AvalancheMediaEditor/Private/Broadcast/ChannelGrid/Slate/SAvaChannel.cpp
new file mode 100644
--- /dev/null
+++ b/Engine/Plugins/Experimental/Avalanche/Source/AvalancheMediaEditor/Private/Broadcast/ChannelGrid/Slate/SAvaChannel.cpp
@@ -0,0 +1,810 @@
+// Copyright Epic Games, Inc. All Rights Reserved.
+
+#include "SAvaChannel.h"
+
+#include "AvaMediaDefines.h"
+#include "AvaMediaEditorStyle.h"
+#include "AvaMediaEditorUtils.h"
+#include "AvalancheBroadcast.h"
+#include "Broadcast/AvaBroadcastEditor.h"
+#include "Broadcast/ChannelGrid/DragDropOps/AvaOutputTileItemDragDropOp.h"
+#include "Broadcast/OutputDevices/AvaOutputClassItem.h"
+#include "Broadcast/OutputDevices/DragDropOps/AvaOutputTreeItemDragDropOp.h"
+#include "Brushes/SlateImageBrush.h"
+#include "ClassIconFinder.h"
+#include "Framework/Commands/GenericCommands.h"
+#include "Framework/MultiBox/MultiBoxBuilder.h"
+#include "IStructureDetailsView.h"
+#include "MediaOutput.h"
+#include "OutputDevices/Slate/SAvaCaptureImage.h"
+#include "PropertyEditorModule.h"
+#include "ScopedTransaction.h"
+#include "Widgets/Colors/SColorBlock.h"
+#include "Widgets/Input/SButton.h"
+#include "Widgets/Input/SMenuAnchor.h"
+#include "Widgets/Layout/SScaleBox.h"
+#include "Widgets/Layout/SScrollBox.h"
+#include "Widgets/Layout/SSeparator.h"
+#include "Widgets/Text/SInlineEditableTextBlock.h"
+#include "Widgets/Views/STileView.h"
+
+#define LOCTEXT_NAMESPACE "SAvaChannel"
+
+void SAvaChannel::Construct(const FArguments& InArgs, const TSharedPtr<FAvaBroadcastEditor>& InBroadcastEditor)
+{
+	BroadcastEditorWeak = InBroadcastEditor;
+	ChannelName = InArgs._ChannelName;
+	CanMaximize = InArgs._CanMaximize;
+	OnMaximizeClicked = InArgs._OnMaximizeClicked;
+	
+	RegisterCommands();
+	
+	FAvaOutputChannel::GetOnChannelChanged().AddRaw(this, &SAvaChannel::OnChannelChanged);
+	InBroadcastEditor->OnOutputTileSelectionChanged.AddRaw(this, &SAvaChannel::OnOutputTileSelectionChanged);
+	
+	ChildSlot
+	[
+		SNew(SOverlay)
+		.Clipping(EWidgetClipping::ClipToBounds)
+		+ SOverlay::Slot()
+		[
+			SNew(SBorder)
+			.BorderImage(&FAppStyle::Get().GetWidgetStyle<FTableRowStyle>(TEXT("SceneOutliner.TableViewRow")).DropIndicator_Onto)
+			.Visibility(this, &SAvaChannel::GetChannelDragVisibility)
+		]
+		+ SOverlay::Slot()
+		.Padding(2.0f)
+		[
+			SNew(SBorder)
+			.BorderImage(FAppStyle::GetBrush("Menu.WidgetBorder"))
+			[
+				MakeChannelVerticalBox()
+			]
+		]
+		//Min/Maximize and X button
+		+ SOverlay::Slot()
+		.VAlign(EVerticalAlignment::VAlign_Top)
+		.HAlign(EHorizontalAlignment::HAlign_Right)
+		.Padding(0.0f, 5.f, 5.f, 0.f)
+		[
+			SNew(SHorizontalBox)
+			+ SHorizontalBox::Slot()
+			.AutoWidth()
+			[
+				SAssignNew(ChannelSettingsMenuAnchor, SMenuAnchor)
+				.Content()
+				[
+					SNew(SButton)
+					.ButtonStyle(FAppStyle::Get(), "SimpleButton")
+					.OnClicked(this, &SAvaChannel::OnChannelSettingsButtonClicked)
+					.ToolTipText(LOCTEXT("ChannelSettingToolTip", "Open the channel settings"))
+					[
+						SNew(SImage)
+						.Image(FAppStyle::Get().GetBrush("Icons.Settings"))
+					]
+				]
+			]
+			+ SHorizontalBox::Slot()
+			.AutoWidth()
+			[
+				SNew(SButton)
+				.ButtonStyle(FAppStyle::Get(), "SimpleButton")
+				.OnClicked(this, &SAvaChannel::OnChannelPinButtonClicked)
+				.ToolTipText(LOCTEXT("ChannelPinToolTip", "Pin the channel across all profiles."))
+				[
+					SNew(SImage)
+					.Image(this, &SAvaChannel::GetChannelPinBrush)
+				]
+			]
+			+ SHorizontalBox::Slot()
+			.AutoWidth()
+			[
+				SNew(SButton)
+				.ButtonStyle(FAppStyle::Get(), "SimpleButton")
+				.OnClicked(this, &SAvaChannel::OnChannelMaximizeButtonClicked)
+				.ToolTipText(this, &SAvaChannel::GetChannelMaximizeRestoreTooltipText)
+				[
+					SNew(SImage)
+					.Image(this, &SAvaChannel::GetChannelMaximizeRestoreBrush)
+				]
+			]
+			+ SHorizontalBox::Slot()
+			.AutoWidth()
+			[
+				SNew(SButton)
+				.ButtonStyle(FAppStyle::Get(), "SimpleButton")
+				.OnClicked(this, &SAvaChannel::OnChannelRemoveButtonClicked)
+				.IsEnabled(this, &SAvaChannel::CanEditChanges)
+				.ToolTipText(LOCTEXT("ChannelRemoveToolTip", "Remove this channel."))
+				[
+					SNew(SImage)
+					.Image(FAppStyle::Get().GetBrush("Icons.X"))
+				]
+			]
+		]
+		//Status Button
+		+ SOverlay::Slot()
+		.HAlign(EHorizontalAlignment::HAlign_Left)
+		.VAlign(EVerticalAlignment::VAlign_Top)
+		.Padding(5.f, 5.f, 0.f, 0.f)
+		[
+			SNew(SHorizontalBox)
+			+ SHorizontalBox::Slot()
+			.AutoWidth()
+			[
+				SNew(SButton)
+				.ButtonStyle(FAppStyle::Get(), "SimpleButton")
+				.OnClicked(this, &SAvaChannel::OnChannelStatusButtonClicked)
+				.ToolTipText(LOCTEXT("ChannelStatusChangeToolTip", "Change the status of this channel."))
+				[
+					SAssignNew(ChannelStatusOptions, SMenuAnchor)
+					.OnGetMenuContent(this, &SAvaChannel::GetChannelStatusOptions)
+					[
+						SNew(SImage)
+						.Image(this, &SAvaChannel::GetChannelStatusBrush)
+					]				
+				]
+			]
+			+ SHorizontalBox::Slot()
+			.Padding(2.f, 0.f, 0.f, 0.f)
+			.HAlign(EHorizontalAlignment::HAlign_Left)
+			.VAlign(EVerticalAlignment::VAlign_Center)
+			.AutoWidth()
+			[
+				SNew(STextBlock)
+				.Text(this, &SAvaChannel::GetChannelStatusText)
+				.ToolTipText_Lambda([this]()
+					{
+						return FText::Format(LOCTEXT("ChannelStatusToolTip", "Current Channel Status: {0}"), GetChannelStatusText());
+					})
+			]
+			+ SHorizontalBox::Slot()
+			.AutoWidth()
+			[
+				SNew(SButton)
+				.ButtonStyle(FAppStyle::Get(), "SimpleButton")
+				.OnClicked(this, &SAvaChannel::OnChannelTypeToggleButtonClicked)
+				.ToolTipText(LOCTEXT("ChannelTypeToggleToolTip", "Toggles between \"Program\" or \"Preview\" channel type."))
+				[
+					SNew(SImage)
+					.Image(this, &SAvaChannel::GetChannelTypeBrush)
+				]
+			]
+		]
+	];
+
+	OutputTileListView->SetStyle(nullptr);	
+	const FAvaOutputChannel& Channel = UAvalancheBroadcast::Get().GetCurrentProfile().GetChannel(ChannelName);
+	OnChannelBroadcastStateChanged(Channel);
+	OnChannelMediaOutputsChanged(Channel);
+}
+
+SAvaChannel::~SAvaChannel()
+{
+	FAvaOutputChannel::GetOnChannelChanged().RemoveAll(this);	
+	if (TSharedPtr<FAvaBroadcastEditor> BroadcastEditor = BroadcastEditorWeak.Pin())
+	{
+		BroadcastEditor->OnOutputTileSelectionChanged.RemoveAll(this);
+	}
+}
+
+void SAvaChannel::RegisterCommands()
+{
+	MediaOutputCommandList = MakeShared<FUICommandList>();
+	
+	MediaOutputCommandList->MapAction(FGenericCommands::Get().Delete
+		, FExecuteAction::CreateSP(this, &SAvaChannel::DeleteSelectedOutputTiles)
+		, FCanExecuteAction::CreateSP(this, &SAvaChannel::CanEditChanges));
+}
+
+void SAvaChannel::SetPosition(int32 InColumnIndex, int32 InRowIndex)
+{
+	ColumnIndex= InColumnIndex;
+	RowIndex = InRowIndex;
+}
+
+void SAvaChannel::OnChannelChanged(const FAvaOutputChannel& InChannel, EAvaChannelChange InChange)
+{
+	if (EnumHasAnyFlags(InChange, EAvaChannelChange::State))
+	{
+		OnChannelBroadcastStateChanged(InChannel);
+	}
+	if (EnumHasAnyFlags(InChange, EAvaChannelChange::RenderTarget))
+	{
+		OnChannelRenderTargetChanged(InChannel);
+	}
+	if (EnumHasAnyFlags(InChange, EAvaChannelChange::MediaOutputs))
+	{
+		OnChannelMediaOutputsChanged(InChannel);
+	}
+}
+
+void SAvaChannel::OnChannelBroadcastStateChanged(const FAvaOutputChannel& InChannel)
+{	
+	if (InChannel.IsValidChannel() && InChannel.GetChannelName() == ChannelName)
+	{
+		ChannelState = InChannel.GetState();
+		ChannelStateText = FAvaMediaEditorUtils::GetChannelStatusText(ChannelState, InChannel.GetIssueSeverity());
+		ChannelStatusBrush = FAvaMediaEditorUtils::GetChannelStatusBrush(ChannelState, InChannel.GetIssueSeverity());
+		
+		if (ChannelState != EAvaChannelState::Live)
+		{
+			ChannelPreviewBrush.Reset();
+		}
+
+		//Update ChannelPreviewBrush when Visibility Changes. It might be we haven't created it yet.
+		OnChannelRenderTargetChanged(InChannel);
+	}
+}
+
+void SAvaChannel::OnChannelRenderTargetChanged(const FAvaOutputChannel& InChannel)
+{
+	if (InChannel.IsValidChannel() && InChannel.GetChannelName() == ChannelName)
+	{
+		UTextureRenderTarget2D* const RenderTarget = InChannel.GetCurrentRenderTarget(true);
+		
+		//Only Invert Alpha if the Current Render Target isn't the Placeholder RT, since Placeholder uses Widget Rendering.
+		//TODO: Need to check whether the Channel or Media Outputs is doing Invert Alpha rather than figuring it out here
+		bShouldInvertAlpha = InChannel.GetPlaceholderRenderTarget() != RenderTarget;
+
+		if (RenderTarget)
+		{
+			const FVector2D RenderTargetSize(RenderTarget->SizeX, RenderTarget->SizeY);
+			
+			//If Brush is invalid, or the Brush's Texture Target doesn't match the new Render Target, reset the Brush.
+			if (!ChannelPreviewBrush.IsValid() || ChannelPreviewBrush->GetResourceObject() != RenderTarget)
+			{
+				ChannelPreviewBrush = MakeShared<FSlateImageBrush>(RenderTarget, RenderTargetSize);
+			}
+			//If Sizes mismatch, just resizes rather than recreating the Brush with same underlying Resource
+			else if (ChannelPreviewBrush->GetImageSize() !=  RenderTargetSize)
+			{
+				ChannelPreviewBrush->SetImageSize(RenderTargetSize);
+			}
+		}
+	}
+}
+
+void SAvaChannel::OnChannelMediaOutputsChanged(const FAvaOutputChannel& InChannel)
+{
+	if (InChannel.IsValidChannel() && InChannel.GetChannelName() == ChannelName)
+	{
+		TSet<UMediaOutput*> CurrentMediaOutputs(InChannel.GetMediaOutputs());
+		
+		TSet<UMediaOutput*> SeenMediaOutputs;
+		SeenMediaOutputs.Reserve(OutputTileItems.Num());
+		
+		//Remove Items no longer in the Media Outputs List
+		for (TArray<FAvaOutputTileItemPtr>::TIterator Iter = OutputTileItems.CreateIterator(); Iter; ++Iter)
+		{
+			const FAvaOutputTileItemPtr& Item = *Iter;
+
+			//Remove if Item or Underlying Media Output is Invalid, 
+			if (!Item.IsValid())
+			{
+				Iter.RemoveCurrent();
+				continue;
+			}
+
+			UMediaOutput* const MediaOutput = Item->GetMediaOutput();
+			if (!IsValid(MediaOutput) || !CurrentMediaOutputs.Contains(MediaOutput))
+			{
+				Iter.RemoveCurrent();
+			}
+			else
+			{
+				SeenMediaOutputs.Add(MediaOutput);
+			}
+		}
+
+		//Add New Media Outputs as Item in the List
+		{
+			TArray<UMediaOutput*> NewMediaOutputs = CurrentMediaOutputs.Difference(SeenMediaOutputs).Array();
+			OutputTileItems.Reserve(OutputTileItems.Num() + NewMediaOutputs.Num());
+			
+			for (UMediaOutput* const MediaOutput : NewMediaOutputs)
+			{
+				TSharedPtr<FAvaOutputTileItem> Item = MakeShared<FAvaOutputTileItem>(ChannelName, MediaOutput);
+				OutputTileItems.Add(Item);
+			}
+		}
+	
+		if (OutputTileListView.IsValid())
+		{
+			OutputTileListView->RequestListRefresh();
+		}
+	}
+}
+
+void SAvaChannel::OnOutputTileSelectionChanged(const TSharedPtr<FAvaOutputTileItem>& Item)
+{
+	if (OutputTileListView.IsValid() && !OutputTileItems.Contains(Item))
+	{
+		//Clear Selection without Notifying
+		OutputTileListView->Private_ClearSelection();
+	}
+}
+
+void SAvaChannel::DeleteSelectedOutputTiles()
+{
+	if (OutputTileListView.IsValid() && OutputTileListView->GetNumItemsSelected() > 0)
+	{
+		TArray<FAvaOutputTileItemPtr> SelectedTiles = OutputTileListView->GetSelectedItems();
+		TArray<UMediaOutput*> MediaOutputs;
+		MediaOutputs.Reserve(SelectedTiles.Num());
+
+		for (const FAvaOutputTileItemPtr& SelectedTile : SelectedTiles)
+		{
+			if (SelectedTile.IsValid())
+			{
+				MediaOutputs.Add(SelectedTile->GetMediaOutput());
+			}
+		}
+
+		FScopedTransaction Transaction(LOCTEXT("RemoveMediaOutputs", "Remove Media Outputs"));
+		
+		UAvalancheBroadcast& Broadcast = UAvalancheBroadcast::Get();
+		Broadcast.Modify();
+		
+		const int32 RemovedCount = Broadcast.GetCurrentProfile().RemoveChannelMediaOutputs(ChannelName, MediaOutputs);
+		
+		if (RemovedCount == 0)
+		{
+			Transaction.Cancel();
+		}
+	}
+}
+
+FText SAvaChannel::GetChannelStatusText() const
+{
+	return ChannelStateText;
+}
+
+TSharedRef<SWidget> SAvaChannel::GetChannelStatusOptions()
+{
+	FMenuBuilder Builder(true, nullptr);
+	
+	Builder.BeginSection("Status", LOCTEXT("BroadcastStatus", "Broadcast Status"));
+	{
+		const FAvaOutputChannel& Channel = UAvalancheBroadcast::Get().GetCurrentProfile().GetChannel(ChannelName);
+		// An offline channel can't be made live from here.
+		if (Channel.GetState() != EAvaChannelState::Offline)
+		{
+			for (uint8 Index = 0; Index < static_cast<uint8>(EAvaChannelState::Max); ++Index)
+			{
+				const EAvaChannelState State = static_cast<EAvaChannelState>(Index);
+
+				// Only add other states that the broadcast isn't currently in (or offline).
+				if (Channel.GetState() != State && State != EAvaChannelState::Offline)
+				{
+					Builder.AddMenuEntry(StaticEnum<EAvaChannelState>()->GetDisplayNameTextByIndex(Index)
+						, FText()
+						, FSlateIcon()
+						, FUIAction(FExecuteAction::CreateSP(this, &SAvaChannel::OnChannelStatusSelected, State)));	
+				}
+			}
+		}
+	}
+	Builder.EndSection();
+	
+	return Builder.MakeWidget();
+}
+
+TSharedRef<SWidget> SAvaChannel::MakeChannelVerticalBox()
+{
+	return SNew(SVerticalBox)
+		+ SVerticalBox::Slot()
+		.FillHeight(0.1f)
+		.Padding(0.f, 2.f, 0.f, 2.f)
+		.VAlign(EVerticalAlignment::VAlign_Fill)
+		.HAlign(EHorizontalAlignment::HAlign_Fill)
+		.MaxHeight(30.f)
+		[
+			SNew(SScaleBox)
+			.Stretch(EStretch::ScaleToFitY)
+			[
+				SAssignNew(ChannelNameTextBlock, SInlineEditableTextBlock)
+				.Text(this, &SAvaChannel::GetChannelNameText)
+				.OnVerifyTextChanged(this, &SAvaChannel::OnVerifyChannelNameTextChanged)
+				.OnTextCommitted(this, &SAvaChannel::OnChannelNameTextCommitted)
+				.IsReadOnly(this, &SAvaChannel::IsReadOnly)
+				.Justification(ETextJustify::Center)
+			]
+		]
+		+ SVerticalBox::Slot()
+		.FillHeight(0.15f)
+		.MaxHeight(50.f)
+		.VAlign(EVerticalAlignment::VAlign_Fill)
+		.HAlign(EHorizontalAlignment::HAlign_Fill)
+		[
+			SNew(SOverlay)
+			+ SOverlay::Slot()
+			.HAlign(EHorizontalAlignment::HAlign_Center)
+			[
+				SNew(SScaleBox)
+				.Stretch(EStretch::ScaleToFitY)
+				[
+					MakeMediaOutputsTileView()
+				]
+			]
+			+ SOverlay::Slot()
+			.HAlign(EHorizontalAlignment::HAlign_Center)
+			.VAlign(EVerticalAlignment::VAlign_Center)
+			[
+				SNew(STextBlock)
+				.Visibility(this, &SAvaChannel::GetMediaOutputEmptyTextVisibility)
+				.Text(LOCTEXT("NoMediaOutputsFound", "No Media Outputs. Drag and Drop from the Output Devices List."))
+				.TextStyle(FAppStyle::Get(), "HintText")
+			]
+		]
+		+ SVerticalBox::Slot()
+		.FillHeight(0.75f)
+		.VAlign(EVerticalAlignment::VAlign_Fill)
+		.HAlign(EHorizontalAlignment::HAlign_Fill)
+		[
+			SNew(SBorder)
+			.Visibility(this, &SAvaChannel::GetChannelPreviewVisibility)
+			[
+				SNew(SScaleBox)
+				.Stretch(EStretch::ScaleToFit)
+				.VAlign(EVerticalAlignment::VAlign_Center)
+				.HAlign(EHorizontalAlignment::HAlign_Center)
+				[
+					SNew(SAvaCaptureImage)
+					.ImageArgs(SImage::FArguments()
+						.Image(this, &SAvaChannel::GetChannelPreviewBrush))
+					.ShouldInvertAlpha(this, &SAvaChannel::ShouldInvertAlpha)
+					.EnableGammaCorrection(false)
+					.EnableBlending(false)
+				]
+			]
+		];		
+}
+
+TSharedRef<SWidget> SAvaChannel::MakeMediaOutputsTileView()
+{
+	return SAssignNew(OutputTileListView, SListView<FAvaOutputTileItemPtr>)
+		.Orientation(EOrientation::Orient_Horizontal)
+		.ScrollbarVisibility(EVisibility::Collapsed)
+		.SelectionMode(ESelectionMode::Single)
+		.ListItemsSource(&OutputTileItems)
+		.OnGenerateRow(this, &SAvaChannel::OnGenerateMediaOutputTile)
+		.OnSelectionChanged(this, &SAvaChannel::OnMediaOutputTileSelectionChanged)
+		.OnContextMenuOpening(this, &SAvaChannel::OnMediaOutputTileContextMenuOpening)
+	;
+}
+
+TSharedRef<ITableRow> SAvaChannel::OnGenerateMediaOutputTile(FAvaOutputTileItemPtr Item, const TSharedRef<STableViewBase>& InOwnerTable) const
+{
+	check(Item.IsValid());
+	return SNew(STableRow<FAvaOutputTileItemPtr>, InOwnerTable)
+		.Padding(FMargin(5.f, 5.f))
+		.OnDragDetected(Item.ToSharedRef(), &FAvaOutputTileItem::OnDragDetected)
+		[
+			Item->GenerateTile()
+		];
+}
+
+void SAvaChannel::OnMediaOutputTileSelectionChanged(FAvaOutputTileItemPtr Item, ESelectInfo::Type SelectInfo)
+{
+	if (TSharedPtr<FAvaBroadcastEditor> BroadcastEditor = BroadcastEditorWeak.Pin())
+	{
+		BroadcastEditor->SelectOutputTile(Item);
+	}
+}
+
+TSharedPtr<SWidget> SAvaChannel::OnMediaOutputTileContextMenuOpening() const
+{
+	FMenuBuilder MenuBuilder(true, MediaOutputCommandList);
+	MenuBuilder.BeginSection("Edit", LOCTEXT("Edit", "Edit"));
+	{
+		MenuBuilder.AddMenuEntry(FGenericCommands::Get().Delete);
+	}
+	MenuBuilder.EndSection();
+	return MenuBuilder.MakeWidget();
+}
+
+void SAvaChannel::OnDragEnter(const FGeometry& MyGeometry, const FDragDropEvent& DragDropEvent)
+{
+	SetDragging(false);
+	if (CanEditChanges())
+	{
+		if (const TSharedPtr<FAvaOutputTreeItemDragDropOp> OutputClassDragDropOp = DragDropEvent.GetOperationAs<FAvaOutputTreeItemDragDropOp>())
+		{
+			if (OutputClassDragDropOp->IsValidToDropInChannel(ChannelName))
+			{
+				SetDragging(true);
+			}
+		}
+		else if (const TSharedPtr<FAvaOutputTileItemDragDropOp> OutputTileDragDropOp = DragDropEvent.GetOperationAs<FAvaOutputTileItemDragDropOp>())
+		{
+			if (OutputTileDragDropOp->IsValidToDropInChannel(ChannelName))
+			{
+				SetDragging(true);
+			}
+		}
+	}
+}
+
+void SAvaChannel::OnDragLeave(const FDragDropEvent& DragDropEvent)
+{
+	SetDragging(false);
+}
+
+FReply SAvaChannel::OnDrop(const FGeometry& MyGeometry, const FDragDropEvent& DragDropEvent)
+{
+	SetDragging(false);
+	if (CanEditChanges())
+	{
+		if (const TSharedPtr<FAvaOutputTreeItemDragDropOp> OutputItemDragDropOp = DragDropEvent.GetOperationAs<FAvaOutputTreeItemDragDropOp>())
+		{
+			OutputItemDragDropOp->OnChannelDrop(ChannelName);
+		}
+		else if (const TSharedPtr<FAvaOutputTileItemDragDropOp> OutputTileDragDropOp = DragDropEvent.GetOperationAs<FAvaOutputTileItemDragDropOp>())
+		{
+			OutputTileDragDropOp->OnChannelDrop(ChannelName);
+		}
+	}
+	return FReply::Unhandled();
+}
+
+FReply SAvaChannel::OnKeyDown(const FGeometry& MyGeometry, const FKeyEvent& InKeyEvent)
+{
+	if (MediaOutputCommandList.IsValid() && MediaOutputCommandList->ProcessCommandBindings(InKeyEvent))
+	{
+		return FReply::Handled();
+	}
+	return FReply::Unhandled();
+}
+
+void SAvaChannel::SetDragging(bool bIsDragging)
+{
+	DragVisibility = bIsDragging
+		? EVisibility::SelfHitTestInvisible
+		: EVisibility::Hidden;
+}
+
+FText SAvaChannel::GetChannelNameText() const
+{
+	return FText::FromName(ChannelName);
+}
+
+bool SAvaChannel::OnVerifyChannelNameTextChanged(const FText& InText, FText& OutErrorMessage)
+{
+	return UAvalancheBroadcast::Get().CanRenameChannel(ChannelName, FName(InText.ToString()));
+}
+
+void SAvaChannel::OnChannelNameTextCommitted(const FText& InText, ETextCommit::Type CommitInfo)
+{
+	const FName NewChannelName(InText.ToString());
+
+	FScopedTransaction Transaction(LOCTEXT("RenameChannel", "Rename Channel"));
+	UAvalancheBroadcast& Broadcast = UAvalancheBroadcast::Get();
+	Broadcast.Modify();
+	
+	if (Broadcast.RenameChannel(ChannelName, NewChannelName))
+	{
+		ChannelName = NewChannelName;
+	}
+	else
+	{
+		Transaction.Cancel();
+	}
+}
+
+FReply SAvaChannel::OnChannelStatusButtonClicked()
+{
+	if (ChannelStatusOptions.IsValid() && !ChannelStatusOptions->IsOpen())
+	{
+		ChannelStatusOptions->SetIsOpen(true);
+		return FReply::Handled();
+	}
+	return FReply::Unhandled();
+}
+
+FReply SAvaChannel::OnChannelSettingsButtonClicked()
+{
+	if (!ChannelSettings.IsValid())
+	{
+		FAvaOutputChannel& Channel = UAvalancheBroadcast::Get().GetCurrentProfile().GetChannelMutable(ChannelName);
+		if (!Channel.IsValidChannel())
+		{
+			return FReply::Unhandled();
+		}
+
+		FPropertyEditorModule& PropertyEditorModule = FModuleManager::GetModuleChecked<FPropertyEditorModule>("PropertyEditor");
+
+		FDetailsViewArgs DetailsViewArgs;
+		DetailsViewArgs.NameAreaSettings = FDetailsViewArgs::HideNameArea;
+		DetailsViewArgs.bAllowSearch = false;
+		DetailsViewArgs.bShowOptions = false;
+
+		TSharedRef<FStructOnScope> ChannelStruct = MakeShared<FStructOnScope>(FAvaOutputChannel::StaticStruct(), reinterpret_cast<uint8*>(&Channel));
+
+		ChannelSettings = PropertyEditorModule.CreateStructureDetailView(DetailsViewArgs, FStructureDetailsViewArgs(), ChannelStruct);
+
+		ChannelSettings->GetOnFinishedChangingPropertiesDelegate().AddLambda(
+			[ChannelStructWeak = TWeakPtr<FStructOnScope>(ChannelStruct)](const FPropertyChangedEvent&)
+			{
+				if (TSharedPtr<FStructOnScope> ChannelStruct = ChannelStructWeak.Pin())
+				{
+					const FAvaOutputChannel& Channel = *reinterpret_cast<const FAvaOutputChannel*>(ChannelStruct->GetStructMemory());
+					FAvaOutputChannel::GetOnChannelChanged().Broadcast(Channel, EAvaChannelChange::Settings);
+				}
+			});
+
+		ChannelSettingsMenuAnchor->SetMenuContent(SNew(SBox)
+				.MaxDesiredHeight(500.f)
+				.MinDesiredWidth(150.f)
+				[
+					SNew(SScrollBox)
+					+ SScrollBox::Slot()
+					[
+				    	ChannelSettings->GetWidget().ToSharedRef()
+					]
+				]
+			);
+	}
+
+	ChannelSettingsMenuAnchor->SetIsOpen(true);
+	return FReply::Handled();
+}
+
+FReply SAvaChannel::OnChannelPinButtonClicked()
+{
+	FScopedTransaction Transaction(LOCTEXT("ToggleChannelPin", "Toggle Channel Pin"));
+	
+	UAvalancheBroadcast& Broadcast = UAvalancheBroadcast::Get();	
+	Broadcast.Modify();
+	if (!Broadcast.IsChannelPinned(ChannelName))
+	{
+		Broadcast.PinChannel(ChannelName, Broadcast.GetCurrentProfileName());
+	}
+	else
+	{
+		Broadcast.UnpinChannel(ChannelName);
+	}
+	Broadcast.RebuildProfiles();
+	Broadcast.QueueNotifyChange(EAvaBroadcastChange::ChannelGrid);
+	Broadcast.GetOnChannelsListChanged().Broadcast(Broadcast.GetCurrentProfile());
+	
+	return FReply::Handled();
+}
+
+FReply SAvaChannel::OnChannelTypeToggleButtonClicked()
+{
+	FScopedTransaction Transaction(LOCTEXT("ToggleChannelType", "Toggle Channel Type"));
+
+	UAvalancheBroadcast& Broadcast = UAvalancheBroadcast::Get();
+	Broadcast.Modify();
+
+	const EAvaBroadcastChannelType NewChannelType = (Broadcast.GetChannelType(ChannelName) == EAvaBroadcastChannelType::Preview)
+		? EAvaBroadcastChannelType::Program : EAvaBroadcastChannelType::Preview;
+	Broadcast.SetChannelType(ChannelName, NewChannelType);
+	Broadcast.QueueNotifyChange(EAvaBroadcastChange::ChannelType);
+	Broadcast.GetOnChannelsListChanged().Broadcast(Broadcast.GetCurrentProfile());
+
+	return FReply::Handled();
+}
+
+FReply SAvaChannel::OnChannelMaximizeButtonClicked()
+{
+	if (OnMaximizeClicked.IsBound())
+	{
+		OnMaximizeClicked.Execute(SharedThis(this));
+		return FReply::Handled();
+	}
+	return FReply::Unhandled();
+}
+
+FReply SAvaChannel::OnChannelRemoveButtonClicked()
+{
+	FScopedTransaction Transaction(LOCTEXT("RemoveChannel", "Remove Channel"));
+	
+	UAvalancheBroadcast& Broadcast = UAvalancheBroadcast::Get();
+	Broadcast.Modify();
+	
+	if (Broadcast.GetCurrentProfile().RemoveChannel(ChannelName))
+	{
+		return FReply::Handled();
+	}
+
+	Transaction.Cancel();
+	return FReply::Unhandled();
+}
+
+void SAvaChannel::OnChannelStatusSelected(EAvaChannelState State)
+{
+	FAvaOutputChannel& Channel = UAvalancheBroadcast::Get().GetCurrentProfile().GetChannelMutable(ChannelName);
+	if (Channel.IsValidChannel())
+	{
+		switch (State)
+		{
+		case EAvaChannelState::Idle:
+			Channel.StopChannelBroadcast();
+			break;
+		case EAvaChannelState::Live:
+			Channel.StartChannelBroadcast();
+			break;
+		}
+	}
+}
+
+const FSlateBrush* SAvaChannel::GetChannelStatusBrush() const
+{
+	return ChannelStatusBrush;
+}
+
+const FSlateBrush* SAvaChannel::GetChannelPreviewBrush() const
+{
+	return ChannelPreviewBrush.Get();
+}
+
+const FSlateBrush* SAvaChannel::GetChannelPinBrush() const
+{
+	return UAvalancheBroadcast::Get().IsChannelPinned(ChannelName) ?
+		FAppStyle::Get().GetBrush("Icons.Pinned") : FAppStyle::Get().GetBrush("Icons.Unpinned");
+}
+
+const FSlateBrush* SAvaChannel::GetChannelMaximizeRestoreBrush() const
+{
+	FName BrushName = TEXT("EditorViewport.Front");
+	if (CanMaximize.IsBound() && !CanMaximize.Get())
+	{
+		BrushName = TEXT("LevelEditor.Tabs.Viewports");
+	}
+	return FAppStyle::Get().GetBrush(BrushName);
+}
+
+const FSlateBrush* SAvaChannel::GetChannelTypeBrush() const
+{
+	return UAvalancheBroadcast::Get().GetChannelType(ChannelName) == EAvaBroadcastChannelType::Preview ?
+		FAvaMediaEditorStyle::Get().GetBrush("AvalancheMediaEditor.ChannelTypePreview")
+		: FAvaMediaEditorStyle::Get().GetBrush("AvalancheMediaEditor.ChannelTypeProgram");
+}
+
+EVisibility SAvaChannel::GetMediaOutputEmptyTextVisibility() const
+{
+	return OutputTileItems.IsEmpty()
+		? EVisibility::Visible
+		: EVisibility::Hidden;
+}
+
+EVisibility SAvaChannel::GetChannelPreviewVisibility() const
+{
+	return UAvalancheBroadcast::Get().CanShowPreview()
+		? EVisibility::Visible
+		: EVisibility::Collapsed;
+}
+
+EVisibility SAvaChannel::GetChannelDragVisibility() const
+{
+	return DragVisibility;
+}
+
+bool SAvaChannel::ShouldInvertAlpha() const
+{
+	return bShouldInvertAlpha;
+}
+
+bool SAvaChannel::IsReadOnly() const
+{
+	return ChannelState == EAvaChannelState::Live;
+}
+
+bool SAvaChannel::CanEditChanges() const
+{
+	return !IsReadOnly();
+}
+
+FText SAvaChannel::GetChannelMaximizeRestoreTooltipText() const
+{
+	FText TooltipText = LOCTEXT("ChannelMaximizeToolTip", "Maximize this channel in the view. All other channels will be hidden.");
+	if (CanMaximize.IsBound() && !CanMaximize.Get())
+	{
+		TooltipText = LOCTEXT("ChannelRestoreToolTip", "Restore this channel in the view. All other channels will be visible.");
+	}
+	return TooltipText;
+}
+
+#undef LOCTEXT_NAMESPACE
diff --git a/Engine/Plugins/Experimental/Avalanche/Source/AvalancheMediaEditor/Private/Broadcast/ChannelGrid/Slate/SAvaChannel.h b/Engine/Plugins/Experimental/Avalanche/Source/AvalancheMediaEditor/Private/Broadcast/ChannelGrid/Slate/SAvaChannel.h
new file mode 100644
--- /dev/null
+++ b/Engine/Plugins/Experimental/Avalanche/Source/AvalancheMediaEditor/Private/Broadcast/ChannelGrid/Slate/SAvaChannel.h
@@ -0,0 +1,141 @@
+// Copyright Epic Games, Inc. All Rights Reserved.
+
+#pragma once
+
+#include "Broadcast/ChannelGrid/AvaOutputTileItem.h"
+#include "Channel/AvaOutputChannel.h"
+#include "CoreMinimal.h"
+#include "MediaOutput.h"
+#include "Widgets/SCompoundWidget.h"
+#include "Widgets/Views/STileView.h"
+
+class FAvaBroadcastEditor;
+class FUICommandList;
+class IStructureDetailsView;
+class SAvaChannel;
+class SInlineEditableTextBlock;
+class SMenuAnchor;
+class UMediaOutput;
+struct FAvaOutputChannel;
+struct FSlateImageBrush;
+
+DECLARE_DELEGATE_RetVal(bool, FAvaCanMaximize);
+DECLARE_DELEGATE_OneParam(FAvaOnMaximizeClicked, const TSharedRef<SAvaChannel>&);
+
+class SAvaChannel : public SCompoundWidget
+{
+public:
+	
+	SLATE_BEGIN_ARGS(SAvaChannel) {}
+		SLATE_ARGUMENT(FName, ChannelName)
+		SLATE_ATTRIBUTE(bool, CanMaximize)
+		SLATE_EVENT(FAvaOnMaximizeClicked, OnMaximizeClicked)
+	SLATE_END_ARGS()
+
+	/** Constructs this widget with InArgs */
+	void Construct(const FArguments& InArgs, const TSharedPtr<FAvaBroadcastEditor>& InBroadcastEditor);
+	virtual ~SAvaChannel() override;
+
+	void RegisterCommands();
+	FName GetChannelName() const { return ChannelName; }
+
+	void SetPosition(int32 InColumnIndex, int32 InRowIndex);
+	int32 GetRowIndex() const { return RowIndex; }
+	int32 GetColumnIndex() const { return ColumnIndex; }
+
+	void OnChannelChanged(const FAvaOutputChannel& InChannel, EAvaChannelChange InChange);
+	
+	void OnChannelBroadcastStateChanged(const FAvaOutputChannel& InChannel);
+	void OnChannelRenderTargetChanged(const FAvaOutputChannel& InChannel);
+	void OnChannelMediaOutputsChanged(const FAvaOutputChannel& InChannel);
+	
+	void OnOutputTileSelectionChanged(const TSharedPtr<FAvaOutputTileItem>& InMediaOutputItem);
+	
+	void DeleteSelectedOutputTiles();
+
+	FText GetChannelStatusText() const;
+	
+	TSharedRef<SWidget> GetChannelStatusOptions();
+	TSharedRef<SWidget> MakeChannelVerticalBox();
+	TSharedRef<SWidget> MakeMediaOutputsTileView();
+
+	TSharedRef<ITableRow> OnGenerateMediaOutputTile(FAvaOutputTileItemPtr Item, const TSharedRef<STableViewBase>& InOwnerTable) const;
+	void OnMediaOutputTileSelectionChanged(FAvaOutputTileItemPtr Item,	ESelectInfo::Type SelectInfo);
+	TSharedPtr<SWidget> OnMediaOutputTileContextMenuOpening() const;
+	
+	virtual void OnDragEnter(const FGeometry& MyGeometry, const FDragDropEvent& DragDropEvent) override;
+	virtual void OnDragLeave(const FDragDropEvent& DragDropEvent) override;
+	virtual FReply OnDrop(const FGeometry& MyGeometry, const FDragDropEvent& DragDropEvent) override;
+
+	virtual FReply OnKeyDown(const FGeometry& MyGeometry, const FKeyEvent& InKeyEvent) override;;
+	
+	void SetDragging(bool bIsDragging);
+
+	FText GetChannelNameText() const;
+	bool OnVerifyChannelNameTextChanged(const FText& InText, FText& OutErrorMessage);
+	void OnChannelNameTextCommitted(const FText& InText, ETextCommit::Type CommitInfo);
+
+	FReply OnChannelStatusButtonClicked();
+	FReply OnChannelSettingsButtonClicked();
+	FReply OnChannelPinButtonClicked();
+	FReply OnChannelTypeToggleButtonClicked();
+	FReply OnChannelMaximizeButtonClicked();
+	FReply OnChannelRemoveButtonClicked();
+
+	void OnChannelStatusSelected(EAvaChannelState State);
+
+	const FSlateBrush* GetChannelStatusBrush() const;
+	const FSlateBrush* GetChannelPreviewBrush() const;
+	const FSlateBrush* GetChannelPinBrush() const;
+	const FSlateBrush* GetChannelMaximizeRestoreBrush() const;
+	const FSlateBrush* GetChannelTypeBrush() const;
+
+	EVisibility GetMediaOutputEmptyTextVisibility() const;
+	EVisibility GetChannelPreviewVisibility() const;
+	EVisibility GetChannelDragVisibility() const;
+
+	bool ShouldInvertAlpha() const;
+
+	bool IsReadOnly() const;
+	bool CanEditChanges() const;
+	
+protected:
+	FText GetChannelMaximizeRestoreTooltipText() const;
+	
+	FName ChannelName = NAME_None;
+
+	TSharedPtr<SMenuAnchor> ChannelStatusOptions;
+	
+	TSharedPtr<SListView<FAvaOutputTileItemPtr>> OutputTileListView;
+	
+	TSharedPtr<IStructureDetailsView> ChannelSettings;
+
+	TSharedPtr<SMenuAnchor> ChannelSettingsMenuAnchor;
+
+	TArray<FAvaOutputTileItemPtr> OutputTileItems;
+	
+	TSharedPtr<FSlateImageBrush> ChannelPreviewBrush;
+
+	const FSlateBrush* ChannelStatusBrush = nullptr;
+	
+	TSharedPtr<SInlineEditableTextBlock> ChannelNameTextBlock;
+	
+	TWeakPtr<FAvaBroadcastEditor> BroadcastEditorWeak;
+
+	TSharedPtr<FUICommandList> MediaOutputCommandList;
+	
+	EVisibility DragVisibility = EVisibility::Hidden;
+
+	EAvaChannelState ChannelState = EAvaChannelState::Idle;
+	
+	FText ChannelStateText;
+
+	TAttribute<bool> CanMaximize;
+	FAvaOnMaximizeClicked OnMaximizeClicked;
+
+	int32 RowIndex = INDEX_NONE;
+	
+	int32 ColumnIndex= INDEX_NONE;
+	
+	bool bShouldInvertAlpha = false;
+};
diff --git a/Engine/Plugins/Experimental/Avalanche/Source/AvalancheMediaEditor/Private/Broadcast/OutputDevices/AvaOutputClassItem.cpp b/Engine/Plugins/Experimental/Avalanche/Source/AvalancheMediaEditor/Private/Broadcast/OutputDevices/AvaOutputClassItem.cpp
new file mode 100644
--- /dev/null
+++ b/Engine/Plugins/Experimental/Avalanche/Source/AvalancheMediaEditor/Private/Broadcast/OutputDevices/AvaOutputClassItem.cpp
@@ -0,0 +1,202 @@
+// Copyright Epic Games, Inc. All Rights Reserved.
+
+#include "AvaOutputClassItem.h"
+#include "AvaOutputDeviceItem.h"
+#include "AvaOutputServerItem.h"
+#include "AvalancheBroadcast.h"
+#include "IAvaMediaModule.h"
+#include "IMediaIOCoreDeviceProvider.h"
+#include "IMediaIOCoreModule.h"
+#include "MediaIOCoreDefinitions.h"
+#include "MediaOutput.h"
+#include "OutputDevices/AvaDeviceProviderData.h"
+#include "OutputDevices/AvaMediaOutputUtils.h"
+#include "OutputDevices/IAvaDeviceProviderProxyManager.h"
+#include "ScopedTransaction.h"
+#include "Slate/SAvaOutputTreeItem.h"
+#include "Styling/SlateIconFinder.h"
+
+#define LOCTEXT_NAMESPACE "AvaOutputClassItem"
+
+namespace UE::AvaOutputClassItem::Private
+{
+	IMediaIOCoreDeviceProvider* GetDeviceProvider(const UClass* InMediaOutputClass)
+	{
+		const FName DeviceProviderName = UE::AvaMediaOutputUtils::GetDeviceProviderName(InMediaOutputClass);
+		if (!IMediaIOCoreModule::IsAvailable() || DeviceProviderName.IsNone())
+		{
+			return nullptr;
+		}
+
+		// Remark: even if the class has a device provider name, it is possible
+		// that there is no corresponding device provider.
+		return IMediaIOCoreModule::Get().GetDeviceProvider(DeviceProviderName);
+	}
+
+	void UpdateDeviceName(UAvalancheBroadcast& InBroadcast, const FName& InChannelName, const UMediaOutput* InMediaOutput, const FString& InDeviceName)
+	{
+		FAvaOutputChannel& Channel = InBroadcast.GetCurrentProfile().GetChannelMutable(InChannelName);
+		if (Channel.IsValidChannel())
+		{
+			if (FAvaMediaOutputInfo* ExistingOutputInfo = Channel.GetMediaOutputInfoMutable(InMediaOutput))
+			{
+				ExistingOutputInfo->DeviceName = FName(InDeviceName);	
+			}
+		}
+	}
+}
+
+FText FAvaOutputClassItem::GetDisplayName() const
+{
+	check(OutputClass.IsValid());
+	return OutputClass->GetDisplayNameText();
+}
+
+const FSlateBrush* FAvaOutputClassItem::GetIconBrush() const
+{
+	check(OutputClass.IsValid());
+	return FSlateIconFinder::FindIconBrushForClass(OutputClass.Get());
+}
+
+void FAvaOutputClassItem::RefreshChildren()
+{
+	IMediaIOCoreDeviceProvider* DeviceProvider = UE::AvaOutputClassItem::Private::GetDeviceProvider(OutputClass.Get());
+
+	if (!DeviceProvider)
+	{
+		Children.Reset();
+		return;
+	}
+
+	TSharedPtr<FAvaOutputTreeItem> Parent = ParentWeak.Pin();
+	if (!Parent.IsValid())
+	{
+		return;
+	}
+
+	TSharedPtr<FAvaOutputServerItem> ParentServerItem = StaticCastSharedPtr<FAvaOutputServerItem>(Parent);
+	check(ParentServerItem.IsValid());
+
+	const FName DeviceProviderName = DeviceProvider->GetFName();
+	const IAvaDeviceProviderProxyManager& DeviceProviderProxyManager = IAvaMediaModule::Get().GetDeviceProviderProxyManager();
+
+	TArray<FMediaIOConfiguration> OutputConfigs = DeviceProvider->GetConfigurations(false, true);
+
+	// Fill current devices from output configs of this device provider.
+	TSet<FAvaMediaDevice> CurrentDevices;
+	CurrentDevices.Reserve(OutputConfigs.Num());
+
+	for (const FMediaIOConfiguration& Config : OutputConfigs)
+	{
+		const FString ServerName = DeviceProviderProxyManager.FindServerNameForDevice(DeviceProviderName, Config.MediaConnection.Device.DeviceName);
+		const bool bIsLocal = DeviceProviderProxyManager.IsLocalDevice(DeviceProviderName, Config.MediaConnection.Device.DeviceName);
+		const FAvaDeviceProviderData* DeviceProviderDataForServer = ParentServerItem->GetDeviceProviderData(DeviceProviderName);
+
+		if ((!DeviceProviderDataForServer && bIsLocal)
+			 || (DeviceProviderDataForServer && !bIsLocal && ServerName == ParentServerItem->GetServerName()))
+		{
+			CurrentDevices.Add({ Config.MediaConnection.Device, DeviceProviderName, ServerName });
+		}
+	}
+
+	//A set to contain the Media Devices of the Current Children
+	TSet<FAvaMediaDevice> SeenDevices;
+	SeenDevices.Reserve(CurrentDevices.Num());
+	
+	//Remove Existing Children that are Invalid
+	for (TArray<FAvaOutputTreeItemPtr>::TIterator ItemIt = Children.CreateIterator(); ItemIt; ++ItemIt)
+	{
+		FAvaOutputTreeItemPtr Item(*ItemIt);
+
+		//Remove Invalid Pointers or Items that are not Device Items since a Class Item can only contain Device Items
+		if (!Item.IsValid() || !Item->IsA<FAvaOutputDeviceItem>())
+		{
+			ItemIt.RemoveCurrent();
+			continue;
+		}
+
+		const TSharedPtr<FAvaOutputDeviceItem> DeviceItem = StaticCastSharedPtr<FAvaOutputDeviceItem>(Item);
+		const FAvaMediaDevice& UnderlyingDevice = DeviceItem->GetDevice();
+
+		if (UnderlyingDevice.IsValid() && CurrentDevices.Contains(UnderlyingDevice))
+		{
+			SeenDevices.Add(UnderlyingDevice);
+		}
+		else
+		{
+			//Remove if there's no valid Underlying Device or it's no longer in the set
+			ItemIt.RemoveCurrent();
+		}
+	}
+	
+	//Append the Devices that are not already in the Original Children List
+	{
+		TSharedPtr<FAvaOutputClassItem> This = SharedThis(this);
+		TArray<FAvaMediaDevice> NewDevices = CurrentDevices.Difference(SeenDevices).Array();
+		Children.Reserve(Children.Num() + NewDevices.Num());
+		
+		for (const FAvaMediaDevice& Device : NewDevices)
+		{
+			TSharedPtr<FAvaOutputDeviceItem> DeviceItem = MakeShared<FAvaOutputDeviceItem>(This, Device);
+			Children.Add(DeviceItem);
+		}
+	}
+}
+
+TSharedPtr<SWidget> FAvaOutputClassItem::GenerateRowWidget()
+{
+	return SNew(SAvaOutputTreeItem, SharedThis(this));
+}
+
+UMediaOutput* FAvaOutputClassItem::AddMediaOutputToChannel(FName InTargetChannel, const FAvaMediaOutputInfo& InOutputInfo)
+{
+	check(OutputClass.IsValid());
+
+	UAvalancheBroadcast& Broadcast = UAvalancheBroadcast::Get();
+	
+	FScopedTransaction Transaction(LOCTEXT("AddMediaOutput", "Add Media Output"));
+	Broadcast.Modify();
+
+	FAvaMediaOutputInfo OutputInfo = InOutputInfo;
+	
+	// Support non-enumerated devices (i.e. no device provider).
+	if (!OutputInfo.IsValid() || !OutputInfo.Guid.IsValid())
+	{
+		if (TSharedPtr<FAvaOutputTreeItem> Parent = ParentWeak.Pin())
+		{
+			const TSharedPtr<FAvaOutputServerItem> ParentServerItem = StaticCastSharedPtr<FAvaOutputServerItem>(Parent);
+			check(ParentServerItem.IsValid());
+
+			// Fill the device info with available data.
+			OutputInfo.Guid = FGuid::NewGuid();
+			OutputInfo.ServerName = ParentServerItem->GetServerName();
+
+			// Fetch the device provider name (if there is one) from the class's meta data.
+			OutputInfo.DeviceProviderName = UE::AvaMediaOutputUtils::GetDeviceProviderName(OutputClass.Get());
+
+			// Device name:
+			// There is no enumerated device name. It will be updated from the media output object below.
+		}
+	}
+	
+	UMediaOutput* const MediaOutput = Broadcast.GetCurrentProfile().AddChannelMediaOutput(InTargetChannel, OutputClass.Get(), OutputInfo);
+
+	// If the device name was not provided as input, we need to update it from the Media Output object.
+	if (MediaOutput && OutputInfo.DeviceName.IsNone())
+	{
+		const FString DeviceName = UE::AvaMediaOutputUtils::GetDeviceName(MediaOutput);
+		if (!DeviceName.IsEmpty())
+		{
+			UE::AvaOutputClassItem::Private::UpdateDeviceName(Broadcast, InTargetChannel, MediaOutput, DeviceName);
+		}
+	}
+
+	if (!MediaOutput)
+	{
+		Transaction.Cancel();
+	}
+	
+	return MediaOutput;
+}
+
+#undef LOCTEXT_NAMESPACE
diff --git a/Engine/Plugins/Experimental/Avalanche/Source/AvalancheMediaEditor/Private/Broadcast/OutputDevices/AvaOutputClassItem.h b/Engine/Plugins/Experimental/Avalanche/Source/AvalancheMediaEditor/Private/Broadcast/OutputDevices/AvaOutputClassItem.h
new file mode 100644
--- /dev/null
+++ b/Engine/Plugins/Experimental/Avalanche/Source/AvalancheMediaEditor/Private/Broadcast/OutputDevices/AvaOutputClassItem.h
@@ -0,0 +1,36 @@
+// Copyright Epic Games, Inc. All Rights Reserved.
+
+#pragma once
+
+#include "AvaOutputServerItem.h"
+#include "AvaOutputTreeItem.h"
+#include "UObject/WeakObjectPtr.h"
+#include "UObject/WeakObjectPtrTemplates.h"
+
+class FAvaOutputClassItem : public FAvaOutputTreeItem
+{
+public:
+	UE_AVA_INHERITS_WITH_SUPER(FAvaOutputClassItem, FAvaOutputTreeItem);
+
+	FAvaOutputClassItem(const TSharedPtr<FAvaOutputTreeItem>& InParent, UClass* InOutputClass)
+		: Super(InParent)
+		, OutputClass(InOutputClass)
+	{
+	}
+
+	UClass* GetOutputClass() const
+	{
+		return OutputClass.Get();
+	}
+
+	//~ Begin IAvaOutputTreeItem
+	virtual FText GetDisplayName() const override;
+	virtual const FSlateBrush* GetIconBrush() const override;
+	virtual void RefreshChildren() override;
+	virtual TSharedPtr<SWidget> GenerateRowWidget() override;
+	virtual UMediaOutput* AddMediaOutputToChannel(FName InTargetChannel, const FAvaMediaOutputInfo& InOutputInfo) override;
+	//~ End IAvaOutputTreeItem
+
+protected:
+	TWeakObjectPtr<UClass> OutputClass;
+};
diff --git a/Engine/Plugins/Experimental/Avalanche/Source/AvalancheMediaEditor/Private/Broadcast/OutputDevices/AvaOutputDeviceItem.cpp b/Engine/Plugins/Experimental/Avalanche/Source/AvalancheMediaEditor/Private/Broadcast/OutputDevices/AvaOutputDeviceItem.cpp
new file mode 100644
--- /dev/null
+++ b/Engine/Plugins/Experimental/Avalanche/Source/AvalancheMediaEditor/Private/Broadcast/OutputDevices/AvaOutputDeviceItem.cpp
@@ -0,0 +1,105 @@
+// Copyright Epic Games, Inc. All Rights Reserved.
+
+#include "AvaOutputDeviceItem.h"
+
+#include "AvalancheBroadcast.h"
+#include "AvaOutputClassItem.h"
+#include "Channel/AvaMediaOutputInfo.h"
+#include "IMediaIOCoreDeviceProvider.h"
+#include "IMediaIOCoreModule.h"
+#include "MediaOutput.h"
+#include "Widgets/Text/STextBlock.h"
+
+#define LOCTEXT_NAMESPACE "AvaOutputDeviceItem"
+
+FText FAvaOutputDeviceItem::GetDisplayName() const
+{
+	return Device.GetDisplayNameText();
+}
+
+const FSlateBrush* FAvaOutputDeviceItem::GetIconBrush() const
+{
+	return nullptr;
+}
+
+void FAvaOutputDeviceItem::RefreshChildren()
+{
+	//There shouldn't be any Childrens here
+	Children.Reset();
+}
+
+TSharedPtr<SWidget> FAvaOutputDeviceItem::GenerateRowWidget()
+{
+	return SNew(STextBlock)
+		.Text(GetDisplayName());
+}
+
+bool FAvaOutputDeviceItem::IsValidToDropInChannel(FName InTargetChannelName)
+{
+	// It is invalid to drop a "remote" device in a preview channel.
+	if (UAvalancheBroadcast::Get().GetChannelType(InTargetChannelName) == EAvaBroadcastChannelType::Preview)
+	{
+		return !FAvaMediaOutputInfo::IsRemote(Device.GetServerName());
+	}
+	return true;
+}
+
+UMediaOutput* FAvaOutputDeviceItem::AddMediaOutputToChannel(FName InTargetChannel, const FAvaMediaOutputInfo& /*InOutputInfo*/)
+{
+	if (TSharedPtr<FAvaOutputTreeItem> Parent = ParentWeak.Pin())
+	{
+		FAvaMediaOutputInfo OutputInfo;
+		OutputInfo.Guid = FGuid::NewGuid();
+		OutputInfo.ServerName = Device.GetServerName();
+		OutputInfo.DeviceProviderName = Device.GetDeviceProviderName();
+		OutputInfo.DeviceName = Device.GetDevice().DeviceName;
+
+		UMediaOutput* const MediaOutput = Parent->AddMediaOutputToChannel(InTargetChannel, OutputInfo);
+		
+		if (MediaOutput)
+		{
+			const UClass* const MediaOutputClass = MediaOutput->GetClass();
+			const IMediaIOCoreDeviceProvider* const DeviceProvider = IMediaIOCoreModule::Get().GetDeviceProvider(Device.GetDeviceProviderName());
+			
+			if (DeviceProvider)
+			{
+				FStructProperty* const Property = FindFProperty<FStructProperty>(MediaOutputClass, TEXT("OutputConfiguration"));
+				if (Property && Property->Struct->IsChildOf(FMediaIOOutputConfiguration::StaticStruct()))
+				{
+					if (FMediaIOOutputConfiguration* const OutputConfig = Property->ContainerPtrToValuePtr<FMediaIOOutputConfiguration>(MediaOutput))
+					{
+						{
+							// Initial "default" setup. 
+							// This leaves the OutputConfig object in a possibly invalid state
+							// because it is very likely the selected device's config doesn't match the one returned by
+							// GetDefaultOutputConfiguration().
+							*OutputConfig = DeviceProvider->GetDefaultOutputConfiguration();
+							OutputConfig->MediaConfiguration.MediaConnection.Device = Device.GetDevice();
+						}
+
+						// Fetch all possible output configurations.
+						TArray<FMediaIOOutputConfiguration> OutputConfigs = DeviceProvider->GetOutputConfigurations();
+
+						// Select one that match the device. 
+						// For now just get the first one.
+						for (const FMediaIOOutputConfiguration& Config : OutputConfigs)
+						{
+							if (Config.MediaConfiguration.MediaConnection.Device.DeviceName == Device.GetDevice().DeviceName)
+							{
+								*OutputConfig = Config;
+								break; // get the first one.
+							}
+						}
+
+						FPropertyChangedEvent PropertyChangedEvent(Property, EPropertyChangeType::ValueSet);
+						MediaOutput->PostEditChangeProperty(PropertyChangedEvent);
+					}
+				}
+			}
+		}
+		return MediaOutput;
+	}
+	return nullptr;
+}
+
+#undef LOCTEXT_NAMESPACE
diff --git a/Engine/Plugins/Experimental/Avalanche/Source/AvalancheMediaEditor/Private/Broadcast/OutputDevices/AvaOutputDeviceItem.h b/Engine/Plugins/Experimental/Avalanche/Source/AvalancheMediaEditor/Private/Broadcast/OutputDevices/AvaOutputDeviceItem.h
new file mode 100644
--- /dev/null
+++ b/Engine/Plugins/Experimental/Avalanche/Source/AvalancheMediaEditor/Private/Broadcast/OutputDevices/AvaOutputDeviceItem.h
@@ -0,0 +1,38 @@
+// Copyright Epic Games, Inc. All Rights Reserved.
+
+#pragma once
+
+#include "AvaMediaDevice.h"
+#include "AvaOutputTreeItem.h"
+#include "MediaIOCoreDefinitions.h"
+
+class FAvaOutputClassItem;
+
+class FAvaOutputDeviceItem : public FAvaOutputTreeItem
+{
+public:
+	UE_AVA_INHERITS_WITH_SUPER(FAvaOutputDeviceItem, FAvaOutputTreeItem);
+
+	FAvaOutputDeviceItem(const TSharedPtr<FAvaOutputTreeItem>& InParent, const FAvaMediaDevice& InDevice)
+		: Super(InParent)
+		, Device(InDevice)
+	{
+	}
+
+	const FAvaMediaDevice& GetDevice() const
+	{
+		return Device;
+	}
+
+	//~ Begin IAvaOutputTreeItem
+	virtual FText GetDisplayName() const override;
+	virtual const FSlateBrush* GetIconBrush() const override;
+	virtual void RefreshChildren() override;
+	virtual TSharedPtr<SWidget> GenerateRowWidget() override;
+	virtual bool IsValidToDropInChannel(FName InTargetChannelName) override;
+	virtual UMediaOutput* AddMediaOutputToChannel(FName InTargetChannel, const FAvaMediaOutputInfo& InOutputInfo) override;
+	//~ End IAvaOutputTreeItem
+
+protected:
+	FAvaMediaDevice Device;
+};
diff --git a/Engine/Plugins/Experimental/Avalanche/Source/AvalancheMediaEditor/Private/Broadcast/OutputDevices/AvaOutputRootItem.cpp b/Engine/Plugins/Experimental/Avalanche/Source/AvalancheMediaEditor/Private/Broadcast/OutputDevices/AvaOutputRootItem.cpp
new file mode 100644
--- /dev/null
+++ b/Engine/Plugins/Experimental/Avalanche/Source/AvalancheMediaEditor/Private/Broadcast/OutputDevices/AvaOutputRootItem.cpp
@@ -0,0 +1,102 @@
+// Copyright Epic Games, Inc. All Rights Reserved.
+
+#include "AvaOutputRootItem.h"
+#include "AvaOutputServerItem.h"
+#include "IAvaMediaModule.h"
+#include "Input/Reply.h"
+#include "MediaOutput.h"
+#include "OutputDevices/IAvaDeviceProviderProxyManager.h"
+#include "UObject/UObjectIterator.h"
+
+void FAvaOutputRootItem::RefreshChildren()
+{
+	const IAvaDeviceProviderProxyManager& DeviceProviderProxyManager = IAvaMediaModule::Get().GetDeviceProviderProxyManager();
+
+	const TSet<FString> CurrentServers = DeviceProviderProxyManager.GetServerNames();
+
+	TSet<FString> SeenServerNames;
+	SeenServerNames.Reserve(CurrentServers.Num());
+
+	//Add local server item
+	const TSharedPtr<FAvaOutputServerItem> OutputLocalServerItem = MakeShared<FAvaOutputServerItem>(DeviceProviderProxyManager.GetLocalServerName(), nullptr);
+	Children.Add(OutputLocalServerItem);
+
+	//Remove Existing Children that are Invalid
+	for (TArray<FAvaOutputTreeItemPtr>::TIterator ItemIt = Children.CreateIterator(); ItemIt; ++ItemIt)
+	{
+		FAvaOutputTreeItemPtr Item(*ItemIt);
+
+		//Always going to display local devices group
+		if (Item == OutputLocalServerItem)
+		{
+			continue;
+		}
+
+		//Remove Invalid Pointers or Items that are not Output Server Items since Root can only have the Server Items as Top Level
+		if (!Item.IsValid() || !Item->IsA<FAvaOutputServerItem>())
+		{
+			ItemIt.RemoveCurrent();
+			continue;
+		}
+
+		const TSharedPtr<FAvaOutputServerItem> OutputServerItem = StaticCastSharedPtr<FAvaOutputServerItem>(Item);
+		const FString UnderlyingServerName = OutputServerItem->GetServerName();
+
+		if (!UnderlyingServerName.IsEmpty() && CurrentServers.Contains(UnderlyingServerName))
+		{
+			SeenServerNames.Add(UnderlyingServerName);
+		}
+		else
+		{
+			//Remove if there's no valid Underlying OutputClass or it's no longer in the set
+			ItemIt.RemoveCurrent();
+		}
+	}
+	
+	//Append the New Servers that are not already in the Original Children List
+	{
+		TSet<FString> NewServers = CurrentServers.Difference(SeenServerNames);
+		Children.Reserve(Children.Num() + NewServers.Num());
+		
+		for (const FString& ServerName : NewServers)
+		{
+			if (const TSharedPtr<const FAvaDeviceProviderDataList> ServerDeviceProviderDataList = DeviceProviderProxyManager.GetDeviceProviderDataListForServer(ServerName))
+			{
+				TSharedPtr<FAvaOutputServerItem> OutputServerItem = MakeShared<FAvaOutputServerItem>(ServerName, ServerDeviceProviderDataList);
+				Children.Add(OutputServerItem);
+			}
+		}
+	}
+}
+
+FText FAvaOutputRootItem::GetDisplayName() const
+{
+	//Root shouldn't need to give a Display Name!
+	check(0);
+	return FText::GetEmpty();
+}
+
+const FSlateBrush* FAvaOutputRootItem::GetIconBrush() const
+{
+	check(0);
+	return nullptr;
+}
+
+TSharedPtr<SWidget> FAvaOutputRootItem::GenerateRowWidget()
+{
+	//Root shouldn't need to Generate a Row Widget
+	check(0);
+	return nullptr;
+}
+
+FReply FAvaOutputRootItem::OnDragDetected(const FGeometry& MyGeometry, const FPointerEvent& MouseEvent)
+{
+	check(0);
+	return FReply::Unhandled();
+}
+
+UMediaOutput* FAvaOutputRootItem::AddMediaOutputToChannel(FName InTargetChannel, const FAvaMediaOutputInfo& InOutputInfo)
+{
+	check(0);
+	return nullptr;
+}
diff --git a/Engine/Plugins/Experimental/Avalanche/Source/AvalancheMediaEditor/Private/Broadcast/OutputDevices/AvaOutputRootItem.h b/Engine/Plugins/Experimental/Avalanche/Source/AvalancheMediaEditor/Private/Broadcast/OutputDevices/AvaOutputRootItem.h
new file mode 100644
--- /dev/null
+++ b/Engine/Plugins/Experimental/Avalanche/Source/AvalancheMediaEditor/Private/Broadcast/OutputDevices/AvaOutputRootItem.h
@@ -0,0 +1,25 @@
+// Copyright Epic Games, Inc. All Rights Reserved.
+
+#pragma once
+
+#include "AvaOutputTreeItem.h"
+
+class FAvaOutputRootItem : public FAvaOutputTreeItem
+{
+public:
+	UE_AVA_INHERITS_WITH_SUPER(FAvaOutputRootItem, FAvaOutputTreeItem);
+
+	FAvaOutputRootItem()
+		: Super(nullptr)
+	{}
+
+	//~ Begin IAvaOutputTreeItem
+	virtual void RefreshChildren() override;
+private:
+	virtual FText GetDisplayName() const override;
+	virtual const FSlateBrush* GetIconBrush() const override;
+	virtual TSharedPtr<SWidget> GenerateRowWidget() override;
+	virtual FReply OnDragDetected(const FGeometry& MyGeometry, const FPointerEvent& MouseEvent) override;
+	virtual UMediaOutput* AddMediaOutputToChannel(FName InTargetChannel, const FAvaMediaOutputInfo& InOutputInfo) override;
+	//~ End IAvaOutputTreeItem
+};
diff --git a/Engine/Plugins/Experimental/Avalanche/Source/AvalancheMediaEditor/Private/Broadcast/OutputDevices/AvaOutputServerItem.cpp b/Engine/Plugins/Experimental/Avalanche/Source/AvalancheMediaEditor/Private/Broadcast/OutputDevices/AvaOutputServerItem.cpp
new file mode 100644
--- /dev/null
+++ b/Engine/Plugins/Experimental/Avalanche/Source/AvalancheMediaEditor/Private/Broadcast/OutputDevices/AvaOutputServerItem.cpp
@@ -0,0 +1,126 @@
+// Copyright Epic Games, Inc. All Rights Reserved.
+
+#include "AvaOutputServerItem.h"
+
+#include "AvaOutputClassItem.h"
+#include "AvalancheBroadcast.h"
+#include "AvalancheMediaEditorSettings.h"
+#include "IMediaIOCoreModule.h"
+#include "MediaOutput.h"
+#include "OutputDevices/AvaMediaOutputUtils.h"
+#include "ScopedTransaction.h"
+#include "Slate/SAvaOutputTreeItem.h"
+#include "Styling/AppStyle.h"
+#include "UObject/UObjectIterator.h"
+
+#define LOCTEXT_NAMESPACE "AvaOutputServerItem"
+
+FString FAvaOutputServerItem::GetServerName() const
+{
+	return ServerName;
+}
+
+const FAvaDeviceProviderData* FAvaOutputServerItem::GetDeviceProviderData(FName InDeviceProviderName) const
+{
+	if (!DeviceProviderDataList.IsValid())
+	{
+		return nullptr;
+	}
+
+	for (const FAvaDeviceProviderData& DeviceProviderData : DeviceProviderDataList->DeviceProviders)
+	{
+		if (DeviceProviderData.Name == InDeviceProviderName)
+		{
+			return &DeviceProviderData;
+		}
+	}
+
+	return nullptr;
+}
+
+FText FAvaOutputServerItem::GetDisplayName() const
+{
+	return FText::FromString(ServerName);
+}
+
+const FSlateBrush* FAvaOutputServerItem::GetIconBrush() const
+{
+	return FAppStyle::Get().GetBrush(TEXT("Icons.Server"));
+}
+
+void FAvaOutputServerItem::RefreshChildren()
+{
+	if (!IMediaIOCoreModule::IsAvailable())
+	{
+		Children.Reset();
+		return;
+	}
+
+	// Make a set of all allowed Media Output Classes.
+	// It may have all classes or be restricted to those with the MediaIOCustomLayout MetaData.
+	TSet<UClass*> CurrentOutputClasses;
+	for (UClass* const Class : TObjectRange<UClass>())
+	{
+		const bool bIsMediaOutputClass = Class->IsChildOf(UMediaOutput::StaticClass()) && Class != UMediaOutput::StaticClass();
+		const bool bHasDeviceProvider = UE::AvaMediaOutputUtils::HasDeviceProviderName(Class);
+
+		if (bIsMediaOutputClass && (UAvalancheMediaEditorSettings::Get().bBroadcastShowAllMediaOutputClasses || bHasDeviceProvider))
+		{
+			CurrentOutputClasses.Add(Class);
+		}
+	}
+
+	TSet<UClass*> SeenOutputClasses;
+	SeenOutputClasses.Reserve(CurrentOutputClasses.Num());
+
+	//Remove Existing Children that are Invalid
+	for (TArray<FAvaOutputTreeItemPtr>::TIterator ItemIt = Children.CreateIterator(); ItemIt; ++ItemIt)
+	{
+		FAvaOutputTreeItemPtr Item(*ItemIt);
+
+		//Remove Invalid Pointers or Items that are not Output Class Items since Root can only have the Class Items as Top Level
+		if (!Item.IsValid() || !Item->IsA<FAvaOutputClassItem>())
+		{
+			ItemIt.RemoveCurrent();
+			continue;
+		}
+
+		const TSharedPtr<FAvaOutputClassItem> OutputClassItem = StaticCastSharedPtr<FAvaOutputClassItem>(Item);
+		UClass* const UnderlyingOutputClass = OutputClassItem->GetOutputClass();
+
+		if (UnderlyingOutputClass && CurrentOutputClasses.Contains(UnderlyingOutputClass))
+		{
+			SeenOutputClasses.Add(UnderlyingOutputClass);
+		}
+		else
+		{
+			//Remove if there's no valid Underlying OutputClass or it's no longer in the set
+			ItemIt.RemoveCurrent();
+		}
+	}
+
+	//Append the New Output Classes that are not already in the Original Children List
+	{
+		TSharedPtr<FAvaOutputServerItem> This = SharedThis(this);
+		TArray<UClass*> NewOutputClasses = CurrentOutputClasses.Difference(SeenOutputClasses).Array();
+		Children.Reserve(Children.Num() + NewOutputClasses.Num());
+
+		for (UClass* const NewOutputClass : NewOutputClasses)
+		{
+			TSharedPtr<FAvaOutputClassItem> OutputClassItem = MakeShared<FAvaOutputClassItem>(This, NewOutputClass);
+			Children.Add(OutputClassItem);
+		}
+	}
+}
+
+TSharedPtr<SWidget> FAvaOutputServerItem::GenerateRowWidget()
+{
+	return SNew(SAvaOutputTreeItem, SharedThis(this));
+}
+
+UMediaOutput* FAvaOutputServerItem::AddMediaOutputToChannel(FName InTargetChannel, const FAvaMediaOutputInfo& InOutputInfo)
+{
+	return nullptr;
+}
+
+#undef LOCTEXT_NAMESPACE
diff --git a/Engine/Plugins/Experimental/Avalanche/Source/AvalancheMediaEditor/Private/Broadcast/OutputDevices/AvaOutputServerItem.h b/Engine/Plugins/Experimental/Avalanche/Source/AvalancheMediaEditor/Private/Broadcast/OutputDevices/AvaOutputServerItem.h
new file mode 100644
--- /dev/null
+++ b/Engine/Plugins/Experimental/Avalanche/Source/AvalancheMediaEditor/Private/Broadcast/OutputDevices/AvaOutputServerItem.h
@@ -0,0 +1,37 @@
+// Copyright Epic Games, Inc. All Rights Reserved.
+
+#pragma once
+
+#include "AvaOutputTreeItem.h"
+#include "OutputDevices/AvaDeviceProviderData.h"
+#include "UObject/WeakObjectPtr.h"
+#include "UObject/WeakObjectPtrTemplates.h"
+
+class FAvaOutputServerItem : public FAvaOutputTreeItem
+{
+public:
+	UE_AVA_INHERITS_WITH_SUPER(FAvaOutputServerItem, FAvaOutputTreeItem);
+
+	FAvaOutputServerItem(const FString& InServerName, const TSharedPtr<const FAvaDeviceProviderDataList>& InDeviceProviderDataList)
+		: Super(nullptr)
+		, ServerName(InServerName)
+		, DeviceProviderDataList(InDeviceProviderDataList)
+	{
+	}
+
+	FString GetServerName() const;
+	const FAvaDeviceProviderData* GetDeviceProviderData(FName InDeviceProviderName) const;
+
+	//~ Begin IAvaOutputTreeItem
+	virtual FText GetDisplayName() const override;
+	virtual const FSlateBrush* GetIconBrush() const override;
+	virtual void RefreshChildren() override;
+	virtual TSharedPtr<SWidget> GenerateRowWidget() override;
+	virtual bool IsValidToDropInChannel(FName InTargetChannelName) override { return false; }
+	virtual UMediaOutput* AddMediaOutputToChannel(FName InTargetChannel, const FAvaMediaOutputInfo& InOutputInfo) override;
+	//~ End IAvaOutputTreeItem
+
+protected:
+	FString ServerName;
+	const TSharedPtr<const FAvaDeviceProviderDataList> DeviceProviderDataList;
+};
diff --git a/Engine/Plugins/Experimental/Avalanche/Source/AvalancheMediaEditor/Private/Broadcast/OutputDevices/AvaOutputTreeItem.cpp b/Engine/Plugins/Experimental/Avalanche/Source/AvalancheMediaEditor/Private/Broadcast/OutputDevices/AvaOutputTreeItem.cpp
new file mode 100644
--- /dev/null
+++ b/Engine/Plugins/Experimental/Avalanche/Source/AvalancheMediaEditor/Private/Broadcast/OutputDevices/AvaOutputTreeItem.cpp
@@ -0,0 +1,39 @@
+// Copyright Epic Games, Inc. All Rights Reserved.
+
+#include "AvaOutputTreeItem.h"
+#include "DragDropOps/AvaOutputTreeItemDragDropOp.h"
+
+const TWeakPtr<FAvaOutputTreeItem>& FAvaOutputTreeItem::GetParent() const
+{
+	return ParentWeak;
+}
+
+const TArray<FAvaOutputTreeItemPtr>& FAvaOutputTreeItem::GetChildren() const
+{
+	return Children;
+}
+
+FReply FAvaOutputTreeItem::OnDragDetected(const FGeometry& MyGeometry, const FPointerEvent& MouseEvent)
+{
+	if (MouseEvent.IsMouseButtonDown(EKeys::LeftMouseButton))
+	{
+		return FReply::Handled().BeginDragDrop(FAvaOutputTreeItemDragDropOp::New(SharedThis(this)));
+	}
+	return FReply::Unhandled();
+}
+
+void FAvaOutputTreeItem::RefreshTree(const FAvaOutputTreeItemPtr& InItem)
+{
+	TArray<FAvaOutputTreeItemPtr> ItemsRemainingToRefresh;
+	ItemsRemainingToRefresh.Add(InItem);
+		
+	while (ItemsRemainingToRefresh.Num() > 0)
+	{
+		FAvaOutputTreeItemPtr Item = ItemsRemainingToRefresh.Pop();
+		if (Item.IsValid())
+		{
+			Item->RefreshChildren();
+			ItemsRemainingToRefresh.Append(Item->GetChildren());
+		}
+	}
+}
diff --git a/Engine/Plugins/Experimental/Avalanche/Source/AvalancheMediaEditor/Private/Broadcast/OutputDevices/AvaOutputTreeItem.h b/Engine/Plugins/Experimental/Avalanche/Source/AvalancheMediaEditor/Private/Broadcast/OutputDevices/AvaOutputTreeItem.h
new file mode 100644
--- /dev/null
+++ b/Engine/Plugins/Experimental/Avalanche/Source/AvalancheMediaEditor/Private/Broadcast/OutputDevices/AvaOutputTreeItem.h
@@ -0,0 +1,69 @@
+// Copyright Epic Games, Inc. All Rights Reserved.
+
+#pragma once
+
+#include "AvaType.h"
+#include "Channel/AvaMediaOutputInfo.h"
+
+class FAvaOutputTreeItem;
+class FReply;
+class IAvaOutputTreeItem;
+class SWidget;
+class UMediaOutput;
+struct FGeometry;
+struct FPointerEvent;
+struct FSlateBrush;
+
+typedef TSharedPtr<IAvaOutputTreeItem> FAvaOutputTreeItemPtr;
+
+class IAvaOutputTreeItem : public IAvaTypeCastable, public TSharedFromThis<IAvaOutputTreeItem>
+{
+public:
+	UE_AVA_INHERITS(IAvaOutputTreeItem, IAvaTypeCastable);
+
+	virtual FText GetDisplayName() const = 0;
+
+	virtual const FSlateBrush* GetIconBrush() const = 0;
+	
+	/** Refreshes what the Children are of this Item. (not recursive!) */
+	virtual void RefreshChildren() = 0;
+
+	virtual TSharedPtr<SWidget> GenerateRowWidget() = 0;
+	
+	virtual const TWeakPtr<FAvaOutputTreeItem>& GetParent() const = 0;
+
+	virtual const TArray<FAvaOutputTreeItemPtr>& GetChildren() const = 0;
+
+	virtual FReply OnDragDetected(const FGeometry& MyGeometry, const FPointerEvent& MouseEvent) = 0;
+
+	/** Returns true if it is valid to add this item to the given channel. */
+	virtual bool IsValidToDropInChannel(FName InTargetChannelName) = 0;
+
+	virtual UMediaOutput* AddMediaOutputToChannel(FName InTargetChannel, const FAvaMediaOutputInfo& InOutputInfo) = 0;
+};
+
+class FAvaOutputTreeItem : public IAvaOutputTreeItem
+{
+public:
+	UE_AVA_INHERITS(FAvaOutputTreeItem, IAvaOutputTreeItem);
+
+	FAvaOutputTreeItem(const TSharedPtr<FAvaOutputTreeItem>& InParent)
+		: ParentWeak(InParent)
+	{
+	}
+
+	//~ Begin IAvaOutputTreeItem
+	virtual ~FAvaOutputTreeItem() override {};
+	virtual const TWeakPtr<FAvaOutputTreeItem>& GetParent() const override;
+	virtual const TArray<FAvaOutputTreeItemPtr>& GetChildren() const override;
+	virtual FReply OnDragDetected(const FGeometry& MyGeometry, const FPointerEvent& MouseEvent) override;
+	virtual bool IsValidToDropInChannel(FName InTargetChannelName) override { return true; }
+	//~ End IAvaOutputTreeItem
+
+	/** Calls RefreshChildren() on the tree of item. */
+	static void RefreshTree(const FAvaOutputTreeItemPtr& InItem);
+
+protected:
+	TWeakPtr<FAvaOutputTreeItem> ParentWeak;
+	TArray<FAvaOutputTreeItemPtr> Children;
+};
diff --git a/Engine/Plugins/Experimental/Avalanche/Source/AvalancheMediaEditor/Private/Broadcast/OutputDevices/DragDropOps/AvaOutputTreeItemDragDropOp.cpp b/Engine/Plugins/Experimental/Avalanche/Source/AvalancheMediaEditor/Private/Broadcast/OutputDevices/DragDropOps/AvaOutputTreeItemDragDropOp.cpp
new file mode 100644
--- /dev/null
+++ b/Engine/Plugins/Experimental/Avalanche/Source/AvalancheMediaEditor/Private/Broadcast/OutputDevices/DragDropOps/AvaOutputTreeItemDragDropOp.cpp
@@ -0,0 +1,58 @@
+// Copyright Epic Games, Inc. All Rights Reserved.
+
+#include "AvaOutputTreeItemDragDropOp.h"
+#include "AvalancheBroadcast.h"
+#include "Broadcast/OutputDevices/AvaOutputClassItem.h"
+#include "MediaOutput.h"
+#include "ScopedTransaction.h"
+#include "Styling/SlateIconFinder.h"
+
+#define LOCTEXT_NAMESPACE "AvaOutputTreeItemDragDropOp"
+
+TSharedRef<FAvaOutputTreeItemDragDropOp> FAvaOutputTreeItemDragDropOp::New(const TSharedPtr<FAvaOutputTreeItem>& InOutputTreeItem)
+{
+	TSharedRef<FAvaOutputTreeItemDragDropOp> DragDropOp = MakeShared<FAvaOutputTreeItemDragDropOp>();
+	DragDropOp->Init(InOutputTreeItem);
+	return DragDropOp;
+}
+
+bool FAvaOutputTreeItemDragDropOp::IsValidToDropInChannel(FName InTargetChannelName) const
+{
+	return OutputTreeItem.IsValid() && OutputTreeItem->IsValidToDropInChannel(InTargetChannelName);
+}
+
+FReply FAvaOutputTreeItemDragDropOp::OnChannelDrop(FName InTargetChannelName)
+{
+	if (const TSharedPtr<FAvaOutputTreeItem> Item = GetOutputTreeItem())
+	{
+		UAvalancheBroadcast& Broadcast = UAvalancheBroadcast::Get();
+		
+		FScopedTransaction Transaction(LOCTEXT("OnChannelDrop", "Drop Media Output"));
+		Broadcast.Modify();
+		
+		const UMediaOutput* const MediaOutput = Item->AddMediaOutputToChannel(InTargetChannelName, FAvaMediaOutputInfo());
+		if (IsValid(MediaOutput))
+		{
+			return FReply::Handled();
+		}
+		else
+		{
+			Transaction.Cancel();
+		}
+	}
+	return FReply::Unhandled();
+}
+
+void FAvaOutputTreeItemDragDropOp::Init(const TSharedPtr<FAvaOutputTreeItem>& InOutputTreeItem)
+{
+	OutputTreeItem = InOutputTreeItem;
+	CurrentHoverText = InOutputTreeItem->GetDisplayName();
+	CurrentIconBrush = InOutputTreeItem->GetIconBrush();
+	CurrentIconColorAndOpacity = FSlateColor::UseForeground();
+	MouseCursor = EMouseCursor::GrabHandClosed;
+
+	SetupDefaults();
+	Construct();
+}
+
+#undef LOCTEXT_NAMESPACE
diff --git a/Engine/Plugins/Experimental/Avalanche/Source/AvalancheMediaEditor/Private/Broadcast/OutputDevices/DragDropOps/AvaOutputTreeItemDragDropOp.h b/Engine/Plugins/Experimental/Avalanche/Source/AvalancheMediaEditor/Private/Broadcast/OutputDevices/DragDropOps/AvaOutputTreeItemDragDropOp.h
new file mode 100644
--- /dev/null
+++ b/Engine/Plugins/Experimental/Avalanche/Source/AvalancheMediaEditor/Private/Broadcast/OutputDevices/DragDropOps/AvaOutputTreeItemDragDropOp.h
@@ -0,0 +1,27 @@
+// Copyright Epic Games, Inc. All Rights Reserved.
+
+#pragma once
+
+#include "DragAndDrop/DecoratedDragDropOp.h"
+
+class FAvaOutputTreeItem;
+
+class FAvaOutputTreeItemDragDropOp : public FDecoratedDragDropOp
+{
+public:
+	DRAG_DROP_OPERATOR_TYPE(FAvaOutputClassDragDropOp, FDecoratedDragDropOp)
+
+	static TSharedRef<FAvaOutputTreeItemDragDropOp> New(const TSharedPtr<FAvaOutputTreeItem>& InOutputClassItem);
+
+	bool IsValidToDropInChannel(FName InTargetChannelName) const;
+
+	TSharedPtr<FAvaOutputTreeItem> GetOutputTreeItem() const { return OutputTreeItem; }
+
+	FReply OnChannelDrop(FName InTargetChannelName);
+
+protected:
+	void Init(const TSharedPtr<FAvaOutputTreeItem>& InOutputClassItem);
+
+	/** Keep Reference Count while Drag Dropping */
+	TSharedPtr<FAvaOutputTreeItem> OutputTreeItem;
+};
diff --git a/Engine/Plugins/Experimental/Avalanche/Source/AvalancheMediaEditor/Private/Broadcast/OutputDevices/Slate/SAvaOutputDevices.cpp b/Engine/Plugins/Experimental/Avalanche/Source/AvalancheMediaEditor/Private/Broadcast/OutputDevices/Slate/SAvaOutputDevices.cpp
new file mode 100644
--- /dev/null
+++ b/Engine/Plugins/Experimental/Avalanche/Source/AvalancheMediaEditor/Private/Broadcast/OutputDevices/Slate/SAvaOutputDevices.cpp
@@ -0,0 +1,150 @@
+// Copyright Epic Games, Inc. All Rights Reserved.
+
+#include "SAvaOutputDevices.h"
+#include "AvalancheBroadcast.h"
+#include "AvalancheMediaEditorSettings.h"
+#include "Broadcast/OutputDevices/AvaOutputRootItem.h"
+#include "Broadcast/OutputDevices/AvaOutputServerItem.h"
+#include "Widgets/Views/STreeView.h"
+
+TMap<uint32, bool> SAvaOutputDevices::ItemExpansionStates = {};
+
+SAvaOutputDevices::~SAvaOutputDevices()
+{
+	if (UObjectInitialized())
+	{
+		UAvalancheMediaEditorSettings::GetMutable().OnSettingChanged().RemoveAll(this);
+	}
+}
+
+void SAvaOutputDevices::Construct(const FArguments& InArgs, const TSharedPtr<FAvaBroadcastEditor>& InBroadcastEditor)
+{
+	RootItem = MakeShared<FAvaOutputRootItem>();
+	
+	SAssignNew(OutputTree, STreeView<FAvaOutputTreeItemPtr>)
+		.ItemHeight(20.0f)
+		.SelectionMode(ESelectionMode::Single)
+		.OnGenerateRow(this, &SAvaOutputDevices::OnGenerateItemRow)
+		.OnGetChildren(this, &SAvaOutputDevices::OnGetRowChildren)
+		.OnExpansionChanged(this, &SAvaOutputDevices::OnRowExpansionChanged)
+		.TreeItemsSource(&TopLevelItems);
+	
+	ChildSlot
+	[
+		OutputTree.ToSharedRef()
+	];
+
+	RefreshOutputDevices();
+	
+	BroadcastChangedHandle = UAvalancheBroadcast::Get().AddChangeListener(
+		FOnAvaBroadcastChanged::FDelegate::CreateSP(this, &SAvaOutputDevices::OnBroadcastChanged));
+	UAvalancheMediaEditorSettings::GetMutable().OnSettingChanged().AddSP(this, &SAvaOutputDevices::OnAvaMediaSettingsChanged);
+}
+
+TSharedRef<ITableRow> SAvaOutputDevices::OnGenerateItemRow(FAvaOutputTreeItemPtr InItem
+	, const TSharedRef<STableViewBase>& InOwnerTable)
+{
+	check(InItem.IsValid());
+	return SNew(STableRow<FAvaOutputTreeItemPtr>, InOwnerTable)
+		.ShowWires(false)
+		.Padding(FMargin(5.f, 5.f))
+		.OnDragDetected(InItem.ToSharedRef(), &IAvaOutputTreeItem::OnDragDetected)
+		[
+			InItem->GenerateRowWidget().ToSharedRef()
+		];
+}
+
+void SAvaOutputDevices::OnGetRowChildren(FAvaOutputTreeItemPtr InItem, TArray<FAvaOutputTreeItemPtr>& OutChildren) const
+{
+	if (InItem.IsValid())
+	{
+		OutChildren.Append(InItem->GetChildren());
+	}
+}
+
+void SAvaOutputDevices::OnRowExpansionChanged(FAvaOutputTreeItemPtr InItem, const bool bInIsExpanded)
+{
+	const uint32 Hash = ItemHashRecursive(InItem);
+
+	if (ItemExpansionStates.Contains(Hash))
+	{
+		ItemExpansionStates[Hash] = bInIsExpanded;
+	}
+	else
+	{
+		ItemExpansionStates.Add(Hash, bInIsExpanded);
+	}
+}
+
+uint32 SAvaOutputDevices::ItemHashRecursive(const FAvaOutputTreeItemPtr& InItem) const
+{
+	if (!InItem.IsValid() || InItem->IsA<FAvaOutputRootItem>())
+	{
+		return 0;
+	}
+
+	const uint32 Hash = GetTypeHash(InItem->GetDisplayName().ToString());
+
+	if (const TSharedPtr<FAvaOutputTreeItem>& Parent = InItem->GetParent().Pin())
+	{
+		const uint32 ParentHash = ItemHashRecursive(Parent);
+		if (ParentHash != 0)
+		{
+			const uint32 CombinedHash = HashCombine(ParentHash, Hash);
+			return CombinedHash;
+		}
+	}
+		
+	return Hash;
+}
+
+void SAvaOutputDevices::RefreshOutputDevices()
+{
+	check(RootItem.IsValid());
+	
+	//Refresh Items
+	FAvaOutputTreeItem::RefreshTree(RootItem);
+	TopLevelItems = RootItem->GetChildren();
+	OutputTree->RequestTreeRefresh();
+
+	SetExpansionStatesRecursive(RootItem);
+}
+
+void SAvaOutputDevices::OnBroadcastChanged(EAvaBroadcastChange InChange)
+{
+	if (EnumHasAnyFlags(InChange, EAvaBroadcastChange::OutputDevices))
+	{
+		RefreshOutputDevices();
+	}
+}
+
+void SAvaOutputDevices::OnAvaMediaSettingsChanged(UObject*, FPropertyChangedEvent& InPropertyChangeEvent)
+{
+	static const FName BroadcastShowAllMediaOutputClassesName = GET_MEMBER_NAME_CHECKED(UAvalancheMediaEditorSettings, bBroadcastShowAllMediaOutputClasses);
+	if (InPropertyChangeEvent.GetPropertyName() == BroadcastShowAllMediaOutputClassesName)
+	{
+		RefreshOutputDevices();
+	}
+}
+
+void SAvaOutputDevices::SetExpansionStatesRecursive(const FAvaOutputTreeItemPtr& InRootItem)
+{
+	if (!InRootItem.IsValid())
+	{
+		return;
+	}
+
+	for (const FAvaOutputTreeItemPtr& Child : InRootItem->GetChildren())
+	{
+		if (Child.IsValid())
+		{
+			const uint32 ChildHash = ItemHashRecursive(Child);
+			if (ItemExpansionStates.Contains(ChildHash))
+			{
+				OutputTree->SetItemExpansion(Child, ItemExpansionStates[ChildHash]);
+			}
+
+			SetExpansionStatesRecursive(Child);
+		}
+	}
+}
diff --git a/Engine/Plugins/Experimental/Avalanche/Source/AvalancheMediaEditor/Private/Broadcast/OutputDevices/Slate/SAvaOutputDevices.h b/Engine/Plugins/Experimental/Avalanche/Source/AvalancheMediaEditor/Private/Broadcast/OutputDevices/Slate/SAvaOutputDevices.h
new file mode 100644
--- /dev/null
+++ b/Engine/Plugins/Experimental/Avalanche/Source/AvalancheMediaEditor/Private/Broadcast/OutputDevices/Slate/SAvaOutputDevices.h
@@ -0,0 +1,47 @@
+// Copyright Epic Games, Inc. All Rights Reserved.
+
+#pragma once
+
+#include "AvaMediaDefines.h"
+#include "Broadcast/OutputDevices/AvaOutputTreeItem.h"
+#include "Widgets/SCompoundWidget.h"
+
+class FAvaBroadcastEditor;
+class FAvaOutputRootItem;
+class ITableRow;
+class STableViewBase;
+template<typename ItemType> class STreeView;
+
+class SAvaOutputDevices : public SCompoundWidget
+{
+public:
+	SLATE_BEGIN_ARGS(SAvaOutputDevices){}
+	SLATE_END_ARGS()
+	
+	void Construct(const FArguments& InArgs, const TSharedPtr<FAvaBroadcastEditor>& InBroadcastEditor);
+	virtual ~SAvaOutputDevices() override;
+	
+	TSharedRef<ITableRow> OnGenerateItemRow(FAvaOutputTreeItemPtr InItem, const TSharedRef<STableViewBase>& InOwnerTable);
+	
+	void OnGetRowChildren(FAvaOutputTreeItemPtr InItem, TArray<FAvaOutputTreeItemPtr>& OutChildren) const;
+	void OnRowExpansionChanged(FAvaOutputTreeItemPtr InItem, const bool bInIsExpanded);
+
+	void RefreshOutputDevices();
+
+protected:
+	void OnBroadcastChanged(EAvaBroadcastChange InChange);
+	void OnAvaMediaSettingsChanged(UObject*, FPropertyChangedEvent& InPropertyChangeEvent);
+
+	uint32 ItemHashRecursive(const FAvaOutputTreeItemPtr& InItem) const;
+	void SetExpansionStatesRecursive(const FAvaOutputTreeItemPtr& InRootItem);
+
+	static TMap<uint32, bool> ItemExpansionStates;
+
+	FDelegateHandle BroadcastChangedHandle;
+
+	TSharedPtr<STreeView<FAvaOutputTreeItemPtr>> OutputTree;
+	
+	TArray<FAvaOutputTreeItemPtr> TopLevelItems;
+
+	TSharedPtr<FAvaOutputRootItem> RootItem;
+};
diff --git a/Engine/Plugins/Experimental/Avalanche/Source/AvalancheMediaEditor/Private/Broadcast/OutputDevices/Slate/SAvaOutputTreeItem.cpp b/Engine/Plugins/Experimental/Avalanche/Source/AvalancheMediaEditor/Private/Broadcast/OutputDevices/Slate/SAvaOutputTreeItem.cpp
new file mode 100644
--- /dev/null
+++ b/Engine/Plugins/Experimental/Avalanche/Source/AvalancheMediaEditor/Private/Broadcast/OutputDevices/Slate/SAvaOutputTreeItem.cpp
@@ -0,0 +1,32 @@
+// Copyright Epic Games, Inc. All Rights Reserved.
+
+#include "SAvaOutputTreeItem.h"
+#include "Broadcast/OutputDevices/AvaOutputTreeItem.h"
+#include "Widgets/Images/SImage.h"
+#include "Widgets/SBoxPanel.h"
+#include "Widgets/Text/STextBlock.h"
+
+void SAvaOutputTreeItem::Construct(const FArguments& InArgs, const TSharedPtr<FAvaOutputTreeItem>& InOutputTreeItem)
+{
+	check(InOutputTreeItem.IsValid());
+	OutputTreeItemWeak = InOutputTreeItem;
+
+	ChildSlot
+		[
+			SNew(SHorizontalBox)
+			+ SHorizontalBox::Slot()
+			.AutoWidth()
+			[
+				SNew(SImage)
+				.Image(InOutputTreeItem->GetIconBrush())
+			]
+			+ SHorizontalBox::Slot()
+			.FillWidth(1.f)
+			.Padding(5.f, 0.f, 0.f, 0.f)
+			.VAlign(VAlign_Center)
+			[
+				SNew(STextBlock)
+				.Text(InOutputTreeItem->GetDisplayName())
+			]
+		];
+}
diff --git a/Engine/Plugins/Experimental/Avalanche/Source/AvalancheMediaEditor/Private/Broadcast/OutputDevices/Slate/SAvaOutputTreeItem.h b/Engine/Plugins/Experimental/Avalanche/Source/AvalancheMediaEditor/Private/Broadcast/OutputDevices/Slate/SAvaOutputTreeItem.h
new file mode 100644
--- /dev/null
+++ b/Engine/Plugins/Experimental/Avalanche/Source/AvalancheMediaEditor/Private/Broadcast/OutputDevices/Slate/SAvaOutputTreeItem.h
@@ -0,0 +1,19 @@
+// Copyright Epic Games, Inc. All Rights Reserved.
+
+#pragma once
+
+#include "Widgets/SCompoundWidget.h"
+
+class FAvaOutputTreeItem;
+
+class SAvaOutputTreeItem : public SCompoundWidget
+{
+public:
+	SLATE_BEGIN_ARGS(SAvaOutputTreeItem){}
+	SLATE_END_ARGS()
+	
+	void Construct(const FArguments& InArgs, const TSharedPtr<FAvaOutputTreeItem>& InOutputTreeItem);
+	
+protected:
+	TWeakPtr<FAvaOutputTreeItem> OutputTreeItemWeak;
+};
diff --git a/Engine/Plugins/Experimental/Avalanche/Source/AvalancheMediaEditor/Private/OutputDevices/AvaDisplayMediaOutputFactory.cpp b/Engine/Plugins/Experimental/Avalanche/Source/AvalancheMediaEditor/Private/OutputDevices/AvaDisplayMediaOutputFactory.cpp
new file mode 100644
--- /dev/null
+++ b/Engine/Plugins/Experimental/Avalanche/Source/AvalancheMediaEditor/Private/OutputDevices/AvaDisplayMediaOutputFactory.cpp
@@ -0,0 +1,29 @@
+// Copyright Epic Games, Inc. All Rights Reserved.
+
+#include "AvaDisplayMediaOutputFactory.h"
+#include "OutputDevices/AvaDisplayMediaOutput.h"
+#include "AssetTypeCategories.h"
+
+UAvaDisplayMediaOutputFactory::UAvaDisplayMediaOutputFactory(const FObjectInitializer& ObjectInitializer)
+	: Super(ObjectInitializer)
+{
+	SupportedClass = UAvaDisplayMediaOutput::StaticClass();
+	bCreateNew = true;
+	bEditAfterNew = true;
+}
+
+UObject* UAvaDisplayMediaOutputFactory::FactoryCreateNew(UClass* InClass, UObject* InParent, FName InName
+	, EObjectFlags Flags, UObject* Context, FFeedbackContext* Warn)
+{
+	return NewObject<UAvaDisplayMediaOutput>(InParent, InClass, InName, Flags);
+}
+
+uint32 UAvaDisplayMediaOutputFactory::GetMenuCategories() const
+{
+	return EAssetTypeCategories::Media;
+}
+
+bool UAvaDisplayMediaOutputFactory::ShouldShowInNewMenu() const
+{
+	return true;
+}
diff --git a/Engine/Plugins/Experimental/Avalanche/Source/AvalancheMediaEditor/Private/OutputDevices/AvaDisplayMediaOutputFactory.h b/Engine/Plugins/Experimental/Avalanche/Source/AvalancheMediaEditor/Private/OutputDevices/AvaDisplayMediaOutputFactory.h
new file mode 100644
--- /dev/null
+++ b/Engine/Plugins/Experimental/Avalanche/Source/AvalancheMediaEditor/Private/OutputDevices/AvaDisplayMediaOutputFactory.h
@@ -0,0 +1,21 @@
+// Copyright Epic Games, Inc. All Rights Reserved.
+
+#pragma once
+
+#include "CoreMinimal.h"
+#include "Factories/Factory.h"
+#include "AvaDisplayMediaOutputFactory.generated.h"
+
+/**
+ * Implements a factory for UAvaDisplayMediaOutput objects.
+ */
+UCLASS()
+class UAvaDisplayMediaOutputFactory : public UFactory
+{
+	GENERATED_UCLASS_BODY()
+
+public:
+	virtual UObject* FactoryCreateNew(UClass* InClass, UObject* InParent, FName InName, EObjectFlags Flags, UObject* Context, FFeedbackContext* Warn) override;
+	virtual uint32 GetMenuCategories() const override;
+	virtual bool ShouldShowInNewMenu() const override;
+};
diff --git a/Engine/Plugins/Experimental/Avalanche/Source/AvalancheMediaEditor/Private/OutputDevices/AvaMediaIOOutputConfigurationCustomization.cpp b/Engine/Plugins/Experimental/Avalanche/Source/AvalancheMediaEditor/Private/OutputDevices/AvaMediaIOOutputConfigurationCustomization.cpp
new file mode 100644
--- /dev/null
+++ b/Engine/Plugins/Experimental/Avalanche/Source/AvalancheMediaEditor/Private/OutputDevices/AvaMediaIOOutputConfigurationCustomization.cpp
@@ -0,0 +1,177 @@
+// Copyright Epic Games, Inc. All Rights Reserved.
+
+#include "AvaMediaIOOutputConfigurationCustomization.h"
+#include "DetailCategoryBuilder.h"
+#include "DetailLayoutBuilder.h"
+#include "DetailWidgetRow.h"
+#include "Framework/Application/SlateApplication.h"
+#include "IAvaMediaModule.h"
+#include "IDetailChildrenBuilder.h"
+#include "IMediaIOCoreDeviceProvider.h"
+#include "IMediaIOCoreModule.h"
+#include "MediaIOPermutationsSelectorBuilder.h"
+#include "ObjectEditorUtils.h"
+#include "Widgets/Input/SComboButton.h"
+#include "Widgets/Text/STextBlock.h"
+#include "Widgets/SBoxPanel.h"
+#include "Widgets/SMediaPermutationsSelector.h"
+#include "Widgets/SNullWidget.h"
+#include "Widgets/SWindow.h"
+
+#define LOCTEXT_NAMESPACE "AvaMediaIOOutputConfigurationCustomization"
+
+TSharedRef<IPropertyTypeCustomization> FAvaMediaIOOutputConfigurationCustomization::MakeInstance()
+{
+	return MakeShared<FAvaMediaIOOutputConfigurationCustomization>();
+}
+
+TAttribute<FText> FAvaMediaIOOutputConfigurationCustomization::GetContentText()
+{
+	TWeakPtr<FAvaMediaIOOutputConfigurationCustomization> WeakThisPtr = StaticCastSharedRef<FAvaMediaIOOutputConfigurationCustomization>(AsShared());
+	
+	return MakeAttributeLambda([WeakThisPtr]
+	{
+		const TSharedPtr<FAvaMediaIOOutputConfigurationCustomization> ThisSharedPtr = WeakThisPtr.Pin();
+		if (ThisSharedPtr.IsValid())
+		{
+			// Remark: we need to fetch the device provider within the lambda because it may change, i.e.
+			// if a server disconnects, it's device provider proxy will be removed.
+			const FMediaIOOutputConfiguration* Value = ThisSharedPtr->GetPropertyValueFromPropertyHandle<FMediaIOOutputConfiguration>();
+			if (const IMediaIOCoreDeviceProvider* DeviceProviderPtr = IAvaMediaModule::Get().GetDeviceProvider(ThisSharedPtr->DeviceProviderName, Value))
+			{
+				return DeviceProviderPtr->ToText(*Value);
+			}
+		}
+		return FText::GetEmpty();
+	});
+}
+
+TSharedRef<SWidget> FAvaMediaIOOutputConfigurationCustomization::HandleSourceComboButtonMenuContent()
+{
+	PermutationSelector.Reset();
+
+	const FMediaIOOutputConfiguration* Value = GetPropertyValueFromPropertyHandle<FMediaIOOutputConfiguration>();
+	const IMediaIOCoreDeviceProvider* DeviceProviderPtr = IAvaMediaModule::Get().GetDeviceProvider(DeviceProviderName, Value);
+	if (DeviceProviderPtr == nullptr)
+	{
+		return SNew(STextBlock)
+			.Text(LOCTEXT("NoDeviceProviderFound", "No provider found"));
+	}
+
+	SelectedConfiguration = *GetPropertyValueFromPropertyHandle<FMediaIOOutputConfiguration>();
+	if (!SelectedConfiguration.IsValid())
+	{
+		SelectedConfiguration = DeviceProviderPtr->GetDefaultOutputConfiguration();
+	}
+
+
+	TArray<FMediaIOOutputConfiguration> MediaConfigurations = DeviceProviderPtr->GetOutputConfigurations();
+	if (MediaConfigurations.Num() == 0)
+	{
+		return SNew(STextBlock)
+			.Text(LOCTEXT("NoConfigurationFound", "No configuration found"));
+	}
+
+	// Create the permutation selector
+	auto QuadTypeVisible = [](FName ColumnName, const TArray<FMediaIOOutputConfiguration>& UniquePermutationsForThisColumn)
+	{
+		if (UniquePermutationsForThisColumn.Num() > 0)
+		{
+			return UniquePermutationsForThisColumn[0].MediaConfiguration.MediaConnection.TransportType == EMediaIOTransportType::QuadLink;
+		}
+		return false;
+	};
+
+	auto KeyVisible = [](FName ColumnName, const TArray<FMediaIOOutputConfiguration>& UniquePermutationsForThisColumn)
+	{
+		if (UniquePermutationsForThisColumn.Num() > 0)
+		{
+			return UniquePermutationsForThisColumn[0].OutputType == EMediaIOOutputType::FillAndKey;
+		}
+		return false;
+	};
+
+	auto SyncVisible = [](FName ColumnName, const TArray<FMediaIOOutputConfiguration>& UniquePermutationsForThisColumn)
+	{
+		if (UniquePermutationsForThisColumn.Num() > 0)
+		{
+			return UniquePermutationsForThisColumn[0].OutputReference == EMediaIOReferenceType::Input;
+		}
+		return false;
+	};
+
+	using TSelection = SMediaPermutationsSelector<FMediaIOOutputConfiguration, FMediaIOPermutationsSelectorBuilder>;
+	TSelection::FArguments Arguments;
+	Arguments
+		.PermutationsSource(MoveTemp(MediaConfigurations))
+		.SelectedPermutation(SelectedConfiguration)
+		.OnSelectionChanged(this, &FAvaMediaIOOutputConfigurationCustomization::OnSelectionChanged)
+		.OnButtonClicked(this, &FAvaMediaIOOutputConfigurationCustomization::OnButtonClicked)
+		+ TSelection::Column(FMediaIOPermutationsSelectorBuilder::NAME_OutputType)
+		.Label(LOCTEXT("OutputTypeeLabel", "Output Type"))
+		+ TSelection::Column(FMediaIOPermutationsSelectorBuilder::NAME_DeviceIdentifier)
+		.Label(LOCTEXT("DeviceLabel", "Device"));
+
+	if (DeviceProviderPtr->ShowOutputTransportInSelector())
+	{
+		Arguments
+			+ TSelection::Column(FMediaIOPermutationsSelectorBuilder::NAME_TransportType)
+			.Label(LOCTEXT("DestinationTypeLabel", "Destination"))
+			+ TSelection::Column(FMediaIOPermutationsSelectorBuilder::NAME_QuadType)
+			.Label(LOCTEXT("QuadTypeLabel", "Quad"))
+			.IsColumnVisible_Lambda(QuadTypeVisible);
+	}
+
+	Arguments
+		+ TSelection::Column(FMediaIOPermutationsSelectorBuilder::NAME_Resolution)
+		.Label(LOCTEXT("ResolutionLabel", "Resolution"))
+		+ TSelection::Column(FMediaIOPermutationsSelectorBuilder::NAME_Standard)
+		.Label(LOCTEXT("StandardLabel", "Standard"))
+		+ TSelection::Column(FMediaIOPermutationsSelectorBuilder::NAME_FrameRate)
+		.Label(LOCTEXT("FrameRateLabel", "Frame Rate"));
+
+	if (DeviceProviderPtr->ShowOutputKeyInSelector())
+	{
+		Arguments
+			+ TSelection::Column(FMediaIOPermutationsSelectorBuilder::NAME_KeyPortSource)
+			.Label(LOCTEXT("KeyDestinationTypeLabel", "Key Destination"))
+			.IsColumnVisible_Lambda(KeyVisible);
+	}
+
+	if (DeviceProviderPtr->ShowReferenceInSelector())
+	{
+		Arguments
+			+ TSelection::Column(FMediaIOPermutationsSelectorBuilder::NAME_OutputReference)
+			.Label(LOCTEXT("ReferenceLabel", "Reference"))
+			+ TSelection::Column(FMediaIOPermutationsSelectorBuilder::NAME_SyncPortSource)
+			.Label(LOCTEXT("SyncSourceTypeLabel", "Sync Source"))
+			.IsColumnVisible_Lambda(SyncVisible);
+	}
+
+	TSharedRef<TSelection> Selector = SNew(TSelection) = Arguments;
+	PermutationSelector = Selector;
+	SelectedConfiguration = Selector->GetSelectedItem();
+
+	return Selector;
+}
+
+void FAvaMediaIOOutputConfigurationCustomization::OnSelectionChanged(FMediaIOOutputConfiguration SelectedItem)
+{
+	SelectedConfiguration = SelectedItem;
+}
+
+FReply FAvaMediaIOOutputConfigurationCustomization::OnButtonClicked() const
+{
+	AssignValue(SelectedConfiguration);
+
+	TSharedPtr<SWidget> SharedPermutationSelector = PermutationSelector.Pin();
+	if (SharedPermutationSelector.IsValid())
+	{
+		TSharedRef<SWindow> ParentContextMenuWindow = FSlateApplication::Get().FindWidgetWindow(SharedPermutationSelector.ToSharedRef()).ToSharedRef();
+		FSlateApplication::Get().RequestDestroyWindow(ParentContextMenuWindow);
+	}
+
+	return FReply::Handled();
+}
+
+#undef LOCTEXT_NAMESPACE
diff --git a/Engine/Plugins/Experimental/Avalanche/Source/AvalancheMediaEditor/Private/OutputDevices/AvaMediaIOOutputConfigurationCustomization.h b/Engine/Plugins/Experimental/Avalanche/Source/AvalancheMediaEditor/Private/OutputDevices/AvaMediaIOOutputConfigurationCustomization.h
new file mode 100644
--- /dev/null
+++ b/Engine/Plugins/Experimental/Avalanche/Source/AvalancheMediaEditor/Private/OutputDevices/AvaMediaIOOutputConfigurationCustomization.h
@@ -0,0 +1,26 @@
+// Copyright Epic Games, Inc. All Rights Reserved.
+
+#pragma once
+
+#include "Customizations/MediaIOCustomizationBase.h"
+#include "Input/Reply.h"
+#include "MediaIOCoreDefinitions.h"
+
+/**
+ * Implements a details view customization for the FMediaIOOutputConfiguration
+ */
+class FAvaMediaIOOutputConfigurationCustomization : public FMediaIOCustomizationBase
+{
+public:
+	static TSharedRef<IPropertyTypeCustomization> MakeInstance();
+
+private:
+	virtual TAttribute<FText> GetContentText() override;
+	virtual TSharedRef<SWidget> HandleSourceComboButtonMenuContent() override;
+
+	void OnSelectionChanged(FMediaIOOutputConfiguration SelectedItem);
+	FReply OnButtonClicked() const;
+
+	TWeakPtr<SWidget> PermutationSelector;
+	FMediaIOOutputConfiguration SelectedConfiguration;
+};
diff --git a/Engine/Plugins/Experimental/Avalanche/Source/AvalancheMediaEditor/Private/Playback/AppModes/AvaPlaybackDefaultMode.h b/Engine/Plugins/Experimental/Avalanche/Source/AvalancheMediaEditor/Private/Playback/AppModes/AvaPlaybackDefaultMode.h
new file mode 100644
--- /dev/null
+++ b/Engine/Plugins/Experimental/Avalanche/Source/AvalancheMediaEditor/Private/Playback/AppModes/AvaPlaybackDefaultMode.h
@@ -0,0 +1,11 @@
+// Copyright Epic Games, Inc. All Rights Reserved.
+
+#pragma once
+
+#include "AvaPlaybackAppMode.h"
+
+class FAvaPlaybackDefaultMode : public FAvaPlaybackAppMode
+{
+public:
+	FAvaPlaybackDefaultMode(const TSharedPtr<FAvaPlaybackEditor>& InPlaybackEditor);
+};
