[MotionDesign]

diff --git a/Engine/Plugins/Experimental/Avalanche/Source/AvalancheMediaEditor/Private/Playback/AssetDefinition_AvalanchePlayback.cpp b/Engine/Plugins/Experimental/Avalanche/Source/AvalancheMediaEditor/Private/Playback/AssetDefinition_AvalanchePlayback.cpp
new file mode 100644
--- /dev/null
+++ b/Engine/Plugins/Experimental/Avalanche/Source/AvalancheMediaEditor/Private/Playback/AssetDefinition_AvalanchePlayback.cpp
@@ -0,0 +1,163 @@
+// Copyright Epic Games, Inc. All Rights Reserved.
+
+#include "AssetDefinition_AvalanchePlayback.h"
+
+#include "AvaMediaEditorStyle.h"
+#include "ContentBrowserMenuContexts.h"
+#include "Playback/AvaPlaybackEditor.h"
+#include "Playback/AvalanchePlayback.h"
+
+#define LOCTEXT_NAMESPACE "AssetDefinition_AvalanchePlaylist"
+
+FText UAssetDefinition_AvalanchePlayback::GetAssetDisplayName() const
+{
+	return LOCTEXT("AvaPlaybackAction_Name", "Motion Design Playback");
+}
+
+TSoftClassPtr<UObject> UAssetDefinition_AvalanchePlayback::GetAssetClass() const
+{
+	return UAvalanchePlayback::StaticClass();
+}
+
+FLinearColor UAssetDefinition_AvalanchePlayback::GetAssetColor() const
+{
+	static const FName PlaybackAssetColorName(TEXT("AvalancheMediaEditor.AssetColors.Playback"));
+	return FAvaMediaEditorStyle::Get().GetColor(PlaybackAssetColorName);
+}
+
+TConstArrayView<FAssetCategoryPath> UAssetDefinition_AvalanchePlayback::GetAssetCategories() const
+{
+	static const TArray<FAssetCategoryPath> Categories = {EAssetCategoryPaths::Blueprint};
+	return Categories;
+}
+
+EAssetCommandResult UAssetDefinition_AvalanchePlayback::OpenAssets(const FAssetOpenArgs& OpenArgs) const
+{
+	EAssetCommandResult CommandResult = EAssetCommandResult::Unhandled;
+	if (OpenArgs.OpenMethod == EAssetOpenMethod::Edit)
+	{
+		for (UAvalanchePlayback* Playback : OpenArgs.LoadObjects<UAvalanchePlayback>())
+		{
+			if (Playback)
+			{
+				const TSharedRef<FAvaPlaybackEditor> PlaybackEditor = MakeShared<FAvaPlaybackEditor>();
+				PlaybackEditor->InitPlaybackEditor(EToolkitMode::Standalone, OpenArgs.ToolkitHost, Playback);
+				CommandResult = EAssetCommandResult::Handled;
+			}
+		}
+	}
+	return CommandResult;
+}
+
+namespace MenuExtension_AvalanchePlayback
+{
+	/** Count selected assets of the given type, loaded or not. */
+	template<typename ExpectedAssetType>
+	int32 GetNumSelectedAssets(const UContentBrowserAssetContextMenuContext* Context)
+	{
+		int32 NumSelectedAssets = 0;
+		for (const FAssetData& Asset : Context->SelectedAssets)
+		{
+			if (Asset.IsInstanceOf(ExpectedAssetType::StaticClass()))
+			{
+				++NumSelectedAssets;
+			}
+		}
+		return NumSelectedAssets;
+	}
+
+	/** Returns the count of playbacks that are loaded and playing. Will not load unloaded assets. */
+	static int32 GetNumSelectedPlaybackPlaying(const UContentBrowserAssetContextMenuContext* Context)
+	{
+		int32 NumPlaying = 0;
+		for (const FAssetData& Asset : Context->SelectedAssets)
+		{
+			if (Asset.IsInstanceOf(UAvalanchePlayback::StaticClass()))
+			{
+				if (const UAvalanchePlayback* Playback = Cast<UAvalanchePlayback>(Asset.FastGetAsset(false)))
+				{
+					if (Playback->IsPlaying())
+					{
+						++NumPlaying;
+					}
+				}
+			}
+		}
+		return NumPlaying;
+	}
+	
+	static bool CanExecutePlay(const FToolMenuContext& InContext)
+	{
+		const UContentBrowserAssetContextMenuContext* Context = UContentBrowserAssetContextMenuContext::FindContextWithAssets(InContext);
+		return GetNumSelectedAssets<UAvalanchePlayback>(Context) > GetNumSelectedPlaybackPlaying(Context);
+	}
+
+	static void ExecutePlay(const FToolMenuContext& InContext)
+	{
+		const UContentBrowserAssetContextMenuContext* Context = UContentBrowserAssetContextMenuContext::FindContextWithAssets(InContext);
+
+		for (UAvalanchePlayback* Playback : Context->LoadSelectedObjects<UAvalanchePlayback>())
+		{
+			if (Playback && !Playback->IsPlaying())
+			{
+				Playback->Play();
+			}
+		}
+	}
+
+	static bool CanExecuteStop(const FToolMenuContext& InContext)
+	{
+		const UContentBrowserAssetContextMenuContext* Context = UContentBrowserAssetContextMenuContext::FindContextWithAssets(InContext);
+		return GetNumSelectedPlaybackPlaying(Context) > 0;
+	}
+
+	static void ExecuteStop(const FToolMenuContext& InContext)
+	{
+		const UContentBrowserAssetContextMenuContext* Context = UContentBrowserAssetContextMenuContext::FindContextWithAssets(InContext);
+
+		for (UAvalanchePlayback* Playback : Context->LoadSelectedObjects<UAvalanchePlayback>())
+		{
+			if (Playback && Playback->IsPlaying())
+			{
+				Playback->Stop(EAvaPlaybackStopOptions::Default);
+			}
+		}
+	}
+
+	static FDelayedAutoRegisterHelper DelayedAutoRegister(EDelayedRegisterRunPhase::EndOfEngineInit, []
+	{
+		UToolMenus::RegisterStartupCallback(FSimpleMulticastDelegate::FDelegate::CreateLambda([]()
+		{
+			FToolMenuOwnerScoped OwnerScoped(UE_MODULE_NAME);
+			UToolMenu* Menu = UE::ContentBrowser::ExtendToolMenu_AssetContextMenu(UAvalanchePlayback::StaticClass());
+
+			FToolMenuSection& Section = Menu->FindOrAddSection("GetAssetActions");
+			Section.AddDynamicEntry(NAME_None, FNewToolMenuSectionDelegate::CreateLambda([](FToolMenuSection& InSection)
+			{
+				{
+					const TAttribute<FText> Label = LOCTEXT("AvalanchePlayback_Play", "Play");
+					const TAttribute<FText> ToolTip = LOCTEXT("AvalanchePlayback_PlayTooltip", "Play");
+					const FSlateIcon Icon = FSlateIcon(FAppStyle::GetAppStyleSetName(), "GenericPlay");
+
+					FToolUIAction UIAction;
+					UIAction.ExecuteAction = FToolMenuExecuteAction::CreateStatic(&ExecutePlay);
+					UIAction.CanExecuteAction = FToolMenuCanExecuteAction::CreateStatic(&CanExecutePlay);
+					InSection.AddMenuEntry("AvalanchePlayback_Play", Label, ToolTip, Icon, UIAction);
+				}
+				{
+					const TAttribute<FText> Label = LOCTEXT("AvalanchePlayback_Stop", "Stop");
+					const TAttribute<FText> ToolTip = LOCTEXT("AvalanchePlayback_Stop", "Stop");
+					const FSlateIcon Icon = FSlateIcon(FAppStyle::GetAppStyleSetName(), "GenericPause");
+
+					FToolUIAction UIAction;
+					UIAction.ExecuteAction = FToolMenuExecuteAction::CreateStatic(&ExecuteStop);
+					UIAction.CanExecuteAction = FToolMenuCanExecuteAction::CreateStatic(&CanExecuteStop);
+					InSection.AddMenuEntry("AvalanchePlayback_Stop", Label, ToolTip, Icon, UIAction);
+				}
+			}));
+		}));
+	});
+}
+
+
+#undef LOCTEXT_NAMESPACE
diff --git a/Engine/Plugins/Experimental/Avalanche/Source/AvalancheMediaEditor/Private/Playback/AssetDefinition_AvalanchePlayback.h b/Engine/Plugins/Experimental/Avalanche/Source/AvalancheMediaEditor/Private/Playback/AssetDefinition_AvalanchePlayback.h
new file mode 100644
--- /dev/null
+++ b/Engine/Plugins/Experimental/Avalanche/Source/AvalancheMediaEditor/Private/Playback/AssetDefinition_AvalanchePlayback.h
@@ -0,0 +1,21 @@
+// Copyright Epic Games, Inc. All Rights Reserved.
+
+#pragma once
+
+#include "AssetDefinitionDefault.h"
+#include "AssetDefinition_AvalanchePlayback.generated.h"
+
+UCLASS()
+class UAssetDefinition_AvalanchePlayback : public UAssetDefinitionDefault
+{
+	GENERATED_BODY()
+
+public:
+	//~ Begin UAssetDefinition
+	virtual FText GetAssetDisplayName() const override;
+	virtual TSoftClassPtr<UObject> GetAssetClass() const override;
+	virtual FLinearColor GetAssetColor() const override;
+	virtual TConstArrayView<FAssetCategoryPath> GetAssetCategories() const override;
+	virtual EAssetCommandResult OpenAssets(const FAssetOpenArgs& OpenArgs) const override;
+	//~ End UAssetDefinition
+};
diff --git a/Engine/Plugins/Experimental/Avalanche/Source/AvalancheMediaEditor/Private/Playback/AvaPlaybackEditor.cpp b/Engine/Plugins/Experimental/Avalanche/Source/AvalancheMediaEditor/Private/Playback/AvaPlaybackEditor.cpp
new file mode 100644
--- /dev/null
+++ b/Engine/Plugins/Experimental/Avalanche/Source/AvalancheMediaEditor/Private/Playback/AvaPlaybackEditor.cpp
@@ -0,0 +1,848 @@
+// Copyright Epic Games, Inc. All Rights Reserved.
+
+#include "Playback/AvaPlaybackEditor.h"
+
+#include "AppModes/AvaPlaybackDefaultMode.h"
+#include "AvaMediaEditorStyle.h"
+#include "Broadcast/AvaBroadcastEditor.h"
+#include "EdGraphUtilities.h"
+#include "Framework/Commands/GenericCommands.h"
+#include "Framework/MultiBox/MultiBoxBuilder.h"
+#include "GraphEditorActions.h"
+#include "HAL/PlatformApplicationMisc.h"
+#include "IAvaMediaEditorModule.h"
+#include "Kismet2/BlueprintEditorUtils.h"
+#include "Playback/AvaPlaybackCommands.h"
+#include "Playback/AvalanchePlayback.h"
+#include "Playback/Graph/AvaPlaybackGraph.h"
+#include "Playback/Graph/AvaPlaybackGraphSchema.h"
+#include "Playback/Graph/Nodes/AvaPlaybackGraphNode.h"
+#include "Playback/Graph/Nodes/AvaPlaybackGraphNode_Root.h"
+#include "Playback/Graph/SchemaActions/AvaPlaybackAction_NewComment.h"
+#include "Playback/Nodes/AvaPlaybackNode.h"
+#include "SNodePanel.h"
+#include "ScopedTransaction.h"
+#include "UObject/UObjectIterator.h"
+#include "WorkflowOrientedApp/ApplicationMode.h"
+
+#define LOCTEXT_NAMESPACE "AvaPlaybackEditor"
+
+TMap<TSubclassOf<UAvaPlaybackNode>, TSubclassOf<UAvaPlaybackGraphNode>> FAvaPlaybackEditor::OverrideNodeClasses;
+
+void FAvaPlaybackEditor::InitPlaybackEditor(const EToolkitMode::Type InMode
+	, const TSharedPtr<IToolkitHost>& InitToolkitHost
+	, UAvalanchePlayback* InPlayback)
+{
+	CacheOverrideNodeClasses();
+	
+	AvalanchePlayback = InPlayback;
+	AvalanchePlayback->SetGraphEditor(SharedThis(this));
+	AvalanchePlayback->CreateGraph();
+	
+	CreateDefaultCommands();
+	CreateGraphCommands();
+	
+	const FName PlaybackEditorAppName(TEXT("AvalanchePlaybackEditorApp"));
+	constexpr bool bCreateDefaultStandaloneMenu = true;
+	constexpr bool bCreateDefaultToolbar = true;
+	
+	InitAssetEditor(InMode
+		, InitToolkitHost
+		, PlaybackEditorAppName
+		, FTabManager::FLayout::NullLayout
+		, bCreateDefaultStandaloneMenu
+		, bCreateDefaultToolbar
+		, InPlayback);
+	
+	RegisterApplicationModes();
+	AvalanchePlayback->DryRunGraph();
+}
+
+FName FAvaPlaybackEditor::GetToolkitFName() const
+{
+	return TEXT("AvaPlaybackEditor");
+}
+
+FText FAvaPlaybackEditor::GetBaseToolkitName() const
+{
+	return LOCTEXT("PlaybackAppLabel", "Motion Design Playback Editor");
+}
+
+FString FAvaPlaybackEditor::GetWorldCentricTabPrefix() const
+{
+	return LOCTEXT("PlaybackScriptPrefix", "Script ").ToString();
+}
+
+FLinearColor FAvaPlaybackEditor::GetWorldCentricTabColorScale() const
+{
+	return FLinearColor(0.0f, 0.0f, 0.3f, 0.5f);
+}
+
+UAvalanchePlayback* FAvaPlaybackEditor::GetPlaybackObject() const
+{
+	return AvalanchePlayback.Get();
+}
+
+void FAvaPlaybackEditor::ExtendToolBar(TSharedPtr<FExtender> Extender)
+{
+	Extender->AddToolBarExtension("Asset"
+		, EExtensionHook::After
+		, ToolkitCommands
+		, FToolBarExtensionDelegate::CreateSP(this, &FAvaPlaybackEditor::FillPlayToolBar));
+}
+
+void FAvaPlaybackEditor::FillPlayToolBar(FToolBarBuilder& ToolBarBuilder)
+{
+	TWeakObjectPtr<UAvalanchePlayback> Playback = GetPlaybackObject();
+	
+	ToolBarBuilder.BeginSection(TEXT("Player"));
+	{
+		ToolBarBuilder.AddToolBarButton(FUIAction(FExecuteAction::CreateLambda([Playback]
+				{
+					if (Playback.IsValid())
+					{
+						Playback->Play();
+					}
+				}),
+				FCanExecuteAction::CreateLambda([Playback]
+				{
+					return Playback.IsValid() && !Playback->IsPlaying();
+				}))
+			, NAME_None
+			, LOCTEXT("Play_Label", "Play")
+			, LOCTEXT("Play_ToolTip", "Play")
+			, FSlateIcon(FAppStyle::GetAppStyleSetName(), "Icons.Toolbar.Play")
+		);
+
+		ToolBarBuilder.AddToolBarButton(FUIAction(FExecuteAction::CreateLambda([Playback]
+				{
+					if (Playback.IsValid())
+					{
+						Playback->Stop(EAvaPlaybackStopOptions::ForceImmediate | EAvaPlaybackStopOptions::Unload);
+					}
+				}),
+				FCanExecuteAction::CreateLambda([Playback]
+				{
+					return Playback.IsValid() && Playback->IsPlaying();
+				}))
+			, NAME_None
+			, LOCTEXT("Stop_Label", "Stop")
+			, LOCTEXT("Stop_ToolTip", "Stop")
+			, FSlateIcon(FAppStyle::GetAppStyleSetName(), "Icons.Toolbar.Stop")
+		);
+	}
+	ToolBarBuilder.EndSection();
+	
+	ToolBarBuilder.BeginSection(TEXT("Broadcast"));
+	{
+		ToolBarBuilder.AddToolBarButton(FExecuteAction::CreateStatic(&FAvaBroadcastEditor::OpenBroadcastEditor)
+		, NAME_None
+		, LOCTEXT("Broadcast_Label", "Broadcast")
+		, LOCTEXT("Broadcast_ToolTip", "Opens the Avalanche Broadcast Editor Window")
+		, TAttribute<FSlateIcon>::Create([]() { return IAvaMediaEditorModule::Get().GetToolbarBroadcastButtonIcon(); }));
+	}
+	ToolBarBuilder.EndSection();
+}
+
+void FAvaPlaybackEditor::CacheOverrideNodeClasses()
+{
+	if (OverrideNodeClasses.IsEmpty())
+	{
+		for (UClass* Class : TObjectRange<UClass>())
+		{
+			if (Class->IsChildOf(UAvaPlaybackGraphNode::StaticClass()))
+			{
+				const UAvaPlaybackGraphNode* const GraphNode = Cast<UAvaPlaybackGraphNode>(Class->GetDefaultObject());
+				if (TSubclassOf<UAvaPlaybackNode> PlaybackNodeClass = GraphNode->GetPlaybackNodeClass())
+				{
+					OverrideNodeClasses.Add(PlaybackNodeClass, Class);
+				}
+			}
+		}
+	}	
+}
+
+TSharedRef<SGraphEditor> FAvaPlaybackEditor::CreateGraphEditor()
+{
+	UAvalanchePlayback* const Playback = GetPlaybackObject();
+	check(Playback);
+	
+	FGraphAppearanceInfo AppearanceInfo;
+	AppearanceInfo.CornerText = LOCTEXT("AppearanceCornerText_Playback", "Motion Design PLAYBACK");
+
+	SGraphEditor::FGraphEditorEvents InEvents;
+	InEvents.OnSelectionChanged  = SGraphEditor::FOnSelectionChanged::CreateSP(this, &FAvaPlaybackEditor::OnSelectedNodesChanged);
+	InEvents.OnTextCommitted     = FOnNodeTextCommitted::CreateSP(this, &FAvaPlaybackEditor::OnNodeTitleCommitted);
+
+	GraphEditor = SNew(SGraphEditor)
+		.AdditionalCommands(GraphEditorCommands)
+		.IsEditable(true)
+		.Appearance(AppearanceInfo)
+		.GraphToEdit(Playback->GetGraph())
+		.GraphEvents(InEvents)
+		.AutoExpandActionMenu(true)
+		.ShowGraphStateOverlay(false);
+	
+	return GraphEditor.ToSharedRef();
+}
+
+UEdGraph* FAvaPlaybackEditor::CreatePlaybackGraph(UAvalanchePlayback* InPlayback)
+{
+	UAvaPlaybackGraph* PlaybackGraph = CastChecked<UAvaPlaybackGraph>(FBlueprintEditorUtils::CreateNewGraph(InPlayback
+		, NAME_None
+		, UAvaPlaybackGraph::StaticClass()
+		, UAvaPlaybackGraphSchema::StaticClass()));
+	
+	return PlaybackGraph;
+}
+
+void FAvaPlaybackEditor::SetupPlaybackNode(UEdGraph* InGraph
+	, UAvaPlaybackNode* InPlaybackNode
+	, bool bSelectNewNode)
+{
+	TSubclassOf<UAvaPlaybackGraphNode> NodeClass = UAvaPlaybackGraphNode::StaticClass();
+	
+	//Find Most Relevant Override Node Class
+	{
+		UClass* StartingClass = InPlaybackNode->GetClass();
+		while (StartingClass && StartingClass->IsChildOf(UAvaPlaybackNode::StaticClass()))
+		{
+			if (TSubclassOf<UAvaPlaybackGraphNode>* const OverrideClass = OverrideNodeClasses.Find(StartingClass))
+			{
+				NodeClass = *OverrideClass;
+				break;
+			}
+			StartingClass = StartingClass->GetSuperClass();
+		}
+
+		if (!NodeClass)
+		{
+			NodeClass = UAvaPlaybackGraphNode::StaticClass();
+		}
+	}
+	
+	UAvaPlaybackGraph* const PlaybackGraph = Cast<UAvaPlaybackGraph>(InGraph);
+	check(PlaybackGraph  && NodeClass);
+	
+	UAvaPlaybackGraphNode* const Node = PlaybackGraph->CreatePlaybackGraphNode(NodeClass, bSelectNewNode);
+	check(Node);
+	
+	Node->SetPlaybackNode(InPlaybackNode);
+	Node->CreateNewGuid();
+	Node->PostPlacedNewNode();
+	
+	if (Node->Pins.Num() == 0)
+	{
+		Node->AllocateDefaultPins();
+	}
+}
+
+void FAvaPlaybackEditor::CompilePlaybackNodesFromGraphNodes(UAvalanchePlayback* InPlayback)
+{
+	UEdGraph* const PlaybackGraph = InPlayback->GetGraph();
+	check(PlaybackGraph);
+
+	TArray<UAvaPlaybackNode*> PlaybackNodes;
+	PlaybackNodes.Reserve(PlaybackGraph->Nodes.Num());
+	
+	for (UEdGraphNode* const EdGraphNode : PlaybackGraph->Nodes)
+	{
+		UAvaPlaybackGraphNode* const GraphNode = Cast<UAvaPlaybackGraphNode>(EdGraphNode);
+		if (GraphNode && GraphNode->GetPlaybackNode())
+		{
+			UAvaPlaybackNode* const PlaybackNode = GraphNode->GetPlaybackNode();
+			
+			// Set ChildNodes of each PlaybackNode
+			TArray<UEdGraphPin*> InputPins = GraphNode->GetInputPins();
+			
+			TArray<UAvaPlaybackNode*> ChildNodes;
+			ChildNodes.Reserve(InputPins.Num());
+			
+			for (UEdGraphPin* const InputPin : InputPins)
+			{
+				if (InputPin->LinkedTo.Num() > 0)
+				{
+					//Note: A Playback Node Input Pin can't be connected to multiple sources.
+					UAvaPlaybackGraphNode* GraphChildNode = CastChecked<UAvaPlaybackGraphNode>(InputPin->LinkedTo[0]->GetOwningNode());
+					ChildNodes.Add(GraphChildNode->GetPlaybackNode());
+				}
+				else
+				{
+					ChildNodes.AddZeroed();
+				}
+			}
+
+			PlaybackNode->SetFlags(RF_Transactional);
+			PlaybackNode->Modify();
+			PlaybackNode->SetChildNodes(MoveTemp(ChildNodes));
+			
+			PlaybackNodes.Emplace(PlaybackNode);
+		}
+	}	
+	
+	InPlayback->DryRunGraph();
+	
+	for (UAvaPlaybackNode* const PlaybackNode : PlaybackNodes)
+	{
+		PlaybackNode->PostEditChange();
+	}
+}
+
+void FAvaPlaybackEditor::CreateInputPin(UEdGraphNode* InGraphNode)
+{
+	CastChecked<UAvaPlaybackGraphNode>(InGraphNode)->CreateInputPin();
+}
+
+void FAvaPlaybackEditor::RefreshNode(UEdGraphNode& InNode)
+{
+	if (GraphEditor.IsValid())
+	{
+		GraphEditor->RefreshNode(InNode);
+	}
+}
+
+bool FAvaPlaybackEditor::GetBoundsForSelectedNodes(FSlateRect& Rect, float Padding)
+{
+	if (GraphEditor.IsValid())
+	{
+		return GraphEditor->GetBoundsForSelectedNodes(Rect, Padding);
+	}
+	return false;
+}
+
+TSet<UObject*> FAvaPlaybackEditor::GetSelectedNodes() const
+{
+	if (GraphEditor.IsValid())
+	{
+		return GraphEditor->GetSelectedNodes();
+	}
+	return TSet<UObject*>();
+}
+
+void FAvaPlaybackEditor::OnSelectedNodesChanged(const TSet<UObject*>& NewSelection)
+{
+	TArray<UObject*> Selection;
+	if (NewSelection.Num())
+	{
+		for (TSet<UObject*>::TConstIterator Iter(NewSelection); Iter; ++Iter)
+		{
+			if (UAvaPlaybackGraphNode* const GraphNode = Cast<UAvaPlaybackGraphNode>(*Iter))
+			{
+				Selection.Add(GraphNode->GetPlaybackNode());
+			}
+			else
+			{
+				Selection.Add(*Iter);
+			}
+		}
+	}
+	OnPlaybackSelectionChanged.Broadcast(Selection);
+}
+
+void FAvaPlaybackEditor::OnNodeTitleCommitted(const FText& NewText
+	, ETextCommit::Type CommitInfo
+	, UEdGraphNode* NodeBeingChanged)
+{
+	if (NodeBeingChanged)
+	{
+		const FScopedTransaction Transaction(LOCTEXT("RenameNode", "Rename Node"));
+		NodeBeingChanged->Modify();
+		NodeBeingChanged->OnRenameNode(NewText.ToString());
+	}
+}
+
+void FAvaPlaybackEditor::RegisterApplicationModes()
+{
+	TArray<TSharedRef<FApplicationMode>> ApplicationModes;
+	TSharedPtr<FAvaPlaybackEditor> This = SharedThis(this);
+	
+	ApplicationModes.Add(MakeShared<FAvaPlaybackDefaultMode>(This));
+	//Can add more App Modes here
+	
+	for (const TSharedRef<FApplicationMode>& AppMode : ApplicationModes)
+	{
+		AddApplicationMode(AppMode->GetModeName(), AppMode);
+	}
+
+	SetCurrentMode(FAvaPlaybackAppMode::DefaultMode);
+}
+
+void FAvaPlaybackEditor::CreateDefaultCommands()
+{
+}
+
+void FAvaPlaybackEditor::CreateGraphCommands()
+{
+	GraphEditorCommands = MakeShared<FUICommandList>();
+	{
+		const FGraphEditorCommandsImpl& GraphCommands = FGraphEditorCommands::Get();
+		const FGenericCommands& GenericCommands = FGenericCommands::Get();
+		const FAvaPlaybackCommands& PlaybackCommands = FAvaPlaybackCommands::Get();
+			
+		// Playback Commands
+		GraphEditorCommands->MapAction(PlaybackCommands.AddInputPin,
+			FExecuteAction::CreateSP(this, &FAvaPlaybackEditor::AddInputPin),
+			FCanExecuteAction::CreateSP(this, &FAvaPlaybackEditor::CanAddInputPin));
+
+		GraphEditorCommands->MapAction(PlaybackCommands.RemoveInputPin,
+			FExecuteAction::CreateSP(this, &FAvaPlaybackEditor::RemoveInputPin),
+			FCanExecuteAction::CreateSP(this, &FAvaPlaybackEditor::CanRemoveInputPin));
+		
+		// Graph Editor Commands
+		GraphEditorCommands->MapAction(GraphCommands.CreateComment
+			, FExecuteAction::CreateSP(this, &FAvaPlaybackEditor::CreateComment));
+
+		// Editing commands
+		GraphEditorCommands->MapAction(GenericCommands.SelectAll
+			, FExecuteAction::CreateSP(this, &FAvaPlaybackEditor::SelectAllNodes)
+			, FCanExecuteAction::CreateSP(this, &FAvaPlaybackEditor::CanSelectAllNodes)
+		);
+
+		GraphEditorCommands->MapAction(GenericCommands.Delete
+			, FExecuteAction::CreateSP(this, &FAvaPlaybackEditor::DeleteSelectedNodes)
+			, FCanExecuteAction::CreateSP(this, &FAvaPlaybackEditor::CanDeleteSelectedNodes)
+		);
+		
+		GraphEditorCommands->MapAction(GenericCommands.Copy
+			, FExecuteAction::CreateSP(this, &FAvaPlaybackEditor::CopySelectedNodes)
+			, FCanExecuteAction::CreateSP(this, &FAvaPlaybackEditor::CanCopySelectedNodes)
+		);
+
+		GraphEditorCommands->MapAction(GenericCommands.Cut
+			, FExecuteAction::CreateSP(this, &FAvaPlaybackEditor::CutSelectedNodes)
+			, FCanExecuteAction::CreateSP(this, &FAvaPlaybackEditor::CanCutSelectedNodes)
+		);
+
+		GraphEditorCommands->MapAction(GenericCommands.Paste
+			, FExecuteAction::CreateSP(this, &FAvaPlaybackEditor::PasteNodes)
+			, FCanExecuteAction::CreateSP(this, &FAvaPlaybackEditor::CanPasteNodes)
+		);
+
+		GraphEditorCommands->MapAction(GenericCommands.Duplicate
+			, FExecuteAction::CreateSP(this, &FAvaPlaybackEditor::DuplicateSelectedNodes)
+			, FCanExecuteAction::CreateSP(this, &FAvaPlaybackEditor::CanDuplicateSelectedNodes)
+		);
+
+		// Alignment Commands
+		GraphEditorCommands->MapAction(GraphCommands.AlignNodesTop
+			, FExecuteAction::CreateSP(this, &FAvaPlaybackEditor::OnAlignTop)
+		);
+
+		GraphEditorCommands->MapAction(GraphCommands.AlignNodesMiddle
+			, FExecuteAction::CreateSP(this, &FAvaPlaybackEditor::OnAlignMiddle)
+		);
+
+		GraphEditorCommands->MapAction(GraphCommands.AlignNodesBottom
+			, FExecuteAction::CreateSP(this, &FAvaPlaybackEditor::OnAlignBottom)
+		);
+
+		GraphEditorCommands->MapAction(GraphCommands.AlignNodesLeft
+			, FExecuteAction::CreateSP(this, &FAvaPlaybackEditor::OnAlignLeft)
+		);
+
+		GraphEditorCommands->MapAction(GraphCommands.AlignNodesCenter
+			, FExecuteAction::CreateSP(this, &FAvaPlaybackEditor::OnAlignCenter)
+		);
+
+		GraphEditorCommands->MapAction(GraphCommands.AlignNodesRight
+			, FExecuteAction::CreateSP(this, &FAvaPlaybackEditor::OnAlignRight)
+		);
+
+		GraphEditorCommands->MapAction(GraphCommands.StraightenConnections
+			, FExecuteAction::CreateSP(this, &FAvaPlaybackEditor::OnStraightenConnections)
+		);
+
+		// Distribution Commands
+		GraphEditorCommands->MapAction(GraphCommands.DistributeNodesHorizontally
+			, FExecuteAction::CreateSP(this, &FAvaPlaybackEditor::OnDistributeNodesH)
+		);
+
+		GraphEditorCommands->MapAction(GraphCommands.DistributeNodesVertically
+			, FExecuteAction::CreateSP(this, &FAvaPlaybackEditor::OnDistributeNodesV)
+		);
+	}
+}
+
+bool FAvaPlaybackEditor::CanAddInputPin() const
+{
+	return GetSelectedNodes().Num() == 1;
+}
+
+void FAvaPlaybackEditor::AddInputPin()
+{
+	const TSet<UObject*> SelectedNodes = GetSelectedNodes();
+	// Iterator used but should only contain one node
+	for (TSet<UObject*>::TConstIterator It(SelectedNodes); It; ++It)
+	{
+		if (UAvaPlaybackGraphNode* SelectedNode = Cast<UAvaPlaybackGraphNode>(*It))
+		{
+			SelectedNode->AddInputPin();
+			break;
+		}
+	}
+}
+
+bool FAvaPlaybackEditor::CanRemoveInputPin() const
+{
+	return true;
+}
+
+void FAvaPlaybackEditor::RemoveInputPin()
+{
+	if (GraphEditor.IsValid())
+	{
+		if (UEdGraphPin* const SelectedPin = GraphEditor->GetGraphPinForMenu())
+		{
+			UAvaPlaybackGraphNode* const SelectedNode = Cast<UAvaPlaybackGraphNode>(SelectedPin->GetOwningNode());
+			if (SelectedNode && SelectedNode == SelectedPin->GetOwningNode())
+			{
+				SelectedNode->RemoveInputPin(SelectedPin);
+			}
+		}
+	}
+}
+
+void FAvaPlaybackEditor::CreateComment()
+{
+	if (GraphEditor.IsValid())
+	{
+		FAvaPlaybackAction_NewComment CommentAction;
+		CommentAction.PerformAction(GraphEditor->GetCurrentGraph(), nullptr, GraphEditor->GetPasteLocation());
+	}
+}
+
+void FAvaPlaybackEditor::SelectAllNodes()
+{
+	if (GraphEditor.IsValid())
+	{
+		GraphEditor->SelectAllNodes();
+	}
+}
+
+bool FAvaPlaybackEditor::CanDeleteSelectedNodes() const
+{
+	const TSet<UObject*> SelectedNodes = GetSelectedNodes();
+	if (SelectedNodes.Num() == 1)
+	{
+		for (TSet<UObject*>::TConstIterator It(SelectedNodes); It; ++It)
+		{
+			if (Cast<UAvaPlaybackGraphNode_Root>(*It))
+			{
+				// Return false if only root node is selected, as it can't be deleted
+				return false;
+			}
+		}
+	}
+	return SelectedNodes.Num() > 0;
+}
+
+void FAvaPlaybackEditor::DeleteSelectedNodes()
+{
+	if (!GraphEditor.IsValid() || !GraphEditor->GetCurrentGraph())
+	{
+		return;
+	}
+	
+	const FScopedTransaction Transaction(LOCTEXT("DeleteSelectedNodes", "Delete Selected Nodes"));
+	GraphEditor->GetCurrentGraph()->Modify();
+
+	const TSet<UObject*> SelectedNodes = GetSelectedNodes();
+	GraphEditor->ClearSelectionSet();
+
+	check(AvalanchePlayback.IsValid());
+	for (TSet<UObject*>::TConstIterator It(SelectedNodes); It; ++It)
+	{
+		UEdGraphNode* const Node = CastChecked<UEdGraphNode>(*It);
+		if (Node->CanUserDeleteNode())
+		{
+			if (UAvaPlaybackGraphNode* const PlaybackGraphNode = Cast<UAvaPlaybackGraphNode>(Node))
+			{
+				UAvaPlaybackNode* const NodeToDelete = PlaybackGraphNode->GetPlaybackNode();
+				FBlueprintEditorUtils::RemoveNode(nullptr, PlaybackGraphNode, true);
+
+				// Make sure Playback is updated to match graph
+				AvalanchePlayback->CompilePlaybackNodesFromGraphNodes();
+
+				// Remove this node from the AvalanchePlayback's list of all Playback Nodes
+				AvalanchePlayback->RemovePlaybackNode(NodeToDelete);
+				AvalanchePlayback->MarkPackageDirty();
+			}
+			else
+			{
+				FBlueprintEditorUtils::RemoveNode(nullptr, Node, true);
+			}
+		}
+	}
+}
+
+void FAvaPlaybackEditor::DeleteSelectedDuplicatableNodes()
+{
+	if (!GraphEditor.IsValid())
+	{
+		return;
+	}
+	
+	// Cache off the old selection
+	const TSet<UObject*> OldSelectedNodes = GetSelectedNodes();
+
+	// Clear the selection and only select the nodes that can be duplicated
+	TSet<UObject*> RemainingNodes;
+	GraphEditor->ClearSelectionSet();
+
+	for (TSet<UObject*>::TConstIterator Iter(OldSelectedNodes); Iter; ++Iter)
+	{
+		UEdGraphNode* const Node = Cast<UEdGraphNode>(*Iter);
+		if (Node && Node->CanDuplicateNode())
+		{
+			GraphEditor->SetNodeSelection(Node, true);
+		}
+		else
+		{
+			RemainingNodes.Add(Node);
+		}
+	}
+
+	// Delete the duplicatable nodes
+	DeleteSelectedNodes();
+
+	// Reselect whatever's left from the original selection after the deletion
+	GraphEditor->ClearSelectionSet();
+
+	for (TSet<UObject*>::TConstIterator Iter(RemainingNodes); Iter; ++Iter)
+	{
+		if (UEdGraphNode* const Node = Cast<UEdGraphNode>(*Iter))
+		{
+			GraphEditor->SetNodeSelection(Node, true);
+		}
+	}
+}
+
+bool FAvaPlaybackEditor::CanCopySelectedNodes() const
+{
+	const TSet<UObject*> SelectedNodes = GetSelectedNodes();
+	
+	for (TSet<UObject*>::TConstIterator Iter(SelectedNodes); Iter; ++Iter)
+	{
+		UEdGraphNode* const Node = Cast<UEdGraphNode>(*Iter);
+		
+		//If at least one node can be duplicated, we should allow copy
+		if (Node && Node->CanDuplicateNode())
+		{
+			return true;
+		}
+	}
+	
+	return false;
+}
+
+void FAvaPlaybackEditor::CopySelectedNodes()
+{
+	// Export the selected nodes and place the text on the clipboard
+	const TSet<UObject*> SelectedNodes = GetSelectedNodes();
+	
+	for (TSet<UObject*>::TConstIterator Iter(SelectedNodes); Iter; ++Iter)
+	{
+		if (UAvaPlaybackGraphNode* const Node = Cast<UAvaPlaybackGraphNode>(*Iter))
+		{
+			Node->PrepareForCopying();
+		}
+	}
+	
+	FString ExportedText;
+	FEdGraphUtilities::ExportNodesToText(SelectedNodes, ExportedText);
+	FPlatformApplicationMisc::ClipboardCopy(*ExportedText);
+	
+	for (TSet<UObject*>::TConstIterator Iter(SelectedNodes); Iter; ++Iter)
+	{
+		if (UAvaPlaybackGraphNode* const Node = Cast<UAvaPlaybackGraphNode>(*Iter))
+		{
+			Node->PostCopyNode();
+		}
+	}
+}
+
+bool FAvaPlaybackEditor::CanCutSelectedNodes() const
+{
+	return CanCopySelectedNodes() && CanDeleteSelectedNodes();
+}
+
+void FAvaPlaybackEditor::CutSelectedNodes()
+{
+	CopySelectedNodes();
+	// Cut should only delete nodes that can be duplicated
+	DeleteSelectedDuplicatableNodes();
+}
+
+bool FAvaPlaybackEditor::CanPasteNodes() const
+{
+	if (AvalanchePlayback.IsValid() && AvalanchePlayback->GetGraph())
+	{
+		FString ClipboardContent;
+		FPlatformApplicationMisc::ClipboardPaste(ClipboardContent);
+		return FEdGraphUtilities::CanImportNodesFromText(AvalanchePlayback->GetGraph(), ClipboardContent);
+	}
+	return false;
+}
+
+void FAvaPlaybackEditor::PasteNodes()
+{
+	if (GraphEditor.IsValid())
+	{
+		PasteNodesHere(GraphEditor->GetPasteLocation());
+	}
+}
+
+void FAvaPlaybackEditor::PasteNodesHere(const FVector2D& Location)
+{
+	if (!GraphEditor.IsValid() || !AvalanchePlayback.IsValid() || !AvalanchePlayback->GetGraph())
+	{
+		return;
+	}
+	
+	// Undo/Redo support
+	const FScopedTransaction Transaction(LOCTEXT("PasteNodes", "Paste Playback Nodes"));
+	AvalanchePlayback->GetGraph()->Modify();
+	AvalanchePlayback->Modify();
+
+	// Clear the selection set (newly pasted stuff will be selected)
+	GraphEditor->ClearSelectionSet();
+
+	// Grab the text to paste from the clipboard.
+	FString TextToImport;
+	FPlatformApplicationMisc::ClipboardPaste(TextToImport);
+
+	// Import the nodes
+	TSet<UEdGraphNode*> PastedNodes;
+	FEdGraphUtilities::ImportNodesFromText(AvalanchePlayback->GetGraph(), TextToImport, /*out*/ PastedNodes);
+
+	//Average position of nodes so we can move them while still maintaining relative distances to each other
+	FVector2D AverageNodePosition(0.0f,0.0f);
+
+	for (TSet<UEdGraphNode*>::TIterator It(PastedNodes); It; ++It)
+	{
+		UEdGraphNode* Node = *It;
+		AverageNodePosition.X += Node->NodePosX;
+		AverageNodePosition.Y += Node->NodePosY;
+	}
+
+	if (PastedNodes.Num() > 0)
+	{
+		FVector2D::FReal InverseNumNodes = 1.0 / static_cast<FVector2D::FReal>(PastedNodes.Num());
+		AverageNodePosition.X *= InverseNumNodes;
+		AverageNodePosition.Y *= InverseNumNodes;
+	}
+
+	for (TSet<UEdGraphNode*>::TIterator It(PastedNodes); It; ++It)
+	{
+		UEdGraphNode* const Node = *It;
+		if (UAvaPlaybackGraphNode* PlaybackGraphNode = Cast<UAvaPlaybackGraphNode>(Node))
+		{
+			AvalanchePlayback->AddPlaybackNode(PlaybackGraphNode->GetPlaybackNode());
+		}
+
+		// Select the newly pasted stuff
+		GraphEditor->SetNodeSelection(Node, true);
+
+		Node->NodePosX = (Node->NodePosX - AverageNodePosition.X) + Location.X ;
+		Node->NodePosY = (Node->NodePosY - AverageNodePosition.Y) + Location.Y ;
+
+		Node->SnapToGrid(SNodePanel::GetSnapGridSize());
+
+		// Give new node a different Guid from the old one
+		Node->CreateNewGuid();
+	}
+
+	// Force new pasted nodes to have same connections as graph nodes
+	AvalanchePlayback->CompilePlaybackNodesFromGraphNodes();
+
+	// Update UI
+	GraphEditor->NotifyGraphChanged();
+
+	AvalanchePlayback->PostEditChange();
+	AvalanchePlayback->MarkPackageDirty();
+}
+
+bool FAvaPlaybackEditor::CanDuplicateSelectedNodes() const
+{
+	return CanCopySelectedNodes();
+}
+
+void FAvaPlaybackEditor::DuplicateSelectedNodes()
+{
+	CopySelectedNodes();
+	PasteNodes();
+}
+
+void FAvaPlaybackEditor::OnAlignTop()
+{
+	if (GraphEditor.IsValid())
+	{
+		GraphEditor->OnAlignTop();
+	}
+}
+
+void FAvaPlaybackEditor::OnAlignMiddle()
+{
+	if (GraphEditor.IsValid())
+	{
+		GraphEditor->OnAlignMiddle();
+	}
+}
+
+void FAvaPlaybackEditor::OnAlignBottom()
+{
+	if (GraphEditor.IsValid())
+	{
+		GraphEditor->OnAlignBottom();
+	}
+}
+
+void FAvaPlaybackEditor::OnAlignLeft()
+{
+	if (GraphEditor.IsValid())
+	{
+		GraphEditor->OnAlignLeft();
+	}
+}
+
+void FAvaPlaybackEditor::OnAlignCenter()
+{
+	if (GraphEditor.IsValid())
+	{
+		GraphEditor->OnAlignCenter();
+	}
+}
+
+void FAvaPlaybackEditor::OnAlignRight()
+{
+	if (GraphEditor.IsValid())
+	{
+		GraphEditor->OnAlignRight();
+	}
+}
+
+void FAvaPlaybackEditor::OnStraightenConnections()
+{
+	if (GraphEditor.IsValid())
+	{
+		GraphEditor->OnStraightenConnections();
+	}
+}
+
+void FAvaPlaybackEditor::OnDistributeNodesH()
+{
+	if (GraphEditor.IsValid())
+	{
+		GraphEditor->OnDistributeNodesH();
+	}
+}
+
+void FAvaPlaybackEditor::OnDistributeNodesV()
+{
+	if (GraphEditor.IsValid())
+	{
+		GraphEditor->OnDistributeNodesV();
+	}
+}
+
+#undef LOCTEXT_NAMESPACE
diff --git a/Engine/Plugins/Experimental/Avalanche/Source/AvalancheMediaEditor/Private/Playback/AvaPlaybackFactory.cpp b/Engine/Plugins/Experimental/Avalanche/Source/AvalancheMediaEditor/Private/Playback/AvaPlaybackFactory.cpp
new file mode 100644
--- /dev/null
+++ b/Engine/Plugins/Experimental/Avalanche/Source/AvalancheMediaEditor/Private/Playback/AvaPlaybackFactory.cpp
@@ -0,0 +1,66 @@
+// Copyright Epic Games, Inc. All Rights Reserved.
+
+#include "Playback/AvaPlaybackFactory.h"
+#include "AssetToolsModule.h"
+#include "ClassViewerFilter.h"
+#include "IAssetTools.h"
+#include "Playback/AvalanchePlayback.h"
+
+#define LOCTEXT_NAMESPACE "AvaPlaybackFactory"
+
+UAvaPlaybackFactory::UAvaPlaybackFactory()
+{
+	// Provide the factory with information about how to handle our asset
+	bCreateNew = true;
+	bEditAfterNew = true;
+	SupportedClass = UAvalanchePlayback::StaticClass();
+}
+
+UAvaPlaybackFactory::~UAvaPlaybackFactory()
+{
+}
+
+uint32 UAvaPlaybackFactory::GetMenuCategories() const
+{
+	IAssetTools& AssetTools = FAssetToolsModule::GetModule().Get();
+	return AssetTools.FindAdvancedAssetCategory("AvalancheCategory");
+}
+
+bool UAvaPlaybackFactory::ConfigureProperties()
+{
+	return Super::ConfigureProperties();
+}
+
+bool UAvaPlaybackFactory::ShouldShowInNewMenu() const
+{
+	return Super::ShouldShowInNewMenu();
+}
+
+UObject* UAvaPlaybackFactory::FactoryCreateNew(UClass* Class, UObject* InParent, FName Name, EObjectFlags Flags
+  , UObject* Context, FFeedbackContext* Warn, FName CallingContext)
+{
+	return Super::FactoryCreateNew(Class, InParent, Name, Flags, Context, Warn, CallingContext);
+}
+
+UObject* UAvaPlaybackFactory::FactoryCreateNew(UClass* Class, UObject* InParent, FName Name, EObjectFlags Flags
+  , UObject* Context, FFeedbackContext* Warn)
+{
+	UAvalanchePlayback* Playback = nullptr;
+	if (ensure(SupportedClass == Class))
+	{
+		Playback = NewObject<UAvalanchePlayback>(InParent, Name, Flags);
+	}
+	return Playback;
+}
+
+bool UAvaPlaybackFactory::DoesSupportClass(UClass* Class)
+{
+	return Class == UAvalanchePlayback::StaticClass();
+}
+
+UClass* UAvaPlaybackFactory::ResolveSupportedClass()
+{
+	return Super::ResolveSupportedClass();
+}
+
+#undef LOCTEXT_NAMESPACE
diff --git a/Engine/Plugins/Experimental/Avalanche/Source/AvalancheMediaEditor/Private/Playback/DetailsView/SAvaPlaybackDetailsView.h b/Engine/Plugins/Experimental/Avalanche/Source/AvalancheMediaEditor/Private/Playback/DetailsView/SAvaPlaybackDetailsView.h
new file mode 100644
--- /dev/null
+++ b/Engine/Plugins/Experimental/Avalanche/Source/AvalancheMediaEditor/Private/Playback/DetailsView/SAvaPlaybackDetailsView.h
@@ -0,0 +1,28 @@
+// Copyright Epic Games, Inc. All Rights Reserved.
+
+#pragma once
+
+#include "CoreMinimal.h"
+#include "Widgets/SCompoundWidget.h"
+
+class FAvaPlaybackEditor;
+class IDetailsView;
+
+class SAvaPlaybackDetailsView : public SCompoundWidget
+{
+public:
+	
+	SLATE_BEGIN_ARGS(SAvaPlaybackDetailsView){}
+	SLATE_END_ARGS()
+	
+	void Construct(const FArguments& InArgs, const TSharedPtr<FAvaPlaybackEditor>& InPlaybackEditor);
+	virtual ~SAvaPlaybackDetailsView() override;
+	
+	void OnPlaybackNodeSelectionChanged(const TArray<UObject*>& InSelectedObjects);
+
+protected:
+
+	TWeakPtr<FAvaPlaybackEditor> PlaybackEditorWeak;
+	
+	TSharedPtr<IDetailsView> DetailsView;
+};
diff --git a/Engine/Plugins/Experimental/Avalanche/Source/AvalancheMediaEditor/Private/Playback/Graph/AvaPlaybackGraph.cpp b/Engine/Plugins/Experimental/Avalanche/Source/AvalancheMediaEditor/Private/Playback/Graph/AvaPlaybackGraph.cpp
new file mode 100644
--- /dev/null
+++ b/Engine/Plugins/Experimental/Avalanche/Source/AvalancheMediaEditor/Private/Playback/Graph/AvaPlaybackGraph.cpp
@@ -0,0 +1,15 @@
+// Copyright Epic Games, Inc. All Rights Reserved.
+
+#include "Playback/Graph/AvaPlaybackGraph.h"
+#include "Playback/AvalanchePlayback.h"
+#include "Playback/Graph/Nodes/AvaPlaybackGraphNode.h"
+
+UAvalanchePlayback* UAvaPlaybackGraph::GetPlayback() const
+{
+	return CastChecked<UAvalanchePlayback>(GetOuter());
+}
+
+UAvaPlaybackGraphNode* UAvaPlaybackGraph::CreatePlaybackGraphNode(TSubclassOf<UAvaPlaybackGraphNode> NewNodeClass, bool bSelectNewNode)
+{
+	return CastChecked<UAvaPlaybackGraphNode>(CreateNode(NewNodeClass, false, bSelectNewNode));
+}
diff --git a/Engine/Plugins/Experimental/Avalanche/Source/AvalancheMediaEditor/Private/Playback/Graph/AvaPlaybackGraphSchema.cpp b/Engine/Plugins/Experimental/Avalanche/Source/AvalancheMediaEditor/Private/Playback/Graph/AvaPlaybackGraphSchema.cpp
new file mode 100644
--- /dev/null
+++ b/Engine/Plugins/Experimental/Avalanche/Source/AvalancheMediaEditor/Private/Playback/Graph/AvaPlaybackGraphSchema.cpp
@@ -0,0 +1,300 @@
+// Copyright Epic Games, Inc. All Rights Reserved.
+
+#include "Playback/Graph/AvaPlaybackGraphSchema.h"
+#include "AvaBlueprint.h"
+#include "GraphEditorActions.h"
+#include "Playback/AvalanchePlayback.h"
+#include "Playback/Graph/AvaPlaybackGraph.h"
+#include "Playback/Graph/Nodes/AvaPlaybackGraphNode_Root.h"
+#include "Playback/Graph/SchemaActions/AvaPlaybackAction_NewComment.h"
+#include "Playback/Graph/SchemaActions/AvaPlaybackAction_NewNode.h"
+#include "Playback/Graph/SchemaActions/AvaPlaybackAction_PasteNode.h"
+#include "Playback/IAvaPlaybackEditor.h"
+#include "Playback/Nodes/AvaPlaybackNodeRoot.h"
+#include "ScopedTransaction.h"
+#include "ToolMenu.h"
+#include "UObject/UObjectIterator.h"
+
+#define LOCTEXT_NAMESPACE "AvaPlaybackGraphSchema"
+
+const FLinearColor UAvaPlaybackGraphSchema::ActivePinColor = FLinearColor::White;
+const FLinearColor UAvaPlaybackGraphSchema::InactivePinColor = FLinearColor(0.05f, 0.05f, 0.05f);
+
+// Allowable PinType.PinCategory values
+const FName UAvaPlaybackGraphSchema::PC_ChannelFeed(TEXT("channelfeed"));
+const FName UAvaPlaybackGraphSchema::PC_Event(TEXT("event"));
+
+TArray<TSubclassOf<UAvaPlaybackNode>> UAvaPlaybackGraphSchema::PlaybackNodeClasses;
+
+TSharedPtr<IAvaPlaybackGraphEditor> UAvaPlaybackGraphSchema::GetPlaybackGraphEditor(const UEdGraph* Graph) const
+{
+	if (Graph && Cast<UAvaPlaybackGraph>(Graph))
+	{
+		if (UAvalanchePlayback* const Playback = Cast<UAvaPlaybackGraph>(Graph)->GetPlayback())
+		{
+			return Playback->GetGraphEditor();
+		}
+	}
+	return nullptr;
+}
+
+void UAvaPlaybackGraphSchema::CompilePlaybackNodesFromGraphNodes(UEdGraphNode& Node) const
+{
+	CastChecked<UAvaPlaybackGraph>(Node.GetGraph())->GetPlayback()->CompilePlaybackNodesFromGraphNodes();
+}
+
+bool UAvaPlaybackGraphSchema::ConnectionCausesLoop(const UEdGraphPin* InputPin, const UEdGraphPin* OutputPin) const
+{
+	UAvaPlaybackGraphNode* const InputNode = Cast<UAvaPlaybackGraphNode>(InputPin->GetOwningNode());
+
+	if (InputNode)
+	{
+		// Only nodes representing SoundNodes have outputs
+		UAvaPlaybackGraphNode* const OutputNode = CastChecked<UAvaPlaybackGraphNode>(OutputPin->GetOwningNode());
+
+		if (OutputNode->GetPlaybackNode())
+		{
+			// Grab all child nodes. We can't just test the output because 
+			// the loop could happen from any additional child nodes. 
+			TArray<UAvaPlaybackNode*> Nodes;
+			OutputNode->GetPlaybackNode()->GetAllNodes(Nodes);
+
+			// If our test input is in that set, return true.
+			return Nodes.Contains(InputNode->GetPlaybackNode());
+		}
+	}
+
+	// Simple connection to root node
+	return false;
+}
+
+void UAvaPlaybackGraphSchema::GetGraphContextActions(FGraphContextMenuBuilder& ContextMenuBuilder) const
+{
+	GetPlaybackNodeActions(ContextMenuBuilder, true);
+	GetCommentAction(ContextMenuBuilder, ContextMenuBuilder.CurrentGraph);
+
+	bool bCanPasteNodes = false;
+	if (TSharedPtr<IAvaPlaybackGraphEditor> GraphEditor = GetPlaybackGraphEditor(ContextMenuBuilder.CurrentGraph))
+	{
+		bCanPasteNodes = GraphEditor->CanPasteNodes();
+	}
+	
+	if (!ContextMenuBuilder.FromPin && bCanPasteNodes)
+	{
+		TSharedPtr<FAvaPlaybackAction_PasteNode> NewAction = MakeShared<FAvaPlaybackAction_PasteNode>(FText::GetEmpty()
+			, LOCTEXT("PasteHereAction", "Paste here")
+			, FText::GetEmpty()
+			, 0);
+		
+		ContextMenuBuilder.AddAction(NewAction);
+	}
+}
+
+void UAvaPlaybackGraphSchema::GetContextMenuActions(UToolMenu* Menu, UGraphNodeContextMenuContext* Context) const
+{
+	if (Context->Node)
+	{
+		const UAvaPlaybackGraphNode* const GraphNode = Cast<const UAvaPlaybackGraphNode>(Context->Node);
+		{
+			FToolMenuSection& Section = Menu->AddSection("AvaGraphSchemaNodeActions", LOCTEXT("NodeActionsMenuHeader", "Node Actions"));
+			Section.AddMenuEntry(FGraphEditorCommands::Get().BreakNodeLinks);
+		}
+	}
+	Super::GetContextMenuActions(Menu, Context);
+}
+
+void UAvaPlaybackGraphSchema::CreateDefaultNodesForGraph(UEdGraph& Graph) const
+{
+	UAvalanchePlayback* const Playback = CastChecked<UAvaPlaybackGraph>(&Graph)->GetPlayback();
+	UAvaPlaybackNodeRoot* const RootNode = Playback->ConstructPlaybackNode<UAvaPlaybackNodeRoot>();
+	check(RootNode);
+	
+	SetNodeMetaData(RootNode->GetGraphNode(), FNodeMetadata::DefaultGraphNode);
+}
+
+const FPinConnectionResponse UAvaPlaybackGraphSchema::CanCreateConnection(const UEdGraphPin* PinA
+	, const UEdGraphPin* PinB) const
+{
+	// Make sure the pins are not on the same node
+	if (PinA->GetOwningNode() == PinB->GetOwningNode())
+	{
+		return FPinConnectionResponse(CONNECT_RESPONSE_DISALLOW, LOCTEXT("ConnectionSameNode", "Both are on the same node"));
+	}
+
+	// Pin mismatch in Pin Category
+	if (PinA->PinType.PinCategory != PinB->PinType.PinCategory)
+	{
+		return FPinConnectionResponse(CONNECT_RESPONSE_DISALLOW, LOCTEXT("IncompatibleCategories", "Pin Types are not Compatible"));
+	}
+	
+	// Compare the directions
+	const UEdGraphPin* InputPin = nullptr;
+	const UEdGraphPin* OutputPin = nullptr;
+
+	if (!CategorizePinsByDirection(PinA, PinB, InputPin, OutputPin))
+	{
+		return FPinConnectionResponse(CONNECT_RESPONSE_DISALLOW, LOCTEXT("IncompatibleDirections", "Directions are not compatible"));
+	}
+
+	if (ConnectionCausesLoop(InputPin, OutputPin))
+	{
+		return FPinConnectionResponse(CONNECT_RESPONSE_DISALLOW, LOCTEXT("ConnectionLoop", "Connection would cause loop"));
+	}
+
+	// Break existing connections on inputs only for Channel Feed only
+	// multiple Output connections are acceptable
+	if (InputPin->LinkedTo.Num() > 0)
+	{
+		ECanCreateConnectionResponse ReplyBreakOutputs;
+		if (InputPin == PinA)
+		{
+			ReplyBreakOutputs = CONNECT_RESPONSE_BREAK_OTHERS_A;
+		}
+		else
+		{
+			ReplyBreakOutputs = CONNECT_RESPONSE_BREAK_OTHERS_B;
+		}
+		return FPinConnectionResponse(ReplyBreakOutputs, LOCTEXT("ConnectionReplace", "Replace existing connections"));
+	}
+
+	return FPinConnectionResponse(CONNECT_RESPONSE_MAKE, TEXT(""));
+}
+
+bool UAvaPlaybackGraphSchema::TryCreateConnection(UEdGraphPin* PinA, UEdGraphPin* PinB) const
+{
+	if (UEdGraphSchema::TryCreateConnection(PinA, PinB))
+	{
+		CompilePlaybackNodesFromGraphNodes(*PinA->GetOwningNode());
+		return true;
+	}
+	return false;
+}
+
+FLinearColor UAvaPlaybackGraphSchema::GetPinTypeColor(const FEdGraphPinType& PinType) const
+{
+	if (PinType.PinCategory == PC_ChannelFeed)
+	{
+		return FLinearColor::White;
+	}
+	
+	if (PinType.PinCategory == PC_Event)
+	{
+		return FLinearColor::White;
+	}
+
+	return Super::GetPinTypeColor(PinType);
+}
+
+bool UAvaPlaybackGraphSchema::ShouldHidePinDefaultValue(UEdGraphPin* Pin) const
+{
+	return Super::ShouldHidePinDefaultValue(Pin);
+}
+
+void UAvaPlaybackGraphSchema::BreakNodeLinks(UEdGraphNode& TargetNode) const
+{
+	Super::BreakNodeLinks(TargetNode);
+	CompilePlaybackNodesFromGraphNodes(TargetNode);
+}
+
+void UAvaPlaybackGraphSchema::BreakPinLinks(UEdGraphPin& TargetPin, bool bSendsNodeNotifcation) const
+{
+	const FScopedTransaction Transaction(LOCTEXT("PlaybackGraph_BreakPinLinks", "Break Pin Links"));
+	Super::BreakPinLinks(TargetPin, bSendsNodeNotifcation);
+
+	// if this would notify the node then we need to compile the SoundCue
+	if (bSendsNodeNotifcation)
+	{
+		CompilePlaybackNodesFromGraphNodes(*TargetPin.GetOwningNode());
+	}
+}
+
+void UAvaPlaybackGraphSchema::GetAssetsGraphHoverMessage(const TArray<FAssetData>& Assets, const UEdGraph* HoverGraph
+	, FString& OutTooltipText, bool& OutOkIcon) const
+{
+	Super::GetAssetsGraphHoverMessage(Assets, HoverGraph, OutTooltipText, OutOkIcon);
+}
+
+void UAvaPlaybackGraphSchema::DroppedAssetsOnGraph(const TArray<FAssetData>& Assets, const FVector2D& GraphPosition
+	, UEdGraph* Graph) const
+{
+	Super::DroppedAssetsOnGraph(Assets, GraphPosition, Graph);
+}
+
+void UAvaPlaybackGraphSchema::DroppedAssetsOnNode(const TArray<FAssetData>& Assets, const FVector2D& GraphPosition
+	, UEdGraphNode* Node) const
+{
+	Super::DroppedAssetsOnNode(Assets, GraphPosition, Node);
+}
+
+int32 UAvaPlaybackGraphSchema::GetNodeSelectionCount(const UEdGraph* Graph) const
+{
+	return Super::GetNodeSelectionCount(Graph);
+}
+
+TSharedPtr<FEdGraphSchemaAction> UAvaPlaybackGraphSchema::GetCreateCommentAction() const
+{
+	return MakeShared<FAvaPlaybackAction_NewComment>();
+}
+
+void UAvaPlaybackGraphSchema::CachePlaybackNodeClasses()
+{
+	if (PlaybackNodeClasses.IsEmpty())
+	{
+		// Construct list of non-abstract sound node classes.
+		for (UClass* Class : TObjectRange<UClass>())
+		{
+			if (Class->IsChildOf(UAvaPlaybackNode::StaticClass()) && !Class->HasAnyClassFlags(CLASS_Abstract))
+			{
+				PlaybackNodeClasses.Add(Class);
+			}
+		}
+	}
+}
+
+void UAvaPlaybackGraphSchema::GetPlaybackNodeActions(FGraphActionMenuBuilder& ActionMenuBuilder, bool bShowSelectedActions) const
+{
+	CachePlaybackNodeClasses();
+	
+	for (TSubclassOf<UAvaPlaybackNode> PlaybackNodeClass : PlaybackNodeClasses)
+	{		
+		UAvaPlaybackNode* PlaybackNode = PlaybackNodeClass->GetDefaultObject<UAvaPlaybackNode>();
+		
+		if (!ActionMenuBuilder.FromPin
+			|| ActionMenuBuilder.FromPin->Direction == EGPD_Input
+			|| PlaybackNode->GetMaxChildNodes() > 0)
+		{
+			const FText NodeTitle = PlaybackNode->GetNodeDisplayNameText();
+			const FText NodeCategory = PlaybackNode->GetNodeCategoryText();
+			
+			//New Node Action (for anything other than Playback Node Root)
+			if (!PlaybackNodeClass->IsChildOf(UAvaPlaybackNodeRoot::StaticClass()))
+			{				
+				FFormatNamedArguments Arguments;
+				Arguments.Add(TEXT("Desc"), NodeTitle);
+				
+				TSharedPtr<FAvaPlaybackAction_NewNode> NewNodeAction = MakeShared<FAvaPlaybackAction_NewNode>(NodeCategory
+					, NodeTitle
+					, FText::Format(LOCTEXT("NewPlaybackNodeTooltip", "Adds {Desc} node here"), Arguments)
+					, 0);
+				
+				ActionMenuBuilder.AddAction(NewNodeAction);
+				NewNodeAction->SetPlaybackNodeClass(PlaybackNodeClass);
+			}
+		}
+	}
+}
+
+void UAvaPlaybackGraphSchema::GetCommentAction(FGraphActionMenuBuilder& ActionMenuBuilder, const UEdGraph* CurrentGraph) const
+{
+	if (!ActionMenuBuilder.FromPin)
+	{
+		TSharedPtr<FAvaPlaybackAction_NewComment> NewAction = MakeShared<FAvaPlaybackAction_NewComment>(FText::GetEmpty()
+			, LOCTEXT("AddCommentAction", "Add Comment...")
+			, LOCTEXT("CreateCommentToolTip", "Creates a comment.")
+			, 0);
+		
+		ActionMenuBuilder.AddAction(NewAction);
+	}
+}
+
+#undef LOCTEXT_NAMESPACE 
diff --git a/Engine/Plugins/Experimental/Avalanche/Source/AvalancheMediaEditor/Private/Playback/Graph/Nodes/AvaPlaybackGraphNode.cpp b/Engine/Plugins/Experimental/Avalanche/Source/AvalancheMediaEditor/Private/Playback/Graph/Nodes/AvaPlaybackGraphNode.cpp
new file mode 100644
--- /dev/null
+++ b/Engine/Plugins/Experimental/Avalanche/Source/AvalancheMediaEditor/Private/Playback/Graph/Nodes/AvaPlaybackGraphNode.cpp
@@ -0,0 +1,556 @@
+// Copyright Epic Games, Inc. All Rights Reserved.
+
+#include "Playback/Graph/Nodes/AvaPlaybackGraphNode.h"
+#include "AvalancheMediaEditorSettings.h"
+#include "Framework/Commands/GenericCommands.h"
+#include "GraphEditorActions.h"
+#include "Playback/AvalanchePlayback.h"
+#include "Playback/AvaPlaybackCommands.h"
+#include "Playback/Graph/AvaPlaybackGraph.h"
+#include "Playback/Graph/AvaPlaybackGraphSchema.h"
+#include "Playback/Graph/Nodes/AvaPlaybackGraphNode_Root.h"
+#include "Playback/Graph/Nodes/Slate/SAvaPlaybackGraphNode.h"
+#include "Playback/Nodes/AvaPlaybackNode.h"
+#include "ScopedTransaction.h"
+#include "ToolMenu.h"
+
+#define LOCTEXT_NAMESPACE "AvaPlaybackGraphNode"
+
+UAvalanchePlayback* UAvaPlaybackGraphNode::GetPlayback() const
+{
+	return CastChecked<UAvaPlaybackGraph>(GetGraph())->GetPlayback();
+}
+
+void UAvaPlaybackGraphNode::SetPlaybackNode(UAvaPlaybackNode* InPlaybackNode)
+{
+	PlaybackNode = InPlaybackNode;
+	PlaybackNode->SetGraphNode(this);
+}
+
+double UAvaPlaybackGraphNode::GetLastTimeTicked()
+{
+	return PlaybackNode ? PlaybackNode->GetLastTimeTicked() : -1.0;
+}
+
+double UAvaPlaybackGraphNode::GetChildLastTimeTicked(int32 ChildIndex) const
+{
+	return PlaybackNode ? PlaybackNode->GetChildLastTimeTicked(ChildIndex) : -1.0;
+}
+
+bool UAvaPlaybackGraphNode::CanAddInputPin() const
+{
+	if (PlaybackNode)
+	{
+		// Check if adding another input would exceed max child nodes.
+		return PlaybackNode->GetChildNodes().Num() < PlaybackNode->GetMaxChildNodes();
+	}
+	else
+	{
+		return false;
+	}
+}
+
+void UAvaPlaybackGraphNode::AddInputPin()
+{
+	const FScopedTransaction Transaction(LOCTEXT("PlaybackGraphNodeAddInput", "Add Playback Input"));
+	Modify();
+	CreateInputPin();
+
+	UAvalanchePlayback* const Playback = GetPlayback();
+	Playback->CompilePlaybackNodesFromGraphNodes();
+	Playback->MarkPackageDirty();
+
+	// Refresh the current graph, so the pins can be updated
+	GetGraph()->NotifyGraphChanged();
+}
+
+void UAvaPlaybackGraphNode::CreateInputPin()
+{
+	const int32 InputIndex = GetInputPinCount();
+	
+	UEdGraphPin* const Pin = CreatePin(EGPD_Input
+		, GetInputPinCategory(InputIndex)
+		, GetInputPinSubCategory(InputIndex)
+		, GetInputPinSubCategoryObject(InputIndex)
+		, GetInputPinName(InputIndex));
+
+	check(Pin);
+	
+	if (Pin->PinName.IsNone())
+	{
+		// Makes sure pin has a name for lookup purposes but user will never see it
+		Pin->PinName = CreateUniquePinName(TEXT("Input"));
+		Pin->PinFriendlyName = FText::FromString(TEXT(" "));
+	}
+}
+
+void UAvaPlaybackGraphNode::RemoveInputPin(UEdGraphPin* InGraphPin)
+{
+	const FScopedTransaction Transaction(LOCTEXT("PlaybackGraphNodeRemoveInput", "Remove Playback Input") );
+	Modify();
+	
+	const TArray<UEdGraphPin*> InputPins = GetInputPins();
+	
+	for (int32 InputIndex = 0; InputIndex < InputPins.Num(); InputIndex++)
+	{
+		if (InGraphPin == InputPins[InputIndex])
+		{
+			InGraphPin->MarkAsGarbage();
+			Pins.Remove(InGraphPin);
+			
+			// also remove the PlaybackNode child node so ordering matches
+			PlaybackNode->Modify();
+			PlaybackNode->RemoveChildNode(InputIndex);
+			break;
+		}
+	}
+
+	UAvalanchePlayback* const Playback = GetPlayback();
+	Playback->CompilePlaybackNodesFromGraphNodes();
+	Playback->MarkPackageDirty();
+
+	// Refresh the current graph, so the pins can be updated
+	GetGraph()->NotifyGraphChanged();
+}
+
+void UAvaPlaybackGraphNode::PrepareForCopying()
+{
+	if (PlaybackNode)
+	{
+		// Temporarily take ownership of the PlaybackNode, so that it is not deleted when cutting
+		PlaybackNode->Rename(nullptr, this, REN_DontCreateRedirectors);
+	}
+}
+
+void UAvaPlaybackGraphNode::PostPasteNode()
+{
+	if (PlaybackNode)
+	{
+		PlaybackNode->PostAllocateNode();
+	}
+}
+
+void UAvaPlaybackGraphNode::PostCopyNode()
+{
+	ResetPlaybackNodeOwner();
+}
+
+void UAvaPlaybackGraphNode::ResetPlaybackNodeOwner()
+{
+	if (PlaybackNode)
+	{
+		UAvalanchePlayback* const Playback = GetPlayback();
+		if (PlaybackNode->GetOuter() != Playback)
+		{
+			// Ensures PlaybackNode is owned by the Playback Object
+			PlaybackNode->Rename(nullptr, Playback, REN_DontCreateRedirectors);
+		}
+
+		// Set up the back pointer for newly created sound nodes
+		PlaybackNode->SetGraphNode(this);
+	}
+}
+
+TSubclassOf<UAvaPlaybackNode> UAvaPlaybackGraphNode::GetPlaybackNodeClass() const
+{
+	return UAvaPlaybackNode::StaticClass();
+}
+
+void UAvaPlaybackGraphNode::CreateInputPins()
+{
+	if (PlaybackNode)
+	{
+		for (int32 ChildIndex = 0; ChildIndex < PlaybackNode->GetChildNodes().Num(); ++ChildIndex)
+		{
+			CreateInputPin();
+		}
+	}
+}
+
+void UAvaPlaybackGraphNode::CreateOutputPin()
+{
+	if (!IsRootNode())
+	{
+		UEdGraphPin* const Pin = CreatePin(EGPD_Output
+			, GetOutputPinCategory()
+			, NAME_None);
+		
+		if (Pin->PinName.IsNone())
+		{
+			// Makes sure pin has a name for lookup purposes but user will never see it
+			Pin->PinName = CreateUniquePinName(TEXT("Output"));
+			Pin->PinFriendlyName = FText::FromString(TEXT(" "));;
+		}
+	}
+}
+
+bool UAvaPlaybackGraphNode::IsRootNode() const
+{
+	return !!Cast<UAvaPlaybackGraphNode_Root>(this);
+}
+
+UEdGraphPin* UAvaPlaybackGraphNode::GetOutputPin() const
+{
+	for (UEdGraphPin* const Pin : Pins)
+	{
+		if (Pin->Direction == EEdGraphPinDirection::EGPD_Output)
+		{
+			return Pin;
+		}
+	}
+	return nullptr;
+}
+
+TArray<UEdGraphPin*> UAvaPlaybackGraphNode::GetInputPins() const
+{
+	TArray<UEdGraphPin*> InputPins;
+	
+	//Reserve with the Expected Number of Inputs (since we assume only 1 output pin, reserve PinNum - 1)
+	InputPins.Reserve(FMath::Max(0, Pins.Num() - 1));
+	
+	for (UEdGraphPin* const Pin : Pins)
+	{
+		if (Pin->Direction == EEdGraphPinDirection::EGPD_Input)
+		{
+			InputPins.Add(Pin);
+		}
+	}
+
+	return InputPins;
+}
+
+UEdGraphPin* UAvaPlaybackGraphNode::GetInputPin(int32 InputIndex)  const
+{
+	TArray<UEdGraphPin*> InputPins = GetInputPins();
+	
+	if (InputPins.IsValidIndex(InputIndex))
+	{
+		return InputPins[InputIndex];
+	}
+	
+	checkf(0, TEXT("PlaybackGraphNode::GetInputPin. Index out of Bounds"));
+	return nullptr;
+}
+
+int32 UAvaPlaybackGraphNode::GetInputPinIndex(UEdGraphPin* InInputPin) const
+{
+	if (InInputPin && InInputPin->Direction == EGPD_Input)
+	{
+		int32 Index = 0;
+		for (const UEdGraphPin* const Pin : Pins)
+		{
+			if (Pin && Pin->Direction == EGPD_Input)
+			{
+				if (Pin == InInputPin)
+				{
+					return Index;
+				}
+				++Index;
+			}
+		}
+	}
+	return INDEX_NONE;
+}
+
+int32 UAvaPlaybackGraphNode::GetInputPinCount() const
+{
+	return GetInputPins().Num();
+}
+
+void UAvaPlaybackGraphNode::InsertNewNode(UEdGraphPin* FromPin, UEdGraphPin* NewLinkPin, TSet<UEdGraphNode*>& OutNodeList)
+{
+	const UAvaPlaybackGraphSchema* const Schema = CastChecked<UAvaPlaybackGraphSchema>(GetSchema());
+
+	// The pin we are creating from already has a connection that needs to be broken. We want to "insert" the new node in between, so that the output of the new node is hooked up too
+	UEdGraphPin* OldLinkedPin = FromPin->LinkedTo[0];
+	check(OldLinkedPin);
+
+	FromPin->BreakAllPinLinks();
+
+	// Hook up the old linked pin to the first valid output pin on the new node
+	for (int32 OutputPinIndex=0; OutputPinIndex < Pins.Num(); OutputPinIndex++)
+	{
+		UEdGraphPin* OutputPin = Pins[OutputPinIndex];
+		check(OutputPin);
+		
+		if (Schema->CanCreateConnection(OldLinkedPin, OutputPin).Response == ECanCreateConnectionResponse::CONNECT_RESPONSE_MAKE)
+		{
+			if (Schema->TryCreateConnection(OldLinkedPin, OutputPin))
+			{
+				OutNodeList.Add(OldLinkedPin->GetOwningNode());
+				OutNodeList.Add(this);
+			}
+			break;
+		}
+	}
+
+	if (Schema->TryCreateConnection(FromPin, NewLinkPin))
+	{
+		OutNodeList.Add(FromPin->GetOwningNode());
+		OutNodeList.Add(this);
+	}
+}
+
+void UAvaPlaybackGraphNode::PostLoad()
+{
+	Super::PostLoad();
+
+	// Fixup any Playback Node back pointers that may be out of date
+	if (PlaybackNode)
+	{
+		PlaybackNode->SetGraphNode(this);
+	}
+
+	for (UEdGraphPin* const Pin : Pins)
+	{
+		if (Pin && Pin->PinName.IsNone())
+		{
+			// Makes sure pin has a name for lookup purposes but user will never see it
+			if (Pin->Direction == EGPD_Input)
+			{
+				Pin->PinName = CreateUniquePinName(TEXT("Input"));
+			}
+			else
+			{
+				Pin->PinName = CreateUniquePinName(TEXT("Output"));
+			}
+			Pin->PinFriendlyName = FText::FromString(TEXT(" "));
+		}
+	}
+}
+
+void UAvaPlaybackGraphNode::PostEditImport()
+{
+	Super::PostEditImport();
+	ResetPlaybackNodeOwner();
+}
+
+void UAvaPlaybackGraphNode::PostDuplicate(bool bDuplicateForPIE)
+{
+	Super::PostDuplicate(bDuplicateForPIE);
+	if (!bDuplicateForPIE)
+	{
+		CreateNewGuid();
+	}
+}
+
+FText UAvaPlaybackGraphNode::GetNodeTitle(ENodeTitleType::Type TitleType) const
+{
+	if (PlaybackNode)
+	{
+		return PlaybackNode->GetNodeDisplayNameText();
+	}
+	return Super::GetNodeTitle(TitleType);
+}
+
+FText UAvaPlaybackGraphNode::GetTooltipText() const
+{
+	if (PlaybackNode)
+	{
+		return PlaybackNode->GetNodeTooltipText();
+	}
+	return Super::GetTooltipText();
+}
+
+void UAvaPlaybackGraphNode::GetNodeContextMenuActions(UToolMenu* Menu, UGraphNodeContextMenuContext* Context) const
+{
+	if (Context->Pin)
+	{
+		// If on an input that can be deleted, show option
+		if (Context->Pin->Direction == EGPD_Input && PlaybackNode->GetChildNodes().Num() > PlaybackNode->GetMinChildNodes())
+		{
+			FToolMenuSection& Section = Menu->AddSection("PlaybackGraphRemovePin"
+				, LOCTEXT("PlaybackGraphRemovePin", "Input Pin Actions"));
+			Section.AddMenuEntry(FAvaPlaybackCommands::Get().AddInputPin);
+			Section.AddMenuEntry(FAvaPlaybackCommands::Get().RemoveInputPin);
+		}
+	}
+	else if (Context->Node)
+	{
+		{
+			FToolMenuSection& Section = Menu->AddSection("PlaybackGraphNodeEdit");
+			Section.AddMenuEntry(FGenericCommands::Get().Delete);
+			Section.AddMenuEntry(FGenericCommands::Get().Cut);
+			Section.AddMenuEntry(FGenericCommands::Get().Copy);
+			Section.AddMenuEntry(FGenericCommands::Get().Duplicate);
+			Section.AddSeparator(NAME_None);
+		}
+		
+		{
+			FToolMenuSection& Section = Menu->AddSection("PlaybackGraphNodeOrganization", LOCTEXT("PlaybackGraphNodeOrganization", "Organization"));
+			
+			Section.AddSubMenu("Alignment"
+				, LOCTEXT("AlignmentHeader", "Alignment")
+				, FText()
+				, FNewToolMenuDelegate::CreateLambda([](UToolMenu* SubMenu)
+				{
+					{
+						FToolMenuSection& SubMenuSection = SubMenu->AddSection("PlaybackGraphSchemaAlignment", LOCTEXT("AlignHeader", "Align"));
+						SubMenuSection.AddMenuEntry(FGraphEditorCommands::Get().AlignNodesTop);
+						SubMenuSection.AddMenuEntry(FGraphEditorCommands::Get().AlignNodesMiddle);
+						SubMenuSection.AddMenuEntry(FGraphEditorCommands::Get().AlignNodesBottom);
+						SubMenuSection.AddMenuEntry(FGraphEditorCommands::Get().AlignNodesLeft);
+						SubMenuSection.AddMenuEntry(FGraphEditorCommands::Get().AlignNodesCenter);
+						SubMenuSection.AddMenuEntry(FGraphEditorCommands::Get().AlignNodesRight);
+						SubMenuSection.AddMenuEntry(FGraphEditorCommands::Get().StraightenConnections);
+					}
+	
+					{
+						FToolMenuSection& SubMenuSection = SubMenu->AddSection("PlaybackGraphSchemaDistribution", LOCTEXT("DistributionHeader", "Distribution"));
+						SubMenuSection.AddMenuEntry(FGraphEditorCommands::Get().DistributeNodesHorizontally);
+						SubMenuSection.AddMenuEntry(FGraphEditorCommands::Get().DistributeNodesVertically);
+					}
+				}));
+			
+			Section.AddSeparator(NAME_None);
+		}
+	}
+}
+
+void UAvaPlaybackGraphNode::AllocateDefaultPins()
+{
+	check(Pins.Num() == 0);
+	CreateInputPins();
+	CreateOutputPin();
+}
+
+void UAvaPlaybackGraphNode::ReconstructNode()
+{
+	// Break any links to 'orphan' pins
+	for (int32 PinIndex = 0; PinIndex < Pins.Num(); ++PinIndex)
+	{
+		UEdGraphPin* Pin = Pins[PinIndex];
+		TArray<class UEdGraphPin*>& LinkedToRef = Pin->LinkedTo;
+		for (int32 LinkIdx=0; LinkIdx < LinkedToRef.Num(); LinkIdx++)
+		{
+			UEdGraphPin* OtherPin = LinkedToRef[LinkIdx];
+			// If we are linked to a pin that its owner doesn't know about, break that link
+			if (!OtherPin->GetOwningNode()->Pins.Contains(OtherPin))
+			{
+				Pin->LinkedTo.Remove(OtherPin);
+			}
+		}
+	}
+
+	// Store the old Input and Output pins
+	TArray<UEdGraphPin*> OldInputPins = GetInputPins();
+	UEdGraphPin* OldOutputPin = GetOutputPin();
+	
+	// Move the existing pins to a saved array
+	TArray OldPins(Pins);
+	Pins.Reset();
+
+	// Recreate the new pins
+	AllocateDefaultPins();
+
+	// Get new Input and Output pins
+	TArray<UEdGraphPin*> NewInputPins = GetInputPins();
+	UEdGraphPin* NewOutputPin = GetOutputPin();
+
+	for (int32 PinIndex = 0; PinIndex < OldInputPins.Num(); ++PinIndex)
+	{
+		if (PinIndex < NewInputPins.Num())
+		{
+			NewInputPins[PinIndex]->MovePersistentDataFromOldPin(*OldInputPins[PinIndex]);
+		}
+	}
+	
+	if (OldOutputPin)
+	{
+		NewOutputPin->MovePersistentDataFromOldPin(*OldOutputPin);
+	}
+
+	// Throw away the original pins
+	for (UEdGraphPin* OldPin : OldPins)
+	{
+		OldPin->Modify();
+		UEdGraphNode::DestroyPin(OldPin);
+	}
+}
+
+void UAvaPlaybackGraphNode::AutowireNewNode(UEdGraphPin* FromPin)
+{
+	if (FromPin)
+	{
+		const UAvaPlaybackGraphSchema* const Schema = CastChecked<UAvaPlaybackGraphSchema>(GetSchema());
+		TSet<UEdGraphNode*> NodeList;
+
+		// auto-connect from dragged pin to first compatible pin on the new node
+		for (int32 Index = 0; Index < Pins.Num(); ++Index)
+		{
+			UEdGraphPin* Pin = Pins[Index];
+			check(Pin);
+			
+			FPinConnectionResponse Response = Schema->CanCreateConnection(FromPin, Pin);
+			
+			if (Response.Response == ECanCreateConnectionResponse::CONNECT_RESPONSE_MAKE)
+			{
+				if (Schema->TryCreateConnection(FromPin, Pin))
+				{
+					NodeList.Add(FromPin->GetOwningNode());
+					NodeList.Add(this);
+				}
+				break;
+			}
+			
+			if (Response.Response == ECanCreateConnectionResponse::CONNECT_RESPONSE_BREAK_OTHERS_A)
+			{
+				InsertNewNode(FromPin, Pin, NodeList);
+				break;
+			}
+		}
+
+		// Send all nodes that received a new pin connection a notification
+		for (TSet<UEdGraphNode*>::TConstIterator It = NodeList.CreateConstIterator(); It; ++It)
+		{
+			UEdGraphNode* Node = *It;
+			Node->NodeConnectionListChanged();
+		}
+	}
+}
+
+bool UAvaPlaybackGraphNode::CanCreateUnderSpecifiedSchema(const UEdGraphSchema* Schema) const
+{
+	return Schema->IsA(UAvaPlaybackGraphSchema::StaticClass());
+}
+
+FLinearColor UAvaPlaybackGraphNode::GetNodeTitleColor() const
+{
+	return UAvalancheMediaEditorSettings::Get().PlaybackDefaultNodeColor;
+}
+
+TSharedPtr<SGraphNode> UAvaPlaybackGraphNode::CreateVisualWidget()
+{
+	return SNew(SAvaPlaybackGraphNode, this);
+}
+
+FName UAvaPlaybackGraphNode::GetInputPinName(int32 InputPinIndex) const
+{
+	if (PlaybackNode)
+	{
+		return PlaybackNode->GetInputPinName(InputPinIndex);
+	}
+	return NAME_None;
+}
+
+FName UAvaPlaybackGraphNode::GetInputPinCategory(int32 InputPinIndex) const
+{
+	return UAvaPlaybackGraphSchema::PC_ChannelFeed;
+}
+
+FName UAvaPlaybackGraphNode::GetInputPinSubCategory(int32 InputPinIndex) const
+{
+	return NAME_None;
+}
+
+UObject* UAvaPlaybackGraphNode::GetInputPinSubCategoryObject(int32 InputPinIndex) const
+{
+	return nullptr;
+}
+
+FName UAvaPlaybackGraphNode::GetOutputPinCategory() const
+{
+	return UAvaPlaybackGraphSchema::PC_ChannelFeed;
+}
+
+#undef LOCTEXT_NAMESPACE
diff --git a/Engine/Plugins/Experimental/Avalanche/Source/AvalancheMediaEditor/Private/Playback/Graph/Nodes/AvaPlaybackGraphNode_BlueprintPlayer.cpp b/Engine/Plugins/Experimental/Avalanche/Source/AvalancheMediaEditor/Private/Playback/Graph/Nodes/AvaPlaybackGraphNode_BlueprintPlayer.cpp
new file mode 100644
--- /dev/null
+++ b/Engine/Plugins/Experimental/Avalanche/Source/AvalancheMediaEditor/Private/Playback/Graph/Nodes/AvaPlaybackGraphNode_BlueprintPlayer.cpp
@@ -0,0 +1,27 @@
+// Copyright Epic Games, Inc. All Rights Reserved.
+
+#include "Playback/Graph/Nodes/AvaPlaybackGraphNode_BlueprintPlayer.h"
+#include "AvalancheMediaEditorSettings.h"
+#include "Playback/Graph/AvaPlaybackGraphSchema.h"
+#include "Playback/Nodes/AvaPlaybackNodeBlueprintPlayer.h"
+#include "Slate/SAvaPlaybackGraphNode_Player.h"
+
+TSubclassOf<UAvaPlaybackNode> UAvaPlaybackGraphNode_BlueprintPlayer::GetPlaybackNodeClass() const
+{
+	return UAvaPlaybackNodeBlueprintPlayer::StaticClass();
+}
+
+FName UAvaPlaybackGraphNode_BlueprintPlayer::GetInputPinCategory(int32 InputPinIndex) const
+{
+	return UAvaPlaybackGraphSchema::PC_Event;
+}
+
+FLinearColor UAvaPlaybackGraphNode_BlueprintPlayer::GetNodeTitleColor() const
+{
+	return UAvalancheMediaEditorSettings::Get().PlaybackPlayerNodeColor;
+}
+
+TSharedPtr<SGraphNode> UAvaPlaybackGraphNode_BlueprintPlayer::CreateVisualWidget()
+{
+	return SNew(SAvaPlaybackGraphNode_Player, this);
+}
diff --git a/Engine/Plugins/Experimental/Avalanche/Source/AvalancheMediaEditor/Private/Playback/Graph/Nodes/AvaPlaybackGraphNode_EventAction.cpp b/Engine/Plugins/Experimental/Avalanche/Source/AvalancheMediaEditor/Private/Playback/Graph/Nodes/AvaPlaybackGraphNode_EventAction.cpp
new file mode 100644
--- /dev/null
+++ b/Engine/Plugins/Experimental/Avalanche/Source/AvalancheMediaEditor/Private/Playback/Graph/Nodes/AvaPlaybackGraphNode_EventAction.cpp
@@ -0,0 +1,26 @@
+// Copyright Epic Games, Inc. All Rights Reserved.
+
+#include "Playback/Graph/Nodes/AvaPlaybackGraphNode_EventAction.h"
+#include "AvalancheMediaEditorSettings.h"
+#include "Playback/Graph/AvaPlaybackGraphSchema.h"
+#include "Playback/Nodes/Events/AvaPlaybackNodeAction.h"
+
+TSubclassOf<UAvaPlaybackNode> UAvaPlaybackGraphNode_EventAction::GetPlaybackNodeClass() const
+{
+	return UAvaPlaybackNodeAction::StaticClass();
+}
+
+FName UAvaPlaybackGraphNode_EventAction::GetInputPinCategory(int32 InputPinIndex) const
+{
+	return UAvaPlaybackGraphSchema::PC_Event;
+}
+
+FName UAvaPlaybackGraphNode_EventAction::GetOutputPinCategory() const
+{
+	return UAvaPlaybackGraphSchema::PC_Event;
+}
+
+FLinearColor UAvaPlaybackGraphNode_EventAction::GetNodeTitleColor() const
+{
+	return UAvalancheMediaEditorSettings::Get().PlaybackActionNodeColor;
+}
diff --git a/Engine/Plugins/Experimental/Avalanche/Source/AvalancheMediaEditor/Private/Playback/Graph/Nodes/AvaPlaybackGraphNode_EventFlow.cpp b/Engine/Plugins/Experimental/Avalanche/Source/AvalancheMediaEditor/Private/Playback/Graph/Nodes/AvaPlaybackGraphNode_EventFlow.cpp
new file mode 100644
--- /dev/null
+++ b/Engine/Plugins/Experimental/Avalanche/Source/AvalancheMediaEditor/Private/Playback/Graph/Nodes/AvaPlaybackGraphNode_EventFlow.cpp
@@ -0,0 +1,20 @@
+// Copyright Epic Games, Inc. All Rights Reserved.
+
+#include "Playback/Graph/Nodes/AvaPlaybackGraphNode_EventFlow.h"
+#include "Playback/Graph/AvaPlaybackGraphSchema.h"
+#include "Playback/Nodes/Events/AvaPlaybackNodeFlow.h"
+
+TSubclassOf<UAvaPlaybackNode> UAvaPlaybackGraphNode_EventFlow::GetPlaybackNodeClass() const
+{
+	return UAvaPlaybackNodeFlow::StaticClass();
+}
+
+FName UAvaPlaybackGraphNode_EventFlow::GetInputPinCategory(int32 InputPinIndex) const
+{
+	return UAvaPlaybackGraphSchema::PC_Event;
+}
+
+FName UAvaPlaybackGraphNode_EventFlow::GetOutputPinCategory() const
+{
+	return UAvaPlaybackGraphSchema::PC_Event;
+}
diff --git a/Engine/Plugins/Experimental/Avalanche/Source/AvalancheMediaEditor/Private/Playback/Graph/Nodes/AvaPlaybackGraphNode_EventTrigger.cpp b/Engine/Plugins/Experimental/Avalanche/Source/AvalancheMediaEditor/Private/Playback/Graph/Nodes/AvaPlaybackGraphNode_EventTrigger.cpp
new file mode 100644
--- /dev/null
+++ b/Engine/Plugins/Experimental/Avalanche/Source/AvalancheMediaEditor/Private/Playback/Graph/Nodes/AvaPlaybackGraphNode_EventTrigger.cpp
@@ -0,0 +1,21 @@
+// Copyright Epic Games, Inc. All Rights Reserved.
+
+#include "Playback/Graph/Nodes/AvaPlaybackGraphNode_EventTrigger.h"
+#include "AvalancheMediaEditorSettings.h"
+#include "Playback/Graph/AvaPlaybackGraphSchema.h"
+#include "Playback/Nodes/Events/AvaPlaybackNodeTrigger.h"
+
+TSubclassOf<UAvaPlaybackNode> UAvaPlaybackGraphNode_EventTrigger::GetPlaybackNodeClass() const
+{
+	return UAvaPlaybackNodeTrigger::StaticClass();
+}
+
+FName UAvaPlaybackGraphNode_EventTrigger::GetOutputPinCategory() const
+{
+	return UAvaPlaybackGraphSchema::PC_Event;
+}
+
+FLinearColor UAvaPlaybackGraphNode_EventTrigger::GetNodeTitleColor() const
+{
+	return UAvalancheMediaEditorSettings::Get().PlaybackEventNodeColor;
+}
diff --git a/Engine/Plugins/Experimental/Avalanche/Source/AvalancheMediaEditor/Private/Playback/Graph/Nodes/AvaPlaybackGraphNode_LevelPlayer.cpp b/Engine/Plugins/Experimental/Avalanche/Source/AvalancheMediaEditor/Private/Playback/Graph/Nodes/AvaPlaybackGraphNode_LevelPlayer.cpp
new file mode 100644
--- /dev/null
+++ b/Engine/Plugins/Experimental/Avalanche/Source/AvalancheMediaEditor/Private/Playback/Graph/Nodes/AvaPlaybackGraphNode_LevelPlayer.cpp
@@ -0,0 +1,27 @@
+// Copyright Epic Games, Inc. All Rights Reserved.
+
+#include "Playback/Graph/Nodes/AvaPlaybackGraphNode_LevelPlayer.h"
+#include "AvalancheMediaEditorSettings.h"
+#include "Playback/Graph/AvaPlaybackGraphSchema.h"
+#include "Playback/Nodes/AvaPlaybackNodeLevelPlayer.h"
+#include "Slate/SAvaPlaybackGraphNode_Player.h"
+
+TSubclassOf<UAvaPlaybackNode> UAvaPlaybackGraphNode_LevelPlayer::GetPlaybackNodeClass() const
+{
+	return UAvaPlaybackNodeLevelPlayer::StaticClass();
+}
+
+FName UAvaPlaybackGraphNode_LevelPlayer::GetInputPinCategory(int32 InputPinIndex) const
+{
+	return UAvaPlaybackGraphSchema::PC_Event;
+}
+
+FLinearColor UAvaPlaybackGraphNode_LevelPlayer::GetNodeTitleColor() const
+{
+	return UAvalancheMediaEditorSettings::Get().PlaybackPlayerNodeColor;
+}
+
+TSharedPtr<SGraphNode> UAvaPlaybackGraphNode_LevelPlayer::CreateVisualWidget()
+{
+	return SNew(SAvaPlaybackGraphNode_Player, this);
+}
diff --git a/Engine/Plugins/Experimental/Avalanche/Source/AvalancheMediaEditor/Private/Playback/Graph/Nodes/AvaPlaybackGraphNode_Root.cpp b/Engine/Plugins/Experimental/Avalanche/Source/AvalancheMediaEditor/Private/Playback/Graph/Nodes/AvaPlaybackGraphNode_Root.cpp
new file mode 100644
--- /dev/null
+++ b/Engine/Plugins/Experimental/Avalanche/Source/AvalancheMediaEditor/Private/Playback/Graph/Nodes/AvaPlaybackGraphNode_Root.cpp
@@ -0,0 +1,35 @@
+// Copyright Epic Games, Inc. All Rights Reserved.
+
+#include "Playback/Graph/Nodes/AvaPlaybackGraphNode_Root.h"
+#include "AvalancheMediaEditorSettings.h"
+#include "Playback/Nodes/AvaPlaybackNodeRoot.h"
+#include "Slate/SAvaPlaybackGraphNode_Root.h"
+
+#define LOCTEXT_NAMESPACE "AvaPlaybackGraphNode_Root"
+
+TSubclassOf<UAvaPlaybackNode> UAvaPlaybackGraphNode_Root::GetPlaybackNodeClass() const
+{
+	return UAvaPlaybackNodeRoot::StaticClass();
+}
+
+FLinearColor UAvaPlaybackGraphNode_Root::GetNodeTitleColor() const
+{
+	return UAvalancheMediaEditorSettings::Get().PlaybackChannelsNodeColor;
+}
+
+bool UAvaPlaybackGraphNode_Root::CanUserDeleteNode() const
+{
+	return false;
+}
+
+bool UAvaPlaybackGraphNode_Root::CanDuplicateNode() const
+{
+	return false;
+}
+
+TSharedPtr<SGraphNode> UAvaPlaybackGraphNode_Root::CreateVisualWidget()
+{
+	return SNew(SAvaPlaybackGraphNode_Root, this);
+}
+
+#undef LOCTEXT_NAMESPACE
diff --git a/Engine/Plugins/Experimental/Avalanche/Source/AvalancheMediaEditor/Private/Playback/Graph/Nodes/Slate/SAvaPlaybackGraphNode.cpp b/Engine/Plugins/Experimental/Avalanche/Source/AvalancheMediaEditor/Private/Playback/Graph/Nodes/Slate/SAvaPlaybackGraphNode.cpp
new file mode 100644
--- /dev/null
+++ b/Engine/Plugins/Experimental/Avalanche/Source/AvalancheMediaEditor/Private/Playback/Graph/Nodes/Slate/SAvaPlaybackGraphNode.cpp
@@ -0,0 +1,73 @@
+// Copyright Epic Games, Inc. All Rights Reserved.
+
+#include "Playback/Graph/Nodes/Slate/SAvaPlaybackGraphNode.h"
+#include "GraphEditorSettings.h"
+#include "KismetPins/SGraphPinExec.h"
+#include "Playback/Graph/AvaPlaybackGraphSchema.h"
+#include "Playback/Graph/Nodes/AvaPlaybackGraphNode.h"
+#include "Widgets/SBoxPanel.h"
+
+#define LOCTEXT_NAMESPACE "SAvaPlaybackGraphNode"
+
+void SAvaPlaybackGraphNode::Construct(const FArguments& InArgs, UAvaPlaybackGraphNode* InGraphNode)
+{
+	PlaybackGraphNode = InGraphNode;
+	GraphNode = InGraphNode;
+	
+	SetCursor(EMouseCursor::CardinalCross);
+	UpdateGraphNode();
+	
+	PostConstruct();
+}
+
+void SAvaPlaybackGraphNode::CreateOutputSideAddButton(TSharedPtr<SVerticalBox> OutputBox)
+{
+	TSharedRef<SWidget> AddPinButton = AddPinButtonContent(LOCTEXT("AvaPlaybackNodeAddPinButton", "Add Input")
+		, LOCTEXT("AvaPlaybackNodeAddPinButton_Tooltip", "Adds an input to the Playback node"));
+
+	FMargin AddPinPadding = Settings->GetOutputPinPadding();
+	AddPinPadding.Top += 6.0f;
+
+	OutputBox->AddSlot()
+		.AutoHeight()
+		.VAlign(VAlign_Center)
+		.Padding(AddPinPadding)
+		[
+			AddPinButton
+		];
+}
+
+EVisibility SAvaPlaybackGraphNode::IsAddPinButtonVisible() const
+{
+	EVisibility ButtonVisibility = SGraphNode::IsAddPinButtonVisible();
+	if (ButtonVisibility == EVisibility::Visible)
+	{
+		if (PlaybackGraphNode.IsValid() && !PlaybackGraphNode->CanAddInputPin())
+		{
+			ButtonVisibility = EVisibility::Collapsed;
+		}
+	}
+	return ButtonVisibility;
+}
+
+FReply SAvaPlaybackGraphNode::OnAddPin()
+{
+	if (PlaybackGraphNode.IsValid())
+	{
+		PlaybackGraphNode->AddInputPin();
+		return FReply::Handled();
+	}
+	
+	return FReply::Unhandled();
+}
+
+TSharedPtr<SGraphPin> SAvaPlaybackGraphNode::CreatePinWidget(UEdGraphPin* Pin) const
+{
+	if (Pin->PinType.PinCategory == UAvaPlaybackGraphSchema::PC_Event)
+	{
+		return SNew(SGraphPinExec, Pin);
+	}
+	return SGraphNode::CreatePinWidget(Pin);
+}
+
+#undef LOCTEXT_NAMESPACE 
diff --git a/Engine/Plugins/Experimental/Avalanche/Source/AvalancheMediaEditor/Private/Playback/Graph/Nodes/Slate/SAvaPlaybackGraphNode_Player.cpp b/Engine/Plugins/Experimental/Avalanche/Source/AvalancheMediaEditor/Private/Playback/Graph/Nodes/Slate/SAvaPlaybackGraphNode_Player.cpp
new file mode 100644
--- /dev/null
+++ b/Engine/Plugins/Experimental/Avalanche/Source/AvalancheMediaEditor/Private/Playback/Graph/Nodes/Slate/SAvaPlaybackGraphNode_Player.cpp
@@ -0,0 +1,181 @@
+// Copyright Epic Games, Inc. All Rights Reserved.
+
+#include "SAvaPlaybackGraphNode_Player.h"
+
+#include "Brushes/SlateImageBrush.h"
+#include "Channel/AvaOutputChannel.h"
+#include "Engine/TextureRenderTarget2D.h"
+#include "GraphEditorSettings.h"
+#include "OutputDevices/Slate/SAvaCaptureImage.h"
+#include "Playback/Graph/Nodes/AvaPlaybackGraphNode.h"
+#include "Playback/Nodes/AvaPlaybackNodePlayer.h"
+#include "Widgets/Layout/SBox.h"
+#include "Widgets/Layout/SScaleBox.h"
+#include "Widgets/SBoxPanel.h"
+
+/**
+ * Wrapper to encapsulate a FSlateImageBrush that can be updated with a render target.
+ * If the render target becomes null, the brush is updated to be a solid color.
+ */
+struct SAvaPlaybackGraphNode_Player::FPreviewBrush
+{
+	FSlateImageBrush Brush;
+	FVector2D RenderTargetSize;
+
+	FPreviewBrush(const FVector2D& InRenderTargetSize)
+		: Brush(NAME_None, InRenderTargetSize, FLinearColor::Black)
+		, RenderTargetSize(InRenderTargetSize)
+	{
+
+	}
+
+	void Update(UTextureRenderTarget2D* InRenderTarget)
+	{
+		if (InRenderTarget)
+		{
+			RenderTargetSize = FVector2D(InRenderTarget->SizeX, InRenderTarget->SizeY);
+
+			//If Brush is invalid, or the Brush's Texture Target doesn't match the new Render Target, reset the Brush.
+			if (Brush.GetResourceObject() != InRenderTarget)
+			{
+				Brush = FSlateImageBrush(InRenderTarget, RenderTargetSize);
+			}
+			//If Sizes mismatch, just resizes rather than recreating the Brush with same underlying Resource
+			else if (Brush.GetImageSize() != RenderTargetSize)
+			{
+				Brush.SetImageSize(RenderTargetSize);
+			}
+		}
+		else if (Brush.GetResourceObject() != nullptr)
+		{
+			// Preserve the size of the preview brush so the checkerboard remains the same.
+			Brush = FSlateImageBrush(NAME_None, RenderTargetSize, FLinearColor::Black);
+		}
+	}
+};
+
+SAvaPlaybackGraphNode_Player::SAvaPlaybackGraphNode_Player() = default;
+SAvaPlaybackGraphNode_Player::~SAvaPlaybackGraphNode_Player() = default;
+
+void SAvaPlaybackGraphNode_Player::PostConstruct()
+{
+	check(PlaybackGraphNode.IsValid());
+	PlayerNode = Cast<UAvaPlaybackNodePlayer>(PlaybackGraphNode->GetPlaybackNode());
+	check(PlayerNode.IsValid());
+	
+	FIntPoint RenderTargetSize;
+	if (const UTextureRenderTarget2D* RenderTarget = PlayerNode->GetPreviewRenderTarget())
+	{
+		RenderTargetSize = FIntPoint(RenderTarget->SizeX, RenderTarget->SizeY);
+	}
+	else
+	{
+		RenderTargetSize = FAvaOutputChannel::GetDefaultMediaOutputSize(EAvaBroadcastChannelType::Program);
+	}	
+	PlayerPreviewBrush = MakeUnique<FPreviewBrush>(FVector2d(RenderTargetSize.X, RenderTargetSize.Y));
+}
+
+TSharedRef<SWidget> SAvaPlaybackGraphNode_Player::CreatePreviewWidget()
+{
+	return SNew(SBox)
+		.WidthOverride(106.f)
+		.HeightOverride(106.f)
+		[
+			SNew(SBorder)
+			.Padding(5.0f)
+			.BorderImage(FAppStyle::GetBrush("NoBorder"))
+			[
+				SNew(SScaleBox)
+				.Stretch(EStretch::ScaleToFitY)
+				[
+					SNew(SAvaCaptureImage)
+					.ImageArgs(SImage::FArguments()
+						.Image(this, &SAvaPlaybackGraphNode_Player::GetPlayerPreviewBrush))
+					.ShouldInvertAlpha(true)
+					.EnableGammaCorrection(false)
+					.EnableBlending(false)
+				]
+			]
+		];
+}
+
+const FSlateBrush* SAvaPlaybackGraphNode_Player::GetPlayerPreviewBrush() const
+{
+	// Update the brush here to ensure the cached image attribute is updated with the correct render target.
+	PlayerPreviewBrush->Update(PlayerNode.IsValid() ? PlayerNode->GetPreviewRenderTarget() : nullptr);
+	return &PlayerPreviewBrush->Brush;
+}
+
+void SAvaPlaybackGraphNode_Player::CreateBelowPinControls(TSharedPtr<SVerticalBox> MainBox)
+{
+	if (GraphNode && MainBox.IsValid())
+	{
+		int32 LeftPinCount = InputPins.Num();
+		int32 RightPinCount = OutputPins.Num();
+
+		const float NegativeHPad = FMath::Max<float>(-Settings->PaddingTowardsNodeEdge, 0.0f);
+		const float ExtraPad = 0.0f;
+
+		// Place preview widget based on where the least pins are
+		if (LeftPinCount < RightPinCount || RightPinCount == 0)
+		{
+			LeftNodeBox->AddSlot()
+			.Padding(FMargin(NegativeHPad + ExtraPad, 0.0f, 0.0f, 0.0f))
+			.AutoHeight()
+			.HAlign(HAlign_Left)
+			[
+				CreatePreviewWidget()
+			];
+		}
+		else if (LeftPinCount > RightPinCount)
+		{
+			RightNodeBox->AddSlot()
+			.Padding(FMargin(NegativeHPad + ExtraPad, 0.0f, 0.0f, 0.0f))
+			.AutoHeight()
+			.HAlign(HAlign_Right)
+			[
+				CreatePreviewWidget()
+			];
+		}
+		else
+		{
+			MainBox->AddSlot()
+			.Padding(Settings->GetNonPinNodeBodyPadding())
+			.AutoHeight()
+			[
+				SNew(SHorizontalBox)
+				+SHorizontalBox::Slot()
+				.AutoWidth()
+				[
+					CreatePreviewWidget()
+				]
+			];
+		}
+	}
+}
+
+TSharedRef<SWidget> SAvaPlaybackGraphNode_Player::CreateNodeContentArea()
+{
+	// NODE CONTENT AREA
+	return SNew(SBorder)
+		.BorderImage(FAppStyle::GetBrush("NoBorder"))
+		.HAlign(HAlign_Fill)
+		.VAlign(VAlign_Fill)
+		[
+			SNew(SHorizontalBox)
+			+SHorizontalBox::Slot()
+			.HAlign(HAlign_Left)
+			.FillWidth(1.0f)
+			[
+				// LEFT
+				SAssignNew(LeftNodeBox, SVerticalBox)
+			]
+			+SHorizontalBox::Slot()
+			.AutoWidth()
+			.HAlign(HAlign_Right)
+			[
+				// RIGHT
+				SAssignNew(RightNodeBox, SVerticalBox)
+			]
+		];
+}
diff --git a/Engine/Plugins/Experimental/Avalanche/Source/AvalancheMediaEditor/Private/Playback/Graph/Nodes/Slate/SAvaPlaybackGraphNode_Player.h b/Engine/Plugins/Experimental/Avalanche/Source/AvalancheMediaEditor/Private/Playback/Graph/Nodes/Slate/SAvaPlaybackGraphNode_Player.h
new file mode 100644
--- /dev/null
+++ b/Engine/Plugins/Experimental/Avalanche/Source/AvalancheMediaEditor/Private/Playback/Graph/Nodes/Slate/SAvaPlaybackGraphNode_Player.h
@@ -0,0 +1,34 @@
+// Copyright Epic Games, Inc. All Rights Reserved.
+
+#pragma once
+
+#include "CoreMinimal.h"
+#include "Playback/Graph/Nodes/Slate/SAvaPlaybackGraphNode.h"
+
+class UAvaPlaybackNodePlayer;
+
+class SAvaPlaybackGraphNode_Player : public SAvaPlaybackGraphNode
+{
+public:
+	SAvaPlaybackGraphNode_Player();
+	~SAvaPlaybackGraphNode_Player();
+
+	virtual void PostConstruct() override;
+
+	/** Creates a preview viewport if necessary */
+	TSharedRef<SWidget> CreatePreviewWidget();
+
+	const FSlateBrush* GetPlayerPreviewBrush() const;
+
+	//SGraphNode Interface
+	virtual void CreateBelowPinControls(TSharedPtr<SVerticalBox> MainBox) override;
+	virtual TSharedRef<SWidget> CreateNodeContentArea() override;
+	//~SGraphNode Interface
+
+protected:
+
+	TWeakObjectPtr<UAvaPlaybackNodePlayer> PlayerNode;
+
+	struct FPreviewBrush;
+	TUniquePtr<FPreviewBrush> PlayerPreviewBrush;
+};
diff --git a/Engine/Plugins/Experimental/Avalanche/Source/AvalancheMediaEditor/Private/Playback/Graph/Nodes/Slate/SAvaPlaybackGraphNode_Root.cpp b/Engine/Plugins/Experimental/Avalanche/Source/AvalancheMediaEditor/Private/Playback/Graph/Nodes/Slate/SAvaPlaybackGraphNode_Root.cpp
new file mode 100644
--- /dev/null
+++ b/Engine/Plugins/Experimental/Avalanche/Source/AvalancheMediaEditor/Private/Playback/Graph/Nodes/Slate/SAvaPlaybackGraphNode_Root.cpp
@@ -0,0 +1,14 @@
+// Copyright Epic Games, Inc. All Rights Reserved.
+
+#include "SAvaPlaybackGraphNode_Root.h"
+#include "Playback/Graph/AvaPlaybackGraphSchema.h"
+#include "Playback/Graph/Pins/Slate/SAvaPlaybackGraphPin_Channel.h"
+
+TSharedPtr<SGraphPin> SAvaPlaybackGraphNode_Root::CreatePinWidget(UEdGraphPin* Pin) const
+{
+	if (Pin->PinType.PinCategory == UAvaPlaybackGraphSchema::PC_ChannelFeed)
+	{
+		return SNew(SAvaPlaybackGraphPin_Channel, Pin);
+	}
+	return SGraphNode::CreatePinWidget(Pin);
+}
diff --git a/Engine/Plugins/Experimental/Avalanche/Source/AvalancheMediaEditor/Private/Playback/Graph/Nodes/Slate/SAvaPlaybackGraphNode_Root.h b/Engine/Plugins/Experimental/Avalanche/Source/AvalancheMediaEditor/Private/Playback/Graph/Nodes/Slate/SAvaPlaybackGraphNode_Root.h
new file mode 100644
--- /dev/null
+++ b/Engine/Plugins/Experimental/Avalanche/Source/AvalancheMediaEditor/Private/Playback/Graph/Nodes/Slate/SAvaPlaybackGraphNode_Root.h
@@ -0,0 +1,16 @@
+// Copyright Epic Games, Inc. All Rights Reserved.
+
+#pragma once
+
+#include "CoreMinimal.h"
+#include "Playback/Graph/Nodes/Slate/SAvaPlaybackGraphNode.h"
+
+class UAvaPlaybackGraphNode_Root;
+
+class SAvaPlaybackGraphNode_Root : public SAvaPlaybackGraphNode
+{
+public:
+
+	/** Hook that allows derived classes to supply their own SGraphPin derivatives for any pin. */
+	virtual TSharedPtr<SGraphPin> CreatePinWidget(UEdGraphPin* Pin) const override;
+};
diff --git a/Engine/Plugins/Experimental/Avalanche/Source/AvalancheMediaEditor/Private/Playback/Graph/Pins/Slate/SAvaPlaybackGraphPin_Channel.cpp b/Engine/Plugins/Experimental/Avalanche/Source/AvalancheMediaEditor/Private/Playback/Graph/Pins/Slate/SAvaPlaybackGraphPin_Channel.cpp
new file mode 100644
--- /dev/null
+++ b/Engine/Plugins/Experimental/Avalanche/Source/AvalancheMediaEditor/Private/Playback/Graph/Pins/Slate/SAvaPlaybackGraphPin_Channel.cpp
@@ -0,0 +1,81 @@
+// Copyright Epic Games, Inc. All Rights Reserved.
+
+#include "SAvaPlaybackGraphPin_Channel.h"
+
+#include "AvaMediaEditorStyle.h"
+#include "AvaMediaEditorUtils.h"
+#include "AvalancheBroadcast.h"
+#include "Playback/Graph/AvaPlaybackGraphSchema.h"
+#include "Widgets/Images/SImage.h"
+#include "Widgets/SBoxPanel.h"
+
+void SAvaPlaybackGraphPin_Channel::Construct(const FArguments& InArgs, UEdGraphPin* InPin)
+{
+	SGraphPin::Construct(SGraphPin::FArguments()
+			.UsePinColorForText(true)
+		, InPin);
+	
+	UpdateChannelState();
+}
+
+const FAvaOutputChannel* SAvaPlaybackGraphPin_Channel::GetChannel() const
+{
+	if (const UEdGraphPin* const ChannelPin = GetPinObj())
+	{
+		const FAvaOutputChannel& Channel = UAvalancheBroadcast::Get().GetCurrentProfile().GetChannel(ChannelPin->GetFName());
+		if (Channel.IsValidChannel())
+		{
+			return &Channel;
+		}
+	}
+	return nullptr;
+}
+
+void SAvaPlaybackGraphPin_Channel::UpdateChannelState()
+{
+	bPinEnabled = false;
+	
+	if (const FAvaOutputChannel* const Channel = GetChannel())
+	{
+		bPinEnabled = true;
+		ChannelStateText = FAvaMediaEditorUtils::GetChannelStatusText(Channel->GetState(), Channel->GetIssueSeverity());
+		ChannelStateBrush = FAvaMediaEditorUtils::GetChannelStatusBrush(Channel->GetState(), Channel->GetIssueSeverity());
+	}
+}
+
+TSharedRef<SWidget> SAvaPlaybackGraphPin_Channel::GetDefaultValueWidget()
+{
+	return SNew(SHorizontalBox)
+		+SHorizontalBox::Slot()
+		.AutoWidth()
+		[
+			SNew(SImage)
+			.Image(this, &SAvaPlaybackGraphPin_Channel::GetChannelStatusBrush)
+		]
+		+SHorizontalBox::Slot()
+		.Padding(2.f, 0.f, 0.f, 0.f)
+		.HAlign(EHorizontalAlignment::HAlign_Left)
+		.VAlign(EVerticalAlignment::VAlign_Center)
+		.AutoWidth()
+		[
+			SNew(STextBlock)
+			.Text(this, &SAvaPlaybackGraphPin_Channel::GetChannelStatusText)
+		];
+}
+
+FText SAvaPlaybackGraphPin_Channel::GetChannelStatusText() const
+{
+	return ChannelStateText;
+}
+
+const FSlateBrush* SAvaPlaybackGraphPin_Channel::GetChannelStatusBrush() const
+{
+	return ChannelStateBrush;
+}
+
+FSlateColor SAvaPlaybackGraphPin_Channel::GetPinColor() const
+{
+	return bPinEnabled
+		? UAvaPlaybackGraphSchema::ActivePinColor
+		: UAvaPlaybackGraphSchema::InactivePinColor;
+}	
diff --git a/Engine/Plugins/Experimental/Avalanche/Source/AvalancheMediaEditor/Private/Playback/Graph/Pins/Slate/SAvaPlaybackGraphPin_Channel.h b/Engine/Plugins/Experimental/Avalanche/Source/AvalancheMediaEditor/Private/Playback/Graph/Pins/Slate/SAvaPlaybackGraphPin_Channel.h
new file mode 100644
--- /dev/null
+++ b/Engine/Plugins/Experimental/Avalanche/Source/AvalancheMediaEditor/Private/Playback/Graph/Pins/Slate/SAvaPlaybackGraphPin_Channel.h
@@ -0,0 +1,39 @@
+// Copyright Epic Games, Inc. All Rights Reserved.
+
+#pragma once
+
+#include "CoreMinimal.h"
+#include "SGraphPin.h"
+
+struct FAvaOutputChannel;
+
+class SAvaPlaybackGraphPin_Channel : public SGraphPin
+{
+public:
+
+	SLATE_BEGIN_ARGS(SAvaPlaybackGraphPin_Channel){}
+	SLATE_END_ARGS()
+	
+	void Construct(const FArguments& InArgs, UEdGraphPin* InPin);
+	
+	const FAvaOutputChannel* GetChannel() const;
+	
+	void UpdateChannelState();
+	
+	/** Build the widget we should put into the 'default value' space, shown when nothing connected */
+	virtual TSharedRef<SWidget>	GetDefaultValueWidget() override;
+
+	FText GetChannelStatusText() const;
+	
+	const FSlateBrush* GetChannelStatusBrush() const;
+
+	virtual FSlateColor GetPinColor() const override;
+	
+protected:
+
+	const FSlateBrush* ChannelStateBrush = nullptr;
+	
+	FText ChannelStateText;
+	
+	bool bPinEnabled = true;
+};
diff --git a/Engine/Plugins/Experimental/Avalanche/Source/AvalancheMediaEditor/Private/Playback/Graph/SchemaActions/AvaPlaybackAction_PasteNode.cpp b/Engine/Plugins/Experimental/Avalanche/Source/AvalancheMediaEditor/Private/Playback/Graph/SchemaActions/AvaPlaybackAction_PasteNode.cpp
new file mode 100644
--- /dev/null
+++ b/Engine/Plugins/Experimental/Avalanche/Source/AvalancheMediaEditor/Private/Playback/Graph/SchemaActions/AvaPlaybackAction_PasteNode.cpp
@@ -0,0 +1,20 @@
+// Copyright Epic Games, Inc. All Rights Reserved.
+
+#include "Playback/Graph/SchemaActions/AvaPlaybackAction_PasteNode.h"
+#include "Playback/AvalanchePlayback.h"
+#include "Playback/Graph/AvaPlaybackGraph.h"
+#include "Playback/IAvaPlaybackEditor.h"
+
+UEdGraphNode* FAvaPlaybackAction_PasteNode::PerformAction(UEdGraph* ParentGraph, UEdGraphPin* FromPin
+	, const FVector2D Location, bool bSelectNewNode)
+{
+	UAvalanchePlayback* const Playback = CastChecked<UAvaPlaybackGraph>(ParentGraph)->GetPlayback();
+	check(Playback);
+	
+	if (TSharedPtr<IAvaPlaybackGraphEditor> GraphEditor = Playback->GetGraphEditor())
+	{
+		GraphEditor->PasteNodesHere(Location);
+	}
+	
+	return nullptr;
+}
diff --git a/Engine/Plugins/Experimental/Avalanche/Source/AvalancheMediaEditor/Private/Playback/TabFactories/AvaPlaybackDetailsTabFactory.cpp b/Engine/Plugins/Experimental/Avalanche/Source/AvalancheMediaEditor/Private/Playback/TabFactories/AvaPlaybackDetailsTabFactory.cpp
new file mode 100644
--- /dev/null
+++ b/Engine/Plugins/Experimental/Avalanche/Source/AvalancheMediaEditor/Private/Playback/TabFactories/AvaPlaybackDetailsTabFactory.cpp
@@ -0,0 +1,31 @@
+// Copyright Epic Games, Inc. All Rights Reserved.
+
+#include "AvaPlaybackDetailsTabFactory.h"
+#include "Playback/AvalanchePlayback.h"
+#include "Playback/AvaPlaybackEditor.h"
+#include "Playback/DetailsView/SAvaPlaybackDetailsView.h"
+
+const FName FAvaPlaybackDetailsTabFactory::TabID(TEXT("AvalanchePlaybackDetails"));
+
+#define LOCTEXT_NAMESPACE "AvaPlaybackDetailsTabFactory"
+
+FAvaPlaybackDetailsTabFactory::FAvaPlaybackDetailsTabFactory(const TSharedPtr<FAvaPlaybackEditor>& InPlaybackEditor)
+	: FAvaPlaybackTabFactory(TabID, InPlaybackEditor)
+{
+	TabLabel = LOCTEXT("PlaybackDetails_TabLabel", "Details");
+	TabIcon = FSlateIcon(FAppStyle::GetAppStyleSetName(), "LevelEditor.Tabs.Details");
+
+	bIsSingleton = true;
+
+	ViewMenuDescription = LOCTEXT("PlaybackDetails_ViewMenu_Desc", "Playback Node Details");
+	ViewMenuTooltip = LOCTEXT("PlaybackDetails_ViewMenu_ToolTip", "Playback Node Details");
+}
+
+TSharedRef<SWidget> FAvaPlaybackDetailsTabFactory::CreateTabBody(const FWorkflowTabSpawnInfo& Info) const
+{
+	TSharedPtr<FAvaPlaybackEditor> PlaybackEditor = PlaybackEditorWeak.Pin();
+	check(PlaybackEditor.IsValid());
+	return SNew(SAvaPlaybackDetailsView, PlaybackEditor);
+}
+
+#undef LOCTEXT_NAMESPACE
diff --git a/Engine/Plugins/Experimental/Avalanche/Source/AvalancheMediaEditor/Private/Playback/TabFactories/AvaPlaybackGraphTabFactory.cpp b/Engine/Plugins/Experimental/Avalanche/Source/AvalancheMediaEditor/Private/Playback/TabFactories/AvaPlaybackGraphTabFactory.cpp
new file mode 100644
--- /dev/null
+++ b/Engine/Plugins/Experimental/Avalanche/Source/AvalancheMediaEditor/Private/Playback/TabFactories/AvaPlaybackGraphTabFactory.cpp
@@ -0,0 +1,30 @@
+// Copyright Epic Games, Inc. All Rights Reserved.
+
+#include "AvaPlaybackGraphTabFactory.h"
+#include "GraphEditor.h"
+#include "Playback/AvalanchePlayback.h"
+#include "Playback/AvaPlaybackEditor.h"
+
+const FName FAvaPlaybackGraphTabFactory::TabID(TEXT("AvalanchePlaybackGraph"));
+
+#define LOCTEXT_NAMESPACE "AvaPlaybackGraphTabFactory"
+
+FAvaPlaybackGraphTabFactory::FAvaPlaybackGraphTabFactory(const TSharedPtr<FAvaPlaybackEditor>& InPlaybackEditor)
+	: FAvaPlaybackTabFactory(TabID, InPlaybackEditor)
+{
+	TabLabel = LOCTEXT("PlaybackGraph_TabLabel", "Playback Graph");
+	TabIcon = FSlateIcon(FAppStyle::GetAppStyleSetName(), "GraphEditor.EventGraph_16x");
+
+	bIsSingleton = true;
+
+	ViewMenuDescription = LOCTEXT("PlaybackGraph_ViewMenu_Desc", "Playback Graph");
+	ViewMenuTooltip = LOCTEXT("PlaybackGraph_ViewMenu_ToolTip", "Playback Graph");
+}
+
+TSharedRef<SWidget> FAvaPlaybackGraphTabFactory::CreateTabBody(const FWorkflowTabSpawnInfo& Info) const
+{
+	TSharedRef<FAvaPlaybackEditor> PlaybackEditor = PlaybackEditorWeak.Pin().ToSharedRef();	
+	return PlaybackEditor->CreateGraphEditor();
+}
+
+#undef LOCTEXT_NAMESPACE
diff --git a/Engine/Plugins/Experimental/Avalanche/Source/AvalancheMediaEditor/Private/Playback/TabFactories/AvaPlaybackGraphTabFactory.h b/Engine/Plugins/Experimental/Avalanche/Source/AvalancheMediaEditor/Private/Playback/TabFactories/AvaPlaybackGraphTabFactory.h
new file mode 100644
--- /dev/null
+++ b/Engine/Plugins/Experimental/Avalanche/Source/AvalancheMediaEditor/Private/Playback/TabFactories/AvaPlaybackGraphTabFactory.h
@@ -0,0 +1,15 @@
+// Copyright Epic Games, Inc. All Rights Reserved.
+
+#pragma once
+
+#include "AvaPlaybackTabFactory.h"
+
+class FAvaPlaybackGraphTabFactory : public FAvaPlaybackTabFactory
+{
+public:
+	static const FName TabID;
+	
+	FAvaPlaybackGraphTabFactory(const TSharedPtr<FAvaPlaybackEditor>& InPlaybackEditor);
+	
+	virtual TSharedRef<SWidget> CreateTabBody(const FWorkflowTabSpawnInfo& Info) const override;
+};
diff --git a/Engine/Plugins/Experimental/Avalanche/Source/AvalancheMediaEditor/Private/Playback/TabFactories/AvaPlaybackTabFactory.cpp b/Engine/Plugins/Experimental/Avalanche/Source/AvalancheMediaEditor/Private/Playback/TabFactories/AvaPlaybackTabFactory.cpp
new file mode 100644
--- /dev/null
+++ b/Engine/Plugins/Experimental/Avalanche/Source/AvalancheMediaEditor/Private/Playback/TabFactories/AvaPlaybackTabFactory.cpp
@@ -0,0 +1,10 @@
+// Copyright Epic Games, Inc. All Rights Reserved.
+
+#include "AvaPlaybackTabFactory.h"
+#include "Playback/AvaPlaybackEditor.h"
+
+FAvaPlaybackTabFactory::FAvaPlaybackTabFactory(const FName& InTabID, const TSharedPtr<FAvaPlaybackEditor>& InPlaybackEditor)
+	: FWorkflowTabFactory(InTabID, InPlaybackEditor)
+	, PlaybackEditorWeak(InPlaybackEditor)
+{
+}
diff --git a/Engine/Plugins/Experimental/Avalanche/Source/AvalancheMediaEditor/Private/Playback/TabFactories/AvaPlaybackTabFactory.h b/Engine/Plugins/Experimental/Avalanche/Source/AvalancheMediaEditor/Private/Playback/TabFactories/AvaPlaybackTabFactory.h
new file mode 100644
--- /dev/null
+++ b/Engine/Plugins/Experimental/Avalanche/Source/AvalancheMediaEditor/Private/Playback/TabFactories/AvaPlaybackTabFactory.h
@@ -0,0 +1,19 @@
+// Copyright Epic Games, Inc. All Rights Reserved.
+
+#pragma once
+
+#include "WorkflowOrientedApp/WorkflowTabFactory.h"
+
+class FAvaPlaybackEditor;
+
+//Base class for all Tab Factories in Ava Playback Editor
+class FAvaPlaybackTabFactory : public FWorkflowTabFactory
+{
+public:
+	
+	FAvaPlaybackTabFactory(const FName& InTabID, const TSharedPtr<FAvaPlaybackEditor>& InPlaybackEditor);
+
+protected:
+
+	TWeakPtr<FAvaPlaybackEditor> PlaybackEditorWeak;
+};
diff --git a/Engine/Plugins/Experimental/Avalanche/Source/AvalancheMediaEditor/Private/Playlist/AppModes/AvaPlaylistAppMode.cpp b/Engine/Plugins/Experimental/Avalanche/Source/AvalancheMediaEditor/Private/Playlist/AppModes/AvaPlaylistAppMode.cpp
new file mode 100644
--- /dev/null
+++ b/Engine/Plugins/Experimental/Avalanche/Source/AvalancheMediaEditor/Private/Playlist/AppModes/AvaPlaylistAppMode.cpp
@@ -0,0 +1,49 @@
+// Copyright Epic Games, Inc. All Rights Reserved.
+
+#include "AvaPlaylistAppMode.h"
+#include "Playlist/AvaPlaylistEditor.h"
+
+#define LOCTEXT_NAMESPACE "AvaPlaylistAppMode"
+
+const FName FAvaPlaylistAppMode::DefaultMode(TEXT("DefaultMode"));
+
+FAvaPlaylistAppMode::FAvaPlaylistAppMode(const TSharedPtr<FAvaPlaylistEditor>& InPlaylistEditor, const FName& InModeName)
+	: FApplicationMode(InModeName, FAvaPlaylistAppMode::GetLocalizedMode)
+	, PlaylistEditorWeak(InPlaylistEditor)
+{
+}
+
+void FAvaPlaylistAppMode::RegisterTabFactories(TSharedPtr<FTabManager> InTabManager)
+{
+	TSharedPtr<FAvaPlaylistEditor> PlaybackEditor = PlaylistEditorWeak.Pin();
+	PlaybackEditor->PushTabFactories(TabFactories);
+	FApplicationMode::RegisterTabFactories(InTabManager);
+}
+
+TSharedPtr<FDocumentTabFactory> FAvaPlaylistAppMode::GetDocumentTabFactory(const FName& InName) const
+{
+	if (const TSharedRef<FDocumentTabFactory>* DocFactory = DocumentTabFactories.Find(InName))
+	{
+		return *DocFactory;
+	}
+
+	return nullptr;
+}
+
+FText FAvaPlaylistAppMode::GetLocalizedMode(const FName InMode)
+{
+	static TMap<FName, FText> LocModes;
+
+	if (LocModes.Num() == 0)
+	{
+		LocModes.Add(DefaultMode, LOCTEXT("Playlist_DefaultMode", "Default"));
+	}
+
+	check(InMode != NAME_None);
+	const FText* OutDesc = LocModes.Find(InMode);
+	check(OutDesc);
+	
+	return *OutDesc;
+}
+
+#undef LOCTEXT_NAMESPACE
diff --git a/Engine/Plugins/Experimental/Avalanche/Source/AvalancheMediaEditor/Private/Playlist/AppModes/AvaPlaylistAppMode.h b/Engine/Plugins/Experimental/Avalanche/Source/AvalancheMediaEditor/Private/Playlist/AppModes/AvaPlaylistAppMode.h
new file mode 100644
--- /dev/null
+++ b/Engine/Plugins/Experimental/Avalanche/Source/AvalancheMediaEditor/Private/Playlist/AppModes/AvaPlaylistAppMode.h
@@ -0,0 +1,33 @@
+// Copyright Epic Games, Inc. All Rights Reserved.
+
+#pragma once
+
+#include "WorkflowOrientedApp/ApplicationMode.h"
+#include "WorkflowOrientedApp/WorkflowTabManager.h"
+
+class FAvaPlaylistEditor;
+
+class FAvaPlaylistAppMode : public FApplicationMode
+{
+public:
+
+	// Mode constants
+	static const FName DefaultMode;
+	//Add more here
+	
+	FAvaPlaylistAppMode(const TSharedPtr<FAvaPlaylistEditor>& InPlaylistEditor, const FName& InModeName);
+	virtual void RegisterTabFactories(TSharedPtr<FTabManager> InTabManager) override;
+
+	TSharedPtr<FDocumentTabFactory> GetDocumentTabFactory(const FName& InName) const;
+
+protected:
+	
+	static FText GetLocalizedMode(const FName InMode);
+
+	TWeakPtr<FAvaPlaylistEditor> PlaylistEditorWeak;
+	
+	// Set of spawnable tabs in the mode
+	FWorkflowAllowedTabSet TabFactories;
+
+	TMap<FName, TSharedRef<FDocumentTabFactory>> DocumentTabFactories;
+};
