[MotionDesign]

diff --git a/Engine/Plugins/Experimental/Avalanche/Source/AvalancheMediaEditor/Private/Playlist/AppModes/AvaPlaylistDefaultMode.cpp b/Engine/Plugins/Experimental/Avalanche/Source/AvalancheMediaEditor/Private/Playlist/AppModes/AvaPlaylistDefaultMode.cpp
new file mode 100644
--- /dev/null
+++ b/Engine/Plugins/Experimental/Avalanche/Source/AvalancheMediaEditor/Private/Playlist/AppModes/AvaPlaylistDefaultMode.cpp
@@ -0,0 +1,134 @@
+// Copyright Epic Games, Inc. All Rights Reserved.
+
+#include "AvaPlaylistDefaultMode.h"
+#include "IAvaMediaEditorModule.h"
+#include "Playlist/AvaPlaylistEditor.h"
+#include "Playlist/TabFactories/AvaChannelLayerStatusListTabFactory.h"
+#include "Playlist/TabFactories/AvaChannelStatusListTabFactory.h"
+#include "Playlist/TabFactories/AvaInstancedPageListTabFactory.h"
+#include "Playlist/TabFactories/AvaPageDetailsTabFactory.h"
+#include "Playlist/TabFactories/AvaPageViewerTabFactory.h"
+#include "Playlist/TabFactories/AvaShowControlTabFactory.h"
+#include "Playlist/TabFactories/AvaSubListDocumentTabFactory.h"
+#include "Playlist/TabFactories/AvaSubListTabFactory.h"
+#include "Playlist/TabFactories/AvaTemplatePageListTabFactory.h"
+
+#define LOCTEXT_NAMESPACE "AvaPlaylistDefaultMode"
+
+FAvaPlaylistDefaultMode::FAvaPlaylistDefaultMode(const TSharedPtr<FAvaPlaylistEditor>& InPlaylistEditor)
+	: FAvaPlaylistAppMode(InPlaylistEditor, FAvaPlaylistAppMode::DefaultMode)
+{
+	WorkspaceMenuCategory = FWorkspaceItem::NewGroup(LOCTEXT("WorkspaceMenu_AvaPlaylist", "Motion Design Playlist"));
+	
+	check(InPlaylistEditor.IsValid());
+	TabLayout = FTabManager::NewLayout("AvalanchePlaylistEditor_Default_Layout_V2.4")
+		->AddArea
+		(
+			FTabManager::NewPrimaryArea()
+			->SetOrientation(Orient_Horizontal)
+			->Split
+			(
+				FTabManager::NewSplitter()
+				->SetSizeCoefficient(0.5f)
+				->SetOrientation(Orient_Vertical)
+				->Split
+				(
+					FTabManager::NewStack()
+					->SetSizeCoefficient(0.25f)
+					->AddTab(FAvaTemplatePageListTabFactory::TabID, ETabState::OpenedTab)
+					->SetForegroundTab(FAvaTemplatePageListTabFactory::TabID)
+					->SetHideTabWell(true)
+				)
+				->Split
+				(
+					FTabManager::NewStack()
+					->SetSizeCoefficient(0.05f)
+					->AddTab(FAvaShowControlTabFactory::TabID, ETabState::OpenedTab)
+					->SetForegroundTab(FAvaShowControlTabFactory::TabID)
+					->SetHideTabWell(true)
+				)
+				->Split
+				(
+					FTabManager::NewStack()
+					->SetSizeCoefficient(0.35f)
+					->AddTab(FAvaInstancedPageListTabFactory::TabID, ETabState::OpenedTab)
+					->SetForegroundTab(FAvaInstancedPageListTabFactory::TabID)
+					->SetHideTabWell(false)
+				)
+				->Split
+				(
+					FTabManager::NewStack()
+					->SetSizeCoefficient(0.35f)
+					->AddTab(FAvaSubListTabFactory::TabID, ETabState::ClosedTab)
+					->SetForegroundTab(FAvaSubListTabFactory::TabID)
+					->SetHideTabWell(false)
+				)
+			)
+			->Split
+			(
+				FTabManager::NewSplitter()
+				->SetSizeCoefficient(0.60f)
+				->SetOrientation(Orient_Vertical)
+				->Split
+				(
+					FTabManager::NewSplitter()
+					->SetSizeCoefficient(0.3f)
+					->SetOrientation(Orient_Vertical)
+					->Split
+					(
+						FTabManager::NewStack()
+						->SetSizeCoefficient(0.7f)
+						->AddTab(FAvaPageDetailsTabFactory::TabID, ETabState::OpenedTab)
+						->SetForegroundTab(FAvaPageDetailsTabFactory::TabID)
+						->SetHideTabWell(false)
+					)
+					->Split
+					(
+						FTabManager::NewStack()
+						->SetSizeCoefficient(0.3f)
+						->AddTab(FAvaChannelLayerStatusListTabFactory::TabID, ETabState::OpenedTab)
+						->SetForegroundTab(FAvaChannelLayerStatusListTabFactory::TabID)
+						->SetHideTabWell(false)
+					)
+				)
+				->Split
+				(
+					FTabManager::NewStack()
+					->SetSizeCoefficient(0.02f)
+					->AddTab(FAvaChannelStatusListTabFactory::TabID, ETabState::OpenedTab)
+					->SetForegroundTab(FAvaChannelStatusListTabFactory::TabID)
+					->SetHideTabWell(true)
+				)
+				->Split
+				(
+					FTabManager::NewStack()
+					->SetSizeCoefficient(0.6f)
+					->AddTab(FAvaPageViewerTabFactory::TabID, ETabState::OpenedTab)
+					->SetForegroundTab(FAvaPageViewerTabFactory::TabID)
+					->SetHideTabWell(false)
+				)
+			)
+		);
+
+	// Add Tab Spawners
+	TabFactories.RegisterFactory(MakeShared<FAvaTemplatePageListTabFactory>(InPlaylistEditor));
+	TabFactories.RegisterFactory(MakeShared<FAvaInstancedPageListTabFactory>(InPlaylistEditor));
+	TabFactories.RegisterFactory(MakeShared<FAvaPageDetailsTabFactory>(InPlaylistEditor));
+	TabFactories.RegisterFactory(MakeShared<FAvaShowControlTabFactory>(InPlaylistEditor));
+	TabFactories.RegisterFactory(MakeShared<FAvaSubListTabFactory>(InPlaylistEditor));	
+	TabFactories.RegisterFactory(MakeShared<FAvaPageViewerTabFactory>(InPlaylistEditor));
+	TabFactories.RegisterFactory(MakeShared<FAvaChannelStatusListTabFactory>(InPlaylistEditor));
+	TabFactories.RegisterFactory(MakeShared<FAvaChannelLayerStatusListTabFactory>(InPlaylistEditor));
+
+	DocumentTabFactories.Emplace(FAvaSubListDocumentTabFactory::FactoryId, MakeShared<FAvaSubListDocumentTabFactory>(InPlaylistEditor));
+
+	// Make sure we start with our existing list of extenders instead of creating a new one
+	IAvaMediaEditorModule& AvaMediaEditorModule = IAvaMediaEditorModule::Get();
+	ToolbarExtender = AvaMediaEditorModule.GetPlaylistToolBarExtensibilityManager()->GetAllExtenders(
+		InPlaylistEditor->GetToolkitCommands(),
+		*InPlaylistEditor->GetObjectsCurrentlyBeingEdited()
+	);
+	InPlaylistEditor->ExtendToolBar(ToolbarExtender);
+}
+
+#undef LOCTEXT_NAMESPACE
diff --git a/Engine/Plugins/Experimental/Avalanche/Source/AvalancheMediaEditor/Private/Playlist/AppModes/AvaPlaylistDefaultMode.h b/Engine/Plugins/Experimental/Avalanche/Source/AvalancheMediaEditor/Private/Playlist/AppModes/AvaPlaylistDefaultMode.h
new file mode 100644
--- /dev/null
+++ b/Engine/Plugins/Experimental/Avalanche/Source/AvalancheMediaEditor/Private/Playlist/AppModes/AvaPlaylistDefaultMode.h
@@ -0,0 +1,11 @@
+// Copyright Epic Games, Inc. All Rights Reserved.
+
+#pragma once
+
+#include "AvaPlaylistAppMode.h"
+
+class FAvaPlaylistDefaultMode : public FAvaPlaylistAppMode
+{
+public:
+	FAvaPlaylistDefaultMode(const TSharedPtr<FAvaPlaylistEditor>& InPlaylistEditor);
+};
diff --git a/Engine/Plugins/Experimental/Avalanche/Source/AvalancheMediaEditor/Private/Playlist/AssetDefinition_AvalanchePlaylist.cpp b/Engine/Plugins/Experimental/Avalanche/Source/AvalancheMediaEditor/Private/Playlist/AssetDefinition_AvalanchePlaylist.cpp
new file mode 100644
--- /dev/null
+++ b/Engine/Plugins/Experimental/Avalanche/Source/AvalancheMediaEditor/Private/Playlist/AssetDefinition_AvalanchePlaylist.cpp
@@ -0,0 +1,191 @@
+// Copyright Epic Games, Inc. All Rights Reserved.
+
+#include "AssetDefinition_AvalanchePlaylist.h"
+
+#include "AvaMediaEditorStyle.h"
+#include "ContentBrowserMenuContexts.h"
+#include "Misc/MessageDialog.h"
+#include "Playlist/AvaPlaylistEditor.h"
+#include "Playlist/AvaPlaylistEditorUtils.h"
+#include "Playlist/AvalanchePlaylist.h"
+
+#define LOCTEXT_NAMESPACE "AssetDefinition_AvalanchePlaylist"
+
+FText UAssetDefinition_AvalanchePlaylist::GetAssetDisplayName() const
+{
+	return LOCTEXT("AvaPlaylistAction_Name", "Motion Design Playlist");
+}
+
+TSoftClassPtr<UObject> UAssetDefinition_AvalanchePlaylist::GetAssetClass() const
+{
+	return UAvalanchePlaylist::StaticClass();
+}
+
+FLinearColor UAssetDefinition_AvalanchePlaylist::GetAssetColor() const
+{
+	static const FName RundownAssetColorName(TEXT("AvalancheMediaEditor.AssetColors.Rundown"));
+	return FAvaMediaEditorStyle::Get().GetColor(RundownAssetColorName);
+}
+
+TConstArrayView<FAssetCategoryPath> UAssetDefinition_AvalanchePlaylist::GetAssetCategories() const
+{
+	static const TArray<FAssetCategoryPath> Categories = {EAssetCategoryPaths::Blueprint};
+	return Categories;
+}
+
+EAssetCommandResult UAssetDefinition_AvalanchePlaylist::OpenAssets(const FAssetOpenArgs& OpenArgs) const
+{
+	EAssetCommandResult CommandResult = EAssetCommandResult::Unhandled;
+	if (OpenArgs.OpenMethod == EAssetOpenMethod::Edit)
+	{
+		for (UAvalanchePlaylist* Playlist : OpenArgs.LoadObjects<UAvalanchePlaylist>())
+		{
+			if (Playlist)
+			{
+				const TSharedRef<FAvaPlaylistEditor> PlaylistEditor = MakeShared<FAvaPlaylistEditor>();
+				PlaylistEditor->InitPlaylistEditor(EToolkitMode::Standalone, OpenArgs.ToolkitHost, Playlist);
+				CommandResult = EAssetCommandResult::Handled;
+			}
+		}
+	}
+	return CommandResult;
+}
+
+namespace MenuExtension_AvalanchePlaylist
+{
+	static bool CanExportToJson(const FToolMenuContext& InContext)
+	{
+		return true;
+	}
+	
+	static void ExportToJson(const FToolMenuContext& InContext)
+	{
+		const UContentBrowserAssetContextMenuContext* Context = UContentBrowserAssetContextMenuContext::FindContextWithAssets(InContext);
+
+		for (UAvalanchePlaylist* Playlist : Context->LoadSelectedObjects<UAvalanchePlaylist>())
+		{
+			if (Playlist)
+			{
+				if (Playlist->GetOutermost() && Playlist->GetOutermost()->HasAnyPackageFlags(PKG_DisallowExport))
+				{
+					UE_LOG(LogAvaPlaylist, Error, TEXT("Package disallow export."));
+					continue;
+				}
+
+				using namespace UE::AvaPlaylistEditor::Utils;
+				FString ExportFilename = GetExportFilepath(Playlist, TEXT("json file"), TEXT("json"));
+				if (!ExportFilename.IsEmpty())
+				{
+					SavePlaylistToJson(Playlist, *ExportFilename);
+				}
+			}
+		}
+	}
+
+	static bool CanImportFromJson(const FToolMenuContext& InContext)
+	{
+		const UContentBrowserAssetContextMenuContext* Context = UContentBrowserAssetContextMenuContext::FindContextWithAssets(InContext);
+		for (const UAvalanchePlaylist* Playlist : Context->GetSelectedObjectsInMemory<UAvalanchePlaylist>())
+		{
+			if (Playlist && Playlist->IsPlaying())
+			{
+				return false;
+			}
+		}
+		return true;
+	}
+	
+	static void ImportFromJson(const FToolMenuContext& InContext)
+	{
+		const UContentBrowserAssetContextMenuContext* Context = UContentBrowserAssetContextMenuContext::FindContextWithAssets(InContext);
+
+		for (UAvalanchePlaylist* Playlist : Context->LoadSelectedObjects<UAvalanchePlaylist>())
+		{
+			if (Playlist)
+			{
+				// Ask user to confirm stopping the playlist if it is playing.
+				if (Playlist->IsPlaying())
+				{
+					const FText MessageText = LOCTEXT("StopPagesOnImportQuestion",
+						"All pages must be stopped before importing. Some pages are still playing, do you want to stop all pages?");
+			
+					const EAppReturnType::Type Reply = FMessageDialog::Open(EAppMsgType::YesNoCancel, EAppReturnType::Yes, MessageText);
+
+					if (Reply == EAppReturnType::Cancel)
+					{
+						return;	// cancel the whole operation.
+					}
+
+					if (Reply == EAppReturnType::No)
+					{
+						UE_LOG(LogAvaPlaylist, Warning, TEXT("Skipping import of playlist \"%s\""), *Playlist->GetFullName());
+						continue;
+					}
+					
+					if (Reply == EAppReturnType::Yes)
+					{
+						Playlist->ClosePlaybackContext(true);
+					}
+				}
+				
+				// Ask user to confirm stomping the existing playlist.
+				if (!Playlist->IsEmpty())
+				{
+					const FText MessageText = LOCTEXT("ClearPlaylistOnImportQuestion",
+						"The playlist is not empty, all existing content will be overwritten. Are you sure?");
+			
+					const EAppReturnType::Type Reply = FMessageDialog::Open(EAppMsgType::YesNoCancel, EAppReturnType::Yes, MessageText);
+					if (Reply == EAppReturnType::Cancel)
+					{
+						return;
+					}
+					if (Reply == EAppReturnType::No)
+					{
+						continue;
+					}
+				}
+				
+				using namespace UE::AvaPlaylistEditor::Utils;
+				FString ImportFilename = GetImportFilepath(TEXT("json file"), TEXT("json"));
+				if (!ImportFilename.IsEmpty())
+				{
+					LoadPlaylistFromJson(Playlist, *ImportFilename);
+				}
+			}
+		}
+	}
+
+	static FDelayedAutoRegisterHelper DelayedAutoRegister(EDelayedRegisterRunPhase::EndOfEngineInit, []
+	{
+		UToolMenus::RegisterStartupCallback(FSimpleMulticastDelegate::FDelegate::CreateLambda([]()
+		{
+			FToolMenuOwnerScoped OwnerScoped(UE_MODULE_NAME);
+			UToolMenu* Menu = UE::ContentBrowser::ExtendToolMenu_AssetContextMenu(UAvalanchePlaylist::StaticClass());
+
+			FToolMenuSection& Section = Menu->FindOrAddSection("GetAssetActions");
+			Section.AddDynamicEntry(NAME_None, FNewToolMenuSectionDelegate::CreateLambda([](FToolMenuSection& InSection)
+			{
+				{
+					const TAttribute<FText> Label = LOCTEXT("AvaPlaylist_ExportToJson", "Export to Json");
+					const TAttribute<FText> ToolTip = LOCTEXT("AvaPlaylist_ExportToJsonTooltip", "Export To Json");
+
+					FToolUIAction UIAction;
+					UIAction.ExecuteAction = FToolMenuExecuteAction::CreateStatic(&ExportToJson);
+					UIAction.CanExecuteAction = FToolMenuCanExecuteAction::CreateStatic(&CanExportToJson);
+					InSection.AddMenuEntry("AvaPlaylist_ExportToJson", Label, ToolTip, FSlateIcon(), UIAction);
+				}
+				{
+					const TAttribute<FText> Label = LOCTEXT("AvaPlaylist_ImportFromJson", "Import from Json");
+					const TAttribute<FText> ToolTip = LOCTEXT("AvaPlaylist_ImportFromJsonTooltip", "Import from Json");
+
+					FToolUIAction UIAction;
+					UIAction.ExecuteAction = FToolMenuExecuteAction::CreateStatic(&ImportFromJson);
+					UIAction.CanExecuteAction = FToolMenuCanExecuteAction::CreateStatic(&CanImportFromJson);
+					InSection.AddMenuEntry("AvaPlaylist_ImportFromJson", Label, ToolTip, FSlateIcon(), UIAction);
+				}
+			}));
+		}));
+	});
+}
+
+#undef LOCTEXT_NAMESPACE
diff --git a/Engine/Plugins/Experimental/Avalanche/Source/AvalancheMediaEditor/Private/Playlist/AssetDefinition_AvalanchePlaylist.h b/Engine/Plugins/Experimental/Avalanche/Source/AvalancheMediaEditor/Private/Playlist/AssetDefinition_AvalanchePlaylist.h
new file mode 100644
--- /dev/null
+++ b/Engine/Plugins/Experimental/Avalanche/Source/AvalancheMediaEditor/Private/Playlist/AssetDefinition_AvalanchePlaylist.h
@@ -0,0 +1,21 @@
+// Copyright Epic Games, Inc. All Rights Reserved.
+
+#pragma once
+
+#include "AssetDefinitionDefault.h"
+#include "AssetDefinition_AvalanchePlaylist.generated.h"
+
+UCLASS()
+class UAssetDefinition_AvalanchePlaylist : public UAssetDefinitionDefault
+{
+	GENERATED_BODY()
+
+public:
+	//~ Begin UAssetDefinition
+	virtual FText GetAssetDisplayName() const override;
+	virtual TSoftClassPtr<UObject> GetAssetClass() const override;
+	virtual FLinearColor GetAssetColor() const override;
+	virtual TConstArrayView<FAssetCategoryPath> GetAssetCategories() const override;
+	virtual EAssetCommandResult OpenAssets(const FAssetOpenArgs& OpenArgs) const override;
+	//~ End UAssetDefinition
+};
diff --git a/Engine/Plugins/Experimental/Avalanche/Source/AvalancheMediaEditor/Private/Playlist/AvaPlaylistCommands.cpp b/Engine/Plugins/Experimental/Avalanche/Source/AvalancheMediaEditor/Private/Playlist/AvaPlaylistCommands.cpp
new file mode 100644
--- /dev/null
+++ b/Engine/Plugins/Experimental/Avalanche/Source/AvalancheMediaEditor/Private/Playlist/AvaPlaylistCommands.cpp
@@ -0,0 +1,148 @@
+// Copyright Epic Games, Inc. All Rights Reserved.
+
+#include "Playlist/AvaPlaylistCommands.h"
+
+#define LOCTEXT_NAMESPACE "AvaPlaylistCommands"
+
+void FAvaPlaylistCommands::RegisterCommands()
+{
+	UI_COMMAND(AddTemplate
+		, "Add Template"
+		, "Adds a new Template to the Template List"
+		, EUserInterfaceActionType::Button
+		, FInputChord(EKeys::Insert));
+	
+	UI_COMMAND(CreatePageInstanceFromTemplate
+		, "Create Page(s)"
+		, "Adds a new copy of a Template into the main Page List"
+		, EUserInterfaceActionType::Button
+		, FInputChord(EModifierKey::Control, EKeys::Insert));
+
+	UI_COMMAND(CreateComboTemplate
+		, "Create Combo Template"
+		, "Create a combination template that contains all selected templates. This is only relevant for transition logic templates. All templates that are to be merged into Combination Templates must be in different transition logic layers"
+		, EUserInterfaceActionType::Button
+		, FInputChord());
+
+	UI_COMMAND(RemovePage
+		, "Remove Pages"
+		, "Removes the Selected Pages from the Page List"
+		, EUserInterfaceActionType::Button
+		, FInputChord(EKeys::Delete));
+
+	UI_COMMAND(RenumberPage
+		, "Renumber Page"
+		, "Renumbers the Id of the Selected Page"
+		, EUserInterfaceActionType::Button
+		, FInputChord(EKeys::F3));
+	
+	UI_COMMAND(ReimportPage
+		, "Reimport Page"
+		, "Re-imports the selected page"
+		, EUserInterfaceActionType::Button
+		, FInputChord());
+
+	UI_COMMAND(EditPageSource
+		, "Edit Scene"
+		, "Edit source Avalanche Asset of the selected page"
+		, EUserInterfaceActionType::Button
+		, FInputChord());
+
+	UI_COMMAND(ExportPagesToPlaylist
+		, "Export To Playlist"
+		, "Export selected pages to a Playlist asset"
+		, EUserInterfaceActionType::Button
+		, FInputChord());
+
+	UI_COMMAND(ExportPagesToJson
+		, "Export To Json"
+		, "Export selected pages to an external json file"
+		, EUserInterfaceActionType::Button
+		, FInputChord());
+	
+	UI_COMMAND(ExportPagesToXml
+		, "Export To Xml"
+		, "Export selected pages to an external xml file"
+		, EUserInterfaceActionType::Button
+		, FInputChord());
+
+	UI_COMMAND(PreviewFrame
+		, "Preview Frame"
+		, "Runs the Selected Page without playing animations, but jumping to the Preview Frame set in each animation"
+		, EUserInterfaceActionType::Button
+		, FInputChord());
+
+	UI_COMMAND(PreviewPlay
+		, "Preview In"
+		, "Runs the Selected Page displayed in the Page Preview"
+		, EUserInterfaceActionType::Button
+		, FInputChord(EModifierKey::Control, EKeys::Enter));
+
+	UI_COMMAND(PreviewStop
+		, "Preview Out"
+		, "Stops the Page displayed in the Page Preview"
+		, EUserInterfaceActionType::Button
+		, FInputChord(EModifierKey::Control, EKeys::Multiply));
+
+	UI_COMMAND(PreviewForceStop
+		, "Force Preview Out"
+		, "Forcibly stops the Page displayed in the Page Preview regardless of it's current state."
+		, EUserInterfaceActionType::Button
+		, FInputChord());
+
+	UI_COMMAND(PreviewContinue
+		, "Continue"
+		, "Continues to play the Animation in the Current Page being Played if it has stopped"
+		, EUserInterfaceActionType::Button
+		, FInputChord(EModifierKey::Control, EKeys::Divide));
+
+	UI_COMMAND(PreviewPlayNext
+		, "Preview Next"
+		, "Previews the Next Page in the Page List"
+		, EUserInterfaceActionType::Button
+		, FInputChord(EModifierKey::Control, EKeys::Subtract));
+
+	UI_COMMAND(TakeToProgram
+		, "Take To Program"
+		, "Runs the current previewed page in the selected channel"
+		, EUserInterfaceActionType::Button
+		, FInputChord(EKeys::Enter));
+	
+	UI_COMMAND(Play
+		, "Take In"
+		, "Runs the Selected Page displayed in the selected channel"
+		, EUserInterfaceActionType::Button
+		, FInputChord(EModifierKey::Shift, EKeys::Enter));
+
+	UI_COMMAND(UpdateValues
+		, "Update Values"
+		, "Update Values of the Selected Page displayed in the selected channel"
+		, EUserInterfaceActionType::Button
+		, FInputChord(EKeys::F5));
+	
+	UI_COMMAND(Stop
+		, "Take Out"
+		, "Stops the Page displayed in the selected channel"
+		, EUserInterfaceActionType::Button
+		, FInputChord(EKeys::Multiply));
+
+	UI_COMMAND(ForceStop
+		, "Force Take Out"
+		, "Forcibly stops the Page displayed in the selected channel regardless of it's current state."
+		, EUserInterfaceActionType::Button
+		, FInputChord());
+
+	UI_COMMAND(Continue
+		, "Continue"
+		, "Continues to play the Animation in the Current Page being Played if it has stopped"
+		, EUserInterfaceActionType::Button
+		, FInputChord(EKeys::Divide));
+
+	UI_COMMAND(PlayNext
+		, "Take Next"
+		, "Takes the Next Page in the Page List"
+		, EUserInterfaceActionType::Button
+		, FInputChord(EKeys::Subtract));
+}
+
+#undef LOCTEXT_NAMESPACE
diff --git a/Engine/Plugins/Experimental/Avalanche/Source/AvalancheMediaEditor/Private/Playlist/AvaPlaylistCommands.h b/Engine/Plugins/Experimental/Avalanche/Source/AvalancheMediaEditor/Private/Playlist/AvaPlaylistCommands.h
new file mode 100644
--- /dev/null
+++ b/Engine/Plugins/Experimental/Avalanche/Source/AvalancheMediaEditor/Private/Playlist/AvaPlaylistCommands.h
@@ -0,0 +1,52 @@
+// Copyright Epic Games, Inc. All Rights Reserved.
+
+#pragma once
+
+#include "AvaMediaEditorStyle.h"
+#include "Framework/Commands/Commands.h"
+#include "Styling/AppStyle.h"
+
+class FAvaPlaylistCommands : public TCommands<FAvaPlaylistCommands>
+{
+public:
+	
+	FAvaPlaylistCommands()
+		: TCommands<FAvaPlaylistCommands>(TEXT("AvaPlaylistCommands")
+		, NSLOCTEXT("AvaPlaylistCommands", "AvaPlaylistCommands", "Motion Design Playlist")
+		, NAME_None
+		, FAvaMediaEditorStyle::GetStyleSetName())
+	{
+	}
+
+	/** Initialize commands */
+	virtual void RegisterCommands() override;
+
+	/** Page List Commands */
+	TSharedPtr<FUICommandInfo> AddTemplate;
+	TSharedPtr<FUICommandInfo> CreatePageInstanceFromTemplate;
+	TSharedPtr<FUICommandInfo> CreateComboTemplate;
+	TSharedPtr<FUICommandInfo> RemovePage;
+	TSharedPtr<FUICommandInfo> RenumberPage;
+	TSharedPtr<FUICommandInfo> ReimportPage;
+	TSharedPtr<FUICommandInfo> EditPageSource;
+	TSharedPtr<FUICommandInfo> ExportPagesToPlaylist;
+	TSharedPtr<FUICommandInfo> ExportPagesToJson;
+	TSharedPtr<FUICommandInfo> ExportPagesToXml;
+
+	/** Show Control / Playlist Commands */
+	TSharedPtr<FUICommandInfo> Play;
+	TSharedPtr<FUICommandInfo> UpdateValues;
+	TSharedPtr<FUICommandInfo> Stop;
+	TSharedPtr<FUICommandInfo> ForceStop;
+	TSharedPtr<FUICommandInfo> Continue;
+	TSharedPtr<FUICommandInfo> PlayNext;
+
+	/** Preview Control / Playlist Commands */
+	TSharedPtr<FUICommandInfo> PreviewFrame;
+	TSharedPtr<FUICommandInfo> PreviewPlay;
+	TSharedPtr<FUICommandInfo> PreviewStop;
+	TSharedPtr<FUICommandInfo> PreviewForceStop;
+	TSharedPtr<FUICommandInfo> PreviewContinue;
+	TSharedPtr<FUICommandInfo> PreviewPlayNext;
+	TSharedPtr<FUICommandInfo> TakeToProgram;
+};
diff --git a/Engine/Plugins/Experimental/Avalanche/Source/AvalancheMediaEditor/Private/Playlist/AvaPlaylistDefines.cpp b/Engine/Plugins/Experimental/Avalanche/Source/AvalancheMediaEditor/Private/Playlist/AvaPlaylistDefines.cpp
new file mode 100644
--- /dev/null
+++ b/Engine/Plugins/Experimental/Avalanche/Source/AvalancheMediaEditor/Private/Playlist/AvaPlaylistDefines.cpp
@@ -0,0 +1,7 @@
+// Copyright Epic Games, Inc. All Rights Reserved.
+
+#include "Playlist/AvaPlaylistDefines.h"
+
+const FNumberFormattingOptions UE::AvalanchePlaylist::FEditorMetrics::PageIdFormattingOptions = FNumberFormattingOptions()
+	.SetMinimumIntegralDigits(4)
+	.SetUseGrouping(false);
diff --git a/Engine/Plugins/Experimental/Avalanche/Source/AvalancheMediaEditor/Private/Playlist/AvaPlaylistEditor.cpp b/Engine/Plugins/Experimental/Avalanche/Source/AvalancheMediaEditor/Private/Playlist/AvaPlaylistEditor.cpp
new file mode 100644
--- /dev/null
+++ b/Engine/Plugins/Experimental/Avalanche/Source/AvalancheMediaEditor/Private/Playlist/AvaPlaylistEditor.cpp
@@ -0,0 +1,1859 @@
+// Copyright Epic Games, Inc. All Rights Reserved.
+
+#include "Playlist/AvaPlaylistEditor.h"
+
+#include "AppModes/AvaPlaylistDefaultMode.h"
+#include "Async/Async.h"
+#include "AvalancheBroadcast.h"
+#include "Broadcast/AvaBroadcastEditor.h"
+#include "Framework/Application/IInputProcessor.h"
+#include "Framework/Application/SlateApplication.h"
+#include "Framework/MultiBox/MultiBoxBuilder.h"
+#include "IAvaMediaEditorModule.h"
+#include "IAvaMediaModule.h"
+#include "Misc/MessageDialog.h"
+#include "Pages/Slate/SAvaInstancedPageList.h"
+#include "Pages/Slate/SAvaReadPage.h"
+#include "Pages/Slate/SAvaTemplatePageList.h"
+#include "Playlist/AvaPlaylistCommands.h"
+#include "Playlist/AvaPlaylistEditorUtils.h"
+#include "Playlist/AvaPlaylistPlaybackUtils.h"
+#include "Playlist/AvaRundownEditorMacroCommands.h"
+#include "Playlist/AvaRundownEditorSettings.h"
+#include "Playlist/AvalancheManagedInstanceCache.h"
+#include "Playlist/AvalanchePlaylist.h"
+#include "Playlist/Factories/Filters/AvaPlaylistFactoriesUtils.h"
+#include "Playlist/Filters/AvaPageTextFilter.h"
+#include "Playlist/Pages/PageViews/AvaPageView.h"
+#include "Playlist/Pages/Slate/SAvaPageList.h"
+#include "ScopedTransaction.h"
+#include "TabFactories/AvaInstancedPageListTabFactory.h"
+#include "TabFactories/AvaSubListDocumentTabFactory.h"
+#include "TabFactories/AvaSubListTabFactory.h"
+#include "TabFactories/AvaTemplatePageListTabFactory.h"
+#include "Widgets/Docking/SDockTab.h"
+#include "WorkflowOrientedApp/ApplicationMode.h"
+
+#define LOCTEXT_NAMESPACE "AvaPlaylistEditor"
+
+namespace UE::AvaPlaylistEditor::Private
+{
+	/** Return true if there's an Active Tab and if it is part of this Playlist Editor. */
+	inline bool IsActiveTabPartOfEditor(FAvaPlaylistEditor& InPlaylistEditor)
+	{
+		const TSharedPtr<SDockTab> ActiveTab = FGlobalTabmanager::Get()->GetActiveTab();
+		return ActiveTab.IsValid() && InPlaylistEditor.GetAssociatedTabManager() == ActiveTab->GetTabManagerPtr();
+	}
+
+	inline bool ShouldStopPagesOnClose()
+	{
+		const UAvaRundownEditorSettings* RundownEditorSettings = UAvaRundownEditorSettings::Get();
+		return RundownEditorSettings ? RundownEditorSettings->bShouldStopPagesOnClose : false;
+	}
+
+	inline const UAvaRundownMacroCollection* GetMacroCollection()
+	{
+		const UAvaRundownEditorSettings* RundownEditorSettings = UAvaRundownEditorSettings::Get();
+		if (!RundownEditorSettings || RundownEditorSettings->MacroCollection.IsNull())
+		{
+			return nullptr;
+		}
+
+		return RundownEditorSettings->MacroCollection.LoadSynchronous();
+	}
+
+	FInputChord MakeInputChord(const FKeyEvent& InKeyEvent)
+	{
+		return FInputChord(InKeyEvent.GetKey(),InKeyEvent.IsShiftDown(),  InKeyEvent.IsControlDown(), InKeyEvent.IsAltDown(), InKeyEvent.IsCommandDown());
+	}
+
+	FName GetName(EAvaRundownEditorMacroCommand InCommand)
+	{
+		return FAvaRundownEditorMacroCommands::GetShortCommandName(InCommand);
+	}
+}
+
+class FAvaPlaylistEditorInputProcessor : public IInputProcessor
+{
+public:
+	FAvaPlaylistEditorInputProcessor(const TWeakPtr<FAvaPlaylistEditor>& InPlaylistEditorWeak)
+		: PlaylistEditorWeak(InPlaylistEditorWeak)
+	{
+	}
+
+	virtual void Tick(const float DeltaTime, FSlateApplication& SlateApp, TSharedRef<ICursor> Cursor) override {}
+
+	virtual bool HandleKeyDownEvent(FSlateApplication& SlateApp, const FKeyEvent& InKeyEvent) override
+	{
+		if (const TSharedPtr<FAvaPlaylistEditor> PlaylistEditor = PlaylistEditorWeak.Pin())
+		{
+			if (!UE::AvaPlaylistEditor::Private::IsActiveTabPartOfEditor(*PlaylistEditor))
+			{
+				return false;
+			}
+
+			// Do not pre-process keys that are not relevant to the editor
+			if (!PlaylistEditor->IsKeyRelevant(InKeyEvent))
+			{
+				return false;
+			}
+
+			// Next, Check if the Keyboard Focused Widget (if valid) actually handles these Key Events
+			const TSharedPtr<SWidget> FocusedWidget = FSlateApplicationBase::Get().GetKeyboardFocusedWidget();
+
+			if (FocusedWidget.IsValid() && FocusedWidget->OnKeyDown(FocusedWidget->GetTickSpaceGeometry(), InKeyEvent).IsEventHandled())
+			{
+				// if handled, then just return as we don't want to process the same key again for the widget
+				return true;
+			}
+
+			// If not handled then forward it to Playlist Editor
+			return PlaylistEditor->HandleKeyDownEvent(InKeyEvent);
+		}
+		return false;
+	}
+
+private:
+	TWeakPtr<FAvaPlaylistEditor> PlaylistEditorWeak;
+};
+
+/**
+ *	Console commands are shared between all editors.
+ *	However, only the editor that is the active tab will receive the commands.
+ */
+class FAvaPlaylistEditor::FSharedConsoleCommands
+{
+	struct FPrivateToken { explicit FPrivateToken() = default; };
+public:
+	static void RegisterEditor(const TSharedPtr<FAvaPlaylistEditor>& InPlaylistEditor);
+	static void UnregisterEditor(FAvaPlaylistEditor* InPlaylistEditor);
+
+	static TSharedPtr<FSharedConsoleCommands> GetSharedInstance();
+
+	explicit FSharedConsoleCommands(FPrivateToken)
+	{
+		RegisterConsoleCommands();
+	}
+
+	~FSharedConsoleCommands()
+	{
+		UnregisterConsoleCommands();
+	}
+
+private:
+	void RegisterConsoleCommands();
+	void UnregisterConsoleCommands();
+
+	TSharedPtr<FAvaPlaylistEditor> GetActivePlaylistEditor() const;
+
+	template <typename InFunctionType>
+	void InvokeOnActiveEditor(InFunctionType InFunctionToInvoke) const
+	{
+		if (const TSharedPtr<FAvaPlaylistEditor> ActivePlaylistEditor = GetActivePlaylistEditor())
+		{
+			::Invoke(InFunctionToInvoke, ActivePlaylistEditor);
+		}
+	}
+
+	void StartAutoPlayCommand(const TArray<FString>& InArgs) const
+	{
+		InvokeOnActiveEditor([&InArgs](const TSharedPtr<FAvaPlaylistEditor>& InPlaylistEditor){InPlaylistEditor->StartAutoPlayCommand(InArgs);});
+	}
+
+	void StopAutoPlayCommand(const TArray<FString>& InArgs) const
+	{
+		InvokeOnActiveEditor([&InArgs](const TSharedPtr<FAvaPlaylistEditor>& InPlaylistEditor){InPlaylistEditor->StopAutoPlayCommand(InArgs);});
+	}
+
+	void LoadPageCommand(const TArray<FString>& InArgs) const
+	{
+		InvokeOnActiveEditor([&InArgs](const TSharedPtr<FAvaPlaylistEditor>& InPlaylistEditor){InPlaylistEditor->LoadPageCommand(InArgs);});
+	}
+	
+	void UnloadPageCommand(const TArray<FString>& InArgs) const
+	{
+		InvokeOnActiveEditor([&InArgs](const TSharedPtr<FAvaPlaylistEditor>& InPlaylistEditor){InPlaylistEditor->UnloadPageCommand(InArgs);});
+	}
+
+private:
+	TArray<IConsoleObject*> ConsoleCommands;
+	TArray<TWeakPtr<FAvaPlaylistEditor>> PlaylistEditors;
+};
+
+FAvaPlaylistEditor::FAvaPlaylistEditor()
+	: TextFilterTemplatePage(MakeShared<FAvaPageTextFilter>())
+	, TextFilterInstancedPage(MakeShared<FAvaPageTextFilter>())
+{}
+
+FAvaPlaylistEditor::~FAvaPlaylistEditor()
+{
+	UAvalanchePlaylist* Playlist = AvalanchePlaylist.Get();
+
+	if (IsValid(Playlist))
+	{
+		Playlist->GetOnActiveListChanged().RemoveAll(this);
+		Playlist->GetOnPagePlayerAdded().RemoveAll(this);
+		
+		const bool bStopPages = bStopPagesOnCloseOverride.IsSet() ?
+			bStopPagesOnCloseOverride.GetValue() : UE::AvaPlaylistEditor::Private::ShouldStopPagesOnClose();
+
+		Playlist->ClosePlaybackContext(bStopPages);
+	}
+
+	if (InputProcessor.IsValid() && FSlateApplication::IsInitialized())
+	{
+		FSlateApplication::Get().UnregisterInputPreProcessor(InputProcessor);
+	}
+
+	if (IAvaMediaModule::IsModuleLoaded())
+	{
+		// This is a safety flush to provide the possibility to reset all loaded assets.
+		// This is to force a reload of the assets if they are modified in another editor.
+		IAvaMediaModule::Get().GetManagedInstanceCache().Flush();
+	}
+
+	FSharedConsoleCommands::UnregisterEditor(this);
+}
+
+void FAvaPlaylistEditor::InitPlaylistEditor(const EToolkitMode::Type InMode
+	, const TSharedPtr<IToolkitHost>& InInitToolkitHost
+	, UAvalanchePlaylist* InPlaylist)
+{
+	if (FSlateApplication::IsInitialized())
+	{
+		InputProcessor = MakeShared<FAvaPlaylistEditorInputProcessor>(SharedThis(this));
+		FSlateApplication::Get().RegisterInputPreProcessor(InputProcessor);
+	}
+
+	AvalanchePlaylist = InPlaylist;
+
+	InitVisibilityTemplatePages();
+	InitVisibilityInstancedPages();
+
+	TextFilterTemplatePage->OnChanged().AddSP(this, &FAvaPlaylistEditor::OnTemplateFilterChanged);
+	TextFilterInstancedPage->OnChanged().AddSP(this, &FAvaPlaylistEditor::OnInstancedFilterChanged);
+
+	CreatePlaylistCommands();
+
+	if (IsValid(InPlaylist))
+	{
+		InPlaylist->GetOnActiveListChanged().AddSP(this, &FAvaPlaylistEditor::OnActiveSubListChanged);
+		InPlaylist->GetOnPagePlayerAdded().AddSP(this, &FAvaPlaylistEditor::HandleOnPagePlayerAdded);
+		InPlaylist->InitializePlaybackContext();
+	}
+
+	const FName PlaylistEditorAppName(TEXT("AvalanchePlaylistEditorApp"));
+	constexpr bool bCreateDefaultStandaloneMenu = true;
+	constexpr bool bCreateDefaultToolbar = true;
+
+	InitAssetEditor(InMode
+		, InInitToolkitHost
+		, PlaylistEditorAppName
+		, FTabManager::FLayout::NullLayout
+		, bCreateDefaultStandaloneMenu
+		, bCreateDefaultToolbar
+		, InPlaylist);
+
+	RegisterApplicationModes();
+	FSharedConsoleCommands::RegisterEditor(SharedThis(this));
+
+	CreateSubListTabs();
+}
+
+bool FAvaPlaylistEditor::IsKeyRelevant(const FKeyEvent& InKeyEvent) const
+{
+	if (ReadPageWidget.IsValid() && ReadPageWidget->IsKeyRelevant(InKeyEvent))
+	{
+		return true;
+	}
+
+	using namespace UE::AvaPlaylistEditor::Private;
+	const UAvaRundownMacroCollection* MacroCollection = GetMacroCollection();
+	if (MacroCollection && MacroCollection->HasBindingFor(MakeInputChord(InKeyEvent)))
+	{
+		return true;
+	}
+	
+	return false;
+}
+
+bool FAvaPlaylistEditor::HandleKeyDownEvent(const FKeyEvent& InKeyEvent)
+{
+	if (ReadPageWidget.IsValid() && ReadPageWidget->ProcessPlaylistKeyDown(InKeyEvent))
+	{
+		return true;
+	}
+
+	// Check if we have a binding for this key
+	using namespace UE::AvaPlaylistEditor::Private;
+	const UAvaRundownMacroCollection* MacroCollection = GetMacroCollection();
+	if (!MacroCollection)
+	{
+		return false;
+	}
+
+	const FInputChord InputChord = MakeInputChord(InKeyEvent);
+	const int32 NumCommandsFound = MacroCollection->ForEachCommand(InputChord, [this](const FAvaRundownMacroCommand& InCommand)
+	{
+		if (const FMacroCommandFunction* CommandFunction = GetBindableMacroCommands().Find(InCommand.Name))
+		{
+			TArray<FString> Args;
+			InCommand.Arguments.ParseIntoArrayWS(Args, TEXT(","));
+			(*CommandFunction)(Args);
+			return true;
+		}
+		return false;
+	});
+	
+	return NumCommandsFound > 0;
+}
+
+TSharedPtr<SAvaPageList> FAvaPlaylistEditor::GetTemplateListWidget() const
+{
+	return GetListWidget(FAvaTemplatePageListTabFactory::TabID);
+}
+
+TSharedPtr<SAvaPageList> FAvaPlaylistEditor::GetInstanceListWidget() const
+{
+	return GetListWidget(FAvaInstancedPageListTabFactory::TabID);
+}
+
+TSharedPtr<SAvaPageList> FAvaPlaylistEditor::GetListWidget(const FAvaPageListReference& InPageListReference) const
+{
+	if (InPageListReference.Type == EAvaPageListType::Template)
+	{
+		return GetListWidget(FAvaTemplatePageListTabFactory::TabID);
+	}
+
+	if (InPageListReference.Type == EAvaPageListType::Instance)
+	{
+		return GetListWidget(FAvaInstancedPageListTabFactory::TabID);
+	}
+
+	return GetListWidget(FAvaSubListDocumentTabFactory::GetTabId(InPageListReference.SubListIndex));
+}
+
+TSharedPtr<SAvaPageList> FAvaPlaylistEditor::GetListWidget(const FName& InTabId) const
+{
+	if (TabManager.IsValid())
+	{
+		TSharedPtr<SDockTab> FoundTab = TabManager->FindExistingLiveTab(InTabId);
+
+		if (FoundTab.IsValid())
+		{
+			TSharedPtr<SWidget> Content = FoundTab->GetContent();
+
+			if (Content->GetWidgetClass().GetWidgetType() == SAvaPageList::StaticWidgetClass().GetWidgetType()
+				|| Content->GetWidgetClass().GetWidgetType() == SAvaInstancedPageList::StaticWidgetClass().GetWidgetType()
+				|| Content->GetWidgetClass().GetWidgetType() == SAvaTemplatePageList::StaticWidgetClass().GetWidgetType())
+			{
+				return StaticCastSharedPtr<SAvaPageList>(Content);
+			}
+		}
+	}
+
+	return nullptr;
+}
+
+TSharedPtr<SAvaInstancedPageList> FAvaPlaylistEditor::GetActiveListWidget() const
+{
+	const UAvalanchePlaylist* Playlist = AvalanchePlaylist.Get();
+
+	if (!IsValid(Playlist))
+	{
+		return nullptr;
+	}
+
+	const FAvaPageListReference& ActiveList = Playlist->GetActivePageListReference();
+
+	if (ActiveList.Type == EAvaPageListType::Template)
+	{
+		return nullptr;
+	}
+
+	TSharedPtr<SAvaPageList> PageList;
+
+	if (ActiveList.Type == EAvaPageListType::Instance)
+	{
+		PageList = GetListWidget(FAvaInstancedPageListTabFactory::TabID);
+	}
+	else if (ActiveList.Type == EAvaPageListType::View)
+	{
+		PageList = GetListWidget(FAvaSubListDocumentTabFactory::GetTabId(ActiveList.SubListIndex));
+	}
+
+	if (PageList.IsValid() && PageList->GetWidgetClass().GetWidgetType() == SAvaInstancedPageList::StaticWidgetClass().GetWidgetType())
+	{
+		return StaticCastSharedPtr<SAvaInstancedPageList>(PageList);
+	}
+
+	return nullptr;
+}
+
+TSharedPtr<SAvaPageList> FAvaPlaylistEditor::GetFocusedListWidget() const
+{
+	if (TabManager.IsValid())
+	{
+		const TSharedPtr<SDockTab> ActiveTab = FGlobalTabmanager::Get()->GetActiveTab();
+
+		if (ActiveTab.IsValid() && const_cast<FAvaPlaylistEditor*>(this)->GetAssociatedTabManager() == ActiveTab->GetTabManagerPtr())
+		{
+			const TSharedPtr<SWidget> Content = ActiveTab->GetContent();
+
+			if (Content->GetWidgetClass().GetWidgetType() == SAvaPageList::StaticWidgetClass().GetWidgetType()
+				|| Content->GetWidgetClass().GetWidgetType() == SAvaInstancedPageList::StaticWidgetClass().GetWidgetType()
+				|| Content->GetWidgetClass().GetWidgetType() == SAvaTemplatePageList::StaticWidgetClass().GetWidgetType())
+			{
+				return StaticCastSharedPtr<SAvaPageList>(Content);
+			}
+		}
+	}
+	return nullptr;
+}
+
+int32 FAvaPlaylistEditor::GetFirstSelectedPageOnActiveSubListWidget() const
+{
+	const TSharedPtr<SAvaInstancedPageList> ActiveWidget = GetActiveListWidget();
+
+	if (ActiveWidget.IsValid())
+	{
+		return ActiveWidget->GetFirstSelectedPageId();
+	}
+
+	return FAvalanchePage::InvalidPageId;
+}
+
+TConstArrayView<int32> FAvaPlaylistEditor::GetSelectedPagesOnActiveSubListWidget() const
+{
+	const TSharedPtr<SAvaInstancedPageList> ActiveWidget = GetActiveListWidget();
+
+	if (ActiveWidget.IsValid())
+	{
+		return ActiveWidget->GetSelectedPageIds();
+	}
+
+	return {};
+}
+
+TConstArrayView<int32> FAvaPlaylistEditor::GetSelectedPagesOnFocusedWidget() const
+{
+	const TSharedPtr<SAvaPageList> FocusedWidget = GetFocusedListWidget();
+	if (FocusedWidget.IsValid())
+	{
+		return FocusedWidget->GetSelectedPageIds();
+	}
+
+	return {};
+}
+
+bool FAvaPlaylistEditor::CanAddTemplate() const
+{
+	UAvalanchePlaylist* Playlist = AvalanchePlaylist.Get();
+
+	if (IsValid(Playlist))
+	{
+		return Playlist->CanAddPage();
+	}
+
+	return false;
+}
+
+void FAvaPlaylistEditor::AddTemplate()
+{
+	UAvalanchePlaylist* Playlist = AvalanchePlaylist.Get();
+
+	if (IsValid(Playlist))
+	{
+		FScopedTransaction Transaction(LOCTEXT("AddTemplate", "Add Template"));
+		Playlist->Modify();
+		
+		if (Playlist->AddTemplate() == FAvalanchePage::InvalidPageId)
+		{
+			Transaction.Cancel();
+		}
+	}
+}
+
+bool FAvaPlaylistEditor::CanPlaySelectedPage() const
+{
+	const TSharedPtr<SAvaInstancedPageList> ActiveWidget = GetActiveListWidget();
+
+	if (ActiveWidget.IsValid())
+	{
+		return ActiveWidget->CanPlaySelectedPage();
+	}
+
+	return false;
+}
+
+void FAvaPlaylistEditor::PlaySelectedPage()
+{
+	const TSharedPtr<SAvaInstancedPageList> ActiveWidget = GetActiveListWidget();
+
+	if (ActiveWidget.IsValid())
+	{
+		ActiveWidget->PlaySelectedPage();
+	}
+}
+
+bool FAvaPlaylistEditor::CanUpdateValuesOnSelectedPage() const
+{
+	const TSharedPtr<SAvaInstancedPageList> ActiveWidget = GetActiveListWidget();
+	return ActiveWidget.IsValid() ? ActiveWidget->CanUpdateValuesOnSelectedPage() : false;
+}
+
+void FAvaPlaylistEditor::UpdateValuesOnSelectedPage()
+{
+	const TSharedPtr<SAvaInstancedPageList> ActiveWidget = GetActiveListWidget();
+
+	if (ActiveWidget.IsValid())
+	{
+		ActiveWidget->UpdateValuesOnSelectedPage();
+	}
+}
+
+bool FAvaPlaylistEditor::CanStopSelectedPage(bool bInForce) const
+{
+	const TSharedPtr<SAvaInstancedPageList> ActiveWidget = GetActiveListWidget();
+
+	if (ActiveWidget.IsValid())
+	{
+		return ActiveWidget->CanStopSelectedPage(bInForce);
+	}
+
+	return false;
+}
+
+void FAvaPlaylistEditor::StopSelectedPage(bool bInForce)
+{
+	const TSharedPtr<SAvaInstancedPageList> ActiveWidget = GetActiveListWidget();
+
+	if (ActiveWidget.IsValid())
+	{
+		ActiveWidget->StopSelectedPage(bInForce);
+	}
+}
+
+bool FAvaPlaylistEditor::CanContinueSelectedPage() const
+{
+	const TSharedPtr<SAvaInstancedPageList> ActiveWidget = GetActiveListWidget();
+
+	if (ActiveWidget.IsValid())
+	{
+		return ActiveWidget->CanContinueSelectedPage();
+	}
+
+	return false;
+}
+
+void FAvaPlaylistEditor::ContinueSelectedPage()
+{
+	const TSharedPtr<SAvaInstancedPageList> ActiveWidget = GetActiveListWidget();
+
+	if (ActiveWidget.IsValid())
+	{
+		ActiveWidget->ContinueSelectedPage();
+	}
+}
+
+bool FAvaPlaylistEditor::CanPlayNextPage() const
+{
+	const TSharedPtr<SAvaInstancedPageList> ActiveWidget = GetActiveListWidget();
+
+	if (ActiveWidget.IsValid())
+	{
+		return ActiveWidget->CanPlayNextPage();
+	}
+
+	return false;
+}
+
+void FAvaPlaylistEditor::PlayNextPage()
+{
+	const TSharedPtr<SAvaInstancedPageList> ActiveWidget = GetActiveListWidget();
+
+	if (ActiveWidget.IsValid())
+	{
+		ActiveWidget->PlayNextPage();
+	}
+}
+
+bool FAvaPlaylistEditor::CanPreviewPlaySelectedPage() const
+{
+	const TSharedPtr<SAvaInstancedPageList> ActiveWidget = GetActiveListWidget();
+
+	if (ActiveWidget.IsValid())
+	{
+		if (ActiveWidget->CanPreviewPlaySelectedPage())
+		{
+			return true;
+		}
+
+		// If instanced page list can't preview play the page it may be that you selected a template page.
+		if (const TSharedPtr<SAvaPageList> TemplateWidget = GetTemplateListWidget())
+		{
+			return TemplateWidget->CanPreviewPlaySelectedPage();
+		}
+	}
+
+	return false;
+}
+
+void FAvaPlaylistEditor::PreviewPlaySelectedPage(bool bInToMark)
+{
+	const TSharedPtr<SAvaInstancedPageList> ActiveWidget = GetActiveListWidget();
+
+	if (ActiveWidget.IsValid())
+	{
+		// If instanced page list can't preview play but you reached here than it is for sure a template page selected
+		if (ActiveWidget->CanPreviewPlaySelectedPage())
+		{
+			ActiveWidget->PreviewPlaySelectedPage(bInToMark);
+		}
+		else if (const TSharedPtr<SAvaPageList> TemplateWidget = GetTemplateListWidget())
+		{
+			TemplateWidget->PreviewPlaySelectedPage(bInToMark);
+		}
+	}
+}
+
+bool FAvaPlaylistEditor::CanPreviewStopSelectedPage(bool bInForce) const
+{
+	const TSharedPtr<SAvaInstancedPageList> ActiveWidget = GetActiveListWidget();
+
+	if (ActiveWidget.IsValid())
+	{
+		return ActiveWidget->CanPreviewStopSelectedPage(bInForce);
+	}
+
+	return false;
+}
+
+void FAvaPlaylistEditor::PreviewStopSelectedPage(bool bInForce)
+{
+	const TSharedPtr<SAvaInstancedPageList> ActiveWidget = GetActiveListWidget();
+
+	if (ActiveWidget.IsValid())
+	{
+		ActiveWidget->PreviewStopSelectedPage(bInForce);
+	}
+}
+
+bool FAvaPlaylistEditor::CanPreviewContinueSelectedPage() const
+{
+	const TSharedPtr<SAvaInstancedPageList> ActiveWidget = GetActiveListWidget();
+
+	if (ActiveWidget.IsValid())
+	{
+		return ActiveWidget->CanPreviewContinueSelectedPage();
+	}
+
+	return false;
+}
+
+void FAvaPlaylistEditor::PreviewContinueSelectedPage()
+{
+	const TSharedPtr<SAvaInstancedPageList> ActiveWidget = GetActiveListWidget();
+
+	if (ActiveWidget.IsValid())
+	{
+		ActiveWidget->PreviewContinueSelectedPage();
+	}
+}
+
+bool FAvaPlaylistEditor::CanPreviewPlayNextPage() const
+{
+	const TSharedPtr<SAvaInstancedPageList> ActiveWidget = GetActiveListWidget();
+
+	if (ActiveWidget.IsValid())
+	{
+		return ActiveWidget->CanPreviewPlayNextPage();
+	}
+
+	return false;
+}
+
+void FAvaPlaylistEditor::PreviewPlayNextPage()
+{
+	const TSharedPtr<SAvaInstancedPageList> ActiveWidget = GetActiveListWidget();
+
+	if (ActiveWidget.IsValid())
+	{
+		ActiveWidget->PreviewPlayNextPage();
+	}
+}
+
+bool FAvaPlaylistEditor::CanTakeToProgram() const
+{
+	const TSharedPtr<SAvaInstancedPageList> ActiveWidget = GetActiveListWidget();
+
+	if (ActiveWidget.IsValid())
+	{
+		return ActiveWidget->CanTakeToProgram();
+	}
+
+	return false;
+}
+
+void FAvaPlaylistEditor::TakeToProgram() const
+{
+	const TSharedPtr<SAvaInstancedPageList> ActiveWidget = GetActiveListWidget();
+
+	if (ActiveWidget.IsValid())
+	{
+		ActiveWidget->TakeToProgram();
+	}
+}
+
+bool FAvaPlaylistEditor::CanCreateInstancesFromSelectedTemplates() const
+{
+	UAvalanchePlaylist* Playlist = AvalanchePlaylist.Get();
+
+	if (!IsValid(Playlist))
+	{
+		return false;
+	}
+
+	const TSharedPtr<SAvaPageList> TemplateListWidget = GetTemplateListWidget();
+
+	if (!TemplateListWidget)
+	{
+		return false;
+	}
+
+	const bool bHasSelectedPages = TemplateListWidget.IsValid() && (TemplateListWidget->GetSelectedPageIds().IsEmpty() == false);
+	const bool bCanAddPage = Playlist->CanAddPage();
+
+	return bHasSelectedPages && bCanAddPage;
+}
+
+void FAvaPlaylistEditor::CreateInstancesFromSelectedTemplates()
+{
+	UAvalanchePlaylist* Playlist = AvalanchePlaylist.Get();
+
+	if (!IsValid(Playlist))
+	{
+		return;
+	}
+
+	const TSharedPtr<SAvaPageList> TemplateListWidget = GetTemplateListWidget();
+
+	if (!TemplateListWidget)
+	{
+		return;
+	}
+	
+	FScopedTransaction Transaction(LOCTEXT("CreateInstancesFromSelectedTemplates", "Create Instances From Selected Templates"));
+	Playlist->Modify();
+
+	const TArray<int32> AddedPages = Playlist->AddPagesFromTemplates(TemplateListWidget->GetSelectedPageIds());
+
+	if (!AddedPages.IsEmpty())
+	{
+		const TSharedPtr<SAvaInstancedPageList> PageListWidget = GetActiveListWidget();
+
+		if (PageListWidget.IsValid() && PageListWidget->GetPageListReference().Type == EAvaPageListType::View)
+		{
+			Playlist->AddPagesToSubList(PageListWidget->GetPageListReference().SubListIndex, AddedPages);
+		}
+	}
+	else
+	{
+		Transaction.Cancel();
+	}
+}
+
+bool FAvaPlaylistEditor::CanRemoveSelectedPages() const
+{
+	const UAvalanchePlaylist* Playlist = AvalanchePlaylist.Get();
+
+	if (!IsValid(Playlist))
+	{
+		return false;
+	}
+
+	const TSharedPtr<SAvaPageList> PageListWidget = GetFocusedListWidget();
+
+	return PageListWidget.IsValid() && (PageListWidget->GetSelectedPageIds().IsEmpty() == false)
+		&& Playlist->CanRemovePages(PageListWidget->GetSelectedPageIds());
+}
+
+void FAvaPlaylistEditor::RemoveSelectedPages()
+{
+	UAvalanchePlaylist* Playlist = AvalanchePlaylist.Get();
+
+	if (!IsValid(Playlist))
+	{
+		return;
+	}
+
+	const TSharedPtr<SAvaPageList> PageListWidget = GetFocusedListWidget();
+	if (PageListWidget.IsValid())
+	{
+		FScopedTransaction Transaction(LOCTEXT("RemoveSelectedPages", "Remove Selected Pages"));
+		Playlist->Modify();
+
+		const int32 RemovedCount = Playlist->RemovePages(PageListWidget->GetSelectedPageIds());
+		
+		if (RemovedCount == 0)
+		{
+			Transaction.Cancel();
+		}
+	}
+}
+
+void FAvaPlaylistEditor::RefreshTemplateVisibility()
+{
+	VisibleTemplatePageIds.Reset();
+
+	if (AvalanchePlaylist.IsValid())
+	{
+		for (const FAvalanchePage& Page : AvalanchePlaylist->GetTemplatePages().Pages)
+		{
+			TextFilterTemplatePage->SetItem(Page, AvalanchePlaylist.Get(), EAvaPlaylistSearchListType::Template);
+			if (TextFilterTemplatePage->PassesFilter(Page))
+			{
+				VisibleTemplatePageIds.Add(Page.GetPageId());
+			}
+		}
+	}
+}
+
+void FAvaPlaylistEditor::RefreshInstancedVisibility()
+{
+	VisibleInstancedPageIds.Reset();
+
+	if (AvalanchePlaylist.IsValid())
+	{
+		for (const FAvalanchePage& Page : AvalanchePlaylist->GetInstancedPages().Pages)
+		{
+			TextFilterInstancedPage->SetItem(Page, AvalanchePlaylist.Get(), EAvaPlaylistSearchListType::Instanced);
+			if (TextFilterInstancedPage->PassesFilter(Page))
+			{
+				VisibleInstancedPageIds.Add(Page.GetPageId());
+			}
+		}
+	}
+}
+
+bool FAvaPlaylistEditor::IsTemplatePageVisible(const FAvalanchePage& InPage) const
+{
+	return VisibleTemplatePageIds.Contains(InPage.GetPageId());
+}
+
+bool FAvaPlaylistEditor::IsInstancedPageVisible(const FAvalanchePage& InPage) const
+{
+	return VisibleInstancedPageIds.Contains(InPage.GetPageId());
+}
+
+void FAvaPlaylistEditor::SetSearchText(const FText& InText, EAvaPlaylistSearchListType& InPageListType)
+{
+	switch (InPageListType)
+	{
+	case EAvaPlaylistSearchListType::Template:
+		SetTemplateSearchText(InText);
+		break;
+	case EAvaPlaylistSearchListType::Instanced:
+		SetInstancedSearchText(InText);
+		break;
+
+	case EAvaPlaylistSearchListType::None:
+	default:
+		break;
+	}
+}
+
+FName FAvaPlaylistEditor::GetToolkitFName() const
+{
+	return TEXT("AvaPlaylistEditor");
+}
+
+FText FAvaPlaylistEditor::GetBaseToolkitName() const
+{
+	return LOCTEXT("PlaylistAppLabel", "Motion Design Playlist Editor");
+}
+
+FString FAvaPlaylistEditor::GetWorldCentricTabPrefix() const
+{
+	return LOCTEXT("PlaylistScriptPrefix", "Script ").ToString();
+}
+
+FLinearColor FAvaPlaylistEditor::GetWorldCentricTabColorScale() const
+{
+	return FLinearColor(0.0f, 0.0f, 0.3f, 0.5f);
+}
+
+bool FAvaPlaylistEditor::OnRequestClose(EAssetEditorCloseReason InCloseReason)
+{
+	const UAvalanchePlaylist* Playlist = AvalanchePlaylist.Get();
+
+	// Ask user if they want pages to be stopped.
+	// Only ask if the editor is closed by the user directly.
+	if (IsValid(Playlist) && InCloseReason == EAssetEditorCloseReason::AssetEditorHostClosed)
+	{
+		if (Playlist->IsPlaying())
+		{
+			
+			const FText MessageText = LOCTEXT("StopPagesOnExitQuestion",
+				"The editor is closing and some pages are still playing, do you want to stop all pages?");
+			
+			const EAppReturnType::Type Reply = FMessageDialog::Open(EAppMsgType::YesNoCancel, EAppReturnType::Yes, MessageText);
+
+			if (Reply == EAppReturnType::Cancel)
+			{
+				return false;	// Don't close the editor.
+			}
+			
+			bStopPagesOnCloseOverride = (Reply == EAppReturnType::Yes);
+		}
+	}
+
+	return FWorkflowCentricApplication::OnRequestClose(InCloseReason);
+}
+
+TSharedRef<SWidget> FAvaPlaylistEditor::MakeReadPageWidget()
+{
+	return SAssignNew(ReadPageWidget, SAvaReadPage, SharedThis(this));
+}
+
+void FAvaPlaylistEditor::ExtendToolBar(TSharedPtr<FExtender> InExtender)
+{
+	InExtender->AddToolBarExtension("Asset"
+		, EExtensionHook::After
+		, ToolkitCommands
+		, FToolBarExtensionDelegate::CreateSP(this, &FAvaPlaylistEditor::FillPageToolBar));
+}
+
+void FAvaPlaylistEditor::FillPageToolBar(FToolBarBuilder& OutToolBarBuilder)
+{
+	const FAvaPlaylistCommands& PlaylistCommands = FAvaPlaylistCommands::Get();
+
+	OutToolBarBuilder.BeginSection(TEXT("Pages"));
+	{
+		OutToolBarBuilder.AddToolBarButton(PlaylistCommands.AddTemplate);
+		OutToolBarBuilder.AddToolBarButton(PlaylistCommands.CreatePageInstanceFromTemplate);
+		OutToolBarBuilder.AddToolBarButton(PlaylistCommands.RemovePage);
+	}
+	OutToolBarBuilder.EndSection();
+
+	OutToolBarBuilder.AddSeparator();
+
+	OutToolBarBuilder.BeginSection(TEXT("Broadcast"));
+	{
+		OutToolBarBuilder.AddToolBarButton(FExecuteAction::CreateStatic(&FAvaBroadcastEditor::OpenBroadcastEditor)
+			, NAME_None
+			, LOCTEXT("Broadcast_Label", "Broadcast")
+			, LOCTEXT("Broadcast_ToolTip", "Opens the Avalanche Broadcast Editor Window")
+			, TAttribute<FSlateIcon>::Create([]() { return IAvaMediaEditorModule::Get().GetToolbarBroadcastButtonIcon(); }));
+		OutToolBarBuilder.AddComboButton(
+			FUIAction(FExecuteAction()
+			, FCanExecuteAction::CreateLambda([](){ return !UAvalancheBroadcast::Get().IsBroadcastingAnyChannel();}))
+			, FOnGetContent::CreateSP(this, &FAvaPlaylistEditor::MakeProfileComboButton)
+			, TAttribute<FText>(this, &FAvaPlaylistEditor::GetCurrentProfileName)
+			, LOCTEXT("Profile_ToolTip", "Pick between different Profiles")
+			, FSlateIcon(FAppStyle::Get().GetStyleSetName(), "LevelEditor.Profile")
+			, false);
+		OutToolBarBuilder.AddToolBarButton(
+			FUIAction(FExecuteAction::CreateLambda([]{ UAvalancheBroadcast::Get().StartBroadcast();})
+				, FCanExecuteAction::CreateLambda([]{ return !UAvalancheBroadcast::Get().IsBroadcastingAllChannels();}))
+			, NAME_None
+			, LOCTEXT("Play_Label", "Start All Channels")
+			, LOCTEXT("Play_ToolTip", "Starts Broadcast on all Idle Channels")
+			, FSlateIcon(FAppStyle::GetAppStyleSetName(), "Icons.Toolbar.Play"));
+		OutToolBarBuilder.AddToolBarButton(
+			FUIAction(FExecuteAction::CreateLambda([]{ UAvalancheBroadcast::Get().StopBroadcast();})
+				, FCanExecuteAction::CreateLambda([]{ return UAvalancheBroadcast::Get().IsBroadcastingAnyChannel();}))
+			, NAME_None
+			, LOCTEXT("Stop_Label", "Stop All Channels")
+			, LOCTEXT("Stop_ToolTip", "Stops Broadcast on all Live Channels")
+			, FSlateIcon(FAppStyle::GetAppStyleSetName(), "Icons.Toolbar.Stop"));
+	}
+	OutToolBarBuilder.EndSection();
+}
+
+UAvalanchePlaylist* FAvaPlaylistEditor::GetPlaylist() const
+{
+	return AvalanchePlaylist.Get();
+}
+
+void FAvaPlaylistEditor::MarkAsModified()
+{
+	if (AvalanchePlaylist.IsValid())
+	{
+		AvalanchePlaylist->Modify();
+	}
+}
+
+TArray<TSharedPtr<FAvalancheManagedInstance>> FAvaPlaylistEditor::GetManagedInstancesForPage(const UAvalanchePlaylist* InPlaylist, const FAvalanchePage& InPage)
+{
+	const TArray<FSoftObjectPath> AssetPaths = InPage.GetAvalancheAssetPaths(InPlaylist);
+	
+	TArray<TSharedPtr<FAvalancheManagedInstance>> ManagedInstances;
+	ManagedInstances.Reserve(AssetPaths.Num());
+	
+	FAvalancheManagedInstanceCache& ManagedInstanceCache = IAvaMediaModule::Get().GetManagedInstanceCache();
+	
+	for (const FSoftObjectPath& AssetPath : AssetPaths)
+	{		
+		if (TSharedPtr<FAvalancheManagedInstance> ManagedInstance = ManagedInstanceCache.GetOrLoadAvalancheInstance(AssetPath))
+		{
+			ManagedInstances.Add(ManagedInstance);
+		}
+	}
+	return ManagedInstances;
+}
+
+bool FAvaPlaylistEditor::MergeDefaultRemoteControlValues(const TArray<TSharedPtr<FAvalancheManagedInstance>>& InManagedInstances, FAvalancheRemoteControlValues& OutMergedValues)
+{
+	bool bAllUniqueIds = true;
+	
+	for (const TSharedPtr<FAvalancheManagedInstance>& ManagedInstance : InManagedInstances)
+	{
+		bAllUniqueIds &= OutMergedValues.Merge(ManagedInstance->GetDefaultRemoteControlValues());
+	}
+	
+	return bAllUniqueIds;
+}
+
+void FAvaPlaylistEditor::OnActiveSubListChanged()
+{
+	UAvalanchePlaylist* Playlist = AvalanchePlaylist.Get();
+
+	if (!IsValid(Playlist))
+	{
+		return;
+	}
+
+	const FAvaPageListReference& ActiveList = Playlist->GetActivePageListReference();
+
+	if (ActiveList.Type != EAvaPageListType::View)
+	{
+		return;
+	}
+
+	const FName TabId = FAvaSubListDocumentTabFactory::GetTabId(ActiveList.SubListIndex);
+
+	TSharedPtr<SDockTab> SubListTab = TabManager->FindExistingLiveTab(TabId);
+
+	if (!SubListTab.IsValid())
+	{
+		SubListTab = CreateSubListTab(ActiveList.SubListIndex);
+	}
+
+	if (SubListTab.IsValid())
+	{
+		SubListTab->ActivateInParent(ETabActivationCause::SetDirectly);
+		SubListTab->DrawAttention();
+	}
+}
+
+void FAvaPlaylistEditor::HandleOnPagePlayerAdded(UAvalanchePlaylist* InPlaylist, UAvalanchePagePlayer* InPagePlayer)
+{
+	UAvalanchePlaylist* Playlist = AvalanchePlaylist.Get();
+
+	if (!IsValid(Playlist) || InPlaylist != Playlist || !IsValid(InPagePlayer))
+	{
+		return;
+	}
+
+	// TODO: this should have an editor setting to enable it. "EnablePreviewAutoBroadcast" or something...
+	// Auto start preview channel.
+	if (InPagePlayer->bIsPreview)
+	{
+		UAvalancheBroadcast::Get().ConditionalStartBroadcastChannel(InPagePlayer->ChannelFName);
+	}
+}
+
+TSharedPtr<SDockTab> FAvaPlaylistEditor::CreateSubListTab(int32 InSubListIndex)
+{
+	TSharedPtr<FApplicationMode> AppMode = GetCurrentModePtr();
+
+	if (AppMode.IsValid())
+	{
+		const FName TabId = FAvaSubListDocumentTabFactory::GetTabId(InSubListIndex);
+		TSharedRef<FAvaPlaylistAppMode> PlaylistAppMode = StaticCastSharedRef<FAvaPlaylistAppMode>(AppMode.ToSharedRef());
+		TSharedPtr<FDocumentTabFactory> DocTabFactory = PlaylistAppMode->GetDocumentTabFactory(FAvaSubListDocumentTabFactory::FactoryId);
+
+		if (DocTabFactory.IsValid())
+		{
+			TSharedRef<FAvaSubListDocumentTabFactory> SubListTabFactory = StaticCastSharedRef<FAvaSubListDocumentTabFactory>(DocTabFactory.ToSharedRef());
+
+			FWorkflowTabSpawnInfo Info;
+			Info.TabManager = TabManager;
+			Info.Payload = nullptr;
+			Info.TabInfo = nullptr;
+
+			TSharedPtr<SDockTab> SubListTab = SubListTabFactory->SpawnSubListTab(Info, InSubListIndex);
+
+			if (SubListTab.IsValid())
+			{
+				TabManager->InsertNewDocumentTab(FAvaSubListTabFactory::TabID, TabId, FTabManager::FLiveTabSearch(TabId), SubListTab.ToSharedRef());
+				return SubListTab;
+			}
+		}
+	}
+
+	return nullptr;
+}
+
+void FAvaPlaylistEditor::CreateSubListTabs()
+{
+	UAvalanchePlaylist* Playlist = AvalanchePlaylist.Get();
+
+	if (!IsValid(Playlist))
+	{
+		return;
+	}
+
+	for (int32 SubListIndex = 0; SubListIndex < Playlist->GetSubLists().Num(); ++SubListIndex)
+	{
+		CreateSubListTab(SubListIndex);
+	}
+}
+
+void FAvaPlaylistEditor::RegisterApplicationModes()
+{
+	TArray<TSharedRef<FApplicationMode>> ApplicationModes;
+	TSharedPtr<FAvaPlaylistEditor> This = SharedThis(this);
+
+	ApplicationModes.Add(MakeShared<FAvaPlaylistDefaultMode>(This));
+	//Can add more App Modes here
+
+	for (const TSharedRef<FApplicationMode>& AppMode : ApplicationModes)
+	{
+		AddApplicationMode(AppMode->GetModeName(), AppMode);
+	}
+
+	SetCurrentMode(FAvaPlaylistDefaultMode::DefaultMode);
+}
+
+void FAvaPlaylistEditor::CreatePlaylistCommands()
+{
+	//Playlist Commands
+	{
+		const FAvaPlaylistCommands& PlaylistCommands = FAvaPlaylistCommands::Get();
+
+		ToolkitCommands->MapAction(PlaylistCommands.Play,
+			FExecuteAction::CreateSP(this, &FAvaPlaylistEditor::PlaySelectedPage),
+			FCanExecuteAction::CreateSP(this, &FAvaPlaylistEditor::CanPlaySelectedPage));
+
+		ToolkitCommands->MapAction(PlaylistCommands.UpdateValues,
+			FExecuteAction::CreateSP(this, &FAvaPlaylistEditor::UpdateValuesOnSelectedPage),
+			FCanExecuteAction::CreateSP(this, &FAvaPlaylistEditor::CanUpdateValuesOnSelectedPage));
+		
+		ToolkitCommands->MapAction(PlaylistCommands.Stop,
+			FExecuteAction::CreateSP(this, &FAvaPlaylistEditor::StopSelectedPage, false),
+			FCanExecuteAction::CreateSP(this, &FAvaPlaylistEditor::CanStopSelectedPage, false));
+
+		ToolkitCommands->MapAction(PlaylistCommands.ForceStop,
+			FExecuteAction::CreateSP(this, &FAvaPlaylistEditor::StopSelectedPage, true),
+			FCanExecuteAction::CreateSP(this, &FAvaPlaylistEditor::CanStopSelectedPage, true));
+
+		ToolkitCommands->MapAction(PlaylistCommands.Continue,
+			FExecuteAction::CreateSP(this, &FAvaPlaylistEditor::ContinueSelectedPage),
+			FCanExecuteAction::CreateSP(this, &FAvaPlaylistEditor::CanContinueSelectedPage));
+
+		ToolkitCommands->MapAction(PlaylistCommands.PlayNext,
+			FExecuteAction::CreateSP(this, &FAvaPlaylistEditor::PlayNextPage),
+			FCanExecuteAction::CreateSP(this, &FAvaPlaylistEditor::CanPlayNextPage));
+
+		ToolkitCommands->MapAction(PlaylistCommands.PreviewFrame,
+			FExecuteAction::CreateSP(this, &FAvaPlaylistEditor::PreviewPlaySelectedPage, true),
+			FCanExecuteAction::CreateSP(this, &FAvaPlaylistEditor::CanPreviewPlaySelectedPage));
+
+		ToolkitCommands->MapAction(PlaylistCommands.PreviewPlay,
+			FExecuteAction::CreateSP(this, &FAvaPlaylistEditor::PreviewPlaySelectedPage, false),
+			FCanExecuteAction::CreateSP(this, &FAvaPlaylistEditor::CanPreviewPlaySelectedPage));
+
+		ToolkitCommands->MapAction(PlaylistCommands.PreviewContinue,
+			FExecuteAction::CreateSP(this, &FAvaPlaylistEditor::PreviewContinueSelectedPage),
+			FCanExecuteAction::CreateSP(this, &FAvaPlaylistEditor::CanPreviewContinueSelectedPage));
+
+		ToolkitCommands->MapAction(PlaylistCommands.PreviewStop,
+			FExecuteAction::CreateSP(this, &FAvaPlaylistEditor::PreviewStopSelectedPage, false),
+			FCanExecuteAction::CreateSP(this, &FAvaPlaylistEditor::CanPreviewStopSelectedPage, false));
+
+		ToolkitCommands->MapAction(PlaylistCommands.PreviewForceStop,
+			FExecuteAction::CreateSP(this, &FAvaPlaylistEditor::PreviewStopSelectedPage, true),
+			FCanExecuteAction::CreateSP(this, &FAvaPlaylistEditor::CanPreviewStopSelectedPage, true));
+
+		ToolkitCommands->MapAction(PlaylistCommands.PreviewPlayNext,
+			FExecuteAction::CreateSP(this, &FAvaPlaylistEditor::PreviewPlayNextPage),
+			FCanExecuteAction::CreateSP(this, &FAvaPlaylistEditor::CanPreviewPlayNextPage));
+
+		ToolkitCommands->MapAction(PlaylistCommands.TakeToProgram,
+			FExecuteAction::CreateSP(this, &FAvaPlaylistEditor::TakeToProgram),
+			FCanExecuteAction::CreateSP(this, &FAvaPlaylistEditor::CanTakeToProgram));
+
+		ToolkitCommands->MapAction(PlaylistCommands.AddTemplate,
+			FExecuteAction::CreateSP(this, &FAvaPlaylistEditor::AddTemplate),
+			FCanExecuteAction::CreateSP(this, &FAvaPlaylistEditor::CanAddTemplate));
+
+		ToolkitCommands->MapAction(PlaylistCommands.CreatePageInstanceFromTemplate,
+			FExecuteAction::CreateSP(this, &FAvaPlaylistEditor::CreateInstancesFromSelectedTemplates),
+			FCanExecuteAction::CreateSP(this, &FAvaPlaylistEditor::CanCreateInstancesFromSelectedTemplates));
+
+		ToolkitCommands->MapAction(PlaylistCommands.RemovePage,
+			FExecuteAction::CreateSP(this, &FAvaPlaylistEditor::RemoveSelectedPages),
+			FCanExecuteAction::CreateSP(this, &FAvaPlaylistEditor::CanRemoveSelectedPages));
+	}
+}
+
+FText FAvaPlaylistEditor::GetCurrentProfileName() const
+{
+	return FText::FromName(UAvalancheBroadcast::Get().GetCurrentProfileName());
+}
+
+TSharedRef<SWidget> FAvaPlaylistEditor::MakeProfileComboButton()
+{
+	FMenuBuilder MenuBuilder(true, nullptr);
+
+	TArray<FName> ProfileNames = UAvalancheBroadcast::Get().GetProfileNames();
+	for (FName ProfileName : ProfileNames)
+	{
+		MenuBuilder.AddMenuEntry(FText::FromName(ProfileName)
+			, FText()
+			, FSlateIcon()
+			, FUIAction(FExecuteAction::CreateLambda([ProfileName](){ UAvalancheBroadcast::Get().SetCurrentProfile(ProfileName);}))
+		);
+	}
+
+	return MenuBuilder.MakeWidget();
+}
+
+void FAvaPlaylistEditor::StartAutoPlayCommand(const TArray<FString>& InArgs)
+{
+	const double PlayInterval = (InArgs.Num() > 0) ? FCString::Atod(*InArgs[0]) : FAutoPlayTicker::DefaultPlayInterval;
+	UE_LOG(LogAvaMediaEditor, Log, TEXT("Playlist auto play started, interval: %f seconds."), PlayInterval);
+	AutoPlayTicker = MakeUnique<FAutoPlayTicker>(SharedThis(this), PlayInterval);
+}
+
+void FAvaPlaylistEditor::StopAutoPlayCommand(const TArray<FString>& InArgs)
+{
+	CancelAutoPlay();
+}
+
+void FAvaPlaylistEditor::CancelAutoPlay()
+{
+	if (AutoPlayTicker && !AutoPlayTicker->bIsCancelled)
+	{
+		// Mark as cancelled.
+		AutoPlayTicker->bIsCancelled = true;
+		// We can't delete the object immediately because it might still be ticking so we defer with an async task.
+		TWeakPtr<FAvaPlaylistEditor> PlaylistEditorWeak(SharedThis(this));
+		AsyncTask(ENamedThreads::GameThread, [PlaylistEditorWeak]()
+		{
+			if (const TSharedPtr<FAvaPlaylistEditor> PlaylistEditor = PlaylistEditorWeak.Pin())
+			{
+				PlaylistEditor->AutoPlayTicker.Reset();
+			}
+		});
+	}
+}
+
+void FAvaPlaylistEditor::LoadPageCommand(const TArray<FString>& InArgs)
+{
+	FString Errors;
+	const TArray<int32> PageIds = ArgumentsToPageIds(InArgs, /*bInPreview*/ false, Errors);
+
+	if (!Errors.IsEmpty())
+	{
+		UE_LOG(LogAvaMediaEditor, Error, TEXT("Errors in Load Page command: %s"), *Errors);
+	}
+
+	UAvalanchePlaylist* Playlist = AvalanchePlaylist.Get();
+
+	// Log the status of the operation.
+	auto LoadPage = [Playlist](int32 InPageId, bool bInIsPreview, FName InChannelName)
+	{
+		if (Playlist->GetPageLoadingManager().RequestLoadPage(InPageId, bInIsPreview, InChannelName))
+		{
+			UE_LOG(LogAvaMediaEditor, Display, TEXT("Loaded page %d for channel \"%s\"."), InPageId, *InChannelName.ToString());
+		}
+	};
+
+	const FName PreviewChannelName = UAvalanchePlaylist::GetDefaultPreviewChannelName();
+	
+	for (const int32 PageId : PageIds)
+	{
+		const FAvalanchePage& Page =  AvalanchePlaylist->GetPage(PageId);
+		if (Page.IsValidPage())
+		{
+			if (!Page.IsTemplate())
+			{
+				LoadPage(PageId, false, Page.GetChannelName());
+			}
+			LoadPage(PageId, true, PreviewChannelName);
+		}
+	}
+}
+
+void FAvaPlaylistEditor::UnloadPageCommand(const TArray<FString>& InArgs)
+{
+	FString Errors;
+	const TArray<int32> PageIds = ArgumentsToPageIds(InArgs, /*bInPreview*/ false, Errors);
+
+	if (!Errors.IsEmpty())
+	{
+		UE_LOG(LogAvaMediaEditor, Error, TEXT("Errors in Unload Page command: %s"), *Errors);
+	}
+
+	UAvalanchePlaylist* Playlist = AvalanchePlaylist.Get();
+
+	// Log the status of the operation.
+	auto UnloadPage = [Playlist](int32 InPageId, const FName& InChannelName)
+	{
+		if (Playlist->UnloadPage(InPageId, InChannelName.ToString()))
+		{
+			UE_LOG(LogAvaMediaEditor, Display, TEXT("Unloaded page %d for channel \"%s\"."), InPageId, *InChannelName.ToString());
+		}
+	};
+
+	const FName PreviewChannelName = UAvalanchePlaylist::GetDefaultPreviewChannelName();
+	
+	for (const int32 PageId : PageIds)
+	{
+		const FAvalanchePage& Page =  AvalanchePlaylist->GetPage(PageId);
+		if (Page.IsValidPage())
+		{
+			if (!Page.IsTemplate())
+			{
+				UnloadPage(PageId, Page.GetChannelName());
+			}
+			UnloadPage(PageId, PreviewChannelName);
+		}
+	}
+}
+
+void FAvaPlaylistEditor::PlayPageCommand(const TArray<FString>& InArgs, bool bInPreview)
+{
+	FString Errors;	
+	const TArray<int32> PageIds = ArgumentsToPageIds(InArgs, bInPreview, Errors);
+
+	if (!Errors.IsEmpty())
+	{
+		UE_LOG(LogAvaMediaEditor, Error, TEXT("Errors in Play Page command: %s"), *Errors);
+	}
+
+	UAvalanchePlaylist* Playlist = AvalanchePlaylist.Get();
+
+	if (Playlist && !PageIds.IsEmpty())
+	{
+		Playlist->PlayPages(PageIds, bInPreview ? EAvaPlayType::PreviewFromFrame : EAvaPlayType::PlayFromStart);
+	}
+}
+
+void FAvaPlaylistEditor::ContinuePageCommand(const TArray<FString>& InArgs, bool bInPreview)
+{
+	FString Errors;	
+	const TArray<int32> PageIds = ArgumentsToPageIds(InArgs, bInPreview, Errors);
+
+	if (!Errors.IsEmpty())
+	{
+		UE_LOG(LogAvaMediaEditor, Error, TEXT("Errors in Continue Page command: %s"), *Errors);
+	}
+
+	UAvalanchePlaylist* Playlist = AvalanchePlaylist.Get();
+
+	for (const int32 PageId : PageIds)
+	{
+		if (Playlist->CanContinuePage(PageId, bInPreview))
+		{
+			Playlist->ContinuePage(PageId, bInPreview);
+		}
+	}
+}
+
+void FAvaPlaylistEditor::PlayNextPageCommand(const TArray<FString>& InArgs, bool bInPreview)
+{
+	UAvalanchePlaylist* Playlist = AvalanchePlaylist.Get();
+	if (!AvalanchePlaylist.IsValid())
+	{
+		UE_LOG(LogAvaMediaEditor, Error, TEXT("Errors in Continue Page command: invalid playlist."));
+		return;
+	}
+
+	const int32 NextPageId = FAvaPlaylistPlaybackUtils::GetPageIdToPlayNext(
+		Playlist, UAvalanchePlaylist::InstancePageList, bInPreview, bInPreview ? UAvalanchePlaylist::GetDefaultPreviewChannelName() : NAME_None);
+
+	if (FAvaPlaylistPlaybackUtils::IsPageIdValid(NextPageId))
+	{
+		Playlist->PlayPage(NextPageId, bInPreview ? EAvaPlayType::PreviewFromFrame : EAvaPlayType::PlayFromStart);
+	}
+}
+
+void FAvaPlaylistEditor::StopPageCommand(const TArray<FString>& InArgs, bool bInPreview, bool bInForce)
+{
+	FString Errors;
+	const TArray<int32> PageIds = ArgumentsToPageIds(InArgs, bInPreview, Errors);
+
+	if (!Errors.IsEmpty())
+	{
+		UE_LOG(LogAvaMediaEditor, Error, TEXT("Errors in \"Stop Page\" command: %s"), *Errors);
+	}
+
+	UAvalanchePlaylist* Playlist = AvalanchePlaylist.Get();
+	if (!Playlist)
+	{
+		return;
+	}
+
+	if (Playlist && !PageIds.IsEmpty())
+	{
+		Playlist->StopPages(PageIds, bInForce ? EAvaPlaylistPageStopOptions::ForceNoTransition : EAvaPlaylistPageStopOptions::Default, bInPreview);
+	}
+}
+
+void FAvaPlaylistEditor::TakeToProgramCommand(const TArray<FString>& InArgs)
+{
+	FString Errors;
+	const TArray<int32> PageIds = ArgumentsToPageIds(InArgs, /*bInPreview*/ true, Errors);
+
+	if (!Errors.IsEmpty())
+	{
+		UE_LOG(LogAvaMediaEditor, Error, TEXT("Errors in \"Take To Program\" command: %s"), *Errors);
+	}
+
+	UAvalanchePlaylist* Playlist = AvalanchePlaylist.Get();
+	if (!Playlist)
+	{
+		return;
+	}
+
+	const TArray<int32> PageIdsToPlay = FAvaPlaylistPlaybackUtils::GetPagesToTakeToProgram(Playlist, PageIds);
+
+	if (!PageIdsToPlay.IsEmpty())
+	{
+		Playlist->PlayPages(PageIdsToPlay, EAvaPlayType::PlayFromStart);
+	}
+}
+
+void FAvaPlaylistEditor::StartChannelCommand(const TArray<FString>& InArgs)
+{
+	UAvalancheBroadcast& Broadcast = UAvalancheBroadcast::Get();
+	if (InArgs.IsEmpty() || InArgs[0] == TEXT("all"))
+	{
+		Broadcast.StartBroadcast();
+		return;
+	}
+	
+	for (const FString& Argument : InArgs)
+	{
+		FAvaOutputChannel& Channel = Broadcast.GetCurrentProfile().GetChannelMutable(FName(Argument));
+		if (Channel.IsValidChannel())
+		{
+			Channel.StartChannelBroadcast();
+		}
+		else
+		{
+			UE_LOG(LogAvaMediaEditor, Error, TEXT("Argument \"%s\" is not a valid channel name."), *Argument);
+		}
+	}
+}
+
+void FAvaPlaylistEditor::StopChannelCommand(const TArray<FString>& InArgs)
+{
+	UAvalancheBroadcast& Broadcast = UAvalancheBroadcast::Get();
+	if (InArgs.IsEmpty() || InArgs[0] == TEXT("all"))
+	{
+		Broadcast.StopBroadcast();
+		return;
+	}
+	
+	for (const FString& Argument : InArgs)
+	{
+		FAvaOutputChannel& Channel = Broadcast.GetCurrentProfile().GetChannelMutable(FName(Argument));
+		if (Channel.IsValidChannel())
+		{
+			Channel.StopChannelBroadcast();
+		}
+		else
+		{
+			UE_LOG(LogAvaMediaEditor, Error, TEXT("Argument \"%s\" is not a valid channel name."), *Argument);
+		}
+	}
+}
+
+TArray<int32> FAvaPlaylistEditor::ArgumentsToPageIds(const TArray<FString>& InArgs, bool bInPreview, FString& OutErrors) const
+{
+	if (!AvalanchePlaylist.IsValid())
+	{
+		OutErrors = TEXT("invalid playlist.");
+		return {};
+	}
+
+	UAvalanchePlaylist* Playlist = AvalanchePlaylist.Get();
+
+	TSet<int32> PageIds;
+
+	for (const FString& Arg : InArgs)
+	{
+		if (Arg.IsNumeric())
+		{
+			const int32 PageId = FCString::Atoi(*Arg);
+			const FAvalanchePage& Page =  Playlist->GetPage(PageId);
+			if (!Page.IsValidPage())
+			{
+				OutErrors += FString::Printf(TEXT("argument \"%s\" is not a valid page Id. "), *Arg);
+				continue;
+			}
+			PageIds.Add(PageId);
+		}
+		else if (Arg == TEXT("all"))	// all pages in the list.
+		{
+			for (const FAvalanchePage& Page : Playlist->GetInstancedPages().Pages)
+			{
+				PageIds.Add(Page.GetPageId());
+			}
+			for (const FAvalanchePage& Page : Playlist->GetTemplatePages().Pages)
+			{
+				PageIds.Add(Page.GetPageId());
+			}
+		}
+		else if (Arg == TEXT("any"))	// any playing pages
+		{
+			if (bInPreview)
+			{
+				PageIds.Append(Playlist->GetPreviewingPageIds());
+			}
+			else
+			{
+				PageIds.Append(Playlist->GetPlayingPageIds());
+			}
+		}
+		else if (UAvalancheBroadcast::Get().GetChannelIndex(FName(Arg)) != INDEX_NONE)
+		{
+			// Any playing pages on the specified channel.
+			if (bInPreview)
+			{
+				PageIds.Append(Playlist->GetPreviewingPageIds(FName(Arg)));
+			}
+			else
+			{
+				PageIds.Append(Playlist->GetPlayingPageIds(FName(Arg)));
+			}
+		}
+		else
+		{
+			OutErrors += FString::Printf(TEXT("argument \"%s\" is not recognized. Supported arguments: pageId, channelName, \"all\" or \"any\". "), *Arg);
+		}
+	}
+	return PageIds.Array();
+}
+
+const FAvaPlaylistEditor::FBindableMacroCommands& FAvaPlaylistEditor::GetBindableMacroCommands()
+{
+	if (BindableMacroCommands.IsEmpty())
+	{
+		using namespace UE::AvaPlaylistEditor::Private;
+		BindableMacroCommands.Add(GetName(EAvaRundownEditorMacroCommand::LoadPage), [this](const TArray<FString>& InArgs){LoadPageCommand(InArgs);});
+		BindableMacroCommands.Add(GetName(EAvaRundownEditorMacroCommand::UnloadPage), [this](const TArray<FString>& InArgs){UnloadPageCommand(InArgs);});
+		BindableMacroCommands.Add(GetName(EAvaRundownEditorMacroCommand::TakeIn), [this](const TArray<FString>& InArgs){PlayPageCommand(InArgs, /*bInPreview*/ false);});
+		BindableMacroCommands.Add(GetName(EAvaRundownEditorMacroCommand::TakeOut), [this](const TArray<FString>& InArgs){StopPageCommand(InArgs, /*bInPreview*/ false, /*bInForce*/ false);});
+		BindableMacroCommands.Add(GetName(EAvaRundownEditorMacroCommand::ForceTakeOut), [this](const TArray<FString>& InArgs){StopPageCommand(InArgs, /*bInPreview*/ false, /*bInForce*/ true);});
+		BindableMacroCommands.Add(GetName(EAvaRundownEditorMacroCommand::TakeNext), [this](const TArray<FString>& InArgs){PlayNextPageCommand(InArgs, /*bInPreview*/ false);});
+		BindableMacroCommands.Add(GetName(EAvaRundownEditorMacroCommand::Continue), [this](const TArray<FString>& InArgs){ContinuePageCommand(InArgs, /*bInPreview*/ false);});
+		BindableMacroCommands.Add(GetName(EAvaRundownEditorMacroCommand::PreviewIn), [this](const TArray<FString>& InArgs){PlayPageCommand(InArgs, /*bInPreview*/ true);});
+		BindableMacroCommands.Add(GetName(EAvaRundownEditorMacroCommand::PreviewOut), [this](const TArray<FString>& InArgs){StopPageCommand(InArgs, /*bInPreview*/ true, /*bInForce*/ false);});
+		BindableMacroCommands.Add(GetName(EAvaRundownEditorMacroCommand::ForcePreviewOut), [this](const TArray<FString>& InArgs){StopPageCommand(InArgs, /*bInPreview*/ true, /*bInForce*/ true);});
+		BindableMacroCommands.Add(GetName(EAvaRundownEditorMacroCommand::PreviewNext), [this](const TArray<FString>& InArgs){PlayNextPageCommand(InArgs, /*bInPreview*/ true);});
+		BindableMacroCommands.Add(GetName(EAvaRundownEditorMacroCommand::ContinuePreview), [this](const TArray<FString>& InArgs){ContinuePageCommand(InArgs, /*bInPreview*/ true);});
+		BindableMacroCommands.Add(GetName(EAvaRundownEditorMacroCommand::TakeToProgram), [this](const TArray<FString>& InArgs){TakeToProgramCommand(InArgs);});
+		BindableMacroCommands.Add(GetName(EAvaRundownEditorMacroCommand::StartChannel), [this](const TArray<FString>& InArgs){StartChannelCommand(InArgs);});
+		BindableMacroCommands.Add(GetName(EAvaRundownEditorMacroCommand::StopChannel), [this](const TArray<FString>& InArgs){StopChannelCommand(InArgs);});
+	}
+	return BindableMacroCommands;
+}
+
+void FAvaPlaylistEditor::SetTemplateSearchText(const FText& InText)
+{
+	TextFilterTemplatePage->SetFilterText(InText);
+}
+
+void FAvaPlaylistEditor::SetInstancedSearchText(const FText& InText)
+{
+	TextFilterInstancedPage->SetFilterText(InText);
+}
+
+void FAvaPlaylistEditor::OnTemplateFilterChanged()
+{
+	RefreshTemplateVisibility();
+
+	if (GetTemplateListWidget().IsValid())
+	{
+		GetTemplateListWidget()->Refresh();
+	}
+}
+
+void FAvaPlaylistEditor::OnInstancedFilterChanged()
+{
+	RefreshInstancedVisibility();
+
+	if (GetInstanceListWidget().IsValid())
+	{
+		GetInstanceListWidget()->Refresh();
+	}
+}
+
+void FAvaPlaylistEditor::InitVisibilityTemplatePages()
+{
+	if (const UAvalanchePlaylist* Playlist = GetPlaylist())
+	{
+		for (const FAvalanchePage& Page : Playlist->GetTemplatePages().Pages)
+		{
+			VisibleTemplatePageIds.Add(Page.GetPageId());
+		}
+	}
+}
+
+void FAvaPlaylistEditor::InitVisibilityInstancedPages()
+{
+	if (const UAvalanchePlaylist* Playlist = GetPlaylist())
+	{
+		for (const FAvalanchePage& Page : Playlist->GetInstancedPages().Pages)
+		{
+			VisibleInstancedPageIds.Add(Page.GetPageId());
+		}
+	}
+}
+
+FAvaPlaylistEditor::FAutoPlayTicker::FAutoPlayTicker(TWeakPtr<FAvaPlaylistEditor> InPlaylistEditorWeak, double InTickInterval)
+	: PlaylistEditorWeak(InPlaylistEditorWeak)
+	, PlayInterval(InTickInterval)
+{
+	// Start current page or select and play the first page of the list.
+	if (const TSharedPtr<FAvaPlaylistEditor> PlaylistEditor = PlaylistEditorWeak.Pin())
+	{
+		UAvalanchePlaylist* Playlist = PlaylistEditor->GetPlaylist();
+
+		if (Playlist && Playlist->GetInstancedPages().Pages.Num() > 0)
+		{
+			TMap<FName, int32> PageIdsToPlay = TMap<FName, int32>();
+
+			if (const TSharedPtr<SAvaInstancedPageList> ActivePageList = PlaylistEditor->GetActiveListWidget())
+			{
+				if (ActivePageList->GetPageViews().IsEmpty() == false)
+				{
+					if (ActivePageList->GetSelectedPageIds().IsEmpty())
+					{
+						const FAvaPageViewPtr& FirstPage = ActivePageList->GetPageViews()[0];
+						PageIdsToPlay.Add(Playlist->GetPage(FirstPage->GetPageId()).GetChannelName(), FirstPage->GetPageId());
+						ActivePageList->SelectPage(FirstPage->GetPageId());
+					}
+					else
+					{
+						for (const int32 PageId : ActivePageList->GetSelectedPageIds())
+						{
+							if (PageIdsToPlay.Find(Playlist->GetPage(PageId).GetChannelName()) == nullptr)
+							{
+								PageIdsToPlay.Add(Playlist->GetPage(PageId).GetChannelName(), PageId);
+							}
+						}
+					}
+				}
+			}
+
+			if (PageIdsToPlay.IsEmpty())
+			{
+				// Just go direct to the instanced pages.
+				PageIdsToPlay.Add(Playlist->GetInstancedPages().Pages[0].GetChannelName(), Playlist->GetInstancedPages().Pages[0].GetPageId());
+
+				// And select it if possible.
+				if (TSharedPtr<SAvaPageList> MainPageList = PlaylistEditor->GetInstanceListWidget())
+				{
+					if (MainPageList->GetPageViews().IsEmpty() == false)
+					{
+						if (MainPageList->GetSelectedPageIds().IsEmpty())
+						{
+							MainPageList->SelectPage(Playlist->GetInstancedPages().Pages[0].GetPageId());
+						}
+						else
+						{
+							for (const int32 PageId : MainPageList->GetSelectedPageIds())
+							{
+								if (PageIdsToPlay.Find(Playlist->GetPage(PageId).GetChannelName()) == nullptr)
+								{
+									PageIdsToPlay.Add(Playlist->GetPage(PageId).GetChannelName(), PageId);
+								}
+							}
+						}
+					}
+				}
+			}
+
+			for (const TPair<FName, int32> PageId : PageIdsToPlay)
+			{
+				Playlist->PlayPage(PageId.Value, EAvaPlayType::PlayFromStart);
+			}
+		}
+	}
+}
+
+void FAvaPlaylistEditor::FAutoPlayTicker::Tick(float DeltaTime)
+{
+	const TSharedPtr<FAvaPlaylistEditor> PlaylistEditor = PlaylistEditorWeak.Pin();
+	if (!PlaylistEditor || bIsCancelled)
+	{
+		return;
+	}
+
+	const UAvalanchePlaylist* Playlist = PlaylistEditor->GetPlaylist();
+	if (!Playlist)
+	{
+		UE_LOG(LogAvaMediaEditor, Error, TEXT("Invalid playlist. Cancelling auto play."));
+		PlaylistEditor->CancelAutoPlay();
+		return;
+	}
+
+	if (!bIsPagePlaying)
+	{
+		// Check the status of the page to see if it started playing.
+		const TArray<int32> SelectedPagesId = PlaylistEditor->GetPlaylist()->GetPlayingPageIds();
+		if (SelectedPagesId.IsEmpty())
+		{
+			return;
+		}
+
+		int32 NumberOfPagesPlaying = 0;
+		for (const int32 SelectedPageId : SelectedPagesId)
+		{
+			const FAvalanchePage* SelectedPage = (SelectedPageId != FAvalanchePage::InvalidPageId) ? &Playlist->GetPage(SelectedPageId) : nullptr;
+
+			if (SelectedPage && SelectedPage->IsValidPage())
+			{
+				const TArray<FAvalanchePageStatus> Statuses = SelectedPage->GetPageContextualStatuses(Playlist);
+
+				for (const FAvalanchePageStatus& Status : Statuses)
+				{
+					if (Status.Status == EAvalanchePageStatus::Playing)
+					{
+						// bIsPagePlaying = true;
+						NumberOfPagesPlaying++;
+						NextPageStartTime = FApp::GetCurrentTime() + PlayInterval;
+						break;
+					}
+					else if (Status.Status == EAvalanchePageStatus::Error)
+					{
+						UE_LOG(LogAvaMediaEditor, Error, TEXT("Playback error. Cancelling auto play."));
+						PlaylistEditor->CancelAutoPlay();
+						return;
+					}
+				}
+			}
+			else
+			{
+				UE_LOG(LogAvaMediaEditor, Error, TEXT("Current Page is Invalid. Cancelling auto play."));
+				PlaylistEditor->CancelAutoPlay();
+				return;
+			}
+		}
+		bIsPagePlaying = (NumberOfPagesPlaying == SelectedPagesId.Num());
+	}
+
+	if (bIsPagePlaying && FApp::GetCurrentTime() > NextPageStartTime)
+	{
+		// This will request the page to play.
+		if (const TSharedPtr<SAvaInstancedPageList> ActivePageList = PlaylistEditor->GetActiveListWidget())
+		{
+			const TArray<int32> NextPages = ActivePageList->PlayNextPage();
+			if (NextPages.Num() > 0)
+			{
+				ActivePageList->DeselectPages();
+				ActivePageList->SelectPages(NextPages);
+			}
+		}
+
+		// But we will wait for it to actually start before measuring play time.
+		bIsPagePlaying = false;
+	}
+}
+
+TStatId FAvaPlaylistEditor::FAutoPlayTicker::GetStatId() const
+{
+	RETURN_QUICK_DECLARE_CYCLE_STAT(FAutoPlayTicker, STATGROUP_Tickables);
+}
+
+void FAvaPlaylistEditor::FSharedConsoleCommands::RegisterEditor(const TSharedPtr<FAvaPlaylistEditor>& InPlaylistEditor)
+{
+	if (InPlaylistEditor.IsValid() && !InPlaylistEditor->SharedConsoleCommands.IsValid())
+	{
+		InPlaylistEditor->SharedConsoleCommands = GetSharedInstance();
+		InPlaylistEditor->SharedConsoleCommands->PlaylistEditors.Add(InPlaylistEditor);
+	}
+}
+
+void FAvaPlaylistEditor::FSharedConsoleCommands::UnregisterEditor(FAvaPlaylistEditor* InPlaylistEditor)
+{
+	if (InPlaylistEditor && InPlaylistEditor->SharedConsoleCommands.IsValid())
+	{
+		InPlaylistEditor->SharedConsoleCommands->PlaylistEditors.RemoveAll([InPlaylistEditor](const TWeakPtr<FAvaPlaylistEditor>& Element)
+		{
+			return !Element.IsValid() || Element.HasSameObject(InPlaylistEditor);
+		});
+		InPlaylistEditor->SharedConsoleCommands.Reset();
+	}
+}
+
+TSharedPtr<FAvaPlaylistEditor::FSharedConsoleCommands> FAvaPlaylistEditor::FSharedConsoleCommands::GetSharedInstance()
+{
+	static TWeakPtr<FSharedConsoleCommands> GlobalInstanceWeak;
+	if (GlobalInstanceWeak.IsValid())
+	{
+		return GlobalInstanceWeak.Pin();
+	}
+	TSharedPtr<FSharedConsoleCommands> NewInstance = MakeShared<FSharedConsoleCommands>(FPrivateToken());
+	GlobalInstanceWeak = NewInstance;
+	return NewInstance;
+}
+
+void FAvaPlaylistEditor::FSharedConsoleCommands::RegisterConsoleCommands()
+{
+	ConsoleCommands.Add(IConsoleManager::Get().RegisterConsoleCommand(
+		TEXT("AvaPlaylistEditor.StartAutoPlay"),
+		TEXT("Starts auto play of the currently active playlist editor."),
+		FConsoleCommandWithArgsDelegate::CreateRaw(this, &FSharedConsoleCommands::StartAutoPlayCommand),
+		ECVF_Default
+	));
+	ConsoleCommands.Add(IConsoleManager::Get().RegisterConsoleCommand(
+		TEXT("AvaPlaylistEditor.StopAutoPlay"),
+		TEXT("Stops auto play of the currently active playlist editor."),
+		FConsoleCommandWithArgsDelegate::CreateRaw(this, &FSharedConsoleCommands::StopAutoPlayCommand),
+		ECVF_Default
+	));
+	ConsoleCommands.Add(IConsoleManager::Get().RegisterConsoleCommand(
+		TEXT("AvaPlaylistEditor.LoadPage"),
+		TEXT("Preload the specified page from the current playlist to memory."),
+		FConsoleCommandWithArgsDelegate::CreateRaw(this, &FSharedConsoleCommands::LoadPageCommand),
+		ECVF_Default
+	));
+	ConsoleCommands.Add(IConsoleManager::Get().RegisterConsoleCommand(
+		TEXT("AvaPlaylistEditor.UnloadPage"),
+		TEXT("Unload the specified page from the current playlist."),
+		FConsoleCommandWithArgsDelegate::CreateRaw(this, &FSharedConsoleCommands::UnloadPageCommand),
+		ECVF_Default
+	));
+}
+
+void FAvaPlaylistEditor::FSharedConsoleCommands::UnregisterConsoleCommands()
+{
+	for (IConsoleObject* ConsoleCmd : ConsoleCommands)
+	{
+		IConsoleManager::Get().UnregisterConsoleObject(ConsoleCmd);
+	}
+	ConsoleCommands.Empty();
+}
+
+TSharedPtr<FAvaPlaylistEditor> FAvaPlaylistEditor::FSharedConsoleCommands::GetActivePlaylistEditor() const
+{
+	for (const TWeakPtr<FAvaPlaylistEditor>& PlaylistEditorWeak : PlaylistEditors)
+	{
+		TSharedPtr<FAvaPlaylistEditor> PlaylistEditor = PlaylistEditorWeak.Pin();
+		if (PlaylistEditor && UE::AvaPlaylistEditor::Private::IsActiveTabPartOfEditor(*PlaylistEditor))
+		{
+			return PlaylistEditor;
+		}
+	}
+	return TSharedPtr<FAvaPlaylistEditor>();
+}
+
+#undef LOCTEXT_NAMESPACE
