[MotionDesign]

diff --git a/Engine/Plugins/Experimental/Avalanche/Source/AvalancheMediaEditor/Private/Playlist/AvaPlaylistEditorUtils.cpp b/Engine/Plugins/Experimental/Avalanche/Source/AvalancheMediaEditor/Private/Playlist/AvaPlaylistEditorUtils.cpp
new file mode 100644
--- /dev/null
+++ b/Engine/Plugins/Experimental/Avalanche/Source/AvalancheMediaEditor/Private/Playlist/AvaPlaylistEditorUtils.cpp
@@ -0,0 +1,625 @@
+// Copyright Epic Games, Inc. All Rights Reserved.
+
+#include "Playlist/AvaPlaylistEditorUtils.h"
+
+#include "Backends/JsonStructDeserializerBackend.h"
+#include "Backends/JsonStructSerializerBackend.h"
+#include "Backends/XmlStructSerializerBackend.h"
+#include "ContentBrowserModule.h"
+#include "DesktopPlatformModule.h"
+#include "Framework/Application/SlateApplication.h"
+#include "HAL/FileManager.h"
+#include "IAssetTools.h"
+#include "IContentBrowserSingleton.h"
+#include "JsonObjectConverter.h"
+#include "Misc/PathViews.h"
+#include "Playlist/AvalanchePlaylist.h"
+#include "StructDeserializer.h"
+#include "StructSerializer.h"
+
+#define LOCTEXT_NAMESPACE "AvaPlaylistEditor_Utils"
+
+namespace UE::AvaPlaylistEditor::Utils::Private
+{
+	static FString LastExportPath;
+	static FString LastImportPath;
+	
+	static const FString PageEntriesName = TEXT("Pages");
+
+	// Don't serialize the transient properties.
+	static auto TransientPropertyFiler = [](const FProperty* InCurrentProp, const FProperty* InParentProp)
+	{
+		const bool bIsTransient = InCurrentProp && InCurrentProp->HasAnyPropertyFlags(CPF_Transient); 
+		return !bIsTransient; 
+	};
+
+	struct FPlaylistSerializerPolicies : public FStructSerializerPolicies
+	{
+		FPlaylistSerializerPolicies()
+		{
+			PropertyFilter = TransientPropertyFiler;
+		}
+	};
+
+	struct FPlaylistDeserializerPolicies : public FStructDeserializerPolicies
+	{
+		FPlaylistDeserializerPolicies()
+		{
+			PropertyFilter = TransientPropertyFiler;
+		}
+	};
+	
+	TArray<TSharedPtr<FJsonValue>> PagesToJsonObjects(UAvalanchePlaylist* InPlaylist, const TArray<int32>& InPageIds)
+	{
+		check(InPlaylist);
+		TArray<TSharedPtr<FJsonValue>> PageEntries;
+		PageEntries.Reserve(InPageIds.Num());
+
+		for (const int32 PageId : InPageIds)
+		{
+			const FAvalanchePage& Page = InPlaylist->GetPage(PageId);
+
+			if (Page.IsValidPage())
+			{
+				TSharedRef<FJsonObject> PageObject = MakeShared<FJsonObject>();
+				FJsonObjectConverter::UStructToJsonObject(FAvalanchePage::StaticStruct(), &Page, PageObject, 0 /* CheckFlags */, 0 /* SkipFlags */);
+				PageEntries.Add(MakeShared<FJsonValueObject>(PageObject));
+			}
+		}
+
+		return PageEntries;
+	}
+
+	const FAvalanchePage& FindPage(const TArray<FAvalanchePage>& InPages, int32 InPageIdToFind)
+	{
+		for (const FAvalanchePage& Page : InPages)
+		{
+			if (Page.GetPageId() == InPageIdToFind)
+			{
+				return Page;
+			}
+		}
+		return FAvalanchePage::NullPage;
+	}
+	
+	const FAvalanchePage& FindTemplateForSourcePage(const UAvalanchePlaylist* InPlaylist, const FAvalanchePage& InSourcePage,
+		const TArray<FAvalanchePage>& InSourceTemplates, FImportTemplateMap& InOutImportedTemplateIds)
+	{
+		{
+			// Check if the template is already imported/existing at the given TemplateId.
+			const FAvalanchePage& ExistingTemplate = InPlaylist->GetPage(InOutImportedTemplateIds.GetTemplateId(InSourcePage.GetTemplateId()));
+			if (ExistingTemplate.IsValidPage() && ExistingTemplate.IsTemplate() && ExistingTemplate.GetAvalancheAssetPathDirect() == InSourcePage.GetAvalancheAssetPathDirect())
+			{
+				return ExistingTemplate;
+			}
+		}
+
+		// Fallback: Try to find a match using the source template if available.
+		const FAvalanchePage& SourceTemplate = FindPage(InSourceTemplates, InSourcePage.GetTemplateId());
+		if (SourceTemplate.IsValidPage())
+		{
+			// Try to find that template in the playlist with an exact match (rc values, asset, etc).
+			const FAvalanchePageCollection& PageCollection = InPlaylist->GetTemplatePages();
+			for (const FAvalanchePage& ExistingTemplate : PageCollection.Pages)
+			{
+				if (ExistingTemplate.IsTemplateMatchingByValue(SourceTemplate))
+				{
+					// Keep track of the match we made for next time.
+					InOutImportedTemplateIds.Add(InSourcePage.GetTemplateId(), ExistingTemplate.GetPageId());
+					return ExistingTemplate;
+				}
+			}
+		}
+		
+		return FAvalanchePage::NullPage;
+	}
+
+	bool CopyPageInPlace(UAvalanchePlaylist* InPlaylist, int32 InPageId, const FAvalanchePage& InSourcePage, int32 InTemplateId)
+	{
+		FAvalanchePage& DestinationPage = InPlaylist->GetPage(InPageId);
+		if (DestinationPage.IsValidPage())
+		{
+			DestinationPage = InSourcePage;
+			DestinationPage.SetPageId(InPageId);		// Restore page id.
+			DestinationPage.SetTemplateId(InTemplateId);
+			return true;
+		}
+		UE_LOG(LogAvaPlaylist, Error, TEXT("Failed to copy page in plage: page id %d is not found in destination playlist."), InPageId);
+		return false;
+	}
+}
+
+FString UE::AvaPlaylistEditor::Utils::GetImportFilepath(const TCHAR* InFileDescription, const TCHAR* InExtension)
+{
+	// Reference: UAssetToolsImpl::ExportAssetsInternal.
+	TArray<FString> OpenFilenames;
+	IDesktopPlatform* DesktopPlatform = FDesktopPlatformModule::Get();
+	bool bOpened = false;
+	if (DesktopPlatform)
+	{
+		if (Private::LastImportPath.IsEmpty())
+		{
+			Private::LastImportPath = FPaths::ProjectSavedDir();
+		}
+		
+		const FString FileType = FString::Printf(TEXT("%s (*.%s)|*.%s"), InFileDescription, InExtension, InExtension);
+		
+		bOpened = DesktopPlatform->OpenFileDialog(
+			FSlateApplication::Get().FindBestParentWindowHandleForDialogs(nullptr),
+			FText::Format(LOCTEXT("Import_F", "Import {0}"), FText::FromString(InFileDescription)).ToString(),
+			*Private::LastImportPath,
+			TEXT(""),
+			*FileType,
+			EFileDialogFlags::None,
+			OpenFilenames
+		);
+	}
+	if (bOpened && OpenFilenames.Num() > 0 && OpenFilenames[0].IsEmpty() == false)
+	{
+		Private::LastImportPath = OpenFilenames[0];
+		return OpenFilenames[0];
+	}
+	return FString();
+}
+
+FString UE::AvaPlaylistEditor::Utils::GetExportFilepath(const UObject* InObjectToExport, const TCHAR* InFileDescription, const TCHAR* InExtension)
+{
+	// Reference: UAssetToolsImpl::ExportAssetsInternal.
+	TArray<FString> SaveFilenames;
+	IDesktopPlatform* DesktopPlatform = FDesktopPlatformModule::Get();
+	bool bSave = false;
+	if (DesktopPlatform)
+	{
+		if (Private::LastExportPath.IsEmpty())
+		{
+			Private::LastExportPath = FPaths::ProjectSavedDir();
+		}
+		
+		const FString FileType = FString::Printf(TEXT("%s (*.%s)|*.%s"), InFileDescription, InExtension, InExtension);
+		
+		bSave = DesktopPlatform->SaveFileDialog(
+			FSlateApplication::Get().FindBestParentWindowHandleForDialogs(nullptr),
+			FText::Format(LOCTEXT("Export_F", "Export: {0}"), FText::FromString(InObjectToExport->GetName())).ToString(),
+			*Private::LastExportPath,
+			*InObjectToExport->GetName(),
+			*FileType,
+			EFileDialogFlags::None,
+			SaveFilenames
+		);
+	}
+	if (bSave && SaveFilenames.Num() > 0)
+	{
+		Private::LastExportPath = SaveFilenames[0];
+		return SaveFilenames[0];
+	}
+	return FString();
+}
+
+FString UE::AvaPlaylistEditor::Utils::GetSaveAssetAsPath(const FString& InDefaultPath, const FString& InDefaultAssetName)
+{
+	FSaveAssetDialogConfig SaveAssetDialogConfig;
+	{
+		SaveAssetDialogConfig.DefaultPath = InDefaultPath;
+		SaveAssetDialogConfig.DefaultAssetName = InDefaultAssetName;
+		SaveAssetDialogConfig.AssetClassNames.Add(UAvalanchePlaylist::StaticClass()->GetClassPathName());
+		SaveAssetDialogConfig.ExistingAssetPolicy = ESaveAssetDialogExistingAssetPolicy::AllowButWarn;
+		SaveAssetDialogConfig.DialogTitleOverride = LOCTEXT("SaveAssetDialogTitle", "Save Asset As");
+	}
+	
+	const FContentBrowserModule& ContentBrowserModule = FModuleManager::LoadModuleChecked<FContentBrowserModule>("ContentBrowser");
+	return ContentBrowserModule.Get().CreateModalSaveAssetDialog(SaveAssetDialogConfig);
+}
+
+FString UE::AvaPlaylistEditor::Utils::SerializePagesToJson(UAvalanchePlaylist* InPlaylist, const TArray<int32>& InPageIds)
+{
+	if (!IsValid(InPlaylist))
+	{
+		return FString();
+	}
+	
+	const TArray<TSharedPtr<FJsonValue>> PageEntries = Private::PagesToJsonObjects(InPlaylist, InPageIds);
+	
+	if (PageEntries.IsEmpty())
+	{
+		return FString();
+	}
+
+	const TSharedRef<FJsonObject> RootJsonObject = MakeShared<FJsonObject>();
+	RootJsonObject->SetArrayField(Private::PageEntriesName, PageEntries);
+
+	FString SerializedString;
+	const TSharedRef<TJsonWriter<>> Writer = TJsonWriterFactory<>::Create(&SerializedString);
+
+	FJsonSerializer::Serialize(RootJsonObject, Writer);
+	
+	return SerializedString;
+}
+
+TArray<FAvalanchePage> UE::AvaPlaylistEditor::Utils::DeserializePagesFromJson(const FString& InJsonString)
+{
+	TSharedPtr<FJsonObject> RootJsonObject;
+	const TSharedRef<TJsonReader<>> Reader = TJsonReaderFactory<>::Create(InJsonString);
+
+	if (!FJsonSerializer::Deserialize(Reader, RootJsonObject))
+	{
+		UE_LOG(LogAvaPlaylist, Warning, TEXT("Unable to serialize the pasted text into Json format"));
+		return {};
+	}
+
+	const TArray<TSharedPtr<FJsonValue>>* PageEntries;
+	if (!RootJsonObject->TryGetArrayField(Private::PageEntriesName, PageEntries))
+	{
+		UE_LOG(LogAvaPlaylist, Warning, TEXT("Missing %s entry field in pasted text"), *Private::PageEntriesName);
+		return {};
+	}
+
+	TArray<FAvalanchePage> Pages;
+	Pages.Reserve(PageEntries->Num());
+	for (const TSharedPtr<FJsonValue>& PageEntry : *PageEntries)
+	{
+		if (!PageEntry.IsValid() || PageEntry->Type != EJson::Object)
+		{
+			UE_LOG(LogAvaPlaylist, Warning, TEXT("Invalid page entry. Not an object"));
+			continue;
+		}
+
+		const TSharedPtr<FJsonObject>& PageObject = PageEntry->AsObject();
+		check(PageObject.IsValid());
+
+		FAvalanchePage Page;
+		if (FJsonObjectConverter::JsonObjectToUStruct(PageObject.ToSharedRef(), FAvalanchePage::StaticStruct(), &Page, 0 /* CheckFlags */, 0 /* SkipFlags */))
+		{
+			Pages.Emplace(MoveTemp(Page));
+		}
+		else
+		{
+			UE_LOG(LogAvaPlaylist, Warning, TEXT("Unable to convert Page Entry Json Object to Avalanche Page Struct"));
+		}
+	}
+	return Pages;
+}
+
+TStrongObjectPtr<UAvalanchePlaylist> UE::AvaPlaylistEditor::Utils::ExportPagesToPlaylist(const UAvalanchePlaylist* InPlaylist, const TArray<int32>& InPageIds)
+{
+	if (!InPlaylist || InPageIds.IsEmpty())
+	{
+		return nullptr;
+	}
+
+	TSet<int32> AddedTemplates;
+	TArray<FAvalanchePage> SourcePages;
+	TArray<FAvalanchePage> SourceTemplates;
+	bool bExportTemplates = false;
+	
+	for (const int32 PageId : InPageIds)
+	{
+		const FAvalanchePage& Page = InPlaylist->GetPage(PageId);
+		if (!Page.IsValidPage())
+		{
+			continue;
+		}
+
+		if (Page.IsTemplate() && !AddedTemplates.Contains(Page.GetPageId()))
+		{
+			SourceTemplates.Add(Page);
+			AddedTemplates.Add(Page.GetTemplateId());
+			bExportTemplates = true; 
+			continue;
+		}
+		
+		SourcePages.Add(Page);
+		if (!AddedTemplates.Contains(Page.GetTemplateId()))
+		{
+			const FAvalanchePage& Template = InPlaylist->GetPage(Page.GetTemplateId());
+			if (Template.IsValidPage())
+			{
+				SourceTemplates.Add(Template);
+				AddedTemplates.Add(Template.GetPageId());
+			}
+		}
+	}
+
+	if (!SourcePages.IsEmpty())
+	{
+		const TStrongObjectPtr<UAvalanchePlaylist> NewPlaylist(NewObject<UAvalanchePlaylist>());
+		FImportTemplateMap ImportedTemplateIds;
+		
+		if (bExportTemplates)
+		{
+			ImportTemplatePages(NewPlaylist.Get(), SourceTemplates, ImportedTemplateIds);
+		}
+		
+		ImportInstancedPages(NewPlaylist.Get(), UAvalanchePlaylist::InstancePageList, SourcePages, SourceTemplates, ImportedTemplateIds);
+		return NewPlaylist;
+	}
+
+	return nullptr;
+}
+
+bool UE::AvaPlaylistEditor::Utils::SavePlaylistToXml(const UAvalanchePlaylist* InPlaylist, FArchive& InArchive, EXmlSerializationEncoding InXmlEncoding)
+{
+	if (IsValid(InPlaylist))
+	{
+		// Note: using the StructSerializerBackend produces a more compact format and is more suitable for exporting compared to
+		// FXmlArchiveOutputFormatter. However, it doesn't support serializing UObject in place and the xml deserializer hasn't been
+		// implemented yet. Support for serialization of UObjects in place is not planned to be needed for playlist at the moment.
+		FXmlStructSerializerBackend Backend(InArchive, EStructSerializerBackendFlags::Default);
+		FStructSerializer::Serialize(InPlaylist, *InPlaylist->GetClass(), Backend, Private::FPlaylistSerializerPolicies());
+		Backend.SaveDocument(InXmlEncoding);
+		return true;
+	}
+	return false;
+}
+
+bool UE::AvaPlaylistEditor::Utils::SavePlaylistToXml(const UAvalanchePlaylist* InPlaylist, const TCHAR* InFilepath)
+{
+	if (IsValid(InPlaylist))
+	{
+		const TUniquePtr<FArchive> FileWriter(IFileManager::Get().CreateFileWriter(InFilepath));
+		if (FileWriter)	
+		{
+			const bool bSaved = SavePlaylistToXml(InPlaylist, *FileWriter, EXmlSerializationEncoding::Utf8);
+			FileWriter->Close();
+			return bSaved;
+		}
+	}
+	return false;
+}
+
+bool UE::AvaPlaylistEditor::Utils::SavePlaylistToJson(const UAvalanchePlaylist* InPlaylist, FArchive& InArchive)
+{
+	if (IsValid(InPlaylist))
+	{
+		FJsonStructSerializerBackend Backend(InArchive, EStructSerializerBackendFlags::Default);
+		FStructSerializer::Serialize(InPlaylist, *InPlaylist->GetClass(), Backend, Private::FPlaylistSerializerPolicies());
+		return true;
+	}
+	return false;
+}
+
+bool UE::AvaPlaylistEditor::Utils::SavePlaylistToJson(const UAvalanchePlaylist* InPlaylist, const TCHAR* InFilepath)
+{
+	if (IsValid(InPlaylist))
+	{
+		const TUniquePtr<FArchive> FileWriter(IFileManager::Get().CreateFileWriter(InFilepath));
+		if (FileWriter)
+		{
+			const bool bSerialized = SavePlaylistToJson(InPlaylist, *FileWriter);
+			FileWriter->Close();
+			return bSerialized;
+		}
+	}
+	return false;
+}
+
+bool UE::AvaPlaylistEditor::Utils::LoadPlaylistFromJson(UAvalanchePlaylist* InPlaylist, FArchive& InArchive)
+{
+	if (IsValid(InPlaylist))
+	{
+		// Serializing doesn't reset content, it will add to it, so we need
+		// to explicitly make the playlist empty first.
+		if (!InPlaylist->Empty())
+		{
+			return false;
+		}
+		
+		FJsonStructDeserializerBackend Backend(InArchive);
+		const bool bLoaded = FStructDeserializer::Deserialize(InPlaylist, *InPlaylist->GetClass(), Backend, Private::FPlaylistDeserializerPolicies());
+		if (bLoaded)
+		{
+			InPlaylist->PostLoad();
+		}
+		return bLoaded;
+	}
+	return false;
+}
+
+bool UE::AvaPlaylistEditor::Utils::LoadPlaylistFromJson(UAvalanchePlaylist* InPlaylist, const TCHAR* InFilepath)
+{
+	if (IsValid(InPlaylist))
+	{
+		const TUniquePtr<FArchive> FileReader(IFileManager::Get().CreateFileReader(InFilepath));
+		if (FileReader)
+		{
+			const bool bLoaded = LoadPlaylistFromJson(InPlaylist, *FileReader);
+			FileReader->Close();
+			return bLoaded;
+		}
+	}
+	return false;
+}
+
+bool UE::AvaPlaylistEditor::Utils::CanLoadPlaylistFromFile(const TCHAR* InFilepath)
+{
+	const FStringView Extension = FPathViews::GetExtension(InFilepath);
+	
+	// We only support loading from json file (LoadPlaylistFromJson).
+	return Extension.Equals(TEXT("json"), ESearchCase::IgnoreCase);
+}
+
+TArray<int32> UE::AvaPlaylistEditor::Utils::ImportTemplatePages(UAvalanchePlaylist* InPlaylist, const TArray<FAvalanchePage>& InSourceTemplates, FImportTemplateMap& OutImportedTemplateIds)
+{
+	if (!IsValid(InPlaylist))
+	{
+		return {};
+	}
+
+	TArray<int32> OutTemplateIds;
+	OutTemplateIds.Reserve(InSourceTemplates.Num());
+
+	for (const FAvalanchePage& SourceTemplate : InSourceTemplates)
+	{
+		const int32 SourceTemplateId = SourceTemplate.IsTemplate() ? SourceTemplate.GetPageId() : SourceTemplate.GetTemplateId();
+		// Try to add the template with the id it had in the original list.
+		int32 ImportedTemplateId = InPlaylist->AddTemplate(FAvaPageIdGeneratorParams(SourceTemplateId));
+		
+		if (ImportedTemplateId != FAvalanchePage::InvalidPageId)
+		{
+			OutImportedTemplateIds.Add(SourceTemplateId, ImportedTemplateId);
+			ensure(Private::CopyPageInPlace(InPlaylist, ImportedTemplateId, SourceTemplate, FAvalanchePage::InvalidPageId));
+			OutTemplateIds.Add(ImportedTemplateId);
+		}
+	}
+	return OutTemplateIds;
+}
+
+TArray<int32> UE::AvaPlaylistEditor::Utils::ImportInstancedPages(
+	UAvalanchePlaylist* InPlaylist,
+	const FAvaPageListReference& InPageListReference,
+	const TArray<FAvalanchePage>& InSourcePages,
+	const TArray<FAvalanchePage>& InSourceTemplates,
+	FImportTemplateMap& InOutImportedTemplateIds,
+	const FAvaPageInsertPosition& InInsertPosition)
+{
+	if (!IsValid(InPlaylist))
+	{
+		return {};
+	}
+	
+	TArray<int32> OutPageIds;
+
+	FAvaPageInsertPosition InsertPosition = InInsertPosition;
+
+	// If we are adding above, iteration should be reversed, so the last is added first
+	// and the next to last added above that, etc.
+	const bool bReverseIteration = InsertPosition.IsValid() && !InsertPosition.bAddBelow;
+	
+	for (int32 PageIndex = 0; PageIndex < InSourcePages.Num(); ++PageIndex)
+	{
+		const FAvalanchePage& SourcePage = InSourcePages[bReverseIteration ? InSourcePages.Num() - PageIndex - 1 : PageIndex];
+		
+		if (!SourcePage.IsValidPage())
+		{
+			continue;
+		}
+
+		// Attempt to find/create template for this page
+		int32 ImportedTemplateId = FAvalanchePage::InvalidPageId;
+
+		if (SourcePage.IsTemplate())
+		{
+			// Todo: untested case.
+			// Suspect this of being wrong. if the page is a template, then it needs to be imported
+			// as a template, i.e. ImportTemplatePages.
+			ImportedTemplateId = SourcePage.GetPageId();
+		}
+		else
+		{
+			// The source page has a source template id. It may not match the destination, this is
+			// why we rely on the ImportedTemplateIds to translate that. As a fallback, if the source
+			// templates are provided, it will try to match templates by values.
+			const FAvalanchePage& ExistingTemplate = Private::FindTemplateForSourcePage(InPlaylist, SourcePage, InSourceTemplates, InOutImportedTemplateIds);
+
+			if (ExistingTemplate.IsValidPage())
+			{
+				ImportedTemplateId = ExistingTemplate.GetPageId();
+			}
+			else
+			{
+				// Try to add the template with the id it had in the original list.
+				ImportedTemplateId = InPlaylist->AddTemplate(FAvaPageIdGeneratorParams(SourcePage.GetTemplateId()));
+				
+				if (ImportedTemplateId != FAvalanchePage::InvalidPageId)
+				{
+					// Keep track of correspondence for next pages in the list.
+					InOutImportedTemplateIds.Add(SourcePage.GetTemplateId(), ImportedTemplateId);
+					
+					// We either add the current page "as template", or use the source templates if provided.
+					const FAvalanchePage& SourceTemplate = Private::FindPage(InSourceTemplates, SourcePage.GetTemplateId());
+					
+					ensure(Private::CopyPageInPlace(InPlaylist, ImportedTemplateId, SourceTemplate.IsValidPage() ? SourceTemplate : SourcePage, FAvalanchePage::InvalidPageId));
+				}
+				else
+				{
+					// We are unable to find/create a template for this page.
+					continue;
+				}
+			}
+		}
+
+		// There is a valid imported template at this point.
+		check(ImportedTemplateId != FAvalanchePage::InvalidPageId);
+
+		// We're pasting to the instance list, so just add the page
+		if (InPageListReference.Type == EAvaPageListType::Instance)
+		{
+			// We want to preserve the source PageId if possible.
+			const FAvaPageIdGeneratorParams NewPageIdParams =
+				FAvaPageIdGeneratorParams::FromInsertPositionOrSourceId(SourcePage.GetPageId(), InsertPosition);
+			
+			int32 ImportedPageId = InPlaylist->AddPageFromTemplate(ImportedTemplateId, NewPageIdParams, InsertPosition);
+
+			if (ImportedPageId != FAvalanchePage::InvalidPageId)
+			{
+				ensure(Private::CopyPageInPlace(InPlaylist, ImportedPageId, SourcePage, ImportedTemplateId));
+				OutPageIds.Add(ImportedPageId);			
+				InsertPosition.ConditionalUpdateAdjacentId(ImportedPageId); // Update for next insertion.
+			}
+
+			continue;
+		}
+
+		// We're pasting to a sub list, so we need to check if the page exists in the instance list first
+		const FAvalanchePage& InstancedPage = InPlaylist->GetPage(SourcePage.GetPageId());
+		int32 InstancedPageId = InstancedPage.GetPageId();
+
+		// Add it if it's missing
+		if (!InstancedPage.IsValidPage() || InstancedPage.IsTemplate())
+		{
+			// We want to preserve the source PageId if possible.
+			const FAvaPageIdGeneratorParams NewPageIdParams =
+				FAvaPageIdGeneratorParams::FromInsertPositionOrSourceId(!InstancedPage.IsTemplate() ? SourcePage.GetPageId() : FAvalanchePage::InvalidPageId, InsertPosition);
+
+			// Note: If the page is a template, it will already be imported as templateId.
+			InstancedPageId = InPlaylist->AddPageFromTemplate(ImportedTemplateId, NewPageIdParams, InsertPosition);
+
+			if (InstancedPageId != FAvalanchePage::InvalidPageId)
+			{
+				ensure(Private::CopyPageInPlace(InPlaylist, InstancedPageId, SourcePage, ImportedTemplateId));
+				InsertPosition.ConditionalUpdateAdjacentId(InstancedPageId); // Update for next insertion.
+			}
+			else
+			{
+				// We were unable to create the instance page, so it cannot be added to a sub list
+				continue;
+			}
+		}
+
+		// There is a valid imported instanced page at this point.
+		check(InstancedPageId != FAvalanchePage::InvalidPageId);
+
+		// Now we have our instance page reference, add it to the sublist
+		if (InPlaylist->AddPageToSubList(InPageListReference.SubListIndex, InstancedPageId))
+		{
+			OutPageIds.Add(InstancedPageId);
+		}
+	}
+
+	if (bReverseIteration)
+	{
+		Algo::Reverse(OutPageIds);
+	}
+	
+	return OutPageIds;
+}
+
+TArray<int32> UE::AvaPlaylistEditor::Utils::ImportInstancedPagesFromPlaylist(UAvalanchePlaylist* InPlaylist, const UAvalanchePlaylist* InSourcePlaylist, const FAvaPageInsertPosition& InInsertPosition)
+{
+	if (!IsValid(InPlaylist) || !IsValid(InSourcePlaylist))
+	{
+		return {};
+	}
+	
+	const FAvalanchePageCollection& SourceTemplates = InSourcePlaylist->GetTemplatePages();
+	const FAvalanchePageCollection& SourcePages = InSourcePlaylist->GetInstancedPages();
+	FImportTemplateMap ImportedTemplateIds;
+	
+	return ImportInstancedPages(InPlaylist, UAvalanchePlaylist::InstancePageList, SourcePages.Pages, SourceTemplates.Pages, ImportedTemplateIds, InInsertPosition);
+}
+
+UAvalanchePlaylist* UE::AvaPlaylistEditor::Utils::SaveDuplicatePlaylist(UAvalanchePlaylist* InSourcePlaylist, const FString& InAssetName, const FString& InPackagePath)
+{
+	return Cast<UAvalanchePlaylist>(IAssetTools::Get().DuplicateAsset(InAssetName, InPackagePath, InSourcePlaylist));
+}
+
+#undef LOCTEXT_NAMESPACE
diff --git a/Engine/Plugins/Experimental/Avalanche/Source/AvalancheMediaEditor/Private/Playlist/AvaPlaylistEditorUtils.h b/Engine/Plugins/Experimental/Avalanche/Source/AvalancheMediaEditor/Private/Playlist/AvaPlaylistEditorUtils.h
new file mode 100644
--- /dev/null
+++ b/Engine/Plugins/Experimental/Avalanche/Source/AvalancheMediaEditor/Private/Playlist/AvaPlaylistEditorUtils.h
@@ -0,0 +1,110 @@
+// Copyright Epic Games, Inc. All Rights Reserved.
+
+#pragma once
+
+#include "Containers/Array.h"
+#include "Playlist/AvalanchePage.h"
+#include "Playlist/AvalanchePlaylist.h"
+#include "Templates/SharedPointer.h"
+#include "XmlSerializationDefines.h"
+
+class UAvalanchePlaylist;
+
+namespace UE::AvaPlaylistEditor::Utils
+{
+	FString GetImportFilepath(const TCHAR* InFileDescription, const TCHAR* InExtension);
+	FString GetExportFilepath(const UObject* InObjectToExport, const TCHAR* InFileDescription, const TCHAR* InExtension);
+	FString GetSaveAssetAsPath(const FString& InDefaultPath, const FString& InDefaultAssetName);
+	
+	FString SerializePagesToJson(UAvalanchePlaylist* InPlaylist, const TArray<int32>& InPageIds);
+	TArray<FAvalanchePage> DeserializePagesFromJson(const FString& InJsonString);
+
+	/**
+	 * @brief Creates a transient playlist with only the specified pages.
+	 * @param InPlaylist Source Playlist.
+	 * @param InPageIds Pages to copy to the new playlist.
+	 * @return Created playlist.
+	 */
+	TStrongObjectPtr<UAvalanchePlaylist> ExportPagesToPlaylist(const UAvalanchePlaylist* InPlaylist, const TArray<int32>& InPageIds);
+
+	bool SavePlaylistToXml(const UAvalanchePlaylist* InPlaylist, FArchive& InArchive, EXmlSerializationEncoding InXmlEncoding);
+	bool SavePlaylistToXml(const UAvalanchePlaylist* InPlaylist, const TCHAR* InFilepath);
+
+	bool SavePlaylistToJson(const UAvalanchePlaylist* InPlaylist, FArchive& InArchive);
+	bool SavePlaylistToJson(const UAvalanchePlaylist* InPlaylist, const TCHAR* InFilepath);
+
+	bool LoadPlaylistFromJson(UAvalanchePlaylist* InPlaylist, FArchive& InArchive);
+	bool LoadPlaylistFromJson(UAvalanchePlaylist* InPlaylist, const TCHAR* InFilepath);
+
+	/**
+	 * Check the file extenstion to see if it is a supported format to use 
+	 * with the deserializing functions (LoadPlaylistFrom...).
+	 */
+	bool CanLoadPlaylistFromFile(const TCHAR* InFilepath);
+
+	/**
+	 * Map of "Source TemplateId" to "Destination Template Id".
+	 * 
+	 * During the import of new pages, we need to keep track of the
+	 * destination template Id, i.e. Source Page's template id -> actual id in playlist.
+	 */
+	struct FImportTemplateMap
+	{
+		TMap<int32, int32> TemplateIds;
+
+		void Add(int32 InSourceTemplateId, int32 InDestinationTemplateId)
+		{
+			TemplateIds.Add(InSourceTemplateId, InDestinationTemplateId);
+		}
+
+		bool HasSourceTemplateId(int32 InSourceTemplateId) const
+		{
+			return TemplateIds.Contains(InSourceTemplateId);
+		}
+
+		int32 GetTemplateId(int32 InSourceTemplateId) const
+		{
+			const int32* FoundId = TemplateIds.Find(InSourceTemplateId);
+			return FoundId ? *FoundId : InSourceTemplateId;
+		}
+	};
+	
+	TArray<int32> ImportTemplatePages(UAvalanchePlaylist* InPlaylist, const TArray<FAvalanchePage>& InSourceTemplates, FImportTemplateMap& OutImportedTemplateIds);
+
+	/**
+	 * @brief Import the given pages to the playlist.
+	 * @param InPlaylist Playlist that receives the new pages. 
+	 * @param InPageListReference Reference to the page list to which the pages will be added.
+	 * @param InSourcePages Source pages
+	 * @param InSourceTemplates Source templates, if available, for the source pages. 
+	 * @param InOutImportedTemplateIds Imported template ids to make the correspondence between the source page template id and what was imported.
+	 * @param InInsertPosition Indicate where in the page list to insert the new pages.
+	 * @return Array of added page ids.
+	 *
+	 * Importing pages is a non-trivial operation. It includes adding missing templates. The biggest issue
+	 * is maintaining the link between the page's template id and the imported templates. If available, the
+	 * source templates can be used to reconstruct this relationship by finding an exact match with the templates
+	 * of the destination playlist.
+	 */
+	TArray<int32> ImportInstancedPages(
+		UAvalanchePlaylist* InPlaylist,
+		const FAvaPageListReference& InPageListReference,
+		const TArray<FAvalanchePage>& InSourcePages,
+		const TArray<FAvalanchePage>& InSourceTemplates,
+		FImportTemplateMap& InOutImportedTemplateIds,
+		const FAvaPageInsertPosition& InInsertPosition = FAvaPageInsertPosition());
+
+	/**
+	 * @brief Import/Merge a playlist into another playlist.
+	 * @param InPlaylist Destination playlist.
+	 * @param InSourcePlaylist Source playlist to merge in the destination.
+	 * @param InInsertPosition Indicate where in the page list to insert the new pages.
+	 * @return Array of added page ids.
+	 */
+	TArray<int32> ImportInstancedPagesFromPlaylist(
+		UAvalanchePlaylist* InPlaylist,
+		const UAvalanchePlaylist* InSourcePlaylist,
+		const FAvaPageInsertPosition& InInsertPosition = FAvaPageInsertPosition());
+	
+	UAvalanchePlaylist* SaveDuplicatePlaylist(UAvalanchePlaylist* InSourcePlaylist, const FString& InAssetName, const FString& InPackagePath);
+}
\ No newline at end of file
diff --git a/Engine/Plugins/Experimental/Avalanche/Source/AvalancheMediaEditor/Private/Playlist/AvaPlaylistExporter.cpp b/Engine/Plugins/Experimental/Avalanche/Source/AvalancheMediaEditor/Private/Playlist/AvaPlaylistExporter.cpp
new file mode 100644
--- /dev/null
+++ b/Engine/Plugins/Experimental/Avalanche/Source/AvalancheMediaEditor/Private/Playlist/AvaPlaylistExporter.cpp
@@ -0,0 +1,51 @@
+// Copyright Epic Games, Inc. All Rights Reserved.
+
+#include "Playlist/AvaPlaylistExporter.h"
+
+#include "AvaSerializationUtils.h"
+#include "Playlist/AvaPlaylistEditorUtils.h"
+#include "Playlist/AvalanchePlaylist.h"
+#include "Serialization/MemoryWriter.h"
+
+UAvaPlaylistExporter::UAvaPlaylistExporter()
+{
+	FormatExtension.Add(TEXT("json"));
+	FormatDescription.Add(TEXT("JavaScript Object Notation file"));
+	FormatExtension.Add(TEXT("xml"));
+	FormatDescription.Add(TEXT("eXtensible Markup Language file"));
+	SupportedClass = UAvalanchePlaylist::StaticClass();
+	bText = true;
+}
+
+bool UAvaPlaylistExporter::ExportText(const FExportObjectInnerContext* InContext, UObject* InObject, const TCHAR* InType, FOutputDevice& InAr, FFeedbackContext* InWarn, uint32 InPortFlags)
+{
+	const UAvalanchePlaylist* Playlist = Cast<UAvalanchePlaylist>(InObject);
+	if (!Playlist)
+	{
+		return false;
+	}
+
+	TArray<uint8> OutputBytes;
+	FMemoryWriter Writer(OutputBytes);
+	bool bSavedToBytes = false;
+
+	if ( FCString::Stricmp(InType, TEXT("json")) == 0 )
+	{
+		bSavedToBytes = UE::AvaPlaylistEditor::Utils::SavePlaylistToJson(Playlist, Writer);
+	}
+	else if (FCString::Stricmp(InType, TEXT("xml")) == 0 )
+	{
+		// Note: using WChar encoding for compatibility with BytesToString below.
+		bSavedToBytes = UE::AvaPlaylistEditor::Utils::SavePlaylistToXml(Playlist, Writer, EXmlSerializationEncoding::WChar);
+	}
+
+	if (bSavedToBytes)
+	{
+		FString OutputString;
+		UE::AvaSerializationUtils::JsonValueConversion::BytesToString(OutputBytes, OutputString);
+		InAr.Log(OutputString);
+		return true;
+	}
+	
+	return false;	
+}
\ No newline at end of file
diff --git a/Engine/Plugins/Experimental/Avalanche/Source/AvalancheMediaEditor/Private/Playlist/AvaPlaylistExporter.h b/Engine/Plugins/Experimental/Avalanche/Source/AvalancheMediaEditor/Private/Playlist/AvaPlaylistExporter.h
new file mode 100644
--- /dev/null
+++ b/Engine/Plugins/Experimental/Avalanche/Source/AvalancheMediaEditor/Private/Playlist/AvaPlaylistExporter.h
@@ -0,0 +1,20 @@
+// Copyright Epic Games, Inc. All Rights Reserved.
+
+#pragma once
+
+#include "Exporters/Exporter.h"
+
+#include "AvaPlaylistExporter.generated.h"
+
+UCLASS()
+class UAvaPlaylistExporter : public UExporter
+{
+	GENERATED_BODY()
+
+public:
+	UAvaPlaylistExporter();
+	
+	//~ Begin UExporter Interface
+	virtual bool ExportText(const FExportObjectInnerContext* InContext, UObject* InObject, const TCHAR* InType, FOutputDevice& InAr, FFeedbackContext* InWarn, uint32 InPortFlags) override;
+	//~ End UExporter Interface
+};
\ No newline at end of file
diff --git a/Engine/Plugins/Experimental/Avalanche/Source/AvalancheMediaEditor/Private/Playlist/AvaPlaylistFactory.cpp b/Engine/Plugins/Experimental/Avalanche/Source/AvalancheMediaEditor/Private/Playlist/AvaPlaylistFactory.cpp
new file mode 100644
--- /dev/null
+++ b/Engine/Plugins/Experimental/Avalanche/Source/AvalancheMediaEditor/Private/Playlist/AvaPlaylistFactory.cpp
@@ -0,0 +1,61 @@
+// Copyright Epic Games, Inc. All Rights Reserved.
+
+#include "Playlist/AvaPlaylistFactory.h"
+#include "AssetToolsModule.h"
+#include "IAssetTools.h"
+#include "Playlist/AvalanchePlaylist.h"
+
+UAvaPlaylistFactory::UAvaPlaylistFactory()
+{
+	// Provide the factory with information about how to handle our asset
+	bCreateNew = true;
+	bEditAfterNew = true;
+	SupportedClass = UAvalanchePlaylist::StaticClass();
+}
+
+UAvaPlaylistFactory::~UAvaPlaylistFactory()
+{
+}
+
+uint32 UAvaPlaylistFactory::GetMenuCategories() const
+{
+	IAssetTools& AssetTools = FAssetToolsModule::GetModule().Get();
+	return AssetTools.FindAdvancedAssetCategory("AvalancheCategory");
+}
+
+bool UAvaPlaylistFactory::ConfigureProperties()
+{
+	return Super::ConfigureProperties();
+}
+
+bool UAvaPlaylistFactory::ShouldShowInNewMenu() const
+{
+	return Super::ShouldShowInNewMenu();
+}
+
+UObject* UAvaPlaylistFactory::FactoryCreateNew(UClass* Class, UObject* InParent, FName Name, EObjectFlags Flags
+	, UObject* Context, FFeedbackContext* Warn, FName CallingContext)
+{
+	return Super::FactoryCreateNew(Class, InParent, Name, Flags, Context, Warn, CallingContext);
+}
+
+UObject* UAvaPlaylistFactory::FactoryCreateNew(UClass* Class, UObject* InParent, FName Name, EObjectFlags Flags
+	, UObject* Context, FFeedbackContext* Warn)
+{
+	UAvalanchePlaylist* Playlist = nullptr;
+	if (ensure(SupportedClass == Class))
+	{
+		Playlist = NewObject<UAvalanchePlaylist>(InParent, Name, Flags);
+	}
+	return Playlist;
+}
+
+bool UAvaPlaylistFactory::DoesSupportClass(UClass* Class)
+{
+	return Class == UAvalanchePlaylist::StaticClass();
+}
+
+UClass* UAvaPlaylistFactory::ResolveSupportedClass()
+{
+	return Super::ResolveSupportedClass();
+}
diff --git a/Engine/Plugins/Experimental/Avalanche/Source/AvalancheMediaEditor/Private/Playlist/AvaPlaylistPlaybackUtils.cpp b/Engine/Plugins/Experimental/Avalanche/Source/AvalancheMediaEditor/Private/Playlist/AvaPlaylistPlaybackUtils.cpp
new file mode 100644
--- /dev/null
+++ b/Engine/Plugins/Experimental/Avalanche/Source/AvalancheMediaEditor/Private/Playlist/AvaPlaylistPlaybackUtils.cpp
@@ -0,0 +1,70 @@
+// Copyright Epic Games, Inc. All Rights Reserved.
+
+#include "Playlist/AvaPlaylistPlaybackUtils.h"
+
+#include "Playlist/AvalanchePlaylist.h"
+
+namespace UE::AvaPlaylistPlaybackUtils::Private
+{
+	template <typename Predicate>
+	TArray<int32> FilterSelectedOrPreviewingPages(const UAvalanchePlaylist* InPlaylist, const TArray<int32>& InSelectedPageIds, FName InPreviewChannelName, Predicate InPredicate)
+	{
+		if (InPlaylist)
+		{
+			if (InSelectedPageIds.Num())
+			{
+				return ::Invoke(InPredicate, InPlaylist, InSelectedPageIds);
+			}
+			return ::Invoke(InPredicate, InPlaylist, InPlaylist->GetPreviewingPageIds(InPreviewChannelName));
+		}
+		return {};
+	}
+}
+
+// Selecting the "next page" implies we know what the last page played was and we go to the next.
+// Normally, the last page played was the one currently selected, but the user might decide to select
+// another page for editing purpose, but still want to hit play next from the last played page.
+// This is why the "play head" is now added in the playback context for the given list.
+int32 FAvaPlaylistPlaybackUtils::GetPageIdToPlayNext(const UAvalanchePlaylist* InPlaylist, const FAvaPageListReference& InPageListReference, bool bInPreview, FName InPreviewChannelName)
+{
+	// If there is no playback context, it means no pages are playing and can't have a next page if nothing is playing.
+	if (InPlaylist && InPlaylist->GetPageListPlaybackContextCollection())
+	{
+		if (const TSharedPtr<FAvaPageListPlaybackContext> PlaybackContext = InPlaylist->GetPageListPlaybackContextCollection()->GetContext(bInPreview, InPreviewChannelName))
+		{
+			// Check if the current play head page is playing.
+			const bool bIsHeadPlaying = bInPreview ? InPlaylist->IsPagePreviewing(PlaybackContext->PlayHeadPageId) : InPlaylist->IsPagePlaying(PlaybackContext->PlayHeadPageId); 
+			
+			if (bIsHeadPlaying)
+			{
+				const int32 NextPageId = InPlaylist->GetNextPage(PlaybackContext->PlayHeadPageId, InPageListReference).GetPageId();
+				if (IsPageIdValid(NextPageId) && InPlaylist->CanPlayPage(NextPageId, true))
+				{
+					return NextPageId;
+				}
+			}
+		}
+	}
+	return FAvalanchePage::InvalidPageId;
+}
+
+TArray<int32> FAvaPlaylistPlaybackUtils::GetPagesToTakeToProgram(const UAvalanchePlaylist* InPlaylist, const TArray<int32>& InSelectedPageIds, FName InPreviewChannel)
+{
+	using namespace UE::AvaPlaylistPlaybackUtils::Private;
+	auto KeepPagesToTakeToProgram = [](const UAvalanchePlaylist* InPlaylist, const TArray<int32>& InPageIds)
+	{
+		TArray<int32> OutPageIds;
+		OutPageIds.Reserve(InPageIds.Num());
+		for (int32 PageId : InPageIds)
+		{
+			// We can't take templates to program. Only instanced pages.
+			// Note: CanPlayPage returns true for playing pages, so need to make sure it is not already playing.
+			if (InPlaylist->CanPlayPage(PageId, false) && !InPlaylist->IsPagePlaying(PageId))
+			{
+				OutPageIds.Add(PageId);
+			}
+		}
+		return OutPageIds;
+	};
+	return FilterSelectedOrPreviewingPages(InPlaylist, InSelectedPageIds, InPreviewChannel, KeepPagesToTakeToProgram);
+}
diff --git a/Engine/Plugins/Experimental/Avalanche/Source/AvalancheMediaEditor/Private/Playlist/AvaPlaylistPlaybackUtils.h b/Engine/Plugins/Experimental/Avalanche/Source/AvalancheMediaEditor/Private/Playlist/AvaPlaylistPlaybackUtils.h
new file mode 100644
--- /dev/null
+++ b/Engine/Plugins/Experimental/Avalanche/Source/AvalancheMediaEditor/Private/Playlist/AvaPlaylistPlaybackUtils.h
@@ -0,0 +1,19 @@
+// Copyright Epic Games, Inc. All Rights Reserved.
+
+#pragma once
+
+#include "Playlist/AvalanchePage.h"
+
+class UAvalanchePlaylist;
+
+struct FAvaPlaylistPlaybackUtils
+{
+	static bool IsPageIdValid(int32 InPageId)
+	{
+		return InPageId != FAvalanchePage::InvalidPageId;
+	}
+	
+	static int32 GetPageIdToPlayNext(const UAvalanchePlaylist* InPlaylist, const FAvaPageListReference& InPageListReference, bool bInPreview, FName InPreviewChannelName);
+	
+	static TArray<int32> GetPagesToTakeToProgram(const UAvalanchePlaylist* InPlaylist, const TArray<int32>& InSelectedPageIds, FName InPreviewChannel = NAME_None);
+};
\ No newline at end of file
