[MotionDesign]

diff --git a/Engine/Plugins/Experimental/Avalanche/Source/AvalancheMediaEditor/Private/Playlist/AvaPlaylistServer.cpp b/Engine/Plugins/Experimental/Avalanche/Source/AvalancheMediaEditor/Private/Playlist/AvaPlaylistServer.cpp
new file mode 100644
--- /dev/null
+++ b/Engine/Plugins/Experimental/Avalanche/Source/AvalancheMediaEditor/Private/Playlist/AvaPlaylistServer.cpp
@@ -0,0 +1,1677 @@
+// Copyright Epic Games, Inc. All Rights Reserved.
+
+#include "Playlist/AvaPlaylistServer.h"
+
+#include "AssetRegistry/AssetData.h"
+#include "AssetRegistry/AssetRegistryModule.h"
+#include "AssetRegistry/IAssetRegistry.h"
+#include "Async/Async.h"
+#include "AvaMediaMessageUtils.h"
+#include "AvaRemoteControlUtils.h"
+#include "AvaRenderTargetUtils.h"
+#include "AvalancheBroadcast.h"
+#include "Broadcast/OutputDevices/AvaOutputClassItem.h"
+#include "Broadcast/OutputDevices/AvaOutputDeviceItem.h"
+#include "Broadcast/OutputDevices/AvaOutputRootItem.h"
+#include "Broadcast/OutputDevices/AvaOutputTreeItem.h"
+#include "IAvaMediaModule.h"
+#include "IRemoteControlModule.h"
+#include "ImageUtils.h"
+#include "MediaOutput.h"
+#include "MediaOutputEditorUtils/MediaOutputEditorUtils.h"
+#include "MessageEndpointBuilder.h"
+#include "Misc/FileHelper.h"
+#include "OutputDevices/AvaRenderTargetMediaUtils.h"
+#include "Playback/AvaMediaPlaybackManager.h"
+#include "Playback/AvaMediaPlaybackUtils.h"
+#include "Playlist/AvaPlaylistPlaybackUtils.h"
+#include "Playlist/AvalancheManagedInstanceCache.h"
+#include "Playlist/AvalanchePlaylist.h"
+#include "RenderingThread.h"
+#include "ScopedTransaction.h"
+#include "Serialization/JsonReader.h"
+#include "Serialization/JsonSerializer.h"
+#include "TextureResource.h"
+
+DEFINE_LOG_CATEGORY_STATIC(LogAvaPlaylistServer, Log, All);
+
+#define LOCTEXT_NAMESPACE "AvaPlaylistServer"
+
+namespace UE::AvaPlaylistServer::Private
+{
+	inline FAvaPlaylistPage GetPageInfo(const UAvalanchePlaylist* InPlaylist, const FAvalanchePage& InPage)
+	{
+		FAvaPlaylistPage PageInfo;
+		PageInfo.PageId = InPage.GetPageId();
+		PageInfo.PageName = InPage.GetPageName();
+		PageInfo.IsTemplate = InPage.IsTemplate();
+		PageInfo.TemplateId = InPage.GetTemplateId();
+		PageInfo.CombinedTemplateIds = InPage.GetCombinedTemplateIds();
+		PageInfo.AssetPath = InPage.GetAvalancheAssetPath(InPlaylist);	// Todo: combo templates
+		PageInfo.Statuses = InPage.GetPageStatuses(InPlaylist);
+		PageInfo.TransitionLayerName = InPage.GetTransitionLayer(InPlaylist).ToString(); 	// Todo: combo templates
+		PageInfo.OutputChannel = InPage.GetChannelName().ToString();
+		PageInfo.bIsEnabled = InPage.IsEnabled();
+		PageInfo.bIsPlaying = InPlaylist->IsPagePlaying(InPage);
+		return PageInfo;
+	}
+
+	/**
+	 * Utility function load a playlist asset in memory.
+	 */
+	static TStrongObjectPtr<UAvalanchePlaylist> LoadPlaylist(const FSoftObjectPath& InPlaylistPath)
+	{
+		UObject* Object = InPlaylistPath.ResolveObject();
+		if (!Object)
+		{
+			Object = InPlaylistPath.TryLoad();
+		}
+		return TStrongObjectPtr<UAvalanchePlaylist>(Cast<UAvalanchePlaylist>(Object));
+	}
+
+	static FAvaPlaylistChannel SerializeChannel(const FAvaOutputChannel& InChannel)
+	{
+		FAvaPlaylistChannel Channel;
+		Channel.Name = InChannel.GetChannelName().ToString();
+		const TArray<UMediaOutput*>& MediaOutputs = InChannel.GetMediaOutputs();
+		for (const UMediaOutput* MediaOutput : MediaOutputs)
+		{
+			FAvaPlaylistOutputDeviceItem DeviceItem;
+			DeviceItem.Name = MediaOutput->GetFName().ToString();
+			DeviceItem.Data = FAvaMediaOutputEditorUtils::SerializeMediaOutput(MediaOutput);
+			Channel.Devices.Push(MoveTemp(DeviceItem));
+		}
+
+		return Channel;
+	}
+	
+	// recursively search device and children
+	static FAvaOutputTreeItemPtr RecursiveFindOutputTreeItem(const FAvaOutputTreeItemPtr& InOutputTreeItem, const FString& InDeviceName)
+	{
+		if (!InOutputTreeItem->IsA<FAvaOutputRootItem>() && InDeviceName == InOutputTreeItem->GetDisplayName().ToString())
+		{
+			return InOutputTreeItem;
+		}
+		
+		for (const TSharedPtr<IAvaOutputTreeItem>& Child : InOutputTreeItem->GetChildren())
+		{
+			if (FAvaOutputTreeItemPtr TreeItem = RecursiveFindOutputTreeItem(Child, InDeviceName))
+			{
+				return TreeItem;
+			}
+		}
+		
+		return nullptr;
+	}
+
+	static UMediaOutput* FindChannelMediaOutput(const FAvaOutputChannel& InOutputChannel, const FString& InOutputMediaName)
+	{
+		const TArray<UMediaOutput*>& MediaOutputs = InOutputChannel.GetMediaOutputs();
+		for (UMediaOutput* MediaOutput : MediaOutputs)
+		{
+			if (MediaOutput->GetName() == InOutputMediaName)
+			{
+				return MediaOutput;
+			}
+		}
+		
+		return nullptr;
+	}
+	
+	inline TArray<int32> GetPlayingPages(const UAvalanchePlaylist* InPlaylist, bool bInIsPreview, FName InChannelName)
+	{
+		return bInIsPreview ? InPlaylist->GetPreviewingPageIds(InChannelName) : InPlaylist->GetPlayingPageIds(InChannelName);
+	}
+
+	static bool ContinuePages(UAvalanchePlaylist* InPlaylist, const TArray<int32>& InPageIds, bool bInIsPreview, FName InPreviewChannelName, FString& OutFailureReason)
+	{
+		bool bSuccess = false;
+		for (const int32 PageId : InPageIds)
+		{
+			if (InPlaylist->CanContinuePage(PageId, bInIsPreview, InPreviewChannelName))
+			{
+				bSuccess |= InPlaylist->ContinuePage(PageId, bInIsPreview, InPreviewChannelName);
+			}
+			else if (bInIsPreview)
+			{
+				OutFailureReason.Appendf(TEXT("PageId %d was not previewing on channel \"%s\". "), PageId, *InPreviewChannelName.ToString());
+			}
+			else
+			{
+				OutFailureReason.Appendf(TEXT("PageId %d was not playing. "), PageId);
+			}
+		}
+		return bSuccess;
+	}
+
+	static bool UpdatePagesValues(const UAvalanchePlaylist* InPlaylist, const TArray<int32>& InPageIds, bool bInIsPreview, FName InPreviewChannelName)
+	{
+		bool bSuccess = false;
+		for (const int32 PageId : InPageIds)
+		{
+			bSuccess |= InPlaylist->PushRuntimeRemoteControlValues(PageId, bInIsPreview, InPreviewChannelName);
+		}
+		return bSuccess;
+	}
+}
+
+/**
+ * Holds the render target for copying the channel image.
+ * The render target needs to be held for many frames until it is done.
+ */
+struct FAvaPlaylistServer::FChannelImage
+{
+	// Optional temporary render target for converting pixel format.
+	TStrongObjectPtr<UTextureRenderTarget2D> RenderTarget;
+
+	// Pixels readback from the render target. Format is PF_B8G8R8A8 (for now).
+	TArray<FColor> RawPixels;
+	int32 SizeX = 0;
+	int32 SizeY = 0;
+
+	void UpdateRenderTarget(int32 InSizeX, int32 InSizeY, EPixelFormat InFormat, const FLinearColor& InClearColor)
+	{
+		if (!RenderTarget.IsValid())
+		{
+			static const FName ChannelImageRenderTargetBaseName = TEXT("AvaPlaylistServer_ChannelImageRenderTarget");
+			RenderTarget.Reset(AvaRenderTargetUtils::CreateDefaultRenderTarget(ChannelImageRenderTargetBaseName));
+		}
+	
+		AvaRenderTargetUtils::UpdateRenderTarget(RenderTarget.Get(), FIntPoint(InSizeX, InSizeY), InFormat, InClearColor);
+	}
+
+	void UpdateRawPixels(int32 InSizeX, int32 InSizeY)
+	{
+		SizeX = InSizeX;
+		SizeY = InSizeY;
+		RawPixels.SetNum(InSizeX * InSizeY);
+	}
+};
+
+FAvaPlaylistServer::FAvaPlaylistServer()
+{
+	
+}
+
+FAvaPlaylistServer::~FAvaPlaylistServer()
+{
+	FMessageEndpoint::SafeRelease(MessageEndpoint);
+	
+	for (IConsoleObject* ConsoleCommand : ConsoleCommands)
+	{
+		IConsoleManager::Get().UnregisterConsoleObject(ConsoleCommand);
+	}
+	ConsoleCommands.Empty();
+}
+
+const FMessageAddress& FAvaPlaylistServer::GetMessageAddress() const
+{
+	if (MessageEndpoint.IsValid())
+	{
+		return MessageEndpoint->GetAddress();
+	}
+	static FMessageAddress InvalidMessageAddress;
+	return InvalidMessageAddress;
+}
+
+void FAvaPlaylistServer::Init(const FString& InAssignedHostName)
+{
+	HostName = InAssignedHostName.IsEmpty() ? FPlatformProcess::ComputerName() : InAssignedHostName;
+	
+	MessageEndpoint = FMessageEndpoint::Builder("AvalanchePlaylistServer")
+	.Handling<FAvaPlaylistPing>(this, &FAvaPlaylistServer::HandlePlaylistPing)
+	.Handling<FAvaPlaylistGetPlaylists>(this, &FAvaPlaylistServer::HandleGetPlaylists)
+	.Handling<FAvaPlaylistLoadPlaylist>(this, &FAvaPlaylistServer::HandleLoadPlaylist)
+	.Handling<FAvaPlaylistCreatePage>(this, &FAvaPlaylistServer::HandleCreatePage)
+	.Handling<FAvaPlaylistDeletePage>(this, &FAvaPlaylistServer::HandleDeletePage)
+	.Handling<FAvaPlaylistCreateTemplate>(this, &FAvaPlaylistServer::HandleCreateTemplate)
+	.Handling<FAvaPlaylistDeleteTemplate>(this, &FAvaPlaylistServer::HandleDeleteTemplate)
+	.Handling<FAvaPlaylistChangeTemplateBP>(this, &FAvaPlaylistServer::HandleChangeTemplateBP)
+	.Handling<FAvaPlaylistGetPages>(this, &FAvaPlaylistServer::HandleGetPages)
+	.Handling<FAvaPlaylistGetPageDetails>(this, &FAvaPlaylistServer::HandleGetPageDetails)
+	.Handling<FAvaPlaylistPageChangeChannel>(this, &FAvaPlaylistServer::HandleChangePageChannel)
+	.Handling<FAvaPlaylistUpdatePageFromRCP>(this, &FAvaPlaylistServer::HandleUpdatePageFromRCP)
+	.Handling<FAvaPlaylistPageAction>(this, &FAvaPlaylistServer::HandlePageAction)
+	.Handling<FAvaPlaylistPagePreviewAction>(this, &FAvaPlaylistServer::HandlePagePreviewAction)
+	.Handling<FAvaPlaylistPageActions>(this, &FAvaPlaylistServer::HandlePageActions)
+	.Handling<FAvaPlaylistPagePreviewActions>(this, &FAvaPlaylistServer::HandlePagePreviewActions)
+	.Handling<FAvaPlaylistGetChannel>(this, &FAvaPlaylistServer::HandleGetChannel)
+	.Handling<FAvaPlaylistGetChannels>(this, &FAvaPlaylistServer::HandleGetChannels)
+	.Handling<FAvaPlaylistChannelAction>(this, &FAvaPlaylistServer::HandleChannelAction)
+	.Handling<FAvaPlaylistGetDevices>(this, &FAvaPlaylistServer::HandleGetDevices)
+	.Handling<FAvaPlaylistAddChannelDevice>(this, &FAvaPlaylistServer::HandleAddChannelDevice)
+	.Handling<FAvaPlaylistEditChannelDevice>(this, &FAvaPlaylistServer::HandleEditChannelDevice)
+	.Handling<FAvaPlaylistRemoveChannelDevice>(this, &FAvaPlaylistServer::HandleRemoveChannelDevice)
+	.Handling<FAvaPlaylistGetChannelImage>(this, &FAvaPlaylistServer::HandleGetChannelImage)
+	.NotificationHandling(FOnBusNotification::CreateRaw(this, &FAvaPlaylistServer::OnMessageBusNotification));
+	
+	if (MessageEndpoint.IsValid())
+	{
+		// Subscribe to the server listing requests
+		MessageEndpoint->Subscribe<FAvaPlaylistPing>();
+
+		UE_LOG(LogAvaPlaylistServer, Log, TEXT("Motion Design Playlist Server \"%s\" Started."), *HostName);
+	}
+}
+
+void FAvaPlaylistServer::SetupBroadcastDelegates(UAvalancheBroadcast* InBroadcast)
+{
+	RemoveBroadcastDelegates(InBroadcast);
+	InBroadcast->GetOnChannelsListChanged().AddRaw(this, &FAvaPlaylistServer::OnBroadcastChannelListChanged);
+}
+
+void FAvaPlaylistServer::SetupEditorDelegates()
+{
+	RemoveEditorDelegates();
+	const FAssetRegistryModule& AssetRegistryModule = FModuleManager::LoadModuleChecked<FAssetRegistryModule>("AssetRegistry");
+	AssetRegistryModule.Get().OnAssetAdded().AddRaw(this, &FAvaPlaylistServer::OnAssetAddedOrRemoved);
+	AssetRegistryModule.Get().OnAssetRemoved().AddRaw(this, &FAvaPlaylistServer::OnAssetAddedOrRemoved);
+}
+
+void FAvaPlaylistServer::RemoveBroadcastDelegates(UAvalancheBroadcast* InBroadcast) const
+{
+	InBroadcast->GetOnChannelsListChanged().RemoveAll(this);
+}
+
+void FAvaPlaylistServer::RemoveEditorDelegates() const
+{
+	const FAssetRegistryModule& AssetRegistryModule = FModuleManager::LoadModuleChecked<FAssetRegistryModule>("AssetRegistry");
+	AssetRegistryModule.Get().OnAssetAdded().RemoveAll(this);
+	AssetRegistryModule.Get().OnAssetRemoved().RemoveAll(this);
+}
+
+void FAvaPlaylistServer::OnPageListChanged(const FAvaPageListChangeParams& InParams) const
+{
+	FAvaPlaylistPageListChanged* ReplyMessage = FMessageEndpoint::MakeMessage<FAvaPlaylistPageListChanged>();
+	ReplyMessage->Playlist = FSoftObjectPath(InParams.Playlist).ToString();
+	ReplyMessage->ChangeType = static_cast<uint8>(InParams.ChangeType);
+	ReplyMessage->AffectedPages = InParams.AffectedPages;
+	
+	SendResponse(ReplyMessage, ClientAddresses);
+}
+
+void FAvaPlaylistServer::OnPagesChanged(const UAvalanchePlaylist* InPlaylist, const FAvalanchePage& InPage, EAvaPageChanges InChange) const
+{
+	switch(InChange)
+	{
+	case EAvaPageChanges::AnimationSettings:
+		{
+			PageAnimSettingsChanged(InPlaylist, InPage);
+			break;
+		}
+	case EAvaPageChanges::Blueprint:
+		{
+			PageBlueprintChanged(InPlaylist, InPage, InPage.GetAvalancheAssetPath(InPlaylist).ToString());
+			break;
+		}
+	case EAvaPageChanges::Status:
+		{
+			PageStatusChanged(InPlaylist, InPage);
+			break;
+		}
+	case EAvaPageChanges::Channel:
+		{
+			PageChannelChanged(InPlaylist, InPage, InPage.GetChannelName().ToString());
+			break;
+		}
+	case EAvaPageChanges::RemoteControlValues:
+		{
+			break;
+		}
+	case EAvaPageChanges::All:
+		{
+			break;
+		}
+	case EAvaPageChanges::None:
+		{
+			break;
+		}
+	default:
+		{
+			break;
+		}
+	}
+	
+}
+void FAvaPlaylistServer::PageStatusChanged(const UAvalanchePlaylist* InPlaylist, const FAvalanchePage& InPage) const
+{
+	FAvaPlaylistPagesStatuses* ReplyMessage = FMessageEndpoint::MakeMessage<FAvaPlaylistPagesStatuses>();
+	ReplyMessage->Playlist = FSoftObjectPath(InPlaylist).ToString();
+	ReplyMessage->PageInfo = UE::AvaPlaylistServer::Private::GetPageInfo(InPlaylist, InPage);
+	SendResponse(ReplyMessage, ClientAddresses);
+}
+
+void FAvaPlaylistServer::PageBlueprintChanged(const UAvalanchePlaylist* InPlaylist, const FAvalanchePage& InPage, const FString& InBlueprintPath) const
+{
+	FAvaPlaylistPageBlueprintChanged* ReplyMessage = FMessageEndpoint::MakeMessage<FAvaPlaylistPageBlueprintChanged>();
+	ReplyMessage->Playlist = FSoftObjectPath(InPlaylist).ToString();
+	ReplyMessage->PageId = InPage.GetPageId();
+	ReplyMessage->BlueprintPath = InBlueprintPath;
+	SendResponse(ReplyMessage, ClientAddresses);
+}
+
+void FAvaPlaylistServer::PageChannelChanged(const UAvalanchePlaylist* InPlaylist, const FAvalanchePage& InPage, const FString& InChannelName) const
+{
+	FAvaPlaylistPageChannelChanged* ReplyMessage = FMessageEndpoint::MakeMessage<FAvaPlaylistPageChannelChanged>();
+	ReplyMessage->Playlist = FSoftObjectPath(InPlaylist).ToString();
+	ReplyMessage->PageId = InPage.GetPageId();
+	ReplyMessage->ChannelName = InChannelName;
+	SendResponse(ReplyMessage, ClientAddresses);
+}
+
+void FAvaPlaylistServer::PageAnimSettingsChanged(const UAvalanchePlaylist* InPlaylist, const FAvalanchePage& InPage) const
+{
+	FAvaPlaylistPageAnimSettingsChanged* ReplyMessage = FMessageEndpoint::MakeMessage<FAvaPlaylistPageAnimSettingsChanged>();
+	ReplyMessage->Playlist = FSoftObjectPath(InPlaylist).ToString();
+	ReplyMessage->PageId = InPage.GetPageId();
+	SendResponse(ReplyMessage, ClientAddresses);
+}
+
+void FAvaPlaylistServer::OnBroadcastChannelListChanged(const FAvaBroadcastProfile& InProfile) const
+{
+	FAvaBroadcastChannelListChanged* ReplyMessage = FMessageEndpoint::MakeMessage<FAvaBroadcastChannelListChanged>();
+
+	const TArray<FAvaOutputChannel*>& OutputChannels = InProfile.GetChannels();
+	
+	ReplyMessage->Channels.Reserve(OutputChannels.Num());
+
+	for (const FAvaOutputChannel* OutputChannel : OutputChannels)
+	{
+		FAvaPlaylistChannel Channel = UE::AvaPlaylistServer::Private::SerializeChannel(*OutputChannel);
+		ReplyMessage->Channels.Push(MoveTemp(Channel));
+	}
+	SendResponse(ReplyMessage, ClientAddresses);
+}
+
+void FAvaPlaylistServer::OnAssetAddedOrRemoved(const FAssetData& InAssetData) const
+{
+	using namespace UE::AvaPlaylistServer::Private;
+	if (InAssetData.GetClass() == UAvalanchePlaylist::StaticClass() || FAvaMediaPlaybackUtils::IsPlayableAsset(InAssetData))
+	{
+		FAvaPlaylistAssetsChanged* ReplyMessage = FMessageEndpoint::MakeMessage<FAvaPlaylistAssetsChanged>();
+		ReplyMessage->AssetName = InAssetData.AssetName.ToString();
+		SendResponse(ReplyMessage, ClientAddresses);
+	}
+}
+
+void FAvaPlaylistServer::HandlePlaylistPing(const FAvaPlaylistPing& InMessage, const TSharedRef<IMessageContext, ESPMode::ThreadSafe>& InContext)
+{
+	if (!InMessage.bAuto)
+	{
+		UE_LOG(LogAvaPlaylistServer, Log, TEXT("Received Ping request from %s"), *InContext->GetSender().ToString());
+	}
+
+	SetupEditorDelegates();
+	
+	FAvaPlaylistPong* ReplyMessage = FMessageEndpoint::MakeMessage<FAvaPlaylistPong>();
+	ReplyMessage->RequestId = InMessage.RequestId;
+	ReplyMessage->bAuto = InMessage.bAuto;
+
+	// We still support the initial version.
+	constexpr int32 CurrentMinimumApiVersion = EPlaylistApiVersion::Initial;
+
+	// Consider clients that didn't request a version to be the "initial" version.
+	const int32 RequestedApiVersion = InMessage.RequestedApiVersion != -1 ? InMessage.RequestedApiVersion : EPlaylistApiVersion::Initial;
+
+	// Determine the version we will communicate with this client.
+	int32 HonoredApiVersion = EPlaylistApiVersion::LatestVersion;
+	
+	if (RequestedApiVersion >= CurrentMinimumApiVersion && RequestedApiVersion <= EPlaylistApiVersion::LatestVersion)
+	{
+		HonoredApiVersion = RequestedApiVersion;
+	}
+	
+	ReplyMessage->ApiVersion = HonoredApiVersion;
+	ReplyMessage->MinimumApiVersion = CurrentMinimumApiVersion;
+	ReplyMessage->LatestApiVersion = EPlaylistApiVersion::LatestVersion;
+	ReplyMessage->HostName = HostName;
+
+	FClientInfo& ClientInfo = GetOrAddClientInfo(InContext->GetSender());
+	ClientInfo.ApiVersion = HonoredApiVersion;
+
+	SendResponse(ReplyMessage, InContext->GetSender());
+}
+
+void FAvaPlaylistServer::HandleGetPlaylists(const FAvaPlaylistGetPlaylists& InMessage, const TSharedRef<IMessageContext, ESPMode::ThreadSafe>& InContext)
+{
+	FAvaPlaylistPlaylists* ReplyMessage = FMessageEndpoint::MakeMessage<FAvaPlaylistPlaylists>();
+
+	ReplyMessage->RequestId = InMessage.RequestId;
+	
+	// List all the play list.
+	if (IAssetRegistry* AssetRegistry = IAssetRegistry::Get())
+	{
+		TArray<FAssetData> Assets;
+		if (AssetRegistry->GetAssetsByClass(UAvalanchePlaylist::StaticClass()->GetClassPathName(), Assets))
+		{
+			ReplyMessage->Playlists.Reserve(Assets.Num());		
+			for (const FAssetData& Data : Assets)
+			{
+				ReplyMessage->Playlists.Add(Data.ToSoftObjectPath().ToString());
+			}
+		}
+	}
+
+	SendResponse(ReplyMessage, InContext->GetSender());
+}
+
+void FAvaPlaylistServer::HandleLoadPlaylist(const FAvaPlaylistLoadPlaylist& InMessage, const TSharedRef<IMessageContext, ESPMode::ThreadSafe>& InContext)
+{
+	// If the requested path is empty, we assume this is a request for information only.
+	if (!InMessage.Playlist.IsEmpty())
+	{
+		const FSoftObjectPath NewPlaylistPath(InMessage.Playlist);
+		const UAvalanchePlaylist* Playlist = PlaylistPlaybackCommandData.GetOrLoadPlaylist(NewPlaylistPath,
+			[this](UAvalanchePlaylist* InPlaylist)
+			{
+				PlaylistPlaybackCommandData.ClosePlaybackContext();
+				PlaylistPlaybackCommandData.RemovePlaylistDelegates(this, InPlaylist);
+			},
+			[this](UAvalanchePlaylist* InPlaylist)
+			{
+				PlaylistPlaybackCommandData.SetupPlaylistDelegates(this, InPlaylist);
+				if (InPlaylist)
+				{
+					InPlaylist->InitializePlaybackContext();
+				}
+			});
+		
+		if (!Playlist)
+		{
+			SendMessage(InContext->GetSender(), InMessage.RequestId, ELogVerbosity::Error,
+				TEXT("Playlist \"%s\" not loaded."), *InMessage.Playlist);
+			return;
+		}
+	}
+
+	SendMessage(InContext->GetSender(), InMessage.RequestId, ELogVerbosity::Log,
+		TEXT("Playlist \"%s\" loaded."), *PlaylistPlaybackCommandData.CurrentPlaylistPath.ToString());
+}
+
+void FAvaPlaylistServer::HandleGetPages(const FAvaPlaylistGetPages& InMessage, const TSharedRef<IMessageContext, ESPMode::ThreadSafe>& InContext)
+{
+	UAvalanchePlaylist* Playlist = GetOrLoadPlaylistForEdit(InContext->GetSender(), InMessage.RequestId, InMessage.Playlist);
+	if (!Playlist)
+	{
+		return;
+	}
+	
+	FAvaPlaylistPages* ReplyMessage = FMessageEndpoint::MakeMessage<FAvaPlaylistPages>();
+	ReplyMessage->RequestId = InMessage.RequestId;
+	ReplyMessage->Pages.Reserve(Playlist->GetInstancedPages().Pages.Num() + Playlist->GetTemplatePages().Pages.Num());
+	for (const FAvalanchePage& Page : Playlist->GetInstancedPages().Pages)
+	{
+		FAvaPlaylistPage PageInfo = UE::AvaPlaylistServer::Private::GetPageInfo(Playlist, Page);
+		ReplyMessage->Pages.Add(MoveTemp(PageInfo));
+	}
+	
+	for (const FAvalanchePage& Page : Playlist->GetTemplatePages().Pages)
+	{
+		FAvaPlaylistPage PageInfo = UE::AvaPlaylistServer::Private::GetPageInfo(Playlist, Page);
+		ReplyMessage->Pages.Add(MoveTemp(PageInfo));
+	}
+	
+	SendResponse(ReplyMessage, InContext->GetSender());
+}
+
+void FAvaPlaylistServer::HandleCreatePage(const FAvaPlaylistCreatePage& InMessage,
+	const TSharedRef<IMessageContext, ESPMode::ThreadSafe>& InContext)
+{
+	UAvalanchePlaylist* Playlist = GetOrLoadPlaylistForEdit(InContext->GetSender(), InMessage.RequestId, InMessage.Playlist);
+	if (!Playlist)
+	{
+		return;
+	}
+
+	const FAvalanchePage& Template = Playlist->GetPage(InMessage.TemplateId);
+    if (!Template.IsValidPage() || !Template.IsTemplate())
+    {
+    	LogAndSendMessage(InContext->GetSender(), InMessage.RequestId, ELogVerbosity::Error, TEXT("Template %d is not valid or is not a template"), InMessage.TemplateId);
+    	return;
+    }
+	
+	const int32 PageId = Playlist->AddPageFromTemplate(InMessage.TemplateId);
+	if (PageId != FAvalanchePage::InvalidPageId)
+	{
+		SendMessage(InContext->GetSender(), InMessage.RequestId, ELogVerbosity::Log, TEXT("Page %d Created from Template %d"), PageId, InMessage.TemplateId);
+	}
+	else
+	{
+		SendMessage(InContext->GetSender(), InMessage.RequestId, ELogVerbosity::Error, TEXT("Failed to create a page from Template %d"), InMessage.TemplateId);
+	}
+}
+
+void FAvaPlaylistServer::HandleDeletePage(const FAvaPlaylistDeletePage& InMessage, const TSharedRef<IMessageContext, ESPMode::ThreadSafe>& InContext)
+{
+	UAvalanchePlaylist* Playlist = GetOrLoadPlaylistForEdit(InContext->GetSender(), InMessage.RequestId, InMessage.Playlist);
+	if (!Playlist)
+	{
+		return;
+	}
+
+	const FAvalanchePage& Page = Playlist->GetPage(InMessage.PageId);
+	if (!Page.IsValidPage())
+	{
+		LogAndSendMessage(InContext->GetSender(), InMessage.RequestId, ELogVerbosity::Error, TEXT("page %d is not valid"), InMessage.PageId);
+		return;
+	}
+
+	Playlist->RemovePage(InMessage.PageId);
+	SendMessage(InContext->GetSender(), InMessage.RequestId, ELogVerbosity::Log, TEXT("Page %d deleted"), InMessage.PageId);
+}
+
+void FAvaPlaylistServer::HandleDeleteTemplate(const FAvaPlaylistDeleteTemplate& InMessage,
+	const TSharedRef<IMessageContext, ESPMode::ThreadSafe>& InContext)
+{
+	UAvalanchePlaylist* Playlist = GetOrLoadPlaylistForEdit(InContext->GetSender(), InMessage.RequestId, InMessage.Playlist);
+	if (!Playlist)
+	{
+		return;
+	}
+
+	const FAvalanchePage& Page = Playlist->GetPage(InMessage.PageId);
+	if (!Page.IsValidPage())
+	{
+		LogAndSendMessage(InContext->GetSender(), InMessage.RequestId, ELogVerbosity::Error, TEXT("page %d is not valid"), InMessage.PageId);
+		return;
+	}
+
+	if (!Page.IsTemplate())
+	{
+		LogAndSendMessage(InContext->GetSender(), InMessage.RequestId, ELogVerbosity::Error, TEXT("page %d is not a template"), InMessage.PageId);
+		return;
+	}
+
+	if (!Page.GetInstancedIds().IsEmpty())
+	{
+		LogAndSendMessage(InContext->GetSender(), InMessage.RequestId, ELogVerbosity::Error, TEXT("Template has instanced pages"), InMessage.PageId);
+		return;
+	}
+
+	Playlist->RemovePage(InMessage.PageId);
+	SendMessage(InContext->GetSender(), InMessage.RequestId, ELogVerbosity::Log, TEXT("Page template %d deleted"), InMessage.PageId);
+}
+
+void FAvaPlaylistServer::HandleCreateTemplate(const FAvaPlaylistCreateTemplate& InMessage,
+	const TSharedRef<IMessageContext, ESPMode::ThreadSafe>& InContext)
+{
+	UAvalanchePlaylist* Playlist = GetOrLoadPlaylistForEdit(InContext->GetSender(), InMessage.RequestId, InMessage.Playlist);
+	if (!Playlist)
+	{
+		return;
+	}
+	
+	Playlist->AddTemplate();
+	SendMessage(InContext->GetSender(), InMessage.RequestId, ELogVerbosity::Log, TEXT("Template Created"));
+}
+
+
+void FAvaPlaylistServer::HandleChangeTemplateBP(const FAvaPlaylistChangeTemplateBP& InMessage,
+	const TSharedRef<IMessageContext, ESPMode::ThreadSafe>& InContext)
+{
+	UAvalanchePlaylist* Playlist = GetOrLoadPlaylistForEdit(InContext->GetSender(), InMessage.RequestId, InMessage.Playlist);
+	if (!Playlist)
+	{
+		return;
+	}
+
+	FAvalanchePage& Page = Playlist->GetPage(InMessage.TemplateId);
+	if (Page.IsValidPage() && Page.IsTemplate())
+	{
+		if (Page.UpdateAvalancheAsset(InMessage.AssetPath))
+		{
+			SendMessage(InContext->GetSender(), InMessage.RequestId, ELogVerbosity::Log, TEXT("Blueprint change of template: %d to %s"), InMessage.TemplateId, *InMessage.AssetPath);
+			Playlist->GetOnPagesChanged().Broadcast(Playlist, Page, EAvaPageChanges::Blueprint);
+			return;
+		}
+	}
+	LogAndSendMessage(InContext->GetSender(), InMessage.RequestId, ELogVerbosity::Error, TEXT("Blueprint change of template: %d to %s failed."), InMessage.TemplateId, *InMessage.AssetPath);
+}
+
+void FAvaPlaylistServer::HandleGetPageDetails(const FAvaPlaylistGetPageDetails& InMessage, const TSharedRef<IMessageContext, ESPMode::ThreadSafe>& InContext)
+{	
+	UAvalanchePlaylist* Playlist = GetOrLoadPlaylistForEdit(InContext->GetSender(), InMessage.RequestId, InMessage.Playlist);
+	if (!Playlist)
+	{
+		return;
+	}
+
+	const FAvalanchePage& Page = Playlist->GetPage(InMessage.PageId);
+	if (!Page.IsValidPage())
+	{
+		LogAndSendMessage(InContext->GetSender(), InMessage.RequestId, ELogVerbosity::Error, TEXT("\"PageDetails\" not available: PageId %d is invalid."), InMessage.PageId);
+		return;
+	}
+
+	if (Page.GetAvalancheAssetPath(Playlist).IsNull())
+	{
+		LogAndSendMessage(InContext->GetSender(), InMessage.RequestId, ELogVerbosity::Error, TEXT("Page has no asset selected"), InMessage.PageId);
+		return;
+	}
+	
+	if (InMessage.bLoadRemoteControlPreset)
+	{
+		FAvalancheManagedInstanceCache& ManagedInstanceCache = IAvaMediaModule::Get().GetManagedInstanceCache();
+		const TSharedPtr<FAvalancheManagedInstance> ManagedInstance
+			= ManagedInstanceCache.GetOrLoadAvalancheInstance(Page.GetAvalancheAssetPath(Playlist));
+		
+		if (ManagedInstance.IsValid())
+		{
+			PlaylistEditCommandData.SaveCurrentRemoteControlPresetToPage(true);
+
+			// Applying the controller values can break the WYSIWYG of the editor,
+			// in case multiple controllers set the same exposed entity with different values.
+			// There is no guaranty that the controller actions are self consistent.
+			// To avoid this issue, we apply the controllers first, and then
+			// restore the entity values in a second pass.
+			
+			Page.GetRemoteControlValues().ApplyControllerValuesToRemoteControlPreset(ManagedInstance->GetRemoteControlPreset(), true);
+			Page.GetRemoteControlValues().ApplyEntityValuesToRemoteControlPreset(ManagedInstance->GetRemoteControlPreset());
+
+			// Register the RC Preset to Remote Control module to make it available through WebRC.
+			FAvaRemoteControlUtils::RegisterRemoteControlPreset(ManagedInstance->GetRemoteControlPreset(), /*bInEnsureUniqueId*/ true);
+
+			// Keep track of what is currently registered.
+			PlaylistEditCommandData.ManagedPageId = InMessage.PageId;
+			PlaylistEditCommandData.ManagedInstance = ManagedInstance;
+		}
+	}
+	
+	FAvaPlaylistPageDetails* ReplyMessage = FMessageEndpoint::MakeMessage<FAvaPlaylistPageDetails>();
+	ReplyMessage->RequestId = InMessage.RequestId;
+	ReplyMessage->Playlist = InMessage.Playlist;
+	ReplyMessage->PageInfo = UE::AvaPlaylistServer::Private::GetPageInfo(Playlist, Page);
+	ReplyMessage->RemoteControlValues = Page.GetRemoteControlValues();
+	if (InMessage.bLoadRemoteControlPreset && PlaylistEditCommandData.ManagedInstance.IsValid() && PlaylistEditCommandData.ManagedInstance->GetRemoteControlPreset())
+	{
+		ReplyMessage->RemoteControlPresetName = PlaylistEditCommandData.ManagedInstance->GetRemoteControlPreset()->GetPresetName().ToString();
+		ReplyMessage->RemoteControlPresetId = PlaylistEditCommandData.ManagedInstance->GetRemoteControlPreset()->GetPresetId().ToString();
+	}
+	SendResponse(ReplyMessage, InContext->GetSender());
+}
+
+void FAvaPlaylistServer::HandleChangePageChannel(const FAvaPlaylistPageChangeChannel& InMessage,
+	const TSharedRef<IMessageContext, ESPMode::ThreadSafe>& InContext)
+{
+	UAvalancheBroadcast& Broadcast = UAvalancheBroadcast::Get();
+	if (InMessage.ChannelName.IsEmpty())
+	{
+		LogAndSendMessage(InContext->GetSender(), InMessage.RequestId, ELogVerbosity::Error, TEXT("No Channel Name Provided"));
+		return;
+	}
+
+	const FName ChannelName(InMessage.ChannelName);
+	const FAvaOutputChannel& Channel = Broadcast.GetCurrentProfile().GetChannel(ChannelName);
+	if (!Channel.IsValidChannel())
+	{
+		LogAndSendMessage(InContext->GetSender(), InMessage.RequestId, ELogVerbosity::Error, TEXT("%s is not a valid channel"), *ChannelName.ToString());
+		return;
+	}
+
+	UAvalanchePlaylist* Playlist = GetOrLoadPlaylistForEdit(InContext->GetSender(), InMessage.RequestId, InMessage.Playlist);
+	if (!Playlist)
+	{
+		return;
+	}
+
+	FAvalanchePage& Page = Playlist->GetPage(InMessage.PageId);
+	if (!Page.IsValidPage())
+	{
+		LogAndSendMessage(InContext->GetSender(), InMessage.RequestId, ELogVerbosity::Error, TEXT("PageId %d is invalid."), InMessage.PageId);
+		return;
+	}
+
+	if (Page.GetChannelName() == ChannelName)
+	{
+		SendMessage(InContext->GetSender(), InMessage.RequestId, ELogVerbosity::Log, TEXT("Same Channel Selected"), InMessage.PageId);
+		return;
+	}
+
+	Page.SetChannelName(ChannelName);
+	Playlist->GetOnPagesChanged().Broadcast(Playlist, Page, EAvaPageChanges::Channel);
+	SendMessage(InContext->GetSender(), InMessage.RequestId, ELogVerbosity::Log, TEXT("Channel changed"));
+}
+
+void FAvaPlaylistServer::HandleUpdatePageFromRCP(const FAvaPlaylistUpdatePageFromRCP& InMessage, const TSharedRef<IMessageContext, ESPMode::ThreadSafe>& InContext)
+{
+	// Note that this doesn't save the playlist.
+	PlaylistEditCommandData.SaveCurrentRemoteControlPresetToPage(InMessage.bUnregister);
+	if (InMessage.bUnregister)
+	{
+		PlaylistEditCommandData.ManagedPageId = FAvalanchePage::InvalidPageId;
+		PlaylistEditCommandData.ManagedInstance.Reset();
+	}
+	SendMessage(InContext->GetSender(), InMessage.RequestId, ELogVerbosity::Log, TEXT("\"UpdatePageFromRCP\" Ok."));
+}
+
+void FAvaPlaylistServer::HandlePageAction(const FAvaPlaylistPageAction& InMessage, const TSharedRef<IMessageContext, ESPMode::ThreadSafe>& InContext)
+{
+	const FRequestInfo RequestInfo = {InMessage.RequestId, InContext->GetSender()};
+	HandlePageActions(RequestInfo, {InMessage.PageId}, false, FName(), InMessage.Action);
+}
+
+void FAvaPlaylistServer::HandlePagePreviewAction(const FAvaPlaylistPagePreviewAction& InMessage,
+	const TSharedRef<IMessageContext, ESPMode::ThreadSafe>& InContext)
+{
+	const FRequestInfo RequestInfo = {InMessage.RequestId, InContext->GetSender()};
+	HandlePageActions(RequestInfo, {InMessage.PageId}, true, FName(InMessage.PreviewChannelName), InMessage.Action);
+}
+
+void FAvaPlaylistServer::HandlePageActions(const FAvaPlaylistPageActions& InMessage, const TSharedRef<IMessageContext, ESPMode::ThreadSafe>& InContext)
+{
+	const FRequestInfo RequestInfo = {InMessage.RequestId, InContext->GetSender()};
+	HandlePageActions(RequestInfo, InMessage.PageIds, false, FName(), InMessage.Action);
+}
+
+void FAvaPlaylistServer::HandlePagePreviewActions(const FAvaPlaylistPagePreviewActions& InMessage,
+	const TSharedRef<IMessageContext, ESPMode::ThreadSafe>& InContext)
+{
+	const FRequestInfo RequestInfo = {InMessage.RequestId, InContext->GetSender()};
+	HandlePageActions(RequestInfo, InMessage.PageIds, true, FName(InMessage.PreviewChannelName), InMessage.Action);
+}
+
+void FAvaPlaylistServer::HandleGetChannel(const FAvaPlaylistGetChannel& InMessage, const TSharedRef<IMessageContext, ESPMode::ThreadSafe>& InContext)
+{
+	const FName ChannelName(InMessage.ChannelName);
+	
+	const UAvalancheBroadcast& Broadcast = UAvalancheBroadcast::Get();
+	const FAvaOutputChannel& Channel = Broadcast.GetCurrentProfile().GetChannel(ChannelName);
+
+	if (!Channel.IsValidChannel())
+	{
+		LogAndSendMessage(InContext->GetSender(), InMessage.RequestId, ELogVerbosity::Error, TEXT("\"GetChannel\" Channel \"%s\" not found."), *InMessage.ChannelName);
+		return;
+	}
+	
+	FAvaPlaylistChannelResponse* ReplyMessage = FMessageEndpoint::MakeMessage<FAvaPlaylistChannelResponse>();
+	ReplyMessage->RequestId = InMessage.RequestId;
+	ReplyMessage->Channel = UE::AvaPlaylistServer::Private::SerializeChannel(Channel);
+	SendResponse(ReplyMessage, InContext->GetSender());
+}
+
+void FAvaPlaylistServer::HandleGetChannels(const FAvaPlaylistGetChannels& InMessage, const TSharedRef<IMessageContext, ESPMode::ThreadSafe>& InContext)
+{
+	UAvalancheBroadcast& Broadcast = UAvalancheBroadcast::Get();
+	FAvaPlaylistChannels* ReplyMessage = FMessageEndpoint::MakeMessage<FAvaPlaylistChannels>();
+	ReplyMessage->RequestId = InMessage.RequestId;
+	
+	const TArray<FAvaOutputChannel*>& Channels = Broadcast.GetCurrentProfile().GetChannels();
+	ReplyMessage->Channels.Reserve(Channels.Num());
+
+	for (const FAvaOutputChannel* OutputChannel : Channels)
+	{
+		FAvaPlaylistChannel Channel = UE::AvaPlaylistServer::Private::SerializeChannel(*OutputChannel);
+		ReplyMessage->Channels.Push(MoveTemp(Channel));
+	}
+	
+	SetupBroadcastDelegates(&Broadcast);
+	SendResponse(ReplyMessage, InContext->GetSender());
+}
+
+void FAvaPlaylistServer::HandleChannelAction(const FAvaPlaylistChannelAction& InMessage, const TSharedRef<IMessageContext, ESPMode::ThreadSafe>& InContext)
+{
+	UAvalancheBroadcast& Broadcast = UAvalancheBroadcast::Get();
+	if (InMessage.Action == EAvaPlaylistChannelActions::Start)
+	{
+		if (InMessage.ChannelName.IsEmpty())
+		{
+			Broadcast.StartBroadcast();
+			SendMessage(InContext->GetSender(), InMessage.RequestId, ELogVerbosity::Log, TEXT("\"ChannelAction\" Ok."));
+		}
+		else
+		{
+			const FName ChannelName(InMessage.ChannelName);
+			FAvaOutputChannel& Channel = Broadcast.GetCurrentProfile().GetChannelMutable(ChannelName);
+			if (Channel.IsValidChannel())
+			{
+				Channel.StartChannelBroadcast();
+				SendMessage(InContext->GetSender(), InMessage.RequestId, ELogVerbosity::Log, TEXT("\"ChannelAction\" Ok."));
+			}
+			else
+			{
+				LogAndSendMessage(InContext->GetSender(), InMessage.RequestId, ELogVerbosity::Error, TEXT("\"ChannelAction\" Failed. Reason: Invalid Channel \"%s\"."), *InMessage.ChannelName);
+			}
+		}
+	}
+	else if (InMessage.Action == EAvaPlaylistChannelActions::Stop)
+	{
+		if (InMessage.ChannelName.IsEmpty())
+		{
+			Broadcast.StopBroadcast();
+			SendMessage(InContext->GetSender(), InMessage.RequestId, ELogVerbosity::Log, TEXT("\"ChannelAction\" Ok."));
+		}
+		else
+		{
+			const FName ChannelName(InMessage.ChannelName);
+			FAvaOutputChannel& Channel = Broadcast.GetCurrentProfile().GetChannelMutable(ChannelName);
+			if (Channel.IsValidChannel())
+			{
+				Channel.StopChannelBroadcast();
+				SendMessage(InContext->GetSender(), InMessage.RequestId, ELogVerbosity::Log, TEXT("\"ChannelAction\" Ok."));
+			}
+			else
+			{
+				LogAndSendMessage(InContext->GetSender(), InMessage.RequestId, ELogVerbosity::Error, TEXT("\"ChannelAction\" Failed. Reason: Invalid Channel \"%s\"."), *InMessage.ChannelName);
+			}
+		}
+	}
+	else
+	{
+		LogAndSendMessage(InContext->GetSender(), InMessage.RequestId, ELogVerbosity::Error, TEXT("\"ChannelAction\" Failed. Reason: Invalid Action (must be \"Start\" or \"Stop\"."));
+	}
+}
+
+void FAvaPlaylistServer::HandleGetChannelImage(const FAvaPlaylistGetChannelImage& InMessage, const TSharedRef<IMessageContext, ESPMode::ThreadSafe>& InContext)
+{
+	if (InMessage.ChannelName.IsEmpty())
+	{
+		SendMessage(InContext->GetSender(), InMessage.RequestId, ELogVerbosity::Error, TEXT("\"GetChannelImage\" Failed. Reason: Invalid ChannelName."));
+		return;
+	}
+	
+	const FName ChannelName(InMessage.ChannelName);
+	const FAvaOutputChannel& Channel = UAvalancheBroadcast::Get().GetCurrentProfile().GetChannel(ChannelName);
+	
+	if (!Channel.IsValidChannel())
+	{
+		SendMessage(InContext->GetSender(), InMessage.RequestId, ELogVerbosity::Error, TEXT("\"GetChannelImage\" Failed. Reason: Invalid Channel \"%s\"."), *InMessage.ChannelName);
+		return;
+	}
+
+	const FRequestInfo RequestInfo = {InMessage.RequestId, InContext->GetSender()};
+	
+	UTextureRenderTarget2D* ChannelRenderTarget = Channel.GetCurrentRenderTarget(true);
+	
+	// If the channel's render target is not the desired format, we will need to convert it.
+	TSharedPtr<FChannelImage> ChannelImage;
+	if (AvailableChannelImages.Num() > 0)
+	{
+		ChannelImage = AvailableChannelImages.Pop();
+	}
+	else
+	{
+		ChannelImage = MakeShared<FChannelImage>();
+	}
+
+	if (ChannelRenderTarget->GetFormat() != PF_B8G8R8A8)
+	{
+		ChannelImage->UpdateRenderTarget(ChannelRenderTarget->SizeX, ChannelRenderTarget->SizeY, PF_B8G8R8A8, ChannelRenderTarget->ClearColor);
+	}
+	else
+	{
+		ChannelImage->RenderTarget.Reset(); // No need for conversion.
+	}
+	
+	TWeakPtr<FAvaPlaylistServer> WeakPlaylistServer(SharedThis(this));
+
+	// The conversion is done by the GPU in the render thread.
+	ENQUEUE_RENDER_COMMAND(FAvaConvertChannelImage)(
+		[ChannelRenderTarget, ChannelImage, RequestInfo, WeakPlaylistServer](FRHICommandListImmediate& RHICmdList)
+		{
+			FRHITexture* SourceRHI = ChannelRenderTarget->GetResource()->GetTexture2DRHI();
+			FRHITexture* ReadbackRHI = SourceRHI;
+
+			// Convert if needed.
+			if (ChannelImage->RenderTarget.IsValid())
+			{
+				FRHITexture* DestinationRHI = ChannelImage->RenderTarget->GetResource()->GetTexture2DRHI();
+				UE::AvaRenderTargetMediaUtils::CopyTexture(RHICmdList, SourceRHI, DestinationRHI);
+				ReadbackRHI = DestinationRHI;
+			}
+
+			// Reading Render Target pixels in the render thread to avoid a flush render commands.
+			const FReadSurfaceDataFlags ReadDataFlags(RCM_UNorm, CubeFace_MAX);
+			const FIntRect SourceRect = FIntRect(0, 0, ChannelRenderTarget->SizeX, ChannelRenderTarget->SizeY);
+
+			ChannelImage->UpdateRawPixels(SourceRect.Width(), SourceRect.Height());
+			
+			RHICmdList.ReadSurfaceData(ReadbackRHI, SourceRect, ChannelImage->RawPixels, ReadDataFlags);
+
+			// When the converted render target is ready, we resume the work in the game thread.
+			AsyncTask(ENamedThreads::GameThread, [WeakPlaylistServer, RequestInfo, ChannelImage]()
+			{
+				if (const TSharedPtr<FAvaPlaylistServer> PlaylistServer = WeakPlaylistServer.Pin())
+				{
+					PlaylistServer->FinishGetChannelImage(RequestInfo, ChannelImage);
+				}
+			});
+		});
+}
+
+void FAvaPlaylistServer::HandleAddChannelDevice(const FAvaPlaylistAddChannelDevice& InMessage, const TSharedRef<IMessageContext, ESPMode::ThreadSafe>& InContext)
+{
+	UAvalancheBroadcast& Broadcast = UAvalancheBroadcast::Get();
+	if (InMessage.ChannelName.IsEmpty() || InMessage.MediaOutputName.IsEmpty())
+	{
+		LogAndSendMessage(InContext->GetSender(), InMessage.RequestId, ELogVerbosity::Error, TEXT("\"AddChannelDevice\" Failed. Reason: One or more Empty Parameters."));
+		return;
+	}
+
+	const FName ChannelName(InMessage.ChannelName);
+	const FAvaOutputChannel& OutputChannel = Broadcast.GetCurrentProfile().GetChannel(ChannelName);
+	if (!OutputChannel.IsValidChannel())
+	{
+		LogAndSendMessage(InContext->GetSender(), InMessage.RequestId, ELogVerbosity::Error, TEXT("\"AddChannelDevice\" Failed. Reason: Invalid Channel \"%s\"."), *InMessage.ChannelName);
+		return;    
+	}
+
+	/*
+		We're essentially replicating the UI editor here. The editor:
+		1. Builds an output tree
+		2. Allows drag-and-drop of output/devices to a channel
+		3. AddMediaOutputToChannel() is called
+
+		We don't have immediate drag-and-drop information here, since this is called externally, so we'll rebuild a tree,
+		and recursively search for a match, and then issue the same AddMediaOutputToChannel call the editor UI would've called.
+
+		This won't be called frequently, so it's equivalent to an end-user opening up and adding a device to a channel via
+		the broadcast window (tree rebuild -> drag and drop item)
+	*/
+
+	const FAvaOutputTreeItemPtr OutputDevices = MakeShared<FAvaOutputRootItem>();
+	FAvaOutputTreeItem::RefreshTree(OutputDevices);
+	FAvaOutputTreeItemPtr TreeItem = UE::AvaPlaylistServer::Private::RecursiveFindOutputTreeItem(OutputDevices, InMessage.MediaOutputName);
+	
+	if (!TreeItem.IsValid())
+	{
+		LogAndSendMessage(InContext->GetSender(), InMessage.RequestId, ELogVerbosity::Error, TEXT("\"AddChannelDevice\" Failed. Reason: Invalid Device \"%s\"."), *InMessage.MediaOutputName);
+		return;
+	}
+
+	const FAvaMediaOutputInfo OutputInfo;
+	const UMediaOutput* OutputDevice = TreeItem->AddMediaOutputToChannel(OutputChannel.GetChannelName(), OutputInfo);
+	Broadcast.SaveBroadcast();
+
+	LogAndSendMessage(InContext->GetSender(), InMessage.RequestId, ELogVerbosity::Log, TEXT("\"AddChannelDevice\" successfully added device \"%s\""), *OutputDevice->GetFName().ToString());
+}
+
+void FAvaPlaylistServer::HandleEditChannelDevice(const FAvaPlaylistEditChannelDevice& InMessage,
+	const TSharedRef<IMessageContext, ESPMode::ThreadSafe>& InContext)
+{
+	UAvalancheBroadcast& Broadcast = UAvalancheBroadcast::Get();
+	if (InMessage.ChannelName.IsEmpty() || InMessage.MediaOutputName.IsEmpty())
+	{
+		LogAndSendMessage(InContext->GetSender(), InMessage.RequestId, ELogVerbosity::Error, TEXT("\"EditChannelDevice\" Failed. Reason: One or more Empty Parameters."));
+		return;
+	}
+
+	const FName ChannelName(InMessage.ChannelName);
+	const FAvaOutputChannel& OutputChannel = Broadcast.GetCurrentProfile().GetChannel(ChannelName);
+	if (!OutputChannel.IsValidChannel())
+	{
+		LogAndSendMessage(InContext->GetSender(), InMessage.RequestId, ELogVerbosity::Error, TEXT("\"EditChannelDevice\" Failed. Reason: Invalid Channel \"%s\"."), *InMessage.ChannelName);
+		return;    
+	}
+
+	UMediaOutput* MediaOutput = UE::AvaPlaylistServer::Private::FindChannelMediaOutput(OutputChannel, InMessage.MediaOutputName);
+	
+	if (!MediaOutput)
+	{
+		LogAndSendMessage(InContext->GetSender(), InMessage.RequestId, ELogVerbosity::Error, TEXT("\"EditChannelDevice\" Failed. Reason: Invalid Device \"%s\"."), *InMessage.MediaOutputName);
+		return;
+	}
+	
+	FAvaPlaylistOutputDeviceItem DeviceItem;
+	DeviceItem.Name = InMessage.MediaOutputName;
+	DeviceItem.Data = InMessage.Data;
+
+	FAvaMediaOutputEditorUtils::EditMediaOutput(MediaOutput, DeviceItem.Data);
+
+	Broadcast.SaveBroadcast();
+	
+	SendMessage(InContext->GetSender(), InMessage.RequestId, ELogVerbosity::Log, TEXT("\"EditChannelDevice\". Successfully edited device \"%s\" on \"%s\""), *InMessage.MediaOutputName, *InMessage.ChannelName); 
+}
+
+void FAvaPlaylistServer::HandleRemoveChannelDevice(const FAvaPlaylistRemoveChannelDevice& InMessage,
+	const TSharedRef<IMessageContext, ESPMode::ThreadSafe>& InContext)
+{
+	UAvalancheBroadcast& Broadcast = UAvalancheBroadcast::Get();
+	if (InMessage.ChannelName.IsEmpty() || InMessage.MediaOutputName.IsEmpty())
+	{
+		LogAndSendMessage(InContext->GetSender(), InMessage.RequestId, ELogVerbosity::Error, TEXT("\"RemoveChannelDevice\" Failed. Reason: One or more Empty Parameters."));
+		return;
+	}
+	const FName ChannelName(InMessage.ChannelName);
+	const FAvaOutputChannel& OutputChannel = Broadcast.GetCurrentProfile().GetChannel(ChannelName);
+	if (!OutputChannel.IsValidChannel())
+	{
+		LogAndSendMessage(InContext->GetSender(), InMessage.RequestId, ELogVerbosity::Error, TEXT("\"RemoveChannelDevice\" Failed. Reason: Invalid Channel \"%s\"."), *InMessage.ChannelName);
+		return;    
+	}
+
+	UMediaOutput* MediaOutput = UE::AvaPlaylistServer::Private::FindChannelMediaOutput(OutputChannel, InMessage.MediaOutputName);
+	
+	if (!MediaOutput)
+	{
+		LogAndSendMessage(InContext->GetSender(), InMessage.RequestId, ELogVerbosity::Error, TEXT("\"RemoveChannelDevice\" Failed. Reason: Invalid Device \"%s\"."), *InMessage.MediaOutputName);
+		return;
+	}
+
+	FScopedTransaction Transaction(LOCTEXT("RemoveMediaOutput", "Remove Media Output"));
+	
+	Broadcast.Modify();
+
+	const int32 RemovedCount = Broadcast.GetCurrentProfile().RemoveChannelMediaOutputs(ChannelName, TArray{MediaOutput});
+
+	if (RemovedCount == 0)
+	{
+		Transaction.Cancel();
+		LogAndSendMessage(InContext->GetSender(), InMessage.RequestId, ELogVerbosity::Error, TEXT("\"RemoveChannelDevice\" Didn't remove device."));
+		return;
+	}
+
+	Broadcast.SaveBroadcast();
+	LogAndSendMessage(InContext->GetSender(), InMessage.RequestId, ELogVerbosity::Log, TEXT("\"RemoveChannelDevice\" Removed Device \"%s\""), *InMessage.MediaOutputName);
+}
+
+void FAvaPlaylistServer::HandleGetDevices(const FAvaPlaylistGetDevices& InMessage,
+	const TSharedRef<IMessageContext, ESPMode::ThreadSafe>& InContext)
+{
+	FAvaPlaylistDevicesList* ReplyMessage = FMessageEndpoint::MakeMessage<FAvaPlaylistDevicesList>();
+	ReplyMessage->RequestId = InMessage.RequestId;
+	
+	const FAvaOutputTreeItemPtr OutputDevices = MakeShared<FAvaOutputRootItem>();
+	FAvaOutputTreeItem::RefreshTree(OutputDevices);
+	// OutputDevices here aren't literally a physical device, just a construct representing
+	// output. This convention was pulled from the SAvaOutputDevices->RefreshDevices() call
+	for (const TSharedPtr<IAvaOutputTreeItem>& ClassItem : OutputDevices->GetChildren())
+	{
+		if (const FAvaOutputClassItem* AvaOutputClassItem = ClassItem->CastTo<FAvaOutputClassItem>())
+		{
+			FAvaPlaylistOutputClassItem OutputClassItem;
+			OutputClassItem.Name = ClassItem->GetDisplayName().ToString();
+
+			const TArray<FAvaOutputTreeItemPtr>& Children = AvaOutputClassItem->GetChildren();
+			if (Children.Num() > 0)
+			{
+				for (const TSharedPtr<IAvaOutputTreeItem>& OutputDeviceItem : Children)
+				{
+					if (OutputDeviceItem->IsA<FAvaOutputDeviceItem>())
+					{
+						FAvaPlaylistOutputDeviceItem DeviceItem;
+						DeviceItem.Name = OutputDeviceItem->GetDisplayName().ToString();
+						// Intentionally leaving DeviceItem.Data blank, as it's not usable data by itself
+						// .Data will be filled out on a GetChannels call, where it becomes usable
+
+						OutputClassItem.Devices.Push(DeviceItem);
+					}
+				}
+			}
+			else
+			{
+				FAvaPlaylistOutputDeviceItem DeviceItem;
+				DeviceItem.Name = OutputClassItem.Name;
+				OutputClassItem.Devices.Push(DeviceItem);
+			}
+
+			ReplyMessage->DeviceClasses.Push(OutputClassItem);
+		}
+	}
+
+	SendResponse(ReplyMessage, InContext->GetSender());
+}
+
+void FAvaPlaylistServer::LogAndSendMessage(const FMessageAddress& InSender, int32 InRequestId, ELogVerbosity::Type InVerbosity, const TCHAR* InFormat, ...) const
+{
+	TCHAR TempString[1024];
+	va_list Args;
+	va_start(Args, InFormat);
+	FCString::GetVarArgs(TempString, UE_ARRAY_COUNT(TempString), InFormat, Args);
+	va_end(Args);
+
+	// The UE_LOG macro adds ELogVerbosity:: to the verbosity, which prevents
+	// us from using it with a variable.
+	switch (InVerbosity)
+	{
+	case ELogVerbosity::Type::Log:
+		UE_LOG(LogAvaPlaylistServer, Log, TEXT("%s"), TempString);
+		break;
+	case ELogVerbosity::Type::Display:
+		UE_LOG(LogAvaPlaylistServer, Display, TEXT("%s"), TempString);
+		break;
+	case ELogVerbosity::Type::Warning:
+		UE_LOG(LogAvaPlaylistServer, Warning, TEXT("%s"), TempString);
+		break;
+	case ELogVerbosity::Type::Error:
+		UE_LOG(LogAvaPlaylistServer, Error, TEXT("%s"), TempString);
+		break;
+	default:
+		UE_LOG(LogAvaPlaylistServer, Log, TEXT("%s"), TempString);
+		break;
+	}
+	
+	// Send the error message to the client.
+	SendMessageImpl(InSender, InRequestId, InVerbosity, TempString);
+}
+
+void FAvaPlaylistServer::SendMessage(const FMessageAddress& InSender, int32 InRequestId, ELogVerbosity::Type InVerbosity, const TCHAR* InFormat, ...) const
+{
+	TCHAR TempString[1024];
+	va_list Args;
+	va_start(Args, InFormat);
+	FCString::GetVarArgs(TempString, UE_ARRAY_COUNT(TempString), InFormat, Args);
+	va_end(Args);
+
+	// Send the error message to the client.
+	SendMessageImpl(InSender, InRequestId, InVerbosity, TempString);
+}
+
+void FAvaPlaylistServer::SendMessageImpl(const FMessageAddress& InSender, int32 InRequestId, ELogVerbosity::Type InVerbosity, const TCHAR* InMsg) const
+{
+	FAvaPlaylistServerMsg* ErrorMessage = FMessageEndpoint::MakeMessage<FAvaPlaylistServerMsg>();
+	ErrorMessage->RequestId = InRequestId;
+	ErrorMessage->Verbosity = ToString(InVerbosity);
+	ErrorMessage->Text = InMsg;
+	SendResponse(ErrorMessage, InSender);
+}
+
+void FAvaPlaylistServer::RegisterConsoleCommands()
+{
+	if (ConsoleCommands.Num() != 0)
+	{
+		return;
+	}
+		
+	ConsoleCommands.Add(IConsoleManager::Get().RegisterConsoleCommand(
+			TEXT("AvaPlaylistServer.Status"),
+			TEXT("Display current status of all server info."),
+			FConsoleCommandWithArgsDelegate::CreateRaw(this, &FAvaPlaylistServer::ShowStatusCommand),
+			ECVF_Default
+			));
+
+}
+
+void FAvaPlaylistServer::ShowStatusCommand(const TArray<FString>& InArgs)
+{
+	UE_LOG(LogAvaPlaylistServer, Display, TEXT("Playlist Server: \"%s\""), *HostName);
+	UE_LOG(LogAvaPlaylistServer, Display, TEXT("- Endpoint Bus Address: \"%s\""), MessageEndpoint.IsValid() ? *MessageEndpoint->GetAddress().ToString() : TEXT("Invalid"));
+	UE_LOG(LogAvaPlaylistServer, Display, TEXT("- Computer: \"%s\""), *HostName);
+
+	for (const TPair<FMessageAddress, TSharedPtr<FClientInfo>>& Client : Clients)
+	{
+		const FClientInfo& ClientInfo = *Client.Value;
+		UE_LOG(LogAvaPlaylistServer, Display, TEXT("Connected Client: \"%s\""), *ClientInfo.Address.ToString());
+		UE_LOG(LogAvaPlaylistServer, Display, TEXT("   - Api Version: %d"), ClientInfo.ApiVersion);
+	}
+	
+	UE_LOG(LogAvaPlaylistServer, Display, TEXT("Playlist Caches:"));
+	UE_LOG(LogAvaPlaylistServer, Display, TEXT("- Editing Playlist: \"%s\""), *PlaylistEditCommandData.CurrentPlaylistPath.ToString());
+	UE_LOG(LogAvaPlaylistServer, Display, TEXT("- Editing PageId: \"%d\""), PlaylistEditCommandData.ManagedPageId);
+	UE_LOG(LogAvaPlaylistServer, Display, TEXT("- Playing Playlist: \"%s\""), *PlaylistPlaybackCommandData.CurrentPlaylistPath.ToString());
+
+	if (PlaylistPlaybackCommandData.CurrentPlaylist.IsValid())
+	{
+		TArray<int32> PlayingPages = PlaylistPlaybackCommandData.CurrentPlaylist->GetPlayingPageIds();
+		for (const int32 PlayingPageId : PlayingPages)
+		{
+			UE_LOG(LogAvaPlaylistServer, Display, TEXT("- Playing PageId: \"%d\""), PlayingPageId);
+		}
+		TArray<int32> PreviewingPages = PlaylistPlaybackCommandData.CurrentPlaylist->GetPreviewingPageIds();
+		for (const int32 PreviewingPageId : PreviewingPages)
+		{
+			UE_LOG(LogAvaPlaylistServer, Display, TEXT("- Previewing PageId: \"%d\""), PreviewingPageId);
+		}
+	}
+}
+
+void FAvaPlaylistServer::FinishGetChannelImage(const FRequestInfo& InRequestInfo, const TSharedPtr<FChannelImage>& InChannelImage)
+{
+	FAvaPlaylistChannelImage* ReplyMessage = FMessageEndpoint::MakeMessage<FAvaPlaylistChannelImage>();
+	ReplyMessage->RequestId = InRequestInfo.RequestId;
+	FImage Image;
+	bool bSuccess = false;
+
+	// Note: replacing FImageUtils::GetRenderTargetImage since we already have the raw pixels.
+	{
+		constexpr EPixelFormat Format = PF_B8G8R8A8;
+		const int32 ImageBytes = CalculateImageBytes(InChannelImage->SizeX, InChannelImage->SizeY, 0, Format);
+		Image.RawData.AddUninitialized(ImageBytes);
+		FMemory::Memcpy( Image.RawData.GetData(), InChannelImage->RawPixels.GetData(), InChannelImage->RawPixels.Num() * sizeof(FColor) );
+		Image.SizeX = InChannelImage->SizeX;
+		Image.SizeY = InChannelImage->SizeY;
+		Image.NumSlices = 1;
+		Image.Format = ERawImageFormat::BGRA8;
+		Image.GammaSpace = EGammaSpace::sRGB;
+	}
+
+	// TODO: profile this.
+	// Options: resize the render target on the gpu prior to reading pixels.
+	{
+		FImage ResizedImage;
+		Image.ResizeTo(ResizedImage, Image.GetWidth() * .25f, Image.GetHeight() * .25f, Image.Format, EGammaSpace::Linear);
+
+		TArray64<uint8> CompressedData;
+		if (FImageUtils::CompressImage(CompressedData,TEXT("JPEG"), ResizedImage, 95))
+		{
+			const uint32 SafeMessageSizeLimit = UE::AvaMediaMessageUtils::GetSafeMessageSizeLimit();
+			if (CompressedData.Num() > SafeMessageSizeLimit)
+			{
+				LogAndSendMessage(InRequestInfo.Sender, InRequestInfo.RequestId, ELogVerbosity::Error,
+					TEXT("\"GetChannelImage\" Failed. Reason: (DataSize: %d) is larger that the safe size limit for udp segmenter (%d)."),
+					CompressedData.Num(), SafeMessageSizeLimit);
+				return;
+			}
+
+			ReplyMessage->ImageData.Append(CompressedData.GetData(), CompressedData.GetAllocatedSize());
+			bSuccess = true;
+		}
+	}
+	
+	if (!bSuccess)
+	{
+		LogAndSendMessage(InRequestInfo.Sender, InRequestInfo.RequestId, ELogVerbosity::Error,
+			TEXT("\"GetChannelImage\" Failed. Reason: Unable to retrieve Channel Image."));
+		return;
+	}
+
+	SendResponse(ReplyMessage, InRequestInfo.Sender);
+	
+	// Put the image back in the pool of available images for next request. (or we could abandon it)
+	AvailableChannelImages.Add(InChannelImage);
+}
+
+void FAvaPlaylistServer::HandlePageActions(const FRequestInfo& InRequestInfo, const TArray<int32>& InPageIds,
+	bool bInIsPreview, FName InPreviewChannelName, EAvaPlaylistPageActions InAction) const
+{
+	using namespace UE::AvaPlaylistServer::Private;
+	
+	if (!PlaylistPlaybackCommandData.CurrentPlaylist.IsValid())
+	{
+		LogAndSendMessage(InRequestInfo.Sender, InRequestInfo.RequestId, ELogVerbosity::Error,
+			TEXT("\"PageAction\" Failed. Reason: no play list currently loaded."));
+		return;
+	}
+
+	UAvalanchePlaylist* Playlist = PlaylistPlaybackCommandData.CurrentPlaylist.Get();
+
+	{
+		// Validate the pages - the command will be considered a failure (as a whole) if it contains invalid pages.
+		FString InvalidPages;
+		for (const int32 PageId : InPageIds)
+		{
+			const FAvalanchePage& Page = Playlist->GetPage(PageId);
+			if (!Page.IsValidPage())
+			{
+				InvalidPages.Appendf(TEXT("%s%d"), InvalidPages.IsEmpty() ? TEXT("") : TEXT(", "), PageId);
+			}
+		}
+
+		if (!InvalidPages.IsEmpty())
+		{
+			LogAndSendMessage(InRequestInfo.Sender, InRequestInfo.RequestId, ELogVerbosity::Error,
+				TEXT("\"PageAction\" Failed. Reason: PageIds {%s} are invalid."), *InvalidPages);
+			return;
+		}
+	}
+
+	const FName PreviewChannelName = !InPreviewChannelName.IsNone() ? InPreviewChannelName : UAvalanchePlaylist::GetDefaultPreviewChannelName();
+	// Todo: support program channel name in command.
+	const FName CommandChannelName = bInIsPreview ? InPreviewChannelName : NAME_None; 
+	
+	bool bSuccess = false;
+	FString FailureReason;
+	switch (InAction)
+	{
+		case EAvaPlaylistPageActions::Load:
+			for (const int32 PageId : InPageIds)
+			{
+				bSuccess |= Playlist->GetPageLoadingManager().RequestLoadPage(PageId, bInIsPreview, PreviewChannelName);
+			}
+			break;
+		case EAvaPlaylistPageActions::Unload:
+			for (const int32 PageId : InPageIds)
+			{
+				const FAvalanchePage& Page =  Playlist->GetPage(PageId);
+				if (Page.IsValidPage())
+				{
+					bSuccess |= Playlist->UnloadPage(PageId, (bInIsPreview ? PreviewChannelName : Page.GetChannelName()).ToString());
+				}
+			}
+			break;
+		case EAvaPlaylistPageActions::Play:
+			bSuccess = !Playlist->PlayPages(InPageIds, bInIsPreview ? EAvaPlayType::PreviewFromStart : EAvaPlayType::PlayFromStart, PreviewChannelName).IsEmpty();
+			break;
+		case EAvaPlaylistPageActions::PlayNext:
+			{
+				const int32 NextPageId = FAvaPlaylistPlaybackUtils::GetPageIdToPlayNext(Playlist, UAvalanchePlaylist::InstancePageList, bInIsPreview, PreviewChannelName);
+				if (FAvaPlaylistPlaybackUtils::IsPageIdValid(NextPageId))
+				{
+					bSuccess = Playlist->PlayPage(NextPageId, bInIsPreview ? EAvaPlayType::PreviewFromFrame : EAvaPlayType::PlayFromStart);
+				}
+				break;
+			}
+		case EAvaPlaylistPageActions::Stop:
+			if (InPageIds.IsEmpty())
+			{
+				// If the list of pages is empty, we will stop all the playing pages.
+				const TArray<int32> PageIds = GetPlayingPages(Playlist, bInIsPreview, CommandChannelName);
+				bSuccess = !Playlist->StopPages(PageIds, EAvaPlaylistPageStopOptions::Default, bInIsPreview).IsEmpty();
+			}
+			else
+			{
+				bSuccess = !Playlist->StopPages(InPageIds, EAvaPlaylistPageStopOptions::Default, bInIsPreview).IsEmpty();
+			}
+			break;
+		case EAvaPlaylistPageActions::ForceStop:
+			if (InPageIds.IsEmpty())
+			{
+				// If the list of pages is empty, we will stop all the playing pages.
+				const TArray<int32> PageIds = GetPlayingPages(Playlist, bInIsPreview, CommandChannelName);
+				bSuccess = !Playlist->StopPages(PageIds, EAvaPlaylistPageStopOptions::ForceNoTransition, bInIsPreview).IsEmpty();
+			}
+			else
+			{
+				bSuccess = !Playlist->StopPages(InPageIds, EAvaPlaylistPageStopOptions::ForceNoTransition, bInIsPreview).IsEmpty();
+			}
+			break;
+		case EAvaPlaylistPageActions::Continue:
+			if (InPageIds.IsEmpty())
+			{
+				// If the list of pages is empty, we will continue all the playing pages.
+				const TArray<int32> PageIds = GetPlayingPages(Playlist, bInIsPreview, CommandChannelName);
+				bSuccess = ContinuePages(Playlist, PageIds, bInIsPreview, PreviewChannelName, FailureReason);
+			}
+			else
+			{
+				bSuccess = ContinuePages(Playlist, InPageIds, bInIsPreview, PreviewChannelName, FailureReason);
+			}
+			break;
+		case EAvaPlaylistPageActions::UpdateValues:
+			if (InPageIds.IsEmpty())
+			{
+				// If the list of pages is empty, we will continue all the playing pages.
+				const TArray<int32> PageIds = GetPlayingPages(Playlist, bInIsPreview, CommandChannelName);
+				bSuccess = UpdatePagesValues(Playlist, PageIds, bInIsPreview, PreviewChannelName);
+			}
+			else
+			{
+				bSuccess = UpdatePagesValues(Playlist, InPageIds, bInIsPreview, PreviewChannelName);
+			}
+			break;
+		case EAvaPlaylistPageActions::TakeToProgram:
+			{
+				const TArray<int32> PageIds = FAvaPlaylistPlaybackUtils::GetPagesToTakeToProgram(Playlist, InPageIds, PreviewChannelName);
+				Playlist->PlayPages(PageIds, EAvaPlayType::PlayFromStart);
+			}
+			break;
+		default:
+			FailureReason.Appendf(TEXT("Invalid action. "));
+			break;
+	}
+
+	const TCHAR* CommandName = bInIsPreview ? TEXT("PagePreviewAction") : TEXT("PageAction");
+
+	// For multi-page commands, we consider a partial success as success.
+	// Remote applications are notified of the page status with FAvaPlaylistPagesStatuses.
+	//
+	// Todo:
+	// For pages that failed to execute the command, the failure reason is not sent
+	// to remote applications. Given the more complex status information, we would
+	// probably need a response message for this command with additional error information.
+	
+	if (bSuccess)
+	{
+		SendMessage(InRequestInfo.Sender, InRequestInfo.RequestId, ELogVerbosity::Log,
+			TEXT("\"%s\" Ok."), CommandName);
+	}
+	else if (!FailureReason.IsEmpty())
+	{
+		LogAndSendMessage(InRequestInfo.Sender, InRequestInfo.RequestId, ELogVerbosity::Error,
+			TEXT("\"%s\" Failed. Reason: %s"), CommandName, *FailureReason);
+	}
+	else
+	{
+		LogAndSendMessage(InRequestInfo.Sender, InRequestInfo.RequestId, ELogVerbosity::Error,
+			TEXT("\"%s\" Failed."), CommandName);
+	}
+}
+
+UAvalanchePlaylist* FAvaPlaylistServer::GetOrLoadPlaylistForEdit(const FMessageAddress& InSender, int32 InRequestId, const FString& InPlaylistPath)
+{
+	UAvalanchePlaylist* Playlist;
+	
+	if (!InPlaylistPath.IsEmpty())
+	{
+		// If a path is specified, the playlist gets reloaded
+		// unless it was already loaded from a previous editing command.
+		// This will not affect the currently loaded playlist for playback.
+		const FSoftObjectPath NewPlaylistPath(InPlaylistPath);
+		
+		Playlist = PlaylistEditCommandData.GetOrLoadPlaylist(NewPlaylistPath,
+			[this](UAvalanchePlaylist* InPlaylist)
+			{
+				PlaylistEditCommandData.SaveCurrentRemoteControlPresetToPage(true);
+				PlaylistEditCommandData.RemovePlaylistDelegates(this, InPlaylist);
+			},
+			[this](UAvalanchePlaylist* InPlaylist)
+			{
+				PlaylistEditCommandData.SetupPlaylistDelegates(this, InPlaylist);
+			});
+
+		if (!Playlist)
+		{
+			LogAndSendMessage(InSender, InRequestId, ELogVerbosity::Error, TEXT("Failed to load Playlist \"%s\"."), *InPlaylistPath);
+		}
+	}
+	else
+	{
+		// If the path is not specified, we assume it is using the previously loaded playlist.
+		Playlist = PlaylistEditCommandData.CurrentPlaylist.Get();
+
+		// Note: for backward compatibility with QA python script, we allow this command to use the current "playback" playlist as fallback.
+		if (!Playlist)
+		{
+			Playlist = PlaylistPlaybackCommandData.CurrentPlaylist.Get();
+
+			// Update the edit data accordingly.
+			PlaylistEditCommandData.CurrentPlaylist.Reset(Playlist);
+			PlaylistEditCommandData.CurrentPlaylistPath = PlaylistPlaybackCommandData.CurrentPlaylistPath;
+		}
+
+		if (!Playlist)
+		{
+			LogAndSendMessage(InSender, InRequestId, ELogVerbosity::Error, TEXT("No playlist path specified and no playlist currently loaded."));
+		}
+	}
+	return Playlist;
+}
+
+void FAvaPlaylistServer::OnMessageBusNotification(const FMessageBusNotification& InNotification)
+{
+	// This is called when the websocket client disconnects.
+	if (InNotification.NotificationType == EMessageBusNotification::Unregistered)
+	{
+		TWeakPtr<FAvaPlaylistServer> ServerWeak = SharedThis(this);
+		auto RemoveClient = [ServerWeak, RegistrationAddress = InNotification.RegistrationAddress]()
+		{
+			if (const TSharedPtr<FAvaPlaylistServer> Server = ServerWeak.Pin())
+			{
+				UE_LOG(LogAvaPlaylistServer, Log, TEXT("Client \"%s\" disconnected."), *RegistrationAddress.ToString());
+				Server->Clients.Remove(RegistrationAddress);
+				Server->RefreshClientAddresses();
+			}
+		};
+
+		if (IsInGameThread())
+		{
+			RemoveClient();
+		}
+		else
+		{
+			Async(EAsyncExecution::TaskGraphMainThread, MoveTemp(RemoveClient));
+		}
+	}
+}
+
+void FAvaPlaylistServer::RefreshClientAddresses()
+{
+	ClientAddresses.Reset(Clients.Num());
+	for (const TPair<FMessageAddress, TSharedPtr<FClientInfo>>& Client : Clients)
+	{
+		ClientAddresses.Add(Client.Key);
+	}
+}
+
+UAvalanchePlaylist* FAvaPlaylistServer::FPlaylistCache::GetOrLoadPlaylist(const FSoftObjectPath& InPlaylistPath,
+	const FPlaylistEventFunction InUnloadCurrentPlaylistFunction,
+	const FPlaylistEventFunction InNewPlaylistLoadedFunction)
+{
+	if (CurrentPlaylistPath != InPlaylistPath)
+	{
+		const TStrongObjectPtr<UAvalanchePlaylist> NewPlaylist = UE::AvaPlaylistServer::Private::LoadPlaylist(InPlaylistPath);
+		if (NewPlaylist.IsValid())
+		{
+			if (CurrentPlaylist.Get())
+			{
+				InUnloadCurrentPlaylistFunction(CurrentPlaylist.Get());
+			}
+			InNewPlaylistLoadedFunction(NewPlaylist.Get());
+			CurrentPlaylist = NewPlaylist;
+			CurrentPlaylistPath = InPlaylistPath;
+		}
+		else
+		{
+			return nullptr;
+		}
+	}
+	return CurrentPlaylist.Get();
+}
+
+void FAvaPlaylistServer::FPlaylistCache::SetupPlaylistDelegates(FAvaPlaylistServer* InPlaylistServer, UAvalanchePlaylist* InPlaylist)
+{
+	RemovePlaylistDelegates(InPlaylistServer, InPlaylist);
+
+	if (!InPlaylist)
+	{
+		return;
+	}
+
+	FAvaMediaPlaybackManager& Manager = IAvaMediaModule::Get().GetLocalPlaybackManager();
+	TWeakObjectPtr<UAvalanchePlaylist> PlaylistWeak(InPlaylist);
+	TWeakPtr<FAvaPlaylistServer> PlaylistServerWeak = InPlaylistServer->AsShared();
+	OnPlaybackInstanceStatusChangedDelegateHandle = 
+		Manager.OnPlaybackInstanceStatusChanged.AddLambda([PlaylistWeak, PlaylistServerWeak](const FAvaMediaPlaybackInstance& InPlaybackInstance)
+		{
+			UAvalanchePlaylist* Playlist = PlaylistWeak.Get();
+			const TSharedPtr<FAvaPlaylistServer> PlaylistServer = PlaylistServerWeak.Pin();
+			if (IsValid(Playlist) && PlaylistServer.IsValid())
+			{
+				const int32 PageId = UAvalanchePagePlayer::GetPageIdFromInstanceUserData(InPlaybackInstance.GetInstanceUserData());
+				const FAvalanchePage Page = Playlist->GetPage(PageId);
+				if (Page.IsValidPage())
+				{
+					PlaylistServer->PageStatusChanged(Playlist, Page);
+				}
+			}
+		});
+
+	InPlaylist->GetOnPagesChanged().AddRaw(InPlaylistServer, &FAvaPlaylistServer::OnPagesChanged);
+	InPlaylist->GetOnInstancedPageListChanged().AddRaw(InPlaylistServer, &FAvaPlaylistServer::OnPageListChanged);
+	InPlaylist->GetOnTemplatePageListChanged().AddRaw(InPlaylistServer, &FAvaPlaylistServer::OnPageListChanged);
+}
+
+void FAvaPlaylistServer::FPlaylistCache::RemovePlaylistDelegates(const FAvaPlaylistServer* InPlaylistServer, UAvalanchePlaylist* InPlaylist) const
+{
+	FAvaMediaPlaybackManager& Manager = IAvaMediaModule::Get().GetLocalPlaybackManager();
+	Manager.OnPlaybackInstanceStatusChanged.Remove(OnPlaybackInstanceStatusChangedDelegateHandle);
+
+	if (InPlaylist)
+	{
+		InPlaylist->GetOnPagesChanged().RemoveAll(InPlaylistServer);
+		InPlaylist->GetOnInstancedPageListChanged().RemoveAll(InPlaylistServer);
+		InPlaylist->GetOnTemplatePageListChanged().RemoveAll(InPlaylistServer);
+	}
+}
+
+FAvaPlaylistServer::FPlaylistEditCommandData::~FPlaylistEditCommandData()
+{
+	SaveCurrentRemoteControlPresetToPage(true);
+}
+
+void FAvaPlaylistServer::FPlaylistEditCommandData::SaveCurrentRemoteControlPresetToPage(bool bInUnregister)
+{
+	if (!ManagedInstance.IsValid() || !ManagedInstance->GetRemoteControlPreset())
+	{
+		return;
+	}
+
+	// Check if the RCP was registered.
+	IRemoteControlModule& RemoteControlModule = IRemoteControlModule::Get();
+	const FName CurrentPresetName = ManagedInstance->GetRemoteControlPreset()->GetPresetName();
+	const URemoteControlPreset* ResolvedPreset = RemoteControlModule.ResolvePreset(CurrentPresetName);
+	if (ResolvedPreset != ManagedInstance->GetRemoteControlPreset())
+	{
+		return;
+	}
+
+	if (bInUnregister)
+	{
+		// Unregister from RC module.
+		RemoteControlModule.UnregisterEmbeddedPreset(CurrentPresetName);
+	}
+
+	if (!CurrentPlaylist.IsValid())
+	{
+		return;
+	}
+
+	// Save the modified values to the page.
+	FAvalanchePage& ManagedPage = CurrentPlaylist->GetPage(ManagedPageId);
+	if (!ManagedPage.IsValidPage())
+	{
+		return;
+	}
+
+	constexpr bool bIsDefault = false;
+	FAvalancheRemoteControlValues NewValues;
+	NewValues.CopyFrom(ManagedInstance->GetRemoteControlPreset(), bIsDefault);
+
+	// UpdateRemoteControlValues does half the job by ensuring that missing values are added and
+	// extra values are removed. But it doesn't change existing values.
+	EAvaRemoteControlChanges RemoteControlChanges = ManagedPage.UpdateRemoteControlValues(NewValues, bIsDefault);
+
+	// Modify existing values if different.
+	for (const TPair<FGuid, FAvalancheRemoteControlValue>& NewValue : NewValues.EntityValues)
+	{
+		const FAvalancheRemoteControlValue* ExistingValue = ManagedPage.GetRemoteControlEntityValue(NewValue.Key);
+		
+		if (ExistingValue && !NewValue.Value.IsSameValueAs(*ExistingValue))
+		{
+			ManagedPage.SetRemoteControlEntityValue(NewValue.Key, NewValue.Value);
+			RemoteControlChanges |= EAvaRemoteControlChanges::EntityValues;
+		}
+	}
+	for (const TPair<FGuid, FAvalancheRemoteControlValue>& NewValue : NewValues.ControllerValues)
+	{
+		const FAvalancheRemoteControlValue* ExistingValue = ManagedPage.GetRemoteControlControllerValue(NewValue.Key);
+
+		if (ExistingValue && !NewValue.Value.IsSameValueAs(*ExistingValue))
+		{
+			ManagedPage.SetRemoteControlControllerValue(NewValue.Key, NewValue.Value);
+			RemoteControlChanges |= EAvaRemoteControlChanges::ControllerValues;
+		}
+	}
+
+	if (RemoteControlChanges != EAvaRemoteControlChanges::None)
+	{
+		CurrentPlaylist->NotifyPageRemoteControlValueChanged(ManagedPageId, RemoteControlChanges);
+	}
+}
+
+FAvaPlaylistServer::FPlaylistPlaybackCommandData::~FPlaylistPlaybackCommandData()
+{
+	ClosePlaybackContext();
+}
+
+void FAvaPlaylistServer::FPlaylistPlaybackCommandData::ClosePlaybackContext()
+{
+	if (CurrentPlaylist.IsValid())
+	{
+		// Stop all playing pages.
+		CurrentPlaylist->ClosePlaybackContext(true);
+		CurrentPlaylist.Reset();
+		CurrentPlaylistPath.Reset();
+	}
+}
+
+#undef LOCTEXT_NAMESPACE
\ No newline at end of file
diff --git a/Engine/Plugins/Experimental/Avalanche/Source/AvalancheMediaEditor/Private/Playlist/ChannelLayerStatus/SAvaChannelLayerStatusList.cpp b/Engine/Plugins/Experimental/Avalanche/Source/AvalancheMediaEditor/Private/Playlist/ChannelLayerStatus/SAvaChannelLayerStatusList.cpp
new file mode 100644
--- /dev/null
+++ b/Engine/Plugins/Experimental/Avalanche/Source/AvalancheMediaEditor/Private/Playlist/ChannelLayerStatus/SAvaChannelLayerStatusList.cpp
@@ -0,0 +1,405 @@
+// Copyright Epic Games, Inc. All Rights Reserved.
+
+#include "SAvaChannelLayerStatusList.h"
+#include "IAvaMediaEditorModule.h"
+#include "Misc/CoreDelegates.h"
+#include "Playlist/AvaPlaylistDefines.h"
+#include "Playlist/AvaPlaylistEditor.h"
+#include "Widgets/Colors/SColorBlock.h"
+#include "Widgets/Images/SImage.h"
+#include "Widgets/Layout/SGridPanel.h"
+#include "Widgets/Layout/SScrollBox.h"
+#include "Widgets/Layout/SWrapBox.h"
+
+#define LOCTEXT_NAMESPACE "SAvaChannelLayerStatusList"
+
+namespace UE::AvalancheMediaEditor::Private
+{
+	const TArray<FLinearColor> ComboPageColorPalette = {
+		FStyleColors::AccentRed.GetSpecifiedColor(),
+		FStyleColors::AccentGreen.GetSpecifiedColor(),
+		FStyleColors::AccentBlue.GetSpecifiedColor(),
+		FStyleColors::AccentYellow.GetSpecifiedColor(),
+		FStyleColors::AccentPurple.GetSpecifiedColor(),
+		FStyleColors::AccentOrange.GetSpecifiedColor(),
+		FStyleColors::AccentGray.GetSpecifiedColor(),
+		FStyleColors::AccentBrown.GetSpecifiedColor(),
+		FStyleColors::AccentPink.GetSpecifiedColor(),
+		FStyleColors::AccentWhite.GetSpecifiedColor()
+	};
+}
+
+void SAvaChannelLayerStatusList::Construct(const FArguments& InArgs, const TSharedPtr<FAvaPlaylistEditor>& InPlaylistEditor)
+{
+	PlaylistEditorWeak = InPlaylistEditor;
+
+	ChildSlot
+	[
+		SNew(SScrollBox)
+		.Orientation(EOrientation::Orient_Vertical)
+		+ SScrollBox::Slot()
+		.FillSize(1)
+		[
+			SAssignNew(Container, SWrapBox)
+			.UseAllottedSize(true)
+			.HAlign(HAlign_Left)
+			.Orientation(EOrientation::Orient_Horizontal)
+			.InnerSlotPadding(FVector2D(5.f))
+		]
+	];
+
+	RefreshList();
+
+	UpdateHandle = FCoreDelegates::OnEndFrame.AddSPLambda(this, [this]() { RefreshList(); });
+
+	if (UAvalanchePlaylist* Playlist = InPlaylistEditor->GetPlaylist())
+	{
+		// @TODO: When events work, remove end of frame!
+		//Playlist->GetPlaybackManager().OnPlaybackInstanceStatusChanged.AddSP(this, &SAvaChannelLayerStatusList::OnPlaybackInstanceStatusChanged);
+	}
+}
+
+void SAvaChannelLayerStatusList::RefreshList()
+{
+	Container->ClearChildren();
+	ChannelLayerStatusList.Empty();
+
+	TSharedPtr<FAvaPlaylistEditor> PlaylistEditor = PlaylistEditorWeak.Pin();
+
+	if (!PlaylistEditor.IsValid())
+	{
+		return;
+	}
+
+	UAvalanchePlaylist* Playlist = PlaylistEditor->GetPlaylist();
+
+	if (!IsValid(Playlist))
+	{
+		return;
+	}
+
+	//   Channel,    Layer, List of pages trying to use this layer
+	TMap<FName, TMap<FName, TSet<int32>>> LayerPageRequest;
+
+	// Iterate over the active page players to find what's playing on each layer
+	for (const TObjectPtr<UAvalanchePagePlayer>& PagePlayer : Playlist->GetPagePlayers())
+	{
+		if (!PagePlayer->IsPlaying())
+		{
+			continue;
+		}
+
+		const int32 PageId = PagePlayer->PageId;
+		const FAvalanchePage& Page = Playlist->GetPage(PageId);
+
+		if (!Page.IsValidPage())
+		{
+			continue;
+		}
+
+		const FName& ChannelName = PagePlayer->ChannelFName;
+
+		for (const TObjectPtr<UAvaRundownPlaybackInstancePlayer>& InstancePlayer : PagePlayer->InstancePlayers)
+		{
+			if (!InstancePlayer->IsPlaying())
+			{
+				continue;
+			}
+
+			if (!InstancePlayer->TransitionLayer.IsValid())
+			{
+				continue;
+			}
+
+			const FAvaTag* LayerTag = InstancePlayer->TransitionLayer.GetTag();
+
+			if (!LayerTag)
+			{
+				continue;
+			}
+
+			const FAvalanchePage& TemplatePage = Playlist->GetPage(Page.GetTemplateId());
+
+			if (!TemplatePage.IsValidPage())
+			{
+				continue;
+			}
+
+			const FName& LayerName = LayerTag->TagName;
+			const FText LayerDescription = FText::FromName(LayerTag->TagName);
+			FLinearColor ComboPageColor = FLinearColor::Transparent;
+
+			// If we're not a combo page, just add our id to the layer request
+			if (!TemplatePage.IsComboTemplate())
+			{
+				//UE_LOG(LogAvaMediaEditor, Verbose, TEXT("Non-Template page playing on: %s: %s: %d"), *ChannelName.ToString(), *LayerName.ToString(), PageId);
+				LayerPageRequest.FindOrAdd(ChannelName).FindOrAdd(LayerName).Add(PageId);
+			}
+			/**
+			 * If this is a combo page then we need to work out if our pages are active/overridden.
+			 * For each playing page, record the pages that the entire combo template wanted to play
+			 * and compare later.
+			 */
+			else
+			{
+				// For each layer that this template uses, mark off a request for that layer.
+				for (const FAvaTagHandle& TemplateTag : TemplatePage.GetTransitionLayers(Playlist))
+				{
+					if (!TemplateTag.IsValid())
+					{
+						continue;
+					}
+
+					const FName TemplateLayerName = TemplateTag.ToName();
+
+					if (TemplateLayerName == NAME_None)
+					{
+						continue;
+					}
+
+					//UE_LOG(LogAvaMediaEditor, Warning, TEXT("Template page playing on: %s: %s: %d"), *ChannelName.ToString(), *TemplateLayerName.ToString(), PageId);
+					LayerPageRequest.FindOrAdd(ChannelName).FindOrAdd(TemplateLayerName).Add(PageId);
+				}
+			}
+
+			ChannelLayerStatusList.FindOrAdd(ChannelName).FindOrAdd(LayerName) = {
+				ChannelName,
+				LayerName,
+				PageId,
+				LayerDescription,
+				ComboPageColor,
+				false
+			};
+		}
+	}
+
+	int32 ColorIndex = 0;
+
+	//   PageId, Color
+	TMap<int32, FLinearColor> ComboPageColors;
+
+	LayerPageRequest.KeySort(
+		[](const FName& InA, const FName& InB)
+		{
+			return InA.Compare(InB) < 0;
+		});
+
+	// Now iterate through the actually playing pages and see if they match the requested combo page pages.
+	for (TPair<FName, TMap<FName, TSet<int32>>>& ChannelPair : LayerPageRequest)
+	{
+		const FName& ChannelName = ChannelPair.Key;
+
+		ChannelPair.Value.KeySort(
+			[](const FName& InA, const FName& InB)
+			{
+				return InA.Compare(InB) < 0;
+			});
+
+		for (const TPair<FName, TSet<int32>>& LayerPair : ChannelPair.Value)
+		{
+			const FName& ValueName = LayerPair.Key;
+
+			if (TMap<FName, FAvaChannelLayerEntry>* LiveChannelMap = ChannelLayerStatusList.Find(ChannelName))
+			{
+				if (FAvaChannelLayerEntry* LiveLayerEntry = LiveChannelMap->Find(ValueName))
+				{
+					if (LayerPair.Value.IsEmpty())
+					{
+						continue;
+					}
+
+					// No conflict
+					if (LayerPair.Value.Num() == 1)
+					{
+						const int32 PageId = LayerPair.Value[LayerPair.Value.begin().GetId()];
+
+						if (const FLinearColor* CurrentColor = ComboPageColors.Find(PageId))
+						{
+							LiveLayerEntry->ComboPageColor = *CurrentColor;
+						}
+						else
+						{
+							using namespace UE::AvalancheMediaEditor::Private;
+
+							LiveLayerEntry->ComboPageColor = ComboPageColorPalette[ColorIndex];
+
+							++ColorIndex;
+
+							// If we have too many, just repeat them
+							if (ColorIndex >= ComboPageColorPalette.Num())
+							{
+								ColorIndex = 0;
+							}
+
+							ComboPageColors.Add(PageId, LiveLayerEntry->ComboPageColor);
+						}
+					}
+					// If more than 1 playing page has requested this entry, it must have a conflict.
+					else
+					{
+						LiveLayerEntry->bOverridden = true;
+					}
+				}
+			}
+		}
+	}
+
+	constexpr float Padding = 5.f;
+
+	ChannelLayerStatusList.KeySort(
+		[](const FName& InA, const FName& InB)
+		{
+			return InA.Compare(InB) < 0;
+		});
+
+	// Add new widgets
+	for (const TPair<FName, TMap<FName, FAvaChannelLayerEntry>>& ChannelPair : ChannelLayerStatusList)
+	{
+		const FName& ChannelName = ChannelPair.Key;
+
+		constexpr int32 ColumnStatus = 0;
+		constexpr int32 ColumnLayer = 1;
+		constexpr int32 ColumnPageId = 2;
+		constexpr int32 ColumnCount = 3;
+
+		static const FName PreviewChannelName = "_Preview";
+		static const FText PreviewChannelDescription = LOCTEXT("Preview", "Preview");
+
+		int32 RowIndex = 0;
+
+		TSharedRef<SGridPanel> ChannelInfo = SNew(SGridPanel)
+			+ SGridPanel::Slot(ColumnStatus, RowIndex)
+			.ColumnSpan(ColumnCount)
+			.HAlign(HAlign_Fill)
+			[
+				SNew(SOverlay)
+				+ SOverlay::Slot()
+				[
+					SNew(SColorBlock)
+					.Color(FSlateColor(EStyleColor::Primary).GetSpecifiedColor())
+				]
+				+ SOverlay::Slot()
+				.Padding(Padding)
+				.HAlign(HAlign_Center)
+				[
+					SNew(STextBlock)
+					.Text(ChannelName == PreviewChannelName ? PreviewChannelDescription : FText::FromName(ChannelName))
+					.ColorAndOpacity(FSlateColor(EStyleColor::ForegroundHeader))
+				]
+			];
+
+		++RowIndex;
+
+		ChannelInfo->AddSlot(ColumnStatus, RowIndex)
+			.Padding(0.f)
+			[
+				SNew(SColorBlock)
+				.Color(FSlateColor(EStyleColor::Header).GetSpecifiedColor())
+			];
+
+		ChannelInfo->AddSlot(ColumnLayer, RowIndex)
+			.Padding(1.f, 0.f, 0.f, 0.f)
+			[
+				SNew(SOverlay)
+				+ SOverlay::Slot()
+				[
+					SNew(SColorBlock)
+					.Color(FSlateColor(EStyleColor::Header).GetSpecifiedColor())
+				]
+				+ SOverlay::Slot()
+				.Padding(Padding)
+				.HAlign(HAlign_Center)
+				[
+					SNew(STextBlock)
+					.Text(LOCTEXT("Layer", "Layer"))
+				]
+			];
+
+		ChannelInfo->AddSlot(ColumnPageId, RowIndex)
+			.Padding(1.f, 0.f, 0.f, 0.f)
+			[
+				SNew(SOverlay)
+				+ SOverlay::Slot()
+				[
+					SNew(SColorBlock)
+					.Color(FSlateColor(EStyleColor::Header).GetSpecifiedColor())
+				]
+				+ SOverlay::Slot()
+				.Padding(Padding)
+				.HAlign(HAlign_Center)
+				[
+					SNew(STextBlock)
+					.Text(LOCTEXT("PageId", "Page ID"))
+					.MinDesiredWidth(50.f)
+				]
+			];
+
+		for (const TPair<FName, FAvaChannelLayerEntry>& LayerPair : ChannelPair.Value)
+		{
+			++RowIndex;
+
+			const FAvaChannelLayerEntry& LayerEntry = LayerPair.Value;
+
+			if (LayerEntry.bOverridden)
+			{
+				ChannelInfo->AddSlot(ColumnStatus, RowIndex)
+					.Padding(Padding, Padding, 0.f, Padding)
+					[
+						SNew(SImage)
+						.Image(FAppStyle::Get().GetBrush("Icons.Next"))
+					];
+			}
+			else
+			{
+				ChannelInfo->AddSlot(ColumnStatus, RowIndex)
+					.Padding(Padding)
+					[
+						SNew(SColorBlock)
+						.Color(LayerEntry.ComboPageColor)
+						.Size(FVector2D(16.0))
+						.Visibility(LayerEntry.ComboPageColor != FLinearColor::Transparent ? EVisibility::Visible : EVisibility::Hidden)
+					];
+			}
+
+			ChannelInfo->AddSlot(ColumnLayer, RowIndex)
+				.Padding(Padding, Padding, 0.f, Padding)
+				[
+					SNew(STextBlock)
+					.Text(LayerEntry.LayerDescription)
+				];
+
+			ChannelInfo->AddSlot(ColumnPageId, RowIndex)
+				.Padding(Padding, Padding, 0.f, Padding)
+				[
+					SNew(STextBlock)
+					.Text(FText::AsNumber(LayerEntry.PageId, &UE::AvalanchePlaylist::FEditorMetrics::PageIdFormattingOptions))
+				];
+		}
+
+		Container->AddSlot()
+			[
+				ChannelInfo
+			];
+	}
+
+	/*
+	 * @TODO: When events work, re-enable this.
+	if (UpdateHandle.IsValid())
+	{
+		FCoreDelegates::OnEndFrame.Remove(UpdateHandle);
+		UpdateHandle.Reset();
+	}
+	*/
+}
+
+void SAvaChannelLayerStatusList::OnPlaybackInstanceStatusChanged(const FAvaMediaPlaybackInstance& InPlaybackInstance)
+{
+	if (UpdateHandle.IsValid())
+	{
+		return;
+	}
+
+	UpdateHandle = FCoreDelegates::OnEndFrame.AddSPLambda(this, [this]() { RefreshList(); });
+}
+
+#undef LOCTEXT_NAMESPACE
diff --git a/Engine/Plugins/Experimental/Avalanche/Source/AvalancheMediaEditor/Private/Playlist/ChannelLayerStatus/SAvaChannelLayerStatusList.h b/Engine/Plugins/Experimental/Avalanche/Source/AvalancheMediaEditor/Private/Playlist/ChannelLayerStatus/SAvaChannelLayerStatusList.h
new file mode 100644
--- /dev/null
+++ b/Engine/Plugins/Experimental/Avalanche/Source/AvalancheMediaEditor/Private/Playlist/ChannelLayerStatus/SAvaChannelLayerStatusList.h
@@ -0,0 +1,48 @@
+// Copyright Epic Games, Inc. All Rights Reserved.
+
+#pragma once
+
+#include "Containers/Map.h"
+#include "Templates/SharedPointer.h"
+#include "Widgets/SCompoundWidget.h"
+#include "Delegates/IDelegateInstance.h"
+
+class FAvaMediaPlaybackInstance;
+class FAvaPlaylistEditor;
+class SAvaChannelLayerStatus;
+class SWrapBox;
+
+struct FAvaChannelLayerEntry
+{
+	FName ChannelName;
+	FName LayerName;
+	int32 PageId;
+	FText LayerDescription;
+	FLinearColor ComboPageColor;
+	bool bOverridden;
+};
+
+/*
+ * Widget containing the Status of all ChannelLayers in Broadcast. It contains a list of SAvaChannelLayerStatus widgets
+ */
+class SAvaChannelLayerStatusList : public SCompoundWidget
+{
+public:
+	SLATE_BEGIN_ARGS(SAvaChannelLayerStatusList) {}
+	SLATE_END_ARGS()
+
+	/** Constructs this widget with InArgs */
+	void Construct(const FArguments& InArgs, const TSharedPtr<FAvaPlaylistEditor>& InPlaylistEditor);
+
+protected:
+	TWeakPtr<FAvaPlaylistEditor> PlaylistEditorWeak;
+	TSharedPtr<SWrapBox> Container;
+	FDelegateHandle UpdateHandle;
+
+	//   Channel,    Layer, Details
+	TMap<FName, TMap<FName, FAvaChannelLayerEntry>> ChannelLayerStatusList;
+
+	void RefreshList();
+
+	void OnPlaybackInstanceStatusChanged(const FAvaMediaPlaybackInstance& InPlaybackInstance);
+};
diff --git a/Engine/Plugins/Experimental/Avalanche/Source/AvalancheMediaEditor/Private/Playlist/ChannelStatus/SAvaChannelStatus.cpp b/Engine/Plugins/Experimental/Avalanche/Source/AvalancheMediaEditor/Private/Playlist/ChannelStatus/SAvaChannelStatus.cpp
new file mode 100644
--- /dev/null
+++ b/Engine/Plugins/Experimental/Avalanche/Source/AvalancheMediaEditor/Private/Playlist/ChannelStatus/SAvaChannelStatus.cpp
@@ -0,0 +1,55 @@
+// Copyright Epic Games, Inc. All Rights Reserved.
+
+#include "SAvaChannelStatus.h"
+
+#include "AvaMediaEditorStyle.h"
+#include "AvaMediaEditorUtils.h"
+#include "AvalancheBroadcast.h"
+#include "Components/HorizontalBox.h"
+#include "Widgets/Images/SImage.h"
+#include "Widgets/Layout/SBox.h"
+#include "Widgets/Text/STextBlock.h"
+
+void SAvaChannelStatus::Construct(const FArguments& InArgs, const FAvaOutputChannel& InChannel)
+{
+	ChannelName = InChannel.GetChannelName();
+	FAvaOutputChannel::GetOnChannelChanged().AddSP(this, &SAvaChannelStatus::OnChannelChanged);
+
+	ChildSlot
+	.Padding(FMargin(10.f, 3.f))
+	[
+		SNew(SHorizontalBox)		
+		+ SHorizontalBox::Slot()
+		.AutoWidth()
+		[
+			SAssignNew(ChannelStatusIcon, SImage)
+			.DesiredSizeOverride(FVector2D(16.f))
+		]
+		+ SHorizontalBox::Slot()
+		.FillWidth(1.f)
+		[
+			SNew(STextBlock)
+			.Text(FText::FromName(ChannelName))
+		]
+	];
+
+	OnChannelChanged(InChannel, EAvaChannelChange::State);
+}
+
+SAvaChannelStatus::~SAvaChannelStatus()
+{
+	FAvaOutputChannel::GetOnChannelChanged().RemoveAll(this);
+}
+
+void SAvaChannelStatus::OnChannelChanged(const FAvaOutputChannel& InChannel, EAvaChannelChange InChange)
+{
+	const bool bMatchingChannels = InChannel.GetChannelName() == ChannelName;
+	const bool bStateChanged     = EnumHasAnyFlags(InChange, EAvaChannelChange::State);
+	
+	if (!bMatchingChannels || !bStateChanged)
+	{
+		return;
+	}
+	
+	ChannelStatusIcon->SetImage(FAvaMediaEditorUtils::GetChannelStatusBrush(InChannel.GetState(), InChannel.GetIssueSeverity()));
+}
diff --git a/Engine/Plugins/Experimental/Avalanche/Source/AvalancheMediaEditor/Private/Playlist/ChannelStatus/SAvaChannelStatus.h b/Engine/Plugins/Experimental/Avalanche/Source/AvalancheMediaEditor/Private/Playlist/ChannelStatus/SAvaChannelStatus.h
new file mode 100644
--- /dev/null
+++ b/Engine/Plugins/Experimental/Avalanche/Source/AvalancheMediaEditor/Private/Playlist/ChannelStatus/SAvaChannelStatus.h
@@ -0,0 +1,33 @@
+// Copyright Epic Games, Inc. All Rights Reserved.
+
+#pragma once
+
+#include "CoreMinimal.h"
+#include "Channel/AvaOutputChannel.h"
+#include "Widgets/SCompoundWidget.h"
+
+class SImage;
+class UAvalancheBroadcast;
+
+/*
+ * Widget displaying the Status of a Single Channel
+ */
+class SAvaChannelStatus : public SCompoundWidget
+{
+public:
+	
+	SLATE_BEGIN_ARGS(SAvaChannelStatus) {}
+	SLATE_END_ARGS()
+
+	void Construct(const FArguments& InArgs, const FAvaOutputChannel& InChannel);
+
+	virtual ~SAvaChannelStatus() override;
+	
+	void OnChannelChanged(const FAvaOutputChannel& InChannel, EAvaChannelChange InChange);
+	
+protected:
+
+	FName ChannelName;
+
+	TSharedPtr<SImage> ChannelStatusIcon;
+};
diff --git a/Engine/Plugins/Experimental/Avalanche/Source/AvalancheMediaEditor/Private/Playlist/ChannelStatus/SAvaChannelStatusList.cpp b/Engine/Plugins/Experimental/Avalanche/Source/AvalancheMediaEditor/Private/Playlist/ChannelStatus/SAvaChannelStatusList.cpp
new file mode 100644
--- /dev/null
+++ b/Engine/Plugins/Experimental/Avalanche/Source/AvalancheMediaEditor/Private/Playlist/ChannelStatus/SAvaChannelStatusList.cpp
@@ -0,0 +1,66 @@
+// Copyright Epic Games, Inc. All Rights Reserved.
+
+#include "SAvaChannelStatusList.h"
+#include "AvalancheBroadcast.h"
+#include "SAvaChannelStatus.h"
+#include "Widgets/Layout/SGridPanel.h"
+#include "Widgets/Layout/SWrapBox.h"
+
+void SAvaChannelStatusList::Construct(const FArguments& InArgs)
+{
+	BroadcastWeak = UAvalancheBroadcast::GetAvalancheBroadcast();
+	
+	check(BroadcastWeak.IsValid());
+	
+	BroadcastWeak->AddChangeListener(FOnAvaBroadcastChanged::FDelegate::CreateSP(this
+		, &SAvaChannelStatusList::OnBroadcastChanged));
+	
+	ChildSlot
+	[
+		SAssignNew(WrapBox, SWrapBox)
+		.UseAllottedSize(true)
+		.HAlign(EHorizontalAlignment::HAlign_Center)
+	];
+
+	RefreshList();
+}
+
+SAvaChannelStatusList::~SAvaChannelStatusList()
+{
+	if (BroadcastWeak.IsValid())
+	{
+		BroadcastWeak->RemoveChangeListener(this);
+	}
+}
+
+void SAvaChannelStatusList::RefreshList()
+{
+	UAvalancheBroadcast* const Broadcast = BroadcastWeak.Get();
+	
+	if (!Broadcast || !WrapBox.IsValid())
+	{
+		return;
+	}
+
+	WrapBox->ClearChildren();
+
+	const TArray<FAvaOutputChannel*>& Channels = Broadcast->GetCurrentProfile().GetChannels();
+	
+	for (const FAvaOutputChannel* Channel : Channels)
+	{
+		WrapBox->AddSlot()
+			[
+				SNew(SAvaChannelStatus, *Channel)
+			];
+	}
+}
+
+void SAvaChannelStatusList::OnBroadcastChanged(EAvaBroadcastChange InChange)
+{
+	if (EnumHasAnyFlags(InChange, EAvaBroadcastChange::CurrentProfile
+		| EAvaBroadcastChange::ChannelGrid
+		| EAvaBroadcastChange::ChannelRename))
+	{
+		RefreshList();
+	}
+}
diff --git a/Engine/Plugins/Experimental/Avalanche/Source/AvalancheMediaEditor/Private/Playlist/ChannelStatus/SAvaChannelStatusList.h b/Engine/Plugins/Experimental/Avalanche/Source/AvalancheMediaEditor/Private/Playlist/ChannelStatus/SAvaChannelStatusList.h
new file mode 100644
--- /dev/null
+++ b/Engine/Plugins/Experimental/Avalanche/Source/AvalancheMediaEditor/Private/Playlist/ChannelStatus/SAvaChannelStatusList.h
@@ -0,0 +1,39 @@
+// Copyright Epic Games, Inc. All Rights Reserved.
+
+#pragma once
+
+#include "AvaMediaDefines.h"
+#include "CoreMinimal.h"
+#include "Widgets/SCompoundWidget.h"
+
+class SAvaChannelStatus;
+class SWrapBox;
+class UAvalancheBroadcast;
+
+/*
+ * Widget containing the Status of all Channels in Broadcast. It contains a list of SAvaChannelStatus widgets
+ */
+class SAvaChannelStatusList : public SCompoundWidget
+{
+public:
+	
+	SLATE_BEGIN_ARGS(SAvaChannelStatusList){}
+	SLATE_END_ARGS()
+
+	/** Constructs this widget with InArgs */
+	void Construct(const FArguments& InArgs);
+	
+	virtual ~SAvaChannelStatusList() override;
+
+	void RefreshList();
+	
+	void OnBroadcastChanged(EAvaBroadcastChange InChange);
+
+protected:
+
+	TSharedPtr<SWrapBox> WrapBox;
+
+	TWeakObjectPtr<UAvalancheBroadcast> BroadcastWeak;
+
+	TMap<FName, TSharedPtr<SAvaChannelStatus>> ChannelStatusSlots;
+};
diff --git a/Engine/Plugins/Experimental/Avalanche/Source/AvalancheMediaEditor/Private/Playlist/DetailsView/RemoteControl/Controllers/AvaRCControllerItem.cpp b/Engine/Plugins/Experimental/Avalanche/Source/AvalancheMediaEditor/Private/Playlist/DetailsView/RemoteControl/Controllers/AvaRCControllerItem.cpp
new file mode 100644
--- /dev/null
+++ b/Engine/Plugins/Experimental/Avalanche/Source/AvalancheMediaEditor/Private/Playlist/DetailsView/RemoteControl/Controllers/AvaRCControllerItem.cpp
@@ -0,0 +1,34 @@
+// Copyright Epic Games, Inc. All Rights Reserved.
+
+#include "AvaRCControllerItem.h"
+#include "Controller/RCController.h"
+#include "IRemoteControlUIModule.h"
+#include "RCVirtualProperty.h"
+#include "SAvaRCControllerItemRow.h"
+
+FAvaRCControllerItem::FAvaRCControllerItem(int32 InInstanceIndex, FName InAssetName, URCController* InController, const TSharedRef<IDetailTreeNode>& InTreeNode)
+{
+	AssetName = InAssetName;
+	InstanceIndex = InInstanceIndex;
+	
+	if (InController)
+	{
+		const FName DisplayName = InController->DisplayName.IsNone()
+			? InController->PropertyName
+			: InController->DisplayName;
+		
+		DisplayNameText = FText::FromName(DisplayName);
+		NodeWidgets = InTreeNode->CreateNodeWidgets();
+
+		// Check if the Controller has a custom widget. In that case, overwrite the value widget with it
+		if (const TSharedPtr<SWidget>& CustomControllerWidget = IRemoteControlUIModule::Get().CreateCustomControllerWidget(InController, InTreeNode->CreatePropertyHandle()))
+		{
+			NodeWidgets.ValueWidget = CustomControllerWidget;
+		}
+	}
+}
+
+TSharedRef<ITableRow> FAvaRCControllerItem::CreateWidget(TSharedRef<SAvaRCControllerPanel> InControllerPanel, const TSharedRef<STableViewBase>& InOwnerTable) const
+{
+	return SNew(SAvaRCControllerItemRow, InControllerPanel, InOwnerTable, SharedThis(this));
+}
diff --git a/Engine/Plugins/Experimental/Avalanche/Source/AvalancheMediaEditor/Private/Playlist/DetailsView/RemoteControl/Controllers/AvaRCControllerItem.h b/Engine/Plugins/Experimental/Avalanche/Source/AvalancheMediaEditor/Private/Playlist/DetailsView/RemoteControl/Controllers/AvaRCControllerItem.h
new file mode 100644
--- /dev/null
+++ b/Engine/Plugins/Experimental/Avalanche/Source/AvalancheMediaEditor/Private/Playlist/DetailsView/RemoteControl/Controllers/AvaRCControllerItem.h
@@ -0,0 +1,34 @@
+// Copyright Epic Games, Inc. All Rights Reserved.
+
+#pragma once
+
+#include "IDetailTreeNode.h"
+#include "Internationalization/Text.h"
+#include "Templates/SharedPointer.h"
+
+class ITableRow;
+class SAvaRCControllerPanel;
+class STableViewBase;
+class URCController;
+
+class FAvaRCControllerItem : public TSharedFromThis<FAvaRCControllerItem>
+{
+public:
+	FAvaRCControllerItem(int32 InInstanceIndex, FName InAssetName, URCController* InController, const TSharedRef<IDetailTreeNode>& InTreeNode);
+	
+	TSharedRef<ITableRow> CreateWidget(TSharedRef<SAvaRCControllerPanel> InControllerPanel, const TSharedRef<STableViewBase>& InOwnerTable) const;
+
+	FText GetDisplayName() const { return DisplayNameText; }
+	const FNodeWidgets& GetNodeWidgets() const { return NodeWidgets; }
+	
+	int32 GetInstanceIndex() const { return InstanceIndex; }
+	int32 GetDisplayIndex() const { return DisplayIndex; }
+	FName GetAssetName() const { return AssetName; }
+	
+private:
+	int32 InstanceIndex = 0;
+	int32 DisplayIndex = 0;
+	FName AssetName;
+	FText DisplayNameText;
+	FNodeWidgets NodeWidgets;
+};
diff --git a/Engine/Plugins/Experimental/Avalanche/Source/AvalancheMediaEditor/Private/Playlist/DetailsView/RemoteControl/Controllers/SAvaRCControllerItemRow.cpp b/Engine/Plugins/Experimental/Avalanche/Source/AvalancheMediaEditor/Private/Playlist/DetailsView/RemoteControl/Controllers/SAvaRCControllerItemRow.cpp
new file mode 100644
--- /dev/null
+++ b/Engine/Plugins/Experimental/Avalanche/Source/AvalancheMediaEditor/Private/Playlist/DetailsView/RemoteControl/Controllers/SAvaRCControllerItemRow.cpp
@@ -0,0 +1,71 @@
+// Copyright Epic Games, Inc. All Rights Reserved.
+
+#include "SAvaRCControllerItemRow.h"
+#include "AvaRCControllerItem.h"
+#include "Widgets/Layout/SScissorRectBox.h"
+#include "Widgets/SNullWidget.h"
+#include "Widgets/Text/STextBlock.h"
+
+void SAvaRCControllerItemRow::Construct(const FArguments& InArgs, TSharedRef<SAvaRCControllerPanel> InControllerPanel,
+	const TSharedRef<STableViewBase>& InOwnerTableView, const TSharedPtr<const FAvaRCControllerItem>& InRowItem)
+{
+	ItemPtrWeak = InRowItem;
+	ControllerPanelWeak = InControllerPanel;
+
+	SMultiColumnTableRow<FAvaRCControllerItemPtr>::Construct(FSuperRowType::FArguments(), InOwnerTableView);
+}
+
+TSharedRef<SWidget> SAvaRCControllerItemRow::GenerateWidgetForColumn(const FName& InColumnName)
+{
+	TSharedPtr<const FAvaRCControllerItem> ItemPtr = ItemPtrWeak.Pin();
+
+	if (ItemPtr.IsValid())
+	{
+		if (InColumnName == SAvaRCControllerPanel::ControllerColumnName)
+		{
+			return SNew(SScissorRectBox)
+				[
+					SNew(SBox)
+					.VAlign(EVerticalAlignment::VAlign_Center)
+					.Padding(3.f, 2.f, 3.f, 2.f)
+					[
+						SNew(STextBlock)
+						.Text(ItemPtr->GetDisplayName())
+					]
+				];
+		}
+		else if (InColumnName == SAvaRCControllerPanel::ValueColumnName)
+		{
+			if (ItemPtr->GetNodeWidgets().ValueWidget.IsValid())
+			{
+				return ItemPtr->GetNodeWidgets().ValueWidget.ToSharedRef();
+			}
+			else if (ItemPtr->GetNodeWidgets().WholeRowWidget.IsValid())
+			{
+				return ItemPtr->GetNodeWidgets().WholeRowWidget.ToSharedRef();
+			}
+		}
+		else
+		{
+			TSharedPtr<SAvaRCControllerPanel> ControllerPanel = ControllerPanelWeak.Pin();
+
+			if (ControllerPanel.IsValid())
+			{
+				TSharedPtr<SWidget> Cell = nullptr;
+				const TArray<FAvaRCControllerTableRowExtensionDelegate>& TableRowExtensionDelegates = ControllerPanel->GetTableRowExtensionDelegates(InColumnName);
+
+				for (const FAvaRCControllerTableRowExtensionDelegate& TableRowExtensionDelegate : TableRowExtensionDelegates)
+				{
+					TableRowExtensionDelegate.ExecuteIfBound(ControllerPanel.ToSharedRef(), ItemPtr.ToSharedRef(), Cell);
+				}
+
+				if (Cell.IsValid())
+				{
+					return Cell.ToSharedRef();
+				}
+			}
+		}
+	}
+
+	return SNullWidget::NullWidget;
+}
diff --git a/Engine/Plugins/Experimental/Avalanche/Source/AvalancheMediaEditor/Private/Playlist/DetailsView/RemoteControl/Controllers/SAvaRCControllerItemRow.h b/Engine/Plugins/Experimental/Avalanche/Source/AvalancheMediaEditor/Private/Playlist/DetailsView/RemoteControl/Controllers/SAvaRCControllerItemRow.h
new file mode 100644
--- /dev/null
+++ b/Engine/Plugins/Experimental/Avalanche/Source/AvalancheMediaEditor/Private/Playlist/DetailsView/RemoteControl/Controllers/SAvaRCControllerItemRow.h
@@ -0,0 +1,25 @@
+// Copyright Epic Games, Inc. All Rights Reserved.
+
+#pragma once
+
+#include "AvaRCControllerItem.h"
+#include "SAvaRCControllerPanel.h"
+#include "Templates/SharedPointer.h"
+#include "Widgets/Views/SListView.h"
+#include "Widgets/Views/STableRow.h"
+
+class SAvaRCControllerItemRow : public SMultiColumnTableRow<FAvaRCControllerItemPtr>
+{
+public:
+	SLATE_BEGIN_ARGS(SAvaRCControllerItemRow) { }
+	SLATE_END_ARGS()
+
+	void Construct(const FArguments& InArgs, TSharedRef<SAvaRCControllerPanel> InControllerPanel, 
+		const TSharedRef<STableViewBase>& InOwnerTableView, const TSharedPtr<const FAvaRCControllerItem>& InRowItem);
+
+	virtual TSharedRef<SWidget> GenerateWidgetForColumn(const FName& InColumnName) override;
+
+protected:
+	TWeakPtr<const FAvaRCControllerItem> ItemPtrWeak;
+	TWeakPtr<SAvaRCControllerPanel> ControllerPanelWeak;
+};
