[MotionDesign]

diff --git a/Engine/Plugins/Experimental/Avalanche/Source/AvalancheMediaEditor/Private/Playlist/DetailsView/RemoteControl/Controllers/SAvaRCControllerPanel.cpp b/Engine/Plugins/Experimental/Avalanche/Source/AvalancheMediaEditor/Private/Playlist/DetailsView/RemoteControl/Controllers/SAvaRCControllerPanel.cpp
new file mode 100644
--- /dev/null
+++ b/Engine/Plugins/Experimental/Avalanche/Source/AvalancheMediaEditor/Private/Playlist/DetailsView/RemoteControl/Controllers/SAvaRCControllerPanel.cpp
@@ -0,0 +1,488 @@
+// Copyright Epic Games, Inc. All Rights Reserved.
+
+#include "SAvaRCControllerPanel.h"
+
+#include "AvaRCControllerItem.h"
+#include "Behaviour/Builtin/Path/RCSetAssetByPathBehaviour.h"
+#include "Controller/RCController.h"
+#include "IDetailTreeNode.h"
+#include "IPropertyRowGenerator.h"
+#include "Playback/AvalancheRemoteControl.h"
+#include "Playlist/AvaPlaylistEditor.h"
+#include "Playlist/AvalanchePage.h"
+#include "Playlist/AvalancheManagedInstanceCache.h"
+#include "RCVirtualProperty.h"
+#include "Widgets/Views/SHeaderRow.h"
+#include "Widgets/Views/SListView.h"
+
+#define LOCTEXT_NAMESPACE "AvaRCControllerPanel"
+
+const FName SAvaRCControllerPanel::ControllerColumnName = "ControllerColumn";
+const FName SAvaRCControllerPanel::ValueColumnName = "ValueColumn";
+
+FAvaRCControllerHeaderRowExtensionDelegate SAvaRCControllerPanel::HeaderRowExtensionDelegate;
+TMap<FName, TArray<FAvaRCControllerTableRowExtensionDelegate>> SAvaRCControllerPanel::TableRowExtensionDelegates;
+
+TArray<FAvaRCControllerTableRowExtensionDelegate>& SAvaRCControllerPanel::GetTableRowExtensionDelegates(FName InExtensionName)
+{
+	return TableRowExtensionDelegates.FindOrAdd(InExtensionName);
+}
+
+void SAvaRCControllerPanel::Construct(const FArguments& InArgs, const TSharedPtr<FAvaPlaylistEditor>& InPlaylistEditor)
+{
+	PlaylistEditorWeak = InPlaylistEditor;
+	ActivePageId = FAvalanchePage::InvalidPageId;
+
+	ChildSlot
+	[
+		SNew(SBorder)
+		.Padding(8.0f)
+		.BorderImage(FAppStyle::GetBrush("ToolPanel.GroupBorder"))
+		[
+			SAssignNew(ControllerContainer, SListView<FAvaRCControllerItemPtr>)
+			.ListItemsSource(&ControllerItems)
+			.SelectionMode(ESelectionMode::None)
+			.OnGenerateRow(this, &SAvaRCControllerPanel::OnGenerateControllerRow)
+			.HeaderRow(
+				SNew(SHeaderRow)
+				.CanSelectGeneratedColumn(true)
+				+ SHeaderRow::Column(ControllerColumnName)
+				.DefaultLabel(LOCTEXT("Controller", "Controller"))
+				.FillWidth(0.2f)
+				+ SHeaderRow::Column(ValueColumnName)
+				.DefaultLabel(LOCTEXT("Value", "Value"))
+				.FillWidth(0.8f)
+			)
+		]
+	];
+	
+	Refresh({});
+}
+
+bool SAvaRCControllerPanel::HasRemoteControlPreset(const URemoteControlPreset* InPreset) const
+{
+	for (const TSharedPtr<FAvalancheManagedInstance>& ManagedInstance : ManagedInstances)
+	{
+		if (ManagedInstance && ManagedInstance->GetRemoteControlPreset() == InPreset)
+		{
+			return true;
+		}
+	}
+	return false;
+}
+
+void SAvaRCControllerPanel::OnPageSelectionChanged(const TArray<int32>& InSelectedPageIds)
+{
+	Refresh(InSelectedPageIds);
+}
+
+void SAvaRCControllerPanel::UpdatePropertyRowGenerators(int32 InNumGenerators)
+{
+	if (PropertyRowGenerators.Num() != InNumGenerators)
+	{
+		TArray<TUniquePtr<FPropertyRowGeneratorWrapper>> GeneratorPool = MoveTemp(PropertyRowGenerators);
+
+		PropertyRowGenerators.Empty(InNumGenerators);
+
+		for (int32 Index = 0; Index < InNumGenerators; ++Index)
+		{
+			if (Index < GeneratorPool.Num())
+			{
+				PropertyRowGenerators.Add(MoveTemp(GeneratorPool[Index]));
+			}
+			else
+			{
+				PropertyRowGenerators.Add(MakeUnique<FPropertyRowGeneratorWrapper>(this));
+			}
+		}
+	}
+}
+
+void SAvaRCControllerPanel::RefreshForManagedInstance(int32 InInstanceIndex, const FAvalancheManagedInstance& InManagedInstance, const FAvalanchePage& InPage)
+{
+	URemoteControlPreset* const Preset = InManagedInstance.GetRemoteControlPreset();
+
+	if (!Preset || !PropertyRowGenerators.IsValidIndex(InInstanceIndex)
+		|| !PropertyRowGenerators[InInstanceIndex] || !PropertyRowGenerators[InInstanceIndex]->PropertyRowGenerator)
+	{
+		return;
+	}
+	
+	BindRemoteControlDelegates(Preset);
+
+	if (const TSharedPtr<FStructOnScope> StructOnScope = Preset->GetControllerContainerStructOnScope())
+	{
+		IPropertyRowGenerator& PropertyRowGenerator = *PropertyRowGenerators[InInstanceIndex]->PropertyRowGenerator;
+		
+		// We need one of those for each Preset.
+		PropertyRowGenerator.SetStructure(StructOnScope);
+		PropertyRowGenerators[InInstanceIndex]->PresetWeak = Preset;	// Keep track for proper event routing.
+		
+		const TArray<TSharedRef<IDetailTreeNode>>& RootTreeNodes = PropertyRowGenerator.GetRootTreeNodes();
+		check(RootTreeNodes.Num() <= 1);
+
+		for (const TSharedRef<IDetailTreeNode>& RootTreeNode : RootTreeNodes)
+		{
+			TArray<TSharedRef<IDetailTreeNode>> Children;
+			RootTreeNode->GetChildren(Children);
+
+			for (TSharedRef<IDetailTreeNode>& Child : Children)
+			{
+				FProperty* const Property = Child->CreatePropertyHandle()->GetProperty();
+				check(Property);
+
+				if (Property->IsA<FStrProperty>() || Property->IsA<FTextProperty>())
+				{
+					Property->SetMetaData(TEXT("MultiLine"), TEXT("true"));
+				}
+
+				URCVirtualPropertyBase* const VirtualProperty = Preset->GetController(Property->GetFName());
+
+				if (!VirtualProperty)
+				{
+					continue;
+				}
+
+				// Apply the page value to the controller (sync the managed RCP's controller to the page value).
+				{
+					const FAvalancheRemoteControlValue* ControllerValueFromPage = InPage.GetRemoteControlControllerValue(VirtualProperty->Id);
+					if (!ControllerValueFromPage)
+					{
+						// If the value is not set in the page, fallback to default value from the template.
+						ControllerValueFromPage = InManagedInstance.GetDefaultRemoteControlValues().GetControllerValue(VirtualProperty->Id);
+
+						// Note: we are not adding the controller's default value in the page. (Tentative)
+						// Reason: Given that controller values are only applied for the managed RCP (ui),
+						// there is no need to save it in the page unless the user has modified from default.
+						// This way, we only keep user modified values.
+						
+						if (!ControllerValueFromPage)
+						{
+							UE_LOG(LogAvaMediaRemoteControl, Error, TEXT("Controller \"%s\" (id:%s) doesn't have a template default value."),
+								*VirtualProperty->DisplayName.ToString(), *VirtualProperty->Id.ToString());
+						}
+					}
+					if (ControllerValueFromPage)
+					{
+						using namespace UE::AvalancheRemoteControl;
+						FAvalancheRemoteControlValue CurrentControllerValue;
+						EAvaRemoteControlResult RemoteControlResult = GetValueOfController(VirtualProperty, CurrentControllerValue.Value);
+						if (!Failed(RemoteControlResult))
+						{
+							if (!CurrentControllerValue.IsSameValueAs(*ControllerValueFromPage))
+							{
+								// This serves only to sync the Managed RCP's controller values to the page values.
+								// We want to be able to do this without affecting the entity values (preserve WYSIWYG).
+								// So, we temporarily disable the controller's behaviours.
+								FScopedPushControllerBehavioursEnable PushBehavioursEnable(VirtualProperty, false);
+								if (Failed(RemoteControlResult = SetValueOfController(VirtualProperty, ControllerValueFromPage->Value)))
+								{
+									UE_LOG(LogAvaMediaRemoteControl, Error,
+										TEXT("Controller \"%s\" (id:%s): failed to set value in currently selected page: %s."),
+										*VirtualProperty->DisplayName.ToString(), *VirtualProperty->Id.ToString(), *EnumToString(RemoteControlResult));
+								}
+							}
+						}
+						else
+						{
+							UE_LOG(LogAvaMediaRemoteControl, Error,
+								TEXT("Controller \"%s\" (id:%s): failed to get value in currently selected page: %s."),
+								*VirtualProperty->DisplayName.ToString(), *VirtualProperty->Id.ToString(), *EnumToString(RemoteControlResult));
+						}
+					}
+				}
+
+				if (URCController* const Controller = Cast<URCController>(VirtualProperty))
+				{
+					for (URCBehaviour* const Behavior : Controller->Behaviours)
+					{
+						if (URCSetAssetByPathBehaviour* const AssetByPathBehavior = Cast<URCSetAssetByPathBehaviour>(Behavior))
+						{
+							AssetByPathBehavior->UpdateTargetEntity();
+						}
+					}
+
+					ControllerItems.Add(MakeShared<FAvaRCControllerItem>(InInstanceIndex, InManagedInstance.GetSourceAssetPath().GetAssetFName(), Controller, Child));
+				}
+			}
+		}
+	}
+}
+
+
+void SAvaRCControllerPanel::Refresh(const TArray<int32>& InSelectedPageIds)
+{
+	// Request to Rebuild on next tick
+	ControllerContainer->RebuildList();
+
+	ActivePageId = InSelectedPageIds.IsEmpty() ? FAvalanchePage::InvalidPageId : InSelectedPageIds[0];
+
+	const UAvalanchePlaylist* Playlist = GetPlaylist();
+	const FAvalanchePage& Page = GetActivePage(Playlist);
+	
+	if (!Page.IsValidPage())
+	{
+		ControllerItems.Empty();
+       	ManagedInstances.Reset();
+		return;
+	}
+		
+	ManagedInstances = FAvaPlaylistEditor::GetManagedInstancesForPage(Playlist, Page);
+
+	int32 NumItems = 0;
+	for (const TSharedPtr<FAvalancheManagedInstance>& ManagedInstance : ManagedInstances)
+	{
+		const URemoteControlPreset* Preset = ManagedInstance ? ManagedInstance->GetRemoteControlPreset() : nullptr;
+		if (Preset)
+		{
+			NumItems += Preset->GetNumControllers();
+		}
+	}
+
+	ControllerItems.Empty(NumItems);
+	
+	UpdatePropertyRowGenerators(ManagedInstances.Num());
+	UpdatePageSummary(false);
+
+	for (int32 Index = 0; Index < ManagedInstances.Num(); ++Index)
+	{
+		RefreshForManagedInstance(Index, *ManagedInstances[Index], Page);
+	}
+
+	ControllerItems.Sort([](const FAvaRCControllerItemPtr& A, const FAvaRCControllerItemPtr& B)
+	{
+		if (A->GetInstanceIndex() == B->GetInstanceIndex())
+		{
+			return A->GetDisplayIndex() < B->GetDisplayIndex();
+		}
+		return A->GetInstanceIndex() < B->GetInstanceIndex();
+	});
+}
+
+TSharedRef<ITableRow> SAvaRCControllerPanel::OnGenerateControllerRow(FAvaRCControllerItemPtr InItem
+	, const TSharedRef<STableViewBase>& InOwnerTable)
+{
+	return InItem->CreateWidget(SharedThis(this), InOwnerTable);
+}
+
+void SAvaRCControllerPanel::UpdateDefaultValuesAndRefresh()
+{
+	// Remark: The RC values might be already updated in SAvaPageRemoteControlProps.
+	// But order of callback is not guaranteed. Code could reach here first, so
+	// it needs to update and refresh just in case. Calling UpdateRemoteControlValues
+	// multiple time (from different code paths) is harmless (fast if nothing changed).
+	
+	const TSharedPtr<FAvaPlaylistEditor> PlaylistEditor = PlaylistEditorWeak.Pin();
+	if (!PlaylistEditor)
+	{
+		return;
+	}
+	
+	UAvalanchePlaylist* Playlist = PlaylistEditor->GetPlaylist();
+	if (!Playlist)
+	{
+		return;
+	}
+	
+	const FAvalanchePage& Page = GetActivePage(Playlist);
+	if (!Page.IsValidPage())
+	{
+		return;
+	}
+	
+	ManagedInstances = FAvaPlaylistEditor::GetManagedInstancesForPage(Playlist, ActivePageId);
+					
+	if (ManagedInstances.IsEmpty())
+	{
+		return;
+	}
+	
+	FAvalancheRemoteControlValues MergedDefaultRCValues;
+	FAvaPlaylistEditor::MergeDefaultRemoteControlValues(ManagedInstances, MergedDefaultRCValues);
+
+	// Using the playlist API for event propagation.
+	constexpr bool bUpdateDefaults = true;	
+	if (Playlist->UpdateRemoteControlValues(ActivePageId, MergedDefaultRCValues, bUpdateDefaults) != EAvaRemoteControlChanges::None)
+	{
+		PlaylistEditor->MarkAsModified();
+	}
+	
+	Refresh({ActivePageId});
+}
+
+void SAvaRCControllerPanel::OnRemoteControlControllerModified(URemoteControlPreset* InPreset, const TSet<FGuid>& InModifiedControllerIds)
+{
+	if (!IsValid(InPreset) || !HasRemoteControlPreset(InPreset))
+	{
+		return;
+	}
+
+	if (const TSharedPtr<FAvaPlaylistEditor> PlaylistEditor = PlaylistEditorWeak.Pin())
+	{
+		bool bModified = false;
+
+		for (const FGuid& Id : InModifiedControllerIds)
+		{
+			if (URCVirtualPropertyBase* Controller = InPreset->GetController(Id))
+			{
+				using namespace UE::AvalancheRemoteControl;
+				FAvalancheRemoteControlValue ControllerValue;
+				const EAvaRemoteControlResult Result = GetValueOfController(Controller, ControllerValue.Value);
+
+				if (Failed(Result))
+				{
+					UE_LOG(LogAvaMediaRemoteControl, Error, TEXT("Unable to get value of controller \"%s\" (id:%s): %s."),
+						   *Controller->DisplayName.ToString(), *Controller->Id.ToString(), *EnumToString(Result));
+					continue;
+				}
+
+				const FAvalancheRemoteControlValue* StoredControllerValue = GetSelectedPageControllerValue(Controller);
+
+				if (StoredControllerValue && ControllerValue.IsSameValueAs(*StoredControllerValue))
+				{
+					continue;	// Skip if value is identical.
+				}
+
+				if (!SetSelectedPageControllerValue(Controller, ControllerValue))
+				{
+					UE_LOG(LogAvaMediaRemoteControl, Error, TEXT("Unable to set page controller value for: \"%s\""), *Controller->DisplayName.ToString());
+					continue;
+				}
+
+				bModified = true;
+			}
+		}
+
+		if (bModified)
+		{
+			PlaylistEditor->MarkAsModified();
+		}
+	}
+}
+
+
+void SAvaRCControllerPanel::BindRemoteControlDelegates(URemoteControlPreset* InPreset)
+{
+	if (IsValid(InPreset))
+	{
+		if (!InPreset->OnControllerAdded().IsBoundToObject(this))
+		{
+			InPreset->OnControllerAdded().AddSP(this, &SAvaRCControllerPanel::OnRemoteControlControllerAdded);
+		}
+
+		if (!InPreset->OnControllerRemoved().IsBoundToObject(this))
+		{
+			InPreset->OnControllerRemoved().AddSP(this, &SAvaRCControllerPanel::OnRemoteControlControllerRemoved);
+		}
+
+		if (!InPreset->OnControllerRenamed().IsBoundToObject(this))
+		{
+			InPreset->OnControllerRenamed().AddSP(this, &SAvaRCControllerPanel::OnRemoteControlControllerRenamed);
+		}
+
+		if (!InPreset->OnControllerModified().IsBoundToObject(this))
+		{
+			InPreset->OnControllerModified().AddSP(this, &SAvaRCControllerPanel::OnRemoteControlControllerModified);
+		}
+	}
+}
+
+const FAvalancheRemoteControlValue* SAvaRCControllerPanel::GetSelectedPageControllerValue(const URCVirtualPropertyBase* InController) const
+{
+	const FAvalanchePage& Page = GetActivePage();
+	if (IsValid(InController) && Page.IsValidPage())
+	{
+		return Page.GetRemoteControlControllerValue(InController->Id);
+	}
+	return nullptr;
+}
+
+bool SAvaRCControllerPanel::SetSelectedPageControllerValue(const URCVirtualPropertyBase* InController, const FAvalancheRemoteControlValue& InValue) const
+{
+	UAvalanchePlaylist* Playlist = GetPlaylist();
+	if (IsValid(InController) && Playlist)
+	{
+		// Using the playlist API for event propagation.
+		return Playlist->SetRemoteControlControllerValue(ActivePageId, InController->Id, InValue);
+	}
+	return false;
+}
+
+void SAvaRCControllerPanel::UpdatePageSummary(bool bInIsPresetChanged)
+{
+	FAvalanchePage& Page = GetActivePageMutable();
+	if (Page.IsValidPage())
+	{
+		TArray<const URemoteControlPreset*> Presets;
+		Presets.Reserve(ManagedInstances.Num());
+		for (const TSharedPtr<FAvalancheManagedInstance>& ManagedInstance : ManagedInstances)
+		{
+			if (ManagedInstance && ManagedInstance->GetRemoteControlPreset())
+			{
+				Presets.Add(ManagedInstance->GetRemoteControlPreset());
+			}
+		}					
+		Page.UpdatePageSummary(Presets, bInIsPresetChanged);
+	}
+}
+
+UAvalanchePlaylist* SAvaRCControllerPanel::GetPlaylist() const
+{
+	if (const TSharedPtr<FAvaPlaylistEditor> PlaylistEditor = PlaylistEditorWeak.Pin())
+	{
+		return PlaylistEditor->GetPlaylist();
+	}
+	return nullptr;
+}
+
+const FAvalanchePage& SAvaRCControllerPanel::GetActivePage(const UAvalanchePlaylist* InPlaylist) const
+{
+	if (IsValid(InPlaylist) && ActivePageId != FAvalanchePage::InvalidPageId)
+	{
+		return InPlaylist->GetPage(ActivePageId);
+	}
+	return FAvalanchePage::NullPage;
+}
+
+FAvalanchePage& SAvaRCControllerPanel::GetActivePageMutable(UAvalanchePlaylist* InPlaylist) const
+{
+	if (IsValid(InPlaylist) && ActivePageId != FAvalanchePage::InvalidPageId)
+	{
+		return InPlaylist->GetPage(ActivePageId);
+	}
+	return FAvalanchePage::NullPage;
+}
+
+SAvaRCControllerPanel::FPropertyRowGeneratorWrapper::FPropertyRowGeneratorWrapper(SAvaRCControllerPanel* InParentPanel)
+	: ParentPanel(InParentPanel)
+{
+	FPropertyRowGeneratorArgs Args;
+	Args.bShouldShowHiddenProperties   = true;
+	Args.bAllowMultipleTopLevelObjects = false;
+	Args.NotifyHook = this;
+	PropertyRowGenerator = FModuleManager::GetModuleChecked<FPropertyEditorModule>("PropertyEditor").CreatePropertyRowGenerator(Args);
+}
+
+SAvaRCControllerPanel::FPropertyRowGeneratorWrapper::~FPropertyRowGeneratorWrapper()
+{
+	if (PropertyRowGenerator)
+	{
+		PropertyRowGenerator->OnFinishedChangingProperties().RemoveAll(this);
+	}
+}
+
+void SAvaRCControllerPanel::FPropertyRowGeneratorWrapper::NotifyPostChange(const FPropertyChangedEvent& InPropertyChangedEvent, FProperty* InPropertyThatChanged)
+{
+	FNotifyHook::NotifyPostChange(InPropertyChangedEvent, InPropertyThatChanged);
+	
+	if (URemoteControlPreset* Preset = PresetWeak.Get())
+	{
+		Preset->OnModifyController(InPropertyChangedEvent);
+		if (ParentPanel && InPropertyChangedEvent.ChangeType & EPropertyChangeType::ValueSet)
+		{
+			ParentPanel->UpdatePageSummary(true);
+		}
+	}
+}
+
+#undef LOCTEXT_NAMESPACE
diff --git a/Engine/Plugins/Experimental/Avalanche/Source/AvalancheMediaEditor/Private/Playlist/DetailsView/RemoteControl/Controllers/SAvaRCControllerPanel.h b/Engine/Plugins/Experimental/Avalanche/Source/AvalancheMediaEditor/Private/Playlist/DetailsView/RemoteControl/Controllers/SAvaRCControllerPanel.h
new file mode 100644
--- /dev/null
+++ b/Engine/Plugins/Experimental/Avalanche/Source/AvalancheMediaEditor/Private/Playlist/DetailsView/RemoteControl/Controllers/SAvaRCControllerPanel.h
@@ -0,0 +1,107 @@
+// Copyright Epic Games, Inc. All Rights Reserved.
+
+#pragma once
+
+#include "Misc/NotifyHook.h"
+#include "Templates/SharedPointer.h"
+#include "Widgets/SCompoundWidget.h"
+
+class FAvaPlaylistEditor;
+class FAvaRCControllerItem;
+class FAvalancheManagedInstance;
+class IPropertyRowGenerator;
+class ITableRow;
+class SAvaRCControllerPanel;
+class SHeaderRow;
+class STableViewBase;
+class UAvalanchePlaylist;
+class URCVirtualPropertyBase;
+class URemoteControlPreset;
+struct FAvalanchePage;
+struct FAvalancheRemoteControlValue;
+struct FPropertyChangedEvent;
+template <typename ItemType> class SListView;
+
+using FAvaRCControllerItemPtr = TSharedPtr<FAvaRCControllerItem>;
+
+DECLARE_MULTICAST_DELEGATE_TwoParams(FAvaRCControllerHeaderRowExtensionDelegate, TSharedRef<SAvaRCControllerPanel> Panel, 
+	TSharedRef<SHeaderRow>& HeaderRow)
+DECLARE_DELEGATE_ThreeParams(FAvaRCControllerTableRowExtensionDelegate, TSharedRef<SAvaRCControllerPanel> Panel, 
+	TSharedRef<const FAvaRCControllerItem> ItemPtr, TSharedPtr<SWidget>& CurrentWidget)
+
+class SAvaRCControllerPanel : public SCompoundWidget
+{
+public:
+	static const FName ControllerColumnName;
+	static const FName ValueColumnName;
+
+	static FAvaRCControllerHeaderRowExtensionDelegate& GetHeaderRowExtensionDelegate() { return HeaderRowExtensionDelegate; }
+	static TArray<FAvaRCControllerTableRowExtensionDelegate>& GetTableRowExtensionDelegates(FName InExtensionName);
+
+	SLATE_BEGIN_ARGS(SAvaRCControllerPanel) {}
+	SLATE_END_ARGS()
+	
+	void Construct(const FArguments& InArgs, const TSharedPtr<FAvaPlaylistEditor>& InPlaylistEditor);
+	
+	bool HasRemoteControlPreset(const URemoteControlPreset* InPreset) const;
+
+	void OnPageSelectionChanged(const TArray<int32>& InSelectedPageIds);
+	
+	void Refresh(const TArray<int32>& InSelectedPageIds);
+	
+	TSharedRef<ITableRow> OnGenerateControllerRow(FAvaRCControllerItemPtr InItem, const TSharedRef<STableViewBase>& InOwnerTable);
+
+private:
+	static FAvaRCControllerHeaderRowExtensionDelegate HeaderRowExtensionDelegate;
+	static TMap<FName, TArray<FAvaRCControllerTableRowExtensionDelegate>> TableRowExtensionDelegates;
+
+	void UpdatePropertyRowGenerators(int32 InNumGenerators);
+	
+	void RefreshForManagedInstance(int32 InInstanceIndex, const FAvalancheManagedInstance& InManagedInstance, const FAvalanchePage& InPage);
+	
+	/** Update the current page's remote control values from the defaults then refresh the widget. */
+	void UpdateDefaultValuesAndRefresh();
+
+	void OnRemoteControlControllerAdded(URemoteControlPreset* InPreset, const FName NewControllerName, const FGuid& InControllerId) { UpdateDefaultValuesAndRefresh(); }
+	void OnRemoteControlControllerRemoved(URemoteControlPreset* InPreset, const FGuid& InControllerId) { UpdateDefaultValuesAndRefresh(); }
+	void OnRemoteControlControllerRenamed(URemoteControlPreset* InPreset, const FName InOldLabel, const FName InNewLabel) { UpdateDefaultValuesAndRefresh(); }
+	void OnRemoteControlControllerModified(URemoteControlPreset* InPreset, const TSet<FGuid>& InModifiedControllerIds);
+	void BindRemoteControlDelegates(URemoteControlPreset* InPreset);
+
+	const FAvalancheRemoteControlValue* GetSelectedPageControllerValue(const URCVirtualPropertyBase* InController) const;
+	bool SetSelectedPageControllerValue(const URCVirtualPropertyBase* InController, const FAvalancheRemoteControlValue& InValue) const;
+	void UpdatePageSummary(bool bInIsPresetChanged);
+
+	UAvalanchePlaylist* GetPlaylist() const;
+	const FAvalanchePage& GetActivePage(const UAvalanchePlaylist* InPlaylist) const;
+	const FAvalanchePage& GetActivePage() const { return GetActivePage(GetPlaylist()); }
+	FAvalanchePage& GetActivePageMutable(UAvalanchePlaylist* InPlaylist) const;
+	FAvalanchePage& GetActivePageMutable() const { return GetActivePageMutable(GetPlaylist()); }
+
+	TWeakPtr<FAvaPlaylistEditor> PlaylistEditorWeak;
+
+	class FPropertyRowGeneratorWrapper : public FNotifyHook
+	{
+	public:
+		TSharedPtr<IPropertyRowGenerator> PropertyRowGenerator;
+		TWeakObjectPtr<URemoteControlPreset> PresetWeak;
+		SAvaRCControllerPanel* ParentPanel = nullptr;
+
+		FPropertyRowGeneratorWrapper(SAvaRCControllerPanel* InParentPanel);
+		virtual ~FPropertyRowGeneratorWrapper();
+
+		//~ Begin FNotifyHook
+		virtual void NotifyPostChange(const FPropertyChangedEvent& InPropertyChangedEvent, FProperty* InPropertyThatChanged) override;
+		//~ End FNotifyHook
+	};
+	
+	TArray<TUniquePtr<FPropertyRowGeneratorWrapper>> PropertyRowGenerators;
+
+	TArray<TSharedPtr<FAvalancheManagedInstance>> ManagedInstances;
+	
+	TSharedPtr<SListView<FAvaRCControllerItemPtr>> ControllerContainer;
+	
+	TArray<FAvaRCControllerItemPtr> ControllerItems;
+
+	int32 ActivePageId = -1;
+};
diff --git a/Engine/Plugins/Experimental/Avalanche/Source/AvalancheMediaEditor/Private/Playlist/DetailsView/RemoteControl/Properties/AvaPageRemoteControlWidgetUtils.cpp b/Engine/Plugins/Experimental/Avalanche/Source/AvalancheMediaEditor/Private/Playlist/DetailsView/RemoteControl/Properties/AvaPageRemoteControlWidgetUtils.cpp
new file mode 100644
--- /dev/null
+++ b/Engine/Plugins/Experimental/Avalanche/Source/AvalancheMediaEditor/Private/Playlist/DetailsView/RemoteControl/Properties/AvaPageRemoteControlWidgetUtils.cpp
@@ -0,0 +1,114 @@
+// Copyright Epic Games, Inc. All Rights Reserved.
+
+#include "AvaPageRemoteControlWidgetUtils.h"
+#include "IDetailTreeNode.h"
+#include "PropertyHandle.h"
+#include "SlateOptMacros.h"
+#include "Widgets/SBoxPanel.h"
+
+#define LOCTEXT_NAMESPACE "AvaPageWidgetUtils"
+
+BEGIN_SLATE_FUNCTION_BUILD_OPTIMIZATION
+
+TSharedRef<SWidget> FAvaPageWidgetUtils::CreateNodeValueWidget(const TSharedPtr<IDetailTreeNode>& InNode)
+{
+	FNodeWidgets NodeWidgets = InNode->CreateNodeWidgets();
+
+	TSharedRef<SHorizontalBox> FieldWidget = SNew(SHorizontalBox);
+
+	if (NodeWidgets.ValueWidget)
+	{
+		FieldWidget->AddSlot()
+			.Padding(FMargin(3.0f, 2.0f))
+			.HAlign(HAlign_Right)
+			.FillWidth(1.0f)
+			[
+				NodeWidgets.ValueWidget.ToSharedRef()
+			];
+	}
+	else if (NodeWidgets.WholeRowWidget)
+	{
+		FieldWidget->AddSlot()
+			.Padding(FMargin(3.0f, 2.0f))
+			.FillWidth(1.0f)
+			[
+				NodeWidgets.WholeRowWidget.ToSharedRef()
+			];
+	}
+
+	return FieldWidget;
+}
+
+bool FAvaPageWidgetUtils::FindPropertyHandleRecursive(const TSharedPtr<IPropertyHandle>& InPropertyHandle, const FString& InPropertyNameOrPath, EFindNodeMethod InFindMethod)
+{
+	if (InPropertyHandle && InPropertyHandle->IsValidHandle())
+	{
+		uint32 ChildrenCount = 0;
+		InPropertyHandle->GetNumChildren(ChildrenCount);
+		for (uint32 Index = 0; Index < ChildrenCount; ++Index)
+		{
+			TSharedPtr<IPropertyHandle> ChildHandle = InPropertyHandle->GetChildHandle(Index);
+			if (FindPropertyHandleRecursive(ChildHandle, InPropertyNameOrPath, InFindMethod))
+			{
+				return true;
+			}
+		}
+
+		if (InPropertyHandle->GetProperty())
+		{
+			if (InFindMethod == EFindNodeMethod::Path)
+			{
+				if (InPropertyHandle->GeneratePathToProperty() == InPropertyNameOrPath)
+				{
+					return true;
+				}
+			}
+			else if (InPropertyHandle->GetProperty()->GetName() == InPropertyNameOrPath)
+			{
+				return true;
+			}
+		}
+	}
+
+	return false;
+}
+
+TSharedPtr<IDetailTreeNode> FAvaPageWidgetUtils::FindTreeNodeRecursive(const TSharedRef<IDetailTreeNode>& InRootNode, const FString& InPropertyNameOrPath, EFindNodeMethod InFindMethod)
+{
+	TArray<TSharedRef<IDetailTreeNode>> Children;
+	InRootNode->GetChildren(Children);
+	for (TSharedRef<IDetailTreeNode>& Child : Children)
+	{
+		TSharedPtr<IDetailTreeNode> FoundNode = FindTreeNodeRecursive(Child, InPropertyNameOrPath, InFindMethod);
+		if (FoundNode.IsValid())
+		{
+			return FoundNode;
+		}
+	}
+
+	TSharedPtr<IPropertyHandle> Handle = InRootNode->CreatePropertyHandle();
+	if (FindPropertyHandleRecursive(Handle, InPropertyNameOrPath, InFindMethod))
+	{
+		return InRootNode;
+	}
+
+	return nullptr;
+}
+
+TSharedPtr<IDetailTreeNode> FAvaPageWidgetUtils::FindNode(const TArray<TSharedRef<IDetailTreeNode>>& InRootNodes, const FString& InQualifiedPropertyName, EFindNodeMethod InFindMethod)
+{
+	for (const TSharedRef<IDetailTreeNode>& CategoryNode : InRootNodes)
+	{
+		TSharedPtr<IDetailTreeNode> FoundNode = FindTreeNodeRecursive(CategoryNode, InQualifiedPropertyName, InFindMethod);
+		if (FoundNode.IsValid())
+		{
+			return FoundNode;
+		}
+	}
+
+	return nullptr;
+}
+
+END_SLATE_FUNCTION_BUILD_OPTIMIZATION
+
+#undef LOCTEXT_NAMESPACE
diff --git a/Engine/Plugins/Experimental/Avalanche/Source/AvalancheMediaEditor/Private/Playlist/DetailsView/RemoteControl/Properties/AvaPageRemoteControlWidgetUtils.h b/Engine/Plugins/Experimental/Avalanche/Source/AvalancheMediaEditor/Private/Playlist/DetailsView/RemoteControl/Properties/AvaPageRemoteControlWidgetUtils.h
new file mode 100644
--- /dev/null
+++ b/Engine/Plugins/Experimental/Avalanche/Source/AvalancheMediaEditor/Private/Playlist/DetailsView/RemoteControl/Properties/AvaPageRemoteControlWidgetUtils.h
@@ -0,0 +1,28 @@
+// Copyright Epic Games, Inc. All Rights Reserved.
+
+#pragma once
+
+#include "CoreMinimal.h"
+
+class IDetailTreeNode;
+class IPropertyHandle;
+class SWidget;
+
+/**
+ * A static class containing utility functions for widgets.
+ */
+class FAvaPageWidgetUtils final
+{
+public:
+	enum class EFindNodeMethod : uint8
+	{
+		Name, // Find a node by providing its property name
+		Path // Find a node by providing a property path
+	};
+
+public:
+	static TSharedRef<SWidget> CreateNodeValueWidget(const TSharedPtr<IDetailTreeNode>& InNode);
+	static bool FindPropertyHandleRecursive(const TSharedPtr<IPropertyHandle>& InPropertyHandle, const FString& InPropertyNameOrPath, EFindNodeMethod InFindMethod);
+	static TSharedPtr<IDetailTreeNode> FindTreeNodeRecursive(const TSharedRef<IDetailTreeNode>& InRootNode, const FString& InPropertyNameOrPath, EFindNodeMethod InFindMethod);
+	static TSharedPtr<IDetailTreeNode> FindNode(const TArray<TSharedRef<IDetailTreeNode>>& InRootNodes, const FString& InQualifiedPropertyName, EFindNodeMethod InFindMethod);
+};
diff --git a/Engine/Plugins/Experimental/Avalanche/Source/AvalancheMediaEditor/Private/Playlist/DetailsView/RemoteControl/Properties/AvaRCPropertyItem.cpp b/Engine/Plugins/Experimental/Avalanche/Source/AvalancheMediaEditor/Private/Playlist/DetailsView/RemoteControl/Properties/AvaRCPropertyItem.cpp
new file mode 100644
--- /dev/null
+++ b/Engine/Plugins/Experimental/Avalanche/Source/AvalancheMediaEditor/Private/Playlist/DetailsView/RemoteControl/Properties/AvaRCPropertyItem.cpp
@@ -0,0 +1,17 @@
+// Copyright Epic Games, Inc. All Rights Reserved.
+
+#include "AvaRCPropertyItem.h"
+
+#include "RemoteControlEntity.h"
+#include "SAvaRCPropertyItemRow.h"
+
+FAvaRCPropertyItem::FAvaRCPropertyItem(TSharedRef<FRemoteControlEntity> InEntity, bool bInControlled)
+{
+	Entity = InEntity;
+	bEntityControlled = bInControlled;
+}
+
+TSharedRef<ITableRow> FAvaRCPropertyItem::CreateWidget(TSharedRef<SAvaPageRemoteControlProps> InPropertyPanel, const TSharedRef<STableViewBase>& InOwnerTable) const
+{
+	return SNew(SAvaRCPropertyItemRow, InPropertyPanel, InOwnerTable, SharedThis(this));
+}
diff --git a/Engine/Plugins/Experimental/Avalanche/Source/AvalancheMediaEditor/Private/Playlist/DetailsView/RemoteControl/Properties/AvaRCPropertyItem.h b/Engine/Plugins/Experimental/Avalanche/Source/AvalancheMediaEditor/Private/Playlist/DetailsView/RemoteControl/Properties/AvaRCPropertyItem.h
new file mode 100644
--- /dev/null
+++ b/Engine/Plugins/Experimental/Avalanche/Source/AvalancheMediaEditor/Private/Playlist/DetailsView/RemoteControl/Properties/AvaRCPropertyItem.h
@@ -0,0 +1,30 @@
+// Copyright Epic Games, Inc. All Rights Reserved.
+
+#pragma once
+
+#include "IDetailTreeNode.h"
+#include "Internationalization/Text.h"
+#include "Templates/SharedPointer.h"
+
+class IDetailTreeNode;
+class ITableRow;
+class SAvaPageRemoteControlProps;
+class STableViewBase;
+class URemoteControlPreset;
+struct FAvaPageRemoteControlGenerateWidgetArgs;
+struct FRemoteControlEntity;
+
+class FAvaRCPropertyItem : public TSharedFromThis<FAvaRCPropertyItem>
+{
+public:
+	FAvaRCPropertyItem(TSharedRef<FRemoteControlEntity> InEntity, bool bInControlled);
+
+	TSharedRef<ITableRow> CreateWidget(TSharedRef<SAvaPageRemoteControlProps> InPropertyPanel, const TSharedRef<STableViewBase>& InOwnerTable) const;
+
+	TSharedPtr<FRemoteControlEntity> GetEntity() const { return Entity.Pin(); }
+	bool IsEntityControlled() const { return bEntityControlled; }
+
+private:
+	TWeakPtr<FRemoteControlEntity> Entity = nullptr;
+	bool bEntityControlled = false;
+};
diff --git a/Engine/Plugins/Experimental/Avalanche/Source/AvalancheMediaEditor/Private/Playlist/DetailsView/RemoteControl/Properties/SAvaPageRemoteControlProps.cpp b/Engine/Plugins/Experimental/Avalanche/Source/AvalancheMediaEditor/Private/Playlist/DetailsView/RemoteControl/Properties/SAvaPageRemoteControlProps.cpp
new file mode 100644
--- /dev/null
+++ b/Engine/Plugins/Experimental/Avalanche/Source/AvalancheMediaEditor/Private/Playlist/DetailsView/RemoteControl/Properties/SAvaPageRemoteControlProps.cpp
@@ -0,0 +1,470 @@
+// Copyright Epic Games, Inc. All Rights Reserved.
+
+#include "SAvaPageRemoteControlProps.h"
+
+#include "AvaRCPropertyItem.h"
+#include "IAvaMediaModule.h"
+#include "Playback/AvalancheRemoteControl.h"
+#include "Playlist/AvaPlaylistEditor.h"
+#include "Playlist/AvalancheManagedInstanceCache.h"
+#include "Playlist/AvalanchePage.h"
+#include "SAvaRCPropertyItemRow.h"
+#include "SlateOptMacros.h"
+#include "Widgets/Views/ITableRow.h"
+#include "Widgets/Views/SListView.h"
+#include "Widgets/Views/STableViewBase.h"
+
+#define LOCTEXT_NAMESPACE "AvaPageRemoteControlProps"
+
+const FName SAvaPageRemoteControlProps::PropertyColumnName = "PropertyColumn";
+const FName SAvaPageRemoteControlProps::ValueColumnName = "ValueColumn";
+
+FAvaRCPropertyHeaderRowExtensionDelegate SAvaPageRemoteControlProps::HeaderRowExtensionDelegate;
+TMap<FName, TArray<FAvaRCPropertyTableRowExtensionDelegate>> SAvaPageRemoteControlProps::TableRowExtensionDelegates;
+
+BEGIN_SLATE_FUNCTION_BUILD_OPTIMIZATION
+
+TArray<FAvaRCPropertyTableRowExtensionDelegate>& SAvaPageRemoteControlProps::GetTableRowExtensionDelegates(FName InExtensionName)
+{
+	return TableRowExtensionDelegates.FindOrAdd(InExtensionName);
+}
+
+void SAvaPageRemoteControlProps::Construct(const FArguments& InArgs, TSharedPtr<FAvaPlaylistEditor> InPlaylistEditor)
+{
+	PlaylistEditorWeak = InPlaylistEditor;
+	ActivePageId = FAvalanchePage::InvalidPageId;
+
+	TSharedRef<SHeaderRow> HeaderRow =
+		SNew(SHeaderRow)
+		+ SHeaderRow::Column(PropertyColumnName)
+		.DefaultLabel(LOCTEXT("Property", "Property"))
+		.FixedWidth(150.f)
+		+ SHeaderRow::Column(ValueColumnName)
+		.DefaultLabel(LOCTEXT("Value", "Value"))
+		.FillWidth(1.f);
+
+	HeaderRowExtensionDelegate.Broadcast(SharedThis(this), HeaderRow);
+
+	ChildSlot
+	[
+		SNew(SBorder)
+		.Padding(8.0f)
+		.BorderImage(FAppStyle::GetBrush("ToolPanel.GroupBorder"))
+		[
+			SAssignNew(PropertyContainer, SListView<FAvaRCPropertyItemPtr>)
+			.ListItemsSource(&PropertyItems)
+			.SelectionMode(ESelectionMode::None)
+			.OnGenerateRow(this, &SAvaPageRemoteControlProps::OnGenerateControllerRow)
+			.HeaderRow(HeaderRow)
+		]
+	];
+
+	Refresh({});
+}
+
+END_SLATE_FUNCTION_BUILD_OPTIMIZATION
+
+SAvaPageRemoteControlProps::~SAvaPageRemoteControlProps()
+{
+}
+
+void SAvaPageRemoteControlProps::UpdateDefaultValuesAndRefresh(const TArray<int32>& InSelectedPageIds)
+{
+	// Remark:
+	// This is used in the "reimport page" work flow. We want to be able to
+	// reimport multiple pages at the same time.
+	
+	if (const TSharedPtr<FAvaPlaylistEditor> PlaylistEditor = PlaylistEditorWeak.Pin())
+	{
+		UAvalanchePlaylist* Playlist = PlaylistEditor->GetPlaylist();
+
+		if (IsValid(Playlist))
+		{
+			EAvaRemoteControlChanges Changes = EAvaRemoteControlChanges::None;
+			for (const int32 PageId : InSelectedPageIds)
+			{
+				FAvalanchePage& Page = PlaylistEditor->GetPlaylist()->GetPage(PageId);
+
+				if (Page.IsValidPage())
+				{
+					ManagedInstances = FAvaPlaylistEditor::GetManagedInstancesForPage(Playlist, Page);
+					
+					if (!ManagedInstances.IsEmpty())
+					{
+						FAvalancheRemoteControlValues MergedDefaultRCValues;
+						FAvaPlaylistEditor::MergeDefaultRemoteControlValues(ManagedInstances, MergedDefaultRCValues);
+
+						// Using the playlist API for event propagation.
+						constexpr bool bUpdateDefaults = true;
+						Changes |= Playlist->UpdateRemoteControlValues(ActivePageId, MergedDefaultRCValues, bUpdateDefaults);
+					}
+				}
+			}
+			
+			if (Changes != EAvaRemoteControlChanges::None)
+			{
+				PlaylistEditor->MarkAsModified();
+			}
+			
+			Refresh(InSelectedPageIds);
+		}
+	}
+}
+
+TSharedRef<ITableRow> SAvaPageRemoteControlProps::OnGenerateControllerRow(FAvaRCPropertyItemPtr InItem, const TSharedRef<STableViewBase>& InOwnerTable)
+{
+	return InItem->CreateWidget(SharedThis(this), InOwnerTable);
+}
+
+
+void SAvaPageRemoteControlProps::RefreshTable(const TSet<FGuid>& InEntityIds)
+{
+	for (const FAvaRCPropertyItemPtr& PropertyItem : PropertyItems)
+	{
+		const TSharedPtr<FRemoteControlEntity> Entity = PropertyItem->GetEntity();
+		if (Entity && (InEntityIds.IsEmpty() || InEntityIds.Contains(Entity->GetId())))
+		{
+			TSharedPtr<ITableRow> TableRow = PropertyContainer->WidgetFromItem(PropertyItem);
+
+			if (TableRow.IsValid())
+			{
+				const TSharedRef<SAvaRCPropertyItemRow> ItemRow = StaticCastSharedRef<SAvaRCPropertyItemRow>(TableRow.ToSharedRef());
+				ItemRow->UpdateValue();
+			}
+		}
+	}
+}
+
+void SAvaPageRemoteControlProps::Refresh(const TArray<int32>& InSelectedPageIds)
+{
+	if (!PropertyContainer.IsValid())
+	{
+		return;
+	}
+
+	ActivePageId = InSelectedPageIds.IsEmpty() ? FAvalanchePage::InvalidPageId : InSelectedPageIds[0];
+	
+	if (const TSharedPtr<FAvaPlaylistEditor> PlaylistEditor = PlaylistEditorWeak.Pin())
+	{
+		const UAvalanchePlaylist* Playlist = PlaylistEditor->GetPlaylist();
+
+		if (IsValid(Playlist))
+		{
+			ManagedInstances.Reset();
+
+			if (FAvalanchePage* ActivePage = GetActivePage())
+			{
+				ManagedInstances = FAvaPlaylistEditor::GetManagedInstancesForPage(Playlist, *ActivePage);
+
+				for (const TSharedPtr<FAvalancheManagedInstance>& ManagedInstance : ManagedInstances)
+				{
+					URemoteControlPreset* Preset = ManagedInstance ? ManagedInstance->GetRemoteControlPreset() : nullptr;
+					if (Preset)
+					{
+						BindRemoteControlDelegates(Preset);	
+					}
+				}
+
+				if (!ManagedInstances.IsEmpty())
+				{
+					FAvalancheRemoteControlValues MergedDefaultRCValues;
+					FAvaPlaylistEditor::MergeDefaultRemoteControlValues(ManagedInstances, MergedDefaultRCValues);
+					
+					// Prune any extra stale values. This happens if templates are changed.
+					if (ActivePage->PruneRemoteControlValues(MergedDefaultRCValues) != EAvaRemoteControlChanges::None)
+					{
+						UE_LOG(LogAvaMediaRemoteControl, Log, TEXT("Page %d had stale values that where pruned."), ActivePage->GetPageId());
+						PlaylistEditor->MarkAsModified();
+					}
+				}
+			}
+		}
+
+		struct FAvaPropertyDetails
+		{
+			TSharedRef<FRemoteControlEntity> Entity;
+			bool bEntityControlled;
+		};
+
+		TArray<FAvaPropertyDetails> NewItems;
+
+		for (const TSharedPtr<FAvalancheManagedInstance>& ManagedInstance : ManagedInstances)
+		{
+			URemoteControlPreset* RemoteControlPreset = ManagedInstance ? ManagedInstance->GetRemoteControlPreset() : nullptr;
+			if (!RemoteControlPreset)
+			{
+				continue;
+			}
+			
+			const TArray<TWeakPtr<FRemoteControlEntity>> ExposedEntities = RemoteControlPreset->GetExposedEntities<FRemoteControlEntity>();
+			NewItems.Reserve(ExposedEntities.Num());
+			
+			for (const TWeakPtr<FRemoteControlEntity>& EntityWeakPtr : ExposedEntities)
+			{
+				if (const TSharedPtr<FRemoteControlEntity> Entity = EntityWeakPtr.Pin())
+				{
+					const FAvalancheRemoteControlValue* EntityValue = GetSelectedPageEntityValue(Entity);
+
+					if (!EntityValue)
+					{
+						// If the page doesn't already have a value, we get it from the template's default values.
+						const FAvalancheRemoteControlValue* const DefaultEntityValue = ManagedInstance->GetDefaultRemoteControlValues().GetEntityValue(Entity->GetId()); 
+						
+						if (!DefaultEntityValue)
+						{
+							UE_LOG(LogAvaMediaRemoteControl, Error, TEXT("Entity \"%s\" (id:%s) doesn't have a template default value."),
+								*Entity->GetLabel().ToString(), *Entity->GetId().ToString());
+							// TODO: UX improvement: instead of skipping, could add empty element, with error mark (and error message in tooltip).
+							continue;
+						}
+
+						// Ensure the default values have the default flag.
+						ensure(DefaultEntityValue->bIsDefault == true);
+						
+						// WYSIWYG (Solution):
+						// Capture the default value (flagged as default) in the current page to ensure all values will be applied to runtime RCP.
+						if (!SetSelectedPageEntityValue(Entity, *DefaultEntityValue))
+						{
+							UE_LOG(LogAvaMediaRemoteControl, Error, TEXT("Entity \"%s\" (id:%s): failed to set value in currently selected page."),
+								*Entity->GetLabel().ToString(), *Entity->GetId().ToString());
+						}
+
+						EntityValue = DefaultEntityValue;
+					}
+					
+					// Update Exposed entity value with value from page.
+					using namespace UE::AvalancheRemoteControl;
+					if (const EAvaRemoteControlResult Result = SetValueOfEntity(Entity, EntityValue->Value); Failed(Result))
+					{
+						UE_LOG(LogAvaMediaRemoteControl, Error, TEXT("Entity \"%s\" (id:%s): failed to set entity value: %s."),
+							*Entity->GetLabel().ToString(), *Entity->GetId().ToString(), *EnumToString(Result));
+						// TODO: UX improvement: instead of skipping, could add empty element, with error mark (and error message in tooltip).
+						continue;
+					}
+
+					{
+						const bool bEntityControlled = ManagedInstance->GetDefaultRemoteControlValues().EntitiesControlledByController.Contains(Entity->GetId());
+						NewItems.Add({Entity.ToSharedRef(), bEntityControlled});
+					}
+				}
+			}
+		}
+
+		PropertyContainer->RebuildList();
+
+		bool bRecreateList = NewItems.Num() != PropertyItems.Num();
+
+		if (!bRecreateList)
+		{
+			for (int32 PropertyIdx = 0; PropertyIdx < NewItems.Num(); ++PropertyIdx)
+			{
+				const FAvaRCPropertyItemPtr& PropertyItem = PropertyItems[PropertyIdx];
+				const FAvaPropertyDetails& NewItem = NewItems[PropertyIdx];
+
+				if (!PropertyItem.IsValid())
+				{
+					bRecreateList = true;
+					break;
+				}
+
+				TSharedPtr<FRemoteControlEntity> PropertyEntity = PropertyItem->GetEntity();
+
+				if (!PropertyEntity.IsValid())
+				{
+					bRecreateList = true;
+					break;
+				}
+
+				if (PropertyEntity.Get() != NewItem.Entity.ToSharedPtr().Get())
+				{
+					bRecreateList = true;
+					break;
+				}
+			}
+		}
+
+		if (!bRecreateList)
+		{
+			RefreshTable();
+			return;
+		}
+
+		PropertyItems.Empty();
+
+		for (const FAvaPropertyDetails& NewItem : NewItems)
+		{
+			PropertyItems.Add(MakeShared<FAvaRCPropertyItem>(NewItem.Entity, NewItem.bEntityControlled));
+		}
+	}
+}
+
+// Remarks:
+// This is called by URemoteControlPreset::OnEndFrame() as a result of an entity being modified.
+// However, it doesn't seem to be called (or not always) if the entity is modified by a controller action. 
+void SAvaPageRemoteControlProps::OnRemoteControlExposedPropertiesModified(URemoteControlPreset* InPreset, const TSet<FGuid>& InModifiedProperties)
+{
+	if (!IsValid(InPreset) ||!HasRemoteControlPreset(InPreset))
+	{
+		return;
+	}
+
+	if (const TSharedPtr<FAvaPlaylistEditor> PlaylistEditor = PlaylistEditorWeak.Pin())
+	{
+		bool bModified = false;
+
+		for (const FGuid& Id : InModifiedProperties)
+		{
+			if (const TSharedPtr<FRemoteControlEntity> Entity = InPreset->GetExposedEntity<FRemoteControlEntity>(Id).Pin())
+			{
+				FAvalancheRemoteControlValue EntityValue;
+
+				{
+					using namespace UE::AvalancheRemoteControl;
+					if (const EAvaRemoteControlResult Result = GetValueOfEntity(Entity, EntityValue.Value); Failed(Result))
+					{
+						UE_LOG(LogAvaMediaRemoteControl, Error, TEXT("Unable to get value of entity \"%s\": %s"),
+							*Entity->GetLabel().ToString(), *EnumToString(Result));
+						continue;
+					}
+				}
+				
+				const FAvalancheRemoteControlValue* StoredEntityValue = GetSelectedPageEntityValue(Entity);
+
+				if (StoredEntityValue && StoredEntityValue->IsSameValueAs(EntityValue))
+				{
+					continue;	// Skip if value is identical.
+				}
+
+				if (!SetSelectedPageEntityValue(Entity, EntityValue))
+				{
+					UE_LOG(LogAvaMediaRemoteControl, Error, TEXT("Unable to set page entity value for: \"%s\""), *Entity->GetLabel().ToString());
+					continue;
+				}
+
+				bModified = true;
+			}
+		}
+
+		if (bModified)
+		{
+			PlaylistEditor->MarkAsModified();
+		}
+	}
+}
+
+void SAvaPageRemoteControlProps::OnRemoteControlControllerModified(URemoteControlPreset* InPreset, const TSet<FGuid>& InModifiedControllerIds)
+{
+	if (!IsValid(InPreset) || !HasRemoteControlPreset(InPreset))
+	{
+		return;
+	}
+
+	TSet<FGuid> EntityIds;
+	for (const FGuid& ControllerId : InModifiedControllerIds)
+	{
+		UE::AvalancheRemoteControl::GetEntitiesControlledByController(InPreset, InPreset->GetController(ControllerId), EntityIds);
+	}
+
+	// If a controller changed, we need to propagate the refresh of the field's widgets.
+	// Optimization: only refresh the widgets that are related to the modified controllers.
+	RefreshTable(EntityIds);
+
+	// It seems OnPropertyChangedDelegate (OnExposedPropertiesModified()) is not called when properties are
+	// changed by controllers. Ensure the values are saved by calling our handler directly.
+	OnRemoteControlExposedPropertiesModified(InPreset, EntityIds);
+}
+
+void SAvaPageRemoteControlProps::BindRemoteControlDelegates(URemoteControlPreset* InPreset)
+{
+	if (IsValid(InPreset))
+	{
+		if (!InPreset->OnEntityExposed().IsBoundToObject(this))
+		{
+			InPreset->OnEntityExposed().AddSP(this, &SAvaPageRemoteControlProps::OnRemoteControlEntitiesExposed);
+		}
+
+		if (!InPreset->OnEntityUnexposed().IsBoundToObject(this))
+		{
+			InPreset->OnEntityUnexposed().AddSP(this, &SAvaPageRemoteControlProps::OnRemoteControlEntitiesUnexposed);
+		}
+
+		if (!InPreset->OnEntitiesUpdated().IsBoundToObject(this))
+		{
+			InPreset->OnEntitiesUpdated().AddSP(this, &SAvaPageRemoteControlProps::OnRemoteControlEntitiesUpdated);
+		}
+
+		if (!InPreset->OnExposedPropertiesModified().IsBoundToObject(this))
+		{
+			InPreset->OnExposedPropertiesModified().AddSP(this, &SAvaPageRemoteControlProps::OnRemoteControlExposedPropertiesModified);
+		}
+
+		if (!InPreset->OnControllerModified().IsBoundToObject(this))
+		{
+			InPreset->OnControllerModified().AddSP(this, &SAvaPageRemoteControlProps::OnRemoteControlControllerModified);
+		}
+	}
+}
+
+bool SAvaPageRemoteControlProps::HasRemoteControlPreset(const URemoteControlPreset* InPreset) const
+{
+	for (const TSharedPtr<FAvalancheManagedInstance>& ManagedInstance : ManagedInstances)
+	{
+		if (ManagedInstance && ManagedInstance->GetRemoteControlPreset() == InPreset)
+		{
+			return true;
+		}
+	}
+	return false;
+}
+
+FAvalanchePage* SAvaPageRemoteControlProps::GetActivePage() const
+{
+	if (ActivePageId == FAvalanchePage::InvalidPageId)
+	{
+		return nullptr;
+	}
+
+	if (const TSharedPtr<FAvaPlaylistEditor> PlaylistEditor = PlaylistEditorWeak.Pin())
+	{
+		if (PlaylistEditor->IsPlaylistValid())
+		{
+			FAvalanchePage& Page = PlaylistEditor->GetPlaylist()->GetPage(ActivePageId);
+			return Page.IsValidPage() ? &Page : nullptr;
+		}
+	}
+
+	return nullptr;
+}
+
+const FAvalancheRemoteControlValue* SAvaPageRemoteControlProps::GetSelectedPageEntityValue(const TSharedPtr<FRemoteControlEntity>& InRemoteControlEntity) const
+{
+	if (InRemoteControlEntity.IsValid())
+	{
+		if (const FAvalanchePage* Page = GetActivePage())
+		{
+			return Page->GetRemoteControlEntityValue(InRemoteControlEntity->GetId()); 
+		}
+	}
+
+	return nullptr;
+}
+
+bool SAvaPageRemoteControlProps::SetSelectedPageEntityValue(const TSharedPtr<FRemoteControlEntity>& InRemoteControlEntity, const FAvalancheRemoteControlValue& InValue) const
+{
+	if (InRemoteControlEntity.IsValid())
+	{
+		if (const TSharedPtr<FAvaPlaylistEditor> PlaylistEditor = PlaylistEditorWeak.Pin())
+		{
+			if (PlaylistEditor->IsPlaylistValid())
+			{
+				// Using the playlist API for event propagation.
+				UAvalanchePlaylist* Playlist = PlaylistEditor->GetPlaylist(); 
+				return Playlist->SetRemoteControlEntityValue(ActivePageId, InRemoteControlEntity->GetId(), InValue);
+			}
+		}
+	}
+
+	return false;
+}
+
+#undef LOCTEXT_NAMESPACE
diff --git a/Engine/Plugins/Experimental/Avalanche/Source/AvalancheMediaEditor/Private/Playlist/DetailsView/RemoteControl/Properties/SAvaPageRemoteControlProps.h b/Engine/Plugins/Experimental/Avalanche/Source/AvalancheMediaEditor/Private/Playlist/DetailsView/RemoteControl/Properties/SAvaPageRemoteControlProps.h
new file mode 100644
--- /dev/null
+++ b/Engine/Plugins/Experimental/Avalanche/Source/AvalancheMediaEditor/Private/Playlist/DetailsView/RemoteControl/Properties/SAvaPageRemoteControlProps.h
@@ -0,0 +1,96 @@
+// Copyright Epic Games, Inc. All Rights Reserved.
+
+#pragma once
+
+#include "CoreMinimal.h"
+#include "RemoteControlPreset.h"
+#include "Widgets/SCompoundWidget.h"
+#include "Widgets/Views/SListView.h"
+#include "Widgets/Views/ITableRow.h"
+
+class FAvaPlaylistEditor;
+class FAvaRCPropertyItem;
+class FAvalancheManagedInstance;
+class SAvaPageRemoteControlProps;
+struct FAvalanchePage;
+struct FAvalancheRemoteControlValue;
+
+using FAvaRCPropertyItemPtr = TSharedPtr<FAvaRCPropertyItem>;
+
+DECLARE_MULTICAST_DELEGATE_TwoParams(FAvaRCPropertyHeaderRowExtensionDelegate, TSharedRef<SAvaPageRemoteControlProps> Panel,
+	TSharedRef<SHeaderRow>& HeaderRow)
+DECLARE_DELEGATE_ThreeParams(FAvaRCPropertyTableRowExtensionDelegate, TSharedRef<SAvaPageRemoteControlProps> Panel,
+	TSharedRef<const FAvaRCPropertyItem> ItemPtr, TSharedPtr<SWidget>& CurrentWidget)
+
+/**
+ * The page props implementation for remote control fields.
+ */
+class SAvaPageRemoteControlProps : public SCompoundWidget
+{
+public:
+	static const FName PropertyColumnName;
+	static const FName ValueColumnName;
+
+	static FAvaRCPropertyHeaderRowExtensionDelegate& GetHeaderRowExtensionDelegate() { return HeaderRowExtensionDelegate; }
+	static TArray<FAvaRCPropertyTableRowExtensionDelegate>& GetTableRowExtensionDelegates(FName InExtensionName);
+
+	SLATE_BEGIN_ARGS(SAvaPageRemoteControlProps) {}
+	SLATE_END_ARGS()
+
+	void Construct(const FArguments& InArgs, TSharedPtr<FAvaPlaylistEditor> InPlaylistEditor);
+	virtual ~SAvaPageRemoteControlProps() override;
+
+	/** Update the current page's remote control values from the defaults then refresh the widget. */
+	void UpdateDefaultValuesAndRefresh(const TArray<int32>& InSelectedPageIds);
+
+	/** Refreshes the content of this widget. */
+	void Refresh(const TArray<int32>& InSelectedPageIds);
+
+private:
+	static FAvaRCPropertyHeaderRowExtensionDelegate HeaderRowExtensionDelegate;
+	static TMap<FName, TArray<FAvaRCPropertyTableRowExtensionDelegate>> TableRowExtensionDelegates;
+
+	void OnRemoteControlEntitiesExposed(URemoteControlPreset* InPreset, const FGuid& InEntityId) { UpdateDefaultValuesAndRefresh({GetActivePageId()}); }
+	void OnRemoteControlEntitiesUnexposed(URemoteControlPreset* InPreset, const FGuid& InEntityId) { UpdateDefaultValuesAndRefresh({GetActivePageId()}); }
+	void OnRemoteControlEntitiesUpdated(URemoteControlPreset* InPreset, const TSet<FGuid>& InModifiedEntities) { UpdateDefaultValuesAndRefresh({GetActivePageId()}); }
+	void OnRemoteControlExposedPropertiesModified(URemoteControlPreset* InPreset, const TSet<FGuid>& InModifiedProperties);
+	void OnRemoteControlControllerModified(URemoteControlPreset* InPreset, const TSet<FGuid>& InModifiedControllerIds);
+
+	void BindRemoteControlDelegates(URemoteControlPreset* InPreset);
+
+	bool HasRemoteControlPreset(const URemoteControlPreset* InPreset) const;
+
+	/** Returns the currently selected page if 1 page is currently selected, returns nullptr otherwise. */
+	FAvalanchePage* GetActivePage() const;
+
+	/** Returns the currently selected page id if 1 page is currently selected, returns InvalidPageId otherwise. */
+	int32 GetActivePageId() const { return ActivePageId; }
+
+	/**
+	 * Get Selected Page's entity value corresponding to the given entity (using entity Id to match).
+	 * @return pointer to page's entity value, null if not found.
+	 */
+	const FAvalancheRemoteControlValue* GetSelectedPageEntityValue(const TSharedPtr<FRemoteControlEntity>& InRemoteControlEntity) const;
+
+	/**
+	 * Set (or add) Selected Page's entity value corresponding to the given entity (using entity Id to match).
+	 * @return true if it succeeded, false otherwise.
+	 */
+	bool SetSelectedPageEntityValue(const TSharedPtr<FRemoteControlEntity>& InRemoteControlEntity, const FAvalancheRemoteControlValue& InValue) const;
+
+	TSharedRef<ITableRow> OnGenerateControllerRow(FAvaRCPropertyItemPtr InItem, const TSharedRef<STableViewBase>& InOwnerTable);
+
+	void RefreshTable(const TSet<FGuid>& InEntityIds = TSet<FGuid>());
+
+	TWeakPtr<FAvaPlaylistEditor> PlaylistEditorWeak;
+	
+	TArray<TSharedPtr<FAvalancheManagedInstance>> ManagedInstances;
+
+	/** The widget that lists the property rows. */
+	TSharedPtr<SListView<FAvaRCPropertyItemPtr>> PropertyContainer;
+	
+	/** The data used to back the properties container list view. */
+	TArray<FAvaRCPropertyItemPtr> PropertyItems;
+
+	int32 ActivePageId = -1;
+};
diff --git a/Engine/Plugins/Experimental/Avalanche/Source/AvalancheMediaEditor/Private/Playlist/DetailsView/RemoteControl/Properties/SAvaRCPropertyItemRow.cpp b/Engine/Plugins/Experimental/Avalanche/Source/AvalancheMediaEditor/Private/Playlist/DetailsView/RemoteControl/Properties/SAvaRCPropertyItemRow.cpp
new file mode 100644
--- /dev/null
+++ b/Engine/Plugins/Experimental/Avalanche/Source/AvalancheMediaEditor/Private/Playlist/DetailsView/RemoteControl/Properties/SAvaRCPropertyItemRow.cpp
@@ -0,0 +1,150 @@
+// Copyright Epic Games, Inc. All Rights Reserved.
+
+#include "SAvaRCPropertyItemRow.h"
+#include "AvaPageRemoteControlWidgetUtils.h"
+#include "AvaRCPropertyItem.h"
+#include "IPropertyRowGenerator.h"
+#include "Modules/ModuleManager.h"
+#include "PropertyEditorModule.h"
+#include "Widgets/Layout/SScissorRectBox.h"
+#include "Widgets/SNullWidget.h"
+#include "Widgets/Text/STextBlock.h"
+
+#define LOCTEXT_NAMESPACE "SAvaRCPropertyItemRow"
+
+void SAvaRCPropertyItemRow::Construct(const FArguments& InArgs, TSharedRef<SAvaPageRemoteControlProps> InPropertyPanel,
+	const TSharedRef<STableViewBase>& InOwnerTableView, const TSharedPtr<const FAvaRCPropertyItem>& InRowItem)
+{
+	ItemPtrWeak = InRowItem;
+	PropertyPanelWeak = InPropertyPanel;
+	Generator = nullptr;
+	ValueContainer = nullptr;
+	ValueWidget = nullptr;
+
+	SMultiColumnTableRow<FAvaRCPropertyItemPtr>::Construct(FSuperRowType::FArguments(), InOwnerTableView);
+}
+
+TSharedRef<SWidget> SAvaRCPropertyItemRow::GenerateWidgetForColumn(const FName& InColumnName)
+{
+	TSharedPtr<const FAvaRCPropertyItem> ItemPtr = ItemPtrWeak.Pin();
+
+	if (ItemPtr.IsValid())
+	{
+		if (InColumnName == SAvaPageRemoteControlProps::PropertyColumnName)
+		{
+			return SNew(SScissorRectBox)
+				[
+					SNew(SBox)
+					.VAlign(EVerticalAlignment::VAlign_Center)
+					.Padding(3.f, 2.f, 3.f, 2.f)
+					[
+						SNew(STextBlock)
+						.Text(GetFieldLabel())
+					]
+				];
+		}
+		else if (InColumnName == SAvaPageRemoteControlProps::ValueColumnName)
+		{
+			return SAssignNew(ValueContainer, SBox)
+				[
+					CreateValue()
+				];
+		}
+		else
+		{
+			TSharedPtr<SAvaPageRemoteControlProps> PropertyPanel = PropertyPanelWeak.Pin();
+
+			if (PropertyPanel.IsValid())
+			{
+				TSharedPtr<SWidget> Cell = nullptr;
+				const TArray<FAvaRCPropertyTableRowExtensionDelegate>& TableRowExtensionDelegates = PropertyPanel->GetTableRowExtensionDelegates(InColumnName);
+
+				for (const FAvaRCPropertyTableRowExtensionDelegate& TableRowExtensionDelegate : TableRowExtensionDelegates)
+				{
+					TableRowExtensionDelegate.ExecuteIfBound(PropertyPanel.ToSharedRef(), ItemPtr.ToSharedRef(), Cell);
+				}
+
+				if (Cell.IsValid())
+				{
+					return Cell.ToSharedRef();
+				}
+			}
+		}
+	}
+
+	return SNullWidget::NullWidget;
+}
+
+void SAvaRCPropertyItemRow::UpdateValue()
+{
+	if (ValueContainer.IsValid())
+	{
+		ValueContainer->SetContent(CreateValue());
+	}
+}
+
+FText SAvaRCPropertyItemRow::GetFieldLabel() const
+{
+	if (TSharedPtr<const FAvaRCPropertyItem> ItemPtr = ItemPtrWeak.Pin())
+	{
+		if (TSharedPtr<FRemoteControlEntity> EntityPtr = ItemPtr->GetEntity())
+		{
+			TSharedPtr<FRemoteControlField> FieldPtr = StaticCastSharedPtr<FRemoteControlField>(EntityPtr);
+			return FText::FromName(FieldPtr->FieldName);
+		}
+	}
+
+	return FText::GetEmpty();
+}
+
+TSharedRef<SWidget> SAvaRCPropertyItemRow::CreateValue()
+{
+	if (TSharedPtr<const FAvaRCPropertyItem> ItemPtr = ItemPtrWeak.Pin())
+	{
+		if (TSharedPtr<FRemoteControlEntity> EntityPtr = ItemPtr->GetEntity())
+		{
+			TSharedPtr<FRemoteControlField> FieldPtr = StaticCastSharedPtr<FRemoteControlField>(EntityPtr);
+
+			// For the moment, just use the first object.
+			TArray<UObject*> Objects = FieldPtr->GetBoundObjects();
+
+			if ((FieldPtr->FieldType == EExposedFieldType::Property) && (Objects.Num() > 0))
+			{
+				FPropertyRowGeneratorArgs Args;
+				Generator = FModuleManager::GetModuleChecked<FPropertyEditorModule>("PropertyEditor").CreatePropertyRowGenerator(Args);
+				Generator->SetObjects({Objects[0]});
+
+				if (TSharedPtr<IDetailTreeNode> Node = FAvaPageWidgetUtils::FindNode(Generator->GetRootTreeNodes(), FieldPtr->FieldPathInfo.ToPathPropertyString(), FAvaPageWidgetUtils::EFindNodeMethod::Path))
+				{
+					const FNodeWidgets NodeWidgets = Node->CreateNodeWidgets();
+					ValueWidget = NodeWidgets.WholeRowWidget.IsValid() ? NodeWidgets.WholeRowWidget : NodeWidgets.ValueWidget;
+
+					if (ItemPtr->IsEntityControlled())
+					{
+						ValueWidget = SNew(SHorizontalBox)
+							+ SHorizontalBox::Slot()
+							.VAlign(EVerticalAlignment::VAlign_Center)
+							[
+								ValueWidget.ToSharedRef()
+							]
+						+ SHorizontalBox::Slot()
+							.VAlign(EVerticalAlignment::VAlign_Center)
+							.Padding(3.f, 0.f, 0.f, 0.f)
+							[
+								SNew(STextBlock)
+								.Text(LOCTEXT("Controlled", "(Controlled)"))
+							];
+
+						ValueWidget->SetEnabled(false);
+					}
+
+					return ValueWidget.ToSharedRef();
+				}
+			}
+		}
+	}
+
+	return SNullWidget::NullWidget;
+}
+
+#undef LOCTEXT_NAMESPACE
diff --git a/Engine/Plugins/Experimental/Avalanche/Source/AvalancheMediaEditor/Private/Playlist/DetailsView/RemoteControl/Properties/SAvaRCPropertyItemRow.h b/Engine/Plugins/Experimental/Avalanche/Source/AvalancheMediaEditor/Private/Playlist/DetailsView/RemoteControl/Properties/SAvaRCPropertyItemRow.h
new file mode 100644
--- /dev/null
+++ b/Engine/Plugins/Experimental/Avalanche/Source/AvalancheMediaEditor/Private/Playlist/DetailsView/RemoteControl/Properties/SAvaRCPropertyItemRow.h
@@ -0,0 +1,36 @@
+// Copyright Epic Games, Inc. All Rights Reserved.
+
+#pragma once
+
+#include "IDetailTreeNode.h"
+#include "SAvaPageRemoteControlProps.h"
+#include "Widgets/Views/STableRow.h"
+
+class SBox;
+class SWidget;
+
+class SAvaRCPropertyItemRow : public SMultiColumnTableRow<FAvaRCPropertyItemPtr>
+{
+public:
+	SLATE_BEGIN_ARGS(SAvaRCPropertyItemRow) { }
+	SLATE_END_ARGS()
+
+	void Construct(const FArguments& InArgs, TSharedRef<SAvaPageRemoteControlProps> InPropertyPanel,
+		const TSharedRef<STableViewBase>& InOwnerTableView, const TSharedPtr<const FAvaRCPropertyItem>& InRowItem);
+
+	virtual TSharedRef<SWidget> GenerateWidgetForColumn(const FName& InColumnName) override;
+
+	void UpdateValue();
+
+protected:
+	TWeakPtr<const FAvaRCPropertyItem> ItemPtrWeak;
+	TWeakPtr<SAvaPageRemoteControlProps> PropertyPanelWeak;
+	TSharedPtr<IPropertyRowGenerator> Generator;
+	TSharedPtr<SBox> ValueContainer;
+	TSharedPtr<SWidget> ValueWidget;
+
+	/** Get this field's label. */
+	FText GetFieldLabel() const;
+
+	TSharedRef<SWidget> CreateValue();
+};
diff --git a/Engine/Plugins/Experimental/Avalanche/Source/AvalancheMediaEditor/Private/Playlist/DetailsView/SAvaPageDetails.cpp b/Engine/Plugins/Experimental/Avalanche/Source/AvalancheMediaEditor/Private/Playlist/DetailsView/SAvaPageDetails.cpp
new file mode 100644
--- /dev/null
+++ b/Engine/Plugins/Experimental/Avalanche/Source/AvalancheMediaEditor/Private/Playlist/DetailsView/SAvaPageDetails.cpp
@@ -0,0 +1,441 @@
+// Copyright Epic Games, Inc. All Rights Reserved.
+
+#include "SAvaPageDetails.h"
+
+#include "Async/Async.h"
+#include "AvaMediaEditorStyle.h"
+#include "IAvaMediaModule.h"
+#include "Playlist/AvaPlaylistEditor.h"
+#include "Playlist/AvalancheManagedInstanceCache.h"
+#include "Playlist/AvalanchePage.h"
+#include "Playlist/AvalanchePlaylist.h"
+#include "Playlist/DetailsView/RemoteControl/Properties/SAvaPageRemoteControlProps.h"
+#include "Playlist/Pages/Slate/SAvaInstancedPageList.h"
+#include "Playlist/Pages/Slate/SAvaPageList.h"
+#include "RemoteControl/Controllers/SAvaRCControllerPanel.h"
+#include "Widgets/Images/SImage.h"
+#include "Widgets/Input/SButton.h"
+#include "Widgets/Input/SEditableTextBox.h"
+#include "Widgets/Layout/SExpandableArea.h"
+#include "Widgets/Layout/SScrollBox.h"
+#include "Widgets/Text/STextBlock.h"
+
+#define LOCTEXT_NAMESPACE "SAvaPageDetails"
+
+void SAvaPageDetails::Construct(const FArguments& InArgs, const TSharedPtr<FAvaPlaylistEditor>& InPlaylistEditor)
+{
+	PlaylistEditorWeak = InPlaylistEditor;
+	ActivePageId = FAvalanchePage::InvalidPageId;
+
+	InPlaylistEditor->GetOnPageEvent().AddSP(this, &SAvaPageDetails::OnPageEvent);
+	IAvaMediaModule::Get().GetManagedInstanceCache().OnEntryInvalidated.AddSP(this, &SAvaPageDetails::OnManagedInstanceCacheEntryInvalidated);
+
+	TSharedRef<SHorizontalBox> AnimationHeader = SNew(SHorizontalBox);
+	{
+		static const TArray<FText> HeaderTitles
+    	{
+    		LOCTEXT("HeaderTitleAnimationName", "Animation Name"),
+    		LOCTEXT("HeaderTitleLoopsToPlay", "Loops to Play"),
+    		LOCTEXT("HeaderTitlePlaybackSpeed", "Playback Speed"),
+    		LOCTEXT("HeaderTitlePlayMode", "Play Mode")
+    	};
+
+    	for (const FText& Title : HeaderTitles)
+    	{
+    		AnimationHeader->AddSlot()
+    			.FillWidth(1.f)
+    			[
+    				SNew(STextBlock)
+    				.Text(Title)
+    				.Justification(ETextJustify::Center)
+    			];
+    	}
+	}
+
+	ChildSlot
+	[
+		SNew(SScrollBox)
+		+ SScrollBox::Slot()
+		[
+			SNew(SVerticalBox)
+			+ SVerticalBox::Slot()
+			.AutoHeight()
+			.Padding(FMargin(10.f, 10.f, 10.f, 0.f))
+			[
+				SNew(SHorizontalBox)
+				+ SHorizontalBox::Slot()
+				.VAlign(EVerticalAlignment::VAlign_Center)
+				.MaxWidth(75.f)
+				[
+					SNew(STextBlock)
+					.Text(LOCTEXT("PageId", "Page Id"))
+					.MinDesiredWidth(75.f)
+				]
+				+ SHorizontalBox::Slot()
+				.VAlign(EVerticalAlignment::VAlign_Center)
+				.Padding(5.f, 0.f, 0.f, 0.f)
+				.MaxWidth(70.f)
+				[
+					SNew(SEditableTextBox)
+					.HintText(LOCTEXT("PageIdHint", "Page Id"))
+					.OnTextCommitted(this, &SAvaPageDetails::OnPageIdCommitted)
+					.Text(this, &SAvaPageDetails::GetPageId)
+					.IsEnabled(this, &SAvaPageDetails::HasSelectedPage)
+				]
+				+ SHorizontalBox::Slot()
+				.VAlign(EVerticalAlignment::VAlign_Center)
+				.Padding(5.f, 0.f, 0.f, 0.f)
+				.AutoWidth()
+				[
+					SNew(SButton)
+					.ButtonStyle(FAppStyle::Get(), "SimpleButton")
+					.ToolTipText(LOCTEXT("DuplicatePageTooltip", "DuplicatePage"))
+					.OnClicked(this, &SAvaPageDetails::DuplicateSelectedPage)
+					.IsEnabled(this, &SAvaPageDetails::HasSelectedPage)
+					[
+						SNew(SImage)
+						.Image(FAppStyle::Get().GetBrush("GenericCommands.Duplicate"))
+					]
+				]
+			]
+			+ SVerticalBox::Slot()
+			.AutoHeight()
+			.Padding(FMargin(10.f, 3.f, 10.f, 0.f))
+			[
+				SNew(SHorizontalBox)
+				+ SHorizontalBox::Slot()
+				.VAlign(EVerticalAlignment::VAlign_Center)
+				.MaxWidth(75.f)
+				[
+					SNew(STextBlock)
+					.Text(LOCTEXT("PageName", "Page Name"))
+					.MinDesiredWidth(75.f)
+				]
+				+ SHorizontalBox::Slot()
+				.VAlign(EVerticalAlignment::VAlign_Center)
+				.Padding(5.f, 0.f, 0.f, 0.f)
+				[
+					SNew(SEditableTextBox)
+					.HintText(LOCTEXT("PageNameHint", "Page Name"))
+					.OnTextChanged(this, &SAvaPageDetails::OnPageNameChanged)
+					.Text(this, &SAvaPageDetails::GetPageDescription)
+					.IsEnabled(this, &SAvaPageDetails::HasSelectedPage)
+				]
+			]
+			// Controllers
+			+ SVerticalBox::Slot()
+			.AutoHeight()
+			.Padding(FMargin(0.f, 10.f, 0.f, 0.f))
+			[
+				SAssignNew(RCControllerPanel, SAvaRCControllerPanel, InPlaylistEditor)
+			]
+			// Exposed Properties
+			+ SVerticalBox::Slot()
+			.Padding(FMargin(0.f, 10.f, 0.f, 0.f))
+			.AutoHeight()
+			[
+				SNew(SHorizontalBox)
+				+ SHorizontalBox::Slot()
+				.VAlign(VAlign_Center)
+				.HAlign(HAlign_Left)
+				.Padding(5.f, 0.f, 0.f, 0.f)
+				.AutoWidth()
+				[
+					SNew(SButton)
+					.ContentPadding(0)
+					.ButtonStyle(FAppStyle::Get(), "NoBorder")
+					.OnClicked(this, &SAvaPageDetails::ToggleExposedPropertiesVisibility)
+					.ToolTipText(LOCTEXT("VisibilityButtonToolTip", "Toggle Exposed Properties Visibility"))
+					.Content()
+					[
+						SNew(SImage)
+						.Image(this, &SAvaPageDetails::GetExposedPropertiesVisibilityBrush)
+					]
+				]
+				+ SHorizontalBox::Slot()
+				.VAlign(VAlign_Center)
+				.HAlign(HAlign_Left)
+				.Padding(5.f, 0.f, 0.f, 0.f)
+				.AutoWidth()
+				[
+					SNew(STextBlock)
+					.Text(LOCTEXT("Properties", "Properties"))
+				]
+			]
+			+ SVerticalBox::Slot()
+			.AutoHeight()
+			[
+				SAssignNew(RemoteControlProps, SAvaPageRemoteControlProps, InPlaylistEditor)
+				.Visibility(EVisibility::Collapsed)
+			]
+		]
+	];
+
+	OnPageSelectionChanged({});
+}
+
+SAvaPageDetails::~SAvaPageDetails()
+{
+	if (const TSharedPtr<FAvaPlaylistEditor> PlaylistEditor = PlaylistEditorWeak.Pin())
+	{
+		PlaylistEditor->GetOnPageEvent().RemoveAll(this);
+	}
+	if (IAvaMediaModule::IsModuleLoaded())
+	{
+		IAvaMediaModule::Get().GetManagedInstanceCache().OnEntryInvalidated.RemoveAll(this);
+	}
+}
+
+void SAvaPageDetails::OnPageEvent(const TArray<int32>& InSelectedPageIds, UE::AvalanchePlaylist::EPageEvent InPageEvent)
+{
+	if (InPageEvent == UE::AvalanchePlaylist::EPageEvent::SelectionChanged)
+	{
+		OnPageSelectionChanged(InSelectedPageIds);
+		RemoteControlProps->Refresh(InSelectedPageIds);
+		RCControllerPanel->Refresh(InSelectedPageIds);
+	}
+	else if (InPageEvent == UE::AvalanchePlaylist::EPageEvent::ReimportRequest)
+	{
+		// Note: all widgets are refreshed in this event handler to ensure it is done
+		// after the cache is invalidated.
+		OnPageSelectionChanged(InSelectedPageIds);
+		// Only call SAvaPageRemoteControlProps::UpdateDefaultValuesAndRefresh to have it update the values of all selected pages.
+		RemoteControlProps->UpdateDefaultValuesAndRefresh(InSelectedPageIds);
+		// Pages values already updated above, only need to refresh UI.
+		RCControllerPanel->Refresh(InSelectedPageIds);
+	}
+}
+
+void SAvaPageDetails::OnPageSelectionChanged(const TArray<int32>& InSelectedPageIds)
+{
+	ActivePageId = InSelectedPageIds.IsEmpty() ? FAvalanchePage::InvalidPageId : InSelectedPageIds[0];
+}
+
+void SAvaPageDetails::OnManagedInstanceCacheEntryInvalidated(const FSoftObjectPath& InAssetPath)
+{
+	if (!bRefreshSelectedPageQueued)
+	{
+		if (TSharedPtr<FAvaPlaylistEditor> PlaylistEditor = PlaylistEditorWeak.Pin())
+		{
+			UAvalanchePlaylist* Playlist = PlaylistEditor->GetPlaylist();
+
+			if (IsValid(Playlist))
+			{
+				const FAvalanchePage& SelectedPage = GetSelectedPage();
+
+				if (SelectedPage.IsValidPage())
+				{
+					if (SelectedPage.GetAvalancheAssetPath(Playlist) == InAssetPath)
+					{
+						bRefreshSelectedPageQueued = true;
+						// Queue a refresh on next tick.
+						// We don't want to refresh immediately to avoid issues with
+						// cascading events within the managed blueprint cache.
+						TWeakPtr<SWidget> ThisWeak(AsShared());
+						AsyncTask(ENamedThreads::GameThread, [ThisWeak]()
+							{
+								if (const TSharedPtr<SWidget> ThisWidget = ThisWeak.Pin())
+								{
+									SAvaPageDetails* AvaPageDetails = static_cast<SAvaPageDetails*>(ThisWidget.Get());
+									AvaPageDetails->RefreshSelectedPage();
+									AvaPageDetails->bRefreshSelectedPageQueued = false;
+								}
+							});
+					}
+				}
+			}
+		}
+	}
+}
+
+FReply SAvaPageDetails::ToggleExposedPropertiesVisibility()
+{
+	if (RemoteControlProps->GetVisibility() == EVisibility::Collapsed)
+	{
+		RemoteControlProps->SetVisibility(EVisibility::SelfHitTestInvisible);
+	}
+	else
+	{
+		RemoteControlProps->SetVisibility(EVisibility::Collapsed);
+	}
+
+	return FReply::Handled();
+}
+
+const FSlateBrush* SAvaPageDetails::GetExposedPropertiesVisibilityBrush() const
+{
+	if (RemoteControlProps->GetVisibility() == EVisibility::Collapsed)
+	{
+		return FAppStyle::GetBrush("Level.NotVisibleHighlightIcon16x");
+	}
+	else
+	{
+		return FAppStyle::GetBrush("Level.VisibleHighlightIcon16x");
+	}
+}
+
+const FAvalanchePage& SAvaPageDetails::GetSelectedPage() const
+{
+	const TSharedPtr<FAvaPlaylistEditor> PlaylistEditor = PlaylistEditorWeak.Pin();
+
+	if (HasSelectedPage() && PlaylistEditor && PlaylistEditor->IsPlaylistValid())
+	{
+		return PlaylistEditor->GetPlaylist()->GetPage(ActivePageId);
+	}
+
+	return FAvalanchePage::NullPage;
+}
+
+FAvalanchePage& SAvaPageDetails::GetMutableSelectedPage() const
+{
+	const TSharedPtr<FAvaPlaylistEditor> PlaylistEditor = PlaylistEditorWeak.Pin();
+
+	if (HasSelectedPage() && PlaylistEditor && PlaylistEditor->IsPlaylistValid())
+	{
+		return PlaylistEditor->GetPlaylist()->GetPage(ActivePageId);
+	}
+
+	return FAvalanchePage::NullPage;
+}
+
+void SAvaPageDetails::RefreshSelectedPage()
+{
+	const FAvalanchePage& SelectedPage = GetSelectedPage();
+
+	if (SelectedPage.IsValidPage())
+	{
+		OnPageSelectionChanged({SelectedPage.GetPageId()});
+		RemoteControlProps->UpdateDefaultValuesAndRefresh({SelectedPage.GetPageId()});
+		RCControllerPanel->Refresh({SelectedPage.GetPageId()});
+	}
+}
+
+bool SAvaPageDetails::HasSelectedPage() const
+{
+	if (ActivePageId == FAvalanchePage::InvalidPageId)
+	{
+		return false;
+	}
+
+	return PlaylistEditorWeak.IsValid();
+}
+
+FText SAvaPageDetails::GetPageId() const
+{
+	const FAvalanchePage& SelectedPage = GetSelectedPage();
+
+	if (SelectedPage.IsValidPage())
+	{
+		return FText::AsNumber(SelectedPage.GetPageId(), &UE::AvalanchePlaylist::FEditorMetrics::PageIdFormattingOptions);
+	}
+
+	return FText::GetEmpty();
+}
+
+void SAvaPageDetails::OnPageIdCommitted(const FText& InNewText, ETextCommit::Type InCommitType)
+{
+	switch (InCommitType)
+	{
+		case ETextCommit::OnEnter:
+		case ETextCommit::OnUserMovedFocus:
+			break;
+
+		case ETextCommit::Default:
+		case ETextCommit::OnCleared:
+		default:
+			return;
+	}
+
+	if (!InNewText.IsNumeric())
+	{
+		return;
+	}
+
+	FAvalanchePage& SelectedPage = GetMutableSelectedPage();
+
+	if (!SelectedPage.IsValidPage()) // Not FAvalanchePage::NullPage
+	{
+		return;
+	}
+
+	int32 NewId = FCString::Atoi(*InNewText.ToString());
+
+	if (NewId == SelectedPage.GetPageId())
+	{
+		return;
+	}
+
+	const TSharedPtr<FAvaPlaylistEditor> PlaylistEditor = PlaylistEditorWeak.Pin();
+
+	if (!PlaylistEditor.IsValid() || !IsValid(PlaylistEditor->GetPlaylist()))
+	{
+		return;
+	}
+
+	const bool bSuccessful = PlaylistEditor->GetPlaylist()->RenumberPageId(SelectedPage.GetPageId(), NewId);
+
+	if (bSuccessful)
+	{
+		TSharedPtr<SAvaInstancedPageList> PageList = PlaylistEditor->GetActiveListWidget();
+		
+		if (PageList.IsValid())
+		{
+			PageList->SelectPage(NewId);
+		}
+	}
+}
+
+FText SAvaPageDetails::GetPageDescription() const
+{
+	const FAvalanchePage& SelectedPage = GetSelectedPage();
+
+	if (SelectedPage.IsValidPage())
+	{
+		return SelectedPage.GetPageDescription();
+	}
+
+	return FText::GetEmpty();
+}
+
+void SAvaPageDetails::OnPageNameChanged(const FText& InNewText)
+{
+	FAvalanchePage& SelectedPage = GetMutableSelectedPage();
+
+	if (SelectedPage.IsValidPage()) // Not FAvalanchePage::NullPage
+	{
+		SelectedPage.SetPageFriendlyName(InNewText);
+	}
+}
+
+FReply SAvaPageDetails::DuplicateSelectedPage()
+{
+	FAvalanchePage& SelectedPage = GetMutableSelectedPage();
+
+	if (!SelectedPage.IsValidPage()) // Not FAvalanchePage::NullPage
+	{
+		return FReply::Unhandled();
+	}
+
+	const TSharedPtr<FAvaPlaylistEditor> PlaylistEditor = PlaylistEditorWeak.Pin();
+
+	if (!PlaylistEditor.IsValid())
+	{
+		return FReply::Unhandled();
+	}
+
+	TSharedPtr<SAvaInstancedPageList> PageList = PlaylistEditor->GetActiveListWidget();
+
+	if (!PageList.IsValid())
+	{
+		return FReply::Unhandled();
+	}
+
+	TArray<int32> SelectedPages = TArray<int32>(PageList->GetSelectedPageIds());
+	PageList->SelectPage(SelectedPage.GetPageId());
+	PageList->DuplicateSelectedPages();
+	PageList->SelectPages(SelectedPages);
+
+	return FReply::Handled();
+}
+
+#undef LOCTEXT_NAMESPACE
diff --git a/Engine/Plugins/Experimental/Avalanche/Source/AvalancheMediaEditor/Private/Playlist/DetailsView/SAvaPageDetails.h b/Engine/Plugins/Experimental/Avalanche/Source/AvalancheMediaEditor/Private/Playlist/DetailsView/SAvaPageDetails.h
new file mode 100644
--- /dev/null
+++ b/Engine/Plugins/Experimental/Avalanche/Source/AvalancheMediaEditor/Private/Playlist/DetailsView/SAvaPageDetails.h
@@ -0,0 +1,63 @@
+// Copyright Epic Games, Inc. All Rights Reserved.
+
+#pragma once
+
+#include "CoreMinimal.h"
+#include "Playlist/AvaPlaylistDefines.h"
+#include "Widgets/SCompoundWidget.h"
+
+class FAvaPlaylistEditor;
+class SAvaPageRemoteControlProps;
+class SAvaRCControllerPanel;
+struct FAvalanchePage;
+
+class SAvaPageDetails : public SCompoundWidget
+{
+public:
+	SLATE_BEGIN_ARGS(SAvaPageDetails) {}
+	SLATE_END_ARGS()
+
+	/** Constructs this widget with InArgs */
+	void Construct(const FArguments& InArgs, const TSharedPtr<FAvaPlaylistEditor>& InPlaylistEditor);
+
+	virtual ~SAvaPageDetails() override;
+
+	void OnPageEvent(const TArray<int32>& InSelectedPageIds, UE::AvalanchePlaylist::EPageEvent InPageEvent);
+	void OnPageSelectionChanged(const TArray<int32>& InSelectedPageIds);
+	void OnManagedInstanceCacheEntryInvalidated(const FSoftObjectPath& InAssetPath);
+
+protected:
+	FReply ToggleExposedPropertiesVisibility();
+
+	const FSlateBrush* GetExposedPropertiesVisibilityBrush() const;
+
+	const FAvalanchePage& GetSelectedPage() const;
+	FAvalanchePage& GetMutableSelectedPage() const;
+
+	void RefreshSelectedPage();
+
+	bool HasSelectedPage() const;
+
+	FText GetPageId() const;
+
+	/** Only update page id on commit. */
+	void OnPageIdCommitted(const FText& InNewText, ETextCommit::Type InCommitType);
+
+	FText GetPageDescription() const;
+
+	/** Update page name live. */
+	void OnPageNameChanged(const FText& InNewText);
+
+	FReply DuplicateSelectedPage();
+
+private:
+	TWeakPtr<FAvaPlaylistEditor> PlaylistEditorWeak;
+
+	TSharedPtr<SAvaPageRemoteControlProps> RemoteControlProps;
+
+	TSharedPtr<SAvaRCControllerPanel> RCControllerPanel;
+
+	bool bRefreshSelectedPageQueued = false;
+
+	int32 ActivePageId;
+};
diff --git a/Engine/Plugins/Experimental/Avalanche/Source/AvalancheMediaEditor/Private/Playlist/Factories/Filters/AvaPlaylistFilterChannelExpressionFactory.cpp b/Engine/Plugins/Experimental/Avalanche/Source/AvalancheMediaEditor/Private/Playlist/Factories/Filters/AvaPlaylistFilterChannelExpressionFactory.cpp
new file mode 100644
--- /dev/null
+++ b/Engine/Plugins/Experimental/Avalanche/Source/AvalancheMediaEditor/Private/Playlist/Factories/Filters/AvaPlaylistFilterChannelExpressionFactory.cpp
@@ -0,0 +1,27 @@
+// Copyright Epic Games, Inc. All Rights Reserved.
+
+#include "AvaPlaylistFilterChannelExpressionFactory.h"
+#include "Misc/TextFilterUtils.h"
+#include "Playlist/AvalanchePage.h"
+#include "Playlist/Factories/Filters/AvaPlaylistFactoriesUtils.h"
+
+const FName FAvaPlaylistFilterChannelExpressionFactory::KeyName = FName(TEXT("CHANNEL"));
+
+FName FAvaPlaylistFilterChannelExpressionFactory::GetFilterIdentifier() const
+{
+	return KeyName;
+}
+
+bool FAvaPlaylistFilterChannelExpressionFactory::FilterExpression(const FAvalanchePage& InItem, const FAvaPlaylistTextFilterArgs& InArgs) const
+{
+	return InArgs.ValueToCheck.CompareName(InItem.GetChannelName(), InArgs.ComparisonMode);
+}
+
+bool FAvaPlaylistFilterChannelExpressionFactory::SupportsComparisonOperation(ETextFilterComparisonOperation InComparisonOperation, EAvaPlaylistSearchListType InPlaylistSearchListType) const
+{
+	const bool bIsEqualityOperation = InComparisonOperation == ETextFilterComparisonOperation::Equal 
+		|| InComparisonOperation == ETextFilterComparisonOperation::NotEqual;
+
+	return bIsEqualityOperation 
+	   && InPlaylistSearchListType == EAvaPlaylistSearchListType::Instanced;
+}
diff --git a/Engine/Plugins/Experimental/Avalanche/Source/AvalancheMediaEditor/Private/Playlist/Factories/Filters/AvaPlaylistFilterChannelExpressionFactory.h b/Engine/Plugins/Experimental/Avalanche/Source/AvalancheMediaEditor/Private/Playlist/Factories/Filters/AvaPlaylistFilterChannelExpressionFactory.h
new file mode 100644
--- /dev/null
+++ b/Engine/Plugins/Experimental/Avalanche/Source/AvalancheMediaEditor/Private/Playlist/Factories/Filters/AvaPlaylistFilterChannelExpressionFactory.h
@@ -0,0 +1,17 @@
+// Copyright Epic Games, Inc. All Rights Reserved.
+
+#pragma once
+
+#include "Playlist/Factories/Filters/IAvaPlaylistFilterExpressionFactory.h"
+
+class FAvaPlaylistFilterChannelExpressionFactory : public IAvaPlaylistFilterExpressionFactory
+{
+public:
+	static const FName KeyName;
+
+	//~ Begin IAvaFilterExpressionFactory interface
+	virtual FName GetFilterIdentifier() const override;
+	virtual bool FilterExpression(const FAvalanchePage& InItem, const FAvaPlaylistTextFilterArgs& InArgs) const override;
+	virtual bool SupportsComparisonOperation(ETextFilterComparisonOperation InComparisonOperation, EAvaPlaylistSearchListType InPlaylistSearchListType) const override;
+	//~ End IAvaFilterExpressionFactory interface
+};
diff --git a/Engine/Plugins/Experimental/Avalanche/Source/AvalancheMediaEditor/Private/Playlist/Factories/Filters/AvaPlaylistFilterIdExpressionFactory.cpp b/Engine/Plugins/Experimental/Avalanche/Source/AvalancheMediaEditor/Private/Playlist/Factories/Filters/AvaPlaylistFilterIdExpressionFactory.cpp
new file mode 100644
--- /dev/null
+++ b/Engine/Plugins/Experimental/Avalanche/Source/AvalancheMediaEditor/Private/Playlist/Factories/Filters/AvaPlaylistFilterIdExpressionFactory.cpp
@@ -0,0 +1,28 @@
+// Copyright Epic Games, Inc. All Rights Reserved.
+
+#include "AvaPlaylistFilterIdExpressionFactory.h"
+#include "Misc/TextFilterUtils.h"
+#include "Playlist/AvalanchePage.h"
+#include "Playlist/AvaPlaylistDefines.h"
+
+const FName FAvaPlaylistFilterIdExpressionFactory::KeyName = FName(TEXT("ID"));
+
+FName FAvaPlaylistFilterIdExpressionFactory::GetFilterIdentifier() const
+{
+	return KeyName;
+}
+
+bool FAvaPlaylistFilterIdExpressionFactory::FilterExpression(const FAvalanchePage& InItem, const FAvaPlaylistTextFilterArgs& InArgs) const
+{
+	const FTextFilterString FilterString = FString::FromInt(InItem.GetPageId());
+	if (FilterString.CanCompareNumeric(InArgs.ValueToCheck))
+	{
+		return FilterString.CompareNumeric(InArgs.ValueToCheck, InArgs.ComparisonOperation);
+	}
+	return false;
+}
+
+bool FAvaPlaylistFilterIdExpressionFactory::SupportsComparisonOperation(ETextFilterComparisonOperation InComparisonOperation, EAvaPlaylistSearchListType InPlaylistSearchListType) const
+{
+	return true;
+}
diff --git a/Engine/Plugins/Experimental/Avalanche/Source/AvalancheMediaEditor/Private/Playlist/Factories/Filters/AvaPlaylistFilterIdExpressionFactory.h b/Engine/Plugins/Experimental/Avalanche/Source/AvalancheMediaEditor/Private/Playlist/Factories/Filters/AvaPlaylistFilterIdExpressionFactory.h
new file mode 100644
--- /dev/null
+++ b/Engine/Plugins/Experimental/Avalanche/Source/AvalancheMediaEditor/Private/Playlist/Factories/Filters/AvaPlaylistFilterIdExpressionFactory.h
@@ -0,0 +1,17 @@
+// Copyright Epic Games, Inc. All Rights Reserved.
+
+#pragma once
+
+#include "Playlist/Factories/Filters/IAvaPlaylistFilterExpressionFactory.h"
+
+class FAvaPlaylistFilterIdExpressionFactory : public IAvaPlaylistFilterExpressionFactory
+{
+public:
+	static const FName KeyName;
+
+	//~ Begin IAvaFilterExpressionFactory interface
+	virtual FName GetFilterIdentifier() const override;
+	virtual bool FilterExpression(const FAvalanchePage& InItem, const FAvaPlaylistTextFilterArgs& InArgs) const override;
+	virtual bool SupportsComparisonOperation(ETextFilterComparisonOperation InComparisonOperation, EAvaPlaylistSearchListType InPlaylistSearchListType) const override;
+	//~ End IAvaFilterExpressionFactory interface
+};
diff --git a/Engine/Plugins/Experimental/Avalanche/Source/AvalancheMediaEditor/Private/Playlist/Factories/Filters/AvaPlaylistFilterNameExpressionFactory.cpp b/Engine/Plugins/Experimental/Avalanche/Source/AvalancheMediaEditor/Private/Playlist/Factories/Filters/AvaPlaylistFilterNameExpressionFactory.cpp
new file mode 100644
--- /dev/null
+++ b/Engine/Plugins/Experimental/Avalanche/Source/AvalancheMediaEditor/Private/Playlist/Factories/Filters/AvaPlaylistFilterNameExpressionFactory.cpp
@@ -0,0 +1,23 @@
+// Copyright Epic Games, Inc. All Rights Reserved.
+
+#include "AvaPlaylistFilterNameExpressionFactory.h"
+#include "Misc/TextFilterUtils.h"
+#include "Playlist/AvalanchePage.h"
+
+const FName FAvaPlaylistFilterNameExpressionFactory::KeyName = FName(TEXT("NAME"));
+
+FName FAvaPlaylistFilterNameExpressionFactory::GetFilterIdentifier() const
+{
+	return KeyName;
+}
+
+bool FAvaPlaylistFilterNameExpressionFactory::FilterExpression(const FAvalanchePage& InItem, const FAvaPlaylistTextFilterArgs& InArgs) const
+{
+	const FName Name = FName(InItem.GetPageDescription().ToString());
+	return InArgs.ValueToCheck.CompareName(Name, InArgs.ComparisonMode);
+}
+
+bool FAvaPlaylistFilterNameExpressionFactory::SupportsComparisonOperation(ETextFilterComparisonOperation InComparisonOperation, EAvaPlaylistSearchListType InPlaylistSearchListType) const
+{
+	return InComparisonOperation == ETextFilterComparisonOperation::Equal || InComparisonOperation == ETextFilterComparisonOperation::NotEqual;
+}
diff --git a/Engine/Plugins/Experimental/Avalanche/Source/AvalancheMediaEditor/Private/Playlist/Factories/Filters/AvaPlaylistFilterNameExpressionFactory.h b/Engine/Plugins/Experimental/Avalanche/Source/AvalancheMediaEditor/Private/Playlist/Factories/Filters/AvaPlaylistFilterNameExpressionFactory.h
new file mode 100644
--- /dev/null
+++ b/Engine/Plugins/Experimental/Avalanche/Source/AvalancheMediaEditor/Private/Playlist/Factories/Filters/AvaPlaylistFilterNameExpressionFactory.h
@@ -0,0 +1,17 @@
+// Copyright Epic Games, Inc. All Rights Reserved.
+
+#pragma once
+
+#include "Playlist/Factories/Filters/IAvaPlaylistFilterExpressionFactory.h"
+
+class FAvaPlaylistFilterNameExpressionFactory : public IAvaPlaylistFilterExpressionFactory
+{
+public:
+	static const FName KeyName;
+
+	//~ Begin IAvaFilterExpressionFactory interface
+	virtual FName GetFilterIdentifier() const override;
+	virtual bool FilterExpression(const FAvalanchePage& InItem, const FAvaPlaylistTextFilterArgs& InArgs) const override;
+	virtual bool SupportsComparisonOperation(ETextFilterComparisonOperation InComparisonOperation, EAvaPlaylistSearchListType InPlaylistSearchListType) const override;
+	//~ End IAvaFilterExpressionFactory interface
+};
diff --git a/Engine/Plugins/Experimental/Avalanche/Source/AvalancheMediaEditor/Private/Playlist/Factories/Filters/AvaPlaylistFilterPathExpressionFactory.cpp b/Engine/Plugins/Experimental/Avalanche/Source/AvalancheMediaEditor/Private/Playlist/Factories/Filters/AvaPlaylistFilterPathExpressionFactory.cpp
new file mode 100644
--- /dev/null
+++ b/Engine/Plugins/Experimental/Avalanche/Source/AvalancheMediaEditor/Private/Playlist/Factories/Filters/AvaPlaylistFilterPathExpressionFactory.cpp
@@ -0,0 +1,27 @@
+// Copyright Epic Games, Inc. All Rights Reserved.
+
+#include "AvaPlaylistFilterPathExpressionFactory.h"
+#include "Misc/TextFilterUtils.h"
+#include "Playlist/AvalanchePage.h"
+
+const FName FAvaPlaylistFilterPathExpressionFactory::KeyName = FName(TEXT("ASSET"));
+
+FName FAvaPlaylistFilterPathExpressionFactory::GetFilterIdentifier() const
+{
+	return KeyName;
+}
+
+bool FAvaPlaylistFilterPathExpressionFactory::FilterExpression(const FAvalanchePage& InItem, const FAvaPlaylistTextFilterArgs& InArgs) const
+{
+	if (!InArgs.ItemPlaylist)
+	{
+		return false;
+	}
+	const FName AssetPath = FName(InItem.GetAvalancheAssetPath(InArgs.ItemPlaylist).GetAssetPathString());
+	return InArgs.ValueToCheck.CompareName(AssetPath, InArgs.ComparisonMode);
+}
+
+bool FAvaPlaylistFilterPathExpressionFactory::SupportsComparisonOperation(ETextFilterComparisonOperation InComparisonOperation, EAvaPlaylistSearchListType InPlaylistSearchListType) const
+{
+	return InComparisonOperation == ETextFilterComparisonOperation::Equal || InComparisonOperation == ETextFilterComparisonOperation::NotEqual;
+}
diff --git a/Engine/Plugins/Experimental/Avalanche/Source/AvalancheMediaEditor/Private/Playlist/Factories/Filters/AvaPlaylistFilterPathExpressionFactory.h b/Engine/Plugins/Experimental/Avalanche/Source/AvalancheMediaEditor/Private/Playlist/Factories/Filters/AvaPlaylistFilterPathExpressionFactory.h
new file mode 100644
--- /dev/null
+++ b/Engine/Plugins/Experimental/Avalanche/Source/AvalancheMediaEditor/Private/Playlist/Factories/Filters/AvaPlaylistFilterPathExpressionFactory.h
@@ -0,0 +1,17 @@
+// Copyright Epic Games, Inc. All Rights Reserved.
+
+#pragma once
+
+#include "Playlist/Factories/Filters/IAvaPlaylistFilterExpressionFactory.h"
+
+class FAvaPlaylistFilterPathExpressionFactory : public IAvaPlaylistFilterExpressionFactory
+{
+public:
+	static const FName KeyName;
+
+	//~ Begin IAvaFilterExpressionFactory interface
+	virtual FName GetFilterIdentifier() const override;
+	virtual bool FilterExpression(const FAvalanchePage& InItem, const FAvaPlaylistTextFilterArgs& InArgs) const override;
+	virtual bool SupportsComparisonOperation(ETextFilterComparisonOperation InComparisonOperation, EAvaPlaylistSearchListType InPlaylistSearchListType) const override;
+	//~ End IAvaFilterExpressionFactory interface
+};
diff --git a/Engine/Plugins/Experimental/Avalanche/Source/AvalancheMediaEditor/Private/Playlist/Factories/Filters/AvaPlaylistFilterStatusExpressionFactory.cpp b/Engine/Plugins/Experimental/Avalanche/Source/AvalancheMediaEditor/Private/Playlist/Factories/Filters/AvaPlaylistFilterStatusExpressionFactory.cpp
new file mode 100644
--- /dev/null
+++ b/Engine/Plugins/Experimental/Avalanche/Source/AvalancheMediaEditor/Private/Playlist/Factories/Filters/AvaPlaylistFilterStatusExpressionFactory.cpp
@@ -0,0 +1,35 @@
+// Copyright Epic Games, Inc. All Rights Reserved.
+
+#include "AvaPlaylistFilterStatusExpressionFactory.h"
+#include "Misc/TextFilterUtils.h"
+#include "Playlist/AvalanchePage.h"
+
+const FName FAvaPlaylistFilterStatusExpressionFactory::KeyName = FName(TEXT("STATUS"));
+
+FName FAvaPlaylistFilterStatusExpressionFactory::GetFilterIdentifier() const
+{
+	return KeyName;
+}
+
+bool FAvaPlaylistFilterStatusExpressionFactory::FilterExpression(const FAvalanchePage& InItem, const FAvaPlaylistTextFilterArgs& InArgs) const
+{
+	if (!InArgs.ItemPlaylist)
+	{
+		return false;
+	}
+	const TArray<FAvalanchePageStatus> Statuses = InItem.GetPageContextualStatuses(InArgs.ItemPlaylist);
+	for (const FAvalanchePageStatus& Status : Statuses)
+	{
+		const FName StatusName = StaticEnum<EAvalanchePageStatus>()->GetNameByValue(static_cast<int32>(Status.Status));
+		if (InArgs.ValueToCheck.CompareName(StatusName, InArgs.ComparisonMode))
+		{
+			return true;
+		}
+	}
+	return false;
+}
+
+bool FAvaPlaylistFilterStatusExpressionFactory::SupportsComparisonOperation(ETextFilterComparisonOperation InComparisonOperation, EAvaPlaylistSearchListType InPlaylistSearchListType) const
+{
+	return InComparisonOperation == ETextFilterComparisonOperation::Equal || InComparisonOperation == ETextFilterComparisonOperation::NotEqual;
+}
diff --git a/Engine/Plugins/Experimental/Avalanche/Source/AvalancheMediaEditor/Private/Playlist/Factories/Filters/AvaPlaylistFilterStatusExpressionFactory.h b/Engine/Plugins/Experimental/Avalanche/Source/AvalancheMediaEditor/Private/Playlist/Factories/Filters/AvaPlaylistFilterStatusExpressionFactory.h
new file mode 100644
--- /dev/null
+++ b/Engine/Plugins/Experimental/Avalanche/Source/AvalancheMediaEditor/Private/Playlist/Factories/Filters/AvaPlaylistFilterStatusExpressionFactory.h
@@ -0,0 +1,17 @@
+// Copyright Epic Games, Inc. All Rights Reserved.
+
+#pragma once
+
+#include "Playlist/Factories/Filters/IAvaPlaylistFilterExpressionFactory.h"
+
+class FAvaPlaylistFilterStatusExpressionFactory : public IAvaPlaylistFilterExpressionFactory
+{
+public:
+	static const FName KeyName;
+
+	//~ Begin IAvaFilterExpressionFactory interface
+	virtual FName GetFilterIdentifier() const override;
+	virtual bool FilterExpression(const FAvalanchePage& InItem, const FAvaPlaylistTextFilterArgs& InArgs) const override;
+	virtual bool SupportsComparisonOperation(ETextFilterComparisonOperation InComparisonOperation, EAvaPlaylistSearchListType InPlaylistSearchListType) const override;
+	//~ End IAvaFilterExpressionFactory interface
+};
diff --git a/Engine/Plugins/Experimental/Avalanche/Source/AvalancheMediaEditor/Private/Playlist/Factories/Filters/AvaPlaylistFilterTransitionLayerExpressionFactory.cpp b/Engine/Plugins/Experimental/Avalanche/Source/AvalancheMediaEditor/Private/Playlist/Factories/Filters/AvaPlaylistFilterTransitionLayerExpressionFactory.cpp
new file mode 100644
--- /dev/null
+++ b/Engine/Plugins/Experimental/Avalanche/Source/AvalancheMediaEditor/Private/Playlist/Factories/Filters/AvaPlaylistFilterTransitionLayerExpressionFactory.cpp
@@ -0,0 +1,26 @@
+// Copyright Epic Games, Inc. All Rights Reserved.
+
+#include "AvaPlaylistFilterTransitionLayerExpressionFactory.h"
+#include "Misc/TextFilterUtils.h"
+#include "Playlist/AvalanchePage.h"
+
+const FName FAvaPlaylistFilterTransitionLayerExpressionFactory::KeyName = FName(TEXT("TRANSITIONLAYER"));
+
+FName FAvaPlaylistFilterTransitionLayerExpressionFactory::GetFilterIdentifier() const
+{
+	return KeyName;
+}
+
+bool FAvaPlaylistFilterTransitionLayerExpressionFactory::FilterExpression(const FAvalanchePage& InItem, const FAvaPlaylistTextFilterArgs& InArgs) const
+{
+	if (!InArgs.ItemPlaylist)
+	{
+		return false;
+	}
+	return InArgs.ValueToCheck.CompareFString(InItem.GetTransitionLayer(InArgs.ItemPlaylist).ToString(), InArgs.ComparisonMode);
+}
+
+bool FAvaPlaylistFilterTransitionLayerExpressionFactory::SupportsComparisonOperation(ETextFilterComparisonOperation InComparisonOperation, EAvaPlaylistSearchListType InPlaylistSearchListType) const
+{
+	return InComparisonOperation == ETextFilterComparisonOperation::Equal || InComparisonOperation == ETextFilterComparisonOperation::NotEqual;
+}
diff --git a/Engine/Plugins/Experimental/Avalanche/Source/AvalancheMediaEditor/Private/Playlist/Factories/Filters/AvaPlaylistFilterTransitionLayerExpressionFactory.h b/Engine/Plugins/Experimental/Avalanche/Source/AvalancheMediaEditor/Private/Playlist/Factories/Filters/AvaPlaylistFilterTransitionLayerExpressionFactory.h
new file mode 100644
--- /dev/null
+++ b/Engine/Plugins/Experimental/Avalanche/Source/AvalancheMediaEditor/Private/Playlist/Factories/Filters/AvaPlaylistFilterTransitionLayerExpressionFactory.h
@@ -0,0 +1,17 @@
+// Copyright Epic Games, Inc. All Rights Reserved.
+
+#pragma once
+
+#include "Playlist/Factories/Filters/IAvaPlaylistFilterExpressionFactory.h"
+
+class FAvaPlaylistFilterTransitionLayerExpressionFactory : public IAvaPlaylistFilterExpressionFactory
+{
+public:
+	static const FName KeyName;
+
+	//~ Begin IAvaFilterExpressionFactory interface
+	virtual FName GetFilterIdentifier() const override;
+	virtual bool FilterExpression(const FAvalanchePage& InItem, const FAvaPlaylistTextFilterArgs& InArgs) const override;
+	virtual bool SupportsComparisonOperation(ETextFilterComparisonOperation InComparisonOperation, EAvaPlaylistSearchListType InPlaylistSearchListType) const override;
+	//~ End IAvaFilterExpressionFactory interface
+};
diff --git a/Engine/Plugins/Experimental/Avalanche/Source/AvalancheMediaEditor/Private/Playlist/Factories/Suggestions/AvaPlaylistFilterChannelSuggestionFactory.cpp b/Engine/Plugins/Experimental/Avalanche/Source/AvalancheMediaEditor/Private/Playlist/Factories/Suggestions/AvaPlaylistFilterChannelSuggestionFactory.cpp
new file mode 100644
--- /dev/null
+++ b/Engine/Plugins/Experimental/Avalanche/Source/AvalancheMediaEditor/Private/Playlist/Factories/Suggestions/AvaPlaylistFilterChannelSuggestionFactory.cpp
@@ -0,0 +1,36 @@
+// Copyright Epic Games, Inc. All Rights Reserved.
+
+#include "AvaPlaylistFilterChannelSuggestionFactory.h"
+#include "Playlist/AvalanchePlaylist.h"
+#include "Playlist/Factories/Filters/AvaPlaylistFactoriesUtils.h"
+#include "SAssetSearchBox.h"
+
+#define LOCTEXT_NAMESPACE "AvaPlaylistFilterChannelSuggestionFactory"
+
+const FName FAvaPlaylistFilterChannelSuggestionFactory::KeyName = FName(TEXT("CHANNEL"));
+
+void FAvaPlaylistFilterChannelSuggestionFactory::AddSuggestion(const TSharedRef<FAvaPlaylistFilterSuggestionPayload>& InPayload)
+{
+	const FAvalanchePage& PageItem = UAvalanchePlaylist::GetPageSafe(InPayload->Playlist, InPayload->ItemPageId);
+	if (PageItem.IsValidPage())
+	{
+		const FText ChannelCategoryLabel = LOCTEXT("ChannelCategoryLabel", "Ava-Playlist-Channel");
+		const FString ChannelName = PageItem.GetChannelName().ToString();
+
+		FString ChannelNameSuggestion = FString::Printf(TEXT("Channel=%s"), *ChannelName);
+		const bool bIsFilterValueValid = InPayload->FilterValue.IsEmpty() || ChannelNameSuggestion.Contains(InPayload->FilterValue);
+
+		if (bIsFilterValueValid && !InPayload->FilterCache.Contains(ChannelNameSuggestion))
+		{
+			InPayload->FilterCache.Add(ChannelNameSuggestion);
+			InPayload->PossibleSuggestions.Add(FAssetSearchBoxSuggestion{MoveTemp(ChannelNameSuggestion), FText::FromString(ChannelName), ChannelCategoryLabel});
+		}
+	}
+}
+
+bool FAvaPlaylistFilterChannelSuggestionFactory::SupportSuggestionType(EAvaPlaylistSearchListType InSuggestionType) const
+{
+	return InSuggestionType == EAvaPlaylistSearchListType::Instanced;
+}
+
+#undef LOCTEXT_NAMESPACE
diff --git a/Engine/Plugins/Experimental/Avalanche/Source/AvalancheMediaEditor/Private/Playlist/Factories/Suggestions/AvaPlaylistFilterChannelSuggestionFactory.h b/Engine/Plugins/Experimental/Avalanche/Source/AvalancheMediaEditor/Private/Playlist/Factories/Suggestions/AvaPlaylistFilterChannelSuggestionFactory.h
new file mode 100644
--- /dev/null
+++ b/Engine/Plugins/Experimental/Avalanche/Source/AvalancheMediaEditor/Private/Playlist/Factories/Suggestions/AvaPlaylistFilterChannelSuggestionFactory.h
@@ -0,0 +1,18 @@
+// Copyright Epic Games, Inc. All Rights Reserved.
+
+#pragma once
+
+#include "Playlist/Factories/Filters/IAvaPlaylistFilterSuggestionFactory.h"
+
+class FAvaPlaylistFilterChannelSuggestionFactory : public IAvaPlaylistFilterSuggestionFactory
+{
+public:
+	static const FName KeyName;
+
+	//~ Begin IAvaFilterSuggestionFactory interface
+	virtual FName GetSuggestionIdentifier() const override { return KeyName; }
+	virtual bool IsSimpleSuggestion() const override { return false; }
+	virtual void AddSuggestion(const TSharedRef<FAvaPlaylistFilterSuggestionPayload>& InPayload) override;
+	virtual bool SupportSuggestionType(EAvaPlaylistSearchListType InSuggestionType) const override;
+	//~ End IAvaFilterSuggestionFactory interface
+};
diff --git a/Engine/Plugins/Experimental/Avalanche/Source/AvalancheMediaEditor/Private/Playlist/Factories/Suggestions/AvaPlaylistFilterIdSuggestionFactory.cpp b/Engine/Plugins/Experimental/Avalanche/Source/AvalancheMediaEditor/Private/Playlist/Factories/Suggestions/AvaPlaylistFilterIdSuggestionFactory.cpp
new file mode 100644
--- /dev/null
+++ b/Engine/Plugins/Experimental/Avalanche/Source/AvalancheMediaEditor/Private/Playlist/Factories/Suggestions/AvaPlaylistFilterIdSuggestionFactory.cpp
@@ -0,0 +1,30 @@
+// Copyright Epic Games, Inc. All Rights Reserved.
+
+#include "AvaPlaylistFilterIdSuggestionFactory.h"
+#include "Playlist/AvalanchePage.h"
+#include "Playlist/AvaPlaylistDefines.h"
+#include "Playlist/Factories/Filters/AvaPlaylistFactoriesUtils.h"
+#include "SAssetSearchBox.h"
+
+#define LOCTEXT_NAMESPACE "AvaPlaylistFilterIdSuggestionFactory"
+
+const FName FAvaPlaylistFilterIdSuggestionFactory::KeyName = FName(TEXT("ID"));
+
+void FAvaPlaylistFilterIdSuggestionFactory::AddSuggestion(const TSharedRef<FAvaPlaylistFilterSuggestionPayload>& InPayload)
+{
+	const FText IdCategoryLabel = LOCTEXT("IdCategoryLabel", "Ava-Playlist-Id");
+	FString IdNameSuggestion = TEXT("Id");
+	const bool bIsFilterValueValid = InPayload->FilterValue.IsEmpty() || IdNameSuggestion.Contains(InPayload->FilterValue);
+	if (bIsFilterValueValid && !InPayload->FilterCache.Contains(IdNameSuggestion))
+	{
+		InPayload->FilterCache.Add(IdNameSuggestion);
+		InPayload->PossibleSuggestions.Add(FAssetSearchBoxSuggestion{MoveTemp(IdNameSuggestion), LOCTEXT("IdSuggestion","Id"), IdCategoryLabel});
+	}
+}
+
+bool FAvaPlaylistFilterIdSuggestionFactory::SupportSuggestionType(EAvaPlaylistSearchListType InSuggestionType) const
+{
+	return InSuggestionType != EAvaPlaylistSearchListType::None;
+}
+
+#undef LOCTEXT_NAMESPACE
diff --git a/Engine/Plugins/Experimental/Avalanche/Source/AvalancheMediaEditor/Private/Playlist/Factories/Suggestions/AvaPlaylistFilterIdSuggestionFactory.h b/Engine/Plugins/Experimental/Avalanche/Source/AvalancheMediaEditor/Private/Playlist/Factories/Suggestions/AvaPlaylistFilterIdSuggestionFactory.h
new file mode 100644
--- /dev/null
+++ b/Engine/Plugins/Experimental/Avalanche/Source/AvalancheMediaEditor/Private/Playlist/Factories/Suggestions/AvaPlaylistFilterIdSuggestionFactory.h
@@ -0,0 +1,18 @@
+// Copyright Epic Games, Inc. All Rights Reserved.
+
+#pragma once
+
+#include "Playlist/Factories/Filters/IAvaPlaylistFilterSuggestionFactory.h"
+
+class FAvaPlaylistFilterIdSuggestionFactory : public IAvaPlaylistFilterSuggestionFactory
+{
+public:
+	static const FName KeyName;
+
+	//~ Begin IAvaFilterSuggestionFactory interface
+	virtual FName GetSuggestionIdentifier() const override { return KeyName; }
+	virtual bool IsSimpleSuggestion() const override { return true; }
+	virtual void AddSuggestion(const TSharedRef<FAvaPlaylistFilterSuggestionPayload>& InPayload) override;
+	virtual bool SupportSuggestionType(EAvaPlaylistSearchListType InSuggestionType) const override;
+	//~ End IAvaFilterSuggestionFactory interface
+};
diff --git a/Engine/Plugins/Experimental/Avalanche/Source/AvalancheMediaEditor/Private/Playlist/Factories/Suggestions/AvaPlaylistFilterNameSuggestionFactory.cpp b/Engine/Plugins/Experimental/Avalanche/Source/AvalancheMediaEditor/Private/Playlist/Factories/Suggestions/AvaPlaylistFilterNameSuggestionFactory.cpp
new file mode 100644
--- /dev/null
+++ b/Engine/Plugins/Experimental/Avalanche/Source/AvalancheMediaEditor/Private/Playlist/Factories/Suggestions/AvaPlaylistFilterNameSuggestionFactory.cpp
@@ -0,0 +1,40 @@
+// Copyright Epic Games, Inc. All Rights Reserved.
+
+#include "AvaPlaylistFilterNameSuggestionFactory.h"
+#include "Playlist/AvalanchePlaylist.h"
+#include "Playlist/Factories/Filters/AvaPlaylistFactoriesUtils.h"
+#include "SAssetSearchBox.h"
+
+#define LOCTEXT_NAMESPACE "AvaPlaylistFilterNameSuggestionFactory"
+
+const FName FAvaPlaylistFilterNameSuggestionFactory::KeyName = FName(TEXT("NAME"));
+
+void FAvaPlaylistFilterNameSuggestionFactory::AddSuggestion(const TSharedRef<FAvaPlaylistFilterSuggestionPayload>& InPayload)
+{
+	const FAvalanchePage& PageItem = UAvalanchePlaylist::GetPageSafe(InPayload->Playlist, InPayload->ItemPageId);
+	if (PageItem.IsValidPage())
+	{
+		const FText NameCategoryLabel = LOCTEXT("NameCategoryLabel", "Ava-Playlist-Name");
+		const FText PageText = PageItem.GetPageDescription();
+
+		FString PageName = TEXT("\"");
+		PageName +=	PageItem.GetPageDescription().ToString();
+		PageName.Append(TEXT("\""));
+
+		FString NameSuggestion = FString::Printf(TEXT("Name=%s"), *PageName);
+		const bool bIsFilterValueValid = InPayload->FilterValue.IsEmpty() || NameSuggestion.Contains(InPayload->FilterValue);
+
+		if (bIsFilterValueValid && !InPayload->FilterCache.Contains(NameSuggestion))
+		{
+			InPayload->FilterCache.Add(NameSuggestion);
+			InPayload->PossibleSuggestions.Add(FAssetSearchBoxSuggestion{MoveTemp(NameSuggestion), PageText, NameCategoryLabel});
+		}
+	}
+}
+
+bool FAvaPlaylistFilterNameSuggestionFactory::SupportSuggestionType(EAvaPlaylistSearchListType InSuggestionType) const
+{
+	return InSuggestionType != EAvaPlaylistSearchListType::None;
+}
+
+#undef LOCTEXT_NAMESPACE
