[MotionDesign]

diff --git a/Engine/Plugins/Experimental/Avalanche/Source/AvalancheMediaEditor/Private/Playlist/Factories/Suggestions/AvaPlaylistFilterNameSuggestionFactory.h b/Engine/Plugins/Experimental/Avalanche/Source/AvalancheMediaEditor/Private/Playlist/Factories/Suggestions/AvaPlaylistFilterNameSuggestionFactory.h
new file mode 100644
--- /dev/null
+++ b/Engine/Plugins/Experimental/Avalanche/Source/AvalancheMediaEditor/Private/Playlist/Factories/Suggestions/AvaPlaylistFilterNameSuggestionFactory.h
@@ -0,0 +1,18 @@
+// Copyright Epic Games, Inc. All Rights Reserved.
+
+#pragma once
+
+#include "Playlist/Factories/Filters/IAvaPlaylistFilterSuggestionFactory.h"
+
+class FAvaPlaylistFilterNameSuggestionFactory : public IAvaPlaylistFilterSuggestionFactory
+{
+public:
+	static const FName KeyName;
+
+	//~ Begin IAvaFilterSuggestionFactory interface
+	virtual FName GetSuggestionIdentifier() const override { return KeyName; }
+	virtual bool IsSimpleSuggestion() const override { return false; }
+	virtual void AddSuggestion(const TSharedRef<FAvaPlaylistFilterSuggestionPayload>& InPayload) override;
+	virtual bool SupportSuggestionType(EAvaPlaylistSearchListType InSuggestionType) const override;
+	//~ End IAvaFilterSuggestionFactory interface
+};
diff --git a/Engine/Plugins/Experimental/Avalanche/Source/AvalancheMediaEditor/Private/Playlist/Factories/Suggestions/AvaPlaylistFilterPathSuggestionFactory.cpp b/Engine/Plugins/Experimental/Avalanche/Source/AvalancheMediaEditor/Private/Playlist/Factories/Suggestions/AvaPlaylistFilterPathSuggestionFactory.cpp
new file mode 100644
--- /dev/null
+++ b/Engine/Plugins/Experimental/Avalanche/Source/AvalancheMediaEditor/Private/Playlist/Factories/Suggestions/AvaPlaylistFilterPathSuggestionFactory.cpp
@@ -0,0 +1,36 @@
+// Copyright Epic Games, Inc. All Rights Reserved.
+
+#include "AvaPlaylistFilterPathSuggestionFactory.h"
+#include "Playlist/AvalanchePlaylist.h"
+#include "Playlist/Factories/Filters/AvaPlaylistFactoriesUtils.h"
+#include "SAssetSearchBox.h"
+
+#define LOCTEXT_NAMESPACE "AvaPlaylistFilterPathSuggestionFactory"
+
+const FName FAvaPlaylistFilterPathSuggestionFactory::KeyName = FName(TEXT("ASSET"));
+
+void FAvaPlaylistFilterPathSuggestionFactory::AddSuggestion(const TSharedRef<FAvaPlaylistFilterSuggestionPayload>& InPayload)
+{
+	const FAvalanchePage& PageItem = UAvalanchePlaylist::GetPageSafe(InPayload->Playlist, InPayload->ItemPageId);
+	if (PageItem.IsValidPage())
+	{
+		const FText PathCategoryLabel = LOCTEXT("PathCategoryLabel", "Ava-Playlist-Asset");
+		const FString AssetName = PageItem.GetAvalancheAssetPath(InPayload->Playlist).GetAssetName();
+
+		FString BlueprintNameSuggestion = FString::Printf(TEXT("Asset=%s"), *AssetName);
+		const bool bIsFilterValueValid = InPayload->FilterValue.IsEmpty() || BlueprintNameSuggestion.Contains(InPayload->FilterValue);
+
+		if (bIsFilterValueValid && !InPayload->FilterCache.Contains(BlueprintNameSuggestion))
+		{
+			InPayload->FilterCache.Add(BlueprintNameSuggestion);
+			InPayload->PossibleSuggestions.Add(FAssetSearchBoxSuggestion{MoveTemp(BlueprintNameSuggestion), FText::FromString(AssetName), PathCategoryLabel});
+		}
+	}
+}
+
+bool FAvaPlaylistFilterPathSuggestionFactory::SupportSuggestionType(EAvaPlaylistSearchListType InSuggestionType) const
+{
+	return InSuggestionType != EAvaPlaylistSearchListType::None;
+}
+
+#undef LOCTEXT_NAMESPACE
diff --git a/Engine/Plugins/Experimental/Avalanche/Source/AvalancheMediaEditor/Private/Playlist/Factories/Suggestions/AvaPlaylistFilterPathSuggestionFactory.h b/Engine/Plugins/Experimental/Avalanche/Source/AvalancheMediaEditor/Private/Playlist/Factories/Suggestions/AvaPlaylistFilterPathSuggestionFactory.h
new file mode 100644
--- /dev/null
+++ b/Engine/Plugins/Experimental/Avalanche/Source/AvalancheMediaEditor/Private/Playlist/Factories/Suggestions/AvaPlaylistFilterPathSuggestionFactory.h
@@ -0,0 +1,18 @@
+// Copyright Epic Games, Inc. All Rights Reserved.
+
+#pragma once
+
+#include "Playlist/Factories/Filters/IAvaPlaylistFilterSuggestionFactory.h"
+
+class FAvaPlaylistFilterPathSuggestionFactory : public IAvaPlaylistFilterSuggestionFactory
+{
+public:
+	static const FName KeyName;
+
+	//~ Begin IAvaFilterSuggestionFactory interface
+	virtual FName GetSuggestionIdentifier() const override { return KeyName; }
+	virtual bool IsSimpleSuggestion() const override { return false; }
+	virtual void AddSuggestion(const TSharedRef<FAvaPlaylistFilterSuggestionPayload>& InPayload) override;
+	virtual bool SupportSuggestionType(EAvaPlaylistSearchListType InSuggestionType) const override;
+	//~ End IAvaFilterSuggestionFactory interface
+};
diff --git a/Engine/Plugins/Experimental/Avalanche/Source/AvalancheMediaEditor/Private/Playlist/Factories/Suggestions/AvaPlaylistFilterStatusSuggestionFactory.cpp b/Engine/Plugins/Experimental/Avalanche/Source/AvalancheMediaEditor/Private/Playlist/Factories/Suggestions/AvaPlaylistFilterStatusSuggestionFactory.cpp
new file mode 100644
--- /dev/null
+++ b/Engine/Plugins/Experimental/Avalanche/Source/AvalancheMediaEditor/Private/Playlist/Factories/Suggestions/AvaPlaylistFilterStatusSuggestionFactory.cpp
@@ -0,0 +1,40 @@
+// Copyright Epic Games, Inc. All Rights Reserved.
+
+#include "AvaPlaylistFilterStatusSuggestionFactory.h"
+#include "Playlist/AvalanchePlaylist.h"
+#include "Playlist/Factories/Filters/AvaPlaylistFactoriesUtils.h"
+#include "SAssetSearchBox.h"
+
+#define LOCTEXT_NAMESPACE "AvaPlaylistFilterStatusSuggestionFactory"
+
+const FName FAvaPlaylistFilterStatusSuggestionFactory::KeyName = FName(TEXT("STATUS"));
+
+void FAvaPlaylistFilterStatusSuggestionFactory::AddSuggestion(const TSharedRef<FAvaPlaylistFilterSuggestionPayload>& InPayload)
+{
+	const FAvalanchePage& PageItem = UAvalanchePlaylist::GetPageSafe(InPayload->Playlist, InPayload->ItemPageId);
+	if (PageItem.IsValidPage())
+	{
+		const FText StatusCategoryLabel = LOCTEXT("StatusCategoryLabel", "Ava-Playlist-Status");
+		TArray<FAvalanchePageStatus> StatusPages = PageItem.GetPageContextualStatuses(InPayload->Playlist);
+		for (FAvalanchePageStatus Status : StatusPages)
+		{
+			const FString StatusName = StaticEnum<EAvalanchePageStatus>()->GetNameStringByValue(static_cast<int32>(Status.Status));
+
+			FString StatusNameSuggestion = FString::Printf(TEXT("Status=%s"), *StatusName);
+			const bool bIsFilterValueValid = InPayload->FilterValue.IsEmpty() || StatusNameSuggestion.Contains(InPayload->FilterValue);
+
+			if (bIsFilterValueValid && !InPayload->FilterCache.Contains(StatusNameSuggestion))
+			{
+				InPayload->FilterCache.Add(StatusNameSuggestion);
+				InPayload->PossibleSuggestions.Add(FAssetSearchBoxSuggestion{MoveTemp(StatusNameSuggestion), FText::FromString(StatusName), StatusCategoryLabel});
+			}
+		}
+	}
+}
+
+bool FAvaPlaylistFilterStatusSuggestionFactory::SupportSuggestionType(EAvaPlaylistSearchListType InSuggestionType) const
+{
+	return InSuggestionType != EAvaPlaylistSearchListType::None;
+}
+
+#undef LOCTEXT_NAMESPACE
diff --git a/Engine/Plugins/Experimental/Avalanche/Source/AvalancheMediaEditor/Private/Playlist/Factories/Suggestions/AvaPlaylistFilterStatusSuggestionFactory.h b/Engine/Plugins/Experimental/Avalanche/Source/AvalancheMediaEditor/Private/Playlist/Factories/Suggestions/AvaPlaylistFilterStatusSuggestionFactory.h
new file mode 100644
--- /dev/null
+++ b/Engine/Plugins/Experimental/Avalanche/Source/AvalancheMediaEditor/Private/Playlist/Factories/Suggestions/AvaPlaylistFilterStatusSuggestionFactory.h
@@ -0,0 +1,18 @@
+// Copyright Epic Games, Inc. All Rights Reserved.
+
+#pragma once
+
+#include "Playlist/Factories/Filters/IAvaPlaylistFilterSuggestionFactory.h"
+
+class FAvaPlaylistFilterStatusSuggestionFactory : public IAvaPlaylistFilterSuggestionFactory
+{
+public:
+	static const FName KeyName;
+
+	//~ Begin IAvaFilterSuggestionFactory interface
+	virtual FName GetSuggestionIdentifier() const override { return KeyName; }
+	virtual bool IsSimpleSuggestion() const override { return false; }
+	virtual void AddSuggestion(const TSharedRef<FAvaPlaylistFilterSuggestionPayload>& InPayload) override;
+	virtual bool SupportSuggestionType(EAvaPlaylistSearchListType InSuggestionType) const override;
+	//~ End IAvaFilterSuggestionFactory interface
+};
diff --git a/Engine/Plugins/Experimental/Avalanche/Source/AvalancheMediaEditor/Private/Playlist/Factories/Suggestions/AvaPlaylistFilterTransitionLayerSuggestionFactory.cpp b/Engine/Plugins/Experimental/Avalanche/Source/AvalancheMediaEditor/Private/Playlist/Factories/Suggestions/AvaPlaylistFilterTransitionLayerSuggestionFactory.cpp
new file mode 100644
--- /dev/null
+++ b/Engine/Plugins/Experimental/Avalanche/Source/AvalancheMediaEditor/Private/Playlist/Factories/Suggestions/AvaPlaylistFilterTransitionLayerSuggestionFactory.cpp
@@ -0,0 +1,36 @@
+// Copyright Epic Games, Inc. All Rights Reserved.
+
+#include "AvaPlaylistFilterTransitionLayerSuggestionFactory.h"
+#include "Playlist/AvalanchePlaylist.h"
+#include "Playlist/Factories/Filters/AvaPlaylistFactoriesUtils.h"
+#include "SAssetSearchBox.h"
+
+#define LOCTEXT_NAMESPACE "AvaPlaylistFilterLayerSuggestionFactory"
+
+const FName FAvaPlaylistFilterTransitionLayerSuggestionFactory::KeyName = FName(TEXT("TRANSITIONLAYER"));
+
+void FAvaPlaylistFilterTransitionLayerSuggestionFactory::AddSuggestion(const TSharedRef<FAvaPlaylistFilterSuggestionPayload>& InPayload)
+{
+	const FAvalanchePage& PageItem = UAvalanchePlaylist::GetPageSafe(InPayload->Playlist, InPayload->ItemPageId);
+	if (PageItem.IsValidPage())
+	{
+		const FText TransitionLayerCategoryLabel = LOCTEXT("TransitionLayerCategoryLabel", "Ava-Playlist-Transition-Layer");
+		const FString PageLayer = PageItem.GetTransitionLayer(InPayload->Playlist).ToString();
+
+		FString TransitionLayerSuggestion = FString::Printf(TEXT("TransitionLayer=%s"), *PageLayer);
+		const bool bIsFilterValueValid = InPayload->FilterValue.IsEmpty() || TransitionLayerSuggestion.Contains(InPayload->FilterValue);
+
+		if (bIsFilterValueValid && !InPayload->FilterCache.Contains(TransitionLayerSuggestion))
+		{
+			InPayload->FilterCache.Add(TransitionLayerSuggestion);
+			InPayload->PossibleSuggestions.Add(FAssetSearchBoxSuggestion{MoveTemp(TransitionLayerSuggestion), FText::FromString(PageLayer), TransitionLayerCategoryLabel});
+		}
+	}
+}
+
+bool FAvaPlaylistFilterTransitionLayerSuggestionFactory::SupportSuggestionType(EAvaPlaylistSearchListType InSuggestionType) const
+{
+	return InSuggestionType != EAvaPlaylistSearchListType::None;
+}
+
+#undef LOCTEXT_NAMESPACE
diff --git a/Engine/Plugins/Experimental/Avalanche/Source/AvalancheMediaEditor/Private/Playlist/Factories/Suggestions/AvaPlaylistFilterTransitionLayerSuggestionFactory.h b/Engine/Plugins/Experimental/Avalanche/Source/AvalancheMediaEditor/Private/Playlist/Factories/Suggestions/AvaPlaylistFilterTransitionLayerSuggestionFactory.h
new file mode 100644
--- /dev/null
+++ b/Engine/Plugins/Experimental/Avalanche/Source/AvalancheMediaEditor/Private/Playlist/Factories/Suggestions/AvaPlaylistFilterTransitionLayerSuggestionFactory.h
@@ -0,0 +1,18 @@
+// Copyright Epic Games, Inc. All Rights Reserved.
+
+#pragma once
+
+#include "Playlist/Factories/Filters/IAvaPlaylistFilterSuggestionFactory.h"
+
+class FAvaPlaylistFilterTransitionLayerSuggestionFactory : public IAvaPlaylistFilterSuggestionFactory
+{
+public:
+	static const FName KeyName;
+
+	//~ Begin IAvaFilterSuggestionFactory interface
+	virtual FName GetSuggestionIdentifier() const override { return KeyName; }
+	virtual bool IsSimpleSuggestion() const override { return false; }
+	virtual void AddSuggestion(const TSharedRef<FAvaPlaylistFilterSuggestionPayload>& InPayload) override;
+	virtual bool SupportSuggestionType(EAvaPlaylistSearchListType InSuggestionType) const override;
+	//~ End IAvaFilterSuggestionFactory interface
+};
diff --git a/Engine/Plugins/Experimental/Avalanche/Source/AvalancheMediaEditor/Private/Playlist/Filters/AvaPageFilterExpressionContext.cpp b/Engine/Plugins/Experimental/Avalanche/Source/AvalancheMediaEditor/Private/Playlist/Filters/AvaPageFilterExpressionContext.cpp
new file mode 100644
--- /dev/null
+++ b/Engine/Plugins/Experimental/Avalanche/Source/AvalancheMediaEditor/Private/Playlist/Filters/AvaPageFilterExpressionContext.cpp
@@ -0,0 +1,89 @@
+// Copyright Epic Games, Inc. All Rights Reserved.
+
+#include "AvaPageFilterExpressionContext.h"
+
+#include "IAvaMediaEditorModule.h"
+#include "Playlist/AvaPlaylistDefines.h"
+#include "Playlist/AvalanchePage.h"
+#include "Playlist/AvalanchePlaylist.h"
+#include "Playlist/Factories/Filters/IAvaPlaylistFilterExpressionFactory.h"
+
+void FAvaPageFilterExpressionContext::SetItem(const FAvalanchePage& InItem,  const UAvalanchePlaylist* InPlaylist, EAvaPlaylistSearchListType InPlaylistSearchListType)
+{
+	ItemPageId = InItem.GetPageId();
+	ItemPlaylist = InPlaylist;
+	PlaylistSearchListType = InPlaylistSearchListType;
+
+	FString NameInfo = InItem.GetPageName();
+	ItemName = *NameInfo;
+
+	const FText IdInfo = FText::AsNumber(InItem.GetPageId(), &UE::AvalanchePlaylist::FEditorMetrics::PageIdFormattingOptions);
+	ItemId = *IdInfo.ToString();
+
+	const FString PathInfo = InItem.GetAvalancheAssetPath(InPlaylist).GetAssetPathString();
+	ItemPath = *PathInfo;
+
+	const TArray<FAvalanchePageStatus> Statuses = InItem.GetPageProgramStatuses(InPlaylist);
+	for (const FAvalanchePageStatus& Status : Statuses)
+	{
+		const FName StatusName = StaticEnum<EAvalanchePageStatus>()->GetNameByValue(static_cast<int32>(Status.Status));
+		ItemStatuses.Add(StatusName);
+	}
+}
+
+void FAvaPageFilterExpressionContext::ClearItem()
+{
+	ItemPageId = FAvalanchePage::InvalidPageId;
+	ItemPlaylist = nullptr;
+	ItemName = NAME_None;
+	ItemId = NAME_None;
+	ItemPath = NAME_None;
+	ItemStatuses.Reset();
+}
+
+bool FAvaPageFilterExpressionContext::TestBasicStringExpression(const FTextFilterString& InValue
+	, const ETextFilterTextComparisonMode InTextComparisonMode) const
+{
+	if (InValue.CompareName(ItemName, InTextComparisonMode))
+	{
+		return true;
+	}
+	if (InValue.CompareName(ItemId, InTextComparisonMode))
+	{
+		return true;
+	}
+	if (InValue.CompareName(ItemPath, ETextFilterTextComparisonMode::Partial))
+	{
+		return true;
+	}
+	for (const FName& Status : ItemStatuses)
+	{
+		if (InValue.CompareName(Status, InTextComparisonMode))
+		{
+			return true;
+		}
+	}
+	return false;
+}
+
+bool FAvaPageFilterExpressionContext::TestComplexExpression(const FName& InKey
+	, const FTextFilterString& InValue
+	, const ETextFilterComparisonOperation InComparisonOperation
+	, const ETextFilterTextComparisonMode InTextComparisonMode) const
+{
+	const IAvaMediaEditorModule& AvaMediaEditorModule = IAvaMediaEditorModule::Get();
+
+	const FAvalanchePage& PageItem = UAvalanchePlaylist::GetPageSafe(ItemPlaylist, ItemPageId);
+
+	if (PageItem.IsValidPage() && AvaMediaEditorModule.CanFilterSupportComparisonOperation(InKey, InComparisonOperation, PlaylistSearchListType))
+	{
+		FAvaPlaylistTextFilterArgs FilterArgs;
+		FilterArgs.ValueToCheck = InValue;
+		FilterArgs.ItemPlaylist = ItemPlaylist;
+		FilterArgs.ComparisonMode = InTextComparisonMode;
+		FilterArgs.ComparisonOperation = InComparisonOperation;
+
+		return AvaMediaEditorModule.FilterExpression(InKey, PageItem, FilterArgs);
+	}
+	return false;
+}
diff --git a/Engine/Plugins/Experimental/Avalanche/Source/AvalancheMediaEditor/Private/Playlist/Filters/AvaPageFilterExpressionContext.h b/Engine/Plugins/Experimental/Avalanche/Source/AvalancheMediaEditor/Private/Playlist/Filters/AvaPageFilterExpressionContext.h
new file mode 100644
--- /dev/null
+++ b/Engine/Plugins/Experimental/Avalanche/Source/AvalancheMediaEditor/Private/Playlist/Filters/AvaPageFilterExpressionContext.h
@@ -0,0 +1,52 @@
+// Copyright Epic Games, Inc. All Rights Reserved.
+
+#pragma once
+
+#include "Containers/UnrealString.h"
+#include "Misc/TextFilterExpressionEvaluator.h"
+#include "Playlist/AvaPlaylistDefines.h"
+#include "Playlist/Factories/Filters/AvaPlaylistFactoriesUtils.h"
+
+class FName;
+class FText;
+class UAvalanchePlaylist;
+struct FAvalanchePage;
+
+/** Expression context to test the given asset data against the current text filter */
+class FAvaPageFilterExpressionContext : public ITextFilterExpressionContext
+{
+public:
+	FAvaPageFilterExpressionContext()
+		: ItemPageId(UE::AvalanchePlaylist::InvalidPageId)
+		, ItemPlaylist(nullptr)
+		, PlaylistSearchListType(EAvaPlaylistSearchListType::None)
+	{}
+
+	void SetItem(const FAvalanchePage& InItem, const UAvalanchePlaylist* InPlaylist, EAvaPlaylistSearchListType InPlaylistSearchListType);
+
+	void ClearItem();
+
+	virtual bool TestBasicStringExpression(const FTextFilterString& InValue
+		, const ETextFilterTextComparisonMode InTextComparisonMode) const override;
+
+	virtual bool TestComplexExpression(const FName& InKey
+		, const FTextFilterString& InValue
+		, const ETextFilterComparisonOperation InComparisonOperation
+		, const ETextFilterTextComparisonMode InTextComparisonMode) const override;
+
+private:
+	/** Pointer to the asset we're currently filtering */
+	int32 ItemPageId;
+
+	const UAvalanchePlaylist* ItemPlaylist;
+
+	EAvaPlaylistSearchListType PlaylistSearchListType;
+
+	FName ItemName;
+
+	FName ItemId;
+
+	FName ItemPath;
+
+	TArray<FName> ItemStatuses;
+};
diff --git a/Engine/Plugins/Experimental/Avalanche/Source/AvalancheMediaEditor/Private/Playlist/Filters/AvaPageTextFilter.cpp b/Engine/Plugins/Experimental/Avalanche/Source/AvalancheMediaEditor/Private/Playlist/Filters/AvaPageTextFilter.cpp
new file mode 100644
--- /dev/null
+++ b/Engine/Plugins/Experimental/Avalanche/Source/AvalancheMediaEditor/Private/Playlist/Filters/AvaPageTextFilter.cpp
@@ -0,0 +1,36 @@
+// Copyright Epic Games, Inc. All Rights Reserved.
+
+#include "AvaPageTextFilter.h"
+#include "AvaPageFilterExpressionContext.h"
+#include "Playlist/Factories/Filters/AvaPlaylistFactoriesUtils.h"
+
+FAvaPageTextFilter::FAvaPageTextFilter()
+	: TextFilterExpressionContext(MakeShared<FAvaPageFilterExpressionContext>())
+	, TextFilterExpressionEvaluator(ETextFilterExpressionEvaluatorMode::Complex)
+{
+}
+
+bool FAvaPageTextFilter::PassesFilter(const FAvalanchePage& InItem) const
+{
+	const bool bMatched = TextFilterExpressionEvaluator.TestTextFilter(*TextFilterExpressionContext);
+	TextFilterExpressionContext->ClearItem();
+	return bMatched;
+}
+
+FText FAvaPageTextFilter::GetFilterText() const
+{
+	return TextFilterExpressionEvaluator.GetFilterText();
+}
+
+void FAvaPageTextFilter::SetFilterText(const FText& InFilterText)
+{
+	if (TextFilterExpressionEvaluator.SetFilterText(InFilterText))
+	{
+		OnChanged().Broadcast();
+	}
+}
+
+void FAvaPageTextFilter::SetItem(const FAvalanchePage& InItem, const UAvalanchePlaylist* InPlaylist, EAvaPlaylistSearchListType InPlaylistSearchListType) const
+{
+	TextFilterExpressionContext->SetItem(InItem, InPlaylist, InPlaylistSearchListType);
+}
diff --git a/Engine/Plugins/Experimental/Avalanche/Source/AvalancheMediaEditor/Private/Playlist/Filters/AvaPageTextFilter.h b/Engine/Plugins/Experimental/Avalanche/Source/AvalancheMediaEditor/Private/Playlist/Filters/AvaPageTextFilter.h
new file mode 100644
--- /dev/null
+++ b/Engine/Plugins/Experimental/Avalanche/Source/AvalancheMediaEditor/Private/Playlist/Filters/AvaPageTextFilter.h
@@ -0,0 +1,40 @@
+// Copyright Epic Games, Inc. All Rights Reserved.
+
+#pragma once
+
+#include "Misc/IFilter.h"
+#include "Misc/TextFilterExpressionEvaluator.h"
+#include "Playlist/AvaPlaylistDefines.h"
+#include "Templates/SharedPointer.h"
+
+class FAvaPageFilterExpressionContext;
+class FText;
+class UAvalanchePlaylist;
+enum class EAvaPlaylistSearchListType : uint8;
+
+class FAvaPageTextFilter
+	: public IFilter<const FAvalanchePage&>
+	, public TSharedFromThis<FAvaPageTextFilter>
+{
+public:
+	FAvaPageTextFilter();
+
+	//~ Begin IFilter
+	virtual FChangedEvent& OnChanged() override { return ChangedEvent; }
+	virtual bool PassesFilter(const FAvalanchePage& InItem) const override;
+	//~ End IFilter
+
+	FText GetFilterText() const;
+
+	void SetFilterText(const FText& InFilterText);
+
+	void SetItem(const FAvalanchePage& InItem, const UAvalanchePlaylist* InPlaylist, EAvaPlaylistSearchListType InPlaylistSearchListType) const;
+
+private:
+	/** Transient context data, used when calling PassesFilter. Kept around to minimize re-allocations between multiple calls to PassesFilter */
+	TSharedRef<FAvaPageFilterExpressionContext> TextFilterExpressionContext;
+
+	FTextFilterExpressionEvaluator TextFilterExpressionEvaluator;
+
+	FChangedEvent ChangedEvent;
+};
diff --git a/Engine/Plugins/Experimental/Avalanche/Source/AvalancheMediaEditor/Private/Playlist/MediaOutputEditorUtils/MediaOutputEditorUtils.cpp b/Engine/Plugins/Experimental/Avalanche/Source/AvalancheMediaEditor/Private/Playlist/MediaOutputEditorUtils/MediaOutputEditorUtils.cpp
new file mode 100644
--- /dev/null
+++ b/Engine/Plugins/Experimental/Avalanche/Source/AvalancheMediaEditor/Private/Playlist/MediaOutputEditorUtils/MediaOutputEditorUtils.cpp
@@ -0,0 +1,260 @@
+// Copyright Epic Games, Inc. All Rights Reserved.
+
+#include "MediaOutputEditorUtils.h"
+#include "Dom/JsonObject.h"
+#include "Serialization/JsonReader.h"
+#include "Serialization/JsonSerializer.h"
+#include "Serialization/JsonWriter.h"
+
+DEFINE_LOG_CATEGORY(LogMediaOutputEditorUtils);
+
+FAvaMediaOutputEditorUtils::FAvaMediaOutputEditorUtils()
+{
+}
+
+FAvaMediaOutputEditorUtils::~FAvaMediaOutputEditorUtils()
+{
+}
+
+FString FAvaMediaOutputEditorUtils::SerializeMediaOutput(const UMediaOutput* InMediaOutput)
+{
+	const TSharedPtr<FJsonObject> OutputObject = MakeShareable(new FJsonObject);
+
+	OutputObject->SetStringField("Class", InMediaOutput->GetClass()->GetPathName());
+	OutputObject->SetStringField("Name", InMediaOutput->GetName());
+
+	TArray<TSharedPtr<FJsonValue>> PropertyArray;
+	for (TFieldIterator<FProperty> Iterator(InMediaOutput->GetClass()); Iterator; ++Iterator)
+	{
+		PropertyArray.Add(MakeShareable(new FJsonValueObject(ParsePropertyInfo(Iterator, InMediaOutput))));
+	}
+
+	const TSharedPtr<FJsonObject> PropertiesObject = MakeShareable(new FJsonObject);
+	OutputObject->SetArrayField("Properties", PropertyArray);
+
+	FString OutputString;
+	const TSharedRef<TJsonWriter<>> JsonWriter = TJsonWriterFactory<>::Create(&OutputString);
+	FJsonSerializer::Serialize(OutputObject.ToSharedRef(), JsonWriter);
+	
+	return OutputString;
+}
+
+void FAvaMediaOutputEditorUtils::EditMediaOutput(UMediaOutput* InMediaOutput, const FString& InDeviceData)
+{
+	const TSharedRef<TJsonReader<TCHAR>> JsonReader = TJsonReaderFactory<TCHAR>::Create(InDeviceData);
+	if (TSharedPtr<FJsonObject> NewDeviceData; FJsonSerializer::Deserialize(JsonReader, NewDeviceData))
+	{
+		for (TArray<TSharedPtr<FJsonValue>> PropertyValues = NewDeviceData->GetArrayField(TEXT("Properties"));
+		     TSharedPtr<FJsonValue> PropertyValue : PropertyValues)
+		{
+			const TSharedPtr<FJsonObject> PropertyObject = PropertyValue.Get()->AsObject();
+
+			for (TFieldIterator<FProperty> PropertyIterator(InMediaOutput->GetClass());
+			     PropertyIterator; ++PropertyIterator)
+			{
+				if (PropertyObject->GetStringField(TEXT("Name")) == PropertyIterator->GetName())
+				{
+					SetProperty(InMediaOutput, PropertyObject, PropertyIterator);
+				}
+			}
+		}
+	}
+}
+
+TSharedPtr<FJsonObject> FAvaMediaOutputEditorUtils::ParsePropertyInfo(TFieldIterator<FProperty> InProperty,
+                                                                       const void* InOwnerObject)
+{
+	const FFieldClass* FieldClass = InProperty->GetClass();
+	if (FieldClass == FIntProperty::StaticClass() || FieldClass == FNameProperty::StaticClass() ||
+		FieldClass == FBoolProperty::StaticClass() || FieldClass == FStrProperty::StaticClass())
+	{
+		return ParseElementaryPropertyInfo(InProperty, InOwnerObject);
+	}
+	if (FieldClass == FEnumProperty::StaticClass())
+	{
+		return ParseEnumPropertyInfo(InProperty, InOwnerObject);
+	}
+	if (FieldClass == FStructProperty::StaticClass())
+	{
+		return ParseStructPropertyInfo(InProperty, InOwnerObject);
+	}
+	return MakeShareable(new FJsonObject);
+}
+
+TSharedPtr<FJsonObject> FAvaMediaOutputEditorUtils::ParseElementaryPropertyInfo(
+	TFieldIterator<FProperty> InProperty, const void* InOwnerObject)
+{
+	const TSharedPtr<FJsonObject> PropertyObject = MakeShareable(new FJsonObject);
+	if (InProperty->GetClass() == FIntProperty::StaticClass())
+	{
+		const FIntProperty* IntProperty = CastField<FIntProperty>(InProperty->GetOwnerProperty());
+		const int32 Value = IntProperty->GetPropertyValue_InContainer(InOwnerObject);
+		PropertyObject->SetStringField("Name", InProperty->GetName());
+		PropertyObject->SetStringField("DisplayName", InProperty->GetDisplayNameText().ToString());
+		PropertyObject->SetStringField("Type", "number");
+		PropertyObject->SetNumberField("Value", Value);
+	}
+	else if (InProperty->GetClass() == FBoolProperty::StaticClass())
+	{
+		const FBoolProperty* BoolProperty = CastField<FBoolProperty>(InProperty->GetOwnerProperty());
+		const bool Value = BoolProperty->GetPropertyValue_InContainer(InOwnerObject);
+		PropertyObject->SetStringField("Name", InProperty->GetName());
+		PropertyObject->SetStringField("DisplayName", InProperty->GetDisplayNameText().ToString());
+		PropertyObject->SetStringField("Type", "boolean");
+		PropertyObject->SetBoolField("Value", Value);
+	}
+	else if (InProperty->GetClass() == FStrProperty::StaticClass())
+	{
+		const FStrProperty* TextProperty = CastField<FStrProperty>(InProperty->GetOwnerProperty());
+		const FString Value = TextProperty->GetPropertyValue_InContainer(InOwnerObject);
+		PropertyObject->SetStringField("Name", InProperty->GetName());
+		PropertyObject->SetStringField("DisplayName", InProperty->GetDisplayNameText().ToString());
+		PropertyObject->SetStringField("Type", "string");
+		PropertyObject->SetStringField("Value", Value);
+	}
+	else if (InProperty->GetClass() == FNameProperty::StaticClass())
+	{
+		const FNameProperty* NameProperty = CastField<FNameProperty>(InProperty->GetOwnerProperty());
+		const FName Value = NameProperty->GetPropertyValue_InContainer(InOwnerObject);
+		PropertyObject->SetStringField("Name", InProperty->GetName());
+		PropertyObject->SetStringField("DisplayName", InProperty->GetDisplayNameText().ToString());
+		PropertyObject->SetStringField("Type", "name");
+		PropertyObject->SetStringField("Value", Value.ToString());
+	}
+	else
+	{
+		UE_LOG(LogMediaOutputEditorUtils, Warning, TEXT("Property type: %s is not supported"), *InProperty->GetClass()->GetName());
+	}
+
+	return PropertyObject;
+}
+
+TSharedPtr<FJsonObject> FAvaMediaOutputEditorUtils::ParseEnumPropertyInfo(
+	TFieldIterator<FProperty> InProperty, const void* InOwnerObject)
+{
+	const TSharedPtr<FJsonObject> PropertyObject = MakeShareable(new FJsonObject);
+	const FEnumProperty* EnumProperty = CastField<FEnumProperty>(InProperty->GetOwnerProperty());
+
+	PropertyObject->SetStringField("Name", InProperty->GetName());
+	PropertyObject->SetStringField("DisplayName", InProperty->GetDisplayNameText().ToString());
+	PropertyObject->SetStringField("Type", "enum");
+	PropertyObject->SetStringField("Class", EnumProperty->GetEnum()->GetName());
+
+	const int32 MaxIndex = EnumProperty->GetEnum()->NumEnums();
+	TArray<TSharedPtr<FJsonValue>> EnumArray;
+	for (int32 Index = 0; Index < MaxIndex; ++Index)
+	{
+		const TSharedPtr<FJsonObject> EnumValueObject = MakeShareable(new FJsonObject);
+		EnumValueObject->SetStringField(FString::FromInt(Index),
+		                                *EnumProperty->GetEnum()->GetAuthoredNameStringByIndex(Index));
+		EnumArray.Add(MakeShareable(new FJsonValueObject(EnumValueObject)));
+		PropertyObject->SetArrayField("EnumValues", EnumArray);
+	}
+	const uint8* Value = EnumProperty->ContainerPtrToValuePtr<uint8>(InOwnerObject);
+
+	PropertyObject->SetStringField("Value", *EnumProperty->GetEnum()->GetAuthoredNameStringByValue(*Value));
+	return PropertyObject;
+}
+
+TSharedPtr<FJsonObject> FAvaMediaOutputEditorUtils::ParseStructPropertyInfo(
+	TFieldIterator<FProperty> InProperty, const void* InOwnerObject)
+{
+	const TSharedPtr<FJsonObject> PropertyObject = MakeShareable(new FJsonObject);
+	const FStructProperty* StructProperty = CastField<FStructProperty>(InProperty->GetOwnerProperty());
+
+	PropertyObject->SetStringField("Name", InProperty->GetName());
+	PropertyObject->SetStringField("DisplayName", InProperty->GetDisplayNameText().ToString());
+	PropertyObject->SetStringField("Type", "struct");
+
+	const UScriptStruct* Struct = StructProperty->Struct;
+	PropertyObject->SetStringField("Class", Struct->GetName());
+
+	const void* StructPtr = StructProperty->ContainerPtrToValuePtr<void>(InOwnerObject);
+
+	TArray<TSharedPtr<FJsonValue>> PropertyArray;
+	for (TFieldIterator<FProperty> StructIterator(Struct); StructIterator; ++StructIterator)
+	{
+		PropertyArray.Add(MakeShareable(new FJsonValueObject(ParsePropertyInfo(StructIterator, StructPtr))));
+	}
+	PropertyObject->SetArrayField("StructProperties", PropertyArray);
+
+	return PropertyObject;
+}
+
+void FAvaMediaOutputEditorUtils::SetProperty(void* InOwnerObject, const TSharedPtr<FJsonObject>& InPropertyObject,
+                                              const TFieldIterator<FProperty>& InProperty)
+{
+	if (InPropertyObject->GetStringField(TEXT("Type")) == "number" || InPropertyObject->GetStringField(TEXT("Type")) == "string"
+		|| InPropertyObject->GetStringField(TEXT("Type")) == "boolean" || InPropertyObject->GetStringField(TEXT("Type")) == "name")
+	{
+		SetElementaryProperty(InOwnerObject, InPropertyObject, InProperty);
+	}
+	else if (InPropertyObject->GetStringField(TEXT("Type")) == "enum")
+	{
+		SetEnumProperty(InOwnerObject, InPropertyObject, InProperty);
+	}
+	else if (InPropertyObject->GetStringField(TEXT("Type")) == "struct")
+	{
+		SetStructProperties(InOwnerObject, InPropertyObject, InProperty);
+	}
+}
+
+void FAvaMediaOutputEditorUtils::SetElementaryProperty(void* InOwnerObject,
+	const TSharedPtr<FJsonObject>& InPropertyObject, TFieldIterator<FProperty> InProperty)
+{
+	if (InPropertyObject->GetStringField(TEXT("Type")) == "number")
+	{
+		const int32 NumberValue = InPropertyObject->GetNumberField(TEXT("Value"));
+		InProperty->GetOwnerProperty()->SetValue_InContainer(InOwnerObject, &NumberValue);
+	}
+	else if (InPropertyObject->GetStringField(TEXT("Type")) == "string")
+	{
+		const FString StringValue = InPropertyObject->GetStringField(TEXT("Value"));
+		InProperty->GetOwnerProperty()->SetValue_InContainer(InOwnerObject, &StringValue);
+	}
+	else if (InPropertyObject->GetStringField(TEXT("Type")) == "boolean")
+	{
+		const bool BoolValue = InPropertyObject->GetBoolField(TEXT("Value"));
+		InProperty->GetOwnerProperty()->SetValue_InContainer(InOwnerObject, &BoolValue);
+	}
+	else if (InPropertyObject->GetStringField(TEXT("Type")) == "name")
+	{
+		const FName NameValue = FName(InPropertyObject->GetStringField(TEXT("Value")));
+		InProperty->GetOwnerProperty()->SetValue_InContainer(InOwnerObject, &NameValue);
+	}
+}
+
+void FAvaMediaOutputEditorUtils::SetEnumProperty(void* InOwnerObject,
+	const TSharedPtr<FJsonObject>& InPropertyObject, TFieldIterator<FProperty> InProperty)
+{
+	const FString ValueName = InPropertyObject->GetStringField(TEXT("Value"));
+
+	const FEnumProperty* EnumProperty = CastField<FEnumProperty>(
+		InProperty->GetOwnerProperty());
+	const int32 EnumValue = EnumProperty->GetEnum()->GetValueByName(FName(ValueName));
+	EnumProperty->GetOwnerProperty()->SetValue_InContainer(InOwnerObject, &EnumValue);
+}
+
+void FAvaMediaOutputEditorUtils::SetStructProperties(void* InOwnerObject,
+	const TSharedPtr<FJsonObject>& InPropertyObject, TFieldIterator<FProperty> InProperty)
+{
+	const FStructProperty* StructProperty = CastField<FStructProperty>(
+		InProperty->GetOwnerProperty());
+
+	const UScriptStruct* Struct = StructProperty->Struct;
+	void* StructPtr = StructProperty->ContainerPtrToValuePtr<void>(InOwnerObject);
+
+	for (TArray<TSharedPtr<FJsonValue>> StructPropertyValues = InPropertyObject->GetArrayField(TEXT("StructProperties"));
+			 TSharedPtr<FJsonValue> StructPropertyValue : StructPropertyValues)
+	{
+		const TSharedPtr<FJsonObject> StructPropertyObject = StructPropertyValue.Get()->AsObject();
+
+		for (TFieldIterator<FProperty> StructIterator(Struct); StructIterator; ++StructIterator)
+		{
+			if (StructPropertyObject->GetStringField(TEXT("Name")) == StructIterator->GetName())
+			{
+				SetProperty(StructPtr, StructPropertyObject, StructIterator);
+			}
+		}
+	}
+}
diff --git a/Engine/Plugins/Experimental/Avalanche/Source/AvalancheMediaEditor/Private/Playlist/MediaOutputEditorUtils/MediaOutputEditorUtils.h b/Engine/Plugins/Experimental/Avalanche/Source/AvalancheMediaEditor/Private/Playlist/MediaOutputEditorUtils/MediaOutputEditorUtils.h
new file mode 100644
--- /dev/null
+++ b/Engine/Plugins/Experimental/Avalanche/Source/AvalancheMediaEditor/Private/Playlist/MediaOutputEditorUtils/MediaOutputEditorUtils.h
@@ -0,0 +1,40 @@
+// Copyright Epic Games, Inc. All Rights Reserved.
+
+#pragma once
+
+#include "MediaOutput.h"
+#include "Templates/SharedPointer.h"
+
+DECLARE_LOG_CATEGORY_EXTERN(LogMediaOutputEditorUtils, Warning, All);
+
+class FJsonObject;
+
+/*
+ *	Class designed to deal with Serialization and Editing
+ *	of Avalanche Media Output device data coming from AvalanchePlaylistServer
+ *
+ *	It is meant to give full control and provide all information to Playlist server
+ *	in a form that is usable outside of unreal engine
+ */
+class FAvaMediaOutputEditorUtils
+{
+public:
+	FAvaMediaOutputEditorUtils();
+	virtual ~FAvaMediaOutputEditorUtils();
+
+	static FString SerializeMediaOutput(const UMediaOutput* InMediaOutput);
+	static void EditMediaOutput(UMediaOutput* InMediaOutput, const FString& InDeviceData);
+
+private:
+	//Serialization helper functions
+	static TSharedPtr<FJsonObject> ParsePropertyInfo(TFieldIterator<FProperty> InProperty, const void* InOwnerObject);
+	static TSharedPtr<FJsonObject> ParseElementaryPropertyInfo(TFieldIterator<FProperty> InProperty, const void* InOwnerObject);
+	static TSharedPtr<FJsonObject> ParseEnumPropertyInfo(TFieldIterator<FProperty> InProperty, const void* InOwnerObject);
+	static TSharedPtr<FJsonObject> ParseStructPropertyInfo(TFieldIterator<FProperty> InProperty, const void* InOwnerObject);
+
+	//Helper functions for editing Avalanche Media Output Device
+	static void SetProperty(void* InOwnerObject, const TSharedPtr<FJsonObject>& InPropertyObject, const TFieldIterator<FProperty>& InProperty);
+	static void SetElementaryProperty(void* InOwnerObject, const TSharedPtr<FJsonObject>& InPropertyObject, TFieldIterator<FProperty> InProperty);
+	static void SetEnumProperty(void* InOwnerObject, const TSharedPtr<FJsonObject>& InPropertyObject, TFieldIterator<FProperty> InProperty);
+	static void SetStructProperties(void* InOwnerObject, const TSharedPtr<FJsonObject>& InPropertyObject, TFieldIterator<FProperty> InProperty);
+};
diff --git a/Engine/Plugins/Experimental/Avalanche/Source/AvalancheMediaEditor/Private/Playlist/Pages/AvaPlaylistPageContextMenu.cpp b/Engine/Plugins/Experimental/Avalanche/Source/AvalancheMediaEditor/Private/Playlist/Pages/AvaPlaylistPageContextMenu.cpp
new file mode 100644
--- /dev/null
+++ b/Engine/Plugins/Experimental/Avalanche/Source/AvalancheMediaEditor/Private/Playlist/Pages/AvaPlaylistPageContextMenu.cpp
@@ -0,0 +1,148 @@
+// Copyright Epic Games, Inc. All Rights Reserved.
+
+#include "AvaPlaylistPageContextMenu.h"
+#include "Framework/Commands/GenericCommands.h"
+#include "IAvaMediaEditorModule.h"
+#include "Playlist/AvaPlaylistCommands.h"
+#include "Playlist/AvaPlaylistEditor.h"
+#include "Playlist/Pages/AvaPlaylistPageContext.h"
+#include "ToolMenu.h"
+#include "ToolMenuSection.h"
+#include "ToolMenus.h"
+
+#define LOCTEXT_NAMESPACE "AvaPlaylistPageContextMenu"
+
+TSharedRef<SWidget> FAvaPlaylistPageContextMenu::GeneratePageContextMenuWidget(const TWeakPtr<FAvaPlaylistEditor>& InPlaylistEditorWeak, const FAvaPageListReference& InPageListReference, const TSharedPtr<FUICommandList>& InCommandList)
+{
+	UToolMenus* const ToolMenus = UToolMenus::Get();
+	check(ToolMenus);
+
+	const FName PageContextMenuName = IAvaMediaEditorModule::GetPlaylistPageMenuName();
+
+	if (!ToolMenus->IsMenuRegistered(PageContextMenuName))
+	{
+		UToolMenu* const ContextMenu = ToolMenus->RegisterMenu(PageContextMenuName, NAME_None, EMultiBoxType::Menu);
+		check(ContextMenu);
+
+		ContextMenu->AddDynamicSection("PopulateContextMenu", FNewToolMenuDelegate::CreateStatic(
+			[](UToolMenu* InMenu)
+			{
+				if (!InMenu)
+				{
+					return;
+				}
+				if (UAvaPlaylistPageContext* PageContext = InMenu->FindContext<UAvaPlaylistPageContext>())
+				{
+					if (TSharedPtr<FAvaPlaylistPageContextMenu> SourceContextMenu = PageContext->ContextMenuWeak.Pin())
+					{
+						SourceContextMenu->PopulatePageContextMenu(*InMenu, *PageContext);
+					}
+				}
+			}));
+	}
+
+	UAvaPlaylistPageContext* const ContextObject = NewObject<UAvaPlaylistPageContext>();
+	check(ContextObject);
+	ContextObject->ContextMenuWeak    = SharedThis(this);
+	ContextObject->PlaylistEditorWeak = InPlaylistEditorWeak;
+	ContextObject->PageListReference  = InPageListReference;
+
+	TSharedPtr<FExtender> Extender;
+
+	// Compatibility with IAvaMediaEditorModule Playlist Menu Extensibility Manager
+	if (InPageListReference.Type == EAvaPageListType::Template && InCommandList.IsValid())
+	{
+		TSharedPtr<FExtensibilityManager> MenuExtensibility = IAvaMediaEditorModule::Get().GetPlaylistMenuExtensibilityManager();
+
+		TSharedPtr<FAvaPlaylistEditor> PlaylistEditor = InPlaylistEditorWeak.Pin();
+
+		if (MenuExtensibility.IsValid() && PlaylistEditor.IsValid())
+		{
+			if (const TArray<UObject*>* EditingObjects = PlaylistEditor->GetObjectsCurrentlyBeingEdited())
+			{
+				Extender = MenuExtensibility->GetAllExtenders(InCommandList.ToSharedRef(), *EditingObjects);
+			}
+		}
+	}
+
+	FToolMenuContext Context(InCommandList, Extender, ContextObject);
+	return ToolMenus->GenerateWidget(PageContextMenuName, Context);
+}
+
+void FAvaPlaylistPageContextMenu::PopulatePageContextMenu(UToolMenu& InMenu, UAvaPlaylistPageContext& InContext)
+{
+	TSharedPtr<FExtensibilityManager> MenuExtensibility;
+
+	const FGenericCommands& GenericCommands = FGenericCommands::Get();
+
+	const FAvaPlaylistCommands& PlaylistCommands = FAvaPlaylistCommands::Get();
+
+	const FAvaPageListReference& PageListReference = InContext.GetPageListReference();
+
+	// Page List Operations
+	{
+		FToolMenuSection& Section = InMenu.FindOrAddSection(TEXT("PageListOperations"), LOCTEXT("PlayListOperationsHeader", "Page List Actions"));
+
+		if (PageListReference.Type == EAvaPageListType::Template)
+		{
+			Section.AddMenuEntry(PlaylistCommands.AddTemplate);
+			Section.AddMenuEntry(PlaylistCommands.CreatePageInstanceFromTemplate);
+			Section.AddMenuEntry(PlaylistCommands.CreateComboTemplate);
+		}
+
+		Section.AddMenuEntry(PlaylistCommands.RemovePage);
+		Section.AddMenuEntry(PlaylistCommands.RenumberPage);
+		Section.AddMenuEntry(GenericCommands.Rename);
+		Section.AddMenuEntry(GenericCommands.Cut);
+		Section.AddMenuEntry(GenericCommands.Copy);
+		Section.AddMenuEntry(GenericCommands.Paste);
+		Section.AddMenuEntry(GenericCommands.Duplicate);
+		Section.AddMenuEntry(PlaylistCommands.ReimportPage);
+		Section.AddMenuEntry(PlaylistCommands.EditPageSource);
+
+		if (PageListReference.Type == EAvaPageListType::Template || PageListReference.Type == EAvaPageListType::Instance)
+		{
+			Section.AddSubMenu(TEXT("ExportPages"),
+				LOCTEXT("ExportPagesSubMenu", "Export Pages ..."),
+				LOCTEXT("ExportPagesSubMenuTooltip", "Export selected pages to a separate resource."),
+				FNewMenuDelegate::CreateLambda([&PlaylistCommands](FMenuBuilder& InMenuBuilder)
+				{
+					// When exporting instanced pages, corresponding templates follow along.
+					InMenuBuilder.AddMenuEntry(PlaylistCommands.ExportPagesToPlaylist);
+					InMenuBuilder.AddMenuEntry(PlaylistCommands.ExportPagesToJson);
+					InMenuBuilder.AddMenuEntry(PlaylistCommands.ExportPagesToXml);
+				}),
+				false,
+				FSlateIcon(FAppStyle::GetAppStyleSetName(), "MainFrame.SaveAll")
+			);
+		}
+	}
+
+	// Show Control Actions
+	if (PageListReference.Type != EAvaPageListType::Template)
+	{
+		FToolMenuSection& Section = InMenu.FindOrAddSection(TEXT("ShowControlOperations"), LOCTEXT("ShowControlOperationsHeader", "Show Control Actions"));
+
+		Section.AddMenuEntry(PlaylistCommands.Play);
+		Section.AddMenuEntry(PlaylistCommands.UpdateValues);
+		Section.AddMenuEntry(PlaylistCommands.Continue);
+		Section.AddMenuEntry(PlaylistCommands.Stop);
+		Section.AddMenuEntry(PlaylistCommands.ForceStop);
+		Section.AddMenuEntry(PlaylistCommands.PlayNext);
+	}
+
+	// Preview Actions
+	{
+		FToolMenuSection& Section = InMenu.FindOrAddSection(TEXT("PreviewOperations"), LOCTEXT("PreviewOperationsHeader", "Preview Actions"));
+
+		Section.AddMenuEntry(PlaylistCommands.PreviewPlay);
+		Section.AddMenuEntry(PlaylistCommands.PreviewFrame);
+		Section.AddMenuEntry(PlaylistCommands.PreviewContinue);
+		Section.AddMenuEntry(PlaylistCommands.PreviewStop);
+		Section.AddMenuEntry(PlaylistCommands.PreviewForceStop);
+		Section.AddMenuEntry(PlaylistCommands.PreviewPlayNext);
+		Section.AddMenuEntry(PlaylistCommands.TakeToProgram);
+	}
+}
+
+#undef LOCTEXT_NAMESPACE
diff --git a/Engine/Plugins/Experimental/Avalanche/Source/AvalancheMediaEditor/Private/Playlist/Pages/AvaPlaylistPageContextMenu.h b/Engine/Plugins/Experimental/Avalanche/Source/AvalancheMediaEditor/Private/Playlist/Pages/AvaPlaylistPageContextMenu.h
new file mode 100644
--- /dev/null
+++ b/Engine/Plugins/Experimental/Avalanche/Source/AvalancheMediaEditor/Private/Playlist/Pages/AvaPlaylistPageContextMenu.h
@@ -0,0 +1,21 @@
+// Copyright Epic Games, Inc. All Rights Reserved.
+
+#pragma once
+
+#include "Templates/SharedPointer.h"
+
+class FAvaPlaylistEditor;
+class FUICommandList;
+class SWidget;
+class UAvaPlaylistPageContext;
+class UToolMenu;
+struct FAvaPageListReference;
+
+class FAvaPlaylistPageContextMenu : public TSharedFromThis<FAvaPlaylistPageContextMenu>
+{
+public:
+	TSharedRef<SWidget> GeneratePageContextMenuWidget(const TWeakPtr<FAvaPlaylistEditor>& InPlaylistEditorWeak, const FAvaPageListReference& InPageListReference, const TSharedPtr<FUICommandList>& InCommandList);
+
+private:
+	void PopulatePageContextMenu(UToolMenu& InMenu, UAvaPlaylistPageContext& InContext);
+};
diff --git a/Engine/Plugins/Experimental/Avalanche/Source/AvalancheMediaEditor/Private/Playlist/Pages/Columns/AvaPageAssetNameColumn.cpp b/Engine/Plugins/Experimental/Avalanche/Source/AvalancheMediaEditor/Private/Playlist/Pages/Columns/AvaPageAssetNameColumn.cpp
new file mode 100644
--- /dev/null
+++ b/Engine/Plugins/Experimental/Avalanche/Source/AvalancheMediaEditor/Private/Playlist/Pages/Columns/AvaPageAssetNameColumn.cpp
@@ -0,0 +1,46 @@
+// Copyright Epic Games, Inc. All Rights Reserved.
+
+#include "AvaPageAssetNameColumn.h"
+#include "Playlist/Pages/PageViews/AvaInstancedPageViewImpl.h"
+#include "Playlist/Pages/PageViews/AvaTemplatePageViewImpl.h"
+#include "Widgets/Layout/SBox.h"
+#include "Widgets/Text/STextBlock.h"
+
+#define LOCTEXT_NAMESPACE "AvaPageAssetNameColumn"
+
+FText FAvaPageAssetNameColumn::GetColumnDisplayNameText() const
+{
+	return LOCTEXT("PageId_AssetName", "Asset");
+}
+
+FText FAvaPageAssetNameColumn::GetColumnToolTipText() const
+{
+	return LOCTEXT("PageId_ToolTip", "Name and Id of the template used to create this page.");
+}
+
+SHeaderRow::FColumn::FArguments FAvaPageAssetNameColumn::ConstructHeaderRowColumn()
+{
+	return SHeaderRow::Column(GetColumnId())
+		.DefaultLabel(GetColumnDisplayNameText())
+		.DefaultTooltip(GetColumnToolTipText())
+		.FillWidth(0.1f)
+		.ShouldGenerateWidget(true)
+		.VAlignCell(EVerticalAlignment::VAlign_Center)
+		;
+}
+
+TSharedRef<SWidget> FAvaPageAssetNameColumn::ConstructRowWidget(const FAvaPageViewRef& InPageView
+	, const TSharedPtr<SAvaPageViewRow>& InRow)
+{
+	const TSharedRef<FAvaPageViewImpl> InstancedPageView = StaticCastSharedRef<FAvaPageViewImpl>(InPageView);
+	const UAvalanchePlaylist* Playlist = InPageView->GetPlaylist();
+
+	return SNew(SBox)
+		.Padding(UE::AvalanchePlaylist::FEditorMetrics::ColumnLeftOffset, 0.f, 0.f, 0.f)
+		[
+			SNew(STextBlock)
+			.Text(InstancedPageView, &FAvaPageViewImpl::GetObjectName, Playlist)
+		];
+}
+
+#undef LOCTEXT_NAMESPACE
diff --git a/Engine/Plugins/Experimental/Avalanche/Source/AvalancheMediaEditor/Private/Playlist/Pages/Columns/AvaPageAssetNameColumn.h b/Engine/Plugins/Experimental/Avalanche/Source/AvalancheMediaEditor/Private/Playlist/Pages/Columns/AvaPageAssetNameColumn.h
new file mode 100644
--- /dev/null
+++ b/Engine/Plugins/Experimental/Avalanche/Source/AvalancheMediaEditor/Private/Playlist/Pages/Columns/AvaPageAssetNameColumn.h
@@ -0,0 +1,18 @@
+// Copyright Epic Games, Inc. All Rights Reserved.
+
+#pragma once
+
+#include "Playlist/Pages/Columns/AvaPageViewColumn.h"
+
+class FAvaPageAssetNameColumn : public IAvaPageViewColumn
+{
+public:
+	UE_AVA_INHERITS(FAvaPageAssetNameColumn, IAvaPageViewColumn);
+
+	//~ Begin IAvaPageViewColumn
+	virtual FText GetColumnDisplayNameText() const override;
+	virtual FText GetColumnToolTipText() const override;
+	virtual SHeaderRow::FColumn::FArguments ConstructHeaderRowColumn() override;
+	virtual TSharedRef<SWidget> ConstructRowWidget(const FAvaPageViewRef& InPageView, const TSharedPtr<SAvaPageViewRow>& InRow) override;
+	//~ End IAvaPageViewColumn
+};
diff --git a/Engine/Plugins/Experimental/Avalanche/Source/AvalancheMediaEditor/Private/Playlist/Pages/Columns/AvaPageAssetSelectorColumn.cpp b/Engine/Plugins/Experimental/Avalanche/Source/AvalancheMediaEditor/Private/Playlist/Pages/Columns/AvaPageAssetSelectorColumn.cpp
new file mode 100644
--- /dev/null
+++ b/Engine/Plugins/Experimental/Avalanche/Source/AvalancheMediaEditor/Private/Playlist/Pages/Columns/AvaPageAssetSelectorColumn.cpp
@@ -0,0 +1,86 @@
+// Copyright Epic Games, Inc. All Rights Reserved.
+
+#include "AvaPageAssetSelectorColumn.h"
+
+#include "AssetRegistry/IAssetRegistry.h"
+#include "AvaBlueprint.h"
+#include "Framework/Application/SlateApplication.h"
+#include "Playlist/Pages/PageViews/AvaPageView.h"
+#include "PropertyCustomizationHelpers.h"
+#include "Widgets/Input/SComboButton.h"
+#include "Widgets/Layout/SBox.h"
+#include "Widgets/Text/STextBlock.h"
+
+#define LOCTEXT_NAMESPACE "AvaPageAssetSelectorColumn"
+
+FText FAvaPageAssetSelectorColumn::GetColumnDisplayNameText() const
+{
+	return LOCTEXT("AssetSelectorColumn_Name", "Motion Design Asset");
+}
+
+FText FAvaPageAssetSelectorColumn::GetColumnToolTipText() const
+{
+	return LOCTEXT("AssetSelectorColumn_ToolTip", "Selects a given Motion Design Asset for the Page");
+}
+
+SHeaderRow::FColumn::FArguments FAvaPageAssetSelectorColumn::ConstructHeaderRowColumn()
+{
+	return SHeaderRow::Column(GetColumnId())
+		.DefaultLabel(GetColumnDisplayNameText())
+		.DefaultTooltip(GetColumnToolTipText())
+		.FillWidth(0.25f)
+		.ShouldGenerateWidget(true)
+		.VAlignCell(EVerticalAlignment::VAlign_Center)
+	;
+}
+
+namespace UE::AvaPageAssetSelectorColumn::Private
+{
+	TSharedRef<SWidget> GetAssetPicker(const FAvaPageViewRef& InPageView)
+	{
+		const UAvalanchePlaylist* Playlist = InPageView->GetPlaylist();
+
+		constexpr bool bAllowClear = true;
+		TArray<const UClass*> AllowedClasses;
+		AllowedClasses.Add(UAvalancheBlueprint::StaticClass());
+		AllowedClasses.Add(UWorld::StaticClass());
+
+		const FAssetData AssetData = IAssetRegistry::Get()->GetAssetByObjectPath(InPageView->GetObjectPath(Playlist));
+	
+		return PropertyCustomizationHelpers::MakeAssetPickerWithMenu(
+			AssetData,
+			bAllowClear,
+			AllowedClasses,
+			PropertyCustomizationHelpers::GetNewAssetFactoriesForClasses(AllowedClasses),
+			FOnShouldFilterAsset(),
+			FOnAssetSelected::CreateSP(InPageView, &IAvaPageView::OnObjectChanged),
+			FSimpleDelegate::CreateLambda([]{FSlateApplication::Get().DismissAllMenus();}));
+	}
+}
+
+TSharedRef<SWidget> FAvaPageAssetSelectorColumn::ConstructRowWidget(const FAvaPageViewRef& InPageView,
+                                                                    const TSharedPtr<SAvaPageViewRow>& InRow)
+{
+	const UAvalanchePlaylist* Playlist = InPageView->GetPlaylist();
+
+	// Combo templates don't have an asset selector.
+	// Put a place holder instead.
+	if (!InPageView->HasObjectPath(Playlist))
+	{
+		return SNew(STextBlock)
+		.Text(LOCTEXT("AssetSelectorColumn_ComboPage", "N/A"))
+		.Font(FAppStyle::GetFontStyle(TEXT("PropertyWindow.NormalFont")));
+	}
+	
+	return SNew(SComboButton)
+			.OnGetMenuContent_Lambda([InPageView]{ return UE::AvaPageAssetSelectorColumn::Private::GetAssetPicker(InPageView);})
+			.ContentPadding(FMargin(2.0f, 2.0f))
+			.ButtonContent()
+			[
+				SNew(STextBlock)
+				.Text(InPageView, &IAvaPageView::GetObjectName, Playlist)
+				.Font(FAppStyle::GetFontStyle(TEXT("PropertyWindow.NormalFont")))
+			];
+}
+
+#undef LOCTEXT_NAMESPACE
diff --git a/Engine/Plugins/Experimental/Avalanche/Source/AvalancheMediaEditor/Private/Playlist/Pages/Columns/AvaPageAssetSelectorColumn.h b/Engine/Plugins/Experimental/Avalanche/Source/AvalancheMediaEditor/Private/Playlist/Pages/Columns/AvaPageAssetSelectorColumn.h
new file mode 100644
--- /dev/null
+++ b/Engine/Plugins/Experimental/Avalanche/Source/AvalancheMediaEditor/Private/Playlist/Pages/Columns/AvaPageAssetSelectorColumn.h
@@ -0,0 +1,18 @@
+// Copyright Epic Games, Inc. All Rights Reserved.
+
+#pragma once
+
+#include "Playlist/Pages/Columns/AvaPageViewColumn.h"
+
+class FAvaPageAssetSelectorColumn : public IAvaPageViewColumn
+{
+public:
+	UE_AVA_INHERITS(FAvaPageAssetSelectorColumn, IAvaPageViewColumn);
+
+	//~ Begin IAvaPageViewColumn
+	virtual FText GetColumnDisplayNameText() const override;
+	virtual FText GetColumnToolTipText() const override;
+	virtual SHeaderRow::FColumn::FArguments ConstructHeaderRowColumn() override;
+	virtual TSharedRef<SWidget> ConstructRowWidget(const FAvaPageViewRef& InPageView, const TSharedPtr<SAvaPageViewRow>& InRow) override;
+	//~ End IAvaPageViewColumn
+};
diff --git a/Engine/Plugins/Experimental/Avalanche/Source/AvalancheMediaEditor/Private/Playlist/Pages/Columns/AvaPageChannelSelectorColumn.cpp b/Engine/Plugins/Experimental/Avalanche/Source/AvalancheMediaEditor/Private/Playlist/Pages/Columns/AvaPageChannelSelectorColumn.cpp
new file mode 100644
--- /dev/null
+++ b/Engine/Plugins/Experimental/Avalanche/Source/AvalancheMediaEditor/Private/Playlist/Pages/Columns/AvaPageChannelSelectorColumn.cpp
@@ -0,0 +1,40 @@
+// Copyright Epic Games, Inc. All Rights Reserved.
+
+#include "AvaPageChannelSelectorColumn.h"
+#include "AvaBlueprint.h"
+#include "PropertyCustomizationHelpers.h"
+#include "Slate/SAvaPageChannelSelector.h"
+
+#define LOCTEXT_NAMESPACE "AvaPageChannelSelectorColumn"
+
+FText FAvaPageChannelSelectorColumn::GetColumnDisplayNameText() const
+{
+	return LOCTEXT("ChannelSelectorColumn_ChannelName", "Channel");
+}
+
+FText FAvaPageChannelSelectorColumn::GetColumnToolTipText() const
+{
+	return LOCTEXT("ChannelSelectorColumn_ToolTip", "Selects a Broadcast Channel for the Page");
+}
+
+SHeaderRow::FColumn::FArguments FAvaPageChannelSelectorColumn::ConstructHeaderRowColumn()
+{
+	return SHeaderRow::Column(GetColumnId())
+		.DefaultLabel(GetColumnDisplayNameText())
+		.DefaultTooltip(GetColumnToolTipText())
+		.FillWidth(0.25f)
+		.ShouldGenerateWidget(true)
+		.VAlignCell(EVerticalAlignment::VAlign_Center)
+	;
+}
+
+TSharedRef<SWidget> FAvaPageChannelSelectorColumn::ConstructRowWidget(const FAvaPageViewRef& InPageView, const TSharedPtr<SAvaPageViewRow>& InRow)
+{
+	return SNew(SBox)
+		.Padding(UE::AvalanchePlaylist::FEditorMetrics::ColumnLeftOffset, 0.f, 0.f, 0.f)
+		[
+			SNew(SAvaPageChannelSelector, InPageView, InRow)
+		];
+}
+
+#undef LOCTEXT_NAMESPACE
diff --git a/Engine/Plugins/Experimental/Avalanche/Source/AvalancheMediaEditor/Private/Playlist/Pages/Columns/AvaPageChannelSelectorColumn.h b/Engine/Plugins/Experimental/Avalanche/Source/AvalancheMediaEditor/Private/Playlist/Pages/Columns/AvaPageChannelSelectorColumn.h
new file mode 100644
--- /dev/null
+++ b/Engine/Plugins/Experimental/Avalanche/Source/AvalancheMediaEditor/Private/Playlist/Pages/Columns/AvaPageChannelSelectorColumn.h
@@ -0,0 +1,18 @@
+// Copyright Epic Games, Inc. All Rights Reserved.
+
+#pragma once
+
+#include "Playlist/Pages/Columns/AvaPageViewColumn.h"
+
+class FAvaPageChannelSelectorColumn : public IAvaPageViewColumn
+{
+public:
+	UE_AVA_INHERITS(FAvaPageChannelSelectorColumn, IAvaPageViewColumn);
+
+	//~ Begin IAvaPageViewColumn
+	virtual FText GetColumnDisplayNameText() const override;
+	virtual FText GetColumnToolTipText() const override;
+	virtual SHeaderRow::FColumn::FArguments ConstructHeaderRowColumn() override;
+	virtual TSharedRef<SWidget> ConstructRowWidget(const FAvaPageViewRef& InPageView, const TSharedPtr<SAvaPageViewRow>& InRow) override;
+	//~ End IAvaPageViewColumn
+};
diff --git a/Engine/Plugins/Experimental/Avalanche/Source/AvalancheMediaEditor/Private/Playlist/Pages/Columns/AvaPageEnabledColumn.cpp b/Engine/Plugins/Experimental/Avalanche/Source/AvalancheMediaEditor/Private/Playlist/Pages/Columns/AvaPageEnabledColumn.cpp
new file mode 100644
--- /dev/null
+++ b/Engine/Plugins/Experimental/Avalanche/Source/AvalancheMediaEditor/Private/Playlist/Pages/Columns/AvaPageEnabledColumn.cpp
@@ -0,0 +1,49 @@
+// Copyright Epic Games, Inc. All Rights Reserved.
+
+#include "AvaPageEnabledColumn.h"
+#include "Playlist/AvalanchePage.h"
+#include "Playlist/Pages/PageViews/AvaInstancedPageViewImpl.h"
+#include "Widgets/Input/SCheckBox.h"
+#include "Widgets/Layout/SBox.h"
+
+#define LOCTEXT_NAMESPACE "AvaPageEnabledColumn"
+
+FText FAvaPageEnabledColumn::GetColumnDisplayNameText() const
+{
+	return LOCTEXT("EnabledColumn_Name", "Enabled");
+}
+
+FText FAvaPageEnabledColumn::GetColumnToolTipText() const
+{
+	return LOCTEXT("EnabledColumn_ToolTip", "Determines whether this Page should be considered for Playlist");
+}
+
+SHeaderRow::FColumn::FArguments FAvaPageEnabledColumn::ConstructHeaderRowColumn()
+{
+	return SHeaderRow::Column(GetColumnId())
+		.DefaultLabel(GetColumnDisplayNameText())
+		.DefaultTooltip(GetColumnToolTipText())
+		.FixedWidth(24.0f)
+		.VAlignCell(EVerticalAlignment::VAlign_Center)
+		.HAlignCell(EHorizontalAlignment::HAlign_Center)
+		.ShouldGenerateWidget(true)
+		[
+			SNew(SCheckBox)
+			.IsChecked(true)
+		];
+}
+
+TSharedRef<SWidget> FAvaPageEnabledColumn::ConstructRowWidget(const FAvaPageViewRef& InPageView, const TSharedPtr<SAvaPageViewRow>& InRow)
+{
+	const TSharedRef<FAvaInstancedPageViewImpl> InstancedPageView = StaticCastSharedRef<FAvaInstancedPageViewImpl>(InPageView);
+
+	return SNew(SBox)
+		.Padding(UE::AvalanchePlaylist::FEditorMetrics::ColumnLeftOffset, 0.f, 0.f, 0.f)
+		[
+			SNew(SCheckBox)
+			.IsChecked(InstancedPageView, &FAvaInstancedPageViewImpl::IsEnabled)
+			.OnCheckStateChanged(InstancedPageView, &FAvaInstancedPageViewImpl::SetEnabled)
+		];
+}
+
+#undef LOCTEXT_NAMESPACE
diff --git a/Engine/Plugins/Experimental/Avalanche/Source/AvalancheMediaEditor/Private/Playlist/Pages/Columns/AvaPageEnabledColumn.h b/Engine/Plugins/Experimental/Avalanche/Source/AvalancheMediaEditor/Private/Playlist/Pages/Columns/AvaPageEnabledColumn.h
new file mode 100644
--- /dev/null
+++ b/Engine/Plugins/Experimental/Avalanche/Source/AvalancheMediaEditor/Private/Playlist/Pages/Columns/AvaPageEnabledColumn.h
@@ -0,0 +1,19 @@
+// Copyright Epic Games, Inc. All Rights Reserved.
+
+#pragma once
+
+#include "Playlist/AvaPlaylistDefines.h"
+#include "Playlist/Pages/Columns/AvaPageViewColumn.h"
+
+class FAvaPageEnabledColumn : public IAvaPageViewColumn
+{
+public:
+	UE_AVA_INHERITS(FAvaPageEnabledColumn, IAvaPageViewColumn);
+
+	//~ Begin IAvaPageViewColumn
+	virtual FText GetColumnDisplayNameText() const override;
+	virtual FText GetColumnToolTipText() const override;
+	virtual SHeaderRow::FColumn::FArguments ConstructHeaderRowColumn() override;
+	virtual TSharedRef<SWidget> ConstructRowWidget(const FAvaPageViewRef& InPageView, const TSharedPtr<SAvaPageViewRow>& InRow) override;
+	//~ End IAvaPageViewColumn
+};
diff --git a/Engine/Plugins/Experimental/Avalanche/Source/AvalancheMediaEditor/Private/Playlist/Pages/Columns/AvaPageIdColumn.cpp b/Engine/Plugins/Experimental/Avalanche/Source/AvalancheMediaEditor/Private/Playlist/Pages/Columns/AvaPageIdColumn.cpp
new file mode 100644
--- /dev/null
+++ b/Engine/Plugins/Experimental/Avalanche/Source/AvalancheMediaEditor/Private/Playlist/Pages/Columns/AvaPageIdColumn.cpp
@@ -0,0 +1,41 @@
+// Copyright Epic Games, Inc. All Rights Reserved.
+
+#include "AvaPageIdColumn.h"
+#include "Playlist/Pages/PageViews/AvaInstancedPageView.h"
+#include "Slate/SAvaPageId.h"
+#include "Widgets/Layout/SBox.h"
+
+#define LOCTEXT_NAMESPACE "AvaPageIdColumn"
+
+FText FAvaPageIdColumn::GetColumnDisplayNameText() const
+{
+	return LOCTEXT("PageId_Name", "Page #");
+}
+
+FText FAvaPageIdColumn::GetColumnToolTipText() const
+{
+	return LOCTEXT("PageId_ToolTip", "Page Id");
+}
+
+SHeaderRow::FColumn::FArguments FAvaPageIdColumn::ConstructHeaderRowColumn()
+{
+	return SHeaderRow::Column(GetColumnId())
+		.DefaultLabel(GetColumnDisplayNameText())
+		.DefaultTooltip(GetColumnToolTipText())
+		.FixedWidth(50.f)
+		.ShouldGenerateWidget(true)
+		.VAlignCell(EVerticalAlignment::VAlign_Center)
+	;
+}
+
+TSharedRef<SWidget> FAvaPageIdColumn::ConstructRowWidget(const FAvaPageViewRef& InPageView
+	, const TSharedPtr<SAvaPageViewRow>& InRow)
+{
+	return SNew(SBox)
+		.Padding(UE::AvalanchePlaylist::FEditorMetrics::ColumnLeftOffset, 0.f, 0.f, 0.f)
+		[
+			SNew(SAvaPageId, InPageView)
+		];
+}
+
+#undef LOCTEXT_NAMESPACE
diff --git a/Engine/Plugins/Experimental/Avalanche/Source/AvalancheMediaEditor/Private/Playlist/Pages/Columns/AvaPageIdColumn.h b/Engine/Plugins/Experimental/Avalanche/Source/AvalancheMediaEditor/Private/Playlist/Pages/Columns/AvaPageIdColumn.h
new file mode 100644
--- /dev/null
+++ b/Engine/Plugins/Experimental/Avalanche/Source/AvalancheMediaEditor/Private/Playlist/Pages/Columns/AvaPageIdColumn.h
@@ -0,0 +1,18 @@
+// Copyright Epic Games, Inc. All Rights Reserved.
+
+#pragma once
+
+#include "Playlist/Pages/Columns/AvaPageViewColumn.h"
+
+class FAvaPageIdColumn : public IAvaPageViewColumn
+{
+public:
+	UE_AVA_INHERITS(FAvaPageIdColumn, IAvaPageViewColumn);
+
+	//~ Begin IAvaPageViewColumn
+	virtual FText GetColumnDisplayNameText() const override;
+	virtual FText GetColumnToolTipText() const override;
+	virtual SHeaderRow::FColumn::FArguments ConstructHeaderRowColumn() override;
+	virtual TSharedRef<SWidget> ConstructRowWidget(const FAvaPageViewRef& InPageView, const TSharedPtr<SAvaPageViewRow>& InRow) override;
+	//~ End IAvaPageViewColumn
+};
diff --git a/Engine/Plugins/Experimental/Avalanche/Source/AvalancheMediaEditor/Private/Playlist/Pages/Columns/AvaPageNameColumn.cpp b/Engine/Plugins/Experimental/Avalanche/Source/AvalancheMediaEditor/Private/Playlist/Pages/Columns/AvaPageNameColumn.cpp
new file mode 100644
--- /dev/null
+++ b/Engine/Plugins/Experimental/Avalanche/Source/AvalancheMediaEditor/Private/Playlist/Pages/Columns/AvaPageNameColumn.cpp
@@ -0,0 +1,42 @@
+// Copyright Epic Games, Inc. All Rights Reserved.
+
+#include "AvaPageNameColumn.h"
+
+#include "Playlist/Pages/PageViews/AvaPageView.h"
+#include "Widgets/Layout/SBox.h"
+#include "Slate/SAvaPageName.h"
+
+#define LOCTEXT_NAMESPACE "AvaPageNameColumn"
+
+FText FAvaPageNameColumn::GetColumnDisplayNameText() const
+{
+	return LOCTEXT("PageName_Name", "Description");
+}
+
+FText FAvaPageNameColumn::GetColumnToolTipText() const
+{
+	return LOCTEXT("PageName_ToolTip", "Description of the given Page. Can be null or repeated");
+}
+
+SHeaderRow::FColumn::FArguments FAvaPageNameColumn::ConstructHeaderRowColumn()
+{
+	return SHeaderRow::Column(GetColumnId())
+		.DefaultLabel(GetColumnDisplayNameText())
+		.DefaultTooltip(GetColumnToolTipText())
+		.FillWidth(0.25f)
+		.ShouldGenerateWidget(true)
+		.VAlignCell(EVerticalAlignment::VAlign_Center)
+	;
+}
+
+TSharedRef<SWidget> FAvaPageNameColumn::ConstructRowWidget(const FAvaPageViewRef& InPageView
+	, const TSharedPtr<SAvaPageViewRow>& InRow)
+{
+	return SNew(SBox)
+		.Padding(UE::AvalanchePlaylist::FEditorMetrics::ColumnLeftOffset, 0.f, 0.f, 0.f)
+		[
+			SNew(SAvaPageName, InPageView, InRow)
+		];
+}
+
+#undef LOCTEXT_NAMESPACE
diff --git a/Engine/Plugins/Experimental/Avalanche/Source/AvalancheMediaEditor/Private/Playlist/Pages/Columns/AvaPageNameColumn.h b/Engine/Plugins/Experimental/Avalanche/Source/AvalancheMediaEditor/Private/Playlist/Pages/Columns/AvaPageNameColumn.h
new file mode 100644
--- /dev/null
+++ b/Engine/Plugins/Experimental/Avalanche/Source/AvalancheMediaEditor/Private/Playlist/Pages/Columns/AvaPageNameColumn.h
@@ -0,0 +1,18 @@
+// Copyright Epic Games, Inc. All Rights Reserved.
+
+#pragma once
+
+#include "Playlist/Pages/Columns/AvaPageViewColumn.h"
+
+class FAvaPageNameColumn : public IAvaPageViewColumn
+{
+public:
+	UE_AVA_INHERITS(FAvaPageNameColumn, IAvaPageViewColumn);
+
+	//~ Begin IAvaPageViewColumn
+	virtual FText GetColumnDisplayNameText() const override;
+	virtual FText GetColumnToolTipText() const override;
+	virtual SHeaderRow::FColumn::FArguments ConstructHeaderRowColumn() override;
+	virtual TSharedRef<SWidget> ConstructRowWidget(const FAvaPageViewRef& InPageView, const TSharedPtr<SAvaPageViewRow>& InRow) override;
+	//~ End IAvaPageViewColumn
+};
diff --git a/Engine/Plugins/Experimental/Avalanche/Source/AvalancheMediaEditor/Private/Playlist/Pages/Columns/AvaPageStatusColumn.cpp b/Engine/Plugins/Experimental/Avalanche/Source/AvalancheMediaEditor/Private/Playlist/Pages/Columns/AvaPageStatusColumn.cpp
new file mode 100644
--- /dev/null
+++ b/Engine/Plugins/Experimental/Avalanche/Source/AvalancheMediaEditor/Private/Playlist/Pages/Columns/AvaPageStatusColumn.cpp
@@ -0,0 +1,414 @@
+// Copyright Epic Games, Inc. All Rights Reserved.
+
+#include "AvaPageStatusColumn.h"
+
+#include "AvalancheBroadcast.h"
+#include "AvaMediaEditorStyle.h"
+#include "Playlist/AvalanchePage.h"
+#include "Playlist/AvalanchePlaylist.h"
+#include "Playlist/Pages/PageViews/AvaInstancedPageViewImpl.h"
+#include "Playlist/Pages/PageViews/AvaPageView.h"
+#include "Styling/StyleColors.h"
+#include "Widgets/Images/SImage.h"
+#include "Widgets/Input/SButton.h"
+#include "Widgets/SBoxPanel.h"
+#include "Widgets/Views/SHeaderRow.h"
+
+#define LOCTEXT_NAMESPACE "AvaPageStatusColumn"
+
+FText FAvaPageStatusColumn::GetColumnDisplayNameText() const
+{
+	return LOCTEXT("PageId_Status", "Status");
+}
+
+FText FAvaPageStatusColumn::GetColumnToolTipText() const
+{
+	return LOCTEXT("PageId_ToolTip", "Page Status");
+}
+
+SHeaderRow::FColumn::FArguments FAvaPageStatusColumn::ConstructHeaderRowColumn()
+{
+	return SHeaderRow::Column(GetColumnId())
+		.DefaultLabel(GetColumnDisplayNameText())
+		.DefaultTooltip(GetColumnToolTipText())
+		.FixedWidth(94.f)
+		.ShouldGenerateWidget(true)
+		.VAlignCell(EVerticalAlignment::VAlign_Center)
+	;
+}
+
+TSharedRef<SWidget> FAvaPageStatusColumn::ConstructRowWidget(const FAvaPageViewRef& InPageView
+	, const TSharedPtr<SAvaPageViewRow>& InRow)
+{
+	const TSharedRef<FAvaInstancedPageViewImpl> InstancedPageView = StaticCastSharedRef<FAvaInstancedPageViewImpl>(
+		StaticCastSharedRef<FAvaPageViewImpl>(InPageView)
+		);
+
+	const TWeakPtr<FAvaInstancedPageViewImpl> InstancedPageViewWeak = InstancedPageView;
+	
+	TSharedRef<SHorizontalBox> ButtonList = SNew(SHorizontalBox)
+		+ SHorizontalBox::Slot()
+		.Padding(UE::AvalanchePlaylist::FEditorMetrics::ColumnLeftOffset, 0.f, 0.f, 0.f)
+		[
+			SNew(SButton)
+			.ButtonStyle(FAvaMediaEditorStyle::Get(), "AvalancheMediaEditor.BorderlessButton")
+			.OnClicked(InPageView, &IAvaPageView::OnPreviewButtonClicked)
+			.IsEnabled(InPageView, &IAvaPageView::CanPreview)
+			.ToolTipText(LOCTEXT("Preview", "Preview\n\n- Click: Preview from start\n- +Shift: Use Preview Frame\n- +Control: Continue"))
+			[
+				SNew(SBorder)
+				.Padding(FMargin(0.f, 0.f))
+				.BorderImage_Static(&FAvaPageStatusColumn::GetPreviewBorderBackgroundImage, InstancedPageViewWeak)
+				.BorderBackgroundColor(FLinearColor(1.0f, 0.0f, 0.0f, 1.0f))
+				[
+					SNew(SImage)
+					.ColorAndOpacity_Static(&FAvaPageStatusColumn::GetIsPreviewingButtonColor, InstancedPageViewWeak)
+					.Image(FAvaMediaEditorStyle::Get().GetBrush("AvalancheMediaEditor.MediaPreviewing"))
+				]
+			]
+		];
+	
+	ButtonList->AddSlot()
+		.Padding(UE::AvalanchePlaylist::FEditorMetrics::ColumnLeftOffset, 0.f, 0.f, 0.f)
+		[
+			SNew(SButton)
+			.ButtonStyle(FAvaMediaEditorStyle::Get(), "AvalancheMediaEditor.BorderlessButton")
+			.OnClicked(InstancedPageView, &FAvaInstancedPageViewImpl::OnAssetStatusButtonClicked)
+			.IsEnabled(InstancedPageView, &FAvaInstancedPageViewImpl::CanChangeAssetStatus)
+			.ToolTipText_Static(&FAvaPageStatusColumn::GetAssetStatusButtonTooltip, InstancedPageViewWeak)
+			[
+				SNew(SImage)
+				.ColorAndOpacity_Static(&FAvaPageStatusColumn::GetAssetStatusButtonColor, InstancedPageViewWeak)
+				.Image(FAvaMediaEditorStyle::Get().GetBrush("AvalancheMediaEditor.MediaAssetStatus"))
+			]
+		];
+
+	ButtonList->AddSlot()
+		.Padding(UE::AvalanchePlaylist::FEditorMetrics::ColumnLeftOffset, 0.f, 0.f, 0.f)
+		[
+			SNew(SBorder)
+			.Padding(FMargin(5.f, 4.f))
+			.BorderImage_Static(&FAvaPageStatusColumn::GetProgramBorderBackgroundImage, InstancedPageViewWeak)
+			.BorderBackgroundColor(FLinearColor(1.0f, 0.0f, 0.0f, 1.0f))
+			.ToolTipText_Static(&FAvaPageStatusColumn::GetTakeInTooltip, InstancedPageViewWeak)
+			[
+				SNew(SImage)
+				.ColorAndOpacity_Static(&FAvaPageStatusColumn::GetIsPlayingButtonColor, InstancedPageViewWeak)
+				.Image(FAvaMediaEditorStyle::Get().GetBrush("AvalancheMediaEditor.MediaPlaying"))
+			]
+		];
+
+	return ButtonList;
+}
+
+FSlateColor FAvaPageStatusColumn::GetIsPreviewingButtonColor(const TWeakPtr<FAvaInstancedPageViewImpl> InPageViewWeak)
+{
+	static const FSlateColor Error(FStyleColors::AccentRed.GetSpecifiedColor());
+	static const FSlateColor Syncing(FStyleColors::AccentOrange.GetSpecifiedColor());
+	static const FSlateColor Loading(FStyleColors::AccentBlue.GetSpecifiedColor());
+	static const FSlateColor NotPreviewing(FStyleColors::AccentGray.GetSpecifiedColor());
+	static const FSlateColor Previewing(FStyleColors::AccentGreen.GetSpecifiedColor());
+	
+	if (const FAvaPageViewPtr PageView = InPageViewWeak.Pin())
+	{
+		UAvalanchePlaylist* Previewlist = PageView->GetPlaylist();
+
+		if (IsValid(Previewlist))
+		{
+			FAvalanchePage& Page = Previewlist->GetPage(PageView->GetPageId());
+
+			if (Page.IsValidPage())
+			{
+				const TArray<FAvalanchePageStatus> Statuses = Page.GetPagePreviewStatuses(Previewlist);
+
+				if (FAvalanchePage::StatusesContainsStatus(Statuses, {EAvalanchePageStatus::Unknown,
+					EAvalanchePageStatus::Error,
+					EAvalanchePageStatus::Missing}))
+				{
+					return Error;
+				}
+
+				if (FAvalanchePage::StatusesContainsStatus(Statuses, {EAvalanchePageStatus::Syncing}))
+				{
+					return Syncing;
+				}
+
+				if (FAvalanchePage::StatusesContainsStatus(Statuses, {EAvalanchePageStatus::Loading}))
+				{
+					return Loading;
+				}
+
+				if (FAvalanchePage::StatusesContainsStatus(Statuses, {EAvalanchePageStatus::Previewing}))
+				{
+					return Previewing;
+				}
+
+				if (FAvalanchePage::StatusesContainsStatus(Statuses, {EAvalanchePageStatus::Available,
+					EAvalanchePageStatus::Loaded,
+					EAvalanchePageStatus::Playing}))
+				{
+					return NotPreviewing;
+				}
+			}
+		}
+	}
+
+	return Error;
+}
+
+const FSlateBrush* FAvaPageStatusColumn::GetPreviewBorderBackgroundImage(const TWeakPtr<FAvaInstancedPageViewImpl> InPageViewWeak)
+{
+	bool bIsPlayHead = false;
+	
+	if (const FAvaPageViewPtr PageView = InPageViewWeak.Pin())
+	{
+		const UAvalanchePlaylist* Playlist = PageView->GetPlaylist();
+
+		if (IsValid(Playlist))
+		{
+			if (Playlist->IsPagePreviewing(PageView->GetPageId()))
+			{
+				if (const FAvaPageListPlaybackContextCollection* ContextCollection = Playlist->GetPageListPlaybackContextCollection())
+				{
+					if (const TSharedPtr<FAvaPageListPlaybackContext> Context = ContextCollection->GetContext(true, Playlist->GetDefaultPreviewChannelName()))
+					{
+						bIsPlayHead = Context->PlayHeadPageId == PageView->GetPageId();
+					}
+				}
+			}
+		}
+	}
+	
+	if (bIsPlayHead)
+	{
+		return FAppStyle::GetBrush("Border");
+	}
+	return FAppStyle::GetBrush("NoBorder");
+}
+
+FSlateColor FAvaPageStatusColumn::GetIsPlayingButtonColor(const TWeakPtr<FAvaInstancedPageViewImpl> InPageViewWeak)
+{
+	static const FSlateColor Error(FStyleColors::AccentRed.GetSpecifiedColor());
+	static const FSlateColor NotPlaying(FStyleColors::AccentGray.GetSpecifiedColor());
+	static const FSlateColor Playing(FStyleColors::AccentGreen.GetSpecifiedColor());
+
+	if (const FAvaPageViewPtr PageView = InPageViewWeak.Pin())
+	{
+		UAvalanchePlaylist* Playlist = PageView->GetPlaylist();
+
+		if (IsValid(Playlist))
+		{
+			FAvalanchePage& Page = Playlist->GetPage(PageView->GetPageId());
+
+			if (Page.IsValidPage())
+			{
+				const TArray<FAvalanchePageStatus> Statuses = Page.GetPageProgramStatuses(Playlist);
+
+				if (FAvalanchePage::StatusesContainsStatus(Statuses, {EAvalanchePageStatus::Playing}))
+				{
+					return Playing;
+				}
+
+				if (FAvalanchePage::StatusesContainsStatus(Statuses, {EAvalanchePageStatus::Unknown,
+					EAvalanchePageStatus::Offline,
+					EAvalanchePageStatus::Missing,
+					EAvalanchePageStatus::Error}))
+				{
+					return Error;
+				}
+			}
+		}
+	}
+
+	return NotPlaying;
+}
+
+const FSlateBrush* FAvaPageStatusColumn::GetProgramBorderBackgroundImage(const TWeakPtr<FAvaInstancedPageViewImpl> InPageViewWeak)
+{
+	bool bIsPlayHead = false;
+	
+	if (const FAvaPageViewPtr PageView = InPageViewWeak.Pin())
+	{
+		const UAvalanchePlaylist* Playlist = PageView->GetPlaylist();
+
+		if (IsValid(Playlist))
+		{
+			if (Playlist->IsPagePlaying(PageView->GetPageId()))
+			{
+				if (const FAvaPageListPlaybackContextCollection* ContextCollection = Playlist->GetPageListPlaybackContextCollection())
+				{
+					if (const TSharedPtr<FAvaPageListPlaybackContext> Context = ContextCollection->GetContext(false, NAME_None))
+					{
+						bIsPlayHead = Context->PlayHeadPageId == PageView->GetPageId();
+					}
+				}
+			}
+		}
+	}
+	
+	if (bIsPlayHead)
+	{
+		return FAppStyle::GetBrush("Border");
+	}
+	return FAppStyle::GetBrush("NoBorder");
+}
+
+FText FAvaPageStatusColumn::GetTakeInTooltip(const TWeakPtr<FAvaInstancedPageViewImpl> InPageViewWeak)
+{
+	// static const FText BaseTooltip = LOCTEXT("TakeIn", "Take In\n\n- Click: Play\n- Control+Click: Continue");
+	static const FText BaseTooltip = LOCTEXT("Status", "Page Status: ");
+	static const FText PlayedTooltip = LOCTEXT("Playing", "Playing");
+	static const FText NotPlayedTooltip = LOCTEXT("Stopped", "Stopped");
+	static const FText ErrorTooltip = LOCTEXT("CantPlay", "**Cannot Take In**");
+	static const FText NoOutputs = LOCTEXT("NoOutputs", "No Outputs Selected");
+
+	if (const FAvaPageViewPtr PageView = InPageViewWeak.Pin())
+	{
+		UAvalanchePlaylist* Playlist = PageView->GetPlaylist();
+
+		if (IsValid(Playlist))
+		{
+			TArray<FText> Texts;
+			Texts.Add(BaseTooltip);
+			bool bAddedErrorTooltip = false;
+
+			// Checks whether it can play based on situation, not status
+			if (!Playlist->CanPlayPage(PageView->GetPageId(), false))
+			{
+				Texts.Add(ErrorTooltip);
+				bAddedErrorTooltip = true;
+			}
+
+			const FAvalanchePage& Page = Playlist->GetPage(PageView->GetPageId());
+
+			// Check actual status
+			if (Page.IsValidPage())
+			{
+				const FAvaOutputChannel& Channel = UAvalancheBroadcast::Get().GetCurrentProfile().GetChannel(Page.GetChannelName());
+
+				if (Channel.GetMediaOutputs().IsEmpty())
+				{
+					if (!bAddedErrorTooltip)
+					{
+						Texts.Add(ErrorTooltip);
+						bAddedErrorTooltip = true;
+					}
+
+					Texts.Add(NoOutputs);
+				}
+				else
+				{
+					if (Playlist->IsPagePlaying(Page))
+					{
+						Texts.Add(PlayedTooltip);
+					}
+					else
+					{
+						Texts.Add(NotPlayedTooltip);
+					}
+				}
+			}
+
+			return FText::Join(LOCTEXT("NewLines", "\n\n"), Texts);
+		}
+	}
+
+	return BaseTooltip;
+}
+
+FSlateColor FAvaPageStatusColumn::GetAssetStatusButtonColor(const TWeakPtr<FAvaInstancedPageViewImpl> InPageViewWeak)
+{
+	static const FSlateColor Error(FStyleColors::AccentRed.GetSpecifiedColor());
+	static const FSlateColor Loading(FStyleColors::AccentBlue.GetSpecifiedColor());
+	static const FSlateColor Available(FStyleColors::AccentYellow.GetSpecifiedColor());
+	static const FSlateColor Loaded(FStyleColors::AccentGreen.GetSpecifiedColor());
+	static const FSlateColor Playing(FStyleColors::AccentOrange.GetSpecifiedColor());
+
+	if (const TSharedPtr<FAvaInstancedPageViewImpl> PageView = InPageViewWeak.Pin())
+	{
+		UAvalanchePlaylist* Playlist = PageView->GetPlaylist();
+
+		if (IsValid(Playlist))
+		{
+			FAvalanchePage& Page = Playlist->GetPage(PageView->GetPageId());
+
+			if (Page.IsValidPage())
+			{
+				const TArray<FAvalanchePageStatus> Statuses = Page.GetPageProgramStatuses(Playlist);
+
+				if (FAvalanchePage::StatusesContainsStatus(Statuses, {EAvalanchePageStatus::Unknown,
+					EAvalanchePageStatus::Error,
+					EAvalanchePageStatus::Missing}))
+				{
+					return Error;
+				}
+
+				if (FAvalanchePage::StatusesContainsStatus(Statuses, {EAvalanchePageStatus::Loading}))
+				{
+					return Loading;
+				}
+
+				if (FAvalanchePage::StatusesContainsStatus(Statuses, {EAvalanchePageStatus::Previewing,
+					EAvalanchePageStatus::Playing}))
+				{
+					return Playing;
+				}
+
+				if (FAvalanchePage::StatusesContainsStatus(Statuses, {EAvalanchePageStatus::Loaded}))
+				{
+					return Loaded;
+				}
+
+				if (FAvalanchePage::StatusesContainsStatus(Statuses, {EAvalanchePageStatus::Available}))
+				{
+					return Available;
+				}
+			}
+		}
+	}
+
+	return Error;
+}
+
+FText FAvaPageStatusColumn::GetAssetStatusButtonTooltip(const TWeakPtr<FAvaInstancedPageViewImpl> InPageViewWeak)
+{
+	if (const TSharedPtr<FAvaInstancedPageViewImpl> PageView = InPageViewWeak.Pin())
+	{
+		UAvalanchePlaylist* Playlist = PageView->GetPlaylist();
+
+		if (IsValid(Playlist))
+		{
+			FAvalanchePage& Page = Playlist->GetPage(PageView->GetPageId());
+
+			if (Page.IsValidPage())
+			{
+				const UEnum* PageStatusEnum = StaticEnum<EAvalanchePageStatus>();
+				const UEnum* ChannelTypeEnum = StaticEnum<EAvaBroadcastChannelType>();
+				const TArray<FAvalanchePageStatus> Statuses = Page.GetPageStatuses(Playlist);
+				TArray<FText> StatusTexts;
+				StatusTexts.Reserve(Statuses.Num());
+
+				for (const FAvalanchePageStatus& Status : Statuses)
+				{
+					FText StatusTypeText = FText::Format(LOCTEXT("StatusType", "{0}: {1}"),
+						ChannelTypeEnum->GetDisplayNameTextByValue(static_cast<int64>(Status.Type)),
+						PageStatusEnum->GetDisplayNameTextByValue(static_cast<int64>(Status.Status)));
+					
+					StatusTexts.Add(MoveTemp(StatusTypeText));
+				}
+
+				if (!StatusTexts.IsEmpty())
+				{
+					return FText::Format(
+						LOCTEXT("StatusInfo", "Status: {0}."),
+						FText::Join(LOCTEXT("Separator", ", "), StatusTexts)
+					);
+				}
+			}
+		}
+	}
+
+	return LOCTEXT("NoStatues", "Status Information Unavailable.");
+}
+
+#undef LOCTEXT_NAMESPACE
diff --git a/Engine/Plugins/Experimental/Avalanche/Source/AvalancheMediaEditor/Private/Playlist/Pages/Columns/AvaPageStatusColumn.h b/Engine/Plugins/Experimental/Avalanche/Source/AvalancheMediaEditor/Private/Playlist/Pages/Columns/AvaPageStatusColumn.h
new file mode 100644
--- /dev/null
+++ b/Engine/Plugins/Experimental/Avalanche/Source/AvalancheMediaEditor/Private/Playlist/Pages/Columns/AvaPageStatusColumn.h
@@ -0,0 +1,33 @@
+// Copyright Epic Games, Inc. All Rights Reserved.
+
+#pragma once
+
+#include "Playlist/Pages/Columns/AvaPageViewColumn.h"
+
+class FAvaInstancedPageViewImpl;
+
+class FAvaPageStatusColumn : public IAvaPageViewColumn
+{
+public:
+	UE_AVA_INHERITS(FAvaPageStatusColumn, IAvaPageViewColumn);
+
+	//~ Begin IAvaPageViewColumn
+	virtual FText GetColumnDisplayNameText() const override;
+	virtual FText GetColumnToolTipText() const override;
+	virtual SHeaderRow::FColumn::FArguments ConstructHeaderRowColumn() override;
+	virtual TSharedRef<SWidget> ConstructRowWidget(const FAvaPageViewRef& InPageView, const TSharedPtr<SAvaPageViewRow>& InRow) override;
+	//~ End IAvaPageViewColumn
+
+	static FSlateColor GetIsPreviewingButtonColor(const TWeakPtr<FAvaInstancedPageViewImpl> InPageViewWeak);
+
+	static const FSlateBrush* GetPreviewBorderBackgroundImage(const TWeakPtr<FAvaInstancedPageViewImpl> InPageViewWeak);
+	
+	static FSlateColor GetIsPlayingButtonColor(const TWeakPtr<FAvaInstancedPageViewImpl> InPageViewWeak);
+
+	static const FSlateBrush* GetProgramBorderBackgroundImage(const TWeakPtr<FAvaInstancedPageViewImpl> InPageViewWeak);
+
+	static FText GetTakeInTooltip(const TWeakPtr<FAvaInstancedPageViewImpl> InPageViewWeak);
+
+	static FSlateColor GetAssetStatusButtonColor(const TWeakPtr<FAvaInstancedPageViewImpl> InPageViewWeak);
+	static FText GetAssetStatusButtonTooltip(const TWeakPtr<FAvaInstancedPageViewImpl> InPageViewWeak);
+};
diff --git a/Engine/Plugins/Experimental/Avalanche/Source/AvalancheMediaEditor/Private/Playlist/Pages/Columns/AvaPageTemplateNameColumn.cpp b/Engine/Plugins/Experimental/Avalanche/Source/AvalancheMediaEditor/Private/Playlist/Pages/Columns/AvaPageTemplateNameColumn.cpp
new file mode 100644
--- /dev/null
+++ b/Engine/Plugins/Experimental/Avalanche/Source/AvalancheMediaEditor/Private/Playlist/Pages/Columns/AvaPageTemplateNameColumn.cpp
@@ -0,0 +1,44 @@
+// Copyright Epic Games, Inc. All Rights Reserved.
+
+#include "AvaPageTemplateNameColumn.h"
+#include "Playlist/Pages/PageViews/AvaInstancedPageViewImpl.h"
+#include "Widgets/Layout/SBox.h"
+#include "Widgets/Text/STextBlock.h"
+
+#define LOCTEXT_NAMESPACE "AvaPageTemplateNameColumn"
+
+FText FAvaPageTemplateNameColumn::GetColumnDisplayNameText() const
+{
+	return LOCTEXT("PageId_TemplateName", "Template");
+}
+
+FText FAvaPageTemplateNameColumn::GetColumnToolTipText() const
+{
+	return LOCTEXT("PageId_ToolTip", "Name and Id of the template used to create this page.");
+}
+
+SHeaderRow::FColumn::FArguments FAvaPageTemplateNameColumn::ConstructHeaderRowColumn()
+{
+	return SHeaderRow::Column(GetColumnId())
+		.DefaultLabel(GetColumnDisplayNameText())
+		.DefaultTooltip(GetColumnToolTipText())
+		.FillWidth(0.1f)
+		.ShouldGenerateWidget(true)
+		.VAlignCell(EVerticalAlignment::VAlign_Center)
+		;
+}
+
+TSharedRef<SWidget> FAvaPageTemplateNameColumn::ConstructRowWidget(const FAvaPageViewRef& InPageView
+	, const TSharedPtr<SAvaPageViewRow>& InRow)
+{
+	const TSharedRef<FAvaInstancedPageViewImpl> InstancedPageView = StaticCastSharedRef<FAvaInstancedPageViewImpl>(InPageView);
+
+	return SNew(SBox)
+		.Padding(UE::AvalanchePlaylist::FEditorMetrics::ColumnLeftOffset, 0.f, 0.f, 0.f)
+		[
+			SNew(STextBlock)
+			.Text(InstancedPageView, &FAvaInstancedPageViewImpl::GetTemplateDescription)
+		];
+}
+
+#undef LOCTEXT_NAMESPACE
diff --git a/Engine/Plugins/Experimental/Avalanche/Source/AvalancheMediaEditor/Private/Playlist/Pages/Columns/AvaPageTemplateNameColumn.h b/Engine/Plugins/Experimental/Avalanche/Source/AvalancheMediaEditor/Private/Playlist/Pages/Columns/AvaPageTemplateNameColumn.h
new file mode 100644
--- /dev/null
+++ b/Engine/Plugins/Experimental/Avalanche/Source/AvalancheMediaEditor/Private/Playlist/Pages/Columns/AvaPageTemplateNameColumn.h
@@ -0,0 +1,18 @@
+// Copyright Epic Games, Inc. All Rights Reserved.
+
+#pragma once
+
+#include "Playlist/Pages/Columns/AvaPageViewColumn.h"
+
+class FAvaPageTemplateNameColumn : public IAvaPageViewColumn
+{
+public:
+	UE_AVA_INHERITS(FAvaPageTemplateNameColumn, IAvaPageViewColumn);
+
+	//~ Begin IAvaPageViewColumn
+	virtual FText GetColumnDisplayNameText() const override;
+	virtual FText GetColumnToolTipText() const override;
+	virtual SHeaderRow::FColumn::FArguments ConstructHeaderRowColumn() override;
+	virtual TSharedRef<SWidget> ConstructRowWidget(const FAvaPageViewRef& InPageView, const TSharedPtr<SAvaPageViewRow>& InRow) override;
+	//~ End IAvaPageViewColumn
+};
diff --git a/Engine/Plugins/Experimental/Avalanche/Source/AvalancheMediaEditor/Private/Playlist/Pages/Columns/AvaPageTemplateStatusColumn.cpp b/Engine/Plugins/Experimental/Avalanche/Source/AvalancheMediaEditor/Private/Playlist/Pages/Columns/AvaPageTemplateStatusColumn.cpp
new file mode 100644
--- /dev/null
+++ b/Engine/Plugins/Experimental/Avalanche/Source/AvalancheMediaEditor/Private/Playlist/Pages/Columns/AvaPageTemplateStatusColumn.cpp
@@ -0,0 +1,317 @@
+// Copyright Epic Games, Inc. All Rights Reserved.
+
+#include "AvaPageTemplateStatusColumn.h"
+
+#include "AvalancheBroadcast.h"
+#include "AvaMediaEditorStyle.h"
+#include "Playlist/AvalanchePage.h"
+#include "Playlist/AvalanchePlaylist.h"
+#include "Playlist/Pages/PageViews/AvaPageView.h"
+#include "Playlist/Pages/PageViews/AvaPageViewImpl.h"
+#include "Playlist/Pages/PageViews/AvaTemplatePageViewImpl.h"
+#include "Styling/StyleColors.h"
+#include "Widgets/Images/SImage.h"
+#include "Widgets/Input/SButton.h"
+#include "Widgets/SBoxPanel.h"
+#include "Widgets/Views/SHeaderRow.h"
+
+#define LOCTEXT_NAMESPACE "AvaPageStatusColumn"
+
+FText FAvaPageTemplateStatusColumn::GetColumnDisplayNameText() const
+{
+	return LOCTEXT("PageId_Status", "Status");
+}
+
+FText FAvaPageTemplateStatusColumn::GetColumnToolTipText() const
+{
+	return LOCTEXT("PageId_ToolTip", "Page Status");
+}
+
+SHeaderRow::FColumn::FArguments FAvaPageTemplateStatusColumn::ConstructHeaderRowColumn()
+{
+	return SHeaderRow::Column(GetColumnId())
+		.DefaultLabel(GetColumnDisplayNameText())
+		.DefaultTooltip(GetColumnToolTipText())
+		.FixedWidth(94.f)
+		.ShouldGenerateWidget(true)
+		.VAlignCell(EVerticalAlignment::VAlign_Center)
+	;
+}
+
+TSharedRef<SWidget> FAvaPageTemplateStatusColumn::ConstructRowWidget(const FAvaPageViewRef& InPageView
+	, const TSharedPtr<SAvaPageViewRow>& InRow)
+{
+	const TSharedRef<FAvaTemplatePageViewImpl> TemplatePageView = StaticCastSharedRef<FAvaTemplatePageViewImpl>(
+		StaticCastSharedRef<FAvaPageViewImpl>(InPageView)
+	);
+
+	const TWeakPtr<FAvaTemplatePageViewImpl> TemplatePageViewWeak = TemplatePageView;
+	
+	TSharedRef<SHorizontalBox> ButtonList = SNew(SHorizontalBox)
+		+ SHorizontalBox::Slot()
+		.Padding(UE::AvalanchePlaylist::FEditorMetrics::ColumnLeftOffset, 0.f, 0.f, 0.f)
+		[
+			SNew(SButton)
+			.ButtonStyle(FAvaMediaEditorStyle::Get(), "AvalancheMediaEditor.BorderlessButton")
+			.OnClicked(TemplatePageView, &FAvaTemplatePageViewImpl::OnPreviewButtonClicked)
+			.IsEnabled(TemplatePageView, &FAvaTemplatePageViewImpl::CanPreview)
+			.ToolTipText(LOCTEXT("Preview", "Preview\n\n- Click: Preview from start\n- +Shift: Use Preview Frame\n- +Control: Continue"))
+			[
+				SNew(SImage)
+				.ColorAndOpacity_Static(&FAvaPageTemplateStatusColumn::GetIsPreviewingButtonColor, TemplatePageViewWeak)
+				.Image(FAvaMediaEditorStyle::Get().GetBrush("AvalancheMediaEditor.MediaPreviewing"))
+			]
+		];
+
+	ButtonList->AddSlot()
+		.Padding(UE::AvalanchePlaylist::FEditorMetrics::ColumnLeftOffset, 0.f, 0.f, 0.f)
+		[
+			SNew(SButton)
+			.ButtonStyle(FAvaMediaEditorStyle::Get(), "AvalancheMediaEditor.BorderlessButton")
+			.OnClicked(TemplatePageView, &FAvaTemplatePageViewImpl::OnAssetStatusButtonClicked)
+			.IsEnabled(TemplatePageView, &FAvaTemplatePageViewImpl::CanChangeAssetStatus)
+			.ToolTipText_Static(&FAvaPageTemplateStatusColumn::GetAssetStatusButtonTooltip, TemplatePageViewWeak)
+			[
+				SNew(SImage)
+				.ColorAndOpacity_Static(&FAvaPageTemplateStatusColumn::GetAssetStatusButtonColor, TemplatePageViewWeak)
+				.Image(FAvaMediaEditorStyle::Get().GetBrush("AvalancheMediaEditor.MediaAssetStatus"))
+			]
+		];
+
+	ButtonList->AddSlot()
+		.Padding(UE::AvalanchePlaylist::FEditorMetrics::ColumnLeftOffset, 0.f, 0.f, 0.f)
+		[
+			SNew(SButton)
+			.ButtonStyle(FAvaMediaEditorStyle::Get(), "AvalancheMediaEditor.BorderlessButton")
+			.OnClicked(TemplatePageView, &FAvaTemplatePageViewImpl::OnSyncStatusButtonClicked)
+			.IsEnabled(TemplatePageView, &FAvaTemplatePageViewImpl::CanChangeSyncStatus)
+			.ToolTipText_Static(&FAvaPageTemplateStatusColumn::GetSyncStatusButtonTooltip, TemplatePageViewWeak)
+			[
+				SNew(SImage)
+				.ColorAndOpacity_Static(&FAvaPageTemplateStatusColumn::GetSyncStatusButtonColor, TemplatePageViewWeak)
+				.Image(FAvaMediaEditorStyle::Get().GetBrush("AvalancheMediaEditor.MediaSyncStatus"))
+			]
+		];
+
+	return ButtonList;
+}
+
+FSlateColor FAvaPageTemplateStatusColumn::GetIsPreviewingButtonColor(const TWeakPtr<FAvaTemplatePageViewImpl> InPageViewWeak)
+{
+	static const FSlateColor Error(FStyleColors::AccentRed.GetSpecifiedColor());
+	static const FSlateColor Syncing(FStyleColors::AccentOrange.GetSpecifiedColor());
+	static const FSlateColor Loading(FStyleColors::AccentBlue.GetSpecifiedColor());
+	static const FSlateColor NotPreviewing(FStyleColors::AccentGray.GetSpecifiedColor());
+	static const FSlateColor Previewing(FStyleColors::AccentGreen.GetSpecifiedColor());
+
+	if (const TSharedPtr<FAvaTemplatePageViewImpl> PageView = InPageViewWeak.Pin())
+	{
+		UAvalanchePlaylist* Previewlist = PageView->GetPlaylist();
+
+		if (IsValid(Previewlist))
+		{
+			FAvalanchePage& Page = Previewlist->GetPage(PageView->GetPageId());
+
+			if (Page.IsValidPage())
+			{
+				const TArray<FAvalanchePageStatus> Statuses = Page.GetPagePreviewStatuses(Previewlist);
+
+				if (FAvalanchePage::StatusesContainsStatus(Statuses, {EAvalanchePageStatus::Unknown, 
+					EAvalanchePageStatus::Error, 
+					EAvalanchePageStatus::Missing}))
+				{
+					return Error;
+				}
+
+				if (FAvalanchePage::StatusesContainsStatus(Statuses, {EAvalanchePageStatus::Syncing}))
+				{
+					return Syncing;
+				}
+
+				if (FAvalanchePage::StatusesContainsStatus(Statuses, {EAvalanchePageStatus::Loading}))
+				{
+					return Loading;
+				}
+
+				if (FAvalanchePage::StatusesContainsStatus(Statuses, {EAvalanchePageStatus::Previewing}))
+				{
+					return Previewing;
+				}
+
+				if (FAvalanchePage::StatusesContainsStatus(Statuses, {EAvalanchePageStatus::Available,
+					EAvalanchePageStatus::Loaded,
+					EAvalanchePageStatus::Playing}))
+				{
+					return NotPreviewing;
+				}
+			}
+		}
+	}
+
+	return Error;
+}
+
+FSlateColor FAvaPageTemplateStatusColumn::GetAssetStatusButtonColor(const TWeakPtr<FAvaTemplatePageViewImpl> InPageViewWeak)
+{
+	static const FSlateColor Error(FStyleColors::AccentRed.GetSpecifiedColor());
+	static const FSlateColor Loading(FStyleColors::AccentBlue.GetSpecifiedColor());
+	static const FSlateColor Available(FStyleColors::AccentYellow.GetSpecifiedColor());
+	static const FSlateColor Loaded(FStyleColors::AccentGreen.GetSpecifiedColor());
+
+	if (const TSharedPtr<FAvaTemplatePageViewImpl> PageView = InPageViewWeak.Pin())
+	{
+		UAvalanchePlaylist* Playlist = PageView->GetPlaylist();
+
+		if (IsValid(Playlist))
+		{
+			FAvalanchePage& Page = Playlist->GetPage(PageView->GetPageId());
+
+			if (Page.IsValidPage())
+			{
+				const TArray<FAvalanchePageStatus> Statuses = Page.GetPagePreviewStatuses(Playlist);
+
+				if (FAvalanchePage::StatusesContainsStatus(Statuses, {EAvalanchePageStatus::Unknown, 
+					EAvalanchePageStatus::Error, 
+					EAvalanchePageStatus::Missing}))
+				{
+					return Error;
+				}
+
+				if (FAvalanchePage::StatusesContainsStatus(Statuses, {EAvalanchePageStatus::Loading}))
+				{
+					return Loading;
+				}
+
+				if (FAvalanchePage::StatusesContainsStatus(Statuses, {EAvalanchePageStatus::Loaded}))
+				{
+					return Loaded;
+				}
+
+				if (FAvalanchePage::StatusesContainsStatus(Statuses, {EAvalanchePageStatus::Available}))
+				{
+					return Available;
+				}
+			}
+		}
+	}
+
+	return Error;
+}
+
+FText FAvaPageTemplateStatusColumn::GetAssetStatusButtonTooltip(const TWeakPtr<FAvaTemplatePageViewImpl> InPageViewWeak)
+{
+	if (const TSharedPtr<FAvaTemplatePageViewImpl> PageView = InPageViewWeak.Pin())
+	{
+		UAvalanchePlaylist* Playlist = PageView->GetPlaylist();
+
+		if (IsValid(Playlist))
+		{
+			FAvalanchePage& Page = Playlist->GetPage(PageView->GetPageId());
+
+			if (Page.IsValidPage())
+			{
+				UEnum* StatusEnum = StaticEnum<EAvalanchePageStatus>();
+				const TArray<FAvalanchePageStatus> Statuses = Page.GetPagePreviewStatuses(Playlist);
+				TArray<FText> StatusTexts;
+
+				for (const FAvalanchePageStatus& Status : Statuses)
+				{
+					StatusTexts.Add(StatusEnum->GetDisplayNameTextByValue(static_cast<int64>(Status.Status)));
+				}
+
+				if (!StatusTexts.IsEmpty())
+				{
+					return FText::Format(
+						LOCTEXT("StatusInfo", "Status: {0}."),
+						FText::Join(LOCTEXT("Separator", ", "), StatusTexts)
+					);
+				}
+			}
+		}
+	}
+
+	return LOCTEXT("NoStatues", "Status Information Unavailable.");
+}
+
+FText FAvaPageTemplateStatusColumn::GetSyncStatusButtonTooltip(const TWeakPtr<FAvaTemplatePageViewImpl> InPageViewWeak)
+{
+	static FText Error = LOCTEXT("Error", "Storm Sync Error");
+	static FText NeedsSync = LOCTEXT("NeedsSync", "Storm Sync Required");
+	static FText Syncing = LOCTEXT("Syncing", "Storm Syncing");
+	static FText Synced = LOCTEXT("Synced", "Storm Synced");
+
+	if (const TSharedPtr<FAvaTemplatePageViewImpl> PageView = InPageViewWeak.Pin())
+	{
+		UAvalanchePlaylist* Playlist = PageView->GetPlaylist();
+
+		if (IsValid(Playlist))
+		{
+			FAvalanchePage& Page = Playlist->GetPage(PageView->GetPageId());
+
+			if (Page.IsValidPage())
+			{
+				const TArray<FAvalanchePageStatus> Statuses = Page.GetPagePreviewStatuses(Playlist);
+
+				for (const FAvalanchePageStatus& Status : Statuses)
+				{
+					if (Status.bNeedsSync)
+					{
+						return NeedsSync;
+					}
+				}
+
+				if (FAvalanchePage::StatusesContainsStatus(Statuses, {EAvalanchePageStatus::Syncing}))
+				{
+					return Syncing;
+				}
+
+				return Synced;
+			}
+		}
+	}
+
+	return Error;
+}
+
+FSlateColor FAvaPageTemplateStatusColumn::GetSyncStatusButtonColor(
+	const TWeakPtr<FAvaTemplatePageViewImpl> InPageViewWeak)
+{
+	static const FSlateColor Error(FStyleColors::AccentRed.GetSpecifiedColor());
+	static const FSlateColor NeedsSync(FStyleColors::AccentYellow.GetSpecifiedColor());
+	static const FSlateColor Syncing(FStyleColors::AccentBlue.GetSpecifiedColor());
+	static const FSlateColor Synced(FStyleColors::AccentGreen.GetSpecifiedColor());
+
+	if (const TSharedPtr<FAvaTemplatePageViewImpl> PageView = InPageViewWeak.Pin())
+	{
+		UAvalanchePlaylist* Playlist = PageView->GetPlaylist();
+
+		if (IsValid(Playlist))
+		{
+			FAvalanchePage& Page = Playlist->GetPage(PageView->GetPageId());
+
+			if (Page.IsValidPage())
+			{
+				const TArray<FAvalanchePageStatus> Statuses = Page.GetPagePreviewStatuses(Playlist);
+
+				for (const FAvalanchePageStatus& Status : Statuses)
+				{
+					if (Status.bNeedsSync)
+					{
+						return NeedsSync;
+					}
+				}
+
+				if (FAvalanchePage::StatusesContainsStatus(Statuses, {EAvalanchePageStatus::Syncing}))
+				{
+					return Syncing;
+				}
+
+				return Synced;
+			}
+		}
+	}
+
+	return Error;
+}
+
+#undef LOCTEXT_NAMESPACE
diff --git a/Engine/Plugins/Experimental/Avalanche/Source/AvalancheMediaEditor/Private/Playlist/Pages/Columns/AvaPageTemplateStatusColumn.h b/Engine/Plugins/Experimental/Avalanche/Source/AvalancheMediaEditor/Private/Playlist/Pages/Columns/AvaPageTemplateStatusColumn.h
new file mode 100644
--- /dev/null
+++ b/Engine/Plugins/Experimental/Avalanche/Source/AvalancheMediaEditor/Private/Playlist/Pages/Columns/AvaPageTemplateStatusColumn.h
@@ -0,0 +1,28 @@
+// Copyright Epic Games, Inc. All Rights Reserved.
+
+#pragma once
+
+#include "Playlist/Pages/Columns/AvaPageViewColumn.h"
+
+class FAvaTemplatePageViewImpl;
+
+class FAvaPageTemplateStatusColumn : public IAvaPageViewColumn
+{
+public:
+	UE_AVA_INHERITS(FAvaPageTemplateStatusColumn, IAvaPageViewColumn);
+
+	//~ Begin IAvaPageViewColumn
+	virtual FText GetColumnDisplayNameText() const override;
+	virtual FText GetColumnToolTipText() const override;
+	virtual SHeaderRow::FColumn::FArguments ConstructHeaderRowColumn() override;
+	virtual TSharedRef<SWidget> ConstructRowWidget(const FAvaPageViewRef& InPageView, const TSharedPtr<SAvaPageViewRow>& InRow) override;
+	//~ End IAvaPageViewColumn
+
+	static FSlateColor GetIsPreviewingButtonColor(const TWeakPtr<FAvaTemplatePageViewImpl> InPageViewWeak);
+
+	static FSlateColor GetAssetStatusButtonColor(const TWeakPtr<FAvaTemplatePageViewImpl> InPageViewWeak);
+	static FText GetAssetStatusButtonTooltip(const TWeakPtr<FAvaTemplatePageViewImpl> InPageViewWeak);
+
+	static FText GetSyncStatusButtonTooltip(const TWeakPtr<FAvaTemplatePageViewImpl> InPageViewWeak);
+	static FSlateColor GetSyncStatusButtonColor(const TWeakPtr<FAvaTemplatePageViewImpl> InPageViewWeak);
+};
diff --git a/Engine/Plugins/Experimental/Avalanche/Source/AvalancheMediaEditor/Private/Playlist/Pages/Columns/AvaPageThumbnailColumn.cpp b/Engine/Plugins/Experimental/Avalanche/Source/AvalancheMediaEditor/Private/Playlist/Pages/Columns/AvaPageThumbnailColumn.cpp
new file mode 100644
--- /dev/null
+++ b/Engine/Plugins/Experimental/Avalanche/Source/AvalancheMediaEditor/Private/Playlist/Pages/Columns/AvaPageThumbnailColumn.cpp
@@ -0,0 +1,40 @@
+// Copyright Epic Games, Inc. All Rights Reserved.
+
+#include "AvaPageThumbnailColumn.h"
+#include "Slate/SAvaPageThumbnail.h"
+#include "Widgets/Layout/SBox.h"
+
+#define LOCTEXT_NAMESPACE "AvaPageThumbnailColumn"
+
+FText FAvaPageThumbnailColumn::GetColumnDisplayNameText() const
+{
+	return LOCTEXT("PageThumbnail_Name", "Thumbnail");
+}
+
+FText FAvaPageThumbnailColumn::GetColumnToolTipText() const
+{
+	return LOCTEXT("PageThumbnail_Tooltip", "Preview thumbnail of the page, if available");
+}
+
+SHeaderRow::FColumn::FArguments FAvaPageThumbnailColumn::ConstructHeaderRowColumn()
+{
+	return SHeaderRow::Column(GetColumnId())
+		.DefaultLabel(GetColumnDisplayNameText())
+		.DefaultTooltip(GetColumnToolTipText())
+		.FixedWidth(70.f)
+		.ShouldGenerateWidget(true)
+		.VAlignCell(EVerticalAlignment::VAlign_Center);
+}
+
+TSharedRef<SWidget> FAvaPageThumbnailColumn::ConstructRowWidget(const FAvaPageViewRef& InPageView
+	, const TSharedPtr<SAvaPageViewRow>& InRow)
+{
+	return SNew(SBox)
+		.Padding(UE::AvalanchePlaylist::FEditorMetrics::ColumnLeftOffset, 0.f, 0.f, 0.f)
+		[
+			SNew(SAvaPageThumbnail, InPageView, InRow)
+			.ThumbnailWidgetSize(64)
+		];
+}
+
+#undef LOCTEXT_NAMESPACE
diff --git a/Engine/Plugins/Experimental/Avalanche/Source/AvalancheMediaEditor/Private/Playlist/Pages/Columns/AvaPageThumbnailColumn.h b/Engine/Plugins/Experimental/Avalanche/Source/AvalancheMediaEditor/Private/Playlist/Pages/Columns/AvaPageThumbnailColumn.h
new file mode 100644
--- /dev/null
+++ b/Engine/Plugins/Experimental/Avalanche/Source/AvalancheMediaEditor/Private/Playlist/Pages/Columns/AvaPageThumbnailColumn.h
@@ -0,0 +1,18 @@
+// Copyright Epic Games, Inc. All Rights Reserved.
+
+#pragma once
+
+#include "Playlist/Pages/Columns/AvaPageViewColumn.h"
+
+class FAvaPageThumbnailColumn : public IAvaPageViewColumn
+{
+public:
+	UE_AVA_INHERITS(FAvaPageThumbnailColumn, IAvaPageViewColumn);
+
+	//~ Begin IAvaPageViewColumn
+	virtual FText GetColumnDisplayNameText() const override;
+	virtual FText GetColumnToolTipText() const override;
+	virtual SHeaderRow::FColumn::FArguments ConstructHeaderRowColumn() override;
+	virtual TSharedRef<SWidget> ConstructRowWidget(const FAvaPageViewRef& InPageView, const TSharedPtr<SAvaPageViewRow>& InRow) override;
+	//~ End IAvaPageViewColumn
+};
diff --git a/Engine/Plugins/Experimental/Avalanche/Source/AvalancheMediaEditor/Private/Playlist/Pages/Columns/AvaPageTransitionLayerColumn.cpp b/Engine/Plugins/Experimental/Avalanche/Source/AvalancheMediaEditor/Private/Playlist/Pages/Columns/AvaPageTransitionLayerColumn.cpp
new file mode 100644
--- /dev/null
+++ b/Engine/Plugins/Experimental/Avalanche/Source/AvalancheMediaEditor/Private/Playlist/Pages/Columns/AvaPageTransitionLayerColumn.cpp
@@ -0,0 +1,41 @@
+// Copyright Epic Games, Inc. All Rights Reserved.
+
+#include "AvaPageTransitionLayerColumn.h"
+#include "Playlist/Pages/PageViews/AvaPageView.h"
+#include "Widgets/Layout/SBox.h"
+#include "Widgets/Text/STextBlock.h"
+
+#define LOCTEXT_NAMESPACE "AvaPageTransitionLayerColumn"
+
+FText FAvaPageTransitionLayerColumn::GetColumnDisplayNameText() const
+{
+	return LOCTEXT("TransitionLayerColumn_LayerName", "Layer");
+}
+
+FText FAvaPageTransitionLayerColumn::GetColumnToolTipText() const
+{
+	return LOCTEXT("TransitionLayerColumn_ToolTip", "Transition layer name for the page");
+}
+
+SHeaderRow::FColumn::FArguments FAvaPageTransitionLayerColumn::ConstructHeaderRowColumn()
+{
+	return SHeaderRow::Column(GetColumnId())
+		.DefaultLabel(GetColumnDisplayNameText())
+		.DefaultTooltip(GetColumnToolTipText())
+		.FillWidth(0.25f)
+		.ShouldGenerateWidget(true)
+		.VAlignCell(EVerticalAlignment::VAlign_Center)
+	;
+}
+
+TSharedRef<SWidget> FAvaPageTransitionLayerColumn::ConstructRowWidget(const FAvaPageViewRef& InPageView, const TSharedPtr<SAvaPageViewRow>& InRow)
+{
+	return SNew(SBox)
+		.Padding(UE::AvalanchePlaylist::FEditorMetrics::ColumnLeftOffset, 0.f, 0.f, 0.f)
+		[
+			SNew(STextBlock)
+			.Text(InPageView, &IAvaPageView::GetPageTransitionLayerNameText)
+		];
+}
+
+#undef LOCTEXT_NAMESPACE
diff --git a/Engine/Plugins/Experimental/Avalanche/Source/AvalancheMediaEditor/Private/Playlist/Pages/Columns/AvaPageTransitionLayerColumn.h b/Engine/Plugins/Experimental/Avalanche/Source/AvalancheMediaEditor/Private/Playlist/Pages/Columns/AvaPageTransitionLayerColumn.h
new file mode 100644
--- /dev/null
+++ b/Engine/Plugins/Experimental/Avalanche/Source/AvalancheMediaEditor/Private/Playlist/Pages/Columns/AvaPageTransitionLayerColumn.h
@@ -0,0 +1,18 @@
+// Copyright Epic Games, Inc. All Rights Reserved.
+
+#pragma once
+
+#include "Playlist/Pages/Columns/AvaPageViewColumn.h"
+
+class FAvaPageTransitionLayerColumn : public IAvaPageViewColumn
+{
+public:
+	UE_AVA_INHERITS(FAvaPageTransitionLayerColumn, IAvaPageViewColumn);
+
+	//~ Begin IAvaPageViewColumn
+	virtual FText GetColumnDisplayNameText() const override;
+	virtual FText GetColumnToolTipText() const override;
+	virtual SHeaderRow::FColumn::FArguments ConstructHeaderRowColumn() override;
+	virtual TSharedRef<SWidget> ConstructRowWidget(const FAvaPageViewRef& InPageView, const TSharedPtr<SAvaPageViewRow>& InRow) override;
+	//~ End IAvaPageViewColumn
+};
diff --git a/Engine/Plugins/Experimental/Avalanche/Source/AvalancheMediaEditor/Private/Playlist/Pages/Columns/AvaPageViewColumn.cpp b/Engine/Plugins/Experimental/Avalanche/Source/AvalancheMediaEditor/Private/Playlist/Pages/Columns/AvaPageViewColumn.cpp
new file mode 100644
--- /dev/null
+++ b/Engine/Plugins/Experimental/Avalanche/Source/AvalancheMediaEditor/Private/Playlist/Pages/Columns/AvaPageViewColumn.cpp
@@ -0,0 +1,3 @@
+// Copyright Epic Games, Inc. All Rights Reserved.
+
+#include "Playlist/Pages/Columns/AvaPageViewColumn.h"
