[MotionDesign]

diff --git a/Engine/Plugins/Experimental/Avalanche/Source/AvalancheMediaEditor/Private/Playlist/Pages/Columns/Slate/SAvaPageChannelSelector.cpp b/Engine/Plugins/Experimental/Avalanche/Source/AvalancheMediaEditor/Private/Playlist/Pages/Columns/Slate/SAvaPageChannelSelector.cpp
new file mode 100644
--- /dev/null
+++ b/Engine/Plugins/Experimental/Avalanche/Source/AvalancheMediaEditor/Private/Playlist/Pages/Columns/Slate/SAvaPageChannelSelector.cpp
@@ -0,0 +1,90 @@
+// Copyright Epic Games, Inc. All Rights Reserved.
+
+#include "SAvaPageChannelSelector.h"
+#include "AvalancheBroadcast.h"
+#include "Playlist/Pages/PageViews/AvaInstancedPageView.h"
+#include "Playlist/Pages/PageViews/AvaInstancedPageViewImpl.h"
+#include "Playlist/Pages/PageViews/AvaPageViewImpl.h"
+#include "Playlist/Pages/Slate/SAvaPageViewRow.h"
+#include "Widgets/Text/SInlineEditableTextBlock.h"
+
+void SAvaPageChannelSelector::Construct(const FArguments& InArgs, const FAvaPageViewRef& InPageView, const TSharedPtr<SAvaPageViewRow>& InRow)
+{
+	UpdateChannelNames();
+
+	const TSharedRef<IAvaInstancedPageView> InstancedPageView = StaticCastSharedRef<IAvaInstancedPageView>(
+		StaticCastSharedRef<FAvaInstancedPageViewImpl>(
+			StaticCastSharedRef<FAvaPageViewImpl>(InPageView)
+		)
+	);
+	
+	PageViewWeak = InstancedPageView;
+	PageViewRowWeak = InRow;
+	
+	ChildSlot
+	[
+		SAssignNew(ChannelCombo, SComboBox<FName>)
+		.InitiallySelectedItem(InstancedPageView->GetChannelName())
+		.OptionsSource(&ChannelNames)
+		.OnGenerateWidget(this, &SAvaPageChannelSelector::GenerateChannelWidget)
+		.OnSelectionChanged(this, &SAvaPageChannelSelector::OnChannelSelectionChanged)
+		.OnComboBoxOpening(this, &SAvaPageChannelSelector::OnComboBoxOpening)
+		[
+			SNew(STextBlock)
+			.Text(this, &SAvaPageChannelSelector::GetCurrentChannelName)
+		]
+	];
+}
+
+TSharedRef<SWidget> SAvaPageChannelSelector::GenerateChannelWidget(FName InChannelName)
+{
+	return SNew(STextBlock)
+		.Text(FText::FromName(InChannelName));
+}
+
+void SAvaPageChannelSelector::OnChannelSelectionChanged(SComboBox<FName>::NullableOptionType InProposedSelection, ESelectInfo::Type InSelectInfo)
+{
+	if (const FAvaInstancedPageViewPtr PageView = PageViewWeak.Pin())
+	{
+		PageView->SetChannel(InProposedSelection);
+	}
+}
+
+FText SAvaPageChannelSelector::GetCurrentChannelName() const
+{
+	if (const FAvaInstancedPageViewPtr PageView = PageViewWeak.Pin())
+	{
+		return FText::FromName(PageView->GetChannelName());
+	}
+	return FText();
+}
+
+void SAvaPageChannelSelector::OnComboBoxOpening()
+{
+	UpdateChannelNames();
+	
+	const FAvaInstancedPageViewPtr PageView = PageViewWeak.Pin();
+
+	if (!PageView.IsValid())
+	{
+		return;
+	}
+	
+	check(ChannelCombo.IsValid());
+	ChannelCombo->SetSelectedItem(PageView->GetChannelName());
+}
+
+void SAvaPageChannelSelector::UpdateChannelNames()
+{
+	const UAvalancheBroadcast& AvalancheBroadcast = UAvalancheBroadcast::Get();
+	ChannelNames.Reset(AvalancheBroadcast.GetChannelNameCount());
+	
+	for (int32 ChannelIndex = 0; ChannelIndex < AvalancheBroadcast.GetChannelNameCount(); ++ChannelIndex)
+	{
+		const FName ChannelName = AvalancheBroadcast.GetChannelName(ChannelIndex); 
+		if (AvalancheBroadcast.GetChannelType(ChannelName) == EAvaBroadcastChannelType::Program)
+		{
+			ChannelNames.Add(ChannelName);
+		}
+	}
+}
diff --git a/Engine/Plugins/Experimental/Avalanche/Source/AvalancheMediaEditor/Private/Playlist/Pages/Columns/Slate/SAvaPageChannelSelector.h b/Engine/Plugins/Experimental/Avalanche/Source/AvalancheMediaEditor/Private/Playlist/Pages/Columns/Slate/SAvaPageChannelSelector.h
new file mode 100644
--- /dev/null
+++ b/Engine/Plugins/Experimental/Avalanche/Source/AvalancheMediaEditor/Private/Playlist/Pages/Columns/Slate/SAvaPageChannelSelector.h
@@ -0,0 +1,42 @@
+// Copyright Epic Games, Inc. All Rights Reserved.
+
+#pragma once
+
+#include "CoreMinimal.h"
+#include "Playlist/AvaPlaylistDefines.h"
+#include "Widgets/Input/SComboBox.h"
+#include "Widgets/SCompoundWidget.h"
+
+class IAvaInstancedPageView;
+class SAvaPageViewRow;
+
+class SAvaPageChannelSelector : public SCompoundWidget
+{
+public:
+	SLATE_BEGIN_ARGS(SAvaPageChannelSelector){}
+	SLATE_END_ARGS()
+
+	/** Constructs this widget with InArgs */
+	void Construct(const FArguments& InArgs, const FAvaPageViewRef& InPageView, const TSharedPtr<SAvaPageViewRow>& InRow);
+
+	TSharedRef<SWidget> GenerateChannelWidget(FName InChannelName);
+	
+	void OnChannelSelectionChanged(SComboBox<FName>::NullableOptionType InProposedSelection, ESelectInfo::Type InSelectInfo);
+	
+	FText GetCurrentChannelName() const;
+	
+	void OnComboBoxOpening();
+	
+protected:
+	void UpdateChannelNames();
+
+	TWeakPtr<IAvaInstancedPageView> PageViewWeak;
+	
+	TWeakPtr<SAvaPageViewRow> PageViewRowWeak;
+	
+	TSharedPtr<SComboBox<FName>> ChannelCombo;
+
+	TArray<FName> ChannelNames;
+
+	bool bInEditingMode = false;
+};
diff --git a/Engine/Plugins/Experimental/Avalanche/Source/AvalancheMediaEditor/Private/Playlist/Pages/Columns/Slate/SAvaPageId.cpp b/Engine/Plugins/Experimental/Avalanche/Source/AvalancheMediaEditor/Private/Playlist/Pages/Columns/Slate/SAvaPageId.cpp
new file mode 100644
--- /dev/null
+++ b/Engine/Plugins/Experimental/Avalanche/Source/AvalancheMediaEditor/Private/Playlist/Pages/Columns/Slate/SAvaPageId.cpp
@@ -0,0 +1,137 @@
+// Copyright Epic Games, Inc. All Rights Reserved.
+
+#include "SAvaPageId.h"
+#include "Playlist/AvalanchePlaylist.h"
+#include "ScopedTransaction.h"
+#include "Widgets/Text/SInlineEditableTextBlock.h"
+
+#define LOCTEXT_NAMESPACE "SAvaPageId"
+
+void SAvaPageId::Construct(const FArguments& InArgs, const FAvaPageViewPtr& InPageView)
+{
+	PageViewWeak = InPageView;
+
+	InPageView->GetOnRenumber().AddSP(this, &SAvaPageId::OnRenumberAction);
+	
+	ChildSlot
+	[
+		SAssignNew(InlineTextBlock, SInlineEditableTextBlock)
+		.Text(InPageView.Get(), &IAvaPageView::GetPageIdText)
+		.OnTextCommitted(this, &SAvaPageId::OnTextCommitted)
+		.OnVerifyTextChanged(this, &SAvaPageId::OnVerifyTextChanged)
+		.OnEnterEditingMode(this, &SAvaPageId::OnEnterEditingMode)
+		.OnExitEditingMode(this, &SAvaPageId::OnExitEditingMode)
+	];
+}
+
+SAvaPageId::~SAvaPageId()
+{
+	if (PageViewWeak.IsValid())
+	{
+		PageViewWeak.Pin()->GetOnRenumber().RemoveAll(this);
+	}
+}
+
+void SAvaPageId::OnRenumberAction(EAvaPageActionState InRenumberAction)
+{
+	if (InRenumberAction == EAvaPageActionState::Requested)
+	{
+		check(InlineTextBlock.IsValid());
+		InlineTextBlock->EnterEditingMode();
+	}
+}
+
+void SAvaPageId::OnTextCommitted(const FText& InText, ETextCommit::Type InCommitInfo)
+{
+	const FAvaPageViewPtr PageView = PageViewWeak.Pin();
+
+	if (!PageView)
+	{
+		return;
+	}
+	
+	switch (InCommitInfo)
+	{
+	case ETextCommit::OnEnter:
+		//falls through
+				
+	case ETextCommit::OnUserMovedFocus:
+		RenumberPageId(InText, PageView);
+		PageView->GetOnRenumber().Broadcast(EAvaPageActionState::Completed);
+		break;
+
+	case ETextCommit::Default:
+		//falls through
+				
+	case ETextCommit::OnCleared:
+		//falls through
+			
+	default:
+		PageView->GetOnRenumber().Broadcast(EAvaPageActionState::Cancelled);
+		break;
+	}
+}
+
+bool SAvaPageId::OnVerifyTextChanged(const FText& InText, FText& OutErrorMessage)
+{
+	const FString String = InText.ToString();
+	
+	if (!String.IsNumeric())
+	{
+		OutErrorMessage = LOCTEXT("NonNumericError", "Text is not Numeric");
+		return false;
+	}
+
+	UAvalanchePlaylist* const Playlist = PageViewWeak.IsValid()
+		? PageViewWeak.Pin()->GetPlaylist()
+		: nullptr;
+
+	if (!Playlist)
+	{
+		OutErrorMessage = LOCTEXT("NoPlaylist", "No Playlist found for Item");
+		return false;
+	}
+	
+	const int32 Id = FCString::Atoi(*String);
+
+	if (Playlist->GetPage(Id).IsValidPage())
+	{
+		OutErrorMessage = LOCTEXT("ExistingPageWithId", "A Page already exists with input Id");
+		return false;
+	}
+	
+	return true;
+}
+
+void SAvaPageId::OnEnterEditingMode()
+{
+	bInEditingMode = true;
+}
+
+void SAvaPageId::OnExitEditingMode()
+{
+	bInEditingMode = false;
+}
+
+void SAvaPageId::RenumberPageId(const FText& InText, const FAvaPageViewPtr& InPageView)
+{
+	const FString String = InText.ToString();
+	check(String.IsNumeric());
+
+	UAvalanchePlaylist* const Playlist = InPageView->GetPlaylist();
+	check(Playlist);
+	
+	const int32 Id = FCString::Atoi(*String);	
+
+	FScopedTransaction Transaction(LOCTEXT("RenumberPageId", "Renumber Page Id"));
+	Playlist->Modify();
+	
+	const bool bResult = Playlist->RenumberPageId(InPageView->GetPageId(), Id);
+
+	if (!bResult)
+	{
+		Transaction.Cancel();
+	}
+}
+
+#undef LOCTEXT_NAMESPACE
\ No newline at end of file
diff --git a/Engine/Plugins/Experimental/Avalanche/Source/AvalancheMediaEditor/Private/Playlist/Pages/Columns/Slate/SAvaPageId.h b/Engine/Plugins/Experimental/Avalanche/Source/AvalancheMediaEditor/Private/Playlist/Pages/Columns/Slate/SAvaPageId.h
new file mode 100644
--- /dev/null
+++ b/Engine/Plugins/Experimental/Avalanche/Source/AvalancheMediaEditor/Private/Playlist/Pages/Columns/Slate/SAvaPageId.h
@@ -0,0 +1,42 @@
+// Copyright Epic Games, Inc. All Rights Reserved.
+
+#pragma once
+
+#include "CoreMinimal.h"
+#include "Playlist/AvaPlaylistDefines.h"
+#include "Playlist/Pages/PageViews/AvaPageView.h"
+#include "Widgets/SCompoundWidget.h"
+
+class SInlineEditableTextBlock;
+
+class SAvaPageId : public SCompoundWidget
+{
+public:
+	
+	SLATE_BEGIN_ARGS(SAvaPageId){}
+	SLATE_END_ARGS()
+
+	void Construct(const FArguments& InArgs, const FAvaPageViewPtr& InPageView);
+
+	virtual ~SAvaPageId() override;
+
+	void OnRenumberAction(EAvaPageActionState InRenumberAction);
+	
+	void OnTextCommitted(const FText& InText, ETextCommit::Type InCommitInfo);
+	
+	bool OnVerifyTextChanged(const FText& InText, FText& OutErrorMessage);
+
+	void OnEnterEditingMode();
+	
+	void OnExitEditingMode();
+
+	void RenumberPageId(const FText& InText, const FAvaPageViewPtr& InPageView);
+	
+protected:
+
+	TWeakPtr<IAvaPageView> PageViewWeak;
+	
+	TSharedPtr<SInlineEditableTextBlock> InlineTextBlock;
+
+	bool bInEditingMode = false;
+};
diff --git a/Engine/Plugins/Experimental/Avalanche/Source/AvalancheMediaEditor/Private/Playlist/Pages/Columns/Slate/SAvaPageName.cpp b/Engine/Plugins/Experimental/Avalanche/Source/AvalancheMediaEditor/Private/Playlist/Pages/Columns/Slate/SAvaPageName.cpp
new file mode 100644
--- /dev/null
+++ b/Engine/Plugins/Experimental/Avalanche/Source/AvalancheMediaEditor/Private/Playlist/Pages/Columns/Slate/SAvaPageName.cpp
@@ -0,0 +1,104 @@
+// Copyright Epic Games, Inc. All Rights Reserved.
+
+#include "SAvaPageName.h"
+#include "Framework/Application/SlateApplication.h"
+#include "Playlist/Pages/PageViews/AvaPageView.h"
+#include "Playlist/Pages/Slate/SAvaPageViewRow.h"
+#include "Widgets/Text/SInlineEditableTextBlock.h"
+
+void SAvaPageName::Construct(const FArguments& InArgs, const FAvaPageViewPtr& InPageView, const TSharedPtr<SAvaPageViewRow>& InRow)
+{
+	PageViewWeak = InPageView;
+	PageViewRowWeak = InRow;
+
+	InPageView->GetOnRename().AddSP(this, &SAvaPageName::OnRenameAction);
+	
+	ChildSlot
+	[
+		SAssignNew(InlineTextBlock, SInlineEditableTextBlock)
+		.Text(InPageView.Get(), &IAvaPageView::GetPageDescription)
+		.OnTextCommitted(this, &SAvaPageName::OnTextCommitted)
+		.OnVerifyTextChanged(this, &SAvaPageName::OnVerifyTextChanged)
+		.OnEnterEditingMode(this, &SAvaPageName::OnEnterEditingMode)
+		.OnExitEditingMode(this, &SAvaPageName::OnExitEditingMode)
+		.IsSelected(InRow.Get(), &SAvaPageViewRow::IsSelectedExclusively)
+		.IsReadOnly(this, &SAvaPageName::IsReadOnly)
+	];
+}
+
+SAvaPageName::~SAvaPageName()
+{
+	if (PageViewWeak.IsValid())
+	{
+		PageViewWeak.Pin()->GetOnRename().RemoveAll(this);
+	}
+}
+
+void SAvaPageName::OnRenameAction(EAvaPageActionState InRenameAction)
+{
+	if (InRenameAction == EAvaPageActionState::Requested)
+	{
+		check(InlineTextBlock.IsValid());
+		bIsReadOnly = false;	// Only allow entering edit mode from rename action.
+		InlineTextBlock->EnterEditingMode();
+	}
+}
+
+void SAvaPageName::OnTextCommitted(const FText& InText, ETextCommit::Type InCommitInfo)
+{
+	if (const FAvaPageViewPtr PageView = PageViewWeak.Pin())
+	{
+		switch (InCommitInfo)
+		{
+		case ETextCommit::OnEnter:
+			//falls through
+				
+		case ETextCommit::OnUserMovedFocus:
+			RenamePage(InText, PageView);
+			PageView->GetOnRename().Broadcast(EAvaPageActionState::Completed);
+			break;
+
+		case ETextCommit::Default:
+			//falls through
+				
+		case ETextCommit::OnCleared:
+			//falls through
+			
+		default:
+			PageView->GetOnRename().Broadcast(EAvaPageActionState::Cancelled);
+			break;
+		}
+	}
+}
+
+bool SAvaPageName::OnVerifyTextChanged(const FText& InText, FText& OutErrorMessage)
+{
+	return true;
+}
+
+void SAvaPageName::OnEnterEditingMode()
+{
+	bInEditingMode = true;
+}
+
+void SAvaPageName::OnExitEditingMode()
+{
+	bInEditingMode = false;
+	bIsReadOnly = true;
+}
+
+void SAvaPageName::RenamePage(const FText& InText, const FAvaPageViewPtr& InPageView)
+{
+	check(InPageView.IsValid());
+
+	const bool bRenameSuccessful = InPageView->RenameFriendlyName(InText);
+	if (bRenameSuccessful)
+	{
+		FSlateApplication::Get().ClearKeyboardFocus(EFocusCause::Mouse);
+	}
+}
+
+bool SAvaPageName::IsReadOnly() const
+{
+	return bIsReadOnly;
+}
\ No newline at end of file
diff --git a/Engine/Plugins/Experimental/Avalanche/Source/AvalancheMediaEditor/Private/Playlist/Pages/Columns/Slate/SAvaPageName.h b/Engine/Plugins/Experimental/Avalanche/Source/AvalancheMediaEditor/Private/Playlist/Pages/Columns/Slate/SAvaPageName.h
new file mode 100644
--- /dev/null
+++ b/Engine/Plugins/Experimental/Avalanche/Source/AvalancheMediaEditor/Private/Playlist/Pages/Columns/Slate/SAvaPageName.h
@@ -0,0 +1,47 @@
+// Copyright Epic Games, Inc. All Rights Reserved.
+
+#pragma once
+
+#include "CoreMinimal.h"
+#include "Playlist/AvaPlaylistDefines.h"
+#include "Widgets/SCompoundWidget.h"
+
+class IAvaPageView;
+class SAvaPageViewRow;
+class SInlineEditableTextBlock;
+enum class EAvaPageActionState : uint8;
+
+class SAvaPageName : public SCompoundWidget
+{
+public:
+	
+	SLATE_BEGIN_ARGS(SAvaPageName){}
+	SLATE_END_ARGS()
+
+	/** Constructs this widget with InArgs */
+	void Construct(const FArguments& InArgs, const FAvaPageViewPtr& InPageView, const TSharedPtr<SAvaPageViewRow>& InRow);
+	virtual ~SAvaPageName() override;
+
+	void OnRenameAction(EAvaPageActionState InRenameAction);
+	
+	void OnTextCommitted(const FText& InText, ETextCommit::Type InCommitInfo);
+	bool OnVerifyTextChanged(const FText& InText, FText& OutErrorMessage);
+
+	void OnEnterEditingMode();
+	void OnExitEditingMode();
+
+	void RenamePage(const FText& InText, const FAvaPageViewPtr& InPageView);
+
+	bool IsReadOnly() const;
+	
+protected:
+
+	TWeakPtr<IAvaPageView> PageViewWeak;
+	
+	TWeakPtr<SAvaPageViewRow> PageViewRowWeak;
+	
+	TSharedPtr<SInlineEditableTextBlock> InlineTextBlock;
+
+	bool bInEditingMode = false;
+	bool bIsReadOnly = true;
+};
diff --git a/Engine/Plugins/Experimental/Avalanche/Source/AvalancheMediaEditor/Private/Playlist/Pages/Columns/Slate/SAvaPageThumbnail.cpp b/Engine/Plugins/Experimental/Avalanche/Source/AvalancheMediaEditor/Private/Playlist/Pages/Columns/Slate/SAvaPageThumbnail.cpp
new file mode 100644
--- /dev/null
+++ b/Engine/Plugins/Experimental/Avalanche/Source/AvalancheMediaEditor/Private/Playlist/Pages/Columns/Slate/SAvaPageThumbnail.cpp
@@ -0,0 +1,140 @@
+// Copyright Epic Games, Inc. All Rights Reserved.
+
+#include "SAvaPageThumbnail.h"
+#include "AssetRegistry/IAssetRegistry.h"
+#include "AssetThumbnail.h"
+#include "Framework/Application/SlateApplication.h"
+#include "Playlist/AvalanchePlaylist.h"
+#include "Playlist/Pages/PageViews/AvaPageView.h"
+#include "Playlist/Pages/Slate/SAvaPageViewRow.h"
+#include "ThumbnailRendering/ThumbnailManager.h"
+#include "Widgets/Layout/SScaleBox.h"
+
+void SAvaPageThumbnail::Construct(const FArguments& InArgs, const FAvaPageViewPtr& InPageView, const TSharedPtr<SAvaPageViewRow>& InRow)
+{
+	PageViewWeak = InPageView;
+	PageViewRowWeak = InRow;
+	ThumbnailWidgetSize = InArgs._ThumbnailWidgetSize;
+
+	InitThumbnailWidget();
+
+	ChildSlot
+	[
+		SNew(SBorder)
+		.Padding(FMargin(0.0f, 0.0f, 5.0f, 5.0f))
+		.BorderImage(FAppStyle::Get().GetBrush("ContentBrowser.AssetTileItem.DropShadow"))
+		[
+			SNew(SBox)
+			.Padding(0.0f)
+			.WidthOverride(ThumbnailWidgetSize)
+			.HeightOverride(ThumbnailWidgetSize)
+			[
+				SNew(SScaleBox)
+				.Stretch(EStretch::ScaleToFit)
+				.StretchDirection(EStretchDirection::Both)
+				[
+					SNew(SBorder)
+					.Padding(0.0f)
+					.BorderImage(FAppStyle::Get().GetBrush("ContentBrowser.AssetTileItem.ThumbnailAreaBackground"))
+					[
+						ThumbnailWidget.ToSharedRef()
+					]
+				]
+			]
+		]
+	];
+}
+
+void SAvaPageThumbnail::OnAssetUpdate(const UAvalanchePlaylist* InPlaylist, const FAvalanchePage& InPage, EAvaPageChanges InPageChangeType)
+{
+	if (!EnumHasAnyFlags(InPageChangeType, EAvaPageChanges::Blueprint) || !InPlaylist)
+	{
+		return;
+	}
+
+	if (const TSharedPtr<IAvaPageView> PageView = PageViewWeak.Pin())
+	{
+		int32 PageIdToCheck = PageView->GetPageId();
+		if (!PageView->IsTemplate())
+		{
+			PageIdToCheck = InPlaylist->GetPage(PageView->GetPageId()).GetTemplateId();
+		}
+
+		if (AssetThumbnail.IsValid() && InPage.IsValidPage() && InPage.IsTemplate() && PageIdToCheck == InPage.GetPageId())
+		{
+			const FSoftObjectPath AssetPath = InPage.GetAvalancheAssetPath(InPlaylist);
+			const FAssetData Data = IAssetRegistry::Get()->GetAssetByObjectPath(AssetPath);
+
+			AssetThumbnail->SetAsset(Data);
+			AssetThumbnail->RefreshThumbnail();
+			// Just to ensure that the viewport won't turn black on an asset which wasn't loaded yet
+			bWasAssetLoaded = false;
+		}
+	}
+}
+
+void SAvaPageThumbnail::InitThumbnailWidget()
+{
+	if (const TSharedPtr<IAvaPageView> PageView =  PageViewWeak.Pin())
+	{
+		if (UAvalanchePlaylist* Playlist = PageView->GetPlaylist())
+		{
+			Playlist->GetOnPagesChanged().AddSP(this, &SAvaPageThumbnail::OnAssetUpdate);
+
+			FAvalanchePage& Page = Playlist->GetPage(PageView->GetPageId());
+			if (!PageView->IsTemplate())
+			{
+				Page = Playlist->GetPage(Page.GetTemplateId());
+			}
+			const FSoftObjectPath AssetPath = Page.GetAvalancheAssetPath(Playlist);
+			const FAssetData Data = IAssetRegistry::Get()->GetAssetByObjectPath(AssetPath);
+
+
+			AssetThumbnail = MakeShared<FAssetThumbnail>(Data, ThumbnailWidgetSize * 2, ThumbnailWidgetSize * 2, UThumbnailManager::Get().GetSharedThumbnailPool());
+			AssetThumbnail->SetRealTime(false);
+
+			ThumbnailWidget = AssetThumbnail->MakeThumbnailWidget();
+
+			Tooltip = FSlateApplication::Get().MakeToolTip(NSLOCTEXT("SAvaPageThumbnail", "SAvaPageThumbnail_Tooltip", "Tooltip"));
+			Tooltip->SetContentWidget(SNew(SBox)
+					.Padding(0.0f)
+					.WidthOverride(ThumbnailWidgetSize * 2)
+					.HeightOverride(ThumbnailWidgetSize * 2)
+					[
+						SNew(SScaleBox)
+						.Stretch(EStretch::ScaleToFit)
+						.StretchDirection(EStretchDirection::Both)
+						[
+							AssetThumbnail->MakeThumbnailWidget()
+						]
+					]);
+
+			ThumbnailWidget->SetToolTip(Tooltip);
+		}
+	}
+
+	if (!ThumbnailWidget.IsValid())
+	{
+		ThumbnailWidget = SNullWidget::NullWidget;
+	}
+}
+
+void SAvaPageThumbnail::Tick(const FGeometry& AllottedGeometry, const double InCurrentTime, const float InDeltaTime)
+{
+	if (!bWasAssetLoaded && AssetThumbnail->GetAsset())
+	{
+		bWasAssetLoaded = true;
+		AssetThumbnail->RefreshThumbnail();
+	}
+}
+
+SAvaPageThumbnail::~SAvaPageThumbnail()
+{
+	if (const TSharedPtr<IAvaPageView> PageView =  PageViewWeak.Pin())
+	{
+		if (UAvalanchePlaylist* Playlist = PageView->GetPlaylist())
+		{
+			Playlist->GetOnPagesChanged().RemoveAll(this);
+		}
+	}
+}
diff --git a/Engine/Plugins/Experimental/Avalanche/Source/AvalancheMediaEditor/Private/Playlist/Pages/Columns/Slate/SAvaPageThumbnail.h b/Engine/Plugins/Experimental/Avalanche/Source/AvalancheMediaEditor/Private/Playlist/Pages/Columns/Slate/SAvaPageThumbnail.h
new file mode 100644
--- /dev/null
+++ b/Engine/Plugins/Experimental/Avalanche/Source/AvalancheMediaEditor/Private/Playlist/Pages/Columns/Slate/SAvaPageThumbnail.h
@@ -0,0 +1,51 @@
+// Copyright Epic Games, Inc. All Rights Reserved.
+
+#pragma once
+
+#include "Playlist/AvaPlaylistDefines.h"
+#include "Templates/SharedPointer.h"
+#include "Widgets/SCompoundWidget.h"
+
+class FAssetThumbnail;
+class IToolTip;
+class SAvaPageViewRow;
+class UAvalanchePlaylist;
+enum class EAvaPageChanges : uint8;
+
+class SAvaPageThumbnail : public SCompoundWidget
+{
+public:
+
+	SLATE_BEGIN_ARGS(SAvaPageThumbnail)
+		: _ThumbnailWidgetSize(64)
+	{}
+
+	SLATE_ARGUMENT(int32, ThumbnailWidgetSize)
+
+	SLATE_END_ARGS()
+
+	/** Constructs this widget with InArgs */
+	void Construct(const FArguments& InArgs, const FAvaPageViewPtr& InPageView, const TSharedPtr<SAvaPageViewRow>& InRow);
+	virtual void Tick(const FGeometry& AllottedGeometry, const double InCurrentTime, const float InDeltaTime) override;
+	~SAvaPageThumbnail() override;
+
+private:
+	void OnAssetUpdate(const UAvalanchePlaylist* InPlaylist, const FAvalanchePage& InPage, EAvaPageChanges InPageChangeType);
+
+	void InitThumbnailWidget();
+
+private:
+	TWeakPtr<IAvaPageView> PageViewWeak;
+
+	TWeakPtr<SAvaPageViewRow> PageViewRowWeak;
+
+	TSharedPtr<FAssetThumbnail> AssetThumbnail;
+
+	TSharedPtr<SWidget> ThumbnailWidget;
+
+	TSharedPtr<IToolTip> Tooltip;
+
+	bool bWasAssetLoaded = false;
+
+	int32 ThumbnailWidgetSize = 0;
+};
diff --git a/Engine/Plugins/Experimental/Avalanche/Source/AvalancheMediaEditor/Private/Playlist/Pages/PageViews/AvaInstancedPageViewImpl.cpp b/Engine/Plugins/Experimental/Avalanche/Source/AvalancheMediaEditor/Private/Playlist/Pages/PageViews/AvaInstancedPageViewImpl.cpp
new file mode 100644
--- /dev/null
+++ b/Engine/Plugins/Experimental/Avalanche/Source/AvalancheMediaEditor/Private/Playlist/Pages/PageViews/AvaInstancedPageViewImpl.cpp
@@ -0,0 +1,189 @@
+// Copyright Epic Games, Inc. All Rights Reserved.
+
+#include "AvaInstancedPageViewImpl.h"
+
+#include "AvalancheBroadcast.h"
+#include "Framework/Application/SlateApplication.h"
+#include "Input/Reply.h"
+#include "Playlist/AvalanchePage.h"
+#include "Playlist/AvalanchePlaylist.h"
+#include "Playlist/AvaPlaylistDefines.h"
+
+#define LOCTEXT_NAMESPACE "AvaPageViewImpl"
+
+FAvaInstancedPageViewImpl::FAvaInstancedPageViewImpl(int32 InPageId, UAvalanchePlaylist* InPlaylist, const TSharedPtr<SAvaPageList>& InPageList)
+	: FAvaPageViewImpl(InPageId, InPlaylist, InPageList)
+{
+}
+
+FReply FAvaInstancedPageViewImpl::OnPlayButtonClicked()
+{
+	UAvalanchePlaylist* Playlist = GetPlaylist();
+
+	if (IsValid(Playlist))
+	{
+		const FAvalanchePage& Page = GetPage();
+
+		if (Page.IsValidPage())
+		{
+			const TArray<FAvalanchePageStatus> Statuses = Page.GetPageProgramStatuses(Playlist);
+			FModifierKeysState KeyState = FSlateApplication::Get().GetModifierKeys();
+
+			const bool bContinue = KeyState.IsControlDown() || KeyState.IsCommandDown();
+
+			if (bContinue)
+			{
+				Playlist->ContinuePage(Page.GetPageId(), false);
+			}
+			else if (FAvalanchePage::StatusesContainsStatus(Statuses, {EAvalanchePageStatus::Playing}))
+			{
+				Playlist->PlayPage(Page.GetPageId(), EAvaPlayType::PlayFromStart);
+			}
+			else
+			{
+				Playlist->PlayPage(Page.GetPageId(), EAvaPlayType::PlayFromStart);
+			}
+
+			return FReply::Handled();
+		}
+	}
+
+	return FReply::Unhandled();
+}
+
+bool FAvaInstancedPageViewImpl::CanPlay() const
+{
+	UAvalanchePlaylist* Playlist = GetPlaylist();
+
+	if (IsValid(Playlist))
+	{
+		const FAvalanchePage& Page = GetPage();
+
+		if (Page.IsValidPage())
+		{
+			const FAvaOutputChannel& Channel = UAvalancheBroadcast::Get().GetCurrentProfile().GetChannel(Page.GetChannelName());
+
+			if (Channel.GetMediaOutputs().Num() == 0)
+			{
+				return false;
+			}
+
+			const TArray<FAvalanchePageStatus> Statuses = Page.GetPageProgramStatuses(Playlist);
+
+			const bool bHasError = FAvalanchePage::StatusesContainsStatus(Statuses, {EAvalanchePageStatus::Error, EAvalanchePageStatus::Loading,
+				EAvalanchePageStatus::Missing, EAvalanchePageStatus::Syncing, EAvalanchePageStatus::Unknown});
+
+			if (bHasError)
+			{
+				return false;
+			}
+
+			const bool bCanPlay = FAvalanchePage::StatusesContainsStatus(Statuses, {EAvalanchePageStatus::Available, EAvalanchePageStatus::Loaded,
+				EAvalanchePageStatus::Playing, EAvalanchePageStatus::NeedsSync});
+
+			if (bCanPlay)
+			{
+				return true;
+			}
+		}
+	}
+
+	return false;
+}
+
+ECheckBoxState FAvaInstancedPageViewImpl::IsEnabled() const
+{
+	const FAvalanchePage& Page = GetPage();
+	if (Page.IsValidPage() && Page.IsEnabled())
+	{
+		return  ECheckBoxState::Checked;
+	}
+	return ECheckBoxState::Unchecked;
+}
+
+void FAvaInstancedPageViewImpl::SetEnabled(ECheckBoxState InState)
+{
+	if (!IsPageSelected())
+	{
+		SetPageSelection(EAvaPageViewSelectionChangeType::ReplaceSelection);
+	}
+
+	PerformWorkOnPages(LOCTEXT("SetEnabled", "Set Enabled"),
+		[InState](FAvalanchePage& InPage)->bool
+		{
+			InPage.SetEnabled(InState == ECheckBoxState::Checked);
+			return true;
+		});
+}
+
+FName FAvaInstancedPageViewImpl::GetChannelName() const
+{
+	const FAvalanchePage& Page = GetPage();
+	return Page.IsValidPage() ? Page.GetChannelName() : FName();
+}
+
+bool FAvaInstancedPageViewImpl::SetChannel(FName InChannel)
+{
+	if (!IsPageSelected())
+	{
+		SetPageSelection(EAvaPageViewSelectionChangeType::ReplaceSelection);
+	}
+
+	return PerformWorkOnPages(LOCTEXT("SetChannel", "Set Channel"),
+		[this, InChannel](FAvalanchePage& InPage)->bool
+		{
+			InPage.SetChannelName(InChannel);
+			GetPlaylist()->GetOnPagesChanged().Broadcast(GetPlaylist(), InPage, EAvaPageChanges::Channel);
+			return true;
+		});
+}
+
+const FAvalanchePage& FAvaInstancedPageViewImpl::GetTemplate() const
+{
+	UAvalanchePlaylist* Playlist = PlaylistWeak.Get();
+
+	if (IsValid(Playlist))
+	{
+		const FAvalanchePage& Page = Playlist->GetPage(PageId);
+
+		if (Page.IsValidPage())
+		{
+			if (Page.IsTemplate())
+			{
+				return Page;
+			}
+
+			const FAvalanchePage& Template = Playlist->GetPage(Page.GetTemplateId());
+
+			if (Template.IsValidPage())
+			{
+				return Template;
+			}
+		}
+	}
+
+	return FAvalanchePage::NullPage;
+}
+
+FText FAvaInstancedPageViewImpl::GetTemplateDescription() const
+{
+	const FAvalanchePage& Template = GetTemplate();
+
+	if (!Template.IsValidPage())
+	{
+		return FText::GetEmpty();
+	}
+
+	return FText::Format(
+		LOCTEXT("TemplateFormat", "{0}: {1}"),
+		FText::AsNumber(Template.GetPageId(), &UE::AvalanchePlaylist::FEditorMetrics::PageIdFormattingOptions),
+		Template.GetPageDescription()
+	);
+}
+
+bool FAvaInstancedPageViewImpl::IsTemplate() const 
+{ 
+	return false;
+}
+
+#undef LOCTEXT_NAMESPACE
diff --git a/Engine/Plugins/Experimental/Avalanche/Source/AvalancheMediaEditor/Private/Playlist/Pages/PageViews/AvaInstancedPageViewImpl.h b/Engine/Plugins/Experimental/Avalanche/Source/AvalancheMediaEditor/Private/Playlist/Pages/PageViews/AvaInstancedPageViewImpl.h
new file mode 100644
--- /dev/null
+++ b/Engine/Plugins/Experimental/Avalanche/Source/AvalancheMediaEditor/Private/Playlist/Pages/PageViews/AvaInstancedPageViewImpl.h
@@ -0,0 +1,28 @@
+// Copyright Epic Games, Inc. All Rights Reserved.
+
+#pragma once
+
+#include "AvaPageViewImpl.h"
+#include "Playlist/Pages/PageViews/AvaInstancedPageView.h"
+
+class FAvaInstancedPageViewImpl : public FAvaPageViewImpl, public IAvaInstancedPageView
+{
+public:
+	UE_AVA_INHERITS(FAvaInstancedPageViewImpl, FAvaPageViewImpl, IAvaInstancedPageView);
+
+	FAvaInstancedPageViewImpl(int32 InPageId, UAvalanchePlaylist* InPlaylist, const TSharedPtr<SAvaPageList>& InPageList);
+
+	virtual bool IsTemplate() const override;
+
+	virtual FReply OnPlayButtonClicked() override;
+	virtual bool CanPlay() const override;
+
+	virtual ECheckBoxState IsEnabled() const override;
+	virtual void SetEnabled(ECheckBoxState InState) override;
+
+	virtual FName GetChannelName() const override;
+	virtual bool SetChannel(FName InChannel) override;
+
+	virtual const FAvalanchePage& GetTemplate() const override;
+	virtual FText GetTemplateDescription() const override;
+};
diff --git a/Engine/Plugins/Experimental/Avalanche/Source/AvalancheMediaEditor/Private/Playlist/Pages/PageViews/AvaPageViewImpl.cpp b/Engine/Plugins/Experimental/Avalanche/Source/AvalancheMediaEditor/Private/Playlist/Pages/PageViews/AvaPageViewImpl.cpp
new file mode 100644
--- /dev/null
+++ b/Engine/Plugins/Experimental/Avalanche/Source/AvalancheMediaEditor/Private/Playlist/Pages/PageViews/AvaPageViewImpl.cpp
@@ -0,0 +1,416 @@
+// Copyright Epic Games, Inc. All Rights Reserved.
+
+#include "AvaPageViewImpl.h"
+
+#include "AssetRegistry/AssetData.h"
+#include "Framework/Application/SlateApplication.h"
+#include "Playlist/AvaPlaylistDefines.h"
+#include "Playlist/AvalanchePage.h"
+#include "Playlist/AvalanchePlaylist.h"
+#include "Playlist/Pages/Slate/SAvaPageList.h"
+#include "ScopedTransaction.h"
+
+#define LOCTEXT_NAMESPACE "AvaPageViewImpl"
+
+FAvaPageViewImpl::FAvaPageViewImpl(int32 InPageId, UAvalanchePlaylist* InPlaylist, const TSharedPtr<SAvaPageList>& InPageList)
+	: PageId(InPageId)
+	, PlaylistWeak(InPlaylist)
+	, PageListWeak(InPageList)
+{
+}
+
+UAvalanchePlaylist* FAvaPageViewImpl::GetPlaylist() const
+{
+	return PlaylistWeak.Get();
+}
+
+int32 FAvaPageViewImpl::GetPageId() const
+{
+	const FAvalanchePage& Page = GetPage();
+	return Page.IsValidPage()
+		? Page.GetPageId()
+		: FAvalanchePage::InvalidPageId;
+}
+
+FText FAvaPageViewImpl::GetPageIdText() const
+{
+	const int32 Id = GetPageId();
+	return Id != FAvalanchePage::InvalidPageId
+		? FText::AsNumber(Id, &UE::AvalanchePlaylist::FEditorMetrics::PageIdFormattingOptions)
+		: LOCTEXT("InvalidIdText", "(invalid)");
+}
+
+FText FAvaPageViewImpl::GetPageNameText() const
+{
+	const FAvalanchePage& Page = GetPage();
+
+	return Page.IsValidPage()
+		? FText::FromString(Page.GetPageName())
+		: LOCTEXT("EmptyPageNameText", "");
+}
+
+FText FAvaPageViewImpl::GetPageTransitionLayerNameText() const
+{
+	if (const UAvalanchePlaylist* Playlist = PlaylistWeak.Get())
+	{
+		const FAvalanchePage& Page = GetPage();
+		if (Page.IsValidPage())
+		{
+			if (Page.HasTransitionLogic(Playlist))
+			{
+				FString TransitionLayers;
+				const int32 NumTemplates = Page.GetNumTemplates(Playlist);
+				for (int32 TemplateIndex = 0; TemplateIndex < NumTemplates; ++TemplateIndex)
+				{
+					if (!TransitionLayers.IsEmpty())
+					{
+						TransitionLayers += TEXT(", ");
+					}
+					TransitionLayers += Page.GetTransitionLayer(Playlist, TemplateIndex).ToString();
+				}
+				return FText::FromString(TransitionLayers);
+			}
+
+			return LOCTEXT("PageTransitionLayerText_NA", "N/A");
+		}
+	}
+	// Either invalid playlist or invalid page.
+	return LOCTEXT("PageTransitionLayerText_Invalid", "(invalid)");
+}
+
+FText FAvaPageViewImpl::GetPageSummary() const
+{
+	const FAvalanchePage& Page = GetPage();
+
+	return Page.IsValidPage()
+		? Page.GetPageSummary()
+		: LOCTEXT("EmptyPageSummaryText", "");
+}
+
+FText FAvaPageViewImpl::GetPageDescription() const
+{
+	const FAvalanchePage& Page = GetPage();
+
+	return Page.IsValidPage()
+		? Page.GetPageDescription()
+		: LOCTEXT("EmptyPageDescriptionText", "");
+}
+
+bool FAvaPageViewImpl::HasObjectPath(const UAvalanchePlaylist* InPlaylist) const
+{
+	const FAvalanchePage& Page = GetPage();
+	return Page.IsValidPage() && !Page.IsComboTemplate();
+}
+
+FSoftObjectPath FAvaPageViewImpl::GetObjectPath(const UAvalanchePlaylist* InPlaylist) const
+{
+	const FAvalanchePage& Page = GetPage();
+	if (Page.IsValidPage())
+	{
+		return Page.GetAvalancheAssetPath(InPlaylist);
+	}
+
+	return FSoftObjectPath();
+}
+
+FText FAvaPageViewImpl::GetObjectName(const UAvalanchePlaylist* InPlaylist) const
+{
+	const FAvalanchePage& Page = GetPage();
+	if (Page.IsValidPage())
+	{
+		if (Page.ResolveTemplate(InPlaylist).IsComboTemplate())
+		{
+			// Since combo templates don't have an asset selector, use same placeholder.
+			return LOCTEXT("AssetName_ComboPage_NA", "N/A");
+		}
+		return FText::FromString(Page.GetAvalancheAssetPath(InPlaylist).GetAssetName());
+	}
+	return FText::GetEmpty();
+}
+
+void FAvaPageViewImpl::OnObjectChanged(const FAssetData& InAssetData)
+{
+	if (!IsPageSelected())
+	{
+		SetPageSelection(EAvaPageViewSelectionChangeType::ReplaceSelection);
+	}
+
+	PerformWorkOnPages(LOCTEXT("SetAvaBlueprint", "Set Avalanche Blueprint"),
+		[this, &InAssetData](FAvalanchePage& InPage)->bool
+		{
+			if (!InPage.UpdateAvalancheAsset(InAssetData.GetSoftObjectPath())) 
+			{
+ 				return false;
+			}
+			GetPlaylist()->GetOnPagesChanged().Broadcast(GetPlaylist(), InPage, EAvaPageChanges::Blueprint);
+			return true;
+		});
+}
+
+bool FAvaPageViewImpl::Rename(const FText& InNewName)
+{
+	UAvalanchePlaylist* const Playlist = PlaylistWeak.Get();
+	if (!Playlist)
+	{
+		return false;
+	}
+
+	FAvalanchePage& Page  = Playlist->GetPage(PageId);
+	if (Page.IsValidPage())
+	{
+		FScopedTransaction Transaction(LOCTEXT("RenamePage", "Rename Page"));
+		Playlist->Modify();
+		
+		Page.Rename(InNewName.ToString());
+		return true;
+	}
+	return false;
+}
+
+bool FAvaPageViewImpl::RenameFriendlyName(const FText& InNewName)
+{
+	UAvalanchePlaylist* const Playlist = PlaylistWeak.Get();
+	if (!Playlist)
+	{
+		return false;
+	}
+
+	FAvalanchePage& Page  = Playlist->GetPage(PageId);
+	if (Page.IsValidPage())
+	{
+		FScopedTransaction Transaction(LOCTEXT("RenamePage", "Rename Page"));
+		Playlist->Modify();
+		
+		Page.RenameFriendlyName(InNewName.ToString());
+		return true;
+	}
+	return false;
+}
+
+FReply FAvaPageViewImpl::OnAssetStatusButtonClicked()
+{
+	return FReply::Handled();
+}
+
+bool FAvaPageViewImpl::CanChangeAssetStatus() const
+{
+	UAvalanchePlaylist* Playlist = PlaylistWeak.Get();
+
+	if (IsValid(Playlist))
+	{
+		const FAvalanchePage& Page = Playlist->GetPage(PageId);
+
+		if (Page.IsValidPage())
+		{
+			const TArray<FAvalanchePageStatus> Statuses = Page.GetPageContextualStatuses(Playlist);
+
+			return !FAvalanchePage::StatusesContainsStatus(Statuses, {EAvalanchePageStatus::Error, EAvalanchePageStatus::Loaded, EAvalanchePageStatus::Missing,
+				EAvalanchePageStatus::Playing, EAvalanchePageStatus::Previewing, EAvalanchePageStatus::Syncing, EAvalanchePageStatus::Unknown});
+		}
+	}
+
+	return false;
+}
+
+FReply FAvaPageViewImpl::OnPreviewButtonClicked()
+{
+	UAvalanchePlaylist* Playlist = GetPlaylist();
+
+	if (IsValid(Playlist))
+	{
+		const FAvalanchePage& Page = GetPage();
+
+		if (Page.IsValidPage())
+		{
+			const TArray<FAvalanchePageStatus> Statuses = Page.GetPagePreviewStatuses(Playlist);
+			const bool bIsPreviewing = FAvalanchePage::StatusesContainsStatus(Statuses, {EAvalanchePageStatus::Previewing});
+			const int32 ThisPageId = Page.GetPageId();
+
+			// Alt = restart
+			// Control = continue
+			// Shift = frame
+			const bool bFromFrame = FSlateApplication::Get().GetModifierKeys().IsShiftDown();
+			const bool bContinue = FSlateApplication::Get().GetModifierKeys().IsControlDown()
+				|| FSlateApplication::Get().GetModifierKeys().IsCommandDown();
+
+			const EAvaPlayType PreviewType = bFromFrame ? EAvaPlayType::PreviewFromFrame : EAvaPlayType::PreviewFromStart;
+
+			if (bIsPreviewing && bContinue)
+			{
+				Playlist->ContinuePage(ThisPageId, true);
+			}
+			else
+			{
+				Playlist->PlayPage(ThisPageId, PreviewType);
+			}
+
+			return FReply::Handled();
+		}
+	}
+
+	return FReply::Unhandled();
+}
+
+bool FAvaPageViewImpl::CanPreview() const
+{
+	UAvalanchePlaylist* Playlist = GetPlaylist();
+
+	if (IsValid(Playlist))
+	{
+		const FAvalanchePage& Page = GetPage();
+
+		if (Page.IsValidPage())
+		{
+			const TArray<FAvalanchePageStatus> Statuses = Page.GetPagePreviewStatuses(Playlist);
+			const bool bIsPreviewing = FAvalanchePage::StatusesContainsStatus(Statuses, {EAvalanchePageStatus::Previewing});
+			const int32 ThisPageId = Page.GetPageId();
+
+			// Alt = restart
+			// Control = continue
+			// Shift = frame
+			const bool bFromFrame = FSlateApplication::Get().GetModifierKeys().IsShiftDown();
+			const bool bContinue = FSlateApplication::Get().GetModifierKeys().IsControlDown()
+				|| FSlateApplication::Get().GetModifierKeys().IsCommandDown();
+			const bool bRestartPreview = FSlateApplication::Get().GetModifierKeys().IsAltDown();
+
+			if (bIsPreviewing)
+			{
+				if (bContinue)
+				{
+					return Playlist->CanContinuePage(ThisPageId, true);
+				}
+				else
+				{
+					return Playlist->CanStopPage(ThisPageId, EAvaPlaylistPageStopOptions::Default, true);
+
+					// Unable to test if we can play as well.
+				}
+			}
+			else
+			{
+				return Playlist->CanPlayPage(ThisPageId, true);
+			}
+		}
+	}
+
+	return false;
+}
+
+bool FAvaPageViewImpl::IsPageSelected() const
+{
+	TSharedPtr<SAvaPageList> PageList = PageListWeak.Pin();
+	if (!PageList.IsValid())
+	{
+		return false;
+	}
+
+	return PageList->GetSelectedPageIds().Contains(PageId);
+}
+
+bool FAvaPageViewImpl::SetPageSelection(EAvaPageViewSelectionChangeType InSelectionChangeType)
+{
+	TSharedPtr<SAvaPageList> PageList = PageListWeak.Pin();
+	if (!PageList.IsValid())
+	{
+		return false;
+	}
+
+	switch (InSelectionChangeType)
+	{
+		case EAvaPageViewSelectionChangeType::Deselect:
+			if (PageList->GetSelectedPageIds().Contains(PageId))
+			{
+				PageList->DeselectPage(PageId);
+			}
+
+			return true;
+
+		case EAvaPageViewSelectionChangeType::AddToSelection:
+			if (!PageList->GetSelectedPageIds().Contains(PageId))
+			{
+				PageList->SelectPage(PageId, false);
+			}
+
+			return true;
+
+		case EAvaPageViewSelectionChangeType::ReplaceSelection:
+			PageList->DeselectPages();
+			PageList->SelectPage(PageId, false);
+			return true;
+
+		default:
+			// Not possible
+			return false;
+	}
+}
+
+bool FAvaPageViewImpl::PerformWorkOnPages(const FText& InTransactionSessionName, TFunction<bool(FAvalanchePage&)>&& InWork)
+{
+	UAvalanchePlaylist* const Playlist = PlaylistWeak.Get();
+	if (!Playlist)
+	{
+		return false;
+	}
+
+	FAvalanchePage& UnderlyingPage  = Playlist->GetPage(PageId);
+	if (!UnderlyingPage.IsValidPage())
+	{
+		return false;
+	}
+
+	TSet<FAvalanchePage*> PagesToPerformWork;
+	PagesToPerformWork.Add(&UnderlyingPage);
+
+	if (TSharedPtr<SAvaPageList> PageList = PageListWeak.Pin())
+	{
+		TConstArrayView<int32> SelectedPageIds = PageList->GetSelectedPageIds();
+
+		if (SelectedPageIds.Contains(UnderlyingPage.GetPageId()))
+		{
+			for (int32 SelectedPageId : SelectedPageIds)
+			{
+				FAvalanchePage& CurrentPage = Playlist->GetPage(SelectedPageId);
+
+				if (CurrentPage.IsValidPage())
+				{
+					PagesToPerformWork.Add(&CurrentPage);
+				}
+			}
+		}
+	}
+	
+	if (!PagesToPerformWork.IsEmpty())
+	{
+		FScopedTransaction Transaction(InTransactionSessionName);
+		Playlist->Modify();
+
+		int32 WorkDoneCount = 0;
+		
+		for (FAvalanchePage* const PageToPerformWork : PagesToPerformWork)
+		{
+			if (PageToPerformWork && InWork(*PageToPerformWork))
+			{
+				++WorkDoneCount;
+			}
+		}
+
+		if (WorkDoneCount == 0)
+		{
+			Transaction.Cancel();
+		}
+		return WorkDoneCount > 0;
+	}
+
+	return false;
+}
+
+const FAvalanchePage& FAvaPageViewImpl::GetPage() const
+{
+	if (PlaylistWeak.IsValid())
+	{
+		return PlaylistWeak->GetPage(PageId);
+	}
+	return FAvalanchePage::NullPage;
+}
+
+#undef LOCTEXT_NAMESPACE
diff --git a/Engine/Plugins/Experimental/Avalanche/Source/AvalancheMediaEditor/Private/Playlist/Pages/PageViews/AvaPageViewImpl.h b/Engine/Plugins/Experimental/Avalanche/Source/AvalancheMediaEditor/Private/Playlist/Pages/PageViews/AvaPageViewImpl.h
new file mode 100644
--- /dev/null
+++ b/Engine/Plugins/Experimental/Avalanche/Source/AvalancheMediaEditor/Private/Playlist/Pages/PageViews/AvaPageViewImpl.h
@@ -0,0 +1,74 @@
+// Copyright Epic Games, Inc. All Rights Reserved.
+
+#pragma once
+
+#include "Playlist/Pages/PageViews/AvaPageView.h"
+#include "UObject/WeakObjectPtr.h"
+#include "UObject/WeakObjectPtrTemplates.h"
+
+class FReply;
+class SAvaPageList;
+class UAvalanchePlaylist;
+struct FAssetData;
+struct FAvalanchePage;
+
+class FAvaPageViewImpl : public IAvaPageView
+{
+public:
+	UE_AVA_INHERITS(FAvaPageViewImpl, IAvaPageView);
+
+	FAvaPageViewImpl(int32 InPageId, UAvalanchePlaylist* InPlaylist, const TSharedPtr<SAvaPageList>& InPageList);
+
+	virtual UAvalanchePlaylist* GetPlaylist() const override;
+	
+	virtual int32 GetPageId() const override;
+	virtual FText GetPageIdText() const override;
+	virtual FText GetPageNameText() const override;
+	virtual FText GetPageTransitionLayerNameText() const override;
+	virtual FText GetPageSummary() const override;
+	virtual FText GetPageDescription() const override;
+
+	virtual bool HasObjectPath(const UAvalanchePlaylist* InPlaylist) const override;
+	virtual FSoftObjectPath GetObjectPath(const UAvalanchePlaylist* InPlaylist) const override;
+	virtual FText GetObjectName(const UAvalanchePlaylist* InPlaylist) const override;
+	virtual void OnObjectChanged(const FAssetData& InAssetData) override;
+
+	virtual bool Rename(const FText& InNewName) override;
+	virtual bool RenameFriendlyName(const FText& InNewName) override;
+	virtual FOnPageAction& GetOnRename() override { return OnRename; }
+	virtual FOnPageAction& GetOnRenumber() override { return OnRenumber; }
+
+	virtual FReply OnAssetStatusButtonClicked() override;
+	virtual bool CanChangeAssetStatus() const override;
+	
+	virtual FReply OnPreviewButtonClicked() override;
+	virtual bool CanPreview() const override;
+
+	virtual bool IsPageSelected() const override;
+	virtual bool SetPageSelection(EAvaPageViewSelectionChangeType InSelectionChangeType) override;
+
+protected:
+	/**
+	 * Performs the given work on the underlying Page of this Page View
+	 * and, if this page is selected, to do this work on the rest of the selected pages too
+	 * @param InTransactionSessionName the session name to use for the transaction
+	 * @param InWork the function to execute on each page. needs to return true for the work to be considered done
+	 * @return whether any work was performed and a transaction was completed 
+	 */
+	bool PerformWorkOnPages(const FText& InTransactionSessionName, TFunction<bool(FAvalanchePage&)>&& InWork);
+	
+	const FAvalanchePage& GetPage() const;
+
+protected:
+	/** The Id used to search the Actual Page from Playlist */
+	int32 PageId;
+	
+	TWeakObjectPtr<UAvalanchePlaylist> PlaylistWeak;
+
+	TWeakPtr<SAvaPageList> PageListWeak;
+	
+	FOnPageAction OnRename;
+	
+	FOnPageAction OnRenumber;
+
+};
diff --git a/Engine/Plugins/Experimental/Avalanche/Source/AvalancheMediaEditor/Private/Playlist/Pages/PageViews/AvaTemplatePageViewImpl.cpp b/Engine/Plugins/Experimental/Avalanche/Source/AvalancheMediaEditor/Private/Playlist/Pages/PageViews/AvaTemplatePageViewImpl.cpp
new file mode 100644
--- /dev/null
+++ b/Engine/Plugins/Experimental/Avalanche/Source/AvalancheMediaEditor/Private/Playlist/Pages/PageViews/AvaTemplatePageViewImpl.cpp
@@ -0,0 +1,84 @@
+// Copyright Epic Games, Inc. All Rights Reserved.
+
+#include "AvaTemplatePageViewImpl.h"
+#include "Playlist/AvalanchePlaylist.h"
+#include "Input/Reply.h"
+#include "Playlist/AvalanchePage.h"
+
+FAvaTemplatePageViewImpl::FAvaTemplatePageViewImpl(int32 InPageId, UAvalanchePlaylist* InPlaylist, const TSharedPtr<SAvaPageList>& InPageList)
+		: FAvaPageViewImpl(InPageId, InPlaylist, InPageList)
+{
+}
+
+UAvalanchePlaylist* FAvaTemplatePageViewImpl::GetPlaylist() const
+{ 
+	return FAvaPageViewImpl::GetPlaylist(); 
+}
+
+FReply FAvaTemplatePageViewImpl::OnPlayButtonClicked()
+{
+	return FReply::Handled();
+}
+
+bool FAvaTemplatePageViewImpl::CanPlay() const
+{
+	return false;
+}
+
+FReply FAvaTemplatePageViewImpl::OnSyncStatusButtonClicked()
+{
+	if (!CanChangeSyncStatus())
+	{
+		return FReply::Unhandled();
+	}
+
+	UAvalanchePlaylist* Playlist = GetPlaylist();
+
+	if (IsValid(Playlist))
+	{
+		const FAvalanchePage& Page = GetPage();
+
+		if (Page.IsValidPage())
+		{
+			//TODO: Actual sync
+
+			return FReply::Handled();
+		}
+	}
+
+	return FReply::Unhandled();
+}
+
+bool FAvaTemplatePageViewImpl::CanChangeSyncStatus() const
+{
+	UAvalanchePlaylist* Playlist = GetPlaylist();
+
+	if (IsValid(Playlist))
+	{
+		const FAvalanchePage& Page = GetPage();
+
+		if (Page.IsValidPage())
+		{
+			TArray<FAvalanchePageStatus> Statuses = Page.GetPagePreviewStatuses(Playlist);
+			bool bNeedsSync = false;
+
+			for (const FAvalanchePageStatus& Status : Statuses)
+			{
+				if (Status.bNeedsSync)
+				{
+					bNeedsSync = true;
+					break;
+				}
+			}
+
+			return bNeedsSync;
+		}
+	}
+
+	return false;
+}
+
+bool FAvaTemplatePageViewImpl::IsTemplate() const 
+{ 
+	return true;
+}
diff --git a/Engine/Plugins/Experimental/Avalanche/Source/AvalancheMediaEditor/Private/Playlist/Pages/PageViews/AvaTemplatePageViewImpl.h b/Engine/Plugins/Experimental/Avalanche/Source/AvalancheMediaEditor/Private/Playlist/Pages/PageViews/AvaTemplatePageViewImpl.h
new file mode 100644
--- /dev/null
+++ b/Engine/Plugins/Experimental/Avalanche/Source/AvalancheMediaEditor/Private/Playlist/Pages/PageViews/AvaTemplatePageViewImpl.h
@@ -0,0 +1,26 @@
+// Copyright Epic Games, Inc. All Rights Reserved.
+
+#pragma once
+
+#include "AvaPageViewImpl.h"
+#include "Playlist/Pages/PageViews/AvaTemplatePageView.h"
+
+class FReply;
+struct FAssetData;
+
+class FAvaTemplatePageViewImpl : public FAvaPageViewImpl, public IAvaTemplatePageView
+{
+public:
+	UE_AVA_INHERITS(FAvaTemplatePageViewImpl, FAvaPageViewImpl, IAvaTemplatePageView);
+
+	FAvaTemplatePageViewImpl(int32 InPageId, UAvalanchePlaylist* InPlaylist, const TSharedPtr<SAvaPageList>& InPageList);
+
+	virtual UAvalanchePlaylist* GetPlaylist() const override;
+
+	virtual bool IsTemplate() const override;
+	virtual FReply OnPlayButtonClicked() override;
+	virtual bool CanPlay() const override;
+
+	FReply OnSyncStatusButtonClicked();
+	bool CanChangeSyncStatus() const;
+};
diff --git a/Engine/Plugins/Experimental/Avalanche/Source/AvalancheMediaEditor/Private/Playlist/Pages/Slate/SAvaInstancedPageList.cpp b/Engine/Plugins/Experimental/Avalanche/Source/AvalancheMediaEditor/Private/Playlist/Pages/Slate/SAvaInstancedPageList.cpp
new file mode 100644
--- /dev/null
+++ b/Engine/Plugins/Experimental/Avalanche/Source/AvalancheMediaEditor/Private/Playlist/Pages/Slate/SAvaInstancedPageList.cpp
@@ -0,0 +1,1374 @@
+// Copyright Epic Games, Inc. All Rights Reserved.
+
+#include "SAvaInstancedPageList.h"
+
+#include "Misc/PathViews.h"
+#include "Playlist/AvaPlaylistCommands.h"
+#include "Playlist/AvaPlaylistEditor.h"
+#include "Playlist/AvaPlaylistEditorUtils.h"
+#include "Playlist/AvaPlaylistPlaybackUtils.h"
+#include "Playlist/AvaRundownEditorSettings.h"
+#include "Playlist/AvalanchePage.h"
+#include "Playlist/AvalanchePlaylist.h"
+#include "Playlist/Factories/Filters/AvaPlaylistFactoriesUtils.h"
+#include "Playlist/Pages/Columns/AvaPageAssetNameColumn.h"
+#include "Playlist/Pages/Columns/AvaPageChannelSelectorColumn.h"
+#include "Playlist/Pages/Columns/AvaPageEnabledColumn.h"
+#include "Playlist/Pages/Columns/AvaPageIdColumn.h"
+#include "Playlist/Pages/Columns/AvaPageNameColumn.h"
+#include "Playlist/Pages/Columns/AvaPageStatusColumn.h"
+#include "Playlist/Pages/Columns/AvaPageTemplateNameColumn.h"
+#include "Playlist/Pages/PageViews/AvaInstancedPageViewImpl.h"
+#include "Playlist/TabFactories/AvaInstancedPageListTabFactory.h"
+#include "Playlist/TabFactories/AvaSubListDocumentTabFactory.h"
+
+#include "SAvaPageList.h"
+#include "ScopedTransaction.h"
+#include "Styling/StyleColors.h"
+#include "Widgets/Docking/SDockTab.h"
+#include "Widgets/Input/SButton.h"
+#include "Widgets/Input/SEditableTextBox.h"
+#include "Widgets/SNullWidget.h"
+#include "Widgets/Views/SListView.h"
+
+#define LOCTEXT_NAMESPACE "SAvaInstancedPageList"
+
+void SAvaInstancedPageList::PrivateRegisterAttributes(struct FSlateAttributeDescriptor::FInitializer&)
+{
+
+}
+
+void SAvaInstancedPageList::Construct(const FArguments& InArgs, TSharedPtr<FAvaPlaylistEditor> InPlaylistEditor, const FAvaPageListReference& InPageListReference)
+{
+	SAvaPageList::Construct(SAvaPageList::FArguments(), InPlaylistEditor, InPageListReference, EAvaPlaylistSearchListType::Instanced);
+
+	UAvalanchePlaylist* const Playlist = InPlaylistEditor->GetPlaylist();
+	check(Playlist);
+
+	check(InPageListReference.Type == EAvaPageListType::Instance || Playlist->IsValidSubList(PageListReference));
+
+	if (PageListReference.Type == EAvaPageListType::Instance)
+	{
+		TabId = FAvaInstancedPageListTabFactory::TabID;
+		Playlist->GetOnInstancedPageListChanged().AddSP(this, &SAvaInstancedPageList::OnInstancedPageListChanged);
+	}
+	else
+	{
+		FAvaSubListDocumentTabFactory::GetTabId(InPageListReference.SubListIndex);
+		Playlist->GetSubList(InPageListReference.SubListIndex).OnPageListChanged.AddSP(this, &SAvaInstancedPageList::OnInstancedPageListChanged);
+	}
+
+	if (Playlist->IsValidSubList(PageListReference))
+	{
+		SearchBar->InsertSlot(0)
+			.AutoWidth()
+			.VAlign(EVerticalAlignment::VAlign_Center)
+			.Padding(5.f, 0.f)
+			[
+				SNew(STextBlock)
+				.Text(LOCTEXT("Name", "Name:"))
+			];
+
+		SearchBar->InsertSlot(1)
+			.AutoWidth()
+			.VAlign(EVerticalAlignment::VAlign_Center)
+			.Padding(5.f, 0.f)
+			[
+				SNew(SEditableTextBox)
+				.HintText(LOCTEXT("PageViewName", "Page View Name"))
+				.OnTextCommitted(this, &SAvaInstancedPageList::OnPageViewNameCommitted)
+				.Text(this, &SAvaInstancedPageList::GetPageViewName)
+				.MinDesiredWidth(150.f)
+			];
+	}
+
+	SearchBar->AddSlot()
+		.AutoWidth()
+		.VAlign(EVerticalAlignment::VAlign_Center)
+		.Padding(5.f, 0.f)
+		[
+			SNew(SButton)
+			.ToolTipText(LOCTEXT("MakeActiveTooltip", "Page List used for Show Controls (Take In, etc.)"))
+			.OnClicked(this, &SAvaInstancedPageList::MakeActive)
+			.IsEnabled(this, &SAvaInstancedPageList::CanMakeActive)
+			.Visibility_Lambda([this]()
+			{
+				if (TSharedPtr<FAvaPlaylistEditor> PlaylistEditor = PlaylistEditorWeak.Pin())
+				{
+					UAvalanchePlaylist* const Playlist = PlaylistEditor->GetPlaylist();
+			
+					if (IsValid(Playlist))
+					{
+						return Playlist->GetSubLists().IsEmpty()? EVisibility::Collapsed : EVisibility::Visible;
+					}
+				}
+				return EVisibility::Collapsed;
+			})
+			.ButtonColorAndOpacity(this, &SAvaInstancedPageList::GetMakeActiveButtonColor)
+			[
+				SNew(STextBlock)
+				.Text(LOCTEXT("ActivePage", "Show-Control Pages"))
+			]
+		];
+
+	Refresh();
+}
+
+SAvaInstancedPageList::~SAvaInstancedPageList()
+{
+	if (UAvalanchePlaylist* const Playlist = GetValidPlaylist())
+	{
+		if (PageListReference.Type == EAvaPageListType::Instance)
+		{
+			Playlist->GetOnInstancedPageListChanged().RemoveAll(this);
+		}
+		else if (Playlist->IsValidSubList(PageListReference))
+		{
+			Playlist->GetSubList(PageListReference.SubListIndex).OnPageListChanged.RemoveAll(this);
+		}
+	}
+}
+
+void SAvaInstancedPageList::Refresh()
+{
+	if (const TSharedPtr<FAvaPlaylistEditor> PlaylistEditor = PlaylistEditorWeak.Pin())
+	{
+		UAvalanchePlaylist* const Playlist = PlaylistEditor->GetPlaylist();
+		check(Playlist);
+
+		TArray<FAvalanchePage> Pages = Playlist->GetInstancedPages().Pages;
+
+		if (Playlist->IsValidSubList(PageListReference))
+		{
+			Pages.Empty();
+
+			for (const int32 PageId : Playlist->GetSubList(PageListReference.SubListIndex).PageIds)
+			{
+				if (const int32* Index = Playlist->GetInstancedPages().PageIndices.Find(PageId))
+				{
+					Pages.Add(Playlist->GetInstancedPages().Pages[*Index]);
+				}
+			}
+		}
+
+		PageViews.Reset(Pages.Num());
+
+		for (const FAvalanchePage& Page : Pages)
+		{
+			if (PlaylistEditor->IsInstancedPageVisible(Page))
+			{
+				PageViews.Emplace(MakeShared<FAvaInstancedPageViewImpl>(Page.GetPageId(), Playlist, SharedThis(this)));
+			}
+		}
+
+		PageListView->RequestListRefresh();
+	}
+}
+
+void SAvaInstancedPageList::CreateColumns()
+{
+	HeaderRow = SNew(SHeaderRow)
+		.Visibility(EVisibility::Visible)
+		.CanSelectGeneratedColumn(true);
+
+	Columns.Empty();
+	HeaderRow->ClearColumns();
+
+	//TODO: Extensibility?
+	TArray<TSharedPtr<IAvaPageViewColumn>> FoundColumns;
+	FoundColumns.Add(MakeShared<FAvaPageEnabledColumn>());
+	FoundColumns.Add(MakeShared<FAvaPageIdColumn>());
+	FoundColumns.Add(MakeShared<FAvaPageTemplateNameColumn>());
+	FoundColumns.Add(MakeShared<FAvaPageNameColumn>());
+	FoundColumns.Add(MakeShared<FAvaPageAssetNameColumn>());
+	FoundColumns.Add(MakeShared<FAvaPageChannelSelectorColumn>());
+	FoundColumns.Add(MakeShared<FAvaPageStatusColumn>());
+
+	for (const TSharedPtr<IAvaPageViewColumn>& Column : FoundColumns)
+	{
+		const FName ColumnId = Column->GetColumnId();
+		Columns.Add(ColumnId, Column);
+		HeaderRow->AddColumn(Column->ConstructHeaderRowColumn());
+		HeaderRow->SetShowGeneratedColumn(ColumnId, false);
+	}
+}
+
+TSharedPtr<SWidget> SAvaInstancedPageList::OnContextMenuOpening()
+{
+	if (TSharedPtr<FAvaPlaylistEditor> PlaylistEditor = PlaylistEditorWeak.Pin())
+	{
+		return GetPageListContextMenu();
+	}
+
+	return SNullWidget::NullWidget;
+}
+
+void SAvaInstancedPageList::BindCommands()
+{
+	SAvaPageList::BindCommands();
+
+	//Playlist Commands
+	{
+		const FAvaPlaylistCommands& PlaylistCommands = FAvaPlaylistCommands::Get();
+
+		CommandList->MapAction(PlaylistCommands.RemovePage,
+			FExecuteAction::CreateSP(this, &SAvaPageList::RemoveSelectedPages),
+			FCanExecuteAction::CreateSP(this, &SAvaPageList::CanRemoveSelectedPages));
+
+		CommandList->MapAction(PlaylistCommands.RenumberPage,
+			FExecuteAction::CreateSP(this, &SAvaPageList::RenumberSelectedPage),
+			FCanExecuteAction::CreateSP(this, &SAvaPageList::CanRenumberSelectedPage));
+
+		CommandList->MapAction(PlaylistCommands.ReimportPage,
+			FExecuteAction::CreateSP(this, &SAvaPageList::ReimportSelectedPage),
+			FCanExecuteAction::CreateSP(this, &SAvaPageList::CanReimportSelectedPage));
+
+		CommandList->MapAction(PlaylistCommands.EditPageSource,
+			FExecuteAction::CreateSP(this, &SAvaPageList::EditSelectedPageSource),
+			FCanExecuteAction::CreateSP(this, &SAvaPageList::CanEditSelectedPageSource));
+
+		CommandList->MapAction(PlaylistCommands.ExportPagesToPlaylist,
+			FExecuteAction::CreateSP(this, &SAvaPageList::ExportSelectedPagesToPlaylist),
+			FCanExecuteAction::CreateSP(this, &SAvaPageList::CanExportSelectedPagesToPlaylist));
+
+		CommandList->MapAction(PlaylistCommands.ExportPagesToJson,
+			FExecuteAction::CreateSP(this, &SAvaPageList::ExportSelectedPagesToExternalFile, TEXT("json")),
+			FCanExecuteAction::CreateSP(this, &SAvaPageList::CanExportSelectedPagesToExternalFile, TEXT("json")));
+
+		CommandList->MapAction(PlaylistCommands.ExportPagesToXml,
+			FExecuteAction::CreateSP(this, &SAvaPageList::ExportSelectedPagesToExternalFile, TEXT("xml")),
+			FCanExecuteAction::CreateSP(this, &SAvaPageList::CanExportSelectedPagesToExternalFile, TEXT("xml")));
+
+		CommandList->MapAction(PlaylistCommands.Play,
+			FExecuteAction::CreateSP(this, &SAvaInstancedPageList::PlaySelectedPage),
+			FCanExecuteAction::CreateSP(this, &SAvaInstancedPageList::CanPlaySelectedPage));
+
+		CommandList->MapAction(PlaylistCommands.UpdateValues,
+			FExecuteAction::CreateSP(this, &SAvaInstancedPageList::UpdateValuesOnSelectedPage),
+			FCanExecuteAction::CreateSP(this, &SAvaInstancedPageList::CanUpdateValuesOnSelectedPage));
+
+		CommandList->MapAction(PlaylistCommands.Stop,
+			FExecuteAction::CreateSP(this, &SAvaInstancedPageList::StopSelectedPage, false),
+			FCanExecuteAction::CreateSP(this, &SAvaInstancedPageList::CanStopSelectedPage, false));
+
+		CommandList->MapAction(PlaylistCommands.ForceStop,
+			FExecuteAction::CreateSP(this, &SAvaInstancedPageList::StopSelectedPage, true),
+			FCanExecuteAction::CreateSP(this, &SAvaInstancedPageList::CanStopSelectedPage, true));
+
+		CommandList->MapAction(PlaylistCommands.Continue,
+			FExecuteAction::CreateSP(this, &SAvaInstancedPageList::ContinueSelectedPage),
+			FCanExecuteAction::CreateSP(this, &SAvaInstancedPageList::CanContinueSelectedPage));
+
+		CommandList->MapAction(PlaylistCommands.PlayNext,
+			FExecuteAction::CreateSP(this, &SAvaInstancedPageList::PlayNextPageNoReturn),
+			FCanExecuteAction::CreateSP(this, &SAvaInstancedPageList::CanPlayNextPage));
+
+		CommandList->MapAction(PlaylistCommands.PreviewFrame,
+			FExecuteAction::CreateSP(this, &SAvaPageList::PreviewPlaySelectedPage, true),
+			FCanExecuteAction::CreateSP(this, &SAvaPageList::CanPreviewPlaySelectedPage));
+
+		CommandList->MapAction(PlaylistCommands.PreviewPlay,
+			FExecuteAction::CreateSP(this, &SAvaPageList::PreviewPlaySelectedPage, false),
+			FCanExecuteAction::CreateSP(this, &SAvaPageList::CanPreviewPlaySelectedPage));
+
+		CommandList->MapAction(PlaylistCommands.PreviewStop,
+			FExecuteAction::CreateSP(this, &SAvaPageList::PreviewStopSelectedPage, false),
+			FCanExecuteAction::CreateSP(this, &SAvaPageList::CanPreviewStopSelectedPage, false));
+
+		CommandList->MapAction(PlaylistCommands.PreviewForceStop,
+			FExecuteAction::CreateSP(this, &SAvaPageList::PreviewStopSelectedPage, true),
+			FCanExecuteAction::CreateSP(this, &SAvaPageList::CanPreviewStopSelectedPage, true));
+
+		CommandList->MapAction(PlaylistCommands.PreviewContinue,
+			FExecuteAction::CreateSP(this, &SAvaPageList::PreviewContinueSelectedPage),
+			FCanExecuteAction::CreateSP(this, &SAvaPageList::CanPreviewContinueSelectedPage));
+
+		CommandList->MapAction(PlaylistCommands.PreviewPlayNext,
+			FExecuteAction::CreateSP(this, &SAvaPageList::PreviewPlayNextPage),
+			FCanExecuteAction::CreateSP(this, &SAvaPageList::CanPreviewPlayNextPage));
+
+		CommandList->MapAction(PlaylistCommands.TakeToProgram,
+			FExecuteAction::CreateSP(this, &SAvaPageList::TakeToProgram),
+			FCanExecuteAction::CreateSP(this, &SAvaPageList::CanTakeToProgram));
+	}
+}
+
+bool SAvaInstancedPageList::HandleDropAvalancheAssets(const TArray<FSoftObjectPath>& InAvaAssets, EItemDropZone InDropZone, const FAvaPageViewPtr& InItem)
+{
+	UAvalanchePlaylist* Playlist = GetValidPlaylist();
+
+	if (!Playlist)
+	{
+		return false;
+	}
+	
+	if (PageListReference.Type == EAvaPageListType::View && !Playlist->IsValidSubList(PageListReference))
+	{
+		return false;
+	}
+
+	FAvaPageInsertPosition InsertAt = MakeInsertPosition(InDropZone, InItem);
+	
+	TArray<FSoftObjectPath> NewAvaAssets = InAvaAssets;
+	TArray<int32> NewPageIds;
+
+	// If we are adding above, they should be reversed, so the last is added first
+	// and the next to last added above that, etc.
+	if (InsertAt.IsAddAbove())
+	{
+		Algo::Reverse(NewAvaAssets);
+	}
+
+	bool bHasValidAssets = false;
+	Playlist->Modify();
+
+	for (const FSoftObjectPath& AvaAsset : NewAvaAssets)
+	{
+		if (AvaAsset.IsNull())
+		{
+			continue;
+		}
+
+		// Create Template
+		const int32 NewTemplateId = Playlist->AddTemplate();
+
+		if (NewTemplateId == FAvalanchePage::InvalidPageId)
+		{
+			continue;
+		}
+
+		Playlist->GetPage(NewTemplateId).UpdateAvalancheAsset(AvaAsset);
+
+		int32 NewInstanceId = FAvalanchePage::InvalidPageId;
+
+		if (PageListReference.Type == EAvaPageListType::Instance && InsertAt.IsValid())
+		{
+			NewInstanceId = Playlist->AddPageFromTemplate(NewTemplateId, FAvaPageIdGeneratorParams::FromInsertPosition(InsertAt), InsertAt);
+		}
+		else
+		{
+			NewInstanceId = Playlist->AddPageFromTemplate(NewTemplateId);
+		}
+
+		if (NewInstanceId == FAvalanchePage::InvalidPageId)
+		{
+			continue;
+		}
+
+		if (PageListReference.Type == EAvaPageListType::View)
+		{
+			Playlist->AddPageToSubList(PageListReference.SubListIndex, NewInstanceId, InsertAt);
+		}
+
+		InsertAt.ConditionalUpdateAdjacentId(NewInstanceId);
+		NewPageIds.Add(NewInstanceId);
+		bHasValidAssets = true;
+	}
+
+	if (bHasValidAssets)
+	{
+		Refresh();
+		DeselectPages();
+		SelectPages(NewPageIds);
+	}
+
+	return bHasValidAssets;
+}
+
+bool SAvaInstancedPageList::HandleDropPlaylists(const TArray<FSoftObjectPath>& InPlaylistPaths, EItemDropZone InDropZone, const FAvaPageViewPtr& InItem)
+{
+	UAvalanchePlaylist* Playlist = GetValidPlaylist();
+
+	if (!Playlist)
+	{
+		return false;
+	}
+
+	FAvaPageInsertPosition InsertPosition = MakeInsertPosition(InDropZone, InItem);
+	
+	TArray<FSoftObjectPath> NewPlaylistPaths = InPlaylistPaths;
+	TArray<int32> NewPageIds;
+	
+	// If we are adding above, they should be reversed, so the last is added first
+	// and the next to last added above that, etc.
+	if (!InsertPosition.bAddBelow)
+	{
+		Algo::Reverse(NewPlaylistPaths);
+	}
+	
+	Playlist->Modify();
+
+	for (const FSoftObjectPath& PlaylistPath : NewPlaylistPaths)
+	{
+		if (PlaylistPath.IsNull())
+		{
+			continue;
+		}
+
+		if (const UAvalanchePlaylist* SourcePlaylist = Cast<UAvalanchePlaylist>(PlaylistPath.TryLoad()))
+		{
+			TArray<int32> ImportedPageIds = UE::AvaPlaylistEditor::Utils::ImportInstancedPagesFromPlaylist(Playlist, SourcePlaylist, InsertPosition);
+			if (!ImportedPageIds.IsEmpty())
+			{
+				NewPageIds.Append(ImportedPageIds);
+				InsertPosition.ConditionalUpdateAdjacentId(NewPageIds.Last());	// todo: last or first, depends on AddBelow.
+			}
+		}
+	}
+
+	if (!NewPageIds.IsEmpty())
+	{
+		Refresh();
+		DeselectPages();
+		SelectPages(NewPageIds);
+	}
+
+	return !NewPageIds.IsEmpty();
+}
+
+bool SAvaInstancedPageList::HandleDropPageIds(const FAvaPageListReference& InPageListReference, const TArray<int32>& InPageIds, EItemDropZone InDropZone, const FAvaPageViewPtr& InItem)
+{
+	// Nothing to do.
+	if (InPageIds.IsEmpty())
+	{
+		return true;
+	}
+	
+	UAvalanchePlaylist* Playlist = GetValidPlaylist();
+
+	if (!Playlist)
+	{
+		return false;
+	}
+
+	const bool bFromTemplates = InPageListReference.Type == EAvaPageListType::Template;
+	const bool bFromMainList = InPageListReference.Type == EAvaPageListType::Instance;
+	const bool bFromSubList = Playlist->IsValidSubList(InPageListReference);
+
+	if (!bFromTemplates && !bFromMainList && !bFromSubList)
+	{
+		return false;
+	}
+
+	const bool bToMainList = PageListReference.Type == EAvaPageListType::Instance;
+	const bool bToSubList = Playlist->IsValidSubList(PageListReference);
+
+	if (!bToMainList && !bToSubList)
+	{
+		return false;
+	}
+
+	FScopedTransaction DropTransaction(LOCTEXT("DropTransaction", "Drop Pages"));
+	Playlist->Modify();
+	bool bDidSomething = false;
+
+	if (bToMainList)
+	{
+		if (bFromTemplates)
+		{
+			bDidSomething = HandleDropPageIdsOnMainListFromTemplates(InPageIds, InDropZone, InItem);
+		}
+		else if (bFromMainList)
+		{
+			bDidSomething = HandleDropPageIdsOnMainListFromMainList(InPageIds, InDropZone, InItem);
+		}
+		// Cannot drop from sub list to main list
+	}
+	else if (bToSubList)
+	{
+		if (bFromTemplates)
+		{
+			bDidSomething = HandleDropPageIdsOnSubListFromTemplates(InPageIds, InDropZone, InItem);
+		}
+		else if (bFromMainList)
+		{
+			bDidSomething = HandleDropPageIdsOnSubListFromMainList(InPageIds, InDropZone, InItem);
+		}
+		else if (bFromSubList)
+		{
+			bDidSomething = HandleDropPageIdsOnSubListFromSubList(InPageListReference.SubListIndex, InPageIds, InDropZone, InItem);
+		}
+	}
+
+	if (!bDidSomething)
+	{
+		// We didn't do anything.
+		DropTransaction.Cancel();
+	}
+
+	return bDidSomething;
+}
+
+bool SAvaInstancedPageList::HandleDropPageIdsOnMainListFromTemplates(const TArray<int32>& InPageIds, EItemDropZone InDropZone, const FAvaPageViewPtr& InItem)
+{
+	if (PageListReference.Type != EAvaPageListType::Instance)
+	{
+		return false;
+	}
+
+	// Nothing to do.
+	if (InPageIds.IsEmpty())
+	{
+		return true;
+	}
+
+	UAvalanchePlaylist* Playlist = GetValidPlaylist();
+
+	if (!Playlist || !Playlist->CanAddPage())
+	{
+		return false;
+	}
+
+	FAvaPageInsertPosition InsertAt = MakeInsertPosition(InDropZone, InItem);	
+
+	if (InsertAt.IsValid())
+	{
+		const int32 DroppedOnPageIndex = Playlist->GetInstancedPages().GetPageIndex(InsertAt.AdjacentId);
+
+		if (DroppedOnPageIndex == INDEX_NONE)
+		{
+			InsertAt.AdjacentId = FAvalanchePage::InvalidPageId;
+		}
+	}
+	
+	TArray<int32> ActualPageIds = InPageIds;
+	TArray<int32> NewPageIds;
+	NewPageIds.Reserve(InPageIds.Num());
+
+	if (InsertAt.IsAddAbove())
+	{
+		Algo::Reverse(ActualPageIds);
+	}
+
+	for (const int32 PageId : ActualPageIds)
+	{
+		int32 NewInstanceId = Playlist->AddPageFromTemplate(PageId, FAvaPageIdGeneratorParams::FromInsertPosition(InsertAt), InsertAt);
+
+		if (NewInstanceId != FAvalanchePage::InvalidPageId)
+		{
+			NewPageIds.Add(NewInstanceId);
+			InsertAt.ConditionalUpdateAdjacentId(NewInstanceId);
+		}
+	}
+
+	if (NewPageIds.IsEmpty())
+	{
+		return false;
+	}
+
+	SelectPages(NewPageIds, true);
+
+	return true;
+}
+
+bool SAvaInstancedPageList::HandleDropExternalFiles(const TArray<FString>& InFiles, EItemDropZone InDropZone, const FAvaPageViewPtr& InItem)
+{
+	UAvalanchePlaylist* Playlist = GetValidPlaylist();
+
+	if (!Playlist)
+	{
+		return false;
+	}
+	
+	FAvaPageInsertPosition InsertAt = MakeInsertPosition(InDropZone, InItem);
+	
+	TArray<FString> NewFiles = InFiles;
+	TArray<int32> NewPageIds;
+
+	// If we are adding above, they should be reversed, so the last is added first
+	// and the next to last added above that, etc.
+	if (InsertAt.IsAddAbove())
+	{
+		Algo::Reverse(NewFiles);
+	}
+	
+	Playlist->Modify();
+	
+	for (const FString& File : InFiles)
+	{
+		if (FPathViews::GetExtension(File).Equals(TEXT("json"), ESearchCase::IgnoreCase))
+		{
+			using namespace UE::AvaPlaylistEditor::Utils;
+			const TStrongObjectPtr<UAvalanchePlaylist> TmpPlaylist(NewObject<UAvalanchePlaylist>());
+			
+			if (LoadPlaylistFromJson(TmpPlaylist.Get(), *File))
+			{
+				if (!TmpPlaylist->GetInstancedPages().Pages.IsEmpty())
+				{
+					TArray<int32> ImportedPageIds = ImportInstancedPagesFromPlaylist(GetPlaylist(), TmpPlaylist.Get(), InsertAt);
+					if (!ImportedPageIds.IsEmpty())
+					{
+						NewPageIds.Append(ImportedPageIds);
+						InsertAt.ConditionalUpdateAdjacentId(NewPageIds.Last());	// todo: last or first, depends on AddBelow.
+					}
+					else
+					{
+						UE_LOG(LogAvaPlaylist, Error, TEXT("Failed to merge %s in current playlist."), *File);	
+					}
+				}
+				else
+				{
+					UE_LOG(LogAvaPlaylist, Warning, TEXT("Merging %s has no instanced pages to merge."), *File);
+				}
+			}
+			else
+			{
+				UE_LOG(LogAvaPlaylist, Error, TEXT("%s is not a valid playlist."), *File);
+			}
+		}
+		else
+		{
+			UE_LOG(LogAvaPlaylist, Error, TEXT("%s is not a supported file type. Only playlists in \"json\" format are supported."), *File);
+		}
+	}
+	
+	if (!NewPageIds.IsEmpty())
+	{
+		Refresh();
+		DeselectPages();
+		SelectPages(NewPageIds);
+	}
+
+	return !NewPageIds.IsEmpty();
+}
+
+bool SAvaInstancedPageList::HandleDropPageIdsOnMainListFromMainList(const TArray<int32>& InPageIds, EItemDropZone InDropZone, const FAvaPageViewPtr& InItem)
+{
+	if (PageListReference.Type != EAvaPageListType::Instance)
+	{
+		return false;
+	}
+
+	// Nothing to do.
+	if (InPageIds.IsEmpty())
+	{
+		return true;
+	}
+
+	UAvalanchePlaylist* Playlist = GetValidPlaylist();
+
+	if (!Playlist || !Playlist->CanChangePageOrder())
+	{
+		return false;
+	}
+
+	const FAvalanchePageCollection& InstancedPageCollection = Playlist->GetInstancedPages();
+	int32 DroppedOnPageIndex = INDEX_NONE;
+
+	if (InItem.IsValid() && InstancedPageCollection.PageIndices.Contains(InItem->GetPageId()))
+	{
+		DroppedOnPageIndex = InstancedPageCollection.PageIndices[InItem->GetPageId()];
+
+		// Nothing to do.
+		if (InPageIds.Num() == 1 && DroppedOnPageIndex == InPageIds[0])
+		{
+			return true;
+		}
+	}
+
+	TArray<int32> MovedPageIdIndices;
+	TArray<int32> NewPageOrder;
+	TArray<int32> NewSelectedIds;
+
+	MovedPageIdIndices.Reserve(InPageIds.Num());
+	NewPageOrder.Reserve(PageViews.Num());
+	NewSelectedIds.Reserve(InPageIds.Num());
+
+	// This has the byproduct of removing invalid ids.
+	for (int32 PageId : InPageIds)
+	{
+		if (const int32* PageIndexPtr = InstancedPageCollection.PageIndices.Find(PageId))
+		{
+			MovedPageIdIndices.Add(*PageIndexPtr);
+			NewSelectedIds.Add(PageId);
+		}
+	}
+
+	// Nothing to do.
+	if (MovedPageIdIndices.IsEmpty())
+	{
+		return true;
+	}
+
+	for (int32 PageViewIndex = 0; PageViewIndex < PageViews.Num(); ++PageViewIndex)
+	{
+		const bool bHasMovedThisPage = MovedPageIdIndices.Contains(PageViewIndex);
+
+		if (PageViewIndex == DroppedOnPageIndex)
+		{
+			const bool bAddBefore = InDropZone == EItemDropZone::AboveItem;
+
+			if (bAddBefore)
+			{
+				NewPageOrder.Append(MovedPageIdIndices);
+			}
+
+			// If we moved the page we dropped onto, it will already be in the list
+			if (!bHasMovedThisPage)
+			{
+				NewPageOrder.Add(PageViewIndex);
+			}
+
+			if (!bAddBefore)
+			{
+				NewPageOrder.Append(MovedPageIdIndices);
+			}
+
+			continue;
+		}
+
+		if (bHasMovedThisPage)
+		{
+			continue;
+		}
+
+		NewPageOrder.Add(PageViewIndex);
+	}
+
+	Playlist->ChangePageOrder(UAvalanchePlaylist::InstancePageList, NewPageOrder);
+	SelectPages(NewSelectedIds, true);
+
+	return true;
+}
+
+bool SAvaInstancedPageList::HandleDropPageIdsOnSubListFromTemplates(const TArray<int32>& InPageIds, EItemDropZone InDropZone, const FAvaPageViewPtr& InItem)
+{
+	if (PageListReference.Type != EAvaPageListType::View)
+	{
+		return false;
+	}
+
+	// Nothing to do.
+	if (InPageIds.IsEmpty())
+	{
+		return true;
+	}
+
+	UAvalanchePlaylist* Playlist = GetValidPlaylist();
+
+	if (!Playlist || !Playlist->CanAddPage() || !Playlist->IsValidSubList(PageListReference))
+	{
+		return false;
+	}
+
+	FAvaPageInsertPosition InsertAt = MakeInsertPosition(InDropZone, InItem);
+	
+	if (InsertAt.IsValid())
+	{
+		// Make sure the sublist has the dropped on item. (It should)
+		const FAvalancheSubList& ToSubList = Playlist->GetSubList(PageListReference.SubListIndex);
+		if (!ToSubList.PageIds.Contains(InsertAt.AdjacentId))
+		{
+			InsertAt.AdjacentId = FAvalanchePage::InvalidPageId;
+		}
+	}
+	
+	TArray<int32> ActualPageIds = InPageIds;
+	TArray<int32> NewPageIds;
+	NewPageIds.Reserve(InPageIds.Num());
+
+	if (InsertAt.IsAddAbove())
+	{
+		Algo::Reverse(ActualPageIds);
+	}
+
+	for (const int32 PageId : ActualPageIds)
+	{
+		const int32 NewInstanceId = Playlist->AddPageFromTemplate(PageId, FAvaPageIdGeneratorParams::FromInsertPosition(InsertAt), InsertAt);
+
+		if (NewInstanceId != FAvalanchePage::InvalidPageId)
+		{
+			NewPageIds.Add(NewInstanceId);
+		}
+	}
+
+	if (NewPageIds.IsEmpty())
+	{
+		return false;
+	}
+
+	for (const int32 PageId : NewPageIds)
+	{
+		if (Playlist->AddPageToSubList(PageListReference.SubListIndex, PageId, InsertAt))
+		{
+			InsertAt.ConditionalUpdateAdjacentId(PageId);
+		}
+	}
+
+	SelectPages(NewPageIds, true);
+
+	return true;
+}
+
+bool SAvaInstancedPageList::HandleDropPageIdsOnSubListFromMainList(const TArray<int32>& InPageIds, EItemDropZone InDropZone, const FAvaPageViewPtr& InItem)
+{
+	if (PageListReference.Type != EAvaPageListType::View)
+	{
+		return false;
+	}
+
+	// Nothing to do.
+	if (InPageIds.IsEmpty())
+	{
+		return true;
+	}
+
+	UAvalanchePlaylist* Playlist = GetValidPlaylist();
+
+	if (!Playlist || !Playlist->CanAddPage() || !Playlist->IsValidSubList(PageListReference))
+	{
+		return false;
+	}
+
+	FAvalancheSubList& ToSubList = Playlist->GetSubList(PageListReference.SubListIndex);
+
+	FAvaPageInsertPosition InsertAt = MakeInsertPosition(InDropZone, InItem);
+	
+	int32 DroppedOnPageIndex = INDEX_NONE;
+
+	if (InsertAt.IsValid())
+	{
+		DroppedOnPageIndex = ToSubList.PageIds.IndexOfByKey(InsertAt.AdjacentId);
+		if (DroppedOnPageIndex == INDEX_NONE)
+		{
+			InsertAt.AdjacentId = FAvalanchePage::InvalidPageId;
+		}
+	}
+
+	TArray<int32> ActualPageIds;
+	ActualPageIds.Reserve(InPageIds.Num());
+
+	// Do not add pages already in this sub list
+	for (int32 PageId : InPageIds)
+	{
+		if (!ToSubList.PageIds.Contains(PageId))
+		{
+			ActualPageIds.Add(PageId);
+		}
+	}
+
+	if (InsertAt.IsAddAbove())
+	{
+		Algo::Reverse(ActualPageIds);
+	}
+
+	if (InsertAt.IsValid() && InPageIds.Contains(InsertAt.AdjacentId))
+	{
+		// If we have a valid index and we're adding above, iterate up until we find one we're not removing
+		if (DroppedOnPageIndex != INDEX_NONE && InsertAt.IsAddAbove())
+		{
+			for (DroppedOnPageIndex = DroppedOnPageIndex - 1; DroppedOnPageIndex > INDEX_NONE; --DroppedOnPageIndex)
+			{
+				if (InPageIds.Contains(ToSubList.PageIds[DroppedOnPageIndex]) == false)
+				{
+					break;
+				}
+			}
+		}
+
+		// If we have no valid index, iterate from the end of the array until we find out that we aren't removing
+		if (DroppedOnPageIndex == INDEX_NONE)
+		{
+			for (DroppedOnPageIndex = ToSubList.PageIds.Num() - 1; DroppedOnPageIndex > INDEX_NONE; --DroppedOnPageIndex)
+			{
+				if (InPageIds.Contains(ToSubList.PageIds[DroppedOnPageIndex]) == false)
+				{
+					break;
+				}
+			}
+		}
+		else if (InsertAt.IsAddBelow())
+		{
+			for (DroppedOnPageIndex = DroppedOnPageIndex + 1; DroppedOnPageIndex < ToSubList.PageIds.Num(); ++DroppedOnPageIndex)
+			{
+				if (InPageIds.Contains(ToSubList.PageIds[DroppedOnPageIndex]) == false)
+				{
+					break;
+				}
+			}
+		}
+
+		// If we still have no valid index, just forget about dropping on a page
+		if (DroppedOnPageIndex == INDEX_NONE || !ToSubList.PageIds.IsValidIndex(DroppedOnPageIndex))
+		{
+			InsertAt.AdjacentId = FAvalanchePage::InvalidPageId;
+		}
+		// Make sure our id is correct.
+		else
+		{
+			InsertAt.AdjacentId = ToSubList.PageIds[DroppedOnPageIndex];
+		}
+	}
+
+	for (const int32 PageId : ActualPageIds)
+	{
+		Playlist->AddPageToSubList(PageListReference.SubListIndex, PageId, InsertAt);
+		InsertAt.ConditionalUpdateAdjacentId(PageId);
+	}
+
+	SelectPages(ActualPageIds, true);
+
+	return true;
+}
+
+bool SAvaInstancedPageList::HandleDropPageIdsOnSubListFromSubList(int32 InFromList, const TArray<int32>& InPageIds, EItemDropZone InDropZone, const FAvaPageViewPtr& InItem)
+{
+	if (PageListReference.Type != EAvaPageListType::View)
+	{
+		return false;
+	}
+
+	UAvalanchePlaylist* Playlist = GetValidPlaylist();
+	if (!Playlist || !Playlist->IsValidSubList(PageListReference) || !Playlist->IsValidSubListIndex(InFromList))
+	{
+		return false;
+	}
+
+	if (PageListReference.SubListIndex == InFromList)
+	{
+		if (!Playlist->CanChangePageOrder())
+		{
+			return false;
+		}
+	}
+	else
+	{
+		if (!Playlist->CanAddPage())
+		{
+			return false;
+		}
+	}
+
+	FAvalancheSubList& ToSubList = Playlist->GetSubList(PageListReference.SubListIndex);
+
+	FAvaPageInsertPosition InsertAt = MakeInsertPosition(InDropZone, InItem);
+	
+	int32 DroppedOnPageIndex = INDEX_NONE;
+
+	if (InsertAt.IsValid())
+	{
+		DroppedOnPageIndex = ToSubList.PageIds.IndexOfByKey(InsertAt.AdjacentId);
+
+		if (DroppedOnPageIndex == INDEX_NONE)
+		{
+			InsertAt.AdjacentId = FAvalanchePage::InvalidPageId;
+		}
+	}
+
+	TArray<int32> ActualPageIds;
+	ActualPageIds.Reserve(InPageIds.Num());
+
+	// Nothing to do
+	if (PageListReference.SubListIndex == InFromList)
+	{
+		if (ToSubList.PageIds.Num() <= 1)
+		{
+			return true;
+		}
+
+		// Remove and re-add all pages
+		ActualPageIds = InPageIds;
+	}
+	else
+	{
+		// Do not add pages already in this sub list
+		for (int32 PageId : InPageIds)
+		{
+			if (!ToSubList.PageIds.Contains(PageId))
+			{
+				ActualPageIds.Add(PageId);
+			}
+		}
+	}
+
+	// Nothing to do
+	if (ActualPageIds.IsEmpty())
+	{
+		return true;
+	}
+
+	if (InsertAt.IsAddAbove())
+	{
+		Algo::Reverse(ActualPageIds);
+	}
+
+	if (InsertAt.IsValid() && InPageIds.Contains(InsertAt.AdjacentId))
+	{
+		// If we have a valid index and we're adding above, iterate up until we find one we're not removing
+		if (DroppedOnPageIndex != INDEX_NONE && InsertAt.IsAddAbove())
+		{
+			for (DroppedOnPageIndex = DroppedOnPageIndex - 1; DroppedOnPageIndex > INDEX_NONE; --DroppedOnPageIndex)
+			{
+				if (InPageIds.Contains(ToSubList.PageIds[DroppedOnPageIndex]) == false)
+				{
+					break;
+				}
+			}
+		}
+
+		// If we have no valid index, iterate from the end of the array until we find out that we aren't removing
+		if (DroppedOnPageIndex == INDEX_NONE)
+		{
+			for (DroppedOnPageIndex = ToSubList.PageIds.Num() - 1; DroppedOnPageIndex > INDEX_NONE; --DroppedOnPageIndex)
+			{
+				if (InPageIds.Contains(ToSubList.PageIds[DroppedOnPageIndex]) == false)
+				{
+					break;
+				}
+			}
+		}
+		else if (InsertAt.IsAddBelow())
+		{
+			for (DroppedOnPageIndex = DroppedOnPageIndex + 1; DroppedOnPageIndex < ToSubList.PageIds.Num(); ++DroppedOnPageIndex)
+			{
+				if (InPageIds.Contains(ToSubList.PageIds[DroppedOnPageIndex]) == false)
+				{
+					break;
+				}
+			}
+		}
+
+		// If we still have no valid index, just forget about dropping on a page
+		if (DroppedOnPageIndex == INDEX_NONE || !ToSubList.PageIds.IsValidIndex(DroppedOnPageIndex))
+		{
+			InsertAt.AdjacentId = FAvalanchePage::InvalidPageId;
+		}
+		// Make sure our id is correct.
+		else
+		{
+			InsertAt.AdjacentId = ToSubList.PageIds[DroppedOnPageIndex];
+		}
+	}
+
+	Playlist->RemovePagesFromSubList(InFromList, InPageIds);
+
+	for (const int32 PageId : ActualPageIds)
+	{
+		Playlist->AddPageToSubList(PageListReference.SubListIndex, PageId, InsertAt);
+		InsertAt.ConditionalUpdateAdjacentId(PageId);
+	}
+
+	SelectPages(ActualPageIds, true);
+
+	return true;
+}
+
+void SAvaInstancedPageList::OnTabActivated(TSharedRef<SDockTab> InDockTab, ETabActivationCause InActivationCause)
+{
+	MakeActive();
+}
+
+void SAvaInstancedPageList::PlaySelectedPage() const
+{
+	if (UAvalanchePlaylist* Playlist = GetPlaylist())
+	{
+		const TArray<int32> PageIds = GetPagesToTakeIn();
+		Playlist->PlayPages(PageIds, EAvaPlayType::PlayFromStart);
+	}
+}
+
+bool SAvaInstancedPageList::CanPlaySelectedPage() const
+{
+	const TArray<int32> PageIds = GetPagesToTakeIn();
+	return !PageIds.IsEmpty();
+}
+
+bool SAvaInstancedPageList::CanUpdateValuesOnSelectedPage() const
+{
+	const TArray<int32> PageIds = GetPagesToUpdate();
+	return !PageIds.IsEmpty();
+}
+
+void SAvaInstancedPageList::UpdateValuesOnSelectedPage()
+{
+	if (const UAvalanchePlaylist* Playlist = GetPlaylist())
+	{
+		const TArray<int32> PageIds = GetPagesToUpdate();
+		for (const int32 PageId : PageIds)
+		{
+			Playlist->PushRuntimeRemoteControlValues(PageId, false);
+		}
+	}
+}
+
+void SAvaInstancedPageList::ContinueSelectedPage() const
+{
+	if (UAvalanchePlaylist* Playlist = GetPlaylist())
+	{
+		const TArray<int32> PageIds = GetPagesToContinue();
+		for (const int32 PageId : PageIds)
+		{
+			Playlist->ContinuePage(PageId, false);
+		}
+	}
+}
+
+bool SAvaInstancedPageList::CanContinueSelectedPage() const
+{
+	const TArray<int32> PageIds = GetPagesToContinue();
+	return !PageIds.IsEmpty();
+}
+
+void SAvaInstancedPageList::StopSelectedPage(bool bInForce) const
+{
+	if (UAvalanchePlaylist* Playlist = GetPlaylist())
+	{
+		const EAvaPlaylistPageStopOptions StopOptions = bInForce ? EAvaPlaylistPageStopOptions::ForceNoTransition : EAvaPlaylistPageStopOptions::Default;
+		const TArray<int32> PageIds = GetPagesToTakeOut(bInForce);
+		Playlist->StopPages(PageIds, StopOptions, false);
+	}
+}
+
+bool SAvaInstancedPageList::CanStopSelectedPage(bool bInForce) const
+{
+	const TArray<int32> PageIds = GetPagesToTakeOut(bInForce);
+	return !PageIds.IsEmpty();
+}
+
+TArray<int32> SAvaInstancedPageList::PlayNextPage() const
+{
+	TArray<int32> NextPageIds;
+	if (UAvalanchePlaylist* Playlist = GetPlaylist())
+	{
+		const int32 PageIdToTakeNext = GetPageIdToTakeNext();
+		if (IsPageIdValid(PageIdToTakeNext))
+		{
+			Playlist->PlayPage(PageIdToTakeNext, EAvaPlayType::PlayFromStart);
+			NextPageIds.Add(PageIdToTakeNext);
+		}
+	}
+	return NextPageIds;
+}
+
+bool SAvaInstancedPageList::CanPlayNextPage() const
+{
+	const int32 PageIdToTakeNext = GetPageIdToTakeNext();
+	return IsPageIdValid(PageIdToTakeNext);
+}
+
+void SAvaInstancedPageList::OnInstancedPageListChanged(const FAvaPageListChangeParams& InParams)
+{
+	if (const TSharedPtr<FAvaPlaylistEditor> PlaylistEditor = PlaylistEditorWeak.Pin())
+	{
+		PlaylistEditor->RefreshInstancedVisibility();
+	}
+	Refresh();
+}
+
+FReply SAvaInstancedPageList::MakeActive()
+{
+	if (CanMakeActive())
+	{
+		GetPlaylist()->SetActivePageList(PageListReference);
+		return FReply::Handled();
+	}
+
+	return FReply::Unhandled();
+}
+
+bool SAvaInstancedPageList::CanMakeActive() const
+{
+	if (const UAvalanchePlaylist* const Playlist = GetValidPlaylist())
+	{
+		if (Playlist->IsPlaying())
+		{
+			return false;
+		}
+		return (Playlist->GetActivePageListReference() != PageListReference);
+	}
+	return false;
+}
+
+FSlateColor SAvaInstancedPageList::GetMakeActiveButtonColor() const
+{
+	static const FSlateColor Inactive(FStyleColors::AccentGray.GetSpecifiedColor());
+	static const FSlateColor Active(FStyleColors::PrimaryHover.GetSpecifiedColor());
+
+	if (const UAvalanchePlaylist* const Playlist = GetValidPlaylist())
+	{
+		if (Playlist->GetActivePageListReference() == PageListReference)
+		{
+			return Active;
+		}
+	}
+
+	return Inactive;
+}
+
+FText SAvaInstancedPageList::GetPageViewName() const
+{
+	if (UAvalanchePlaylist* const Playlist = GetValidPlaylist())
+	{
+		if (Playlist->IsValidSubList(PageListReference))
+		{
+			return Playlist->GetSubList(PageListReference.SubListIndex).Name;
+		}
+	}
+
+	return FText::GetEmpty();
+}
+
+void SAvaInstancedPageList::OnPageViewNameCommitted(const FText& InNewText, ETextCommit::Type InCommitType)
+{
+	if (InCommitType == ETextCommit::Default || InCommitType == ETextCommit::OnCleared)
+	{
+		return;
+	}
+
+	if (UAvalanchePlaylist* const Playlist = GetValidPlaylist())
+	{
+		if (Playlist->IsValidSubList(PageListReference))
+		{
+			Playlist->GetSubList(PageListReference.SubListIndex).Name = InNewText;
+			Playlist->GetSubList(PageListReference.SubListIndex).OnPageListChanged.Broadcast({Playlist, EAvaPageListChange::RenamedPageView, {}});
+
+			if (TSharedPtr<SDockTab> MyTab = MyTabWeak.Pin())
+			{
+				MyTab->SetLabel(InNewText);
+			}
+		}
+	}
+}
+
+TArray<int32> SAvaInstancedPageList::AddPastedPages(const TArray<FAvalanchePage>& InPages)
+{
+	if (UAvalanchePlaylist* Playlist = GetValidPlaylist())
+	{
+		using namespace UE::AvaPlaylistEditor::Utils;
+		FImportTemplateMap ImportedTemplateIds;
+		return ImportInstancedPages(Playlist, PageListReference, InPages, {}, ImportedTemplateIds);
+	}
+	return {};
+}
+
+TArray<int32> SAvaInstancedPageList::FilterPlayingPages(FFilterPageFunctionRef InFilterPageFunction) const
+{
+	if (const UAvalanchePlaylist* Playlist = GetPlaylist())
+	{
+		return InFilterPageFunction(Playlist, Playlist->GetPlayingPageIds());
+	}
+	return {};
+}
+
+TArray<int32> SAvaInstancedPageList::FilterSelectedOrPlayingPages(FFilterPageFunctionRef InFilterPageFunction, const bool bInAllowFallback) const
+{
+	if (const UAvalanchePlaylist* Playlist = GetPlaylist())
+	{
+		if (!SelectedPageIds.IsEmpty())
+		{
+			const TArray<int32> SelectedPages = InFilterPageFunction(Playlist, SelectedPageIds);
+			if (!bInAllowFallback || !SelectedPages.IsEmpty())
+			{
+				return SelectedPages;
+			}
+		}
+		return InFilterPageFunction(Playlist, Playlist->GetPlayingPageIds());
+	}
+	return {};
+}
+
+TArray<int32> SAvaInstancedPageList::FilterPageSetForProgram(FFilterPageFunctionRef InFilterPageFunction, const EAvaRundownPageSet InPageSet) const
+{
+	switch (InPageSet)
+	{
+	case EAvaRundownPageSet::SelectedOrPlayingStrict:
+		return FilterSelectedOrPlayingPages(InFilterPageFunction, /*bAllowFallback*/ false);
+	case EAvaRundownPageSet::SelectedOrPlaying:
+		return FilterSelectedOrPlayingPages(InFilterPageFunction, /*bAllowFallback*/ true);
+	case EAvaRundownPageSet::Selected:
+		return FilterSelectedPages(InFilterPageFunction);
+	case EAvaRundownPageSet::Playing:
+		return FilterPlayingPages(InFilterPageFunction);
+	default:
+		return FilterSelectedOrPlayingPages(InFilterPageFunction, /*bAllowFallback*/ true);
+	}
+}
+
+TArray<int32> SAvaInstancedPageList::GetPagesToTakeIn() const
+{
+	auto KeepPagesToPlay = [](const UAvalanchePlaylist* InPlaylist, const TArray<int32>& InPageIds)
+	{
+		TArray<int32> OutPageIds;
+		OutPageIds.Reserve(InPageIds.Num());
+		for (int32 PageId : InPageIds)
+		{
+			if (InPlaylist->CanPlayPage(PageId, false))
+			{
+				OutPageIds.Add(PageId);
+			}
+		}
+		return OutPageIds;
+	};
+	return FilterSelectedPages(KeepPagesToPlay);
+}
+
+TArray<int32> SAvaInstancedPageList::GetPagesToTakeOut(bool bInForce) const
+{
+	const EAvaPlaylistPageStopOptions StopOptions = bInForce ? EAvaPlaylistPageStopOptions::ForceNoTransition : EAvaPlaylistPageStopOptions::Default;
+	auto KeepPagesToStop = [StopOptions](const UAvalanchePlaylist* InPlaylist, const TArray<int32>& InPageIds)
+	{
+		TArray<int32> OutPageIds;
+		OutPageIds.Reserve(InPageIds.Num());
+		for (int32 PageId : InPageIds)
+		{
+			if (InPlaylist->CanStopPage(PageId, StopOptions, false))
+			{
+				OutPageIds.Add(PageId);
+			}
+		}
+		return OutPageIds;
+	};
+
+	const UAvaRundownEditorSettings* RundownEditorSettings = UAvaRundownEditorSettings::Get();
+	const EAvaRundownPageSet PageSet = RundownEditorSettings ? RundownEditorSettings->TakeOutActionPageSet : EAvaRundownPageSet::SelectedOrPlaying;
+	return FilterPageSetForProgram(KeepPagesToStop, PageSet);
+}
+
+TArray<int32> SAvaInstancedPageList::GetPagesToContinue() const
+{
+	auto KeepPagesToContinue = [](const UAvalanchePlaylist* InPlaylist, const TArray<int32>& InPageIds)
+	{
+		TArray<int32> OutPageIds;
+		OutPageIds.Reserve(InPageIds.Num());
+		for (int32 PageId : InPageIds)
+		{
+			if (InPlaylist->CanContinuePage(PageId, false))
+			{
+				OutPageIds.Add(PageId);
+			}
+		}
+		return OutPageIds;
+	};
+
+	const UAvaRundownEditorSettings* RundownEditorSettings = UAvaRundownEditorSettings::Get();
+	const EAvaRundownPageSet PageSet = RundownEditorSettings ? RundownEditorSettings->ContinueActionPageSet : EAvaRundownPageSet::SelectedOrPlaying;
+	return FilterPageSetForProgram(KeepPagesToContinue, PageSet);
+}
+
+TArray<int32> SAvaInstancedPageList::GetPagesToUpdate() const
+{
+	auto KeepPagesToUpdate = [](const UAvalanchePlaylist* InPlaylist, const TArray<int32>& InPageIds)
+	{
+		TArray<int32> OutPageIds;
+		OutPageIds.Reserve(InPageIds.Num());
+		for (int32 PageId : InPageIds)
+		{
+			if (InPlaylist->IsPagePlaying(PageId))
+			{
+				OutPageIds.Add(PageId);
+			}
+		}
+		return OutPageIds;
+	};
+
+	const UAvaRundownEditorSettings* RundownEditorSettings = UAvaRundownEditorSettings::Get();
+	const EAvaRundownPageSet PageSet = RundownEditorSettings ? RundownEditorSettings->UpdateValuesActionPageSet : EAvaRundownPageSet::SelectedOrPlaying;
+	return FilterPageSetForProgram(KeepPagesToUpdate, PageSet);
+}
+
+int32 SAvaInstancedPageList::GetPageIdToTakeNext() const
+{
+	return FAvaPlaylistPlaybackUtils::GetPageIdToPlayNext(GetPlaylist(), GetPageListReference(), /*bInPreview*/ false, NAME_None);
+}
+
+#undef LOCTEXT_NAMESPACE
diff --git a/Engine/Plugins/Experimental/Avalanche/Source/AvalancheMediaEditor/Private/Playlist/Pages/Slate/SAvaInstancedPageList.h b/Engine/Plugins/Experimental/Avalanche/Source/AvalancheMediaEditor/Private/Playlist/Pages/Slate/SAvaInstancedPageList.h
new file mode 100644
--- /dev/null
+++ b/Engine/Plugins/Experimental/Avalanche/Source/AvalancheMediaEditor/Private/Playlist/Pages/Slate/SAvaInstancedPageList.h
@@ -0,0 +1,100 @@
+// Copyright Epic Games, Inc. All Rights Reserved.
+
+#pragma once
+
+#include "SAvaPageList.h"
+#include "Playlist/AvaPlaylistDefines.h"
+
+class SDockTab;
+enum ETabActivationCause : uint8;
+
+class SAvaInstancedPageList : public SAvaPageList
+{
+	SLATE_DECLARE_WIDGET(SAvaInstancedPageList, SAvaPageList)
+
+public:
+	SLATE_BEGIN_ARGS(SAvaInstancedPageList) {}
+	SLATE_END_ARGS()
+
+	/** Constructs this widget with InArgs */
+	void Construct(const FArguments& InArgs, TSharedPtr<FAvaPlaylistEditor> InPlaylistEditor, const FAvaPageListReference& PageListReference);
+	virtual ~SAvaInstancedPageList() override;
+
+	//~ Begin SAvaPageList
+	virtual void Refresh() override;
+	virtual void CreateColumns() override;
+	virtual TSharedPtr<SWidget> OnContextMenuOpening() override;
+	virtual void BindCommands() override;
+	virtual bool HandleDropAvalancheAssets(const TArray<FSoftObjectPath>& InAvaAssets, EItemDropZone InDropZone, const FAvaPageViewPtr& InItem) override;
+	virtual bool HandleDropPlaylists(const TArray<FSoftObjectPath>& InPlaylistPaths, EItemDropZone InDropZone, const FAvaPageViewPtr& InItem) override;
+	virtual bool HandleDropPageIds(const FAvaPageListReference& InPageListReference, const TArray<int32>& InPageIds, EItemDropZone InDropZone, const FAvaPageViewPtr& InItem) override;
+	virtual bool HandleDropExternalFiles(const TArray<FString>& InFiles, EItemDropZone InDropZone, const FAvaPageViewPtr& InItem) override;
+	//~ End SAvaPageList
+
+	bool HandleDropPageIdsOnMainListFromTemplates(const TArray<int32>& InPageIds, EItemDropZone InDropZone, const FAvaPageViewPtr& InItem);
+	bool HandleDropPageIdsOnMainListFromMainList(const TArray<int32>& InPageIds, EItemDropZone InDropZone, const FAvaPageViewPtr& InItem);
+	bool HandleDropPageIdsOnSubListFromTemplates(const TArray<int32>& InPageIds, EItemDropZone InDropZone, const FAvaPageViewPtr& InItem);
+	bool HandleDropPageIdsOnSubListFromMainList(const TArray<int32>& InPageIds, EItemDropZone InDropZone, const FAvaPageViewPtr& InItem);
+	bool HandleDropPageIdsOnSubListFromSubList(int32 InFromList, const TArray<int32>& InPageIds, EItemDropZone InDropZone, const FAvaPageViewPtr& InItem);
+
+	void OnTabActivated(TSharedRef<SDockTab> InDockTab, ETabActivationCause InActivationCause);
+
+	FName GetTabId() const { return TabId; }
+
+	TSharedPtr<SDockTab> GetMyTab() const { return MyTabWeak.Pin(); }
+	void SetMyTab(TSharedRef<SDockTab> InTab) { MyTabWeak = InTab; }
+
+	/** Plays the currently selected page. */
+	void PlaySelectedPage() const;
+	bool CanPlaySelectedPage() const;
+
+	bool CanUpdateValuesOnSelectedPage() const;
+	void UpdateValuesOnSelectedPage();
+	
+	void ContinueSelectedPage() const;
+	bool CanContinueSelectedPage() const;
+
+	void StopSelectedPage(bool bInForce) const;
+	bool CanStopSelectedPage(bool bInForce) const;
+
+	TArray<int32> PlayNextPage() const;
+	bool CanPlayNextPage() const;
+
+	TArray<int32> GetPageIdsToTakeNext() const
+	{
+		const int32 NextPageId = GetPageIdToTakeNext();
+		if (IsPageIdValid(NextPageId))
+		{
+			return { NextPageId };
+		}
+		return {};
+	}
+	
+private:
+	void OnInstancedPageListChanged(const FAvaPageListChangeParams& InParams);
+
+	TArray<int32> FilterPlayingPages(FFilterPageFunctionRef InFilterPageFunction) const;
+	TArray<int32> FilterSelectedOrPlayingPages(FFilterPageFunctionRef InFilterPageFunction, const bool bInAllowFallback) const;
+	TArray<int32> FilterPageSetForProgram(FFilterPageFunctionRef InFilterPageFunction, const EAvaRundownPageSet InPageSet) const;
+
+	TArray<int32> GetPagesToTakeIn() const;
+	TArray<int32> GetPagesToTakeOut(bool bInForce) const;
+	TArray<int32> GetPagesToContinue() const;
+	TArray<int32> GetPagesToUpdate() const;
+	int32 GetPageIdToTakeNext() const;
+
+protected:
+	FName TabId;
+	TWeakPtr<SDockTab> MyTabWeak;
+
+	FReply MakeActive();
+	bool CanMakeActive() const;
+	FSlateColor GetMakeActiveButtonColor() const;
+
+	FText GetPageViewName() const;
+	void OnPageViewNameCommitted(const FText& InNewText, ETextCommit::Type InCommitType);
+
+	void PlayNextPageNoReturn() const { PlayNextPage(); }
+
+	virtual TArray<int32> AddPastedPages(const TArray<FAvalanchePage>& InPages) override;
+};
