[MotionDesign]

diff --git a/Engine/Plugins/Experimental/Avalanche/Source/AvalancheMediaEditor/Private/Playlist/Pages/Slate/SAvaPageList.cpp b/Engine/Plugins/Experimental/Avalanche/Source/AvalancheMediaEditor/Private/Playlist/Pages/Slate/SAvaPageList.cpp
new file mode 100644
--- /dev/null
+++ b/Engine/Plugins/Experimental/Avalanche/Source/AvalancheMediaEditor/Private/Playlist/Pages/Slate/SAvaPageList.cpp
@@ -0,0 +1,1288 @@
+// Copyright Epic Games, Inc. All Rights Reserved.
+
+#include "SAvaPageList.h"
+
+#include "AvaAssetTags.h"
+#include "DragAndDrop/AssetDragDropOp.h"
+#include "Editor/EditorWidgets/Public/SAssetSearchBox.h"
+#include "Framework/Commands/GenericCommands.h"
+#include "Framework/MultiBox/MultiBoxBuilder.h"
+#include "HAL/PlatformApplicationMisc.h"
+#include "IAvaMediaEditorModule.h"
+#include "Misc/TextFilterExpressionEvaluator.h"
+#include "Playlist/AvaPlaylistCommands.h"
+#include "Playlist/AvaPlaylistEditor.h"
+#include "Playlist/AvaPlaylistEditorUtils.h"
+#include "Playlist/AvaPlaylistPlaybackUtils.h"
+#include "Playlist/AvaRundownEditorSettings.h"
+#include "Playlist/AvalanchePage.h"
+#include "Playlist/AvalanchePlaylist.h"
+#include "Playlist/Factories/Filters/AvaPlaylistFactoriesUtils.h"
+#include "Playlist/Factories/Filters/IAvaPlaylistFilterSuggestionFactory.h"
+#include "Playlist/Pages/AvaPlaylistPageContext.h"
+#include "Playlist/Pages/AvaPlaylistPageContextMenu.h"
+#include "Playlist/Pages/PageViews/AvaPageViewImpl.h"
+#include "Playlist/Pages/Slate/SAvaPageViewRow.h"
+#include "ScopedTransaction.h"
+#include "Widgets/SBoxPanel.h"
+#include "Widgets/SNullWidget.h"
+#include "Widgets/Views/SListView.h"
+
+#define LOCTEXT_NAMESPACE "SAvaPageList"
+
+namespace UE::AvalanchePlaylist::Private
+{
+	static const FString PageClipboardPrefix = TEXT("AvalanchePages");
+	static const FString PageEntriesName = TEXT("Pages");
+
+	void ExtractAssetSearchFilterTerms(const FText& SearchText, FString* OutFilterKey, FString* OutFilterValue, int32* OutSuggestionInsertionIndex)
+	{
+		const FString SearchString = SearchText.ToString();
+
+		if (OutFilterKey)
+		{
+			OutFilterKey->Reset();
+		}
+		if (OutFilterValue)
+		{
+			OutFilterValue->Reset();
+		}
+		if (OutSuggestionInsertionIndex)
+		{
+			*OutSuggestionInsertionIndex = SearchString.Len();
+		}
+
+		// Build the search filter terms so that we can inspect the tokens
+		FTextFilterExpressionEvaluator LocalFilter(ETextFilterExpressionEvaluatorMode::Complex);
+		LocalFilter.SetFilterText(SearchText);
+	
+		// Inspect the tokens to see what the last part of the search term was
+		// If it was a key->value pair then we'll use that to control what kinds of results we show
+		// For anything else we just use the text from the last token as our filter term to allow incremental auto-complete
+		const TArray<FExpressionToken>& FilterTokens = LocalFilter.GetFilterExpressionTokens();
+		if (FilterTokens.Num() > 0)
+		{
+			const FExpressionToken& LastToken = FilterTokens.Last();
+			// If the last token is a text token, then consider it as a value and walk back to see if we also have a key
+			if (LastToken.Node.Cast<TextFilterExpressionParser::FTextToken>())
+			{
+				if (OutFilterValue)
+				{
+					*OutFilterValue = LastToken.Context.GetString();
+				}
+				if (OutSuggestionInsertionIndex)
+				{
+					*OutSuggestionInsertionIndex = FMath::Min(*OutSuggestionInsertionIndex, LastToken.Context.GetCharacterIndex());
+				}
+				if (FilterTokens.IsValidIndex(FilterTokens.Num() - 2))
+				{
+					const FExpressionToken& ComparisonToken = FilterTokens[FilterTokens.Num() - 2];
+					if (ComparisonToken.Node.Cast<TextFilterExpressionParser::FEqual>())
+					{
+						if (FilterTokens.IsValidIndex(FilterTokens.Num() - 3))
+						{
+							const FExpressionToken& KeyToken = FilterTokens[FilterTokens.Num() - 3];
+							if (KeyToken.Node.Cast<TextFilterExpressionParser::FTextToken>())
+							{
+								if (OutFilterKey)
+								{
+									*OutFilterKey = KeyToken.Context.GetString();
+								}
+								if (OutSuggestionInsertionIndex)
+								{
+									*OutSuggestionInsertionIndex = FMath::Min(*OutSuggestionInsertionIndex, KeyToken.Context.GetCharacterIndex());
+								}
+							}
+						}
+					}
+				}
+			}
+			// If the last token is a comparison operator, then walk back and see if we have a key
+			else if (LastToken.Node.Cast<TextFilterExpressionParser::FEqual>())
+			{
+				if (FilterTokens.IsValidIndex(FilterTokens.Num() - 2))
+				{
+					const FExpressionToken& KeyToken = FilterTokens[FilterTokens.Num() - 2];
+					if (KeyToken.Node.Cast<TextFilterExpressionParser::FTextToken>())
+					{
+						if (OutFilterKey)
+						{
+							*OutFilterKey = KeyToken.Context.GetString();
+						}
+						if (OutSuggestionInsertionIndex)
+						{
+							*OutSuggestionInsertionIndex = FMath::Min(*OutSuggestionInsertionIndex, KeyToken.Context.GetCharacterIndex());
+						}
+					}
+				}
+			}
+		}
+	}
+}
+
+void SAvaPageList::PrivateRegisterAttributes(struct FSlateAttributeDescriptor::FInitializer&)
+{
+
+}
+
+void SAvaPageList::Construct(const FArguments& InArgs, TSharedPtr<FAvaPlaylistEditor> InPlaylistEditor, const FAvaPageListReference& InPageListReference, EAvaPlaylistSearchListType InPageListType)
+{
+	PlaylistEditorWeak = InPlaylistEditor;
+	check(InPlaylistEditor.IsValid());
+
+	PageContextMenu = MakeShared<FAvaPlaylistPageContextMenu>();
+
+	PageListReference = InPageListReference;
+	PageListType = InPageListType;
+	CommandList = MakeShared<FUICommandList>();
+	BindCommands();
+
+	UAvalanchePlaylist* const Playlist = InPlaylistEditor->GetPlaylist();
+	check(Playlist);
+
+	InPlaylistEditor->GetOnPageEvent().AddSP(this, &SAvaPageList::OnPageEvent);
+
+	CreateColumns();
+
+	ChildSlot
+		[
+			SNew(SVerticalBox)
+			+ SVerticalBox::Slot()
+			.AutoHeight()
+			.Padding(2.f)
+			[
+				SAssignNew(SearchBar, SHorizontalBox)
+				// Search Box
+				+ SHorizontalBox::Slot()
+				.Padding(4.f, 2.f)
+				[
+					SAssignNew(AssetSearchBoxPtr, SAssetSearchBox)
+					.HintText(LOCTEXT("FilterSearch", "Search..."))
+					.ToolTipText(LOCTEXT("FilterSearchHint", "Type here to search..."))
+					.ShowSearchHistory(true)
+					.OnTextChanged(this, &SAvaPageList::OnSearchTextChanged)
+					.OnTextCommitted(this, &SAvaPageList::OnSearchTextCommitted)
+					.OnAssetSearchBoxSuggestionFilter(this, &SAvaPageList::OnSearchBoxSuggestionFilter)
+					.OnAssetSearchBoxSuggestionChosen(this, &SAvaPageList::OnAssetSearchBoxSuggestionChosen)
+					.DelayChangeNotificationsWhileTyping(true)
+					.Visibility(EVisibility::Visible)
+					.AddMetaData<FTagMetaData>(FTagMetaData(TEXT("SAvaPageListSearchBox")))
+				]
+			]
+			// Pages Collection
+			+ SVerticalBox::Slot()
+			.FillHeight(1.f)
+			[
+				SNew(SBorder)
+				.BorderImage(FAppStyle::GetBrush("ToolPanel.GroupBorder"))
+				.Padding(0.f)
+				[
+					SAssignNew(PageListView, SListView<FAvaPageViewPtr>)
+					.HeaderRow(HeaderRow)
+					.ListItemsSource(&PageViews)
+					.HandleSpacebarSelection(true)
+					.SelectionMode(ESelectionMode::Multi)
+					.ClearSelectionOnClick(true)
+					.OnGenerateRow(this, &SAvaPageList::GeneratePageTableRow)
+					.OnSelectionChanged(this, &SAvaPageList::OnPageSelected)
+					.OnContextMenuOpening(this, &SAvaPageList::OnContextMenuOpening)
+				]
+			]
+		];
+
+	Refresh();
+}
+
+SAvaPageList::~SAvaPageList()
+{
+	if (PlaylistEditorWeak.IsValid())
+	{
+		TSharedPtr<FAvaPlaylistEditor> PlaylistEditor = PlaylistEditorWeak.Pin();
+
+		if (PlaylistEditor.IsValid())
+		{
+			PlaylistEditor->GetOnPageEvent().RemoveAll(this);
+		}
+	}
+}
+
+void SAvaPageList::OnPageEvent(const TArray<int32>& InSelectedPageIds, UE::AvalanchePlaylist::EPageEvent InPageEvent)
+{
+	using namespace  UE::AvalanchePlaylist;
+
+	switch (InPageEvent)
+	{
+		case EPageEvent::SelectionRequest:
+			OnPageSelectionRequested(InSelectedPageIds);
+			break;
+
+		case EPageEvent::RenameRequest:
+			OnPageItemActionRequested(InSelectedPageIds, &IAvaPageView::GetOnRename);
+			break;
+
+		case EPageEvent::RenumberRequest:
+			OnPageItemActionRequested(InSelectedPageIds, &IAvaPageView::GetOnRenumber);
+			break;
+	}
+}
+
+void SAvaPageList::OnPageItemActionRequested(const TArray<int32>& InPageIds, IAvaPageView::FOnPageAction& (IAvaPageView::* InFunc)())
+{
+	//Only allow one page id to be renamed at a time
+	const TSet<int32> PageIdSet(InPageIds);
+	for (const FAvaPageViewPtr& PageView : PageViews)
+	{
+		if (PageIdSet.Contains(PageView->GetPageId()))
+		{
+			(*PageView.*InFunc)().Broadcast(EAvaPageActionState::Requested);
+		}
+	}
+}
+
+void SAvaPageList::OnPageSelectionRequested(const TArray<int32>& InPageIds)
+{
+	TArray<FAvaPageViewPtr> DesiredSelectedItems;
+	DesiredSelectedItems.Empty(InPageIds.Num());
+	for (int32 PageId : InPageIds)
+	{
+		DesiredSelectedItems.Add(GetPageViewPtr(PageId));
+	}
+	PageListView->ClearSelection();
+	PageListView->SetItemSelection(DesiredSelectedItems, true);
+}
+
+TSharedRef<ITableRow> SAvaPageList::GeneratePageTableRow(FAvaPageViewPtr InPageView, const TSharedRef<STableViewBase>& InOwnerTable)
+{
+	check(InPageView.IsValid());
+	return SNew(SAvaPageViewRow, InPageView, SharedThis(this));
+}
+
+void SAvaPageList::OnPageSelected(FAvaPageViewPtr InPageView, ESelectInfo::Type InSelectInfo)
+{
+	if (PageListView.IsValid())
+	{
+		TArray<FAvaPageViewPtr> SelectedItems;
+		PageListView->GetSelectedItems(SelectedItems);
+		SelectedPageIds.Empty(SelectedItems.Num());
+
+		for (const FAvaPageViewPtr& PageView : SelectedItems)
+		{
+			if (PageView.IsValid())
+			{
+				const int32 SelectedPage = PageView->GetPageId();
+				SelectedPageIds.Add(SelectedPage);
+			}
+		}
+
+		TSharedPtr<FAvaPlaylistEditor> PlaylistEditor = GetPlaylistEditor();
+
+		if (PlaylistEditor.IsValid())
+		{
+			PlaylistEditor->GetOnPageEvent().Broadcast(SelectedPageIds, UE::AvalanchePlaylist::EPageEvent::SelectionChanged);
+		}
+	}
+}
+
+TSharedPtr<IAvaPageViewColumn> SAvaPageList::FindColumn(FName InColumnName) const
+{
+	if (const TSharedPtr<IAvaPageViewColumn>* const FoundColumn = Columns.Find(InColumnName))
+	{
+		return *FoundColumn;
+	}
+	return nullptr;
+}
+
+FAvaPageViewPtr SAvaPageList::GetPageViewPtr(int32 InPageId) const
+{
+	for (const FAvaPageViewPtr& PageView : PageViews)
+	{
+		if (PageView->GetPageId() == InPageId)
+		{
+			return PageView;
+		}
+	}
+	return nullptr;
+}
+
+void SAvaPageList::SelectPage(int32 InPageId, bool bInScrollIntoView)
+{
+	if (PageListView.IsValid())
+	{
+		for (const FAvaPageViewPtr& PageView : PageViews)
+		{
+			if (PageView->GetPageId() == InPageId)
+			{
+				PageListView->SetItemSelection(PageView, true, ESelectInfo::Direct);
+
+				if (bInScrollIntoView)
+				{
+					PageListView->RequestScrollIntoView(PageView);
+				}
+
+				return;
+			}
+		}
+	}
+}
+
+void SAvaPageList::SelectPages(const TArray<int32>& InPageIds, bool bInScrollIntoView)
+{
+	if (PageListView.IsValid())
+	{
+		for (int32 PageId : InPageIds)
+		{
+			SelectPage(PageId, bInScrollIntoView);
+		}
+	}
+}
+
+void SAvaPageList::DeselectPage(int32 InPageId)
+{
+	if (PageListView.IsValid())
+	{
+		for (const FAvaPageViewPtr& PageView : PageViews)
+		{
+			if (PageView->GetPageId() == InPageId)
+			{
+				PageListView->SetItemSelection(PageView, false, ESelectInfo::Direct);
+				return;
+			}
+		}
+	}
+}
+
+void SAvaPageList::DeselectPages()
+{
+	if (PageListView.IsValid())
+	{
+		PageListView->ClearSelection();
+	}
+}
+
+TArray<int32> SAvaPageList::GetPlayingPageIds() const
+{
+	if (const UAvalanchePlaylist* Playlist = GetPlaylist())
+	{
+		return Playlist->GetPlayingPageIds();
+	}
+	return TArray<int32>();
+}
+
+TSharedPtr<FAvaPlaylistEditor> SAvaPageList::GetPlaylistEditor() const
+{
+	return PlaylistEditorWeak.Pin();
+}
+
+UAvalanchePlaylist* SAvaPageList::GetPlaylist() const
+{
+	const TSharedPtr<FAvaPlaylistEditor> PlaylistEditor = GetPlaylistEditor();
+	return PlaylistEditor.IsValid() ? PlaylistEditor->GetPlaylist() : nullptr;
+}
+
+UAvalanchePlaylist* SAvaPageList::GetValidPlaylist() const
+{
+	UAvalanchePlaylist* Playlist = GetPlaylist();
+	return IsValid(Playlist) ? Playlist : nullptr;
+}
+
+void SAvaPageList::BindCommands()
+{
+	//Generic Commands
+	{
+		const FGenericCommands& GenericCommands = FGenericCommands::Get();
+
+		CommandList->MapAction(GenericCommands.Rename,
+			FExecuteAction::CreateSP(this, &SAvaPageList::RenameSelectedPage),
+			FCanExecuteAction::CreateSP(this, &SAvaPageList::CanRenameSelectedPage));
+
+		CommandList->MapAction(GenericCommands.Delete,
+			FExecuteAction::CreateSP(this, &SAvaPageList::RemoveSelectedPages),
+			FCanExecuteAction::CreateSP(this, &SAvaPageList::CanRemoveSelectedPages));
+
+		CommandList->MapAction(GenericCommands.Cut,
+			FExecuteAction::CreateSP(this, &SAvaPageList::CutSelectedPages),
+			FCanExecuteAction::CreateSP(this, &SAvaPageList::CanCutSelectedPages));
+
+		CommandList->MapAction(GenericCommands.Copy,
+			FExecuteAction::CreateSP(this, &SAvaPageList::CopySelectedPages),
+			FCanExecuteAction::CreateSP(this, &SAvaPageList::CanCopySelectedPages));
+
+		CommandList->MapAction(GenericCommands.Paste,
+			FExecuteAction::CreateSP(this, &SAvaPageList::Paste),
+			FCanExecuteAction::CreateSP(this, &SAvaPageList::CanPaste));
+
+		CommandList->MapAction(GenericCommands.Duplicate,
+			FExecuteAction::CreateSP(this, &SAvaPageList::DuplicateSelectedPages),
+			FCanExecuteAction::CreateSP(this, &SAvaPageList::CanDuplicateSelectedPages));
+	}
+}
+
+int32 SAvaPageList::GetFirstSelectedPageId() const
+{
+	return SelectedPageIds.Num() > 0 ? SelectedPageIds[0] : FAvalanchePage::InvalidPageId;
+}
+
+TSharedRef<SWidget> SAvaPageList::GetPageListContextMenu()
+{
+	return PageContextMenu->GeneratePageContextMenuWidget(PlaylistEditorWeak, PageListReference, CommandList);
+}
+
+bool SAvaPageList::CanAddPage() const
+{
+	if (const UAvalanchePlaylist* Playlist = GetPlaylist(); IsValid(Playlist))
+	{
+		return Playlist->CanAddPage();
+	}
+	return false;
+}
+
+bool SAvaPageList::CanAddTemplate() const
+{
+	return CanAddPage();
+}
+
+bool SAvaPageList::CanCreateInstance() const
+{
+	return CanAddPage();
+}
+
+bool SAvaPageList::CanCopySelectedPages() const
+{
+	return !SelectedPageIds.IsEmpty() && IsValid(GetPlaylist());
+}
+
+void SAvaPageList::CopySelectedPages()
+{
+	if (!CanCopySelectedPages())
+	{
+		return;
+	}
+
+	UAvalanchePlaylist* Playlist = GetPlaylist();
+	if (!IsValid(Playlist))
+	{
+		UE_LOG(LogAvaPlaylist, Error, TEXT("Can't copy from invalid playlist."));
+		return;
+	}
+
+	FString SerializedString = UE::AvaPlaylistEditor::Utils::SerializePagesToJson(Playlist, SelectedPageIds);
+	
+	// Add Prefix to quickly identify whether current clipboard is from Avalanche Pages or not
+	SerializedString = *FString::Printf(TEXT("%s%s"), *UE::AvalanchePlaylist::Private::PageClipboardPrefix, *SerializedString);
+
+	FPlatformApplicationMisc::ClipboardCopy(*SerializedString);
+}
+
+bool SAvaPageList::CanCutSelectedPages() const
+{
+	return CanCopySelectedPages() && CanRemoveSelectedPages();
+}
+
+void SAvaPageList::CutSelectedPages()
+{
+	if (CanCutSelectedPages())
+	{
+		CopySelectedPages();
+		RemoveSelectedPages();
+	}
+}
+
+bool SAvaPageList::CanPaste() const
+{
+	if (!CanAddPage())
+	{
+		return false;
+	}
+
+	FString PastedText;
+	FPlatformApplicationMisc::ClipboardPaste(PastedText);
+	return PastedText.StartsWith(UE::AvalanchePlaylist::Private::PageClipboardPrefix);
+}
+
+void SAvaPageList::Paste()
+{
+	if (!CanPaste())
+	{
+		return;
+	}
+
+	UAvalanchePlaylist* Playlist = GetPlaylist();
+	if (!IsValid(Playlist))
+	{
+		return;
+	}	
+
+	FString PastedText;
+	FPlatformApplicationMisc::ClipboardPaste(PastedText);
+	PastedText.RightChopInline(UE::AvalanchePlaylist::Private::PageClipboardPrefix.Len());
+
+	const TArray<FAvalanchePage> Pages = UE::AvaPlaylistEditor::Utils::DeserializePagesFromJson(PastedText);
+	
+	if (Pages.Num() > 0)
+	{
+		FScopedTransaction Transaction(LOCTEXT("AddPages", "Add Pages"));
+		Playlist->Modify();
+
+		const TArray<int32> PageIds = AddPastedPages(Pages);
+
+		if (PageIds.IsEmpty())
+		{
+			Transaction.Cancel();
+		}
+		else
+		{
+			DeselectPages();
+			SelectPages(PageIds);
+		}
+	}
+}
+
+bool SAvaPageList::CanDuplicateSelectedPages() const
+{
+	return CanCopySelectedPages() && CanAddPage();
+}
+
+void SAvaPageList::DuplicateSelectedPages()
+{
+	if (CanDuplicateSelectedPages())
+	{
+		CopySelectedPages();
+		Paste();
+	}
+}
+
+bool SAvaPageList::CanRemoveSelectedPages() const
+{
+	if (!SelectedPageIds.IsEmpty())
+	{
+		if (const UAvalanchePlaylist* Playlist = GetPlaylist(); IsValid(Playlist))
+		{
+			return Playlist->CanRemovePages(SelectedPageIds);
+		}
+	}
+	return false;
+}
+
+void SAvaPageList::RemoveSelectedPages()
+{
+	if (CanRemoveSelectedPages())
+	{
+		if (UAvalanchePlaylist* Playlist = GetPlaylist(); IsValid(Playlist))
+		{
+			FScopedTransaction Transaction(LOCTEXT("RemoveSelectedPages", "Remove Selected Pages"));
+			Playlist->Modify();
+
+			int32 RemovedCount = 0;
+
+			if (PageListReference.Type != EAvaPageListType::View)
+			{
+				RemovedCount = Playlist->RemovePages(SelectedPageIds);
+			}
+			else if (Playlist->IsValidSubList(PageListReference))
+			{
+				RemovedCount = Playlist->RemovePagesFromSubList(PageListReference.SubListIndex, SelectedPageIds);
+			}
+
+			if (RemovedCount == 0)
+			{
+				Transaction.Cancel();
+			}
+		}
+	}
+}
+
+bool SAvaPageList::CanRenameSelectedPage() const
+{
+	if (SelectedPageIds.Num() == 1)
+	{
+		if (const UAvalanchePlaylist* Playlist = GetPlaylist(); IsValid(Playlist))
+		{
+			return Playlist->GetPage(SelectedPageIds[0]).IsValidPage();
+		}
+	}
+	return false;
+}
+
+void SAvaPageList::RenameSelectedPage()
+{
+	if (CanRenameSelectedPage())
+	{
+		OnPageEvent(SelectedPageIds, UE::AvalanchePlaylist::EPageEvent::RenameRequest);
+	}
+}
+
+bool SAvaPageList::CanRenumberSelectedPage() const
+{
+	if (SelectedPageIds.Num() == 1)
+	{
+		if (const UAvalanchePlaylist* Playlist = GetPlaylist(); IsValid(Playlist))
+		{
+			return Playlist->CanRenumberPageId(SelectedPageIds[0]);
+		}
+	}
+	return false;
+}
+
+void SAvaPageList::RenumberSelectedPage()
+{
+	if (CanRenumberSelectedPage())
+	{
+		OnPageEvent(SelectedPageIds, UE::AvalanchePlaylist::EPageEvent::RenumberRequest);
+	}
+}
+
+bool SAvaPageList::CanReimportSelectedPage() const
+{
+	return SelectedPageIds.Num() > 0 && GetPlaylistEditor().IsValid();
+}
+
+void SAvaPageList::ReimportSelectedPage() const
+{
+	if (CanReimportSelectedPage())
+	{
+		if (const TSharedPtr<FAvaPlaylistEditor> PlaylistEditor = GetPlaylistEditor())
+		{
+			UAvalanchePlaylist* Playlist = PlaylistEditor->GetPlaylist();
+			if (IsValid(Playlist))
+			{
+				// Enforce invalidation of the avalanche managed instance cache for the selected page(s).
+				Playlist->InvalidateManagedInstanceCacheForPages(SelectedPageIds);
+				Playlist->UpdateAvalancheAssetForPages(SelectedPageIds, true);
+			}
+
+			PlaylistEditor->GetOnPageEvent().Broadcast(SelectedPageIds, UE::AvalanchePlaylist::EPageEvent::ReimportRequest);
+		}
+	}
+}
+
+bool SAvaPageList::CanEditSelectedPageSource() const
+{
+	if (SelectedPageIds.Num() == 1)
+	{
+		if (UAvalanchePlaylist* Playlist = GetPlaylist(); IsValid(Playlist))
+		{
+			const FAvalanchePage& Page = Playlist->GetPage(SelectedPageIds[0]);
+			return Page.IsValidPage();
+		}
+	}
+	return false;
+}
+
+void SAvaPageList::EditSelectedPageSource()
+{
+	if (SelectedPageIds.Num() == 1)
+	{
+		if (UAvalanchePlaylist* Playlist = GetPlaylist(); IsValid(Playlist))
+		{
+			const FAvalanchePage& Page = Playlist->GetPage(SelectedPageIds[0]);
+			if (Page.IsValidPage())
+			{
+				GEditor->GetEditorSubsystem<UAssetEditorSubsystem>()->OpenEditorForAsset(Page.GetAvalancheAssetPath(Playlist));
+			}
+		}
+	}
+}
+
+bool SAvaPageList::CanExportSelectedPagesToPlaylist()
+{
+	return SelectedPageIds.Num() > 0;
+}
+
+void SAvaPageList::ExportSelectedPagesToPlaylist()
+{
+	const UAvalanchePlaylist* SourcePlaylist = GetPlaylist();
+	if (!SourcePlaylist || SelectedPageIds.IsEmpty())
+	{
+		return;
+	}
+
+	using namespace UE::AvaPlaylistEditor::Utils;
+	if (const TStrongObjectPtr<UAvalanchePlaylist> NewPlaylist = ExportPagesToPlaylist(SourcePlaylist, SelectedPageIds))
+	{	
+		const UPackage* const PlaylistPackage = SourcePlaylist->GetPackage();
+		const FString DefaultPackagePath = FPackageName::GetLongPackagePath(PlaylistPackage->GetName());
+		const FString DefaultAssetName = SourcePlaylist->GetName() + TEXT("_Exported");		
+		const FString SaveObjectPath = GetSaveAssetAsPath(DefaultPackagePath, DefaultAssetName);
+
+		if (!SaveObjectPath.IsEmpty())
+		{
+			const FString PackageName = FPackageName::ObjectPathToPackageName(SaveObjectPath);
+			const FString PackagePath = FPackageName::GetLongPackagePath(PackageName);
+			const FString AssetName = FPackageName::GetLongPackageAssetName(PackageName);
+			SaveDuplicatePlaylist(NewPlaylist.Get(), AssetName, PackagePath);
+		}
+	}
+}
+
+bool SAvaPageList::CanExportSelectedPagesToExternalFile(const TCHAR* InType)
+{
+	return SelectedPageIds.Num() > 0;
+}
+
+void SAvaPageList::ExportSelectedPagesToExternalFile(const TCHAR* InType)
+{
+	const UAvalanchePlaylist* SourcePlaylist = GetPlaylist();
+	if (!SourcePlaylist || SelectedPageIds.IsEmpty())
+	{
+		return;
+	}
+
+	using namespace UE::AvaPlaylistEditor::Utils;
+	if (const TStrongObjectPtr<UAvalanchePlaylist> NewPlaylist = ExportPagesToPlaylist(SourcePlaylist, SelectedPageIds))
+	{
+		if (FCString::Stricmp(InType, TEXT("json")) == 0)
+		{
+			const FString ExportFilename = GetExportFilepath(SourcePlaylist, TEXT("json file"), TEXT("json"));
+			if (!ExportFilename.IsEmpty())
+			{
+				SavePlaylistToJson(NewPlaylist.Get(), *ExportFilename);
+			}
+		}
+		else if (FCString::Stricmp(InType, TEXT("xml")) == 0)
+		{
+			const FString ExportFilename = GetExportFilepath(SourcePlaylist, TEXT("xml file"), TEXT("xml"));
+			if (!ExportFilename.IsEmpty())
+			{
+				SavePlaylistToXml(NewPlaylist.Get(), *ExportFilename);
+			}
+		}
+		else
+		{
+			UE_LOG(LogAvaPlaylist, Error, TEXT("Export Pages to external file doesn't support type \"%s\"."), InType);
+		}
+	}
+}
+
+bool SAvaPageList::CanPreviewPlaySelectedPage() const
+{
+	const TArray<int32> PageIds = GetPagesToPreviewIn();
+	return !PageIds.IsEmpty();
+}
+
+void SAvaPageList::PreviewPlaySelectedPage(bool bInToMark) const
+{
+	if (UAvalanchePlaylist* Playlist = GetPlaylist())
+	{
+		const TArray<int32> PageIds = GetPagesToPreviewIn();
+		Playlist->PlayPages(PageIds, bInToMark ? EAvaPlayType::PreviewFromFrame : EAvaPlayType::PreviewFromStart);
+	}
+}
+
+bool SAvaPageList::CanPreviewStopSelectedPage(bool bInForce) const
+{
+	const TArray<int32> PageIds = GetPagesToPreviewOut(bInForce);
+	return !PageIds.IsEmpty();
+}
+
+void SAvaPageList::PreviewStopSelectedPage(bool bInForce) const
+{
+	if (UAvalanchePlaylist* Playlist = GetPlaylist())
+	{
+		const EAvaPlaylistPageStopOptions StopOptions = bInForce ? EAvaPlaylistPageStopOptions::ForceNoTransition : EAvaPlaylistPageStopOptions::Default;
+		const TArray<int32> PageIds = GetPagesToPreviewOut(bInForce);
+		Playlist->StopPages(PageIds, StopOptions, true);
+	}
+}
+
+bool SAvaPageList::CanPreviewContinueSelectedPage() const
+{
+	const TArray<int32> PageIds = GetPagesToPreviewContinue();
+	return !PageIds.IsEmpty();
+}
+
+void SAvaPageList::PreviewContinueSelectedPage() const
+{
+	if (UAvalanchePlaylist* Playlist = GetPlaylist())
+	{
+		const TArray<int32> PageIds = GetPagesToPreviewContinue();
+		for (const int32 PageId : PageIds)
+		{
+			Playlist->ContinuePage(PageId, true);
+		}
+	}
+}
+
+bool SAvaPageList::CanPreviewPlayNextPage() const
+{
+	const int32 PageIdToPreviewNext = GetPageIdToPreviewNext();
+	return IsPageIdValid(PageIdToPreviewNext);
+}
+
+void SAvaPageList::PreviewPlayNextPage()
+{
+	if (UAvalanchePlaylist* Playlist = GetPlaylist())
+	{
+		const int32 PageIdToPreviewNext = GetPageIdToPreviewNext();
+		if (IsPageIdValid(PageIdToPreviewNext))
+		{
+			Playlist->PlayPage(PageIdToPreviewNext, EAvaPlayType::PreviewFromStart);
+			DeselectPages();
+			SelectPages({PageIdToPreviewNext});
+		}
+	}
+}
+
+bool SAvaPageList::CanTakeToProgram() const
+{
+	const TArray<int32> PageIds = GetPagesToTakeToProgram();
+	return !PageIds.IsEmpty();
+}
+
+void SAvaPageList::TakeToProgram() const
+{
+	if (UAvalanchePlaylist* Playlist = GetPlaylist(); IsValid(Playlist))
+	{
+		const TArray<int32> PageIds = GetPagesToTakeToProgram();
+		Playlist->PlayPages(PageIds, EAvaPlayType::PlayFromStart);
+	}
+}
+
+namespace UE::AvaPageList::Private
+{
+	static bool IsAvaAsset(const FAssetData& InAssetData, bool bInLogInfo)
+	{
+		static const FName& AvalancheSceneAssetTag = UE::Ava::AssetTags::AvalancheScene;
+		static const FString& AssetTagValueEnabled = UE::Ava::AssetTags::Values::Enabled;
+
+		const UClass* AssetClass = InAssetData.GetClass(EResolveClass::Yes);
+		if (!IsValid(AssetClass))
+		{
+			return false;
+		}
+		
+		const EAvalancheAssetType AssetType = FAvaSoftAssetPath::GetAssetTypeFromClass(AssetClass, true);
+		if (AssetType == EAvalancheAssetType::Unknown)
+		{
+			return false;
+		}
+		
+		// If the asset is a level, we check if it has the avalanche scene tag.
+		if (AssetType == EAvalancheAssetType::World)
+		{
+			const FAssetTagValueRef AvalancheSceneTag = InAssetData.TagsAndValues.FindTag(AvalancheSceneAssetTag);
+			
+			if (!AvalancheSceneTag.IsSet())
+			{
+				if (bInLogInfo)
+				{
+					UE_LOG(LogAvaMediaEditor, Display,
+						TEXT("Level Asset \"%s\" is not an Avalanche Scene (Asset Tag not found)."),
+						*InAssetData.GetSoftObjectPath().ToString());
+				}
+				return false;
+			}
+			
+			if (!AvalancheSceneTag.Equals(AssetTagValueEnabled))
+			{
+				if (bInLogInfo)
+				{
+					UE_LOG(LogAvaMediaEditor, Display,
+						TEXT("Level Asset \"%s\" is an Avalanche Scene but not enabled."),
+						*InAssetData.GetSoftObjectPath().ToString());
+				}
+				return false;						
+			}
+		}
+		return true;
+	}
+}
+
+bool SAvaPageList::IsAssetDropSupported(const FAssetData& InAsset, const FSoftObjectPath& InDestinationPlaylistPath)
+{
+	using namespace UE::AvaPageList::Private;
+	if (IsAvaAsset(InAsset, false))
+	{
+		return true;
+	}
+	if (InAsset.IsInstanceOf<UAvalanchePlaylist>(EResolveClass::Yes))
+	{
+		// Don't drop a playlist in itself.
+		if (InAsset.GetSoftObjectPath() != InDestinationPlaylistPath)
+		{
+			return true;
+		}
+	}
+	return false;
+}
+
+TArray<FSoftObjectPath> SAvaPageList::FilterAvaAssetPaths(const TArray<FAssetData>& InAssets)
+{
+	using namespace UE::AvaPageList::Private;
+	TArray<FSoftObjectPath> AvaAssets;
+	AvaAssets.Reserve(InAssets.Num());
+	
+	for (const FAssetData& AssetData : InAssets)
+	{
+		if (IsAvaAsset(AssetData, true))
+		{
+			AvaAssets.Add(AssetData.GetSoftObjectPath());
+		}
+	}
+	return AvaAssets;
+}
+
+TArray<FSoftObjectPath> SAvaPageList::FilterPlaylistPaths(const TArray<FAssetData>& InAssets, const FSoftObjectPath& InDestinationPlaylistPath)
+{
+	TArray<FSoftObjectPath> Playlists;
+	Playlists.Reserve(InAssets.Num());
+
+	for (const FAssetData& AssetData : InAssets)
+	{
+		if (AssetData.IsInstanceOf<UAvalanchePlaylist>(EResolveClass::Yes))
+		{
+			FSoftObjectPath PlaylistPath = AssetData.GetSoftObjectPath();
+			// Don't drop a playlist in itself.
+			if (PlaylistPath != InDestinationPlaylistPath)
+			{
+				Playlists.Add(MoveTemp(PlaylistPath));
+			}
+		}
+	}
+	
+	return Playlists;
+}
+
+FAvaPageInsertPosition SAvaPageList::MakeInsertPosition(EItemDropZone InDropZone, const FAvaPageViewPtr& InItem)
+{
+	FAvaPageInsertPosition InsertPosition;
+	InsertPosition.AdjacentId = InItem.IsValid() ? InItem->GetPageId() : FAvalanchePage::InvalidPageId;
+	InsertPosition.bAddBelow = InDropZone != EItemDropZone::AboveItem;
+	return InsertPosition;
+}
+
+bool SAvaPageList::CanHandleDragObjects(const FDragDropEvent& InDragDropEvent) const
+{
+	if (const TSharedPtr<FAssetDragDropOp> AssetDragDropOp = InDragDropEvent.GetOperationAs<FAssetDragDropOp>())
+	{
+		const FSoftObjectPath DestinationPlaylistPath(GetPlaylist());
+		for (const FAssetData& AssetData : AssetDragDropOp->GetAssets())
+		{
+			if (IsAssetDropSupported(AssetData, DestinationPlaylistPath))
+			{
+				return true;
+			}
+		}
+		return false;
+	}
+
+	if (const TSharedPtr<FAvaPageViewRowDragDropOp> PageDragDropOp = InDragDropEvent.GetOperationAs<FAvaPageViewRowDragDropOp>())
+	{
+		return true;
+	}
+
+	if (const TSharedPtr<FExternalDragOperation> ExternalDragDropOp = InDragDropEvent.GetOperationAs<FExternalDragOperation>())
+	{
+		if (ExternalDragDropOp->HasFiles())
+		{
+			for (const FString& File : ExternalDragDropOp->GetFiles())
+			{
+				if (UE::AvaPlaylistEditor::Utils::CanLoadPlaylistFromFile(*File))
+				{
+					return true;
+				}
+			}
+		}
+	}
+	
+	return false;
+}
+
+bool SAvaPageList::HandleDropEvent(const FDragDropEvent& InDragDropEvent, EItemDropZone InDropZone, const FAvaPageViewPtr& InItem)
+{
+	if (const TSharedPtr<FAssetDragDropOp> AssetDragDropOp = InDragDropEvent.GetOperationAs<FAssetDragDropOp>())
+	{
+		const TArray<FSoftObjectPath> AvaAssets = FilterAvaAssetPaths(AssetDragDropOp->GetAssets());
+		const TArray<FSoftObjectPath> Playlists = FilterPlaylistPaths(AssetDragDropOp->GetAssets(), FSoftObjectPath(GetPlaylist()));
+
+		FScopedTransaction DropTransaction(LOCTEXT("DropTransaction", "Drop Avalanche Assets onto Page List"));
+
+		bool bIsHandled = false;
+		if (!Playlists.IsEmpty())
+		{
+			bIsHandled |= HandleDropPlaylists(Playlists, InDropZone, InItem);
+		}
+		if (!AvaAssets.IsEmpty())
+		{
+			bIsHandled |= HandleDropAvalancheAssets(AvaAssets, InDropZone, InItem);
+		}
+		return bIsHandled;
+	}
+
+	if (const TSharedPtr<FAvaPageViewRowDragDropOp> PageDragDropOp = InDragDropEvent.GetOperationAs<FAvaPageViewRowDragDropOp>())
+	{
+		if (const TSharedPtr<SAvaPageList> FromPageList = PageDragDropOp->GetPageList())
+		{
+			return HandleDropPageIds(FromPageList->GetPageListReference(), PageDragDropOp->GetDraggedIds(), InDropZone, InItem);
+		}
+		return false;
+	}
+
+	if (const TSharedPtr<FExternalDragOperation> ExternalDragDropOp = InDragDropEvent.GetOperationAs<FExternalDragOperation>())
+	{
+		FScopedTransaction DropTransaction(LOCTEXT("DropTransaction", "Drop External Asset onto Page List"));
+
+		return ExternalDragDropOp->HasFiles() && HandleDropExternalFiles(ExternalDragDropOp->GetFiles(), InDropZone, InItem);
+	}
+
+	return false;
+}
+
+FReply SAvaPageList::OnDragDetected(const FGeometry& InMyGeometry, const FPointerEvent& InMouseEvent)
+{
+	return FReply::Handled();
+}
+
+FReply SAvaPageList::OnDragOver(const FGeometry& InMyGeometry, const FDragDropEvent& InDragDropEvent)
+{
+	if (CanHandleDragObjects(InDragDropEvent))
+	{
+		return FReply::Handled();
+	}
+	
+	return FReply::Unhandled();
+}
+
+FReply SAvaPageList::OnDrop(const FGeometry& InMyGeometry, const FDragDropEvent& InDragDropEvent)
+{
+	if (HandleDropEvent(InDragDropEvent, EItemDropZone::OntoItem, nullptr))
+	{
+		return FReply::Handled();
+	}
+		
+	return FReply::Unhandled();
+}
+
+FReply SAvaPageList::OnKeyDown(const FGeometry& InMyGeometry, const FKeyEvent& InKeyEvent)
+{
+	if (CommandList.IsValid() && CommandList->ProcessCommandBindings(InKeyEvent))
+	{
+		return FReply::Handled();
+	}
+
+	return FReply::Unhandled();
+}
+
+FText SAvaPageList::OnAssetSearchBoxSuggestionChosen(const FText& InSearchText, const FString& InSuggestion)
+{
+	int32 SuggestionInsertionIndex = 0;
+	UE::AvalanchePlaylist::Private::ExtractAssetSearchFilterTerms(InSearchText, nullptr, nullptr, &SuggestionInsertionIndex);
+
+	FString SearchString = InSearchText.ToString();
+	SearchString.RemoveAt(SuggestionInsertionIndex, SearchString.Len() - SuggestionInsertionIndex, false);
+	SearchString.Append(InSuggestion);
+	return FText::FromString(SearchString);
+}
+
+void SAvaPageList::OnSearchTextChanged(const FText& FilterText)
+{
+	if (TSharedPtr<FAvaPlaylistEditor> PlaylistEditor = PlaylistEditorWeak.Pin())
+	{
+		PlaylistEditor->SetSearchText(FilterText, PageListType);
+	}
+}
+
+void SAvaPageList::OnSearchTextCommitted(const FText& FilterText, ETextCommit::Type CommitType)
+{
+	if (TSharedPtr<FAvaPlaylistEditor> PlaylistEditor = PlaylistEditorWeak.Pin())
+	{
+		PlaylistEditor->SetSearchText(FilterText, PageListType);
+	}
+}
+
+void SAvaPageList::OnSearchBoxSuggestionFilter(const FText& InSearchText, TArray<FAssetSearchBoxSuggestion>& OutPossibleSuggestions, FText& OutSuggestionHighlightText)
+{
+	// We don't bind the suggestion list, so this list should be empty as we populate it here based on the search term
+	check(OutPossibleSuggestions.IsEmpty());
+
+	FString FilterKey;
+	FString FilterValue;
+	UE::AvalanchePlaylist::Private::ExtractAssetSearchFilterTerms(InSearchText, &FilterKey, &FilterValue, nullptr);
+
+	const IAvaMediaEditorModule& AvaMediaEditorModule = IAvaMediaEditorModule::Get();
+
+	TSet<FString> FilterCache;
+	const TSharedRef<FAvaPlaylistFilterSuggestionPayload> SimplePayload =
+				MakeShared<FAvaPlaylistFilterSuggestionPayload>(FAvaPlaylistFilterSuggestionPayload{ OutPossibleSuggestions, FilterValue, FAvalanchePage::InvalidPageId, nullptr, FilterCache });
+	for (const TSharedPtr<IAvaPlaylistFilterSuggestionFactory>& SimpleSuggestion : AvaMediaEditorModule.GetSimpleSuggestions(PageListType))
+	{
+		SimpleSuggestion->AddSuggestion(SimplePayload);
+	}
+
+	if (const TSharedPtr<FAvaPlaylistEditor> PlaylistEditor = PlaylistEditorWeak.Pin())
+	{
+		if (const UAvalanchePlaylist* Playlist = PlaylistEditor->GetPlaylist())
+		{
+			const TSharedRef<FAvaPlaylistFilterSuggestionPayload> ComplexPayload =
+				MakeShared<FAvaPlaylistFilterSuggestionPayload>(FAvaPlaylistFilterSuggestionPayload{ OutPossibleSuggestions, FilterValue, FAvalanchePage::InvalidPageId, Playlist, FilterCache });
+
+			for (const FAvalanchePage& Page : GetPagesByType(PageListType))
+			{
+				for (const TSharedPtr<IAvaPlaylistFilterSuggestionFactory>& ComplexSuggestion : AvaMediaEditorModule.GetComplexSuggestions(PageListType))
+				{
+					ComplexPayload->ItemPageId = Page.GetPageId();
+					ComplexSuggestion->AddSuggestion(ComplexPayload);
+				}
+			}
+		}
+	}
+	OutSuggestionHighlightText = FText::FromString(FilterValue);
+}
+
+TArray<FAvalanchePage> SAvaPageList::GetPagesByType(EAvaPlaylistSearchListType InPlaylistSearchListType) const
+{
+	TArray<FAvalanchePage> Pages = TArray<FAvalanchePage>();
+
+	if (const TSharedPtr<FAvaPlaylistEditor> PlaylistEditor = PlaylistEditorWeak.Pin())
+	{
+		if (const UAvalanchePlaylist* Playlist = PlaylistEditor->GetPlaylist())
+		{
+			switch (InPlaylistSearchListType)
+			{
+			case EAvaPlaylistSearchListType::Template:
+				Pages = Playlist->GetTemplatePages().Pages;
+				break;
+
+			case EAvaPlaylistSearchListType::Instanced:
+				Pages = Playlist->GetInstancedPages().Pages;
+				break;
+
+			case EAvaPlaylistSearchListType::None:
+			default:
+				break;
+			}
+		}
+	}
+
+	return Pages;
+}
+
+TArray<int32> SAvaPageList::FilterSelectedPages(FFilterPageFunctionRef InFilterPageFunction) const
+{
+	if (const UAvalanchePlaylist* Playlist = GetPlaylist())
+	{
+		if (!SelectedPageIds.IsEmpty())
+		{
+			return InFilterPageFunction(Playlist, SelectedPageIds);
+		}
+	}
+	return {};
+}
+	
+TArray<int32> SAvaPageList::FilterPreviewingPages(FFilterPageFunctionRef InFilterPageFunction) const
+{
+	if (const UAvalanchePlaylist* Playlist = GetPlaylist())
+	{
+		return InFilterPageFunction(Playlist, Playlist->GetPreviewingPageIds());
+	}
+	return {};
+}
+	
+TArray<int32> SAvaPageList::FilterSelectedOrPreviewingPages(FFilterPageFunctionRef InFilterPageFunction, const bool bInAllowFallback) const
+{
+	if (const UAvalanchePlaylist* Playlist = GetPlaylist())
+	{
+		if (!SelectedPageIds.IsEmpty())
+		{
+			const TArray<int32> SelectedPages = InFilterPageFunction(Playlist, SelectedPageIds);
+			if (!bInAllowFallback || !SelectedPages.IsEmpty())
+			{
+				return SelectedPages;
+			}
+		}
+		return InFilterPageFunction(Playlist, Playlist->GetPreviewingPageIds());
+	}
+	return {};
+}
+	
+TArray<int32> SAvaPageList::FilterPageSetForPreview(FFilterPageFunctionRef InFilterPageFunction, const EAvaRundownPageSet InPageSet) const
+{
+	switch (InPageSet)
+	{
+	case EAvaRundownPageSet::SelectedOrPlayingStrict:
+		return FilterSelectedOrPreviewingPages(InFilterPageFunction, /*bAllowFallback*/ false);
+	case EAvaRundownPageSet::SelectedOrPlaying:
+		return FilterSelectedOrPreviewingPages(InFilterPageFunction, /*bAllowFallback*/ true);
+	case EAvaRundownPageSet::Selected:
+		return FilterSelectedPages(InFilterPageFunction);
+	case EAvaRundownPageSet::Playing:
+		return FilterPreviewingPages(InFilterPageFunction);
+	default:
+		return FilterSelectedOrPreviewingPages(InFilterPageFunction, /*bAllowFallback*/ true);
+	}
+}
+
+TArray<int32> SAvaPageList::GetPagesToPreviewIn() const
+{
+	auto KeepPagesToPreviewIn = [](const UAvalanchePlaylist* InPlaylist, const TArray<int32>& InPageIds)
+	{
+		TArray<int32> OutPageIds;
+		OutPageIds.Reserve(InPageIds.Num());
+		for (int32 PageId : InPageIds)
+		{
+			if (InPlaylist->CanPlayPage(PageId, true))
+			{
+				OutPageIds.Add(PageId);
+			}
+		}
+		return OutPageIds;
+	};
+	return FilterSelectedPages(KeepPagesToPreviewIn);
+}
+
+TArray<int32> SAvaPageList::GetPagesToPreviewOut(bool bInForce) const
+{
+	EAvaPlaylistPageStopOptions StopOptions = bInForce ? EAvaPlaylistPageStopOptions::ForceNoTransition : EAvaPlaylistPageStopOptions::Default;
+	auto KeepPagesToPreviewStop = [StopOptions](const UAvalanchePlaylist* InPlaylist, const TArray<int32>& InPageIds)
+	{
+		TArray<int32> OutPageIds;
+		OutPageIds.Reserve(InPageIds.Num());
+		for (int32 PageId : InPageIds)
+		{
+			if (InPlaylist->CanStopPage(PageId, StopOptions, true))
+			{
+				OutPageIds.Add(PageId);
+			}
+		}
+		return OutPageIds;
+	};
+
+	const UAvaRundownEditorSettings* RundownEditorSettings = UAvaRundownEditorSettings::Get();
+	const EAvaRundownPageSet PageSet = RundownEditorSettings ? RundownEditorSettings->PreviewOutActionPageSet : EAvaRundownPageSet::SelectedOrPlaying;
+	return FilterPageSetForPreview(KeepPagesToPreviewStop, PageSet);
+}
+
+TArray<int32> SAvaPageList::GetPagesToPreviewContinue() const
+{
+	auto KeepPagesToPreviewContinue = [](const UAvalanchePlaylist* InPlaylist, const TArray<int32>& InPageIds)
+	{
+		TArray<int32> OutPageIds;
+		OutPageIds.Reserve(InPageIds.Num());
+		for (int32 PageId : InPageIds)
+		{
+			if (InPlaylist->CanContinuePage(PageId, true))
+			{
+				OutPageIds.Add(PageId);
+			}
+		}
+		return OutPageIds;
+	};
+
+	const UAvaRundownEditorSettings* RundownEditorSettings = UAvaRundownEditorSettings::Get();
+	const EAvaRundownPageSet PageSet = RundownEditorSettings ? RundownEditorSettings->PreviewContinueActionPageSet : EAvaRundownPageSet::SelectedOrPlaying;
+	return FilterPageSetForPreview(KeepPagesToPreviewContinue, PageSet);
+}
+
+TArray<int32> SAvaPageList::GetPagesToTakeToProgram() const
+{
+	// Remark: taking to program will take all previewing pages that are not already playing
+	// regardless of what is selected.
+	return FAvaPlaylistPlaybackUtils::GetPagesToTakeToProgram(GetPlaylist(), {});
+}
+
+int32 SAvaPageList::GetPageIdToPreviewNext() const
+{
+	return FAvaPlaylistPlaybackUtils::GetPageIdToPlayNext(
+		GetPlaylist(), GetPageListReference(), /*bInPreview*/ true, UAvalanchePlaylist::GetDefaultPreviewChannelName());
+}
+
+#undef LOCTEXT_NAMESPACE
diff --git a/Engine/Plugins/Experimental/Avalanche/Source/AvalancheMediaEditor/Private/Playlist/Pages/Slate/SAvaPageList.h b/Engine/Plugins/Experimental/Avalanche/Source/AvalancheMediaEditor/Private/Playlist/Pages/Slate/SAvaPageList.h
new file mode 100644
--- /dev/null
+++ b/Engine/Plugins/Experimental/Avalanche/Source/AvalancheMediaEditor/Private/Playlist/Pages/Slate/SAvaPageList.h
@@ -0,0 +1,228 @@
+// Copyright Epic Games, Inc. All Rights Reserved.
+
+#pragma once
+
+#include "AvaMediaDefines.h"
+#include "Playlist/AvaPlaylistDefines.h"
+#include "Playlist/AvaRundownEditorDefines.h"
+#include "Playlist/AvalanchePlaylist.h"
+#include "Playlist/Pages/PageViews/AvaPageView.h"
+#include "Widgets/SCompoundWidget.h"
+
+class FAvaPlaylistEditor;
+class FAvaPlaylistPageContextMenu;
+class FUICommandList;
+class IAvaPageViewColumn;
+class ITableRow;
+class SHeaderRow;
+class SHorizontalBox;
+template <typename ItemType>
+class SListView;
+class SSearchBox;
+class SAssetSearchBox;
+class STableViewBase;
+struct FAssetSearchBoxSuggestion;
+struct FAvalanchePage;
+struct FAvaPageListChangeParams;
+enum class EItemDropZone;
+enum class EAvaPlaylistSearchListType : uint8;
+
+class SAvaPageList : public SCompoundWidget
+{
+	SLATE_DECLARE_WIDGET(SAvaPageList, SCompoundWidget)
+
+public:
+	SLATE_BEGIN_ARGS(SAvaPageList) {}
+	SLATE_END_ARGS()
+
+	/** Constructs this widget with InArgs */
+	void Construct(const FArguments& InArgs, TSharedPtr<FAvaPlaylistEditor> InPlaylistEditor, const FAvaPageListReference& InPageListReference, EAvaPlaylistSearchListType InPageListType);
+	virtual ~SAvaPageList() override;
+
+	const FAvaPageListReference& GetPageListReference() const { return PageListReference; }
+
+	void OnPageEvent(const TArray<int32>& InSelectedPageIds, UE::AvalanchePlaylist::EPageEvent InPageEvent);
+
+	void OnPageItemActionRequested(const TArray<int32>& InPageIds, IAvaPageView::FOnPageAction&(IAvaPageView::*InFunc)());
+
+	void OnPageSelectionRequested(const TArray<int32>& InPageIds);
+
+	TSharedRef<ITableRow> GeneratePageTableRow(FAvaPageViewPtr InPageView, const TSharedRef<STableViewBase>& InOwnerTable);
+
+	void OnPageSelected(FAvaPageViewPtr InPageView, ESelectInfo::Type InSelectInfo);
+
+	TSharedPtr<SListView<FAvaPageViewPtr>> GetPageListView() { return PageListView; }
+
+	TSharedPtr<IAvaPageViewColumn> FindColumn(FName InColumnName) const;
+
+	FAvaPageViewPtr GetPageViewPtr(int32 InPageId) const;
+
+	void SelectPage(int32 InPageId, bool bInScrollIntoView = true);
+	void SelectPages(const TArray<int32>& InPageIds, bool bInScrollIntoView = true);
+	void DeselectPage(int32 InPageId);
+	void DeselectPages();
+
+	const TArray<int32>& GetSelectedPageIds() const { return SelectedPageIds; }
+	TArray<int32> GetPlayingPageIds() const;
+	const TArray<FAvaPageViewPtr>& GetPageViews() const { return PageViews; }
+
+	TSharedPtr<FAvaPlaylistEditor> GetPlaylistEditor() const;
+
+	UAvalanchePlaylist* GetPlaylist() const;
+	UAvalanchePlaylist* GetValidPlaylist() const;
+
+	virtual void BindCommands();
+
+	virtual void Refresh() = 0;
+
+	virtual void CreateColumns() = 0;
+
+	virtual TSharedPtr<SWidget> OnContextMenuOpening() = 0;
+	
+	int32 GetFirstSelectedPageId() const;
+
+	TSharedRef<SWidget> GetPageListContextMenu();
+
+	EAvaPlaylistSearchListType GetPageListType() const { return PageListType; }
+
+	bool CanCopySelectedPages() const;
+	void CopySelectedPages();
+
+	bool CanCutSelectedPages() const;
+	void CutSelectedPages();
+
+	bool CanPaste() const;
+	void Paste();
+
+	bool CanDuplicateSelectedPages() const;
+	void DuplicateSelectedPages();
+
+	bool CanAddPage() const;
+	bool CanAddTemplate() const;
+	bool CanCreateInstance() const;
+
+	bool CanRemoveSelectedPages() const;
+	void RemoveSelectedPages();
+
+	bool CanRenameSelectedPage() const;
+	void RenameSelectedPage();
+
+	bool CanRenumberSelectedPage() const;
+	void RenumberSelectedPage();
+
+	bool CanReimportSelectedPage() const;
+	void ReimportSelectedPage() const;
+
+	bool CanEditSelectedPageSource() const;
+	void EditSelectedPageSource();
+
+	bool CanExportSelectedPagesToPlaylist();
+	void ExportSelectedPagesToPlaylist();
+
+	bool CanExportSelectedPagesToExternalFile(const TCHAR* InType);
+	void ExportSelectedPagesToExternalFile(const TCHAR* InType);
+
+	bool CanPreviewPlaySelectedPage() const;
+	void PreviewPlaySelectedPage(bool bInToMark) const;
+
+	bool CanPreviewStopSelectedPage(bool bInForce) const;
+	void PreviewStopSelectedPage(bool bInForce) const;
+
+	bool CanPreviewContinueSelectedPage() const;
+	void PreviewContinueSelectedPage() const;
+
+	bool CanPreviewPlayNextPage() const;
+	void PreviewPlayNextPage();
+
+	bool CanTakeToProgram() const;
+	void TakeToProgram() const;
+
+	/** For the given asset, check if it supported for drop operation. */
+	static bool IsAssetDropSupported(const FAssetData& InAsset, const FSoftObjectPath& InDestinationPlaylistPath);
+	
+	/** From the given array of asset data, filter only ava assets. */
+	static TArray<FSoftObjectPath> FilterAvaAssetPaths(const TArray<FAssetData>& InAssets);
+
+	/** From the given array of asset data, keep only the playlists. */
+	static TArray<FSoftObjectPath> FilterPlaylistPaths(const TArray<FAssetData>& InAssets, const FSoftObjectPath& InDestinationPlaylistPath);
+
+	/** Helper to convert drop parameters to insert position. */
+	static FAvaPageInsertPosition MakeInsertPosition(EItemDropZone InDropZone, const FAvaPageViewPtr& InItem);
+
+	virtual bool CanHandleDragObjects(const FDragDropEvent& InDragDropEvent) const;
+	virtual bool HandleDropEvent(const FDragDropEvent& InDragDropEvent, EItemDropZone InDropZone, const FAvaPageViewPtr& InItem);
+	
+	virtual bool HandleDropAvalancheAssets(const TArray<FSoftObjectPath>& InAvaAssets, EItemDropZone InDropZone, const FAvaPageViewPtr& InItem) = 0;
+	virtual bool HandleDropPlaylists(const TArray<FSoftObjectPath>& InPlaylistPaths, EItemDropZone InDropZone, const FAvaPageViewPtr& InItem) = 0;
+	virtual bool HandleDropPageIds(const FAvaPageListReference& PageListReference, const TArray<int32>& InPageIds, EItemDropZone InDropZone, const FAvaPageViewPtr& InItem) = 0;
+	virtual bool HandleDropExternalFiles(const TArray<FString>& InFiles, EItemDropZone InDropZone, const FAvaPageViewPtr& InItem) = 0;
+
+	//~ Begin SWidget
+	virtual FReply OnDragDetected(const FGeometry& InMyGeometry, const FPointerEvent& InMouseEvent) override;
+	virtual FReply OnDragOver(const FGeometry& InMyGeometry, const FDragDropEvent& InDragDropEvent) override;
+	virtual FReply OnDrop(const FGeometry& InMyGeometry, const FDragDropEvent& InDragDropEvent) override;
+	virtual FReply OnKeyDown(const FGeometry& InMyGeometry, const FKeyEvent& InKeyEvent) override;
+	//~ End SWidget
+
+protected: 
+	static bool IsPageIdValid(int32 InPageId)
+	{
+		return InPageId != FAvalanchePage::InvalidPageId;
+	}
+
+	using FFilterPageFunctionRef = TFunctionRef<TArray<int32>(const UAvalanchePlaylist* InPlaylist, const TArray<int32>& InPageIds)>;
+	
+	TArray<int32> FilterSelectedPages(FFilterPageFunctionRef InFilterPageFunction) const;
+	TArray<int32> FilterPreviewingPages(FFilterPageFunctionRef InFilterPageFunction) const;
+	TArray<int32> FilterSelectedOrPreviewingPages(FFilterPageFunctionRef InFilterPageFunction, const bool bInAllowFallback) const;
+	TArray<int32> FilterPageSetForPreview(FFilterPageFunctionRef InFilterPageFunction, const EAvaRundownPageSet InPageSet) const;
+	
+	TArray<int32> GetPagesToPreviewIn() const;
+	TArray<int32> GetPagesToPreviewOut(bool bInForce) const;
+	TArray<int32> GetPagesToPreviewContinue() const;
+	TArray<int32> GetPagesToTakeToProgram() const;
+	int32 GetPageIdToPreviewNext() const;
+	
+private:
+	FText OnAssetSearchBoxSuggestionChosen(const FText& InSearchText, const FString& InSuggestion);
+
+	void OnSearchTextChanged(const FText& FilterText);
+
+	void OnSearchTextCommitted(const FText& FilterText, ETextCommit::Type CommitType);
+
+	void OnSearchBoxSuggestionFilter(const FText& InSearchText, TArray<FAssetSearchBoxSuggestion>& OutPossibleSuggestions, FText& OutSuggestionHighlightText);
+
+	TArray<FAvalanchePage> GetPagesByType(EAvaPlaylistSearchListType InPlaylistSearchListType) const;
+
+protected:
+	TWeakPtr<FAvaPlaylistEditor> PlaylistEditorWeak;
+
+	TSharedPtr<FAvaPlaylistPageContextMenu> PageContextMenu;
+
+	FAvaPageListReference PageListReference;
+
+	TSharedPtr<SHorizontalBox> SearchBar;
+
+	TSharedPtr<SHeaderRow> HeaderRow;
+
+	TMap<FName, TSharedPtr<IAvaPageViewColumn>> Columns;
+
+	TSharedPtr<SListView<FAvaPageViewPtr>> PageListView;
+
+	TArray<FAvaPageViewPtr> PageViews;
+
+	TArray<int32> SelectedPageIds;
+
+	TSharedPtr<FUICommandList> CommandList;
+
+	/**
+	 * For template and instance lists, it will return the list of page ids created.
+	 * For a sub list it will return the list of page ids that were added to the page view.
+	 */
+	virtual TArray<int32> AddPastedPages(const TArray<FAvalanchePage>& InPages) = 0;
+
+private:
+	TSharedPtr<SAssetSearchBox> AssetSearchBoxPtr;
+
+	EAvaPlaylistSearchListType PageListType;
+};
diff --git a/Engine/Plugins/Experimental/Avalanche/Source/AvalancheMediaEditor/Private/Playlist/Pages/Slate/SAvaPageViewRow.cpp b/Engine/Plugins/Experimental/Avalanche/Source/AvalancheMediaEditor/Private/Playlist/Pages/Slate/SAvaPageViewRow.cpp
new file mode 100644
--- /dev/null
+++ b/Engine/Plugins/Experimental/Avalanche/Source/AvalancheMediaEditor/Private/Playlist/Pages/Slate/SAvaPageViewRow.cpp
@@ -0,0 +1,188 @@
+// Copyright Epic Games, Inc. All Rights Reserved.
+
+#include "Playlist/Pages/Slate/SAvaPageViewRow.h"
+#include "AvaBlueprint.h"
+#include "AvaMediaEditorStyle.h"
+#include "DragAndDrop/AssetDragDropOp.h"
+#include "Playlist/Pages/Columns/AvaPageViewColumn.h"
+#include "SAvaPageList.h"
+#include "ScopedTransaction.h"
+#include "Widgets/Layout/SBorder.h"
+#include "Widgets/SNullWidget.h"
+#include "Widgets/Text/STextBlock.h"
+#include "Widgets/Views/SListView.h"
+
+#define LOCTEXT_NAMESPACE "SAvaPageViewRow"
+
+namespace UE::AvalancheMediaEditor::RowUtils
+{
+	FName ActivePageBrushName = FName(TEXT("TableView.ActivePageBrush"));
+	FName ActivePageHoveredBrushName = FName(TEXT("TableView.Hovered.ActivePageBrush"));
+	FName ActivePageSelectedBrushName = FName(TEXT("TableView.Selected.ActivePageBrush"));
+	FName DisabledPageBrushName = FName(TEXT("TableView.DisabledPageBrush"));
+	FName DisabledPageHoveredBrushName = FName(TEXT("TableView.Hovered.DisabledPageBrush"));
+	FName DisabledPageSelectedBrushName = FName(TEXT("TableView.Selected.DisabledPageBrush"));
+}
+
+void SAvaPageViewRow::Construct(const FArguments& InArgs, FAvaPageViewPtr InPageView, const TSharedRef<SAvaPageList>& InPageList)
+{
+	PageViewWeak = InPageView;
+	PageListWeak = InPageList;
+
+	TSharedPtr<SListView<FAvaPageViewPtr>> PageListView = InPageList->GetPageListView();
+	check(PageListView.IsValid());
+	
+	SMultiColumnTableRow::Construct(
+		FSuperRowType::FArguments()
+		.OnDragDetected(this, &SAvaPageViewRow::OnPageViewDragDetected)
+		.OnCanAcceptDrop(this, &SAvaPageViewRow::OnPageViewCanAcceptDrop)
+		.OnAcceptDrop(this, &SAvaPageViewRow::OnPageViewAcceptDrop)
+		, PageListView.ToSharedRef()
+	);
+	SetBorderImage(TAttribute<const FSlateBrush*>(this, &SAvaPageViewRow::GetBorder));
+}
+
+TSharedRef<SWidget> SAvaPageViewRow::GenerateWidgetForColumn(const FName& InColumnName)
+{
+	if (const FAvaPageViewPtr PageView = PageViewWeak.Pin())
+	{
+		if (const TSharedPtr<IAvaPageViewColumn> Column = PageListWeak.Pin()->FindColumn(InColumnName))
+		{
+			const TSharedRef<SWidget> RowWidget = Column->ConstructRowWidget(PageView.ToSharedRef(), SharedThis(this));
+			RowWidget->SetOnMouseDoubleClick(FPointerEventHandler::CreateLambda([this](const FGeometry&, const FPointerEvent&)
+			{
+				return PageViewWeak.IsValid()? PageViewWeak.Pin()->OnPreviewButtonClicked() : FReply::Handled();
+			}));
+			return RowWidget;
+		}
+	}
+	return SNullWidget::NullWidget;
+}
+
+const FSlateBrush* SAvaPageViewRow::GetBorder() const
+{
+	if (const TSharedPtr<IAvaPageView> PageView = PageViewWeak.Pin())
+	{
+		const UAvalanchePlaylist* const Playlist = PageView->GetPlaylist();
+		if (Playlist && !PageView->IsTemplate())
+		{
+			const FAvalanchePage& Page = Playlist->GetPage(PageView->GetPageId());
+			const TArray<FAvalanchePageStatus> Statuses = Page.GetPageContextualStatuses(Playlist);
+
+			if (FAvalanchePage::StatusesContainsStatus(Statuses, { EAvalanchePageStatus::Loaded, EAvalanchePageStatus::Playing }))
+			{
+				return GetRowColorBrush(EAvalancheRowState::Active);
+			}
+
+			if (FAvalanchePage::StatusesContainsStatus(Statuses, { EAvalanchePageStatus::Error, EAvalanchePageStatus::Missing }))
+			{
+				return GetRowColorBrush(EAvalancheRowState::Disabled);
+			}
+		}
+	}
+
+	return SMultiColumnTableRow::GetBorder();
+}
+
+FReply SAvaPageViewRow::OnPageViewDragDetected(const FGeometry& InMyGeometry, const FPointerEvent& InPointerEvent)
+{
+	TSharedPtr<SAvaPageList> PageList = PageListWeak.Pin();
+
+	if (!PageList.IsValid() || PageList->GetSelectedPageIds().IsEmpty())
+	{
+		return FReply::Unhandled();
+	}
+
+	return FReply::Handled().BeginDragDrop(FAvaPageViewRowDragDropOp::New(PageList));
+}
+
+TOptional<EItemDropZone> SAvaPageViewRow::OnPageViewCanAcceptDrop(const FDragDropEvent& InDragDropEvent, EItemDropZone InDropZone, FAvaPageViewPtr InItem)
+{
+	if (TSharedPtr<SAvaPageList> PageList = PageListWeak.Pin())
+	{
+		if (PageList->CanHandleDragObjects(InDragDropEvent))
+		{
+			return InDropZone;
+		}
+	}
+
+	return TOptional<EItemDropZone>();
+}
+
+FReply SAvaPageViewRow::OnPageViewAcceptDrop(const FDragDropEvent& InDragDropEvent, EItemDropZone InDropZone, FAvaPageViewPtr InItem)
+{
+	if (const TSharedPtr<SAvaPageList> PageList = PageListWeak.Pin())
+	{
+		if (PageList->HandleDropEvent(InDragDropEvent, InDropZone, InItem))
+		{
+			return FReply::Handled();
+		}
+	}
+
+	return FReply::Unhandled();
+}
+
+const FSlateBrush* SAvaPageViewRow::GetRowColorBrush(EAvalancheRowState InRowState) const
+{
+	using namespace UE::AvalancheMediaEditor::RowUtils;
+	FName NormalStateBrushName;
+	FName HoveredStateBrushName;
+	FName SelectedStateBrushName;
+
+	if (InRowState == EAvalancheRowState::Active)
+	{
+		NormalStateBrushName = ActivePageBrushName;
+		HoveredStateBrushName = ActivePageHoveredBrushName;
+		SelectedStateBrushName = ActivePageSelectedBrushName;
+	}
+	else
+	{
+		NormalStateBrushName = DisabledPageBrushName;
+		HoveredStateBrushName = DisabledPageHoveredBrushName;
+		SelectedStateBrushName = DisabledPageSelectedBrushName;
+	}
+
+	if (IsHovered())
+	{
+		return FAvaMediaEditorStyle::Get().GetBrush(HoveredStateBrushName);
+	}
+
+	if (IsSelected())
+	{
+		return FAvaMediaEditorStyle::Get().GetBrush(SelectedStateBrushName);
+	}
+
+	return FAvaMediaEditorStyle::Get().GetBrush(NormalStateBrushName);
+}
+
+TSharedPtr<SWidget> FAvaPageViewRowDragDropOp::GetDefaultDecorator() const
+{
+	TArray<FText> IdTexts;
+
+	for (int32 DraggedId : DraggedIds)
+	{
+		IdTexts.Add(FText::AsNumber(DraggedId, &UE::AvalanchePlaylist::FEditorMetrics::PageIdFormattingOptions));
+	}
+
+	return SNew(SBorder)
+		.Padding(5.f)
+		.BorderImage(FAppStyle::GetBrush("Graph.ConnectorFeedback.Border"))
+		.Content()
+		[
+			SNew(STextBlock)
+			.Text(FText::Join(LOCTEXT("Comma", ", "), IdTexts))
+		];
+}
+
+TSharedRef<FAvaPageViewRowDragDropOp> FAvaPageViewRowDragDropOp::New(const TSharedPtr<SAvaPageList>& InPageList)
+{
+	check(InPageList.IsValid());
+
+	TSharedRef<FAvaPageViewRowDragDropOp> Operation = MakeShared<FAvaPageViewRowDragDropOp>();
+	Operation->PageListWeak = InPageList;
+	Operation->DraggedIds = InPageList->GetSelectedPageIds(); // Cache this just in case.
+	Operation->Construct();
+	return Operation;
+}
+
+#undef LOCTEXT_NAMESPACE
diff --git a/Engine/Plugins/Experimental/Avalanche/Source/AvalancheMediaEditor/Private/Playlist/Pages/Slate/SAvaReadPage.cpp b/Engine/Plugins/Experimental/Avalanche/Source/AvalancheMediaEditor/Private/Playlist/Pages/Slate/SAvaReadPage.cpp
new file mode 100644
--- /dev/null
+++ b/Engine/Plugins/Experimental/Avalanche/Source/AvalancheMediaEditor/Private/Playlist/Pages/Slate/SAvaReadPage.cpp
@@ -0,0 +1,195 @@
+// Copyright Epic Games, Inc. All Rights Reserved.
+
+#include "SAvaReadPage.h"
+
+#include "Playlist/AvaPlaylistEditor.h"
+#include "SAvaInstancedPageList.h"
+#include "SAvaPageList.h"
+#include "Framework/Application/SlateApplication.h"
+#include "Widgets/Input/SButton.h"
+#include "Widgets/Input/SEditableTextBox.h"
+
+#define LOCTEXT_NAMESPACE "SAvaReadPage"
+
+class SAvaReadPageEditableTextBox : public SEditableTextBox
+{
+public:
+	bool IsKeyRelevant(const FKeyEvent& InKeyEvent) const
+	{
+		static const TSet<FKey> NumPadKeys { EKeys::NumPadZero, EKeys::NumPadOne, EKeys::NumPadTwo
+			, EKeys::NumPadThree, EKeys::NumPadFour, EKeys::NumPadFive, EKeys::NumPadSix, EKeys::NumPadSeven
+			, EKeys::NumPadEight, EKeys::NumPadNine };
+
+		const bool bTextNeedsFocus = EditableText.IsValid() && !EditableText->HasKeyboardFocus();
+
+		return bTextNeedsFocus && NumPadKeys.Contains(InKeyEvent.GetKey());
+	}
+
+	bool ProcessPlaylistKeyDown(const FKeyEvent& InKeyEvent)
+	{
+		// Check if Editable Text needs focus, and that the Key Event is a NumPad Key event
+		if (IsKeyRelevant(InKeyEvent))
+		{
+			SetText(FText::GetEmpty());
+			TSharedRef<SWidget> TextWidget = EditableText.ToSharedRef();
+			FSlateApplication::Get().SetKeyboardFocus(TextWidget);
+			TextWidget->OnKeyDown(TextWidget->GetTickSpaceGeometry(), InKeyEvent);
+			return true;
+		}
+		return false;
+	}
+};
+
+void SAvaReadPage::Construct(const FArguments& InArgs, TSharedPtr<FAvaPlaylistEditor> InPlaylistEditor)
+{
+	PlaylistEditorWeak = InPlaylistEditor;
+	InPlaylistEditor->GetOnPageEvent().AddSP(this, &SAvaReadPage::OnPageEvent);
+
+	const int32 PageId = InPlaylistEditor->GetFirstSelectedPageOnActiveSubListWidget();
+	const FText PageIdText = PageId != FAvalanchePage::InvalidPageId
+		? FText::AsNumber(PageId, &UE::AvalanchePlaylist::FEditorMetrics::PageIdFormattingOptions)
+		: FText::GetEmpty();
+
+	// Small overestimate of 5 digit width.
+	// Only relevant when text is empty (as numbers themselves are formatted for min 5 digits)
+	constexpr float ReadPageWidth = 40.f;
+
+	ChildSlot
+	.Padding(2.f, 0.f)
+	[
+		SNew(SHorizontalBox)
+		+ SHorizontalBox::Slot()
+		.FillWidth(1.f)
+		.VAlign(EVerticalAlignment::VAlign_Center)
+		[
+			SAssignNew(ReadPageText, SAvaReadPageEditableTextBox)
+			.Text(PageIdText)
+			.OnVerifyTextChanged(this, &SAvaReadPage::OnVerifyTextChanged)
+			.OnTextCommitted(this, &SAvaReadPage::OnTextCommitted)
+			.MinDesiredWidth(ReadPageWidth)
+		]
+		+ SHorizontalBox::Slot()
+		.AutoWidth()
+		.VAlign(EVerticalAlignment::VAlign_Center)
+		[
+			SNew(SButton)
+			.ButtonStyle(FAppStyle::Get(), "SimpleButton")
+			.Text(LOCTEXT("ReadPageButton", "Read Page"))
+			.OnClicked(this, &SAvaReadPage::OnReadPageClicked)
+		]
+	];
+}
+
+SAvaReadPage::~SAvaReadPage()
+{
+	if (PlaylistEditorWeak.IsValid())
+	{
+		PlaylistEditorWeak.Pin()->GetOnPageEvent().RemoveAll(this);
+	}
+}
+
+void SAvaReadPage::OnPageEvent(const TArray<int32>& InSelectedPageIds, UE::AvalanchePlaylist::EPageEvent InPageEvent)
+{
+	if (!InSelectedPageIds.IsEmpty())
+	{
+		ReadPageText->SetText(FText::AsNumber(InSelectedPageIds[0]
+			, &UE::AvalanchePlaylist::FEditorMetrics::PageIdFormattingOptions));
+	}
+	else
+	{
+		ReadPageText->SetText(FText::GetEmpty());
+	}
+}
+
+bool SAvaReadPage::IsKeyRelevant(const FKeyEvent& InKeyEvent) const
+{
+	if (ReadPageText.IsValid())
+	{
+		return ReadPageText->IsKeyRelevant(InKeyEvent);
+	}
+	return false;
+}
+
+bool SAvaReadPage::ProcessPlaylistKeyDown(const FKeyEvent& InKeyEvent)
+{
+	if (ReadPageText.IsValid())
+	{
+		return ReadPageText->ProcessPlaylistKeyDown(InKeyEvent);
+	}
+	return false;
+}
+
+bool SAvaReadPage::OnVerifyTextChanged(const FText& InText, FText& OutErrorMessage) const
+{
+	const FString String = InText.ToString();
+
+	if (String.IsEmpty())
+	{
+		//Allow for the Text to be Empty (i.e. no pages are selected)
+		return true;
+	}
+	
+	if (!String.IsNumeric())
+	{
+		OutErrorMessage = LOCTEXT("NonNumericError", "Text is not Numeric");
+		return false;
+	}
+
+	UAvalanchePlaylist* const Playlist = PlaylistEditorWeak.IsValid()
+		? PlaylistEditorWeak.Pin()->GetPlaylist()
+		: nullptr;
+
+	if (!Playlist)
+	{
+		OutErrorMessage = LOCTEXT("NoPlaylist", "No Playlist found for Item");
+		return false;
+	}
+	
+	const int32 Id = FCString::Atoi(*String);
+
+	if (!Playlist->GetPage(Id).IsValidPage())
+	{
+		OutErrorMessage = LOCTEXT("NoPageFound", "No Page with given Page Id exists");
+		return false;
+	}
+	
+	return true;
+}
+
+void SAvaReadPage::OnTextCommitted(const FText& InText, ETextCommit::Type InCommitType)
+{
+	TSharedPtr<FAvaPlaylistEditor> PlaylistEditor = PlaylistEditorWeak.Pin();
+	check(PlaylistEditor.IsValid());
+
+	TSharedPtr<SAvaInstancedPageList> PageList = PlaylistEditor->GetActiveListWidget();
+
+	if (PageList.IsValid())
+	{
+		if (InText.IsNumeric())
+		{
+			const int32 Id = FCString::Atoi(*InText.ToString());
+
+			FModifierKeysState KeyState = FSlateApplication::Get().GetModifierKeys();
+			const bool bAddtoSelection = KeyState.IsControlDown() || KeyState.IsCommandDown() || KeyState.IsAltDown();
+
+			if (!bAddtoSelection)
+			{
+				PageList->DeselectPages();
+			}
+
+			PageList->SelectPage(Id);
+		}
+		else if (InText.IsEmpty())
+		{
+			PageList->DeselectPages();
+		}
+	}
+}
+
+FReply SAvaReadPage::OnReadPageClicked()
+{
+	OnTextCommitted(ReadPageText->GetText(), ETextCommit::Type::OnEnter);
+	return FReply::Handled();
+}
+
+#undef LOCTEXT_NAMESPACE
diff --git a/Engine/Plugins/Experimental/Avalanche/Source/AvalancheMediaEditor/Private/Playlist/Pages/Slate/SAvaReadPage.h b/Engine/Plugins/Experimental/Avalanche/Source/AvalancheMediaEditor/Private/Playlist/Pages/Slate/SAvaReadPage.h
new file mode 100644
--- /dev/null
+++ b/Engine/Plugins/Experimental/Avalanche/Source/AvalancheMediaEditor/Private/Playlist/Pages/Slate/SAvaReadPage.h
@@ -0,0 +1,38 @@
+// Copyright Epic Games, Inc. All Rights Reserved.
+
+#pragma once
+
+#include "Playlist/AvaPlaylistDefines.h"
+#include "Widgets/SCompoundWidget.h"
+
+class FAvaPlaylistEditor;
+class SAvaReadPageEditableTextBox;
+
+class SAvaReadPage : public SCompoundWidget
+{
+public:
+	
+	SLATE_BEGIN_ARGS(SAvaReadPage) {}
+	SLATE_END_ARGS()
+
+	void Construct(const FArguments& InArgs, TSharedPtr<FAvaPlaylistEditor> InPlaylistEditor);
+
+	virtual ~SAvaReadPage() override;
+
+	void OnPageEvent(const TArray<int32>& InSelectedPageIds, UE::AvalanchePlaylist::EPageEvent InPageEvent);
+
+	bool IsKeyRelevant(const FKeyEvent& InKeyEvent) const;
+
+	bool ProcessPlaylistKeyDown(const FKeyEvent& InKeyEvent);
+
+	bool OnVerifyTextChanged(const FText& InText, FText& OutErrorMessage) const;
+
+	void OnTextCommitted(const FText& InText, ETextCommit::Type InCommitType);
+
+	FReply OnReadPageClicked();
+
+private:
+	TWeakPtr<FAvaPlaylistEditor> PlaylistEditorWeak;
+
+	TSharedPtr<SAvaReadPageEditableTextBox> ReadPageText;
+};
diff --git a/Engine/Plugins/Experimental/Avalanche/Source/AvalancheMediaEditor/Private/Playlist/Pages/Slate/SAvaSubListStartPage.cpp b/Engine/Plugins/Experimental/Avalanche/Source/AvalancheMediaEditor/Private/Playlist/Pages/Slate/SAvaSubListStartPage.cpp
new file mode 100644
--- /dev/null
+++ b/Engine/Plugins/Experimental/Avalanche/Source/AvalancheMediaEditor/Private/Playlist/Pages/Slate/SAvaSubListStartPage.cpp
@@ -0,0 +1,54 @@
+// Copyright Epic Games, Inc. All Rights Reserved.
+
+#include "SAvaSubListStartPage.h"
+
+#include "Playlist/AvaPlaylistEditor.h"
+#include "Styling/AppStyle.h"
+#include "Widgets/Input/SButton.h"
+#include "Widgets/Layout/SBox.h"
+#include "Widgets/Text/STextBlock.h"
+
+#define LOCTEXT_NAMESPACE "SAvaSubListStartPage"
+
+void SAvaSubListStartPage::Construct(const FArguments& InArgs, TSharedPtr<FAvaPlaylistEditor> InPlaylistEditor)
+{
+	PlaylistEditorWeak = InPlaylistEditor;
+
+	ChildSlot
+	[
+		SNew(SBox)
+		.Padding(10.f, 10.f)
+		.HAlign(EHorizontalAlignment::HAlign_Center)
+		.VAlign(EVerticalAlignment::VAlign_Top)
+		[
+			SNew(SButton)
+			.ToolTipText(LOCTEXT("AddSubListTooltip", "Add Page View"))
+			.OnClicked(this, &SAvaSubListStartPage::OnCreateSubListClicked)
+			[
+				SNew(STextBlock)
+				.Text(LOCTEXT("AddSubList", "Add Page View"))
+			]
+		]
+	];
+}
+
+FReply SAvaSubListStartPage::OnCreateSubListClicked()
+{
+	TSharedPtr<FAvaPlaylistEditor> PlaylistEditor = PlaylistEditorWeak.Pin();
+
+	if (PlaylistEditor.IsValid())
+	{
+		UAvalanchePlaylist* Playlist = PlaylistEditor->GetPlaylist();
+
+		if (IsValid(Playlist))
+		{
+			Playlist->AddSubList();
+
+			return FReply::Handled();
+		}
+	}
+
+	return FReply::Unhandled();
+}
+
+#undef LOCTEXT_NAMESPACE
diff --git a/Engine/Plugins/Experimental/Avalanche/Source/AvalancheMediaEditor/Private/Playlist/Pages/Slate/SAvaSubListStartPage.h b/Engine/Plugins/Experimental/Avalanche/Source/AvalancheMediaEditor/Private/Playlist/Pages/Slate/SAvaSubListStartPage.h
new file mode 100644
--- /dev/null
+++ b/Engine/Plugins/Experimental/Avalanche/Source/AvalancheMediaEditor/Private/Playlist/Pages/Slate/SAvaSubListStartPage.h
@@ -0,0 +1,22 @@
+// Copyright Epic Games, Inc. All Rights Reserved.
+
+#pragma once
+
+#include "Widgets/SCompoundWidget.h"
+
+class FAvaPlaylistEditor;
+
+class SAvaSubListStartPage : public SCompoundWidget
+{
+public:
+	SLATE_BEGIN_ARGS(SAvaSubListStartPage) {}
+	SLATE_END_ARGS()
+
+	void Construct(const FArguments& InArgs, TSharedPtr<FAvaPlaylistEditor> InPlaylistEditor);
+	virtual ~SAvaSubListStartPage() override {}
+
+protected:
+	TWeakPtr<FAvaPlaylistEditor> PlaylistEditorWeak;
+
+	FReply OnCreateSubListClicked();
+};
diff --git a/Engine/Plugins/Experimental/Avalanche/Source/AvalancheMediaEditor/Private/Playlist/Pages/Slate/SAvaTemplatePageList.cpp b/Engine/Plugins/Experimental/Avalanche/Source/AvalancheMediaEditor/Private/Playlist/Pages/Slate/SAvaTemplatePageList.cpp
new file mode 100644
--- /dev/null
+++ b/Engine/Plugins/Experimental/Avalanche/Source/AvalancheMediaEditor/Private/Playlist/Pages/Slate/SAvaTemplatePageList.cpp
@@ -0,0 +1,502 @@
+// Copyright Epic Games, Inc. All Rights Reserved.
+
+#include "SAvaTemplatePageList.h"
+
+#include "Playlist/AvaPlaylistCommands.h"
+#include "Playlist/AvaPlaylistEditor.h"
+#include "Playlist/AvaPlaylistEditorUtils.h"
+#include "Playlist/AvalanchePage.h"
+#include "Playlist/AvalanchePlaylist.h"
+#include "Playlist/Factories/Filters/AvaPlaylistFactoriesUtils.h"
+#include "Playlist/Pages/Columns/AvaPageAssetSelectorColumn.h"
+#include "Playlist/Pages/Columns/AvaPageIdColumn.h"
+#include "Playlist/Pages/Columns/AvaPageNameColumn.h"
+#include "Playlist/Pages/Columns/AvaPageTemplateStatusColumn.h"
+#include "Playlist/Pages/Columns/AvaPageThumbnailColumn.h"
+#include "Playlist/Pages/Columns/AvaPageTransitionLayerColumn.h"
+#include "Playlist/Pages/PageViews/AvaTemplatePageViewImpl.h"
+#include "Widgets/SNullWidget.h"
+#include "Widgets/Views/SListView.h"
+
+#define LOCTEXT_NAMESPACE "SAvaTemplatePageList"
+
+void SAvaTemplatePageList::PrivateRegisterAttributes(struct FSlateAttributeDescriptor::FInitializer&)
+{
+
+}
+
+void SAvaTemplatePageList::Construct(const FArguments& InArgs, TSharedPtr<FAvaPlaylistEditor> InPlaylistEditor)
+{
+	SAvaPageList::Construct(SAvaPageList::FArguments(), InPlaylistEditor, UAvalanchePlaylist::TemplatePageList, EAvaPlaylistSearchListType::Template);
+
+	PlaylistEditorWeak = InPlaylistEditor;
+	check(InPlaylistEditor.IsValid());
+
+	UAvalanchePlaylist* const Playlist = InPlaylistEditor->GetPlaylist();
+	check(Playlist);
+
+	Playlist->GetOnTemplatePageListChanged().AddSP(this, &SAvaTemplatePageList::OnTemplatePageListChanged);
+
+	Refresh();
+}
+
+SAvaTemplatePageList::~SAvaTemplatePageList()
+{
+	if (UAvalanchePlaylist* const Playlist = GetValidPlaylist())
+	{
+		Playlist->GetOnTemplatePageListChanged().RemoveAll(this);
+	}
+}
+
+void SAvaTemplatePageList::Refresh()
+{
+	if (const TSharedPtr<FAvaPlaylistEditor> PlaylistEditor = PlaylistEditorWeak.Pin())
+	{
+		UAvalanchePlaylist* const Playlist = PlaylistEditor->GetPlaylist();
+		if (!Playlist)
+		{
+			return;
+		}
+
+		const FAvalanchePageCollection& PageCollection = Playlist->GetTemplatePages();
+		PageViews.Reset(PageCollection.Pages.Num());
+
+		for (const FAvalanchePage& Page : PageCollection.Pages)
+		{
+			if (PlaylistEditor->IsTemplatePageVisible(Page))
+			{
+				PageViews.Emplace(MakeShared<FAvaTemplatePageViewImpl>(Page.GetPageId(), Playlist, SharedThis(this)));
+			}
+		}
+
+		PageListView->RequestListRefresh();
+	}
+}
+
+void SAvaTemplatePageList::CreateColumns()
+{
+	HeaderRow = SNew(SHeaderRow)
+		.Visibility(EVisibility::Visible)
+		.CanSelectGeneratedColumn(true);
+
+	Columns.Empty();
+	HeaderRow->ClearColumns();
+
+	//TODO: Extensibility?
+	TArray<TSharedPtr<IAvaPageViewColumn>> FoundColumns;
+	FoundColumns.Add(MakeShared<FAvaPageThumbnailColumn>());
+	FoundColumns.Add(MakeShared<FAvaPageIdColumn>());
+	FoundColumns.Add(MakeShared<FAvaPageNameColumn>());
+	FoundColumns.Add(MakeShared<FAvaPageAssetSelectorColumn>());
+	FoundColumns.Add(MakeShared<FAvaPageTransitionLayerColumn>());
+	FoundColumns.Add(MakeShared<FAvaPageTemplateStatusColumn>());
+
+	for (const TSharedPtr<IAvaPageViewColumn>& Column : FoundColumns)
+	{
+		const FName ColumnId = Column->GetColumnId();
+		Columns.Add(ColumnId, Column);
+		HeaderRow->AddColumn(Column->ConstructHeaderRowColumn());
+		HeaderRow->SetShowGeneratedColumn(ColumnId, false);
+	}
+}
+
+TSharedPtr<SWidget> SAvaTemplatePageList::OnContextMenuOpening()
+{
+	if (TSharedPtr<FAvaPlaylistEditor> PlaylistEditor = PlaylistEditorWeak.Pin())
+	{
+		return GetPageListContextMenu();
+	}
+
+	return SNullWidget::NullWidget;
+}
+
+void SAvaTemplatePageList::BindCommands()
+{
+	SAvaPageList::BindCommands();
+
+	//Playlist Commands
+	{
+		const FAvaPlaylistCommands& PlaylistCommands = FAvaPlaylistCommands::Get();
+
+		CommandList->MapAction(PlaylistCommands.AddTemplate,
+			FExecuteAction::CreateSP(this, &SAvaTemplatePageList::AddTemplate),
+			FCanExecuteAction::CreateSP(this, &SAvaPageList::CanAddTemplate));
+
+		CommandList->MapAction(PlaylistCommands.CreatePageInstanceFromTemplate,
+			FExecuteAction::CreateSP(this, &SAvaTemplatePageList::CreateInstance),
+			FCanExecuteAction::CreateSP(this, &SAvaPageList::CanCreateInstance));
+
+		CommandList->MapAction(PlaylistCommands.CreateComboTemplate,
+			FExecuteAction::CreateSP(this, &SAvaTemplatePageList::CreateComboTemplate),
+			FCanExecuteAction::CreateSP(this, &SAvaTemplatePageList::CanCreateComboTemplate));
+
+		CommandList->MapAction(PlaylistCommands.RemovePage,
+			FExecuteAction::CreateSP(this, &SAvaPageList::RemoveSelectedPages),
+			FCanExecuteAction::CreateSP(this, &SAvaPageList::CanRemoveSelectedPages));
+
+		CommandList->MapAction(PlaylistCommands.RenumberPage,
+			FExecuteAction::CreateSP(this, &SAvaPageList::RenumberSelectedPage),
+			FCanExecuteAction::CreateSP(this, &SAvaPageList::CanRenumberSelectedPage));
+
+		CommandList->MapAction(PlaylistCommands.ReimportPage,
+			FExecuteAction::CreateSP(this, &SAvaPageList::ReimportSelectedPage),
+			FCanExecuteAction::CreateSP(this, &SAvaPageList::CanReimportSelectedPage));
+
+		CommandList->MapAction(PlaylistCommands.EditPageSource,
+			FExecuteAction::CreateSP(this, &SAvaPageList::EditSelectedPageSource),
+			FCanExecuteAction::CreateSP(this, &SAvaPageList::CanEditSelectedPageSource));
+
+		CommandList->MapAction(PlaylistCommands.PreviewFrame,
+			FExecuteAction::CreateSP(this, &SAvaPageList::PreviewPlaySelectedPage, true),
+			FCanExecuteAction::CreateSP(this, &SAvaPageList::CanPreviewPlaySelectedPage));
+
+		CommandList->MapAction(PlaylistCommands.PreviewPlay,
+			FExecuteAction::CreateSP(this, &SAvaPageList::PreviewPlaySelectedPage, false),
+			FCanExecuteAction::CreateSP(this, &SAvaPageList::CanPreviewPlaySelectedPage));
+
+		CommandList->MapAction(PlaylistCommands.PreviewStop,
+			FExecuteAction::CreateSP(this, &SAvaPageList::PreviewStopSelectedPage, false),
+			FCanExecuteAction::CreateSP(this, &SAvaPageList::CanPreviewStopSelectedPage, false));
+
+		CommandList->MapAction(PlaylistCommands.PreviewForceStop,
+			FExecuteAction::CreateSP(this, &SAvaPageList::PreviewStopSelectedPage, true),
+			FCanExecuteAction::CreateSP(this, &SAvaPageList::CanPreviewStopSelectedPage, true));
+
+		CommandList->MapAction(PlaylistCommands.PreviewContinue,
+			FExecuteAction::CreateSP(this, &SAvaPageList::PreviewContinueSelectedPage),
+			FCanExecuteAction::CreateSP(this, &SAvaPageList::CanPreviewContinueSelectedPage));
+
+		CommandList->MapAction(PlaylistCommands.PreviewPlayNext,
+			FExecuteAction::CreateSP(this, &SAvaPageList::PreviewPlayNextPage),
+			FCanExecuteAction::CreateSP(this, &SAvaPageList::CanPreviewPlayNextPage));
+
+		CommandList->MapAction(PlaylistCommands.TakeToProgram,
+			FExecuteAction::CreateSP(this, &SAvaPageList::TakeToProgram),
+			FCanExecuteAction::CreateSP(this, &SAvaPageList::CanTakeToProgram));
+	}
+}
+
+bool SAvaTemplatePageList::HandleDropAvalancheAssets(const TArray<FSoftObjectPath>& InAvaAssets, EItemDropZone InDropZone, const FAvaPageViewPtr& InItem)
+{
+	UAvalanchePlaylist* Playlist = GetValidPlaylist();
+	if (!Playlist)
+	{
+		return false;
+	}
+
+	FAvaPageInsertPosition InsertAt = MakeInsertPosition(InDropZone, InItem);
+	TArray<FSoftObjectPath> NewAvaAssets = InAvaAssets;
+	TArray<int32> NewTemplateIds;
+
+	// If we are adding above, they should be reversed, so the last is added first
+	// and the next to last added above that, etc.
+	if (InsertAt.IsAddAbove())
+	{
+		Algo::Reverse(NewAvaAssets);
+	}
+
+	bool bHasValidAssets = false;
+	Playlist->Modify();
+
+	for (const FSoftObjectPath& AvaAsset : NewAvaAssets)
+	{
+		if (AvaAsset.IsNull())
+		{
+			continue;
+		}
+
+		int32 NewTemplateId = Playlist->AddTemplate(FAvaPageIdGeneratorParams::FromInsertPosition(InsertAt));
+		
+		if (NewTemplateId == FAvalanchePage::InvalidPageId)
+		{
+			continue;
+		}
+
+		Playlist->GetPage(NewTemplateId).UpdateAvalancheAsset(AvaAsset);
+		InsertAt.ConditionalUpdateAdjacentId(NewTemplateId);
+		NewTemplateIds.Add(NewTemplateId);
+
+		bHasValidAssets = true;
+	}
+
+	if (bHasValidAssets)
+	{
+		Refresh();
+		DeselectPages();
+		SelectPages(NewTemplateIds);
+	}
+
+	return bHasValidAssets;
+}
+
+bool SAvaTemplatePageList::HandleDropPlaylists(const TArray<FSoftObjectPath>& InPlaylistPaths, EItemDropZone InDropZone, const FAvaPageViewPtr& InItem)
+{
+	// Not supported directly.
+	// The templates will import automatically when the playlist pages are imported.
+	return false;	
+}
+
+bool SAvaTemplatePageList::HandleDropPageIds(const FAvaPageListReference& InPageListReference, const TArray<int32>& InPageIds, EItemDropZone InDropZone, const FAvaPageViewPtr& InItem)
+{
+	// Can only drop templates onto the templates list.
+	if (InPageListReference.Type != EAvaPageListType::Template)
+	{
+		return false;
+	}
+
+	// Nothing to do.
+	if (InPageIds.IsEmpty())
+	{
+		return true;
+	}
+
+	UAvalanchePlaylist* Playlist = GetValidPlaylist();
+	if (!Playlist || !Playlist->CanChangePageOrder())
+	{
+		return false;
+	}
+	
+	const FAvalanchePageCollection& TemplatePageCollection = Playlist->GetTemplatePages();
+	int32 DroppedOnPageIndex = INDEX_NONE;
+
+	if (InItem.IsValid() && TemplatePageCollection.PageIndices.Contains(InItem->GetPageId()))
+	{
+		DroppedOnPageIndex = TemplatePageCollection.PageIndices[InItem->GetPageId()];
+
+		// Nothing to do.
+		if (InPageIds.Num() == 1 && DroppedOnPageIndex == InPageIds[0])
+		{
+			return true;
+		}
+	}
+
+	TArray<int32> MovedPageIdIndices;
+	TArray<int32> NewPageOrder;
+	TArray<int32> NewSelectedIds;
+
+	MovedPageIdIndices.Reserve(InPageIds.Num());
+	NewPageOrder.Reserve(PageViews.Num());
+	NewSelectedIds.Reserve(InPageIds.Num());
+
+	// This has the byproduct of removing invalid ids.
+	for (int32 PageId : InPageIds)
+	{
+		if (const int32* PageIndexPtr = TemplatePageCollection.PageIndices.Find(PageId))
+		{
+			MovedPageIdIndices.Add(*PageIndexPtr);
+			NewSelectedIds.Add(PageId);
+		}
+	}
+
+	// Nothing to do.
+	if (MovedPageIdIndices.IsEmpty())
+	{
+		return true;
+	}
+
+	for (int32 PageViewIndex = 0; PageViewIndex < PageViews.Num(); ++PageViewIndex)
+	{
+		const bool bHasMovedThisPage = MovedPageIdIndices.Contains(PageViewIndex);
+
+		if (PageViewIndex == DroppedOnPageIndex)
+		{
+			const bool bAddBefore = InDropZone == EItemDropZone::AboveItem;
+
+			if (bAddBefore)
+			{
+				NewPageOrder.Append(MovedPageIdIndices);
+			}
+
+			// If we moved the page we dropped onto, it will already be in the list
+			if (!bHasMovedThisPage)
+			{
+				NewPageOrder.Add(PageViewIndex);
+			}
+
+			if (!bAddBefore)
+			{
+				NewPageOrder.Append(MovedPageIdIndices);
+			}
+
+			continue;
+		}
+
+		if (bHasMovedThisPage)
+		{
+			continue;
+		}
+
+		NewPageOrder.Add(PageViewIndex);
+	}
+
+	Playlist->ChangePageOrder(PageListReference, NewPageOrder);
+	SelectPages(NewSelectedIds, true);
+
+	return true;
+}
+
+bool SAvaTemplatePageList::HandleDropExternalFiles(const TArray<FString>& InFiles, EItemDropZone InDropZone, const FAvaPageViewPtr& InItem)
+{
+	// Not supported directly.
+	// The templates will import automatically when the playlist pages are imported.
+	return false;
+}
+
+void SAvaTemplatePageList::AddTemplate()
+{
+	if (UAvalanchePlaylist* Playlist = GetValidPlaylist())
+	{
+		int32 LastIndex = FAvalanchePage::InvalidPageId;
+
+		for (const int32 SelectedIndex : SelectedPageIds)
+		{
+			LastIndex = FMath::Max(LastIndex, SelectedIndex);
+		}
+
+		const int32 NewTemplateId = Playlist->AddTemplate(FAvaPageIdGeneratorParams(LastIndex));
+
+		if (NewTemplateId != FAvalanchePage::InvalidPageId)
+		{
+			DeselectPages();
+			SelectPage(NewTemplateId);
+		}
+	}
+}
+
+void SAvaTemplatePageList::CreateInstance()
+{
+	if (SelectedPageIds.IsEmpty())
+	{
+		return;
+	}
+
+	if (UAvalanchePlaylist* Playlist = GetValidPlaylist())
+	{
+		const TArray<int32> NewPageIds = Playlist->AddPagesFromTemplates(SelectedPageIds);
+
+		if (!NewPageIds.IsEmpty())
+		{
+			DeselectPages();
+			SelectPages(NewPageIds);
+		}
+	}
+}
+
+void SAvaTemplatePageList::CreateComboTemplate()
+{
+	UAvalanchePlaylist* Playlist = GetValidPlaylist();
+	
+	if (SelectedPageIds.IsEmpty() || !Playlist)
+	{
+		return;
+	}
+
+	TSet<FAvaTagId> LayerIds;
+	TArray<int32> TemplateIds;
+	FAvalancheRemoteControlValues MergedValues;
+	
+	for (const int32 SelectedPageId : SelectedPageIds)
+	{
+		const FAvalanchePage& Page = Playlist->GetPage(SelectedPageId);
+		if (!Page.IsValidPage())
+		{
+			UE_LOG(LogAvaPlaylist, Error, TEXT("CreateComboTemplate: Template %d is not valid."), SelectedPageId);
+			continue;
+		}
+		if (!Page.IsTemplate())
+		{
+			UE_LOG(LogAvaPlaylist, Error, TEXT("CreateComboTemplate: Page %d is not a template."), SelectedPageId);
+			continue;
+		}
+		if (Page.IsComboTemplate())
+		{
+			UE_LOG(LogAvaPlaylist, Error, TEXT("CreateComboTemplate: Template %d is already a combo template."), SelectedPageId);
+			continue;
+		}
+		if (!Page.HasTransitionLogic(Playlist))
+		{
+			UE_LOG(LogAvaPlaylist, Error, TEXT("CreateComboTemplate: Template %d doesn't have transition logic."), SelectedPageId);
+			continue;
+		}
+		if (!Page.GetTransitionLayer(Playlist).IsValid())
+		{
+			UE_LOG(LogAvaPlaylist, Error, TEXT("CreateComboTemplate: Template %d doesn't have a valid transition logic layer."), SelectedPageId);
+			continue;
+		}
+		if (LayerIds.Contains(Page.GetTransitionLayer(Playlist).TagId))
+		{
+			UE_LOG(LogAvaPlaylist, Error, TEXT("CreateComboTemplate: Template %d's layer %s is already in the selection."), SelectedPageId, *Page.GetTransitionLayer(Playlist).ToString());
+			continue;
+		}
+
+		// Make sure the RC values can merge correctly. If not, original template needs fixing.
+		if (MergedValues.HasIdCollisions(Page.GetRemoteControlValues()))
+		{
+			UE_LOG(LogAvaPlaylist, Error, TEXT("CreateComboTemplate: Template %d's RemoteControl values have Id collisions with other templates in the selection."), SelectedPageId);
+			continue;
+		}
+		
+		MergedValues.Merge(Page.GetRemoteControlValues());
+		LayerIds.Add(Page.GetTransitionLayer(Playlist).TagId);
+		TemplateIds.Add(Page.GetPageId());
+	}
+
+	if (TemplateIds.Num() > 1)
+	{
+		Playlist->AddComboTemplate(TemplateIds);
+	}
+	else
+	{
+		UE_LOG(LogAvaPlaylist, Error, TEXT("CreateComboTemplate: Need at least 2 suitable templates to create a combo template."));
+	}
+}
+
+bool SAvaTemplatePageList::CanCreateComboTemplate()
+{
+	UAvalanchePlaylist* Playlist = GetValidPlaylist();
+	
+	if (SelectedPageIds.Num() < 2 || !Playlist)
+	{
+		return false;
+	}
+
+	TSet<FAvaTagId> LayerIds;
+
+	for (const int32 SelectedPageId : SelectedPageIds)
+	{
+		const FAvalanchePage& Page = Playlist->GetPage(SelectedPageId);
+		if (Page.IsValidPage()
+			&& Page.IsTemplate()
+			&& !Page.IsComboTemplate()
+			&& Page.HasTransitionLogic(Playlist)
+			&& Page.GetTransitionLayer(Playlist).IsValid()
+			&& !LayerIds.Contains(Page.GetTransitionLayer(Playlist).TagId))
+		{
+			LayerIds.Add(Page.GetTransitionLayer(Playlist).TagId);
+		}
+	}
+	// Need more than one template to create a combo template.
+	return LayerIds.Num() > 1;
+}
+
+TArray<int32> SAvaTemplatePageList::AddPastedPages(const TArray<FAvalanchePage>& InPages)
+{
+	if (UAvalanchePlaylist* Playlist = GetValidPlaylist())
+	{
+		using namespace UE::AvaPlaylistEditor::Utils;
+		FImportTemplateMap ImportedTemplateIds;
+		return ImportTemplatePages(Playlist, InPages, ImportedTemplateIds);
+	}
+	return {};
+}
+
+void SAvaTemplatePageList::OnTemplatePageListChanged(const FAvaPageListChangeParams& InParams)
+{
+	if (const TSharedPtr<FAvaPlaylistEditor> PlaylistEditor = PlaylistEditorWeak.Pin())
+	{
+		PlaylistEditor->RefreshTemplateVisibility();
+	}
+	Refresh();
+}
+
+#undef LOCTEXT_NAMESPACE
diff --git a/Engine/Plugins/Experimental/Avalanche/Source/AvalancheMediaEditor/Private/Playlist/Pages/Slate/SAvaTemplatePageList.h b/Engine/Plugins/Experimental/Avalanche/Source/AvalancheMediaEditor/Private/Playlist/Pages/Slate/SAvaTemplatePageList.h
new file mode 100644
--- /dev/null
+++ b/Engine/Plugins/Experimental/Avalanche/Source/AvalancheMediaEditor/Private/Playlist/Pages/Slate/SAvaTemplatePageList.h
@@ -0,0 +1,52 @@
+// Copyright Epic Games, Inc. All Rights Reserved.
+
+#pragma once
+
+#include "SAvaPageList.h"
+#include "Playlist/AvalanchePlaylist.h"
+#include "Playlist/AvaPlaylistDefines.h"
+
+class FAvaPlaylistEditor;
+class IAvaPageViewColumn;
+class ITableRow;
+class SHeaderRow;
+template <typename ItemType> class SListView;
+class SSearchBox;
+class STableViewBase;
+struct FAvalanchePage;
+
+class SAvaTemplatePageList : public SAvaPageList
+{
+	SLATE_DECLARE_WIDGET(SAvaTemplatePageList, SAvaPageList)
+
+public:
+	SLATE_BEGIN_ARGS(SAvaTemplatePageList) {}
+	SLATE_END_ARGS()
+
+	/** Constructs this widget with InArgs */
+	void Construct(const FArguments& InArgs, TSharedPtr<FAvaPlaylistEditor> InPlaylistEditor);
+	virtual ~SAvaTemplatePageList() override;
+
+	//~ Begin SAvaPageList
+	virtual void Refresh() override;
+	virtual void CreateColumns() override;
+	virtual TSharedPtr<SWidget> OnContextMenuOpening() override;
+	virtual void BindCommands() override;
+	virtual bool HandleDropAvalancheAssets(const TArray<FSoftObjectPath>& InAvaAssets, EItemDropZone InDropZone, const FAvaPageViewPtr& InItem) override;
+	virtual bool HandleDropPlaylists(const TArray<FSoftObjectPath>& InPlaylistPaths, EItemDropZone InDropZone, const FAvaPageViewPtr& InItem) override;
+	virtual bool HandleDropPageIds(const FAvaPageListReference& InPageListReference, const TArray<int32>& InPageIds, EItemDropZone InDropZone, const FAvaPageViewPtr& InItem) override;
+	virtual bool HandleDropExternalFiles(const TArray<FString>& InFiles, EItemDropZone InDropZone, const FAvaPageViewPtr& InItem) override;
+	//~ End SAvaPageList
+
+	void AddTemplate();
+
+	void CreateInstance();
+	void CreateComboTemplate();
+	bool CanCreateComboTemplate();
+
+protected:
+	virtual TArray<int32> AddPastedPages(const TArray<FAvalanchePage>& InPages) override;
+
+private:
+	void OnTemplatePageListChanged(const FAvaPageListChangeParams& InParams);
+};
diff --git a/Engine/Plugins/Experimental/Avalanche/Source/AvalancheMediaEditor/Private/Playlist/Preview/SAvaPagePreview.cpp b/Engine/Plugins/Experimental/Avalanche/Source/AvalancheMediaEditor/Private/Playlist/Preview/SAvaPagePreview.cpp
new file mode 100644
--- /dev/null
+++ b/Engine/Plugins/Experimental/Avalanche/Source/AvalancheMediaEditor/Private/Playlist/Preview/SAvaPagePreview.cpp
@@ -0,0 +1,323 @@
+// Copyright Epic Games, Inc. All Rights Reserved.
+
+#include "SAvaPagePreview.h"
+#include "AvaMediaEditorStyle.h"
+#include "AvalancheMediaSettings.h"
+#include "Brushes/SlateImageBrush.h"
+#include "Engine/TextureRenderTarget2D.h"
+#include "Framework/Commands/UICommandList.h"
+#include "Framework/MultiBox/MultiBoxBuilder.h"
+#include "Framework/Notifications/NotificationManager.h"
+#include "IAvaMediaEditorModule.h"
+#include "ISettingsModule.h"
+#include "OutputDevices/Slate/SAvaCaptureImage.h"
+#include "Playlist/AvaPlaylistCommands.h"
+#include "Playlist/AvaPlaylistEditor.h"
+#include "Playlist/AvaRundownEditorSettings.h"
+#include "Playlist/AvalanchePlaylist.h"
+#include "Playlist/Preview/SAvaPreviewChannelSelector.h"
+#include "Widgets/Layout/SScaleBox.h"
+#include "Widgets/Layout/SSpacer.h"
+#include "Widgets/Notifications/SNotificationList.h"
+
+#define LOCTEXT_NAMESPACE "SAvaPagePreview"
+
+/**
+ * Wrapper to encapsulate a FSlateImageBrush that can be updated with a render target.
+ * If the render target becomes null, the brush is updated to be a solid color.
+ */
+struct SAvaPagePreview::FPreviewBrush
+{
+	FSlateImageBrush Brush;
+	FVector2D RenderTargetSize;
+
+	FPreviewBrush(const FVector2D& InRenderTargetSize)
+		: Brush(NAME_None, InRenderTargetSize, FLinearColor::Black)
+		, RenderTargetSize(InRenderTargetSize)
+	{
+
+	}
+
+	void Update(UTextureRenderTarget2D* InRenderTarget)
+	{
+		if (InRenderTarget)
+		{
+			RenderTargetSize = FVector2D(InRenderTarget->SizeX, InRenderTarget->SizeY);
+
+			//If Brush is invalid, or the Brush's Texture Target doesn't match the new Render Target, reset the Brush.
+			if (Brush.GetResourceObject() != InRenderTarget)
+			{
+				Brush = FSlateImageBrush(InRenderTarget, RenderTargetSize);
+			}
+			//If Sizes mismatch, just resizes rather than recreating the Brush with same underlying Resource
+			else if (Brush.GetImageSize() != RenderTargetSize)
+			{
+				Brush.SetImageSize(RenderTargetSize);
+			}
+		}
+		else if (Brush.GetResourceObject() != nullptr || Brush.GetImageSize() != RenderTargetSize)
+		{
+			// Preserve the size of the preview brush so the checkerboard remains the same.
+			Brush = FSlateImageBrush(NAME_None, RenderTargetSize, FLinearColor::Black);
+		}
+	}
+};
+
+void SAvaPagePreview::Construct(const FArguments& InArgs, const TSharedPtr<FAvaPlaylistEditor>& InPlaylistEditor)
+{
+	PlaylistEditorWeak = InPlaylistEditor;
+
+	const FIntPoint DefaultRenderTargetSize = UAvalancheMediaSettings::Get().PreviewDefaultResolution;
+	PreviewBrush = MakeUnique<FPreviewBrush>(FVector2d(DefaultRenderTargetSize.X, DefaultRenderTargetSize.Y));
+
+	ChildSlot
+	[
+		SNew(SVerticalBox)
+		+ SVerticalBox::Slot()
+		.AutoHeight()
+		[
+			CreatePagePreviewToolBar(InPlaylistEditor->GetToolkitCommands())
+		]
+		+ SVerticalBox::Slot()
+		.FillHeight(1.f)
+		[
+			SNew(SScaleBox)
+			.Stretch(EStretch::ScaleToFit)
+			.StretchDirection(EStretchDirection::Both)
+			[
+				SNew( SOverlay )
+				+SOverlay::Slot()
+				[
+					SNew( SImage )
+					.Image( FAvaMediaEditorStyle::Get().GetBrush( "AvalancheMediaEditor.Checkerboard" ) )
+					.Visibility(this, &SAvaPagePreview::GetCheckerboardVisibility)
+				]
+				+SOverlay::Slot()
+				[
+					SNew(SAvaCaptureImage)
+					.ImageArgs(SImage::FArguments()
+						.Image(this, &SAvaPagePreview::GetPreviewBrush))
+					.ShouldInvertAlpha(true)
+					.EnableGammaCorrection(false)
+					.EnableBlending(this, &SAvaPagePreview::IsBlendingEnabled)
+				]
+			]
+		]
+	];
+}
+
+SAvaPagePreview::SAvaPagePreview() = default;
+SAvaPagePreview::~SAvaPagePreview() = default;
+
+TSharedRef<SWidget> SAvaPagePreview::CreatePagePreviewToolBar(const TSharedRef<FUICommandList>& InCommandList)
+{
+	FSlimHorizontalToolBarBuilder ToolBarBuilder(InCommandList, FMultiBoxCustomization::None);
+	ToolBarBuilder.SetStyle(&FAvaMediaEditorStyle::Get(), "AvalancheMediaEditor.ToolBar");
+
+	const FAvaPlaylistCommands& PlaylistCommands = FAvaPlaylistCommands::Get();
+
+	ToolBarBuilder.BeginSection(TEXT("ShowControl"));
+	{
+		ToolBarBuilder.AddToolBarButton(PlaylistCommands.PreviewFrame);
+
+		ToolBarBuilder.BeginStyleOverride("AvalancheMediaEditor.ToolBarGreenButtonOverride");
+		{
+			ToolBarBuilder.AddToolBarButton(PlaylistCommands.PreviewPlay);
+		}
+		ToolBarBuilder.EndStyleOverride();
+
+		ToolBarBuilder.AddToolBarButton(PlaylistCommands.PreviewContinue);
+		ToolBarBuilder.AddToolBarButton(PlaylistCommands.PreviewStop);
+		ToolBarBuilder.AddToolBarButton(PlaylistCommands.PreviewPlayNext);
+
+		ToolBarBuilder.BeginStyleOverride("AvalancheMediaEditor.ToolBarRedButtonOverride");
+		{
+			ToolBarBuilder.AddToolBarButton(PlaylistCommands.TakeToProgram);
+		}
+		ToolBarBuilder.EndStyleOverride();
+
+		if (TSharedPtr<FAvaPlaylistEditor> PlaylistEditor = PlaylistEditorWeak.Pin())
+		{
+			ToolBarBuilder.AddWidget(PlaylistEditor->MakeReadPageWidget());	
+		}
+	}
+	ToolBarBuilder.EndSection();
+
+	ToolBarBuilder.BeginSection("Settings");
+	{
+		ToolBarBuilder.AddWidget(SNew(SSpacer), NAME_None, false, HAlign_Right);
+
+		ToolBarBuilder.AddToolBarButton(
+				FUIAction(
+					FExecuteAction::CreateSP(this, &SAvaPagePreview::HandleCheckerboardActionExecute),
+					FCanExecuteAction::CreateLambda([]	{ return true;}))
+				, NAME_None
+				, FText()
+				, LOCTEXT("ToggleAlpha_ToolTip", "Toggle alpha preview (checker board).")
+				, FSlateIcon(FAppStyle::GetAppStyleSetName(), "Checkerboard")
+			);
+
+		ToolBarBuilder.BeginStyleOverride("AvalancheMediaEditor.CalloutToolbar");
+		{
+			ToolBarBuilder.AddComboButton(
+				FUIAction(),
+				FOnGetContent::CreateSP(this, &SAvaPagePreview::OnGenerateSettingsMenu),
+				LOCTEXT("SettingsMenu", "Preview Settings"),
+				FText::GetEmpty(),
+				FSlateIcon(FAppStyle::GetAppStyleSetName(), "Icons.Settings")
+			);
+		}
+		ToolBarBuilder.EndStyleOverride();
+	}
+	ToolBarBuilder.EndSection();	
+
+	return ToolBarBuilder.MakeWidget();
+}
+
+TSharedRef<SWidget> SAvaPagePreview::OnGenerateSettingsMenu()
+{
+	TSharedPtr<FAvaPlaylistEditor> PlaylistEditor = PlaylistEditorWeak.Pin();
+	TSharedPtr<FUICommandList> ToolkitCommands = PlaylistEditor.IsValid() ? PlaylistEditor->GetToolkitCommands() : TSharedPtr<FUICommandList>();
+	
+	FMenuBuilder MenuBuilder(true, ToolkitCommands);
+	MenuBuilder.BeginSection("ResolutionSection", LOCTEXT("ResolutionSectionHeader", "Resolution Options"));
+	{
+		// Add a bunch of standard resolutions.
+		AddResolutionMenuEntry(MenuBuilder, LOCTEXT("Res_3840x2160", "3840 x 2160"), FIntPoint(3840,2160));
+		AddResolutionMenuEntry(MenuBuilder, LOCTEXT("Res_1920x1080", "1920 x 1080"), FIntPoint(1920,1080));
+		AddResolutionMenuEntry(MenuBuilder, LOCTEXT("Res_960x540", "960 x 540"), FIntPoint(960,540));
+		AddResolutionMenuEntry(MenuBuilder, LOCTEXT("Res_480x270", "480 x 270"), FIntPoint(480,270));
+	}
+	MenuBuilder.EndSection();
+
+	MenuBuilder.AddMenuSeparator();
+
+	MenuBuilder.AddWidget(SNew(SAvaPreviewChannelSelector), LOCTEXT("PreviewChannelSelectorLabel", "Channel"));
+	
+	MenuBuilder.AddMenuSeparator();
+
+	MenuBuilder.AddMenuEntry(
+		LOCTEXT("Settings", "Settings"),
+		LOCTEXT("Settings_Tooltip", "Opens the Avalanche Media Settings."),
+		FSlateIcon(),
+		FUIAction(
+			FExecuteAction::CreateSP(this, &SAvaPagePreview::HandleSettingsActionExecute)
+		),
+		NAME_None,
+		EUserInterfaceActionType::Button);
+	
+	return MenuBuilder.MakeWidget();
+}
+
+UTextureRenderTarget2D* SAvaPagePreview::GetPreviewRenderTarget() const
+{
+	const TSharedPtr<FAvaPlaylistEditor> Editor = PlaylistEditorWeak.Pin();
+	if (Editor.IsValid() && Editor->IsPlaylistValid())
+	{
+		return Editor->GetPlaylist()->GetPreviewRenderTarget();
+	}
+	return nullptr;
+}
+
+const FSlateBrush* SAvaPagePreview::GetPreviewBrush() const
+{
+	// Update the brush here to ensure the cached image attribute is updated with the correct render target.
+	PreviewBrush->Update(GetPreviewRenderTarget());
+	return &PreviewBrush->Brush;
+}
+
+bool SAvaPagePreview::IsBlendingEnabled() const
+{
+	const UAvaRundownEditorSettings* Settings = UAvaRundownEditorSettings::Get();
+	return Settings ? Settings->bPreviewCheckerBoard : false;
+}
+
+EVisibility SAvaPagePreview::GetCheckerboardVisibility() const
+{	
+	if (const UAvaRundownEditorSettings* Settings = UAvaRundownEditorSettings::Get())
+	{
+		return Settings->bPreviewCheckerBoard ? EVisibility::Visible : EVisibility::Hidden;
+	}
+	return EVisibility::Hidden;
+}
+
+void SAvaPagePreview::SetPreviewResolution(FIntPoint InResolution) const
+{
+	// Update settings
+	UAvalancheMediaSettings& AvalancheMediaSettings = UAvalancheMediaSettings::GetMutable();
+	AvalancheMediaSettings.PreviewDefaultResolution = InResolution;
+	AvalancheMediaSettings.SaveConfig();
+
+	// Resize render target
+	UTextureRenderTarget2D* RenderTarget = GetPreviewRenderTarget();
+	if (RenderTarget)
+	{
+		RenderTarget->ResizeTarget(InResolution.X, InResolution.Y);
+	}
+	else
+	{
+		// The brush needs a resize.
+		PreviewBrush->RenderTargetSize = FVector2D(InResolution.X, InResolution.Y);
+	}
+	
+	// Refresh the brush.
+	PreviewBrush->Update(RenderTarget);
+}
+
+bool SAvaPagePreview::IsPreviewResolution(FIntPoint InResolution) const
+{
+	return UAvalancheMediaSettings::Get().PreviewDefaultResolution == InResolution;
+}
+
+void SAvaPagePreview::AddResolutionMenuEntry(FMenuBuilder& InOutMenuBuilder, const FText& Label, const FIntPoint& InResolution)
+{
+	InOutMenuBuilder.AddMenuEntry(
+			Label,
+			FText(),
+			FSlateIcon(),
+			FUIAction(
+				FExecuteAction::CreateSP(this, &SAvaPagePreview::SetPreviewResolution, InResolution),
+				FCanExecuteAction::CreateLambda([]() { return true; }),
+				FIsActionChecked::CreateSP(this, &SAvaPagePreview::IsPreviewResolution, InResolution)
+			),
+			NAME_None, EUserInterfaceActionType::RadioButton);
+}
+
+void SAvaPagePreview::HandleCheckerboardActionExecute() const
+{
+	UAvaRundownEditorSettings* RundownEditorSettings = UAvaRundownEditorSettings::GetMutable();
+	if (!RundownEditorSettings)
+	{
+		UE_LOG(LogAvaMediaEditor, Error, TEXT("Failed to retrieve Rundown Editor Settings."));
+		return;
+	}
+	
+	bool& bShowCheckerBoard = RundownEditorSettings->bPreviewCheckerBoard;
+	bShowCheckerBoard = !bShowCheckerBoard;
+	RundownEditorSettings->SaveConfig();	// We want this to be persistent.
+	if (bShowCheckerBoard)
+	{
+		const IConsoleVariable* PropagateAlphaCVar = IConsoleManager::Get().FindConsoleVariable(TEXT("r.PostProcessing.PropagateAlpha"));
+		if (PropagateAlphaCVar && PropagateAlphaCVar->GetInt() != 2)
+		{
+			const FText NotificationText = LOCTEXT("AlphaSupport",
+				"An output requested Alpha Support but the required project setting is not enabled!\n"
+				"Go to Project Settings > Rendering > PostProcessing > 'Enable Alpha Channel Support in Post Processing' and set it to 'Allow through tonemapper'.");
+
+			FNotificationInfo Info(NotificationText);
+			Info.ExpireDuration = 5.0f;	// The message is long, need more time to read it.
+			FSlateNotificationManager::Get().AddNotification(Info);
+
+			// Also output to logs.
+			UE_LOG(LogAvaMediaEditor, Warning, TEXT("%s"), *NotificationText.ToString());
+		}
+	}
+}
+
+void SAvaPagePreview::HandleSettingsActionExecute() const
+{
+	// Bring up editor of UAvalancheMediaSettings.
+	FModuleManager::LoadModuleChecked<ISettingsModule>("Settings").ShowViewer("Project", "Motion Design", "Playback & Broadcast");
+}
+
+#undef LOCTEXT_NAMESPACE
diff --git a/Engine/Plugins/Experimental/Avalanche/Source/AvalancheMediaEditor/Private/Playlist/Preview/SAvaPagePreview.h b/Engine/Plugins/Experimental/Avalanche/Source/AvalancheMediaEditor/Private/Playlist/Preview/SAvaPagePreview.h
new file mode 100644
--- /dev/null
+++ b/Engine/Plugins/Experimental/Avalanche/Source/AvalancheMediaEditor/Private/Playlist/Preview/SAvaPagePreview.h
@@ -0,0 +1,46 @@
+// Copyright Epic Games, Inc. All Rights Reserved.
+
+#pragma once
+
+#include "CoreMinimal.h"
+#include "Widgets/SCompoundWidget.h"
+
+class FAvaPlaylistEditor;
+class FMenuBuilder;
+class FUICommandList;
+class UTextureRenderTarget2D;
+
+class SAvaPagePreview : public SCompoundWidget
+{
+public:	
+	SLATE_BEGIN_ARGS(SAvaPagePreview) {}
+	SLATE_END_ARGS()
+
+	/** Constructs this widget with InArgs */
+	void Construct(const FArguments& InArgs, const TSharedPtr<FAvaPlaylistEditor>& InPlaylistEditor);
+	SAvaPagePreview();
+	virtual ~SAvaPagePreview() override;
+
+	TSharedRef<SWidget> CreatePagePreviewToolBar(const TSharedRef<FUICommandList>& InCommandList);
+	TSharedRef<SWidget> OnGenerateSettingsMenu();
+
+protected:
+	const FSlateBrush* GetPreviewBrush() const;
+	bool IsBlendingEnabled() const;
+	EVisibility GetCheckerboardVisibility() const;
+	void SetPreviewResolution(FIntPoint InResolution) const;
+	bool IsPreviewResolution(FIntPoint InResolution) const;
+	
+	UTextureRenderTarget2D* GetPreviewRenderTarget() const;
+	
+	void AddResolutionMenuEntry(FMenuBuilder& InOutMenuBuilder, const FText& Label, const FIntPoint& InResolution);
+	
+	void HandleCheckerboardActionExecute() const;
+	void HandleSettingsActionExecute() const;
+
+protected:
+	TWeakPtr<FAvaPlaylistEditor> PlaylistEditorWeak;
+
+	struct FPreviewBrush;
+	TUniquePtr<FPreviewBrush> PreviewBrush;
+};
\ No newline at end of file
diff --git a/Engine/Plugins/Experimental/Avalanche/Source/AvalancheMediaEditor/Private/Playlist/Preview/SAvaPreviewChannelSelector.cpp b/Engine/Plugins/Experimental/Avalanche/Source/AvalancheMediaEditor/Private/Playlist/Preview/SAvaPreviewChannelSelector.cpp
new file mode 100644
--- /dev/null
+++ b/Engine/Plugins/Experimental/Avalanche/Source/AvalancheMediaEditor/Private/Playlist/Preview/SAvaPreviewChannelSelector.cpp
@@ -0,0 +1,72 @@
+// Copyright Epic Games, Inc. All Rights Reserved.
+
+#include "SAvaPreviewChannelSelector.h"
+#include "AvalancheBroadcast.h"
+#include "AvalancheMediaSettings.h"
+
+void SAvaPreviewChannelSelector::Construct(const FArguments& InArgs)
+{
+	UpdateChannelNames();
+	
+	ChildSlot
+	[
+		SAssignNew(ChannelCombo, SComboBox<FName>)
+		.InitiallySelectedItem(GetPreviewChannelNameFromSettings())
+		.OptionsSource(&ChannelNames)
+		.OnGenerateWidget(this, &SAvaPreviewChannelSelector::GenerateChannelWidget)
+		.OnSelectionChanged(this, &SAvaPreviewChannelSelector::OnChannelSelectionChanged)
+		.OnComboBoxOpening(this, &SAvaPreviewChannelSelector::OnComboBoxOpening)
+		[
+			SNew(STextBlock)
+			.Text(this, &SAvaPreviewChannelSelector::GetCurrentChannelName)
+		]
+	];
+}
+
+TSharedRef<SWidget> SAvaPreviewChannelSelector::GenerateChannelWidget(FName InChannelName)
+{
+	return SNew(STextBlock)
+		.Text(FText::FromName(InChannelName));
+}
+
+void SAvaPreviewChannelSelector::OnChannelSelectionChanged(SComboBox<FName>::NullableOptionType InProposedSelection, ESelectInfo::Type InSelectInfo)
+{
+	UAvalancheMediaSettings& AvalancheMediaSettings = UAvalancheMediaSettings::GetMutable();
+	AvalancheMediaSettings.PreviewChannelName = !InProposedSelection.IsNone() ? InProposedSelection.ToString() : FString();
+	AvalancheMediaSettings.SaveConfig();
+}
+
+FText SAvaPreviewChannelSelector::GetCurrentChannelName() const
+{
+	return FText::FromName(GetPreviewChannelNameFromSettings());
+}
+
+FName SAvaPreviewChannelSelector::GetPreviewChannelNameFromSettings()
+{
+	const UAvalancheMediaSettings& Settings = UAvalancheMediaSettings::Get();	
+	return !Settings.PreviewChannelName.IsEmpty() ? FName(Settings.PreviewChannelName) : NAME_None; 
+}
+
+void SAvaPreviewChannelSelector::OnComboBoxOpening()
+{
+	UpdateChannelNames();
+	
+	check(ChannelCombo.IsValid());
+	ChannelCombo->SetSelectedItem(GetPreviewChannelNameFromSettings());
+}
+
+void SAvaPreviewChannelSelector::UpdateChannelNames()
+{
+	const UAvalancheBroadcast& AvaBroadcast = UAvalancheBroadcast::Get();
+	ChannelNames.Reset(AvaBroadcast.GetChannelNameCount() + 1);
+	ChannelNames.Add(NAME_None);	// Add none to allow user to deselect.
+	
+	for (int32 ChannelIndex = 0; ChannelIndex < AvaBroadcast.GetChannelNameCount(); ++ChannelIndex)
+	{
+		const FName ChannelName = AvaBroadcast.GetChannelName(ChannelIndex);
+		if (AvaBroadcast.GetChannelType(ChannelName) == EAvaBroadcastChannelType::Preview)
+		{
+			ChannelNames.Add(ChannelName);
+		}
+	}
+}
