[MotionDesign]

diff --git a/Engine/Plugins/Experimental/Avalanche/Source/AvalancheMediaEditor/Private/Playlist/Preview/SAvaPreviewChannelSelector.h b/Engine/Plugins/Experimental/Avalanche/Source/AvalancheMediaEditor/Private/Playlist/Preview/SAvaPreviewChannelSelector.h
new file mode 100644
--- /dev/null
+++ b/Engine/Plugins/Experimental/Avalanche/Source/AvalancheMediaEditor/Private/Playlist/Preview/SAvaPreviewChannelSelector.h
@@ -0,0 +1,33 @@
+// Copyright Epic Games, Inc. All Rights Reserved.
+
+#pragma once
+
+#include "Widgets/Input/SComboBox.h"
+#include "Widgets/SCompoundWidget.h"
+
+class SAvaPreviewChannelSelector : public SCompoundWidget
+{
+public:
+	SLATE_BEGIN_ARGS(SAvaPreviewChannelSelector){}
+	SLATE_END_ARGS()
+
+	/** Constructs this widget with InArgs */
+	void Construct(const FArguments& InArgs);
+
+	TSharedRef<SWidget> GenerateChannelWidget(FName InChannelName);
+	
+	void OnChannelSelectionChanged(SComboBox<FName>::NullableOptionType InProposedSelection, ESelectInfo::Type InSelectInfo);
+	
+	FText GetCurrentChannelName() const;
+
+	static FName GetPreviewChannelNameFromSettings();
+
+	void OnComboBoxOpening();
+
+protected:
+	void UpdateChannelNames();
+	
+	TSharedPtr<SComboBox<FName>> ChannelCombo;
+
+	TArray<FName> ChannelNames;
+};
diff --git a/Engine/Plugins/Experimental/Avalanche/Source/AvalancheMediaEditor/Private/Playlist/ShowControl/SAvaShowControl.cpp b/Engine/Plugins/Experimental/Avalanche/Source/AvalancheMediaEditor/Private/Playlist/ShowControl/SAvaShowControl.cpp
new file mode 100644
--- /dev/null
+++ b/Engine/Plugins/Experimental/Avalanche/Source/AvalancheMediaEditor/Private/Playlist/ShowControl/SAvaShowControl.cpp
@@ -0,0 +1,184 @@
+// Copyright Epic Games, Inc. All Rights Reserved.
+
+#include "SAvaShowControl.h"
+#include "Framework/MultiBox/MultiBoxBuilder.h"
+#include "Playback/AvalanchePlayback.h"
+#include "Playlist/AvaPlaylistCommands.h"
+#include "Playlist/AvaPlaylistEditor.h"
+#include "Playlist/Pages/Slate/SAvaInstancedPageList.h"
+#include "Playlist/Pages/Slate/SAvaPageList.h"
+#include "Playlist/TabFactories/AvaSubListDocumentTabFactory.h"
+
+#define LOCTEXT_NAMESPACE "SAvaShowControl"
+
+void SAvaShowControl::Construct(const FArguments& InArgs, const TSharedPtr<FAvaPlaylistEditor>& InPlaylistEditor)
+{
+	PlaylistEditor = InPlaylistEditor;
+	
+	ChildSlot
+	[
+		SNew(SVerticalBox)
+		+ SVerticalBox::Slot()
+		.AutoHeight()
+		[
+			BuildShowControlToolBar(InPlaylistEditor->GetToolkitCommands())
+		]
+	];
+}
+
+TSharedRef<SWidget> SAvaShowControl::BuildShowControlToolBar(const TSharedRef<FUICommandList>& InCommandList)
+{
+	FSlimHorizontalToolBarBuilder ToolBarBuilder(InCommandList, FMultiBoxCustomization::None);
+	ToolBarBuilder.SetStyle(&FAvaMediaEditorStyle::Get(), "AvalancheMediaEditor.ToolBar");
+
+	const FAvaPlaylistCommands& PlaylistCommands = FAvaPlaylistCommands::Get();
+
+	ToolBarBuilder.BeginSection(TEXT("ShowControl"));
+	{
+		ToolBarBuilder.BeginStyleOverride("AvalancheMediaEditor.ToolBarRedButtonOverride");
+		{
+			ToolBarBuilder.AddToolBarButton(PlaylistCommands.Play);
+		}
+		ToolBarBuilder.EndStyleOverride();
+
+		ToolBarBuilder.AddToolBarButton(PlaylistCommands.Continue);
+		ToolBarBuilder.AddToolBarButton(PlaylistCommands.Stop);
+
+		ToolBarBuilder.BeginStyleOverride("AvalancheMediaEditor.ToolBarRedButtonOverride");
+		{
+			ToolBarBuilder.AddToolBarButton(PlaylistCommands.PlayNext);
+		}
+		ToolBarBuilder.EndStyleOverride();
+
+		ToolBarBuilder.AddSeparator();
+		ToolBarBuilder.AddWidget(CreateActiveListWidget());
+		ToolBarBuilder.AddSeparator();
+		ToolBarBuilder.AddWidget(CreateNextPageWidget());
+	}
+	ToolBarBuilder.EndSection();
+
+	return ToolBarBuilder.MakeWidget();
+}
+
+TSharedRef<SWidget> SAvaShowControl::CreateActiveListWidget()
+{
+	return
+		SNew(SHorizontalBox)
+		+ SHorizontalBox::Slot()
+		.VAlign(EVerticalAlignment::VAlign_Center)
+		.Padding(5.f, 0.f)
+		.AutoWidth()
+		[
+			SNew(STextBlock)
+			.Text(LOCTEXT("ActiveView", "Active View:"))
+		]
+		+ SHorizontalBox::Slot()
+		.VAlign(EVerticalAlignment::VAlign_Center)
+		.Padding(0.f, 0.f, 5.f, 0.f)
+		.AutoWidth()
+		[
+			SNew(STextBlock)
+			.Text(this, &SAvaShowControl::GetActiveListName)
+		];
+}
+
+FText SAvaShowControl::GetActiveListName() const
+{
+	if (PlaylistEditor.IsValid())
+	{
+		UAvalanchePlaylist* Playlist = PlaylistEditor->GetPlaylist();
+
+		if (IsValid(Playlist))
+		{
+			const FAvaPageListReference& ActiveList = Playlist->GetActivePageListReference();
+
+			// Can't happen. Here for completeness.
+			if (ActiveList.Type == EAvaPageListType::Template)
+			{
+				static const FText Templates(LOCTEXT("Templates", "Templates"));
+				return Templates;
+			}
+
+			if (ActiveList.Type == EAvaPageListType::Instance)
+			{
+				static const FText AllPages(LOCTEXT("AllPages", "All Pages"));
+				return AllPages;
+			}
+
+			return FText::Format(LOCTEXT("PlaylistSubListDocument_TabLabel", "Page View {0}"), FText::AsNumber(ActiveList.SubListIndex + 1));
+		}
+	}
+
+	return FText::GetEmpty();
+}
+
+TSharedRef<SWidget> SAvaShowControl::CreateNextPageWidget()
+{
+	return
+		SNew(SHorizontalBox)
+		+ SHorizontalBox::Slot()
+		.VAlign(EVerticalAlignment::VAlign_Center)
+		.Padding(5.f, 0.f)
+		.AutoWidth()
+		[
+			SNew(STextBlock)
+			.Text(LOCTEXT("NextUp", "Next Up:"))
+		]
+		+ SHorizontalBox::Slot()
+		.VAlign(EVerticalAlignment::VAlign_Center)
+		.Padding(0.f, 0.f, 5.f, 0.f)
+		.AutoWidth()
+		[
+			SNew(STextBlock)
+			.Text(this, &SAvaShowControl::GetNextPageName)
+		];
+}
+
+FText SAvaShowControl::GetNextPageName() const
+{
+	static const FText NoPage(LOCTEXT("None", "-"));
+
+	if (PlaylistEditor.IsValid())
+	{
+		if (const UAvalanchePlaylist* Playlist = PlaylistEditor->GetPlaylist(); IsValid(Playlist))
+		{
+			if (Playlist->GetInstancedPages().Pages.IsEmpty())
+			{
+				return NoPage;
+			}
+
+			const TSharedPtr<SAvaInstancedPageList> ActiveSubListWidget = PlaylistEditor->GetActiveListWidget();
+
+			if (!ActiveSubListWidget.IsValid() || ActiveSubListWidget->GetPlayingPageIds().IsEmpty())
+			{
+				return NoPage;
+			}
+
+			TArray<int32> NextUps = ActiveSubListWidget->GetPageIdsToTakeNext();
+
+			FText NextUpMessage = FText::GetEmpty();
+			FFormatOrderedArguments FormatOrderedArguments;
+			
+			for (const int32 NextUp : NextUps)
+			{
+				if (const int32* PageIndex = Playlist->GetInstancedPages().PageIndices.Find(NextUp))
+				{
+					FormatOrderedArguments.Add(FText::Format(
+							LOCTEXT("NextUpFormat", "{0}: {1}"),
+							FText::AsNumber(NextUp, &UE::AvalanchePlaylist::FEditorMetrics::PageIdFormattingOptions),
+							FText::FromString(Playlist->GetInstancedPages().Pages[*PageIndex].GetPageName())
+						));
+				}
+			}
+
+			if (FormatOrderedArguments.Num() > 0)
+			{
+				return FText::Join(LOCTEXT("NextUpDelimiter", "\n"), FormatOrderedArguments);
+			}
+		}
+	}
+
+	return NoPage;
+}
+
+#undef LOCTEXT_NAMESPACE
diff --git a/Engine/Plugins/Experimental/Avalanche/Source/AvalancheMediaEditor/Private/Playlist/ShowControl/SAvaShowControl.h b/Engine/Plugins/Experimental/Avalanche/Source/AvalancheMediaEditor/Private/Playlist/ShowControl/SAvaShowControl.h
new file mode 100644
--- /dev/null
+++ b/Engine/Plugins/Experimental/Avalanche/Source/AvalancheMediaEditor/Private/Playlist/ShowControl/SAvaShowControl.h
@@ -0,0 +1,30 @@
+// Copyright Epic Games, Inc. All Rights Reserved.
+
+#pragma once
+
+#include "CoreMinimal.h"
+#include "Widgets/SCompoundWidget.h"
+
+class FUICommandList;
+class FAvaPlaylistEditor;
+
+class SAvaShowControl : public SCompoundWidget
+{
+public:	
+	SLATE_BEGIN_ARGS(SAvaShowControl) {}
+	SLATE_END_ARGS()
+
+	/** Constructs this widget with InArgs */
+	void Construct(const FArguments& InArgs, const TSharedPtr<FAvaPlaylistEditor>& InPlaylistEditor);
+
+	TSharedRef<SWidget> BuildShowControlToolBar(const TSharedRef<FUICommandList>& InCommandList);
+
+protected:	
+	TSharedPtr<FAvaPlaylistEditor> PlaylistEditor;
+
+	TSharedRef<SWidget> CreateActiveListWidget();
+	FText GetActiveListName() const;
+
+	TSharedRef<SWidget> CreateNextPageWidget();
+	FText GetNextPageName() const;
+};
diff --git a/Engine/Plugins/Experimental/Avalanche/Source/AvalancheMediaEditor/Private/Playlist/TabFactories/AvaChannelLayerStatusListTabFactory.cpp b/Engine/Plugins/Experimental/Avalanche/Source/AvalancheMediaEditor/Private/Playlist/TabFactories/AvaChannelLayerStatusListTabFactory.cpp
new file mode 100644
--- /dev/null
+++ b/Engine/Plugins/Experimental/Avalanche/Source/AvalancheMediaEditor/Private/Playlist/TabFactories/AvaChannelLayerStatusListTabFactory.cpp
@@ -0,0 +1,29 @@
+// Copyright Epic Games, Inc. All Rights Reserved.
+
+#include "AvaChannelLayerStatusListTabFactory.h"
+#include "AvaMediaEditorStyle.h"
+#include "Playlist/ChannelLayerStatus/SAvaChannelLayerStatusList.h"
+
+const FName FAvaChannelLayerStatusListTabFactory::TabID(TEXT("AvalancheChannelLayerStatusList"));
+
+#define LOCTEXT_NAMESPACE "AvaChannelLayerStatusListTabFactory"
+
+FAvaChannelLayerStatusListTabFactory::FAvaChannelLayerStatusListTabFactory(const TSharedPtr<FAvaPlaylistEditor>& InPlaylistEditor)
+	: FAvaPlaylistTabFactory(TabID, InPlaylistEditor)
+{
+	TabLabel = LOCTEXT("ChannelLayerStatus_TabLabel", "Layer Status");
+	TabIcon = FSlateIcon(FAvaMediaEditorStyle::GetStyleSetName(), "AvalancheMediaEditor.BroadcastIcon");
+
+	bIsSingleton = true;
+	bShouldAutosize = true;
+
+	ViewMenuDescription = LOCTEXT("ChannelLayerStatus_ViewMenu_Desc", "Displays the Status of all Channel Layers in Broadcast");
+	ViewMenuTooltip = LOCTEXT("ChannelLayerStatus_ViewMenu_ToolTip", "Displays the Status of all Channel Layers in Broadcast");
+}
+
+TSharedRef<SWidget> FAvaChannelLayerStatusListTabFactory::CreateTabBody(const FWorkflowTabSpawnInfo& Info) const
+{
+	return SNew(SAvaChannelLayerStatusList, PlaylistEditorWeak.Pin());
+}
+
+#undef LOCTEXT_NAMESPACE
diff --git a/Engine/Plugins/Experimental/Avalanche/Source/AvalancheMediaEditor/Private/Playlist/TabFactories/AvaChannelLayerStatusListTabFactory.h b/Engine/Plugins/Experimental/Avalanche/Source/AvalancheMediaEditor/Private/Playlist/TabFactories/AvaChannelLayerStatusListTabFactory.h
new file mode 100644
--- /dev/null
+++ b/Engine/Plugins/Experimental/Avalanche/Source/AvalancheMediaEditor/Private/Playlist/TabFactories/AvaChannelLayerStatusListTabFactory.h
@@ -0,0 +1,15 @@
+// Copyright Epic Games, Inc. All Rights Reserved.
+
+#pragma once
+
+#include "AvaPlaylistTabFactory.h"
+
+class FAvaChannelLayerStatusListTabFactory : public FAvaPlaylistTabFactory
+{
+public:
+	static const FName TabID;
+
+	FAvaChannelLayerStatusListTabFactory(const TSharedPtr<FAvaPlaylistEditor>& InPlaylistEditor);
+
+	virtual TSharedRef<SWidget> CreateTabBody(const FWorkflowTabSpawnInfo& Info) const override;
+};
diff --git a/Engine/Plugins/Experimental/Avalanche/Source/AvalancheMediaEditor/Private/Playlist/TabFactories/AvaChannelStatusListTabFactory.cpp b/Engine/Plugins/Experimental/Avalanche/Source/AvalancheMediaEditor/Private/Playlist/TabFactories/AvaChannelStatusListTabFactory.cpp
new file mode 100644
--- /dev/null
+++ b/Engine/Plugins/Experimental/Avalanche/Source/AvalancheMediaEditor/Private/Playlist/TabFactories/AvaChannelStatusListTabFactory.cpp
@@ -0,0 +1,29 @@
+// Copyright Epic Games, Inc. All Rights Reserved.
+
+#include "AvaChannelStatusListTabFactory.h"
+#include "AvaMediaEditorStyle.h"
+#include "Playlist/ChannelStatus/SAvaChannelStatusList.h"
+
+const FName FAvaChannelStatusListTabFactory::TabID(TEXT("AvalancheChannelStatusList"));
+
+#define LOCTEXT_NAMESPACE "AvaChannelStatusListTabFactory"
+
+FAvaChannelStatusListTabFactory::FAvaChannelStatusListTabFactory(const TSharedPtr<FAvaPlaylistEditor>& InPlaylistEditor)
+	: FAvaPlaylistTabFactory(TabID, InPlaylistEditor)
+{
+	TabLabel = LOCTEXT("ChannelStatus_TabLabel", "Channel Status");
+	TabIcon = FSlateIcon(FAvaMediaEditorStyle::GetStyleSetName(), "AvalancheMediaEditor.BroadcastIcon");
+
+	bIsSingleton = true;
+	bShouldAutosize = true;
+
+	ViewMenuDescription = LOCTEXT("ChannelStatus_ViewMenu_Desc", "Displays the Status of all Channels in Broadcast");
+	ViewMenuTooltip = LOCTEXT("ChannelStatus_ViewMenu_ToolTip", "Displays the Status of all Channels in Broadcast");
+}
+
+TSharedRef<SWidget> FAvaChannelStatusListTabFactory::CreateTabBody(const FWorkflowTabSpawnInfo& Info) const
+{
+	return SNew(SAvaChannelStatusList);
+}
+
+#undef LOCTEXT_NAMESPACE
diff --git a/Engine/Plugins/Experimental/Avalanche/Source/AvalancheMediaEditor/Private/Playlist/TabFactories/AvaChannelStatusListTabFactory.h b/Engine/Plugins/Experimental/Avalanche/Source/AvalancheMediaEditor/Private/Playlist/TabFactories/AvaChannelStatusListTabFactory.h
new file mode 100644
--- /dev/null
+++ b/Engine/Plugins/Experimental/Avalanche/Source/AvalancheMediaEditor/Private/Playlist/TabFactories/AvaChannelStatusListTabFactory.h
@@ -0,0 +1,15 @@
+// Copyright Epic Games, Inc. All Rights Reserved.
+
+#pragma once
+
+#include "AvaPlaylistTabFactory.h"
+
+class FAvaChannelStatusListTabFactory : public FAvaPlaylistTabFactory
+{
+public:
+	static const FName TabID;
+	
+	FAvaChannelStatusListTabFactory(const TSharedPtr<FAvaPlaylistEditor>& InPlaylistEditor);
+	
+	virtual TSharedRef<SWidget> CreateTabBody(const FWorkflowTabSpawnInfo& Info) const override;
+};
diff --git a/Engine/Plugins/Experimental/Avalanche/Source/AvalancheMediaEditor/Private/Playlist/TabFactories/AvaInstancedPageListTabFactory.cpp b/Engine/Plugins/Experimental/Avalanche/Source/AvalancheMediaEditor/Private/Playlist/TabFactories/AvaInstancedPageListTabFactory.cpp
new file mode 100644
--- /dev/null
+++ b/Engine/Plugins/Experimental/Avalanche/Source/AvalancheMediaEditor/Private/Playlist/TabFactories/AvaInstancedPageListTabFactory.cpp
@@ -0,0 +1,47 @@
+// Copyright Epic Games, Inc. All Rights Reserved.
+
+#include "AvaInstancedPageListTabFactory.h"
+#include "Playlist/AvaPlaylistEditor.h"
+#include "Playlist/Pages/Slate/SAvaInstancedPageList.h"
+#include "Widgets/Docking/SDockTab.h"
+
+const FName FAvaInstancedPageListTabFactory::TabID(TEXT("AvalancheInstancedPlaylistPageList"));
+
+#define LOCTEXT_NAMESPACE "AvaInstancedPageListTabFactory"
+
+FAvaInstancedPageListTabFactory::FAvaInstancedPageListTabFactory(const TSharedPtr<FAvaPlaylistEditor>& InPlaylistEditor)
+	: FAvaPlaylistTabFactory(TabID, InPlaylistEditor)
+{
+	TabLabel = LOCTEXT("PlaylistPageList_TabLabel", "Pages");
+	TabIcon = FSlateIcon(FAppStyle::GetAppStyleSetName(), "PlacementBrowser.Icons.All");
+
+	bIsSingleton = true;
+
+	ViewMenuDescription = LOCTEXT("PlaylistPageList_ViewMenu_Desc", "Pages");
+	ViewMenuTooltip = LOCTEXT("PlaylistPageList_ViewMenu_ToolTip", "Pages");
+}
+
+TSharedRef<SWidget> FAvaInstancedPageListTabFactory::CreateTabBody(const FWorkflowTabSpawnInfo& InInfo) const
+{
+	return SNew(SAvaInstancedPageList, PlaylistEditorWeak.Pin(), UAvalanchePlaylist::InstancePageList);
+}
+
+TSharedRef<SDockTab> FAvaInstancedPageListTabFactory::OnSpawnTab(const FSpawnTabArgs& InSpawnArgs,
+	TWeakPtr<FTabManager> InWeakTabManager) const
+{
+	TSharedRef<SDockTab> NewTab = FAvaPlaylistTabFactory::OnSpawnTab(InSpawnArgs, InWeakTabManager);
+	TSharedRef<SWidget> Content = NewTab->GetContent();
+
+	if (Content->GetWidgetClass().GetWidgetType() == SBorder::StaticWidgetClass().GetWidgetType()
+		&& Content->GetTypeAsString() == "SBorder") // Subclasses of SBorder that do not propertly declare themselves will have the same widget class.
+	{
+		Content = StaticCastSharedRef<SBorder>(Content)->GetContent();
+	}
+
+	TSharedRef<SAvaInstancedPageList> PageList = StaticCastSharedRef<SAvaInstancedPageList>(Content);
+	NewTab->SetOnTabActivated(SDockTab::FOnTabActivatedCallback::CreateSP(PageList, &SAvaInstancedPageList::OnTabActivated));
+
+	return NewTab;
+}
+
+#undef LOCTEXT_NAMESPACE
diff --git a/Engine/Plugins/Experimental/Avalanche/Source/AvalancheMediaEditor/Private/Playlist/TabFactories/AvaInstancedPageListTabFactory.h b/Engine/Plugins/Experimental/Avalanche/Source/AvalancheMediaEditor/Private/Playlist/TabFactories/AvaInstancedPageListTabFactory.h
new file mode 100644
--- /dev/null
+++ b/Engine/Plugins/Experimental/Avalanche/Source/AvalancheMediaEditor/Private/Playlist/TabFactories/AvaInstancedPageListTabFactory.h
@@ -0,0 +1,18 @@
+// Copyright Epic Games, Inc. All Rights Reserved.
+
+#pragma once
+
+#include "AvaPlaylistTabFactory.h"
+
+class FAvaInstancedPageListTabFactory : public FAvaPlaylistTabFactory
+{
+public:
+	static const FName TabID;
+	
+	FAvaInstancedPageListTabFactory(const TSharedPtr<FAvaPlaylistEditor>& InPlaylistEditor);
+	
+	virtual TSharedRef<SWidget> CreateTabBody(const FWorkflowTabSpawnInfo& InInfo) const override;
+
+protected:
+	virtual TSharedRef<SDockTab> OnSpawnTab(const FSpawnTabArgs& InSpawnArgs, TWeakPtr<FTabManager> InWeakTabManager) const override;
+};
diff --git a/Engine/Plugins/Experimental/Avalanche/Source/AvalancheMediaEditor/Private/Playlist/TabFactories/AvaPageDetailsTabFactory.cpp b/Engine/Plugins/Experimental/Avalanche/Source/AvalancheMediaEditor/Private/Playlist/TabFactories/AvaPageDetailsTabFactory.cpp
new file mode 100644
--- /dev/null
+++ b/Engine/Plugins/Experimental/Avalanche/Source/AvalancheMediaEditor/Private/Playlist/TabFactories/AvaPageDetailsTabFactory.cpp
@@ -0,0 +1,28 @@
+// Copyright Epic Games, Inc. All Rights Reserved.
+
+#include "AvaPageDetailsTabFactory.h"
+#include "Playlist/AvaPlaylistEditor.h"
+#include "Playlist/DetailsView/SAvaPageDetails.h"
+
+const FName FAvaPageDetailsTabFactory::TabID(TEXT("AvalanchePlaylistPageDetails"));
+
+#define LOCTEXT_NAMESPACE "AvaPageDetailsTabFactory"
+
+FAvaPageDetailsTabFactory::FAvaPageDetailsTabFactory(const TSharedPtr<FAvaPlaylistEditor>& InPlaylistEditor)
+	: FAvaPlaylistTabFactory(TabID, InPlaylistEditor)
+{
+	TabLabel = LOCTEXT("PlaylistPageDetails_TabLabel", "Page Details");
+	TabIcon = FSlateIcon(FAppStyle::GetAppStyleSetName(), "LevelEditor.Tabs.Details");
+
+	bIsSingleton = true;
+
+	ViewMenuDescription = LOCTEXT("PlaylistPageDetails_ViewMenu_Desc", "Page Details");
+	ViewMenuTooltip = LOCTEXT("PlaylistPageDetails_ViewMenu_ToolTip", "Page Details");
+}
+
+TSharedRef<SWidget> FAvaPageDetailsTabFactory::CreateTabBody(const FWorkflowTabSpawnInfo& Info) const
+{
+	return SNew(SAvaPageDetails, PlaylistEditorWeak.Pin());
+}
+
+#undef LOCTEXT_NAMESPACE
diff --git a/Engine/Plugins/Experimental/Avalanche/Source/AvalancheMediaEditor/Private/Playlist/TabFactories/AvaPageDetailsTabFactory.h b/Engine/Plugins/Experimental/Avalanche/Source/AvalancheMediaEditor/Private/Playlist/TabFactories/AvaPageDetailsTabFactory.h
new file mode 100644
--- /dev/null
+++ b/Engine/Plugins/Experimental/Avalanche/Source/AvalancheMediaEditor/Private/Playlist/TabFactories/AvaPageDetailsTabFactory.h
@@ -0,0 +1,15 @@
+// Copyright Epic Games, Inc. All Rights Reserved.
+
+#pragma once
+
+#include "AvaPlaylistTabFactory.h"
+
+class FAvaPageDetailsTabFactory : public FAvaPlaylistTabFactory
+{
+public:
+	static const FName TabID;
+	
+	FAvaPageDetailsTabFactory(const TSharedPtr<FAvaPlaylistEditor>& InPlaylistEditor);
+	
+	virtual TSharedRef<SWidget> CreateTabBody(const FWorkflowTabSpawnInfo& Info) const override;
+};
diff --git a/Engine/Plugins/Experimental/Avalanche/Source/AvalancheMediaEditor/Private/Playlist/TabFactories/AvaPageViewerTabFactory.cpp b/Engine/Plugins/Experimental/Avalanche/Source/AvalancheMediaEditor/Private/Playlist/TabFactories/AvaPageViewerTabFactory.cpp
new file mode 100644
--- /dev/null
+++ b/Engine/Plugins/Experimental/Avalanche/Source/AvalancheMediaEditor/Private/Playlist/TabFactories/AvaPageViewerTabFactory.cpp
@@ -0,0 +1,30 @@
+// Copyright Epic Games, Inc. All Rights Reserved.
+
+#include "AvaPageViewerTabFactory.h"
+#include "Playlist/AvaPlaylistEditor.h"
+#include "Playlist/Preview/SAvaPagePreview.h"
+
+const FName FAvaPageViewerTabFactory::TabID(TEXT("AvalanchePlaylistPageViewer"));
+
+#define LOCTEXT_NAMESPACE "AvalanchePlaylist"
+
+FAvaPageViewerTabFactory::FAvaPageViewerTabFactory(const TSharedPtr<FAvaPlaylistEditor>& InPlaylistEditor)
+	: FAvaPlaylistTabFactory(TabID, InPlaylistEditor)
+{
+	TabLabel = LOCTEXT("PlaylistPageViewer_TabLabel", "Page Preview");
+	TabIcon = FSlateIcon(FAppStyle::GetAppStyleSetName(), "LevelEditor.Tabs.VirtualProduction");
+
+	bIsSingleton = true;
+
+	ViewMenuDescription = LOCTEXT("PlaylistPageViewer_ViewMenu_Desc", "Previews the selected page");
+	ViewMenuTooltip = LOCTEXT("PlaylistPageViewer_ViewMenu_ToolTip", "Previews the selected page");
+}
+
+TSharedRef<SWidget> FAvaPageViewerTabFactory::CreateTabBody(const FWorkflowTabSpawnInfo& Info) const
+{
+	TSharedPtr<FAvaPlaylistEditor> PlaylistEditor = PlaylistEditorWeak.Pin();
+	
+	return SNew(SAvaPagePreview, PlaylistEditor);
+}
+
+#undef LOCTEXT_NAMESPACE
diff --git a/Engine/Plugins/Experimental/Avalanche/Source/AvalancheMediaEditor/Private/Playlist/TabFactories/AvaPageViewerTabFactory.h b/Engine/Plugins/Experimental/Avalanche/Source/AvalancheMediaEditor/Private/Playlist/TabFactories/AvaPageViewerTabFactory.h
new file mode 100644
--- /dev/null
+++ b/Engine/Plugins/Experimental/Avalanche/Source/AvalancheMediaEditor/Private/Playlist/TabFactories/AvaPageViewerTabFactory.h
@@ -0,0 +1,15 @@
+// Copyright Epic Games, Inc. All Rights Reserved.
+
+#pragma once
+
+#include "AvaPlaylistTabFactory.h"
+
+class FAvaPageViewerTabFactory  : public FAvaPlaylistTabFactory
+{
+public:
+	static const FName TabID;
+	
+	FAvaPageViewerTabFactory(const TSharedPtr<FAvaPlaylistEditor>& InPlaylistEditor);
+	
+	virtual TSharedRef<SWidget> CreateTabBody(const FWorkflowTabSpawnInfo& Info) const override;
+};
diff --git a/Engine/Plugins/Experimental/Avalanche/Source/AvalancheMediaEditor/Private/Playlist/TabFactories/AvaPlaylistTabFactory.cpp b/Engine/Plugins/Experimental/Avalanche/Source/AvalancheMediaEditor/Private/Playlist/TabFactories/AvaPlaylistTabFactory.cpp
new file mode 100644
--- /dev/null
+++ b/Engine/Plugins/Experimental/Avalanche/Source/AvalancheMediaEditor/Private/Playlist/TabFactories/AvaPlaylistTabFactory.cpp
@@ -0,0 +1,10 @@
+// Copyright Epic Games, Inc. All Rights Reserved.
+
+#include "AvaPlaylistTabFactory.h"
+#include "Playlist/AvaPlaylistEditor.h"
+
+FAvaPlaylistTabFactory::FAvaPlaylistTabFactory(const FName& InTabID, const TSharedPtr<FAvaPlaylistEditor>& InPlaylistEditor)
+	: FWorkflowTabFactory(InTabID, InPlaylistEditor)
+	, PlaylistEditorWeak(InPlaylistEditor)
+{
+}
diff --git a/Engine/Plugins/Experimental/Avalanche/Source/AvalancheMediaEditor/Private/Playlist/TabFactories/AvaPlaylistTabFactory.h b/Engine/Plugins/Experimental/Avalanche/Source/AvalancheMediaEditor/Private/Playlist/TabFactories/AvaPlaylistTabFactory.h
new file mode 100644
--- /dev/null
+++ b/Engine/Plugins/Experimental/Avalanche/Source/AvalancheMediaEditor/Private/Playlist/TabFactories/AvaPlaylistTabFactory.h
@@ -0,0 +1,20 @@
+// Copyright Epic Games, Inc. All Rights Reserved.
+
+#pragma once
+
+#include "WorkflowOrientedApp/WorkflowTabFactory.h"
+
+class FAvaPlaylistEditor;
+
+//Base class for all Tab Factories in Ava Playlist Editor
+class FAvaPlaylistTabFactory : public FWorkflowTabFactory
+{
+public:
+
+	FAvaPlaylistTabFactory(const FName& InTabID, const TSharedPtr<FAvaPlaylistEditor>& InPlaylistEditor);
+
+protected:
+
+	TWeakPtr<FAvaPlaylistEditor> PlaylistEditorWeak;
+};
+
diff --git a/Engine/Plugins/Experimental/Avalanche/Source/AvalancheMediaEditor/Private/Playlist/TabFactories/AvaShowControlTabFactory.cpp b/Engine/Plugins/Experimental/Avalanche/Source/AvalancheMediaEditor/Private/Playlist/TabFactories/AvaShowControlTabFactory.cpp
new file mode 100644
--- /dev/null
+++ b/Engine/Plugins/Experimental/Avalanche/Source/AvalancheMediaEditor/Private/Playlist/TabFactories/AvaShowControlTabFactory.cpp
@@ -0,0 +1,29 @@
+// Copyright Epic Games, Inc. All Rights Reserved.
+
+#include "AvaShowControlTabFactory.h"
+#include "Playlist/AvaPlaylistEditor.h"
+#include "Playlist/ShowControl/SAvaShowControl.h"
+
+const FName FAvaShowControlTabFactory::TabID(TEXT("AvalanchePlaylistShowControl"));
+
+#define LOCTEXT_NAMESPACE "AvaShowControlTabFactory"
+
+FAvaShowControlTabFactory::FAvaShowControlTabFactory(const TSharedPtr<FAvaPlaylistEditor>& InPlaylistEditor)
+	: FAvaPlaylistTabFactory(TabID, InPlaylistEditor)
+{
+	TabLabel = LOCTEXT("PlaylistShowControl_TabLabel", "Show Control");
+	TabIcon = FSlateIcon(FAppStyle::GetAppStyleSetName(), "Icons.Toolbar.Settings");
+
+	bIsSingleton = true;
+	bShouldAutosize = false;
+
+	ViewMenuDescription = LOCTEXT("PlaylistShowControl_ViewMenu_Desc", "Show Control");
+	ViewMenuTooltip = LOCTEXT("PlaylistShowControl_ViewMenu_ToolTip", "Show Control");
+}
+
+TSharedRef<SWidget> FAvaShowControlTabFactory::CreateTabBody(const FWorkflowTabSpawnInfo& Info) const
+{
+	return SNew(SAvaShowControl, PlaylistEditorWeak.Pin());
+}
+
+#undef LOCTEXT_NAMESPACE
diff --git a/Engine/Plugins/Experimental/Avalanche/Source/AvalancheMediaEditor/Private/Playlist/TabFactories/AvaShowControlTabFactory.h b/Engine/Plugins/Experimental/Avalanche/Source/AvalancheMediaEditor/Private/Playlist/TabFactories/AvaShowControlTabFactory.h
new file mode 100644
--- /dev/null
+++ b/Engine/Plugins/Experimental/Avalanche/Source/AvalancheMediaEditor/Private/Playlist/TabFactories/AvaShowControlTabFactory.h
@@ -0,0 +1,15 @@
+// Copyright Epic Games, Inc. All Rights Reserved.
+
+#pragma once
+
+#include "AvaPlaylistTabFactory.h"
+
+class FAvaShowControlTabFactory : public FAvaPlaylistTabFactory
+{
+public:
+	static const FName TabID;
+	
+	FAvaShowControlTabFactory(const TSharedPtr<FAvaPlaylistEditor>& InPlaylistEditor);
+	
+	virtual TSharedRef<SWidget> CreateTabBody(const FWorkflowTabSpawnInfo& Info) const override;
+};
diff --git a/Engine/Plugins/Experimental/Avalanche/Source/AvalancheMediaEditor/Private/Playlist/TabFactories/AvaSubListDocumentTabFactory.cpp b/Engine/Plugins/Experimental/Avalanche/Source/AvalancheMediaEditor/Private/Playlist/TabFactories/AvaSubListDocumentTabFactory.cpp
new file mode 100644
--- /dev/null
+++ b/Engine/Plugins/Experimental/Avalanche/Source/AvalancheMediaEditor/Private/Playlist/TabFactories/AvaSubListDocumentTabFactory.cpp
@@ -0,0 +1,64 @@
+// Copyright Epic Games, Inc. All Rights Reserved.
+
+#include "AvaSubListDocumentTabFactory.h"
+#include "Playlist/AvaPlaylistEditor.h"
+#include "Playlist/Pages/Slate/SAvaInstancedPageList.h"
+#include "Widgets/Docking/SDockTab.h"
+
+#define LOCTEXT_NAMESPACE "FAvaSubListDocumentTabFactory"
+
+const FName FAvaSubListDocumentTabFactory::FactoryId = "AvaSubListTabFactory";
+const FString FAvaSubListDocumentTabFactory::BaseTabName(TEXT("AvaSubListDocument"));
+
+FName FAvaSubListDocumentTabFactory::GetTabId(int32 InSubListIndex)
+{
+	return FName(BaseTabName + "_" + FString::FromInt(InSubListIndex));
+}
+
+FAvaSubListDocumentTabFactory::FAvaSubListDocumentTabFactory(const TSharedPtr<FAvaPlaylistEditor>& InPlaylistEditor)
+	: FDocumentTabFactory(FactoryId, InPlaylistEditor)
+	, PlaylistEditorWeak(InPlaylistEditor)
+	, SubListIndex(UAvalanchePlaylist::InstancePageList.SubListIndex)
+{
+	TabIcon = FSlateIcon(FAppStyle::GetAppStyleSetName(), "PlacementBrowser.Icons.All");
+}
+
+TSharedRef<SWidget> FAvaSubListDocumentTabFactory::CreateTabBody(const FWorkflowTabSpawnInfo& InInfo) const
+{
+	if (SubListIndex <= UAvalanchePlaylist::InstancePageList.SubListIndex)
+	{
+		return SNullWidget::NullWidget;
+	}
+
+	return SNew(SAvaInstancedPageList, PlaylistEditorWeak.Pin(), UAvalanchePlaylist::CreateSubListReference(SubListIndex));
+}
+
+TSharedRef<SDockTab> FAvaSubListDocumentTabFactory::SpawnSubListTab(const FWorkflowTabSpawnInfo& InInfo, int32 InSubListIndex)
+{
+	SubListIndex = InSubListIndex;
+	TabIdentifier = GetTabId(InSubListIndex);
+	TabLabel = FText::Format(LOCTEXT("PlaylistSubListDocument_TabLabel", "Page View {0}"), FText::AsNumber(SubListIndex + 1));
+	ViewMenuDescription = FText::Format(LOCTEXT("PlaylistSubListDocument_ViewMenu_Desc", "Page View {0}"), FText::AsNumber(SubListIndex + 1));
+	ViewMenuTooltip = FText::Format(LOCTEXT("PlaylistSubListDocument_ViewMenu_ToolTip", "Page View {0}"), FText::AsNumber(SubListIndex + 1));
+
+	TSharedRef<SDockTab> NewTab = SpawnTab(InInfo);
+
+	TSharedRef<SAvaInstancedPageList> PageList = StaticCastSharedRef<SAvaInstancedPageList>(NewTab->GetContent());
+	NewTab->SetOnTabActivated(SDockTab::FOnTabActivatedCallback::CreateSP(PageList, &SAvaInstancedPageList::OnTabActivated));
+	PageList->SetMyTab(NewTab);
+
+	if (TSharedPtr<FAvaPlaylistEditor> PlaylistEditor = PageList->GetPlaylistEditor())
+	{
+		UAvalanchePlaylist* Playlist = PlaylistEditor->GetPlaylist();
+
+		if (IsValid(Playlist) && Playlist->IsValidSubListIndex(InSubListIndex)
+			&& !Playlist->GetSubList(InSubListIndex).Name.IsEmpty())
+		{
+			NewTab->SetLabel(Playlist->GetSubList(InSubListIndex).Name);
+		}
+	}
+
+	return NewTab;
+}
+
+#undef LOCTEXT_NAMESPACE
diff --git a/Engine/Plugins/Experimental/Avalanche/Source/AvalancheMediaEditor/Private/Playlist/TabFactories/AvaSubListDocumentTabFactory.h b/Engine/Plugins/Experimental/Avalanche/Source/AvalancheMediaEditor/Private/Playlist/TabFactories/AvaSubListDocumentTabFactory.h
new file mode 100644
--- /dev/null
+++ b/Engine/Plugins/Experimental/Avalanche/Source/AvalancheMediaEditor/Private/Playlist/TabFactories/AvaSubListDocumentTabFactory.h
@@ -0,0 +1,28 @@
+// Copyright Epic Games, Inc. All Rights Reserved.
+
+#pragma once
+
+#include "WorkflowOrientedApp/WorkflowTabFactory.h"
+
+class FAvaPlaylistEditor;
+
+/** Base class for all Tab Factories in Ava SubListDocument Editor */
+class FAvaSubListDocumentTabFactory : public FDocumentTabFactory
+{
+public:
+	static const FName FactoryId;
+	static const FString BaseTabName;
+
+	static FName GetTabId(int32 InSubListIndex);
+
+	FAvaSubListDocumentTabFactory(const TSharedPtr<FAvaPlaylistEditor>& InSubListDocumentEditor);
+
+	virtual TSharedRef<SWidget> CreateTabBody(const FWorkflowTabSpawnInfo& InInfo) const override;
+
+	TSharedRef<SDockTab> SpawnSubListTab(const FWorkflowTabSpawnInfo& InInfo, int32 InSubListIndex);
+
+protected:
+	TWeakPtr<FAvaPlaylistEditor> PlaylistEditorWeak;
+	int32 SubListIndex;
+};
+
diff --git a/Engine/Plugins/Experimental/Avalanche/Source/AvalancheMediaEditor/Private/Playlist/TabFactories/AvaSubListTabFactory.cpp b/Engine/Plugins/Experimental/Avalanche/Source/AvalancheMediaEditor/Private/Playlist/TabFactories/AvaSubListTabFactory.cpp
new file mode 100644
--- /dev/null
+++ b/Engine/Plugins/Experimental/Avalanche/Source/AvalancheMediaEditor/Private/Playlist/TabFactories/AvaSubListTabFactory.cpp
@@ -0,0 +1,32 @@
+// Copyright Epic Games, Inc. All Rights Reserved.
+
+#include "AvaSubListTabFactory.h"
+#include "Playlist/AvaPlaylistEditor.h"
+#include "Playlist/Pages/Slate/SAvaSubListStartPage.h"
+
+#define LOCTEXT_NAMESPACE "FAvaSubListTabFactory"
+
+const FName FAvaSubListTabFactory::TabID(TEXT("AvalanchePlaylistSubList"));
+
+FAvaSubListTabFactory::FAvaSubListTabFactory(const TSharedPtr<FAvaPlaylistEditor>& InPlaylistEditor)
+	: FAvaPlaylistTabFactory(TabID, InPlaylistEditor)
+{
+	PlaylistEditorWeak = InPlaylistEditor;
+
+	TabLabel = LOCTEXT("PlaylistSubList_TabLabel", "Page View Start");
+	TabIcon = FSlateIcon(FAppStyle::GetAppStyleSetName(), "PlacementBrowser.Icons.All");
+
+	bIsSingleton = true;
+
+	ViewMenuDescription = LOCTEXT("PlaylistSubList_ViewMenu_Desc", "Page View Start");
+	ViewMenuTooltip = LOCTEXT("PlaylistSubList_ViewMenu_ToolTip", "Page View Start");
+}
+
+TSharedRef<SWidget> FAvaSubListTabFactory::CreateTabBody(const FWorkflowTabSpawnInfo& InInfo) const
+{
+	TSharedPtr<FAvaPlaylistEditor> PlaylistEditor = PlaylistEditorWeak.Pin();
+
+	return SNew(SAvaSubListStartPage, PlaylistEditor);
+}
+
+#undef LOCTEXT_NAMESPACE
diff --git a/Engine/Plugins/Experimental/Avalanche/Source/AvalancheMediaEditor/Private/Playlist/TabFactories/AvaSubListTabFactory.h b/Engine/Plugins/Experimental/Avalanche/Source/AvalancheMediaEditor/Private/Playlist/TabFactories/AvaSubListTabFactory.h
new file mode 100644
--- /dev/null
+++ b/Engine/Plugins/Experimental/Avalanche/Source/AvalancheMediaEditor/Private/Playlist/TabFactories/AvaSubListTabFactory.h
@@ -0,0 +1,22 @@
+// Copyright Epic Games, Inc. All Rights Reserved.
+
+#pragma once
+
+#include "AvaPlaylistTabFactory.h"
+
+class FAvaPlaylistEditor;
+
+//Base class for all Tab Factories in Ava SubList Editor
+class FAvaSubListTabFactory : public FAvaPlaylistTabFactory
+{
+public:
+	static const FName TabID;
+
+	FAvaSubListTabFactory(const TSharedPtr<FAvaPlaylistEditor>& InSubListEditor);
+
+	virtual TSharedRef<SWidget> CreateTabBody(const FWorkflowTabSpawnInfo& InInfo) const override;
+
+protected:
+	TWeakPtr<FAvaPlaylistEditor> PlaylistEditorWeak;
+};
+
diff --git a/Engine/Plugins/Experimental/Avalanche/Source/AvalancheMediaEditor/Private/Playlist/TabFactories/AvaTemplatePageListTabFactory.cpp b/Engine/Plugins/Experimental/Avalanche/Source/AvalancheMediaEditor/Private/Playlist/TabFactories/AvaTemplatePageListTabFactory.cpp
new file mode 100644
--- /dev/null
+++ b/Engine/Plugins/Experimental/Avalanche/Source/AvalancheMediaEditor/Private/Playlist/TabFactories/AvaTemplatePageListTabFactory.cpp
@@ -0,0 +1,28 @@
+// Copyright Epic Games, Inc. All Rights Reserved.
+
+#include "AvaTemplatePageListTabFactory.h"
+#include "Playlist/AvaPlaylistEditor.h"
+#include "Playlist/Pages/Slate/SAvaTemplatePageList.h"
+
+const FName FAvaTemplatePageListTabFactory::TabID(TEXT("AvalancheTemplatePlaylistPageList"));
+
+#define LOCTEXT_NAMESPACE "AvaTemplatePageListTabFactory"
+
+FAvaTemplatePageListTabFactory::FAvaTemplatePageListTabFactory(const TSharedPtr<FAvaPlaylistEditor>& InPlaylistEditor)
+	: FAvaPlaylistTabFactory(TabID, InPlaylistEditor)
+{
+	TabLabel = LOCTEXT("PlaylistPageList_TabLabel", "Templates");
+	TabIcon = FSlateIcon(FAppStyle::GetAppStyleSetName(), "PlacementBrowser.Icons.All");
+
+	bIsSingleton = true;
+
+	ViewMenuDescription = LOCTEXT("PlaylistPageList_ViewMenu_Desc", "Templates");
+	ViewMenuTooltip = LOCTEXT("PlaylistPageList_ViewMenu_ToolTip", "Templates");
+}
+
+TSharedRef<SWidget> FAvaTemplatePageListTabFactory::CreateTabBody(const FWorkflowTabSpawnInfo& InInfo) const
+{
+	return SNew(SAvaTemplatePageList, PlaylistEditorWeak.Pin());
+}
+
+#undef LOCTEXT_NAMESPACE
diff --git a/Engine/Plugins/Experimental/Avalanche/Source/AvalancheMediaEditor/Private/Playlist/TabFactories/AvaTemplatePageListTabFactory.h b/Engine/Plugins/Experimental/Avalanche/Source/AvalancheMediaEditor/Private/Playlist/TabFactories/AvaTemplatePageListTabFactory.h
new file mode 100644
--- /dev/null
+++ b/Engine/Plugins/Experimental/Avalanche/Source/AvalancheMediaEditor/Private/Playlist/TabFactories/AvaTemplatePageListTabFactory.h
@@ -0,0 +1,15 @@
+// Copyright Epic Games, Inc. All Rights Reserved.
+
+#pragma once
+
+#include "AvaPlaylistTabFactory.h"
+
+class FAvaTemplatePageListTabFactory : public FAvaPlaylistTabFactory
+{
+public:
+	static const FName TabID;
+	
+	FAvaTemplatePageListTabFactory(const TSharedPtr<FAvaPlaylistEditor>& InPlaylistEditor);
+	
+	virtual TSharedRef<SWidget> CreateTabBody(const FWorkflowTabSpawnInfo& InInfo) const override;
+};
diff --git a/Engine/Plugins/Experimental/Avalanche/Source/AvalancheMediaEditor/Private/Tests/AvaPlaylistServerTests.cpp b/Engine/Plugins/Experimental/Avalanche/Source/AvalancheMediaEditor/Private/Tests/AvaPlaylistServerTests.cpp
new file mode 100644
--- /dev/null
+++ b/Engine/Plugins/Experimental/Avalanche/Source/AvalancheMediaEditor/Private/Tests/AvaPlaylistServerTests.cpp
@@ -0,0 +1,250 @@
+// Copyright Epic Games, Inc. All Rights Reserved.
+
+#include "AvalancheBroadcast.h"
+#include "Broadcast/OutputDevices/AvaOutputRootItem.h"
+#include "Broadcast/OutputDevices/AvaOutputTreeItem.h"
+#include "IAvaMediaEditorModule.h"
+#include "MediaOutput.h"
+#include "MessageEndpointBuilder.h"
+#include "Misc/AutomationTest.h"
+#include "Misc/FileHelper.h"
+#include "Playlist/AvaPlaylistServer.h"
+#include "Playlist/MediaOutputEditorUtils/MediaOutputEditorUtils.h"
+#include "UObject/UObjectIterator.h"
+
+DEFINE_LOG_CATEGORY_STATIC(LogAvaPlaylistServerTests, Log, All);
+
+namespace UE::AvaPlaylistServerTests
+{
+	class FPlaylistServerTestBase : public FAutomationTestBase
+	{
+	public:
+		FPlaylistServerTestBase(const FString& InName, const bool bInComplexTask)
+		: FAutomationTestBase(InName, bInComplexTask)
+		{
+			// UE-185435 - Blackmagic lib generates an error if the driver is not installed.
+			LocalSuppressedLogCategories.Add(TEXT("LogBlackmagicCore"));
+		}
+
+		//~ Begin FAutomationTestBase
+		virtual bool ShouldCaptureLogCategory(const FName& Category) const override
+		{
+			if (LocalSuppressedLogCategories.Contains(Category))
+			{
+				return false;
+			}
+			return FAutomationTestBase::ShouldCaptureLogCategory(Category);
+		}
+		//~ End FAutomationTestBase
+
+	protected:
+		TSet<FName> LocalSuppressedLogCategories;
+	};
+
+	class FTestClient
+	{
+	public:
+		TSharedPtr<FMessageEndpoint> MessageEndpoint;
+		FAutomationTestBase* Test = nullptr;
+		TSet<int32> ReceivedRequestIds;
+
+		bool HasReceivedResponse(int32 InRequestId) const
+		{
+			return ReceivedRequestIds.Contains(InRequestId);
+		}
+		
+		void Init(FAutomationTestBase* InTest)
+		{
+			Test = InTest;
+			
+			MessageEndpoint = FMessageEndpoint::Builder("AvaMediaPlaylistTestClient")
+				.Handling<FAvaPlaylistChannelImage>(this, &FTestClient::HandlePlaylistChannelImage)
+				.Handling<FAvaPlaylistServerMsg>(this, &FTestClient::HandleServerMessage);
+		}
+		
+		void HandlePlaylistChannelImage(const FAvaPlaylistChannelImage& InMessage, const TSharedRef<IMessageContext, ESPMode::ThreadSafe>& InContext)
+		{
+			ReceivedRequestIds.Add(InMessage.RequestId);
+			
+			UE_LOG(LogAvaPlaylistServerTests, Log, TEXT("Received Playlist Channel Image successfully."));
+
+			const FString SaveName = FString::Printf(TEXT("%s%s"), *FPaths::ProjectSavedDir(), TEXT("ChannelImageTest.jpeg"));
+			UE_LOG(LogAvaPlaylistServerTests, Display, TEXT("Saving Image to %s"), *SaveName);
+
+			FFileHelper::SaveArrayToFile(InMessage.ImageData, *SaveName);
+		}
+
+		void HandleServerMessage(const FAvaPlaylistServerMsg& InMessage, const TSharedRef<IMessageContext, ESPMode::ThreadSafe>& InContext)
+		{
+			ReceivedRequestIds.Add(InMessage.RequestId);
+
+			const ELogVerbosity::Type Verbosity = ParseLogVerbosityFromString(InMessage.Verbosity);
+				
+			if (Verbosity <= ELogVerbosity::Error && Test)
+			{
+				Test->AddError(InMessage.Text, 1);
+			}
+		}
+	};
+
+	TSharedPtr<FAvaPlaylistServer> GetOrCreatePlaylistServer()
+	{
+		TSharedPtr<FAvaPlaylistServer> PlaylistServer = IAvaMediaEditorModule::Get().GetPlaylistServer();
+		if (!PlaylistServer)
+		{
+			UE_LOG(LogAvaPlaylistServerTests, Log, TEXT("Playlist Server not started. Starting one temporarily for the test."));
+
+			// Start a temporary server. It will be deleted when the last latent command is finished.
+			PlaylistServer = MakeShared<FAvaPlaylistServer>();
+			PlaylistServer->Init(TEXT(""));
+		}
+		return PlaylistServer;
+	}
+
+	void BackupBroadcastConfig()
+	{
+		const UAvalancheBroadcast& Broadcast = UAvalancheBroadcast::Get();
+		const FString SaveFilepath = Broadcast.GetBroadcastSaveFilepath();
+		const FString BackupFilepath = SaveFilepath + TEXT(".tests.backup");
+		IFileManager::Get().Copy(*BackupFilepath, *SaveFilepath);
+		UE_LOG(LogAvaPlaylistServerTests, Log, TEXT("Broadcast config backed up."));
+	}
+
+	void RestoreBroadcastConfig()
+	{
+		UAvalancheBroadcast& Broadcast = UAvalancheBroadcast::Get();
+		const FString SaveFilepath = Broadcast.GetBroadcastSaveFilepath();
+		const FString BackupFilepath = SaveFilepath + TEXT(".tests.backup");
+		if (!IFileManager::Get().FileExists(*BackupFilepath))
+		{
+			UE_LOG(LogAvaPlaylistServerTests, Error, TEXT("Failed to restore broadcast config: backup file \"%s\" doesn't exist."), *BackupFilepath);
+			return;
+		}
+
+		IFileManager::Get().Copy(*SaveFilepath, *BackupFilepath);
+		IFileManager::Get().Delete(*BackupFilepath);
+
+		// Reload broadcast from backup file.
+		Broadcast.StopBroadcast();
+		Broadcast.LoadBroadcast();
+
+		UE_LOG(LogAvaPlaylistServerTests, Log, TEXT("Broadcast config restored."));
+	}
+}
+
+DEFINE_LATENT_AUTOMATION_COMMAND_THREE_PARAMETER(FPlaylistServerWaitForResponse, int32, RequestId, TSharedPtr<UE::AvaPlaylistServerTests::FTestClient>, TestClient, TSharedPtr<FAvaPlaylistServer>, PlaylistServer);
+bool FPlaylistServerWaitForResponse::Update()
+{
+	if (!TestClient.IsValid())
+	{
+		UE_LOG(LogAvaPlaylistServerTests, Error, TEXT("Test Client is not valid."));
+		return true;
+	}
+	
+	return TestClient->HasReceivedResponse(RequestId);
+}
+
+DEFINE_LATENT_AUTOMATION_COMMAND(FPlaylistServerRestoreBroadcastConfig);
+bool FPlaylistServerRestoreBroadcastConfig::Update()
+{
+	using namespace UE::AvaPlaylistServerTests;
+	RestoreBroadcastConfig();
+	return true;
+}
+
+IMPLEMENT_SIMPLE_AUTOMATION_TEST(FPlaylistServerGetChannelImage, "Avalanche.PlaylistServer.GetChannelImage", (EAutomationTestFlags::EditorContext | EAutomationTestFlags::ProductFilter));
+bool FPlaylistServerGetChannelImage::RunTest(const FString& Parameters)
+{
+	using namespace UE::AvaPlaylistServerTests;
+	
+	const TSharedPtr<FAvaPlaylistServer> PlaylistServer = GetOrCreatePlaylistServer();
+	
+	// Build a test client for this request.
+	const TSharedPtr<FTestClient> TestClient = MakeShared<FTestClient>();
+	TestClient->Init(this);
+
+	// Send the request message.
+	constexpr int32 RequestId = 1000;
+	FAvaPlaylistGetChannelImage* Request = FMessageEndpoint::MakeMessage<FAvaPlaylistGetChannelImage>();
+	Request->RequestId = RequestId;
+	Request->ChannelName = UAvalancheBroadcast::Get().GetChannelName(0).ToString();
+	
+	TestClient->MessageEndpoint->Send(Request, PlaylistServer->GetMessageAddress());
+	
+	ADD_LATENT_AUTOMATION_COMMAND(FPlaylistServerWaitForResponse(RequestId, TestClient, PlaylistServer));
+	return true;
+}
+
+IMPLEMENT_CUSTOM_SIMPLE_AUTOMATION_TEST(FPlaylistServerAddChannelDevice, UE::AvaPlaylistServerTests::FPlaylistServerTestBase, "Avalanche.PlaylistServer.AddChannelDevice", (EAutomationTestFlags::EditorContext | EAutomationTestFlags::ProductFilter));
+bool FPlaylistServerAddChannelDevice::RunTest(const FString& Parameters)
+{
+	using namespace UE::AvaPlaylistServerTests;
+
+	// Need to backup the broadcast config because we are about to make changes.
+	BackupBroadcastConfig();
+	
+	const TSharedPtr<FAvaPlaylistServer> PlaylistServer = GetOrCreatePlaylistServer();
+
+	// Build a test client for this request.
+	const TSharedPtr<FTestClient> TestClient = MakeShared<FTestClient>();
+	TestClient->Init(this);
+
+	// Send the request message.
+	FAvaPlaylistAddChannelDevice* Request = FMessageEndpoint::MakeMessage<FAvaPlaylistAddChannelDevice>();
+	constexpr int32 RequestId = 1000;
+	Request->RequestId = RequestId;
+	Request->ChannelName = UAvalancheBroadcast::Get().GetChannelName(0).ToString();
+	Request->MediaOutputName = TEXT("InvalidDevice");
+
+	// Find an existing output so the request succeeds.
+	const FAvaOutputTreeItemPtr OutputDevices = MakeShared<FAvaOutputRootItem>();
+	FAvaOutputTreeItem::RefreshTree(OutputDevices);
+	
+	if (OutputDevices->GetChildren().Num() > 0)
+	{
+		FAvaOutputTreeItemPtr ExistingOutput = OutputDevices->GetChildren()[0];
+		while (ExistingOutput.IsValid() && ExistingOutput->GetChildren().Num() > 0)
+		{
+			ExistingOutput = ExistingOutput->GetChildren()[0];
+		}
+		if (ExistingOutput.IsValid())
+		{
+			Request->MediaOutputName = ExistingOutput->GetDisplayName().ToString();
+			UE_LOG(LogAvaPlaylistServerTests, Log, TEXT("Selected media output for test: \"%s\"."), *Request->MediaOutputName);
+		}
+	}
+	
+	TestClient->MessageEndpoint->Send(Request, PlaylistServer->GetMessageAddress());
+	
+	ADD_LATENT_AUTOMATION_COMMAND(FPlaylistServerWaitForResponse(RequestId, TestClient, PlaylistServer));
+	ADD_LATENT_AUTOMATION_COMMAND(FPlaylistServerRestoreBroadcastConfig());
+	return true;
+}
+
+
+IMPLEMENT_SIMPLE_AUTOMATION_TEST(FPlaylistServerMediaOutputSerialization, "Avalanche.PlaylistServer.MediaOutputSerialization", (EAutomationTestFlags::EditorContext | EAutomationTestFlags::ProductFilter));
+bool FPlaylistServerMediaOutputSerialization::RunTest(const FString& Parameters)
+{
+
+	for (UClass* const Class : TObjectRange<UClass>())
+	{
+		// For now, we test all the media output classes, technically, this should work with all of them.
+		// The result of this test depend on the plugins that are enabled. 
+		// The blackmagic plugin should be enabled at least.
+		const bool bIsMediaOutputClass = Class->IsChildOf(UMediaOutput::StaticClass()) && Class != UMediaOutput::StaticClass();
+		if (bIsMediaOutputClass)
+		{
+			UMediaOutput* const MediaOutput = NewObject<UMediaOutput>(GetTransientPackage(), Class, NAME_None, RF_Transactional);
+			FString MediaOutputJson = FAvaMediaOutputEditorUtils::SerializeMediaOutput(MediaOutput);
+
+			const FString SaveName = FString::Printf(TEXT("%sSerializationTest_%s.json"), *FPaths::ProjectSavedDir(), *Class->GetName());
+			UE_LOG(LogAvaPlaylistServerTests, Display, TEXT("Serializing Media Output \"%s\" to \"%s\""), *Class->GetName(), *SaveName);
+
+			TUniquePtr<FArchive> Archive(IFileManager::Get().CreateFileWriter(*SaveName));
+			*Archive << MediaOutputJson;
+			
+			Archive->Close();
+		}
+	}
+	return true;
+}
\ No newline at end of file
diff --git a/Engine/Plugins/Experimental/Avalanche/Source/AvalancheMediaEditor/Public/AvalancheMediaEditorSettings.h b/Engine/Plugins/Experimental/Avalanche/Source/AvalancheMediaEditor/Public/AvalancheMediaEditorSettings.h
new file mode 100644
--- /dev/null
+++ b/Engine/Plugins/Experimental/Avalanche/Source/AvalancheMediaEditor/Public/AvalancheMediaEditorSettings.h
@@ -0,0 +1,51 @@
+// Copyright Epic Games, Inc. All Rights Reserved.
+
+#pragma once
+
+#include "Engine/DeveloperSettings.h"
+#include "AvalancheMediaEditorSettings.generated.h"
+
+UCLASS(config=EditorPerProjectUserSettings, meta=(DisplayName="Playback & Broadcast"))
+class AVALANCHEMEDIAEDITOR_API UAvalancheMediaEditorSettings : public UDeveloperSettings
+{
+	GENERATED_BODY()
+
+public:
+	UAvalancheMediaEditorSettings();
+	
+	static const UAvalancheMediaEditorSettings& Get() {return *GetSingletonInstance();}
+	static UAvalancheMediaEditorSettings& GetMutable() {return *GetSingletonInstance();}
+
+	UPROPERTY(Config, EditAnywhere, Category = "Broadcast", meta = (InlineEditConditionToggle))
+	bool bBroadcastEnforceMaxChannelCount = true;
+	
+	UPROPERTY(Config, EditAnywhere, Category = "Broadcast", meta = (DisplayName = "Max Channel Count", EditCondition = "bBroadcastEnforceMaxChannelCount"))
+	int32 BroadcastMaxChannelCount = 9;
+
+	/**
+	 * By default, only the media output classes with a device provider are listed as output devices.
+	 * If this is set to true, then all the media output classes will be listed, thus allowing use of
+	 * outputs that don't correspond to a physical device, such as NDI, File output, etc.
+	 * Caution: enabling this mode may expose Media Output classes that are not usable.
+	 */
+	UPROPERTY(Config, EditAnywhere, Category = "Broadcast", meta = (DisplayName = "Show All Media Output Classes"))
+	bool bBroadcastShowAllMediaOutputClasses = false;
+	
+	UPROPERTY(Config, EditAnywhere, Category = "Playback", meta = (DisplayName = "Default Node Color"))
+	FLinearColor PlaybackDefaultNodeColor;
+
+	UPROPERTY(Config, EditAnywhere, Category = "Playback", meta = (DisplayName = "Channels Node Color"))
+	FLinearColor PlaybackChannelsNodeColor;
+	
+	UPROPERTY(Config, EditAnywhere, Category = "Playback", meta = (DisplayName = "Player Node Color"))
+	FLinearColor PlaybackPlayerNodeColor;
+	
+	UPROPERTY(Config, EditAnywhere, Category = "Playback", meta = (DisplayName = "Event Node Color"))
+	FLinearColor PlaybackEventNodeColor;
+
+	UPROPERTY(Config, EditAnywhere, Category = "Playback", meta = (DisplayName = "Action Node Color"))
+	FLinearColor PlaybackActionNodeColor;
+
+private:
+	static UAvalancheMediaEditorSettings* GetSingletonInstance();
+};
diff --git a/Engine/Plugins/Experimental/Avalanche/Source/AvalancheMediaEditor/Public/Broadcast/ChannelGrid/AvaOutputTileItem.h b/Engine/Plugins/Experimental/Avalanche/Source/AvalancheMediaEditor/Public/Broadcast/ChannelGrid/AvaOutputTileItem.h
new file mode 100644
--- /dev/null
+++ b/Engine/Plugins/Experimental/Avalanche/Source/AvalancheMediaEditor/Public/Broadcast/ChannelGrid/AvaOutputTileItem.h
@@ -0,0 +1,62 @@
+// Copyright Epic Games, Inc. All Rights Reserved.
+
+#pragma once
+
+#include "CoreMinimal.h"
+#include "Channel/AvaOutputChannel.h"
+
+struct FAvaOutputChannel;
+class FAvaOutputTileItem;
+class UMediaOutput;
+
+using FAvaOutputTileItemPtr = TSharedPtr<FAvaOutputTileItem>;
+
+class FAvaOutputTileItem : public TSharedFromThis<FAvaOutputTileItem>
+{
+public:
+
+	FAvaOutputTileItem(FName InChannelName, UMediaOutput* InMediaOutput);
+	~FAvaOutputTileItem();
+	
+	const FAvaOutputChannel& GetChannel() const;
+	FAvaOutputChannel& GetChannel();
+	
+	UMediaOutput* GetMediaOutput() const { return MediaOutput.Get(); }
+
+	TSharedRef<SWidget> GenerateTile() const;	
+
+	FText GetDisplayText() const;
+	FText GetMediaOutputStatusText() const;
+	FText GetToolTipText() const;
+
+	const FSlateBrush* GetMediaOutputIcon() const;
+	const FSlateBrush* GetMediaOutputStatusBrush() const;
+
+	FReply OnDragDetected(const FGeometry& MyGeometry, const FPointerEvent& MouseEvent);
+	
+protected:
+
+	void OnMediaOutputPropertyChanged(UObject* InObject, FPropertyChangedEvent& PropertyChangedEvent);
+	void OnChannelChanged(const FAvaOutputChannel& InChannel, EAvaChannelChange InChange);
+	void OnMediaOutputStateChanged(const FAvaOutputChannel& InChannel, const UMediaOutput* InMediaOutput);
+	void OnBroadcastChanged(EAvaBroadcastChange InChange);
+
+	void UpdateInfo();
+	
+	FText FindLatestDisplayText() const;
+
+protected:
+	
+	FName ChannelName = NAME_None;
+
+	TWeakObjectPtr<UMediaOutput> MediaOutput;
+	
+	FText MediaOutputDisplayText;
+	FText MediaOutputStatusText;
+	FText MediaOutputToolTipText;
+	
+	const FSlateBrush* MediaOutputStatusBrush = nullptr;
+
+	FDelegateHandle BroadcastChangedHandle;
+};
+
diff --git a/Engine/Plugins/Experimental/Avalanche/Source/AvalancheMediaEditor/Public/IAvaMediaEditorModule.h b/Engine/Plugins/Experimental/Avalanche/Source/AvalancheMediaEditor/Public/IAvaMediaEditorModule.h
new file mode 100644
--- /dev/null
+++ b/Engine/Plugins/Experimental/Avalanche/Source/AvalancheMediaEditor/Public/IAvaMediaEditorModule.h
@@ -0,0 +1,95 @@
+// Copyright Epic Games, Inc. All Rights Reserved.
+
+#pragma once
+
+#include "CoreMinimal.h"
+#include "Modules/ModuleInterface.h"
+#include "Modules/ModuleManager.h"
+#include "Toolkits/AssetEditorToolkit.h"
+
+DECLARE_LOG_CATEGORY_EXTERN(LogAvaMediaEditor, Log, All);
+
+class FAvaPlaylistServer;
+class IAvaPlaylistFilterSuggestionFactory;
+struct FAvalanchePage;
+struct FAvaPlaylistTextFilterArgs;
+enum class EAvaPlaylistSearchListType : uint8;
+
+class IAvaMediaEditorModule : public IModuleInterface
+{
+	static constexpr const TCHAR* ModuleName = TEXT("AvalancheMediaEditor");
+
+public:
+	static bool IsLoaded()
+	{
+		return FModuleManager::Get().IsModuleLoaded(ModuleName);
+	}
+
+	static IAvaMediaEditorModule& Get()
+	{
+		return FModuleManager::LoadModuleChecked<IAvaMediaEditorModule>(ModuleName);
+	}
+
+	/** Returns the tool menu name used for Page Context Menu */
+	static FName GetPlaylistPageMenuName()
+	{
+		return TEXT("AvaPlaylistPageContextMenu");
+	}
+
+	virtual FSlateIcon GetToolbarBroadcastButtonIcon() const = 0;
+
+	/** Returns the toolbar extensibility manager for the Broadcast Editor */
+	virtual TSharedPtr<FExtensibilityManager> GetBroadcastToolBarExtensibilityManager() = 0;
+
+	/** Returns the toolbar extensibility manager for the Playback Editor */
+	virtual TSharedPtr<FExtensibilityManager> GetPlaybackToolBarExtensibilityManager() = 0;
+
+	/** Returns the toolbar extensibility manager for the Playback Editor */
+	virtual TSharedPtr<FExtensibilityManager> GetPlaylistToolBarExtensibilityManager() = 0;
+
+	/**
+	 * Returns the context menu extensibility manager for the Playlist Editor's Template Pages
+	 * @remark prefer extending with the UToolMenu named after IAvaMediaEditorModule::GetPlaylistPageMenuName, and using UAvaPlaylistPageContext to retrieve context information
+	 */
+	virtual TSharedPtr<FExtensibilityManager> GetPlaylistMenuExtensibilityManager() = 0;
+
+	DECLARE_MULTICAST_DELEGATE(FOnPlaylistServerStarted);
+	virtual FOnPlaylistServerStarted& GetOnPlaylistServerStarted() = 0;
+
+	DECLARE_MULTICAST_DELEGATE(FOnPlaylistServerStopped);
+	virtual FOnPlaylistServerStopped& GetOnPlaylistServerStopped() = 0;
+
+	virtual TSharedPtr<FAvaPlaylistServer> GetPlaylistServer() const = 0;
+
+	/**
+ 	* Check if current playlist filter expression factory support the comparison operation
+ 	* @param InFilterKey Filter key to get the playlist filter expression factory needed
+ 	* @param InOperation Operation to check if supported
+ 	* @param InPlaylistSearchListType Type of the Search List either Template or Instanced
+ 	* @return True if operation is supported, False otherwise
+ 	*/
+	virtual bool CanFilterSupportComparisonOperation(const FName& InFilterKey, ETextFilterComparisonOperation InOperation, EAvaPlaylistSearchListType InPlaylistSearchListType) const = 0;
+
+	/**
+ 	* Evaluate the expression and return the result
+ 	* @param InFilterKey Filter Key to get the Factory
+ 	* @param InItem Item that is currently checked
+ 	* @param InArgs Args to evaluate the expression see FAvaPlaylistTextFilterArgs for more information
+ 	* @return True if the expression evaluated to True, False otherwise
+ 	*/
+	virtual bool FilterExpression(const FName& InFilterKey, const FAvalanchePage& InItem, const FAvaPlaylistTextFilterArgs& InArgs) const = 0;
+
+	/**
+ 	* Get all simple suggestions with the given type (Template/Instanced/All)
+ 	* @param InSuggestionType Type of suggestion to get, see EAvaPlaylistSearchListType for more information
+ 	* @return An Array containing all suggestions of the given type
+ 	*/
+	virtual TArray<TSharedPtr<IAvaPlaylistFilterSuggestionFactory>> GetSimpleSuggestions(EAvaPlaylistSearchListType InSuggestionType) const = 0;
+
+	/**
+ 	* Get all complex suggestions with the given type (Template/Instanced/All)
+ 	* @param InSuggestionType Type of suggestion to get, see EAvaPlaylistSearchListType for more information
+ 	* @return An Array containing all suggestions of the given type
+ 	*/
+	virtual TArray<TSharedPtr<IAvaPlaylistFilterSuggestionFactory>> GetComplexSuggestions(EAvaPlaylistSearchListType InSuggestionType) const = 0;
+};
diff --git a/Engine/Plugins/Experimental/Avalanche/Source/AvalancheMediaEditor/Public/Playback/AvaPlaybackEditor.h b/Engine/Plugins/Experimental/Avalanche/Source/AvalancheMediaEditor/Public/Playback/AvaPlaybackEditor.h
new file mode 100644
--- /dev/null
+++ b/Engine/Plugins/Experimental/Avalanche/Source/AvalancheMediaEditor/Public/Playback/AvaPlaybackEditor.h
@@ -0,0 +1,124 @@
+// Copyright Epic Games, Inc. All Rights Reserved.
+
+#pragma once
+
+#include "CoreMinimal.h"
+#include "Playback/IAvaPlaybackEditor.h"
+#include "Templates/SharedPointer.h"
+#include "WorkflowOrientedApp/WorkflowCentricApplication.h"
+
+class FSlateRect;
+class SGraphEditor;
+class UAvaPlaybackGraphNode;
+
+class FAvaPlaybackEditor
+	: public FWorkflowCentricApplication
+	, public IAvaPlaybackGraphEditor
+{
+public:
+
+	void InitPlaybackEditor(const EToolkitMode::Type InMode
+		, const TSharedPtr<IToolkitHost>& InitToolkitHost
+		, UAvalanchePlayback* InPlayback);
+
+	//IToolkit Interface
+	virtual FName GetToolkitFName() const override;
+	virtual FText GetBaseToolkitName() const override;
+	virtual FString GetWorldCentricTabPrefix() const override;
+	virtual FLinearColor GetWorldCentricTabColorScale() const override;
+	//~IToolkit Interface
+
+	UAvalanchePlayback* GetPlaybackObject() const;
+	
+	void ExtendToolBar(TSharedPtr<FExtender> Extender);
+	void FillPlayToolBar(FToolBarBuilder& ToolBarBuilder);
+
+public:
+
+	static void CacheOverrideNodeClasses();
+	
+	TSharedRef<SGraphEditor> CreateGraphEditor();
+	
+	//IAvaPlaybackGraphEditor Interface
+	virtual UEdGraph* CreatePlaybackGraph(UAvalanchePlayback* InPlayback) override;
+	virtual void SetupPlaybackNode(UEdGraph* InGraph, UAvaPlaybackNode* InPlaybackNode, bool bSelectNewNode) override;
+	virtual void CompilePlaybackNodesFromGraphNodes(UAvalanchePlayback* InPlayback) override;
+	virtual void CreateInputPin(UEdGraphNode* InGraphNode) override;
+	virtual void RefreshNode(UEdGraphNode& InGraphNode) override;
+	virtual bool GetBoundsForSelectedNodes(FSlateRect& Rect, float Padding) override;
+	virtual TSet<UObject*> GetSelectedNodes() const override;
+	//~IAvaPlaybackGraphEditor Interface
+	
+	/** Called when the selection changes in the GraphEditor */
+	void OnSelectedNodesChanged(const TSet<UObject*>& NewSelection);
+
+	/**
+	 * Called when a node's title is committed for a rename
+	 *
+	 * @param	NewText				New title text
+	 * @param	CommitInfo			How text was committed
+	 * @param	NodeBeingChanged	The node being changed
+	 */
+	void OnNodeTitleCommitted(const FText& NewText, ETextCommit::Type CommitInfo, UEdGraphNode* NodeBeingChanged);
+	
+	DECLARE_MULTICAST_DELEGATE_OneParam(FOnPlaybackSelectionChanged, const TArray<UObject*>&);
+	FOnPlaybackSelectionChanged OnPlaybackSelectionChanged;
+	
+protected:
+	
+	void RegisterApplicationModes();
+	
+	void CreateDefaultCommands();
+	void CreateGraphCommands();
+
+public:
+	
+	bool CanAddInputPin() const;
+	void AddInputPin();
+
+	bool CanRemoveInputPin() const;
+	void RemoveInputPin();
+	
+	void CreateComment();
+	
+	bool CanSelectAllNodes() const { return true; }
+	void SelectAllNodes();
+
+	bool CanDeleteSelectedNodes() const;
+	void DeleteSelectedNodes();
+	void DeleteSelectedDuplicatableNodes();
+	
+	bool CanCopySelectedNodes() const;
+	void CopySelectedNodes();
+	
+	bool CanCutSelectedNodes() const;
+	void CutSelectedNodes();
+	
+	virtual bool CanPasteNodes() const override;
+	void PasteNodes();
+	virtual void PasteNodesHere(const FVector2D& Location) override;
+	
+	bool CanDuplicateSelectedNodes() const;
+	void DuplicateSelectedNodes();
+
+	void OnAlignTop();
+	void OnAlignMiddle();
+	void OnAlignBottom();
+	void OnAlignLeft();
+	void OnAlignCenter();
+	void OnAlignRight();
+
+	void OnStraightenConnections();
+	void OnDistributeNodesH();
+	void OnDistributeNodesV();
+	
+protected:
+	
+	TWeakObjectPtr<UAvalanchePlayback> AvalanchePlayback;
+	
+	TSharedPtr<FUICommandList> GraphEditorCommands;
+
+	TSharedPtr<SGraphEditor> GraphEditor;
+	
+	static TMap<TSubclassOf<UAvaPlaybackNode>, TSubclassOf<UAvaPlaybackGraphNode>> OverrideNodeClasses;
+};
diff --git a/Engine/Plugins/Experimental/Avalanche/Source/AvalancheMediaEditor/Public/Playback/AvaPlaybackFactory.h b/Engine/Plugins/Experimental/Avalanche/Source/AvalancheMediaEditor/Public/Playback/AvaPlaybackFactory.h
new file mode 100644
--- /dev/null
+++ b/Engine/Plugins/Experimental/Avalanche/Source/AvalancheMediaEditor/Public/Playback/AvaPlaybackFactory.h
@@ -0,0 +1,30 @@
+// Copyright Epic Games, Inc. All Rights Reserved.
+
+#pragma once
+
+#include "CoreMinimal.h"
+#include "Factories/Factory.h"
+#include "AvaPlaybackFactory.generated.h"
+
+UCLASS()
+class AVALANCHEMEDIAEDITOR_API UAvaPlaybackFactory : public UFactory
+{
+	GENERATED_BODY()
+
+public:
+
+	UAvaPlaybackFactory();
+	virtual ~UAvaPlaybackFactory() override;
+
+protected:
+	
+	//~ Begin UFactory Interface
+	virtual uint32 GetMenuCategories() const override;
+	virtual bool ConfigureProperties() override;
+	virtual bool ShouldShowInNewMenu() const override;
+	virtual UObject* FactoryCreateNew(UClass* Class, UObject* InParent, FName Name, EObjectFlags Flags, UObject* Context, FFeedbackContext* Warn, FName CallingContext) override;
+	virtual UObject* FactoryCreateNew(UClass* Class, UObject* InParent, FName Name, EObjectFlags Flags, UObject* Context, FFeedbackContext* Warn) override;
+	virtual bool DoesSupportClass(UClass* Class) override;
+	virtual UClass* ResolveSupportedClass() override;
+	//~ Begin UFactory Interface
+};
diff --git a/Engine/Plugins/Experimental/Avalanche/Source/AvalancheMediaEditor/Public/Playback/Graph/AvaPlaybackGraph.h b/Engine/Plugins/Experimental/Avalanche/Source/AvalancheMediaEditor/Public/Playback/Graph/AvaPlaybackGraph.h
new file mode 100644
--- /dev/null
+++ b/Engine/Plugins/Experimental/Avalanche/Source/AvalancheMediaEditor/Public/Playback/Graph/AvaPlaybackGraph.h
@@ -0,0 +1,23 @@
+// Copyright Epic Games, Inc. All Rights Reserved.
+
+#pragma once
+
+#include "CoreMinimal.h"
+#include "EdGraph/EdGraph.h"
+#include "AvaPlaybackGraph.generated.h"
+
+class UAvalanchePlayback;
+class UAvaPlaybackGraphNode;
+
+UCLASS()
+class UAvaPlaybackGraph : public UEdGraph
+{
+	GENERATED_BODY()
+
+public:
+	
+	UAvalanchePlayback* GetPlayback() const;
+	
+	UAvaPlaybackGraphNode* CreatePlaybackGraphNode(TSubclassOf<UAvaPlaybackGraphNode> NewNodeClass
+		, bool bSelectNewNode = true);
+};
diff --git a/Engine/Plugins/Experimental/Avalanche/Source/AvalancheMediaEditor/Public/Playback/Graph/AvaPlaybackGraphSchema.h b/Engine/Plugins/Experimental/Avalanche/Source/AvalancheMediaEditor/Public/Playback/Graph/AvaPlaybackGraphSchema.h
new file mode 100644
--- /dev/null
+++ b/Engine/Plugins/Experimental/Avalanche/Source/AvalancheMediaEditor/Public/Playback/Graph/AvaPlaybackGraphSchema.h
@@ -0,0 +1,65 @@
+// Copyright Epic Games, Inc. All Rights Reserved.
+
+#pragma once
+
+#include "CoreMinimal.h"
+#include "EdGraph/EdGraphSchema.h"
+#include "AvaPlaybackGraphSchema.generated.h"
+
+class IAvaPlaybackGraphEditor;
+class UAvalanchePlayback;
+class UAvaPlaybackNode;
+
+UCLASS()
+class UAvaPlaybackGraphSchema : public UEdGraphSchema
+{
+	GENERATED_BODY()
+
+protected:
+
+	TSharedPtr<IAvaPlaybackGraphEditor> GetPlaybackGraphEditor(const UEdGraph* Graph) const;
+
+public:
+
+	static const FLinearColor ActivePinColor;
+	static const FLinearColor InactivePinColor;
+	
+	// Allowable PinType.PinCategory values
+	AVALANCHEMEDIAEDITOR_API static const FName PC_ChannelFeed;
+	AVALANCHEMEDIAEDITOR_API static const FName PC_Event;
+
+	void CompilePlaybackNodesFromGraphNodes(UEdGraphNode& Node) const;
+	
+	/** Check whether connecting these pins would cause a loop */
+	bool ConnectionCausesLoop(const UEdGraphPin* InputPin, const UEdGraphPin* OutputPin) const;
+	
+	//UEdGraphSchema Interface.
+	virtual void GetGraphContextActions(FGraphContextMenuBuilder& ContextMenuBuilder) const override;
+	virtual void GetContextMenuActions(UToolMenu* Menu, class UGraphNodeContextMenuContext* Context) const override;
+	virtual void CreateDefaultNodesForGraph(UEdGraph& Graph) const override;
+	virtual bool CanDuplicateGraph(UEdGraph* InSourceGraph) const override { return false; }
+	virtual bool ShouldAlwaysPurgeOnModification() const override { return true; }
+	virtual const FPinConnectionResponse CanCreateConnection(const UEdGraphPin* PinA, const UEdGraphPin* PinB) const override;
+	virtual bool TryCreateConnection(UEdGraphPin* PinA, UEdGraphPin* PinB) const override;
+	virtual FLinearColor GetPinTypeColor(const FEdGraphPinType& PinType) const override;
+	virtual bool ShouldHidePinDefaultValue(UEdGraphPin* Pin) const override;
+	virtual void BreakNodeLinks(UEdGraphNode& TargetNode) const override;
+	virtual void BreakPinLinks(UEdGraphPin& TargetPin, bool bSendsNodeNotifcation) const override;
+	virtual void GetAssetsGraphHoverMessage(const TArray<FAssetData>& Assets, const UEdGraph* HoverGraph, FString& OutTooltipText, bool& OutOkIcon) const override;
+	virtual void DroppedAssetsOnGraph(const TArray<struct FAssetData>& Assets, const FVector2D& GraphPosition, UEdGraph* Graph) const override;
+	virtual void DroppedAssetsOnNode(const TArray<FAssetData>& Assets, const FVector2D& GraphPosition, UEdGraphNode* Node) const override;
+	virtual int32 GetNodeSelectionCount(const UEdGraph* Graph) const override;
+	virtual TSharedPtr<FEdGraphSchemaAction> GetCreateCommentAction() const override;
+	//~UEdGraphSchema
+
+	static void CachePlaybackNodeClasses();
+	
+	/** Adds actions for creating every type of Playback Node */
+	void GetPlaybackNodeActions(FGraphActionMenuBuilder& ActionMenuBuilder, bool bShowSelectedActions) const;
+	void GetCommentAction(FGraphActionMenuBuilder& ActionMenuBuilder, const UEdGraph* CurrentGraph = nullptr) const;
+
+protected:
+
+	static TArray<TSubclassOf<UAvaPlaybackNode>> PlaybackNodeClasses;
+	
+};
diff --git a/Engine/Plugins/Experimental/Avalanche/Source/AvalancheMediaEditor/Public/Playback/Graph/Nodes/AvaPlaybackGraphNode.h b/Engine/Plugins/Experimental/Avalanche/Source/AvalancheMediaEditor/Public/Playback/Graph/Nodes/AvaPlaybackGraphNode.h
new file mode 100644
--- /dev/null
+++ b/Engine/Plugins/Experimental/Avalanche/Source/AvalancheMediaEditor/Public/Playback/Graph/Nodes/AvaPlaybackGraphNode.h
@@ -0,0 +1,108 @@
+// Copyright Epic Games, Inc. All Rights Reserved.
+
+#pragma once
+
+#include "CoreMinimal.h"
+#include "EdGraph/EdGraphNode.h"
+#include "Templates/SubclassOf.h"
+#include "AvaPlaybackGraphNode.generated.h"
+
+class UAvalanchePlayback;
+class UAvaPlaybackNode;
+
+UCLASS()
+class AVALANCHEMEDIAEDITOR_API UAvaPlaybackGraphNode : public UEdGraphNode
+{
+	GENERATED_BODY()
+
+public:
+
+	UAvalanchePlayback* GetPlayback() const;
+	
+	void SetPlaybackNode(UAvaPlaybackNode* InPlaybackNode);
+
+	double GetLastTimeTicked();
+	double GetChildLastTimeTicked(int32 ChildIndex) const;
+	
+	bool CanAddInputPin() const;
+	void AddInputPin();
+	void CreateInputPin();
+
+	void RemoveInputPin(UEdGraphPin* InGraphPin);
+
+	//~ Begin UEdGraphNode interface
+	virtual void PrepareForCopying() override;
+	virtual void PostPasteNode() override;
+	//~ End UEdGraphNode interface
+
+	void PostCopyNode();
+
+	/** Make sure the PlaybackNode is owned by the Playback Object */
+	void ResetPlaybackNodeOwner();
+	
+	//Return which Playback Node this Graph Node should Represent
+	virtual TSubclassOf<UAvaPlaybackNode> GetPlaybackNodeClass() const;
+
+	UAvaPlaybackNode* GetPlaybackNode() const
+	{
+		return PlaybackNode;
+	}
+	
+protected:
+	
+	/** Create all of the input pins required */
+	void CreateInputPins();
+	void CreateOutputPin();
+
+public:
+	
+	/** Is this the undeletable root node */
+	bool IsRootNode() const;
+	
+	UEdGraphPin* GetOutputPin() const;
+	TArray<UEdGraphPin*> GetInputPins() const;
+	
+	UEdGraphPin* GetInputPin(int32 InputIndex) const;
+	int32 GetInputPinIndex(UEdGraphPin* InInputPin) const;
+	int32 GetInputPinCount() const;
+
+	/**
+	 * Handles inserting the node between the FromPin and what the FromPin was original connected to
+	 *
+	 * @param FromPin			The pin this node is being spawned from
+	 * @param NewLinkPin		The new pin the FromPin will connect to
+	 * @param OutNodeList		Any nodes that are modified will get added to this list for notification purposes
+	 */
+	void InsertNewNode(UEdGraphPin* FromPin, UEdGraphPin* NewLinkPin, TSet<UEdGraphNode*>& OutNodeList);
+
+	//UObject interface
+	virtual void PostLoad() override;
+	virtual void PostEditImport() override;
+	virtual void PostDuplicate(bool bDuplicateForPIE) override;
+	//~UObject interface
+	
+	//UEdGraphNode
+	virtual FText GetNodeTitle(ENodeTitleType::Type TitleType) const override;
+	virtual FText GetTooltipText() const override;
+	virtual void GetNodeContextMenuActions(UToolMenu* Menu, UGraphNodeContextMenuContext* Context) const override;
+	virtual void AllocateDefaultPins() override;
+	virtual void ReconstructNode() override;
+	virtual void AutowireNewNode(UEdGraphPin* FromPin) override;
+	virtual bool CanCreateUnderSpecifiedSchema(const UEdGraphSchema* Schema) const override;
+	virtual FLinearColor GetNodeTitleColor() const override;
+	virtual TSharedPtr<SGraphNode> CreateVisualWidget() override;
+	//~UEdGraphNode
+	
+	virtual FName GetInputPinName(int32 InputPinIndex) const;
+	virtual FName GetInputPinCategory(int32 InputPinIndex) const;
+	virtual FName GetInputPinSubCategory(int32 InputPinIndex) const;
+	virtual UObject* GetInputPinSubCategoryObject(int32 InputPinIndex) const;
+
+	virtual FName GetOutputPinCategory() const;
+	
+protected:
+	
+	// The Playback Node this represents
+	UPROPERTY(VisibleAnywhere, Instanced, Category=Playback)
+	TObjectPtr<UAvaPlaybackNode> PlaybackNode;
+};
diff --git a/Engine/Plugins/Experimental/Avalanche/Source/AvalancheMediaEditor/Public/Playback/Graph/Nodes/AvaPlaybackGraphNode_BlueprintPlayer.h b/Engine/Plugins/Experimental/Avalanche/Source/AvalancheMediaEditor/Public/Playback/Graph/Nodes/AvaPlaybackGraphNode_BlueprintPlayer.h
new file mode 100644
--- /dev/null
+++ b/Engine/Plugins/Experimental/Avalanche/Source/AvalancheMediaEditor/Public/Playback/Graph/Nodes/AvaPlaybackGraphNode_BlueprintPlayer.h
@@ -0,0 +1,20 @@
+// Copyright Epic Games, Inc. All Rights Reserved.
+
+#pragma once
+
+#include "AvaPlaybackGraphNode.h"
+#include "AvaPlaybackGraphNode_BlueprintPlayer.generated.h"
+
+UCLASS()
+class AVALANCHEMEDIAEDITOR_API UAvaPlaybackGraphNode_BlueprintPlayer : public UAvaPlaybackGraphNode
+{
+	GENERATED_BODY()
+	
+	virtual TSubclassOf<UAvaPlaybackNode> GetPlaybackNodeClass() const override;
+	virtual FName GetInputPinCategory(int32 InputPinIndex) const override;
+	
+	//UEdGraphNode interface
+	virtual FLinearColor GetNodeTitleColor() const override;
+	virtual TSharedPtr<SGraphNode> CreateVisualWidget() override;
+	//~UEdGraphNode interface
+};
diff --git a/Engine/Plugins/Experimental/Avalanche/Source/AvalancheMediaEditor/Public/Playback/Graph/Nodes/AvaPlaybackGraphNode_EventAction.h b/Engine/Plugins/Experimental/Avalanche/Source/AvalancheMediaEditor/Public/Playback/Graph/Nodes/AvaPlaybackGraphNode_EventAction.h
new file mode 100644
--- /dev/null
+++ b/Engine/Plugins/Experimental/Avalanche/Source/AvalancheMediaEditor/Public/Playback/Graph/Nodes/AvaPlaybackGraphNode_EventAction.h
@@ -0,0 +1,23 @@
+// Copyright Epic Games, Inc. All Rights Reserved.
+
+#pragma once
+
+#include "AvaPlaybackGraphNode.h"
+#include "AvaPlaybackGraphNode_EventAction.generated.h"
+
+UCLASS()
+class AVALANCHEMEDIAEDITOR_API UAvaPlaybackGraphNode_EventAction : public UAvaPlaybackGraphNode
+{
+	GENERATED_BODY()
+
+public:
+	
+	virtual TSubclassOf<UAvaPlaybackNode> GetPlaybackNodeClass() const override;
+
+	virtual FName GetInputPinCategory(int32 InputPinIndex) const override;
+	virtual FName GetOutputPinCategory() const override;
+
+	//UEdGraphNode interface
+	virtual FLinearColor GetNodeTitleColor() const override;
+	//~UEdGraphNode interface
+};
diff --git a/Engine/Plugins/Experimental/Avalanche/Source/AvalancheMediaEditor/Public/Playback/Graph/Nodes/AvaPlaybackGraphNode_EventFlow.h b/Engine/Plugins/Experimental/Avalanche/Source/AvalancheMediaEditor/Public/Playback/Graph/Nodes/AvaPlaybackGraphNode_EventFlow.h
new file mode 100644
--- /dev/null
+++ b/Engine/Plugins/Experimental/Avalanche/Source/AvalancheMediaEditor/Public/Playback/Graph/Nodes/AvaPlaybackGraphNode_EventFlow.h
@@ -0,0 +1,19 @@
+// Copyright Epic Games, Inc. All Rights Reserved.
+
+#pragma once
+
+#include "AvaPlaybackGraphNode.h"
+#include "AvaPlaybackGraphNode_EventFlow.generated.h"
+
+UCLASS()
+class AVALANCHEMEDIAEDITOR_API UAvaPlaybackGraphNode_EventFlow : public UAvaPlaybackGraphNode
+{
+	GENERATED_BODY()
+
+public:
+	
+	virtual TSubclassOf<UAvaPlaybackNode> GetPlaybackNodeClass() const override;
+
+	virtual FName GetInputPinCategory(int32 InputPinIndex) const override;
+	virtual FName GetOutputPinCategory() const override;
+};
diff --git a/Engine/Plugins/Experimental/Avalanche/Source/AvalancheMediaEditor/Public/Playback/Graph/Nodes/AvaPlaybackGraphNode_EventTrigger.h b/Engine/Plugins/Experimental/Avalanche/Source/AvalancheMediaEditor/Public/Playback/Graph/Nodes/AvaPlaybackGraphNode_EventTrigger.h
new file mode 100644
--- /dev/null
+++ b/Engine/Plugins/Experimental/Avalanche/Source/AvalancheMediaEditor/Public/Playback/Graph/Nodes/AvaPlaybackGraphNode_EventTrigger.h
@@ -0,0 +1,21 @@
+// Copyright Epic Games, Inc. All Rights Reserved.
+
+#pragma once
+
+#include "AvaPlaybackGraphNode.h"
+#include "AvaPlaybackGraphNode_EventTrigger.generated.h"
+
+UCLASS()
+class AVALANCHEMEDIAEDITOR_API UAvaPlaybackGraphNode_EventTrigger : public UAvaPlaybackGraphNode
+{
+	GENERATED_BODY()
+
+public:
+	
+	virtual TSubclassOf<UAvaPlaybackNode> GetPlaybackNodeClass() const override;
+	virtual FName GetOutputPinCategory() const override;
+
+	//UEdGraphNode interface
+	virtual FLinearColor GetNodeTitleColor() const override;
+	//~UEdGraphNode interface
+};
diff --git a/Engine/Plugins/Experimental/Avalanche/Source/AvalancheMediaEditor/Public/Playback/Graph/Nodes/AvaPlaybackGraphNode_LevelPlayer.h b/Engine/Plugins/Experimental/Avalanche/Source/AvalancheMediaEditor/Public/Playback/Graph/Nodes/AvaPlaybackGraphNode_LevelPlayer.h
new file mode 100644
--- /dev/null
+++ b/Engine/Plugins/Experimental/Avalanche/Source/AvalancheMediaEditor/Public/Playback/Graph/Nodes/AvaPlaybackGraphNode_LevelPlayer.h
@@ -0,0 +1,20 @@
+// Copyright Epic Games, Inc. All Rights Reserved.
+
+#pragma once
+
+#include "AvaPlaybackGraphNode.h"
+#include "AvaPlaybackGraphNode_LevelPlayer.generated.h"
+
+UCLASS()
+class AVALANCHEMEDIAEDITOR_API UAvaPlaybackGraphNode_LevelPlayer : public UAvaPlaybackGraphNode
+{
+	GENERATED_BODY()
+	
+	virtual TSubclassOf<UAvaPlaybackNode> GetPlaybackNodeClass() const override;
+	virtual FName GetInputPinCategory(int32 InputPinIndex) const override;
+	
+	//UEdGraphNode interface
+	virtual FLinearColor GetNodeTitleColor() const override;
+	virtual TSharedPtr<SGraphNode> CreateVisualWidget() override;
+	//~UEdGraphNode interface
+};
diff --git a/Engine/Plugins/Experimental/Avalanche/Source/AvalancheMediaEditor/Public/Playback/Graph/Nodes/AvaPlaybackGraphNode_Root.h b/Engine/Plugins/Experimental/Avalanche/Source/AvalancheMediaEditor/Public/Playback/Graph/Nodes/AvaPlaybackGraphNode_Root.h
new file mode 100644
--- /dev/null
+++ b/Engine/Plugins/Experimental/Avalanche/Source/AvalancheMediaEditor/Public/Playback/Graph/Nodes/AvaPlaybackGraphNode_Root.h
@@ -0,0 +1,21 @@
+// Copyright Epic Games, Inc. All Rights Reserved.
+
+#pragma once
+
+#include "AvaPlaybackGraphNode.h"
+#include "AvaPlaybackGraphNode_Root.generated.h"
+
+UCLASS()
+class AVALANCHEMEDIAEDITOR_API UAvaPlaybackGraphNode_Root : public UAvaPlaybackGraphNode
+{
+	GENERATED_BODY()
+	
+	virtual TSubclassOf<UAvaPlaybackNode> GetPlaybackNodeClass() const override;
+
+	//UEdGraphNode interface
+	virtual FLinearColor GetNodeTitleColor() const override;
+	virtual bool CanUserDeleteNode() const override;
+	virtual bool CanDuplicateNode() const override;
+	virtual TSharedPtr<SGraphNode> CreateVisualWidget() override;
+	//~UEdGraphNode interface
+};
diff --git a/Engine/Plugins/Experimental/Avalanche/Source/AvalancheMediaEditor/Public/Playback/Graph/Nodes/Slate/SAvaPlaybackGraphNode.h b/Engine/Plugins/Experimental/Avalanche/Source/AvalancheMediaEditor/Public/Playback/Graph/Nodes/Slate/SAvaPlaybackGraphNode.h
new file mode 100644
--- /dev/null
+++ b/Engine/Plugins/Experimental/Avalanche/Source/AvalancheMediaEditor/Public/Playback/Graph/Nodes/Slate/SAvaPlaybackGraphNode.h
@@ -0,0 +1,34 @@
+// Copyright Epic Games, Inc. All Rights Reserved.
+
+#pragma once
+
+#include "CoreMinimal.h"
+#include "SGraphNode.h"
+
+class SVerticalBox;
+class UAvaPlaybackGraphNode;
+
+class AVALANCHEMEDIAEDITOR_API SAvaPlaybackGraphNode : public SGraphNode
+{
+public:
+	
+	SLATE_BEGIN_ARGS(SAvaPlaybackGraphNode) {}
+	SLATE_END_ARGS()
+
+	/** Constructs this widget with InArgs */
+	void Construct(const FArguments& InArgs, UAvaPlaybackGraphNode* InGraphNode);
+	virtual void PostConstruct() {};
+	
+protected:
+	
+	//SGraphNode Interface
+	virtual void CreateOutputSideAddButton(TSharedPtr<SVerticalBox> OutputBox) override;
+	virtual EVisibility IsAddPinButtonVisible() const override;
+	virtual FReply OnAddPin() override;
+	virtual TSharedPtr<SGraphPin> CreatePinWidget(UEdGraphPin* Pin) const override;
+	//~SGraphNode Interface
+	
+protected:
+
+	TWeakObjectPtr<UAvaPlaybackGraphNode> PlaybackGraphNode;
+};
diff --git a/Engine/Plugins/Experimental/Avalanche/Source/AvalancheMediaEditor/Public/Playback/Graph/SchemaActions/AvaPlaybackAction_NewNode.h b/Engine/Plugins/Experimental/Avalanche/Source/AvalancheMediaEditor/Public/Playback/Graph/SchemaActions/AvaPlaybackAction_NewNode.h
new file mode 100644
--- /dev/null
+++ b/Engine/Plugins/Experimental/Avalanche/Source/AvalancheMediaEditor/Public/Playback/Graph/SchemaActions/AvaPlaybackAction_NewNode.h
@@ -0,0 +1,47 @@
+// Copyright Epic Games, Inc. All Rights Reserved.
+
+#pragma once
+
+#include "CoreMinimal.h"
+#include "EdGraph/EdGraphSchema.h"
+#include "Templates/SubclassOf.h"
+#include "AvaPlaybackAction_NewNode.generated.h"
+
+class UAvaPlaybackNode;
+
+USTRUCT()
+struct AVALANCHEMEDIAEDITOR_API FAvaPlaybackAction_NewNode : public FEdGraphSchemaAction
+{
+	
+	GENERATED_BODY()
+
+public:
+	
+	FAvaPlaybackAction_NewNode() = default;
+	
+	FAvaPlaybackAction_NewNode(FText InNodeCategory, FText InMenuDesc, FText InToolTip, const int32 InGrouping)
+		: FEdGraphSchemaAction(MoveTemp(InNodeCategory), MoveTemp(InMenuDesc), MoveTemp(InToolTip), InGrouping)
+	{
+	}
+
+	//~ Begin FEdGraphSchemaAction Interface
+	virtual UEdGraphNode* PerformAction(class UEdGraph* ParentGraph
+		, UEdGraphPin* FromPin
+		, const FVector2D Location
+		, bool bSelectNewNode = true) override;
+	//~ End FEdGraphSchemaAction Interface
+
+	void SetPlaybackNodeClass(TSubclassOf<UAvaPlaybackNode> InPlaybackNodeClass);
+
+protected:
+	
+	/** Connects new node to output of selected nodes */
+	void ConnectToSelectedNodes(UAvaPlaybackNode* NewNode, UEdGraph* ParentGraph) const;
+
+protected:
+	
+	/** Class of node we want to create */
+	UPROPERTY()
+	TSubclassOf<UAvaPlaybackNode> PlaybackNodeClass;
+	
+};
diff --git a/Engine/Plugins/Experimental/Avalanche/Source/AvalancheMediaEditor/Public/Playlist/AvaPlaylistDefines.h b/Engine/Plugins/Experimental/Avalanche/Source/AvalancheMediaEditor/Public/Playlist/AvaPlaylistDefines.h
new file mode 100644
--- /dev/null
+++ b/Engine/Plugins/Experimental/Avalanche/Source/AvalancheMediaEditor/Public/Playlist/AvaPlaylistDefines.h
@@ -0,0 +1,50 @@
+// Copyright Epic Games, Inc. All Rights Reserved.
+
+#pragma once
+
+#include "Internationalization/Text.h"
+#include "Templates/SharedPointer.h"
+
+class IAvaInstancedPageView;
+class IAvaTemplatePageView;
+class IAvaPageView;
+struct FAvalanchePage;
+
+typedef TWeakPtr<IAvaPageView> FAvaPageViewWeak;
+typedef TSharedPtr<IAvaPageView> FAvaPageViewPtr;
+typedef TSharedRef<IAvaPageView> FAvaPageViewRef;
+typedef TWeakPtr<IAvaTemplatePageView> FAvaTemplatePageViewWeak;
+typedef TSharedPtr<IAvaTemplatePageView> FAvaTemplatePageViewPtr;
+typedef TSharedRef<IAvaTemplatePageView> FAvaTemplatePageViewRef;
+typedef TWeakPtr<IAvaInstancedPageView> FAvaInstancedPageViewWeak;
+typedef TSharedPtr<IAvaInstancedPageView> FAvaInstancedPageViewPtr;
+typedef TSharedRef<IAvaInstancedPageView> FAvaInstancedPageViewRef;
+
+namespace UE::AvalanchePlaylist
+{
+	enum {InvalidPageId	= -1};
+	
+	struct AVALANCHEMEDIAEDITOR_API FEditorMetrics
+	{
+		static constexpr float ColumnLeftOffset = 5.f;
+		static const FNumberFormattingOptions PageIdFormattingOptions;
+	};
+	
+	enum class EPageEvent : uint8
+	{
+		/** Called after the Selection has Changed */ 
+		SelectionChanged,
+
+		/** Called to sync the Selection of current Selected Items in the Caller */
+		SelectionRequest,
+
+		/** Called to Request a Rename of the Selected Items */
+		RenameRequest,
+
+		/** Called to Request Renumbering Page Id of Selected Items */
+		RenumberRequest,
+
+		/** Called to Request Reimporting Assets of Selected Items */
+		ReimportRequest,
+	};
+}
diff --git a/Engine/Plugins/Experimental/Avalanche/Source/AvalancheMediaEditor/Public/Playlist/AvaPlaylistEditor.h b/Engine/Plugins/Experimental/Avalanche/Source/AvalancheMediaEditor/Public/Playlist/AvaPlaylistEditor.h
new file mode 100644
--- /dev/null
+++ b/Engine/Plugins/Experimental/Avalanche/Source/AvalancheMediaEditor/Public/Playlist/AvaPlaylistEditor.h
@@ -0,0 +1,224 @@
+// Copyright Epic Games, Inc. All Rights Reserved.
+
+#pragma once
+
+#include "AvaPlaylistDefines.h"
+#include "Playlist/AvaRundownMacroCollection.h"
+#include "Playlist/AvalanchePlaylist.h"
+#include "TickableEditorObject.h"
+#include "WorkflowOrientedApp/WorkflowCentricApplication.h"
+
+class FAvaPageTextFilter;
+class FAvaPlaylistEditorInputProcessor;
+class FAvalancheManagedInstance;
+class SAvaInstancedPageList;
+class SAvaPageList;
+class SAvaReadPage;
+enum class EAvaPlaylistSearchListType : uint8;
+
+class AVALANCHEMEDIAEDITOR_API FAvaPlaylistEditor : public FWorkflowCentricApplication
+{
+public:
+	FAvaPlaylistEditor();
+
+	virtual ~FAvaPlaylistEditor() override;
+
+	void InitPlaylistEditor(const EToolkitMode::Type InMode
+		, const TSharedPtr<IToolkitHost>& InitToolkitHost
+		, UAvalanchePlaylist* InPlaylist);
+
+	//~ Begin IToolkit
+	virtual FName GetToolkitFName() const override;
+	virtual FText GetBaseToolkitName() const override;
+	virtual FString GetWorldCentricTabPrefix() const override;
+	virtual FLinearColor GetWorldCentricTabColorScale() const override;
+	//~ End IToolkit
+
+	//~ Begin FAssetEditorToolkit
+	virtual bool OnRequestClose(EAssetEditorCloseReason InCloseReason) override;
+	//~ End FAssetEditorToolkit
+	
+	TSharedRef<SWidget> MakeReadPageWidget();
+
+	void ExtendToolBar(TSharedPtr<FExtender> InExtender);
+	void FillPageToolBar(FToolBarBuilder& OutToolBarBuilder);
+
+	bool IsPlaylistValid() const { return AvalanchePlaylist.IsValid(); }
+	UAvalanchePlaylist* GetPlaylist() const;
+	void MarkAsModified();
+
+	DECLARE_MULTICAST_DELEGATE_TwoParams(FOnPageEvent, const TArray<int32>&, UE::AvalanchePlaylist::EPageEvent);
+	FOnPageEvent& GetOnPageEvent() { return OnPageEvent; }
+
+	static TArray<TSharedPtr<FAvalancheManagedInstance>> GetManagedInstancesForPage(const UAvalanchePlaylist* InPlaylist, const FAvalanchePage& InPage);
+	static bool MergeDefaultRemoteControlValues(const TArray<TSharedPtr<FAvalancheManagedInstance>>& InManagedInstances, FAvalancheRemoteControlValues& OutMergedValues);
+
+	bool IsKeyRelevant(const FKeyEvent& InKeyEvent) const;
+	bool HandleKeyDownEvent(const FKeyEvent& InKeyEvent);
+
+	TSharedPtr<SAvaPageList> GetTemplateListWidget() const;
+	TSharedPtr<SAvaPageList> GetInstanceListWidget() const;
+	TSharedPtr<SAvaPageList> GetListWidget(const FAvaPageListReference& InPageListReference) const;
+	TSharedPtr<SAvaPageList> GetListWidget(const FName& InTabId) const;
+
+	/** Returns the currently active Instanced Page List. */
+	TSharedPtr<SAvaInstancedPageList> GetActiveListWidget() const;
+	
+	/** Returns the currently focused Page List, including templates, instances or views. */
+	TSharedPtr<SAvaPageList> GetFocusedListWidget() const;
+
+	int32 GetFirstSelectedPageOnActiveSubListWidget() const;
+	TConstArrayView<int32> GetSelectedPagesOnActiveSubListWidget() const;
+	TConstArrayView<int32> GetSelectedPagesOnFocusedWidget() const;
+
+	bool CanAddTemplate() const;
+	void AddTemplate();
+
+	bool CanPlaySelectedPage() const;
+	void PlaySelectedPage();
+
+	bool CanUpdateValuesOnSelectedPage() const;
+	void UpdateValuesOnSelectedPage();
+	
+	bool CanStopSelectedPage(bool bInForce) const;
+	void StopSelectedPage(bool bInForce);
+
+	bool CanContinueSelectedPage() const;
+	void ContinueSelectedPage();
+
+	bool CanPlayNextPage() const;
+	void PlayNextPage();
+
+	bool CanPreviewPlaySelectedPage() const;
+	void PreviewPlaySelectedPage(bool bInToMark);
+
+	bool CanPreviewContinueSelectedPage() const;
+	void PreviewContinueSelectedPage();
+
+	bool CanPreviewStopSelectedPage(bool bInForce) const;
+	void PreviewStopSelectedPage(bool bInForce);
+
+	bool CanPreviewPlayNextPage() const;
+	void PreviewPlayNextPage();
+
+	bool CanTakeToProgram() const;
+	void TakeToProgram() const;
+
+	bool CanCreateInstancesFromSelectedTemplates() const;
+	void CreateInstancesFromSelectedTemplates();
+
+	bool CanRemoveSelectedPages() const;
+	void RemoveSelectedPages();
+
+	void RefreshTemplateVisibility();
+	void RefreshInstancedVisibility();
+
+	bool IsTemplatePageVisible(const FAvalanchePage& InPage) const;
+	bool IsInstancedPageVisible(const FAvalanchePage& InPage) const;
+
+	void SetSearchText(const FText& InText, EAvaPlaylistSearchListType& InPageListType);
+
+protected:
+	void RegisterApplicationModes();
+
+	void CreatePlaylistCommands();
+
+	FText GetCurrentProfileName() const;
+
+	TSharedRef<SWidget> MakeProfileComboButton();
+
+	void StartAutoPlayCommand(const TArray<FString>& InArgs);
+	void StopAutoPlayCommand(const TArray<FString>& InArgs);
+	void CancelAutoPlay();
+	
+	void LoadPageCommand(const TArray<FString>& InArgs);
+	void UnloadPageCommand(const TArray<FString>& InArgs);
+	void PlayPageCommand(const TArray<FString>& InArgs, bool bInPreview);
+	void ContinuePageCommand(const TArray<FString>& InArgs, bool bInPreview);
+	void PlayNextPageCommand(const TArray<FString>& InArgs, bool bInPreview);
+	void StopPageCommand(const TArray<FString>& InArgs, bool bInPreview, bool bInForce);
+	void TakeToProgramCommand(const TArray<FString>& InArgs);
+	void StartChannelCommand(const TArray<FString>& InArgs);
+	void StopChannelCommand(const TArray<FString>& InArgs);
+	
+	TArray<int32> ArgumentsToPageIds(const TArray<FString>& InArgs, bool bInPreview, FString& OutErrors) const;
+
+	using FMacroCommandFunction = TFunction<void(const TArray<FString>& InArgs)>;
+	using FBindableMacroCommands = TMap<FName, FMacroCommandFunction>;
+	const FBindableMacroCommands& GetBindableMacroCommands();
+
+private:
+	void SetTemplateSearchText(const FText& InText);
+
+	void SetInstancedSearchText(const FText& InText);
+
+	void OnTemplateFilterChanged();
+
+	void OnInstancedFilterChanged();
+
+	void InitVisibilityTemplatePages();
+
+	void InitVisibilityInstancedPages();
+
+protected:
+	/**
+	 *	Tickable object that keeps track of the play time of a page and will automatically
+	 *	play the next page once the specified play time has elapsed. The measure of play time
+	 *	begins when the page has started playing, the loading time is ignored.
+	 *	This is tracked from the page's status.
+	 */
+	struct FAutoPlayTicker : public FTickableEditorObject
+	{
+		static constexpr double DefaultPlayInterval = 10;
+
+		TWeakPtr<FAvaPlaylistEditor> PlaylistEditorWeak;
+		double PlayInterval = DefaultPlayInterval;
+		double NextPageStartTime = 0;
+		bool bIsPagePlaying = false;
+		bool bIsCancelled = false;
+
+		FAutoPlayTicker(TWeakPtr<FAvaPlaylistEditor> InPlaylistEditorWeak, double InPlayInterval);
+
+		//~ Begin FTickableEditorObject
+		virtual void Tick(float DeltaTime) override;
+		virtual ETickableTickType GetTickableTickType() const override { return ETickableTickType::Always; }
+		virtual TStatId GetStatId() const override;
+		//~ End FTickableEditorObject
+	};
+	TUniquePtr<FAutoPlayTicker> AutoPlayTicker;
+
+	TSharedPtr<SDockTab> CreateSubListTab(int32 InSubListIndex);
+	void CreateSubListTabs();
+
+	void OnActiveSubListChanged();
+	void HandleOnPagePlayerAdded(UAvalanchePlaylist* InPlaylist, UAvalanchePagePlayer* InPagePlayer);
+
+	TWeakObjectPtr<UAvalanchePlaylist> AvalanchePlaylist;
+
+	FOnPageEvent OnPageEvent;
+
+	TMap<FName, TWeakPtr<SDockTab>> SubListTabs;
+
+private:
+	TSharedPtr<SAvaReadPage> ReadPageWidget;
+
+	TSharedPtr<FAvaPlaylistEditorInputProcessor> InputProcessor;
+
+	TSharedRef<FAvaPageTextFilter> TextFilterTemplatePage;
+	TArray<int32> VisibleTemplatePageIds;
+
+	TSharedRef<FAvaPageTextFilter> TextFilterInstancedPage;
+	TArray<int32> VisibleInstancedPageIds;
+
+	/**
+	 * All playlist editor instances share the same command handlers.
+	 * The commands get dispatched only to the active editor.
+	 */
+	class FSharedConsoleCommands;
+	TSharedPtr<FSharedConsoleCommands> SharedConsoleCommands;
+
+	FBindableMacroCommands BindableMacroCommands;
+	
+	/** Optional user specified action for stopping pages. */
+	TOptional<bool> bStopPagesOnCloseOverride;
+};
diff --git a/Engine/Plugins/Experimental/Avalanche/Source/AvalancheMediaEditor/Public/Playlist/AvaPlaylistFactory.h b/Engine/Plugins/Experimental/Avalanche/Source/AvalancheMediaEditor/Public/Playlist/AvaPlaylistFactory.h
new file mode 100644
--- /dev/null
+++ b/Engine/Plugins/Experimental/Avalanche/Source/AvalancheMediaEditor/Public/Playlist/AvaPlaylistFactory.h
@@ -0,0 +1,30 @@
+// Copyright Epic Games, Inc. All Rights Reserved.
+
+#pragma once
+
+#include "CoreMinimal.h"
+#include "Factories/Factory.h"
+#include "AvaPlaylistFactory.generated.h"
+
+UCLASS()
+class AVALANCHEMEDIAEDITOR_API UAvaPlaylistFactory : public UFactory
+{
+	GENERATED_BODY()
+
+public:
+
+	UAvaPlaylistFactory();
+	virtual ~UAvaPlaylistFactory() override;
+
+protected:
+	
+	//~ Begin UFactory Interface
+	virtual uint32 GetMenuCategories() const override;
+	virtual bool ConfigureProperties() override;
+	virtual bool ShouldShowInNewMenu() const override;
+	virtual UObject* FactoryCreateNew(UClass* Class, UObject* InParent, FName Name, EObjectFlags Flags, UObject* Context, FFeedbackContext* Warn, FName CallingContext) override;
+	virtual UObject* FactoryCreateNew(UClass* Class, UObject* InParent, FName Name, EObjectFlags Flags, UObject* Context, FFeedbackContext* Warn) override;
+	virtual bool DoesSupportClass(UClass* Class) override;
+	virtual UClass* ResolveSupportedClass() override;
+	//~ Begin UFactory Interface
+};
