[MotionDesign]

diff --git a/Engine/Plugins/Experimental/Avalanche/Source/AvalancheMediaEditor/Public/Playlist/AvaPlaylistMessages.h b/Engine/Plugins/Experimental/Avalanche/Source/AvalancheMediaEditor/Public/Playlist/AvaPlaylistMessages.h
new file mode 100644
--- /dev/null
+++ b/Engine/Plugins/Experimental/Avalanche/Source/AvalancheMediaEditor/Public/Playlist/AvaPlaylistMessages.h
@@ -0,0 +1,689 @@
+// Copyright Epic Games, Inc. All Rights Reserved.
+
+#pragma once
+
+#include "Playlist/AvalanchePage.h"
+
+#include "AvaPlaylistMessages.generated.h"
+
+namespace EPlaylistApiVersion
+{
+	/**
+	 * Defines the protocol version of the Playlist Server API.
+	 *
+	 * API versioning is used to provide legacy support either on
+	 * the client side or server side for non compatible changes.
+	 * Clients can request a version of the API that they where implemented against,
+	 * if the server can still honor the request it will accept.
+	 */
+	enum Type
+	{
+		Initial = 1,
+
+		// -----<new versions can be added before this line>-------------------------------------------------
+		// - this needs to be the last line (see note below)
+		VersionPlusOne,
+		LatestVersion = VersionPlusOne - 1
+	};
+}
+
+
+USTRUCT()
+struct AVALANCHEMEDIAEDITOR_API FAvaPlaylistMsgBase
+{
+	GENERATED_BODY()
+public:
+	UPROPERTY()
+	int32 RequestId = -1;
+};
+
+USTRUCT()
+struct AVALANCHEMEDIAEDITOR_API FAvaPlaylistServerMsg : public FAvaPlaylistMsgBase
+{
+	GENERATED_BODY()
+public:
+	UPROPERTY()
+	FString Verbosity;
+
+	UPROPERTY()
+	FString Text;
+};
+
+/** Request published by client to discover servers. */
+USTRUCT()
+struct AVALANCHEMEDIAEDITOR_API FAvaPlaylistPing : public FAvaPlaylistMsgBase
+{
+	GENERATED_BODY()
+public:
+	/** True if the request originates from an automatic timer. False if requests originates from user interaction. */
+	UPROPERTY()
+	bool bAuto = true;
+
+	/**
+	 * API Version the client has been implemented against.
+	 * If none (-1) the server will consider the initial version is requested.
+	 */
+	UPROPERTY()
+	int32 RequestedApiVersion = -1;
+};
+
+/** Response sent by server to client to be discovered. */
+USTRUCT()
+struct AVALANCHEMEDIAEDITOR_API FAvaPlaylistPong : public FAvaPlaylistMsgBase
+{
+	GENERATED_BODY()
+
+public:
+	/** True if it is a reply to an auto ping. Mirrors the bAuto flag from Ping message. */
+	UPROPERTY()
+	bool bAuto = true;
+
+	/**
+	 * API Version the server will communicate with for this client.
+	 * The server may honor the requested version if possible.
+	 * Versions newer than server implementation will obviously not be honored either.
+	 * Clients should expect an older server to reply with an older version.
+	 */
+	UPROPERTY()
+	int32 ApiVersion = -1;
+
+	/** Minimum API Version the server implements. */
+	UPROPERTY()
+	int32 MinimumApiVersion = -1;
+
+	/** Latest API Version the server support. */
+	UPROPERTY()
+	int32 LatestApiVersion = -1;
+
+	UPROPERTY()
+	FString HostName;
+};
+
+/**
+ *	Request list of playlist that can be opened on the current server.
+ */
+USTRUCT()
+struct AVALANCHEMEDIAEDITOR_API FAvaPlaylistGetPlaylists : public FAvaPlaylistMsgBase
+{
+	GENERATED_BODY()
+};
+
+/**
+ *	List of all playlists.
+ *	Expected Response from FAvaPlaylistGetPlaylists.
+ */
+USTRUCT()
+struct AVALANCHEMEDIAEDITOR_API FAvaPlaylistPlaylists : public FAvaPlaylistMsgBase
+{
+	GENERATED_BODY()
+	
+public:
+	UPROPERTY()
+	TArray<FString> Playlists;
+};
+
+/**
+ *	Request that the given playlist be opened.
+ *	Only one playlist can be opened at a time. If another playlist
+ *	is opened, it will be closed and all currently playing pages stopped.
+ *	If the path is empty, nothing will be done and the server will reply with
+ *	a FAvaPlaylistLoadedPlaylist message indicating which playlist is currently loaded.
+ */
+USTRUCT()
+struct AVALANCHEMEDIAEDITOR_API FAvaPlaylistLoadPlaylist : public FAvaPlaylistMsgBase
+{
+	GENERATED_BODY()
+	
+public:
+	UPROPERTY()
+	FString Playlist;
+};
+
+/**
+ * Request the list of pages from the given playlist.
+ */
+USTRUCT()
+struct AVALANCHEMEDIAEDITOR_API FAvaPlaylistGetPages : public FAvaPlaylistMsgBase
+{
+	GENERATED_BODY()
+public:
+	UPROPERTY()
+	FString Playlist;
+};
+
+USTRUCT()
+struct AVALANCHEMEDIAEDITOR_API FAvaPlaylistCreatePage : public FAvaPlaylistMsgBase
+{
+	GENERATED_BODY()
+
+	UPROPERTY()
+	FString Playlist;
+
+	UPROPERTY()
+	int32 TemplateId = FAvalanchePage::InvalidPageId;
+};
+
+USTRUCT()
+struct AVALANCHEMEDIAEDITOR_API FAvaPlaylistDeletePage : public FAvaPlaylistMsgBase
+{
+	GENERATED_BODY()
+
+	UPROPERTY()
+	FString Playlist;
+
+	UPROPERTY()
+	int32 PageId = FAvalanchePage::InvalidPageId;
+};
+
+USTRUCT()
+struct AVALANCHEMEDIAEDITOR_API FAvaPlaylistCreateTemplate : public FAvaPlaylistMsgBase
+{
+	GENERATED_BODY()
+
+	UPROPERTY()
+	FString Playlist;
+};
+
+USTRUCT()
+struct AVALANCHEMEDIAEDITOR_API FAvaPlaylistDeleteTemplate : public FAvaPlaylistMsgBase
+{
+	GENERATED_BODY()
+
+	UPROPERTY()
+	FString Playlist;
+
+	UPROPERTY()
+	int32 PageId = FAvalanchePage::InvalidPageId;
+};
+
+USTRUCT()
+struct AVALANCHEMEDIAEDITOR_API FAvaPlaylistChangeTemplateBP : public FAvaPlaylistMsgBase
+{
+	GENERATED_BODY()
+public:
+	UPROPERTY()
+	FString Playlist;
+	
+	UPROPERTY()
+	int32 TemplateId = FAvalanchePage::InvalidPageId;
+
+	UPROPERTY()
+	FString AssetPath;
+};
+
+USTRUCT()
+struct AVALANCHEMEDIAEDITOR_API FAvaPlaylistPage
+{
+	GENERATED_BODY()
+public:
+	UPROPERTY()
+	int32 PageId = FAvalanchePage::InvalidPageId;
+
+	UPROPERTY()
+	FString PageName;
+
+	UPROPERTY()
+	bool IsTemplate = false;
+
+	UPROPERTY()
+	int32 TemplateId = FAvalanchePage::InvalidPageId;
+
+	UPROPERTY()
+	TArray<int32> CombinedTemplateIds;
+
+	UPROPERTY()
+	FSoftObjectPath AssetPath;
+
+	UPROPERTY()
+	TArray<FAvalanchePageStatus> Statuses;
+
+	UPROPERTY()
+	FString TransitionLayerName;
+	
+	UPROPERTY()
+	FString OutputChannel;
+
+	UPROPERTY()
+	bool bIsEnabled = false;
+
+	UPROPERTY()
+	bool bIsPlaying = false;
+};
+
+/*
+ * List of pages from the current playlist.
+ */
+USTRUCT()
+struct AVALANCHEMEDIAEDITOR_API FAvaPlaylistPages : public FAvaPlaylistMsgBase
+{
+	GENERATED_BODY()
+public:
+	UPROPERTY()
+	TArray<FAvaPlaylistPage> Pages;
+};
+
+/**
+ * Request the page details from the given playlist.
+ */
+USTRUCT()
+struct AVALANCHEMEDIAEDITOR_API FAvaPlaylistGetPageDetails : public FAvaPlaylistMsgBase
+{
+	GENERATED_BODY()
+	
+public:
+	UPROPERTY()
+	FString Playlist;
+
+	UPROPERTY()
+	int32 PageId = FAvalanchePage::InvalidPageId;
+
+	/** This will request that a managed asset instance gets loaded to be
+	 * accessible through WebRC. */
+	UPROPERTY()
+	bool bLoadRemoteControlPreset = false;
+};
+
+/**
+ *	Server response to FAvaPlaylistGetPageDetails request.
+ */
+USTRUCT()
+struct AVALANCHEMEDIAEDITOR_API FAvaPlaylistPageDetails : public FAvaPlaylistMsgBase
+{
+	GENERATED_BODY()
+	
+public:
+	UPROPERTY()
+	FString Playlist;
+
+	UPROPERTY()
+	FAvaPlaylistPage PageInfo;
+
+	UPROPERTY()
+	FAvalancheRemoteControlValues RemoteControlValues;
+
+	/** Name of the remote control preset to resolve through WebRC API. */
+	UPROPERTY()
+	FString RemoteControlPresetName;
+
+	UPROPERTY()
+	FString RemoteControlPresetId;
+};
+
+USTRUCT()
+struct AVALANCHEMEDIAEDITOR_API FAvaPlaylistPagesStatuses : public FAvaPlaylistMsgBase
+{
+	GENERATED_BODY()
+
+public:
+	UPROPERTY()
+	FString Playlist;
+
+	UPROPERTY()
+	FAvaPlaylistPage PageInfo;
+};
+
+USTRUCT()
+struct AVALANCHEMEDIAEDITOR_API FAvaPlaylistPageListChanged : public FAvaPlaylistMsgBase
+{
+	GENERATED_BODY()
+
+	UPROPERTY()
+	FString Playlist;
+
+	/** See EAvaPageListChange flags. */
+	UPROPERTY()
+	uint8 ChangeType = 0;
+
+	UPROPERTY();
+	TArray<int32> AffectedPages;
+};
+
+USTRUCT()
+struct AVALANCHEMEDIAEDITOR_API FAvaPlaylistPageBlueprintChanged : public FAvaPlaylistMsgBase
+{
+	GENERATED_BODY()
+
+	UPROPERTY()
+	FString Playlist;
+
+	UPROPERTY()
+	int32 PageId = FAvalanchePage::InvalidPageId;
+	
+	UPROPERTY()
+	FString BlueprintPath;
+};
+
+USTRUCT()
+struct AVALANCHEMEDIAEDITOR_API FAvaPlaylistPageChannelChanged : public FAvaPlaylistMsgBase
+{
+	GENERATED_BODY()
+
+	UPROPERTY()
+	FString Playlist;
+	
+	UPROPERTY()
+	int32 PageId = FAvalanchePage::InvalidPageId;
+
+	UPROPERTY()
+	FString ChannelName;
+};
+
+USTRUCT()
+struct AVALANCHEMEDIAEDITOR_API FAvaPlaylistPageAnimSettingsChanged : public FAvaPlaylistMsgBase
+{
+	GENERATED_BODY()
+
+	UPROPERTY()
+	FString Playlist;
+
+	UPROPERTY()
+	int32 PageId = FAvalanchePage::InvalidPageId;
+};
+
+USTRUCT()
+struct AVALANCHEMEDIAEDITOR_API FAvaPlaylistPageChangeChannel : public FAvaPlaylistMsgBase
+{
+	GENERATED_BODY()
+
+	UPROPERTY()
+	FString Playlist;
+
+	UPROPERTY()
+	int32 PageId = FAvalanchePage::InvalidPageId;
+
+	UPROPERTY()
+	FString ChannelName;
+};
+
+/** This is a request to save the managed RCP back to the corresponding page. */
+USTRUCT()
+struct AVALANCHEMEDIAEDITOR_API FAvaPlaylistUpdatePageFromRCP : public FAvaPlaylistMsgBase
+{
+	GENERATED_BODY()
+
+public:
+	/** Unregister the Remote Control Preset from the WebRC. */
+	UPROPERTY()
+	bool bUnregister = false;
+};
+
+/** Supported Page actions for playback. */
+UENUM()
+enum class EAvaPlaylistPageActions
+{
+	None,
+	Load,
+	Unload,
+	Play,
+	PlayNext,
+	Stop,
+	ForceStop,
+	Continue,
+	UpdateValues,
+	TakeToProgram
+};
+
+USTRUCT()
+struct AVALANCHEMEDIAEDITOR_API FAvaPlaylistPageAction : public FAvaPlaylistMsgBase
+{
+	GENERATED_BODY()
+public:
+	UPROPERTY()
+	int32 PageId = FAvalanchePage::InvalidPageId;
+
+	UPROPERTY()
+	EAvaPlaylistPageActions Action = EAvaPlaylistPageActions::None;
+};
+
+USTRUCT()
+struct AVALANCHEMEDIAEDITOR_API FAvaPlaylistPagePreviewAction : public FAvaPlaylistPageAction
+{
+	GENERATED_BODY()
+public:
+	/** Specify which preview channel to use. If left empty, the playlist's default preview channel is used. */
+	UPROPERTY()
+	FString PreviewChannelName;
+};
+
+/**
+ * Command to execute an action on multiple pages at the same time.
+ * This is necessary for pages to be part of the same transition.
+ */
+USTRUCT()
+struct AVALANCHEMEDIAEDITOR_API FAvaPlaylistPageActions : public FAvaPlaylistMsgBase
+{
+	GENERATED_BODY()
+public:
+	UPROPERTY()
+	TArray<int32> PageIds;
+
+	UPROPERTY()
+	EAvaPlaylistPageActions Action = EAvaPlaylistPageActions::None;
+};
+
+USTRUCT()
+struct AVALANCHEMEDIAEDITOR_API FAvaPlaylistPagePreviewActions : public FAvaPlaylistPageActions
+{
+	GENERATED_BODY()
+public:
+	/** Specify which preview channel to use. If left empty, the playlist's default preview channel is used. */
+	UPROPERTY()
+	FString PreviewChannelName;
+};
+
+UENUM()
+enum class EAvaPlaylistPageEvents
+{
+	None,
+	AnimStarted,
+	AnimPaused,
+	AnimFinished
+};
+
+USTRUCT()
+struct AVALANCHEMEDIAEDITOR_API FAvaPlaylistPageEvent : public FAvaPlaylistMsgBase
+{
+	GENERATED_BODY()
+public:
+	UPROPERTY()
+	int32 PageId = FAvalanchePage::InvalidPageId;
+
+	UPROPERTY()
+	EAvaPlaylistPageEvents Event = EAvaPlaylistPageEvents::None;
+};
+
+USTRUCT()
+struct AVALANCHEMEDIAEDITOR_API FAvaPlaylistOutputDeviceItem
+{
+	GENERATED_BODY()
+public:
+	UPROPERTY()
+	FString Name;
+
+	/**
+	 * Raw Json string representing a serialized UMediaOutput.
+	 */
+	UPROPERTY()
+	FString Data;
+};
+
+USTRUCT()
+struct AVALANCHEMEDIAEDITOR_API FAvaPlaylistOutputClassItem
+{
+	GENERATED_BODY()
+public:
+	UPROPERTY()
+	FString Name;
+
+	UPROPERTY()
+	TArray<FAvaPlaylistOutputDeviceItem> Devices;
+};
+
+USTRUCT()
+struct AVALANCHEMEDIAEDITOR_API FAvaPlaylistDevicesList : public FAvaPlaylistMsgBase
+{
+	GENERATED_BODY()
+public:
+	UPROPERTY()
+	TArray<FAvaPlaylistOutputClassItem> DeviceClasses;
+};
+
+USTRUCT()
+struct AVALANCHEMEDIAEDITOR_API FAvaPlaylistGetChannel : public FAvaPlaylistMsgBase
+{
+	GENERATED_BODY()
+public:
+	UPROPERTY()
+	FString ChannelName;
+};
+
+USTRUCT()
+struct AVALANCHEMEDIAEDITOR_API FAvaPlaylistGetChannels : public FAvaPlaylistMsgBase
+{
+	GENERATED_BODY()
+};
+
+USTRUCT()
+struct AVALANCHEMEDIAEDITOR_API FAvaPlaylistChannel
+{
+	GENERATED_BODY()
+public:
+	UPROPERTY()
+	FString Name;
+
+	UPROPERTY()
+	TArray<FAvaPlaylistOutputDeviceItem> Devices;
+};
+
+USTRUCT()
+struct AVALANCHEMEDIAEDITOR_API FAvaBroadcastChannelListChanged : public FAvaPlaylistMsgBase
+{
+	GENERATED_BODY()
+
+	UPROPERTY()
+	TArray<FAvaPlaylistChannel> Channels;
+};
+
+USTRUCT()
+struct AVALANCHEMEDIAEDITOR_API FAvaPlaylistChannelResponse : public FAvaPlaylistMsgBase
+{
+	GENERATED_BODY()
+public:
+	UPROPERTY()
+	FAvaPlaylistChannel Channel;
+};
+
+USTRUCT()
+struct AVALANCHEMEDIAEDITOR_API FAvaPlaylistChannels : public FAvaPlaylistMsgBase
+{
+	GENERATED_BODY()
+public:
+	UPROPERTY()
+	TArray<FAvaPlaylistChannel> Channels;
+};
+
+USTRUCT()
+struct AVALANCHEMEDIAEDITOR_API FAvaPlaylistAssetsChanged : public FAvaPlaylistMsgBase
+{
+	GENERATED_BODY()
+
+	UPROPERTY()
+	FString AssetName;
+};
+
+// Channel actions
+UENUM()
+enum class EAvaPlaylistChannelActions
+{
+	None,
+	Start,
+	Stop
+};
+
+USTRUCT()
+struct AVALANCHEMEDIAEDITOR_API FAvaPlaylistChannelAction : public FAvaPlaylistMsgBase
+{
+	GENERATED_BODY()
+public:
+	UPROPERTY()
+	FString ChannelName;
+
+	UPROPERTY()
+	EAvaPlaylistChannelActions Action = EAvaPlaylistChannelActions::None;
+};
+
+USTRUCT()
+struct AVALANCHEMEDIAEDITOR_API FAvaPlaylistGetDevices : public FAvaPlaylistMsgBase
+{
+	GENERATED_BODY()
+};
+
+USTRUCT()
+struct AVALANCHEMEDIAEDITOR_API FAvaPlaylistAddChannelDevice : public FAvaPlaylistMsgBase
+{
+	GENERATED_BODY()
+public:
+	UPROPERTY()
+	FString ChannelName;
+
+	UPROPERTY()
+	FString MediaOutputName;
+};
+
+USTRUCT()
+struct AVALANCHEMEDIAEDITOR_API FAvaPlaylistEditChannelDevice : public FAvaPlaylistMsgBase
+{
+	GENERATED_BODY()
+public:
+	UPROPERTY()
+	FString ChannelName;
+
+	UPROPERTY()
+	FString MediaOutputName;
+
+	UPROPERTY()
+	FString Data;
+};
+
+USTRUCT()
+struct AVALANCHEMEDIAEDITOR_API FAvaPlaylistRemoveChannelDevice : public FAvaPlaylistMsgBase
+{
+	GENERATED_BODY()
+public:
+	UPROPERTY()
+	FString ChannelName;
+
+	UPROPERTY()
+	FString MediaOutputName;
+};
+
+/* No difference from FAvaPlaylistOutputDeviceItem except this is meant to return a single device response. */
+USTRUCT()
+struct AVALANCHEMEDIAEDITOR_API FAvaPlaylistOutputDeviceItemResponse : public FAvaPlaylistMsgBase
+{
+	GENERATED_BODY()
+public:
+	UPROPERTY()
+	FString Name;
+
+	UPROPERTY()
+	FString Data;
+};
+
+USTRUCT()
+struct AVALANCHEMEDIAEDITOR_API FAvaPlaylistGetChannelImage : public FAvaPlaylistMsgBase
+{
+	GENERATED_BODY()
+public:
+	UPROPERTY()
+	FString ChannelName;
+};
+
+USTRUCT()
+struct AVALANCHEMEDIAEDITOR_API FAvaPlaylistChannelImage : public FAvaPlaylistMsgBase
+{
+	GENERATED_BODY()
+public:
+	UPROPERTY()
+	TArray<uint8> ImageData;
+};
diff --git a/Engine/Plugins/Experimental/Avalanche/Source/AvalancheMediaEditor/Public/Playlist/AvaPlaylistServer.h b/Engine/Plugins/Experimental/Avalanche/Source/AvalancheMediaEditor/Public/Playlist/AvaPlaylistServer.h
new file mode 100644
--- /dev/null
+++ b/Engine/Plugins/Experimental/Avalanche/Source/AvalancheMediaEditor/Public/Playlist/AvaPlaylistServer.h
@@ -0,0 +1,224 @@
+// Copyright Epic Games, Inc. All Rights Reserved.
+
+#pragma once
+
+#include "AvaBroadcastProfile.h"
+#include "AvaPlaylistMessages.h"
+#include "MessageEndpoint.h"
+#include "Playlist/AvalanchePlaylist.h"
+#include "UObject/SoftObjectPath.h"
+#include "UObject/StrongObjectPtr.h"
+
+class FAvalancheManagedInstance;
+class UMediaOutput;
+
+/**
+ * Implements a play list server that listens to commands on message bus.
+ * The intention is to run a web socket transport bridge so the messages can
+ * come from external applications.
+ */
+class AVALANCHEMEDIAEDITOR_API FAvaPlaylistServer : public TSharedFromThis<FAvaPlaylistServer>
+{
+public:
+	FAvaPlaylistServer();
+	virtual ~FAvaPlaylistServer();
+
+	void Init(const FString& InAssignedHostName);
+
+	void SetupBroadcastDelegates(UAvalancheBroadcast* InBroadcast);
+	void SetupEditorDelegates();
+	void RemoveBroadcastDelegates(UAvalancheBroadcast* InBroadcast) const;
+	void RemoveEditorDelegates() const;
+	
+	void OnPagesChanged(const UAvalanchePlaylist* InPlaylist, const FAvalanchePage& InPage, EAvaPageChanges InChange) const;
+	void OnPageListChanged(const FAvaPageListChangeParams& InParams) const;
+	void PageBlueprintChanged(const UAvalanchePlaylist* InPlaylist, const FAvalanchePage& InPage, const FString& InBlueprintPath) const;
+	void PageStatusChanged(const UAvalanchePlaylist* InPlaylist, const FAvalanchePage& InPage) const;
+	void PageChannelChanged(const UAvalanchePlaylist* InPlaylist, const FAvalanchePage& InPage, const FString& InChannelName) const;
+	void PageAnimSettingsChanged(const UAvalanchePlaylist* InPlaylist, const FAvalanchePage& InPage) const;
+	void OnBroadcastChannelListChanged(const FAvaBroadcastProfile& InProfile) const;
+	void OnAssetAddedOrRemoved(const FAssetData& InAssetData) const;
+
+	/** Returns the endpoint's message address. */
+	const FMessageAddress& GetMessageAddress() const;
+	
+	// Message handlers
+	void HandlePlaylistPing(const FAvaPlaylistPing& InMessage, const TSharedRef<IMessageContext, ESPMode::ThreadSafe>& InContext);
+	void HandleGetPlaylists(const FAvaPlaylistGetPlaylists& InMessage, const TSharedRef<IMessageContext, ESPMode::ThreadSafe>& InContext);
+	void HandleLoadPlaylist(const FAvaPlaylistLoadPlaylist& InMessage, const TSharedRef<IMessageContext, ESPMode::ThreadSafe>& InContext);
+	void HandleGetPages(const FAvaPlaylistGetPages& InMessage, const TSharedRef<IMessageContext, ESPMode::ThreadSafe>& InContext);
+	void HandleCreatePage(const FAvaPlaylistCreatePage& InMessage,const TSharedRef<IMessageContext, ESPMode::ThreadSafe>& InContext);
+	void HandleCreateTemplate(const FAvaPlaylistCreateTemplate& InMessage, const TSharedRef<IMessageContext, ESPMode::ThreadSafe>& InContext);
+	void HandleDeletePage(const FAvaPlaylistDeletePage& InMessage,const TSharedRef<IMessageContext, ESPMode::ThreadSafe>& InContext);
+	void HandleDeleteTemplate(const FAvaPlaylistDeleteTemplate& InMessage,const TSharedRef<IMessageContext, ESPMode::ThreadSafe>& InContext);
+	void HandleChangeTemplateBP(const FAvaPlaylistChangeTemplateBP& InMessage, const TSharedRef<IMessageContext, ESPMode::ThreadSafe>& InContext);
+	void HandleGetPageDetails(const FAvaPlaylistGetPageDetails& InMessage, const TSharedRef<IMessageContext, ESPMode::ThreadSafe>& InContext);
+	void HandleChangePageChannel(const FAvaPlaylistPageChangeChannel& InMessage, const TSharedRef<IMessageContext, ESPMode::ThreadSafe>& InContext);
+	void HandleUpdatePageFromRCP(const FAvaPlaylistUpdatePageFromRCP& InMessage, const TSharedRef<IMessageContext, ESPMode::ThreadSafe>& InContext);
+	void HandlePageAction(const FAvaPlaylistPageAction& InMessage, const TSharedRef<IMessageContext, ESPMode::ThreadSafe>& InContext);
+	void HandlePagePreviewAction(const FAvaPlaylistPagePreviewAction& InMessage, const TSharedRef<IMessageContext, ESPMode::ThreadSafe>& InContext);
+	void HandlePageActions(const FAvaPlaylistPageActions& InMessage, const TSharedRef<IMessageContext, ESPMode::ThreadSafe>& InContext);
+	void HandlePagePreviewActions(const FAvaPlaylistPagePreviewActions& InMessage, const TSharedRef<IMessageContext, ESPMode::ThreadSafe>& InContext);
+
+	void HandleGetChannel(const FAvaPlaylistGetChannel& InMessage, const TSharedRef<IMessageContext, ESPMode::ThreadSafe>& InContext);
+	void HandleGetChannels(const FAvaPlaylistGetChannels& InMessage, const TSharedRef<IMessageContext, ESPMode::ThreadSafe>& InContext);
+	void HandleChannelAction(const FAvaPlaylistChannelAction& InMessage, const TSharedRef<IMessageContext, ESPMode::ThreadSafe>& InContext);
+	void HandleAddChannelDevice(const FAvaPlaylistAddChannelDevice& InMessage, const TSharedRef<IMessageContext, ESPMode::ThreadSafe>& InContext);
+	void HandleEditChannelDevice(const FAvaPlaylistEditChannelDevice& InMessage, const TSharedRef<IMessageContext, ESPMode::ThreadSafe>& InContext);
+	void HandleRemoveChannelDevice(const FAvaPlaylistRemoveChannelDevice& InMessage, const TSharedRef<IMessageContext, ESPMode::ThreadSafe>& InContext);
+	void HandleGetChannelImage(const FAvaPlaylistGetChannelImage& InMessage, const TSharedRef<IMessageContext, ESPMode::ThreadSafe>& InContext);
+	
+	void HandleGetDevices(const FAvaPlaylistGetDevices& InMessage, const TSharedRef<IMessageContext, ESPMode::ThreadSafe>& InContext);
+	
+	void LogAndSendMessage(const FMessageAddress& InSender, int32 InRequestId, ELogVerbosity::Type InVerbosity, const TCHAR* InFormat, ...) const;
+	void SendMessage(const FMessageAddress& InSender, int32 InRequestId, ELogVerbosity::Type InVerbosity, const TCHAR* InFormat, ...) const;
+	void SendMessageImpl(const FMessageAddress& InSender, int32 InRequestId, ELogVerbosity::Type InVerbosity, const TCHAR* InMsg) const;
+
+	void RegisterConsoleCommands();
+	
+	void ShowStatusCommand(const TArray<FString>& InArgs);
+
+protected:
+	struct FRequestInfo
+	{
+		int32 RequestId;
+		FMessageAddress Sender;
+	};
+	
+	struct FChannelImage;
+	void FinishGetChannelImage(const FRequestInfo& InRequestInfo, const TSharedPtr<FChannelImage>& InChannelImage);
+
+	void HandlePageActions(const FRequestInfo& InRequestInfo, const TArray<int32>& InPageIds,
+		bool bInIsPreview, FName InPreviewChannelName, EAvaPlaylistPageActions InAction) const;
+	
+	/**
+	 * Helper function to retrieve the appropriate playlist for editing commands.
+	 */
+	UAvalanchePlaylist* GetOrLoadPlaylistForEdit(const FMessageAddress& InSender, int32 InRequestId, const FString& InPlaylistPath);
+
+	template<typename MessageType>
+	void SendResponse(MessageType* InMessage, const FMessageAddress& InRecipient, EMessageFlags InFlags = EMessageFlags::None) const
+	{
+		if (MessageEndpoint)
+		{
+			MessageEndpoint->Send(InMessage, InRecipient);
+		}
+	}
+	
+	template<typename MessageType>
+	void SendResponse(MessageType* InMessage, const TArray<FMessageAddress>& InRecipients, EMessageFlags InFlags = EMessageFlags::None) const
+	{
+		if (MessageEndpoint)
+		{
+			MessageEndpoint->Send(InMessage, MessageType::StaticStruct(), InFlags, nullptr,
+				InRecipients, FTimespan::Zero(), FDateTime::MaxValue());
+		}
+	}
+
+	void OnMessageBusNotification(const FMessageBusNotification& InNotification);
+
+private:
+	FString HostName;	
+	TSharedPtr<FMessageEndpoint, ESPMode::ThreadSafe> MessageEndpoint;
+	TArray<IConsoleObject*> ConsoleCommands;
+
+	/** Keep information on connected clients. */
+	struct FClientInfo
+	{
+		FMessageAddress Address;
+		/** Api version for communication with this client. */
+		int32 ApiVersion = -1;
+
+		explicit FClientInfo(const FMessageAddress& InAddress) : Address(InAddress) {}
+	};
+
+	/** Keep track of remote clients context information. */
+	TMap<FMessageAddress, TSharedPtr<FClientInfo>> Clients;
+
+	/** Array of just the client addresses for sending responses. */
+	TArray<FMessageAddress> ClientAddresses;
+	
+	FClientInfo* GetClientInfo(const FMessageAddress& InAddress) const
+	{
+		const TSharedPtr<FClientInfo>* ClientInfoPtr = Clients.Find(InAddress);
+		return ClientInfoPtr ? ClientInfoPtr->Get() : nullptr;
+	}
+	
+	FClientInfo& GetOrAddClientInfo(const FMessageAddress& InAddress)
+	{
+		if (const TSharedPtr<FClientInfo>* ExistingClientInfo = Clients.Find(InAddress))
+		{
+			return *ExistingClientInfo->Get();
+		}
+
+		const TSharedPtr<FClientInfo> NewClientInfo = MakeShared<FClientInfo>(InAddress);
+		Clients.Add(InAddress, NewClientInfo);
+		RefreshClientAddresses();
+		return *NewClientInfo;
+	}
+
+	void RefreshClientAddresses();
+	
+	/** Pool of images that can be recycled. */
+	TArray<TSharedPtr<FChannelImage>> AvailableChannelImages;
+	
+	/**
+	 * Keeps a current playlist cached for commands operating on playlist.
+	 * There is only one "current" playlist at a given time.
+	 */
+	struct FPlaylistCache
+	{
+		/** Currently loaded/cached playlist's path. */
+		FSoftObjectPath CurrentPlaylistPath;
+		
+		/** Currently loaded/cached playlist object. */
+		TStrongObjectPtr<UAvalanchePlaylist> CurrentPlaylist;
+
+		FDelegateHandle OnPlaybackInstanceStatusChangedDelegateHandle;
+		
+		typedef TFunctionRef<void(UAvalanchePlaylist*)> FPlaylistEventFunction;
+		
+		/**
+		 * Returns requested playlist specified by InPlaylistPath. Will load it if necessary or returned the cached one if it is the same.
+		 * If the new playlist fails to load, the return value is nullptr, and the previous playlist will remain loaded.
+		 * @param InPlaylistPath	Requested playlist path. 
+		 * @param InUnloadCurrentPlaylistFunction	Function called, with the previously cached playlist, when a new playlist is loaded
+		 *											and previous playlist needs to be unloaded.
+		 * @param InNewPlaylistLoadedFunction	Function called, with the new playlist, when a new playlist is loaded.
+		 */
+		UAvalanchePlaylist* GetOrLoadPlaylist(const FSoftObjectPath& InPlaylistPath, 
+			const FPlaylistEventFunction InUnloadCurrentPlaylistFunction,
+			const FPlaylistEventFunction InNewPlaylistLoadedFunction);
+
+		void SetupPlaylistDelegates(FAvaPlaylistServer* InPlaylistServer, UAvalanchePlaylist* InPlaylist);
+		void RemovePlaylistDelegates(const FAvaPlaylistServer* InPlaylistServer, UAvalanchePlaylist* InPlaylist) const;
+	};
+	
+	/** Cached data for page editing commands (GetPages and GetPageDetails). */
+	struct FPlaylistEditCommandData : public FPlaylistCache
+	{
+		/** PageId of the current managed ava blueprint. */
+		int32 ManagedPageId = FAvalanchePage::InvalidPageId;
+		TSharedPtr<FAvalancheManagedInstance> ManagedInstance;
+
+		~FPlaylistEditCommandData();
+		
+		/**
+		 * Checks if previous RCP was registered.
+		 * If so, save modified values to corresponding page.
+ 		 * This may result in the Playlist to be modified.
+ 		 * Will also unregister RCP from RC Module if requested.
+		 */
+		void SaveCurrentRemoteControlPresetToPage(bool bInUnregister);
+	};
+	FPlaylistEditCommandData PlaylistEditCommandData;
+
+	/** Cached data for playback commands (LoadPlaylist and PageAction). */
+	struct FPlaylistPlaybackCommandData: public FPlaylistCache
+	{
+		~FPlaylistPlaybackCommandData();
+
+		void ClosePlaybackContext();
+	};
+	FPlaylistPlaybackCommandData PlaylistPlaybackCommandData;
+};
diff --git a/Engine/Plugins/Experimental/Avalanche/Source/AvalancheMediaEditor/Public/Playlist/Factories/Filters/AvaPlaylistFactoriesUtils.h b/Engine/Plugins/Experimental/Avalanche/Source/AvalancheMediaEditor/Public/Playlist/Factories/Filters/AvaPlaylistFactoriesUtils.h
new file mode 100644
--- /dev/null
+++ b/Engine/Plugins/Experimental/Avalanche/Source/AvalancheMediaEditor/Public/Playlist/Factories/Filters/AvaPlaylistFactoriesUtils.h
@@ -0,0 +1,15 @@
+// Copyright Epic Games, Inc. All Rights Reserved.
+
+#pragma once
+
+#include "HAL/Platform.h"
+#include "Misc/EnumClassFlags.h"
+
+/** Type of the Page list */
+enum class EAvaPlaylistSearchListType : uint8
+{
+	None      = 0x00, // None
+	Template  = 0x01, // Template pages type
+	Instanced = 0x10 // Instanced pages type
+};
+ENUM_CLASS_FLAGS(EAvaPlaylistSearchListType)
diff --git a/Engine/Plugins/Experimental/Avalanche/Source/AvalancheMediaEditor/Public/Playlist/Factories/Filters/IAvaPlaylistFilterExpressionFactory.h b/Engine/Plugins/Experimental/Avalanche/Source/AvalancheMediaEditor/Public/Playlist/Factories/Filters/IAvaPlaylistFilterExpressionFactory.h
new file mode 100644
--- /dev/null
+++ b/Engine/Plugins/Experimental/Avalanche/Source/AvalancheMediaEditor/Public/Playlist/Factories/Filters/IAvaPlaylistFilterExpressionFactory.h
@@ -0,0 +1,72 @@
+// Copyright Epic Games, Inc. All Rights Reserved.
+
+#pragma once
+
+#include "Misc/TextFilterUtils.h"
+#include "HAL/Platform.h"
+#include "Templates/SharedPointer.h"
+#include "Templates/UnrealTypeTraits.h"
+#include "UObject/NameTypes.h"
+
+class UAvalanchePlaylist;
+struct FAvalanchePage;
+enum class EAvaPlaylistSearchListType : uint8;
+
+/**
+ * Holds the arguments required for the expression evaluation.
+ */
+struct FAvaPlaylistTextFilterArgs
+{
+	/** Item Playlist */
+	const UAvalanchePlaylist* ItemPlaylist;
+
+	/** Second Value */
+	FTextFilterString ValueToCheck;
+
+	/** Comparison Operation requested (Equal/Greater/etc...) still needed for certain Filter */
+	ETextFilterComparisonOperation ComparisonOperation;
+
+	/** Comparison Mode requested see ETextFilterTextComparisonMode for more information */
+	ETextFilterTextComparisonMode ComparisonMode;
+};
+
+class IAvaPlaylistFilterExpressionFactory : public TSharedFromThis<IAvaPlaylistFilterExpressionFactory>
+{
+public:
+	virtual ~IAvaPlaylistFilterExpressionFactory() {}
+
+	/** Get the filter identifier for this factory */
+	virtual FName GetFilterIdentifier() const = 0;
+
+	/**
+ 	 * Create and return an Instance of the AvaPlaylistFilterExpressionFactory Requested
+ 	 * @tparam InPlaylistFilterExpressionFactoryType The type of the factory to instantiate
+ 	 * @param InArgs Additional Args for constructor of the AvaPlaylistFilterExpressionFactory class if needed
+ 	 * @return The new playlist filter factory
+ 	 */
+	template <
+		typename InPlaylistFilterExpressionFactoryType,
+		typename... InArgsType
+		UE_REQUIRES(TIsDerivedFrom<InPlaylistFilterExpressionFactoryType, IAvaPlaylistFilterExpressionFactory>::Value)
+	>
+	static TSharedRef<InPlaylistFilterExpressionFactoryType> MakeInstance(InArgsType&&... InArgs)
+	{
+		return MakeShared<InPlaylistFilterExpressionFactoryType>(Forward<InArgsType>(InArgs)...);
+	}
+
+	/**
+	 * Evaluate the expression with the Value given
+	 * @param InItem Current Item being checked
+	 * @param InArgs The argument containing the data to evaluate the expression see FAvaTextFilterArgs for more information
+	 * @return True if the expression evaluate to True, False otherwise
+	 */
+	virtual bool FilterExpression(const FAvalanchePage& InItem, const FAvaPlaylistTextFilterArgs& InArgs) const = 0;
+
+	/**
+	 * Whether the factory support a given Comparison Operation
+	 * @param InComparisonOperation Comparison Operation to check
+	 * @param InPlaylistSearchListType Type of the Search List either Template or Instanced
+	 * @return True if Factory support passed Comparison Operation, False otherwise
+	 */
+	virtual bool SupportsComparisonOperation(ETextFilterComparisonOperation InComparisonOperation, EAvaPlaylistSearchListType InPlaylistSearchListType) const = 0;
+};
diff --git a/Engine/Plugins/Experimental/Avalanche/Source/AvalancheMediaEditor/Public/Playlist/Factories/Filters/IAvaPlaylistFilterSuggestionFactory.h b/Engine/Plugins/Experimental/Avalanche/Source/AvalancheMediaEditor/Public/Playlist/Factories/Filters/IAvaPlaylistFilterSuggestionFactory.h
new file mode 100644
--- /dev/null
+++ b/Engine/Plugins/Experimental/Avalanche/Source/AvalancheMediaEditor/Public/Playlist/Factories/Filters/IAvaPlaylistFilterSuggestionFactory.h
@@ -0,0 +1,74 @@
+// Copyright Epic Games, Inc. All Rights Reserved.
+
+#pragma once
+
+#include "HAL/Platform.h"
+#include "Playlist/AvaPlaylistDefines.h"
+#include "Templates/SharedPointer.h"
+#include "Templates/UnrealTypeTraits.h"
+
+class UAvalanchePlaylist;
+struct FAvalanchePage;
+struct FAssetSearchBoxSuggestion;
+enum class EAvaPlaylistSearchListType : uint8;
+
+/**
+ * Holds the arguments required for the suggestion factory to get the suggestion.
+ */
+struct FAvaPlaylistFilterSuggestionPayload
+{
+	/** Holds all the suggestion that will be shown */
+	TArray<FAssetSearchBoxSuggestion>& PossibleSuggestions;
+
+	/** The current string value written in the searchbox */
+	const FString FilterValue;
+
+	/** Current item to check */
+	int32 ItemPageId = UE::AvalanchePlaylist::InvalidPageId;
+
+	/** Current item playlist */
+	const UAvalanchePlaylist* Playlist = nullptr;
+
+	/** Used to speed up the check on whether a suggestion is already added */
+	TSet<FString>& FilterCache;
+};
+
+class IAvaPlaylistFilterSuggestionFactory : public TSharedFromThis<IAvaPlaylistFilterSuggestionFactory>
+{
+public:
+	virtual ~IAvaPlaylistFilterSuggestionFactory() {}
+
+	/**
+	  * Create and return an Instance of the AvaPlaylistFilterSuggestionFactory Requested
+	  * @tparam InPlaylistSuggestionFactoryType The type of the factory to instantiate
+	  * @param InArgs Additional Args for constructor of the AvaPlaylistFilterSuggestionFactory class if needed
+	  * @return The new suggestion factory
+	  */
+	template <
+		typename InPlaylistSuggestionFactoryType,
+		typename... InArgsType
+		UE_REQUIRES(TIsDerivedFrom<InPlaylistSuggestionFactoryType, IAvaPlaylistFilterSuggestionFactory>::Value)
+	>
+	static TSharedRef<InPlaylistSuggestionFactoryType> MakeInstance(InArgsType&&... InArgs)
+	{
+		return MakeShared<InPlaylistSuggestionFactoryType>(Forward<InArgsType>(InArgs)...);
+	}
+
+	/** Get the suggestion identifier for this factory */
+	virtual FName GetSuggestionIdentifier() const = 0;
+
+	/** True if the suggestion doesn't need to use the item to add its suggestions */
+	virtual bool IsSimpleSuggestion() const = 0;
+	/**
+	 * Add suggestions entry to be shown
+	 * @param InPayload Arguments to be passed to the factory to get the suggestion to show
+	 */
+	virtual void AddSuggestion(const TSharedRef<FAvaPlaylistFilterSuggestionPayload>& InPayload) = 0;
+
+	/**
+	 * Whether the factory support a given type of suggestion
+	 * @param InSuggestionType Suggestion type to check
+	 * @return True if Factory support passed Suggestion Type, False otherwise
+	 */
+	virtual bool SupportSuggestionType(EAvaPlaylistSearchListType InSuggestionType) const = 0;
+};
diff --git a/Engine/Plugins/Experimental/Avalanche/Source/AvalancheMediaEditor/Public/Playlist/Pages/AvaPlaylistPageContext.h b/Engine/Plugins/Experimental/Avalanche/Source/AvalancheMediaEditor/Public/Playlist/Pages/AvaPlaylistPageContext.h
new file mode 100644
--- /dev/null
+++ b/Engine/Plugins/Experimental/Avalanche/Source/AvalancheMediaEditor/Public/Playlist/Pages/AvaPlaylistPageContext.h
@@ -0,0 +1,37 @@
+// Copyright Epic Games, Inc. All Rights Reserved.
+
+#pragma once
+
+#include "AvaMediaDefines.h"
+#include "Templates/SharedPointer.h"
+#include "UObject/Object.h"
+#include "AvaPlaylistPageContext.generated.h"
+
+class FAvaPlaylistEditor;
+class FAvaPlaylistPageContextMenu;
+
+UCLASS(MinimalAPI)
+class UAvaPlaylistPageContext : public UObject
+{
+	GENERATED_BODY()
+
+	friend FAvaPlaylistPageContextMenu;
+
+public:
+	TSharedPtr<FAvaPlaylistEditor> GetPlaylistEditor() const
+	{
+		return PlaylistEditorWeak.Pin();
+	}
+
+	const FAvaPageListReference& GetPageListReference() const
+	{
+		return PageListReference;
+	}
+
+private:
+	TWeakPtr<FAvaPlaylistPageContextMenu> ContextMenuWeak;
+
+	TWeakPtr<FAvaPlaylistEditor> PlaylistEditorWeak;
+
+	FAvaPageListReference PageListReference;
+};
diff --git a/Engine/Plugins/Experimental/Avalanche/Source/AvalancheMediaEditor/Public/Playlist/Pages/Columns/AvaPageViewColumn.h b/Engine/Plugins/Experimental/Avalanche/Source/AvalancheMediaEditor/Public/Playlist/Pages/Columns/AvaPageViewColumn.h
new file mode 100644
--- /dev/null
+++ b/Engine/Plugins/Experimental/Avalanche/Source/AvalancheMediaEditor/Public/Playlist/Pages/Columns/AvaPageViewColumn.h
@@ -0,0 +1,25 @@
+// Copyright Epic Games, Inc. All Rights Reserved.
+
+#pragma once
+
+#include "AvaType.h"
+#include "CoreMinimal.h"
+#include "Playlist/AvaPlaylistDefines.h"
+#include "Widgets/Views/SHeaderRow.h"
+
+class SAvaPageViewRow;
+
+class IAvaPageViewColumn : public IAvaTypeCastable, public TSharedFromThis<IAvaPageViewColumn>
+{
+public:
+	UE_AVA_INHERITS(IAvaPageViewColumn, IAvaTypeCastable);
+
+	FName GetColumnId() const { return GetTypeId().ToName(); }
+
+	virtual FText GetColumnDisplayNameText() const = 0;
+	virtual FText GetColumnToolTipText() const = 0;
+
+	virtual SHeaderRow::FColumn::FArguments ConstructHeaderRowColumn() = 0;
+	
+	virtual TSharedRef<SWidget> ConstructRowWidget(const FAvaPageViewRef& InPageView, const TSharedPtr<SAvaPageViewRow>& InRow) = 0;
+};
diff --git a/Engine/Plugins/Experimental/Avalanche/Source/AvalancheMediaEditor/Public/Playlist/Pages/PageViews/AvaInstancedPageView.h b/Engine/Plugins/Experimental/Avalanche/Source/AvalancheMediaEditor/Public/Playlist/Pages/PageViews/AvaInstancedPageView.h
new file mode 100644
--- /dev/null
+++ b/Engine/Plugins/Experimental/Avalanche/Source/AvalancheMediaEditor/Public/Playlist/Pages/PageViews/AvaInstancedPageView.h
@@ -0,0 +1,25 @@
+// Copyright Epic Games, Inc. All Rights Reserved.
+
+#pragma once
+
+#include "AvaPageView.h"
+
+enum class ECheckBoxState : uint8;
+struct FAvalanchePage;
+
+class IAvaInstancedPageView
+{
+public:
+	UE_AVA_TYPE(IAvaInstancedPageView);
+
+	virtual ~IAvaInstancedPageView() = default;
+
+	virtual ECheckBoxState IsEnabled() const = 0;
+	virtual void SetEnabled(ECheckBoxState InState) = 0;
+
+	virtual FName GetChannelName() const = 0;
+	virtual bool SetChannel(FName InChannel) = 0;
+
+	virtual const FAvalanchePage& GetTemplate() const = 0;
+	virtual FText GetTemplateDescription() const = 0;
+};
diff --git a/Engine/Plugins/Experimental/Avalanche/Source/AvalancheMediaEditor/Public/Playlist/Pages/PageViews/AvaPageView.h b/Engine/Plugins/Experimental/Avalanche/Source/AvalancheMediaEditor/Public/Playlist/Pages/PageViews/AvaPageView.h
new file mode 100644
--- /dev/null
+++ b/Engine/Plugins/Experimental/Avalanche/Source/AvalancheMediaEditor/Public/Playlist/Pages/PageViews/AvaPageView.h
@@ -0,0 +1,67 @@
+// Copyright Epic Games, Inc. All Rights Reserved.
+
+#pragma once
+
+#include "AvaType.h"
+#include "CoreMinimal.h"
+
+class FReply;
+class UAvalanchePlaylist;
+struct FAssetData;
+
+enum class EAvaPageActionState : uint8
+{
+	None,
+	Requested,
+	Completed,
+	Cancelled,
+};
+
+enum class EAvaPageViewSelectionChangeType : uint8
+{
+	Deselect,
+	ReplaceSelection,
+	AddToSelection
+};
+
+class IAvaPageView : public IAvaTypeCastable, public TSharedFromThis<IAvaPageView>
+{
+public:
+	UE_AVA_INHERITS(IAvaPageView, IAvaTypeCastable);
+
+	virtual UAvalanchePlaylist* GetPlaylist() const = 0;
+	
+	virtual int32 GetPageId() const = 0;
+	virtual FText GetPageIdText() const = 0;
+	virtual FText GetPageNameText() const = 0;
+	virtual FText GetPageTransitionLayerNameText() const = 0;
+
+	virtual FText GetPageSummary() const = 0;
+	virtual FText GetPageDescription() const = 0;
+
+	virtual bool IsTemplate() const = 0;
+
+	virtual bool HasObjectPath(const UAvalanchePlaylist* InPlaylist) const = 0;
+	virtual FSoftObjectPath GetObjectPath(const UAvalanchePlaylist* InPlaylist) const = 0;
+	virtual FText GetObjectName(const UAvalanchePlaylist* InPlaylist) const = 0;
+	virtual void OnObjectChanged(const FAssetData& InAssetData) = 0;
+
+	virtual bool Rename(const FText& InNewName) = 0;
+	virtual bool RenameFriendlyName(const FText& InNewName) = 0;
+	
+	DECLARE_MULTICAST_DELEGATE_OneParam(FOnPageAction, EAvaPageActionState);
+	virtual FOnPageAction& GetOnRename() = 0;
+	virtual FOnPageAction& GetOnRenumber() = 0;
+
+	virtual FReply OnAssetStatusButtonClicked() = 0;
+	virtual bool CanChangeAssetStatus() const = 0;
+
+	virtual FReply OnPreviewButtonClicked() = 0;
+	virtual bool CanPreview() const = 0;
+
+	virtual FReply OnPlayButtonClicked() = 0;
+	virtual bool CanPlay() const = 0;
+
+	virtual bool IsPageSelected() const = 0;
+	virtual bool SetPageSelection(EAvaPageViewSelectionChangeType InSelectionChangeType) = 0;
+};
diff --git a/Engine/Plugins/Experimental/Avalanche/Source/AvalancheMediaEditor/Public/Playlist/Pages/PageViews/AvaTemplatePageView.h b/Engine/Plugins/Experimental/Avalanche/Source/AvalancheMediaEditor/Public/Playlist/Pages/PageViews/AvaTemplatePageView.h
new file mode 100644
--- /dev/null
+++ b/Engine/Plugins/Experimental/Avalanche/Source/AvalancheMediaEditor/Public/Playlist/Pages/PageViews/AvaTemplatePageView.h
@@ -0,0 +1,13 @@
+// Copyright Epic Games, Inc. All Rights Reserved.
+
+#pragma once
+
+#include "AvaType.h"
+
+class IAvaTemplatePageView
+{
+public:
+	UE_AVA_TYPE(IAvaTemplatePageView);
+
+	virtual ~IAvaTemplatePageView() = default;
+};
diff --git a/Engine/Plugins/Experimental/Avalanche/Source/AvalancheMediaEditor/Public/Playlist/Pages/Slate/SAvaPageViewRow.h b/Engine/Plugins/Experimental/Avalanche/Source/AvalancheMediaEditor/Public/Playlist/Pages/Slate/SAvaPageViewRow.h
new file mode 100644
--- /dev/null
+++ b/Engine/Plugins/Experimental/Avalanche/Source/AvalancheMediaEditor/Public/Playlist/Pages/Slate/SAvaPageViewRow.h
@@ -0,0 +1,61 @@
+// Copyright Epic Games, Inc. All Rights Reserved.
+
+#pragma once
+
+#include "CoreMinimal.h"
+#include "Playlist/AvaPlaylistDefines.h"
+#include "Playlist/AvaPlaylistDefines.h"
+#include "Widgets/Views/SListView.h"
+#include "Widgets/Views/STableRow.h"
+
+class SAvaPageList;
+
+enum class EAvalancheRowState
+{
+	Active,
+	Disabled
+};
+
+class FAvaPageViewRowDragDropOp : public FDragDropOperation, public TSharedFromThis<FAvaPageViewRowDragDropOp>
+{
+public:
+	DRAG_DROP_OPERATOR_TYPE(FAvaPageViewRowDragDropOp, FDragDropOperation)
+
+	FName Color;
+
+	virtual TSharedPtr<SWidget> GetDefaultDecorator() const override;
+
+	static TSharedRef<FAvaPageViewRowDragDropOp> New(const TSharedPtr<SAvaPageList>& InPageList);
+
+	TSharedPtr<SAvaPageList> GetPageList() const { return PageListWeak.Pin(); }
+	const TArray<int32>& GetDraggedIds() const { return DraggedIds; }
+
+protected:
+	TWeakPtr<SAvaPageList> PageListWeak;
+	TArray<int32> DraggedIds;
+};
+
+class AVALANCHEMEDIAEDITOR_API SAvaPageViewRow : public SMultiColumnTableRow<FAvaPageViewPtr>
+{
+public:	
+	SLATE_BEGIN_ARGS(SAvaPageViewRow){}
+	SLATE_END_ARGS()
+
+	/** Constructs this widget with InArgs */
+	void Construct(const FArguments& InArgs, FAvaPageViewPtr InPageView, const TSharedRef<SAvaPageList>& InPageList);
+
+	virtual TSharedRef<SWidget> GenerateWidgetForColumn(const FName& InColumnName) override;
+
+	virtual const FSlateBrush* GetBorder() const override;
+
+protected:
+	FReply OnPageViewDragDetected(const FGeometry& InMyGeometry, const FPointerEvent& InPointerEvent);
+	TOptional<EItemDropZone> OnPageViewCanAcceptDrop(const FDragDropEvent& InDragDropEvent, EItemDropZone InDropZone, FAvaPageViewPtr InItem);
+	FReply OnPageViewAcceptDrop(const FDragDropEvent& InDragDropEvent, EItemDropZone InDropZone, FAvaPageViewPtr InItem);
+
+	TWeakPtr<IAvaPageView> PageViewWeak;
+	TWeakPtr<SAvaPageList> PageListWeak;
+
+private:
+	const FSlateBrush* GetRowColorBrush(EAvalancheRowState InRowState) const;
+};
diff --git a/Engine/Plugins/Experimental/Avalanche/Source/AvalancheModifiers/Private/AvaModifiersActorUtils.cpp b/Engine/Plugins/Experimental/Avalanche/Source/AvalancheModifiers/Private/AvaModifiersActorUtils.cpp
new file mode 100644
--- /dev/null
+++ b/Engine/Plugins/Experimental/Avalanche/Source/AvalancheModifiers/Private/AvaModifiersActorUtils.cpp
@@ -0,0 +1,251 @@
+// Copyright Epic Games, Inc. All Rights Reserved.
+
+#include "AvaModifiersActorUtils.h"
+
+#include "AvaDefs.h"
+#include "Components/PrimitiveComponent.h"
+#include "GameFramework/Actor.h"
+
+#if WITH_EDITOR
+#include "AvaOutliner.h"
+#include "AvaOutlinerUtils.h"
+#endif
+
+FBox FAvaModifiersActorUtils::GetActorsBounds(const TSet<TWeakObjectPtr<AActor>>& InActors, const FTransform& InReferenceTransform, bool bInSkipHidden)
+{
+	FBox ActorsBounds = FBox(EForceInit::ForceInit);
+	ActorsBounds.IsValid = 0;
+	
+	FVector OrientedVerts[8];
+	
+	FTransform AccumulatedTransform = InReferenceTransform;
+
+	for (const TWeakObjectPtr<AActor>& ActorWeak : InActors)
+	{
+		const AActor* Actor = ActorWeak.Get();
+		
+		if (!Actor)
+		{
+			continue;
+		}
+		
+		if (bInSkipHidden && !IsActorVisible(Actor))
+		{
+			continue;
+		}
+
+		FBox ActorBounds = FAvaModifiersActorUtils::GetActorBounds(Actor);
+
+		if (ActorBounds.IsValid > 0)
+		{
+			FTransform ActorTransform = Actor->GetTransform();
+			ActorTransform.SetScale3D(FVector::OneVector);
+
+			const FOrientedBox OrientedBox = FAvaModifiersActorUtils::GetOrientedBox(ActorBounds, ActorTransform);
+
+			OrientedBox.CalcVertices(OrientedVerts);
+
+			for (const FVector& OrientedVert : OrientedVerts)
+			{
+				ActorsBounds += AccumulatedTransform.InverseTransformPositionNoScale(OrientedVert);
+			}
+
+			ActorsBounds.IsValid = 1;
+		}
+	}
+	
+	AccumulatedTransform.SetScale3D(FVector::OneVector);
+	
+	return ActorsBounds.TransformBy(AccumulatedTransform);
+}
+
+FBox FAvaModifiersActorUtils::GetActorsBounds(AActor* InActor, bool bInIncludeChildren, bool bInSkipHidden)
+{
+	if (!InActor)
+	{
+		return FBox();
+	}
+	
+	TSet<TWeakObjectPtr<AActor>> AttachedModifyActors {InActor};
+
+	if (bInIncludeChildren)
+	{
+		TArray<AActor*> AttachedActors;
+		InActor->GetAttachedActors(AttachedActors, false, true);
+		Algo::Transform(AttachedActors, AttachedModifyActors, [](AActor* InAttachedActor)
+		{
+			return InAttachedActor;
+		});
+	}
+	
+	return FAvaModifiersActorUtils::GetActorsBounds(AttachedModifyActors, InActor->GetActorTransform(), bInSkipHidden);
+}
+
+FBox FAvaModifiersActorUtils::GetActorBounds(const AActor* InActor)
+{	
+	FBox Box(ForceInit);
+	Box.IsValid = 0;
+	
+	if (!InActor || !InActor->GetRootComponent())
+	{
+		return Box;
+	}
+	
+	FTransform ActorToWorld = InActor->GetTransform();
+	ActorToWorld.SetScale3D(FVector::OneVector);
+	const FTransform WorldToActor = ActorToWorld.Inverse();
+
+	InActor->ForEachComponent<UPrimitiveComponent>(true, [&WorldToActor, &Box](const UPrimitiveComponent* InPrimitiveComponent)
+	{
+		if (!IsValid(InPrimitiveComponent))
+		{
+			return;
+		}
+		
+#if WITH_EDITOR
+		// Ignore Visualization Components, but don't consider them as failed components.
+		if (InPrimitiveComponent->IsVisualizationComponent())
+		{
+			return;
+		}
+#endif
+		
+		const FTransform ComponentToActor = InPrimitiveComponent->GetComponentTransform() * WorldToActor;
+		const FBox ComponentBox = InPrimitiveComponent->CalcBounds(ComponentToActor).GetBox();
+
+		Box += ComponentBox;
+		Box.IsValid = 1;
+	});
+
+	return Box;
+}
+
+FOrientedBox FAvaModifiersActorUtils::GetOrientedBox(const FBox& InLocalBox, const FTransform& InWorldTransform)
+{
+	FOrientedBox OutOrientedBox;
+	OutOrientedBox.Center = InWorldTransform.TransformPosition(InLocalBox.GetCenter());
+
+	OutOrientedBox.AxisX = InWorldTransform.TransformVector(FVector::UnitX());
+	OutOrientedBox.AxisY = InWorldTransform.TransformVector(FVector::UnitY());
+	OutOrientedBox.AxisZ = InWorldTransform.TransformVector(FVector::UnitZ());
+
+	OutOrientedBox.ExtentX = (InLocalBox.Max.X - InLocalBox.Min.X) / 2.f;
+	OutOrientedBox.ExtentY = (InLocalBox.Max.Y - InLocalBox.Min.Y) / 2.f;
+	OutOrientedBox.ExtentZ = (InLocalBox.Max.Z - InLocalBox.Min.Z) / 2.f;
+	return OutOrientedBox;
+}
+
+FVector FAvaModifiersActorUtils::GetVectorAxis(int32 InAxis)
+{
+	FVector FollowedAxisVector = FVector::ZeroVector;
+	
+	if (EnumHasAnyFlags(static_cast<EAvaModifiersAxis>(InAxis), EAvaModifiersAxis::X))
+	{
+		FollowedAxisVector.X = 1;
+	}
+	
+	if (EnumHasAnyFlags(static_cast<EAvaModifiersAxis>(InAxis), EAvaModifiersAxis::Y))
+	{
+		FollowedAxisVector.Y = 1;
+	}
+	
+	if (EnumHasAnyFlags(static_cast<EAvaModifiersAxis>(InAxis), EAvaModifiersAxis::Z))
+	{
+		FollowedAxisVector.Z = 1;
+	}
+
+	return FollowedAxisVector;
+}
+
+bool FAvaModifiersActorUtils::IsAxisVectorEquals(const FVector& InVectorA, const FVector& InVectorB, int32 InCompareAxis)
+{
+	// Compare InVectorA and InVectorB but only component of InCompareAxis
+	const FVector FollowedAxisVector = FAvaModifiersActorUtils::GetVectorAxis(InCompareAxis);
+	if ((InVectorA * FollowedAxisVector).Equals(InVectorB * FollowedAxisVector, 0.1))
+	{
+		return true;
+	}
+
+	return false;
+}
+
+bool FAvaModifiersActorUtils::IsActorNotIsolated(const AActor* InActor)
+{
+	if (!InActor)
+	{
+		return false;
+	}
+	
+#if WITH_EDITOR
+	bool bIsIsolatingActors = false;
+	TArray<TWeakObjectPtr<const AActor>> IsolatedActors;
+
+	TSharedPtr<FAvaOutliner> AvaOutliner = FAvaOutlinerUtils::EditorGetOutliner(InActor->GetWorld());
+	if (AvaOutliner.IsValid())
+	{
+		bIsIsolatingActors = FAvaOutlinerUtils::EditorActorIsolationInfo(AvaOutliner, IsolatedActors);
+	}
+
+	if (bIsIsolatingActors && !IsolatedActors.Contains(InActor))
+	{
+		return true;
+	}
+#endif
+
+	return false;
+}
+
+FRotator FAvaModifiersActorUtils::FindLookAtRotation(const FVector& InEyePosition, const FVector& InTargetPosition, const EAvaAxis InAxis, const bool bInFlipAxis)
+{
+	auto LookAtDirection = [&InEyePosition, &InTargetPosition, bInFlipAxis]() -> FVector
+	{
+		FVector OutDirection = bInFlipAxis ? (InEyePosition - InTargetPosition) : (InTargetPosition - InEyePosition);
+
+		OutDirection.Normalize();
+		if (OutDirection.SizeSquared() < 0.5f)
+		{
+			// Assert possible if OutDirection is not normalized.
+			OutDirection = FVector(1, 0, 0);
+		}
+
+		return OutDirection;
+	};
+
+	FMatrix NewRotation = FRotationMatrix::Identity;
+
+	switch (InAxis)
+	{
+	case EAvaAxis::Horizontal:
+		NewRotation = FRotationMatrix::MakeFromX(LookAtDirection());
+		break;
+	case EAvaAxis::Vertical:
+		NewRotation = FRotationMatrix::MakeFromY(LookAtDirection());
+		break;
+	case EAvaAxis::Depth:
+		NewRotation = FRotationMatrix::MakeFromZ(LookAtDirection());
+		break;
+	}
+
+	return NewRotation.Rotator();
+}
+
+bool FAvaModifiersActorUtils::IsActorVisible(const AActor* InActor)
+{
+	if (!InActor)
+	{
+		return false;
+	}
+
+	const bool bIsActorVisible = !InActor->IsHidden();
+
+#if WITH_EDITOR
+	const bool bIsActorVisibleInEditor = !InActor->IsTemporarilyHiddenInEditor();
+#else
+	constexpr bool bIsActorVisibleInEditor = true;
+#endif
+
+	// If we don't have a root component, consider this true to skip it in final condition
+	const bool bIsRootComponentVisible = !InActor->GetRootComponent() || InActor->GetRootComponent()->IsVisible();
+
+	return bIsActorVisible && bIsActorVisibleInEditor && bIsRootComponentVisible;
+}
diff --git a/Engine/Plugins/Experimental/Avalanche/Source/AvalancheModifiers/Private/AvaModifiersModule.cpp b/Engine/Plugins/Experimental/Avalanche/Source/AvalancheModifiers/Private/AvaModifiersModule.cpp
new file mode 100644
--- /dev/null
+++ b/Engine/Plugins/Experimental/Avalanche/Source/AvalancheModifiers/Private/AvaModifiersModule.cpp
@@ -0,0 +1,5 @@
+﻿// Copyright Epic Games, Inc. All Rights Reserved.
+
+#include "Modules/ModuleManager.h"
+
+IMPLEMENT_MODULE(FDefaultModuleImpl, AvalancheModifiers)
diff --git a/Engine/Plugins/Experimental/Avalanche/Source/AvalancheModifiers/Private/Extensions/AvaRenderStateUpdateModifierExtension.cpp b/Engine/Plugins/Experimental/Avalanche/Source/AvalancheModifiers/Private/Extensions/AvaRenderStateUpdateModifierExtension.cpp
new file mode 100644
--- /dev/null
+++ b/Engine/Plugins/Experimental/Avalanche/Source/AvalancheModifiers/Private/Extensions/AvaRenderStateUpdateModifierExtension.cpp
@@ -0,0 +1,169 @@
+// Copyright Epic Games, Inc. All Rights Reserved.
+
+#include "Extensions/AvaRenderStateUpdateModifierExtension.h"
+
+#include "AvaModifiersActorUtils.h"
+#include "Containers/Ticker.h"
+#include "Modifiers/ActorModifierCoreBase.h"
+
+FAvaRenderStateUpdateModifierExtension::FAvaRenderStateUpdateModifierExtension(IAvaRenderStateUpdateHandler* InExtensionHandler)
+	: ExtensionHandlerWeak(InExtensionHandler)
+{
+	check(InExtensionHandler);
+}
+
+void FAvaRenderStateUpdateModifierExtension::TrackActorVisibility(AActor* InActor)
+{
+	if (!InActor)
+	{
+		return;
+	}
+
+	if (IsActorVisibilityTracked(InActor))
+	{
+		return;
+	}
+
+	TrackedActorsVisibility.Add(InActor, FAvaModifiersActorUtils::IsActorVisible(InActor));
+}
+
+void FAvaRenderStateUpdateModifierExtension::UntrackActorVisibility(AActor* InActor)
+{
+	if (!InActor)
+	{
+		return;
+	}
+
+	TrackedActorsVisibility.Remove(InActor);
+}
+
+bool FAvaRenderStateUpdateModifierExtension::IsActorVisibilityTracked(AActor* InActor) const
+{
+	return InActor && TrackedActorsVisibility.Contains(InActor);
+}
+
+void FAvaRenderStateUpdateModifierExtension::SetTrackedActorsVisibility(const TSet<TWeakObjectPtr<AActor>>& InActors)
+{
+	// Remove unwanted actors, keep value for tracked ones
+	for (TMap<TWeakObjectPtr<AActor>, bool>::TIterator It(TrackedActorsVisibility); It; ++It)
+	{
+		if (!InActors.Contains(It->Key))
+		{
+			It.RemoveCurrent();
+		}
+	}
+
+	// Track value of wanted actors, will not overwrite already tracked one
+	for (const TWeakObjectPtr<AActor>& Actor : InActors)
+	{
+		TrackActorVisibility(Actor.Get());
+	}
+}
+
+void FAvaRenderStateUpdateModifierExtension::SetTrackedActorVisibility(AActor* InActor, bool bInIncludeChildren)
+{
+	if (!InActor)
+	{
+		return;
+	}
+
+	TSet<TWeakObjectPtr<AActor>> TrackActors {InActor};
+
+	if (bInIncludeChildren)
+	{
+		TArray<AActor*> AttachedActors;
+		InActor->GetAttachedActors(AttachedActors, false, true);
+		Algo::Transform(AttachedActors, TrackActors, [](AActor* InAttachedActor)
+		{
+			return InAttachedActor;
+		});
+	}
+
+	SetTrackedActorsVisibility(TrackActors);
+}
+
+void FAvaRenderStateUpdateModifierExtension::OnExtensionEnabled(EActorModifierCoreEnableReason InReason)
+{
+	BindDelegate();
+}
+
+void FAvaRenderStateUpdateModifierExtension::OnExtensionDisabled(EActorModifierCoreDisableReason InReason)
+{
+	UnbindDelegate();
+}
+
+void FAvaRenderStateUpdateModifierExtension::OnRenderStateDirty(UActorComponent& InComponent)
+{
+	AActor* ActorDirty = InComponent.GetOwner();
+	if (!ActorDirty)
+	{
+		return;
+	}
+
+	const UActorModifierCoreBase* Modifier = GetModifier();
+	if (!Modifier || !Modifier->IsModifierEnabled())
+	{
+		return;
+	}
+
+	if (ActorDirty->GetWorld() != GetModifierWorld())
+	{
+		return;
+	}
+
+	if (IAvaRenderStateUpdateHandler* HandlerInterface = ExtensionHandlerWeak.Get())
+	{
+		HandlerInterface->OnRenderStateUpdated(ActorDirty, &InComponent);
+	}
+
+	if (!TrackedActorsVisibility.Contains(ActorDirty))
+	{
+		return;
+	}
+
+	// Execute on next tick otherwise visibility data might not be up to date
+	TWeakObjectPtr<AActor> ActorDirtyWeak(ActorDirty);
+	FTSTicker::GetCoreTicker().AddTicker(FTickerDelegate::CreateSPLambda(this, [this, ActorDirtyWeak](float InDeltaSeconds)->bool
+	{
+		AActor* ActorDirty = ActorDirtyWeak.Get();
+
+		if (!ActorDirty)
+		{
+			return false;
+		}
+
+		bool* bOldVisibility = TrackedActorsVisibility.Find(ActorDirty);
+
+		if (!bOldVisibility)
+		{
+			return false;
+		}
+
+		const bool bNewVisibility = FAvaModifiersActorUtils::IsActorVisible(ActorDirty);
+		bool& bOldVisibilityRef = *bOldVisibility;
+
+		const bool bVisibilityChanged = bOldVisibilityRef != bNewVisibility;
+
+		// Set new visibility
+		bOldVisibilityRef = bNewVisibility;
+
+		IAvaRenderStateUpdateHandler* HandlerInterface = ExtensionHandlerWeak.Get();
+		if (bVisibilityChanged && HandlerInterface)
+		{
+			HandlerInterface->OnActorVisibilityChanged(ActorDirty);
+		}
+
+		return false;
+	}));
+}
+
+void FAvaRenderStateUpdateModifierExtension::BindDelegate()
+{
+	UActorComponent::MarkRenderStateDirtyEvent.RemoveAll(this);
+	UActorComponent::MarkRenderStateDirtyEvent.AddSP(this, &FAvaRenderStateUpdateModifierExtension::OnRenderStateDirty);
+}
+
+void FAvaRenderStateUpdateModifierExtension::UnbindDelegate()
+{
+	UActorComponent::MarkRenderStateDirtyEvent.RemoveAll(this);
+}
diff --git a/Engine/Plugins/Experimental/Avalanche/Source/AvalancheModifiers/Private/Extensions/AvaSceneTreeUpdateModifierExtension.cpp b/Engine/Plugins/Experimental/Avalanche/Source/AvalancheModifiers/Private/Extensions/AvaSceneTreeUpdateModifierExtension.cpp
new file mode 100644
--- /dev/null
+++ b/Engine/Plugins/Experimental/Avalanche/Source/AvalancheModifiers/Private/Extensions/AvaSceneTreeUpdateModifierExtension.cpp
@@ -0,0 +1,477 @@
+// Copyright Epic Games, Inc. All Rights Reserved.
+
+#include "Extensions/AvaSceneTreeUpdateModifierExtension.h"
+#include "AvaActorUtils.h"
+#include "AvaBlueprint.h"
+#include "AvaSceneTree.h"
+#include "GameFramework/Actor.h"
+
+#if WITH_EDITOR
+#include "AvaOutliner.h"
+#include "AvaOutlinerSubsystem.h"
+#include "AvaOutlinerUtils.h"
+#include "Engine/World.h"
+#endif
+
+FAvaSceneTreeUpdateModifierExtension::FAvaSceneTreeUpdateModifierExtension(IAvaSceneTreeUpdateHandler* InExtensionHandler)
+	: ExtensionHandlerWeak(InExtensionHandler)
+{
+	check(InExtensionHandler);
+}
+
+void FAvaSceneTreeUpdateModifierExtension::TrackSceneTree(int32 InTrackedActorIdx, FAvaSceneTreeActor* InTrackedActor)
+{
+	if (!InTrackedActor)
+	{
+		return;
+	}
+
+	InTrackedActor->LocalActorWeak = GetModifierActor();
+	InTrackedActor->ReferenceActorsWeak.Empty();
+	InTrackedActor->ReferenceActorChildrenWeak.Empty();
+	InTrackedActor->ReferenceActorParentsWeak.Empty();
+	InTrackedActor->ReferenceActorDirectChildrenWeak.Empty();
+
+	TrackedActors.Add(InTrackedActorIdx, InTrackedActor);
+
+	CheckTrackedActorUpdate(InTrackedActorIdx);
+}
+
+void FAvaSceneTreeUpdateModifierExtension::UntrackSceneTree(int32 InTrackedActorIdx)
+{
+	if (!TrackedActors.Contains(InTrackedActorIdx))
+	{
+		return;
+	}
+
+	TrackedActors.Remove(InTrackedActorIdx);
+}
+
+FAvaSceneTreeActor* FAvaSceneTreeUpdateModifierExtension::GetTrackedActor(int32 InTrackedActorIdx) const
+{
+	if (FAvaSceneTreeActor* const* TrackedActor = TrackedActors.Find(InTrackedActorIdx))
+	{
+		return *TrackedActor;
+	}
+
+	return nullptr;
+}
+
+void FAvaSceneTreeUpdateModifierExtension::CheckTrackedActorsUpdate() const
+{
+	// Container could change while in range iteration
+	TArray<int32> TrackedKeys;
+	TrackedActors.GenerateKeyArray(TrackedKeys);
+
+	for (const int32 Key : TrackedKeys)
+	{
+		CheckTrackedActorUpdate(Key);
+	}
+}
+
+void FAvaSceneTreeUpdateModifierExtension::OnExtensionEnabled(EActorModifierCoreEnableReason InReason)
+{
+	const UWorld* World = GetModifierWorld();
+	if (!IsValid(World))
+	{
+		return;
+	}
+
+	// When actor are destroyed in world
+	World->RemoveOnActorDestroyededHandler(WorldActorDestroyedDelegate);
+	WorldActorDestroyedDelegate = World->AddOnActorDestroyedHandler(FOnActorDestroyed::FDelegate::CreateSP(this, &FAvaSceneTreeUpdateModifierExtension::OnWorldActorDestroyed));
+
+	// Used to detect visibility changes in siblings
+	USceneComponent::MarkRenderStateDirtyEvent.RemoveAll(this);
+	USceneComponent::MarkRenderStateDirtyEvent.AddSP(this, &FAvaSceneTreeUpdateModifierExtension::OnRenderStateDirty);
+
+#if WITH_EDITOR
+	if (UAvaOutlinerSubsystem* const OutlinerSubsystem = World->GetSubsystem<UAvaOutlinerSubsystem>())
+	{
+		UAvaOutlinerSubsystem::FActorHierarchyChanged& ActorHierarchyChanged = OutlinerSubsystem->OnActorHierarchyChanged();
+
+		ActorHierarchyChanged.RemoveAll(this);
+		ActorHierarchyChanged.AddSP(this, &FAvaSceneTreeUpdateModifierExtension::OnActorHierarchyChanged);
+
+		if (const TSharedPtr<FAvaOutliner> Outliner = OutlinerSubsystem->GetOutliner())
+		{
+			Outliner->OnOutlinerLoaded.RemoveAll(this);
+			Outliner->OnOutlinerLoaded.AddSP(this, &FAvaSceneTreeUpdateModifierExtension::OnOutlinerLoaded);
+		}
+	}
+#endif
+
+	CheckTrackedActorsUpdate();
+}
+
+void FAvaSceneTreeUpdateModifierExtension::OnExtensionDisabled(EActorModifierCoreDisableReason InReason)
+{
+	USceneComponent::MarkRenderStateDirtyEvent.RemoveAll(this);
+
+	const UWorld* World = GetModifierWorld();
+	if (!IsValid(World))
+	{
+		return;
+	}
+
+	World->RemoveOnActorDestroyededHandler(WorldActorDestroyedDelegate);
+
+#if WITH_EDITOR
+	if (UAvaOutlinerSubsystem* const OutlinerSubsystem = World->GetSubsystem<UAvaOutlinerSubsystem>())
+	{
+		OutlinerSubsystem->OnActorHierarchyChanged().RemoveAll(this);
+		if (const TSharedPtr<FAvaOutliner> Outliner = OutlinerSubsystem->GetOutliner())
+		{
+			Outliner->OnOutlinerLoaded.RemoveAll(this);
+		}
+	}
+#endif
+}
+
+#if WITH_EDITOR
+void FAvaSceneTreeUpdateModifierExtension::OnOutlinerLoaded()
+{
+	CheckTrackedActorsUpdate();
+}
+
+void FAvaSceneTreeUpdateModifierExtension::OnActorHierarchyChanged(AActor* InActor, const AActor* InParentActor, EAvaOutlinerHierarchyChangeType InChangeType)
+{
+	if (!InActor)
+	{
+		return;
+	}
+
+	// Container could change while in range iteration
+	TArray<int32> TrackedKeys;
+	TrackedActors.GenerateKeyArray(TrackedKeys);
+
+	for (const int32 Key : TrackedKeys)
+	{
+		const FAvaSceneTreeActor* TrackedActor = GetTrackedActor(Key);
+		if (!TrackedActor || !TrackedActor->LocalActorWeak.IsValid())
+		{
+			continue;
+		}
+
+		CheckTrackedActorUpdate(Key);
+
+		if (TrackedActor)
+		{
+			const AActor* ReferenceActor = TrackedActor->ReferenceActorWeak.Get();
+			IAvaSceneTreeUpdateHandler* HandlerInterface = ExtensionHandlerWeak.Get();
+
+			if (HandlerInterface
+				&& IsValid(ReferenceActor)
+				&& TrackedActor->ReferenceContainer == EAvaReferenceContainer::Other
+				&& (InActor == ReferenceActor || InActor->IsAttachedTo(ReferenceActor)))
+			{
+				HandlerInterface->OnSceneTreeTrackedActorRearranged(Key, InActor);
+			}
+		}
+	}
+}
+#endif
+
+void FAvaSceneTreeUpdateModifierExtension::OnRenderStateDirty(UActorComponent& InComponent)
+{
+	AActor* OwningActor = InComponent.GetOwner();
+	if (!IsValid(OwningActor))
+	{
+		return;
+	}
+
+	if (GetModifierWorld() != OwningActor->GetWorld())
+	{
+		return;
+	}
+
+	// Container could change while in range iteration
+	TArray<int32> TrackedKeys;
+	TrackedActors.GenerateKeyArray(TrackedKeys);
+
+	for (const int32 Key : TrackedKeys)
+	{
+		const FAvaSceneTreeActor* TrackedActor = GetTrackedActor(Key);
+		if (!TrackedActor || !TrackedActor->LocalActorWeak.IsValid())
+		{
+			continue;
+		}
+
+		const AActor* ReferenceActor = TrackedActor->ReferenceActorWeak.Get();
+
+		const bool bIsReferenceActor = ReferenceActor == OwningActor || TrackedActor->ReferenceActorsWeak.Contains(OwningActor);
+		const bool bIsReferenceActorChild = ReferenceActor && OwningActor->IsAttachedTo(ReferenceActor);
+		const bool bIsReferenceActorSibling = ReferenceActor && OwningActor->GetAttachParentActor() == ReferenceActor->GetAttachParentActor();
+		const bool bIsReferenceActorParent = ReferenceActor && ReferenceActor->IsAttachedTo(OwningActor);
+
+		if (!bIsReferenceActor && !bIsReferenceActorChild && !bIsReferenceActorSibling && !bIsReferenceActorParent)
+		{
+			continue;
+		}
+
+		CheckTrackedActorUpdate(Key);
+	}
+}
+
+void FAvaSceneTreeUpdateModifierExtension::OnWorldActorDestroyed(AActor* InActor)
+{
+	if (!InActor)
+	{
+		return;
+	}
+
+	// Delay check from one tick to make sure actor is no longer attached
+	FTSTicker::GetCoreTicker().AddTicker(FTickerDelegate::CreateSPLambda(this, [this](float InDeltaSeconds)->bool
+	{
+		CheckTrackedActorsUpdate();
+		return false;
+	}));
+}
+
+bool FAvaSceneTreeUpdateModifierExtension::IsSameActorArray(const TArray<TWeakObjectPtr<AActor>>& InPreviousActorWeak, const TArray<TWeakObjectPtr<AActor>>& InNewActorWeak) const
+{
+	if (InPreviousActorWeak.Num() != InNewActorWeak.Num())
+	{
+		return false;
+	}
+
+	for (int32 Idx = 0; Idx < InPreviousActorWeak.Num(); Idx++)
+	{
+		if (InPreviousActorWeak[Idx].Get() != InNewActorWeak[Idx].Get())
+		{
+			return false;
+		}
+	}
+
+	return true;
+}
+
+TArray<TWeakObjectPtr<AActor>> FAvaSceneTreeUpdateModifierExtension::GetReferenceActors(const FAvaSceneTreeActor* InTrackedActor) const
+{
+	TArray<TWeakObjectPtr<AActor>> ReferenceActors;
+
+	if (!InTrackedActor)
+	{
+		return ReferenceActors;
+	}
+
+	AActor* LocalActor = InTrackedActor->LocalActorWeak.Get();
+	if (!IsValid(LocalActor))
+	{
+		return ReferenceActors;
+	}
+
+	if (InTrackedActor->ReferenceContainer == EAvaReferenceContainer::Other)
+	{
+		if (AActor* ReferenceActor = InTrackedActor->ReferenceActorWeak.Get())
+		{
+			ReferenceActors.Add(ReferenceActor);
+		}
+
+		return ReferenceActors;
+	}
+
+	AActor* ContextActor = LocalActor;
+	while (AActor* NewReferenceActor = FAvaActorUtils::ActorFromReferenceContainer(ContextActor, InTrackedActor->ReferenceContainer, false))
+	{
+		// Take only siblings
+		if (NewReferenceActor->GetAttachParentActor() != LocalActor->GetAttachParentActor())
+		{
+			break;
+		}
+
+		ReferenceActors.Add(NewReferenceActor);
+		ContextActor = NewReferenceActor;
+
+		if (InTrackedActor->bSkipHiddenActors
+				&& (NewReferenceActor->IsHidden()
+#if WITH_EDITOR
+				|| NewReferenceActor->IsTemporarilyHiddenInEditor()
+#endif
+				))
+		{
+			continue;
+		}
+
+		break;
+	}
+
+	return ReferenceActors;
+}
+
+TSet<TWeakObjectPtr<AActor>> FAvaSceneTreeUpdateModifierExtension::GetChildrenActorsRecursive(const AActor* InActor) const
+{
+	TSet<TWeakObjectPtr<AActor>> ChildrenWeak;
+
+	if (InActor)
+	{
+		TArray<AActor*> ReferenceAttachedActors;
+		InActor->GetAttachedActors(ReferenceAttachedActors, false, true);
+		Algo::Transform(ReferenceAttachedActors, ChildrenWeak, [](AActor* InAttachedActor)->TWeakObjectPtr<AActor>
+		{
+			return InAttachedActor;
+		});
+	}
+
+	return ChildrenWeak;
+}
+
+TArray<TWeakObjectPtr<AActor>> FAvaSceneTreeUpdateModifierExtension::GetDirectChildrenActor(AActor* InActor) const
+{
+	TArray<TWeakObjectPtr<AActor>> DirectChildrenWeak;
+
+	if (!IsValid(InActor))
+	{
+		return DirectChildrenWeak;
+	}
+
+	bool bIsOutlinerAttachedActors = false;
+	TArray<AActor*> DirectChildren;
+
+#if WITH_EDITOR
+	const UWorld* const World = InActor->GetWorld();
+	check(IsValid(World));
+
+	const UAvaOutlinerSubsystem* const OutlinerSubsystem = World->GetSubsystem<UAvaOutlinerSubsystem>();
+	if (IsValid(OutlinerSubsystem))
+	{
+		const TSharedPtr<FAvaOutliner> AvaOutliner = OutlinerSubsystem->GetOutliner();
+		if (AvaOutliner.IsValid())
+		{
+			DirectChildren = FAvaOutlinerUtils::EditorOutlinerChildActors(AvaOutliner, InActor);
+			bIsOutlinerAttachedActors = true;
+		}
+	}
+#endif
+
+	if (!bIsOutlinerAttachedActors)
+	{
+		if (const IAvaSceneInterface* SceneInterface = FAvaActorUtils::GetSceneInterfaceFromActor(InActor))
+		{
+			const FAvaSceneTree& SceneTree = SceneInterface->GetSceneTree();
+			DirectChildren = SceneTree.GetChildActors(InActor);
+			bIsOutlinerAttachedActors = true;
+		}
+	}
+
+	// Fallback for unsupported world
+	if (!bIsOutlinerAttachedActors)
+	{
+		 InActor->GetAttachedActors(DirectChildren, false, false);
+	}
+
+	Algo::Transform(DirectChildren, DirectChildrenWeak, [](AActor* InActor)->TWeakObjectPtr<AActor>
+	{
+		return InActor;
+	});
+
+	return DirectChildrenWeak;
+}
+
+TArray<TWeakObjectPtr<AActor>> FAvaSceneTreeUpdateModifierExtension::GetParentActors(const AActor* InActor) const
+{
+	TArray<TWeakObjectPtr<AActor>> ActorParentsWeak;
+
+	while (InActor && InActor->GetAttachParentActor())
+	{
+		AActor* ParentActor = InActor->GetAttachParentActor();
+		ActorParentsWeak.Add(ParentActor);
+		InActor = ParentActor;
+	}
+
+	return ActorParentsWeak;
+}
+
+void FAvaSceneTreeUpdateModifierExtension::CheckTrackedActorUpdate(int32 InIdx) const
+{
+	if (!IsExtensionEnabled())
+	{
+		return;
+	}
+
+	FAvaSceneTreeActor* TrackedActor = GetTrackedActor(InIdx);
+	if (!TrackedActor)
+	{
+		return;
+	}
+
+	AActor* LocalActor = TrackedActor->LocalActorWeak.IsValid()
+		? TrackedActor->LocalActorWeak.Get()
+		: GetModifierActor();
+
+	if (!LocalActor)
+	{
+		return;
+	}
+
+	// Reapply in case we overwrite the whole struct outside
+	TrackedActor->LocalActorWeak = LocalActor;
+
+	// Gather previous reference actor before clearing array
+	AActor* PreviousReferenceActor = !TrackedActor->ReferenceActorsWeak.IsEmpty()
+		? TrackedActor->ReferenceActorsWeak.Last().Get()
+		: nullptr;
+
+	TrackedActor->ReferenceActorsWeak.Empty();
+
+	AActor* NewReferenceActor = nullptr;
+
+	// Track siblings actors in case their visibility changes
+	for (const TWeakObjectPtr<AActor>& NewReferenceActorWeak : GetReferenceActors(TrackedActor))
+	{
+		AActor* ReferenceActor = NewReferenceActorWeak.Get();
+		if (!ReferenceActor)
+		{
+			continue;
+		}
+
+		TrackedActor->ReferenceActorsWeak.Add(ReferenceActor);
+		NewReferenceActor = ReferenceActor;
+	}
+
+	TrackedActor->ReferenceActorWeak = NewReferenceActor;
+
+	// Gather children of reference actor
+	const TSet<TWeakObjectPtr<AActor>> PreviousReferenceActorChildrenWeak = TrackedActor->ReferenceActorChildrenWeak;
+	TrackedActor->ReferenceActorChildrenWeak = GetChildrenActorsRecursive(NewReferenceActor);
+
+	// Gather direct children ordered of reference actor
+	const TArray<TWeakObjectPtr<AActor>> PreviousDirectChildrenWeak = TrackedActor->ReferenceActorDirectChildrenWeak;
+	TrackedActor->ReferenceActorDirectChildrenWeak = GetDirectChildrenActor(NewReferenceActor);
+
+	// Gather parents of reference actor
+	const TArray<TWeakObjectPtr<AActor>> PreviousParentActorsWeak = TrackedActor->ReferenceActorParentsWeak;
+	TrackedActor->ReferenceActorParentsWeak = GetParentActors(NewReferenceActor);
+
+	const bool bReferenceActorChanged = NewReferenceActor != PreviousReferenceActor;
+	const bool bReferenceActorChildrenChanged = !(TrackedActor->ReferenceActorChildrenWeak.Num() == PreviousReferenceActorChildrenWeak.Num()
+			&& TrackedActor->ReferenceActorChildrenWeak.Difference(PreviousReferenceActorChildrenWeak).IsEmpty());
+	const bool bReferenceActorDirectChildrenChanged = !IsSameActorArray(PreviousDirectChildrenWeak, TrackedActor->ReferenceActorDirectChildrenWeak);
+	const bool bReferenceActorParentChanged = !IsSameActorArray(PreviousParentActorsWeak, TrackedActor->ReferenceActorParentsWeak);
+
+	if (IAvaSceneTreeUpdateHandler* HandlerInterface = ExtensionHandlerWeak.Get())
+	{
+		// Fire event when reference actor changed
+		if (bReferenceActorChanged)
+		{
+			HandlerInterface->OnSceneTreeTrackedActorChanged(InIdx, PreviousReferenceActor, NewReferenceActor);
+		}
+
+		// Fire event when children actors changed
+		if (bReferenceActorChildrenChanged)
+		{
+			HandlerInterface->OnSceneTreeTrackedActorChildrenChanged(InIdx, PreviousReferenceActorChildrenWeak, TrackedActor->ReferenceActorChildrenWeak);
+		}
+
+		// Fire event when direct children actors changed (even order)
+		if (bReferenceActorDirectChildrenChanged)
+		{
+			HandlerInterface->OnSceneTreeTrackedActorDirectChildrenChanged(InIdx, PreviousDirectChildrenWeak, TrackedActor->ReferenceActorDirectChildrenWeak);
+		}
+
+		if (bReferenceActorParentChanged)
+		{
+			HandlerInterface->OnSceneTreeTrackedActorParentChanged(InIdx, PreviousParentActorsWeak, TrackedActor->ReferenceActorParentsWeak);
+		}
+	}
+}
diff --git a/Engine/Plugins/Experimental/Avalanche/Source/AvalancheModifiers/Private/Extensions/AvaTransformUpdateModifierExtension.cpp b/Engine/Plugins/Experimental/Avalanche/Source/AvalancheModifiers/Private/Extensions/AvaTransformUpdateModifierExtension.cpp
new file mode 100644
--- /dev/null
+++ b/Engine/Plugins/Experimental/Avalanche/Source/AvalancheModifiers/Private/Extensions/AvaTransformUpdateModifierExtension.cpp
@@ -0,0 +1,143 @@
+// Copyright Epic Games, Inc. All Rights Reserved.
+
+#include "Extensions/AvaTransformUpdateModifierExtension.h"
+
+#include "Modifiers/ActorModifierCoreBase.h"
+
+FAvaTransformUpdateModifierExtension::FAvaTransformUpdateModifierExtension(IAvaTransformUpdateHandler* InExtensionInterface)
+	: ExtensionHandlerWeak(InExtensionInterface)
+{
+	check(InExtensionInterface);
+}
+
+void FAvaTransformUpdateModifierExtension::TrackActor(AActor* InActor, bool bInReset)
+{
+	if (!InActor || !InActor->GetRootComponent())
+	{
+		return;
+	}
+
+	if (bInReset)
+	{
+		const TSet<TWeakObjectPtr<AActor>> NewSet {InActor};
+		UntrackActors(TrackedActors.Difference(NewSet));
+	}
+	
+	if (TrackedActors.Contains(InActor))
+	{
+		return;
+	}
+
+	if (IsExtensionEnabled())
+	{
+		BindDelegate(InActor);
+	}
+
+	TrackedActors.Add(InActor);
+}
+
+void FAvaTransformUpdateModifierExtension::UntrackActor(AActor* InActor)
+{
+	if (!InActor || !InActor->GetRootComponent())
+	{
+		return;
+	}
+	
+	if (!TrackedActors.Contains(InActor))
+	{
+		return;
+	}
+
+	UnbindDelegate(InActor);
+
+	TrackedActors.Remove(InActor);
+}
+
+void FAvaTransformUpdateModifierExtension::TrackActors(const TSet<TWeakObjectPtr<AActor>>& InActors, bool bInReset)
+{
+	if (bInReset)
+	{
+		UntrackActors(TrackedActors.Difference(InActors));	
+	}
+	
+	for (const TWeakObjectPtr<AActor>& InActor : InActors)
+	{
+		TrackActor(InActor.Get(), false);
+	}
+}
+
+void FAvaTransformUpdateModifierExtension::UntrackActors(const TSet<TWeakObjectPtr<AActor>>& InActors)
+{
+	for (const TWeakObjectPtr<AActor>& InActor : InActors)
+	{
+		UntrackActor(InActor.Get());
+	}
+}
+
+void FAvaTransformUpdateModifierExtension::OnExtensionEnabled(EActorModifierCoreEnableReason InReason)
+{
+	for (TWeakObjectPtr<AActor>& TrackedActor : TrackedActors)
+	{
+		BindDelegate(TrackedActor.Get());
+	}
+}
+
+void FAvaTransformUpdateModifierExtension::OnExtensionDisabled(EActorModifierCoreDisableReason InReason)
+{
+	for (TWeakObjectPtr<AActor>& TrackedActor : TrackedActors)
+	{
+		UnbindDelegate(TrackedActor.Get());
+	}
+}
+
+void FAvaTransformUpdateModifierExtension::OnTransformUpdated(USceneComponent* InComponent, EUpdateTransformFlags InFlags, ETeleportType InType)
+{
+	if (!InComponent)
+	{
+		return;
+	}
+
+	AActor* ActorTransformed = InComponent->GetOwner();
+	if (!ActorTransformed)
+	{
+		return;
+	}
+
+	const UActorModifierCoreBase* Modifier = GetModifier();
+	if (!Modifier || !Modifier->IsModifierEnabled())
+	{
+		return;
+	}
+	
+	if (IAvaTransformUpdateHandler* HandlerInterface = ExtensionHandlerWeak.Get())
+	{
+		HandlerInterface->OnTransformUpdated(ActorTransformed, InFlags == EUpdateTransformFlags::PropagateFromParent);
+	}
+}
+
+void FAvaTransformUpdateModifierExtension::BindDelegate(const AActor* InActor)
+{
+	if (!InActor)
+	{
+		return;
+	}
+	
+	if (USceneComponent* SceneComponent = InActor->GetRootComponent())
+	{
+		SceneComponent->TransformUpdated.RemoveAll(this);
+		SceneComponent->TransformUpdated.AddSP(this, &FAvaTransformUpdateModifierExtension::OnTransformUpdated);
+	}
+}
+
+void FAvaTransformUpdateModifierExtension::UnbindDelegate(const AActor* InActor)
+{
+	if (!InActor)
+	{
+		return;
+	}
+	
+	if (USceneComponent* SceneComponent = InActor->GetRootComponent())
+	{
+		SceneComponent->TransformUpdated.RemoveAll(this);
+	}
+}
diff --git a/Engine/Plugins/Experimental/Avalanche/Source/AvalancheModifiers/Private/Modifiers/AvaAlignBetweenModifier.cpp b/Engine/Plugins/Experimental/Avalanche/Source/AvalancheModifiers/Private/Modifiers/AvaAlignBetweenModifier.cpp
new file mode 100644
--- /dev/null
+++ b/Engine/Plugins/Experimental/Avalanche/Source/AvalancheModifiers/Private/Modifiers/AvaAlignBetweenModifier.cpp
@@ -0,0 +1,257 @@
+// Copyright Epic Games, Inc. All Rights Reserved.
+
+#include "Modifiers/AvaAlignBetweenModifier.h"
+
+#include "GameFramework/Actor.h"
+#include "Misc/ITransaction.h"
+#include "Shared/AvaTransformModifierShared.h"
+
+#define LOCTEXT_NAMESPACE "AvaAlignBetweenModifier"
+
+void UAvaAlignBetweenModifier::OnModifierAdded(EActorModifierCoreEnableReason InReason)
+{
+	Super::OnModifierAdded(InReason);
+
+	AddExtension<FAvaTransformUpdateModifierExtension>(this);
+}
+
+void UAvaAlignBetweenModifier::OnModifierEnabled(EActorModifierCoreEnableReason InReason)
+{
+	Super::OnModifierEnabled(InReason);
+
+	// Save actor layout state
+	if (UAvaTransformModifierShared* LayoutShared = GetShared<UAvaTransformModifierShared>(true))
+	{
+		LayoutShared->SaveActorState(this, GetModifiedActor());
+	}
+
+	SetTransformExtensionReferenceActors();
+}
+
+void UAvaAlignBetweenModifier::OnModifierDisabled(EActorModifierCoreDisableReason InReason)
+{
+	Super::OnModifierDisabled(InReason);
+
+	// Restore actor layout state
+	if (UAvaTransformModifierShared* LayoutShared = GetShared<UAvaTransformModifierShared>(false))
+	{
+		LayoutShared->RestoreActorState(this, GetModifiedActor());
+	}
+}
+
+void UAvaAlignBetweenModifier::Apply()
+{
+	AActor* const ModifyActor = GetModifiedActor();
+	
+	if (ReferenceActors.IsEmpty())
+	{
+		Next();
+		return;
+	}
+
+	const TSet<FAvaAlignBetweenWeightedActor> WeightedActors = GetEnabledReferenceActors();
+
+	float TotalWeight = 0.0f;
+	for (const FAvaAlignBetweenWeightedActor& WeightedActor : WeightedActors)
+	{
+		if (WeightedActor.ActorWeak.IsValid())
+		{
+			TotalWeight += WeightedActor.Weight;	
+		}
+	}
+
+	if (TotalWeight > 0.0f)
+	{
+		FVector AverageWeightedLocation = FVector::ZeroVector;
+		
+		for (const FAvaAlignBetweenWeightedActor& WeightedActor : WeightedActors)
+		{
+			if (const AActor* Actor = WeightedActor.ActorWeak.Get())
+			{
+				AverageWeightedLocation += Actor->GetActorLocation() * (WeightedActor.Weight / TotalWeight);
+			}
+		}
+		
+		ModifyActor->SetActorLocation(AverageWeightedLocation);
+	}
+
+	Next();
+}
+
+void UAvaAlignBetweenModifier::OnModifiedActorTransformed()
+{
+	Super::OnModifiedActorTransformed();
+
+	MarkModifierDirty();
+}
+
+#if WITH_EDITOR
+void UAvaAlignBetweenModifier::PostEditChangeProperty(FPropertyChangedEvent& PropertyChangedEvent)
+{
+	Super::PostEditChangeProperty(PropertyChangedEvent);
+
+	const FName MemberName = PropertyChangedEvent.GetMemberPropertyName();
+
+	static const FName ReferenceActorsPropertyName = GET_MEMBER_NAME_CHECKED(UAvaAlignBetweenModifier, ReferenceActors);
+	
+	if (MemberName == ReferenceActorsPropertyName)
+	{
+		OnReferenceActorsChanged();
+	}
+}
+
+void UAvaAlignBetweenModifier::PostTransacted(const FTransactionObjectEvent& TransactionEvent)
+{
+	OnReferenceActorsChanged();
+	
+	Super::PostTransacted(TransactionEvent);
+}
+#endif // WITH_EDITOR
+
+TSet<AActor*> UAvaAlignBetweenModifier::GetActors(const bool bEnabledOnly) const
+{
+	TSet<AActor*> OutActors;
+
+	for (const FAvaAlignBetweenWeightedActor& WeightedActor : ReferenceActors)
+	{
+		if (!bEnabledOnly|| (bEnabledOnly && WeightedActor.bEnabled))
+		{
+			OutActors.Add(WeightedActor.ActorWeak.Get());
+		}
+	}
+
+	return OutActors;
+}
+
+void UAvaAlignBetweenModifier::SetReferenceActors(TSet<FAvaAlignBetweenWeightedActor> NewReferenceActors)
+{
+	ReferenceActors = NewReferenceActors;
+	OnReferenceActorsChanged();
+}
+
+void UAvaAlignBetweenModifier::AddReferenceActor(const FAvaAlignBetweenWeightedActor& ReferenceActor)
+{
+	const AActor* const ModifyActor = GetModifiedActor();
+	if (!IsValid(ModifyActor))
+	{
+		return;
+	}
+
+	if (!ReferenceActor.ActorWeak.IsValid() || ReferenceActor.ActorWeak == ModifyActor)
+	{
+		return;
+	}
+
+	bool bAlreadyInSet = false;
+	ReferenceActors.Add(ReferenceActor, &bAlreadyInSet);
+
+	if (!bAlreadyInSet)
+	{
+		SetTransformExtensionReferenceActors();
+		
+		MarkModifierDirty();
+	}
+}
+
+bool UAvaAlignBetweenModifier::RemoveReferenceActor(AActor* const Actor)
+{
+	if (!IsValid(Actor))
+	{
+		return false;
+	}
+
+	if (ReferenceActors.Remove(FAvaAlignBetweenWeightedActor(Actor)) > 0)
+	{
+		MarkModifierDirty();
+		
+		return true;
+	}
+
+	return false;
+}
+
+bool UAvaAlignBetweenModifier::FindReferenceActor(AActor* InActor, FAvaAlignBetweenWeightedActor& OutReferenceActor) const
+{
+	if (!IsValid(InActor))
+	{
+		return false;
+	}
+
+	if (const FAvaAlignBetweenWeightedActor* OutReference = ReferenceActors.Find(FAvaAlignBetweenWeightedActor(InActor)))
+	{
+		OutReferenceActor = *OutReference;
+		return true;
+	}
+
+	return false;
+}
+
+void UAvaAlignBetweenModifier::OnModifierCDOSetup(FActorModifierCoreMetadata& InMetadata)
+{
+	Super::OnModifierCDOSetup(InMetadata);
+
+	InMetadata.SetName(TEXT("AlignBetween"));
+	InMetadata.SetCategory(TEXT("Layout"));
+#if WITH_EDITOR
+	InMetadata.SetDescription(LOCTEXT("ModifierDescription", "Positions an actor between a group of weighted actors"));
+#endif
+}
+
+void UAvaAlignBetweenModifier::OnTransformUpdated(AActor* InActor, bool bInParentMoved)
+{
+	// Check actor is reference actor enabled
+	FAvaAlignBetweenWeightedActor ReferenceActor;
+	FindReferenceActor(InActor, ReferenceActor);
+
+	if (ReferenceActor.IsValid() && ReferenceActor.bEnabled && ReferenceActor.Weight > 0)
+	{
+		MarkModifierDirty();
+	}
+}
+
+void UAvaAlignBetweenModifier::SetTransformExtensionReferenceActors()
+{
+	if (FAvaTransformUpdateModifierExtension* TransformExtension = GetExtension<FAvaTransformUpdateModifierExtension>())
+	{
+		TSet<TWeakObjectPtr<AActor>> ExtensionActors;
+		Algo::Transform(ReferenceActors, ExtensionActors, [](const FAvaAlignBetweenWeightedActor& InActor)->AActor*
+		{
+			return InActor.ActorWeak.Get();
+		});
+		TransformExtension->TrackActors(ExtensionActors, true);
+	}
+}
+
+void UAvaAlignBetweenModifier::OnReferenceActorsChanged()
+{
+	// Make sure the modifying actor is not part of the array.
+	AActor* const ModifyActor = GetModifiedActor();
+	if (!IsValid(ModifyActor))
+	{
+		return;
+	}
+	
+	ReferenceActors.Remove(FAvaAlignBetweenWeightedActor(ModifyActor));
+
+	SetTransformExtensionReferenceActors();
+
+	MarkModifierDirty();
+}
+
+TSet<FAvaAlignBetweenWeightedActor> UAvaAlignBetweenModifier::GetEnabledReferenceActors() const
+{
+	TSet<FAvaAlignBetweenWeightedActor> OutReferences;
+	OutReferences.Reserve(ReferenceActors.Num());
+
+	for (const FAvaAlignBetweenWeightedActor& WeightedActor : ReferenceActors)
+	{
+		if (WeightedActor.IsValid())
+		{
+			OutReferences.Add(WeightedActor);
+		}
+	}
+
+	return OutReferences;
+}
+
+#undef LOCTEXT_NAMESPACE
diff --git a/Engine/Plugins/Experimental/Avalanche/Source/AvalancheModifiers/Private/Modifiers/AvaArrangeBaseModifier.cpp b/Engine/Plugins/Experimental/Avalanche/Source/AvalancheModifiers/Private/Modifiers/AvaArrangeBaseModifier.cpp
new file mode 100644
--- /dev/null
+++ b/Engine/Plugins/Experimental/Avalanche/Source/AvalancheModifiers/Private/Modifiers/AvaArrangeBaseModifier.cpp
@@ -0,0 +1,78 @@
+// Copyright Epic Games, Inc. All Rights Reserved.
+
+#include "Modifiers/AvaArrangeBaseModifier.h"
+
+#include "Extensions/AvaTransformUpdateModifierExtension.h"
+#include "GameFramework/Actor.h"
+#include "Misc/ITransaction.h"
+#include "Shared/AvaTransformModifierShared.h"
+#include "Shared/AvaVisibilityModifierShared.h"
+
+void UAvaArrangeBaseModifier::OnModifierAdded(EActorModifierCoreEnableReason InReason)
+{
+	Super::OnModifierAdded(InReason);
+
+	AddExtension<FAvaRenderStateUpdateModifierExtension>(this);
+	AddExtension<FAvaTransformUpdateModifierExtension>(this);
+
+	if (FAvaSceneTreeUpdateModifierExtension* SceneExtension = GetExtension<FAvaSceneTreeUpdateModifierExtension>())
+	{
+		ReferenceActor.ReferenceContainer = EAvaReferenceContainer::Other;
+		ReferenceActor.ReferenceActorWeak = GetModifiedActor();
+		ReferenceActor.bSkipHiddenActors = false;
+		
+		SceneExtension->TrackSceneTree(0, &ReferenceActor);
+	}
+}
+
+void UAvaArrangeBaseModifier::OnModifierDisabled(EActorModifierCoreDisableReason InReason)
+{
+	Super::OnModifierDisabled(InReason);
+
+	if (UAvaTransformModifierShared* LayoutShared = GetShared<UAvaTransformModifierShared>(false))
+	{
+		LayoutShared->RestoreActorsState(this);
+	}
+
+	if (UAvaVisibilityModifierShared* VisibilityShared = GetShared<UAvaVisibilityModifierShared>(false))
+	{
+		VisibilityShared->RestoreActorsState(this);
+	}
+}
+
+void UAvaArrangeBaseModifier::OnModifiedActorTransformed()
+{
+	Super::OnModifiedActorTransformed();
+
+	MarkModifierDirty();
+}
+
+void UAvaArrangeBaseModifier::OnSceneTreeTrackedActorChildrenChanged(int32 InIdx,
+	const TSet<TWeakObjectPtr<AActor>>& InPreviousChildrenActors,
+	const TSet<TWeakObjectPtr<AActor>>& InNewChildrenActors)
+{
+	Super::OnSceneTreeTrackedActorChildrenChanged(InIdx, InPreviousChildrenActors, InNewChildrenActors);
+
+	const AActor* const ModifyActor = GetModifiedActor();
+	if (!IsValid(ModifyActor))
+	{
+		return;
+	}
+
+	MarkModifierDirty();
+}
+
+void UAvaArrangeBaseModifier::OnSceneTreeTrackedActorDirectChildrenChanged(int32 InIdx,
+	const TArray<TWeakObjectPtr<AActor>>& InPreviousChildrenActors,
+	const TArray<TWeakObjectPtr<AActor>>& InNewChildrenActors)
+{
+	Super::OnSceneTreeTrackedActorDirectChildrenChanged(InIdx, InPreviousChildrenActors, InNewChildrenActors);
+
+	const AActor* const ModifyActor = GetModifiedActor();
+	if (!IsValid(ModifyActor))
+	{
+		return;
+	}
+
+	MarkModifierDirty();
+}
diff --git a/Engine/Plugins/Experimental/Avalanche/Source/AvalancheModifiers/Private/Modifiers/AvaAttachmentBaseModifier.cpp b/Engine/Plugins/Experimental/Avalanche/Source/AvalancheModifiers/Private/Modifiers/AvaAttachmentBaseModifier.cpp
new file mode 100644
--- /dev/null
+++ b/Engine/Plugins/Experimental/Avalanche/Source/AvalancheModifiers/Private/Modifiers/AvaAttachmentBaseModifier.cpp
@@ -0,0 +1,22 @@
+// Copyright Epic Games, Inc. All Rights Reserved.
+
+#include "Modifiers/AvaAttachmentBaseModifier.h"
+
+#include "GameFramework/Actor.h"
+
+void UAvaAttachmentBaseModifier::OnModifierCDOSetup(FActorModifierCoreMetadata& InMetadata)
+{
+	Super::OnModifierCDOSetup(InMetadata);
+	
+	InMetadata.SetCompatibilityRule([](const AActor* InActor)->bool
+	{
+		return InActor && InActor->FindComponentByClass<USceneComponent>();		
+	});
+}
+
+void UAvaAttachmentBaseModifier::OnModifierAdded(EActorModifierCoreEnableReason InReason)
+{
+	Super::OnModifierAdded(InReason);
+
+	AddExtension<FAvaSceneTreeUpdateModifierExtension>(this);
+}
diff --git a/Engine/Plugins/Experimental/Avalanche/Source/AvalancheModifiers/Private/Modifiers/AvaAutoFollowModifier.cpp b/Engine/Plugins/Experimental/Avalanche/Source/AvalancheModifiers/Private/Modifiers/AvaAutoFollowModifier.cpp
new file mode 100644
--- /dev/null
+++ b/Engine/Plugins/Experimental/Avalanche/Source/AvalancheModifiers/Private/Modifiers/AvaAutoFollowModifier.cpp
@@ -0,0 +1,403 @@
+// Copyright Epic Games, Inc. All Rights Reserved.
+
+#include "Modifiers/AvaAutoFollowModifier.h"
+
+#include "AvaModifiersActorUtils.h"
+#include "AvaActorUtils.h"
+#include "Extensions/AvaSceneTreeUpdateModifierExtension.h"
+#include "Extensions/AvaTransformUpdateModifierExtension.h"
+#include "GameFramework/Actor.h"
+#include "Shared/AvaTransformModifierShared.h"
+
+#define LOCTEXT_NAMESPACE "AvaAutoFollowModifier"
+
+bool UAvaAutoFollowModifier::IsModifierDirtyable() const
+{
+	AActor* const ActorModified = GetModifiedActor();
+	AActor* const FollowedActor = ReferenceActor.ReferenceActorWeak.Get();
+
+	if (!IsValid(FollowedActor) || !IsValid(ActorModified))
+	{
+		return Super::IsModifierDirtyable();
+	}
+	
+	const FBox ReferenceActorLocalBounds = FAvaModifiersActorUtils::GetActorsBounds(FollowedActor, true);
+	const FBox ModifiedActorLocalBounds = FAvaModifiersActorUtils::GetActorsBounds(ActorModified, true);
+	
+	// Compare bounds/center to detect changes
+	if (ReferenceActorLocalBounds.Equals(CachedReferenceBounds, 0.01)
+		&& ModifiedActorLocalBounds.Equals(CachedModifiedBounds, 0.01))
+	{
+		return Super::IsModifierDirtyable();
+	}
+	
+	return true;
+}
+
+void UAvaAutoFollowModifier::OnModifierCDOSetup(FActorModifierCoreMetadata& InMetadata)
+{
+	Super::OnModifierCDOSetup(InMetadata);
+	
+	InMetadata.AllowTick(true);
+	InMetadata.SetName(TEXT("AutoFollow"));
+	InMetadata.SetCategory(TEXT("Layout"));
+#if WITH_EDITOR
+	InMetadata.SetDescription(LOCTEXT("ModifierDescription", "Positions an actor relative to another actor using their bounds"));
+#endif
+}
+
+void UAvaAutoFollowModifier::OnModifierAdded(EActorModifierCoreEnableReason InReason)
+{
+	Super::OnModifierAdded(InReason);
+
+	AddExtension<FAvaTransformUpdateModifierExtension>(this);
+	AddExtension<FAvaRenderStateUpdateModifierExtension>(this);
+
+	if (FAvaSceneTreeUpdateModifierExtension* SceneExtension = GetExtension<FAvaSceneTreeUpdateModifierExtension>())
+	{
+		SceneExtension->TrackSceneTree(0, &ReferenceActor);
+	}
+}
+
+void UAvaAutoFollowModifier::OnModifierEnabled(EActorModifierCoreEnableReason InReason)
+{
+	Super::OnModifierEnabled(InReason);
+
+	// Save actor layout state
+	if (UAvaTransformModifierShared* LayoutShared = GetShared<UAvaTransformModifierShared>(true))
+	{
+		LayoutShared->SaveActorState(this, GetModifiedActor());
+	}
+}
+
+void UAvaAutoFollowModifier::OnModifierDisabled(EActorModifierCoreDisableReason InReason)
+{
+	Super::OnModifierDisabled(InReason);
+
+	// Restore actor layout state
+	if (UAvaTransformModifierShared* LayoutShared = GetShared<UAvaTransformModifierShared>(false))
+	{
+		LayoutShared->RestoreActorState(this, GetModifiedActor());
+	}
+}
+
+void UAvaAutoFollowModifier::OnModifiedActorTransformed()
+{
+	const AActor* const FollowedActor = ReferenceActor.ReferenceActorWeak.Get();
+	const AActor* const ActorModified = GetModifiedActor();
+	
+	if (!FollowedActor || !ActorModified)
+	{
+		return;
+	}
+
+	// Compare current location and previous followed location but only component of followed axis to allow movements
+	if (FAvaModifiersActorUtils::IsAxisVectorEquals(ActorModified->GetActorLocation(), CachedFollowLocation, FollowedAxis))
+	{
+		return;
+	}
+	
+	MarkModifierDirty();
+}
+
+void UAvaAutoFollowModifier::Apply()
+{
+	AActor* const ModifyActor = GetModifiedActor();
+	
+	AActor* const FollowedActor = ReferenceActor.ReferenceActorWeak.Get();
+	if (!IsValid(FollowedActor))
+	{
+		Next();
+		return;
+	}
+
+	const FVector FollowAxisVector = FAvaModifiersActorUtils::GetVectorAxis(FollowedAxis);
+	if (FollowAxisVector.IsNearlyZero())
+	{
+		Next();
+		return;
+	}
+
+	// Get Padding to use
+	const FVector DistancePadding = FMath::Lerp<FVector>(DefaultDistance, MaxDistance, Progress / 100.0f);
+
+	// Get actors bounds
+	CachedReferenceBounds = FAvaModifiersActorUtils::GetActorsBounds(FollowedActor, true);
+	CachedModifiedBounds = FAvaModifiersActorUtils::GetActorsBounds(ModifyActor, true);
+
+	// Check if bounds are valid
+	const bool bReferenceActorZeroSizeBounds = CachedReferenceBounds.GetSize().IsNearlyZero();
+	const bool bModifyActorZeroSizeBounds = CachedModifiedBounds.GetSize().IsNearlyZero();
+
+	// Get actors location 
+	const FVector ReferenceActorLocation = FollowedActor->GetActorLocation();
+	const FVector ModifyActorLocation = ModifyActor->GetActorLocation();
+	
+	// Get bounds center (pivot)
+	const FVector ReferenceActorCenter = !bReferenceActorZeroSizeBounds ? CachedReferenceBounds.GetCenter() : ReferenceActorLocation;
+	const FVector ModifierActorCenter = !bModifyActorZeroSizeBounds ? CachedModifiedBounds.GetCenter() : ModifyActorLocation;
+
+	// Get bounds extents
+	const FVector ReferenceActorExtent = CachedReferenceBounds.GetExtent();
+	const FVector ModifierActorExtent = CachedModifiedBounds.GetExtent();
+	
+	// using world space extents so that the modified actor is moved also taking into account reference actor rotation
+	const FVector ReferenceActorLocalOffset = ReferenceActorExtent * OffsetAxis;
+	const FVector ModifierActorLocalOffset = ModifierActorExtent * OffsetAxis;
+	
+	// Use user alignments for followed and modify actors
+    const FVector ReferenceActorBoundsOffset = FollowedAlignment.LocalBoundsOffset(FBox(-ReferenceActorExtent, ReferenceActorExtent));
+    const FVector ModifierActorBoundsOffset = LocalAlignment.LocalBoundsOffset(FBox(-ModifierActorExtent, ModifierActorExtent));
+	
+	// this offset can be non-zero when the actor pivot and bounds origin do not coincide: we need to take this into account
+	const FVector ReferenceActorPivotToBoundsOffset = ReferenceActorLocation - ReferenceActorCenter;
+	const FVector ModifiedActorPivotToBoundsOffset = ModifyActorLocation - ModifierActorCenter;
+
+	const FVector OffsetLocation =
+		// Reference actor extent - reference actor alignment
+		ReferenceActorLocalOffset - ReferenceActorBoundsOffset
+		// Modified actor extent + modified actor alignment
+		+ ModifierActorLocalOffset + ModifierActorBoundsOffset
+		// Distance progress offset
+		+ DistancePadding
+		// finally removing any existing pivot to bounds offset
+		+ (ModifiedActorPivotToBoundsOffset - ReferenceActorPivotToBoundsOffset);
+
+	// target location needs to start from reference actor bounds location + proper location offset
+	CachedFollowLocation = ModifyActorLocation + (ReferenceActorLocation - ModifyActorLocation + OffsetLocation) * FollowAxisVector;
+	
+	ModifyActor->SetActorLocation(CachedFollowLocation);
+	
+	Next();
+}
+
+void UAvaAutoFollowModifier::PostLoad()
+{
+	if (!bDeprecatedPropertiesMigrated)
+	{
+		ReferenceActor.ReferenceContainer = ReferenceContainer_DEPRECATED;
+		ReferenceActor.ReferenceActorWeak = ReferenceActorWeak_DEPRECATED;
+		ReferenceActor.bSkipHiddenActors = bIgnoreHiddenActors_DEPRECATED;
+		
+		bDeprecatedPropertiesMigrated = true;
+	}
+	
+	Super::PostLoad();
+}
+
+#if WITH_EDITOR
+void UAvaAutoFollowModifier::PostEditChangeProperty(FPropertyChangedEvent& PropertyChangedEvent)
+{
+	Super::PostEditChangeProperty(PropertyChangedEvent);
+
+	const FName MemberName = PropertyChangedEvent.GetMemberPropertyName();
+
+	static const FName ReferenceActorPropertyName = GET_MEMBER_NAME_CHECKED(UAvaAutoFollowModifier, ReferenceActor);
+	static const FName DefaultDistancePropertyName = GET_MEMBER_NAME_CHECKED(UAvaAutoFollowModifier, DefaultDistance);
+	static const FName MaxDistancePropertyName = GET_MEMBER_NAME_CHECKED(UAvaAutoFollowModifier, MaxDistance);
+	static const FName ProgressPropertyName = GET_MEMBER_NAME_CHECKED(UAvaAutoFollowModifier, Progress);
+	static const FName FollowerAlignmentPropertyName = GET_MEMBER_NAME_CHECKED(UAvaAutoFollowModifier, FollowedAlignment);
+	static const FName LocalAlignmentPropertyName = GET_MEMBER_NAME_CHECKED(UAvaAutoFollowModifier, LocalAlignment);
+	static const FName OffsetAxisPropertyName = GET_MEMBER_NAME_CHECKED(UAvaAutoFollowModifier, OffsetAxis);
+	static const FName FollowedAxisPropertyName = GET_MEMBER_NAME_CHECKED(UAvaAutoFollowModifier, FollowedAxis);
+	
+	if (MemberName == ReferenceActorPropertyName)
+	{
+		OnReferenceActorChanged();
+	}
+	else if (MemberName == DefaultDistancePropertyName
+		|| MemberName == MaxDistancePropertyName
+		|| MemberName == ProgressPropertyName
+		|| MemberName == FollowerAlignmentPropertyName
+		|| MemberName == LocalAlignmentPropertyName
+		|| MemberName == OffsetAxisPropertyName)
+	{
+		MarkModifierDirty();
+	}
+	else if (MemberName == FollowedAxisPropertyName)
+	{
+		OnFollowedAxisChanged();	
+	}
+}
+
+void UAvaAutoFollowModifier::PostTransacted(const FTransactionObjectEvent& TransactionEvent)
+{
+	OnReferenceActorChanged();
+	
+	Super::PostTransacted(TransactionEvent);
+}
+#endif // WITH_EDITOR
+
+void UAvaAutoFollowModifier::SetReferenceActor(const FAvaSceneTreeActor& InReferenceActor)
+{
+	if (ReferenceActor == InReferenceActor)
+	{
+		return;
+	}
+	
+	ReferenceActor = InReferenceActor;
+	OnReferenceActorChanged();
+}
+
+void UAvaAutoFollowModifier::SetFollowedAxis(int32 InFollowedAxis)
+{
+	if (FollowedAxis == InFollowedAxis)
+	{
+		return;
+	}
+
+	FollowedAxis = InFollowedAxis;
+	OnFollowedAxisChanged();
+}
+
+void UAvaAutoFollowModifier::OnFollowedAxisChanged()
+{
+	const FVector FollowedAxisVector = FAvaModifiersActorUtils::GetVectorAxis(FollowedAxis);
+
+	LocalAlignment.bUseDepth = FollowedAlignment.bUseDepth = FollowedAxisVector.X != 0;
+	LocalAlignment.bUseHorizontal = FollowedAlignment.bUseHorizontal = FollowedAxisVector.Y != 0;
+	LocalAlignment.bUseVertical = FollowedAlignment.bUseVertical = FollowedAxisVector.Z != 0;
+
+	MarkModifierDirty();
+}
+
+void UAvaAutoFollowModifier::SetDefaultDistance(const FVector& NewDefaultDistance)
+{
+	DefaultDistance = NewDefaultDistance;
+
+	MarkModifierDirty();
+}
+
+void UAvaAutoFollowModifier::SetMaxDistance(const FVector& NewMaxDistance)
+{
+	MaxDistance = NewMaxDistance;
+
+	MarkModifierDirty();
+}
+
+void UAvaAutoFollowModifier::SetProgress(const FVector& NewProgress)
+{
+	Progress = NewProgress;
+
+	MarkModifierDirty();
+}
+
+void UAvaAutoFollowModifier::SetFollowedAlignment(const FAvaAnchorAlignment& NewFollowedAlignment)
+{
+	FollowedAlignment = NewFollowedAlignment;
+
+	MarkModifierDirty();
+}
+
+void UAvaAutoFollowModifier::SetLocalAlignment(const FAvaAnchorAlignment& NewLocalAlignment)
+{
+	LocalAlignment = NewLocalAlignment;
+
+	MarkModifierDirty();
+}
+
+void UAvaAutoFollowModifier::SetOffsetAxis(const FVector& NewOffsetAxis)
+{
+	OffsetAxis = NewOffsetAxis;
+
+	MarkModifierDirty();
+}
+
+void UAvaAutoFollowModifier::OnTransformUpdated(AActor* InActor, bool bInParentMoved)
+{
+	const AActor* ActorModified = GetModifiedActor();
+	if (!IsValid(ActorModified))
+	{
+		return;
+	}
+	
+	const AActor* FollowedActor = ReferenceActor.ReferenceActorWeak.Get();
+	const bool bIsAttachedToReferenceActor = InActor->IsAttachedTo(FollowedActor);
+	const bool bIsReferenceActor = InActor == FollowedActor;
+	
+	if (IsValid(FollowedActor) && (bIsAttachedToReferenceActor || bIsReferenceActor))
+	{
+		MarkModifierDirty();
+	}
+}
+
+void UAvaAutoFollowModifier::OnRenderStateUpdated(AActor* InActor, UActorComponent* InComponent)
+{
+	const AActor* ActorModified = GetModifiedActor();
+	if (!IsValid(ActorModified))
+	{
+		return;
+	}
+	
+	const AActor* FollowedActor = ReferenceActor.ReferenceActorWeak.Get();
+	const UAvaTransformModifierShared* LayoutShared = GetShared<UAvaTransformModifierShared>(false);
+	
+	if (!InActor || !FollowedActor || !LayoutShared)
+	{
+		return;
+	}
+	
+	const bool bIsReferenceActor = FollowedActor == InActor;
+	const bool bIsAttachedToReferenceActor = InActor->IsAttachedTo(FollowedActor);
+	const bool bModifierDirtyable = IsModifierDirtyable();
+
+	if (!bIsReferenceActor && !bIsAttachedToReferenceActor && !bModifierDirtyable)
+	{
+		return;
+	}
+	
+	MarkModifierDirty();
+}
+
+void UAvaAutoFollowModifier::OnSceneTreeTrackedActorChanged(int32 InIdx, AActor* InPreviousActor, AActor* InNewActor)
+{
+	Super::OnSceneTreeTrackedActorChanged(InIdx, InPreviousActor, InNewActor);
+
+	if (InNewActor == GetModifiedActor())
+	{
+		OnReferenceActorChanged();
+		return;
+	}
+
+	// Untrack reference actor
+	if (FAvaTransformUpdateModifierExtension* TransformExtension = GetExtension<FAvaTransformUpdateModifierExtension>())
+	{
+		TransformExtension->UntrackActor(InPreviousActor);
+		TransformExtension->TrackActor(InNewActor, true);
+	}
+
+	MarkModifierDirty();
+}
+
+void UAvaAutoFollowModifier::OnSceneTreeTrackedActorChildrenChanged(int32 InIdx,
+	const TSet<TWeakObjectPtr<AActor>>& InPreviousChildrenActors,
+	const TSet<TWeakObjectPtr<AActor>>& InNewChildrenActors)
+{
+	Super::OnSceneTreeTrackedActorChildrenChanged(InIdx, InPreviousChildrenActors, InNewChildrenActors);
+
+	// Untrack reference actor children
+	if (FAvaTransformUpdateModifierExtension* TransformExtension = GetExtension<FAvaTransformUpdateModifierExtension>())
+	{
+		TransformExtension->UntrackActors(InPreviousChildrenActors);
+		TransformExtension->TrackActors(InNewChildrenActors, false);
+	}
+
+	ChildrenActorsWeak = InNewChildrenActors;
+	
+	MarkModifierDirty();
+}
+
+void UAvaAutoFollowModifier::OnReferenceActorChanged()
+{
+	if (ReferenceActor.ReferenceActorWeak.Get() == GetModifiedActor())
+	{
+		ReferenceActor.ReferenceActorWeak = nullptr;
+	}
+
+	if (const FAvaSceneTreeUpdateModifierExtension* SceneExtension = GetExtension<FAvaSceneTreeUpdateModifierExtension>())
+	{
+		SceneExtension->CheckTrackedActorUpdate(0);
+	}
+}
+
+#undef LOCTEXT_NAMESPACE
