[MotionDesign]

diff --git a/Engine/Plugins/Experimental/Avalanche/Source/AvalancheModifiers/Private/Modifiers/AvaAutoSizeModifier.cpp b/Engine/Plugins/Experimental/Avalanche/Source/AvalancheModifiers/Private/Modifiers/AvaAutoSizeModifier.cpp
new file mode 100644
--- /dev/null
+++ b/Engine/Plugins/Experimental/Avalanche/Source/AvalancheModifiers/Private/Modifiers/AvaAutoSizeModifier.cpp
@@ -0,0 +1,456 @@
+// Copyright Epic Games, Inc. All Rights Reserved.
+
+#include "Modifiers/AvaAutoSizeModifier.h"
+
+#include "AvaActorUtils.h"
+#include "Async/Async.h"
+#include "AvaModifiersActorUtils.h"
+#include "Components/DynamicMeshComponent.h"
+#include "DynamicMesh/MeshTransforms.h"
+#include "AvaShapeActor.h"
+#include "DynamicMeshes/AvaShape2DDynMeshBase.h"
+#include "Shared/AvaTransformModifierShared.h"
+
+#define LOCTEXT_NAMESPACE "AvaAutoSizeModifier"
+
+bool UAvaAutoSizeModifier::IsModifierDirtyable() const
+{
+	const AActor* const ActorModified = GetModifiedActor();
+	AActor* const TrackedActor = ReferenceActor.ReferenceActorWeak.Get();
+
+	if (!IsValid(TrackedActor) || !IsValid(ActorModified))
+	{
+		return Super::IsModifierDirtyable();
+	}
+
+	const FBox ReferenceActorLocalBounds = FAvaModifiersActorUtils::GetActorsBounds(TrackedActor, bIncludeChildren, true);
+
+	if (ReferenceActorLocalBounds.Equals(CachedReferenceBounds, 0.01))
+	{
+		return Super::IsModifierDirtyable();
+	}
+
+	return true;
+}
+
+void UAvaAutoSizeModifier::OnModifierCDOSetup(FActorModifierCoreMetadata& InMetadata)
+{
+	Super::OnModifierCDOSetup(InMetadata);
+
+	InMetadata.AllowTick(true);
+	InMetadata.SetName(TEXT("AutoSize"));
+	InMetadata.SetCategory(TEXT("Layout"));
+#if WITH_EDITOR
+	InMetadata.SetDescription(LOCTEXT("ModifierDescription", "The modified actor will be resized to act as a background for the specified actor"));
+#endif
+
+	InMetadata.SetCompatibilityRule([](const AActor* InActor)->bool
+	{
+		bool bSupported = false;
+		if (InActor)
+		{
+			if (const UDynamicMeshComponent* DynMeshComponent = InActor->FindComponentByClass<UDynamicMeshComponent>())
+			{
+				DynMeshComponent->ProcessMesh([&bSupported](const FDynamicMesh3& ProcessMesh)
+				{
+					bSupported = ProcessMesh.VertexCount() > 0 && static_cast<FBox>(ProcessMesh.GetBounds(true)).GetSize().X == 0;
+				});
+			}
+		}
+		return bSupported;
+	});
+}
+
+void UAvaAutoSizeModifier::OnModifiedActorTransformed()
+{
+	Super::OnModifiedActorTransformed();
+
+	const AActor* ActorModified = GetModifiedActor();
+	const AActor* TrackedActor = ReferenceActor.ReferenceActorWeak.Get();
+	if (!TrackedActor || !ActorModified)
+	{
+		return;
+	}
+
+	// Compare current location and previous followed location but only component of followed axis to allow movements
+	if (FAvaModifiersActorUtils::IsAxisVectorEquals(ActorModified->GetActorLocation(), CachedFollowLocation, FollowedAxis))
+	{
+		return;
+	}
+
+	MarkModifierDirty();
+}
+
+void UAvaAutoSizeModifier::OnModifierAdded(EActorModifierCoreEnableReason InReason)
+{
+	Super::OnModifierAdded(InReason);
+
+	AddExtension<FAvaTransformUpdateModifierExtension>(this);
+	AddExtension<FAvaRenderStateUpdateModifierExtension>(this);
+	AddExtension<FAvaSceneTreeUpdateModifierExtension>(this);
+
+	if (FAvaSceneTreeUpdateModifierExtension* SceneExtension = GetExtension<FAvaSceneTreeUpdateModifierExtension>())
+	{
+		SceneExtension->TrackSceneTree(0, &ReferenceActor);
+	}
+}
+
+void UAvaAutoSizeModifier::OnModifierEnabled(EActorModifierCoreEnableReason InReason)
+{
+	Super::OnModifierEnabled(InReason);
+
+	// Save actor layout state
+	if (UAvaTransformModifierShared* LayoutShared = GetShared<UAvaTransformModifierShared>(true))
+	{
+		LayoutShared->SaveActorState(this, GetModifiedActor());
+	}
+
+	if (InReason == EActorModifierCoreEnableReason::User)
+	{
+		if (const AAvaShapeActor* const ReferenceAsShapeActor = Cast<AAvaShapeActor>(GetModifiedActor()))
+		{
+			if (UAvaShape2DDynMeshBase* const Shape2DDynMesh = Cast<UAvaShape2DDynMeshBase>(ReferenceAsShapeActor->GetDynamicMesh()))
+			{
+				ShapeDynMesh2DWeak = Shape2DDynMesh;
+				PreModifierShapeDynMesh2DSize = Shape2DDynMesh->GetSize2D();
+			}
+		}
+	}
+
+	bDeprecatedPropertiesMigrated = true;
+}
+
+void UAvaAutoSizeModifier::OnModifierDisabled(EActorModifierCoreDisableReason InReason)
+{
+	Super::OnModifierDisabled(InReason);
+
+	// Restore actor layout state
+	if (UAvaTransformModifierShared* LayoutShared = GetShared<UAvaTransformModifierShared>(false))
+	{
+		LayoutShared->RestoreActorState(this, GetModifiedActor());
+	}
+
+	if (UAvaShape2DDynMeshBase* const Shape2DDynMesh = ShapeDynMesh2DWeak.Get())
+	{
+		Shape2DDynMesh->SetSize2D(PreModifierShapeDynMesh2DSize);
+	}
+}
+
+void UAvaAutoSizeModifier::Apply()
+{
+	AActor* const CurrentlyModifiedActor = GetModifiedActor();
+
+	AActor* TrackedActor = ReferenceActor.ReferenceActorWeak.Get();
+	if (!IsValid(TrackedActor))
+	{
+		Next();
+		return;
+	}
+
+	UDynamicMeshComponent* const DynMeshComponent = GetMeshComponent();
+	if (!IsValid(DynMeshComponent))
+	{
+		Fail(LOCTEXT("InvalidDynamicMeshComponent", "Invalid dynamic mesh component on modified actor"));
+		return;
+	}
+
+	if (FAvaRenderStateUpdateModifierExtension* RenderExtension = GetExtension<FAvaRenderStateUpdateModifierExtension>())
+	{
+		RenderExtension->SetTrackedActorVisibility(TrackedActor, bIncludeChildren);
+	}
+
+	CachedReferenceBounds = FAvaModifiersActorUtils::GetActorsBounds(TrackedActor, bIncludeChildren, true);
+
+	const FBox ModifiedActorBounds = FAvaModifiersActorUtils::GetActorsBounds(CurrentlyModifiedActor, false);
+
+	FVector ModifiedActorBoundsOrigin;
+	FVector ModifiedActorBoundsExtent;
+	ModifiedActorBounds.GetCenterAndExtents(ModifiedActorBoundsOrigin, ModifiedActorBoundsExtent);
+
+	const FMargin& DesiredPadding = Padding;
+	FVector ReferenceBoundsExtent = CachedReferenceBounds.GetExtent();
+
+	// Add padding only if there is content inside otherwise don't
+	if (!CachedReferenceBounds.GetExtent().IsNearlyZero())
+	{
+		ReferenceBoundsExtent += FVector(0, DesiredPadding.Right + DesiredPadding.Left, DesiredPadding.Top + DesiredPadding.Bottom);
+	}
+
+	// Avoid division by 0
+	if (ModifiedActorBoundsExtent.IsNearlyZero())
+	{
+		ModifiedActorBoundsExtent += FVector(UE_KINDA_SMALL_NUMBER);
+	}
+
+	const float ScaleRatioY = ReferenceBoundsExtent.Y/ModifiedActorBoundsExtent.Y;
+	const float ScaleRatioZ = ReferenceBoundsExtent.Z/ModifiedActorBoundsExtent.Z;
+	const FVector ScaleRatio = FVector(1.0f, ScaleRatioY, ScaleRatioZ);
+
+	// Move only if there is content inside and at least an axis is followed
+	const FVector FollowAxisVector = FAvaModifiersActorUtils::GetVectorAxis(FollowedAxis);
+	if (!FollowAxisVector.IsNearlyZero() && !CachedReferenceBounds.GetSize().IsNearlyZero())
+	{
+		const FVector ModifiedActorPivotToBoundsOffset = ModifiedActorBoundsOrigin - CurrentlyModifiedActor->GetActorLocation();
+		const FVector ReferenceActorPivotToBoundsOffset = CachedReferenceBounds.GetCenter() - TrackedActor->GetActorLocation();
+		const FVector PaddingOffset = FVector(0.f, DesiredPadding.Right - DesiredPadding.Left, DesiredPadding.Top - DesiredPadding.Bottom);
+
+		const FVector DesiredLocationOffset =
+			// Add margin padding
+			PaddingOffset
+			// Remove pivot actor location difference for modified actor
+			- ModifiedActorPivotToBoundsOffset
+			// add pivot offset difference for reference actor
+			+ ReferenceActorPivotToBoundsOffset;
+
+		CachedFollowLocation =
+			CurrentlyModifiedActor->GetActorLocation()
+			+ (TrackedActor->GetActorLocation() - CurrentlyModifiedActor->GetActorLocation() + DesiredLocationOffset) * FollowAxisVector;
+
+		CurrentlyModifiedActor->SetActorLocation(CachedFollowLocation);
+	}
+
+	// check if we can use Shape2DDynMesh SetSize2D for this modified mesh.
+	// this allows us to keep properly scaled corner bevels and slants
+	if (UAvaShape2DDynMeshBase* const Shape2DDynMesh = ShapeDynMesh2DWeak.Get())
+	{
+		FVector2D DesiredSize2D = FVector2D::Max(
+			FVector2D(ReferenceBoundsExtent.Y, ReferenceBoundsExtent.Z)*2.0f,
+			UAvaShape2DDynMeshBase::MinSize2D
+		);
+
+		switch (FitMode)
+		{
+			case EAvaAutoSizeFitMode::WidthAndHeight:
+				// nothing to do
+				break;
+			case EAvaAutoSizeFitMode::WidthOnly:
+				{
+					DesiredSize2D = FVector2D(DesiredSize2D.X, Shape2DDynMesh->GetSize2D().Y);
+				}
+				break;
+			case EAvaAutoSizeFitMode::HeightOnly:
+				{
+					DesiredSize2D = FVector2D(Shape2DDynMesh->GetSize2D().X, DesiredSize2D.Y);
+				}
+				break;
+			default: ;
+		}
+
+		// Not liking this as this will update the shape and trigger a stack update,
+		// we might create a loop where the modifier dirties the shape and the shape dirties modifiers
+		if (!Shape2DDynMesh->GetSize2D().Equals(DesiredSize2D, 0.01))
+		{
+			Shape2DDynMesh->SetSize2D(DesiredSize2D);
+		}
+	}
+	else // for all other dynamic meshes, just scale all vertices
+	{
+		FVector NewScaleValue = CurrentlyModifiedActor->GetActorScale3D() * ScaleRatio;
+
+		switch (FitMode)
+		{
+			case EAvaAutoSizeFitMode::WidthAndHeight:
+				break;
+			case EAvaAutoSizeFitMode::WidthOnly:
+				NewScaleValue.Z = 1.0f;
+				break;
+			case EAvaAutoSizeFitMode::HeightOnly:
+				NewScaleValue.Y = 1.0f;
+				break;
+			default: ;
+		}
+
+		DynMeshComponent->GetDynamicMesh()->EditMesh([&NewScaleValue](FDynamicMesh3& EditMesh)
+		{
+			MeshTransforms::Scale(EditMesh, NewScaleValue, FVector::Zero(), true);
+		}
+		, EDynamicMeshChangeType::GeneralEdit, EDynamicMeshAttributeChangeFlags::Unknown, false);
+	}
+
+	Next();
+}
+
+void UAvaAutoSizeModifier::PostLoad()
+{
+	if (!bDeprecatedPropertiesMigrated
+		&& ReferenceActor.ReferenceContainer == EAvaReferenceContainer::Other
+		&& ReferenceActor.ReferenceActorWeak == nullptr)
+	{
+		ReferenceActor.ReferenceContainer = ReferenceContainer_DEPRECATED;
+		ReferenceActor.ReferenceActorWeak = ReferenceActorWeak_DEPRECATED;
+		ReferenceActor.bSkipHiddenActors = bIgnoreHiddenActors_DEPRECATED;
+	}
+
+	bDeprecatedPropertiesMigrated = true;
+
+	Super::PostLoad();
+}
+
+#if WITH_EDITOR
+void UAvaAutoSizeModifier::PostEditChangeProperty(FPropertyChangedEvent& PropertyChangedEvent)
+{
+	Super::PostEditChangeProperty(PropertyChangedEvent);
+
+	const FName MemberName = PropertyChangedEvent.GetMemberPropertyName();
+
+	static const FName ReferenceActorPropertyName = GET_MEMBER_NAME_CHECKED(UAvaAutoSizeModifier, ReferenceActor);
+	static const FName PaddingPropertyName = GET_MEMBER_NAME_CHECKED(UAvaAutoSizeModifier, Padding);
+	static const FName IncludeChildrenName = GET_MEMBER_NAME_CHECKED(UAvaAutoSizeModifier, bIncludeChildren);
+	static const FName FitModePropertyName = GET_MEMBER_NAME_CHECKED(UAvaAutoSizeModifier, FitMode);
+
+	if (MemberName == ReferenceActorPropertyName)
+	{
+		OnReferenceActorChanged();
+	}
+	else if (MemberName == PaddingPropertyName
+		|| MemberName == FitModePropertyName
+		|| MemberName == IncludeChildrenName)
+	{
+		MarkModifierDirty();
+	}
+}
+#endif
+
+void UAvaAutoSizeModifier::SetReferenceActor(const FAvaSceneTreeActor& InReferenceActor)
+{
+	if (ReferenceActor == InReferenceActor)
+	{
+		return;
+	}
+
+	ReferenceActor = InReferenceActor;
+	OnReferenceActorChanged();
+}
+
+void UAvaAutoSizeModifier::SetFollowedAxis(int32 InFollowedAxis)
+{
+	if (FollowedAxis == InFollowedAxis)
+	{
+		return;
+	}
+
+	FollowedAxis = InFollowedAxis;
+	MarkModifierDirty();
+}
+
+void UAvaAutoSizeModifier::SetPadding(const FMargin& InPadding)
+{
+	if (Padding != InPadding)
+	{
+		Padding = InPadding;
+
+		MarkModifierDirty();
+	}
+}
+
+void UAvaAutoSizeModifier::SetFitMode(const EAvaAutoSizeFitMode InFitMode)
+{
+	if (FitMode != InFitMode)
+	{
+		FitMode = InFitMode;
+
+		MarkModifierDirty();
+	}
+}
+
+void UAvaAutoSizeModifier::SetIncludeChildren(bool bInIncludeChildren)
+{
+	if (bIncludeChildren == bInIncludeChildren)
+	{
+		return;
+	}
+
+	bIncludeChildren = bInIncludeChildren;
+	MarkModifierDirty();
+}
+
+void UAvaAutoSizeModifier::OnTransformUpdated(AActor* InActor, bool bInParentMoved)
+{
+	if (bInParentMoved)
+	{
+		return;
+	}
+
+	OnRenderStateUpdated(InActor, nullptr);
+}
+
+void UAvaAutoSizeModifier::OnRenderStateUpdated(AActor* InActor, UActorComponent* InComponent)
+{
+	const AActor* ActorModified = GetModifiedActor();
+	if (!ActorModified)
+	{
+		return;
+	}
+
+	const AActor* TrackedActor = ReferenceActor.ReferenceActorWeak.Get();
+	if (!TrackedActor)
+	{
+		return;
+	}
+
+	const bool bIsReferenceActor = InActor == TrackedActor;
+	const bool bIsReferenceActorChildren = InActor->IsAttachedTo(TrackedActor);
+
+	if (bIsReferenceActor
+		|| (bIncludeChildren && bIsReferenceActorChildren))
+	{
+		const bool bIsModifierDirtyable = IsModifierDirtyable();
+
+		// Only update if bounds changed
+		if (bIsModifierDirtyable)
+		{
+			MarkModifierDirty();
+		}
+	}
+}
+
+void UAvaAutoSizeModifier::OnActorVisibilityChanged(AActor* InActor)
+{
+	OnRenderStateUpdated(InActor, nullptr);
+}
+
+void UAvaAutoSizeModifier::OnSceneTreeTrackedActorChanged(int32 InIdx, AActor* InPreviousActor, AActor* InNewActor)
+{
+	if (InNewActor == GetModifiedActor())
+	{
+		OnReferenceActorChanged();
+		return;
+	}
+
+	// Untrack reference actor
+	if (FAvaTransformUpdateModifierExtension* TransformExtension = GetExtension<FAvaTransformUpdateModifierExtension>())
+	{
+		TransformExtension->UntrackActor(InPreviousActor);
+		TransformExtension->TrackActor(InNewActor, true);
+	}
+
+	MarkModifierDirty();
+}
+
+void UAvaAutoSizeModifier::OnSceneTreeTrackedActorChildrenChanged(int32 InIdx,
+	const TSet<TWeakObjectPtr<AActor>>& InPreviousChildrenActors,
+	const TSet<TWeakObjectPtr<AActor>>& InNewChildrenActors)
+{
+	if (bIncludeChildren)
+	{
+		if (IsModifierDirtyable())
+		{
+			MarkModifierDirty();
+		}
+	}
+}
+
+void UAvaAutoSizeModifier::OnReferenceActorChanged()
+{
+	if (ReferenceActor.ReferenceActorWeak.Get() == GetModifiedActor())
+	{
+		ReferenceActor.ReferenceActorWeak = nullptr;
+	}
+
+	if (const FAvaSceneTreeUpdateModifierExtension* SceneExtension = GetExtension<FAvaSceneTreeUpdateModifierExtension>())
+	{
+		SceneExtension->CheckTrackedActorUpdate(0);
+	}
+}
+
+#undef LOCTEXT_NAMESPACE
diff --git a/Engine/Plugins/Experimental/Avalanche/Source/AvalancheModifiers/Private/Modifiers/AvaBevelModifier.cpp b/Engine/Plugins/Experimental/Avalanche/Source/AvalancheModifiers/Private/Modifiers/AvaBevelModifier.cpp
new file mode 100644
--- /dev/null
+++ b/Engine/Plugins/Experimental/Avalanche/Source/AvalancheModifiers/Private/Modifiers/AvaBevelModifier.cpp
@@ -0,0 +1,181 @@
+// Copyright Epic Games, Inc. All Rights Reserved.
+
+#include "Modifiers/AvaBevelModifier.h"
+
+#include "Async/Async.h"
+#include "Components/DynamicMeshComponent.h"
+#include "DynamicMesh/MeshNormals.h"
+#include "DynamicMesh/Operations/MergeCoincidentMeshEdges.h"
+#include "GroupTopology.h"
+#include "Operations/MeshBevel.h"
+#include "Parameterization/DynamicMeshUVEditor.h"
+
+#define LOCTEXT_NAMESPACE "AvaBevelModifier"
+
+void UAvaBevelModifier::OnModifierCDOSetup(FActorModifierCoreMetadata& InMetadata)
+{
+	Super::OnModifierCDOSetup(InMetadata);
+
+	InMetadata.SetName(TEXT("Bevel"));
+	InMetadata.SetCategory(TEXT("Geometry"));
+#if WITH_EDITOR
+	InMetadata.SetDescription(LOCTEXT("ModifierDescription", "Create chamfered or rounded corners to geometry that smooth edges and corners"));
+#endif
+}
+
+void UAvaBevelModifier::Apply()
+{
+	UDynamicMeshComponent* const DynMeshComp = GetMeshComponent();
+	
+	if (!IsValid(DynMeshComp))
+	{
+		Fail(LOCTEXT("InvalidDynamicMeshComponent", "Invalid dynamic mesh component on modified actor"));
+		return;
+	}
+
+	if (Inset <= 0.f)
+	{
+		Next();
+		return;
+	}
+	
+	using namespace UE::Geometry;
+			
+	DynMeshComp->GetDynamicMesh()->EditMesh([this, DynMeshComp](FDynamicMesh3& EditMesh) 
+	{
+		if (EditMesh.TriangleCount() > 0)
+		{
+			// weld edges
+			FMergeCoincidentMeshEdges Welder(&EditMesh);
+			Welder.Apply();
+		}
+
+		float BevelDistance = Inset;
+		float Divider = 4.f;
+				
+		// Apply bevel operator for multiple iterations
+		TArray<int32> NewTriangles;
+		for (int32 ItIdx = 0; ItIdx < Iterations; ItIdx++)
+		{
+			const FGroupTopology Topology(&EditMesh, true);
+					
+			FMeshBevel Bevel;
+			Bevel.InsetDistance = BevelDistance;
+			Bevel.InitializeFromGroupTopology(EditMesh, Topology);
+			Bevel.Apply(EditMesh, nullptr);
+					
+			NewTriangles.Append(Bevel.NewTriangles);
+
+			// lets reduce the distance to avoid reversed triangles and overlapping triangles
+			BevelDistance /= Divider;
+			Divider /= 2.f;
+		}
+
+		// Get polygroup layer for back side
+		FDynamicMeshPolygroupAttribute* const BevelPolygroup = FindOrCreatePolygroupLayer(EditMesh, UAvaBevelModifier::BevelPolygroupLayerName, &NewTriangles);
+
+		// TODO : Fix UVs temp, UV modifier needed instead of this
+		FRotator BoxRotation = DynMeshComp->GetComponentTransform().Rotator();
+		FBox MeshBounds = static_cast<FBox>(EditMesh.GetBounds(true));
+		const FTransform PlaneTransform(BoxRotation, MeshBounds.GetCenter(), FVector::OneVector);
+		FDynamicMeshUVOverlay* UVOverlay = EditMesh.Attributes()->GetUVLayer(0);
+		FDynamicMeshUVEditor UVEditor(&EditMesh, UVOverlay);
+		const FFrame3d ProjectionFrame(PlaneTransform);
+		FUVEditResult Result;
+		UVEditor.SetTriangleUVsFromBoxProjection(NewTriangles, [](const FVector3d& Pos)
+		{
+			return Pos;
+		}
+		, ProjectionFrame, MeshBounds.GetSize(), 3, &Result);
+				
+	}, EDynamicMeshChangeType::GeneralEdit, EDynamicMeshAttributeChangeFlags::Unknown, false);	
+
+	Next();
+}
+
+#if WITH_EDITOR
+void UAvaBevelModifier::PostEditChangeProperty(FPropertyChangedEvent& PropertyChangedEvent)
+{
+	Super::PostEditChangeProperty(PropertyChangedEvent);
+
+	const FName MemberName = PropertyChangedEvent.GetMemberPropertyName();
+	
+	static const FName InsetName = GET_MEMBER_NAME_CHECKED(UAvaBevelModifier, Inset);
+	static const FName IterationsName = GET_MEMBER_NAME_CHECKED(UAvaBevelModifier, Iterations);
+	
+	if (MemberName == InsetName)
+	{
+		OnInsetChanged();
+	}
+	else if (MemberName == IterationsName)
+	{
+		OnIterationsChanged();
+	}
+}
+#endif
+
+void UAvaBevelModifier::SetInset(float InInset)
+{
+	if (Inset == InInset)
+	{
+		return;
+	}
+
+	Inset = FMath::Clamp(InInset, UAvaBevelModifier::MinInset, GetMaxBevel());
+	
+	OnInsetChanged();
+}
+
+void UAvaBevelModifier::SetIterations(int32 InIterations)
+{
+	if (Iterations == InIterations)
+	{
+		return;
+	}
+
+	Iterations = FMath::Clamp(InIterations, UAvaBevelModifier::MinIterations, UAvaBevelModifier::MaxIterations);
+	
+	OnIterationsChanged();
+}
+
+void UAvaBevelModifier::OnInsetChanged()
+{
+	Inset = FMath::Min(Inset, GetMaxBevel());
+	MarkModifierDirty();
+}
+
+void UAvaBevelModifier::OnIterationsChanged()
+{
+	Inset = FMath::Min(Inset, GetMaxBevel());
+	MarkModifierDirty();
+}
+
+float UAvaBevelModifier::GetMaxBevel() const
+{
+	float MaxBevelDistance = 0;
+	
+	if (!PreModifierCachedMesh.IsSet())
+	{
+		return MaxBevelDistance;
+	}
+
+	const FBox Bounds = static_cast<FBox>(PreModifierCachedMesh.GetValue().GetBounds(true));
+	const FVector Size3d = Bounds.GetSize();
+
+	const float MinBevelDistance = FMath::Min3(Size3d.X / 2, Size3d.Y / 2, Size3d.Z / 2);
+	float Divider = 4.f;
+	float BevelDistance = MinBevelDistance;
+	MaxBevelDistance = MinBevelDistance;
+	
+	for (int32 Idx = 1; Idx < Iterations; Idx++)
+	{
+		const float Reduce = BevelDistance / Divider;
+		MaxBevelDistance -= Reduce;
+		BevelDistance = Reduce;
+		Divider /= 2.f;
+	}
+	
+	return MaxBevelDistance;
+}
+
+#undef LOCTEXT_NAMESPACE
diff --git a/Engine/Plugins/Experimental/Avalanche/Source/AvalancheModifiers/Private/Modifiers/AvaDynamicMeshConverterModifier.cpp b/Engine/Plugins/Experimental/Avalanche/Source/AvalancheModifiers/Private/Modifiers/AvaDynamicMeshConverterModifier.cpp
new file mode 100644
--- /dev/null
+++ b/Engine/Plugins/Experimental/Avalanche/Source/AvalancheModifiers/Private/Modifiers/AvaDynamicMeshConverterModifier.cpp
@@ -0,0 +1,835 @@
+// Copyright Epic Games, Inc. All Rights Reserved.
+
+#include "Modifiers/AvaDynamicMeshConverterModifier.h"
+
+#include "Async/Async.h"
+#include "Components/BrushComponent.h"
+#include "Components/DynamicMeshComponent.h"
+#include "Components/SkeletalMeshComponent.h"
+#include "Components/StaticMeshComponent.h"
+#include "DynamicMesh/MeshTransforms.h"
+#include "DynamicMeshEditor.h"
+#include "Engine/StaticMesh.h"
+#include "GeometryScript/MeshAssetFunctions.h"
+#include "GeometryScript/MeshBasicEditFunctions.h"
+#include "GeometryScript/SceneUtilityFunctions.h"
+#include "ProceduralMeshComponent.h"
+#include "Engine/CollisionProfile.h"
+
+#if WITH_EDITOR
+#include "AssetRegistry/AssetRegistryModule.h"
+#include "AssetToolsModule.h"
+#include "Dialogs/DlgPickAssetPath.h"
+#endif
+
+#define LOCTEXT_NAMESPACE "AvaDynamicMeshConverterModifier"
+
+FAvaDynamicMeshConverterModifierComponentState::FAvaDynamicMeshConverterModifierComponentState(UPrimitiveComponent* InPrimitiveComponent)
+	: Component(InPrimitiveComponent)
+{
+	if (Component.IsValid())
+	{
+		if (const AActor* ComponentOwner = Component->GetOwner())
+		{
+#if WITH_EDITOR
+			bActorHiddenInGame = ComponentOwner->IsHidden();
+			bActorHiddenInEditor = ComponentOwner->IsTemporarilyHiddenInEditor();
+#endif
+			if (const USceneComponent* RootComponent = ComponentOwner->GetRootComponent())
+			{
+				bComponentVisible = RootComponent->IsVisible();
+				bComponentHiddenInGame = RootComponent->bHiddenInGame;
+			}
+		}
+	}
+}
+
+void UAvaDynamicMeshConverterModifier::OnModifierCDOSetup(FActorModifierCoreMetadata& InMetadata)
+{
+	Super::OnModifierCDOSetup(InMetadata);
+
+	InMetadata.SetName(TEXT("DynamicMeshConverter"));
+	InMetadata.SetCategory(TEXT("Conversion"));
+#if WITH_EDITOR
+	InMetadata.SetDescription(LOCTEXT("ModifierDescription", "Converts a non dynamic mesh actor into a dynamic mesh"));
+#endif
+	InMetadata.SetCompatibilityRule([](const AActor* InActor)->bool
+	{
+		return InActor && !InActor->FindComponentByClass<UDynamicMeshComponent>();
+	});
+}
+
+void UAvaDynamicMeshConverterModifier::OnModifierAdded(EActorModifierCoreEnableReason InReason)
+{
+	Super::OnModifierAdded(InReason);
+
+	AddDynamicMeshComponent();
+	bConvertMesh = true;
+}
+
+void UAvaDynamicMeshConverterModifier::OnModifierEnabled(EActorModifierCoreEnableReason InReason)
+{
+	Super::OnModifierEnabled(InReason);
+
+	// On load update state by converting source again
+	if (InReason == EActorModifierCoreEnableReason::Load)
+	{
+		bConvertMesh = true;
+		MarkModifierDirty();
+	}
+}
+
+void UAvaDynamicMeshConverterModifier::RestorePreState()
+{
+	UAvaGeometryBaseModifier::RestorePreState();
+	
+	for (FAvaDynamicMeshConverterModifierComponentState& ConvertedComponent : ConvertedComponents)
+	{
+		if (ConvertedComponent.Component.IsValid())
+		{
+			AActor* ComponentActor = ConvertedComponent.Component->GetOwner();
+
+			if (ComponentActor != GetModifiedActor())
+			{
+#if WITH_EDITOR
+				// Hide actor but do not hide ourselves
+				ComponentActor->SetActorHiddenInGame(ConvertedComponent.bActorHiddenInGame);
+				ComponentActor->SetIsTemporarilyHiddenInEditor(ConvertedComponent.bActorHiddenInEditor);
+#endif
+			}
+			else if (USceneComponent* RootComponent = ComponentActor->GetRootComponent())
+			{
+				// In the meantime hide root component but later hide component itself
+				RootComponent->SetHiddenInGame(ConvertedComponent.bComponentHiddenInGame);
+				RootComponent->SetVisibility(ConvertedComponent.bComponentVisible);
+			}
+		}
+	}
+}
+
+void UAvaDynamicMeshConverterModifier::OnModifierRemoved(EActorModifierCoreDisableReason InReason)
+{
+	Super::OnModifierRemoved(InReason);
+
+	if (InReason != EActorModifierCoreDisableReason::Destroyed)
+	{
+		RemoveDynamicMeshComponent();
+	}
+}
+
+void UAvaDynamicMeshConverterModifier::Apply()
+{
+	if (!IsMeshValid())
+	{
+		Fail(LOCTEXT("InvalidDynamicMeshComponent", "Invalid dynamic mesh component on modified actor"));
+		return;
+	}
+
+	UDynamicMeshComponent* DynMeshComponent = GetMeshComponent();
+
+	if (!IsValid(DynMeshComponent))
+	{
+		Fail(LOCTEXT("InvalidDynamicMeshComponent", "Invalid dynamic mesh component on modified actor"));
+		return;
+	}
+	
+	using namespace UE::Geometry;
+	
+	if (bConvertMesh)
+	{
+		DynMeshComponent->SetNumMaterials(0);
+		DynMeshComponent->EditMesh([this, DynMeshComponent](FDynamicMesh3& AppendToMesh)
+		{
+			AppendToMesh.Clear();
+			FMeshIndexMappings TmpMappings;
+			FDynamicMeshEditor Editor(&AppendToMesh);
+			FGeometryScriptAppendMeshOptions AppendOptions;
+			AppendOptions.CombineMode = EGeometryScriptCombineAttributesMode::EnableAllMatching;
+			int32 MaterialCount = 0;
+			// Convert meshes
+			ConvertedComponents.Empty();
+			ConvertComponents(ConvertedComponents);
+			for (const FAvaDynamicMeshConverterModifierComponentState& OutConvert : ConvertedComponents)
+			{
+				// Enable matching attributes & append mesh
+				AppendOptions.UpdateAttributesForCombineMode(AppendToMesh, OutConvert.Mesh);
+				Editor.AppendMesh(&OutConvert.Mesh, TmpMappings);
+				
+				if (OutConvert.Mesh.HasAttributes() && OutConvert.Mesh.Attributes()->HasMaterialID())
+				{
+					// Fix triangles materials linking
+					const FDynamicMeshMaterialAttribute* FromMaterialIDAttrib = OutConvert.Mesh.Attributes()->GetMaterialID();
+					FDynamicMeshMaterialAttribute* ToMaterialIDAttrib = AppendToMesh.Attributes()->GetMaterialID();
+					TMap<int32, int32> MaterialMap;
+					for (const TPair<int32, int32>& FromToTId : TmpMappings.GetTriangleMap().GetForwardMap())
+					{
+						const int32 FromMatId = FromMaterialIDAttrib->GetValue(FromToTId.Key);
+						const int32 ToMatId = FromMatId + MaterialCount;
+						MaterialMap.Add(FromMatId, ToMatId);
+						ToMaterialIDAttrib->SetNewValue(FromToTId.Value, ToMatId);
+					}
+					MaterialCount += MaterialMap.Num();
+					
+					// Reapply original materials
+					if (OutConvert.Component.IsValid())
+					{
+						for (const TPair<int32, int32>& MatPair : MaterialMap)
+						{
+							DynMeshComponent->SetMaterial(MatPair.Value, OutConvert.Component->GetMaterial(MatPair.Key));
+						}
+					}
+				}
+				
+				if (OutConvert.Component.IsValid())
+				{
+					AActor* ComponentActor = OutConvert.Component->GetOwner();
+					
+					// Hide converted component
+					if (bHideConvertedMesh)
+					{
+						if (ComponentActor != GetModifiedActor())
+						{
+							ComponentActor->SetActorHiddenInGame(true);
+#if WITH_EDITOR
+							ComponentActor->SetIsTemporarilyHiddenInEditor(true);
+#endif
+						}
+						else if (USceneComponent* RootComponent = ComponentActor->GetRootComponent())
+						{
+							RootComponent->SetHiddenInGame(true);
+							RootComponent->SetVisibility(false);
+						}
+					}
+				}
+				
+				TmpMappings.Reset();
+			}
+			
+			ConvertedMesh = AppendToMesh;
+		});
+	}
+	else if (ConvertedMesh.IsSet())
+	{
+		DynMeshComponent->EditMesh([this](FDynamicMesh3& AppendToMesh)
+		{
+			AppendToMesh = ConvertedMesh.GetValue();
+		});
+	}
+
+	Next();
+}
+
+#if WITH_EDITOR
+void UAvaDynamicMeshConverterModifier::PostEditChangeProperty(FPropertyChangedEvent& PropertyChangedEvent)
+{
+	Super::PostEditChangeProperty(PropertyChangedEvent);
+
+	const FName MemberName = PropertyChangedEvent.GetMemberPropertyName();
+	
+	static const FName SourceActorName = GET_MEMBER_NAME_CHECKED(UAvaDynamicMeshConverterModifier, SourceActorWeak);
+
+	if (MemberName == SourceActorName)
+	{
+		OnSourceActorChanged();
+	}
+}
+
+void UAvaDynamicMeshConverterModifier::ConvertToStaticMeshAsset()
+{
+	using namespace UE::Geometry;
+
+	UDynamicMeshComponent* DynMeshComponent = GetMeshComponent();
+	const AActor* OwningActor = GetModifiedActor();
+
+	if (!OwningActor || !DynMeshComponent)
+	{
+		return;
+	}
+	
+	// generate name for asset
+	const FString NewNameSuggestion = TEXT("SM_Avalanche_") + OwningActor->GetActorNameOrLabel();
+	FString PackageName = FString(TEXT("/Game/Meshes/")) + NewNameSuggestion;
+	FString AssetName;
+
+	const FAssetToolsModule& AssetToolsModule = FModuleManager::LoadModuleChecked<FAssetToolsModule>("AssetTools");
+	AssetToolsModule.Get().CreateUniqueAssetName(PackageName, TEXT(""), PackageName, AssetName);
+
+	const TSharedPtr<SDlgPickAssetPath> PickAssetPathWidget =
+		SNew(SDlgPickAssetPath)
+		.Title(LOCTEXT("ConvertToStaticMeshPickName", "Choose New StaticMesh Location"))
+		.DefaultAssetPath(FText::FromString(PackageName));
+
+	if (PickAssetPathWidget->ShowModal() != EAppReturnType::Ok)
+	{
+		return;
+	}
+
+	// get input name provided by user
+	FString UserPackageName = PickAssetPathWidget->GetFullAssetPath().ToString();
+	FName MeshName(*FPackageName::GetLongPackageAssetName(UserPackageName));
+
+	// is input name valid ?
+	if (MeshName == NAME_None)
+	{
+		// Use default if invalid
+		UserPackageName = PackageName;
+		MeshName = *AssetName;
+	}
+
+	const FDynamicMesh3* MeshIn = DynMeshComponent->GetMesh();
+
+	// empty mesh do not export
+	if (!MeshIn || MeshIn->TriangleCount() == 0)
+	{
+		return;
+	}
+	
+	// find/create package
+	UPackage* Package = CreatePackage(*UserPackageName);
+	check(Package);
+
+	// Create StaticMesh object
+	UStaticMesh* DestinationMesh = NewObject<UStaticMesh>(Package, MeshName, RF_Public | RF_Standalone);
+	UDynamicMesh* SourceMesh = DynMeshComponent->GetDynamicMesh();
+	
+	// export options
+	FGeometryScriptCopyMeshToAssetOptions AssetOptions;
+	AssetOptions.bReplaceMaterials = false;
+	AssetOptions.bEnableRecomputeNormals = false;
+	AssetOptions.bEnableRecomputeTangents = false;
+	AssetOptions.bEnableRemoveDegenerates = true;
+	
+	// LOD options
+	FGeometryScriptMeshWriteLOD TargetLOD;
+	TargetLOD.LODIndex = 0;
+
+	EGeometryScriptOutcomePins OutResult;
+
+	UGeometryScriptLibrary_StaticMeshFunctions::CopyMeshToStaticMesh(SourceMesh, DestinationMesh, AssetOptions, TargetLOD, OutResult);
+	DestinationMesh->GetBodySetup()->AggGeom = DynMeshComponent->GetBodySetup()->AggGeom;
+	
+	if (OutResult == EGeometryScriptOutcomePins::Success)
+	{
+		// Notify asset registry of new asset
+		FAssetRegistryModule::AssetCreated(DestinationMesh);
+	}
+}
+#undef LOCTEXT_NAMESPACE
+#endif
+
+void UAvaDynamicMeshConverterModifier::SetSourceActorWeak(const TWeakObjectPtr<AActor>& InActor)
+{
+	if (InActor.Get() == SourceActorWeak.Get())
+	{
+		return;
+	}
+
+	SourceActorWeak = InActor;
+	OnSourceActorChanged();
+}
+
+void UAvaDynamicMeshConverterModifier::SetComponentType(int32 InComponentType)
+{
+	if (InComponentType == ComponentType)
+	{
+		return;
+	}
+
+	ComponentType = InComponentType;
+}
+
+void UAvaDynamicMeshConverterModifier::SetFilterActorMode(EAvaDynamicMeshConverterModifierFilter InFilter)
+{
+	FilterActorMode = InFilter;
+}
+
+void UAvaDynamicMeshConverterModifier::SetFilterActorClasses(const TSet<TSubclassOf<AActor>>& InClasses)
+{
+	FilterActorClasses = InClasses;
+}
+
+void UAvaDynamicMeshConverterModifier::SetIncludeAttachedActors(bool bInInclude)
+{
+	bIncludeAttachedActors = bInInclude;
+}
+
+void UAvaDynamicMeshConverterModifier::SetHideConvertedMesh(bool bInHide)
+{
+	bHideConvertedMesh = bInHide;
+}
+
+void UAvaDynamicMeshConverterModifier::OnSourceActorChanged()
+{
+	const AActor* SourceActor = SourceActorWeak.Get();
+	const AActor* ActorModified = GetModifiedActor();
+	if (!SourceActor || !ActorModified)
+	{
+		return;
+	}
+	bHideConvertedMesh = SourceActor == ActorModified || SourceActor->IsAttachedTo(ActorModified);
+}
+
+void UAvaDynamicMeshConverterModifier::ConvertComponents(TArray<FAvaDynamicMeshConverterModifierComponentState>& OutResults) const
+{
+	if (!IsMeshValid() || !SourceActorWeak.IsValid())
+	{
+		return;
+	}
+	const FTransform SourceTransform = SourceActorWeak->GetActorTransform();
+	UDynamicMeshComponent* DynMeshComponent = GetMeshComponent();
+	UDynamicMesh* OutputDynamicMesh = DynMeshComponent->GetDynamicMesh();
+	static const FGeometryScriptCopyMeshFromAssetOptions FromMeshOptions;
+	static FGeometryScriptMeshReadLOD FromMeshLOD;
+	FromMeshLOD.LODType = EGeometryScriptLODType::SourceModel;
+	TArray<AActor*> FilteredActors;
+	GetFilteredActors(FilteredActors);
+	if (HasFlag(EAvaDynamicMeshConverterModifierType::StaticMeshComponent))
+	{
+		TArray<UStaticMeshComponent*> Components;
+		GetStaticMeshComponents(FilteredActors, Components);
+		for (UStaticMeshComponent* Component : Components)
+		{
+			UStaticMesh* StaticMesh = Component->GetStaticMesh();
+
+			// convert to dynamic mesh
+			EGeometryScriptOutcomePins OutResult;
+
+			UGeometryScriptLibrary_StaticMeshFunctions::CopyMeshFromStaticMesh(StaticMesh, OutputDynamicMesh, FromMeshOptions, FromMeshLOD, OutResult);
+			if (OutResult == EGeometryScriptOutcomePins::Success)
+			{
+				// Transform the new mesh relative to the component
+				const FTransform RelativeTransform = Component->GetComponentTransform().GetRelativeTransform(SourceTransform);
+				OutputDynamicMesh->EditMesh([&OutResults, Component, RelativeTransform](FDynamicMesh3& EditMesh)
+				{
+					MeshTransforms::ApplyTransform(EditMesh, RelativeTransform);
+					FAvaDynamicMeshConverterModifierComponentState NewResult(Component);
+					NewResult.Mesh = MoveTemp(EditMesh);
+					OutResults.Add(NewResult);
+					// replace by empty mesh
+					FDynamicMesh3 EmptyMesh;
+					EditMesh = MoveTemp(EmptyMesh);
+				});
+			}
+		}
+	}
+	if (HasFlag(EAvaDynamicMeshConverterModifierType::DynamicMeshComponent))
+	{
+		TArray<UDynamicMeshComponent*> Components;
+		GetDynamicMeshComponents(FilteredActors, Components);
+		for (UDynamicMeshComponent* Component : Components)
+		{
+			const UDynamicMesh* DynamicMesh = Component->GetDynamicMesh();
+			// Transform the new mesh relative to the component
+			const FTransform RelativeTransform = Component->GetComponentTransform().GetRelativeTransform(SourceTransform);
+			// Create a copy
+			DynamicMesh->ProcessMesh([&OutResults, Component, RelativeTransform](const FDynamicMesh3& EditMesh)
+			{
+				FDynamicMesh3 CopyMesh = EditMesh;
+				MeshTransforms::ApplyTransform(CopyMesh, RelativeTransform);
+				FAvaDynamicMeshConverterModifierComponentState NewResult(Component);
+				NewResult.Mesh = MoveTemp(CopyMesh);
+				OutResults.Add(NewResult);
+			});
+		}
+	}
+	if (HasFlag(EAvaDynamicMeshConverterModifierType::SkeletalMeshComponent))
+	{
+		TArray<USkeletalMeshComponent*> Components;
+		GetSkeletalMeshComponents(FilteredActors, Components);
+		for (USkeletalMeshComponent* Component : Components)
+		{
+			USkeletalMesh* SkeletalMesh = Component->GetSkeletalMeshAsset();
+
+			// convert to dynamic mesh
+			EGeometryScriptOutcomePins OutResult;
+
+			UGeometryScriptLibrary_StaticMeshFunctions::CopyMeshFromSkeletalMesh(SkeletalMesh, OutputDynamicMesh, FromMeshOptions, FromMeshLOD, OutResult);
+			if (OutResult == EGeometryScriptOutcomePins::Success)
+			{
+				// Transform the new mesh relative to the component
+				const FTransform RelativeTransform = Component->GetComponentTransform().GetRelativeTransform(SourceTransform);
+				OutputDynamicMesh->EditMesh([&OutResults, Component, RelativeTransform](FDynamicMesh3& EditMesh)
+				{
+					MeshTransforms::ApplyTransform(EditMesh, RelativeTransform);
+					FAvaDynamicMeshConverterModifierComponentState NewResult(Component);
+					NewResult.Mesh = MoveTemp(EditMesh);
+					OutResults.Add(NewResult);
+					// replace by empty mesh
+					FDynamicMesh3 EmptyMesh;
+					EditMesh = MoveTemp(EmptyMesh);
+				});
+			}
+		}
+	}
+	if (HasFlag(EAvaDynamicMeshConverterModifierType::BrushComponent))
+	{
+		static FGeometryScriptCopyMeshFromComponentOptions Options;
+		Options.RequestedLOD = FromMeshLOD;
+		FTransform Transform;
+		TArray<UBrushComponent*> Components;
+		GetBrushComponents(FilteredActors, Components);
+		for (UBrushComponent* Component : Components)
+		{
+			// convert to dynamic mesh
+			EGeometryScriptOutcomePins OutResult;
+
+			UGeometryScriptLibrary_SceneUtilityFunctions::CopyMeshFromComponent(Component, OutputDynamicMesh, Options, false, Transform, OutResult);
+			if (OutResult == EGeometryScriptOutcomePins::Success)
+			{
+				// Transform the new mesh relative to the component
+				const FTransform RelativeTransform = Component->GetComponentTransform().GetRelativeTransform(SourceTransform);
+				OutputDynamicMesh->EditMesh([&OutResults, Component, RelativeTransform](FDynamicMesh3& EditMesh)
+				{
+					MeshTransforms::ApplyTransform(EditMesh, RelativeTransform);
+					FAvaDynamicMeshConverterModifierComponentState NewResult(Component);
+					NewResult.Mesh = MoveTemp(EditMesh);
+					OutResults.Add(NewResult);
+					// replace by empty mesh
+					FDynamicMesh3 EmptyMesh;
+					EditMesh = MoveTemp(EmptyMesh);
+				});
+			}
+		}
+	}
+	if (HasFlag(EAvaDynamicMeshConverterModifierType::ProceduralMeshComponent))
+	{
+		using namespace UE::Geometry;
+		TArray<UProceduralMeshComponent*> Components;
+		GetProceduralMeshComponents(FilteredActors, Components);
+		for (UProceduralMeshComponent* Component : Components)
+		{
+			const int32 SectionCount = Component->GetNumSections();
+			if (SectionCount == 0)
+			{
+				continue;
+			}
+
+			// Transform the new mesh relative to the component
+			const FTransform RelativeTransform = Component->GetComponentTransform().GetRelativeTransform(SourceTransform);
+			
+			FAvaDynamicMeshConverterModifierComponentState NewResult(Component);
+			NewResult.Mesh.EnableAttributes();
+			NewResult.Mesh.Attributes()->EnablePrimaryColors();
+			NewResult.Mesh.Attributes()->EnableMaterialID();
+			NewResult.Mesh.Attributes()->SetNumNormalLayers(1);
+			NewResult.Mesh.Attributes()->SetNumUVLayers(1);
+			NewResult.Mesh.Attributes()->SetNumPolygroupLayers(1);
+			NewResult.Mesh.Attributes()->EnableTangents();
+			
+			FDynamicMeshColorOverlay* ColorOverlay = NewResult.Mesh.Attributes()->PrimaryColors();
+			FDynamicMeshNormalOverlay* NormalOverlay = NewResult.Mesh.Attributes()->PrimaryNormals();
+			FDynamicMeshUVOverlay* UVOverlay = NewResult.Mesh.Attributes()->PrimaryUV();
+			FDynamicMeshMaterialAttribute* MaterialAttr = NewResult.Mesh.Attributes()->GetMaterialID();
+			FDynamicMeshPolygroupAttribute* GroupAttr = NewResult.Mesh.Attributes()->GetPolygroupLayer(0);
+			FDynamicMeshNormalOverlay* TangentOverlay = NewResult.Mesh.Attributes()->PrimaryTangents();
+			
+			for (int32 SectionIdx = 0; SectionIdx < SectionCount; SectionIdx++)
+			{
+				if (FProcMeshSection* Section = Component->GetProcMeshSection(SectionIdx))
+				{
+					if (Section->bSectionVisible)
+					{
+						TArray<int32> VtxIds;
+						TArray<int32> NormalIds;
+						TArray<int32> ColorIds;
+						TArray<int32> UVIds;
+						TArray<int32> TaIds;
+						
+						// copy vertices data (position, normal, color, UV, tangent)
+						for (FProcMeshVertex& SectionVertex : Section->ProcVertexBuffer)
+						{
+							int32 VId = NewResult.Mesh.AppendVertex(SectionVertex.Position);
+							VtxIds.Add(VId);
+							
+							int32 NId = NormalOverlay->AppendElement(static_cast<FVector3f>(SectionVertex.Normal));
+							NormalIds.Add(NId);
+							
+							int32 CId = ColorOverlay->AppendElement(static_cast<FVector4f>(SectionVertex.Color));
+							ColorIds.Add(CId);
+							
+							int32 UVId = UVOverlay->AppendElement(static_cast<FVector2f>(SectionVertex.UV0));
+							UVIds.Add(UVId);
+
+							int32 TaId = TangentOverlay->AppendElement(static_cast<FVector3f>(SectionVertex.Tangent.TangentX));
+							TaIds.Add(TaId);
+						}
+						
+						// copy tris data
+						if (Section->ProcIndexBuffer.Num() % 3 != 0)
+						{
+							continue;
+						}
+						for (int32 Idx = 0; Idx < Section->ProcIndexBuffer.Num(); Idx+=3)
+						{
+							int32 VIdx1 = Section->ProcIndexBuffer[Idx];
+							int32 VIdx2 = Section->ProcIndexBuffer[Idx + 1];
+							int32 VIdx3 = Section->ProcIndexBuffer[Idx + 2];
+							
+							int32 VId1 = VtxIds[VIdx1];
+							int32 VId2 = VtxIds[VIdx2];
+							int32 VId3 = VtxIds[VIdx3];
+
+							int32 TId = NewResult.Mesh.AppendTriangle(VId1, VId2, VId3, SectionIdx);
+							if (TId < 0)
+							{
+								continue;
+							}
+
+							NormalOverlay->SetTriangle(TId, FIndex3i(NormalIds[VIdx1], NormalIds[VIdx2], NormalIds[VIdx3]), true);
+							ColorOverlay->SetTriangle(TId, FIndex3i(ColorIds[VIdx1], ColorIds[VIdx2], ColorIds[VIdx3]), true);
+							UVOverlay->SetTriangle(TId, FIndex3i(UVIds[VIdx1], UVIds[VIdx2], UVIds[VIdx3]), true);
+							TangentOverlay->SetTriangle(TId, FIndex3i(TaIds[VIdx1], TaIds[VIdx2], TaIds[VIdx3]), true);
+							
+							MaterialAttr->SetValue(TId, SectionIdx);
+							GroupAttr->SetValue(TId, SectionIdx);
+						}
+					}
+				}
+			}
+
+			if (NewResult.Mesh.TriangleCount() > 0)
+			{
+				MeshTransforms::ApplyTransform(NewResult.Mesh, RelativeTransform);
+				OutResults.Add(NewResult);
+			}
+		}
+	}
+}
+
+bool UAvaDynamicMeshConverterModifier::HasFlag(EAvaDynamicMeshConverterModifierType InFlag) const
+{
+	return EnumHasAnyFlags(static_cast<EAvaDynamicMeshConverterModifierType>(ComponentType), InFlag);
+}
+
+void UAvaDynamicMeshConverterModifier::AddDynamicMeshComponent()
+{
+	UDynamicMeshComponent* DynMeshComponent = GetMeshComponent();
+	
+	if (DynMeshComponent)
+	{
+		return;
+	}
+
+	AActor* ActorModified = GetModifiedActor();
+
+	if (!IsValid(ActorModified))
+	{
+		return;
+	}
+	
+#if WITH_EDITOR
+	ActorModified->Modify();
+	Modify();
+#endif
+
+	const UClass* const NewComponentClass = UDynamicMeshComponent::StaticClass();
+
+	// Construct the new component and attach as needed
+	DynMeshComponent = NewObject<UDynamicMeshComponent>(ActorModified
+		, NewComponentClass
+		, MakeUniqueObjectName(ActorModified, NewComponentClass, TEXT("DynamicMeshComponent"))
+		, RF_Transactional);
+
+	// Add to SerializedComponents array so it gets saved
+	ActorModified->AddInstanceComponent(DynMeshComponent);
+	DynMeshComponent->OnComponentCreated();
+	DynMeshComponent->RegisterComponent();
+
+	if (USceneComponent* RootComponent = ActorModified->GetRootComponent())
+	{
+		static const FAttachmentTransformRules AttachRules(EAttachmentRule::SnapToTarget, EAttachmentRule::SnapToTarget, EAttachmentRule::SnapToTarget, false);
+		DynMeshComponent->AttachToComponent(RootComponent, AttachRules);
+	}
+	else
+	{
+		ActorModified->SetRootComponent(DynMeshComponent);
+	}
+
+	DynMeshComponent->SetCollisionProfileName(UCollisionProfile::BlockAll_ProfileName);
+	DynMeshComponent->SetGenerateOverlapEvents(true);
+
+#if WITH_EDITOR
+	// Rerun construction scripts
+	ActorModified->RerunConstructionScripts();
+#endif
+	
+	bComponentCreated = true;
+}
+
+void UAvaDynamicMeshConverterModifier::RemoveDynamicMeshComponent()
+{
+	UDynamicMeshComponent* DynMeshComponent = GetMeshComponent();
+	
+	if (!DynMeshComponent)
+	{
+		return;
+	}
+
+	// Did we create the component or was it already there
+	if (!bComponentCreated)
+	{
+		return;
+	}
+	
+	AActor* ActorModified = GetModifiedActor();
+
+	if (!IsValid(ActorModified))
+	{
+		return;
+	}
+	
+#if WITH_EDITOR
+	ActorModified->Modify();
+	Modify();
+#endif
+
+	const FDetachmentTransformRules DetachRules(EDetachmentRule::KeepWorld, false);
+	DynMeshComponent->DetachFromComponent(DetachRules);
+	
+	ActorModified->RemoveInstanceComponent(DynMeshComponent);
+	DynMeshComponent->DestroyComponent(false);
+	
+	bComponentCreated = false;
+}
+
+void UAvaDynamicMeshConverterModifier::GetFilteredActors(TArray<AActor*>& OutActors) const
+{
+	if (AActor* OriginActor = SourceActorWeak.Get())
+	{
+		OutActors.Add(OriginActor);
+		if (bIncludeAttachedActors)
+		{
+			OriginActor->GetAttachedActors(OutActors, false, true);
+		}
+		// Filter actor class
+		if (FilterActorMode != EAvaDynamicMeshConverterModifierFilter::None)
+		{
+			for (int32 Idx = OutActors.Num() - 1; Idx >= 0; Idx--)
+			{
+				const AActor* CurrentActor = OutActors[Idx];
+				if (!IsValid(CurrentActor))
+				{
+					continue;
+				}
+				// Include this actor if it's in the filter class
+				if (FilterActorMode == EAvaDynamicMeshConverterModifierFilter::Include)
+				{
+					if (!FilterActorClasses.Contains(CurrentActor->GetClass()))
+					{
+						OutActors.RemoveAt(Idx);
+					}
+				}
+				else // Exclude this actor if it's in the filter class
+				{
+					if (FilterActorClasses.Contains(CurrentActor->GetClass()))
+					{
+						OutActors.RemoveAt(Idx);
+					}
+				}
+			}
+		}
+	}
+}
+
+void UAvaDynamicMeshConverterModifier::GetStaticMeshComponents(const TArray<AActor*>& InActors, TArray<UStaticMeshComponent*>& OutComponents) const
+{
+	for (const AActor* Actor : InActors)
+	{
+		TArray<UStaticMeshComponent*> OutMeshComponents;
+		Actor->GetComponents(OutMeshComponents, false);
+		OutComponents.Append(OutMeshComponents);
+	}
+	// remove all invalid components
+	OutComponents.RemoveAll([](const UStaticMeshComponent* InComponent)->bool
+	{
+#if WITH_EDITOR
+		return !IsValid(InComponent) || InComponent->IsVisualizationComponent();
+#else
+		return !IsValid(InComponent);
+#endif
+	});
+}
+
+void UAvaDynamicMeshConverterModifier::GetDynamicMeshComponents(const TArray<AActor*>& InActors, TArray<UDynamicMeshComponent*>& OutComponents) const
+{
+	for (const AActor* Actor : InActors)
+	{
+		TArray<UDynamicMeshComponent*> OutMeshComponents;
+		Actor->GetComponents(OutMeshComponents, false);
+		OutComponents.Append(OutMeshComponents);
+	}
+	// remove all invalid components & and the modifier created component
+	OutComponents.RemoveAll([this](const UDynamicMeshComponent* InComponent)->bool
+	{
+#if WITH_EDITOR
+		return !IsValid(InComponent) || InComponent == GetMeshComponent() || InComponent->IsVisualizationComponent();
+#else
+		return !IsValid(InComponent) || InComponent == GetMeshComponent();
+#endif
+	});
+}
+
+void UAvaDynamicMeshConverterModifier::GetSkeletalMeshComponents(const TArray<AActor*>& InActors, TArray<USkeletalMeshComponent*>& OutComponents) const
+{
+	for (const AActor* Actor : InActors)
+	{
+		TArray<USkeletalMeshComponent*> OutMeshComponents;
+		Actor->GetComponents(OutMeshComponents, false);
+		OutComponents.Append(OutMeshComponents);
+	}
+	// remove all invalid components
+	OutComponents.RemoveAll([this](const USkeletalMeshComponent* InComponent)->bool
+	{
+#if WITH_EDITOR
+		return !IsValid(InComponent) || InComponent->IsVisualizationComponent();
+#else
+		return !IsValid(InComponent);
+#endif
+	});
+}
+
+void UAvaDynamicMeshConverterModifier::GetBrushComponents(const TArray<AActor*>& InActors, TArray<UBrushComponent*>& OutComponents) const
+{
+	for (const AActor* Actor : InActors)
+	{
+		TArray<UBrushComponent*> OutMeshComponents;
+		Actor->GetComponents(OutMeshComponents, false);
+		OutComponents.Append(OutMeshComponents);
+	}
+	// remove all invalid components
+	OutComponents.RemoveAll([this](const UBrushComponent* InComponent)->bool
+	{
+#if WITH_EDITOR
+		return !IsValid(InComponent) || InComponent->IsVisualizationComponent();
+#else
+		return !IsValid(InComponent);
+#endif
+	});
+}
+
+void UAvaDynamicMeshConverterModifier::GetProceduralMeshComponents(const TArray<AActor*>& InActors, TArray<UProceduralMeshComponent*>& OutComponents) const
+{
+	for (const AActor* Actor : InActors)
+	{
+		TArray<UProceduralMeshComponent*> OutMeshComponents;
+		Actor->GetComponents(OutMeshComponents, false);
+		OutComponents.Append(OutMeshComponents);
+	}
+	// remove all invalid components
+	OutComponents.RemoveAll([this](const UProceduralMeshComponent* InComponent)->bool
+	{
+#if WITH_EDITOR
+		return !IsValid(InComponent) || InComponent->IsVisualizationComponent();
+#else
+		return !IsValid(InComponent);
+#endif
+	});
+}
+
+void UAvaDynamicMeshConverterModifier::ConvertMesh()
+{
+	bConvertMesh = true;
+	MarkModifierDirty();
+}
+
+#undef LOCTEXT_NAMESPACE
diff --git a/Engine/Plugins/Experimental/Avalanche/Source/AvalancheModifiers/Private/Modifiers/AvaGridArrangeModifier.cpp b/Engine/Plugins/Experimental/Avalanche/Source/AvalancheModifiers/Private/Modifiers/AvaGridArrangeModifier.cpp
new file mode 100644
--- /dev/null
+++ b/Engine/Plugins/Experimental/Avalanche/Source/AvalancheModifiers/Private/Modifiers/AvaGridArrangeModifier.cpp
@@ -0,0 +1,214 @@
+// Copyright Epic Games, Inc. All Rights Reserved.
+
+#include "Modifiers/AvaGridArrangeModifier.h"
+
+#include "AvaDefs.h"
+#include "AvaModifiersActorUtils.h"
+#include "GameFramework/Actor.h"
+#include "Shared/AvaTransformModifierShared.h"
+#include "Shared/AvaVisibilityModifierShared.h"
+
+#define LOCTEXT_NAMESPACE "AvaGridArrangeModifier"
+
+#if WITH_EDITOR
+void UAvaGridArrangeModifier::PostEditChangeProperty(FPropertyChangedEvent& PropertyChangedEvent)
+{
+	Super::PostEditChangeProperty(PropertyChangedEvent);
+
+	const FName MemberName = PropertyChangedEvent.GetMemberPropertyName();
+
+	static const FName CountPropertyName = GET_MEMBER_NAME_CHECKED(UAvaGridArrangeModifier, Count);
+	static const FName SpreadPropertyName = GET_MEMBER_NAME_CHECKED(UAvaGridArrangeModifier, Spread);
+	static const FName StartCornerPropertyName = GET_MEMBER_NAME_CHECKED(UAvaGridArrangeModifier, StartCorner);
+	static const FName StartDirectionPropertyName = GET_MEMBER_NAME_CHECKED(UAvaGridArrangeModifier, StartDirection);
+	
+	if (MemberName == CountPropertyName
+		|| MemberName == SpreadPropertyName
+		|| MemberName == StartCornerPropertyName
+		|| MemberName == StartDirectionPropertyName)
+	{
+		MarkModifierDirty();
+	}
+}
+#endif
+
+void UAvaGridArrangeModifier::SetCount(const FIntPoint& NewCount)
+{
+	if (Count == NewCount)
+	{
+		return;
+	}
+	
+	Count = NewCount;
+	MarkModifierDirty();
+}
+
+void UAvaGridArrangeModifier::SetSpread(const FVector2D& NewSpread)
+{
+	if (Spread == NewSpread)
+	{
+		return;
+	}
+	
+	Spread = NewSpread;
+	MarkModifierDirty();
+}
+
+void UAvaGridArrangeModifier::SetStartCorner(EAvaCorner2D InCorner)
+{
+	if (StartCorner == InCorner)
+	{
+		return;
+	}
+	
+	StartCorner = InCorner;
+	MarkModifierDirty();
+}
+
+void UAvaGridArrangeModifier::SetStartDirection(EAvaGridArrangeDirection InDirection)
+{
+	if (StartDirection == InDirection)
+	{
+		return;
+	}
+
+	StartDirection = InDirection;
+	MarkModifierDirty();
+}
+
+void UAvaGridArrangeModifier::GetArrangementLayout(EAvaCorner2D& OutStartCorner, EAvaGridArrangeDirection& OutStartDirection) const
+{
+	OutStartCorner = StartCorner;
+	OutStartDirection = StartDirection;
+}
+
+void UAvaGridArrangeModifier::SetArrangementLayout(const EAvaCorner2D NewStartCorner, const EAvaGridArrangeDirection NewStartDirection)
+{
+	StartCorner = NewStartCorner;
+	StartDirection = NewStartDirection;
+
+	MarkModifierDirty();
+}
+
+void UAvaGridArrangeModifier::OnModifierCDOSetup(FActorModifierCoreMetadata& InMetadata)
+{
+	Super::OnModifierCDOSetup(InMetadata);
+
+	InMetadata.SetName(TEXT("GridArrange"));
+	InMetadata.SetCategory(TEXT("Layout"));
+#if WITH_EDITOR
+	InMetadata.SetDescription(LOCTEXT("ModifierDescription", "Positions child actors in a 2D grid format"));
+#endif
+}
+
+void UAvaGridArrangeModifier::Apply()
+{
+	AActor* const ModifyActor = GetModifiedActor();
+
+	// Early exit if the modify actor is NOT being isolated. The outliner will manage the visibility for the actor and it's children.
+	if (FAvaModifiersActorUtils::IsActorNotIsolated(ModifyActor))
+	{
+		Next();
+		return;
+	}
+
+	const FAvaSceneTreeUpdateModifierExtension* SceneExtension = GetExtension<FAvaSceneTreeUpdateModifierExtension>();
+	if (!SceneExtension)
+	{
+		Fail(LOCTEXT("InvalidSceneExtension", "Scene extension could not be found"));
+		return;
+	}
+	
+	const TArray<TWeakObjectPtr<AActor>> AttachedActors = SceneExtension->GetDirectChildrenActor(ModifyActor);
+	const int32 TotalSlotCount = Count.X * Count.Y;
+	const int32 AttachedActorCount = AttachedActors.Num();
+
+	auto GetGridX = [this](const int32 ChildIndex)
+	{
+		return StartDirection == EAvaGridArrangeDirection::Horizontal
+			? ChildIndex % Count.X
+			: ChildIndex / Count.X;
+	};
+	auto GetGridY = [this](const int32 ChildIndex)
+	{
+		return StartDirection == EAvaGridArrangeDirection::Horizontal
+			? ChildIndex / Count.X
+			: ChildIndex % Count.X;
+	};
+
+	auto GetReversedGridX = [this, GetGridX](const int32 ChildIndex) { return (Count.X - 1) - GetGridX(ChildIndex); };
+	auto GetReversedGridY = [this, GetGridY](const int32 ChildIndex) { return (Count.Y - 1) - GetGridY(ChildIndex); };
+
+	UAvaVisibilityModifierShared* VisibilityShared = GetShared<UAvaVisibilityModifierShared>(true);
+	UAvaTransformModifierShared* LayoutShared = GetShared<UAvaTransformModifierShared>(true);
+
+	TSet<TWeakObjectPtr<AActor>> NewChildrenActorsWeak;
+	for (int32 ChildIndex = 0; ChildIndex < AttachedActorCount; ++ChildIndex)
+	{
+		AActor* AttachedActor = AttachedActors[ChildIndex].Get();
+
+		if (!AttachedActor)
+		{
+			continue;
+		}
+		
+		{
+			// Track all new children actors
+			TArray<AActor*> ChildrenActors { AttachedActor };
+			AttachedActor->GetAttachedActors(ChildrenActors, false, true);
+			for (AActor* ChildActor : ChildrenActors)
+			{
+				NewChildrenActorsWeak.Add(ChildActor);
+			}
+		}
+
+		// No need to handle nested children actor, only direct children, visibility will propagate
+		if (AttachedActor->GetAttachParentActor() != ModifyActor)
+		{
+			continue;
+		}
+		
+		// Track this actor visibility state
+		const bool bIsVisible = ChildIndex < TotalSlotCount;
+		VisibilityShared->SetActorVisibility(this, AttachedActor, !bIsVisible, true);
+
+		FVector RelativeOffset = FVector::ZeroVector;
+		switch (StartCorner)
+		{
+			case EAvaCorner2D::TopLeft:
+				RelativeOffset.Y = GetGridX(ChildIndex);
+				RelativeOffset.Z = GetGridY(ChildIndex);
+			break;
+			case EAvaCorner2D::TopRight:
+				RelativeOffset.Y = -GetReversedGridX(ChildIndex);
+				RelativeOffset.Z = GetGridY(ChildIndex);
+			break;
+			case EAvaCorner2D::BottomLeft:
+				RelativeOffset.Y = GetGridX(ChildIndex);
+				RelativeOffset.Z = -GetReversedGridY(ChildIndex);
+			break;
+			case EAvaCorner2D::BottomRight:
+				RelativeOffset.Y = -GetReversedGridX(ChildIndex);
+				RelativeOffset.Z = -GetReversedGridY(ChildIndex);
+			break;
+		}
+
+		RelativeOffset.Y *= Spread.X;
+		RelativeOffset.Z *= -Spread.Y;
+
+		// Track this actor layout state
+		LayoutShared->SaveActorState(this, AttachedActor);
+		AttachedActor->SetActorRelativeLocation(RelativeOffset);
+	}
+
+	// Untrack previous actors that are not attached anymore
+	const TSet<TWeakObjectPtr<AActor>> UntrackActors = ChildrenActorsWeak.Difference(NewChildrenActorsWeak);
+	LayoutShared->RestoreActorsState(this, UntrackActors);
+	VisibilityShared->RestoreActorsState(this, UntrackActors);
+
+	ChildrenActorsWeak = NewChildrenActorsWeak;
+
+	Next();
+}
+
+#undef LOCTEXT_NAMESPACE
diff --git a/Engine/Plugins/Experimental/Avalanche/Source/AvalancheModifiers/Private/Modifiers/AvaJustifyModifier.cpp b/Engine/Plugins/Experimental/Avalanche/Source/AvalancheModifiers/Private/Modifiers/AvaJustifyModifier.cpp
new file mode 100644
--- /dev/null
+++ b/Engine/Plugins/Experimental/Avalanche/Source/AvalancheModifiers/Private/Modifiers/AvaJustifyModifier.cpp
@@ -0,0 +1,434 @@
+// Copyright Epic Games, Inc. All Rights Reserved.
+
+#include "Modifiers/AvaJustifyModifier.h"
+
+#include "AvaModifiersActorUtils.h"
+#include "Components/ActorComponent.h"
+#include "Components/PrimitiveComponent.h"
+#include "GameFramework/Actor.h"
+#include "Shared/AvaTransformModifierShared.h"
+
+#define LOCTEXT_NAMESPACE "AvaJustifyModifier"
+
+void UAvaJustifyModifier::OnModifierCDOSetup(FActorModifierCoreMetadata& InMetadata)
+{
+	Super::OnModifierCDOSetup(InMetadata);
+
+	InMetadata.AllowTick(true);
+	InMetadata.SetName(TEXT("Justify"));
+	InMetadata.SetCategory(TEXT("Layout"));
+#if WITH_EDITOR
+	InMetadata.SetDescription(LOCTEXT("ModifierDescription", "Aligns child actors according to the specified justification option, based on their bounding boxes"));
+#endif
+}
+
+bool UAvaJustifyModifier::IsModifierDirtyable() const
+{
+	const AActor* ActorModified = GetModifiedActor();
+	if (!IsValid(ActorModified))
+	{
+		return Super::IsModifierDirtyable();
+	}
+
+	const FBox TrackedActorLocalBounds = FAvaModifiersActorUtils::GetActorsBounds(ChildrenActorsWeak, FTransform::Identity, true);
+
+	if (TrackedActorLocalBounds.Equals(CachedTrackedBounds, 0.01))
+	{
+		return Super::IsModifierDirtyable();
+	}
+	
+	return true;
+}
+
+void UAvaJustifyModifier::OnModifiedActorTransformed()
+{
+	// Do nothing if we move, only justify children
+}
+
+void UAvaJustifyModifier::Apply()
+{
+	const AActor* const ActorModified = GetModifiedActor();
+
+	// Gather newly attached children and save current state
+	TSet<TWeakObjectPtr<AActor>> NewChildrenActorsWeak;
+	GetChildrenActors(NewChildrenActorsWeak);
+
+	// the children bounds need to be aligned to the modified actor
+	CachedTrackedBounds = FAvaModifiersActorUtils::GetActorsBounds(NewChildrenActorsWeak, FTransform::Identity, true);
+	
+	FVector BoundsCenter;
+	FVector BoundsExtent;
+	CachedTrackedBounds.GetCenterAndExtents(BoundsCenter, BoundsExtent);
+
+	const FVector AlignmentOffset = GetAlignmentOffset(BoundsExtent);
+	const FVector AnchorOffsetVector = GetAnchorOffset();
+	
+	// used to store the offset needed to constrain or un-constrain the justification axis, the modified actor and its children movements
+	// constraints start set to the bounds center vector, and are changed below based on different combinations
+	const FVector ConstraintVector = GetConstraintVector(BoundsCenter, ActorModified->GetActorLocation());
+	const FVector ChildLocationOffset = ConstraintVector + AlignmentOffset - AnchorOffsetVector;
+
+	UAvaTransformModifierShared* LayoutShared = GetShared<UAvaTransformModifierShared>(true);
+
+	// unregister children transform update callbacks, we don't want them firing while the modifier is updating
+	if (FAvaTransformUpdateModifierExtension* TransformExtension = GetExtension<FAvaTransformUpdateModifierExtension>())
+	{
+		TransformExtension->UntrackActors(ChildrenActorsWeak);
+	}
+
+	// Lets track children actors visibility to receive callback
+	if (FAvaRenderStateUpdateModifierExtension* RenderExtension = GetExtension<FAvaRenderStateUpdateModifierExtension>())
+	{
+		RenderExtension->SetTrackedActorsVisibility(NewChildrenActorsWeak);
+	}
+	
+	// Update child actor position
+	for (TWeakObjectPtr<AActor>& ChildActorWeak : NewChildrenActorsWeak)
+	{
+		if (AActor* const Child = ChildActorWeak.Get())
+		{
+			if (Child->GetAttachParentActor() != ActorModified)
+			{
+				continue;
+			}
+
+			LayoutShared->SaveActorState(this, Child);
+			const FVector ChildRelativeLocation = Child->GetRootComponent()->GetRelativeLocation();
+			const FVector DesiredChildLocation = ChildRelativeLocation - ChildLocationOffset;
+			Child->SetActorRelativeLocation(DesiredChildLocation);
+		}
+	}
+		
+	// Untrack previous actors that are not attached anymore
+	LayoutShared->RestoreActorsState(this, ChildrenActorsWeak.Difference(NewChildrenActorsWeak));
+
+	ChildrenActorsWeak = NewChildrenActorsWeak;
+	
+	if (FAvaTransformUpdateModifierExtension* TransformExtension = GetExtension<FAvaTransformUpdateModifierExtension>())
+	{
+		TransformExtension->TrackActors(ChildrenActorsWeak, true);
+	}
+
+	Next();
+}
+
+#if WITH_EDITOR
+void UAvaJustifyModifier::PostEditChangeProperty(FPropertyChangedEvent& InPropertyChangedEvent)
+{
+	Super::PostEditChangeProperty(InPropertyChangedEvent);
+	
+	const FName PropertyName = InPropertyChangedEvent.GetMemberPropertyName();
+	
+	static const FName HorizontalAlignmentPropertyName = GET_MEMBER_NAME_CHECKED(UAvaJustifyModifier, HorizontalAlignment);
+	static const FName VerticalAlignmentPropertyName = GET_MEMBER_NAME_CHECKED(UAvaJustifyModifier, VerticalAlignment);
+	static const FName DepthAlignmentPropertyName = GET_MEMBER_NAME_CHECKED(UAvaJustifyModifier, DepthAlignment);
+	static const FName HorizontalAnchorPropertyName = GET_MEMBER_NAME_CHECKED(UAvaJustifyModifier, HorizontalAnchor);
+	static const FName VerticalAnchorPropertyName = GET_MEMBER_NAME_CHECKED(UAvaJustifyModifier, VerticalAnchor);
+	static const FName DepthAnchorPropertyName = GET_MEMBER_NAME_CHECKED(UAvaJustifyModifier, DepthAnchor);
+	
+	if (PropertyName == HorizontalAlignmentPropertyName ||
+		PropertyName == VerticalAlignmentPropertyName ||
+		PropertyName == DepthAlignmentPropertyName)
+	{
+		MarkModifierDirty();
+	}
+}
+#endif
+
+void UAvaJustifyModifier::SetHorizontalAlignment(EAvaJustifyHorizontal InHorizontalAlignment)
+{
+	HorizontalAlignment = InHorizontalAlignment;
+
+	MarkModifierDirty();
+}
+
+void UAvaJustifyModifier::SetVerticalAlignment(EAvaJustifyVertical InVerticalAlignment)
+{
+	VerticalAlignment = InVerticalAlignment;
+
+	MarkModifierDirty();
+}
+
+void UAvaJustifyModifier::SetDepthAlignment(EAvaJustifyDepth InDepthAlignment)
+{
+	DepthAlignment = InDepthAlignment;
+
+	MarkModifierDirty();
+}
+
+void UAvaJustifyModifier::SetHorizontalAnchor(const float InHorizontalAnchor)
+{
+	HorizontalAnchor = InHorizontalAnchor;
+
+	MarkModifierDirty();
+}
+
+void UAvaJustifyModifier::SetVerticalAnchor(const float InVerticalAnchor)
+{
+	VerticalAnchor = InVerticalAnchor;
+
+	MarkModifierDirty();
+}
+
+void UAvaJustifyModifier::SetDepthAnchor(const float InDepthAnchor)
+{
+	DepthAnchor = InDepthAnchor;
+
+	MarkModifierDirty();
+}
+
+void UAvaJustifyModifier::OnSceneTreeTrackedActorDirectChildrenChanged(int32 InIdx,
+	const TArray<TWeakObjectPtr<AActor>>& InPreviousChildrenActors,
+	const TArray<TWeakObjectPtr<AActor>>& InNewChildrenActors)
+{
+	// Overwrite this from parent, don't do anything here, not needed for justify
+}
+
+void UAvaJustifyModifier::OnRenderStateUpdated(AActor* InActor, UActorComponent* InComponent)
+{
+	Super::OnRenderStateUpdated(InActor, InComponent);
+
+	const AActor* ActorModified = GetModifiedActor();
+	if (!IsValid(ActorModified))
+	{
+		return;
+	}
+
+	// Only handle direct child of modified actor
+	if (InActor->GetAttachParentActor() != ActorModified)
+	{
+		return;
+	}
+
+	// Is the modifier dirtyable
+	const bool bModifierDirtyable = IsModifierDirtyable();
+	
+	if (bModifierDirtyable)
+	{
+		MarkModifierDirty();
+	}
+}
+
+void UAvaJustifyModifier::OnActorVisibilityChanged(AActor* InActor)
+{
+	Super::OnActorVisibilityChanged(InActor);
+
+	const AActor* ActorModified = GetModifiedActor();
+	if (!IsValid(ActorModified))
+	{
+		return;
+	}
+
+	if (!InActor->IsAttachedTo(ActorModified))
+	{
+		return;
+	}
+	
+	MarkModifierDirty();
+}
+
+void UAvaJustifyModifier::OnTransformUpdated(AActor* InActor, bool bInParentMoved)
+{
+	Super::OnTransformUpdated(InActor, bInParentMoved);
+
+	const AActor* ActorModified = GetModifiedActor();
+	if (!ActorModified || !InActor)
+	{
+		return;
+	}
+
+	if (!InActor->IsAttachedTo(ActorModified))
+	{
+		return;
+	}
+	
+	if (bInParentMoved)
+	{
+		return;
+	}
+	
+	// if there's at least one justification constraint active, we need to rearrange all children based on their position update: let's update the modifier
+	if (HasHorizontalAlignment() || HasVerticalAlignment() || HasDepthAlignment())
+	{
+		MarkModifierDirty();
+	}
+}
+
+FVector UAvaJustifyModifier::MakeConstrainedAxisVector() const
+{
+	/* Generating a vector, the axes of which are:
+	 * 0.0 --> the axis is unconstrained
+	 * 1.0 --> the axis is constrained
+	 */
+
+	return FVector(static_cast<float>(HasDepthAlignment()),
+		static_cast<float>(HasHorizontalAlignment()),
+		static_cast<float>(HasVerticalAlignment()));
+}
+
+void UAvaJustifyModifier::GetChildrenActors(TSet<TWeakObjectPtr<AActor>>& OutChildren) const
+{
+	ForEachActor<AActor>([this, &OutChildren](AActor* InActor)->bool
+	{
+		if (IsValid(InActor))
+		{
+			OutChildren.Add(InActor);
+		}
+		return true;
+	}, EActorModifierCoreLookup::AllChildren);
+}
+
+void UAvaJustifyModifier::GetTrackedActors(const TSet<TWeakObjectPtr<AActor>>& InChildrenActors, TArray<TWeakObjectPtr<const AActor>>& OutTrackedActors) const
+{
+	for (const TWeakObjectPtr<AActor>& ChildActorWeak : InChildrenActors)
+	{
+		const AActor* ChildActor = ChildActorWeak.Get();
+		if (!IsValid(ChildActor))
+		{
+			continue;
+		}
+
+		// Only track visible actors, skip collapsed one
+		if (!FAvaModifiersActorUtils::IsActorVisible(ChildActor))
+		{
+			continue;
+		}
+		
+		OutTrackedActors.Add(ChildActor);
+	}
+}
+
+FVector UAvaJustifyModifier::GetConstraintVector(const FVector& InBoundsCenter, const FVector& InModifiedActorPosition) const
+{
+	const FVector BoundsCenterToModifiedActor = InBoundsCenter - InModifiedActorPosition;
+	const FVector ConstrainedAxisVector = MakeConstrainedAxisVector();
+
+	// The ConstrainedAxisVector is used to filter out unconstrained axes from the BoundsCenterToModifiedActor vector
+	const FVector ConstraintVector = ConstrainedAxisVector * BoundsCenterToModifiedActor;
+	
+	return ConstraintVector;
+}
+
+bool UAvaJustifyModifier::HasDepthAlignment() const
+{
+	return DepthAlignment != EAvaJustifyDepth::None;
+}
+
+bool UAvaJustifyModifier::HasHorizontalAlignment() const
+{
+	return HorizontalAlignment != EAvaJustifyHorizontal::None;
+}
+
+bool UAvaJustifyModifier::HasVerticalAlignment() const
+{
+	return VerticalAlignment != EAvaJustifyVertical::None;
+}
+
+FVector UAvaJustifyModifier::GetAnchorOffset() const
+{
+	// used to add a custom anchor offset, independent from bounds size
+	// note: some axes will be zeroed below, if their respective alignment is Off
+	FVector AnchorOffsetVector = FVector(DepthAnchor, HorizontalAnchor, VerticalAnchor);
+
+	// setup depth alignment and anchor offsets
+	switch (DepthAlignment)
+	{
+		case EAvaJustifyDepth::None:
+			AnchorOffsetVector.X = 0.0f; // ignore depth anchor reference
+			break;
+		default:
+			break;
+	}
+	
+	// setup horizontal alignment and anchor offsets
+	switch (HorizontalAlignment)
+	{				
+		case EAvaJustifyHorizontal::None:
+			AnchorOffsetVector.Y = 0.0f; // ignore horizontal anchor reference
+			break;
+		default:
+			break;
+	}
+
+	// setup vertical alignment and anchor offsets
+	switch (VerticalAlignment)
+	{
+		case EAvaJustifyVertical::None:
+			AnchorOffsetVector.Z = 0.0f; // ignore vertical anchor reference
+			break;
+		default:
+			break;
+	}
+
+	return AnchorOffsetVector;
+}
+
+FVector UAvaJustifyModifier::GetAlignmentOffset(const FVector& InExtent) const
+{
+	// will be filled with justification values, based on bounds size
+	FVector AlignmentOffset = FVector::ZeroVector;
+
+	// setup depth alignment and anchor offsets
+	switch (DepthAlignment)
+	{
+		case EAvaJustifyDepth::None:
+			AlignmentOffset.X = 0.0f;
+			break;
+				
+		case EAvaJustifyDepth::Front:
+			AlignmentOffset.X = InExtent.X;
+			break;
+				
+		case EAvaJustifyDepth::Center:
+			AlignmentOffset.X = 0.0f;
+			break;
+				
+		case EAvaJustifyDepth::Back:
+			AlignmentOffset.X = -InExtent.X;
+			break;
+	}
+	
+	// setup horizontal alignment and anchor offsets
+	switch (HorizontalAlignment)
+	{				
+		case EAvaJustifyHorizontal::None:
+			AlignmentOffset.Y = 0.0f;
+			break;
+			
+		case EAvaJustifyHorizontal::Left:
+			AlignmentOffset.Y = -InExtent.Y;
+			break;
+				
+		case EAvaJustifyHorizontal::Center:
+			AlignmentOffset.Y = 0.0f;
+			break;
+				
+		case EAvaJustifyHorizontal::Right:
+			AlignmentOffset.Y = InExtent.Y;
+			break;
+	}
+
+	// setup vertical alignment and anchor offsets
+	switch (VerticalAlignment)
+	{
+		case EAvaJustifyVertical::None:
+			AlignmentOffset.Z = 0.0f;
+			break;
+				
+		case EAvaJustifyVertical::Top:
+			AlignmentOffset.Z = InExtent.Z;
+			break;
+			
+		case EAvaJustifyVertical::Center:
+			AlignmentOffset.Z = 0.0f;
+			break;
+			
+		case EAvaJustifyVertical::Bottom:
+			AlignmentOffset.Z = -InExtent.Z;
+			break;
+	}
+	
+	return AlignmentOffset;
+}
+
+#undef LOCTEXT_NAMESPACE
diff --git a/Engine/Plugins/Experimental/Avalanche/Source/AvalancheModifiers/Private/Modifiers/AvaLookAtModifier.cpp b/Engine/Plugins/Experimental/Avalanche/Source/AvalancheModifiers/Private/Modifiers/AvaLookAtModifier.cpp
new file mode 100644
--- /dev/null
+++ b/Engine/Plugins/Experimental/Avalanche/Source/AvalancheModifiers/Private/Modifiers/AvaLookAtModifier.cpp
@@ -0,0 +1,229 @@
+// Copyright Epic Games, Inc. All Rights Reserved.
+
+#include "Modifiers/AvaLookAtModifier.h"
+
+#include "AvaDefs.h"
+#include "Extensions/AvaTransformUpdateModifierExtension.h"
+#include "GameFramework/Actor.h"
+#include "Misc/ITransaction.h"
+#include "Shared/AvaTransformModifierShared.h"
+
+#define LOCTEXT_NAMESPACE "AvaLookAtModifier"
+
+void UAvaLookAtModifier::OnModifierAdded(EActorModifierCoreEnableReason InReason)
+{
+	Super::OnModifierAdded(InReason);
+
+	AddExtension<FAvaTransformUpdateModifierExtension>(this);
+
+	if (FAvaSceneTreeUpdateModifierExtension* SceneExtension = GetExtension<FAvaSceneTreeUpdateModifierExtension>())
+	{
+		SceneExtension->TrackSceneTree(0, &ReferenceActor);
+	}
+}
+
+void UAvaLookAtModifier::OnModifierEnabled(EActorModifierCoreEnableReason InReason)
+{
+	Super::OnModifierEnabled(InReason);
+	
+	if (UAvaTransformModifierShared* LayoutShared = GetShared<UAvaTransformModifierShared>(true))
+	{
+		LayoutShared->SaveActorState(this, GetModifiedActor());
+	}
+}
+
+void UAvaLookAtModifier::OnModifierDisabled(EActorModifierCoreDisableReason InReason)
+{
+	Super::OnModifierDisabled(InReason);
+
+	if (UAvaTransformModifierShared* LayoutShared = GetShared<UAvaTransformModifierShared>(false))
+	{
+		LayoutShared->RestoreActorState(this, GetModifiedActor());
+	}
+}
+
+void UAvaLookAtModifier::Apply()
+{
+	AActor* const ModifyActor = GetModifiedActor();
+
+	const AActor* LookActor = ReferenceActor.ReferenceActorWeak.Get();
+	if (!IsValid(LookActor))
+	{
+		Next();
+		return;
+	}
+
+	auto LookAtDirection = [this, ModifyActor, LookActor]() -> FVector
+	{
+		FVector OutDirection = bFlipAxis
+			? LookActor->GetActorLocation() - ModifyActor->GetActorLocation()
+			: ModifyActor->GetActorLocation() - LookActor->GetActorLocation();
+
+		OutDirection.Normalize();
+		if (OutDirection.SizeSquared() < 0.5f)
+		{
+			// Assert possible if OutDirection is not normalized.
+			OutDirection = FVector(1, 0, 0);
+		}
+
+		return OutDirection;
+	};
+
+	auto NormalOffset = [this, ModifyActor](const FQuat& InNewRotation) -> FQuat
+	{
+		FQuat OutNormal = bFlipAxis
+			? InNewRotation - ModifyActor->GetActorRotation().Quaternion()
+			: ModifyActor->GetActorRotation().Quaternion() - InNewRotation;
+
+		OutNormal.Normalize();
+
+		return OutNormal;
+	};
+
+	FMatrix NewRotation = FRotationMatrix::Identity;
+		
+	switch (Axis)
+	{
+		case EAvaAxis::Horizontal:
+			NewRotation = FRotationMatrix::MakeFromX(LookAtDirection());
+			break;
+		case EAvaAxis::Vertical:
+			NewRotation = FRotationMatrix::MakeFromY(LookAtDirection());
+			break;
+		case EAvaAxis::Depth:
+			NewRotation = FRotationMatrix::MakeFromZ(LookAtDirection());
+			break;
+	}
+
+	ModifyActor->SetActorRotation(NewRotation.ToQuat());
+
+	Next();
+}
+
+void UAvaLookAtModifier::OnModifiedActorTransformed()
+{
+	MarkModifierDirty();
+}
+
+void UAvaLookAtModifier::PostLoad()
+{
+	if (!bDeprecatedPropertiesMigrated)
+	{
+		ReferenceActor.ReferenceContainer = EAvaReferenceContainer::Other;
+		ReferenceActor.ReferenceActorWeak = ReferenceActorWeak;
+		ReferenceActor.bSkipHiddenActors = false;
+
+		bDeprecatedPropertiesMigrated = true;
+	}
+	
+	Super::PostLoad();
+}
+
+#if WITH_EDITOR
+void UAvaLookAtModifier::PostEditChangeProperty(FPropertyChangedEvent& PropertyChangedEvent)
+{
+	Super::PostEditChangeProperty(PropertyChangedEvent);
+
+	const FName MemberName = PropertyChangedEvent.GetMemberPropertyName();
+	
+	static const FName ReferenceActorPropertyName = GET_MEMBER_NAME_CHECKED(UAvaLookAtModifier, ReferenceActor);
+	static const FName FlipAxisPropertyName = GET_MEMBER_NAME_CHECKED(UAvaLookAtModifier, bFlipAxis);
+	static const FName AxisPropertyName = GET_MEMBER_NAME_CHECKED(UAvaLookAtModifier, Axis);
+	
+	if (MemberName == ReferenceActorPropertyName)
+	{
+		OnReferenceActorChanged();
+	}
+	else if (MemberName == AxisPropertyName
+		|| MemberName == FlipAxisPropertyName)
+	{
+		MarkModifierDirty();
+	}
+}
+
+void UAvaLookAtModifier::PostTransacted(const FTransactionObjectEvent& TransactionEvent)
+{
+	OnReferenceActorChanged();
+	
+	Super::PostTransacted(TransactionEvent);
+}
+#endif // WITH_EDITOR
+
+void UAvaLookAtModifier::SetReferenceActor(const FAvaSceneTreeActor& InReferenceActor)
+{
+	if (ReferenceActor == InReferenceActor)
+	{
+		return;
+	}
+	
+	ReferenceActor = InReferenceActor;
+	OnReferenceActorChanged();
+}
+
+void UAvaLookAtModifier::OnReferenceActorChanged()
+{
+	if (ReferenceActor.ReferenceActorWeak.Get() == GetModifiedActor())
+	{
+		ReferenceActor.ReferenceActorWeak = nullptr;
+	}
+
+	if (const FAvaSceneTreeUpdateModifierExtension* SceneExtension = GetExtension<FAvaSceneTreeUpdateModifierExtension>())
+	{
+		SceneExtension->CheckTrackedActorUpdate(0);
+	}
+}
+
+void UAvaLookAtModifier::SetAxis(const EAvaAxis NewAxis)
+{
+	Axis = NewAxis;
+
+	MarkModifierDirty();
+}
+
+void UAvaLookAtModifier::SetFlipAxis(const bool bNewFlipAxis)
+{
+	bFlipAxis = bNewFlipAxis;
+
+	MarkModifierDirty();
+}
+
+void UAvaLookAtModifier::OnModifierCDOSetup(FActorModifierCoreMetadata& InMetadata)
+{
+	Super::OnModifierCDOSetup(InMetadata);
+
+	InMetadata.SetName(TEXT("LookAt"));
+	InMetadata.SetCategory(TEXT("Layout"));
+#if WITH_EDITOR
+	InMetadata.SetDescription(LOCTEXT("ModifierDescription", "Rotates an actor to face another actor"));
+#endif
+}
+
+void UAvaLookAtModifier::OnTransformUpdated(AActor* InActor, bool bInParentMoved)
+{
+	if (InActor && InActor == ReferenceActor.ReferenceActorWeak.Get())
+	{
+		MarkModifierDirty();
+	}
+}
+
+void UAvaLookAtModifier::OnSceneTreeTrackedActorChanged(int32 InIdx, AActor* InPreviousActor, AActor* InNewActor)
+{
+	Super::OnSceneTreeTrackedActorChanged(InIdx, InPreviousActor, InNewActor);
+	
+	if (InNewActor == GetModifiedActor())
+	{
+		OnReferenceActorChanged();
+		return;
+	}
+
+	// Untrack reference actor
+	if (FAvaTransformUpdateModifierExtension* TransformExtension = GetExtension<FAvaTransformUpdateModifierExtension>())
+	{
+		TransformExtension->UntrackActor(InPreviousActor);
+		TransformExtension->TrackActor(InNewActor, true);
+	}
+
+	MarkModifierDirty();
+}
+
+#undef LOCTEXT_NAMESPACE
diff --git a/Engine/Plugins/Experimental/Avalanche/Source/AvalancheModifiers/Private/Modifiers/AvaPatternModifier.cpp b/Engine/Plugins/Experimental/Avalanche/Source/AvalancheModifiers/Private/Modifiers/AvaPatternModifier.cpp
new file mode 100644
--- /dev/null
+++ b/Engine/Plugins/Experimental/Avalanche/Source/AvalancheModifiers/Private/Modifiers/AvaPatternModifier.cpp
@@ -0,0 +1,402 @@
+// Copyright Epic Games, Inc. All Rights Reserved.
+
+#include "Modifiers/AvaPatternModifier.h"
+
+#include "Async/Async.h"
+#include "Components/DynamicMeshComponent.h"
+#include "DynamicMesh/MeshTransforms.h"
+#include "DynamicMeshEditor.h"
+
+#define LOCTEXT_NAMESPACE "AvaPatternModifier"
+
+void UAvaPatternModifier::OnModifierCDOSetup(FActorModifierCoreMetadata& InMetadata)
+{
+	Super::OnModifierCDOSetup(InMetadata);
+
+	InMetadata.SetName(TEXT("Pattern"));
+	InMetadata.SetCategory(TEXT("Geometry"));
+#if WITH_EDITOR
+	InMetadata.SetDescription(LOCTEXT("ModifierDescription", "Repeats a geometry multiple times following a specific layout pattern"));
+#endif
+}
+
+void UAvaPatternModifier::Apply()
+{
+	if ((Layout == EAvaPatternModifierLayout::Line && LineLayoutOptions.RepeatCount <= 0)
+		|| (Layout == EAvaPatternModifierLayout::Grid && (GridLayoutOptions.RepeatCount.X <= 0 ||  GridLayoutOptions.RepeatCount.Y <= 0))
+		|| (Layout == EAvaPatternModifierLayout::Circle && CircleLayoutOptions.RepeatCount <= 0))
+	{
+		Fail(LOCTEXT("InvalidRepeatCount", "Layout repeat count must be greater than 0"));
+		return;
+	}
+
+	using namespace UE::Geometry;
+
+	UDynamicMeshComponent* TargetMeshComponent = GetMeshComponent();
+
+	if (!IsValid(TargetMeshComponent))
+	{
+		Fail(LOCTEXT("InvalidDynamicMeshComponent", "Invalid dynamic mesh component on modified actor"));
+		return;
+	}
+
+	TargetMeshComponent->GetDynamicMesh()->EditMesh([this](FDynamicMesh3& AppendToMesh)
+	{
+		// copy original mesh once into tmp mesh
+		FDynamicMesh3 TmpMesh = AppendToMesh;
+
+		// Get the original bounds before clearing the mesh
+		OriginalMeshBounds = GetMeshBounds();
+
+		// clear all since we have a copy
+		for (int32 TId : AppendToMesh.TriangleIndicesItr())
+		{
+			AppendToMesh.RemoveTriangle(TId);
+		}
+
+		FVector CenterAxis = FVector::ZeroVector;
+
+		if (Layout == EAvaPatternModifierLayout::Line)
+		{
+			// edit original mesh to append tmp mesh
+			FMeshIndexMappings TmpMappings;
+			FDynamicMeshEditor Editor(&AppendToMesh);
+			const FTransformSRT3d TransformChange = GetLineLayoutTransformChange();
+			if (!LineLayoutOptions.bAccumulateTransform)
+			{
+				FTransform RotationAndScaleTransform;
+				RotationAndScaleTransform.SetLocation(FVector::ZeroVector);
+				RotationAndScaleTransform.SetRotation(TransformChange.GetRotator().Quaternion());
+				RotationAndScaleTransform.SetScale3D(TransformChange.GetScale());
+				FTransformSRT3d RotateScale(RotationAndScaleTransform);
+				MeshTransforms::ApplyTransform(TmpMesh, RotateScale, true);
+			}
+			for (int32 Idx = 0; Idx < LineLayoutOptions.RepeatCount; Idx++)
+			{
+				if (Idx != 0)
+				{
+					if (LineLayoutOptions.bAccumulateTransform)
+					{
+						MeshTransforms::ApplyTransform(TmpMesh, TransformChange, true);
+					}
+					else
+					{
+						MeshTransforms::Translate(TmpMesh, TransformChange.GetTranslation());
+					}
+				}
+				Editor.AppendMesh(&TmpMesh, TmpMappings);
+				TmpMappings.Reset();
+			}
+			if (LineLayoutOptions.bCentered)
+			{
+				if (LineLayoutOptions.Axis == EAvaPatternModifierAxis::X)
+				{
+					CenterAxis.X = 1;
+				}
+				else if (LineLayoutOptions.Axis == EAvaPatternModifierAxis::Y)
+				{
+					CenterAxis.Y = 1;
+				}
+				else if (LineLayoutOptions.Axis == EAvaPatternModifierAxis::Z)
+				{
+					CenterAxis.Z = 1;
+				}
+			}
+		}
+		else if (Layout == EAvaPatternModifierLayout::Grid)
+		{
+			// edit original mesh to append tmp mesh
+			FMeshIndexMappings TmpMappings;
+			FDynamicMeshEditor Editor(&AppendToMesh);
+			const FTransformSRT3d ColTransformChange = GetGridLayoutColTransformChange();
+			const FTransformSRT3d RowTransformChange = GetGridLayoutRowTransformChange();
+			if (!GridLayoutOptions.bAccumulateTransform)
+			{
+				FTransform RotationAndScaleTransform;
+				RotationAndScaleTransform.SetLocation(FVector::ZeroVector);
+				RotationAndScaleTransform.SetRotation(ColTransformChange.GetRotator().Quaternion());
+				RotationAndScaleTransform.SetScale3D(ColTransformChange.GetScale());
+				FTransformSRT3d RotateScale(RotationAndScaleTransform);
+				MeshTransforms::ApplyTransform(TmpMesh, RotateScale, true);
+			}
+			for (int32 Row = 0; Row < GridLayoutOptions.RepeatCount.X; Row++)
+			{
+				for (int32 Col = 0; Col < GridLayoutOptions.RepeatCount.Y; Col++)
+				{
+					if (!(Row == 0 && Col == 0))
+					{
+						if (GridLayoutOptions.bAccumulateTransform)
+						{
+							MeshTransforms::ApplyTransform(TmpMesh, ColTransformChange, true);
+						}
+						else
+						{
+							MeshTransforms::Translate(TmpMesh, ColTransformChange.GetTranslation());
+						}
+					}
+					Editor.AppendMesh(&TmpMesh, TmpMappings);
+					TmpMappings.Reset();
+				}
+				MeshTransforms::Translate(TmpMesh, RowTransformChange.GetTranslation() - ColTransformChange.GetTranslation() * GridLayoutOptions.RepeatCount.Y);
+			}
+			if (GridLayoutOptions.bCentered)
+			{
+				if (GridLayoutOptions.Plane == EAvaPatternModifierPlane::XY)
+				{
+					CenterAxis.X = 1;
+					CenterAxis.Y = 1;
+				}
+				else if (GridLayoutOptions.Plane == EAvaPatternModifierPlane::YZ)
+				{
+					CenterAxis.Y = 1;
+					CenterAxis.Z = 1;
+				}
+				else if (GridLayoutOptions.Plane == EAvaPatternModifierPlane::ZX)
+				{
+					CenterAxis.X = 1;
+					CenterAxis.Z = 1;
+				}
+			}
+		}
+		else if (Layout == EAvaPatternModifierLayout::Circle)
+		{
+			// edit original mesh to append tmp mesh
+			FMeshIndexMappings TmpMappings;
+			FDynamicMeshEditor Editor(&AppendToMesh);
+			if (!CircleLayoutOptions.bAccumulateTransform)
+			{
+				FTransform RotationAndScaleTransform;
+				RotationAndScaleTransform.SetLocation(FVector::ZeroVector);
+				RotationAndScaleTransform.SetRotation(CircleLayoutOptions.Rotation.Quaternion());
+				RotationAndScaleTransform.SetScale3D(CircleLayoutOptions.Scale);
+				FTransformSRT3d RotateScale(RotationAndScaleTransform);
+				MeshTransforms::ApplyTransform(TmpMesh, RotateScale, true);
+			}
+			for (int32 Idx = 0; Idx < CircleLayoutOptions.RepeatCount; Idx++)
+			{
+				FDynamicMesh3 AppendMesh = TmpMesh;
+				FTransformSRT3d TransformChange = GetCircleLayoutTransformChange(Idx);
+				const FVector Translation = TransformChange.GetTranslation();
+
+				if (CircleLayoutOptions.bAccumulateTransform)
+				{
+					TransformChange.SetTranslation(FVector::ZeroVector);
+					MeshTransforms::ApplyTransform(TmpMesh, TransformChange, true);
+				}
+
+				MeshTransforms::Translate(AppendMesh, Translation);
+				Editor.AppendMesh(&AppendMesh, TmpMappings);
+				TmpMappings.Reset();
+			}
+		}
+		// center this mesh
+		if (!CenterAxis.IsZero())
+		{
+			const FBox BoundingBox = static_cast<FBox>(AppendToMesh.GetBounds(true));
+			if (BoundingBox.IsValid)
+			{
+				const FVector BoundingCenter = BoundingBox.GetCenter();
+				// center only on specified axis
+				const FVector Translate = BoundingCenter * CenterAxis;
+				MeshTransforms::Translate(AppendToMesh, -Translate);
+			}
+		}
+	}, EDynamicMeshChangeType::GeneralEdit, EDynamicMeshAttributeChangeFlags::Unknown);
+
+	Next();
+}
+
+#if WITH_EDITOR
+void UAvaPatternModifier::PostEditChangeProperty(FPropertyChangedEvent& PropertyChangedEvent)
+{
+	Super::PostEditChangeProperty(PropertyChangedEvent);
+
+	const FName MemberName = PropertyChangedEvent.GetMemberPropertyName();
+
+	static const FName LayoutName = GET_MEMBER_NAME_CHECKED(UAvaPatternModifier, Layout);
+	static const FName LineLayoutOptionsName = GET_MEMBER_NAME_CHECKED(UAvaPatternModifier, LineLayoutOptions);
+	static const FName GridLayoutOptionsName = GET_MEMBER_NAME_CHECKED(UAvaPatternModifier, GridLayoutOptions);
+	static const FName CircleLayoutOptionsName = GET_MEMBER_NAME_CHECKED(UAvaPatternModifier, CircleLayoutOptions);
+
+	if (MemberName == LayoutName)
+	{
+		OnLayoutChanged();
+	}
+	else if (MemberName == LineLayoutOptionsName)
+	{
+		OnLineLayoutOptionsChanged();
+	}
+	else if (MemberName == GridLayoutOptionsName)
+	{
+		OnGridLayoutOptionsChanged();
+	}
+	else if (MemberName == CircleLayoutOptionsName)
+	{
+		OnCircleLayoutOptionsChanged();
+	}
+}
+#endif
+
+void UAvaPatternModifier::SetLayout(EAvaPatternModifierLayout InLayout)
+{
+	if (Layout == InLayout)
+	{
+		return;
+	}
+
+	Layout = InLayout;
+	OnLayoutChanged();
+}
+
+void UAvaPatternModifier::SetLineLayoutOptions(const FAvaPatternModifierLineLayoutOptions& InOptions)
+{
+	LineLayoutOptions = InOptions;
+	OnLineLayoutOptionsChanged();
+}
+
+void UAvaPatternModifier::SetGridLayoutOptions(const FAvaPatternModifierGridLayoutOptions& InOptions)
+{
+	GridLayoutOptions = InOptions;
+	OnGridLayoutOptionsChanged();
+}
+
+void UAvaPatternModifier::SetCircleLayoutOptions(const FAvaPatternModifierCircleLayoutOptions& InOptions)
+{
+	CircleLayoutOptions = InOptions;
+	OnCircleLayoutOptionsChanged();
+}
+
+void UAvaPatternModifier::OnLayoutChanged()
+{
+	MarkModifierDirty();
+}
+
+void UAvaPatternModifier::OnLineLayoutOptionsChanged()
+{
+	if (Layout == EAvaPatternModifierLayout::Line)
+	{
+		MarkModifierDirty();
+	}
+}
+
+void UAvaPatternModifier::OnGridLayoutOptionsChanged()
+{
+	if (Layout == EAvaPatternModifierLayout::Grid)
+	{
+		MarkModifierDirty();
+	}
+}
+
+void UAvaPatternModifier::OnCircleLayoutOptionsChanged()
+{
+	if (Layout == EAvaPatternModifierLayout::Circle)
+	{
+		MarkModifierDirty();
+	}
+}
+
+UE::Geometry::FTransformSRT3d UAvaPatternModifier::GetLineLayoutTransformChange() const
+{
+	const FVector Axis = (LineLayoutOptions.bAxisInverted ? -1 : 1 ) *
+		(LineLayoutOptions.Axis == EAvaPatternModifierAxis::X ? FVector::XAxisVector : (LineLayoutOptions.Axis == EAvaPatternModifierAxis::Y ? FVector::YAxisVector : FVector::ZAxisVector));
+	const FVector Size3D = OriginalMeshBounds.GetSize();
+	const FVector Translation = Axis * Size3D + Axis * LineLayoutOptions.Spacing;
+	FTransform TransformChange;
+	TransformChange.SetRotation(LineLayoutOptions.Rotation.Quaternion());
+	TransformChange.SetLocation(Translation);
+	TransformChange.SetScale3D(LineLayoutOptions.Scale);
+	return UE::Geometry::FTransformSRT3d(TransformChange);
+}
+
+UE::Geometry::FTransformSRT3d UAvaPatternModifier::GetGridLayoutColTransformChange() const
+{
+	FVector ColAxis;
+	if (GridLayoutOptions.Plane == EAvaPatternModifierPlane::XY)
+	{
+		ColAxis = FVector(0, 1, 0);
+	}
+	else if (GridLayoutOptions.Plane == EAvaPatternModifierPlane::YZ)
+	{
+		ColAxis = FVector(0, 0, 1);
+	}
+	else if (GridLayoutOptions.Plane == EAvaPatternModifierPlane::ZX)
+	{
+		ColAxis = FVector(1, 0, 0);
+	}
+	ColAxis *= (GridLayoutOptions.AxisInverted.bY ? -1 : 1);
+	const FVector Size3D = OriginalMeshBounds.GetSize();
+	const FVector Translation = ColAxis * Size3D + ColAxis * GridLayoutOptions.Spacing.Y;
+	FTransform TransformChange;
+	TransformChange.SetRotation(GridLayoutOptions.Rotation.Quaternion());
+	TransformChange.SetLocation(Translation);
+	TransformChange.SetScale3D(GridLayoutOptions.Scale);
+	return UE::Geometry::FTransformSRT3d(TransformChange);
+}
+
+UE::Geometry::FTransformSRT3d UAvaPatternModifier::GetGridLayoutRowTransformChange() const
+{
+	FVector RowAxis;
+	if (GridLayoutOptions.Plane == EAvaPatternModifierPlane::XY)
+	{
+		RowAxis = FVector(1, 0, 0);
+	}
+	else if (GridLayoutOptions.Plane == EAvaPatternModifierPlane::YZ)
+	{
+		RowAxis = FVector(0, 1, 0);
+	}
+	else if (GridLayoutOptions.Plane == EAvaPatternModifierPlane::ZX)
+	{
+		RowAxis = FVector(0, 0, 1);
+	}
+	RowAxis *= (GridLayoutOptions.AxisInverted.bX ? -1 : 1);
+	const FVector Size3D = OriginalMeshBounds.GetSize();
+	const FVector Translation = RowAxis * Size3D + RowAxis * GridLayoutOptions.Spacing.X;
+	FTransform TransformChange;
+	TransformChange.SetRotation(GridLayoutOptions.Rotation.Quaternion());
+	TransformChange.SetLocation(Translation);
+	TransformChange.SetScale3D(GridLayoutOptions.Scale);
+	return UE::Geometry::FTransformSRT3d(TransformChange);
+}
+
+UE::Geometry::FTransformSRT3d UAvaPatternModifier::GetCircleLayoutTransformChange(int32 Idx) const
+{
+	const int32 Count = CircleLayoutOptions.RepeatCount;
+	const float Radius = CircleLayoutOptions.Radius;
+	const float StartAngle = FMath::DegreesToRadians(CircleLayoutOptions.StartAngle);
+	const float FullAngle = FMath::DegreesToRadians(CircleLayoutOptions.FullAngle);
+	const float AngleStep = FullAngle / Count;
+
+	const float ZeroAngle = StartAngle;
+	const float ZeroX = Radius * FMath::Cos(ZeroAngle);
+	const float ZeroY = Radius * FMath::Sin(ZeroAngle);
+
+	const float CurrentAngle = StartAngle + Idx * AngleStep;
+	const float X = Radius * FMath::Cos(CurrentAngle);
+	const float Y = Radius * FMath::Sin(CurrentAngle);
+
+	const bool bCenterPlane = CircleLayoutOptions.bCentered;
+
+	FVector Translation = FVector::ZeroVector;
+	if (CircleLayoutOptions.Plane == EAvaPatternModifierPlane::XY)
+	{
+		const FVector TranslationZero = bCenterPlane ? FVector::ZeroVector : FVector(ZeroX, ZeroY, 0);
+		Translation = FVector(X, Y, 0) - TranslationZero;
+	}
+	else if (CircleLayoutOptions.Plane == EAvaPatternModifierPlane::YZ)
+	{
+		const FVector TranslationZero = bCenterPlane ? FVector::ZeroVector : FVector(0, ZeroX, ZeroY);
+		Translation = FVector(0, X, Y) - TranslationZero;
+	}
+	else if (CircleLayoutOptions.Plane == EAvaPatternModifierPlane::ZX)
+	{
+		const FVector TranslationZero = bCenterPlane ? FVector::ZeroVector : FVector(ZeroX, 0, ZeroY);
+		Translation = FVector(X, 0, Y) - TranslationZero;
+	}
+	FTransform TransformChange;
+	TransformChange.SetLocation(Translation);
+	TransformChange.SetRotation(CircleLayoutOptions.Rotation.Quaternion());
+	TransformChange.SetScale3D(CircleLayoutOptions.Scale);
+	return UE::Geometry::FTransformSRT3d(TransformChange);
+}
+
+#undef LOCTEXT_NAMESPACE
diff --git a/Engine/Plugins/Experimental/Avalanche/Source/AvalancheModifiers/Private/Modifiers/AvaRadialArrangeModifier.cpp b/Engine/Plugins/Experimental/Avalanche/Source/AvalancheModifiers/Private/Modifiers/AvaRadialArrangeModifier.cpp
new file mode 100644
--- /dev/null
+++ b/Engine/Plugins/Experimental/Avalanche/Source/AvalancheModifiers/Private/Modifiers/AvaRadialArrangeModifier.cpp
@@ -0,0 +1,286 @@
+// Copyright Epic Games, Inc. All Rights Reserved.
+
+#include "Modifiers/AvaRadialArrangeModifier.h"
+
+#include "AvaDefs.h"
+#include "AvaModifiersActorUtils.h"
+#include "GameFramework/Actor.h"
+#include "Shared/AvaTransformModifierShared.h"
+#include "Shared/AvaVisibilityModifierShared.h"
+
+#define LOCTEXT_NAMESPACE "AvaRadialArrangeModifier"
+
+#if WITH_EDITOR
+void UAvaRadialArrangeModifier::PostEditChangeProperty(FPropertyChangedEvent& PropertyChangedEvent)
+{
+	Super::PostEditChangeProperty(PropertyChangedEvent);
+
+	const FName MemberName = PropertyChangedEvent.GetMemberPropertyName();
+
+	static const TSet<FName> PropertiesName =
+	{
+		GET_MEMBER_NAME_CHECKED(UAvaRadialArrangeModifier, Count),
+		GET_MEMBER_NAME_CHECKED(UAvaRadialArrangeModifier, Rings),
+		GET_MEMBER_NAME_CHECKED(UAvaRadialArrangeModifier, InnerRadius),
+		GET_MEMBER_NAME_CHECKED(UAvaRadialArrangeModifier, OuterRadius),
+		GET_MEMBER_NAME_CHECKED(UAvaRadialArrangeModifier, StartAngle),
+		GET_MEMBER_NAME_CHECKED(UAvaRadialArrangeModifier, EndAngle),
+		GET_MEMBER_NAME_CHECKED(UAvaRadialArrangeModifier, Arrangement),
+		GET_MEMBER_NAME_CHECKED(UAvaRadialArrangeModifier, bStartFromOuterRadius),
+		GET_MEMBER_NAME_CHECKED(UAvaRadialArrangeModifier, bOrient),
+		GET_MEMBER_NAME_CHECKED(UAvaRadialArrangeModifier, OrientAxis),
+		GET_MEMBER_NAME_CHECKED(UAvaRadialArrangeModifier, bFlipOrient)
+	};
+
+	if (PropertiesName.Contains(MemberName))
+	{
+		MarkModifierDirty();
+	}
+}
+#endif // WITH_EDITOR
+
+void UAvaRadialArrangeModifier::SetCount(const int32 NewCount)
+{
+	Count = NewCount;
+
+	MarkModifierDirty();
+}
+
+void UAvaRadialArrangeModifier::SetRings(const int32 NewRings)
+{
+	Rings = NewRings;
+
+	MarkModifierDirty();
+}
+
+void UAvaRadialArrangeModifier::SetInnerRadius(const float NewInnerRadius)
+{
+	InnerRadius = NewInnerRadius;
+
+	MarkModifierDirty();
+}
+
+void UAvaRadialArrangeModifier::SetOuterRadius(const float NewOuterRadius)
+{
+	OuterRadius = NewOuterRadius;
+
+	MarkModifierDirty();
+}
+
+void UAvaRadialArrangeModifier::SetStartAngle(const float NewStartAngle)
+{
+	StartAngle = NewStartAngle;
+
+	MarkModifierDirty();
+}
+
+void UAvaRadialArrangeModifier::SetEndAngle(const float NewEndAngle)
+{
+	EndAngle = NewEndAngle;
+
+	MarkModifierDirty();
+}
+
+void UAvaRadialArrangeModifier::SetArrangement(const EAvaRadialArrangement NewArrangement)
+{
+	Arrangement = NewArrangement;
+
+	MarkModifierDirty();
+}
+
+void UAvaRadialArrangeModifier::SetStartFromOuterRadius(const bool bMewStartFromOuterRadius)
+{
+	bStartFromOuterRadius = bMewStartFromOuterRadius;
+
+	MarkModifierDirty();
+}
+
+void UAvaRadialArrangeModifier::SetOrient(const bool bNewOrient)
+{
+	bOrient = bNewOrient;
+
+	MarkModifierDirty();
+}
+
+void UAvaRadialArrangeModifier::SetOrientAxis(const EAvaAxis NewOrientAxis)
+{
+	OrientAxis = NewOrientAxis;
+
+	MarkModifierDirty();
+}
+
+void UAvaRadialArrangeModifier::SetFlipOrient(const bool bNewFlipOrient)
+{
+	bFlipOrient = bNewFlipOrient;
+
+	MarkModifierDirty();
+}
+
+void UAvaRadialArrangeModifier::OnModifierCDOSetup(FActorModifierCoreMetadata& InMetadata)
+{
+	Super::OnModifierCDOSetup(InMetadata);
+
+	InMetadata.SetName(TEXT("RadialArrange"));
+	InMetadata.SetCategory(TEXT("Layout"));
+#if WITH_EDITOR
+	InMetadata.SetDescription(LOCTEXT("ModifierDescription", "Positions child actors in a 2D radial format"));
+#endif
+}
+
+void UAvaRadialArrangeModifier::Apply()
+{
+	AActor* ModifyActor = GetModifiedActor();
+
+	// Early exit if the modify actor is NOT being isolated. The outliner will manage the visibility for the actor and it's children.
+	if (FAvaModifiersActorUtils::IsActorNotIsolated(ModifyActor))
+	{
+		Next();
+		return;
+	}
+	
+	const FAvaSceneTreeUpdateModifierExtension* SceneExtension = GetExtension<FAvaSceneTreeUpdateModifierExtension>();
+	if (!SceneExtension)
+	{
+		Fail(LOCTEXT("InvalidSceneExtension", "Scene extension could not be found"));
+		return;
+	}
+
+	TArray<TWeakObjectPtr<AActor>> AttachedActors = SceneExtension->GetDirectChildrenActor(ModifyActor);
+	const int32 AttachedActorCount = AttachedActors.Num();
+	const int32 TotalSlotCount = Count == -1 ? AttachedActorCount : FMath::Min(AttachedActorCount, Count);
+
+	// open distance in degrees where children will be placed
+	const float AngleOpenDistance = EndAngle - StartAngle;
+	const float RadiusDistance = Rings > 1 ? OuterRadius - InnerRadius : 0.0f;
+	const float RadiusDistancePerRing = RadiusDistance / Rings;
+
+	auto CalculateRelativeOffset = [this, RadiusDistancePerRing](float InAngleInDegrees, const int32 InRingIndex) -> FVector
+	{
+		InAngleInDegrees = FRotator::NormalizeAxis(InAngleInDegrees);
+
+		double SlotSin = 0.0f;
+		double SlotCos = 0.0f;
+		FMath::SinCos(&SlotSin, &SlotCos, FMath::DegreesToRadians(InAngleInDegrees));
+
+		const float RingStartOffset = RadiusDistancePerRing * InRingIndex;
+		const float ChildRadius = InnerRadius + RingStartOffset;
+
+		return FVector(0.0f, ChildRadius * SlotCos, ChildRadius * SlotSin);
+	};
+
+	UAvaTransformModifierShared* LayoutShared = GetShared<UAvaTransformModifierShared>(true);
+	UAvaVisibilityModifierShared* VisibilityShared = GetShared<UAvaVisibilityModifierShared>(true);
+
+	TSet<TWeakObjectPtr<AActor>> NewChildrenActorsWeak;
+	for (int32 ChildIndex = 0; ChildIndex < AttachedActorCount; ++ChildIndex)
+	{
+		AActor* AttachedActor = AttachedActors[ChildIndex].Get();
+
+		if (!AttachedActor)
+		{
+			continue;
+		}
+
+		{
+			// Track all new children actors
+			TArray<AActor*> ChildrenActors { AttachedActor };
+			AttachedActor->GetAttachedActors(ChildrenActors, false, true);
+			for (AActor* ChildActor : ChildrenActors)
+			{
+				NewChildrenActorsWeak.Add(ChildActor);
+			}
+		}
+		
+		// No need to handle nested children actor, only direct children, visibility will propagate
+		if (AttachedActor->GetAttachParentActor() != ModifyActor)
+		{
+			continue;
+		}
+		
+		// Track this actor visibility state
+		const bool bIsVisible = ChildIndex < TotalSlotCount;
+		VisibilityShared->SetActorVisibility(this, AttachedActor, !bIsVisible, true);
+
+		const int32 ChildIndexToUse = ChildIndex;
+		float RingAngleStep = 0.0f;
+		int32 RingIndex = 0;
+		int32 RingColumnIndex = 0;
+		float SlotAngle = 0.0f;
+
+		switch (Arrangement)
+		{
+			/**
+			 * Each radial ring will contain the same number of elements.
+			 * The space between elements in the outer rings will be greater than the inner rings.
+			 */
+			case EAvaRadialArrangement::Monospace:
+			{
+				const int32 ChildrenPerRing = FMath::Max(1, FMath::CeilToInt32((float)TotalSlotCount / Rings));
+
+				RingAngleStep = ChildrenPerRing > 1 ? (AngleOpenDistance / (ChildrenPerRing - 1)) : 0.0f;
+
+				RingColumnIndex = ChildIndexToUse % ChildrenPerRing;
+				RingIndex = FMath::FloorToInt32((float)ChildIndexToUse / ChildrenPerRing);
+
+				SlotAngle = StartAngle
+					+ (RingAngleStep * RingColumnIndex)
+					+ 90; // adding 90 degrees to make 0 degrees face up instead of right
+
+				break;
+			}
+			/**
+			 * All elements in all radial rings have the same spacing between them.
+			 * The number of elements in the inner rings will be greater than the outer rings.
+			 */
+			 // @TODO: back engineer this Viz Artist arrangement mode
+			case EAvaRadialArrangement::Equal:
+			{
+				//const float RingAngleStep = AngleOpenDistance / ChildrenPerRing;
+				//RingAngleStep = (AngleOpenDistance * Rings) / TotalSlotCount;
+				RingAngleStep = (AngleOpenDistance / TotalSlotCount) * Rings;
+
+				const int32 ChildrenPerRing = FMath::Max(1, FMath::CeilToInt32((float)TotalSlotCount / Rings));
+
+				RingColumnIndex = ChildIndexToUse % ChildrenPerRing;
+				RingIndex = FMath::FloorToInt32((float)ChildIndexToUse / ChildrenPerRing);
+
+				SlotAngle = StartAngle
+					+ (RingAngleStep * RingColumnIndex)
+					+ 90; // adding 90 degrees to make 0 degrees face up instead of right
+
+				break;
+			}
+		}
+
+		if (bStartFromOuterRadius)
+		{
+			RingIndex = Rings - (RingIndex - 1);
+		}
+
+		const FVector RelativeOffset = CalculateRelativeOffset(SlotAngle, RingIndex);
+
+		// Track this actor layout state
+		LayoutShared->SaveActorState(this, AttachedActor);
+		AttachedActor->SetActorRelativeLocation(RelativeOffset);
+
+		if (bOrient)
+		{
+			const FVector EyePosition = AttachedActor->GetActorLocation();
+			const FVector TargetPosition = ModifyActor->GetActorLocation();
+			const FRotator NewRotation = FAvaModifiersActorUtils::FindLookAtRotation(EyePosition, TargetPosition, OrientAxis, bFlipOrient);
+
+			AttachedActor->SetActorRotation(NewRotation);
+		}
+	}
+	
+	// Untrack previous actors that are not attached anymore
+	const TSet<TWeakObjectPtr<AActor>> UntrackActors = ChildrenActorsWeak.Difference(NewChildrenActorsWeak);
+	LayoutShared->RestoreActorsState(this, UntrackActors);
+	VisibilityShared->RestoreActorsState(this, UntrackActors);
+
+	ChildrenActorsWeak = NewChildrenActorsWeak;
+
+	Next();
+}
+
+#undef LOCTEXT_NAMESPACE
