[MotionDesign]

diff --git a/Engine/Plugins/Experimental/Avalanche/Source/AvalancheModifiers/Private/Shared/AvaTransformModifierShared.cpp b/Engine/Plugins/Experimental/Avalanche/Source/AvalancheModifiers/Private/Shared/AvaTransformModifierShared.cpp
new file mode 100644
--- /dev/null
+++ b/Engine/Plugins/Experimental/Avalanche/Source/AvalancheModifiers/Private/Shared/AvaTransformModifierShared.cpp
@@ -0,0 +1,258 @@
+// Copyright Epic Games, Inc. All Rights Reserved.
+
+#include "Shared/AvaTransformModifierShared.h"
+
+#include "Modifiers/AvaBaseModifier.h"
+#include "GameFramework/Actor.h"
+
+void FAvaTransformSharedModifierState::Save(const AActor* InActor)
+{
+	if (const UAvaBaseModifier* Modifier = ModifierWeak.Get())
+	{
+		if (InActor)
+		{
+			ActorTransform = InActor->GetActorTransform();
+		}
+	}
+}
+
+void FAvaTransformSharedModifierState::Restore(AActor* InActor, EAvaTransformSharedModifier InRestoreState) const
+{
+	if (const UAvaBaseModifier* Modifier = ModifierWeak.Get())
+	{
+		if (InActor)
+		{
+			FTransform RestoreTransform = ActorTransform;
+			const FTransform& CurrentActorTransform = InActor->GetActorTransform();
+			
+			if (!EnumHasAnyFlags(InRestoreState, EAvaTransformSharedModifier::Location))
+			{
+				RestoreTransform.SetLocation(CurrentActorTransform.GetLocation());
+			}
+
+			if (!EnumHasAnyFlags(InRestoreState, EAvaTransformSharedModifier::Rotation))
+			{
+				RestoreTransform.SetRotation(CurrentActorTransform.GetRotation());
+			}
+
+			if (!EnumHasAnyFlags(InRestoreState, EAvaTransformSharedModifier::Scale))
+			{
+				RestoreTransform.SetScale3D(CurrentActorTransform.GetScale3D());
+			}
+
+			if (!CurrentActorTransform.Equals(RestoreTransform))
+			{
+				InActor->SetActorTransform(RestoreTransform);
+			}
+		}
+	}
+}
+
+void FAvaTransformSharedActorState::Save()
+{
+	if (const AActor* Actor = ActorWeak.Get())
+	{
+		ActorTransform = Actor->GetActorTransform();
+	}
+}
+
+void FAvaTransformSharedActorState::Restore(EAvaTransformSharedModifier InRestoreState) const
+{
+	if (AActor* Actor = ActorWeak.Get())
+	{
+		FTransform RestoreTransform = ActorTransform;
+		const FTransform& CurrentActorTransform = Actor->GetActorTransform();
+			
+		if (!EnumHasAnyFlags(InRestoreState, EAvaTransformSharedModifier::Location))
+		{
+			RestoreTransform.SetLocation(CurrentActorTransform.GetLocation());
+		}
+
+		if (!EnumHasAnyFlags(InRestoreState, EAvaTransformSharedModifier::Rotation))
+		{
+			RestoreTransform.SetRotation(CurrentActorTransform.GetRotation());
+		}
+
+		if (!EnumHasAnyFlags(InRestoreState, EAvaTransformSharedModifier::Scale))
+		{
+			RestoreTransform.SetScale3D(CurrentActorTransform.GetScale3D());
+		}
+
+		if (!CurrentActorTransform.Equals(RestoreTransform))
+		{
+			Actor->SetActorTransform(RestoreTransform);
+		}
+	}
+}
+
+void UAvaTransformModifierShared::SaveActorState(UAvaBaseModifier* InModifierContext, AActor* InActor)
+{
+	if (!IsValid(InActor))
+	{
+		return;
+	}
+
+	FAvaTransformSharedActorState& ActorState = ActorStates.FindOrAdd(FAvaTransformSharedActorState(InActor));
+
+	if (ActorState.ModifierStates.IsEmpty())
+	{
+		ActorState.Save();
+	}
+
+	bool bAlreadyInSet = false;
+	FAvaTransformSharedModifierState& ModifierState = ActorState.ModifierStates.FindOrAdd(FAvaTransformSharedModifierState(InModifierContext), &bAlreadyInSet);
+
+	if (!bAlreadyInSet)
+	{
+		ModifierState.Save(InActor);
+		ActorState.ModifierStates.Add(ModifierState);
+	}
+}
+
+void UAvaTransformModifierShared::RestoreActorState(UAvaBaseModifier* InModifierContext, AActor* InActor, EAvaTransformSharedModifier InRestoreTransform)
+{
+	if (!IsValid(InActor))
+	{
+		return;
+	}
+
+	FAvaTransformSharedActorState* ActorState = ActorStates.Find(FAvaTransformSharedActorState(InActor));
+	if (!ActorState)
+	{
+		return;
+	}
+
+	const FAvaTransformSharedModifierState* ActorModifierState = ActorState->ModifierStates.Find(FAvaTransformSharedModifierState(InModifierContext));
+	if (!ActorModifierState)
+	{
+		return;
+	}
+
+	// restore modifier state and remove it
+	ActorModifierState->Restore(InActor, InRestoreTransform);
+	ActorState->ModifierStates.Remove(*ActorModifierState);
+
+	// Restore original actor state and remove it
+	if (ActorState->ModifierStates.IsEmpty())
+	{
+		ActorState->Restore(InRestoreTransform);
+		ActorStates.Remove(*ActorState);
+	}
+}
+
+FAvaTransformSharedActorState* UAvaTransformModifierShared::FindActorState(AActor* InActor)
+{
+	if (!IsValid(InActor))
+	{
+		return nullptr;
+	}
+
+	return ActorStates.Find(FAvaTransformSharedActorState(InActor));
+}
+
+TSet<FAvaTransformSharedActorState*> UAvaTransformModifierShared::FindActorsState(UAvaBaseModifier* InModifierContext)
+{
+	TSet<FAvaTransformSharedActorState*> ModifierActorStates;
+	
+	for (FAvaTransformSharedActorState& ActorState : ActorStates)
+	{
+		if (ActorState.ModifierStates.Contains(FAvaTransformSharedModifierState(InModifierContext)))
+		{
+			ModifierActorStates.Add(&ActorState);
+		}
+	}
+	
+	return ModifierActorStates;
+}
+
+void UAvaTransformModifierShared::RestoreActorsState(UAvaBaseModifier* InModifierContext, const TSet<AActor*>* InActors, EAvaTransformSharedModifier InRestoreTransform)
+{
+	const FAvaTransformSharedModifierState SearchModifierState(InModifierContext);
+	TSet<AActor*> LinkedModifierActors;
+	TSet<UActorModifierCoreBase*> LinkedActorModifiers;
+	
+	for (const FAvaTransformSharedActorState& ActorState : ActorStates)
+	{
+		AActor* Actor = ActorState.ActorWeak.Get();
+		if (!Actor)
+		{
+			continue;
+		}
+
+		if (!ActorState.ModifierStates.Contains(SearchModifierState))
+		{
+			continue;
+		}
+
+		if (InActors && !InActors->Contains(Actor))
+		{
+			continue;
+		}
+		
+		// Collect actors affected by modifier
+		LinkedModifierActors.Add(Actor);
+
+		// Collect linked actor modifiers
+		for (const FAvaTransformSharedModifierState& ModifierState : ActorState.ModifierStates)
+		{
+			if (UActorModifierCoreBase* Modifier = ModifierState.ModifierWeak.Get())
+			{
+				LinkedActorModifiers.Add(Modifier);
+			}
+		}
+	}
+
+	// Locking state to prevent from updating when restoring state
+	// When destroyed : Unlocking state of modifier
+	FActorModifierCoreScopedLock ModifiersLock(LinkedActorModifiers);
+
+	// Restore actor state
+	for (AActor* Actor : LinkedModifierActors)
+	{
+		RestoreActorState(InModifierContext, Actor, InRestoreTransform);
+	}
+}
+
+void UAvaTransformModifierShared::RestoreActorsState(UAvaBaseModifier* InModifierContext, const TSet<TWeakObjectPtr<AActor>>& InActors, EAvaTransformSharedModifier InRestoreTransform)
+{
+	TSet<AActor*> Actors;
+	Algo::Transform(InActors, Actors, [](const TWeakObjectPtr<AActor>& InActor)->AActor*{ return InActor.Get(); });
+
+	RestoreActorsState(InModifierContext, &Actors, InRestoreTransform);
+}
+
+bool UAvaTransformModifierShared::IsActorStateSaved(UAvaBaseModifier* InModifierContext, AActor* InActor)
+{
+	if (const FAvaTransformSharedActorState* ActorState = FindActorState(InActor))
+	{
+		return ActorState->ModifierStates.Contains(FAvaTransformSharedModifierState(InModifierContext));
+	}
+	
+	return false;
+}
+
+bool UAvaTransformModifierShared::IsActorsStateSaved(UAvaBaseModifier* InModifierContext)
+{
+	const FAvaTransformSharedModifierState ModifierState(InModifierContext);
+	
+	for (const FAvaTransformSharedActorState& ActorState : ActorStates)
+	{
+		if (ActorState.ModifierStates.Contains(ModifierState))
+		{
+			return true;
+		}
+	}
+	
+	return false;
+}
+
+void UAvaTransformModifierShared::PostLoad()
+{
+	Super::PostLoad();
+
+	// Remove invalid items when loading
+	for (FAvaTransformSharedActorState& ActorState : ActorStates)
+	{
+		ActorState.ModifierStates.Remove(FAvaTransformSharedModifierState(nullptr));
+	}
+}
diff --git a/Engine/Plugins/Experimental/Avalanche/Source/AvalancheModifiers/Private/Shared/AvaTranslucentPriorityModifierShared.cpp b/Engine/Plugins/Experimental/Avalanche/Source/AvalancheModifiers/Private/Shared/AvaTranslucentPriorityModifierShared.cpp
new file mode 100644
--- /dev/null
+++ b/Engine/Plugins/Experimental/Avalanche/Source/AvalancheModifiers/Private/Shared/AvaTranslucentPriorityModifierShared.cpp
@@ -0,0 +1,325 @@
+// Copyright Epic Games, Inc. All Rights Reserved.
+
+#include "Shared/AvaTranslucentPriorityModifierShared.h"
+#include "AvaActorUtils.h"
+#include "AvaBlueprint.h"
+#include "AvaOutlinerUtils.h"
+#include "Camera/CameraActor.h"
+#include "Camera/CameraComponent.h"
+#include "Framework/AvalancheGameInstance.h"
+
+#if WITH_EDITOR
+#include "AvaOutliner.h"
+#include "AvaOutlinerSubsystem.h"
+#endif
+
+void UAvaTranslucentPriorityModifierShared::SetComponentsState(UAvaTranslucentPriorityModifier* InModifierContext, const TSet<TWeakObjectPtr<UPrimitiveComponent>>& InComponents)
+{
+	if (!IsValid(InModifierContext))
+	{
+		return;
+	}
+
+	// Remove old component states
+	for (TSet<FAvaTranslucentPriorityModifierComponentState>::TIterator It(ComponentStates); It; ++It)
+	{
+		UPrimitiveComponent* PrimitiveComponent = It->PrimitiveComponentWeak.Get();
+
+		if (PrimitiveComponent && InComponents.Contains(PrimitiveComponent))
+		{
+			It->ModifierWeak = InModifierContext;
+			continue;
+		}
+
+		const UAvaTranslucentPriorityModifier* Modifier = It->ModifierWeak.Get();
+		if (!PrimitiveComponent || !Modifier || Modifier == InModifierContext)
+		{
+			It.RemoveCurrent();
+		}
+	}
+
+	// Add missing components
+	for (const TWeakObjectPtr<UPrimitiveComponent>& PrimitiveComponentWeak : InComponents)
+	{
+		if (UPrimitiveComponent* PrimitiveComponent = PrimitiveComponentWeak.Get())
+		{
+			SaveComponentState(InModifierContext, PrimitiveComponent, true);
+		}
+	}
+}
+
+void UAvaTranslucentPriorityModifierShared::SaveComponentState(UAvaTranslucentPriorityModifier* InModifierContext, UPrimitiveComponent* InComponent, bool bInOverrideContext)
+{
+	if (!IsValid(InComponent))
+	{
+		return;
+	}
+
+	FAvaTranslucentPriorityModifierComponentState& ComponentState = ComponentStates.FindOrAdd(FAvaTranslucentPriorityModifierComponentState(InComponent));
+
+	if (!ComponentState.ModifierWeak.IsValid())
+	{
+		ComponentState.Save();
+		bInOverrideContext = true;
+	}
+
+	if (bInOverrideContext)
+	{
+		ComponentState.ModifierWeak = InModifierContext;
+	}
+}
+
+void UAvaTranslucentPriorityModifierShared::RestoreComponentState(const UAvaTranslucentPriorityModifier* InModifierContext, UPrimitiveComponent* InComponent, bool bInClearState)
+{
+	if (!IsValid(InComponent))
+	{
+		return;
+	}
+
+	if (const FAvaTranslucentPriorityModifierComponentState* ComponentState = ComponentStates.Find(FAvaTranslucentPriorityModifierComponentState(InComponent)))
+	{
+		if (ComponentState->ModifierWeak == InModifierContext)
+		{
+			ComponentState->Restore();
+
+			if (bInClearState)
+			{
+				ComponentStates.Remove(*ComponentState);
+
+				// This is needed in case other parent translucent modifier take ownership of the component
+				UActorComponent::MarkRenderStateDirtyEvent.Broadcast(*InComponent);
+			}
+		}
+	}
+}
+
+void UAvaTranslucentPriorityModifierShared::RestoreComponentsState(const UAvaTranslucentPriorityModifier* InModifierContext, bool bInClearState)
+{
+	if (!InModifierContext)
+	{
+		return;
+	}
+
+	for (TSet<FAvaTranslucentPriorityModifierComponentState>::TIterator It(ComponentStates); It; ++It)
+	{
+		if (InModifierContext == It->ModifierWeak)
+		{
+			RestoreComponentState(It->ModifierWeak.Get(), It->PrimitiveComponentWeak.Get(), bInClearState);
+		}
+	}
+}
+
+void UAvaTranslucentPriorityModifierShared::RestoreComponentsState(const UAvaTranslucentPriorityModifier* InModifierContext, const TSet<TWeakObjectPtr<UPrimitiveComponent>>& InComponents, bool bInClearState)
+{
+	if (!InModifierContext)
+	{
+		return;
+	}
+
+	for (const TWeakObjectPtr<UPrimitiveComponent>& Component : InComponents)
+	{
+		RestoreComponentState(InModifierContext, Component.Get(), bInClearState);
+	}
+}
+
+UAvaTranslucentPriorityModifier* UAvaTranslucentPriorityModifierShared::FindModifierContext(UPrimitiveComponent* InComponent) const
+{
+	UAvaTranslucentPriorityModifier* ModifierContext = nullptr;
+
+	if (!InComponent)
+	{
+		return ModifierContext;
+	}
+
+	if (const FAvaTranslucentPriorityModifierComponentState* ComponentState = ComponentStates.Find(FAvaTranslucentPriorityModifierComponentState(InComponent)))
+	{
+		ModifierContext = ComponentState->ModifierWeak.Get();
+	}
+
+	return ModifierContext;
+}
+
+TArray<const FAvaTranslucentPriorityModifierComponentState*> UAvaTranslucentPriorityModifierShared::GetSortedComponentStates(UAvaTranslucentPriorityModifier* InModifierContext) const
+{
+	TArray<const FAvaTranslucentPriorityModifierComponentState*> SortedComponentStates;
+
+	if (!InModifierContext)
+	{
+		return SortedComponentStates;
+	}
+
+	SortedComponentStates.Reserve(ComponentStates.Num());
+
+	if (InModifierContext->GetMode() == EAvaTranslucentPriorityModifierMode::Manual)
+	{
+		// Only gather components the modifier is handling
+		for (const FAvaTranslucentPriorityModifierComponentState& ComponentState : ComponentStates)
+		{
+			const UAvaTranslucentPriorityModifier* ModifierContext = ComponentState.ModifierWeak.Get();
+			const UPrimitiveComponent* PrimitiveComponent = ComponentState.PrimitiveComponentWeak.Get();
+
+			if (!ModifierContext || !PrimitiveComponent)
+			{
+				continue;
+			}
+
+			if (ModifierContext == InModifierContext)
+			{
+				SortedComponentStates.Add(&ComponentState);
+			}
+		}
+	}
+	else if (InModifierContext->GetMode() == EAvaTranslucentPriorityModifierMode::AutoCameraDistance)
+	{
+		// Gather components with same camera context
+		const ACameraActor* CameraActor = InModifierContext->GetCameraActorWeak().Get();
+
+		if (!CameraActor || CameraActor->GetWorld() != InModifierContext->GetWorld())
+		{
+			return SortedComponentStates;
+		}
+
+		for (const FAvaTranslucentPriorityModifierComponentState& ComponentState : ComponentStates)
+		{
+			const UAvaTranslucentPriorityModifier* ModifierContext = ComponentState.ModifierWeak.Get();
+			const UPrimitiveComponent* PrimitiveComponent = ComponentState.PrimitiveComponentWeak.Get();
+
+			if (!ModifierContext || !PrimitiveComponent)
+			{
+				continue;
+			}
+
+			if (ModifierContext->GetMode() == InModifierContext->GetMode()
+				&& ModifierContext->CameraActorWeak == CameraActor)
+			{
+				SortedComponentStates.Add(&ComponentState);
+			}
+		}
+
+		// Sort current modifiers by distance X to a specific camera
+		SortedComponentStates.Sort([CameraActor](const FAvaTranslucentPriorityModifierComponentState& InComponentA, const FAvaTranslucentPriorityModifierComponentState& InComponentB)->bool
+		{
+			const FVector AForwardComponentLocation = CameraActor->GetActorForwardVector() * InComponentA.GetComponentLocation();
+			const float ADist = FVector::Distance(AForwardComponentLocation, CameraActor->GetActorLocation());
+
+			const FVector BForwardComponentLocation = CameraActor->GetActorForwardVector() * InComponentB.GetComponentLocation();
+			const float BDist = FVector::Distance(BForwardComponentLocation, CameraActor->GetActorLocation());
+
+			return ADist > BDist;
+		});
+	}
+	else if (InModifierContext->GetMode() == EAvaTranslucentPriorityModifierMode::AutoOutlinerTree)
+	{
+		const AActor* ModifiedActor = InModifierContext->GetModifiedActor();
+		UWorld* const World = ModifiedActor->GetTypedOuter<UWorld>();
+
+		for (const FAvaTranslucentPriorityModifierComponentState& ComponentState : ComponentStates)
+		{
+			const UAvaTranslucentPriorityModifier* ModifierContext = ComponentState.ModifierWeak.Get();
+			const UPrimitiveComponent* PrimitiveComponent = ComponentState.PrimitiveComponentWeak.Get();
+
+			if (!ModifierContext || !PrimitiveComponent)
+			{
+				continue;
+			}
+
+			const AActor* OwningActor = ModifierContext->GetModifiedActor();
+			if (!OwningActor)
+			{
+				continue;
+			}
+
+			if (ModifierContext->GetMode() == InModifierContext->GetMode()
+				&& OwningActor->GetTypedOuter<UWorld>() == World)
+			{
+				SortedComponentStates.Add(&ComponentState);
+			}
+		}
+
+		// Sort current modifiers by outliner hierarchy
+#if WITH_EDITOR
+		// Use the outliner first if available for sorting
+		TSharedPtr<FAvaOutliner> AvaOutliner = FAvaOutlinerUtils::EditorGetOutliner(World);
+		if (AvaOutliner.IsValid())
+		{
+			SortedComponentStates.Sort([AvaOutliner](const FAvaTranslucentPriorityModifierComponentState& InComponentA, const FAvaTranslucentPriorityModifierComponentState& InComponentB)->bool
+			{
+				const FAvaOutlinerItemPtr OutlinerItemA = AvaOutliner->FindItem(InComponentA.GetOwningActor());
+				const FAvaOutlinerItemPtr OutlinerItemB = AvaOutliner->FindItem(InComponentB.GetOwningActor());
+
+				return FAvaOutliner::CompareOutlinerItemOrder(OutlinerItemB, OutlinerItemA);
+			});
+		}
+		else
+		{
+#endif
+			// Fallback on scene tree sorting
+			if (const IAvaSceneInterface* SceneInterface = FAvaActorUtils::GetSceneInterfaceFromActor(ModifiedActor))
+			{
+				const FAvaSceneTree& SceneTree = SceneInterface->GetSceneTree();
+
+				SortedComponentStates.Sort([&SceneTree, World](const FAvaTranslucentPriorityModifierComponentState& InComponentA, const FAvaTranslucentPriorityModifierComponentState& InComponentB)->bool
+				{
+					const FAvaSceneTreeNode* SceneItemA = SceneTree.FindTreeNode(FAvaSceneItem(InComponentA.GetOwningActor(), World));
+					const FAvaSceneTreeNode* SceneItemB = SceneTree.FindTreeNode(FAvaSceneItem(InComponentB.GetOwningActor(), World));
+
+					return FAvaSceneTree::CompareTreeItemOrder(SceneItemB, SceneItemA);
+				});
+			}
+#if WITH_EDITOR
+		}
+#endif
+	}
+
+	return SortedComponentStates;
+}
+
+void UAvaTranslucentPriorityModifierShared::PostLoad()
+{
+	Super::PostLoad();
+
+	// Clean up invalid components
+	for (TSet<FAvaTranslucentPriorityModifierComponentState>::TIterator It(ComponentStates); It; ++It)
+	{
+		if (!It->PrimitiveComponentWeak.IsValid())
+		{
+			It.RemoveCurrent();
+		}
+	}
+}
+
+void FAvaTranslucentPriorityModifierComponentState::Save()
+{
+	if (const UPrimitiveComponent* Component = PrimitiveComponentWeak.Get())
+	{
+		SortPriority = Component->TranslucencySortPriority;
+	}
+}
+
+void FAvaTranslucentPriorityModifierComponentState::Restore() const
+{
+	if (UPrimitiveComponent* Component = PrimitiveComponentWeak.Get())
+	{
+		Component->SetTranslucentSortPriority(SortPriority);
+	}
+}
+
+AActor* FAvaTranslucentPriorityModifierComponentState::GetOwningActor() const
+{
+	if (const UPrimitiveComponent* Component = PrimitiveComponentWeak.Get())
+	{
+		return Component->GetOwner();
+	}
+
+	return nullptr;
+}
+
+FVector FAvaTranslucentPriorityModifierComponentState::GetComponentLocation() const
+{
+	if (const UPrimitiveComponent* Component = PrimitiveComponentWeak.Get())
+	{
+		return Component->GetComponentLocation();
+	}
+
+	return FVector::ZeroVector;
+}
diff --git a/Engine/Plugins/Experimental/Avalanche/Source/AvalancheModifiers/Private/Shared/AvaVisibilityModifierShared.cpp b/Engine/Plugins/Experimental/Avalanche/Source/AvalancheModifiers/Private/Shared/AvaVisibilityModifierShared.cpp
new file mode 100644
--- /dev/null
+++ b/Engine/Plugins/Experimental/Avalanche/Source/AvalancheModifiers/Private/Shared/AvaVisibilityModifierShared.cpp
@@ -0,0 +1,252 @@
+// Copyright Epic Games, Inc. All Rights Reserved.
+
+#include "Shared/AvaVisibilityModifierShared.h"
+
+void FAvaVisibilitySharedModifierState::Save(const AActor* InActor)
+{
+	if (const UAvaBaseModifier* Modifier = ModifierWeak.Get())
+	{
+		if (InActor)
+		{
+#if WITH_EDITOR
+			bActorHiddenInEditor = InActor->IsTemporarilyHiddenInEditor();
+#endif
+			bActorHiddenInGame = InActor->IsHidden();
+		}
+	}
+}
+
+void FAvaVisibilitySharedModifierState::Restore(AActor* InActor) const
+{
+	if (const UAvaBaseModifier* Modifier = ModifierWeak.Get())
+	{
+		if (InActor)
+		{
+#if WITH_EDITOR
+			InActor->SetIsTemporarilyHiddenInEditor(bActorHiddenInEditor);
+#endif
+			InActor->SetActorHiddenInGame(bActorHiddenInGame);
+		}
+	}
+}
+
+void FAvaVisibilitySharedActorState::Save()
+{
+	if (const AActor* Actor = ActorWeak.Get())
+	{
+#if WITH_EDITOR
+		bActorHiddenInEditor = Actor->IsTemporarilyHiddenInEditor();
+#endif
+		bActorHiddenInGame = Actor->IsHidden();
+	}
+}
+
+void FAvaVisibilitySharedActorState::Restore() const
+{
+	if (AActor* Actor = ActorWeak.Get())
+	{
+#if WITH_EDITOR
+		Actor->SetIsTemporarilyHiddenInEditor(bActorHiddenInEditor);
+#endif
+		Actor->SetActorHiddenInGame(bActorHiddenInGame);
+	}
+}
+
+void UAvaVisibilityModifierShared::SaveActorState(UAvaBaseModifier* InModifierContext, AActor* InActor)
+{
+	if (!IsValid(InActor))
+	{
+		return;
+	}
+
+	FAvaVisibilitySharedActorState& ActorState = ActorStates.FindOrAdd(FAvaVisibilitySharedActorState(InActor));
+
+	if (ActorState.ModifierStates.IsEmpty())
+	{
+		ActorState.Save();
+	}
+	
+	FAvaVisibilitySharedModifierState ModifierState(InModifierContext);
+	if (!ActorState.ModifierStates.Contains(ModifierState))
+	{
+		ModifierState.Save(InActor);
+		ActorState.ModifierStates.Add(InModifierContext);
+	}
+}
+
+void UAvaVisibilityModifierShared::RestoreActorState(UAvaBaseModifier* InModifierContext, AActor* InActor)
+{
+	if (!IsValid(InActor))
+	{
+		return;
+	}
+
+	if (FAvaVisibilitySharedActorState* ActorState = ActorStates.Find(FAvaVisibilitySharedActorState(InActor)))
+	{
+		if (const FAvaVisibilitySharedModifierState* ActorModifierState = ActorState->ModifierStates.Find(FAvaVisibilitySharedModifierState(InModifierContext)))
+		{
+			// restore modifier state and remove it
+			ActorModifierState->Restore(InActor);
+			ActorState->ModifierStates.Remove(*ActorModifierState);
+
+			// Restore original actor state and remove it
+			if (ActorState->ModifierStates.IsEmpty())
+			{
+				ActorState->Restore();
+				ActorStates.Remove(*ActorState);
+			}
+		}
+	}
+}
+
+FAvaVisibilitySharedActorState* UAvaVisibilityModifierShared::FindActorState(AActor* InActor)
+{
+	if (!IsValid(InActor))
+	{
+		return nullptr;
+	}
+
+	return ActorStates.Find(FAvaVisibilitySharedActorState(InActor));
+}
+
+void UAvaVisibilityModifierShared::SetActorVisibility(UAvaBaseModifier* InModifierContext, AActor* InActor, bool bInHidden, bool bInRecurse, EAvaVisibilityActor InActorVisibility)
+{
+	if (!IsValid(InActor))
+	{
+		return;
+	}
+
+	TArray<AActor*> Actors { InActor };
+
+	if (bInRecurse)
+	{
+		constexpr bool bResetArray = false;
+		constexpr bool bRecursivelyIncludeAttachedActors = true;
+		InActor->GetAttachedActors(Actors, bResetArray, bRecursivelyIncludeAttachedActors);
+	}
+
+	SetActorsVisibility(InModifierContext, Actors, bInHidden, InActorVisibility);
+}
+
+void UAvaVisibilityModifierShared::SetActorsVisibility(UAvaBaseModifier* InModifierContext, TArray<AActor*> InActors, bool bInHidden, EAvaVisibilityActor InActorVisibility)
+{
+	for (AActor* Actor : InActors)
+	{
+		if (!Actor)
+		{
+			continue;
+		}
+		
+		SaveActorState(InModifierContext, Actor);
+
+#if WITH_EDITOR
+		if (EnumHasAnyFlags(InActorVisibility, EAvaVisibilityActor::Editor))
+		{
+			if (Actor->IsTemporarilyHiddenInEditor() != bInHidden)
+			{
+				Actor->SetIsTemporarilyHiddenInEditor(bInHidden);
+			}
+		}
+#endif
+		
+		if (EnumHasAnyFlags(InActorVisibility, EAvaVisibilityActor::Game))
+		{
+			if (Actor->IsHidden() != bInHidden)
+			{
+				Actor->SetActorHiddenInGame(bInHidden);
+			}
+		}
+	}
+}
+
+void UAvaVisibilityModifierShared::RestoreActorsState(UAvaBaseModifier* InModifierContext, const TSet<AActor*>* InActors)
+{
+	const FAvaVisibilitySharedModifierState SearchModifierState(InModifierContext);
+	TSet<AActor*> LinkedModifierActors;
+	TSet<UActorModifierCoreBase*> LinkedActorModifiers;
+	
+	for (const FAvaVisibilitySharedActorState& ActorState : ActorStates)
+	{
+		AActor* Actor = ActorState.ActorWeak.Get();
+		if (!Actor)
+		{
+			continue;
+		}
+
+		if (!ActorState.ModifierStates.Contains(SearchModifierState))
+		{
+			continue;
+		}
+
+		if (InActors && !InActors->Contains(Actor))
+		{
+			continue;
+		}
+		
+		// Collect actors affected by modifier
+		LinkedModifierActors.Add(Actor);
+
+		// Collect linked actor modifiers
+		for (const FAvaVisibilitySharedModifierState& ModifierState : ActorState.ModifierStates)
+		{
+			if (UActorModifierCoreBase* Modifier = ModifierState.ModifierWeak.Get())
+			{
+				LinkedActorModifiers.Add(Modifier);
+			}
+		}
+	}
+
+	// Locking state to prevent from updating when restoring state
+	// When destroyed : Unlocking state of modifier
+	FActorModifierCoreScopedLock ModifiersLock(LinkedActorModifiers);
+
+	// Restore actor state
+	for (AActor* Actor : LinkedModifierActors)
+	{
+		RestoreActorState(InModifierContext, Actor);
+	}
+}
+
+void UAvaVisibilityModifierShared::RestoreActorsState(UAvaBaseModifier* InModifierContext, const TSet<TWeakObjectPtr<AActor>>& InActors)
+{
+	TSet<AActor*> Actors;
+	Algo::Transform(InActors, Actors, [](const TWeakObjectPtr<AActor>& InActor)->AActor*{ return InActor.Get(); });
+
+	RestoreActorsState(InModifierContext, &Actors);
+}
+
+bool UAvaVisibilityModifierShared::IsActorStateSaved(UAvaBaseModifier* InModifierContext, AActor* InActor)
+{
+	if (const FAvaVisibilitySharedActorState* ActorState = FindActorState(InActor))
+	{
+		return ActorState->ModifierStates.Contains(FAvaVisibilitySharedModifierState(InModifierContext));
+	}
+	
+	return false;
+}
+
+bool UAvaVisibilityModifierShared::IsActorsStateSaved(UAvaBaseModifier* InModifierContext)
+{
+	const FAvaVisibilitySharedModifierState ModifierState(InModifierContext);
+	
+	for (const FAvaVisibilitySharedActorState& ActorState : ActorStates)
+	{
+		if (ActorState.ModifierStates.Contains(ModifierState))
+		{
+			return true;
+		}
+	}
+	
+	return false;
+}
+
+void UAvaVisibilityModifierShared::PostLoad()
+{
+	Super::PostLoad();
+
+	// Remove invalid items when loading
+	for (FAvaVisibilitySharedActorState& ActorState : ActorStates)
+	{
+		ActorState.ModifierStates.Remove(FAvaVisibilitySharedModifierState(nullptr));
+	}
+}
diff --git a/Engine/Plugins/Experimental/Avalanche/Source/AvalancheModifiers/Public/AvaModifiersParametricMaterial.h b/Engine/Plugins/Experimental/Avalanche/Source/AvalancheModifiers/Public/AvaModifiersParametricMaterial.h
new file mode 100644
--- /dev/null
+++ b/Engine/Plugins/Experimental/Avalanche/Source/AvalancheModifiers/Public/AvaModifiersParametricMaterial.h
@@ -0,0 +1,42 @@
+// Copyright Epic Games, Inc. All Rights Reserved.
+
+#pragma once
+
+#include "Math/Color.h"
+#include "UObject/ObjectPtr.h"
+#include "AvaModifiersParametricMaterial.generated.h"
+
+class UMaterial;
+class UMaterialInstanceDynamic;
+class UObject;
+
+/** Use this if you want a parametric material */
+USTRUCT(BlueprintType)
+struct AVALANCHEMODIFIERS_API FAvaModifiersParametricMaterial
+{
+	GENERATED_BODY()
+	
+	FAvaModifiersParametricMaterial();
+
+	FAvaModifiersParametricMaterial(const FAvaModifiersParametricMaterial& Other);
+	FAvaModifiersParametricMaterial& operator=(const FAvaModifiersParametricMaterial& Other);
+
+	UPROPERTY()
+	FLinearColor MaskColor;
+
+	UMaterial* GetDefaultMaterial() const;
+	UMaterialInstanceDynamic* GetMaterial() const;
+	void ApplyChanges(UObject* Outer = nullptr);
+
+protected:
+	UPROPERTY()
+	TObjectPtr<UMaterial> DefaultMaterial;
+
+	UPROPERTY(Transient, Instanced)
+	TObjectPtr<UMaterialInstanceDynamic> InstanceMaterial;
+
+	void ApplyParams() const;
+	void CreateAndApply(UObject* Outer = nullptr);
+	void EnsureCurrentMaterial(UObject* Outer = nullptr);
+	UMaterial* LoadResource() const;
+};
\ No newline at end of file
diff --git a/Engine/Plugins/Experimental/Avalanche/Source/AvalancheModifiers/Public/Extensions/AvaRenderStateUpdateModifierExtension.h b/Engine/Plugins/Experimental/Avalanche/Source/AvalancheModifiers/Public/Extensions/AvaRenderStateUpdateModifierExtension.h
new file mode 100644
--- /dev/null
+++ b/Engine/Plugins/Experimental/Avalanche/Source/AvalancheModifiers/Public/Extensions/AvaRenderStateUpdateModifierExtension.h
@@ -0,0 +1,71 @@
+// Copyright Epic Games, Inc. All Rights Reserved.
+
+#pragma once
+
+#include "Modifiers/ActorModifierCoreExtension.h"
+#include "UObject/WeakInterfacePtr.h"
+#include "AvaRenderStateUpdateModifierExtension.generated.h"
+
+class AActor;
+class UActorComponent;
+
+UINTERFACE(MinimalAPI, NotBlueprintType, meta=(CannotImplementInterfaceInBlueprint))
+class UAvaRenderStateUpdateHandler : public UInterface
+{
+	GENERATED_BODY()
+};
+
+/** Implement this interface to handle extension event */
+class AVALANCHEMODIFIERS_API IAvaRenderStateUpdateHandler
+{
+	GENERATED_BODY()
+
+public:
+	/** Callback when a render state actor in this world changes */
+	virtual void OnRenderStateUpdated(AActor* InActor, UActorComponent* InComponent) = 0;
+
+	/** Callback when a tracked actor visibility has changed */
+	virtual void OnActorVisibilityChanged(AActor* InActor) = 0;
+};
+
+/**
+ * This extension tracks specific actors for render state updates,
+ * when an update happens it will invoke IAvaRenderStateUpdateExtension function
+ */
+class AVALANCHEMODIFIERS_API FAvaRenderStateUpdateModifierExtension : public FActorModifierCoreExtension
+{
+
+public:
+	explicit FAvaRenderStateUpdateModifierExtension(IAvaRenderStateUpdateHandler* InExtensionHandler);
+
+	/** Adds an actor to track for visibility */
+	void TrackActorVisibility(AActor* InActor);
+
+	/** Removes a tracked actor for visibility */
+	void UntrackActorVisibility(AActor* InActor);
+
+	/** Checks if actor is tracked for visibility */
+	bool IsActorVisibilityTracked(AActor* InActor) const;
+
+	/** Sets current tracked actors, removes any actors not included */
+	void SetTrackedActorsVisibility(const TSet<TWeakObjectPtr<AActor>>& InActors);
+
+	/** Sets current tracked actors with actor and its children, removes any actors not included */
+	void SetTrackedActorVisibility(AActor* InActor, bool bInIncludeChildren);
+
+protected:
+	//~ Begin FActorModifierCoreExtension
+	virtual void OnExtensionEnabled(EActorModifierCoreEnableReason InReason) override;
+	virtual void OnExtensionDisabled(EActorModifierCoreDisableReason InReason) override;
+	//~ End FActorModifierCoreExtension
+
+private:
+	void OnRenderStateDirty(UActorComponent& InComponent);
+
+	void BindDelegate();
+	void UnbindDelegate();
+
+	TWeakInterfacePtr<IAvaRenderStateUpdateHandler> ExtensionHandlerWeak;
+
+	TMap<TWeakObjectPtr<AActor>, bool> TrackedActorsVisibility;
+};
diff --git a/Engine/Plugins/Experimental/Avalanche/Source/AvalancheModifiers/Public/Extensions/AvaSceneTreeUpdateModifierExtension.h b/Engine/Plugins/Experimental/Avalanche/Source/AvalancheModifiers/Public/Extensions/AvaSceneTreeUpdateModifierExtension.h
new file mode 100644
--- /dev/null
+++ b/Engine/Plugins/Experimental/Avalanche/Source/AvalancheModifiers/Public/Extensions/AvaSceneTreeUpdateModifierExtension.h
@@ -0,0 +1,138 @@
+// Copyright Epic Games, Inc. All Rights Reserved.
+
+#pragma once
+
+#include "AvaDefs.h"
+#include "Delegates/IDelegateInstance.h"
+#include "Modifiers/ActorModifierCoreExtension.h"
+#include "UObject/WeakInterfacePtr.h"
+#include "AvaSceneTreeUpdateModifierExtension.generated.h"
+
+enum class EAvaOutlinerHierarchyChangeType : uint8;
+enum class EAvaReferenceContainer : uint8;
+
+USTRUCT(BlueprintType)
+struct AVALANCHEMODIFIERS_API FAvaSceneTreeActor
+{
+	friend class FAvaSceneTreeUpdateModifierExtension;
+	
+	GENERATED_BODY()
+
+	FAvaSceneTreeActor() = default;
+	
+	explicit FAvaSceneTreeActor(AActor* InActor)
+		: ReferenceContainer(EAvaReferenceContainer::Other)
+		, ReferenceActorWeak(InActor)
+	{}
+	
+	/** The method for finding a reference actor based on it's position in the parent's hierarchy */
+	UPROPERTY(EditAnywhere, BlueprintReadWrite, Category="Outliner")
+	EAvaReferenceContainer ReferenceContainer = EAvaReferenceContainer::Other;
+
+	/** The actor being followed by the modifier. This is user selectable if the Reference Container is set to "Other" */
+	UPROPERTY(EditAnywhere, BlueprintReadWrite, Category="Outliner", meta=(DisplayName="Reference Actor", EditCondition = "ReferenceContainer == EAvaReferenceContainer::Other"))
+	TWeakObjectPtr<AActor> ReferenceActorWeak = nullptr;
+
+	/** If true, will search for the next visible actor based on the selected reference container */
+	UPROPERTY(EditAnywhere, BlueprintReadWrite, Category="Outliner", meta=(EditCondition="ReferenceContainer != EAvaReferenceContainer::Other", EditConditionHides))
+	bool bSkipHiddenActors = false;
+
+	bool operator==(const FAvaSceneTreeActor& InOther) const
+	{
+		return ReferenceContainer == InOther.ReferenceContainer
+			&& ReferenceActorWeak == InOther.ReferenceActorWeak
+			&& bSkipHiddenActors == InOther.bSkipHiddenActors;
+	}
+
+protected:
+	/** All children of reference actor to compare with new set children for changes */
+	UPROPERTY(Transient)
+	TSet<TWeakObjectPtr<AActor>> ReferenceActorChildrenWeak;
+
+	/** Direct children of reference actor where order counts */
+	UPROPERTY(Transient)
+	TArray<TWeakObjectPtr<AActor>> ReferenceActorDirectChildrenWeak;
+	
+	/** Tracked references actors, if we skip hidden actors, we still need to track those for visibility changes, can be rebuilt */
+	UPROPERTY(Transient)
+	TArray<TWeakObjectPtr<AActor>> ReferenceActorsWeak;
+
+	/** Parents of the reference actor */
+	UPROPERTY(Transient)
+	TArray<TWeakObjectPtr<AActor>> ReferenceActorParentsWeak;
+
+	/** Actor from which we start resolving this reference actor */
+	UPROPERTY(Transient)
+	TWeakObjectPtr<AActor> LocalActorWeak = nullptr;
+};
+
+UINTERFACE(MinimalAPI, NotBlueprintType, meta=(CannotImplementInterfaceInBlueprint))
+class UAvaSceneTreeUpdateHandler : public UInterface
+{
+	GENERATED_BODY()
+};
+
+/** Implement this interface to handle extension event */
+class AVALANCHEMODIFIERS_API IAvaSceneTreeUpdateHandler
+{
+	GENERATED_BODY()
+	
+public:
+	virtual void OnSceneTreeTrackedActorChanged(int32 InIdx, AActor* InPreviousActor, AActor* InNewActor) = 0;
+
+	virtual void OnSceneTreeTrackedActorChildrenChanged(int32 InIdx, const TSet<TWeakObjectPtr<AActor>>& InPreviousChildrenActors, const TSet<TWeakObjectPtr<AActor>>& InNewChildrenActors) = 0;
+
+	virtual void OnSceneTreeTrackedActorDirectChildrenChanged(int32 InIdx, const TArray<TWeakObjectPtr<AActor>>& InPreviousChildrenActors, const TArray<TWeakObjectPtr<AActor>>& InNewChildrenActors) = 0;
+
+	virtual void OnSceneTreeTrackedActorParentChanged(int32 InIdx, const TArray<TWeakObjectPtr<AActor>>& InPreviousParentActor, const TArray<TWeakObjectPtr<AActor>>& InNewParentActor) = 0;
+
+	virtual void OnSceneTreeTrackedActorRearranged(int32 InIdx, AActor* InRearrangedActor) = 0;
+};
+
+/**
+ * This extension tracks specific actors for render state updates,
+ * when an update happens it will dirty the modifier it is attached on if filter passes
+ */
+class AVALANCHEMODIFIERS_API FAvaSceneTreeUpdateModifierExtension : public FActorModifierCoreExtension
+{
+	
+public:
+	explicit FAvaSceneTreeUpdateModifierExtension(IAvaSceneTreeUpdateHandler* InExtensionHandler);
+
+	void TrackSceneTree(int32 InTrackedActorIdx, FAvaSceneTreeActor* InTrackedActor);
+	void UntrackSceneTree(int32 InTrackedActorIdx);
+
+	FAvaSceneTreeActor* GetTrackedActor(int32 InTrackedActorIdx) const;
+
+	void CheckTrackedActorsUpdate() const;
+	void CheckTrackedActorUpdate(int32 InIdx) const;
+	
+	TSet<TWeakObjectPtr<AActor>> GetChildrenActorsRecursive(const AActor* InActor) const;
+	TArray<TWeakObjectPtr<AActor>> GetDirectChildrenActor(AActor* InActor) const;
+	TArray<TWeakObjectPtr<AActor>> GetParentActors(const AActor* InActor) const;
+	
+protected:
+	//~ Begin FActorModifierCoreExtension
+	virtual void OnExtensionEnabled(EActorModifierCoreEnableReason InReason) override;
+	virtual void OnExtensionDisabled(EActorModifierCoreDisableReason InReason) override;
+	//~ End FActorModifierCoreExtension
+
+private:
+#if WITH_EDITOR
+	virtual void OnOutlinerLoaded();
+	virtual void OnActorHierarchyChanged(AActor* InActor, const AActor* InParentActor, EAvaOutlinerHierarchyChangeType InChangeType);
+#endif
+	
+	void OnRenderStateDirty(UActorComponent& InComponent);
+	void OnWorldActorDestroyed(AActor* InActor);
+
+	bool IsSameActorArray(const TArray<TWeakObjectPtr<AActor>>& InPreviousActorWeak, const TArray<TWeakObjectPtr<AActor>>& InNewActorWeak) const;
+	
+	TArray<TWeakObjectPtr<AActor>> GetReferenceActors(const FAvaSceneTreeActor* InTrackedActor) const;
+	
+	TWeakInterfacePtr<IAvaSceneTreeUpdateHandler> ExtensionHandlerWeak;
+
+	TMap<int32, FAvaSceneTreeActor*> TrackedActors;
+	
+	FDelegateHandle WorldActorDestroyedDelegate;
+};
\ No newline at end of file
diff --git a/Engine/Plugins/Experimental/Avalanche/Source/AvalancheModifiers/Public/Extensions/AvaTransformUpdateModifierExtension.h b/Engine/Plugins/Experimental/Avalanche/Source/AvalancheModifiers/Public/Extensions/AvaTransformUpdateModifierExtension.h
new file mode 100644
--- /dev/null
+++ b/Engine/Plugins/Experimental/Avalanche/Source/AvalancheModifiers/Public/Extensions/AvaTransformUpdateModifierExtension.h
@@ -0,0 +1,56 @@
+// Copyright Epic Games, Inc. All Rights Reserved.
+
+#pragma once
+
+#include "Modifiers/ActorModifierCoreExtension.h"
+#include "UObject/WeakInterfacePtr.h"
+#include "AvaTransformUpdateModifierExtension.generated.h"
+
+UINTERFACE(MinimalAPI, NotBlueprintType, meta=(CannotImplementInterfaceInBlueprint))
+class UAvaTransformUpdateHandler : public UInterface
+{
+	GENERATED_BODY()
+};
+
+/** Implement this interface to handle extension event */
+class AVALANCHEMODIFIERS_API IAvaTransformUpdateHandler
+{
+	GENERATED_BODY()
+
+public:
+	/** Callback when a tracked actor transform changes */
+	virtual void OnTransformUpdated(AActor* InActor, bool bInParentMoved) = 0;
+};
+
+/**
+ * This extension tracks specific actors for transform updates,
+ * when an update happens it will invoke the IAvaTransformUpdateExtension function
+ */
+class AVALANCHEMODIFIERS_API FAvaTransformUpdateModifierExtension : public FActorModifierCoreExtension
+{
+
+public:
+	explicit FAvaTransformUpdateModifierExtension(IAvaTransformUpdateHandler* InExtensionHandler);
+
+	void TrackActor(AActor* InActor, bool bInReset);
+	void UntrackActor(AActor* InActor);
+
+	void TrackActors(const TSet<TWeakObjectPtr<AActor>>& InActors, bool bInReset);
+	void UntrackActors(const TSet<TWeakObjectPtr<AActor>>& InActors);
+
+protected:
+	//~ Begin FActorModifierCoreExtension
+	virtual void OnExtensionEnabled(EActorModifierCoreEnableReason InReason) override;
+	virtual void OnExtensionDisabled(EActorModifierCoreDisableReason InReason) override;
+	//~ End FActorModifierCoreExtension
+
+private:
+	void OnTransformUpdated(USceneComponent* InComponent, EUpdateTransformFlags InFlags, ETeleportType InType);
+
+	void BindDelegate(const AActor* InActor);
+	void UnbindDelegate(const AActor* InActor);
+
+	TWeakInterfacePtr<IAvaTransformUpdateHandler> ExtensionHandlerWeak;
+
+	TSet<TWeakObjectPtr<AActor>> TrackedActors;
+};
\ No newline at end of file
diff --git a/Engine/Plugins/Experimental/Avalanche/Source/AvalancheModifiers/Public/Modifiers/AvaAlignBetweenModifier.h b/Engine/Plugins/Experimental/Avalanche/Source/AvalancheModifiers/Public/Modifiers/AvaAlignBetweenModifier.h
new file mode 100644
--- /dev/null
+++ b/Engine/Plugins/Experimental/Avalanche/Source/AvalancheModifiers/Public/Modifiers/AvaAlignBetweenModifier.h
@@ -0,0 +1,124 @@
+// Copyright Epic Games, Inc. All Rights Reserved.
+
+#pragma once
+
+#include "AvaBaseModifier.h"
+#include "Extensions/AvaTransformUpdateModifierExtension.h"
+#include "AvaAlignBetweenModifier.generated.h"
+
+class AActor;
+
+/** Represents an actor with a weight and an enabled state. */
+USTRUCT(BlueprintType)
+struct AVALANCHEMODIFIERS_API FAvaAlignBetweenWeightedActor
+{
+	GENERATED_BODY()
+
+	FAvaAlignBetweenWeightedActor()
+	{}
+	
+	explicit FAvaAlignBetweenWeightedActor(AActor* InActor)
+		: ActorWeak(InActor)
+	{}
+	
+	explicit FAvaAlignBetweenWeightedActor(AActor* InActor, float InWeight, bool bInEnabled)
+		: ActorWeak(InActor)
+		, Weight(InWeight)
+		, bEnabled(bInEnabled)
+	{}
+	
+	/** Returns true if the actor is valid and the state is enabled. */
+	bool IsValid() const
+	{
+		return ActorWeak.IsValid() && bEnabled;
+	}
+
+	friend uint32 GetTypeHash(const FAvaAlignBetweenWeightedActor& InItem)
+	{
+		return GetTypeHash(InItem.ActorWeak);
+	}
+
+	bool operator==(const FAvaAlignBetweenWeightedActor& Other) const
+	{
+		return ActorWeak == Other.ActorWeak;
+	}
+	
+	/** An actor that will effect the placement location. */
+	UPROPERTY(EditAnywhere, BlueprintReadWrite, Category = "Motion Design")
+	TWeakObjectPtr<AActor> ActorWeak;
+
+	/** How much effect this actor has on the placement location. */
+	UPROPERTY(EditAnywhere, BlueprintReadWrite, Category = "Motion Design", Meta = (ClampMin = "0.0", UIMin = "0.0"))
+	float Weight = 0.0f;
+
+	/** If true, will consider this weighted actor when calculating the placement location. */
+	UPROPERTY(EditAnywhere, BlueprintReadWrite, Category = "Motion Design")
+	bool bEnabled = false;
+};
+
+/**
+ * Moves the modifying actor to the averaged location between an array of specified actors.
+ */
+UCLASS(BlueprintType)
+class AVALANCHEMODIFIERS_API UAvaAlignBetweenModifier : public UAvaBaseModifier
+	, public IAvaTransformUpdateHandler
+{
+	GENERATED_BODY()
+	
+public:
+	//~ Begin UObject
+#if WITH_EDITOR
+	virtual void PostEditChangeProperty(FPropertyChangedEvent& PropertyChangedEvent) override;
+	virtual void PostTransacted(const FTransactionObjectEvent& TransactionEvent) override;
+#endif
+	//~ End UObject
+
+	/** Gets all actors from their reference actor structs. */
+	UFUNCTION(BlueprintPure, Category = "Motion Design|Modifiers|AlignBetween")
+	TSet<AActor*> GetActors(const bool bEnabledOnly = false) const;
+
+	/** Returns all valid reference actors that enabled and have a weight greater than 0. */
+	TSet<FAvaAlignBetweenWeightedActor> GetEnabledReferenceActors() const;
+	
+	/** Gets all reference actors and their weights. */
+	UFUNCTION(BlueprintPure, Category = "Motion Design|Modifiers|AlignBetween")
+	TSet<FAvaAlignBetweenWeightedActor> GetReferenceActors() const { return ReferenceActors; }
+
+	/** Sets all reference actors and their weights. */
+	UFUNCTION(BlueprintCallable, Category = "Motion Design|Modifiers|AlignBetween")
+	void SetReferenceActors(TSet<FAvaAlignBetweenWeightedActor> NewReferenceActors);
+
+	/** Adds an actor to the reference list. */
+	UFUNCTION(BlueprintCallable, Category = "Motion Design|Modifiers|AlignBetween")
+	void AddReferenceActor(const FAvaAlignBetweenWeightedActor& ReferenceActor);
+
+	/** Removes an actor from the reference list. */
+	UFUNCTION(BlueprintCallable, Category = "Motion Design|Modifiers|AlignBetween")
+	bool RemoveReferenceActor(AActor* const Actor);
+
+	/** Finds an actor in the reference list. */
+	UFUNCTION(BlueprintCallable, Category = "Motion Design|Modifiers|AlignBetween")
+	bool FindReferenceActor(AActor* InActor, FAvaAlignBetweenWeightedActor& OutReferenceActor) const;
+
+protected:
+	//~ Begin UActorModifierCoreBase
+	virtual void OnModifierCDOSetup(FActorModifierCoreMetadata& InMetadata) override;
+	virtual void OnModifierAdded(EActorModifierCoreEnableReason InReason) override;
+	virtual void OnModifierEnabled(EActorModifierCoreEnableReason InReason) override;
+	virtual void OnModifierDisabled(EActorModifierCoreDisableReason InReason) override;
+	virtual void Apply() override;
+	virtual void OnModifiedActorTransformed() override;
+	//~ End UActorModifierCoreBase
+	
+	void OnReferenceActorsChanged();
+
+	//~ Begin IAvaTransformUpdatedExtension
+	virtual void OnTransformUpdated(AActor* InActor, bool bInParentMoved) override;
+	//~ End IAvaTransformUpdatedExtension
+	
+	void SetTransformExtensionReferenceActors();
+	
+	/** Editable set of reference actors and weights used to calculate the average location for this actor */
+	UPROPERTY(EditAnywhere, BlueprintReadWrite, Setter="SetReferenceActors", Getter="GetReferenceActors", Category="AlignBetween", meta=(AllowPrivateAccess="true"))
+	TSet<FAvaAlignBetweenWeightedActor> ReferenceActors;
+};
diff --git a/Engine/Plugins/Experimental/Avalanche/Source/AvalancheModifiers/Public/Modifiers/AvaArrangeBaseModifier.h b/Engine/Plugins/Experimental/Avalanche/Source/AvalancheModifiers/Public/Modifiers/AvaArrangeBaseModifier.h
new file mode 100644
--- /dev/null
+++ b/Engine/Plugins/Experimental/Avalanche/Source/AvalancheModifiers/Public/Modifiers/AvaArrangeBaseModifier.h
@@ -0,0 +1,46 @@
+// Copyright Epic Games, Inc. All Rights Reserved.
+
+#pragma once
+
+#include "AvaAttachmentBaseModifier.h"
+#include "Extensions/AvaRenderStateUpdateModifierExtension.h"
+#include "Extensions/AvaTransformUpdateModifierExtension.h"
+#include "AvaArrangeBaseModifier.generated.h"
+
+class AActor;
+
+/**
+ * Abstract base class for modifiers dealing with arrangement and attachment actors on self
+ */
+UCLASS(Abstract)
+class AVALANCHEMODIFIERS_API UAvaArrangeBaseModifier : public UAvaAttachmentBaseModifier
+	, public IAvaRenderStateUpdateHandler
+	, public IAvaTransformUpdateHandler
+{
+	GENERATED_BODY()
+	
+protected:
+	//~ Begin UActorModifierCoreBase
+	virtual void OnModifierAdded(EActorModifierCoreEnableReason InReason) override;
+	virtual void OnModifierDisabled(EActorModifierCoreDisableReason InReason) override;
+	virtual void OnModifiedActorTransformed() override;
+	//~ End UActorModifierCoreBase
+	
+	//~ Begin IAvaSceneTreeUpdateModifierExtension
+	virtual void OnSceneTreeTrackedActorChildrenChanged(int32 InIdx, const TSet<TWeakObjectPtr<AActor>>& InPreviousChildrenActors, const TSet<TWeakObjectPtr<AActor>>& InNewChildrenActors) override;
+	virtual void OnSceneTreeTrackedActorDirectChildrenChanged(int32 InIdx, const TArray<TWeakObjectPtr<AActor>>& InPreviousChildrenActors, const TArray<TWeakObjectPtr<AActor>>& InNewChildrenActors) override;
+	//~ End IAvaSceneTreeUpdateModifierExtension
+
+	//~ Begin IAvaRenderStateUpdateExtension
+	virtual void OnRenderStateUpdated(AActor* InActor, UActorComponent* InComponent) override {}
+	virtual void OnActorVisibilityChanged(AActor* InActor) override {}
+	//~ End IAvaRenderStateUpdateExtension
+
+	//~ Begin IAvaTransformUpdateExtension
+	virtual void OnTransformUpdated(AActor* InActor, bool bInParentMoved) override {}
+	//~ End IAvaTransformUpdateExtension
+
+	/** Used to track self modified actor for changes */
+	UPROPERTY(Transient, DuplicateTransient, TextExportTransient)
+	FAvaSceneTreeActor ReferenceActor;
+};
diff --git a/Engine/Plugins/Experimental/Avalanche/Source/AvalancheModifiers/Public/Modifiers/AvaAttachmentBaseModifier.h b/Engine/Plugins/Experimental/Avalanche/Source/AvalancheModifiers/Public/Modifiers/AvaAttachmentBaseModifier.h
new file mode 100644
--- /dev/null
+++ b/Engine/Plugins/Experimental/Avalanche/Source/AvalancheModifiers/Public/Modifiers/AvaAttachmentBaseModifier.h
@@ -0,0 +1,36 @@
+// Copyright Epic Games, Inc. All Rights Reserved.
+
+#pragma once
+
+#include "AvaBaseModifier.h"
+#include "Extensions/AvaSceneTreeUpdateModifierExtension.h"
+#include "AvaAttachmentBaseModifier.generated.h"
+
+class AActor;
+
+/**
+ * Abstract base class for all modifiers that deal with attachments
+ */
+UCLASS(Abstract)
+class AVALANCHEMODIFIERS_API UAvaAttachmentBaseModifier : public UAvaBaseModifier
+	, public IAvaSceneTreeUpdateHandler
+{
+	GENERATED_BODY()
+
+protected:
+	//~ Begin UActorModifierCoreBase
+	virtual void OnModifierCDOSetup(FActorModifierCoreMetadata& InMetadata) override;
+	virtual void OnModifierAdded(EActorModifierCoreEnableReason InReason) override;
+	//~ End UActorModifierCoreBase
+	
+	//~ Begin IAvaSceneTreeUpdateModifierExtension
+	virtual void OnSceneTreeTrackedActorChanged(int32 InIdx, AActor* InPreviousActor, AActor* InNewActor) override {}
+	virtual void OnSceneTreeTrackedActorChildrenChanged(int32 InIdx, const TSet<TWeakObjectPtr<AActor>>& InPreviousChildrenActors, const TSet<TWeakObjectPtr<AActor>>& InNewChildrenActors) override {}
+	virtual void OnSceneTreeTrackedActorDirectChildrenChanged(int32 InIdx, const TArray<TWeakObjectPtr<AActor>>& InPreviousChildrenActors, const TArray<TWeakObjectPtr<AActor>>& InNewChildrenActors) override {}
+	virtual void OnSceneTreeTrackedActorParentChanged(int32 InIdx, const TArray<TWeakObjectPtr<AActor>>& InPreviousParentActor, const TArray<TWeakObjectPtr<AActor>>& InNewParentActor) override {}
+	virtual void OnSceneTreeTrackedActorRearranged(int32 InIdx, AActor* InRearrangedActor) override {}
+	//~ End IAvaSceneTreeUpdateModifierExtension
+
+	UPROPERTY(Transient, DuplicateTransient, TextExportTransient)
+	TSet<TWeakObjectPtr<AActor>> ChildrenActorsWeak;
+};
diff --git a/Engine/Plugins/Experimental/Avalanche/Source/AvalancheModifiers/Public/Modifiers/AvaAutoFollowModifier.h b/Engine/Plugins/Experimental/Avalanche/Source/AvalancheModifiers/Public/Modifiers/AvaAutoFollowModifier.h
new file mode 100644
--- /dev/null
+++ b/Engine/Plugins/Experimental/Avalanche/Source/AvalancheModifiers/Public/Modifiers/AvaAutoFollowModifier.h
@@ -0,0 +1,206 @@
+// Copyright Epic Games, Inc. All Rights Reserved.
+
+#pragma once
+
+#include "AvaDefs.h"
+#include "AvaAttachmentBaseModifier.h"
+#include "AvaModifiersActorUtils.h"
+#include "Components/SceneComponent.h"
+#include "Extensions/AvaRenderStateUpdateModifierExtension.h"
+#include "Extensions/AvaTransformUpdateModifierExtension.h"
+#include "AvaAutoFollowModifier.generated.h"
+
+class AActor;
+class UActorComponent;
+
+/**
+ * Moves the modifying actor along with a specified actor relative to the specified actor's bounds.
+ */
+UCLASS(BlueprintType)
+class AVALANCHEMODIFIERS_API UAvaAutoFollowModifier : public UAvaAttachmentBaseModifier
+	, public IAvaTransformUpdateHandler
+	, public IAvaRenderStateUpdateHandler
+{
+	GENERATED_BODY()
+	
+public:
+	//~ Begin UObject
+	virtual void PostLoad() override;
+#if WITH_EDITOR
+	virtual void PostEditChangeProperty(FPropertyChangedEvent& PropertyChangedEvent) override;
+	virtual void PostTransacted(const FTransactionObjectEvent& TransactionEvent) override;
+#endif
+	//~ End UObject
+
+	UFUNCTION(BlueprintCallable, Category = "Motion Design|Modifiers|AutoFollow")
+	void SetReferenceActor(const FAvaSceneTreeActor& InReferenceActor);
+
+	UFUNCTION(BlueprintPure, Category = "Motion Design|Modifiers|AutoFollow")
+	const FAvaSceneTreeActor& GetReferenceActor() const
+	{
+		return ReferenceActor;
+	}
+
+	UFUNCTION(BlueprintPure, Category = "Motion Design|Modifiers|AutoFollow")
+	int32 GetFollowedAxis() const
+	{
+		return FollowedAxis;
+	}
+	
+	UFUNCTION(BlueprintCallable, Category = "Motion Design|Modifiers|AutoFollow")
+	void SetFollowedAxis(int32 InFollowedAxis);
+
+	/** Gets the distance from this actor to the followed actor. */
+	UFUNCTION(BlueprintPure, Category = "Motion Design|Modifiers|AutoFollow")
+	FVector GetDefaultDistance() const
+	{
+		return DefaultDistance;
+	}
+
+	/** Sets the distance from this actor to the followed actor. */
+	UFUNCTION(BlueprintCallable, Category = "Motion Design|Modifiers|AutoFollow")
+	void SetDefaultDistance(const FVector& NewDefaultDistance);
+
+	/** Gets the maximum distance from this actor to the followed actor. */
+	UFUNCTION(BlueprintPure, Category = "Motion Design|Modifiers|AutoFollow")
+	FVector GetMaxDistance() const
+	{
+		return MaxDistance;
+	}
+
+	/** Sets the maximum distance from this actor to the followed actor. */
+	UFUNCTION(BlueprintCallable, Category = "Motion Design|Modifiers|AutoFollow")
+	void SetMaxDistance(const FVector& NewMaxDistance);
+
+	/** Gets the percent % progress from the maximum distance to the default distance. */
+	UFUNCTION(BlueprintPure, Category = "Motion Design|Modifiers|AutoFollow")
+	FVector GetProgress() const
+	{
+		return Progress;
+	}
+
+	/** Sets the percent % progress from the maximum distance to the default distance. */
+	UFUNCTION(BlueprintCallable, Category = "Motion Design|Modifiers|AutoFollow")
+	void SetProgress(const FVector& NewProgress);
+
+	/** Gets the alignment for the followed actor's center. */
+	UFUNCTION(BlueprintPure, Category = "Motion Design|Modifiers|AutoFollow")
+	FAvaAnchorAlignment GetFollowedAlignment() const
+	{
+		return FollowedAlignment;
+	}
+
+	/** Sets the alignment for the followed actor's center. */
+	UFUNCTION(BlueprintCallable, Category = "Motion Design|Modifiers|AutoFollow")
+	void SetFollowedAlignment(const FAvaAnchorAlignment& NewFollowedAlignment);
+
+	/** Gets the alignment for this actor's center. */
+	UFUNCTION(BlueprintPure, Category = "Motion Design|Modifiers|AutoFollow")
+	FAvaAnchorAlignment GetLocalAlignment() const
+	{
+		return LocalAlignment;
+	}
+
+	/** Sets the alignment for this actor's center. */
+	UFUNCTION(BlueprintCallable, Category = "Motion Design|Modifiers|AutoFollow")
+	void SetLocalAlignment(const FAvaAnchorAlignment& NewLocalAlignment);
+
+	/** Gets the axis direction to offset this actor from the followed actor's bounds. */
+	UFUNCTION(BlueprintPure, Category = "Motion Design|Modifiers|AutoFollow")
+	FVector GetOffsetAxis() const
+	{
+		return OffsetAxis;
+	}
+
+	/** Sets the axis direction to offset this actor from the followed actor's bounds. */
+	UFUNCTION(BlueprintCallable, Category = "Motion Design|Modifiers|AutoFollow")
+	void SetOffsetAxis(const FVector& NewOffsetAxis);
+	
+protected:
+	//~ Begin UActorModifierCoreBase
+	virtual bool IsModifierDirtyable() const override;
+	virtual void OnModifierCDOSetup(FActorModifierCoreMetadata& InMetadata) override;
+	virtual void OnModifierAdded(EActorModifierCoreEnableReason InReason) override;
+	virtual void OnModifierEnabled(EActorModifierCoreEnableReason InReason) override;
+	virtual void OnModifierDisabled(EActorModifierCoreDisableReason InReason) override;
+	virtual void OnModifiedActorTransformed() override;
+	virtual void Apply() override;
+	//~ End UActorModifierCoreBase
+	
+	//~ Begin IAvaTransformUpdatedExtension
+	virtual void OnTransformUpdated(AActor* InActor, bool bInParentMoved) override;
+	//~ End IAvaTransformUpdatedExtension
+
+	//~ Begin IAvaRenderStateUpdateExtension
+	virtual void OnRenderStateUpdated(AActor* InActor, UActorComponent* InComponent) override;
+	virtual void OnActorVisibilityChanged(AActor* InActor) override {}
+	//~ End IAvaRenderStateUpdateExtension
+
+	//~ Begin IAvaSceneTreeUpdateModifierExtension
+	virtual void OnSceneTreeTrackedActorChanged(int32 InIdx, AActor* InPreviousActor, AActor* InNewActor) override;
+	virtual void OnSceneTreeTrackedActorChildrenChanged(int32 InIdx, const TSet<TWeakObjectPtr<AActor>>& InPreviousChildrenActors, const TSet<TWeakObjectPtr<AActor>>& InNewChildrenActors) override;
+	//~ End IAvaSceneTreeUpdateModifierExtension
+
+	void OnReferenceActorChanged();
+
+	void OnFollowedAxisChanged();
+
+	UPROPERTY(EditAnywhere, BlueprintReadWrite, Setter="SetReferenceActor", Getter="GetReferenceActor", Category="AutoFollow", meta=(ShowOnlyInnerProperties, AllowPrivateAccess="true"))
+	FAvaSceneTreeActor ReferenceActor;
+
+	/** The method for finding a reference actor based on it's position in the parent's hierarchy. */
+	UPROPERTY(meta=(DeprecatedProperty, DeprecationMessage="Use ReferenceActor instead"))
+	EAvaReferenceContainer ReferenceContainer_DEPRECATED = EAvaReferenceContainer::Other;
+
+	/** The actor being followed by the modifier. This is user selectable if the Reference Container is set to "Other". */
+	UPROPERTY(meta=(DeprecatedProperty, DeprecationMessage="Use ReferenceActor instead"))
+	TWeakObjectPtr<AActor> ReferenceActorWeak_DEPRECATED = nullptr;
+	
+	/** If true, will search for the next visible actor based on the selected reference container. */
+	UPROPERTY(meta=(DeprecatedProperty, DeprecationMessage="Use ReferenceActor instead"))
+	bool bIgnoreHiddenActors_DEPRECATED = false;
+
+	/** Which axis should we follow */
+	UPROPERTY(EditInstanceOnly, BlueprintReadWrite, Setter="SetFollowedAxis", Getter="GetFollowedAxis", Category="AutoFollow", meta=(Bitmask, BitmaskEnum="/Script/AvalancheModifiers.EAvaModifiersAxis", AllowPrivateAccess="true"))
+	int32 FollowedAxis = static_cast<int32>(
+		EAvaModifiersAxis::Y |
+		EAvaModifiersAxis::Z
+	);
+	
+	/** Based on followed axis, the direction to offset this actor from the followed actor's bounds. */
+	UPROPERTY(EditAnywhere, BlueprintReadWrite, Setter="SetOffsetAxis", Getter="GetOffsetAxis", Interp, Category="AutoFollow", meta=(UIMin="-1.0", UIMax="1.0", AllowPrivateAccess="true"))
+	FVector OffsetAxis = FVector(0, 1, 0);
+	
+	/** The alignment for the followed actor's center. */
+	UPROPERTY(EditAnywhere, BlueprintReadWrite, Setter="SetFollowedAlignment", Getter="GetFollowedAlignment", Category="AutoFollow", meta=(AllowPrivateAccess="true"))
+	FAvaAnchorAlignment FollowedAlignment;
+	
+	/** The alignment for this actor's center. */
+	UPROPERTY(EditAnywhere, BlueprintReadWrite, Setter="SetLocalAlignment", Getter="GetLocalAlignment", Category="AutoFollow", meta=(AllowPrivateAccess="true"))
+	FAvaAnchorAlignment LocalAlignment;
+
+	/** The distance from this actor to the followed actor. */
+	UPROPERTY(EditAnywhere, BlueprintReadWrite, Setter="SetDefaultDistance", Getter="GetDefaultDistance", Interp, Category="AutoFollow", meta=(AllowPrivateAccess="true"))
+	FVector DefaultDistance;
+
+	/** The maximum distance from this actor to the followed actor. */
+	UPROPERTY(EditAnywhere, BlueprintReadWrite, Setter="SetMaxDistance", Getter="GetMaxDistance", Interp, Category="AutoFollow", meta=(AllowPrivateAccess="true"))
+	FVector MaxDistance;
+
+	/** Percent % progress from the maximum distance to the default distance. */
+	UPROPERTY(EditAnywhere, BlueprintReadWrite, Setter="SetProgress", Getter="GetProgress", Interp, Category="AutoFollow", meta=(ClampMin="0.0", UIMin="0.0", ClampMax="100.0", UIMax="100.0", AllowPrivateAccess="true"))
+	FVector Progress;
+
+private:
+	UPROPERTY(Transient, DuplicateTransient, TextExportTransient)
+	FVector CachedFollowLocation = FVector::ZeroVector;
+
+	UPROPERTY(Transient, DuplicateTransient, TextExportTransient)
+	FBox CachedReferenceBounds = FBox(EForceInit::ForceInit);
+	
+	UPROPERTY(Transient, DuplicateTransient, TextExportTransient)
+	FBox CachedModifiedBounds = FBox(EForceInit::ForceInit);
+
+	UPROPERTY()
+	bool bDeprecatedPropertiesMigrated = false;
+};
\ No newline at end of file
diff --git a/Engine/Plugins/Experimental/Avalanche/Source/AvalancheModifiers/Public/Modifiers/AvaAutoSizeModifier.h b/Engine/Plugins/Experimental/Avalanche/Source/AvalancheModifiers/Public/Modifiers/AvaAutoSizeModifier.h
new file mode 100644
--- /dev/null
+++ b/Engine/Plugins/Experimental/Avalanche/Source/AvalancheModifiers/Public/Modifiers/AvaAutoSizeModifier.h
@@ -0,0 +1,169 @@
+// Copyright Epic Games, Inc. All Rights Reserved.
+
+#pragma once
+
+#include "AvaDefs.h"
+#include "AvaGeometryBaseModifier.h"
+#include "AvaModifiersActorUtils.h"
+#include "Extensions/AvaRenderStateUpdateModifierExtension.h"
+#include "Extensions/AvaSceneTreeUpdateModifierExtension.h"
+#include "Extensions/AvaTransformUpdateModifierExtension.h"
+#include "Layout/Margin.h"
+#include "AvaAutoSizeModifier.generated.h"
+
+class UActorComponent;
+class UAvaShape2DDynMeshBase;
+
+UENUM()
+enum class EAvaAutoSizeFitMode : uint8
+{
+	WidthAndHeight,
+	WidthOnly,
+	HeightOnly
+};
+
+/**
+ * Adapts the modified actor geometry size/scale and position so that it acts as a background for a specified actor
+ */
+UCLASS(BlueprintType)
+class AVALANCHEMODIFIERS_API UAvaAutoSizeModifier : public UAvaGeometryBaseModifier
+	, public IAvaTransformUpdateHandler
+	, public IAvaRenderStateUpdateHandler
+	, public IAvaSceneTreeUpdateHandler
+{
+	GENERATED_BODY()
+
+public:
+	//~ Begin UObject
+	virtual void PostLoad() override;
+#if WITH_EDITOR
+	virtual void PostEditChangeProperty(FPropertyChangedEvent& PropertyChangedEvent) override;
+#endif
+	//~ End UObject
+
+	UFUNCTION(BlueprintCallable, Category = "Motion Design|Modifiers|AutoSize")
+	void SetReferenceActor(const FAvaSceneTreeActor& InReferenceActor);
+
+	UFUNCTION(BlueprintPure, Category = "Motion Design|Modifiers|AutoSize")
+	const FAvaSceneTreeActor& GetReferenceActor() const
+	{
+		return ReferenceActor;
+	}
+
+	UFUNCTION(BlueprintPure, Category = "Motion Design|Modifiers|AutoSize")
+	int32 GetFollowedAxis() const
+	{
+		return FollowedAxis;
+	}
+
+	UFUNCTION(BlueprintCallable, Category = "Motion Design|Modifiers|AutoSize")
+	void SetFollowedAxis(int32 InFollowedAxis);
+
+	/** Gets the actor affecting the modifier. This is user selectable if the Reference Container is set to "Other". */
+	UFUNCTION(BlueprintPure, Category = "Motion Design|Modifiers|AutoSize")
+	const FMargin& GetPadding() const
+	{
+		return Padding;
+	}
+
+	/** Sets the actor affecting the modifier. This is user selectable if the Reference Container is set to "Other". */
+	UFUNCTION(BlueprintCallable, Category = "Motion Design|Modifiers|AutoSize")
+	void SetPadding(const FMargin& InPadding);
+
+	UFUNCTION(BlueprintPure, Category = "Motion Design|Modifiers|AutoSize")
+	EAvaAutoSizeFitMode GetFitMode() const
+	{
+		return FitMode;
+	}
+
+	UFUNCTION(BlueprintCallable, Category = "Motion Design|Modifiers|AutoSize")
+	void SetFitMode(const EAvaAutoSizeFitMode InFitMode);
+
+	UFUNCTION(BlueprintCallable, Category = "Motion Design|Modifiers|AutoSize")
+	void SetIncludeChildren(bool bInIncludeChildren);
+
+	UFUNCTION(BlueprintPure, Category = "Motion Design|Modifiers|AutoSize")
+	bool GetIncludeChildren() const
+	{
+		return bIncludeChildren;
+	}
+
+protected:
+	//~ Begin UActorModifierCoreBase
+	virtual bool IsModifierDirtyable() const override;
+	virtual void OnModifierCDOSetup(FActorModifierCoreMetadata& InMetadata) override;
+	virtual void OnModifiedActorTransformed() override;
+	virtual void OnModifierAdded(EActorModifierCoreEnableReason InReason) override;
+	virtual void OnModifierEnabled(EActorModifierCoreEnableReason InReason) override;
+	virtual void OnModifierDisabled(EActorModifierCoreDisableReason InReason) override;
+	virtual void Apply() override;
+	//~ End UActorModifierCoreBase
+
+	//~ Begin IAvaTransformUpdatedExtension
+	virtual void OnTransformUpdated(AActor* InActor, bool bInParentMoved) override;
+	//~ End IAvaTransformUpdatedExtension
+
+	//~ Begin IAvaRenderStateUpdateExtension
+	virtual void OnRenderStateUpdated(AActor* InActor, UActorComponent* InComponent) override;
+	virtual void OnActorVisibilityChanged(AActor* InActor) override;
+	//~ End IAvaRenderStateUpdateExtension
+
+	//~ Begin IAvaSceneTreeUpdateModifierExtension
+	virtual void OnSceneTreeTrackedActorChanged(int32 InIdx, AActor* InPreviousActor, AActor* InNewActor) override;
+	virtual void OnSceneTreeTrackedActorChildrenChanged(int32 InIdx, const TSet<TWeakObjectPtr<AActor>>& InPreviousChildrenActors, const TSet<TWeakObjectPtr<AActor>>& InNewChildrenActors) override;
+	virtual void OnSceneTreeTrackedActorDirectChildrenChanged(int32 InIdx, const TArray<TWeakObjectPtr<AActor>>& InPreviousChildrenActors, const TArray<TWeakObjectPtr<AActor>>& InNewChildrenActors) override {}
+	virtual void OnSceneTreeTrackedActorParentChanged(int32 InIdx, const TArray<TWeakObjectPtr<AActor>>& InPreviousParentActor, const TArray<TWeakObjectPtr<AActor>>& InNewParentActor) override {}
+	virtual void OnSceneTreeTrackedActorRearranged(int32 InIdx, AActor* InRearrangedActor) override {}
+	//~ End IAvaSceneTreeUpdateModifierExtension
+
+	void OnReferenceActorChanged();
+
+	UPROPERTY(EditAnywhere, BlueprintReadWrite, Setter="SetReferenceActor", Getter="GetReferenceActor", Category="AutoSize", meta=(ShowOnlyInnerProperties, AllowPrivateAccess="true"))
+	FAvaSceneTreeActor ReferenceActor;
+
+	/** The method for finding a reference actor based on it's position in the parent's hierarchy. */
+	UPROPERTY(meta=(DeprecatedProperty, DeprecationMessage="Use ReferenceActor instead"))
+	EAvaReferenceContainer ReferenceContainer_DEPRECATED = EAvaReferenceContainer::Other;
+
+	/** The actor affecting the modifier. This is user selectable if the Reference Container is set to "Other". */
+	UPROPERTY(meta=(DeprecatedProperty, DeprecationMessage="Use ReferenceActor instead"))
+	TWeakObjectPtr<AActor> ReferenceActorWeak_DEPRECATED = nullptr;
+
+	/** If true, will search for the next visible actor based on the selected reference container. */
+	UPROPERTY(meta=(DeprecatedProperty, DeprecationMessage="Use ReferenceActor instead"))
+	bool bIgnoreHiddenActors_DEPRECATED = false;
+
+	/** Which axis should we follow, if none selected, it will not follow */
+	UPROPERTY(EditInstanceOnly, BlueprintReadWrite, Setter="SetFollowedAxis", Getter="GetFollowedAxis", Category="AutoSize", meta=(Bitmask, BitmaskEnum="/Script/AvalancheModifiers.EAvaModifiersAxis", AllowPrivateAccess="true"))
+	int32 FollowedAxis = static_cast<int32>(
+		EAvaModifiersAxis::Y |
+		EAvaModifiersAxis::Z
+	);
+
+	/* Padding around reference bounds */
+	UPROPERTY(EditAnywhere, BlueprintReadWrite, Setter="SetPadding", Getter="GetPadding", Interp, Category="AutoSize", meta=(AllowPrivateAccess="true"))
+	FMargin Padding = FMargin(0.f);
+
+	UPROPERTY(EditAnywhere, BlueprintReadWrite, Setter="SetFitMode", Getter="GetFitMode", Category="AutoSize", meta=(AllowPrivateAccess="true"))
+	EAvaAutoSizeFitMode FitMode = EAvaAutoSizeFitMode::WidthAndHeight;
+
+	/** If true, will include children bounds too and compute the new size */
+	UPROPERTY(EditAnywhere, BlueprintReadWrite, Setter="SetIncludeChildren", Getter="GetIncludeChildren", Category="AutoSize", meta=(AllowPrivateAccess="true"))
+	bool bIncludeChildren = true;
+
+private:
+	UPROPERTY()
+	FVector2D PreModifierShapeDynMesh2DSize;
+
+	UPROPERTY()
+	TWeakObjectPtr<UAvaShape2DDynMeshBase> ShapeDynMesh2DWeak;
+
+	UPROPERTY(Transient, DuplicateTransient, TextExportTransient)
+	FVector CachedFollowLocation = FVector::ZeroVector;
+
+	UPROPERTY(Transient, DuplicateTransient, TextExportTransient)
+	FBox CachedReferenceBounds = FBox(EForceInit::ForceInit);
+
+	UPROPERTY()
+	bool bDeprecatedPropertiesMigrated = false;
+};
diff --git a/Engine/Plugins/Experimental/Avalanche/Source/AvalancheModifiers/Public/Modifiers/AvaBaseModifier.h b/Engine/Plugins/Experimental/Avalanche/Source/AvalancheModifiers/Public/Modifiers/AvaBaseModifier.h
new file mode 100644
--- /dev/null
+++ b/Engine/Plugins/Experimental/Avalanche/Source/AvalancheModifiers/Public/Modifiers/AvaBaseModifier.h
@@ -0,0 +1,13 @@
+// Copyright Epic Games, Inc. All Rights Reserved.
+
+#pragma once
+
+#include "Modifiers/ActorModifierCoreBase.h"
+#include "AvaBaseModifier.generated.h"
+
+/** This is the base class for all avalanche modifiers */
+UCLASS(Abstract)
+class AVALANCHEMODIFIERS_API UAvaBaseModifier : public UActorModifierCoreBase
+{
+	GENERATED_BODY()
+};
diff --git a/Engine/Plugins/Experimental/Avalanche/Source/AvalancheModifiers/Public/Modifiers/AvaBendModifier.h b/Engine/Plugins/Experimental/Avalanche/Source/AvalancheModifiers/Public/Modifiers/AvaBendModifier.h
new file mode 100644
--- /dev/null
+++ b/Engine/Plugins/Experimental/Avalanche/Source/AvalancheModifiers/Public/Modifiers/AvaBendModifier.h
@@ -0,0 +1,73 @@
+// Copyright Epic Games, Inc. All Rights Reserved.
+
+#pragma once
+
+#include "AvaGeometryBaseModifier.h"
+#include "AvaBendModifier.generated.h"
+
+UCLASS(BlueprintType)
+class AVALANCHEMODIFIERS_API UAvaBendModifier : public UAvaGeometryBaseModifier
+{
+	GENERATED_BODY()
+	
+public:
+	//~ Begin UObject
+#if WITH_EDITOR
+	virtual void PostEditChangeProperty(FPropertyChangedEvent& PropertyChangedEvent) override;
+#endif
+	//~End UObject
+
+	UFUNCTION()
+	void SetAngle(float InAngle);
+	float GetAngle() const { return Angle; }
+
+	UFUNCTION()
+	void SetExtent(float InExtent);
+	float GetExtent() const { return Extent; }
+
+	UFUNCTION()
+	void SetBendPosition(const FVector& InBendPosition);
+	const FVector& GetBendPosition() const { return BendPosition; }
+
+	UFUNCTION()
+	void SetBendRotation(const FRotator& InBendRotation);
+	const FRotator& GetBendRotation() const { return BendRotation; }
+
+	UFUNCTION()
+	void SetSymmetricExtents(bool bInSymmetricExtents);
+	bool GetSymmetricExtents() const { return bSymmetricExtents; }
+
+	UFUNCTION()
+	void SetBidirectional(bool bInBidirectional);
+	bool GetBidirectional() const { return bBidirectional; }
+
+protected:
+	//~ Begin UActorModifierCoreBase
+	virtual void OnModifierCDOSetup(FActorModifierCoreMetadata& InMetadata) override;
+	virtual void Apply() override;
+	//~ End UActorModifierCoreBase
+	
+	void OnBendTransformChanged();
+	void OnBendOptionChanged();
+	
+	UPROPERTY(EditInstanceOnly, BlueprintReadWrite, Setter="SetBendPosition", Getter="GetBendPosition", Category="Bend", meta=(AllowPrivateAccess="true"))
+	FVector BendPosition = FVector::ZeroVector;
+	
+	UPROPERTY(EditInstanceOnly, BlueprintReadWrite, Setter="SetBendRotation", Getter="GetBendRotation", Category="Bend", meta=(AllowPrivateAccess="true"))
+	FRotator BendRotation = FRotator::ZeroRotator;
+	
+	UPROPERTY(EditInstanceOnly, BlueprintReadWrite, Setter="SetAngle", Getter="GetAngle", Category="Bend", meta=(ClampMin="-360.0", ClampMax="360.0", AllowPrivateAccess="true"))
+	float Angle = 25;
+	
+	UPROPERTY(EditInstanceOnly, BlueprintReadWrite, Setter="SetExtent", Getter="GetExtent", Category="Bend", meta=(ClampMin="0", ClampMax="1.0", AllowPrivateAccess="true"))
+	float Extent = 1.0;
+	
+	UPROPERTY(EditInstanceOnly, BlueprintReadWrite, Setter="SetSymmetricExtents", Getter="GetSymmetricExtents", Category="Bend", meta=(AllowPrivateAccess="true"))
+	bool bSymmetricExtents = true;
+	
+	UPROPERTY(EditInstanceOnly, BlueprintReadWrite, Setter="SetBidirectional", Getter="GetBidirectional", Category="Bend", meta=(AllowPrivateAccess="true"))
+	bool bBidirectional = false;
+	
+	UPROPERTY()
+	float ModifiedMeshMaxExtent = 50;
+};
diff --git a/Engine/Plugins/Experimental/Avalanche/Source/AvalancheModifiers/Public/Modifiers/AvaBevelModifier.h b/Engine/Plugins/Experimental/Avalanche/Source/AvalancheModifiers/Public/Modifiers/AvaBevelModifier.h
new file mode 100644
--- /dev/null
+++ b/Engine/Plugins/Experimental/Avalanche/Source/AvalancheModifiers/Public/Modifiers/AvaBevelModifier.h
@@ -0,0 +1,49 @@
+// Copyright Epic Games, Inc. All Rights Reserved.
+
+#pragma once
+
+#include "AvaGeometryBaseModifier.h"
+#include "AvaBevelModifier.generated.h"
+
+UCLASS(BlueprintType)
+class AVALANCHEMODIFIERS_API UAvaBevelModifier : public UAvaGeometryBaseModifier
+{
+	GENERATED_BODY()
+
+public:
+	static inline const FName BevelPolygroupLayerName = TEXT("Bevel");
+	static constexpr float MinInset = 0;
+	static constexpr int32 MinIterations = 1;
+	static constexpr int32 MaxIterations = 3;
+
+	//~ Begin UObject
+#if WITH_EDITOR
+	virtual void PostEditChangeProperty(FPropertyChangedEvent& PropertyChangedEvent) override;
+#endif
+	//~ End UObject
+	
+	UFUNCTION()
+	void SetInset(float InBevel);
+	float GetInset() const { return Inset; }
+
+	UFUNCTION()
+	void SetIterations(int32 InIterations);
+	int32 GetIterations() const { return Iterations; }
+
+protected:
+	//~ Begin UActorModifierCoreBase
+	virtual void OnModifierCDOSetup(FActorModifierCoreMetadata& InMetadata) override;
+	virtual void Apply() override;
+	//~ End UActorModifierCoreBase
+	
+	void OnInsetChanged();
+	void OnIterationsChanged();
+
+	float GetMaxBevel() const;
+	
+	UPROPERTY(EditInstanceOnly, BlueprintReadWrite, Setter="SetInset", Getter="GetInset", Category="Bevel", meta=(ClampMin="0.0", AllowPrivateAccess="true"))
+	float Inset = 1.0f;
+
+	UPROPERTY(EditInstanceOnly, BlueprintReadWrite, Setter="SetIterations", Getter="GetIterations", Category="Bevel", meta=(ClampMin="1", ClampMax="3", AllowPrivateAccess="true"))
+	int32 Iterations = 1;
+};
diff --git a/Engine/Plugins/Experimental/Avalanche/Source/AvalancheModifiers/Public/Modifiers/AvaDynamicMeshConverterModifier.h b/Engine/Plugins/Experimental/Avalanche/Source/AvalancheModifiers/Public/Modifiers/AvaDynamicMeshConverterModifier.h
new file mode 100644
--- /dev/null
+++ b/Engine/Plugins/Experimental/Avalanche/Source/AvalancheModifiers/Public/Modifiers/AvaDynamicMeshConverterModifier.h
@@ -0,0 +1,174 @@
+// Copyright Epic Games, Inc. All Rights Reserved.
+
+#pragma once
+
+#include "AvaGeometryBaseModifier.h"
+#include "DynamicMesh/DynamicMesh3.h"
+#include "Templates/SubclassOf.h"
+#include "AvaDynamicMeshConverterModifier.generated.h"
+
+class UBrushComponent;
+class UDynamicMeshComponent;
+class UPrimitiveComponent;
+class UProceduralMeshComponent;
+class USkeletalMeshComponent;
+class UStaticMeshComponent;
+
+/** Components that can be converted to dynamic mesh */
+UENUM(BlueprintType, meta = (Bitflags, UseEnumValuesAsMaskValuesInEditor = "true"))
+enum class EAvaDynamicMeshConverterModifierType : uint8
+{
+	None = 0 UMETA(Hidden),
+	StaticMeshComponent = 1 << 0,
+	DynamicMeshComponent = 1 << 1,
+	SkeletalMeshComponent = 1 << 2,
+	BrushComponent = 1 << 3,
+	ProceduralMeshComponent = 1 << 4
+};
+ENUM_CLASS_FLAGS(EAvaDynamicMeshConverterModifierType);
+
+UENUM(BlueprintType)
+enum class EAvaDynamicMeshConverterModifierFilter : uint8
+{
+	None,
+	Include,
+	Exclude
+};
+
+USTRUCT()
+struct FAvaDynamicMeshConverterModifierComponentState
+{
+	GENERATED_BODY()
+
+	explicit FAvaDynamicMeshConverterModifierComponentState() {}
+	explicit FAvaDynamicMeshConverterModifierComponentState(UPrimitiveComponent* InPrimitiveComponent);
+	
+	/** The component we are converting to dynamic mesh */
+	UPROPERTY()
+	TWeakObjectPtr<UPrimitiveComponent> Component = nullptr;
+
+	/** The default visibility of the actor converted component in game */
+	UPROPERTY()
+	bool bActorHiddenInGame = true;
+
+	/** The default visibility of the actor converted component in editor */
+	UPROPERTY()
+	bool bActorHiddenInEditor = true;
+
+	/** The default visibility of the actor converted component in game */
+	UPROPERTY()
+	bool bComponentHiddenInGame = false;
+
+	/** The default visibility of the converted component in editor */
+	UPROPERTY()
+	bool bComponentVisible = true;
+	/** The component converted dynamic mesh*/
+	UE::Geometry::FDynamicMesh3 Mesh;
+};
+
+UCLASS(BlueprintType)
+class AVALANCHEMODIFIERS_API UAvaDynamicMeshConverterModifier : public UAvaGeometryBaseModifier
+{
+	GENERATED_BODY()
+
+public:
+	//~ Begin UObject
+#if WITH_EDITOR
+	virtual void PostEditChangeProperty(FPropertyChangedEvent& PropertyChangedEvent) override;
+	/** Export current dynamic mesh to static mesh asset */
+	UFUNCTION(CallInEditor, Category="DynamicMeshConverter")
+	virtual void ConvertToStaticMeshAsset();
+#endif
+	//~ End UObject
+
+	void SetSourceActorWeak(const TWeakObjectPtr<AActor>& InActor);
+	TWeakObjectPtr<AActor> GetSourceActorWeak() const { return SourceActorWeak; }
+
+	void SetComponentType(int32 InComponentType);
+	int32 GetComponentType() const { return ComponentType; }
+
+	void SetFilterActorMode(EAvaDynamicMeshConverterModifierFilter InFilter);
+	EAvaDynamicMeshConverterModifierFilter GetFilterActorMode() const { return FilterActorMode; }
+
+	void SetFilterActorClasses(const TSet<TSubclassOf<AActor>>& InClasses);
+	const TSet<TSubclassOf<AActor>>& GetFilterActorClasses() const { return FilterActorClasses; }
+
+	void SetIncludeAttachedActors(bool bInInclude);
+	bool GetIncludeAttachedActors() const { return bIncludeAttachedActors; }
+
+	void SetHideConvertedMesh(bool bInHide);
+	bool GetHideConvertedMesh() const { return bHideConvertedMesh; }
+	
+protected:
+	//~ Begin UActorModifierCoreBase
+	virtual void OnModifierCDOSetup(FActorModifierCoreMetadata& InMetadata) override;
+	virtual void OnModifierAdded(EActorModifierCoreEnableReason InReason) override;
+	virtual void OnModifierEnabled(EActorModifierCoreEnableReason InReason) override;
+	virtual void RestorePreState() override;
+	virtual void Apply() override;
+	virtual void OnModifierRemoved(EActorModifierCoreDisableReason InReason) override;
+	//~ End UActorModifierCoreBase
+	
+	void OnSourceActorChanged();
+
+	void ConvertComponents(TArray<FAvaDynamicMeshConverterModifierComponentState>& OutResults) const;
+	bool HasFlag(EAvaDynamicMeshConverterModifierType InFlag) const;
+	
+	void AddDynamicMeshComponent();
+	void RemoveDynamicMeshComponent();
+
+	void GetFilteredActors(TArray<AActor*>& OutActors) const;
+	
+	void GetStaticMeshComponents(const TArray<AActor*>& InActors, TArray<UStaticMeshComponent*>& OutComponents) const;
+	void GetDynamicMeshComponents(const TArray<AActor*>& InActors, TArray<UDynamicMeshComponent*>& OutComponents) const;
+	void GetSkeletalMeshComponents(const TArray<AActor*>& InActors, TArray<USkeletalMeshComponent*>& OutComponents) const;
+	void GetBrushComponents(const TArray<AActor*>& InActors, TArray<UBrushComponent*>& OutComponents) const;
+	void GetProceduralMeshComponents(const TArray<AActor*>& InActors, TArray<UProceduralMeshComponent*>& OutComponents) const;
+
+	/** Triggers the update of the mesh */
+	UFUNCTION(CallInEditor, Category="DynamicMeshConverter")
+	void ConvertMesh();
+
+	/** What actor should we copy from, by default is self */
+	UPROPERTY(EditInstanceOnly, BlueprintReadWrite, Setter="SetSourceActorWeak", Getter="GetSourceActorWeak", Category="DynamicMeshConverter", meta=(DisplayName="SourceActor", AllowPrivateAccess="true"))
+	TWeakObjectPtr<AActor> SourceActorWeak = GetModifiedActor();
+	
+	/** Which components should we take into account for the conversion */
+	UPROPERTY(EditInstanceOnly, BlueprintReadWrite, Setter="SetComponentType", Getter="GetComponentType", Category="DynamicMeshConverter", meta=(Bitmask, BitmaskEnum="/Script/AvalancheModifiers.EAvaDynamicMeshConverterModifierType", AllowPrivateAccess="true"))
+	int32 ComponentType = static_cast<int32>(
+		EAvaDynamicMeshConverterModifierType::StaticMeshComponent |
+		EAvaDynamicMeshConverterModifierType::DynamicMeshComponent |
+		EAvaDynamicMeshConverterModifierType::SkeletalMeshComponent |
+		EAvaDynamicMeshConverterModifierType::BrushComponent |
+		EAvaDynamicMeshConverterModifierType::ProceduralMeshComponent);
+
+	/** Actor filter mode : none, include or exclude specific actor class */
+	UPROPERTY(EditInstanceOnly, BlueprintReadWrite, Setter="SetFilterActorMode", Getter="GetFilterActorMode", Category="DynamicMeshConverter", meta=(AllowPrivateAccess="true"))
+	EAvaDynamicMeshConverterModifierFilter FilterActorMode = EAvaDynamicMeshConverterModifierFilter::None;
+
+	/** Actor class to use as filter when gathering actors to convert */
+	UPROPERTY(EditInstanceOnly, BlueprintReadWrite, Setter="SetFilterActorClasses", Getter="GetFilterActorClasses", Category="DynamicMeshConverter", meta=(EditCondition="FilterActorMode != EAvaDynamicMeshConverterModifierFilter::None", AllowPrivateAccess="true"))
+	TSet<TSubclassOf<AActor>> FilterActorClasses;
+
+	/** Checks and convert all attached actors below this actor */
+	UPROPERTY(EditInstanceOnly, BlueprintReadWrite, Setter="SetIncludeAttachedActors", Getter="GetIncludeAttachedActors", Category="DynamicMeshConverter", meta=(AllowPrivateAccess="true"))
+	bool bIncludeAttachedActors = true;
+
+	/** Change visibility of source mesh once they are converted to dynamic mesh, by default will convert itself so hide converted mesh is true */
+	UPROPERTY(EditInstanceOnly, BlueprintReadWrite, Setter="SetHideConvertedMesh", Getter="GetHideConvertedMesh", Category="DynamicMeshConverter", meta=(AllowPrivateAccess="true"))
+	bool bHideConvertedMesh = true;
+
+	/** Did we create the dynamic mesh component from this modifier or retrieved it */
+	UPROPERTY()
+	bool bComponentCreated = false;
+
+	/** Components converted to dynamic mesh */
+	UPROPERTY()
+	TArray<FAvaDynamicMeshConverterModifierComponentState> ConvertedComponents;
+	
+	/** Cached mesh to set state without reconverting again */
+	TOptional<UE::Geometry::FDynamicMesh3> ConvertedMesh;
+
+	/** Do we convert mesh on next execution */
+	bool bConvertMesh = false;
+};
diff --git a/Engine/Plugins/Experimental/Avalanche/Source/AvalancheModifiers/Public/Modifiers/AvaExtrudeModifier.h b/Engine/Plugins/Experimental/Avalanche/Source/AvalancheModifiers/Public/Modifiers/AvaExtrudeModifier.h
new file mode 100644
--- /dev/null
+++ b/Engine/Plugins/Experimental/Avalanche/Source/AvalancheModifiers/Public/Modifiers/AvaExtrudeModifier.h
@@ -0,0 +1,71 @@
+// Copyright Epic Games, Inc. All Rights Reserved.
+
+#pragma once
+
+#include "AvaGeometryBaseModifier.h"
+#include "AvaExtrudeModifier.generated.h"
+
+UENUM()
+enum class EAvaExtrudeMode : uint8
+{
+	Opposite,
+	Front,
+	Symmetrical
+};
+
+/** This modifier extrude triangles from a 2D shape with a specific depth and optionally closes the back */
+UCLASS(BlueprintType)
+class AVALANCHEMODIFIERS_API UAvaExtrudeModifier : public UAvaGeometryBaseModifier
+{
+	GENERATED_BODY()
+
+public:
+	static inline const FName ExtrudePolygroupLayerName = TEXT("ExtrudeSide");
+	static inline const FName BackPolygroupLayerName = TEXT("ExtrudeBack");
+	
+	//~ Begin UObject
+	virtual void Serialize(FArchive& InArchive) override;
+#if WITH_EDITOR
+	virtual void PostEditChangeProperty(FPropertyChangedEvent& PropertyChangedEvent) override;
+#endif
+	//~ End UObject
+	
+	UFUNCTION()
+	void SetDepth(float InDepth);
+	float GetDepth() const { return Depth; }
+	
+	UFUNCTION()
+	void SetCloseBack(bool bInCloseBack);
+	bool GetCloseBack() const { return bCloseBack; }
+
+	UFUNCTION()
+	void SetExtrudeMode(EAvaExtrudeMode InExtrudeMode);
+	EAvaExtrudeMode GetExtrudeMode() const { return ExtrudeMode; }
+
+protected:
+	//~ Begin UActorModifierCoreBase
+	virtual void OnModifierCDOSetup(FActorModifierCoreMetadata& InMetadata) override;
+	virtual void Apply() override;
+	//~ End UActorModifierCoreBase
+	
+	void OnDepthChanged();
+	void OnCloseBackChanged();
+	void OnExtrudeModeChanged();
+	FVector GetExtrudeDirection() const;
+	
+	/** Handles mesh depth to extrude primary section */
+	UPROPERTY(EditInstanceOnly, BlueprintReadWrite, Setter="SetDepth", Getter="GetDepth", Category="Extrude", meta=(ClampMin="0", DisplayName="Extrude depth", Units="Centimeters", AllowPrivateAccess="true"))
+	float Depth = 30.f;
+	
+	/** Closes the back of the extrude for a 2D shape for example */
+	UPROPERTY(EditInstanceOnly, BlueprintReadWrite, Setter="SetCloseBack", Getter="GetCloseBack", Category="Extrude", meta=(DisplayName="Close back", AllowPrivateAccess="true"))
+	bool bCloseBack = true;
+	
+	/** Moves the mesh in the opposite extrude direction by the depth distance */
+	UPROPERTY(meta=(DeprecatedProperty, DeprecationMessage="Use ExtrudeMode instead"))
+	bool bMoveMeshOppositeDirection_DEPRECATED = true;
+
+	/** Specifies the Extrude direction */
+	UPROPERTY(EditInstanceOnly, BlueprintReadWrite, Setter="SetExtrudeMode", Getter="GetExtrudeMode", Category="Extrude", meta=(DisplayName = "Extrude Mode", AllowPrivateAccess="true"))
+	EAvaExtrudeMode ExtrudeMode = EAvaExtrudeMode::Opposite;
+};
diff --git a/Engine/Plugins/Experimental/Avalanche/Source/AvalancheModifiers/Public/Modifiers/AvaGlobalOpacityModifier.h b/Engine/Plugins/Experimental/Avalanche/Source/AvalancheModifiers/Public/Modifiers/AvaGlobalOpacityModifier.h
new file mode 100644
--- /dev/null
+++ b/Engine/Plugins/Experimental/Avalanche/Source/AvalancheModifiers/Public/Modifiers/AvaGlobalOpacityModifier.h
@@ -0,0 +1,46 @@
+// Copyright Epic Games, Inc. All Rights Reserved.
+
+#pragma once
+
+#include "Modifiers/AvaMaterialParameterModifier.h"
+#include "AvaGlobalOpacityModifier.generated.h"
+
+class UDMMaterialComponent;
+enum class EDMUpdateType : uint8;
+
+/** This modifier sets global opacity parameters on an actor with Material Designer Instances generated with the Material Designer */
+UCLASS(BlueprintType)
+class AVALANCHEMODIFIERS_API UAvaGlobalOpacityModifier : public UAvaMaterialParameterModifier
+{
+	GENERATED_BODY()
+
+	static inline const FName MaterialDesignerGlobalOpacityValueName = FName("VALUE_GlobalOpacity");
+	
+public:
+	UAvaGlobalOpacityModifier();
+
+	//~ Begin UObject
+#if WITH_EDITOR
+	virtual void PostEditChangeProperty(FPropertyChangedEvent& PropertyChangedEvent) override;
+#endif
+	//~ End UObject
+
+	UFUNCTION()
+	void SetGlobalOpacity(float InOpacity);
+	float GetGlobalOpacity() const { return GlobalOpacity; }
+	
+protected:
+	//~ Begin UActorModifierCoreBase
+	virtual void OnModifierCDOSetup(FActorModifierCoreMetadata& InMetadata) override;
+	//~ End UActorModifierCoreBase
+
+	void OnGlobalOpacityChanged();
+
+	virtual void OnActorMaterialAdded(UMaterialInstanceDynamic* InAdded) override;
+	virtual void OnActorMaterialRemoved(UMaterialInstanceDynamic* InRemoved) override;
+	void OnDynamicMaterialValueChanged(UDMMaterialComponent* InComponent, EDMUpdateType InUpdateType);
+	
+	/** Global opacity to set on all Material Designer Instances */
+	UPROPERTY(EditAnywhere, BlueprintReadWrite, Setter="SetGlobalOpacity", Getter="GetGlobalOpacity", Category="Material Parameter", meta=(ClampMin="0", ClampMax="1", AllowPrivateAccess="true"))
+	float GlobalOpacity = 1.f;
+};
diff --git a/Engine/Plugins/Experimental/Avalanche/Source/AvalancheModifiers/Public/Modifiers/AvaGridArrangeModifier.h b/Engine/Plugins/Experimental/Avalanche/Source/AvalancheModifiers/Public/Modifiers/AvaGridArrangeModifier.h
new file mode 100644
--- /dev/null
+++ b/Engine/Plugins/Experimental/Avalanche/Source/AvalancheModifiers/Public/Modifiers/AvaGridArrangeModifier.h
@@ -0,0 +1,90 @@
+// Copyright Epic Games, Inc. All Rights Reserved.
+
+#pragma once
+
+#include "AvaArrangeBaseModifier.h"
+#include "AvaGridArrangeModifier.generated.h"
+
+enum class EAvaCorner2D : uint8;
+
+/** Specifies a starting direction when laying out a grid arrangement. */
+UENUM(BlueprintType)
+enum class EAvaGridArrangeDirection : uint8
+{
+	Horizontal,
+	Vertical
+};
+
+/**
+ * Arranges child actors in a 2D grid format.
+ */
+UCLASS(BlueprintType)
+class AVALANCHEMODIFIERS_API UAvaGridArrangeModifier : public UAvaArrangeBaseModifier
+{
+	GENERATED_BODY()
+	
+public:
+	//~ Begin UObject
+#if WITH_EDITOR
+	virtual void PostEditChangeProperty(FPropertyChangedEvent& PropertyChangedEvent) override;
+#endif
+	//~ End UObject
+
+	/** Returns the width and height of the grid. */
+	UFUNCTION(BlueprintPure, Category = "Motion Design|Modifiers|GridArrange")
+	FIntPoint GetCount() const { return Count; }
+
+	/** Sets the width and height of the grid. */
+	UFUNCTION(BlueprintCallable, Category = "Motion Design|Modifiers|GridArrange")
+	void SetCount(const FIntPoint& NewCount);
+
+	/** Returns the distance between each horizontal and vertical child. */
+	UFUNCTION(BlueprintPure, Category = "Motion Design|Modifiers|GridArrange")
+	FVector2D GetSpread() const { return Spread; }
+
+	/** Sets the distance between each horizontal and vertical child. */
+	UFUNCTION(BlueprintCallable, Category = "Motion Design|Modifiers|GridArrange")
+	void SetSpread(const FVector2D& NewSpread);
+
+	UFUNCTION(BlueprintCallable, Category = "Motion Design|Modifiers|GridArrange")
+	void SetStartCorner(EAvaCorner2D InCorner);
+	
+	UFUNCTION(BlueprintPure, Category = "Motion Design|Modifiers|GridArrange")
+	EAvaCorner2D GetStartCorner() const { return StartCorner; }
+
+	UFUNCTION(BlueprintCallable, Category = "Motion Design|Modifiers|GridArrange")
+	void SetStartDirection(EAvaGridArrangeDirection InDirection);
+	
+	UFUNCTION(BlueprintPure, Category = "Motion Design|Modifiers|GridArrange")
+	EAvaGridArrangeDirection GetStartDirection() const { return StartDirection; }
+	
+	/** Returns the 2D corner and direction from which to start the arrangement. */
+	UFUNCTION(BlueprintPure, Category = "Motion Design|Modifiers|GridArrange")
+	void GetArrangementLayout(UPARAM(DisplayName = "StartCorner") EAvaCorner2D& OutStartCorner, UPARAM(DisplayName = "Direction") EAvaGridArrangeDirection& OutStartDirection) const;
+
+	/** Sets the 2D corner and direction from which to start the arrangement. */
+	UFUNCTION(BlueprintCallable, Category = "Motion Design|Modifiers|GridArrange")
+	void SetArrangementLayout(const EAvaCorner2D NewStartCorner, const EAvaGridArrangeDirection NewStartDirection);
+
+protected:
+	//~ Begin UActorModifierCoreBase
+	virtual void OnModifierCDOSetup(FActorModifierCoreMetadata& InMetadata) override;
+	virtual void Apply() override;
+	//~ End UActorModifierCoreBase
+	
+	/** The width and height of the grid. */
+	UPROPERTY(EditAnywhere, BlueprintReadWrite, Setter="SetCount", Getter="GetCount", Category="GridArrange", Meta=(ClampMin="1", UIMin="1", AllowPrivateAccess="true"))
+	FIntPoint Count = FIntPoint(1, 1);
+
+	/** The distance between each horizontal and vertical child. */
+	UPROPERTY(EditAnywhere, BlueprintReadWrite, Setter="SetSpread", Getter="GetSpread", Interp, Category="GridArrange", meta=(AllowPrivateAccess="true"))
+	FVector2D Spread = FVector2D::ZeroVector;
+
+	/** The 2D corner from which to start the arrangement. */
+	UPROPERTY(EditAnywhere, BlueprintReadWrite, Setter="SetStartCorner", Getter="GetStartCorner", Category="GridArrange", meta=(AllowPrivateAccess="true"))
+	EAvaCorner2D StartCorner;
+
+	/** The direction from which to start the arrangement. */
+	UPROPERTY(EditAnywhere, BlueprintReadWrite, Setter="SetStartDirection", Getter="GetStartDirection", Category="GridArrange", meta=(AllowPrivateAccess="true"))
+	EAvaGridArrangeDirection StartDirection;
+};
diff --git a/Engine/Plugins/Experimental/Avalanche/Source/AvalancheModifiers/Public/Modifiers/AvaJustifyModifier.h b/Engine/Plugins/Experimental/Avalanche/Source/AvalancheModifiers/Public/Modifiers/AvaJustifyModifier.h
new file mode 100644
--- /dev/null
+++ b/Engine/Plugins/Experimental/Avalanche/Source/AvalancheModifiers/Public/Modifiers/AvaJustifyModifier.h
@@ -0,0 +1,145 @@
+// Copyright Epic Games, Inc. All Rights Reserved.
+
+#pragma once
+
+#include "AvaArrangeBaseModifier.h"
+#include "Components/ActorComponent.h"
+#include "AvaJustifyModifier.generated.h"
+
+UENUM(BlueprintType)
+enum class EAvaJustifyVertical : uint8
+{
+	None,
+	Top,
+	Center,
+	Bottom,
+};
+
+UENUM(BlueprintType)
+enum class EAvaJustifyHorizontal : uint8
+{
+	None,
+	Left,
+	Center,
+	Right,
+};
+
+UENUM(BlueprintType)
+enum class EAvaJustifyDepth : uint8
+{
+	None,
+	Front,
+	Center,
+	Back
+};
+
+/**
+ * Justify Modifier
+ * 
+ * Aligns child actors, based on their bounding box, according to the specified justification
+ */
+UCLASS(BlueprintType)
+class AVALANCHEMODIFIERS_API UAvaJustifyModifier : public UAvaArrangeBaseModifier
+{
+	GENERATED_BODY()
+
+public:
+	//~ Begin UObject
+#if WITH_EDITOR
+	virtual void PostEditChangeProperty(FPropertyChangedEvent& InPropertyChangedEvent) override;
+#endif
+	//~ End UObject
+
+	UFUNCTION(BlueprintPure, Category = "Justify")
+	EAvaJustifyHorizontal GetHorizontalAlignment() const { return HorizontalAlignment; }
+	
+	UFUNCTION(BlueprintCallable, Category = "Justify")
+	void SetHorizontalAlignment(EAvaJustifyHorizontal InHorizontalAlignment);
+
+	UFUNCTION(BlueprintPure, Category = "Justify")
+	EAvaJustifyVertical GetVerticalAlignment() const { return VerticalAlignment; }
+	
+	UFUNCTION(BlueprintCallable, Category = "Justify")
+	void SetVerticalAlignment(EAvaJustifyVertical InVerticalAlignment);
+
+	UFUNCTION(BlueprintPure, Category = "Justify")
+	EAvaJustifyDepth GetDepthAlignment() const { return DepthAlignment; }
+	
+	UFUNCTION(BlueprintCallable, Category = "Justify")
+	void SetDepthAlignment(EAvaJustifyDepth InDepthAlignment);
+
+	UFUNCTION(BlueprintPure, Category = "Justify")
+	float GetHorizontalAnchor() const { return HorizontalAnchor; }
+	
+	UFUNCTION(BlueprintCallable, Category = "Justify")
+	void SetHorizontalAnchor(float InHorizontalAnchor);
+
+	UFUNCTION(BlueprintPure, Category = "Justify")
+	float GetVerticalAnchor() const { return VerticalAnchor; }
+	
+	UFUNCTION(BlueprintCallable, Category = "Justify")
+	void SetVerticalAnchor(float InVerticalAnchor);
+
+	UFUNCTION(BlueprintPure, Category = "Justify")
+	float GetDepthAnchor() const { return DepthAnchor; }
+	
+	UFUNCTION(BlueprintCallable, Category = "Justify")
+	void SetDepthAnchor(float InDepthAnchor);
+
+protected:
+	//~ Begin UActorModifierCoreBase
+	virtual void OnModifierCDOSetup(FActorModifierCoreMetadata& InMetadata) override;
+	virtual bool IsModifierDirtyable() const override;
+	virtual void OnModifiedActorTransformed() override;
+	virtual void Apply() override;
+	//~ End UActorModifierCoreBase
+	
+	//~ Begin IAvaSceneTreeUpdateModifierExtension
+	virtual void OnSceneTreeTrackedActorDirectChildrenChanged(int32 InIdx, const TArray<TWeakObjectPtr<AActor>>& InPreviousChildrenActors, const TArray<TWeakObjectPtr<AActor>>& InNewChildrenActors) override;
+	//~ End IAvaSceneTreeUpdateModifierExtension
+	
+	//~ Begin IAvaRenderStateUpdateExtension
+	virtual void OnRenderStateUpdated(AActor* InActor, UActorComponent* InComponent) override;
+	virtual void OnActorVisibilityChanged(AActor* InActor) override;
+	//~ End IAvaRenderStateUpdateExtension
+	
+	//~ Begin IAvaTransformUpdateExtension
+	virtual void OnTransformUpdated(AActor* InActor, bool bInParentMoved) override;
+	//~ End IAvaTransformUpdateExtension
+	
+	// Get tracked actors from children actors
+	void GetChildrenActors(TSet<TWeakObjectPtr<AActor>>& OutChildren) const;
+	void GetTrackedActors(const TSet<TWeakObjectPtr<AActor>>& InChildrenActors, TArray<TWeakObjectPtr<const AActor>>& OutTrackedActors) const;
+
+	bool HasDepthAlignment() const;
+	bool HasHorizontalAlignment() const;
+	bool HasVerticalAlignment() const;
+
+	FVector GetConstraintVector(const FVector& InBoundsCenter, const FVector& InModifiedActorPosition) const;
+	FVector MakeConstrainedAxisVector() const;
+	FVector GetAnchorOffset() const;
+	FVector GetAlignmentOffset(const FVector& InExtent) const;
+	
+	UPROPERTY(EditAnywhere, BlueprintReadWrite, Setter="SetHorizontalAlignment", Getter="GetHorizontalAlignment", Category="Justify", meta=(AllowPrivateAccess="true"))
+	EAvaJustifyHorizontal HorizontalAlignment;
+
+	UPROPERTY(EditAnywhere, BlueprintReadWrite, Setter="SetVerticalAlignment", Getter="GetVerticalAlignment", Category="Justify", meta=(AllowPrivateAccess="true"))
+	EAvaJustifyVertical VerticalAlignment;
+	
+	UPROPERTY(EditAnywhere, BlueprintReadWrite, Setter="SetDepthAlignment", Getter="GetDepthAlignment", Category="Justify", meta=(AllowPrivateAccess="true"))
+	EAvaJustifyDepth DepthAlignment;
+
+	UPROPERTY(EditAnywhere, BlueprintReadWrite, Setter="SetHorizontalAnchor", Getter="GetHorizontalAnchor", Interp, Category="Justify", meta=(EditCondition="HorizontalAlignment != EAvaJustifyHorizontal::None", EditConditionHides, AllowPrivateAccess="true"))
+	float HorizontalAnchor;
+
+	UPROPERTY(EditAnywhere, BlueprintReadWrite, Setter="SetVerticalAnchor", Getter="GetVerticalAnchor", Interp, Category="Justify", meta=(EditCondition="VerticalAlignment != EAvaJustifyVertical::None", EditConditionHides, AllowPrivateAccess="true"))
+	float VerticalAnchor;
+
+	UPROPERTY(EditAnywhere, BlueprintReadWrite, Setter="SetDepthAnchor", Getter="GetDepthAnchor", Interp, Category="Justify", meta=(EditCondition = "DepthAlignment != EAvaJustifyDepth::None", EditConditionHides, AllowPrivateAccess="true"))
+	float DepthAnchor;
+
+private:
+	/** Cached actors bounds to detect a change in tick */
+	UPROPERTY(Transient, DuplicateTransient, TextExportTransient)
+	FBox CachedTrackedBounds = FBox(EForceInit::ForceInit);
+};
diff --git a/Engine/Plugins/Experimental/Avalanche/Source/AvalancheModifiers/Public/Modifiers/AvaLookAtModifier.h b/Engine/Plugins/Experimental/Avalanche/Source/AvalancheModifiers/Public/Modifiers/AvaLookAtModifier.h
new file mode 100644
--- /dev/null
+++ b/Engine/Plugins/Experimental/Avalanche/Source/AvalancheModifiers/Public/Modifiers/AvaLookAtModifier.h
@@ -0,0 +1,98 @@
+// Copyright Epic Games, Inc. All Rights Reserved.
+
+#pragma once
+
+#include "AvaAttachmentBaseModifier.h"
+#include "Extensions/AvaTransformUpdateModifierExtension.h"
+#include "AvaLookAtModifier.generated.h"
+
+class AActor;
+enum class EAvaAxis : uint8;
+
+/**
+ * Rotates the modifying actor to point it's specified axis at another actor.
+ */
+UCLASS(BlueprintType)
+class AVALANCHEMODIFIERS_API UAvaLookAtModifier : public UAvaAttachmentBaseModifier
+	, public IAvaTransformUpdateHandler
+{
+	GENERATED_BODY()
+	
+public:
+	//~ Begin UObject
+	virtual void PostLoad() override;
+#if WITH_EDITOR
+	virtual void PostEditChangeProperty(FPropertyChangedEvent& PropertyChangedEvent) override;
+	virtual void PostTransacted(const FTransactionObjectEvent& TransactionEvent) override;
+#endif
+	//~ End UObject
+
+	UFUNCTION(BlueprintCallable, Category = "Motion Design|Modifiers|LookAt")
+	void SetReferenceActor(const FAvaSceneTreeActor& InReferenceActor);
+	
+	UFUNCTION(BlueprintPure, Category = "Motion Design|Modifiers|LookAt")
+	const FAvaSceneTreeActor& GetReferenceActor() const
+	{
+		return ReferenceActor;
+	}
+	
+	/** Returns the axis that will point towards t he reference actor. */
+	UFUNCTION(BlueprintPure, Category = "Motion Design|Modifiers|LookAt")
+	EAvaAxis GetAxis() const
+	{
+		return Axis;
+	}
+
+	/** Sets the axis that will point towards the reference actor. */
+	UFUNCTION(BlueprintCallable, Category = "Motion Design|Modifiers|LookAt")
+	void SetAxis(const EAvaAxis NewAxis);
+
+	/** Returns true if flipping the look-at rotation axis. */
+	UFUNCTION(BlueprintPure, Category = "Motion Design|Modifiers|LookAt")
+	bool GetFlipAxis() const
+	{
+		return bFlipAxis;
+	}
+
+	/** Sets the look-at direction to be flipped. */
+	UFUNCTION(BlueprintCallable, Category = "Motion Design|Modifiers|LookAt")
+	void SetFlipAxis(const bool bNewFlipAxis);
+
+protected:
+	//~ Begin UActorModifierCoreBase
+	virtual void OnModifierCDOSetup(FActorModifierCoreMetadata& InMetadata) override;
+	virtual void OnModifierAdded(EActorModifierCoreEnableReason InReason) override;
+	virtual void OnModifierEnabled(EActorModifierCoreEnableReason InReason) override;
+	virtual void OnModifierDisabled(EActorModifierCoreDisableReason InReason) override;
+	virtual void Apply() override;
+	virtual void OnModifiedActorTransformed() override;
+	//~ End UActorModifierCoreBase
+	
+	//~ Begin IAvaTransformUpdateExtension
+	virtual void OnTransformUpdated(AActor* InActor, bool bInParentMoved) override;
+	//~ End IAvaTransformUpdateExtension
+	
+	//~ Begin IAvaSceneTreeUpdateModifierExtension
+	virtual void OnSceneTreeTrackedActorChanged(int32 InIdx, AActor* InPreviousActor, AActor* InNewActor) override;
+	//~ End IAvaSceneTreeUpdateModifierExtension
+	
+	void OnReferenceActorChanged();
+	
+	UPROPERTY(EditAnywhere, BlueprintReadWrite, Setter="SetReferenceActor", Getter="GetReferenceActor", Category="LookAt", meta=(ShowOnlyInnerProperties, AllowPrivateAccess="true"))
+	FAvaSceneTreeActor ReferenceActor;
+
+	/** The actor to look at. */
+	UPROPERTY(meta=(DeprecatedProperty, DeprecationMessage="Use ReferenceActor instead"))
+	TWeakObjectPtr<AActor> ReferenceActorWeak;
+
+	/** The axis that will point towards the reference actor. */
+	UPROPERTY(EditAnywhere, BlueprintReadWrite, Setter="SetAxis", Getter="GetAxis", Category="LookAt", meta=(AllowPrivateAccess="true"))
+	EAvaAxis Axis;
+
+	/** If true, will flip the look-at direction. */
+	UPROPERTY(EditAnywhere, BlueprintReadWrite, Setter="SetFlipAxis", Getter="GetFlipAxis", Category="LookAt", meta=(AllowPrivateAccess="true"))
+	bool bFlipAxis = false;
+
+	UPROPERTY()
+	bool bDeprecatedPropertiesMigrated = false;
+};
