[MotionDesign]

diff --git a/Engine/Plugins/Experimental/Avalanche/Source/AvalancheModifiers/Public/Modifiers/AvaMaterialParameterModifier.h b/Engine/Plugins/Experimental/Avalanche/Source/AvalancheModifiers/Public/Modifiers/AvaMaterialParameterModifier.h
new file mode 100644
--- /dev/null
+++ b/Engine/Plugins/Experimental/Avalanche/Source/AvalancheModifiers/Public/Modifiers/AvaMaterialParameterModifier.h
@@ -0,0 +1,118 @@
+// Copyright Epic Games, Inc. All Rights Reserved.
+
+#pragma once
+
+#include "AvaArrangeBaseModifier.h"
+#include "UObject/ObjectPtr.h"
+#include "AvaMaterialParameterModifier.generated.h"
+
+class UMaterialInstanceDynamic;
+class UTexture;
+
+USTRUCT(BlueprintType)
+struct FAvaMaterialParameterMap
+{
+	GENERATED_BODY()
+
+	UPROPERTY(EditAnywhere, BlueprintReadWrite, Category = "MaterialParameter")
+	TMap<FName, float> ScalarParameters;
+
+	UPROPERTY(EditAnywhere, BlueprintReadWrite, Category = "MaterialParameter")
+	TMap<FName, FLinearColor> VectorParameters;
+
+	UPROPERTY(EditAnywhere, BlueprintReadWrite, Category = "MaterialParameter")
+	TMap<FName, TObjectPtr<UTexture>> TextureParameters;
+
+	/** Matches the input parameter key map and removes all unused keys, does not touch current values */
+	void MatchKeys(const FAvaMaterialParameterMap& InParameterMap);
+
+	/** Apply those parameters on this Material Designer Instance */
+	void Set(UMaterialInstanceDynamic* InMaterial);
+
+	/** Read those parameters from this Material Designer Instance and save them */
+	void Get(UMaterialInstanceDynamic* InMaterial);
+};
+
+/** This modifier sets specified dynamic materials parameters on an actor and its children */
+UCLASS(BlueprintType)
+class AVALANCHEMODIFIERS_API UAvaMaterialParameterModifier : public UAvaArrangeBaseModifier
+{
+	GENERATED_BODY()
+	
+public:
+	UAvaMaterialParameterModifier();
+
+	//~ Begin UObject
+#if WITH_EDITOR
+	virtual void PostEditChangeProperty(FPropertyChangedEvent& PropertyChangedEvent) override;
+#endif
+	//~ End UObject
+
+	void SetMaterialParameters(const FAvaMaterialParameterMap& InParameterMap);
+	const FAvaMaterialParameterMap& GetMaterialParameters() const { return MaterialParameters; }
+	
+	void SetUpdateChildren(bool bInUpdateChildren);
+	bool GetUpdateChildren() const { return bUpdateChildren; }
+
+protected:
+	//~ Begin UActorModifierCoreBase
+	virtual void OnModifierCDOSetup(FActorModifierCoreMetadata& InMetadata) override;
+	virtual void OnModifierEnabled(EActorModifierCoreEnableReason InReason) override;
+	virtual void OnModifierDisabled(EActorModifierCoreDisableReason InReason) override;
+	virtual void OnModifiedActorTransformed() override;
+	virtual void RestorePreState() override;
+	virtual void Apply() override;
+	//~ End UActorModifierCoreBase
+	
+	//~ Begin IAvaSceneTreeUpdateModifierExtension
+	virtual void OnSceneTreeTrackedActorChildrenChanged(int32 InIdx, const TSet<TWeakObjectPtr<AActor>>& InPreviousChildrenActors, const TSet<TWeakObjectPtr<AActor>>& InNewChildrenActors) override;
+	virtual void OnSceneTreeTrackedActorDirectChildrenChanged(int32 InIdx, const TArray<TWeakObjectPtr<AActor>>& InPreviousChildrenActors, const TArray<TWeakObjectPtr<AActor>>& InNewChildrenActors) override;
+	//~ End IAvaSceneTreeUpdateModifierExtension
+
+	/** Read and save original values */
+	void SaveMaterialParameters();
+	/** Write and restore original values */
+	void RestoreMaterialParameters();
+
+#if WITH_EDITOR
+	/** Called when a property changes, used to detect material changes */
+	virtual void OnActorPropertyChanged(UObject* InObject, FPropertyChangedEvent& InChangeEvent);
+#endif
+	
+	void ScanActorMaterials();
+
+	void OnMaterialParametersChanged();
+	void OnUpdateChildrenChanged();
+
+	virtual void OnActorMaterialAdded(UMaterialInstanceDynamic* InAdded) {}
+	virtual void OnActorMaterialRemoved(UMaterialInstanceDynamic* InRemoved) {}
+	
+	/** Checks if this actor has a Material Designer Instance or that we already track one */
+	bool IsActorSupported(const AActor* InActor) const;
+
+	/** Retrieves all Material Designer Instance from a primitive component */
+	TSet<UMaterialInstanceDynamic*> GetComponentDynamicMaterials(const UPrimitiveComponent* InComponent) const;
+	
+	/** Which parameters should we set on the Material Designer Instance,
+	 * use EditCondition="bShowMaterialParameters && bShowMaterialParameters" otherwise edit inline boolean appear in details when it should not */
+	UPROPERTY(EditAnywhere, BlueprintReadWrite, Setter="SetMaterialParameters", Getter="GetMaterialParameters", Category="MaterialParameter", meta=(EditCondition="bShowMaterialParameters && bShowMaterialParameters", EditConditionHides, AllowPrivateAccess="true"))
+	FAvaMaterialParameterMap MaterialParameters;
+
+	/** Used to restore Material Designer Instance parameters to their original state */
+	UPROPERTY()
+	TMap<TWeakObjectPtr<UMaterialInstanceDynamic>, FAvaMaterialParameterMap> SavedMaterialParameters;
+
+	/** Filter material type for child modifiers */
+	UPROPERTY(Transient)
+	TSubclassOf<UMaterialInstanceDynamic> MaterialClass;
+	
+	/** Will also look into attached children actors */
+	UPROPERTY(EditAnywhere, BlueprintReadWrite, Setter="SetUpdateChildren", Getter="GetUpdateChildren", Category="MaterialParameter", meta=(AllowPrivateAccess="true"))
+	bool bUpdateChildren = true;
+
+#if WITH_EDITORONLY_DATA
+	/** Used by child classes to override MaterialParameters */
+	UPROPERTY(Transient)
+	bool bShowMaterialParameters = true;
+#endif
+};
diff --git a/Engine/Plugins/Experimental/Avalanche/Source/AvalancheModifiers/Public/Modifiers/AvaMirrorModifier.h b/Engine/Plugins/Experimental/Avalanche/Source/AvalancheModifiers/Public/Modifiers/AvaMirrorModifier.h
new file mode 100644
--- /dev/null
+++ b/Engine/Plugins/Experimental/Avalanche/Source/AvalancheModifiers/Public/Modifiers/AvaMirrorModifier.h
@@ -0,0 +1,85 @@
+// Copyright Epic Games, Inc. All Rights Reserved.
+
+#pragma once
+
+#include "AvaGeometryBaseModifier.h"
+#include "AvaModifiersPreviewPlane.h"
+#include "AvaMirrorModifier.generated.h"
+
+class UStaticMesh;
+
+UCLASS(BlueprintType)
+class AVALANCHEMODIFIERS_API UAvaMirrorModifier : public UAvaGeometryBaseModifier
+{
+	GENERATED_BODY()
+	
+public:
+	//~ Begin UObject
+#if WITH_EDITOR
+	virtual void PostEditChangeProperty(FPropertyChangedEvent& PropertyChangedEvent) override;
+#endif
+	//~ End UObject
+	
+	UFUNCTION()
+	void SetMirrorFramePosition(const FVector& InMirrorFramePosition);
+	const FVector& GetMirrorFramePosition() const { return MirrorFramePosition; }
+
+	UFUNCTION()
+	void SetMirrorFrameRotation(const FRotator& InMirrorFrameRotation);
+	const FRotator& GetMirrorFrameRotation() const { return MirrorFrameRotation; }
+	
+	UFUNCTION()
+	void SetApplyPlaneCut(bool bInApplyPlaneCut);
+	bool GetApplyPlaneCut() const { return bApplyPlaneCut; }
+
+	UFUNCTION()
+	void SetFlipCutSide(bool bInFlipCutSide);
+	bool GetFlipCutSide() const { return bFlipCutSide; }
+
+	UFUNCTION()
+	void SetWeldAlongPlane(bool bInWeldAlongPlane);
+	bool GetWeldAlongPlane() const { return bWeldAlongPlane; }
+	
+protected:
+	//~ Begin UActorModifierCoreBase
+	virtual void OnModifierCDOSetup(FActorModifierCoreMetadata& InMetadata) override;
+	virtual void OnModifierAdded(EActorModifierCoreEnableReason InReason) override;
+	virtual void Apply() override;
+	virtual void OnModifierDisabled(EActorModifierCoreDisableReason InReason) override;
+	virtual void OnModifierRemoved(EActorModifierCoreDisableReason InReason) override;
+	//~ End UActorModifierCoreBase
+	
+	void OnMirrorFrameChanged();
+	void OnMirrorOptionChanged();
+
+#if WITH_EDITOR
+	void CreatePreviewComponent();
+	void DestroyPreviewComponent();
+	void UpdatePreviewComponent();
+	
+	void OnShowMirrorFrameChanged();
+#endif
+	
+    UPROPERTY(EditInstanceOnly, BlueprintReadWrite, Setter="SetMirrorFramePosition", Getter="GetMirrorFramePosition", Category="Mirror", meta=(AllowPrivateAccess="true"))
+	FVector MirrorFramePosition = FVector::ZeroVector;
+
+	UPROPERTY(EditInstanceOnly, BlueprintReadWrite, Setter="SetMirrorFrameRotation", Getter="GetMirrorFrameRotation", Category="Mirror", meta=(AllowPrivateAccess="true"))
+	FRotator MirrorFrameRotation = FRotator(0, 0, 90);
+
+	UPROPERTY(EditInstanceOnly, BlueprintReadWrite, Setter="SetApplyPlaneCut", Getter="GetApplyPlaneCut", Category="Mirror", meta=(AllowPrivateAccess="true"))
+	bool bApplyPlaneCut = true;
+
+	UPROPERTY(EditInstanceOnly, BlueprintReadWrite, Setter="SetFlipCutSide", Getter="GetFlipCutSide", Category="Mirror", meta=(AllowPrivateAccess="true"))
+	bool bFlipCutSide = false;
+
+	UPROPERTY(EditInstanceOnly, BlueprintReadWrite, Setter="SetWeldAlongPlane", Getter="GetWeldAlongPlane", Category="Mirror", meta=(AllowPrivateAccess="true"))
+	bool bWeldAlongPlane = true;
+
+#if WITH_EDITORONLY_DATA
+	UPROPERTY(EditInstanceOnly, BlueprintReadWrite, Category="Mirror", meta=(AllowPrivateAccess="true"))
+	bool bShowMirrorFrame = false;
+
+	UPROPERTY(Transient, DuplicateTransient, TextExportTransient)
+	FAvaModifierPreviewPlane PreviewPlane;
+#endif
+};
diff --git a/Engine/Plugins/Experimental/Avalanche/Source/AvalancheModifiers/Public/Modifiers/AvaNormalModifier.h b/Engine/Plugins/Experimental/Avalanche/Source/AvalancheModifiers/Public/Modifiers/AvaNormalModifier.h
new file mode 100644
--- /dev/null
+++ b/Engine/Plugins/Experimental/Avalanche/Source/AvalancheModifiers/Public/Modifiers/AvaNormalModifier.h
@@ -0,0 +1,103 @@
+// Copyright Epic Games, Inc. All Rights Reserved.
+
+#pragma once
+
+#include "AvaGeometryBaseModifier.h"
+#include "AvaNormalModifier.generated.h"
+
+UENUM(BlueprintType)
+enum class EAvaNormalModifierSplitMethod : uint8
+{
+	/** Do not split, leave as it is */
+	None,
+	/** Each vertex will have a split normal between tris */
+	Vertex,
+	/** Shared vertex between triangles will have a split normal */
+	Triangle,
+	/** Vertices of a same face grouped together will have a split normal */
+	PolyGroup,
+	/** Vertices above a certain angle threshold will have a split normal */
+	Threshold
+};
+
+UCLASS(BlueprintType)
+class AVALANCHEMODIFIERS_API UAvaNormalModifier : public UAvaGeometryBaseModifier
+{
+	GENERATED_BODY()
+
+public:
+	//~ Begin UObject
+#if WITH_EDITOR
+	virtual void PostEditChangeProperty(FPropertyChangedEvent& PropertyChangedEvent) override;
+#endif
+	//~ End UObject
+
+	UFUNCTION()
+	void SetAngleWeighted(bool bInAngleWeighted);
+	bool GetAngleWeighted() const { return bAngleWeighted; }
+	
+	UFUNCTION()
+	void SetAreaWeighted(bool bInAreaWeighted);
+	bool GetAreaWeighted() const { return bAreaWeighted; }
+	
+	UFUNCTION()
+	void SetInvert(bool bInInvert);
+	bool GetInvert() const { return bInvert; }
+	
+	UFUNCTION()
+	void SetSplitMethod(EAvaNormalModifierSplitMethod InSplitMethod);
+	EAvaNormalModifierSplitMethod GetSplitMethod() const { return SplitMethod; }
+	
+	UFUNCTION()
+	void SetAngleThreshold(float InAngleThreshold);
+	float GetAngleThreshold() const { return AngleThreshold; }
+
+	UFUNCTION()
+	void SetPolyGroupLayerIdx(int32 InPolyGroupLayer);
+	int32 GetPolyGroupLayerIdx() const;
+
+	UFUNCTION()
+	void SetPolyGroupLayer(FString& InString);
+	FString GetPolyGroupLayer() const { return PolyGroupLayer; }
+
+protected:
+	//~ Begin UActorModifierCoreBase
+	virtual void OnModifierCDOSetup(FActorModifierCoreMetadata& InMetadata) override;
+	virtual void OnModifierAdded(EActorModifierCoreEnableReason InReason) override;
+	virtual void Apply() override;
+	//~ End UActorModifierCoreBase
+	
+	void OnAngleWeightedChanged();
+	void OnAreaWeightedChanged();
+	void OnInvertChanged();
+	void OnSplitMethodChanged();
+	void OnAngleThresholdChanged();
+	void OnPolyGroupLayerChanged();
+
+	UFUNCTION()
+	TArray<FString> GetPolyGroupLayers() const;
+
+	/** Recompute normals and weight them by angle */
+	UPROPERTY(EditInstanceOnly, BlueprintReadWrite, Setter="SetAngleWeighted", Getter="GetAngleWeighted", Category="Normal", meta=(AllowPrivateAccess="true"))
+	bool bAngleWeighted = true;
+
+	/** Recompute normals and weight them by area */
+	UPROPERTY(EditInstanceOnly, BlueprintReadWrite, Setter="SetAreaWeighted", Getter="GetAreaWeighted", Category="Normal", meta=(AllowPrivateAccess="true"))
+	bool bAreaWeighted = true;
+
+	/** Recompute normals and invert normals and triangles */
+	UPROPERTY(EditInstanceOnly, BlueprintReadWrite, Setter="SetInvert", Getter="GetInvert", Category="Normal", meta=(AllowPrivateAccess="true"))
+	bool bInvert = false;
+
+	/** Recompute normals and use a split method */
+	UPROPERTY(EditInstanceOnly, BlueprintReadWrite, Setter="SetSplitMethod", Getter="GetSplitMethod", Category="Normal", meta=(AllowPrivateAccess="true"))
+	EAvaNormalModifierSplitMethod SplitMethod = EAvaNormalModifierSplitMethod::Threshold;
+
+	/** Angle to compare and split normal when threshold method is chosen */
+	UPROPERTY(EditInstanceOnly, BlueprintReadWrite, Setter="SetAngleThreshold", Getter="GetAngleThreshold", Category="Normal", meta=(ClampMin="0.0", ClampMax="180.0", EditCondition="SplitMethod == EAvaNormalModifierSplitMethod::Threshold", AllowPrivateAccess="true"))
+	float AngleThreshold = 60.f;
+
+	/** PolyGroup to use to split normal from when PolyGroup method is chosen */
+	UPROPERTY(EditInstanceOnly, BlueprintReadWrite, Setter="SetPolyGroupLayer", Getter="GetPolyGroupLayer", Category="Normal", meta=(GetOptions="GetPolyGroupLayers", EditCondition="SplitMethod == EAvaNormalModifierSplitMethod::PolyGroup", AllowPrivateAccess="true"))
+	FString PolyGroupLayer = TEXT("None");
+};
diff --git a/Engine/Plugins/Experimental/Avalanche/Source/AvalancheModifiers/Public/Modifiers/AvaOutlineModifier.h b/Engine/Plugins/Experimental/Avalanche/Source/AvalancheModifiers/Public/Modifiers/AvaOutlineModifier.h
new file mode 100644
--- /dev/null
+++ b/Engine/Plugins/Experimental/Avalanche/Source/AvalancheModifiers/Public/Modifiers/AvaOutlineModifier.h
@@ -0,0 +1,63 @@
+// Copyright Epic Games, Inc. All Rights Reserved.
+
+#pragma once
+
+#include "AvaGeometryBaseModifier.h"
+#include "AvaOutlineModifier.generated.h"
+
+UENUM(BlueprintType)
+enum class EAvaOutlineMode : uint8
+{
+	Outset,
+	Inset
+};
+
+/** This modifier adds an outline around a 2D shape with a specific distance */
+UCLASS(BlueprintType)
+class AVALANCHEMODIFIERS_API UAvaOutlineModifier : public UAvaGeometryBaseModifier
+{
+	GENERATED_BODY()
+
+public:
+	//~ Begin UObject
+#if WITH_EDITOR
+	virtual void PostEditChangeProperty(FPropertyChangedEvent& PropertyChangedEvent) override;
+#endif
+	//~ End UObject
+
+	UFUNCTION()
+	void SetMode(EAvaOutlineMode InMode);
+	EAvaOutlineMode GetMode() const { return Mode; }
+	
+	UFUNCTION()
+	void SetDistance(float InDistance);
+	float GetDistance() const { return Distance; }
+
+	UFUNCTION()
+	void SetRemoveInside(bool bInRemoveInside);
+	bool GetRemoveInside() const { return bRemoveInside; }
+
+	float GetMaxInsetDistance() const;
+
+protected:
+	//~ Begin UActorModifierCoreBase
+	virtual void OnModifierCDOSetup(FActorModifierCoreMetadata& InMetadata) override;
+	virtual void Apply() override;
+	//~ End UActorModifierCoreBase
+	
+	void OnModeChanged();
+	void OnDistanceChanged();
+	void OnRemoveInsideChanged();
+
+	/** Set the mode like inset or outset */
+	UPROPERTY(EditInstanceOnly, BlueprintReadWrite, Setter="SetMode", Getter="GetMode", Category="Outline", meta=(AllowPrivateAccess="true"))
+	EAvaOutlineMode Mode = EAvaOutlineMode::Outset;
+	
+	/** Set the distance for the outline */
+	UPROPERTY(EditInstanceOnly, BlueprintReadWrite, Setter="SetDistance", Getter="GetDistance", Category="Outline", meta=(ClampMin="0", AllowPrivateAccess="true"))
+	float Distance = 10.f;
+
+	/** Remove the inside part and create a hole in the shape */
+	UPROPERTY(EditInstanceOnly, BlueprintReadWrite, Setter="SetRemoveInside", Getter="GetRemoveInside", Category="Outline", meta=(AllowPrivateAccess="true"))
+	bool bRemoveInside = true;
+};
diff --git a/Engine/Plugins/Experimental/Avalanche/Source/AvalancheModifiers/Public/Modifiers/AvaPatternModifier.h b/Engine/Plugins/Experimental/Avalanche/Source/AvalancheModifiers/Public/Modifiers/AvaPatternModifier.h
new file mode 100644
--- /dev/null
+++ b/Engine/Plugins/Experimental/Avalanche/Source/AvalancheModifiers/Public/Modifiers/AvaPatternModifier.h
@@ -0,0 +1,214 @@
+// Copyright Epic Games, Inc. All Rights Reserved.
+
+#pragma once
+
+#include "AvaGeometryBaseModifier.h"
+#include "TransformTypes.h"
+#include "AvaPatternModifier.generated.h"
+
+UENUM(BlueprintType)
+enum class EAvaPatternModifierLayout : uint8
+{
+	Line = 0,
+	Grid = 1,
+	Circle = 2
+};
+
+UENUM(BlueprintType)
+enum class EAvaPatternModifierAxis : uint8
+{
+	X = 0,
+	Y = 1,
+	Z = 2
+};
+
+UENUM(BlueprintType)
+enum class EAvaPatternModifierPlane : uint8
+{
+	XY = 0,
+	ZX = 1,
+	YZ = 2
+};
+
+USTRUCT(BlueprintType)
+struct FVector2b
+{
+	GENERATED_BODY()
+
+	UPROPERTY(EditInstanceOnly, BlueprintReadWrite, Category = "Values")
+	bool bX = false;
+
+	UPROPERTY(EditInstanceOnly, BlueprintReadWrite, Category = "Values")
+	bool bY = false;
+};
+
+USTRUCT(BlueprintType)
+struct FVector3b
+{
+	GENERATED_BODY()
+
+	UPROPERTY(EditInstanceOnly, BlueprintReadWrite, Category = "Values")
+	bool bX = false;
+
+	UPROPERTY(EditInstanceOnly, BlueprintReadWrite, Category = "Values")
+	bool bY = false;
+
+	UPROPERTY(EditInstanceOnly, BlueprintReadWrite, Category = "Values")
+	bool bZ = false;
+};
+
+USTRUCT(BlueprintType)
+struct FAvaPatternModifierLineLayoutOptions
+{
+	GENERATED_BODY()
+
+	UPROPERTY(EditInstanceOnly, BlueprintReadWrite, Category="Pattern")
+	EAvaPatternModifierAxis Axis = EAvaPatternModifierAxis::Y;
+
+	UPROPERTY(EditInstanceOnly, BlueprintReadWrite, Category="Pattern")
+	bool bAxisInverted = false;
+
+	UPROPERTY(EditInstanceOnly, BlueprintReadWrite, Category="Pattern", meta=(ClampMin="1"))
+	int32 RepeatCount = 4;
+
+	UPROPERTY(EditInstanceOnly, BlueprintReadWrite, Category="Pattern")
+	float Spacing = 0.f;
+
+	// Center the layout based on the axis
+	UPROPERTY(EditInstanceOnly, BlueprintReadWrite, Category="Pattern")
+	bool bCentered = true;
+
+	UPROPERTY(EditInstanceOnly, BlueprintReadWrite, Category="Pattern")
+	bool bAccumulateTransform = false;
+
+	UPROPERTY(EditInstanceOnly, BlueprintReadWrite, Category="Pattern")
+	FRotator Rotation = FRotator::ZeroRotator;
+
+	UPROPERTY(EditInstanceOnly, BlueprintReadWrite, Category="Pattern")
+	FVector Scale = FVector::OneVector;
+};
+
+USTRUCT(BlueprintType)
+struct FAvaPatternModifierGridLayoutOptions
+{
+	GENERATED_BODY()
+
+	UPROPERTY(EditInstanceOnly, BlueprintReadWrite, Category="Pattern")
+	EAvaPatternModifierPlane Plane = EAvaPatternModifierPlane::YZ;
+
+	UPROPERTY(EditInstanceOnly, BlueprintReadWrite, Category="Pattern")
+	FVector2b AxisInverted;
+
+	UPROPERTY(EditInstanceOnly, BlueprintReadWrite, Category="Pattern", meta=(ClampMin="1"))
+	FIntPoint RepeatCount = FIntPoint(2, 2); // Row, Column
+
+	UPROPERTY(EditInstanceOnly, BlueprintReadWrite, Category="Pattern")
+	FVector2D Spacing = FVector2D(0.f);
+
+	// Center the layout based on the plane
+	UPROPERTY(EditInstanceOnly, BlueprintReadWrite, Category="Pattern")
+	bool bCentered = true;
+
+	UPROPERTY(EditInstanceOnly, BlueprintReadWrite, Category="Pattern")
+	bool bAccumulateTransform = false;
+
+	UPROPERTY(EditInstanceOnly, BlueprintReadWrite, Category="Pattern")
+	FRotator Rotation = FRotator::ZeroRotator;
+
+	UPROPERTY(EditInstanceOnly, BlueprintReadWrite, Category="Pattern")
+	FVector Scale = FVector::OneVector;
+};
+
+USTRUCT(BlueprintType)
+struct FAvaPatternModifierCircleLayoutOptions
+{
+	GENERATED_BODY()
+
+	UPROPERTY(EditInstanceOnly, BlueprintReadWrite, Category="Pattern")
+	EAvaPatternModifierPlane Plane = EAvaPatternModifierPlane::YZ;
+
+	UPROPERTY(EditInstanceOnly, BlueprintReadWrite, Category="Pattern")
+	float Radius = 100.f;
+
+	UPROPERTY(EditInstanceOnly, BlueprintReadWrite, Category="Pattern")
+	float StartAngle = 180.f;
+
+	UPROPERTY(EditInstanceOnly, BlueprintReadWrite, Category="Pattern")
+	float FullAngle = 360.f;
+
+	UPROPERTY(EditInstanceOnly, BlueprintReadWrite, Category="Pattern", meta=(ClampMin="1"))
+	int32 RepeatCount = 4;
+
+	// Center the layout based on the plane
+	UPROPERTY(EditInstanceOnly, BlueprintReadWrite, Category="Pattern")
+	bool bCentered = true;
+
+	UPROPERTY(EditInstanceOnly, BlueprintReadWrite, Category="Pattern")
+	bool bAccumulateTransform = false;
+
+	UPROPERTY(EditInstanceOnly, BlueprintReadWrite, Category="Pattern")
+	FRotator Rotation = FRotator::ZeroRotator;
+
+	UPROPERTY(EditInstanceOnly, BlueprintReadWrite, Category="Pattern")
+	FVector Scale = FVector::OneVector;
+};
+
+/** This modifier clones a shape following various layouts and options */
+UCLASS(BlueprintType, AutoExpandCategories=(Pattern))
+class AVALANCHEMODIFIERS_API UAvaPatternModifier : public UAvaGeometryBaseModifier
+{
+	GENERATED_BODY()
+
+public:
+	//~ Begin UObject
+#if WITH_EDITOR
+	virtual void PostEditChangeProperty(FPropertyChangedEvent& PropertyChangedEvent) override;
+#endif
+	//~ End UObject
+
+	void SetLayout(EAvaPatternModifierLayout InLayout);
+	EAvaPatternModifierLayout GetLayout() const { return Layout; }
+
+	void SetLineLayoutOptions(const FAvaPatternModifierLineLayoutOptions& InOptions);
+	const FAvaPatternModifierLineLayoutOptions& GetLineLayoutOptions() const { return LineLayoutOptions; }
+
+	void SetGridLayoutOptions(const FAvaPatternModifierGridLayoutOptions& InOptions);
+	const FAvaPatternModifierGridLayoutOptions& GetGridLayoutOptions() const { return GridLayoutOptions; }
+
+	void SetCircleLayoutOptions(const FAvaPatternModifierCircleLayoutOptions& InOptions);
+	const FAvaPatternModifierCircleLayoutOptions& GetCircleLayoutOptions() const { return CircleLayoutOptions; }
+
+protected:
+	//~ Begin UActorModifierCoreBase
+	virtual void OnModifierCDOSetup(FActorModifierCoreMetadata& InMetadata) override;
+	virtual void Apply() override;
+	//~ End UActorModifierCoreBase
+
+	void OnLayoutChanged();
+	void OnLineLayoutOptionsChanged();
+	void OnGridLayoutOptionsChanged();
+	void OnCircleLayoutOptionsChanged();
+
+	UE::Geometry::FTransformSRT3d GetLineLayoutTransformChange() const;
+	UE::Geometry::FTransformSRT3d GetGridLayoutColTransformChange() const;
+	UE::Geometry::FTransformSRT3d GetGridLayoutRowTransformChange() const;
+	UE::Geometry::FTransformSRT3d GetCircleLayoutTransformChange(int32 Idx) const;
+
+	UPROPERTY(EditInstanceOnly, BlueprintReadWrite, Setter="SetLayout", Getter="GetLayout", Category="Pattern", meta=(AllowPrivateAccess="true"))
+	EAvaPatternModifierLayout Layout = EAvaPatternModifierLayout::Line;
+
+	/** Line layout options */
+	UPROPERTY(EditInstanceOnly, BlueprintReadWrite, Setter="SetLineLayoutOptions", Getter="GetLineLayoutOptions", Category="Pattern", meta=(EditCondition="Layout == EAvaPatternModifierLayout::Line", EditConditionHides, AllowPrivateAccess="true"))
+	FAvaPatternModifierLineLayoutOptions LineLayoutOptions;
+
+	/** Grid layout options */
+	UPROPERTY(EditInstanceOnly, BlueprintReadWrite, Setter="SetGridLayoutOptions", Getter="GetGridLayoutOptions", Category="Pattern", meta=(EditCondition="Layout == EAvaPatternModifierLayout::Grid", EditConditionHides, AllowPrivateAccess="true"))
+	FAvaPatternModifierGridLayoutOptions GridLayoutOptions;
+
+	/** Circle layout options */
+	UPROPERTY(EditInstanceOnly, BlueprintReadWrite, Setter="SetCircleLayoutOptions", Getter="GetCircleLayoutOptions", Category="Pattern", meta=(EditCondition="Layout == EAvaPatternModifierLayout::Circle", EditConditionHides, AllowPrivateAccess="true"))
+	FAvaPatternModifierCircleLayoutOptions CircleLayoutOptions;
+
+	UPROPERTY(Transient, DuplicateTransient, TextExportTransient)
+	FBox OriginalMeshBounds;
+};
diff --git a/Engine/Plugins/Experimental/Avalanche/Source/AvalancheModifiers/Public/Modifiers/AvaPlaneCutModifier.h b/Engine/Plugins/Experimental/Avalanche/Source/AvalancheModifiers/Public/Modifiers/AvaPlaneCutModifier.h
new file mode 100644
--- /dev/null
+++ b/Engine/Plugins/Experimental/Avalanche/Source/AvalancheModifiers/Public/Modifiers/AvaPlaneCutModifier.h
@@ -0,0 +1,81 @@
+// Copyright Epic Games, Inc. All Rights Reserved.
+
+#pragma once
+
+#include "AvaGeometryBaseModifier.h"
+#include "AvaModifiersPreviewPlane.h"
+#include "AvaPlaneCutModifier.generated.h"
+
+/** This modifier cuts a shape based on a 2D plane */
+UCLASS(BlueprintType)
+class AVALANCHEMODIFIERS_API UAvaPlaneCutModifier : public UAvaGeometryBaseModifier
+{
+	GENERATED_BODY()
+
+public:
+	//~ Begin UObject
+#if WITH_EDITOR
+	virtual void PostEditChangeProperty(FPropertyChangedEvent& InPropertyChangedEvent) override;
+#endif
+	//~ End UObject
+
+	UFUNCTION()
+	void SetPlaneOrigin(float InOrigin);
+	float GetPlaneOrigin() const { return PlaneOrigin; }
+
+	UFUNCTION()
+	void SetPlaneRotation(const FRotator& InRotation);
+	const FRotator& GetPlaneRotation() const { return PlaneRotation; }
+
+	UFUNCTION()
+	void SetInvertCut(bool bInInvertCut);
+	bool GetInvertCut() const { return bInvertCut; }
+
+	UFUNCTION()
+	void SetFillHoles(bool bInFillHoles);
+	bool GetFillHoles() const { return bFillHoles; }
+
+protected:
+	//~ Begin UActorModifierCoreBase
+	virtual void OnModifierCDOSetup(FActorModifierCoreMetadata& InMetadata) override;
+	virtual void OnModifierAdded(EActorModifierCoreEnableReason InReason) override;
+	virtual void Apply() override;
+	virtual void OnModifierDisabled(EActorModifierCoreDisableReason InReason) override;
+	virtual void OnModifierRemoved(EActorModifierCoreDisableReason InReason) override;
+	//~ End UActorModifierCoreBase
+	
+	/** Returns actual location of plane bounds restricted */
+	FVector GetPlaneLocation() const;
+	
+	void OnPlaneRotationChanged();
+	void OnFillHolesChanged();
+	void OnInvertCutChanged();
+	void OnPlaneOriginChanged();
+
+#if WITH_EDITOR
+	void OnUsePreviewChanged();
+	void CreatePreviewComponent();
+	void DestroyPreviewComponent();
+	void UpdatePreviewComponent();
+#endif
+
+	UPROPERTY(EditInstanceOnly, BlueprintReadWrite, Setter="SetPlaneOrigin", Getter="GetPlaneOrigin", Category="PlaneCut", meta=(Delta="0.5", LinearDeltaSensitivity="1", AllowPrivateAccess="true"))
+	float PlaneOrigin = 0.f;
+
+	UPROPERTY(EditInstanceOnly, BlueprintReadWrite, Setter="SetPlaneRotation", Getter="GetPlaneRotation", Category="PlaneCut", meta=(AllowPrivateAccess="true"))
+	FRotator PlaneRotation = FRotator(0, 0, 90);
+
+	UPROPERTY(EditInstanceOnly, BlueprintReadWrite, Setter="SetInvertCut", Getter="GetInvertCut", Category="PlaneCut", meta=(AllowPrivateAccess="true"))
+	bool bInvertCut = false;
+	
+	UPROPERTY(EditInstanceOnly, BlueprintReadWrite, Setter="SetFillHoles", Getter="GetFillHoles", Category="PlaneCut", meta=(AllowPrivateAccess="true"))
+	bool bFillHoles = true;
+
+#if WITH_EDITORONLY_DATA
+	UPROPERTY(EditInstanceOnly, BlueprintReadWrite, Category="PlaneCut", meta=(AllowPrivateAccess="true"))
+	bool bUsePreview = false;
+
+	UPROPERTY(Transient, DuplicateTransient, TextExportTransient)
+	FAvaModifierPreviewPlane PreviewPlane;
+#endif
+};
diff --git a/Engine/Plugins/Experimental/Avalanche/Source/AvalancheModifiers/Public/Modifiers/AvaRadialArrangeModifier.h b/Engine/Plugins/Experimental/Avalanche/Source/AvalancheModifiers/Public/Modifiers/AvaRadialArrangeModifier.h
new file mode 100644
--- /dev/null
+++ b/Engine/Plugins/Experimental/Avalanche/Source/AvalancheModifiers/Public/Modifiers/AvaRadialArrangeModifier.h
@@ -0,0 +1,176 @@
+// Copyright Epic Games, Inc. All Rights Reserved.
+
+#pragma once
+
+#include "AvaArrangeBaseModifier.h"
+#include "AvaRadialArrangeModifier.generated.h"
+
+/** Specifies how child elements will be arranged radially. */
+UENUM(BlueprintType)
+enum class EAvaRadialArrangement : uint8
+{
+	/**
+	 * Each radial ring will contain the same number of elements.
+	 * The space between elements in the outer rings will be greater than the inner rings.
+	 */
+	Monospace,
+	/**
+	 * All elements in all radial rings have the same spacing between them.
+	 * The number of elements in the inner rings will be greater than the outer rings.
+	 */
+	 Equal
+};
+
+/**
+ * Arranges child actors in a circular rings around its center
+ */
+UCLASS(BlueprintType)
+class AVALANCHEMODIFIERS_API UAvaRadialArrangeModifier : public UAvaArrangeBaseModifier
+{
+	GENERATED_BODY()
+	
+public:
+	//~ Begin UObject
+#if WITH_EDITOR
+	virtual void PostEditChangeProperty(FPropertyChangedEvent& PropertyChangedEvent) override;
+#endif
+	//~ End UObject
+
+	/** Gets the number of child elements to use in the arrangement. Children whose index is greater than or equal to this value will be hidden. */
+	UFUNCTION(BlueprintPure, Category = "Motion Design|Modifiers|RadialArrange")
+	int32 GetCount() const { return Count; }
+
+	/** Sets the number of child elements to use in the arrangement. Children whose index is greater than or equal to this value will be hidden. */
+	UFUNCTION(BlueprintCallable, Category = "Motion Design|Modifiers|RadialArrange")
+	void SetCount(const int32 NewCount);
+
+	/** Gets the number of rings. */
+	UFUNCTION(BlueprintPure, Category = "Motion Design|Modifiers|RadialArrange")
+	int32 GetRings() const { return Rings; }
+
+	/** Sets the number of rings. */
+	UFUNCTION(BlueprintCallable, Category = "Motion Design|Modifiers|RadialArrange")
+	void SetRings(const int32 NewRings);
+
+	/** Gets the radius from the center to the first inner ring. */
+	UFUNCTION(BlueprintPure, Category = "Motion Design|Modifiers|RadialArrange")
+	float GetInnerRadius() const { return InnerRadius; }
+
+	/** Sets the radius from the center to the first inner ring. */
+	UFUNCTION(BlueprintCallable, Category = "Motion Design|Modifiers|RadialArrange")
+	void SetInnerRadius(const float NewInnerRadius);
+
+	/** Gets the radius from the center to the last outer ring. */
+	UFUNCTION(BlueprintPure, Category = "Motion Design|Modifiers|RadialArrange")
+	float GetOuterRadius() const { return OuterRadius; }
+
+	/** Sets the radius from the center to the last outer ring. */
+	UFUNCTION(BlueprintCallable, Category = "Motion Design|Modifiers|RadialArrange")
+	void SetOuterRadius(const float NewOuterRadius);
+
+	/** Gets the start angle for the arrangement space and moving clockwise. 0 = Up, -90 = Left, 90 = Right */
+	UFUNCTION(BlueprintPure, Category = "Motion Design|Modifiers|RadialArrange")
+	float GetStartAngle() const { return StartAngle; }
+
+	/** Sets the start angle for the arrangement space and moving clockwise. 0 = Up, -90 = Left, 90 = Right */
+	UFUNCTION(BlueprintCallable, Category = "Motion Design|Modifiers|RadialArrange")
+	void SetStartAngle(const float NewStartAngle);
+
+	/** Gets the end angle for the arrangement space. 0 = Up, -90 = Left, 90 = Right */
+	UFUNCTION(BlueprintPure, Category = "Motion Design|Modifiers|RadialArrange")
+	float GetEndAngle() const { return EndAngle; }
+
+	/** Sets the end angle for the arrangement space. 0 = Up, -90 = Left, 90 = Right */
+	UFUNCTION(BlueprintCallable, Category = "Motion Design|Modifiers|RadialArrange")
+	void SetEndAngle(const float NewEndAngle);
+
+	/** Defines how to arrange the child elements around the center. */
+	UFUNCTION(BlueprintPure, Category = "Motion Design|Modifiers|RadialArrange")
+	EAvaRadialArrangement GetArrangement() const { return Arrangement; }
+
+	/** Defines how to arrange the child elements around the center. */
+	UFUNCTION(BlueprintCallable, Category = "Motion Design|Modifiers|RadialArrange")
+	void SetArrangement(const EAvaRadialArrangement NewArrangement);
+
+	/** If true, will arrange the child elements starting from the outer radius and moving to the inner radius. Has no effect if only using one ring. */
+	UFUNCTION(BlueprintPure, Category = "Motion Design|Modifiers|RadialArrange")
+	bool GetStartFromOuterRadius() const { return bStartFromOuterRadius; }
+
+	/** If true, will arrange the child elements starting from the outer radius and moving to the inner radius. Has no effect if only using one ring. */
+	UFUNCTION(BlueprintCallable, Category = "Motion Design|Modifiers|RadialArrange")
+	void SetStartFromOuterRadius(const bool bMewStartFromOuterRadius);
+
+	/** If true, will orient the selected axis torwards the center. */
+	UFUNCTION(BlueprintPure, Category = "Motion Design|Modifiers|RadialArrange")
+	bool GetOrient() const { return bOrient; }
+
+	/** If true, will orient the selected axis torwards the center. */
+	UFUNCTION(BlueprintCallable, Category = "Motion Design|Modifiers|RadialArrange")
+	void SetOrient(const bool bNewOrient);
+
+	/** Gets the axis to look at the center. */
+	UFUNCTION(BlueprintPure, Category = "Motion Design|Modifiers|RadialArrange")
+	EAvaAxis GetOrientAxis() const { return OrientAxis; }
+
+	/** Sets the axis to look at the center. */
+	UFUNCTION(BlueprintCallable, Category = "Motion Design|Modifiers|RadialArrange")
+	void SetOrientAxis(const EAvaAxis NewOrientAxis);
+
+	/** If true, will flip the center orientation to face outwards. */
+	UFUNCTION(BlueprintPure, Category = "Motion Design|Modifiers|RadialArrange")
+	bool GetFlipOrient() const { return bFlipOrient; }
+
+	/** If true, will flip the center orientation to face outwards. */
+	UFUNCTION(BlueprintCallable, Category = "Motion Design|Modifiers|RadialArrange")
+	void SetFlipOrient(const bool bNewFlipOrient);
+
+protected:
+	//~ Begin UActorModifierCoreBase
+	virtual void OnModifierCDOSetup(FActorModifierCoreMetadata& InMetadata) override;
+	virtual void Apply() override;
+	//~ End UActorModifierCoreBase
+	
+	/** The number of child elements to limit in the arrangement, or -1 if unlimited. Children whose index is greater than or equal to this value will be hidden. */
+	UPROPERTY(EditAnywhere, BlueprintReadWrite, Setter="SetCount", Getter="GetCount", Category="RadialArrange", meta=(ClampMin="-1", UIMin="-1", AllowPrivateAccess="true"))
+	int32 Count = -1;
+
+	/** The number of rings. */
+	UPROPERTY(EditAnywhere, BlueprintReadWrite, Setter="SetRings", Getter="GetRings", Category="RadialArrange", meta=(ClampMin="1", UIMin="1", AllowPrivateAccess="true"))
+	int32 Rings = 1;
+
+	/** The radius from the center to the first inner ring. */
+	UPROPERTY(EditAnywhere, BlueprintReadWrite, Setter="SetInnerRadius", Getter="GetInnerRadius", Interp, Category="RadialArrange", meta=(ClampMin="0.0", UIMin="0.0", AllowPrivateAccess="true"))
+	float InnerRadius = 70.0f;
+
+	/** The radius from the center to the last outer ring. */
+	UPROPERTY(EditAnywhere, BlueprintReadWrite, Setter="SetOuterRadius", Getter="GetOuterRadius", Interp, Category="RadialArrange", meta=(EditCondition="Rings > 1", EditConditionHides, ClampMin="0.0", UIMin="0.0", AllowPrivateAccess="true"))
+	float OuterRadius = 200.0f;
+	
+	/** The start angle for the arrangement space and moving clockwise. 0 = Up, -90 = Left, 90 = Right */
+	UPROPERTY(EditAnywhere, BlueprintReadWrite, Setter="SetStartAngle", Getter="GetStartAngle", Interp, Category="RadialArrange", meta=(ClampMin="-180.0", UIMin="-180.0", ClampMax="180.0", UIMax="180.0", AllowPrivateAccess="true"))
+	float StartAngle = -180.0f;
+
+	/** The end angle for the arrangement space. 0 = Up, -90 = Left, 90 = Right */
+	UPROPERTY(EditAnywhere, BlueprintReadWrite, Setter="SetEndAngle", Getter="GetEndAngle", Interp, Category="RadialArrange", meta=(ClampMin = "-180.0", UIMin = "-180.0", ClampMax = "180.0", UIMax = "180.0", AllowPrivateAccess="true"))
+	float EndAngle = 180.0f;
+
+	/** Defines how to arrange the child elements around the center. */
+	UPROPERTY(EditAnywhere, BlueprintReadWrite, Setter="SetArrangement", Getter="GetArrangement", Category="RadialArrange", meta=(AllowPrivateAccess="true"))
+	EAvaRadialArrangement Arrangement = EAvaRadialArrangement::Equal;
+
+	/** If true, will arrange the child elements starting from the outer radius and moving to the inner radius. Has no effect if only using one ring. */
+	UPROPERTY(EditAnywhere, BlueprintReadWrite, Setter="SetStartFromOuterRadius", Getter="GetStartFromOuterRadius", Category="RadialArrange", meta=(AllowPrivateAccess="true"))
+	bool bStartFromOuterRadius;
+
+	/** If true, will orient the selected axis torwards the center. */
+	UPROPERTY(EditAnywhere, BlueprintReadWrite, Setter="SetOrient", Getter="GetOrient", Category="RadialArrange", meta=(AllowPrivateAccess="true"))
+	bool bOrient;
+
+	/** The axis to look at the center. */
+	UPROPERTY(EditAnywhere, BlueprintReadWrite, Setter="SetOrientAxis", Getter="GetOrientAxis", Category="RadialArrange", meta=(EditCondition="bOrient", EditConditionHides, AllowPrivateAccess="true"))
+	EAvaAxis OrientAxis;
+
+	/** If true, will flip the orientation axis to the opposite direction. */
+	UPROPERTY(EditAnywhere, BlueprintReadWrite, Setter="SetFlipOrient", Getter="GetFlipOrient", Category="RadialArrange", meta=(EditCondition="bOrient", EditConditionHides, AllowPrivateAccess="true"))
+	bool bFlipOrient;
+};
diff --git a/Engine/Plugins/Experimental/Avalanche/Source/AvalancheModifiers/Public/Modifiers/AvaSplineSweepModifier.h b/Engine/Plugins/Experimental/Avalanche/Source/AvalancheModifiers/Public/Modifiers/AvaSplineSweepModifier.h
new file mode 100644
--- /dev/null
+++ b/Engine/Plugins/Experimental/Avalanche/Source/AvalancheModifiers/Public/Modifiers/AvaSplineSweepModifier.h
@@ -0,0 +1,206 @@
+// Copyright Epic Games, Inc. All Rights Reserved.
+
+#pragma once
+
+#include "AvaGeometryBaseModifier.h"
+#include "AvaPropertyChangeDispatcher.h"
+#include "Extensions/AvaRenderStateUpdateModifierExtension.h"
+#include "AvaSplineSweepModifier.generated.h"
+
+class USplineComponent;
+
+UENUM(BlueprintType)
+enum class EAvaSplineSweepSampleMode : uint8
+{
+	/** Samples the full spline distance : steps on the full distance */
+	FullDistance,
+	/** Samples a custom distance : steps per custom distance */
+	CustomDistance
+};
+
+UCLASS(BlueprintType)
+class AVALANCHEMODIFIERS_API UAvaSplineSweepModifier : public UAvaGeometryBaseModifier
+	, public IAvaRenderStateUpdateHandler
+{
+	GENERATED_BODY()
+
+public:
+	static inline constexpr int32 MaxSampleCount = 100;
+
+	//~ Begin UObject
+#if WITH_EDITOR
+	virtual void PostEditChangeProperty(FPropertyChangedEvent& InPropertyChangedEvent) override;
+#endif
+	//~ End UObject
+
+	void SetSplineActorWeak(TWeakObjectPtr<AActor> InSplineActorWeak);
+	TWeakObjectPtr<AActor> GetSplineActorWeak() const
+	{
+		return SplineActorWeak;
+	}
+
+	UFUNCTION(BlueprintCallable, Category = "Motion Design|Modifiers|SplineSweep")
+	void SetSampleMode(EAvaSplineSweepSampleMode InMode);
+
+	UFUNCTION(BlueprintPure, Category = "Motion Design|Modifiers|SplineSweep")
+	EAvaSplineSweepSampleMode GetSampleMode() const
+	{
+		return SampleMode;
+	}
+
+	UFUNCTION(BlueprintCallable, Category = "Motion Design|Modifiers|SplineSweep")
+	void SetSampleDistance(float InDistance);
+
+	UFUNCTION(BlueprintPure, Category = "Motion Design|Modifiers|SplineSweep")
+	float GetSampleDistance() const
+	{
+		return SampleDistance;
+	}
+
+	UFUNCTION(BlueprintCallable, Category = "Motion Design|Modifiers|SplineSweep")
+	void SetSteps(int32 InSteps);
+
+	UFUNCTION(BlueprintPure, Category = "Motion Design|Modifiers|SplineSweep")
+	int32 GetSteps() const
+	{
+		return Steps;
+	}
+
+	UFUNCTION(BlueprintCallable, Category = "Motion Design|Modifiers|SplineSweep")
+	void SetProgressOffset(float InOffset);
+
+	UFUNCTION(BlueprintPure, Category = "Motion Design|Modifiers|SplineSweep")
+	float GetProgressOffset() const
+	{
+		return ProgressOffset;
+	}
+
+	UFUNCTION(BlueprintCallable, Category = "Motion Design|Modifiers|SplineSweep")
+	void SetProgressStart(float InStart);
+
+	UFUNCTION(BlueprintPure, Category = "Motion Design|Modifiers|SplineSweep")
+	float GetProgressStart() const
+	{
+		return ProgressStart;
+	}
+
+	UFUNCTION(BlueprintCallable, Category = "Motion Design|Modifiers|SplineSweep")
+	void SetProgressEnd(float InEnd);
+
+	UFUNCTION(BlueprintPure, Category = "Motion Design|Modifiers|SplineSweep")
+	float GetProgressEnd() const
+	{
+		return ProgressEnd;
+	}
+
+	UFUNCTION(BlueprintCallable, Category = "Motion Design|Modifiers|SplineSweep")
+	void SetScaleStart(float InScaleStart);
+
+	UFUNCTION(BlueprintPure, Category = "Motion Design|Modifiers|SplineSweep")
+	float GetScaleStart() const
+	{
+		return ScaleStart;
+	}
+
+	UFUNCTION(BlueprintCallable, Category = "Motion Design|Modifiers|SplineSweep")
+	void SetScaleEnd(float InScaleEnd);
+
+	UFUNCTION(BlueprintPure, Category = "Motion Design|Modifiers|SplineSweep")
+	float GetScaleEnd() const
+	{
+		return ScaleEnd;
+	}
+
+	UFUNCTION(BlueprintCallable, Category = "Motion Design|Modifiers|SplineSweep")
+	void SetCapped(bool bInCapped);
+
+	UFUNCTION(BlueprintPure, Category = "Motion Design|Modifiers|SplineSweep")
+	bool GetCapped() const
+	{
+		return bCapped;
+	}
+
+	UFUNCTION(BlueprintCallable, Category = "Motion Design|Modifiers|SplineSweep")
+	void SetLooped(bool bInLooped);
+
+	UFUNCTION(BlueprintPure, Category = "Motion Design|Modifiers|SplineSweep")
+	bool GetLooped() const
+	{
+		return bLooped;
+	}
+
+protected:
+	//~ Begin IAvaRenderStateUpdateHandler
+	virtual void OnRenderStateUpdated(AActor* InActor, UActorComponent* InComponent) override;
+	virtual void OnActorVisibilityChanged(AActor* InActor) override {}
+	//~ End IAvaRenderStateUpdateHandler
+
+	//~ Begin UActorModifierCoreBase
+	virtual void OnModifierCDOSetup(FActorModifierCoreMetadata& InMetadata) override;
+	virtual void OnModifierAdded(EActorModifierCoreEnableReason InReason) override;
+	virtual void Apply() override;
+	//~ End UActorModifierCoreBase
+
+	void OnModifierOptionsChanged();
+	void OnSplineActorWeakChanged();
+	void OnSplineOptionsChanged();
+	bool SampleSpline();
+
+	/** Spline actor to retrieve the USplineComponent from */
+	UPROPERTY(EditInstanceOnly, BlueprintReadWrite, Setter="SetSplineActorWeak", Getter="GetSplineActorWeak", Category="SplineSweep", meta=(AllowPrivateAccess="true"))
+	TWeakObjectPtr<AActor> SplineActorWeak;
+
+	/** How do we sample the spline */
+	UPROPERTY(EditInstanceOnly, BlueprintReadWrite, Setter="SetSampleMode", Getter="GetSampleMode", Category="SplineSweep", meta=(AllowPrivateAccess="true"))
+	EAvaSplineSweepSampleMode SampleMode = EAvaSplineSweepSampleMode::FullDistance;
+
+	/** Custom sample distance for steps per distance mode */
+	UPROPERTY(EditInstanceOnly, BlueprintReadWrite, Setter="SetSampleDistance", Getter="GetSampleDistance", Category="SplineSweep", meta=(ClampMin="1", EditCondition="SampleMode == EAvaSplineSweepSampleMode::CustomDistance", EditConditionHides, AllowPrivateAccess="true"))
+	float SampleDistance = 1000.f;
+
+	/** The sample count defines the precision of the sweep */
+	UPROPERTY(EditInstanceOnly, BlueprintReadWrite, Setter="SetSteps", Getter="GetSteps", Category="SplineSweep", meta=(ClampMin="0", ClampMax="100", AllowPrivateAccess="true"))
+	int32 Steps = 10;
+
+	/** Sample range offset of the spline, for closed loop spline range can wrap around */
+	UPROPERTY(EditInstanceOnly, BlueprintReadWrite, Setter="SetProgressOffset", Getter="GetProgressOffset", Category="SplineSweep", meta=(ClampMin="0", AllowPrivateAccess="true"))
+	float ProgressOffset = 0.f;
+
+	/** Sample start range of the spline */
+	UPROPERTY(EditInstanceOnly, BlueprintReadWrite, Setter="SetProgressStart", Getter="GetProgressStart", Category="SplineSweep", meta=(ClampMin="0", ClampMax="1", AllowPrivateAccess="true"))
+	float ProgressStart = 0.f;
+
+	/** Sample end range of the spline */
+	UPROPERTY(EditInstanceOnly, BlueprintReadWrite, Setter="SetProgressEnd", Getter="GetProgressEnd", Category="SplineSweep", meta=(ClampMin="0", ClampMax="1", AllowPrivateAccess="true"))
+	float ProgressEnd = 1.f;
+
+	/** Start scale of the spline mesh */
+	UPROPERTY(EditInstanceOnly, BlueprintReadWrite, Setter="SetScaleStart", Getter="GetScaleStart", Category="SplineSweep", meta=(ClampMin="0", AllowPrivateAccess="true"))
+	float ScaleStart = 1.f;
+
+	/** End scale of the spline mesh */
+	UPROPERTY(EditInstanceOnly, BlueprintReadWrite, Setter="SetScaleEnd", Getter="GetScaleEnd", Category="SplineSweep", meta=(ClampMin="0", AllowPrivateAccess="true"))
+	float ScaleEnd = 1.f;
+
+	/** Whether start and end are closed, loop must be false */
+	UPROPERTY(EditInstanceOnly, BlueprintReadWrite, Setter="SetCapped", Getter="GetCapped", Category="SplineSweep", meta=(EditCondition="!bLooped", AllowPrivateAccess="true"))
+	bool bCapped = true;
+
+	/** Whether we close the whole spline path, if spline loops this will be true */
+	UPROPERTY(EditInstanceOnly, BlueprintReadWrite, Setter="SetLooped", Getter="GetLooped", Category="SplineSweep", meta=(AllowPrivateAccess="true"))
+	bool bLooped = false;
+
+	/** Spline component to sample from */
+	UPROPERTY()
+	TWeakObjectPtr<USplineComponent> SplineComponentWeak;
+
+	/** Sample points taken from the spline */
+	UPROPERTY(Transient)
+	TArray<FTransform> SplineSamples;
+
+private:
+#if WITH_EDITOR
+	/** Used for PECP */
+	static const TAvaPropertyChangeDispatcher<UAvaSplineSweepModifier> PropertyChangeDispatcher;
+#endif
+};
diff --git a/Engine/Plugins/Experimental/Avalanche/Source/AvalancheModifiers/Public/Modifiers/AvaSubdivideModifier.h b/Engine/Plugins/Experimental/Avalanche/Source/AvalancheModifiers/Public/Modifiers/AvaSubdivideModifier.h
new file mode 100644
--- /dev/null
+++ b/Engine/Plugins/Experimental/Avalanche/Source/AvalancheModifiers/Public/Modifiers/AvaSubdivideModifier.h
@@ -0,0 +1,59 @@
+// Copyright Epic Games, Inc. All Rights Reserved.
+
+#pragma once
+
+#include "AvaGeometryBaseModifier.h"
+#include "AvaSubdivideModifier.generated.h"
+
+UENUM()
+enum class EAvaSubdivisionType : uint8
+{
+	Selective,
+	Uniform,
+	PN
+};
+
+UCLASS(BlueprintType)
+class AVALANCHEMODIFIERS_API UAvaSubdivideModifier : public UAvaGeometryBaseModifier
+{
+	GENERATED_BODY()
+
+public:
+	static constexpr int MinSubdivideCuts = 1;
+    static constexpr int MaxSubdivideCuts = 15;
+
+	//~ Begin UObject
+#if WITH_EDITOR
+	virtual void PostEditChangeProperty(FPropertyChangedEvent& PropertyChangedEvent) override;
+#endif
+	//~ End UObject
+	
+	UFUNCTION()
+	void SetCuts(int32 InCuts);
+	int32 GetCuts() const { return Cuts; }
+	
+	UFUNCTION()
+	void SetRecomputeNormals(bool bInRecomputeNormals);
+	bool GetRecomputeNormals() const { return bRecomputeNormals; }
+
+	UFUNCTION()
+	void SetType(EAvaSubdivisionType InType);
+	EAvaSubdivisionType GetType() const { return Type; }
+
+protected:
+	//~ Begin UActorModifierCoreBase
+	virtual void OnModifierCDOSetup(FActorModifierCoreMetadata& InMetadata) override;
+	virtual void Apply() override;
+	//~ End UActorModifierCoreBase
+	
+	void OnOptionsChanged();
+	
+	UPROPERTY(EditInstanceOnly, BlueprintReadWrite, Setter="SetCuts", Getter="GetCuts", Category="Subdivide", meta=(ClampMin="1", ClampMax="15", AllowPrivateAccess="true"))
+	int32 Cuts = 2;
+	
+	UPROPERTY(EditInstanceOnly, BlueprintReadWrite, Setter="SetRecomputeNormals", Getter="GetRecomputeNormals", Category="Subdivide", meta=(EditCondition="Type == EAvaSubdivisionType::PN", EditConditionHides, AllowPrivateAccess="true"))
+	bool bRecomputeNormals = true;
+	
+	UPROPERTY(EditInstanceOnly, BlueprintReadWrite, Setter="SetType", Getter="GetType", Category="Subdivide", meta=(AllowPrivateAccess="true"))
+	EAvaSubdivisionType Type = EAvaSubdivisionType::Uniform;
+};
diff --git a/Engine/Plugins/Experimental/Avalanche/Source/AvalancheModifiers/Public/Modifiers/AvaTaperModifier.h b/Engine/Plugins/Experimental/Avalanche/Source/AvalancheModifiers/Public/Modifiers/AvaTaperModifier.h
new file mode 100644
--- /dev/null
+++ b/Engine/Plugins/Experimental/Avalanche/Source/AvalancheModifiers/Public/Modifiers/AvaTaperModifier.h
@@ -0,0 +1,111 @@
+// Copyright Epic Games, Inc. All Rights Reserved.
+
+#pragma once
+
+#include "AvaGeometryBaseModifier.h"
+#include "Tools/AvaTaperTool.h"
+#include "AvaTaperModifier.generated.h"
+
+UENUM()
+enum class EAvaTaperReferenceFrame : uint8
+{
+	MeshCenter,
+	Custom
+};
+
+UENUM()
+enum class EAvaTaperExtent : uint8
+{
+	WholeShape,
+	Custom
+};
+
+UCLASS(BlueprintType)
+class AVALANCHEMODIFIERS_API UAvaTaperModifier : public UAvaGeometryBaseModifier
+{
+	GENERATED_BODY()
+
+public:
+	static constexpr int MinTaperLatticeResolution = 1;
+	static constexpr int MaxTaperLatticeResolution = 20;
+
+	//~ Begin UObject
+#if WITH_EDITOR
+	virtual void PostEditChangeProperty(FPropertyChangedEvent& PropertyChangedEvent) override;
+#endif
+	//~ End UObject
+	
+	UFUNCTION()
+	void SetAmount(float InAmount);
+	float GetAmount() const { return Amount; }
+
+	UFUNCTION()
+	void SetUpperExtent(float InUpperExtent);
+	float GetUpperExtent() const { return UpperExtent; }
+
+	UFUNCTION()
+	void SetLowerExtent(float InLowerExtent);
+	float GetLowerExtent() const { return LowerExtent; }
+
+	UFUNCTION()
+	void SetExtent(EAvaTaperExtent InExtent);
+	EAvaTaperExtent GetExtent() const { return Extent; }
+	
+	UFUNCTION()
+	void SetInterpolationType(EAvaTaperInterpolationType InInterpolationType);
+	EAvaTaperInterpolationType GetInterpolationType() const { return InterpolationType; }
+
+	UFUNCTION()
+	void SetReferenceFrame(EAvaTaperReferenceFrame InReferenceFrame);
+	EAvaTaperReferenceFrame GetReferenceFrame() const { return ReferenceFrame; }
+
+	UFUNCTION()
+	void SetResolution(int32 InResolution);
+	int32 GetResolution() const { return Resolution; }
+
+	UFUNCTION()
+	void SetOffset(FVector2D InOffset);
+	FVector2D GetOffset() const { return Offset; }
+	
+protected:
+	//~ Begin UActorModifierCoreBase
+	virtual void OnModifierCDOSetup(FActorModifierCoreMetadata& InMetadata) override;
+	virtual void Apply() override;
+	//~ End UActorModifierCoreBase
+	
+	void CreateTaperTool();
+	
+	void OnParameterChanged();
+	
+	FVector2D GetRequiredOffset() const;
+	FVector2D GetRequiredExtent() const;
+
+	int32 GetSubdividersCuts() const;
+	
+	UPROPERTY(EditInstanceOnly, BlueprintReadWrite, Setter="SetAmount", Getter="GetAmount", Category="Taper", meta=(ClampMin="0.0", ClampMax="1.0", AllowPrivateAccess="true"))
+	float Amount = 0.0;
+
+	UPROPERTY(EditInstanceOnly, BlueprintReadWrite, Setter="SetExtent", Getter="GetExtent", Category="Taper", meta=(AllowPrivateAccess="true"))
+	EAvaTaperExtent Extent = EAvaTaperExtent::WholeShape;
+	
+	UPROPERTY(EditInstanceOnly, BlueprintReadWrite, Setter="SetUpperExtent", Getter="GetUpperExtent", Category="Taper", meta=(EditCondition="Extent == EAvaTaperExtent::Custom", EditConditionHides, ClampMin="0", ClampMax="100", Units="Percent", ToolTip="100%: shape top.\n0%: shape bottom.", AllowPrivateAccess="true"))
+	float UpperExtent = 100;
+
+	UPROPERTY(EditInstanceOnly, BlueprintReadWrite, Setter="SetLowerExtent", Getter="GetLowerExtent", Category="Taper", meta=(EditCondition="Extent == EAvaTaperExtent::Custom", EditConditionHides, ClampMin="0", ClampMax="100", Units="Percent", ToolTip="100%: shape bottom.\n0%: shape top.", AllowPrivateAccess="true"))
+	float LowerExtent = 100;
+
+	UPROPERTY(EditInstanceOnly, BlueprintReadWrite, Setter="SetInterpolationType", Getter="GetInterpolationType", Category="Taper", meta=(AllowPrivateAccess="true"))
+	EAvaTaperInterpolationType InterpolationType = EAvaTaperInterpolationType::Linear;
+	
+	UPROPERTY(EditInstanceOnly, BlueprintReadWrite, Setter="SetResolution", Getter="GetResolution", Category="Taper", meta=(ClampMin="1", ClampMax="20", ToolTip="The number of vertical control points used to apply the taper. If the modifier is in a stack with Subdivide modifiers, taper will use the max value between Resolution and the total subdivision cuts.", AllowPrivateAccess="true"))
+	int32 Resolution = 5;
+
+	UPROPERTY(EditInstanceOnly, BlueprintReadWrite, Setter="SetReferenceFrame", Getter="GetReferenceFrame", Category="Taper", meta=(AllowPrivateAccess="true"))
+	EAvaTaperReferenceFrame ReferenceFrame = EAvaTaperReferenceFrame::MeshCenter;
+	
+	UPROPERTY(EditInstanceOnly, BlueprintReadWrite, Setter="SetOffset", Getter="GetOffset", Category="Taper", meta=(EditCondition="ReferenceFrame == EAvaTaperReferenceFrame::Custom", EditConditionHides, AllowPrivateAccess="true"))
+	FVector2D Offset = FVector2D::ZeroVector;
+
+	UPROPERTY()
+	TObjectPtr<UAvaTaperTool> TaperTool = nullptr;
+};
diff --git a/Engine/Plugins/Experimental/Avalanche/Source/AvalancheModifiers/Public/Modifiers/AvaVisibilityModifier.h b/Engine/Plugins/Experimental/Avalanche/Source/AvalancheModifiers/Public/Modifiers/AvaVisibilityModifier.h
new file mode 100644
--- /dev/null
+++ b/Engine/Plugins/Experimental/Avalanche/Source/AvalancheModifiers/Public/Modifiers/AvaVisibilityModifier.h
@@ -0,0 +1,80 @@
+// Copyright Epic Games, Inc. All Rights Reserved.
+
+#pragma once
+
+#include "AvaArrangeBaseModifier.h"
+#include "AvaVisibilityModifier.generated.h"
+
+class AActor;
+
+/**
+ * Controls the visibility of a range of child actors by index.
+ */
+UCLASS(BlueprintType)
+class AVALANCHEMODIFIERS_API UAvaVisibilityModifier : public UAvaArrangeBaseModifier
+{
+	GENERATED_BODY()
+	
+public:
+	//~ Begin UObject
+#if WITH_EDITOR
+	virtual void PostEditChangeProperty(FPropertyChangedEvent& PropertyChangedEvent) override;
+#endif
+	//~ End UObject
+
+	/** Returns true if hiding the child index range instead of showing. */
+	UFUNCTION(BlueprintPure, Category = "Motion Design|Modifiers|Visibility")
+	bool GetInvertVisibility() const { return bInvertVisibility; }
+
+	/** Sets the child index range to hide instead of showing. */
+	UFUNCTION(BlueprintCallable, Category = "Motion Design|Modifiers|Visibility")
+	void SetInvertVisibility(const bool bNewInvertVisibility);
+
+	UFUNCTION(BlueprintPure, Category = "Motion Design|Modifiers|Visibility")
+	int32 GetIndex() const { return Index; }
+	
+	UFUNCTION(BlueprintCallable, Category = "Motion Design|Modifiers|Visibility")
+	void SetIndex(int32 InIndex);
+	
+	UFUNCTION(BlueprintPure, Category = "Motion Design|Modifiers|Visibility")
+	bool GetTreatAsRange() const { return bTreatAsRange; }
+	
+	UFUNCTION(BlueprintCallable, Category = "Motion Design|Modifiers|Visibility")
+	void SetTreatAsRange(const bool bInTreatAsRange);
+
+protected:
+	//~ Begin UActorModifierCoreBase
+	virtual void OnModifierCDOSetup(FActorModifierCoreMetadata& InMetadata) override;
+	virtual void Apply() override;
+	//~ End UActorModifierCoreBase
+	
+	//~ Begin IAvaRenderStateUpdateExtension
+	virtual void OnRenderStateUpdated(AActor* InActor, UActorComponent* InComponent) override;
+	//~ End IAvaRenderStateUpdateExtension
+
+	/** Used by other modifiers of this class to check if we are hiding an actor */
+	bool IsChildActorHidden(AActor* InActor) const;
+
+	/** Get first visibility modifier found above this actor */
+	UAvaVisibilityModifier* GetFirstModifierAbove(AActor* InActor);
+	UAvaVisibilityModifier* GetLastModifierAbove(AActor* InActor);
+
+	/** Will find the direct children of InParentActor that contains this child actor */
+	AActor* GetDirectChildren(AActor* InParentActor, AActor* InChildActor);
+
+	/** Child index to set visibility, visible if bInvertVisibility is false else hidden */
+	UPROPERTY(EditAnywhere, BlueprintReadWrite, Setter="SetIndex", Getter="GetIndex", Category="Visibility", meta=(ClampMin="0", UIMin="0", AllowPrivateAccess="true"))
+	int32 Index = 0;
+
+	/** Treat index as a range from 0 to index */
+	UPROPERTY(EditAnywhere, BlueprintReadWrite, Setter="SetTreatAsRange", Getter="GetTreatAsRange", Category="Visibility", meta=(AllowPrivateAccess="true"))
+	bool bTreatAsRange = false;
+	
+	/** If true, will hide the child index range instead of showing. */
+	UPROPERTY(EditAnywhere, BlueprintReadWrite, Setter="SetInvertVisibility", Getter="GetInvertVisibility", Category="Visibility", meta=(AllowPrivateAccess="true"))
+	bool bInvertVisibility = false;
+
+	/** Visibility defined for direct children of this actor, used by other modifiers of this type to check */
+	UPROPERTY(Transient, DuplicateTransient, TextExportTransient)
+	TMap<TWeakObjectPtr<AActor>, bool> DirectChildrenActorsWeak;
+};
diff --git a/Engine/Plugins/Experimental/Avalanche/Source/AvalancheModifiers/Public/Shared/AvaTransformModifierShared.h b/Engine/Plugins/Experimental/Avalanche/Source/AvalancheModifiers/Public/Shared/AvaTransformModifierShared.h
new file mode 100644
--- /dev/null
+++ b/Engine/Plugins/Experimental/Avalanche/Source/AvalancheModifiers/Public/Shared/AvaTransformModifierShared.h
@@ -0,0 +1,141 @@
+// Copyright Epic Games, Inc. All Rights Reserved.
+
+#pragma once
+
+#include "Modifiers/ActorModifierCoreSharedObject.h"
+#include "Modifiers/AvaBaseModifier.h"
+#include "GameFramework/Actor.h"
+#include "UObject/WeakObjectPtrTemplates.h"
+#include "AvaTransformModifierShared.generated.h"
+
+class AActor;
+class UAvaBaseModifier;
+
+UENUM()
+enum class EAvaTransformSharedModifier : uint8
+{
+	None,
+	Location = 1 << 0,
+	Rotation = 1 << 1,
+	Scale = 1 << 2,
+	LocationRotation = Location | Rotation,
+	LocationScale = Location | Scale,
+	RotationScale = Rotation | Scale,
+	All = Location | Rotation | Scale
+};
+
+USTRUCT()
+struct AVALANCHEMODIFIERS_API FAvaTransformSharedModifierState
+{
+	GENERATED_BODY()
+
+	FAvaTransformSharedModifierState() {}
+	FAvaTransformSharedModifierState(UAvaBaseModifier* InModifier)
+		: ModifierWeak(InModifier)
+	{}
+
+	/** Save this modifier state if valid */
+	void Save(const AActor* InActor);
+
+	/** Restore this modifier state if valid */
+	void Restore(AActor* InActor, EAvaTransformSharedModifier InRestoreState) const;
+	
+	friend uint32 GetTypeHash(const FAvaTransformSharedModifierState& InItem)
+	{
+		return GetTypeHash(InItem.ModifierWeak);
+	}
+
+	bool operator==(const FAvaTransformSharedModifierState& Other) const
+	{
+		return ModifierWeak == Other.ModifierWeak;
+	}
+
+	UPROPERTY()
+	TWeakObjectPtr<UAvaBaseModifier> ModifierWeak;
+	
+	UPROPERTY()
+	FTransform ActorTransform;
+};
+
+USTRUCT()
+struct AVALANCHEMODIFIERS_API FAvaTransformSharedActorState
+{
+	GENERATED_BODY()
+
+	FAvaTransformSharedActorState() {}
+	FAvaTransformSharedActorState(AActor* InActor)
+		: ActorWeak(InActor)
+	{}
+
+	/** Save this actor state if valid */
+	void Save();
+
+	/** Restore this actor state if valid */
+	void Restore(EAvaTransformSharedModifier InRestoreState) const;
+
+	friend uint32 GetTypeHash(const FAvaTransformSharedActorState& InItem)
+	{
+		return GetTypeHash(InItem.ActorWeak);
+	}
+
+	bool operator==(const FAvaTransformSharedActorState& Other) const
+	{
+		return ActorWeak == Other.ActorWeak;
+	}
+
+	/** Modifiers that are currently watching this state and locking it */
+	UPROPERTY()
+	TSet<FAvaTransformSharedModifierState> ModifierStates;
+
+	/** Actor that this state is describing */
+	UPROPERTY()
+	TWeakObjectPtr<AActor> ActorWeak;
+
+	/** Pre state transform saved */
+	UPROPERTY()
+	FTransform ActorTransform;
+};
+
+/**
+ * Singleton class for transform modifiers to share data between each other
+ * Used because multiple modifier could be watching/updating an actor
+ * We want to save the state of that actor once before any modifier changes it
+ * and restore it when no other modifier is watching it
+ */
+UCLASS()
+class AVALANCHEMODIFIERS_API UAvaTransformModifierShared : public UActorModifierCoreSharedObject
+{
+	GENERATED_BODY()
+	
+public:
+	/** Save actor state, adds it if it is not tracked */
+	void SaveActorState(UAvaBaseModifier* InModifierContext, AActor* InActor);
+
+	/** Restore actor state, removes it if no other modifier track that actor state */
+	void RestoreActorState(UAvaBaseModifier* InModifierContext, AActor* InActor, EAvaTransformSharedModifier InRestoreTransform = EAvaTransformSharedModifier::All);
+
+	/** Get the actor state for a specific actor */
+	FAvaTransformSharedActorState* FindActorState(AActor* InActor);
+
+	/** Get all actor state related to a modifier */
+	TSet<FAvaTransformSharedActorState*> FindActorsState(UAvaBaseModifier* InModifierContext);
+
+	/** Restore all actors states linked to this modifier */
+	void RestoreActorsState(UAvaBaseModifier* InModifierContext, const TSet<AActor*>* InActors = nullptr, EAvaTransformSharedModifier InRestoreTransform = EAvaTransformSharedModifier::All);
+
+	/** Restore all specified actors linked to this modifier */
+	void RestoreActorsState(UAvaBaseModifier* InModifierContext, const TSet<TWeakObjectPtr<AActor>>& InActors, EAvaTransformSharedModifier InRestoreTransform = EAvaTransformSharedModifier::All);
+
+	/** Returns true, if this modifier is tracking this actor */
+	bool IsActorStateSaved(UAvaBaseModifier* InModifierContext, AActor* InActor);
+
+	/** Returns true, if this modifier is tracking any actor */
+	bool IsActorsStateSaved(UAvaBaseModifier* InModifierContext);
+	
+protected:
+	virtual void PostLoad() override;
+	
+	/** Actor state before any modifier applied to it */
+	UPROPERTY()
+	TSet<FAvaTransformSharedActorState> ActorStates;
+};
diff --git a/Engine/Plugins/Experimental/Avalanche/Source/AvalancheModifiers/Public/Shared/AvaVisibilityModifierShared.h b/Engine/Plugins/Experimental/Avalanche/Source/AvalancheModifiers/Public/Shared/AvaVisibilityModifierShared.h
new file mode 100644
--- /dev/null
+++ b/Engine/Plugins/Experimental/Avalanche/Source/AvalancheModifiers/Public/Shared/AvaVisibilityModifierShared.h
@@ -0,0 +1,152 @@
+// Copyright Epic Games, Inc. All Rights Reserved.
+
+#pragma once
+
+#include "Modifiers/ActorModifierCoreSharedObject.h"
+#include "Modifiers/AvaBaseModifier.h"
+#include "GameFramework/Actor.h"
+#include "UObject/WeakObjectPtrTemplates.h"
+#include "AvaVisibilityModifierShared.generated.h"
+
+class AActor;
+class UAvaBaseModifier;
+
+UENUM(meta=(Bitflags))
+enum class EAvaVisibilityActor : uint8
+{
+	None          = 0,
+	Game          = 1 << 0,
+	Editor        = 1 << 1,
+	GameAndEditor = Game | Editor
+};
+ENUM_CLASS_FLAGS(EAvaVisibilityActor);
+
+USTRUCT()
+struct AVALANCHEMODIFIERS_API FAvaVisibilitySharedModifierState
+{
+	GENERATED_BODY()
+
+	FAvaVisibilitySharedModifierState() {}
+	FAvaVisibilitySharedModifierState(UAvaBaseModifier* InModifier)
+		: ModifierWeak(InModifier)
+	{}
+
+	/** Save this modifier state if valid */
+	void Save(const AActor* InActor);
+
+	/** Restore this modifier state if valid */
+	void Restore(AActor* InActor) const;
+	
+	friend uint32 GetTypeHash(const FAvaVisibilitySharedModifierState& InItem)
+	{
+		return GetTypeHash(InItem.ModifierWeak);
+	}
+
+	bool operator==(const FAvaVisibilitySharedModifierState& Other) const
+	{
+		return ModifierWeak == Other.ModifierWeak;
+	}
+
+	UPROPERTY()
+	TWeakObjectPtr<UAvaBaseModifier> ModifierWeak;
+
+#if WITH_EDITORONLY_DATA
+	/** Pre state editor visibility saved */
+	UPROPERTY()
+	bool bActorHiddenInEditor = false;
+#endif
+
+	/** Pre state game visibility saved */
+	UPROPERTY()
+	bool bActorHiddenInGame = false;
+};
+
+USTRUCT()
+struct AVALANCHEMODIFIERS_API FAvaVisibilitySharedActorState
+{
+	GENERATED_BODY()
+
+	FAvaVisibilitySharedActorState() {}
+	FAvaVisibilitySharedActorState(AActor* InActor)
+		: ActorWeak(InActor)
+	{}
+	
+	/** Save this actor state if valid */
+	void Save();
+
+	/** Restore this actor state if valid */
+	void Restore() const;
+	
+	friend uint32 GetTypeHash(const FAvaVisibilitySharedActorState& InItem)
+	{
+		return GetTypeHash(InItem.ActorWeak);
+	}
+
+	bool operator==(const FAvaVisibilitySharedActorState& Other) const
+	{
+		return ActorWeak == Other.ActorWeak;
+	}
+
+	/** Modifiers that are currently watching this state and locking it */
+	UPROPERTY()
+	TSet<FAvaVisibilitySharedModifierState> ModifierStates;
+
+	/** Actor that this state is describing */
+	UPROPERTY()
+	TWeakObjectPtr<AActor> ActorWeak;
+	
+#if WITH_EDITORONLY_DATA
+	/** Pre state editor visibility saved */
+	UPROPERTY()
+	bool bActorHiddenInEditor = false;
+#endif
+
+	/** Pre state game visibility saved */
+	UPROPERTY()
+	bool bActorHiddenInGame = false;
+};
+
+/**
+ * Singleton class for visibility modifiers to share data between each other
+ * Used because multiple modifier could be watching/updating an actor
+ * We want to save the state of that actor once before any modifier changes it
+ * and restore it when no other modifier is watching it
+ */
+UCLASS()
+class AVALANCHEMODIFIERS_API UAvaVisibilityModifierShared : public UActorModifierCoreSharedObject
+{
+	GENERATED_BODY()
+	
+public:
+	/** Watch actor state, adds it if it is not tracked */
+	void SaveActorState(UAvaBaseModifier* InModifierContext, AActor* InActor);
+
+	/** Unwatch actor state, removes it if no other modifier track that actor state */
+	void RestoreActorState(UAvaBaseModifier* InModifierContext, AActor* InActor);
+
+	/** Gather original state before any modifier is applied if there is one */
+	FAvaVisibilitySharedActorState* FindActorState(AActor* InActor);
+
+	/** Set actor visibility in game or editor and recurse, tracks original state if not tracked */
+	void SetActorVisibility(UAvaBaseModifier* InModifierContext, AActor* InActor, bool bInHidden, bool bInRecurse = false, EAvaVisibilityActor InActorVisibility = EAvaVisibilityActor::GameAndEditor);
+
+	/** Set actors visibility in game or editor, tracks original state if not tracked */
+	void SetActorsVisibility(UAvaBaseModifier* InModifierContext, TArray<AActor*> InActors, bool bInHidden, EAvaVisibilityActor InActorVisibility = EAvaVisibilityActor::GameAndEditor);
+
+	/** Unwatch all actors states linked to this modifier */
+	void RestoreActorsState(UAvaBaseModifier* InModifierContext, const TSet<AActor*>* InActors = nullptr);
+
+	void RestoreActorsState(UAvaBaseModifier* InModifierContext, const TSet<TWeakObjectPtr<AActor>>& InActors);
+	
+	/** Returns true, if this modifier is tracking this actor */
+	bool IsActorStateSaved(UAvaBaseModifier* InModifierContext, AActor* InActor);
+
+	/** Returns true, if this modifier is tracking any actor */
+	bool IsActorsStateSaved(UAvaBaseModifier* InModifierContext);
+protected:
+	virtual void PostLoad() override;
+	
+	/** Actor state before any modifier applied to it */
+	UPROPERTY()
+	TSet<FAvaVisibilitySharedActorState> ActorStates;
+};
diff --git a/Engine/Plugins/Experimental/Avalanche/Source/AvalancheModifiersEditor/Private/AvaModifiersEditorStyle.cpp b/Engine/Plugins/Experimental/Avalanche/Source/AvalancheModifiersEditor/Private/AvaModifiersEditorStyle.cpp
new file mode 100644
--- /dev/null
+++ b/Engine/Plugins/Experimental/Avalanche/Source/AvalancheModifiersEditor/Private/AvaModifiersEditorStyle.cpp
@@ -0,0 +1,85 @@
+// Copyright Epic Games, Inc. All Rights Reserved.
+
+#include "AvaModifiersEditorStyle.h"
+#include "GenericPlatform/GenericPlatformFile.h"
+#include "HAL/PlatformFileManager.h"
+#include "Interfaces/IPluginManager.h"
+#include "Misc/Paths.h"
+#include "Modifiers/ActorModifierCoreBase.h"
+#include "Styling/SlateStyleMacros.h"
+#include "Styling/StyleColors.h"
+#include "Styling/SlateTypes.h"
+#include "Styling/SlateStyleRegistry.h"
+#include "Subsystems/ActorModifierCoreSubsystem.h"
+
+FAvaModifiersEditorStyle::FAvaModifiersEditorStyle()
+	: FSlateStyleSet(TEXT("AvalancheModifiersEditor"))
+{
+	const FVector2f Icon16x16(16.f, 16.f);
+
+	const TSharedPtr<IPlugin> Plugin = IPluginManager::Get().FindPlugin(TEXT("Avalanche"));
+
+	check(Plugin.IsValid());
+
+	ContentRootDir = FPaths::Combine(Plugin->GetBaseDir(), TEXT("Resources"));
+
+	Set("ClassIcon.AvaBaseModifier", new IMAGE_BRUSH_SVG("Icons/ModifierIcons/BaseModifier", Icon16x16));
+
+	ModifierCategoriesColors =
+	{
+		{TEXT("Unspecified"), FSlateColor(FActorModifierCoreMetadata::DefaultColor)},
+		{TEXT("Geometry"), FSlateColor(FLinearColor::Yellow.Desaturate(0.25))},
+		{TEXT("Transform"), FSlateColor(FLinearColor::Blue.Desaturate(0.25))},
+		{TEXT("Layout"), FSlateColor(FLinearColor::Green.Desaturate(0.25))},
+		{TEXT("Rendering"), FSlateColor(FLinearColor::Red.Desaturate(0.25))}
+	};
+
+	if (const UActorModifierCoreSubsystem* ModifierSubsystem = UActorModifierCoreSubsystem::Get())
+	{
+		// loop through already registered factories if subsystem exists
+		ModifierSubsystem->ForEachModifierMetadata([this](FActorModifierCoreMetadata& InMetadata)->bool
+		{
+			OnModifierClassRegistered(InMetadata);
+			return true;
+		});
+	}
+
+	FSlateStyleRegistry::RegisterSlateStyle(*this);
+
+	UActorModifierCoreSubsystem::OnModifierClassRegisteredDelegate.AddRaw(this, &FAvaModifiersEditorStyle::OnModifierClassRegistered);
+}
+
+const FSlateColor& FAvaModifiersEditorStyle::GetModifierCategoryColor(FName CategoryName)
+{
+	if (ModifierCategoriesColors.Contains(CategoryName))
+	{
+		return ModifierCategoriesColors[CategoryName];
+	}
+
+	return ModifierCategoriesColors[TEXT("Unspecified")];
+}
+
+void FAvaModifiersEditorStyle::OnModifierClassRegistered(const FActorModifierCoreMetadata& InMetadata)
+{
+	// We need to update the metadata to set various options
+	FActorModifierCoreMetadata& MutableMetadata = const_cast<FActorModifierCoreMetadata&>(InMetadata);
+
+	const FVector2f Icon16x16(16.f, 16.f);
+	const FString ModifierIconPath = TEXT("Icons/ModifierIcons/") + InMetadata.GetName().ToString() + TEXT("Modifier");
+
+	MutableMetadata.SetColor(GetModifierCategoryColor(InMetadata.GetCategory()).GetSpecifiedColor());
+
+	if (FPlatformFileManager::Get().GetPlatformFile().FileExists(*RootToContentDir(ModifierIconPath, TEXT(".svg"))))
+	{
+		const FName ModifierClassStyleName(TEXT("ClassIcon.") + InMetadata.GetClass()->GetName());
+		Set(ModifierClassStyleName, new IMAGE_BRUSH_SVG(ModifierIconPath, Icon16x16));
+		MutableMetadata.SetIcon(FSlateIcon(GetStyleSetName(), ModifierClassStyleName));
+	}
+}
+
+FAvaModifiersEditorStyle::~FAvaModifiersEditorStyle()
+{
+	UActorModifierCoreSubsystem::OnModifierClassRegisteredDelegate.RemoveAll(this);
+
+	FSlateStyleRegistry::UnRegisterSlateStyle(*this);
+}
\ No newline at end of file
diff --git a/Engine/Plugins/Experimental/Avalanche/Source/AvalancheModifiersEditor/Private/AvaModifiersEditorStyle.h b/Engine/Plugins/Experimental/Avalanche/Source/AvalancheModifiersEditor/Private/AvaModifiersEditorStyle.h
new file mode 100644
--- /dev/null
+++ b/Engine/Plugins/Experimental/Avalanche/Source/AvalancheModifiersEditor/Private/AvaModifiersEditorStyle.h
@@ -0,0 +1,29 @@
+// Copyright Epic Games, Inc. All Rights Reserved.
+
+#pragma once
+
+#include "Delegates/IDelegateInstance.h"
+#include "Styling/SlateStyle.h"
+
+struct FActorModifierCoreMetadata;
+
+class FAvaModifiersEditorStyle : public FSlateStyleSet
+{
+public:
+	FAvaModifiersEditorStyle();
+
+	virtual ~FAvaModifiersEditorStyle() override;
+
+	static FAvaModifiersEditorStyle& Get()
+	{
+		static FAvaModifiersEditorStyle StyleSet;
+		return StyleSet;
+	}
+
+	const FSlateColor& GetModifierCategoryColor(FName CategoryName);
+
+private:
+	void OnModifierClassRegistered(const FActorModifierCoreMetadata& InMetadata);
+
+	TMap<FName, FSlateColor> ModifierCategoriesColors;
+};
diff --git a/Engine/Plugins/Experimental/Avalanche/Source/AvalancheOutliner/Private/AvaOutlinerModule.h b/Engine/Plugins/Experimental/Avalanche/Source/AvalancheOutliner/Private/AvaOutlinerModule.h
new file mode 100644
--- /dev/null
+++ b/Engine/Plugins/Experimental/Avalanche/Source/AvalancheOutliner/Private/AvaOutlinerModule.h
@@ -0,0 +1,114 @@
+// Copyright Epic Games, Inc. All Rights Reserved.
+
+#pragma once
+
+#include "IAvaOutlinerModule.h"
+#include "ItemProxies/AvaOutlinerItemProxyRegistry.h"
+
+class IAvaFilterExpressionFactory;
+
+class FAvaOutlinerModule : public IAvaOutlinerModule
+{
+public:
+	//~ Begin IModuleInterface
+	virtual void StartupModule() override;
+	virtual void ShutdownModule() override;
+	//~ End IModuleInterface
+
+	//~ Begin IAvaOutlinerModule
+	virtual FOnInitOutlinerToolMenu& GetOnInitOutlinerToolBar() override { return OnInitOutlinerToolBar; }
+	virtual FOnInitOutlinerToolMenu& GetOnInitOutlinerItemContextMenu() override { return OnInitOutlinerItemContextMenu; }
+	virtual FOnExtendOutlinerToolMenu& GetOnExtendOutlinerToolBar() override  { return OnExtendOutlinerToolBar; }
+	virtual FOnExtendOutlinerToolMenu& GetOnExtendOutlinerItemContextMenu() override   { return OnExtendOutlinerItemContextMenu; }
+	virtual FOnExtendOutlinerColumns& GetOnExtendOutlinerColumns() override { return OnExtendOutlinerColumns; }
+	virtual FAvaOutlinerItemProxyRegistry& GetItemProxyRegistry() override { return ItemProxyRegistry; }
+	virtual FOnExtendItemProxiesForItem& GetOnExtendItemProxiesForItem() override { return OnGetItemProxiesForItem; }
+	virtual void UnregisterFilterExpressionFactory(const FName& InFilterIdentifier) override;
+	virtual void UnregisterFilterSuggestionFactory(const FName& InFilterIdentifier) override;
+	virtual bool CanFilterSupportComparisonOperation(const FName& InFilterKey, const ETextFilterComparisonOperation InOperation) const override;
+	virtual bool FilterExpression(const FName& InFilterKey, const IAvaOutlinerItem& InItem, const FAvaTextFilterArgs& InArgs) const override;
+	virtual TArray<TSharedPtr<IAvaFilterSuggestionFactory>> GetSuggestions(const EAvaFilterSuggestionType& InSuggestionType) const override;
+	virtual FSlateIcon FindOverrideIcon(TSharedPtr<const FAvaOutlinerItem> InItem) const override;
+	//~ End IAvaOutlinerModule
+
+protected:
+	/** Delegate to extend the Tool Menu Context of the Outliner ToolBar */
+	FOnInitOutlinerToolMenu OnInitOutlinerToolBar;
+
+	/** Delegate to extend the Tool Menu Context of the Outliner Item Context Menu */
+	FOnInitOutlinerToolMenu OnInitOutlinerItemContextMenu;
+
+	/** Delegate to extend the Outliner ToolBar */
+	FOnExtendOutlinerToolMenu OnExtendOutlinerToolBar;
+
+	/** Delegate to extend the Item Context Menu */
+	FOnExtendOutlinerToolMenu OnExtendOutlinerItemContextMenu;
+
+	/** Delegates to add in additional Outliner Columns */
+	FOnExtendOutlinerColumns OnExtendOutlinerColumns;
+
+	/** Delegate to add in Item Proxies for a given Item / Item Type */
+	FOnExtendItemProxiesForItem OnGetItemProxiesForItem;
+
+	/** Module Instance of the Item Registry */
+	FAvaOutlinerItemProxyRegistry ItemProxyRegistry;
+
+	//~ Begin IAvaOutlinerModule
+	virtual void UnregisterOverriddenIcon_Internal(const FAvaTypeId& InItemTypeId, const FName& InSpecializationIdentifier) override;
+	virtual void RegisterOverridenIcon_Internal(const FAvaTypeId& InItemTypeId, const TSharedRef<IAvaOutlinerIconCustomization>& InAvaOutlinerIconCustomization) override;
+	virtual void RegisterFilterExpressionFactory_Internal(const TSharedRef<IAvaFilterExpressionFactory>& InFilterExpressionFactory) override;
+	virtual void RegisterFilterSuggestionFactory_Internal(const TSharedRef<IAvaFilterSuggestionFactory>& InFilterSuggestionFactory) override;
+	//~ End IAvaOutlinerModule
+
+private:
+	/**
+	 * Get the customization for the given item if any
+	 * @param InItem Item to search the customization for
+	 * @return The customization for the item if any, nullptr otherwise
+	 */
+	TSharedPtr<IAvaOutlinerIconCustomization> GetCustomizationForItem(const TSharedPtr<const FAvaOutlinerItem>& InItem) const;
+
+	/**
+	 * Registers default filter expression factories.
+	 */
+	void RegisterFilterExpressionFactories();
+
+	/**
+	  * Registers default filter suggestion factories.
+	  */
+	void RegisterFilterSuggestionFactories();
+
+private:
+	/** Holds all the FilterExpressionFactory */
+	TMap<FName, TSharedPtr<IAvaFilterExpressionFactory>> FilterExpressionFactories;
+
+	/** Holds all the FilterSuggestionFactory */
+	TMap<FName, TSharedPtr<IAvaFilterSuggestionFactory>> FilterSuggestionsFactories;
+
+	/**
+	 * Hold the key of the map containing the IconCustomizations
+	 */
+	struct FIconCustomizationKey
+	{
+		/** OutlinerClass FName (ex. AvaOutlinerItem) */
+		FAvaTypeId ItemTypeId = FAvaTypeId::Invalid();
+
+		/** Specialization identifier (ex. AvaShapeActor class FName for ActorIcon customization) */
+		FName CustomizationSpecializationIdentifier;
+
+		bool operator==(const FIconCustomizationKey& InOther) const
+		{
+			return ItemTypeId == InOther.ItemTypeId
+				&& CustomizationSpecializationIdentifier == InOther.CustomizationSpecializationIdentifier;
+		}
+
+		friend uint32 GetTypeHash(FIconCustomizationKey const& InCustomizationKey)
+		{
+			return HashCombine(GetTypeHash(InCustomizationKey.ItemTypeId)
+				, GetTypeHash(InCustomizationKey.CustomizationSpecializationIdentifier));
+		}
+	};
+
+	/** Holds all the IconCustomizations */
+	TMap<FIconCustomizationKey, TSharedPtr<IAvaOutlinerIconCustomization>> IconRegistry;
+};
diff --git a/Engine/Plugins/Experimental/Avalanche/Source/AvalancheOutliner/Private/AvaOutlinerStyle.cpp b/Engine/Plugins/Experimental/Avalanche/Source/AvalancheOutliner/Private/AvaOutlinerStyle.cpp
new file mode 100644
--- /dev/null
+++ b/Engine/Plugins/Experimental/Avalanche/Source/AvalancheOutliner/Private/AvaOutlinerStyle.cpp
@@ -0,0 +1,89 @@
+// Copyright Epic Games, Inc. All Rights Reserved.
+
+#include "AvaOutlinerStyle.h"
+#include "Interfaces/IPluginManager.h"
+#include "Styling/AppStyle.h"
+#include "Styling/SlateStyleRegistry.h"
+#include "Styling/StyleColors.h"
+#include "Styling/SlateTypes.h"
+#include "Styling/ToolBarStyle.h"
+
+TSharedPtr<FSlateStyleSet> FAvaOutlinerStyle::StyleInstance;
+
+void FAvaOutlinerStyle::Initialize()
+{
+	if (!StyleInstance.IsValid())
+	{
+		StyleInstance = Create();
+		FSlateStyleRegistry::RegisterSlateStyle(*StyleInstance);
+	}
+}
+
+void FAvaOutlinerStyle::Shutdown()
+{
+	FSlateStyleRegistry::UnRegisterSlateStyle(*StyleInstance);
+	ensure(StyleInstance.IsUnique());
+	StyleInstance.Reset();
+}
+
+FName FAvaOutlinerStyle::GetStyleSetName()
+{
+	static FName StyleSetName(TEXT("AvalancheOutliner"));
+	return StyleSetName;
+}
+
+const ISlateStyle& FAvaOutlinerStyle::Get()
+{
+	return *StyleInstance;
+}
+
+const FLinearColor& FAvaOutlinerStyle::GetColor(FName InPropertyName, const ANSICHAR* InSpecifier)
+{
+	return StyleInstance->GetColor(InPropertyName, InSpecifier);
+}
+
+const FSlateBrush* FAvaOutlinerStyle::GetBrush(FName InPropertyName, const ANSICHAR* InSpecifier)
+{
+	return StyleInstance->GetBrush(InPropertyName, InSpecifier);
+}
+
+#define BORDER_BRUSH(RelativePath, ...) FSlateBorderBrush(Style->RootToContentDir(RelativePath, TEXT(".png")), __VA_ARGS__)
+#define BOX_BRUSH(RelativePath, ...) FSlateBoxBrush(Style->RootToContentDir(RelativePath, TEXT(".png")), __VA_ARGS__)
+#define CORE_IMAGE_BRUSH(RelativePath, ...) FSlateImageBrush(Style->RootToCoreContentDir(RelativePath, TEXT(".png")), __VA_ARGS__)
+#define IMAGE_BRUSH(RelativePath, ...) FSlateImageBrush(Style->RootToContentDir(RelativePath, TEXT(".png")), __VA_ARGS__)
+
+const FVector2D Icon8x8(8.f, 8.f);
+const FVector2D Icon64x64(64.f, 64.f);
+const FVector2D Icon40x40(40.0f, 40.0f);
+const FVector2D Icon24x24(24.0f, 24.0f);
+const FVector2D Icon20x20(20.0f, 20.0f);
+const FVector2D Icon16x16(16.0f, 16.0f);
+const FVector2D Icon12x12(12.0f, 12.0f);
+
+TSharedRef<FSlateStyleSet> FAvaOutlinerStyle::Create()
+{
+	TSharedRef<FSlateStyleSet> Style = MakeShared<FSlateStyleSet>(GetStyleSetName());
+	const TSharedPtr<IPlugin> Plugin = IPluginManager::Get().FindPlugin(TEXT("Avalanche"));
+
+	check(Plugin.IsValid());
+	if (Plugin.IsValid())
+	{
+		Style->SetContentRoot(FPaths::Combine(Plugin->GetBaseDir(), TEXT("Resources")));
+		Style->SetCoreContentRoot(FPaths::EngineContentDir() / TEXT("Slate"));
+	}
+	
+	Style->Set("AvalancheOutliner.FilterIcon", new IMAGE_BRUSH("Icons/OutlinerIcons/FilterIcon", Icon20x20));
+	
+	// Table View Row Style
+	const FTableRowStyle& TableRowStyle = FAppStyle::Get().GetWidgetStyle<FTableRowStyle>("SceneOutliner.TableViewRow");
+	Style->Set("AvalancheOutliner.TableViewRow", FTableRowStyle(TableRowStyle)
+		.SetEvenRowBackgroundBrush(FSlateNoResource())
+		.SetOddRowBackgroundBrush(FSlateColorBrush(FStyleColors::Background)));
+	
+	return Style;
+}
+
+#undef BORDER_BRUSH
+#undef BOX_BRUSH
+#undef CORE_IMAGE_BRUSH
+#undef IMAGE_BRUSH
diff --git a/Engine/Plugins/Experimental/Avalanche/Source/AvalancheOutliner/Private/AvaOutlinerStyle.h b/Engine/Plugins/Experimental/Avalanche/Source/AvalancheOutliner/Private/AvaOutlinerStyle.h
new file mode 100644
--- /dev/null
+++ b/Engine/Plugins/Experimental/Avalanche/Source/AvalancheOutliner/Private/AvaOutlinerStyle.h
@@ -0,0 +1,32 @@
+// Copyright Epic Games, Inc. All Rights Reserved.
+
+#pragma once
+
+#include "Styling/ISlateStyle.h"
+#include "Styling/SlateStyle.h"
+
+struct FLinearColor;
+struct FSlateBrush;
+
+class FAvaOutlinerStyle
+{
+public:
+	static void Initialize();
+	static void Shutdown();
+
+	static FName GetStyleSetName();
+	static const ISlateStyle& Get();
+
+	static const FLinearColor& GetColor(FName InPropertyName, const ANSICHAR* InSpecifier = nullptr);
+	static const FSlateBrush* GetBrush(FName InPropertyName, const ANSICHAR* InSpecifier = nullptr);
+
+	template<typename InWidgetStyleType>
+	static const InWidgetStyleType& GetWidgetStyle(FName InPropertyName, const ANSICHAR* InSpecifier = nullptr)
+	{
+		return StyleInstance->GetWidgetStyle<InWidgetStyleType>(InPropertyName, InSpecifier);
+	}
+
+private:
+	static TSharedRef<FSlateStyleSet> Create();
+	static TSharedPtr<FSlateStyleSet> StyleInstance;
+};
diff --git a/Engine/Plugins/Experimental/Avalanche/Source/AvalancheOutliner/Private/Item/AvaOutlinerTreeRoot.cpp b/Engine/Plugins/Experimental/Avalanche/Source/AvalancheOutliner/Private/Item/AvaOutlinerTreeRoot.cpp
new file mode 100644
--- /dev/null
+++ b/Engine/Plugins/Experimental/Avalanche/Source/AvalancheOutliner/Private/Item/AvaOutlinerTreeRoot.cpp
@@ -0,0 +1,172 @@
+// Copyright Epic Games, Inc. All Rights Reserved.
+
+#include "Item/AvaOutlinerTreeRoot.h"
+
+#include "ActorFactories/ActorFactory.h"
+#include "AssetSelection.h"
+#include "AvaOutliner.h"
+#include "AvaOutlinerSubsystem.h"
+#include "DragAndDrop/AssetDragDropOp.h"
+#include "DragDropOps/AvaOutlinerItemDragDropOp.h"
+#include "EngineUtils.h"
+#include "Item/AvaOutlinerActor.h"
+#include "ItemActions/AvaOutlinerAddItem.h"
+#include "Textures/SlateIcon.h"
+#include "Widgets/SNullWidget.h"
+#include "Widgets/Views/STableRow.h"
+
+void FAvaOutlinerTreeRoot::FindChildren(TArray<FAvaOutlinerItemPtr>& OutChildren, bool bRecursive)
+{
+	Super::FindChildren(OutChildren, bRecursive);
+
+	UWorld* const World = Outliner.GetWorld();
+	if (!World)
+	{
+		return;
+	}
+
+	// All the Actors that don't have an Outliner Parent will be at the Root
+	const TArray<TWeakObjectPtr<AActor>> RootChildren = Outliner.GetActorSceneOutlinerChildren(nullptr);
+
+	// Worst case and most likely case: all Outliner Children are valid.
+	// Note: if recursive, re-allocations will still need to be done past this reserve count, as it's unknown at this point how many items are going to be added
+	OutChildren.Reserve(OutChildren.Num() + RootChildren.Num());
+
+	for (const TWeakObjectPtr<AActor>& ChildWeak : RootChildren)
+	{
+		if (AActor* const Child = ChildWeak.Get())
+		{
+			const FAvaOutlinerItemPtr ChildActorItem = Outliner.FindOrAdd<FAvaOutlinerActor>(Child);
+			const FAvaOutlinerItemFlagGuard Guard(ChildActorItem, EAvaOutlinerItemFlags::IgnorePendingKill);
+			OutChildren.Add(ChildActorItem);
+			if (bRecursive)
+			{
+				ChildActorItem->FindChildren(OutChildren, bRecursive);
+			}
+		}
+	}
+}
+
+bool FAvaOutlinerTreeRoot::CanAddChild(const FAvaOutlinerItemPtr& InChild) const
+{
+	return Super::CanAddChild(InChild)
+		&& InChild->CanBeTopLevel();
+}
+
+bool FAvaOutlinerTreeRoot::AddChild(const FAvaOutlinerAddItemParams& InAddItemParams)
+{
+	// if current parent is the root, it means it's just rearranging
+	const bool bRearranging = InAddItemParams.Item.IsValid() && InAddItemParams.Item->GetParent().Get() == this;
+
+	// Is it a new actor that we just spawned
+	const bool bSpawning = InAddItemParams.Item.IsValid() && !Children.Contains(InAddItemParams.Item);
+	
+	const bool bResult = FAvaOutlinerItem::AddChild(InAddItemParams);
+	
+	if ((bSpawning || bRearranging) && InAddItemParams.Item->IsA<FAvaOutlinerActor>())
+	{
+		UAvaOutlinerSubsystem* const OutlinerSubsystem = Outliner.GetOutlinerSubsystem();
+		
+		AActor* const Actor = StaticCastSharedPtr<FAvaOutlinerActor>(InAddItemParams.Item)->GetActor();
+		
+		if (OutlinerSubsystem && Actor)
+		{
+			const EAvaOutlinerHierarchyChangeType Type = bSpawning ? EAvaOutlinerHierarchyChangeType::Attached : EAvaOutlinerHierarchyChangeType::Rearranged;
+			
+			OutlinerSubsystem->BroadcastActorHierarchyChanged(Actor, nullptr, Type);
+		}
+	}
+	
+	return bResult;
+}
+
+bool FAvaOutlinerTreeRoot::IsAllowedInOutliner() const
+{
+	checkNoEntry();
+	return false;
+}
+
+FText FAvaOutlinerTreeRoot::GetDisplayName() const
+{
+	checkNoEntry();
+	return FText::GetEmpty();
+}
+
+FText FAvaOutlinerTreeRoot::GetClassName() const
+{
+	checkNoEntry();
+	return FText::GetEmpty();
+}
+
+FText FAvaOutlinerTreeRoot::GetIconTooltipText() const
+{
+	checkNoEntry();
+	return FText::GetEmpty();
+}
+
+FSlateIcon FAvaOutlinerTreeRoot::GetIcon() const
+{
+	checkNoEntry();
+	return FSlateIcon();
+}
+
+TSharedRef<SWidget> FAvaOutlinerTreeRoot::GenerateLabelWidget(const TSharedRef<SAvaOutlinerTreeRow>& InRow)
+{
+	checkNoEntry();
+	return SNullWidget::NullWidget;
+}
+
+bool FAvaOutlinerTreeRoot::CanRename() const
+{
+	checkNoEntry();
+	return false;
+}
+
+bool FAvaOutlinerTreeRoot::Rename(const FString& InName)
+{
+	checkNoEntry();
+	return false;
+}
+
+TOptional<EItemDropZone> FAvaOutlinerTreeRoot::CanAcceptDrop(const FDragDropEvent& InDragDropEvent, EItemDropZone InDropZone)
+{
+	if (const TSharedPtr<FAssetDragDropOp> AssetDragDropOp = InDragDropEvent.GetOperationAs<FAssetDragDropOp>())
+	{
+		//Can't drop if a Single of the Assets is an Invalid Asset
+		for (const FAssetData& Asset : AssetDragDropOp->GetAssets())
+		{
+			UActorFactory* ActorFactory = AssetDragDropOp->GetActorFactory();
+			if (!ActorFactory)
+			{
+				ActorFactory = FActorFactoryAssetProxy::GetFactoryForAsset(Asset);
+			}
+			if (!ActorFactory || !ActorFactory->CanPlaceElementsFromAssetData(Asset))
+			{
+				return TOptional<EItemDropZone>();
+			}
+		}
+		return InDropZone;
+	}
+
+	return Super::CanAcceptDrop(InDragDropEvent, InDropZone);
+}
+
+FReply FAvaOutlinerTreeRoot::AcceptDrop(const FDragDropEvent& InDragDropEvent, EItemDropZone InDropZone)
+{
+	// Force Tree Root to always be Onto Item so we always create new Items as Children of Root
+	const EItemDropZone ForcedDropZone = EItemDropZone::OntoItem;
+
+	if (const TSharedPtr<FAssetDragDropOp> AssetDragDropOp = InDragDropEvent.GetOperationAs<FAssetDragDropOp>())
+	{
+		const UWorld* const World = Outliner.GetWorld();
+		check(World);
+		return CreateItemsFromAssetDrop(AssetDragDropOp, ForcedDropZone, World->GetCurrentLevel());
+	}
+
+	return Super::AcceptDrop(InDragDropEvent, ForcedDropZone);
+}
+
+FAvaOutlinerItemId FAvaOutlinerTreeRoot::CalculateItemId() const
+{
+	return FAvaOutlinerItemId(TEXT("OutlinerRoot"));
+}
diff --git a/Engine/Plugins/Experimental/Avalanche/Source/AvalancheOutliner/Private/ItemActions/AvaOutlinerAddItem.cpp b/Engine/Plugins/Experimental/Avalanche/Source/AvalancheOutliner/Private/ItemActions/AvaOutlinerAddItem.cpp
new file mode 100644
--- /dev/null
+++ b/Engine/Plugins/Experimental/Avalanche/Source/AvalancheOutliner/Private/ItemActions/AvaOutlinerAddItem.cpp
@@ -0,0 +1,124 @@
+// Copyright Epic Games, Inc. All Rights Reserved.
+
+#include "ItemActions/AvaOutlinerAddItem.h"
+#include "AvaOutliner.h"
+#include "Item/AvaOutlinerObject.h"
+#include "Item/AvaOutlinerTreeRoot.h"
+#include "Item/IAvaOutlinerItem.h"
+#include "Selection/AvaOutlinerScopedSelection.h"
+#include "Widgets/Views/STableRow.h"
+
+FAvaOutlinerAddItem::FAvaOutlinerAddItem(const FAvaOutlinerAddItemParams& InAddItemParams)
+	: AddParams(InAddItemParams)
+{
+}
+
+bool FAvaOutlinerAddItem::ShouldTransact() const
+{
+	return EnumHasAnyFlags(AddParams.Flags, EAvaOutlinerAddItemFlags::Transact);
+}
+
+void FAvaOutlinerAddItem::Execute(FAvaOutliner& InOutliner)
+{
+	if (AddParams.Item.IsValid())
+	{
+		//Try to Create Children on Find
+		if (EnumHasAnyFlags(AddParams.Flags, EAvaOutlinerAddItemFlags::AddChildren))
+		{
+			constexpr bool bRecursiveFind = true;
+			TArray<FAvaOutlinerItemPtr> Children;
+			AddParams.Item->FindValidChildren(Children, bRecursiveFind);
+		}
+
+		const FAvaOutlinerItemPtr ParentItem = AddParams.Item->GetParent();
+
+		//If this Array has elements in it, then we need to stop a Circular Dependency from forming
+		{
+			TArray<FAvaOutlinerItemPtr> PathToRelativeItem = AddParams.Item->FindPath({AddParams.RelativeItem});
+			if (PathToRelativeItem.Num() > 0 && ParentItem.IsValid())
+			{
+				FAvaOutlinerAddItemParams CircularSolverParams;
+				CircularSolverParams.Item                     = PathToRelativeItem[0];
+				CircularSolverParams.RelativeItem             = AddParams.Item;
+				CircularSolverParams.RelativeDropZone         = EItemDropZone::AboveItem;
+				CircularSolverParams.Flags                    = AddParams.Flags;
+				CircularSolverParams.AttachmentTransformRules = AddParams.AttachmentTransformRules;
+
+				ParentItem->AddChild(CircularSolverParams);
+			}
+		}
+
+		if (AddParams.RelativeItem.IsValid())
+		{
+			const TSharedPtr<IAvaOutlinerItem> RelativeItemParent = AddParams.RelativeItem->GetParent();
+
+			//If it's onto item, the Relative Item is going to be the Parent
+			if (!AddParams.RelativeDropZone.IsSet() || AddParams.RelativeDropZone == EItemDropZone::OntoItem)
+			{
+				//If the Relative Item is Onto and it's the same as the Current Parent, shift Item up in the Hierarchy
+				//(as long as the parent is a valid one)
+				if (AddParams.RelativeItem == ParentItem && RelativeItemParent)
+				{
+					AddParams.RelativeDropZone = EItemDropZone::BelowItem;
+					RelativeItemParent->AddChild(AddParams);
+				}
+				else
+				{
+					AddParams.RelativeItem->AddChild(AddParams);
+				}
+			}
+			//else we place it as a Sibling to the Relative Item
+			else if (RelativeItemParent)
+			{
+				RelativeItemParent->AddChild(AddParams);
+			}
+			//if no parent, then add it to the tree root
+			else
+			{
+				InOutliner.GetTreeRoot()->AddChild(AddParams);
+			}
+		}
+		else
+		{
+			//If no Relative Item, add to tree root
+			InOutliner.GetTreeRoot()->AddChild(AddParams);
+		}
+
+		if (const FEditorModeTools* const ModeTools = InOutliner.GetModeTools())
+		{
+			const FAvaOutlinerScopedSelection ScopedSelection(*ModeTools, EAvaOutlinerScopedSelectionPurpose::Read);
+
+			//Automatically Select Item if it's Selected in Mode Tools
+			if (AddParams.Item->IsSelected(ScopedSelection))
+			{
+				// Select in Outliner but don't signal selection as we already have it selected in Mode Tools
+				AddParams.Flags = EAvaOutlinerAddItemFlags::Select;
+				AddParams.SelectionFlags &= ~EAvaOutlinerItemSelectionFlags::SignalSelectionChange;
+			}
+			//Signal Selection Change when we attempt to select this item in the Outliner but it isn't selected in Mode Tools
+			else if (EnumHasAnyFlags(AddParams.Flags, EAvaOutlinerAddItemFlags::Select))
+			{
+				AddParams.SelectionFlags |= EAvaOutlinerItemSelectionFlags::SignalSelectionChange;
+			}
+		}
+
+		if (EnumHasAnyFlags(AddParams.Flags, EAvaOutlinerAddItemFlags::Select))
+		{
+			InOutliner.SelectItems({AddParams.Item}, AddParams.SelectionFlags);
+		}
+
+		InOutliner.SetOutlinerModified();
+	}
+}
+
+void FAvaOutlinerAddItem::OnObjectsReplaced(const TMap<UObject*, UObject*>& InReplacementMap, bool bRecursive)
+{
+	if (AddParams.Item.IsValid())
+	{
+		AddParams.Item->OnObjectsReplaced(InReplacementMap, bRecursive);
+	}
+	if (AddParams.RelativeItem.IsValid())
+	{
+		AddParams.RelativeItem->OnObjectsReplaced(InReplacementMap, bRecursive);
+	}
+}
