[MotionDesign]

diff --git a/Engine/Plugins/Experimental/Avalanche/Source/AvalancheOutliner/Private/Tests/AvaOutliner.spec.cpp b/Engine/Plugins/Experimental/Avalanche/Source/AvalancheOutliner/Private/Tests/AvaOutliner.spec.cpp
new file mode 100644
--- /dev/null
+++ b/Engine/Plugins/Experimental/Avalanche/Source/AvalancheOutliner/Private/Tests/AvaOutliner.spec.cpp
@@ -0,0 +1,302 @@
+// Copyright Epic Games, Inc. All Rights Reserved.
+
+#include "AvaOutliner.h"
+#include "AvaOutlinerView.h"
+#include "AvaOutlinerTestUtils.h"
+#include "Engine/DirectionalLight.h"
+#include "Engine/StaticMeshActor.h"
+#include "EngineUtils.h"
+#include "Item/AvaOutlinerActor.h"
+#include "Misc/AutomationTest.h"
+
+#if WITH_AUTOMATION_TESTS
+
+BEGIN_DEFINE_SPEC(FAvaOutlinerSpec, "Avalanche.Outliner"
+	, EAutomationTestFlags::EngineFilter
+	| EAutomationTestFlags::EditorContext
+	| EAutomationTestFlags::ApplicationContextMask)
+
+	// Types of common situations that are used across multiple tests
+	enum class EWhenType
+	{
+		NoItemsSelected,
+		ItemsSelected
+	};
+	
+	void When(const FString& InDescription, TFunction<void()>&& DoWork)
+	{
+		Describe("When " + InDescription, MoveTemp(DoWork));
+	}
+
+	void When(EWhenType InWhenType, TFunction<void()>&& DoWork)
+	{
+		switch (InWhenType)
+		{
+			case EWhenType::NoItemsSelected:
+				When("no items are selected", MoveTemp(DoWork));
+				break;
+
+			case EWhenType::ItemsSelected:
+				When("multiple items are selected", [this, DoWork = MoveTemp(DoWork)]()
+					{
+						BeforeEach([this]()
+							{
+								OutlinerView->SelectItems(TestItems);
+							});
+
+						DoWork();
+					});
+				break;
+		}
+	}
+
+	TSharedPtr<FAvaOutliner> Outliner;
+
+	// The outliner view instance created before each test
+	TSharedPtr<FAvaOutlinerView> OutlinerView;
+
+	// Custom test provider to fill some of the info/actions requested by the outliner (e.g. USelection Instances, or Actor Duplication)
+	TSharedPtr<UE::AvalancheOutliner::Private::FAvaOutlinerProviderTest> OutlinerProvider;
+
+	// A list of outliner items to test things with
+	TArray<FAvaOutlinerItemPtr> TestItems;
+
+END_DEFINE_SPEC(FAvaOutlinerSpec)
+
+void FAvaOutlinerSpec::Define()
+{
+	BeforeEach([this]()
+	{
+		OutlinerProvider = MakeShared<UE::AvalancheOutliner::Private::FAvaOutlinerProviderTest>();
+
+		constexpr int32 OutlinerViewId = 0;
+
+		Outliner = OutlinerProvider->GetOutliner();
+		Outliner->RegisterOutlinerView(OutlinerViewId);
+		Outliner->Refresh();
+
+		OutlinerView = Outliner->GetOutlinerView(0);
+
+		const TArray<UObject*> TestObjects
+		{
+			OutlinerProvider->GetDirectionalLight(),
+			OutlinerProvider->GetFloor(),
+			OutlinerProvider->GetSkySphere(),
+		};
+
+		TestItems = OutlinerProvider->GetOutlinerItems(TestObjects);
+	 });
+
+	Describe("Spawning an Actor", [this]()
+	{
+		BeforeEach([this]()
+		{
+			OutlinerProvider->TestSpawnActor();
+			Outliner->Refresh();
+		});
+
+		It("should register an item for the actor", [this]()
+		{
+			const FAvaOutlinerItemPtr Item = Outliner->FindItem(OutlinerProvider->GetTestSpawnActor());
+			TestTrue(TEXT("Item found for the actor"), Item.IsValid());
+
+			if (!TestTrue(TEXT("Item is an actor item"), Item.IsValid() && Item->IsA<FAvaOutlinerActor>()))
+			{
+				return;
+			}
+
+			const AActor* const Actor = Item->CastTo<FAvaOutlinerActor>()->GetActor();
+			TestTrue(TEXT("Actor item matches the spawned actor"), Actor == OutlinerProvider->GetTestSpawnActor());	
+		});
+	});
+	
+	Describe("Duplicating Selected Items", [this]()
+	{
+		When(EWhenType::NoItemsSelected, [this]()
+		{
+			It("should keep the outliner with no selections", [this]()
+			{
+				OutlinerView->DuplicateSelected();
+				Outliner->Refresh();
+				TestEqual(TEXT("Selected Item count is 0")
+					, Outliner->GetSelectedItemCount()
+					, 0);
+			});
+		});
+
+		When(EWhenType::ItemsSelected, [this]()
+		{
+			Describe("When selected objects have duplicated", [this]()
+			{
+				BeforeEach([this]()
+				{
+					OutlinerView->DuplicateSelected();
+				});
+				
+				It("should have a matching pending item action count", [this]()
+				{
+					TestEqual(TEXT("Number of items enqueued match the requested number")
+						, Outliner->GetPendingItemActionCount()
+						, TestItems.Num());
+				});
+
+				Describe("When outliner has refreshed", [this]()
+				{
+					BeforeEach([this]()
+					{
+						Outliner->Refresh();
+					});
+
+					It("should have a matching pending item action count", [this]()
+					{
+						TestEqual(TEXT("The template item count and the number of copies match")
+							, Outliner->GetSelectedItemCount()
+							, TestItems.Num());	
+					});
+
+					It("should have the duplicate copies selected instead of the template", [this]()
+					{
+						// These should be the new items that were created from duplication
+						const TArray<FAvaOutlinerItemPtr> DuplicateItems = Outliner->GetSelectedItems();
+
+						TestFalse(TEXT("The newly selected item count is not zero"), DuplicateItems.IsEmpty());
+
+						constexpr int32 TestItemIndex = 0;
+							
+						FAvaOutliner::SortItems(TestItems);
+							
+						TestNotEqual(TEXT("The newly selected items are different from the old selected items")
+							, DuplicateItems[TestItemIndex]
+							, TestItems[TestItemIndex]);
+					});
+
+					It("should have each duplicate item right above its respective template item", [this]()
+					{
+						constexpr int32 TestItemIndex = 0;
+
+						FAvaOutliner::SortItems(TestItems);
+						
+						const FAvaOutlinerItemPtr TemplateItem  = TestItems[TestItemIndex];
+						const FAvaOutlinerItemPtr DuplicateItem = Outliner->GetSelectedItems()[TestItemIndex];
+
+						TestEqual(TEXT("The template item and duplicate item share the same parent item")
+							, TemplateItem->GetParent()
+							, DuplicateItem->GetParent());
+
+						const FAvaOutlinerItemPtr ParentItem = TemplateItem->GetParent();
+
+						const int32 TemplateItemIndex  = ParentItem->GetChildIndex(TemplateItem);
+						const int32 DuplicateItemIndex = ParentItem->GetChildIndex(DuplicateItem);
+
+						TestEqual(TEXT("The duplicate item is right above template item")
+							, TemplateItemIndex - 1
+							, DuplicateItemIndex);
+					});
+				});
+			});
+		});	
+	});
+
+	Describe("Grouping Selected Items", [this]()
+	{
+		When("grouping with null", [this]()
+		{
+			When(EWhenType::NoItemsSelected, [this]()
+			{
+				It("should keep the outliner with no selections", [this]()
+				{
+					Outliner->GroupSelection(nullptr);
+					Outliner->Refresh();
+					TestEqual(TEXT("Selected Item count is 0")
+						, Outliner->GetSelectedItemCount()
+						, 0);
+				});
+			});
+
+			When(EWhenType::ItemsSelected, [this]()
+			{
+				It("should keep the outliner with the same selected items", [this]()
+				{
+					const TArray<FAvaOutlinerItemPtr> PrevSelectedItems = Outliner->GetSelectedItems();
+
+					Outliner->GroupSelection(nullptr);
+					Outliner->Refresh();
+
+					const TArray<FAvaOutlinerItemPtr> CurrSelectedItems = Outliner->GetSelectedItems();
+
+					TestEqual(TEXT("Selected Item count remains the same")
+						, PrevSelectedItems.Num()
+						, CurrSelectedItems.Num());
+
+					TestEqual(TEXT("Selected Items remain the same")
+						, PrevSelectedItems[0]
+						, CurrSelectedItems[0]);
+				});
+			});
+		});
+
+		When("grouping with an already existing actor", [this]()
+		{
+			When(EWhenType::NoItemsSelected, [this]()
+			{
+				BeforeEach([this]()
+				{
+					AActor* const GroupingActor = OutlinerProvider->GetVolumeActor();
+					TestNotNull(TEXT("Grouping actor exists"), GroupingActor);
+					Outliner->GroupSelection(GroupingActor);
+					Outliner->Refresh();
+				});
+
+				It("should keep the outliner with no selections", [this]()
+				{
+					TestEqual(TEXT("Selected item count is zero"), Outliner->GetSelectedItemCount(), 0);
+				});
+			});
+
+			When(EWhenType::ItemsSelected, [this]()
+			{
+				BeforeEach([this]()
+				{
+					AActor* const GroupingActor = OutlinerProvider->GetVolumeActor();
+					TestNotNull(TEXT("Grouping actor exists"), GroupingActor);
+					Outliner->GroupSelection(GroupingActor);
+					Outliner->Refresh();
+				});
+
+				It("should keep the outliner with the selected items", [this]()
+				{
+					const TArray<FAvaOutlinerItemPtr> SelectedItems = Outliner->GetSelectedItems();
+
+					TestEqual(TEXT("Selected item count is consistent"), SelectedItems.Num(), TestItems.Num());
+
+					if (TestNotEqual(TEXT("Selected item count is not zero"), SelectedItems.Num(), 0))
+					{
+						TestEqual(TEXT("The selected items are consistent"), TestItems[0], SelectedItems[0]);
+					}
+				});
+
+				It("should have placed the selected items under the grouping actor", [this]()
+				{
+					const AActor* const GroupingActor = OutlinerProvider->GetVolumeActor();
+					const FAvaOutlinerItemPtr GroupingItem = Outliner->FindItem(GroupingActor);
+
+					TestNotNull(TEXT("Grouping item exists in outliner"), GroupingItem.Get());
+
+					const int32 ChildIndex = GroupingItem->GetChildIndex(TestItems[0]);
+
+					TestTrue(TEXT("the selected items are under the grouping item"), ChildIndex != INDEX_NONE);
+				});
+			});
+		});
+	});
+
+	AfterEach([this]()
+	{
+		TestItems.Reset();
+		OutlinerView.Reset();
+		Outliner.Reset();
+		OutlinerProvider.Reset();
+	});
+}
+
+#endif
diff --git a/Engine/Plugins/Experimental/Avalanche/Source/AvalancheOutliner/Public/AvaOutliner.h b/Engine/Plugins/Experimental/Avalanche/Source/AvalancheOutliner/Public/AvaOutliner.h
new file mode 100644
--- /dev/null
+++ b/Engine/Plugins/Experimental/Avalanche/Source/AvalancheOutliner/Public/AvaOutliner.h
@@ -0,0 +1,473 @@
+// Copyright Epic Games, Inc. All Rights Reserved.
+
+#pragma once
+
+#include "AvaOutlinerDefines.h"
+#include "EditorUndoClient.h"
+#include "Engine/EngineTypes.h"
+#include "IAvaOutlinerModule.h"
+#include "Item/AvaOutlinerItemId.h"
+#include "Item/AvaOutlinerItemProxy.h"
+#include "Item/IAvaOutlinerItem.h"
+#include "ItemProxies/AvaOutlinerItemProxyRegistry.h"
+#include "ItemProxies/IAvaOutlinerItemProxyFactory.h"
+#include "TickableEditorObject.h"
+
+class FAvaEditorSelection;
+class FAvaOutlinerTreeRoot;
+class FAvaOutlinerView;
+class FEditorModeTools;
+class FTransaction;
+class FUICommandList;
+class IAvaOutlinerAction;
+class IAvaOutlinerProvider;
+class UAvaOutlinerSubsystem;
+enum class EItemDropZone;
+struct FAvaOutlinerSaveState;
+struct FAvaSceneItem;
+
+/** 
+ * The Outliner Object that is commonly instanced once per World
+ * (unless for advanced use where there are different outliner instances with different item ordering and behaviors).
+ * This is the object that dictates core outliner behavior like how items are sorted, which items are allowed, etc.
+ * Views are the objects that take this core behavior and show a part of it (e.g. through filters).
+ */
+class AVALANCHEOUTLINER_API FAvaOutliner
+	: public TSharedFromThis<FAvaOutliner>
+	, public FTickableEditorObject
+	, public FEditorUndoClient
+{
+public:
+	FAvaOutliner(IAvaOutlinerProvider& InOutlinerProvider);
+	
+	virtual ~FAvaOutliner() override;
+	
+	/**
+	 * Gets the Outliner Subsystem of the World this Outliner is responsible for
+	 * NOTE: The OutlinerSubsystem's Outliner Instance should be the same as this one in the default implementation 
+	 * but can differ if there are multiple FAvaOutliner instances being used in custom implementations
+	 */
+	UAvaOutlinerSubsystem* GetOutlinerSubsystem() const;
+
+	/**
+	 * Determines whether the given actor can be presented in the Outliner, at all.
+	 * This is a permanent check unlike filters that are temporary.
+	 */
+	bool IsActorAllowedInOutliner(const AActor* InActor) const;
+
+	/**
+	 * Determines whether the given scene component can be presented in the Outliner, at all.
+	 * This is a permanent check unlike filters that are temporary.
+	 */
+	bool IsComponentAllowedInOutliner(const USceneComponent* InComponent) const;
+
+	bool CanProcessActorSpawn(AActor* InActor) const;
+	
+	/** Sets the Command List that the Outliner Views will use to append their Command Lists to */
+	void SetBaseCommandList(const TSharedPtr<FUICommandList>& InBaseCommandList);
+
+	TSharedPtr<FUICommandList> GetBaseCommandList() const;
+
+	/** Serializes the Outliner Save State to the given archive */
+	void Serialize(FArchive& Ar);
+
+	DECLARE_MULTICAST_DELEGATE(FOnOutlinerLoaded);
+	FOnOutlinerLoaded OnOutlinerLoaded;
+
+	IAvaOutlinerProvider& GetProvider() const { return OutlinerProvider; }
+
+	FAvaOutlinerItemProxyRegistry& GetItemProxyRegistry() { return ItemProxyRegistry; }
+
+	/** Gathers the Type Names of all the Item Proxies that are registered both in the outliner proxy registry and the module's */
+	TArray<FName> GetRegisteredItemProxyTypeNames() const;
+
+	/** Gathers all previously existing and new Item Proxies for a given Item */
+	void GetItemProxiesForItem(const FAvaOutlinerItemPtr& InItem, TArray<TSharedPtr<FAvaOutlinerItemProxy>>& OutItemProxies);
+
+	/** Tries to find the Item Proxy Factory for the given Item Proxy Type Name */
+	IAvaOutlinerItemProxyFactory* GetItemProxyFactory(FName InItemProxyTypeName) const
+	{
+		// First look for the Registry in Outliner
+		if (IAvaOutlinerItemProxyFactory* Factory = ItemProxyRegistry.GetItemProxyFactory(InItemProxyTypeName))
+		{
+			return Factory;
+		}
+		// Fallback to finding the Factory in the Module if the Outliner did not find it
+		return IAvaOutlinerModule::Get().GetItemProxyRegistry().GetItemProxyFactory(InItemProxyTypeName);
+	}
+
+	/** Tries to find the Item Proxy Factory for the given Item Proxy Type Name */
+	template<typename InItemProxyType, typename = typename TEnableIf<TIsDerivedFrom<InItemProxyType, FAvaOutlinerItemProxy>::IsDerived>::Type>
+	IAvaOutlinerItemProxyFactory* GetItemProxyFactory() const
+	{
+		// First look for the Registry in Outliner
+		if (IAvaOutlinerItemProxyFactory* Factory = ItemProxyRegistry.GetItemProxyFactory<InItemProxyType>())
+		{
+			return Factory;
+		}
+		// Fallback to finding the Factory in the Module if the Outliner did not find it
+		return IAvaOutlinerModule::Get().GetItemProxyRegistry().GetItemProxyFactory<InItemProxyType>();
+	}
+
+	/**
+	 * Tries to get the Item Proxy Factory for the given Item Proxy type, first trying the Outliner Registry then the Module's
+	 * then returns an existing item proxy created via the factory, or creates one if there's no existing item proxy
+	 * @returns the Item Proxy created by the Factory. Can be null if no factory was found or if the factory intentionally returns null
+	 */
+	template<typename InItemProxyType, typename = typename TEnableIf<TIsDerivedFrom<InItemProxyType, FAvaOutlinerItemProxy>::IsDerived>::Type>
+	TSharedPtr<FAvaOutlinerItemProxy> GetOrCreateItemProxy(const FAvaOutlinerItemPtr& InParentItem)
+	{
+		if (!InParentItem.IsValid() || !InParentItem->IsAllowedInOutliner())
+		{
+			return nullptr;
+		}
+		
+		IAvaOutlinerItemProxyFactory* const Factory = GetItemProxyFactory<InItemProxyType>();
+		if (!Factory)
+		{
+			return nullptr;
+		}
+		
+		TSharedPtr<FAvaOutlinerItemProxy> OutItemProxy;
+		if (FAvaOutlinerItemPtr ExistingItemProxy = FindItem(FAvaOutlinerItemId(InParentItem, *Factory)))
+		{
+			check(ExistingItemProxy->IsA<FAvaOutlinerItemProxy>());
+			ExistingItemProxy->SetParent(InParentItem);
+			OutItemProxy = StaticCastSharedPtr<FAvaOutlinerItemProxy>(ExistingItemProxy);
+		}
+		else
+		{
+			OutItemProxy = Factory->CreateItemProxy(*this, InParentItem);
+		}
+		RegisterItem(OutItemProxy);
+		return OutItemProxy;
+	}
+	
+	const TSharedRef<FAvaOutlinerSaveState>& GetSaveState() const { return SaveState; }
+
+	/** Returns whether the Outliner is in Read-only mode */
+	bool IsOutlinerLocked() const;
+
+	void HandleUndoRedoTransaction(const FTransaction* Transaction, bool bIsUndo);
+
+	//~ Begin FEditorUndoClient
+	virtual void PostUndo(bool bSuccess) override;
+	virtual void PostRedo(bool bSuccess) override;
+	//~ End FEditorUndoClient
+
+	//~ Begin FTickableObjectBase
+	virtual TStatId GetStatId() const override;
+	virtual void Tick(float InDeltaTime) override;
+	//~ End FTickableObjectBase
+
+	/**
+	 * Creates a new set of items in the outliner based on the given template items
+	 * @param InItems the template items to use for duplication
+	 * @param InRelativeItem the existing outliner item that items should use as positional reference
+	 * @param InRelativeDropZone where to put the duplicate items with relation to InRelativeItem (above, below, onto)
+	 */
+	void DuplicateItems(TArray<FAvaOutlinerItemPtr> InItems
+		, FAvaOutlinerItemPtr InRelativeItem
+		, TOptional<EItemDropZone> InRelativeDropZone);
+	
+	/** Register a new Outliner View to the Outliner to the given id, replacing the old view that was bound to the given id */
+	TSharedPtr<FAvaOutlinerView> RegisterOutlinerView(int32 InOutlinerViewId);
+
+	/** Unregisters the Outliner View bound to the given id */
+	void UnregisterOutlinerView(int32 InOutlinerViewId);
+
+	/** Sets the given Outliner View Id as the most recent Outliner View */
+	void UpdateRecentOutlinerViews(int32 InOutlinerViewId);
+
+	/** Gets the Outliner View bound to the given id */
+	TSharedPtr<FAvaOutlinerView> GetOutlinerView(int32 InOutlinerViewId) const;
+
+	/** Gets the outliner view that was most recently used (i.e. called FAvaOutliner::UpdateRecentOutlinerViews) */
+	TSharedPtr<FAvaOutlinerView> GetMostRecentOutlinerView() const;
+
+	/** Executes the given predicate for each Outliner View registered */
+	void ForEachOutlinerView(const TFunction<void(const TSharedPtr<FAvaOutlinerView>& InOutlinerView)>& InPredicate) const;
+
+	/** Registers the given Item, replacing the old one. */
+	void RegisterItem(const FAvaOutlinerItemPtr& InItem);
+
+	/** Unregisters the Item having the given ItemId */
+	void UnregisterItem(const FAvaOutlinerItemId& InItemId);
+
+	/**
+	 * Instantiates a new item action without adding it to the Pending Actions Queue.
+	 * This should only be used directly when planning to enqueue multiple actions.
+	 * @see FAvaOutliner::EnqueueItemActions
+	 */
+	template<typename InItemActionType, typename = typename TEnableIf<TIsDerivedFrom<InItemActionType, IAvaOutlinerAction>::IsDerived>::Type, typename ...InArgTypes>
+	TSharedRef<InItemActionType> NewItemAction(InArgTypes&&... InArgs)
+	{
+		return MakeShared<InItemActionType>(Forward<InArgTypes>(InArgs)...);
+	}
+
+	/**
+	 * Instantiates a single new item action and immediately adds it to the Pending Actions Queue.
+	 * Ideal for when dealing with a single action.
+	 * For multiple actions use FAvaOutliner::EnqueueItemActions.
+	 */
+	template<typename InItemActionType, typename = typename TEnableIf<TIsDerivedFrom<InItemActionType, IAvaOutlinerAction>::IsDerived>::Type, typename ...InArgTypes>
+	void EnqueueItemAction(InArgTypes&&... InArgs)
+	{
+		EnqueueItemActions({ NewItemAction<InItemActionType>(Forward<InArgTypes>(InArgs)...) });
+	}
+
+	/** Adds the given actions to the Pending Action Queue */
+	void EnqueueItemActions(const TArray<TSharedPtr<IAvaOutlinerAction>>& InItemActions);
+
+	/** Adds the given actions to the Pending Action Queue */
+	void EnqueueItemActions(TArray<TSharedPtr<IAvaOutlinerAction>>&& InItemActions) noexcept;
+
+	/** Returns the number of actions that been added to the queue so far before triggering a refresh */
+	int32 GetPendingItemActionCount() const { return PendingActions.Num(); }
+
+	/** Instantiates a new Item and automatically registers it to the Outliner */
+	template<typename InItemType, typename = typename TEnableIf<TIsDerivedFrom<InItemType, IAvaOutlinerItem>::IsDerived>::Type, typename ...InArgTypes>
+	TSharedRef<InItemType> FindOrAdd(InArgTypes&&... InArgs)
+	{
+		TSharedRef<InItemType> Item = MakeShared<InItemType>(*this, Forward<InArgTypes>(InArgs)...);
+
+		// If an existing item already exists and has a valid state, use that and forget about the newly created
+		FAvaOutlinerItemPtr ExistingItem = FindItem(Item->GetItemId());
+		if (ExistingItem.IsValid() && ExistingItem->IsItemValid() && ExistingItem->IsA<InItemType>())
+		{
+			return StaticCastSharedPtr<InItemType>(ExistingItem).ToSharedRef();
+		}
+		
+		if (Item->IsAllowedInOutliner())
+		{
+			RegisterItem(Item);
+		}
+		
+		return Item;
+	}
+
+	/** Returns whether the Outliner is currently in need of a Refresh */
+	bool NeedsRefresh() const;
+
+	/** Ensures that the next time Refresh is called in tick, Refresh will be called */
+	void RequestRefresh();
+
+	/**
+	 * Flushes the Pending Actions from the Queue while also updating the state of the Outliner.
+	 * Calling it directly is forcing it to happen.
+	 * If a refresh is needed it will be called on the next tick automatically.
+	 */
+	void Refresh();
+
+	/**
+	 * Finds the Registered Item that has the given Id
+	 * @returns the Item with the given Id, or null if the item does not exist or was not registered to the Outliner
+	 */
+	FAvaOutlinerItemPtr FindItem(const FAvaOutlinerItemId& InItemId) const;
+
+	/**
+	 * Gets the color pair (color name, linear color) related to the Item
+	 * @param InItem the item to query
+	 * @param bRecurseParent whether to get the color of the parent (recursively) if the given item does not have a color by itself
+	 * @returns the matching color pair or unset if item is invalid or no color could be found.
+	 */
+	TOptional<FAvaOutlinerColorPair> FindItemColor(const FAvaOutlinerItemPtr& InItem, bool bRecurseParent = true) const;
+
+	/** Pairs the Item with the given color name, overriding the inherited color if different */
+	void SetItemColor(const FAvaOutlinerItemPtr& InItem, const FName& InColorName);
+
+	/** Removes the Color pairing of the given Item (can still have an inherited color though) */
+	void RemoveItemColor(const FAvaOutlinerItemPtr& InItem);
+
+	/** Returns the color map from the Outliner Settings */
+	const TMap<FName, FLinearColor>& GetColorMap() const;
+
+	/**
+	 * Replaces the Item's Id in the Item Map. This can be due to an object item changing it's object
+	 * (e.g. a bp component getting destroyed and recreated, the item should be the same but the underlying component will not be)
+	 */
+	void NotifyItemIdChanged(const FAvaOutlinerItemId& OldId, const FAvaOutlinerItemPtr& InItem);
+
+	/** Returns the currently selected items in the most recent outliner view (since this list can vary between outliner views) */
+	TArray<FAvaOutlinerItemPtr> GetSelectedItems() const;
+
+	/** Returns the number of currently selected items in the most recent outliner view */
+	int32 GetSelectedItemCount() const;
+
+	/** Gets the Tree Root Item of the Outliner */
+	TSharedRef<FAvaOutlinerTreeRoot> GetTreeRoot() const
+	{
+		return RootItem;
+	}
+
+	/**
+	 * Selects the given Items on all Outliner Views
+	 * @param InItems the items to select
+	 * @param InFlags how the items should be selected (appended, notify of selections, etc)
+	 */
+	void SelectItems(const TArray<FAvaOutlinerItemPtr>& InItems, EAvaOutlinerItemSelectionFlags InFlags = EAvaOutlinerItemSelectionFlags::SignalSelectionChange) const;
+
+	/**
+	 * Clears the Item Selection from all Outliner Views
+	 * @param bSignalSelectionChange whether to notify the change in selection
+	 */
+	void ClearItemSelection(bool bSignalSelectionChange) const;
+
+	/**
+	 * Adds or Removes the Ignore Notify Flags to prevent certain actions from automatically happening when they're triggered
+	 * @param InFlag the ignore flag to add or remove
+	 * @param bIgnore whether to add (true) or remove (false) the flag
+	 */
+	void SetIgnoreNotify(EAvaOutlinerIgnoreNotifyFlags InFlag, bool bIgnore)
+	{
+		if (bIgnore)
+		{
+			EnumAddFlags(IgnoreNotifyFlags, InFlag);
+		}
+		else
+		{
+			EnumRemoveFlags(IgnoreNotifyFlags, InFlag);
+		}
+	}
+
+	/** Gets the closest item to all the given items while also being their common ancestor */
+	static FAvaOutlinerItemPtr FindLowestCommonAncestor(const TArray<FAvaOutlinerItemPtr>& Items);
+
+	/**
+	 * Compares the absolute order of the items in the Outliner and returns true if A comes before B in the outliner.
+	 * Useful to use when sorting Items.
+	 */
+	static bool CompareOutlinerItemOrder(const FAvaOutlinerItemPtr& A, const FAvaOutlinerItemPtr& B);
+
+	/** Converts the given Outliner Item to a Scene Item that can be serialized in the Scene Tree */
+	static FAvaSceneItem MakeSceneItemFromOutlinerItem(const FAvaOutlinerItemPtr& InItem);
+
+	/**
+	 * Helper function to sort the given array of items based on their ordering in the Outliner
+	 * @see FAvaOutliner::CompareOutlinerItemOrder
+	 */
+	static void SortItems(TArray<FAvaOutlinerItemPtr>& OutOutlinerItems, bool bInReverseOrder = false);
+
+	/** Normalizes the given Items by removing selected items that have their parent item also in the selection */
+	static void NormalizeItems(TArray<FAvaOutlinerItemPtr>& InOutItems);
+
+	/**
+	 * Gets all the Selected Items and puts/attaches them under the given Grouping Actor.
+	 * Requires that the Grouping Actor is valid and spawned in the World.
+	 */
+	void GroupSelection(AActor* InGroupingActor, const TOptional<FAttachmentTransformRules>& InTransformRules = TOptional<FAttachmentTransformRules>());
+
+	/** Gets the Editor Mode Tools used to handle selections */
+	FEditorModeTools* GetModeTools() const;
+
+	/** Have the given Selected Items sync to the USelection Instances of Mode Tools */
+	void SyncModeToolsSelection(const TArray<FAvaOutlinerItemPtr>& InSelectedItems) const;
+
+	/** Called when the objects have been selected and notified through USelection Instances in Mode Tools */
+	void OnObjectSelectionChanged(const FAvaEditorSelection& InEditorSelection);
+
+	/** Gets the World the Outliner is working with */
+	UWorld* GetWorld() const;
+
+	/** Gets all the Actors that have as their AActor::GetSceneOutlinerParent the given InParentActor */
+	TArray<TWeakObjectPtr<AActor>> GetActorSceneOutlinerChildren(AActor* InParentActor) const;
+
+	/** Tries to add the new Actor to the Outliner if not added already and if the Outliner allows the given actor to be added */
+	void OnActorSpawned(AActor* InActor);
+
+	/**
+	 * Should be called when Actors have been copied and give Ava Outliner opportunity to add to the Buffer to copy the Outliner data for those Actors
+	 * @param InOutCopiedData the data to process / append to for copy
+	 * @param InCopiedActors the actors to copy
+	 */
+	void OnActorsCopied(FString& InOutCopiedData, TConstArrayView<AActor*> InCopiedActors);
+
+	/**
+	 * Should be called when Actors have been pasted to parse the data that was filled in by FAvaOutliner::OnActorsCopied
+	 * @param InPastedData pasted string data
+	 * @param InPastedActors map of the original actor name to its created actor on paste
+	 */
+	void OnActorsPasted(FStringView InPastedData, const TMap<FName, AActor*>& InPastedActors);
+
+	/**
+	 * Handles when Actors have been Duplicated
+	 * @param InDuplicateActorMap the map of the Duplicate Actors to their Templates
+	 * @param InRelativeItem he item to use as positional reference to where to place the duplicate items in the outliner
+	 * @param InRelativeDropZone where to put the duplicate items relative to the InRelativeItem (above, below, onto)
+	 */
+	void OnActorsDuplicated(const TMap<AActor*, AActor*>& InDuplicateActorMap
+		, FAvaOutlinerItemPtr InRelativeItem = nullptr
+		, TOptional<EItemDropZone> InRelativeDropZone = TOptional<EItemDropZone>());
+
+	/** Called when an Actor has been destroyed. This enqueues the Removal the Actor Item to the Pending Action Queue */
+	void OnActorDestroyed(AActor* InActor);
+
+	/** Called when an Actor's attachment has changed. This triggers a refresh */
+	void OnActorAttachmentChanged(AActor* InActor, const AActor* InParent, bool bAttach);
+
+	/** Called the engine replaces an object. A common example is when a BP Component is destroyed, and replaced */
+	void OnObjectsReplaced(const TMap<UObject*, UObject*>& InReplacementMap);
+
+	/** Gets the Default Transform to use when spawning actors via the Outliner (e.g. dragging a spawnable asset into an item in the outliner) */
+	FTransform GetActorDefaultSpawnTransform() const;
+
+	/** Marks the Outliner dirty. This triggers IAvaOutlinerProvider::OnOutlinerModified on next tick */
+	void SetOutlinerModified();
+
+private:
+	void AddItem(const FAvaOutlinerItemPtr& InItem);
+	void RemoveItem(const FAvaOutlinerItemId& InItemId);
+	void ForEachItem(TFunctionRef<void(const FAvaOutlinerItemPtr&)> InFunc);
+
+	/** Interface providing the Outliner with logic such as Actor Duplication, Mode Tools, and extensibility options */
+	IAvaOutlinerProvider& OutlinerProvider;
+
+	/** the base command list gotten from the outliner provider to append our command list to */
+	TWeakPtr<FUICommandList> BaseCommandListWeak;
+
+	/** the map of the registered items */
+	TMap<FAvaOutlinerItemId, FAvaOutlinerItemPtr> ItemMap;
+
+	TMap<FAvaOutlinerItemId, FAvaOutlinerItemPtr> ItemsPendingAdd;
+
+	TSet<FAvaOutlinerItemId> ItemsPendingRemove;
+
+	/** Map of an Actor to its Child Actors defined as such via AActor::GetSceneOutlinerParent */
+	TMap<TWeakObjectPtr<AActor>, TArray<TWeakObjectPtr<AActor>>> SceneOutlinerParentMap;
+
+	/** The root of all the items in the outliner */
+	TSharedRef<FAvaOutlinerTreeRoot> RootItem;
+
+	/** Outliner's Item Proxy Factory Registry Instance. This takes precedence over the Module's Factory Registry */
+	FAvaOutlinerItemProxyRegistry ItemProxyRegistry;
+	
+	/** the current pending actions before refresh is called */
+	TArray<TSharedPtr<IAvaOutlinerAction>> PendingActions;
+
+	/** the list of objects pending selection processing. filled in when getting notifies from the USelection instances */
+	TSharedPtr<TArray<TWeakObjectPtr<UObject>>> ObjectsLastSelected;
+
+	/** the save state object to help serialize the outliner state */
+	TSharedRef<FAvaOutlinerSaveState> SaveState;
+
+	/** the map of registered outliner views */
+	TMap<int32, TSharedPtr<FAvaOutlinerView>> OutlinerViews;
+
+	/** List of Outliner View Ids in order from least recent to most recent (i.e. Index 0 is least recent) */
+	TArray<int32> RecentOutlinerViews;
+
+	/** the current events to ignore and not handle automatically */
+	EAvaOutlinerIgnoreNotifyFlags IgnoreNotifyFlags = EAvaOutlinerIgnoreNotifyFlags::None;
+
+	/** Flag indicating whether the Outliner has been changed this tick and should call IAvaOutlinerProvider::OnOutlinerModified next tick */
+	bool bOutlinerDirty = false;
+	
+	/** Flag indicating Refreshing is taking place */
+	bool bRefreshing = false;
+
+	/** Flag indicating that a refresh must take place next tick */
+	bool bRefreshRequested = false;
+
+	/** Flag indicating that the Item Map is iterating */
+	bool bIteratingItemMap = false;
+};
diff --git a/Engine/Plugins/Experimental/Avalanche/Source/AvalancheOutliner/Public/Factories/Filters/AvaFilterExpressionFactoryUtils.h b/Engine/Plugins/Experimental/Avalanche/Source/AvalancheOutliner/Public/Factories/Filters/AvaFilterExpressionFactoryUtils.h
new file mode 100644
--- /dev/null
+++ b/Engine/Plugins/Experimental/Avalanche/Source/AvalancheOutliner/Public/Factories/Filters/AvaFilterExpressionFactoryUtils.h
@@ -0,0 +1,47 @@
+﻿// Copyright Epic Games, Inc. All Rights Reserved.
+
+#pragma once
+
+#include "Components/ActorComponent.h"
+#include "GameFramework/Actor.h"
+#include "Item/AvaOutlinerActor.h"
+#include "Misc/TextFilterUtils.h"
+
+namespace UE::AvalancheFactory::Private
+{
+	bool inline FilterActorComponentsByType(const FAvaOutlinerActor* InAvaOutlinerActor, FName InValueToCheck, ETextFilterComparisonOperation InComparisonOperation,ETextFilterTextComparisonMode InComparisonMode)
+	{
+		TArray<UActorComponent*> Components;
+		InAvaOutlinerActor->GetActor()->GetComponents<UActorComponent>(Components);
+		for (const UActorComponent* Component : Components)
+		{
+			FString ComponentName = Component->GetClass()->GetFName().ToString();
+			ComponentName.RemoveSpacesInline();
+			ComponentName.ToUpperInline();
+			const bool bIsComponentMatch = TextFilterUtils::TestBasicStringExpression(ComponentName, InValueToCheck, InComparisonMode);
+			if (InComparisonOperation == ETextFilterComparisonOperation::Equal ? bIsComponentMatch : !bIsComponentMatch)
+			{
+				return true;
+			}
+		}
+		return false;
+	}
+
+	bool inline FilterActorComponentsByName(const FAvaOutlinerActor* InAvaOutlinerActor, FName InValueToCheck, ETextFilterComparisonOperation InComparisonOperation,ETextFilterTextComparisonMode InComparisonMode)
+	{
+		TArray<UActorComponent*> Components;
+		InAvaOutlinerActor->GetActor()->GetComponents<UActorComponent>(Components);
+		for (const UActorComponent* Component : Components)
+		{
+			FString ComponentName = Component->GetName();
+			ComponentName.RemoveSpacesInline();
+			ComponentName.ToUpperInline();
+			const bool bIsComponentMatch = TextFilterUtils::TestBasicStringExpression(ComponentName, InValueToCheck, InComparisonMode);
+			if (InComparisonOperation == ETextFilterComparisonOperation::Equal ? bIsComponentMatch : !bIsComponentMatch)
+			{
+				return true;
+			}
+		}
+		return false;
+	}
+}
diff --git a/Engine/Plugins/Experimental/Avalanche/Source/AvalancheOutliner/Public/IAvaOutlinerModule.h b/Engine/Plugins/Experimental/Avalanche/Source/AvalancheOutliner/Public/IAvaOutlinerModule.h
new file mode 100644
--- /dev/null
+++ b/Engine/Plugins/Experimental/Avalanche/Source/AvalancheOutliner/Public/IAvaOutlinerModule.h
@@ -0,0 +1,193 @@
+// Copyright Epic Games, Inc. All Rights Reserved.
+
+#pragma once
+
+#include "AvaOutlinerDefines.h"
+#include "Containers/ContainersFwd.h"
+#include "Delegates/Delegate.h"
+#include "Filters/IAvaFilterExpressionFactory.h"
+#include "Filters/IAvaFilterSuggestionFactory.h"
+#include "Icon/AvaOutlinerIconCustomization.h"
+#include "Item/AvaOutlinerItem.h"
+#include "Modules/ModuleInterface.h"
+#include "Modules/ModuleManager.h"
+#include "Templates/SharedPointer.h"
+
+class FAvaOutliner;
+class FAvaOutlinerItem;
+class FAvaOutlinerItemProxy;
+class FAvaOutlinerItemProxyRegistry;
+class UToolMenu;
+struct FAvaTextFilterArgs;
+struct FSlateIcon;
+struct FToolMenuContext;
+
+class IAvaOutlinerModule : public IModuleInterface
+{
+public:
+	/** Returns the instance of this module, assuming it's loaded */
+	static IAvaOutlinerModule& Get()
+	{
+		return FModuleManager::LoadModuleChecked<IAvaOutlinerModule>(TEXT("AvalancheOutliner"));
+	};
+	/** Returns whether this module has been loaded  */
+	static bool IsLoaded()
+	{
+		return FModuleManager::Get().IsModuleLoaded(TEXT("AvalancheOutliner"));
+	};
+
+	DECLARE_MULTICAST_DELEGATE_OneParam(FOnInitOutlinerToolMenu, FToolMenuContext&);
+
+	/** Called when Initializing Outliner ToolBar to extend its Tool Menu Context */
+	virtual FOnInitOutlinerToolMenu& GetOnInitOutlinerToolBar() = 0;
+
+	/** Called when Initializing Outliner Item Context Menu to extend its Tool Menu Context */
+	virtual FOnInitOutlinerToolMenu& GetOnInitOutlinerItemContextMenu() = 0;
+
+	DECLARE_MULTICAST_DELEGATE_OneParam(FOnExtendOutlinerToolMenu, UToolMenu*);
+
+	/** Option to extend the Outliner ToolBar */
+	virtual FOnExtendOutlinerToolMenu& GetOnExtendOutlinerToolBar() = 0;
+
+	/** Option to extend the Outliner Item Context Menu */
+	virtual FOnExtendOutlinerToolMenu& GetOnExtendOutlinerItemContextMenu() = 0;
+
+	DECLARE_MULTICAST_DELEGATE_OneParam(FOnExtendOutlinerColumns, class FAvaOutlinerColumnExtender&);
+
+	/** Extends Outliner Columns regardless of the Outliner Provider used */
+	virtual FOnExtendOutlinerColumns& GetOnExtendOutlinerColumns() = 0;
+
+	/** Get the Module's Item Proxy Registry. Used as back up in case the Outliner's Instance does not have the Item type registered */
+	virtual FAvaOutlinerItemProxyRegistry& GetItemProxyRegistry() = 0;
+
+	DECLARE_MULTICAST_DELEGATE_ThreeParams(FOnExtendItemProxiesForItem, FAvaOutliner&, const FAvaOutlinerItemPtr&, TArray<TSharedPtr<FAvaOutlinerItemProxy>>&);
+
+	/** Gets the Delegate to call for adding in the Item Proxies to put under the given Item */
+	virtual FOnExtendItemProxiesForItem& GetOnExtendItemProxiesForItem() = 0;
+
+	/**
+ 	 * Create and Register the given filter expression factory
+ 	 * @tparam InFilterExpressionFactoryType The type of filter expression factory to create and register
+ 	 * @tparam InArgsType Constructor args types if needed
+ 	 * @param InArgs Additional args passed to the constructor of the filter expression factory if needed
+ 	 */
+	template<typename InFilterExpressionFactoryType, typename... InArgsType
+			UE_REQUIRES(TIsDerivedFrom<InFilterExpressionFactoryType, IAvaFilterExpressionFactory>::Value)>
+	void RegisterFilterExpressionFactory(InArgsType&&... InArgs)
+	{
+		RegisterFilterExpressionFactory_Internal(IAvaFilterExpressionFactory::MakeInstance<InFilterExpressionFactoryType>(Forward<InArgsType>(InArgs)...));
+	}
+
+	/**
+  	 * Create and Register the given filter suggestion factory
+  	 * @tparam InSuggestionFactoryType The type of filter suggestion factory to create and register
+  	 * @tparam InArgsType Constructor args types if needed
+  	 * @param InArgs Additional args passed to the constructor of the filter suggestion factory if needed
+  	 */
+	template<typename InSuggestionFactoryType, typename... InArgsType
+			UE_REQUIRES(TIsDerivedFrom<InSuggestionFactoryType, IAvaFilterSuggestionFactory>::Value)>
+	void RegisterFilterSuggestionFactory(InArgsType&&... InArgs)
+	{
+		RegisterFilterSuggestionFactory_Internal(IAvaFilterSuggestionFactory::MakeInstance<InSuggestionFactoryType>(Forward<InArgsType>(InArgs)...));
+	}
+
+	/**
+	 * Register an icon customization for the given OutlinerItemClass
+	 * @tparam InOutlinerItemClass Outliner class to add a customization for
+	 * @tparam InIconCustomization Icon customization to add
+	 * @tparam InArgsType Additional constructor parameter type for the customization class (ex. UClass* for Object/Actor customization)
+	 * @param InArgs Additional constructor parameter for the customization class
+	 * @return The customization created
+	 */
+	template<typename InOutlinerItemClass, typename InIconCustomization, typename... InArgsType
+		UE_REQUIRES(TIsDerivedFrom<InOutlinerItemClass, FAvaOutlinerItem>::Value && TIsDerivedFrom<InIconCustomization, IAvaOutlinerIconCustomization>::Value)>
+	InIconCustomization& RegisterOverriddenIcon(InArgsType&&... InArgs)
+	{
+		TSharedRef<InIconCustomization> OutlinerIconCustomization = MakeShared<InIconCustomization>(Forward<InArgsType>(InArgs)...);
+		RegisterOverridenIcon_Internal(TAvaType<InOutlinerItemClass>::GetTypeId(), OutlinerIconCustomization);
+		return OutlinerIconCustomization.Get();
+	}
+
+	/**
+	 * Unregister an icon customization
+	 * @tparam InOutlinerItemClass Outliner item class (ex. FAvaOutlinerItem)
+	 * @param InAdditionalKeyParameter The additional key parameter to found the customization (ex. ObjectClass FName for ObjectIcon customization)
+	 */
+	template<typename InOutlinerItemClass
+		UE_REQUIRES(TIsDerivedFrom<InOutlinerItemClass, FAvaOutlinerItem>::Value)>
+	void UnregisterOverriddenIcon(FName InAdditionalKeyParameter)
+	{
+		UnregisterOverriddenIcon_Internal(TAvaType<InOutlinerItemClass>::GetTypeId(), InAdditionalKeyParameter);
+	}
+
+	/**
+	 * Get the custom icon for the item if a customization for it exist
+	 * @param InItem The item to search the icon for
+	 * @return The custom icon for the item if a customization is found, FSlateIcon() otherwise
+	 */
+	virtual FSlateIcon FindOverrideIcon(TSharedPtr<const FAvaOutlinerItem> InItem) const = 0;
+
+	/**
+ 	 * Unregister a previously registered filter expression factory. 
+ 	 * @param InFilterIdentifier 
+ 	 */
+	virtual void UnregisterFilterSuggestionFactory(const FName& InFilterIdentifier) = 0;
+
+	/**
+	 * Unregister a previously registered filter expression factory. 
+	 * @param InFilterIdentifier 
+	 */
+	virtual void UnregisterFilterExpressionFactory(const FName& InFilterIdentifier) = 0;
+
+	/**
+	 *	Check if current filter expression factory support the comparison operation
+	 * @param InFilterKey Filter key to get the filter expression factory needed
+	 * @param InOperation Operation to check if supported
+	 * @return True if operation is supported, False otherwise
+	 */
+	virtual bool CanFilterSupportComparisonOperation(const FName& InFilterKey, const ETextFilterComparisonOperation InOperation) const = 0;
+
+	/**
+	 *	Evaluate the expression and return the result 
+	 * @param InFilterKey Filter Key to get the Factory
+	 * @param InItem Item that is currently checked
+	 * @param InArgs Args to evaluate the expression see FAvaTextFilterArgs for more information
+	 * @return True if the expression evaluated to True, False otherwise
+	 */
+	virtual bool FilterExpression(const FName& InFilterKey, const IAvaOutlinerItem& InItem, const FAvaTextFilterArgs& InArgs) const = 0;
+
+	/**
+	 * Get all suggestions with the given type (Generic/ItemBased/All)
+	 * @param InSuggestionType Type of suggestion to get, see EFilterSuggestionType for more information
+	 * @return An Array containing all suggestions of the given type
+	 */
+	virtual TArray<TSharedPtr<IAvaFilterSuggestionFactory>> GetSuggestions(const EAvaFilterSuggestionType& InSuggestionType) const = 0;
+
+protected:
+
+	/**
+	 * Adds the given Customization to the Map containing all of them if not already added.
+	 * @param InItemTypeId Outliner Item Type Id
+	 * @param InAvaOutlinerIconCustomization Customization to add
+	 */
+	virtual void RegisterOverridenIcon_Internal(const FAvaTypeId& InItemTypeId, const TSharedRef<IAvaOutlinerIconCustomization>& InAvaOutlinerIconCustomization) = 0;
+
+	/**
+	 * Unregister the customization with the given key
+	 * @param InItemTypeId Outliner Item Type Id
+	 * @param InSpecializationIdentifier Specialization identifier (ex. AActor class FName for ActorIconCustomization)
+	 */
+	virtual void UnregisterOverriddenIcon_Internal(const FAvaTypeId& InItemTypeId, const FName& InSpecializationIdentifier) = 0;
+
+	/**
+	 * Adds the given factory to the Map containing all of them if not already added.
+	 * @param InFilterExpressionFactory Filter expression factory to register.
+	 */
+	virtual void RegisterFilterExpressionFactory_Internal(const TSharedRef<IAvaFilterExpressionFactory>& InFilterExpressionFactory) = 0;
+
+	/**
+	 * Adds the given factory to the Map containing all of them if not already added.
+	 * @param InFilterSuggestionFactory Filter suggestion factory to register.
+	 */
+	virtual void RegisterFilterSuggestionFactory_Internal(const TSharedRef<IAvaFilterSuggestionFactory>& InFilterSuggestionFactory) = 0;
+};
diff --git a/Engine/Plugins/Experimental/Avalanche/Source/AvalancheOutliner/Public/Icon/AvaOutlinerObjectIconCustomization.h b/Engine/Plugins/Experimental/Avalanche/Source/AvalancheOutliner/Public/Icon/AvaOutlinerObjectIconCustomization.h
new file mode 100644
--- /dev/null
+++ b/Engine/Plugins/Experimental/Avalanche/Source/AvalancheOutliner/Public/Icon/AvaOutlinerObjectIconCustomization.h
@@ -0,0 +1,31 @@
+// Copyright Epic Games, Inc. All Rights Reserved.
+
+#pragma once
+
+#include "AvaOutlinerIconCustomization.h"
+#include "Delegates/Delegate.h"
+#include "Delegates/DelegateCombinations.h"
+#include "Templates/SharedPointer.h"
+
+class UClass;
+
+DECLARE_DELEGATE_RetVal_OneParam(FSlateIcon /** OutIconForItem */, FOnGetOverriddenObjectIcon, TSharedPtr<const FAvaOutlinerItem> /** InItemToCustomize */)
+
+class AVALANCHEOUTLINER_API FAvaOutlinerObjectIconCustomization : public IAvaOutlinerIconCustomization
+{
+public:
+	FAvaOutlinerObjectIconCustomization(const UClass* InSupportedClass);
+
+	//~ Begin IAvaOutlinerIconCustomization
+	virtual FName GetOutlinerItemIdentifier() const override { return SupportedClassName; }
+	virtual bool HasOverrideIcon(TSharedPtr<const FAvaOutlinerItem> InOutlinerItem) const override;
+	virtual FSlateIcon GetOverrideIcon(TSharedPtr<const FAvaOutlinerItem> InOutlinerItem) const override;
+	//~ End IAvaOutlinerIconCustomization
+
+	void SetOverriddenIcon(const FOnGetOverriddenObjectIcon& InOverriddenIcon);
+
+protected:
+	FName SupportedClassName;
+
+	FOnGetOverriddenObjectIcon OnGetOverriddenIcon;
+};
diff --git a/Engine/Plugins/Experimental/Avalanche/Source/AvalancheOutliner/Public/ItemActions/AvaOutlinerAddItem.h b/Engine/Plugins/Experimental/Avalanche/Source/AvalancheOutliner/Public/ItemActions/AvaOutlinerAddItem.h
new file mode 100644
--- /dev/null
+++ b/Engine/Plugins/Experimental/Avalanche/Source/AvalancheOutliner/Public/ItemActions/AvaOutlinerAddItem.h
@@ -0,0 +1,77 @@
+// Copyright Epic Games, Inc. All Rights Reserved.
+
+#pragma once
+
+#include "AvaOutlinerItemAction.h"
+#include "Engine/EngineTypes.h"
+
+class FAvaOutliner;
+
+enum class EItemDropZone;
+
+enum class EAvaOutlinerAddItemFlags : uint8
+{
+	None        = 0,
+	/** Also add the children of the given item even if they were not made into their own Add Item Action */
+	AddChildren = 1 << 0,
+	/** Select this Item on Add */
+	Select      = 1 << 1,
+	/** Make a Transaction for this Action */
+	Transact    = 1 << 2,
+};
+
+ENUM_CLASS_FLAGS(EAvaOutlinerAddItemFlags);
+
+struct FAvaOutlinerAddItemParams
+{
+	FAvaOutlinerAddItemParams(const FAvaOutlinerItemPtr& InItem = nullptr
+			, EAvaOutlinerAddItemFlags InFlags            = EAvaOutlinerAddItemFlags::None
+			, const FAvaOutlinerItemPtr& InRelativeItem   = nullptr
+			, TOptional<EItemDropZone> InRelativeDropZone = TOptional<EItemDropZone>())
+		: Item(InItem)
+		, RelativeItem(InRelativeItem)
+		, RelativeDropZone(InRelativeDropZone)
+		, Flags(InFlags)
+		, SelectionFlags(EAvaOutlinerItemSelectionFlags::None)
+	{
+	}
+
+	/** The Item to Add */
+	FAvaOutlinerItemPtr Item;
+
+	/** The Item to use as base in where to place the Item */
+	FAvaOutlinerItemPtr RelativeItem;
+
+	/** The Placement Order from the Relative Item (Onto/Inside, Above, Below) */
+	TOptional<EItemDropZone> RelativeDropZone;
+
+	/** Some Extra Flags for what to do when Adding or After Adding the Items */
+	EAvaOutlinerAddItemFlags Flags;
+
+	/** Flags to Indicate how we should Select the Item. This only applies if bool bSelectItem is true */
+	EAvaOutlinerItemSelectionFlags SelectionFlags;
+
+	/** Optional Transform override Rule when Attaching Items */
+	TOptional<FAttachmentTransformRules> AttachmentTransformRules;
+};
+
+/**
+ * Item Action responsible of adding an item to the Tree under a given optional Parent.
+ * If Parent is null, it is added as a Top Level Item.
+ */
+class AVALANCHEOUTLINER_API FAvaOutlinerAddItem : public IAvaOutlinerAction
+{
+public:
+	UE_AVA_INHERITS(FAvaOutlinerAddItem, IAvaOutlinerAction);
+
+	FAvaOutlinerAddItem(const FAvaOutlinerAddItemParams& InAddItemParams);
+
+	//~ Begin IAvaOutlinerAction
+	virtual bool ShouldTransact() const override;
+	virtual void Execute(FAvaOutliner& InOutliner) override;
+	virtual void OnObjectsReplaced(const TMap<UObject*, UObject*>& InReplacementMap, bool bRecursive) override;
+	//~ End IAvaOutlinerAction
+
+protected:
+	FAvaOutlinerAddItemParams AddParams;
+};
diff --git a/Engine/Plugins/Experimental/Avalanche/Source/AvalancheOutliner/Public/ItemActions/AvaOutlinerRemoveItem.h b/Engine/Plugins/Experimental/Avalanche/Source/AvalancheOutliner/Public/ItemActions/AvaOutlinerRemoveItem.h
new file mode 100644
--- /dev/null
+++ b/Engine/Plugins/Experimental/Avalanche/Source/AvalancheOutliner/Public/ItemActions/AvaOutlinerRemoveItem.h
@@ -0,0 +1,37 @@
+// Copyright Epic Games, Inc. All Rights Reserved.
+
+#pragma once
+
+#include "AvaOutlinerItemAction.h"
+#include "Engine/EngineTypes.h"
+
+struct FAvaOutlinerRemoveItemParams
+{
+	FAvaOutlinerRemoveItemParams(const FAvaOutlinerItemPtr& InItem = nullptr)
+		: Item(InItem)
+	{
+	}
+	FAvaOutlinerItemPtr Item;
+
+	/** Optional Transform override Rule when Detaching Items */
+	TOptional<FDetachmentTransformRules> DetachmentTransformRules;
+};
+
+/**
+ * Item Action responsible of removing/unregistering Items from the Tree
+ */
+class AVALANCHEOUTLINER_API FAvaOutlinerRemoveItem : public IAvaOutlinerAction
+{
+public:
+	UE_AVA_INHERITS(FAvaOutlinerRemoveItem, IAvaOutlinerAction);
+
+	FAvaOutlinerRemoveItem(const FAvaOutlinerRemoveItemParams& InRemoveItemParams);
+
+	//~ Begin IAvaOutlinerAction
+	virtual void Execute(FAvaOutliner& InOutliner) override;
+	virtual void OnObjectsReplaced(const TMap<UObject*, UObject*>& InReplacementMap, bool bRecursive) override;
+	//~ End IAvaOutlinerAction
+
+protected:
+	FAvaOutlinerRemoveItemParams RemoveParams;
+};
diff --git a/Engine/Plugins/Experimental/Avalanche/Source/AvalancheOutliner/Public/ToolMenuContext/AvaOutlinerItemsContext.h b/Engine/Plugins/Experimental/Avalanche/Source/AvalancheOutliner/Public/ToolMenuContext/AvaOutlinerItemsContext.h
new file mode 100644
--- /dev/null
+++ b/Engine/Plugins/Experimental/Avalanche/Source/AvalancheOutliner/Public/ToolMenuContext/AvaOutlinerItemsContext.h
@@ -0,0 +1,32 @@
+// Copyright Epic Games, Inc. All Rights Reserved.
+
+#pragma once
+
+#include "AvaOutlinerDefines.h"
+#include "Containers/Array.h"
+#include "Containers/ContainersFwd.h"
+#include "Templates/SharedPointer.h"
+#include "UObject/Object.h"
+#include "AvaOutlinerItemsContext.generated.h"
+
+class FAvaOutlinerView;
+
+UCLASS(MinimalAPI)
+class UAvaOutlinerItemsContext : public UObject
+{
+	GENERATED_BODY()
+	
+	friend FAvaOutlinerView;
+	
+public:
+	UAvaOutlinerItemsContext() = default;
+	
+	TSharedPtr<FAvaOutlinerView> GetOutlinerView() const { return OutlinerViewWeak.Pin(); }
+	
+	TConstArrayView<FAvaOutlinerItemWeakPtr> GetItems() const { return ItemListWeak; }
+	
+private:
+	TWeakPtr<FAvaOutlinerView> OutlinerViewWeak;
+	
+	TArray<FAvaOutlinerItemWeakPtr> ItemListWeak;
+};
diff --git a/Engine/Plugins/Experimental/Avalanche/Source/AvalancheOutliner/Public/ToolMenuContext/AvaOutlinerToolBarContext.h b/Engine/Plugins/Experimental/Avalanche/Source/AvalancheOutliner/Public/ToolMenuContext/AvaOutlinerToolBarContext.h
new file mode 100644
--- /dev/null
+++ b/Engine/Plugins/Experimental/Avalanche/Source/AvalancheOutliner/Public/ToolMenuContext/AvaOutlinerToolBarContext.h
@@ -0,0 +1,25 @@
+// Copyright Epic Games, Inc. All Rights Reserved.
+
+#pragma once
+
+#include "Templates/SharedPointer.h"
+#include "UObject/Object.h"
+#include "AvaOutlinerToolBarContext.generated.h"
+
+class FAvaOutlinerView;
+
+UCLASS(MinimalAPI)
+class UAvaOutlinerToolBarContext : public UObject
+{
+	GENERATED_BODY()
+	
+	friend FAvaOutlinerView;
+	
+public:
+	UAvaOutlinerToolBarContext() = default;
+	
+	TSharedPtr<FAvaOutlinerView> GetOutlinerView() const { return OutlinerViewWeak.Pin(); }
+	
+private:
+	TWeakPtr<FAvaOutlinerView> OutlinerViewWeak;
+};
diff --git a/Engine/Plugins/Experimental/Avalanche/Source/AvalanchePropertyAnimator/Private/AvaPropertyAnimatorModule.cpp b/Engine/Plugins/Experimental/Avalanche/Source/AvalanchePropertyAnimator/Private/AvaPropertyAnimatorModule.cpp
new file mode 100644
--- /dev/null
+++ b/Engine/Plugins/Experimental/Avalanche/Source/AvalanchePropertyAnimator/Private/AvaPropertyAnimatorModule.cpp
@@ -0,0 +1,5 @@
+﻿// Copyright Epic Games, Inc. All Rights Reserved.
+
+#include "Modules/ModuleManager.h"
+
+IMPLEMENT_MODULE(FDefaultModuleImpl, AvalanchePropertyAnimator)
\ No newline at end of file
diff --git a/Engine/Plugins/Experimental/Avalanche/Source/AvalanchePropertyAnimator/Private/TimeSources/AvaPropertyAnimatorSequenceTimeSource.cpp b/Engine/Plugins/Experimental/Avalanche/Source/AvalanchePropertyAnimator/Private/TimeSources/AvaPropertyAnimatorSequenceTimeSource.cpp
new file mode 100644
--- /dev/null
+++ b/Engine/Plugins/Experimental/Avalanche/Source/AvalanchePropertyAnimator/Private/TimeSources/AvaPropertyAnimatorSequenceTimeSource.cpp
@@ -0,0 +1,227 @@
+// Copyright Epic Games, Inc. All Rights Reserved.
+
+#include "TimeSources/AvaPropertyAnimatorSequenceTimeSource.h"
+
+#include "AvaSequence.h"
+#include "AvaSequencePlayer.h"
+#include "MovieSceneSequence.h"
+#include "UObject/UObjectIterator.h"
+
+#if WITH_EDITOR
+#include "ISequencerModule.h"
+#include "Modules/ModuleManager.h"
+#endif // WITH_EDITOR
+
+#if WITH_EDITOR
+void UAvaPropertyAnimatorSequenceTimeSource::PostEditChangeProperty(FPropertyChangedEvent& InPropertyChangedEvent)
+{
+	Super::PostEditChangeProperty(InPropertyChangedEvent);
+
+	const FName MemberName = InPropertyChangedEvent.GetMemberPropertyName();
+
+	static const FName SequenceNamePropertyName = GET_MEMBER_NAME_CHECKED(UAvaPropertyAnimatorSequenceTimeSource, SequenceName);
+
+	if (MemberName == SequenceNamePropertyName)
+	{
+		OnSequenceChanged();
+	}
+}
+#endif // WITH_EDITOR
+
+void UAvaPropertyAnimatorSequenceTimeSource::SetSequenceName(const FString& InSequenceName)
+{
+	if (SequenceName == InSequenceName)
+	{
+		return;
+	}
+
+	const TArray<FString> Sequences = GetSequenceNames();
+	if (!Sequences.Contains(InSequenceName))
+	{
+		return;
+	}
+
+	SequenceName = InSequenceName;
+	OnSequenceChanged();
+}
+
+void UAvaPropertyAnimatorSequenceTimeSource::OnTimeSourceActive()
+{
+	Super::OnTimeSourceActive();
+
+	OnSequenceChanged();
+}
+
+void UAvaPropertyAnimatorSequenceTimeSource::OnTimeSourceInactive()
+{
+	Super::OnTimeSourceInactive();
+
+	UAvaSequencePlayer::OnSequenceFinished().RemoveAll(this);
+	UAvaSequencePlayer::OnSequenceStarted().RemoveAll(this);
+
+	SequencePlayerWeak.Reset();
+	SequencePlayerWeak = nullptr;
+}
+
+#if WITH_EDITOR
+void UAvaPropertyAnimatorSequenceTimeSource::OnSequencerCreated(TSharedRef<ISequencer> InSequencer)
+{
+	SequencerWeak = InSequencer;
+}
+
+TSharedPtr<ISequencer> UAvaPropertyAnimatorSequenceTimeSource::GetSequencer() const
+{
+	if (const UAvaPropertyAnimatorSequenceTimeSource* Default = GetDefault<UAvaPropertyAnimatorSequenceTimeSource>())
+	{
+		return Default->SequencerWeak.Pin();
+	}
+
+	return nullptr;
+}
+#endif // WITH_EDITOR
+
+double UAvaPropertyAnimatorSequenceTimeSource::GetTimeElapsed()
+{
+	// Valid when a sequence player is playing
+	if (const UAvaSequencePlayer* SequencePlayer = SequencePlayerWeak.Get())
+	{
+		return SequencePlayer->GetCurrentTime().AsSeconds();
+	}
+
+#if WITH_EDITOR
+	// Use sequencer global time if sequencer active
+	if (const TSharedPtr<ISequencer> Sequencer = GetSequencer())
+	{
+		return Sequencer->GetGlobalTime().AsSeconds();
+	}
+#endif
+
+	return 0;
+}
+
+bool UAvaPropertyAnimatorSequenceTimeSource::IsTimeSourceReady() const
+{
+	return SequencePlayerWeak.IsValid()
+#if WITH_EDITOR
+		|| GetSequencer()
+#endif
+	;
+}
+
+void UAvaPropertyAnimatorSequenceTimeSource::OnTimeSourceRegistered()
+{
+	Super::OnTimeSourceRegistered();
+
+#if WITH_EDITOR
+	ISequencerModule& SequencerModule = FModuleManager::LoadModuleChecked<ISequencerModule>("Sequencer");
+	OnSequencerCreatedHandle = SequencerModule.RegisterOnSequencerCreated(FOnSequencerCreated::FDelegate::CreateUObject(this, &UAvaPropertyAnimatorSequenceTimeSource::OnSequencerCreated));
+#endif // WITH_EDITOR
+}
+
+void UAvaPropertyAnimatorSequenceTimeSource::OnTimeSourceUnregistered()
+{
+	Super::OnTimeSourceUnregistered();
+
+#if WITH_EDITOR
+	ISequencerModule& SequencerModule = FModuleManager::LoadModuleChecked<ISequencerModule>("Sequencer");
+	SequencerModule.UnregisterOnSequencerCreated(OnSequencerCreatedHandle);
+	OnSequencerCreatedHandle.Reset();
+#endif // WITH_EDITOR
+}
+
+TArray<FString> UAvaPropertyAnimatorSequenceTimeSource::GetSequenceNames() const
+{
+	TArray<FString> SequenceNames {TEXT("")};
+
+	const UWorld* World = GetWorld();
+
+	for (const UMovieSceneSequence* MovieSceneSequence : TObjectRange<UMovieSceneSequence>())
+	{
+		if (MovieSceneSequence && MovieSceneSequence->GetWorld() == World)
+		{
+			SequenceNames.Add(MovieSceneSequence->GetName());
+		}
+	}
+
+	return SequenceNames;
+}
+
+void UAvaPropertyAnimatorSequenceTimeSource::OnSequenceChanged()
+{
+	// Reset sequence selection
+	if (SequenceName.IsEmpty())
+	{
+		SequenceWeak = nullptr;
+		SequenceWeak.Reset();
+
+		SequencePlayerWeak = nullptr;
+		SequencePlayerWeak.Reset();
+
+		UAvaSequencePlayer::OnSequenceFinished().RemoveAll(this);
+		UAvaSequencePlayer::OnSequenceFinished().RemoveAll(this);
+		UAvaSequencePlayer::OnSequenceStarted().RemoveAll(this);
+
+		return;
+	}
+
+	// Find sequence with that name
+	const UWorld* World = GetWorld();
+
+	for (UMovieSceneSequence* MovieSceneSequence : TObjectRange<UMovieSceneSequence>())
+	{
+		if (MovieSceneSequence && MovieSceneSequence->GetWorld() == World)
+		{
+			if (MovieSceneSequence->GetName() == SequenceName)
+			{
+				SequenceWeak = MovieSceneSequence;
+				break;
+			}
+		}
+	}
+
+	if (!SequenceWeak.IsValid())
+	{
+		return;
+	}
+
+	UAvaSequencePlayer::OnSequenceFinished().RemoveAll(this);
+	UAvaSequencePlayer::OnSequenceStarted().RemoveAll(this);
+	UAvaSequencePlayer::OnSequenceStarted().AddUObject(this, &UAvaPropertyAnimatorSequenceTimeSource::OnSequenceStarted);
+}
+
+void UAvaPropertyAnimatorSequenceTimeSource::OnSequenceStarted(UAvaSequencePlayer* InPlayer, UAvaSequence* InSequence)
+{
+	if (!InPlayer || !InSequence)
+	{
+		return;
+	}
+
+	const UMovieSceneSequence* Sequence = SequenceWeak.Get();
+	if (InSequence != Sequence)
+	{
+		return;
+	}
+
+	SequencePlayerWeak = InPlayer;
+
+	UAvaSequencePlayer::OnSequenceFinished().RemoveAll(this);
+	UAvaSequencePlayer::OnSequenceFinished().AddUObject(this, &UAvaPropertyAnimatorSequenceTimeSource::OnSequenceFinished);
+}
+
+void UAvaPropertyAnimatorSequenceTimeSource::OnSequenceFinished(UAvaSequencePlayer* InPlayer, UAvaSequence* InSequence)
+{
+	if (!InPlayer || !InSequence)
+	{
+		return;
+	}
+
+	const UMovieSceneSequence* Sequence = SequenceWeak.Get();
+	if (Sequence != InSequence)
+	{
+		return;
+	}
+
+	UAvaSequencePlayer::OnSequenceFinished().RemoveAll(this);
+	SequencePlayerWeak = nullptr;
+	SequencePlayerWeak.Reset();
+}
diff --git a/Engine/Plugins/Experimental/Avalanche/Source/AvalanchePropertyAnimator/Private/TimeSources/AvaPropertyAnimatorSequenceTimeSource.h b/Engine/Plugins/Experimental/Avalanche/Source/AvalanchePropertyAnimator/Private/TimeSources/AvaPropertyAnimatorSequenceTimeSource.h
new file mode 100644
--- /dev/null
+++ b/Engine/Plugins/Experimental/Avalanche/Source/AvalanchePropertyAnimator/Private/TimeSources/AvaPropertyAnimatorSequenceTimeSource.h
@@ -0,0 +1,80 @@
+// Copyright Epic Games, Inc. All Rights Reserved.
+
+#pragma once
+
+#include "TimeSources/PropertyAnimatorCoreRangeTimeSource.h"
+#include "AvaPropertyAnimatorSequenceTimeSource.generated.h"
+
+class ISequencer;
+class UAvaSequence;
+class UAvaSequencePlayer;
+class UMovieSceneSequence;
+
+/**
+ * Time source that follows specific sequence time
+ */
+UCLASS()
+class UAvaPropertyAnimatorSequenceTimeSource : public UPropertyAnimatorCoreRangeTimeSource
+{
+	GENERATED_BODY()
+
+public:
+	UAvaPropertyAnimatorSequenceTimeSource()
+		: UPropertyAnimatorCoreRangeTimeSource(TEXT("Sequence"))
+	{}
+
+	//~ Begin UObject
+#if WITH_EDITOR
+	virtual void PostEditChangeProperty(FPropertyChangedEvent& InPropertyChangedEvent) override;
+#endif // WITH_EDITOR
+	//~ End UObject
+
+	void SetSequenceName(const FString& InSequenceName);
+	FString GetSequenceName() const
+	{
+		return SequenceName;
+	}
+
+	//~ Begin UPropertyAnimatorCoreTimeSourceBase
+	virtual double GetTimeElapsed() override;
+	virtual bool IsTimeSourceReady() const override;
+	//~ End UPropertyAnimatorCoreTimeSourceBase
+
+protected:
+	//~ Begin UPropertyAnimatorCoreTimeSourceBase
+	virtual void OnTimeSourceRegistered() override;
+	virtual void OnTimeSourceUnregistered() override;
+	virtual void OnTimeSourceActive() override;
+	virtual void OnTimeSourceInactive() override;
+	//~ End UPropertyAnimatorCoreTimeSourceBase
+
+	UFUNCTION()
+	TArray<FString> GetSequenceNames() const;
+
+	void OnSequenceChanged();
+
+	void OnSequenceStarted(UAvaSequencePlayer* InPlayer, UAvaSequence* InSequence);
+	void OnSequenceFinished(UAvaSequencePlayer* InPlayer, UAvaSequence* InSequence);
+
+	UPROPERTY(EditInstanceOnly, BlueprintReadWrite, Setter="SetSequenceName", Getter="GetSequenceName", Category="Animator", meta=(GetOptions="GetSequenceNames"))
+	FString SequenceName;
+
+	UPROPERTY()
+	TWeakObjectPtr<UMovieSceneSequence> SequenceWeak;
+
+	UPROPERTY(Transient)
+	TWeakObjectPtr<UAvaSequencePlayer> SequencePlayerWeak;
+
+private:
+#if WITH_EDITOR
+	void OnSequencerCreated(TSharedRef<ISequencer> InSequencer);
+
+	TSharedPtr<ISequencer> GetSequencer() const;
+
+	// Last sequencer registered
+	TWeakPtr<ISequencer> SequencerWeak;
+
+	// Handle only for CDO object
+	FDelegateHandle OnSequencerCreatedHandle;
+#endif // WITH_EDITOR
+};
\ No newline at end of file
diff --git a/Engine/Plugins/Experimental/Avalanche/Source/AvalanchePropertyAnimatorEditor/Private/Outliner/AvaPropertyAnimatorEditorOutlinerProxy.cpp b/Engine/Plugins/Experimental/Avalanche/Source/AvalanchePropertyAnimatorEditor/Private/Outliner/AvaPropertyAnimatorEditorOutlinerProxy.cpp
new file mode 100644
--- /dev/null
+++ b/Engine/Plugins/Experimental/Avalanche/Source/AvalanchePropertyAnimatorEditor/Private/Outliner/AvaPropertyAnimatorEditorOutlinerProxy.cpp
@@ -0,0 +1,138 @@
+// Copyright Epic Games, Inc. All Rights Reserved.
+
+#include "Outliner/AvaPropertyAnimatorEditorOutlinerProxy.h"
+
+#include "Animators/PropertyAnimatorCoreBase.h"
+#include "AvaOutliner.h"
+#include "Components/PropertyAnimatorCoreComponent.h"
+#include "Item/AvaOutlinerActor.h"
+#include "Outliner/AvaPropertyAnimatorEditorOutliner.h"
+#include "Selection/AvaOutlinerScopedSelection.h"
+#include "Styling/SlateIconFinder.h"
+
+#define LOCTEXT_NAMESPACE "AvaPropertyAnimatorEditorOutlinerProxy"
+
+FAvaPropertyAnimatorEditorOutlinerProxy::FAvaPropertyAnimatorEditorOutlinerProxy(FAvaOutliner& InOutliner, const FAvaOutlinerItemPtr& InParentItem)
+	: Super(InOutliner, InParentItem)
+{
+	ItemIcon = FSlateIconFinder::FindIconForClass(UPropertyAnimatorCoreComponent::StaticClass());
+}
+
+UPropertyAnimatorCoreComponent* FAvaPropertyAnimatorEditorOutlinerProxy::GetPropertyAnimatorComponent() const
+{
+	const FAvaOutlinerItemPtr Parent = GetParent();
+
+	if (!Parent.IsValid())
+	{
+		return nullptr;
+	}
+
+	const FAvaOutlinerActor* const ActorItem = Parent->CastTo<FAvaOutlinerActor>();
+
+	if (!ActorItem)
+	{
+		return nullptr;
+	}
+
+	const AActor* Actor = ActorItem->GetActor();
+
+	if (!Actor)
+	{
+		return nullptr;
+	}
+
+	return Actor->FindComponentByClass<UPropertyAnimatorCoreComponent>();
+}
+
+void FAvaPropertyAnimatorEditorOutlinerProxy::OnItemRegistered()
+{
+	Super::OnItemRegistered();
+	BindDelegates();
+}
+
+void FAvaPropertyAnimatorEditorOutlinerProxy::OnItemUnregistered()
+{
+	Super::OnItemUnregistered();
+	UnbindDelegates();
+}
+
+void FAvaPropertyAnimatorEditorOutlinerProxy::Select(FAvaOutlinerScopedSelection& InSelection) const
+{
+	if (UPropertyAnimatorCoreComponent* const PropertyAnimatorComponent = GetPropertyAnimatorComponent())
+	{
+		InSelection.Select(PropertyAnimatorComponent);
+	}
+}
+
+FText FAvaPropertyAnimatorEditorOutlinerProxy::GetDisplayName() const
+{
+	return LOCTEXT("DisplayName", "Animators");
+}
+
+FSlateIcon FAvaPropertyAnimatorEditorOutlinerProxy::GetIcon() const
+{
+	return ItemIcon;
+}
+
+FText FAvaPropertyAnimatorEditorOutlinerProxy::GetIconTooltipText() const
+{
+	return LOCTEXT("Tooltip", "Shows all the animators found in the property animator component of an actor");
+}
+
+void FAvaPropertyAnimatorEditorOutlinerProxy::GetProxiedItems(const TSharedRef<IAvaOutlinerItem>& InParent
+	, TArray<FAvaOutlinerItemPtr>& OutChildren, bool bInRecursive)
+{
+	if (const UPropertyAnimatorCoreComponent* const PropertyAnimatorComponent = GetPropertyAnimatorComponent())
+	{
+		for (UPropertyAnimatorCoreBase* const PropertyAnimator : PropertyAnimatorComponent->GetAnimators())
+		{
+			if (!PropertyAnimator)
+			{
+				continue;
+			}
+
+			const FAvaOutlinerItemPtr AnimatorItem = Outliner.FindOrAdd<FAvaPropertyAnimatorEditorOutliner>(PropertyAnimator);
+			AnimatorItem->SetParent(SharedThis(this));
+
+			OutChildren.Add(AnimatorItem);
+
+			if (bInRecursive)
+			{
+				AnimatorItem->FindChildren(OutChildren, bInRecursive);
+			}
+		}
+	}
+}
+
+void FAvaPropertyAnimatorEditorOutlinerProxy::BindDelegates()
+{
+	UnbindDelegates();
+	UPropertyAnimatorCoreBase::OnAnimatorCreatedDelegate.AddSP(this, &FAvaPropertyAnimatorEditorOutlinerProxy::OnPropertyAnimatorUpdated);
+	UPropertyAnimatorCoreBase::OnAnimatorRemovedDelegate.AddSP(this, &FAvaPropertyAnimatorEditorOutlinerProxy::OnPropertyAnimatorUpdated);
+	UPropertyAnimatorCoreBase::OnAnimatorRenamedDelegate.AddSP(this, &FAvaPropertyAnimatorEditorOutlinerProxy::OnPropertyAnimatorUpdated);
+}
+
+void FAvaPropertyAnimatorEditorOutlinerProxy::UnbindDelegates()
+{
+	UPropertyAnimatorCoreBase::OnAnimatorCreatedDelegate.RemoveAll(this);
+	UPropertyAnimatorCoreBase::OnAnimatorRemovedDelegate.RemoveAll(this);
+	UPropertyAnimatorCoreBase::OnAnimatorRenamedDelegate.RemoveAll(this);
+}
+
+void FAvaPropertyAnimatorEditorOutlinerProxy::OnPropertyAnimatorUpdated(UPropertyAnimatorCoreBase* InAnimator)
+{
+	const UPropertyAnimatorCoreComponent* PropertyAnimatorComponent = GetPropertyAnimatorComponent();
+
+	if (IsValid(InAnimator) && IsValid(PropertyAnimatorComponent))
+	{
+		const UPropertyAnimatorCoreComponent* UpdatedComponent = InAnimator->GetTypedOuter<UPropertyAnimatorCoreComponent>();
+
+		if (UpdatedComponent == PropertyAnimatorComponent)
+		{
+			RefreshChildren();
+			Outliner.RequestRefresh();
+		}
+	}
+}
+
+#undef LOCTEXT_NAMESPACE
diff --git a/Engine/Plugins/Experimental/Avalanche/Source/AvalancheRemoteControlEditor/Public/Outliner/AvaOutlinerRCTrackerComponent.h b/Engine/Plugins/Experimental/Avalanche/Source/AvalancheRemoteControlEditor/Public/Outliner/AvaOutlinerRCTrackerComponent.h
new file mode 100644
--- /dev/null
+++ b/Engine/Plugins/Experimental/Avalanche/Source/AvalancheRemoteControlEditor/Public/Outliner/AvaOutlinerRCTrackerComponent.h
@@ -0,0 +1,42 @@
+// Copyright Epic Games, Inc. All Rights Reserved.
+
+#pragma once
+
+#include "Item/AvaOutlinerObject.h"
+#include "Textures/SlateIcon.h"
+
+class URemoteControlTrackerComponent;
+
+/**
+ * Inheriting from FAvaOutlinerObject and not from FAvaOutlinerComponent, which only takes into account Scene Components.
+ * This class needs instead to handle URemoteControlTrackerComponent, which is of type UActorComponent, and not Scene Component.
+ */
+class AVALANCHEREMOTECONTROLEDITOR_API FAvaOutlinerRCTrackerComponent : public FAvaOutlinerObject
+{
+public:
+	UE_AVA_INHERITS_WITH_SUPER(FAvaOutlinerRCTrackerComponent, FAvaOutlinerObject);
+	
+	DECLARE_MULTICAST_DELEGATE_OneParam(FOnAvaOutlinerRCTrackerComponentSelected, const URemoteControlTrackerComponent*)
+	static FOnAvaOutlinerRCTrackerComponentSelected& OnRemoteControlTrackerSelected() { return OnRCTrackerSelectedDelegate; }
+	
+	FAvaOutlinerRCTrackerComponent(FAvaOutliner& InOutliner, URemoteControlTrackerComponent* InComponent);
+
+	//~ Begin IAvaOutlinerItem
+	virtual FSlateIcon GetIcon() const override;
+	virtual bool ShowVisibility(EAvaOutlinerVisibilityType InVisibilityType) const override;
+	virtual bool GetVisibility(EAvaOutlinerVisibilityType InVisibilityType) const override;
+	virtual void OnVisibilityChanged(EAvaOutlinerVisibilityType InVisibilityType, bool bInNewVisibility) override;
+	virtual void Select(FAvaOutlinerScopedSelection& InSelection) const override;
+	//~ End IAvaOutlinerItem
+	
+protected:
+	//~ Begin FAvaOutlinerObjectItem
+	virtual void SetObject_Impl(UObject* InObject) override;
+	//~ End FAvaOutlinerObjectItem
+
+	static FOnAvaOutlinerRCTrackerComponentSelected OnRCTrackerSelectedDelegate;
+
+	TWeakObjectPtr<URemoteControlTrackerComponent> TrackerComponentWeak;
+	
+	FSlateIcon TrackerIcon;
+};
diff --git a/Engine/Plugins/Experimental/Avalanche/Source/AvalancheSVGEditor/Private/AvalancheSVGEditor.cpp b/Engine/Plugins/Experimental/Avalanche/Source/AvalancheSVGEditor/Private/AvalancheSVGEditor.cpp
new file mode 100644
--- /dev/null
+++ b/Engine/Plugins/Experimental/Avalanche/Source/AvalancheSVGEditor/Private/AvalancheSVGEditor.cpp
@@ -0,0 +1,135 @@
+// Copyright Epic Games, Inc. All Rights Reserved.
+
+#include "AvalancheSVGEditor.h"
+#include "AvaInteractiveToolsDelegates.h"
+#include "IAvalancheInteractiveToolsModule.h"
+#include "Modifiers/ActorModifierCoreStack.h"
+#include "Modifiers/AvaBevelModifier.h"
+#include "Modifiers/AvaExtrudeModifier.h"
+#include "ProceduralMeshes/SVGDynamicMeshComponent.h"
+#include "ProceduralMeshes/SVGStrokeComponent.h"
+#include "SVGEngineSubsystem.h"
+#include "SVGImporterEditorCommands.h"
+#include "SVGShapesParentActor.h"
+#include "Subsystems/ActorModifierCoreSubsystem.h"
+#include "Tool/AvaSVGActorTool.h"
+
+#define LOCTEXT_NAMESPACE "AvalancheSVGEditorModule"
+
+void FAvaSVGEditorModule::StartupModule()
+{
+	FSVGImporterEditorCommands::Register();
+	FAvaInteractiveToolsDelegates::GetRegisterToolsDelegate().AddRaw(this, &FAvaSVGEditorModule::RegisterTools);
+
+	USVGEngineSubsystem::OnSVGActorSplit().BindRaw(this, &FAvaSVGEditorModule::OnSVGActorSplit);
+}
+
+void FAvaSVGEditorModule::ShutdownModule()
+{
+	FAvaInteractiveToolsDelegates::GetRegisterToolsDelegate().RemoveAll(this);
+
+	USVGEngineSubsystem::OnSVGActorSplit().Unbind();
+}
+
+void FAvaSVGEditorModule::RegisterTools(IAvalancheInteractiveToolsModule* InModule)
+{
+	InModule->RegisterTool(
+		IAvalancheInteractiveToolsModule::Get().CategoryNameActor,
+		GetDefault<UAvaSVGActorTool>()->GetToolParameters()
+	);
+}
+
+void FAvaSVGEditorModule::OnSVGActorSplit(ASVGShapesParentActor* InSVGShapesParent)
+{
+	if (!InSVGShapesParent)
+	{
+		return;
+	}
+
+	const UActorModifierCoreSubsystem* ModifierCoreSubsystem = UActorModifierCoreSubsystem::Get();
+
+	if (!ModifierCoreSubsystem)
+	{
+		return;
+	}
+
+	TMap<TObjectPtr<AActor>, TObjectPtr<USVGDynamicMeshComponent>> ShapesMap;
+	InSVGShapesParent->GetShapes(ShapesMap);
+	for (const TPair<TObjectPtr<AActor>, TObjectPtr<USVGDynamicMeshComponent>>& ShapePair : ShapesMap)
+	{
+		if (USVGDynamicMeshComponent* Shape = ShapePair.Value)
+		{
+			const float ExtrudeValue = Shape->GetExtrudeDepth();
+			ESVGExtrudeType ExtrudeType = Shape->ExtrudeType;
+			Shape->FlattenShape();
+
+			if (AActor* ShapeActor = ShapePair.Key)
+			{
+				UActorModifierCoreStack* ModifierStack = ModifierCoreSubsystem->GetActorModifierStack(ShapeActor);
+				if (!ModifierStack)
+				{
+					ModifierStack = ModifierCoreSubsystem->AddActorModifierStack(ShapeActor);
+				}
+
+				FActorModifierCoreStackInsertOp ExtrudeModifierInsertOp;
+				ExtrudeModifierInsertOp.NewModifierName = GetDefault<UAvaExtrudeModifier>()->GetModifierName();
+
+				if (!ModifierCoreSubsystem->GetAllowedModifiers(ShapeActor).Contains(ExtrudeModifierInsertOp.NewModifierName))
+				{
+					continue;
+				}
+
+				UActorModifierCoreBase* Modifier = nullptr;
+
+				Modifier = ModifierStack->InsertModifier(ExtrudeModifierInsertOp);
+				if (UAvaExtrudeModifier* ExtrudeModifier = Cast<UAvaExtrudeModifier>(Modifier))
+				{
+					EAvaExtrudeMode ExtrudeMode = EAvaExtrudeMode::Symmetrical;
+
+					switch (ExtrudeType)
+					{
+						case ESVGExtrudeType::FrontFaceOnly:
+							ExtrudeMode = EAvaExtrudeMode::Front;
+							break;
+
+						case ESVGExtrudeType::None:
+						case ESVGExtrudeType::FrontBackMirror:
+							ExtrudeMode = EAvaExtrudeMode::Symmetrical;
+							break;
+
+						default: ;
+					}
+
+					ExtrudeModifier->SetDepth(ExtrudeValue);
+					ExtrudeModifier->SetExtrudeMode(ExtrudeMode);
+
+					if (Shape->IsA<USVGStrokeComponent>())
+					{
+						ExtrudeModifier->SetCloseBack(false);
+					}
+				}
+
+				if (Shape->Bevel > 0.0f)
+				{
+					FActorModifierCoreStackInsertOp BevelModifierInsertOp;
+					BevelModifierInsertOp.NewModifierName = GetDefault<UAvaBevelModifier>()->GetModifierName();
+
+					if (!ModifierCoreSubsystem->GetAllowedModifiers(ShapeActor).Contains(BevelModifierInsertOp.NewModifierName))
+					{
+						continue;
+					}
+
+					Modifier = ModifierStack->InsertModifier(BevelModifierInsertOp);
+					if (UAvaBevelModifier* BevelModifier = Cast<UAvaBevelModifier>(Modifier))
+					{
+						BevelModifier->SetInset(Shape->Bevel);
+					}
+				}
+			}
+		}
+	}
+}
+
+#undef LOCTEXT_NAMESPACE
+
+IMPLEMENT_MODULE(FAvaSVGEditorModule, AvalancheSVGEditorModule)
diff --git a/Engine/Plugins/Experimental/Avalanche/Source/AvalancheSVGEditor/Private/Tool/AvaSVGActorTool.cpp b/Engine/Plugins/Experimental/Avalanche/Source/AvalancheSVGEditor/Private/Tool/AvaSVGActorTool.cpp
new file mode 100644
--- /dev/null
+++ b/Engine/Plugins/Experimental/Avalanche/Source/AvalancheSVGEditor/Private/Tool/AvaSVGActorTool.cpp
@@ -0,0 +1,37 @@
+// Copyright Epic Games, Inc. All Rights Reserved.
+
+#include "Tool/AvaSVGActorTool.h"
+#include "AvaInteractiveToolsSettings.h"
+#include "Builders/AvaInteractiveToolsToolBuilder.h"
+#include "SVGActor.h"
+#include "SVGImporterEditorCommands.h"
+
+UAvaSVGActorTool::UAvaSVGActorTool()
+{
+	ActorClass = ASVGActor::StaticClass();
+}
+
+bool UAvaSVGActorTool::UseIdentityRotation() const
+{
+	return ConditionalIdentityRotation();
+}
+
+FName UAvaSVGActorTool::GetCategoryName()
+{
+	return IAvalancheInteractiveToolsModule::CategoryNameActor;
+}
+
+FAvaInteractiveToolsToolParameters UAvaSVGActorTool::GetToolParameters() const
+{
+	return {
+		FSVGImporterEditorCommands::Get().SpawnSVGActor,
+		TEXT("SVG Actor Tool"),
+		6000,
+		FAvalancheInteractiveToolsCreateBuilder::CreateLambda(
+			[](UEdMode* InEdMode)
+			{
+				return UAvaInteractiveToolsToolBuilder::CreateToolBuilder<UAvaSVGActorTool>(InEdMode);
+			}),
+		ActorClass
+	};
+}
diff --git a/Engine/Plugins/Experimental/Avalanche/Source/AvalancheSVGEditor/Private/Tool/AvaSVGActorTool.h b/Engine/Plugins/Experimental/Avalanche/Source/AvalancheSVGEditor/Private/Tool/AvaSVGActorTool.h
new file mode 100644
--- /dev/null
+++ b/Engine/Plugins/Experimental/Avalanche/Source/AvalancheSVGEditor/Private/Tool/AvaSVGActorTool.h
@@ -0,0 +1,22 @@
+// Copyright Epic Games, Inc. All Rights Reserved.
+
+#pragma once
+
+#include "Tools/AvaInteractiveToolsActorPointToolBase.h"
+#include "AvaSVGActorTool.generated.h"
+
+UCLASS()
+class UAvaSVGActorTool : public UAvaInteractiveToolsActorPointToolBase
+{
+	GENERATED_BODY()
+
+public:
+	UAvaSVGActorTool();
+
+	//~ Begin UAvaInteractiveToolsToolBase
+	virtual bool UseIdentityLocation() const override { return false; }
+	virtual bool UseIdentityRotation() const override;
+	virtual FName GetCategoryName() override;
+	virtual FAvaInteractiveToolsToolParameters GetToolParameters() const override;
+	//~ End UAvaInteractiveToolsToolBase
+};
diff --git a/Engine/Plugins/Experimental/Avalanche/Source/AvalancheSVGEditor/Public/AvalancheSVGEditor.h b/Engine/Plugins/Experimental/Avalanche/Source/AvalancheSVGEditor/Public/AvalancheSVGEditor.h
new file mode 100644
--- /dev/null
+++ b/Engine/Plugins/Experimental/Avalanche/Source/AvalancheSVGEditor/Public/AvalancheSVGEditor.h
@@ -0,0 +1,23 @@
+// Copyright Epic Games, Inc. All Rights Reserved.
+
+#pragma once
+
+#include "Containers/ContainersFwd.h"
+#include "Modules/ModuleManager.h"
+
+class ASVGShapesParentActor;
+class IAvalancheInteractiveToolsModule;
+class UEdMode;
+struct FAvaInteractiveToolsToolParameters;
+
+class FAvaSVGEditorModule : public IModuleInterface
+{
+public:
+    virtual void StartupModule() override;
+    virtual void ShutdownModule() override;
+
+private:
+	void RegisterTools(IAvalancheInteractiveToolsModule* InModule);
+
+	void OnSVGActorSplit(ASVGShapesParentActor* InSVGShapesParent);
+};
diff --git a/Engine/Plugins/Experimental/Avalanche/Source/AvalancheSequence/Private/Director/AvaSequenceDirectorBlueprint.cpp b/Engine/Plugins/Experimental/Avalanche/Source/AvalancheSequence/Private/Director/AvaSequenceDirectorBlueprint.cpp
new file mode 100644
--- /dev/null
+++ b/Engine/Plugins/Experimental/Avalanche/Source/AvalancheSequence/Private/Director/AvaSequenceDirectorBlueprint.cpp
@@ -0,0 +1,94 @@
+// Copyright Epic Games, Inc. All Rights Reserved.
+
+#include "Director/AvaSequenceDirectorBlueprint.h"
+#include "Director/AvaSequenceDirectorGeneratedClass.h"
+#include "IAvaSequenceProvider.h"
+
+TConstArrayView<FAvaSequenceInfo> UAvaSequenceDirectorBlueprint::GetSequenceInfos()
+{
+	UpdateSequenceInfos();
+	return SequenceInfos;
+}
+
+#if WITH_EDITOR
+bool UAvaSequenceDirectorBlueprint::OnOuterWorldRenamed(const TCHAR* InName, UObject* InNewOuter, ERenameFlags InRenameFlags)
+{
+	const bool bIsRenameTest = (InRenameFlags & REN_Test);
+
+	TArray<UObject*> LastEditedDocumentsObjects;
+	if (!bIsRenameTest)
+	{
+		LastEditedDocumentsObjects.Reserve(LastEditedDocuments.Num());
+		for (const FEditedDocumentInfo& LastEditedDocument : LastEditedDocuments)
+		{
+			LastEditedDocumentsObjects.Add(LastEditedDocument.EditedObjectPath.ResolveObject());
+		}
+	}
+
+	if (!RenameGeneratedClasses(InName, InNewOuter, InRenameFlags))
+	{
+		return false;
+	}
+
+	if (!bIsRenameTest)
+	{
+		for (int32 DocumentIndex = 0; DocumentIndex < LastEditedDocuments.Num(); ++DocumentIndex)
+		{
+			if (LastEditedDocumentsObjects[DocumentIndex])
+			{
+				LastEditedDocuments[DocumentIndex].EditedObjectPath = LastEditedDocumentsObjects[DocumentIndex];
+			}
+		}
+	}
+
+	return true;
+}
+
+UClass* UAvaSequenceDirectorBlueprint::GetBlueprintClass() const
+{
+	return UAvaSequenceDirectorGeneratedClass::StaticClass();
+}
+
+void UAvaSequenceDirectorBlueprint::GetReparentingRules(TSet<const UClass*>& OutAllowedChildrenOfClasses, TSet<const UClass*>& OutDisallowedChildrenOfClasses) const
+{
+	OutAllowedChildrenOfClasses.Add(UAvaSequenceDirectorBlueprint::StaticClass());
+}
+#endif
+
+void UAvaSequenceDirectorBlueprint::UpdateSequenceInfos()
+{
+	UAvaSequence* OwningSequence = GetTypedOuter<UAvaSequence>();
+	if (!OwningSequence)
+	{
+		return;
+	}
+
+	IAvaSequenceProvider* SequenceProvider = OwningSequence->GetSequenceProvider();
+	if (!SequenceProvider)
+	{
+		return;
+	}
+
+	// Skip if it's not loaded yet
+	UObject* SequenceProviderObject = SequenceProvider->ToUObject();
+	if (SequenceProviderObject && !SequenceProviderObject->HasAllFlags(RF_LoadCompleted))
+	{
+		return;
+	}
+
+	const TArray<TObjectPtr<UAvaSequence>>& Sequences = SequenceProvider->GetSequences();
+	SequenceInfos.Empty(Sequences.Num());
+
+	for (UAvaSequence* Sequence : ReverseIterate(Sequences))
+	{
+		if (!IsValid(Sequence))
+		{
+			continue;
+		}
+
+		FAvaSequenceInfo Info;
+		Info.SequenceName = Sequence->GetFName();
+		Info.Sequence     = Sequence;
+		SequenceInfos.Add(MoveTemp(Info));
+	}
+}
diff --git a/Engine/Plugins/Experimental/Avalanche/Source/AvalancheSequence/Private/Transition/AvaTransitionInitializeSequence.cpp b/Engine/Plugins/Experimental/Avalanche/Source/AvalancheSequence/Private/Transition/AvaTransitionInitializeSequence.cpp
new file mode 100644
--- /dev/null
+++ b/Engine/Plugins/Experimental/Avalanche/Source/AvalancheSequence/Private/Transition/AvaTransitionInitializeSequence.cpp
@@ -0,0 +1,33 @@
+// Copyright Epic Games, Inc. All Rights Reserved.
+
+#include "Transition/AvaTransitionInitializeSequence.h"
+#include "AvaSequencePlaybackObject.h"
+
+TArray<UAvaSequencePlayer*> FAvaTransitionInitializeSequence::ExecuteSequenceTask(FStateTreeExecutionContext& InContext) const
+{
+	IAvaSequencePlaybackObject* PlaybackObject = GetPlaybackObject(InContext);
+	if (!PlaybackObject)
+	{
+		return TArray<UAvaSequencePlayer*>();
+	}
+
+	FAvaSequencePlayParams PlaySettings;
+	PlaySettings.Start = PlaySettings.End = InitializeTime;
+	PlaySettings.PlayMode = PlayMode;
+
+	switch (QueryType)
+	{
+	case EAvaTransitionSequenceQueryType::Name:
+		return PlaybackObject->PlaySequencesByLabel(SequenceName, PlaySettings);
+
+	case EAvaTransitionSequenceQueryType::Tag:
+		if (const FAvaTag* Tag = SequenceTag.GetTag())
+		{
+			return PlaybackObject->PlaySequencesByTag(*Tag, bPerformExactMatch, PlaySettings);
+		}
+		return TArray<UAvaSequencePlayer*>();
+	}
+
+	checkNoEntry();
+	return TArray<UAvaSequencePlayer*>();
+}
diff --git a/Engine/Plugins/Experimental/Avalanche/Source/AvalancheSequence/Private/Transition/AvaTransitionPlaySequenceTask.cpp b/Engine/Plugins/Experimental/Avalanche/Source/AvalancheSequence/Private/Transition/AvaTransitionPlaySequenceTask.cpp
new file mode 100644
--- /dev/null
+++ b/Engine/Plugins/Experimental/Avalanche/Source/AvalancheSequence/Private/Transition/AvaTransitionPlaySequenceTask.cpp
@@ -0,0 +1,60 @@
+// Copyright Epic Games, Inc. All Rights Reserved.
+
+#include "Transition/AvaTransitionPlaySequenceTask.h"
+#include "AvaSequence.h"
+#include "AvaSequencePlaybackObject.h"
+#include "AvaSequencePlayer.h"
+#include "AvaTransitionContext.h"
+#include "StateTreeExecutionContext.h"
+#include "StateTreeLinker.h"
+
+#define LOCTEXT_NAMESPACE "AvaTransitionPlaySequenceTask"
+
+FText FAvaTransitionPlaySequenceTask::GenerateDescription() const
+{
+	FText DescriptionFormat = LOCTEXT("TaskDescription", "Play {SequenceQueryText} {AdditionalArgs}");
+
+	FFormatNamedArguments NamedArgs;
+	NamedArgs.Add(TEXT("SequenceQueryText"), GetSequenceQueryText());
+
+	TArray<FText> AdditionalArgs;
+	if (PlaySettings.PlayMode != EAvaSequencePlayMode::Forward)
+	{
+		AdditionalArgs.Add(UEnum::GetDisplayValueAsText(PlaySettings.PlayMode));
+	}
+	AdditionalArgs.Add(UEnum::GetDisplayValueAsText(WaitType));
+
+	if (!AdditionalArgs.IsEmpty())
+	{
+		NamedArgs.Add(TEXT("AdditionalArgs"), FText::Format(INVTEXT("( {0} )"), FText::Join(FText::FromString(TEXT(" | ")), AdditionalArgs)));	
+	}
+
+	return FText::Format(DescriptionFormat, NamedArgs);
+}
+
+TArray<UAvaSequencePlayer*> FAvaTransitionPlaySequenceTask::ExecuteSequenceTask(FStateTreeExecutionContext& InContext) const
+{
+	IAvaSequencePlaybackObject* PlaybackObject = GetPlaybackObject(InContext);
+	if (!PlaybackObject)
+	{
+		return TArray<UAvaSequencePlayer*>();
+	}
+
+	switch (QueryType)
+	{
+	case EAvaTransitionSequenceQueryType::Name:
+		return PlaybackObject->PlaySequencesByLabel(SequenceName, PlaySettings);
+
+	case EAvaTransitionSequenceQueryType::Tag:
+		if (const FAvaTag* Tag = SequenceTag.GetTag())
+		{
+			return PlaybackObject->PlaySequencesByTag(*Tag, bPerformExactMatch, PlaySettings);
+		}
+		return TArray<UAvaSequencePlayer*>();
+	}
+
+	checkNoEntry();
+	return TArray<UAvaSequencePlayer*>();
+}
+
+#undef LOCTEXT_NAMESPACE
diff --git a/Engine/Plugins/Experimental/Avalanche/Source/AvalancheSequence/Private/Transition/AvaTransitionSequenceTask.cpp b/Engine/Plugins/Experimental/Avalanche/Source/AvalancheSequence/Private/Transition/AvaTransitionSequenceTask.cpp
new file mode 100644
--- /dev/null
+++ b/Engine/Plugins/Experimental/Avalanche/Source/AvalancheSequence/Private/Transition/AvaTransitionSequenceTask.cpp
@@ -0,0 +1,158 @@
+// Copyright Epic Games, Inc. All Rights Reserved.
+
+#include "Transition/AvaTransitionSequenceTask.h"
+#include "AvaSequence.h"
+#include "AvaSequencePlaybackObject.h"
+#include "AvaSequencePlayer.h"
+#include "AvaSequenceSubsystem.h"
+#include "AvaTransitionContext.h"
+#include "StateTreeExecutionContext.h"
+#include "StateTreeLinker.h"
+#include "Transition/AvaTransitionSequenceUtils.h"
+
+#define LOCTEXT_NAMESPACE "AvaTransitionSequenceTaskBase"
+
+TArray<UAvaSequencePlayer*> FAvaTransitionSequenceTaskBase::ExecuteSequenceTask(FStateTreeExecutionContext& InContext) const
+{
+	return TArray<UAvaSequencePlayer*>();
+}
+
+bool FAvaTransitionSequenceTaskBase::IsSequenceQueryValid() const
+{
+	switch (QueryType)
+	{
+	case EAvaTransitionSequenceQueryType::Name:
+		return !SequenceName.IsNone();
+
+	case EAvaTransitionSequenceQueryType::Tag:
+		return SequenceTag.IsValid();
+	}
+	checkNoEntry();
+	return false;
+}
+
+IAvaSequencePlaybackObject* FAvaTransitionSequenceTaskBase::GetPlaybackObject(FStateTreeExecutionContext& InContext) const
+{
+	const FAvaTransitionContext& TransitionContext = InContext.GetExternalData(TransitionContextHandle);
+	const UAvaSequenceSubsystem& SequenceSubsystem = InContext.GetExternalData(SequenceSubsystemHandle);
+
+	return FAvaTransitionSequenceUtils::GetPlaybackObject(TransitionContext, SequenceSubsystem);
+}
+
+TArray<TWeakObjectPtr<UAvaSequence>> FAvaTransitionSequenceTaskBase::GetActiveSequences(TConstArrayView<UAvaSequencePlayer*> InSequencePlayers) const
+{
+	TArray<TWeakObjectPtr<UAvaSequence>> OutActiveSequences;
+	OutActiveSequences.Reserve(InSequencePlayers.Num());
+
+	Algo::TransformIf(InSequencePlayers, OutActiveSequences,
+		[](UAvaSequencePlayer* InPlayer)->bool
+		{
+			return IsValid(InPlayer);
+		},
+		[](UAvaSequencePlayer* InPlayer)->TWeakObjectPtr<UAvaSequence>
+		{
+			return TWeakObjectPtr<UAvaSequence>(InPlayer->GetAvaSequence());
+		});
+
+	return OutActiveSequences;
+}
+
+EStateTreeRunStatus FAvaTransitionSequenceTaskBase::WaitForActiveSequences(FStateTreeExecutionContext& InContext) const
+{
+	FAvaTransitionSequenceInstanceData& InstanceData = InContext.GetInstanceData(*this);
+
+	EAvaTransitionSequenceWaitType WaitType = GetWaitType();
+
+	// No wait should only succeed in Enter State, and so return failed run if active sequences is empty 
+	if (WaitType == EAvaTransitionSequenceWaitType::NoWait)
+	{
+		if (InstanceData.ActiveSequences.IsEmpty())
+		{
+			return EStateTreeRunStatus::Failed;
+		}
+		return EStateTreeRunStatus::Succeeded;
+	}
+
+	IAvaSequencePlaybackObject* PlaybackObject = GetPlaybackObject(InContext);
+	if (!PlaybackObject)
+	{
+		return EStateTreeRunStatus::Failed;
+	}
+
+	return FAvaTransitionSequenceUtils::UpdatePlayerRunStatus(*PlaybackObject, InstanceData.ActiveSequences, WaitType);
+}
+
+void FAvaTransitionSequenceTaskBase::StopActiveSequences(FStateTreeExecutionContext& InContext) const
+{
+	IAvaSequencePlaybackObject* PlaybackObject = GetPlaybackObject(InContext);
+	if (!PlaybackObject)
+	{
+		return;
+	}
+
+	FAvaTransitionSequenceInstanceData& InstanceData = InContext.GetInstanceData(*this);
+
+	for (const TWeakObjectPtr<UAvaSequence>& SequenceWeak : InstanceData.ActiveSequences)
+	{
+		if (UAvaSequence* Sequence = SequenceWeak.Get())
+		{
+			PlaybackObject->StopSequence(Sequence);
+		}
+	}
+}
+
+FText FAvaTransitionSequenceTaskBase::GetSequenceQueryText() const
+{
+	switch (QueryType)
+	{
+	case EAvaTransitionSequenceQueryType::Name:
+		return FText::Format(INVTEXT("'{0}'"), FText::FromName(SequenceName));
+
+	case EAvaTransitionSequenceQueryType::Tag:
+		return FText::Format(LOCTEXT("SequenceQueryTag", "tag '{0}'"), FText::FromName(SequenceTag.ToName()));
+	}
+
+	checkNoEntry();
+	return FText::GetEmpty();
+}
+
+EStateTreeRunStatus FAvaTransitionSequenceTaskBase::EnterState(FStateTreeExecutionContext& InContext
+	, const FStateTreeTransitionResult& InTransition) const
+{
+	if (!IsSequenceQueryValid())
+	{
+		return EStateTreeRunStatus::Failed;
+	}
+
+	TArray<UAvaSequencePlayer*> SequencePlayers = ExecuteSequenceTask(InContext);
+	if (SequencePlayers.IsEmpty())
+	{
+		return EStateTreeRunStatus::Succeeded;
+	}
+
+	FInstanceDataType& InstanceData = InContext.GetInstanceData(*this);
+	InstanceData.ActiveSequences = FAvaTransitionSequenceUtils::GetSequences(SequencePlayers);
+	return WaitForActiveSequences(InContext);
+}
+
+EStateTreeRunStatus FAvaTransitionSequenceTaskBase::Tick(FStateTreeExecutionContext& InContext, const float InDeltaTime) const
+{
+	return WaitForActiveSequences(InContext);
+}
+
+void FAvaTransitionSequenceTaskBase::ExitState(FStateTreeExecutionContext& InContext, const FStateTreeTransitionResult& InTransition) const
+{
+	if (InTransition.CurrentRunStatus == EStateTreeRunStatus::Stopped)
+	{
+		StopActiveSequences(InContext);
+	}
+}
+
+bool FAvaTransitionSequenceTaskBase::Link(FStateTreeLinker& InLinker)
+{
+	FAvaTransitionTask::Link(InLinker);
+	InLinker.LinkExternalData(SequenceSubsystemHandle);
+	return true;
+}
+
+#undef LOCTEXT_NAMESPACE
diff --git a/Engine/Plugins/Experimental/Avalanche/Source/AvalancheSequence/Public/AvaAnimBindingReference.h b/Engine/Plugins/Experimental/Avalanche/Source/AvalancheSequence/Public/AvaAnimBindingReference.h
new file mode 100644
--- /dev/null
+++ b/Engine/Plugins/Experimental/Avalanche/Source/AvalancheSequence/Public/AvaAnimBindingReference.h
@@ -0,0 +1,63 @@
+// Copyright Epic Games, Inc. All Rights Reserved.
+
+#pragma once
+
+#include "Containers/Array.h"
+#include "Containers/Map.h"
+#include "Containers/Set.h"
+#include "Containers/UnrealString.h"
+#include "Misc/Guid.h"
+#include "AvaAnimBindingReference.generated.h"
+
+class UObject;
+
+/**
+ * DEPRECATED
+ * Avalanche-side Bindings have been softly deprecated after moving UAvaSequence (previously UAvaAnimation)
+ * to inherit from ULevelSequence.
+ * This is softly deprecated so that Sequence duplications still copy over these bindings.
+ */
+
+/**
+ * An external reference to an object, resolvable through an arbitrary context.
+ * Bindings consist of the relative path to the object from a specific outer (the context)
+ */
+USTRUCT()
+struct FAvaAnimBindingReference
+{
+	GENERATED_BODY();
+
+	friend class UAvaSequence;
+
+private:
+	/** Object path relative to a passed in context object */
+	UPROPERTY()
+	FString ObjectPath;
+};
+
+/** Wrapping UStruct for the References Array */
+USTRUCT()
+struct FAvaAnimBindingReferenceArray
+{
+	GENERATED_BODY()
+
+	friend class UAvaSequence;
+
+private:
+	UPROPERTY()
+	TArray<FAvaAnimBindingReference> References;
+};
+
+/** UStruct that stores a one to many mapping from an Object Binding ID, to object references that pertain to that ID. */
+USTRUCT()
+struct FAvaAnimBindingReferences
+{
+	GENERATED_BODY()
+
+	friend class UAvaSequence;
+
+private:
+	/** The map from object binding ID to an array of references that pertain to that ID */
+	UPROPERTY()
+	TMap<FGuid, FAvaAnimBindingReferenceArray> BindingIdToReferences;
+};
diff --git a/Engine/Plugins/Experimental/Avalanche/Source/AvalancheSequence/Public/AvaSequenceSubsystem.h b/Engine/Plugins/Experimental/Avalanche/Source/AvalancheSequence/Public/AvaSequenceSubsystem.h
new file mode 100644
--- /dev/null
+++ b/Engine/Plugins/Experimental/Avalanche/Source/AvalancheSequence/Public/AvaSequenceSubsystem.h
@@ -0,0 +1,46 @@
+// Copyright Epic Games, Inc. All Rights Reserved.
+
+#pragma once
+
+#include "Containers/Set.h"
+#include "Delegates/IDelegateInstance.h"
+#include "Subsystems/WorldSubsystem.h"
+#include "UObject/WeakInterfacePtr.h"
+#include "AvaSequenceSubsystem.generated.h"
+
+class IAvaSequenceController;
+class IAvaSequencePlaybackObject;
+class IAvaSequenceProvider;
+class UAvaSequence;
+class ULevel;
+
+UCLASS()
+class AVALANCHESEQUENCE_API UAvaSequenceSubsystem : public UWorldSubsystem
+{
+	GENERATED_BODY()
+
+public:
+	UAvaSequenceSubsystem() = default;
+
+	static UAvaSequenceSubsystem* Get(UObject* InPlaybackContext);
+
+	static TSharedRef<IAvaSequenceController> CreateSequenceController(UAvaSequence& InSequence, IAvaSequencePlaybackObject* InPlaybackObject);
+
+	IAvaSequencePlaybackObject* FindOrCreatePlaybackObject(ULevel* InLevel, IAvaSequenceProvider& InSequenceProvider);
+
+	IAvaSequencePlaybackObject* FindPlaybackObject(ULevel* InLevel) const;
+
+	void AddPlaybackObject(IAvaSequencePlaybackObject* InPlaybackObject);
+
+	void RemovePlaybackObject(IAvaSequencePlaybackObject* InPlaybackObject);
+
+protected:
+	//~ Begin UWorldSubsystem
+	virtual bool DoesSupportWorldType(const EWorldType::Type InWorldType) const override;
+	//~ End UWorldSubsystem
+
+private:
+	bool EnsureLevelIsAppropriate(ULevel*& InLevel) const;
+
+	TArray<TWeakInterfacePtr<IAvaSequencePlaybackObject>> PlaybackObjects;
+};
diff --git a/Engine/Plugins/Experimental/Avalanche/Source/AvalancheSequence/Public/Transition/AvaTransitionInitializeSequence.h b/Engine/Plugins/Experimental/Avalanche/Source/AvalancheSequence/Public/Transition/AvaTransitionInitializeSequence.h
new file mode 100644
--- /dev/null
+++ b/Engine/Plugins/Experimental/Avalanche/Source/AvalancheSequence/Public/Transition/AvaTransitionInitializeSequence.h
@@ -0,0 +1,24 @@
+// Copyright Epic Games, Inc. All Rights Reserved.
+
+#pragma once
+
+#include "AvaSequenceShared.h"
+#include "AvaTransitionSequenceTask.h"
+#include "AvaTransitionInitializeSequence.generated.h"
+
+USTRUCT(DisplayName = "Initialize Sequence", Category="Sequence Playback")
+struct AVALANCHESEQUENCE_API FAvaTransitionInitializeSequence : public FAvaTransitionSequenceTaskBase
+{
+	GENERATED_BODY()
+
+	//~ Begin FAvaTransitionSequenceTaskBase
+	virtual EAvaTransitionSequenceWaitType GetWaitType() const { return EAvaTransitionSequenceWaitType::NoWait; }
+	virtual TArray<UAvaSequencePlayer*> ExecuteSequenceTask(FStateTreeExecutionContext& InContext) const override;
+	//~ End FAvaTransitionSequenceTaskBase
+
+	UPROPERTY(EditAnywhere, Category = "Motion Design Sequence")
+	FAvaSequenceTime InitializeTime = FAvaSequenceTime(0.0);
+
+	UPROPERTY(EditAnywhere, Category="Parameter")
+	EAvaSequencePlayMode PlayMode = EAvaSequencePlayMode::Forward;
+};
diff --git a/Engine/Plugins/Experimental/Avalanche/Source/AvalancheSequence/Public/Transition/AvaTransitionSequenceTask.h b/Engine/Plugins/Experimental/Avalanche/Source/AvalancheSequence/Public/Transition/AvaTransitionSequenceTask.h
new file mode 100644
--- /dev/null
+++ b/Engine/Plugins/Experimental/Avalanche/Source/AvalancheSequence/Public/Transition/AvaTransitionSequenceTask.h
@@ -0,0 +1,106 @@
+// Copyright Epic Games, Inc. All Rights Reserved.
+
+#pragma once
+
+#include "AvaSequenceShared.h"
+#include "AvaTagHandle.h"
+#include "AvaTransitionSequenceEnums.h"
+#include "Tasks/AvaTransitionTask.h"
+#include "AvaTransitionSequenceTask.generated.h"
+
+class IAvaSequencePlaybackObject;
+class UAvaSequence;
+class UAvaSequencePlayer;
+class UAvaSequenceSubsystem;
+
+USTRUCT()
+struct FAvaTransitionSequenceInstanceData
+{
+	GENERATED_BODY()
+
+	/** Sequences being played */
+	UPROPERTY()
+	TArray<TWeakObjectPtr<UAvaSequence>> ActiveSequences;
+};
+
+USTRUCT(meta=(Hidden))
+struct AVALANCHESEQUENCE_API FAvaTransitionSequenceTaskBase : public FAvaTransitionTask
+{
+	GENERATED_BODY()
+
+	using FInstanceDataType = FAvaTransitionSequenceInstanceData;
+
+	/**
+	 * Execute the Sequence Task (overriden by implementation)
+	 * @return the sequence players that are relevant to the task
+	 */
+	virtual TArray<UAvaSequencePlayer*> ExecuteSequenceTask(FStateTreeExecutionContext& InContext) const;
+
+	/** Gets the Wait Type to use when Waiting for Active Sequences */
+	virtual EAvaTransitionSequenceWaitType GetWaitType() const { return EAvaTransitionSequenceWaitType::None; }
+
+	/** Determines whether the Current Sequence information is valid for query */
+	bool IsSequenceQueryValid() const;
+
+	/** Attempts to Retrieve the Playback Object from the given Execution Context */
+	IAvaSequencePlaybackObject* GetPlaybackObject(FStateTreeExecutionContext& InContext) const;
+
+	/** Gets all the Sequences from the provided Sequence Players that are Active (playing) */
+	TArray<TWeakObjectPtr<UAvaSequence>> GetActiveSequences(TConstArrayView<UAvaSequencePlayer*> InSequencePlayers) const;
+
+	/**
+	 * Helper function to determine the Tree Run Status by updating and checking if all activated Sequence Players
+	 * are in a state that match the Wait Type
+	 * @param InContext the execution context where the Instance Data of the Active Sequences is retrieved
+	 * @return the suggested tree run status
+	 */
+	EStateTreeRunStatus WaitForActiveSequences(FStateTreeExecutionContext& InContext) const;
+
+	/**
+	 * Helper function to stop all the currently active sequences
+	 * @param InContext the execution context where the Instance Data of the Active Sequences is retrieved
+	 */
+	void StopActiveSequences(FStateTreeExecutionContext& InContext) const;
+
+	FText GetSequenceQueryText() const;
+
+	//~ Begin FStateTreeTaskBase
+	virtual EStateTreeRunStatus EnterState(FStateTreeExecutionContext& InContext, const FStateTreeTransitionResult& InTransition) const override;
+	virtual EStateTreeRunStatus Tick(FStateTreeExecutionContext& InContext, const float InDeltaTime) const override;
+	virtual void ExitState(FStateTreeExecutionContext& InContext, const FStateTreeTransitionResult& InTransition) const override;
+	//~ End FStateTreeTaskBase
+
+	//~ Begin FStateTreeNodeBase
+	virtual const UStruct* GetInstanceDataType() const override { return FInstanceDataType::StaticStruct(); }
+	virtual bool Link(FStateTreeLinker& InLinker) override;
+	//~ End FStateTreeNodeBase
+
+	UPROPERTY(EditAnywhere, Category="Parameter", DisplayName="Sequence Query Type", meta=(DisplayPriority=1))
+	EAvaTransitionSequenceQueryType QueryType = EAvaTransitionSequenceQueryType::Name;
+
+	UPROPERTY(EditAnywhere, Category="Parameter", meta=(DisplayPriority=1, EditCondition="QueryType==EAvaTransitionSequenceQueryType::Name", EditConditionHides))
+	FName SequenceName;
+
+	UPROPERTY(EditAnywhere, Category="Parameter", meta=(DisplayPriority=1, EditCondition="QueryType==EAvaTransitionSequenceQueryType::Tag", EditConditionHides))
+	FAvaTagHandle SequenceTag;
+
+	UPROPERTY()
+	bool bPerformExactMatch = false;
+
+	TStateTreeExternalDataHandle<UAvaSequenceSubsystem> SequenceSubsystemHandle;
+};
+
+/** Base Task but with additional Parameters */
+USTRUCT(meta=(Hidden))
+struct AVALANCHESEQUENCE_API FAvaTransitionSequenceTask : public FAvaTransitionSequenceTaskBase
+{
+	GENERATED_BODY()
+
+	//~ Begin FAvaTransitionSequenceTaskBase
+	virtual EAvaTransitionSequenceWaitType GetWaitType() const override { return WaitType; }
+	//~ End FAvaTransitionSequenceTaskBase
+
+	/** The wait type before this task completes */
+	UPROPERTY(EditAnywhere, Category="Parameter")
+	EAvaTransitionSequenceWaitType WaitType = EAvaTransitionSequenceWaitType::WaitUntilStop;
+};
diff --git a/Engine/Plugins/Experimental/Avalanche/Source/AvalancheSequencer/Private/AvaSequenceActorFactory.h b/Engine/Plugins/Experimental/Avalanche/Source/AvalancheSequencer/Private/AvaSequenceActorFactory.h
new file mode 100644
--- /dev/null
+++ b/Engine/Plugins/Experimental/Avalanche/Source/AvalancheSequencer/Private/AvaSequenceActorFactory.h
@@ -0,0 +1,29 @@
+// Copyright Epic Games, Inc. All Rights Reserved.
+
+#pragma once
+
+#include "ActorFactories/ActorFactory.h"
+#include "AvaSequenceActorFactory.generated.h"
+
+/**
+ * Actor Factory with no spawn implementation for UAvaSequence
+ * Used only to prevent custom UAvaSequence Asset Drag Drops from creating Level Sequence Players (e.g. when dragging from Sequence Panel)
+ */
+UCLASS()
+class UAvaSequenceActorFactory : public UActorFactory
+{
+	GENERATED_BODY()
+
+public:
+	UAvaSequenceActorFactory();
+
+protected:
+	//~ Begin UActorFactory
+	virtual bool CanCreateActorFrom(const FAssetData& InAssetData, FText& OutErrorMessage) override;
+	virtual AActor* SpawnActor(UObject* InAsset, ULevel* InLevel, const FTransform& InTransform, const FActorSpawnParameters& InSpawnParams) override;
+	//~ End UActorFactory
+
+	//~ Begin IAssetFactoryInterface
+	virtual bool CanPlaceElementsFromAssetData(const FAssetData& InAssetData) override;
+	//~ End IAssetFactoryInterface
+};
