[MotionDesign]

diff --git a/Engine/Plugins/Experimental/Avalanche/Source/AvalancheSequencer/Private/Clipboard/AvaSequenceCopyableBinding.h b/Engine/Plugins/Experimental/Avalanche/Source/AvalancheSequencer/Private/Clipboard/AvaSequenceCopyableBinding.h
new file mode 100644
--- /dev/null
+++ b/Engine/Plugins/Experimental/Avalanche/Source/AvalancheSequencer/Private/Clipboard/AvaSequenceCopyableBinding.h
@@ -0,0 +1,51 @@
+// Copyright Epic Games, Inc. All Rights Reserved.
+
+#pragma once
+
+#include "UObject/ObjectMacros.h"
+#include "MovieScene.h"
+#include "AvaSequenceCopyableBinding.generated.h"
+
+UCLASS()
+class UAvaSequenceCopyableBinding : public UObject
+{
+	GENERATED_BODY()
+
+public:
+	/**
+	 * Spawnables need to know about their Object Template but we cannot rely on automatic serialization due to the object
+	 * template belonging to the Movie Scene (it gets serialized as a reference). Instead we manually serialize the object
+	 * so that we can duplicate it into a new object (which is stored in this variable) but we don't want this exported with
+	 * the rest of the text as it'll fall back to the same reference issue. Marking this as TextExportTransient solves this.
+	 */
+	UPROPERTY(TextExportTransient)
+	TObjectPtr<UObject> SpawnableObjectTemplate;
+
+	/**
+	 * Tracks are also owned by the owning Movie Sequence. We manually copy the tracks out of a binding when we copy,
+	 * because the binding stores them as a reference to a privately owned object. We store these copied tracks here,
+	 * and then restore them upon paste to re-create the tracks with the correct owner.
+	 */
+	UPROPERTY()
+	TArray<TObjectPtr<UMovieSceneTrack>> Tracks;
+
+	UPROPERTY()
+	FMovieSceneBinding Binding;
+
+	UPROPERTY()
+	FMovieSceneSpawnable Spawnable;
+
+	UPROPERTY()
+	FMovieScenePossessable Possessable;
+
+	/** FNames of the Actor Names bound. Used for lookup on Import */
+	UPROPERTY()
+	TArray<FName> BoundActorNames;
+
+	/** Path of the Subobjects bound relative to the Resolution Context. Used only if Resolution Context is NOT the Playback Context */
+	UPROPERTY()
+	TArray<FString> BoundObjectPaths;
+
+	UPROPERTY()
+	TArray<FName> FolderPath;
+};
diff --git a/Engine/Plugins/Experimental/Avalanche/Source/AvalancheSequencer/Private/Clipboard/AvaSequenceExporter.cpp b/Engine/Plugins/Experimental/Avalanche/Source/AvalancheSequencer/Private/Clipboard/AvaSequenceExporter.cpp
new file mode 100644
--- /dev/null
+++ b/Engine/Plugins/Experimental/Avalanche/Source/AvalancheSequencer/Private/Clipboard/AvaSequenceExporter.cpp
@@ -0,0 +1,284 @@
+// Copyright Epic Games, Inc. All Rights Reserved.
+
+#include "AvaSequenceExporter.h"
+#include "AvaSequenceCopyableBinding.h"
+#include "AvaSequencer.h"
+#include "GameFramework/Actor.h"
+#include "MovieScene.h"
+#include "MovieSceneBindingProxy.h"
+#include "AvaSequence.h"
+#include "UnrealExporter.h"
+#include "UObject/Package.h"
+
+namespace UE::AvaSequencer::Private
+{
+	class FAvaSequenceExportObjectInnerContext : public FExportObjectInnerContext
+	{
+	public:
+		explicit FAvaSequenceExportObjectInnerContext(UObject* InPlaybackContext)
+			: PlaybackContext(InPlaybackContext)
+		{
+		}
+		virtual ~FAvaSequenceExportObjectInnerContext() override = default;
+
+		//~ Begin FExportObjectInnerContext
+		virtual bool IsObjectSelected(const UObject* InObject) const override { return true; }
+		//~ End FExportObjectInnerContext
+
+		void SetBoundActors(const TArray<AActor*>& InBoundActors)
+		{
+			BoundActors = InBoundActors;
+		}
+
+		TConstArrayView<AActor*> GetBoundActors() const { return BoundActors; }
+
+		UObject* GetPlaybackContext() const { return PlaybackContext; }
+
+	private:
+		TObjectPtr<UObject> PlaybackContext;
+
+		TArray<TObjectPtr<AActor>> BoundActors;
+	};
+}
+
+FAvaSequenceExporter::FAvaSequenceExporter(const TSharedRef<FAvaSequencer>& InAvaSequencer)
+	: AvaSequencerWeak(InAvaSequencer)
+{
+}
+
+void FAvaSequenceExporter::ExportText(FString& InOutCopiedData, TConstArrayView<AActor*> InCopiedActors)
+{
+	TSharedPtr<FAvaSequencer> AvaSequencer = AvaSequencerWeak.Pin();
+	if (!AvaSequencer.IsValid())
+	{
+		return;
+	}
+
+	TMap<UAvaSequence*, TArray<AActor*>> SequenceMap;
+
+	// Gather the Sequences that the Copied Actors are bound to
+	for (AActor* const Actor : InCopiedActors)
+	{
+		TArray<UAvaSequence*> Sequences = AvaSequencer->GetSequencesForObject(Actor);
+		for (UAvaSequence* const Sequence : Sequences)
+		{
+			SequenceMap.FindOrAdd(Sequence).Add(Actor);
+		}
+	}
+
+	// If there are no sequences, no need to add anything on our side, early return.
+	if (SequenceMap.IsEmpty())
+	{
+		return;
+	}
+
+	UObject* const PlaybackContext = AvaSequencer->GetPlaybackContext();
+	UE::AvaSequencer::Private::FAvaSequenceExportObjectInnerContext ExportContext(PlaybackContext);
+
+	const TCHAR* const Filetype = TEXT("copy");
+	const uint32 PortFlags = PPF_DeepCompareInstances | PPF_ExportsNotFullyQualified;
+	const int32 IndentLevel = 0;
+
+	for (const TPair<UAvaSequence*, TArray<AActor*>>& Pair : SequenceMap)
+	{
+		UAvaSequence* const Sequence = Pair.Key;
+		const TArray<AActor*>& BoundActors = Pair.Value;
+
+		if (BoundActors.IsEmpty())
+		{
+			continue;
+		}
+
+		ExportContext.SetBoundActors(BoundActors);
+
+		FStringOutputDevice Ar;
+		UExporter::ExportToOutputDevice(&ExportContext, Sequence, nullptr, Ar, Filetype, IndentLevel, PortFlags);
+		InOutCopiedData += Ar;
+	}
+}
+
+UAvaSequenceExporter::UAvaSequenceExporter()
+{
+	SupportedClass = UAvaSequence::StaticClass();
+	bText = true;
+	PreferredFormatIndex = 0;
+	FormatExtension.Add(TEXT("copy"));
+	FormatDescription.Add(TEXT("Motion Design Sequence"));
+}
+
+bool UAvaSequenceExporter::ExportText(const FExportObjectInnerContext* InContext, UObject* InObject, const TCHAR* InType
+	, FOutputDevice& Ar, FFeedbackContext* InWarn, uint32 InPortFlags)
+{
+	UAvaSequence* const Sequence = Cast<UAvaSequence>(InObject);
+	if (!InContext || !IsValid(Sequence))
+	{
+		return false;
+	}
+
+	const UE::AvaSequencer::Private::FAvaSequenceExportObjectInnerContext& Context
+		= static_cast<const UE::AvaSequencer::Private::FAvaSequenceExportObjectInnerContext&>(*InContext);
+
+	UObject* const PlaybackContext = Context.GetPlaybackContext();
+
+	// Gather guids for the object nodes and any child object nodes
+	TArray<FMovieSceneBindingProxy> Bindings;
+
+	auto TryGetBinding = [&Bindings, Sequence](UObject* InBoundObject)->bool
+		{
+			const FGuid Guid = Sequence->FindGuidFromObject(InBoundObject);
+
+			if (Guid.IsValid())
+			{
+				Bindings.Add(FMovieSceneBindingProxy(Guid, Sequence));
+				return true;
+			}
+
+			return false;
+		};
+
+	constexpr bool bIncludeNestedObjects  = true;
+	constexpr EObjectFlags ExclusionFlags = RF_Transient | RF_TextExportTransient | RF_BeginDestroyed | RF_FinishDestroyed;
+
+	for (AActor* const Actor : Context.GetBoundActors())
+	{
+		if (TryGetBinding(Actor))
+		{
+			TArray<UObject*> Subobjects;
+			GetObjectsWithOuter(Actor, Subobjects, bIncludeNestedObjects, ExclusionFlags);
+
+			for (UObject* const Subobject : Subobjects)
+			{
+				TryGetBinding(Subobject);
+			}
+		}
+	}
+
+	if (Bindings.IsEmpty())
+	{
+		return false;
+	}
+
+	Ar.Logf(TEXT("%sBegin Sequence Label=%s\r\n")
+		, FCString::Spc(TextIndent)
+		, *Sequence->GetLabel().ToString());
+
+	UObject* const CopyableBindingOuter = GetTransientPackage();
+
+	TArray<UAvaSequenceCopyableBinding*> ObjectsToExport;
+	ObjectsToExport.Reserve(Bindings.Num());
+
+	for (const FMovieSceneBindingProxy& ObjectBinding : Bindings)
+	{
+		UMovieScene* const MovieScene = ObjectBinding.GetMovieScene();
+		if (!IsValid(MovieScene))
+		{
+			continue;
+		}
+
+		UAvaSequenceCopyableBinding* const CopyableBinding = NewObject<UAvaSequenceCopyableBinding>(CopyableBindingOuter, NAME_None, RF_Transient);
+
+		ObjectsToExport.Add(CopyableBinding);
+
+		if (FMovieScenePossessable* const Possessable = MovieScene->FindPossessable(ObjectBinding.BindingID))
+		{
+			UObject* ResolutionContext = FAvaSequencer::FindResolutionContext(*Sequence
+				, *MovieScene
+				, Possessable->GetParent()
+				, PlaybackContext);
+
+			CopyableBinding->Possessable = *Possessable;
+
+			TArray<UObject*, TInlineAllocator<1>> BoundObjects;
+			Sequence->LocateBoundObjects(CopyableBinding->Possessable.GetGuid(), UE::UniversalObjectLocator::FResolveParams(ResolutionContext), BoundObjects);
+
+			// Store the names of the bound objects so that they can be found on paste
+			for (UObject* const BoundObject : BoundObjects)
+			{
+				if (!BoundObject)
+				{
+					continue;
+				}
+
+				if (AActor* const Actor = Cast<AActor>(BoundObject))
+				{
+					CopyableBinding->BoundActorNames.Add(Actor->GetFName());
+				}
+				else if (ResolutionContext != PlaybackContext)
+				{
+					CopyableBinding->BoundObjectPaths.Add(BoundObject->GetPathName(ResolutionContext));
+				}
+			}
+		}
+		else
+		{
+			FMovieSceneSpawnable* const Spawnable = MovieScene->FindSpawnable(ObjectBinding.BindingID);
+			if (Spawnable)
+			{
+				CopyableBinding->Spawnable = *Spawnable;
+
+				// We manually serialize the spawnable object template so that it's not a reference to a privately owned object. Spawnables all have unique copies of their template objects anyways.
+				// Object Templates are re-created on paste (based on these templates) with the correct ownership set up.
+				CopyableBinding->SpawnableObjectTemplate = Spawnable->GetObjectTemplate();
+			}
+		}
+
+		if (const FMovieSceneBinding* const Binding = MovieScene->FindBinding(ObjectBinding.BindingID))
+		{
+			CopyableBinding->Binding = *Binding;
+			for (UMovieSceneTrack* const Track : Binding->GetTracks())
+			{
+				// Tracks suffer from the same issues as Spawnable's Object Templates (reference to a privately owned object). We'll manually serialize the tracks to copy them,
+				// and then restore them on paste.
+				UMovieSceneTrack* const DuplicatedTrack = Cast<UMovieSceneTrack>(StaticDuplicateObject(Track, CopyableBinding));
+				CopyableBinding->Tracks.Add(DuplicatedTrack);
+			}
+		}
+	}
+
+	UAvaSequenceExporter::ExportBindings(ObjectsToExport, TEXT("copy"), TextIndent + 3, Ar, InWarn);
+
+	Ar.Logf(TEXT("%sEnd Sequence\r\n"), FCString::Spc(TextIndent));
+
+	return true;
+}
+
+void UAvaSequenceExporter::ExportBindings(const TArray<UAvaSequenceCopyableBinding*>& InObjectsToExport, const TCHAR* InType
+	, int32 InTextIndent
+	, FOutputDevice& Ar
+	, FFeedbackContext* InWarn)
+{
+	// Clear the mark state for saving.
+	UnMarkAllObjects(static_cast<EObjectMark>(EObjectMark::OBJECTMARK_TagExp | EObjectMark::OBJECTMARK_TagImp));
+
+	UObject* const Outer = GetTransientPackage();
+
+	const uint32 PortFlags = PPF_ExportsNotFullyQualified | PPF_Copy | PPF_Delimited;
+
+	FExportObjectInnerContext EmptyContext;
+
+	for (UAvaSequenceCopyableBinding* const ObjectToExport : InObjectsToExport)
+	{
+		// We can't use TextExportTransient on USTRUCTS (which our object contains) so we're going to manually null out some references before serializing them. These references are
+		// serialized manually into the archive, as the auto-serialization will only store a reference (to a privately owned object) which creates issues on deserialization. Attempting 
+		// to deserialize these private objects throws a superflous error in the console that makes it look like things went wrong when they're actually OK and expected.
+		TArray<UMovieSceneTrack*> OldTracks = ObjectToExport->Binding.StealTracks(nullptr);
+
+		UObject* const OldSpawnableTemplate = ObjectToExport->Spawnable.GetObjectTemplate();
+
+		ObjectToExport->Spawnable.SetObjectTemplate(nullptr);
+
+		UExporter::ExportToOutputDevice(&EmptyContext, ObjectToExport, nullptr, Ar, InType
+			, InTextIndent, PortFlags, false, Outer);
+
+		// Restore the references (as we don't want to modify the original in the event of a copy operation!)
+		ObjectToExport->Binding.SetTracks(MoveTemp(OldTracks), nullptr);
+		ObjectToExport->Spawnable.SetObjectTemplate(OldSpawnableTemplate);
+
+		// We manually export the object template for the same private-ownership reason as above. Templates need to be re-created anyways as each Spawnable contains its own copy of the template.
+		if (ObjectToExport->SpawnableObjectTemplate)
+		{
+			UExporter::ExportToOutputDevice(&EmptyContext, ObjectToExport->SpawnableObjectTemplate, nullptr, Ar, InType
+				, InTextIndent, PortFlags);
+		}
+	}
+}
diff --git a/Engine/Plugins/Experimental/Avalanche/Source/AvalancheSequencer/Private/Clipboard/AvaSequenceImporter.cpp b/Engine/Plugins/Experimental/Avalanche/Source/AvalancheSequencer/Private/Clipboard/AvaSequenceImporter.cpp
new file mode 100644
--- /dev/null
+++ b/Engine/Plugins/Experimental/Avalanche/Source/AvalancheSequencer/Private/Clipboard/AvaSequenceImporter.cpp
@@ -0,0 +1,559 @@
+// Copyright Epic Games, Inc. All Rights Reserved.
+
+#include "AvaSequenceImporter.h"
+#include "AvaSequenceCopyableBinding.h"
+#include "AvaSequence.h"
+#include "AvaSequencer.h"
+#include "EngineUtils.h"
+#include "Factories.h"
+#include "IAvaSequenceProvider.h"
+#include "ISequencer.h"
+#include "MovieScenePossessable.h"
+#include "SequencerUtilities.h"
+#include "UObject/Package.h"
+
+FAvaSequenceImporter::FAvaSequenceImporter(const TSharedRef<FAvaSequencer>& InAvaSequencer)
+	: AvaSequencerWeak(InAvaSequencer)
+{
+}
+
+void FAvaSequenceImporter::ImportText(FStringView InPastedData, const TMap<FName, AActor*>& InPastedActors)
+{
+	TSharedPtr<FAvaSequencer> AvaSequencer = AvaSequencerWeak.Pin();
+	if (!AvaSequencer.IsValid())
+	{
+		return;
+	}
+
+	IAvaSequencerProvider& Provider  = AvaSequencer->GetProvider();
+	TSharedRef<ISequencer> Sequencer = AvaSequencer->GetSequencer();
+
+	UAvaSequence* const OriginallyViewedSequence = AvaSequencer->GetViewedSequence();
+
+	const TCHAR* Buffer = InPastedData.GetData();
+
+	FString StringLine;
+	while (FParse::Line(&Buffer, StringLine))
+	{
+		const TCHAR* Str = *StringLine;
+		if (!ParseCommand(&Str, TEXT("Begin")))
+		{
+			continue;
+		}
+
+		FName SequenceLabel;
+		FParse::Value(Str, TEXT("Label="), SequenceLabel);
+
+		UAvaSequence* const SequenceToUse = GetOrCreateSequence(*AvaSequencer, SequenceLabel);
+		if (!ensure(IsValid(SequenceToUse)))
+		{
+			continue;
+		}
+
+		UMovieScene* const MovieScene = SequenceToUse->GetMovieScene();
+		if (!ensure(IsValid(MovieScene)))
+		{
+			continue;
+		}
+
+		UsedSequences.Add(SequenceToUse);
+
+		AvaSequencer->SetViewedSequence(SequenceToUse);
+
+		FString BindingsString, BindingsStringLine;
+		while (!ParseCommand(&Buffer, TEXT("End")) && FParse::Line(&Buffer, BindingsStringLine))
+		{
+			BindingsString += *BindingsStringLine;
+			BindingsString += TEXT("\r\n");
+		}
+
+		const TArray<UAvaSequenceCopyableBinding*> ImportedBindings = ImportBindings(*Sequencer, BindingsString);
+		ProcessImportedBindings(ImportedBindings, InPastedActors);
+	}
+
+	AvaSequencer->SetViewedSequence(OriginallyViewedSequence);
+}
+
+bool FAvaSequenceImporter::ParseCommand(const TCHAR** InStream, const TCHAR* InToken)
+{
+	static const TCHAR* const SequenceToken = TEXT("Sequence");
+	const TCHAR* Original = *InStream;
+
+	if (FParse::Command(InStream, InToken) && FParse::Command(InStream, SequenceToken))
+	{
+		return true;
+	}
+
+	*InStream = Original;
+
+	return false;
+}
+
+void FAvaSequenceImporter::ResetCopiedTracksFlags(UMovieSceneTrack* InTrack)
+{
+	if (!IsValid(InTrack))
+	{
+		return;
+	}
+
+	InTrack->ClearFlags(RF_Transient);
+
+	for (UMovieSceneSection* Section : InTrack->GetAllSections())
+	{
+		Section->ClearFlags(RF_Transient);
+		Section->PostPaste();
+	}
+}
+
+UAvaSequence* FAvaSequenceImporter::GetOrCreateSequence(FAvaSequencer& InAvaSequencer, FName InSequenceLabel)
+{
+	IAvaSequenceProvider* const SequenceProvider = InAvaSequencer.GetProvider().GetSequenceProvider();
+	if (!SequenceProvider)
+	{
+		return nullptr;
+	}
+
+	// Find an unused sequence that matches the given label
+	const TObjectPtr<UAvaSequence>* const  FoundSequence = SequenceProvider->GetSequences().FindByPredicate(
+		[InSequenceLabel, this](const UAvaSequence* const InSequence)
+		{
+			return IsValid(InSequence)
+				&& InSequence->GetLabel() == InSequenceLabel
+				&& !UsedSequences.Contains(InSequence);
+		});
+
+	if (FoundSequence)
+	{
+		return *FoundSequence;
+	}
+
+	UAvaSequence* const NewSequence = InAvaSequencer.CreateSequence();
+	NewSequence->SetLabel(InSequenceLabel);
+	SequenceProvider->AddSequence(NewSequence);
+
+	return NewSequence;
+}
+
+TArray<UAvaSequenceCopyableBinding*> FAvaSequenceImporter::ImportBindings(ISequencer& InSequencerRef, const FString& InBindingsString)
+{
+	class FAvaObjectBindingTextFactory : public FCustomizableTextObjectFactory
+	{
+	public:
+		FAvaObjectBindingTextFactory(ISequencer& InSequencer) : FCustomizableTextObjectFactory(GWarn), Sequencer(InSequencer)
+		{
+		}
+
+		virtual bool CanCreateClass(UClass* InObjectClass, bool& bOmitSubObjs) const override
+		{
+			if (InObjectClass->IsChildOf<UAvaSequenceCopyableBinding>())
+			{
+				return true;
+			}
+			return Sequencer.GetSpawnRegister().CanSpawnObject(InObjectClass);
+		}
+
+		virtual void ProcessConstructedObject(UObject* InObject) override
+		{
+			check(InObject);
+
+			if (InObject->IsA<UAvaSequenceCopyableBinding>())
+			{
+				UAvaSequenceCopyableBinding* CopyableBinding = Cast<UAvaSequenceCopyableBinding>(InObject);
+				ImportedBindings.Add(CopyableBinding);
+			}
+			else
+			{
+				NewSpawnableObjectTemplates.Add(InObject);
+			}
+		}
+
+		TArray<UAvaSequenceCopyableBinding*> ImportedBindings;
+		TArray<UObject*> NewSpawnableObjectTemplates;
+
+	private:
+		ISequencer& Sequencer;
+	};
+
+	UPackage* const TempPackage = NewObject<UPackage>(nullptr, TEXT("/Engine/Editor/AvalancheSequencer/Transient"), RF_Transient);
+	TempPackage->AddToRoot();
+
+	// Turn the text buffer into objects
+	FAvaObjectBindingTextFactory Factory(InSequencerRef);
+	Factory.ProcessBuffer(TempPackage, RF_Transactional, InBindingsString);
+
+	// We had to explicitly serialize object templates due to them being a reference to a privately owned object.
+	// We now deserialize these object template copies
+	// and match them up with their MovieSceneCopyableBinding again.
+	int32 SpawnableObjectTemplateIndex = 0;
+	for (UAvaSequenceCopyableBinding* const ImportedObject : Factory.ImportedBindings)
+	{
+		if (ImportedObject->Spawnable.GetGuid().IsValid() && SpawnableObjectTemplateIndex < Factory.NewSpawnableObjectTemplates.Num())
+		{
+			// This Spawnable Object Template is owned by our transient package, so you'll need to change the owner if you want to keep it later.
+			ImportedObject->SpawnableObjectTemplate = Factory.NewSpawnableObjectTemplates[SpawnableObjectTemplateIndex++];
+		}
+	}
+
+	// Remove the temp package from the root now that it has served its purpose
+	TempPackage->RemoveFromRoot();
+
+	return Factory.ImportedBindings;
+}
+
+void FAvaSequenceImporter::ProcessImportedBindings(const TArray<UAvaSequenceCopyableBinding*>& InImportedBindings
+	, const TMap<FName, AActor*>& InPastedActors)
+{
+	if (InImportedBindings.IsEmpty())
+	{
+		return;
+	}
+
+	TSharedPtr<FAvaSequencer> AvaSequencer = AvaSequencerWeak.Pin();
+	if (!ensureAlways(AvaSequencer.IsValid()))
+	{
+		return;
+	}
+
+	TSharedRef<ISequencer> Sequencer = AvaSequencer->GetSequencer();
+	UAvaSequence* const Sequence     = AvaSequencer->GetViewedSequence();
+	UObject* const PlaybackContext   = Sequencer->GetPlaybackContext();
+
+	if (!ensureAlways(Sequence && PlaybackContext))
+	{
+		return;
+	}
+
+	UMovieScene* const MovieScene = Sequence->GetMovieScene();
+	if (!ensureAlways(MovieScene))
+	{
+		return;
+	}
+
+	struct FImportedPossessable
+	{
+		TArray<FString> BoundObjectPaths;
+	};
+	TMap<FGuid, FImportedPossessable> ImportedPossessableMap;
+	TMap<FGuid, FGuid> OldToNewGuidMap;
+	TArray<FMovieSceneBinding> PastedBindings;
+
+	ImportedPossessableMap.Reserve(InImportedBindings.Num());
+	OldToNewGuidMap.Reserve(InImportedBindings.Num());
+	PastedBindings.Reserve(InImportedBindings.Num());
+
+	FImportContext ImportContext { Sequencer, *Sequence, *MovieScene, Sequencer->GetFocusedTemplateID(), InPastedActors, nullptr };
+
+	for (UAvaSequenceCopyableBinding* const CopyableBinding : InImportedBindings)
+	{
+		// Clear transient flags on the imported tracks
+		for (UMovieSceneTrack* const CopiedTrack : CopyableBinding->Tracks)
+		{
+			FAvaSequenceImporter::ResetCopiedTracksFlags(CopiedTrack);
+		}
+
+		ImportContext.CopyableBinding = CopyableBinding;
+
+		// Possessable Imported Binding
+		if (CopyableBinding->Possessable.GetGuid().IsValid())
+		{
+			FMovieSceneBinding NewBinding = BindPossessable(ImportContext);
+
+			FImportedPossessable ImportedPossessable;
+			ImportedPossessable.BoundObjectPaths = CopyableBinding->BoundObjectPaths;
+			ImportedPossessableMap.Add(NewBinding.GetObjectGuid(), MoveTemp(ImportedPossessable));
+
+			OldToNewGuidMap.Add(CopyableBinding->Possessable.GetGuid(), NewBinding.GetObjectGuid());
+			PastedBindings.Add(MoveTemp(NewBinding));
+		}
+		// Spawnable Imported Binding
+		else if (CopyableBinding->Spawnable.GetGuid().IsValid())
+		{
+			FMovieSceneBinding NewBinding = BindSpawnable(ImportContext);
+			OldToNewGuidMap.Add(CopyableBinding->Spawnable.GetGuid(), NewBinding.GetObjectGuid());
+			PastedBindings.Add(MoveTemp(NewBinding));
+		}
+	}
+
+	// Fix up possessables' parent guids
+	for (const TPair<FGuid, FImportedPossessable>& Pair : ImportedPossessableMap)
+	{
+		FMovieScenePossessable* const Possessable = MovieScene->FindPossessable(Pair.Key);
+		if (!Possessable)
+		{
+			continue;
+		}
+
+		if (FGuid* const NewParentGuid = OldToNewGuidMap.Find(Possessable->GetParent()))
+		{
+			Possessable->SetParent(*NewParentGuid, MovieScene);
+		}
+	}
+
+	// Fix possessables that have bound object paths
+	for (const TPair<FGuid, FImportedPossessable>& Pair : ImportedPossessableMap)
+	{
+		const FGuid& PossessableGuid = Pair.Key;
+		const FImportedPossessable& ImportedPossessable = Pair.Value;
+
+		if (ImportedPossessable.BoundObjectPaths.IsEmpty())
+		{
+			continue;
+		}
+
+		FMovieScenePossessable* const Possessable = MovieScene->FindPossessable(PossessableGuid);
+		if (!Possessable)
+		{
+			continue;
+		}
+
+		auto FindObjectsFunc = [&ImportedPossessableMap, Sequence](const FGuid& InGuid, UObject* InContextChecked)
+			{
+				TArray<UObject*, TInlineAllocator<1>> FoundObjects;
+
+				// Attempt to use the bound object paths that were copied over through the copyable binding
+				// to have an easier time trying to resolve the object
+				if (const FImportedPossessable* const FoundEntry = ImportedPossessableMap.Find(InGuid))
+				{
+					for (const FString& BoundObjectPath : FoundEntry->BoundObjectPaths)
+					{
+						if (UObject* FoundObject = FindObject<UObject>(InContextChecked, *BoundObjectPath))
+						{
+							FoundObjects.Add(FoundObject);
+						}
+					}
+				}
+
+				// If nothing was found, use default method to locate the object
+				if (FoundObjects.IsEmpty() || !FoundObjects[0])
+				{
+					Sequence->LocateBoundObjects(InGuid, UE::UniversalObjectLocator::FResolveParams(InContextChecked), FoundObjects);
+				}
+
+				return FoundObjects;
+			};
+
+		UObject* const ResolutionContext = FAvaSequencer::FindResolutionContext(*Sequence
+			, *MovieScene
+			, Possessable->GetParent()
+			, PlaybackContext
+			, FindObjectsFunc);
+
+		TArray<UObject*> ObjectsToBind;
+		for (const FString& BoundObjectPath : ImportedPossessable.BoundObjectPaths)
+		{
+			if (UObject* FoundObject = FindObject<UObject>(ResolutionContext, *BoundObjectPath))
+			{
+				ObjectsToBind.Add(FoundObject);
+			}
+		}
+
+		AddObjectsToBinding(Sequencer
+			, ObjectsToBind
+			, FMovieSceneBindingProxy(PossessableGuid, Sequence)
+			, ResolutionContext);
+	}
+
+	Sequencer->RestorePreAnimatedState();
+
+	Sequencer->OnMovieSceneBindingsPasted().Broadcast(PastedBindings);
+
+	Sequencer->NotifyMovieSceneDataChanged(EMovieSceneDataChangeType::MovieSceneStructureItemsChanged);
+
+	// Refresh all immediately so that spawned actors will be generated immediately
+	Sequencer->ForceEvaluate();
+
+	// Remap bindings in sections (ie. attach tracks)
+	for (const TPair<FGuid, FGuid>& Pair : OldToNewGuidMap)
+	{
+		FSequencerUtilities::UpdateBindingIDs(Sequencer, Pair.Key, Pair.Value);
+	}
+}
+
+FMovieSceneBinding FAvaSequenceImporter::BindPossessable(const FImportContext& InContext)
+{
+	const FGuid NewGuid = FGuid::NewGuid();
+
+	FMovieSceneBinding NewBinding(NewGuid, InContext.CopyableBinding->Binding.GetName(), InContext.CopyableBinding->Tracks);
+	{
+		FMovieScenePossessable NewPossessable = InContext.CopyableBinding->Possessable;
+		NewPossessable.SetGuid(NewGuid);
+		InContext.MovieScene.AddPossessable(NewPossessable, NewBinding);
+	}
+
+	// Find the actors that this pasted binding should bind to
+	TArray<UObject*> ActorsToBind;
+
+	// Only attempt to bind actor names here, as this is the only thing that can be rebound at this point
+	// Subobjects first need their parents (actors) to resolve correctly, so wait on second pass
+	for (const FName& BoundActorName : InContext.CopyableBinding->BoundActorNames)
+	{
+		if (AActor* const * const FoundActor = InContext.PastedActors.Find(BoundActorName))
+		{
+			AActor* const Actor = *FoundActor;
+			if (!Actor)
+			{
+				continue;
+			}
+
+			// Don't bind if this actor is already bound
+			if (InContext.Sequencer->FindObjectId(*Actor, InContext.SequenceIDRef).IsValid() || ActorsToBind.Contains(Actor))
+			{
+				continue;
+			}
+
+			ActorsToBind.Add(Actor);
+		}
+	}
+
+	// Bind the actors
+	if (!ActorsToBind.IsEmpty())
+	{
+		AddObjectsToBinding(InContext.Sequencer
+			, ActorsToBind
+			, FMovieSceneBindingProxy(NewGuid, &InContext.Sequence)
+			, InContext.Sequencer->GetPlaybackContext());
+	}
+
+	return NewBinding;
+}
+
+FMovieSceneBinding FAvaSequenceImporter::BindSpawnable(const FImportContext& InContext)
+{
+	// We need to let the sequence create the spawnable so that it has everything set up properly internally.
+	// This is required to get spawnables with the correct references to object templates, object templates with
+	// correct owners, etc. However, making a new spawnable also creates the binding for us - this is a problem
+	// because we need to use our binding (which has tracks associated with it). To solve this, we let it create
+	// an object template based off of our (transient package owned) template, then find the newly created binding
+	// and update it.
+
+	FGuid NewGuid;
+	if (InContext.CopyableBinding->SpawnableObjectTemplate)
+	{
+		NewGuid = InContext.Sequencer->MakeNewSpawnable(*InContext.CopyableBinding->SpawnableObjectTemplate
+			, nullptr
+			, false);
+	}
+	else
+	{
+		FMovieSceneSpawnable NewSpawnable{};
+		NewSpawnable.SetGuid(FGuid::NewGuid());
+		NewSpawnable.SetName(InContext.CopyableBinding->Spawnable.GetName());
+
+		InContext.MovieScene.AddSpawnable(NewSpawnable, FMovieSceneBinding(NewSpawnable.GetGuid(), NewSpawnable.GetName()));
+
+		NewGuid = NewSpawnable.GetGuid();
+	}
+
+	FMovieSceneBinding NewBinding(NewGuid, InContext.CopyableBinding->Binding.GetName(), InContext.CopyableBinding->Tracks);
+
+	FMovieSceneSpawnable* const Spawnable = InContext.MovieScene.FindSpawnable(NewGuid);
+
+	// Copy the name of the original spawnable too.
+	Spawnable->SetName(InContext.CopyableBinding->Spawnable.GetName());
+
+	// Clear the transient flags on the copyable binding before assigning to the new spawnable
+	for (UMovieSceneTrack* const Track : NewBinding.GetTracks())
+	{
+		FAvaSequenceImporter::ResetCopiedTracksFlags(Track);
+	}
+
+	// Replace the auto-generated binding with our deserialized bindings (which has our tracks)
+	InContext.MovieScene.ReplaceBinding(NewGuid, NewBinding);
+
+	return NewBinding;
+}
+
+void FAvaSequenceImporter::AddObjectsToBinding(const TSharedRef<ISequencer>& InSequencer
+	, const TArray<UObject*>& InObjectsToAdd
+	, const FMovieSceneBindingProxy& InObjectBinding
+	, UObject* InResolutionContext)
+{
+	UAvaSequence* Sequence = CastChecked<UAvaSequence>(InObjectBinding.Sequence);
+	if (!Sequence)
+	{
+		return;
+	}
+
+	UMovieScene* const MovieScene = Sequence->GetMovieScene();
+	if (!MovieScene || InObjectsToAdd.IsEmpty())
+	{
+		return;
+	}
+
+	UClass* ObjectClass = nullptr;
+	int32 ValidObjectCount = 0;
+
+	FGuid Guid = InObjectBinding.BindingID;
+
+	TArrayView<TWeakObjectPtr<>> WeakObjectsInSequence = Sequence->FindObjectsFromGuid(Guid);
+	for (TWeakObjectPtr<> ObjectWeak : WeakObjectsInSequence)
+	{
+		if (UObject* Object = ObjectWeak.Get())
+		{
+			ObjectClass = Object->GetClass();
+			++ValidObjectCount;
+		}
+	}
+
+	Sequence->Modify();
+	MovieScene->Modify();
+
+	TArray<UObject*> AddedObjects;
+	AddedObjects.Reserve(InObjectsToAdd.Num());
+
+	for (UObject* ObjectToAdd : InObjectsToAdd)
+	{
+		// Skip invalid objects or objects already in the sequence
+		if (!ObjectToAdd || WeakObjectsInSequence.Contains(ObjectToAdd))
+		{
+			continue;
+		}
+
+		// Skip if the object has no common class with the objects already in the binding
+		if (ObjectClass && !UClass::FindCommonBase(ObjectToAdd->GetClass(), ObjectClass))
+		{
+			continue;
+		}
+
+		// if no objects are in the binding, set the class to this object's
+		if (!ObjectClass)
+		{
+			ObjectClass = ObjectToAdd->GetClass();
+		}
+
+		FMovieScenePossessable* Possessable = MovieScene->FindPossessable(Guid);
+		if (!ensureAlways(Possessable))
+		{
+			continue;
+		}
+
+		ObjectToAdd->Modify();
+		if (!Possessable->BindSpawnableObject(InSequencer->GetFocusedTemplateID(), ObjectToAdd, InSequencer->GetSharedPlaybackState()))
+		{
+			Sequence->BindPossessableObject(Guid, *ObjectToAdd, InResolutionContext);
+		}
+		AddedObjects.Add(ObjectToAdd);
+	}
+
+	// Update Labels
+	if (ValidObjectCount + AddedObjects.Num() > 0)
+	{
+		FMovieScenePossessable* const Possessable = MovieScene->FindPossessable(Guid);
+		if (Possessable && ObjectClass)
+		{
+			// If there are multiple objects within the same possessable, name possessable as "ClassName (Count)"
+			if (ValidObjectCount + AddedObjects.Num() > 1)
+			{
+				Possessable->SetName(FString::Printf(TEXT("%s (%d)")
+					, *ObjectClass->GetName()
+					, ValidObjectCount + AddedObjects.Num()));
+			}
+			else if (!AddedObjects.IsEmpty())
+			{
+				Possessable->SetName(FAvaSequencer::GetObjectName(AddedObjects[0]));
+			}
+			Possessable->SetPossessedObjectClass(ObjectClass);
+		}
+	}
+}
+
diff --git a/Engine/Plugins/Experimental/Avalanche/Source/AvalancheSequencer/Private/Commands/Stagger/AvaSequencerStagger.cpp b/Engine/Plugins/Experimental/Avalanche/Source/AvalancheSequencer/Private/Commands/Stagger/AvaSequencerStagger.cpp
new file mode 100644
--- /dev/null
+++ b/Engine/Plugins/Experimental/Avalanche/Source/AvalancheSequencer/Private/Commands/Stagger/AvaSequencerStagger.cpp
@@ -0,0 +1,202 @@
+// Copyright Epic Games, Inc. All Rights Reserved.
+
+#include "AvaSequencerStagger.h"
+#include "AvaSequencer.h"
+#include "Commands/AvaSequencerCommands.h"
+#include "Framework/Application/SlateApplication.h"
+#include "Framework/Commands/UICommandList.h"
+#include "ISequencer.h"
+#include "MVVM/Selection/Selection.h"
+#include "MVVM/ViewModels/LayerBarModel.h"
+#include "Misc/FrameTime.h"
+#include "SAvaSequencerStaggerSettings.h"
+#include "ScopedTransaction.h"
+
+#define LOCTEXT_NAMESPACE "AvaSequencerStagger"
+
+FAvaSequencerStagger::FStaggerElement::FStaggerElement(UE::Sequencer::TViewModelPtr<UE::Sequencer::FLayerBarModel> InLayerBarModel)
+	: LayerBarModel(MoveTemp(InLayerBarModel))
+{
+}
+
+void FAvaSequencerStagger::FStaggerElement::ComputeRange()
+{
+	ComputedRange = LayerBarModel->ComputeRange();
+}
+
+void FAvaSequencerStagger::MapAction(const TSharedRef<FUICommandList>& InCommandList)
+{
+	InCommandList->MapAction(FAvaSequencerCommands::Get().StaggerLayerBars
+		, FExecuteAction::CreateSP(this, &FAvaSequencerStagger::Execute)
+		, FCanExecuteAction::CreateSP(this, &FAvaSequencerStagger::CanExecute));
+}
+
+bool FAvaSequencerStagger::CanExecute() const
+{
+	return GatherStaggerElements().Num() > 1;
+}
+
+void FAvaSequencerStagger::Execute()
+{
+	using namespace UE::Sequencer;
+
+	TArray<FStaggerElement> StaggerElements = GatherStaggerElements();
+	if (StaggerElements.Num() <= 1)
+	{
+		return;
+	}
+
+	FAvaSequencerStaggerSettings Settings;
+	if (!GetSettings(Settings))
+	{
+		return;
+	}
+
+	TSharedRef<ISequencer> Sequencer = Owner.GetSequencer();
+
+	Settings.Shift = ConvertFrameTime(Settings.Shift
+		, Sequencer->GetFocusedDisplayRate()
+		, Sequencer->GetFocusedTickResolution()).FrameNumber;
+
+	Stagger(StaggerElements, Settings);
+}
+
+TArray<FAvaSequencerStagger::FStaggerElement> FAvaSequencerStagger::GatherStaggerElements() const
+{
+	using namespace UE::Sequencer;
+
+	TSharedPtr<FSequencerSelection> SequencerSelection = GetSequencerSelection();
+	if (!SequencerSelection.IsValid())
+	{
+		return TArray<FStaggerElement>();
+	}
+
+	// Gather all the Layer Bar Models Selected
+	TArray<TViewModelPtr<FLayerBarModel>> LayerBarModels;
+	LayerBarModels.Reserve(SequencerSelection->Outliner.Num());
+
+	for (FViewModelPtr ViewModel : SequencerSelection->Outliner)
+	{
+		if (!ViewModel.IsValid())
+		{
+			continue;
+		}
+
+		TSharedPtr<ITrackAreaExtension> TrackArea = ViewModel.ImplicitCast();
+		if (!TrackArea.IsValid())
+		{
+			continue;
+		}
+
+		for (const FViewModelPtr& TrackAreaModel : TrackArea->GetTopLevelChildTrackAreaModels())
+		{
+			if (TViewModelPtr<FLayerBarModel> LayerBarModel = TrackAreaModel.ImplicitCast())
+			{
+				LayerBarModels.Add(LayerBarModel);
+			}
+		}
+	}
+
+	// Remove all Layer Models that have a Descendant that is Selected
+	LayerBarModels.RemoveAll([SequencerSelection](const TViewModelPtr<FLayerBarModel>& InLayerBarModel)
+		{
+			TViewModelPtr<IOutlinerExtension> OutlinerExtension = InLayerBarModel->GetLinkedOutlinerItem();
+			if (!OutlinerExtension.IsValid())
+			{
+				return false;
+			}
+
+			constexpr bool bIncludeThis = false;
+			for (TViewModelPtr<IOutlinerExtension> ChildOutlinerExtension : OutlinerExtension.AsModel()->GetDescendantsOfType<IOutlinerExtension>(bIncludeThis))
+			{
+				if (SequencerSelection->Outliner.IsSelected(ChildOutlinerExtension))
+				{
+					return true;
+				}
+			}
+			return false;
+		});
+
+	TArray<FStaggerElement> StaggerElements;
+	StaggerElements.Append(MoveTemp(LayerBarModels));
+	return StaggerElements;
+}
+
+bool FAvaSequencerStagger::GetSettings(FAvaSequencerStaggerSettings& OutSettings)
+{
+	TSharedRef<SAvaSequencerStaggerSettings> SettingsWidget = SNew(SAvaSequencerStaggerSettings);
+
+	TSharedRef<SWindow> SettingsWindow =
+		SNew(SWindow)
+		.Title(LOCTEXT("StaggerSettings", "Stagger Settings"))
+		.SizingRule(ESizingRule::Autosized)
+		.SupportsMaximize(false)
+		.SupportsMinimize(false)
+		.Content()
+		[
+			SettingsWidget
+		];
+
+	TSharedPtr<SWidget> ParentWindow = FSlateApplication::Get().FindBestParentWindowForDialogs(Owner.GetSequencer()->GetSequencerWidget());
+	FSlateApplication::Get().AddModalWindow(SettingsWindow, ParentWindow);
+
+	if (SettingsWidget->GetReturnType() == EAppReturnType::Ok)
+	{
+		OutSettings = SettingsWidget->GetSettings();
+		return true;
+	}
+
+	return false;
+}
+
+void FAvaSequencerStagger::Stagger(TArrayView<FStaggerElement> InStaggerElements, const FAvaSequencerStaggerSettings& InSettings)
+{
+	if (InStaggerElements.Num() <= 1)
+	{
+		return;
+	}
+
+	// Compute Range for Each Element to Stagger
+	for (FStaggerElement& Element : InStaggerElements)
+	{
+		Element.ComputeRange();
+	}
+
+	FScopedTransaction Transaction(LOCTEXT("Stagger", "Stagger"));
+
+	// Find Earliest Stagger Point
+	FFrameNumber NextStaggerPoint = InStaggerElements[0].ComputedRange.GetLowerBoundValue();
+
+	switch (InSettings.StartPosition)
+	{
+	case FAvaSequencerStaggerSettings::EStartPosition::FirstSelected:
+		break;
+	case FAvaSequencerStaggerSettings::EStartPosition::FirstInTimeline:
+		for (const FStaggerElement& Element : InStaggerElements)
+		{
+			const FFrameNumber& ElementLowerBound = Element.ComputedRange.GetLowerBoundValue();
+			if (ElementLowerBound < NextStaggerPoint)
+			{
+				NextStaggerPoint = ElementLowerBound;
+			}
+		}
+		break;
+	}
+
+	for (const FStaggerElement& Element : InStaggerElements)
+	{
+		Element.LayerBarModel->Offset(NextStaggerPoint - Element.ComputedRange.GetLowerBoundValue());
+
+		switch (InSettings.OperationPoint)
+		{
+		case FAvaSequencerStaggerSettings::EOperationPoint::Start:
+			NextStaggerPoint += InSettings.Shift.FrameNumber;
+			break;
+		case FAvaSequencerStaggerSettings::EOperationPoint::End:
+			NextStaggerPoint += Element.ComputedRange.Size<FFrameNumber>() + InSettings.Shift.FrameNumber;
+			break;
+		}
+	}
+}
+
+#undef LOCTEXT_NAMESPACE
diff --git a/Engine/Plugins/Experimental/Avalanche/Source/AvalancheSequencer/Private/Commands/Stagger/AvaSequencerStagger.h b/Engine/Plugins/Experimental/Avalanche/Source/AvalancheSequencer/Private/Commands/Stagger/AvaSequencerStagger.h
new file mode 100644
--- /dev/null
+++ b/Engine/Plugins/Experimental/Avalanche/Source/AvalancheSequencer/Private/Commands/Stagger/AvaSequencerStagger.h
@@ -0,0 +1,53 @@
+// Copyright Epic Games, Inc. All Rights Reserved.
+
+#pragma once
+
+#include "AvaSequencerStaggerSettings.h"
+#include "Commands/AvaSequencerAction.h"
+#include "Containers/ArrayView.h"
+#include "Containers/ContainersFwd.h"
+#include "MVVM/ViewModelPtr.h"
+#include "Math/Range.h"
+#include "Misc/FrameNumber.h"
+#include "SequencerCoreFwd.h"
+
+namespace UE::Sequencer
+{
+	class FLayerBarModel;
+}
+
+class FAvaSequencerStagger : public FAvaSequencerAction
+{
+	struct FStaggerElement
+	{
+		explicit FStaggerElement(UE::Sequencer::TViewModelPtr<UE::Sequencer::FLayerBarModel> InLayerBarModel);
+
+		void ComputeRange();
+
+		UE::Sequencer::TViewModelPtr<UE::Sequencer::FLayerBarModel> LayerBarModel;
+
+		TRange<FFrameNumber> ComputedRange;
+	};
+
+public:
+	explicit FAvaSequencerStagger(FAvaSequencer& InOwner)
+		: FAvaSequencerAction(InOwner)
+	{
+	}
+
+protected:
+	//~ Begin FAvaSequencerAction
+	virtual void MapAction(const TSharedRef<FUICommandList>& InCommandList) override;
+	//~ End FAvaSequencerAction
+
+	bool CanExecute() const;
+
+	void Execute();
+
+private:
+	TArray<FStaggerElement> GatherStaggerElements() const;
+
+	bool GetSettings(FAvaSequencerStaggerSettings& OutSettings);
+
+	void Stagger(TArrayView<FStaggerElement> InStaggerElements, const FAvaSequencerStaggerSettings& InSettings);
+};
diff --git a/Engine/Plugins/Experimental/Avalanche/Source/AvalancheSequencer/Private/Commands/Stagger/AvaSequencerStaggerSettings.h b/Engine/Plugins/Experimental/Avalanche/Source/AvalancheSequencer/Private/Commands/Stagger/AvaSequencerStaggerSettings.h
new file mode 100644
--- /dev/null
+++ b/Engine/Plugins/Experimental/Avalanche/Source/AvalancheSequencer/Private/Commands/Stagger/AvaSequencerStaggerSettings.h
@@ -0,0 +1,24 @@
+// Copyright Epic Games, Inc. All Rights Reserved.
+
+#pragma once
+
+#include "Misc/FrameTime.h"
+
+struct FAvaSequencerStaggerSettings
+{
+	enum EOperationPoint : uint8
+	{
+		Start = 0,
+		End = 1
+	};
+	EOperationPoint OperationPoint = EOperationPoint::End;
+
+	enum EStartPosition : uint8
+	{
+		FirstSelected = 0,
+		FirstInTimeline = 1
+	};
+	EStartPosition StartPosition = EStartPosition::FirstSelected;
+
+	FFrameTime Shift;
+};
diff --git a/Engine/Plugins/Experimental/Avalanche/Source/AvalancheSequencer/Private/Commands/Stagger/SAvaSequencerStaggerSettings.cpp b/Engine/Plugins/Experimental/Avalanche/Source/AvalancheSequencer/Private/Commands/Stagger/SAvaSequencerStaggerSettings.cpp
new file mode 100644
--- /dev/null
+++ b/Engine/Plugins/Experimental/Avalanche/Source/AvalancheSequencer/Private/Commands/Stagger/SAvaSequencerStaggerSettings.cpp
@@ -0,0 +1,266 @@
+// Copyright Epic Games, Inc. All Rights Reserved.
+
+#include "SAvaSequencerStaggerSettings.h"
+#include "Framework/Application/SlateApplication.h"
+#include "Input/Reply.h"
+#include "Misc/Optional.h"
+#include "SPrimaryButton.h"
+#include "Widgets/Images/SImage.h"
+#include "Widgets/Input/SButton.h"
+#include "Widgets/Input/SNumericEntryBox.h"
+#include "Widgets/Layout/SBox.h"
+#include "Widgets/Layout/SUniformGridPanel.h"
+#include "Widgets/SBoxPanel.h"
+#include "Widgets/Text/STextBlock.h"
+
+#define LOCTEXT_NAMESPACE "SAvaSequencerStaggerSettings"
+
+FAvaSequencerStaggerSettings SAvaSequencerStaggerSettings::Settings;
+
+void SAvaSequencerStaggerSettings::Construct(const FArguments& InArgs)
+{
+	ChildSlot
+	[
+		SNew(SVerticalBox)
+		+ SVerticalBox::Slot()
+		.AutoHeight()
+		.Padding(5.f)
+		[
+			ConstructStartPositionRadioGroup()
+		]
+		+ SVerticalBox::Slot()
+		.AutoHeight()
+		.Padding(5.f, 0.f, 5.f, 5.f)
+		[
+			ConstructOperationPointRadioGroup()
+		]
+		+ SVerticalBox::Slot()
+		.Padding(10.f, 5.f, 10.f, 10.f)
+		[
+			SNew(SHorizontalBox)
+			+ SHorizontalBox::Slot()
+			.AutoWidth()
+			.VAlign(VAlign_Center)
+			[
+				SNew(SButton)
+				.ContentPadding(2.f)
+				.ButtonStyle(FAppStyle::Get(), TEXT("SimpleButton"))
+				.ToolTipText(LOCTEXT("ResetToDefaultsToolTip", "Reset All to Defaults"))
+				.OnClicked(this, &SAvaSequencerStaggerSettings::OnResetToDefaults)
+				[
+					SNew(SBox)
+					.WidthOverride(16.f)
+					.HeightOverride(16.f)
+					[
+						SNew(SImage)
+						.Image(FAppStyle::GetBrush(TEXT("PropertyWindow.DiffersFromDefault")))
+					]
+				]
+			]
+			+ SHorizontalBox::Slot()
+			.AutoWidth()
+			.VAlign(VAlign_Center)
+			.Padding(10.f, 0.f, 0.f, 0.f)
+			[
+				SNew(STextBlock)
+				.Text(LOCTEXT("ShiftFramesLabel", "Shift Frames:"))
+				.Justification(ETextJustify::Right)
+			]
+			+ SHorizontalBox::Slot()
+			.VAlign(VAlign_Center)
+			.Padding(10.f, 0.f, 0.f, 0.f)
+			[
+				SNew(SNumericEntryBox<int32>)
+				.AllowSpin(true)
+				.MinSliderValue(-60)
+				.MaxSliderValue(60)
+				.Value(this, &SAvaSequencerStaggerSettings::GetShiftFrame)
+				.OnValueChanged(this, &SAvaSequencerStaggerSettings::OnShiftFrameChanged)
+			]
+		]
+		+ SVerticalBox::Slot()
+		.AutoHeight()
+		.Padding(10.f, 5.f, 10.f, 10.f)
+		.HAlign(HAlign_Right)
+		[
+			SNew(SUniformGridPanel)
+			.SlotPadding(FMargin(5.f, 0.f))
+			+ SUniformGridPanel::Slot(0, 0)
+			[
+				SNew(SPrimaryButton)
+				.Text(LOCTEXT("OkButtonText", "OK"))
+				.OnClicked(this, &SAvaSequencerStaggerSettings::OnReturnButtonClicked, EAppReturnType::Ok)
+			]
+			+ SUniformGridPanel::Slot(1, 0)
+			[
+				SNew(SButton)
+				.HAlign(HAlign_Center)
+				.Text(LOCTEXT("CancelButtonText", "Cancel"))
+				.OnClicked(this, &SAvaSequencerStaggerSettings::OnReturnButtonClicked, EAppReturnType::Cancel)
+			]
+		]
+	];
+}
+
+TSharedRef<SCheckBox> SAvaSequencerStaggerSettings::ConstructRadioButton(const TAttribute<ECheckBoxState>& InCheckBoxState
+	, const FOnCheckStateChanged& InOnCheckStateChanged
+	, const TAttribute<FText>& InText
+	, const TAttribute<FText>& InToolTipText)
+{
+	return SNew(SCheckBox)
+		.Style(FAppStyle::Get(), TEXT("RadioButton"))
+		.ToolTipText(InToolTipText)
+		.IsChecked(InCheckBoxState)
+		.OnCheckStateChanged(InOnCheckStateChanged)
+		[
+			SNew(STextBlock)
+			.Text(InText)
+		];
+}
+
+TSharedRef<SWidget> SAvaSequencerStaggerSettings::ConstructStartPositionRadioGroup()
+{
+	auto CreateRadioButton = [this](const FAvaSequencerStaggerSettings::EStartPosition InStartPosition
+		, const TAttribute<FText>& InText
+		, const TAttribute<FText>& InToolTipText) -> TSharedRef<SCheckBox>
+	{
+		const TAttribute<ECheckBoxState> CheckBoxState = TAttribute<ECheckBoxState>::CreateLambda(
+			[this, InStartPosition]() -> ECheckBoxState
+			{
+				return (Settings.StartPosition == InStartPosition)
+					? ECheckBoxState::Checked : ECheckBoxState::Unchecked;
+			});
+		const FOnCheckStateChanged OnCheckStateChanged = FOnCheckStateChanged::CreateLambda(
+			[this, InStartPosition](const ECheckBoxState InState)
+			{
+				if (InState == ECheckBoxState::Checked)
+				{
+					Settings.StartPosition = InStartPosition;
+				}
+			});
+		return ConstructRadioButton(CheckBoxState, OnCheckStateChanged, InText, InToolTipText);
+	};
+
+	return SNew(SBorder)
+		.BorderImage(FAppStyle::GetBrush(TEXT("ToolPanel.GroupBorder")))
+		.Padding(3.f)
+		[
+			SNew(SHorizontalBox)
+			+ SHorizontalBox::Slot()
+			.FillWidth(0.35f)
+			.HAlign(HAlign_Right)
+			.VAlign(VAlign_Center)
+			[
+				SNew(STextBlock)
+				.Text(LOCTEXT("ShiftFromLabel", "First:"))
+			]
+			+ SHorizontalBox::Slot()
+			.FillWidth(0.65f)
+			.Padding(10.f, 0.f, 0.f, 0.f)
+			[
+				SNew(SVerticalBox)
+				+ SVerticalBox::Slot()
+				.AutoHeight()
+				.Padding(0.f)
+				[
+					CreateRadioButton(FAvaSequencerStaggerSettings::EStartPosition::FirstSelected
+						, LOCTEXT("FirstSelectedLabel", "Selected")
+						, LOCTEXT("FirstSelectedTooltip", "Uses the first selected sequence, regardless of time, to begin the operation"))
+				]
+				+ SVerticalBox::Slot()
+				.AutoHeight()
+				.Padding(0.f, 1.f, 0.f, 0.f)
+				[
+					CreateRadioButton(FAvaSequencerStaggerSettings::EStartPosition::FirstInTimeline
+						, LOCTEXT("FirstInTimelineLabel", "In Timeline")
+						, LOCTEXT("FirstInTimelineTooltip", "Uses the start of the first sequence in time, from all selected sequences, to begin the operation"))
+				]
+			]
+		];
+}
+
+TSharedRef<SWidget> SAvaSequencerStaggerSettings::ConstructOperationPointRadioGroup()
+{
+	auto CreateRadioButton = [this](const FAvaSequencerStaggerSettings::EOperationPoint InOperationPoint
+		, const TAttribute<FText>& InText
+		, const TAttribute<FText>& InToolTipText) -> TSharedRef<SCheckBox>
+	{
+		const TAttribute<ECheckBoxState> CheckBoxState = TAttribute<ECheckBoxState>::CreateLambda(
+			[this, InOperationPoint]() -> ECheckBoxState
+			{
+				return (Settings.OperationPoint == InOperationPoint)
+					? ECheckBoxState::Checked : ECheckBoxState::Unchecked;
+			});
+		const FOnCheckStateChanged OnCheckStateChanged = FOnCheckStateChanged::CreateLambda(
+			[this, InOperationPoint](const ECheckBoxState InState)
+			{
+				if (InState == ECheckBoxState::Checked)
+				{
+					Settings.OperationPoint = InOperationPoint;
+				}
+			});
+		return ConstructRadioButton(CheckBoxState, OnCheckStateChanged, InText, InToolTipText);
+	};
+
+	return SNew(SBorder)
+		.Padding(3.f)
+		.BorderImage(FAppStyle::GetBrush(TEXT("ToolPanel.GroupBorder")))
+		[
+			SNew(SHorizontalBox)
+			+ SHorizontalBox::Slot()
+			.FillWidth(0.35f)
+			.HAlign(HAlign_Right)
+			.VAlign(VAlign_Center)
+			[
+				SNew(STextBlock)
+				.Text(LOCTEXT("ShiftFromLabel", "From:"))
+			]
+			+ SHorizontalBox::Slot()
+			.FillWidth(0.65f)
+			.Padding(10.f, 0.f, 0.f, 0.f)
+			[
+				SNew(SVerticalBox)
+				+ SVerticalBox::Slot()
+				.AutoHeight()
+				.Padding(0.f)
+				[
+					CreateRadioButton(FAvaSequencerStaggerSettings::EOperationPoint::Start
+						, LOCTEXT("ShiftFromStartLabel", "Layer Start")
+						, LOCTEXT("ShiftFromStartToolTip", "Staggers each sequence from the start of the previous sequence in the selection order"))
+				]
+				+ SVerticalBox::Slot()
+				.AutoHeight()
+				.Padding(0.f, 1.f, 0.f, 0.f)
+				[
+					CreateRadioButton(FAvaSequencerStaggerSettings::EOperationPoint::End
+						, LOCTEXT("ShiftFromEndLabel", "Layer End")
+						, LOCTEXT("ShiftFromEndLabelToolTip", "Staggers each sequence from the end of the previous sequence in the selection order"))
+				]
+			]
+		];
+}
+
+TOptional<int32> SAvaSequencerStaggerSettings::GetShiftFrame() const
+{
+	return Settings.Shift.FrameNumber.Value;
+}
+
+void SAvaSequencerStaggerSettings::OnShiftFrameChanged(int32 InNewValue)
+{
+	Settings.Shift.FrameNumber.Value = InNewValue;
+}
+
+FReply SAvaSequencerStaggerSettings::OnReturnButtonClicked(EAppReturnType::Type InReturnType)
+{
+	ReturnType = InReturnType;
+	FSlateApplication::Get().FindWidgetWindow(AsShared())->RequestDestroyWindow();
+	return FReply::Handled();
+}
+
+FReply SAvaSequencerStaggerSettings::OnResetToDefaults()
+{
+	Settings = FAvaSequencerStaggerSettings();
+	return FReply::Handled();
+}
+
+#undef LOCTEXT_NAMESPACE
diff --git a/Engine/Plugins/Experimental/Avalanche/Source/AvalancheSequencer/Private/Commands/Stagger/SAvaSequencerStaggerSettings.h b/Engine/Plugins/Experimental/Avalanche/Source/AvalancheSequencer/Private/Commands/Stagger/SAvaSequencerStaggerSettings.h
new file mode 100644
--- /dev/null
+++ b/Engine/Plugins/Experimental/Avalanche/Source/AvalancheSequencer/Private/Commands/Stagger/SAvaSequencerStaggerSettings.h
@@ -0,0 +1,46 @@
+// Copyright Epic Games, Inc. All Rights Reserved.
+
+#pragma once
+
+#include "AvaSequencerStaggerSettings.h"
+#include "GenericPlatform/GenericPlatformMisc.h"
+#include "Misc/Attribute.h"
+#include "Misc/OptionalFwd.h"
+#include "Widgets/Input/SCheckBox.h"
+#include "Widgets/SCompoundWidget.h"
+
+class FReply;
+class SWidget;
+
+class SAvaSequencerStaggerSettings : public SCompoundWidget
+{
+public:
+	SLATE_BEGIN_ARGS(SAvaSequencerStaggerSettings) {}
+	SLATE_END_ARGS()
+
+	void Construct(const FArguments& InArgs);
+
+	const FAvaSequencerStaggerSettings& GetSettings() const { return Settings; }
+
+	EAppReturnType::Type GetReturnType() const { return ReturnType; }
+
+	TOptional<int32> GetShiftFrame() const;
+
+	void OnShiftFrameChanged(int32 InNewValue);
+
+	FReply OnReturnButtonClicked(EAppReturnType::Type InReturnType);
+
+private:
+	TSharedRef<SCheckBox> ConstructRadioButton(const TAttribute<ECheckBoxState>& InCheckBoxState
+		, const FOnCheckStateChanged& InOnCheckStateChanged
+		, const TAttribute<FText>& InText
+		, const TAttribute<FText>& InToolTipText);
+	TSharedRef<SWidget> ConstructStartPositionRadioGroup();
+	TSharedRef<SWidget> ConstructOperationPointRadioGroup();
+
+	FReply OnResetToDefaults();
+
+	static FAvaSequencerStaggerSettings Settings;
+
+	EAppReturnType::Type ReturnType = EAppReturnType::Cancel;
+};
diff --git a/Engine/Plugins/Experimental/Avalanche/Source/AvalancheSequencer/Private/Customization/AvaSequenceTimeCustomization.cpp b/Engine/Plugins/Experimental/Avalanche/Source/AvalancheSequencer/Private/Customization/AvaSequenceTimeCustomization.cpp
new file mode 100644
--- /dev/null
+++ b/Engine/Plugins/Experimental/Avalanche/Source/AvalancheSequencer/Private/Customization/AvaSequenceTimeCustomization.cpp
@@ -0,0 +1,74 @@
+// Copyright Epic Games, Inc. All Rights Reserved.
+
+#include "AvaSequenceTimeCustomization.h"
+#include "AvaSequenceShared.h"
+#include "DetailWidgetRow.h"
+#include "Widgets/Layout/SBox.h"
+#include "Widgets/SBoxPanel.h"
+
+void FAvaSequenceTimeCustomization::CustomizeHeader(TSharedRef<IPropertyHandle> InPropertyHandle, FDetailWidgetRow& InHeaderRow, IPropertyTypeCustomizationUtils& InCustomizationUtils)
+{
+	TSharedPtr<IPropertyHandle> PositionTypeHandle = InPropertyHandle->GetChildHandle(GET_MEMBER_NAME_CHECKED(FAvaSequenceTime, TimeType));
+
+	TSharedRef<SHorizontalBox> PropertyValueBox = SNew(SHorizontalBox);
+
+	auto AddPropertySlot = [PropertyValueBox, InPropertyHandle](FName InPropertyName)
+	{
+		TSharedPtr<IPropertyHandle> ChildPropertyHandle = InPropertyHandle->GetChildHandle(InPropertyName);
+
+		TSharedRef<SWidget> PropertyValueWidget = ChildPropertyHandle->CreatePropertyValueWidget(/*bDisplayDefaultPropertyButtons*/false);
+
+		TAttribute<EVisibility>::FGetter VisibilityGetter = TAttribute<EVisibility>::FGetter::CreateSPLambda(ChildPropertyHandle.Get(), [ChildPropertyHandle]
+			{
+				return ChildPropertyHandle->IsEditable()
+					? EVisibility::SelfHitTestInvisible
+					: EVisibility::Collapsed;
+			});
+
+		PropertyValueWidget->SetVisibility(TAttribute<EVisibility>::Create(VisibilityGetter));
+
+		PropertyValueBox->AddSlot()
+			.VAlign(VAlign_Center)
+			.HAlign(HAlign_Fill)
+			[
+				PropertyValueWidget
+			];
+	};
+
+	AddPropertySlot(GET_MEMBER_NAME_CHECKED(FAvaSequenceTime, Frame));
+	AddPropertySlot(GET_MEMBER_NAME_CHECKED(FAvaSequenceTime, Seconds));
+	AddPropertySlot(GET_MEMBER_NAME_CHECKED(FAvaSequenceTime, MarkLabel));
+
+	PropertyValueBox->AddSlot()
+		.AutoWidth()
+		[
+			InPropertyHandle->CreateDefaultPropertyButtonWidgets()
+		];
+
+	InHeaderRow
+		.NameContent()
+		[
+			SNew(SBox)
+			.MinDesiredWidth(300.f)
+			.MaxDesiredWidth(300.f)
+			[
+				SNew(SHorizontalBox)
+				+ SHorizontalBox::Slot()
+				.AutoWidth()
+				[
+					InPropertyHandle->CreatePropertyNameWidget()
+				]
+				+ SHorizontalBox::Slot()
+				.FillWidth(1.f)
+				.Padding(0.f, 0.f, 5.f, 0.f)
+				.HAlign(HAlign_Right)
+				[
+					PositionTypeHandle->CreatePropertyValueWidget(/*bDisplayDefaultPropertyButtons*/false)
+				]
+			]
+		]
+		.ValueContent()
+		[
+			PropertyValueBox
+		];
+}
diff --git a/Engine/Plugins/Experimental/Avalanche/Source/AvalancheSequencer/Private/DetailsView/SAvaKeyFrameEdit.cpp b/Engine/Plugins/Experimental/Avalanche/Source/AvalancheSequencer/Private/DetailsView/SAvaKeyFrameEdit.cpp
new file mode 100644
--- /dev/null
+++ b/Engine/Plugins/Experimental/Avalanche/Source/AvalancheSequencer/Private/DetailsView/SAvaKeyFrameEdit.cpp
@@ -0,0 +1,108 @@
+// Copyright Epic Games, Inc. All Rights Reserved.
+
+#include "SAvaKeyFrameEdit.h"
+#include "AvaSequencer.h"
+#include "EaseCurveTool/AvaEaseCurveTool.h"
+#include "EaseCurveTool/AvaEaseCurveToolSettings.h"
+#include "MVVM/Selection/Selection.h"
+#include "MVVM/ViewModels/SequencerEditorViewModel.h"
+#include "Widgets/Layout/SScrollBox.h"
+
+using namespace UE::Sequencer;
+
+#define LOCTEXT_NAMESPACE "SAvaKeyFrameEdit"
+
+namespace UE::AvaSequencer::Private
+{
+	/** Utility widget to auto update horizontal alignment. */
+	class SDynamicHAlign : public SCompoundWidget
+	{
+	public:
+		SLATE_BEGIN_ARGS(SDynamicHAlign)
+			: _HAlign(HAlign_Fill)
+		{}
+			SLATE_ATTRIBUTE(EHorizontalAlignment, HAlign)
+			SLATE_ARGUMENT(TSharedPtr<SWidget>, Content)
+		SLATE_END_ARGS()
+
+		void Construct(const FArguments& InArgs)
+		{
+			SetCanTick(true);
+
+			HAlign = InArgs._HAlign;
+
+			if (InArgs._Content.IsValid())
+			{
+				ChildSlot
+					.HAlign(HAlign.Get())
+					[
+						InArgs._Content.ToSharedRef()
+					];
+			}
+		}
+
+	protected:
+		//~ Begin SWidget
+		virtual void Tick(const FGeometry& InAllottedGeometry, const double InCurrentTime, const float InDeltaTime) override
+		{
+			SCompoundWidget::Tick(InAllottedGeometry, InCurrentTime, InDeltaTime);
+
+			if (ChildSlot.GetHorizontalAlignment() != HAlign.Get())
+			{
+				ChildSlot.SetHorizontalAlignment(HAlign.Get());
+			}
+		}
+		//~ End SWidget
+
+		TAttribute<EHorizontalAlignment> HAlign;
+	};
+}
+
+void SAvaKeyFrameEdit::Construct(const FArguments& InArgs, TSharedRef<FAvaSequencer> InSequencer)
+{
+	AvaSequencerWeak = InSequencer;
+
+	KeyEditData = InArgs._KeyEditData;
+
+	const TSharedRef<ISequencer> Sequencer = InSequencer->GetSequencer();
+
+	if (const TSharedPtr<FSequencerEditorViewModel> SequencerViewModel = Sequencer->GetViewModel())
+	{
+		SequencerSelectionWeak = SequencerViewModel->GetSelection();
+	}
+
+	CurveEaseTool = MakeShared<FAvaEaseCurveTool>(InSequencer);
+
+	if (!AvaSequencerWeak.IsValid())
+	{
+		return;
+	}
+
+	const TSharedPtr<FAvaSequencer> AvaSequencer = AvaSequencerWeak.Pin();
+
+	ChildSlot
+		[
+			SNew(SScrollBox)
+			+ SScrollBox::Slot()
+			.AutoSize()
+			.HAlign(HAlign_Fill)
+			.Padding(1.f, 2.f, 2.f, 3.f)
+			[
+				/** Horizontal alignment is polled every frame from settings.
+				Fill will stretch the ease curve graph editor horizontally in the details panel.
+				Left/Right/Center will maintain 1:1 ratio. */
+				SNew(UE::AvaSequencer::Private::SDynamicHAlign)
+				.HAlign_UObject(GetDefault<UAvaEaseCurveToolSettings>(), &UAvaEaseCurveToolSettings::GetGraphHAlign)
+				.Content(CurveEaseTool->GenerateWidget())
+			]
+			+ SScrollBox::Slot()
+			.AutoSize()
+			.HAlign(HAlign_Fill)
+			[
+				SNew(SKeyEditInterface, AvaSequencer->GetSequencer())
+				.EditData(KeyEditData)
+			]
+		];
+}
+
+#undef LOCTEXT_NAMESPACE
diff --git a/Engine/Plugins/Experimental/Avalanche/Source/AvalancheSequencer/Private/DetailsView/SAvaSequenceDetails.cpp b/Engine/Plugins/Experimental/Avalanche/Source/AvalancheSequencer/Private/DetailsView/SAvaSequenceDetails.cpp
new file mode 100644
--- /dev/null
+++ b/Engine/Plugins/Experimental/Avalanche/Source/AvalancheSequencer/Private/DetailsView/SAvaSequenceDetails.cpp
@@ -0,0 +1,181 @@
+// Copyright Epic Games, Inc. All Rights Reserved.
+
+#include "SAvaSequenceDetails.h"
+#include "Framework/Application/SlateApplication.h"
+#include "Section/IAvaSequenceSectionDetails.h"
+#include "Widgets/Input/SCheckBox.h"
+#include "Widgets/Layout/SBox.h"
+#include "Widgets/Layout/SScrollBox.h"
+#include "Widgets/SBoxPanel.h"
+#include "Widgets/Text/STextBlock.h"
+
+#define LOCTEXT_NAMESPACE "SAvaSequenceDetails"
+
+void SAvaSequenceDetails::Construct(const FArguments& InArgs
+	, const TSharedPtr<FAvaSequencer>& InAvaSequencer
+	, TArray<TSharedRef<IAvaSequenceSectionDetails>>&& InSections)
+{
+	AvaSequencerWeak = InAvaSequencer;
+	Sections = MoveTemp(InSections);
+
+	OnSelectedSectionsChanged = InArgs._OnSelectedSectionsChanged;
+	SelectedSections = InArgs._InitiallySelectedSections;
+
+	SectionBox = SNew(SHorizontalBox);
+
+	ContentBox = SNew(SScrollBox)
+		.Orientation(Orient_Vertical);
+
+	ChildSlot
+	[
+		SNew(SVerticalBox)
+		+ SVerticalBox::Slot()
+		.AutoHeight()
+		[
+			SectionBox.ToSharedRef()
+		]
+		+ SVerticalBox::Slot()
+		.FillHeight(1.f)
+		[
+			ContentBox.ToSharedRef()
+		]
+	];
+
+	BuildSections();
+}
+
+void SAvaSequenceDetails::BuildSections()
+{
+	SectionBox->ClearChildren();
+	ContentBox->ClearChildren();
+
+	if (!AvaSequencerWeak.IsValid() || Sections.IsEmpty())
+	{
+		return;
+	}
+
+	TSharedRef<FAvaSequencer> AvaSequencer = AvaSequencerWeak.Pin().ToSharedRef();
+
+	for (const TSharedRef<IAvaSequenceSectionDetails>& Section : Sections)
+	{
+		const FName SectionName = Section->GetSectionName();
+
+		SectionBox->AddSlot()
+		.Padding(2.f, 2.f)
+		[
+			SNew(SBox)
+			.Padding(FMargin(0))
+			.Visibility(this, &SAvaSequenceDetails::GetSectionButtonVisibility, Section->AsWeak())
+			[
+				SNew(SCheckBox)
+				.Style(FAppStyle::Get(), "DetailsView.SectionButton")
+				.OnCheckStateChanged(this, &SAvaSequenceDetails::OnSectionSelected, SectionName)
+				.IsChecked(this, &SAvaSequenceDetails::GetSectionCheckBoxState, SectionName)
+				[
+					SNew(STextBlock)
+					.TextStyle(FAppStyle::Get(), "SmallText")
+					.Text(Section->GetSectionDisplayName())
+					.Justification(ETextJustify::Center)
+				]
+			]
+		];
+
+		ContentBox->AddSlot()
+		.AutoSize()
+		[
+			SNew(SBox)
+			.Visibility(this, &SAvaSequenceDetails::GetSectionContentVisibility, SectionName, Section->AsWeak())
+			[
+				Section->CreateContentWidget(AvaSequencer)
+			]
+		];
+	}
+
+	if (SelectedSections.Num() == 0)
+	{
+		// Find first section that is visible
+		TSharedRef<IAvaSequenceSectionDetails>* const FoundSection = Sections.FindByPredicate(
+			[](const TSharedRef<IAvaSequenceSectionDetails>& InSection)
+			{
+				return InSection->ShouldShowSection();
+			});
+
+		checkf(FoundSection, TEXT("There should be at least one section visible to have initially selected!"));
+
+		SelectedSections.Add((*FoundSection)->GetSectionName());
+	}
+}
+
+void SAvaSequenceDetails::OnSectionSelected(ECheckBoxState InCheckBoxState, FName InSectionName)
+{
+	const bool bIsControlDown = FSlateApplication::Get().GetModifierKeys().IsControlDown();
+
+	switch (InCheckBoxState)
+	{
+	case ECheckBoxState::Unchecked:
+		if (bIsControlDown)
+		{
+			SelectedSections.Remove(InSectionName);
+
+			// Force to always have a Selected Section, making it unable to deselect the last one
+			if (SelectedSections.IsEmpty())
+			{
+				SelectedSections.Add(InSectionName);
+			}
+		}
+		else
+		{
+			SelectedSections.Reset();
+			SelectedSections.Add(InSectionName);
+		}
+		break;
+
+	case ECheckBoxState::Checked:
+		if (bIsControlDown)
+		{
+			SelectedSections.Add(InSectionName);
+		}
+		else
+		{
+			SelectedSections.Reset();
+			SelectedSections.Add(InSectionName);
+		}
+		break;
+	}
+
+	OnSelectedSectionsChanged.ExecuteIfBound(SelectedSections);
+}
+
+bool SAvaSequenceDetails::IsSectionSelected(FName InSectionName) const
+{
+	return SelectedSections.Contains(InSectionName);
+}
+
+bool SAvaSequenceDetails::ShouldShowSection(TWeakPtr<IAvaSequenceSectionDetails> InSectionWeak) const
+{
+	TSharedPtr<IAvaSequenceSectionDetails> Section = InSectionWeak.Pin();
+	return Section.IsValid() && Section->ShouldShowSection();
+}
+
+EVisibility SAvaSequenceDetails::GetSectionButtonVisibility(TWeakPtr<IAvaSequenceSectionDetails> InSectionWeak) const
+{
+	return ShouldShowSection(InSectionWeak)
+		? EVisibility::SelfHitTestInvisible
+		: EVisibility::Collapsed;
+}
+
+EVisibility SAvaSequenceDetails::GetSectionContentVisibility(FName InSectionName, TWeakPtr<IAvaSequenceSectionDetails> InSectionWeak) const
+{
+	return IsSectionSelected(InSectionName) && ShouldShowSection(InSectionWeak)
+		? EVisibility::SelfHitTestInvisible
+		: EVisibility::Collapsed;
+}
+
+ECheckBoxState SAvaSequenceDetails::GetSectionCheckBoxState(FName InSectionName) const
+{
+	return IsSectionSelected(InSectionName)
+		? ECheckBoxState::Checked
+		: ECheckBoxState::Unchecked;
+}
+
+#undef LOCTEXT_NAMESPACE 
diff --git a/Engine/Plugins/Experimental/Avalanche/Source/AvalancheSequencer/Private/DetailsView/SAvaSequenceDetails.h b/Engine/Plugins/Experimental/Avalanche/Source/AvalancheSequencer/Private/DetailsView/SAvaSequenceDetails.h
new file mode 100644
--- /dev/null
+++ b/Engine/Plugins/Experimental/Avalanche/Source/AvalancheSequencer/Private/DetailsView/SAvaSequenceDetails.h
@@ -0,0 +1,58 @@
+// Copyright Epic Games, Inc. All Rights Reserved.
+
+#pragma once
+
+#include "Containers/Array.h"
+#include "Containers/Set.h"
+#include "Styling/SlateTypes.h"
+#include "Templates/SharedPointer.h"
+#include "UObject/NameTypes.h"
+#include "UObject/WeakObjectPtr.h"
+#include "Widgets/SCompoundWidget.h"
+
+class FAvaSequencer;
+class IAvaSequenceSectionDetails;
+class SHorizontalBox;
+class SScrollBox;
+
+DECLARE_DELEGATE_OneParam(FAvaSelectedSectionsChanged, const TSet<FName>&)
+
+class SAvaSequenceDetails : public SCompoundWidget
+{
+public:
+	SLATE_BEGIN_ARGS(SAvaSequenceDetails){}
+		SLATE_ARGUMENT(TSet<FName>, InitiallySelectedSections)
+		SLATE_EVENT(FAvaSelectedSectionsChanged, OnSelectedSectionsChanged)
+	SLATE_END_ARGS()
+
+	void Construct(const FArguments& InArgs
+		, const TSharedPtr<FAvaSequencer>& InAvaSequencer
+		, TArray<TSharedRef<IAvaSequenceSectionDetails>>&& InSections);
+
+private:
+	void BuildSections();
+
+	void OnSectionSelected(ECheckBoxState InCheckBoxState, FName InSectionName);
+
+	bool IsSectionSelected(FName InSectionName) const;
+
+	bool ShouldShowSection(TWeakPtr<IAvaSequenceSectionDetails> InSectionWeak) const;
+
+	EVisibility GetSectionButtonVisibility(TWeakPtr<IAvaSequenceSectionDetails> InSectionWeak) const;
+
+	EVisibility GetSectionContentVisibility(FName InSectionName, TWeakPtr<IAvaSequenceSectionDetails> InSectionWeak) const;
+
+	ECheckBoxState GetSectionCheckBoxState(FName InSectionName) const;
+
+	TWeakPtr<FAvaSequencer> AvaSequencerWeak;
+
+	TArray<TSharedRef<IAvaSequenceSectionDetails>> Sections;
+
+	TSet<FName> SelectedSections;
+
+	FAvaSelectedSectionsChanged OnSelectedSectionsChanged;
+
+	TSharedPtr<SHorizontalBox> SectionBox;
+
+	TSharedPtr<SScrollBox> ContentBox;
+};
diff --git a/Engine/Plugins/Experimental/Avalanche/Source/AvalancheSequencer/Private/DetailsView/Section/AvaSequencePlaybackDetails.cpp b/Engine/Plugins/Experimental/Avalanche/Source/AvalancheSequencer/Private/DetailsView/Section/AvaSequencePlaybackDetails.cpp
new file mode 100644
--- /dev/null
+++ b/Engine/Plugins/Experimental/Avalanche/Source/AvalancheSequencer/Private/DetailsView/Section/AvaSequencePlaybackDetails.cpp
@@ -0,0 +1,51 @@
+// Copyright Epic Games, Inc. All Rights Reserved.
+
+#include "AvaSequencePlaybackDetails.h"
+#include "AvaSequencePlaybackObject.h"
+#include "AvaSequence.h"
+#include "AvaSequencer.h"
+#include "CustomDetailsViewArgs.h"
+#include "CustomDetailsViewModule.h"
+
+#define LOCTEXT_NAMESPACE "AvaSequencePlaybackDetails"
+
+FName FAvaSequencePlaybackDetails::GetSectionName() const
+{
+	return TEXT("Playback");
+}
+
+FText FAvaSequencePlaybackDetails::GetSectionDisplayName() const
+{
+	return LOCTEXT("PlaybackLabel", "Playback");
+}
+
+TSharedRef<SWidget> FAvaSequencePlaybackDetails::CreateContentWidget(const TSharedRef<FAvaSequencer>& InAvaSequencer)
+{
+	AvaSequencerWeak = InAvaSequencer;
+
+	FCustomDetailsViewArgs CustomDetailsViewArgs;
+	CustomDetailsViewArgs.IndentAmount           = 0.f;
+	CustomDetailsViewArgs.ValueColumnWidth       = 0.5f;
+	CustomDetailsViewArgs.bShowCategories        = true;
+	CustomDetailsViewArgs.bAllowGlobalExtensions = true;
+	CustomDetailsViewArgs.CategoryAllowList.Allow(TEXT("Scheduled Playback"));
+	CustomDetailsViewArgs.ExpansionState.Add(FCustomDetailsViewItemId::MakeCategoryId("Scheduled Playback"), true);
+
+	TSharedRef<ICustomDetailsView> PlaybackDetailsView = ICustomDetailsViewModule::Get().CreateCustomDetailsView(CustomDetailsViewArgs);
+
+	IAvaSequencePlaybackObject* PlaybackObject = InAvaSequencer->GetProvider().GetPlaybackObject();
+	if (ensureAlways(PlaybackObject))
+	{
+		PlaybackDetailsView->SetObject(PlaybackObject->ToUObject());
+	}
+
+	return PlaybackDetailsView;
+}
+
+bool FAvaSequencePlaybackDetails::ShouldShowSection() const
+{
+	TSharedPtr<FAvaSequencer> AvaSequencer = AvaSequencerWeak.Pin();
+	return AvaSequencer.IsValid() && IsValid(AvaSequencer->GetViewedSequence());
+}
+
+#undef LOCTEXT_NAMESPACE
diff --git a/Engine/Plugins/Experimental/Avalanche/Source/AvalancheSequencer/Private/DetailsView/Section/AvaSequencePlaybackDetails.h b/Engine/Plugins/Experimental/Avalanche/Source/AvalancheSequencer/Private/DetailsView/Section/AvaSequencePlaybackDetails.h
new file mode 100644
--- /dev/null
+++ b/Engine/Plugins/Experimental/Avalanche/Source/AvalancheSequencer/Private/DetailsView/Section/AvaSequencePlaybackDetails.h
@@ -0,0 +1,17 @@
+// Copyright Epic Games, Inc. All Rights Reserved.
+
+#pragma once
+
+#include "IAvaSequenceSectionDetails.h"
+
+class FAvaSequencePlaybackDetails : public IAvaSequenceSectionDetails
+{
+	//~ Begin IAvaSequenceSectionDetails
+	virtual FName GetSectionName() const override;
+	virtual FText GetSectionDisplayName() const override;
+	virtual TSharedRef<SWidget> CreateContentWidget(const TSharedRef<FAvaSequencer>& InAvaSequencer) override;
+	virtual bool ShouldShowSection() const;
+	//~ End IAvaSequenceSectionDetails
+
+	TWeakPtr<FAvaSequencer> AvaSequencerWeak;
+};
diff --git a/Engine/Plugins/Experimental/Avalanche/Source/AvalancheSequencer/Private/DetailsView/Section/AvaSequenceSelectionDetails.cpp b/Engine/Plugins/Experimental/Avalanche/Source/AvalancheSequencer/Private/DetailsView/Section/AvaSequenceSelectionDetails.cpp
new file mode 100644
--- /dev/null
+++ b/Engine/Plugins/Experimental/Avalanche/Source/AvalancheSequencer/Private/DetailsView/Section/AvaSequenceSelectionDetails.cpp
@@ -0,0 +1,225 @@
+// Copyright Epic Games, Inc. All Rights Reserved.
+
+#include "AvaSequenceSelectionDetails.h"
+#include "AvaSequencer.h"
+#include "CustomDetailsViewArgs.h"
+#include "CustomDetailsViewModule.h"
+#include "DetailsView/SAvaKeyFrameEdit.h"
+#include "IKeyArea.h"
+#include "MVVM/Selection/Selection.h"
+#include "MVVM/ViewModels/SequencerEditorViewModel.h"
+#include "SKeyEditInterface.h"
+#include "Styling/AppStyle.h"
+#include "Widgets/Layout/SBorder.h"
+
+#define LOCTEXT_NAMESPACE "AvaSequenceSelectionDetails"
+
+namespace UE::AvaSequencer::Private
+{
+	FKeyEditData GetKeyEditData(const UE::Sequencer::FKeySelection& InKeySelection)
+	{
+		if (InKeySelection.Num() == 1)
+		{
+			for (const FKeyHandle Key : InKeySelection)
+			{
+				if (const TSharedPtr<UE::Sequencer::FChannelModel> Channel = InKeySelection.GetModelForKey(Key))
+				{
+					FKeyEditData KeyEditData;
+					KeyEditData.KeyStruct     = Channel->GetKeyArea()->GetKeyStruct(Key);
+					KeyEditData.OwningSection = Channel->GetSection();
+					return KeyEditData;
+				}
+			}
+		}
+		else
+		{
+			TArray<FKeyHandle> KeyHandles;
+			UMovieSceneSection* CommonSection = nullptr;
+			for (FKeyHandle Key : InKeySelection)
+			{
+				TSharedPtr<UE::Sequencer::FChannelModel> Channel = InKeySelection.GetModelForKey(Key);
+				if (Channel.IsValid())
+				{
+					KeyHandles.Add(Key);
+					if (!CommonSection)
+					{
+						CommonSection = Channel->GetSection();
+					}
+					else if (CommonSection != Channel->GetSection())
+					{
+						CommonSection = nullptr;
+						break;
+					}
+				}
+			}
+
+			if (CommonSection)
+			{
+				FKeyEditData KeyEditData;
+				KeyEditData.KeyStruct     = CommonSection->GetKeyStruct(KeyHandles);
+				KeyEditData.OwningSection = CommonSection;
+				return KeyEditData;
+			}
+		}
+
+		return FKeyEditData();
+	}
+
+	TSharedPtr<UE::Sequencer::FSequencerSelection> GetSelection(const ISequencer& InSequencer)
+	{
+		const TSharedPtr<UE::Sequencer::FSequencerEditorViewModel> ViewModel = InSequencer.GetViewModel();
+		if (!ViewModel.IsValid())
+		{
+			return nullptr;
+		}
+
+		return ViewModel->GetSelection();
+	}
+}
+
+FName FAvaSequenceSelectionDetails::GetSectionName() const
+{
+	return TEXT("Selection");
+}
+
+FText FAvaSequenceSelectionDetails::GetSectionDisplayName() const
+{
+	return LOCTEXT("SelectionLabel", "Selection");
+}
+
+TSharedRef<SWidget> FAvaSequenceSelectionDetails::CreateContentWidget(const TSharedRef<FAvaSequencer>& InAvaSequencer)
+{
+	AvaSequencerWeak = InAvaSequencer;
+
+	ContentBorder = SNew(SBorder)
+		.BorderImage(FAppStyle::GetBrush("NoBorder"));
+
+	const ISequencer& Sequencer = InAvaSequencer->GetSequencer().Get();
+
+	if (const TSharedPtr<UE::Sequencer::FSequencerSelection> SequencerSelection = UE::AvaSequencer::Private::GetSelection(Sequencer))
+	{
+		SequencerSelection->KeySelection.OnChanged.AddSP(this, &FAvaSequenceSelectionDetails::OnSequencerSelectionChanged);
+		SequencerSelection->TrackArea.OnChanged.AddSP(this, &FAvaSequenceSelectionDetails::OnSequencerSelectionChanged);
+		SequencerSelection->Outliner.OnChanged.AddSP(this, &FAvaSequenceSelectionDetails::OnSequencerSelectionChanged);
+
+		OnSequencerSelectionChanged();
+	}
+
+	return ContentBorder.ToSharedRef();
+}
+
+void FAvaSequenceSelectionDetails::OnSequencerSelectionChanged()
+{
+	const TSharedPtr<FAvaSequencer> AvaSequencer = AvaSequencerWeak.Pin();
+	if (!AvaSequencer.IsValid())
+	{
+		ContentBorder->ClearContent();
+		return;
+	}
+
+	const TSharedRef<ISequencer> Sequencer = AvaSequencer->GetSequencer();
+
+	const TSharedPtr<UE::Sequencer::FSequencerSelection> SequencerSelection = UE::AvaSequencer::Private::GetSelection(*Sequencer);
+	if (!SequencerSelection.IsValid())
+	{
+		ContentBorder->ClearContent();
+		return;
+	}
+
+	/**
+	 * @TODO: This does not check for other selections being made.
+	 * Ex. Selecting both a key and a section will always show the key details.
+	 * 1) Should we disallow details for a key selection if there is also a section selection?
+	 * 2) Should we show all the details for the all the selections in the details view?
+	 */
+
+	if (SequencerSelection->KeySelection.Num() > 0)
+	{
+		const FKeyEditData KeyEditData = UE::AvaSequencer::Private::GetKeyEditData(SequencerSelection->KeySelection);
+		if (KeyEditData.KeyStruct.IsValid())
+		{
+			const TSharedRef<SAvaKeyFrameEdit> KeyFrameEditor = SNew(SAvaKeyFrameEdit, AvaSequencer.ToSharedRef())
+				.KeyEditData(this, &FAvaSequenceSelectionDetails::GetKeyEditData);
+			ContentBorder->SetContent(KeyFrameEditor);
+		}
+		else
+		{
+			ContentBorder->SetContent(CreateHintText(LOCTEXT("InvalidKeyCombination", "Selected keys must belong to the same section.")));
+		}
+
+		return;
+	}
+
+	TArray<UObject*> ObjectsToView;
+
+	TArray<UMovieSceneSection*> MovieSceneSectionObjects;
+	Sequencer->GetSelectedSections(MovieSceneSectionObjects);
+
+	if (MovieSceneSectionObjects.Num() > 0)
+	{
+		ObjectsToView.Append(MovieSceneSectionObjects);
+	}
+	else
+	{
+		TArray<UMovieSceneTrack*> MovieSceneTrackObjects;
+		Sequencer->GetSelectedTracks(MovieSceneTrackObjects);
+
+		if (MovieSceneTrackObjects.Num() > 0)
+		{
+			ObjectsToView.Append(MovieSceneTrackObjects);
+		}
+	}
+
+	if (ObjectsToView.Num() > 0)
+	{
+		FCustomDetailsViewArgs CustomDetailsViewArgs;
+		CustomDetailsViewArgs.IndentAmount = 0.f;
+		CustomDetailsViewArgs.bShowCategories = true;
+		CustomDetailsViewArgs.bAllowGlobalExtensions = true;
+		CustomDetailsViewArgs.bDefaultItemsExpanded = true;
+
+		DetailsView = ICustomDetailsViewModule::Get().CreateCustomDetailsView(CustomDetailsViewArgs);
+		DetailsView->SetObjects(ObjectsToView);
+
+		ContentBorder->SetContent(DetailsView.ToSharedRef());
+
+		return;
+	}
+
+	ContentBorder->SetContent(CreateHintText(LOCTEXT("NoSelection", "Select an object to view details.")));
+}
+
+TSharedRef<SWidget> FAvaSequenceSelectionDetails::CreateHintText(const FText& InMessage)
+{
+	return SNew(SBox)
+		.HAlign(HAlign_Center)
+		.Padding(2.0f, 24.0f, 2.0f, 2.0f)
+		[
+			SNew(STextBlock)
+			.Text(InMessage)
+			.TextStyle(FAppStyle::Get(), "HintText")
+		];
+}
+
+FKeyEditData FAvaSequenceSelectionDetails::GetKeyEditData() const
+{
+	const TSharedPtr<FAvaSequencer> AvaSequencer = AvaSequencerWeak.Pin();
+	if (!AvaSequencer.IsValid())
+	{
+		ContentBorder->ClearContent();
+		return FKeyEditData();
+	}
+
+	const TSharedRef<ISequencer> Sequencer = AvaSequencer->GetSequencer();
+
+	const TSharedPtr<UE::Sequencer::FSequencerSelection> SequencerSelection = UE::AvaSequencer::Private::GetSelection(*Sequencer);
+	if (!SequencerSelection.IsValid())
+	{
+		ContentBorder->ClearContent();
+		return FKeyEditData();
+	}
+
+	return UE::AvaSequencer::Private::GetKeyEditData(SequencerSelection->KeySelection);
+}
+
+#undef LOCTEXT_NAMESPACE
diff --git a/Engine/Plugins/Experimental/Avalanche/Source/AvalancheSequencer/Private/DetailsView/Section/AvaSequenceSelectionDetails.h b/Engine/Plugins/Experimental/Avalanche/Source/AvalancheSequencer/Private/DetailsView/Section/AvaSequenceSelectionDetails.h
new file mode 100644
--- /dev/null
+++ b/Engine/Plugins/Experimental/Avalanche/Source/AvalancheSequencer/Private/DetailsView/Section/AvaSequenceSelectionDetails.h
@@ -0,0 +1,35 @@
+// Copyright Epic Games, Inc. All Rights Reserved.
+
+#pragma once
+
+#include "IAvaSequenceSectionDetails.h"
+
+class FAvaSequencer;
+class FName;
+class FText;
+class ICustomDetailsView;
+class SBorder;
+class SWidget;
+struct FKeyEditData;
+
+class FAvaSequenceSelectionDetails : public IAvaSequenceSectionDetails
+{
+protected:
+	//~ Begin IAvaSequenceSectionDetails
+	virtual FName GetSectionName() const override;
+	virtual FText GetSectionDisplayName() const override;
+	virtual TSharedRef<SWidget> CreateContentWidget(const TSharedRef<FAvaSequencer>& InAvaSequencer) override;
+	//~ End IAvaSequenceSectionDetails
+
+	virtual void OnSequencerSelectionChanged();
+
+	virtual TSharedRef<SWidget> CreateHintText(const FText& InMessage);
+
+	FKeyEditData GetKeyEditData() const;
+
+	TWeakPtr<FAvaSequencer> AvaSequencerWeak;
+
+	TSharedPtr<SBorder> ContentBorder;
+
+	TSharedPtr<ICustomDetailsView> DetailsView;
+};
diff --git a/Engine/Plugins/Experimental/Avalanche/Source/AvalancheSequencer/Private/DetailsView/Section/AvaSequenceSettingsDetails.cpp b/Engine/Plugins/Experimental/Avalanche/Source/AvalancheSequencer/Private/DetailsView/Section/AvaSequenceSettingsDetails.cpp
new file mode 100644
--- /dev/null
+++ b/Engine/Plugins/Experimental/Avalanche/Source/AvalancheSequencer/Private/DetailsView/Section/AvaSequenceSettingsDetails.cpp
@@ -0,0 +1,56 @@
+// Copyright Epic Games, Inc. All Rights Reserved.
+
+#include "AvaSequenceSettingsDetails.h"
+#include "AvaSequence.h"
+#include "AvaSequencer.h"
+#include "CustomDetailsViewArgs.h"
+#include "CustomDetailsViewModule.h"
+
+#define LOCTEXT_NAMESPACE "AvaSequenceSettingsDetails"
+
+FName FAvaSequenceSettingsDetails::GetSectionName() const
+{
+	return TEXT("Settings");
+}
+
+FText FAvaSequenceSettingsDetails::GetSectionDisplayName() const
+{
+	return LOCTEXT("SettingsLabel", "Sequence");
+}
+
+TSharedRef<SWidget> FAvaSequenceSettingsDetails::CreateContentWidget(const TSharedRef<FAvaSequencer>& InAvaSequencer)
+{
+	AvaSequencerWeak = InAvaSequencer;
+
+	FCustomDetailsViewArgs CustomDetailsViewArgs;
+	CustomDetailsViewArgs.IndentAmount = 0.f;
+	CustomDetailsViewArgs.bShowCategories = true;
+	CustomDetailsViewArgs.bAllowGlobalExtensions = true;
+	CustomDetailsViewArgs.CategoryAllowList.Allow(TEXT("Sequence Settings"));
+	CustomDetailsViewArgs.ExpansionState.Add(FCustomDetailsViewItemId::MakeCategoryId("Sequence Settings"), true);
+	CustomDetailsViewArgs.ExpansionState.Add(FCustomDetailsViewItemId::MakePropertyId<UAvaSequence>(TEXT("Marks")), true);
+
+	SettingsDetailsView = ICustomDetailsViewModule::Get().CreateCustomDetailsView(CustomDetailsViewArgs);
+
+	InAvaSequencer->GetOnViewedSequenceChanged().AddSP(this, &FAvaSequenceSettingsDetails::OnViewedSequenceChanged);
+
+	if (UAvaSequence* const ViewedSequence = InAvaSequencer->GetViewedSequence())
+	{
+		OnViewedSequenceChanged(ViewedSequence);	
+	}
+
+	return SettingsDetailsView.ToSharedRef();
+}
+
+bool FAvaSequenceSettingsDetails::ShouldShowSection() const
+{
+	TSharedPtr<FAvaSequencer> AvaSequencer = AvaSequencerWeak.Pin();
+	return AvaSequencer.IsValid() && IsValid(AvaSequencer->GetViewedSequence());
+}
+
+void FAvaSequenceSettingsDetails::OnViewedSequenceChanged(UAvaSequence* InSequence)
+{
+	SettingsDetailsView->SetObject(InSequence);
+}
+
+#undef LOCTEXT_NAMESPACE
diff --git a/Engine/Plugins/Experimental/Avalanche/Source/AvalancheSequencer/Private/DetailsView/Section/AvaSequenceSettingsDetails.h b/Engine/Plugins/Experimental/Avalanche/Source/AvalancheSequencer/Private/DetailsView/Section/AvaSequenceSettingsDetails.h
new file mode 100644
--- /dev/null
+++ b/Engine/Plugins/Experimental/Avalanche/Source/AvalancheSequencer/Private/DetailsView/Section/AvaSequenceSettingsDetails.h
@@ -0,0 +1,24 @@
+// Copyright Epic Games, Inc. All Rights Reserved.
+
+#pragma once
+
+#include "IAvaSequenceSectionDetails.h"
+
+class UAvaSequence;
+class ICustomDetailsView;
+
+class FAvaSequenceSettingsDetails : public IAvaSequenceSectionDetails
+{
+	//~ Begin IAvaSequenceSectionDetails
+	virtual FName GetSectionName() const override;
+	virtual FText GetSectionDisplayName() const override;
+	virtual TSharedRef<SWidget> CreateContentWidget(const TSharedRef<FAvaSequencer>& InAvaSequencer) override;
+	virtual bool ShouldShowSection() const override;
+	//~ End IAvaSequenceSectionDetails
+
+	void OnViewedSequenceChanged(UAvaSequence* InSequence);
+
+	TWeakPtr<FAvaSequencer> AvaSequencerWeak;
+
+	TSharedPtr<ICustomDetailsView> SettingsDetailsView;
+};
diff --git a/Engine/Plugins/Experimental/Avalanche/Source/AvalancheSequencer/Private/DetailsView/Section/IAvaSequenceSectionDetails.h b/Engine/Plugins/Experimental/Avalanche/Source/AvalancheSequencer/Private/DetailsView/Section/IAvaSequenceSectionDetails.h
new file mode 100644
--- /dev/null
+++ b/Engine/Plugins/Experimental/Avalanche/Source/AvalancheSequencer/Private/DetailsView/Section/IAvaSequenceSectionDetails.h
@@ -0,0 +1,24 @@
+// Copyright Epic Games, Inc. All Rights Reserved.
+
+#pragma once
+
+#include "Templates/SharedPointer.h"
+
+class FAvaSequencer;
+class FName;
+class FText;
+class SWidget;
+
+class IAvaSequenceSectionDetails : public TSharedFromThis<IAvaSequenceSectionDetails>
+{
+public:
+	virtual ~IAvaSequenceSectionDetails() {}
+
+	virtual FName GetSectionName() const = 0;
+
+	virtual FText GetSectionDisplayName() const = 0;
+
+	virtual TSharedRef<SWidget> CreateContentWidget(const TSharedRef<FAvaSequencer>& InAvaSequencer) = 0;
+
+	virtual bool ShouldShowSection() const { return true; }
+};
diff --git a/Engine/Plugins/Experimental/Avalanche/Source/AvalancheSequencer/Private/EaseCurveTool/AvaEaseCurve.cpp b/Engine/Plugins/Experimental/Avalanche/Source/AvalancheSequencer/Private/EaseCurveTool/AvaEaseCurve.cpp
new file mode 100644
--- /dev/null
+++ b/Engine/Plugins/Experimental/Avalanche/Source/AvalancheSequencer/Private/EaseCurveTool/AvaEaseCurve.cpp
@@ -0,0 +1,140 @@
+// Copyright Epic Games, Inc. All Rights Reserved.
+
+#include "AvaEaseCurve.h"
+
+UAvaEaseCurve::UAvaEaseCurve()
+{
+	StartKeyHandle = FloatCurve.AddKey(0, 0);
+	EndKeyHandle = FloatCurve.AddKey(1, 1);
+}
+
+void UAvaEaseCurve::OnCurveChanged(const TArray<FRichCurveEditInfo>& InRichCurveEditInfos) 
+{
+	UCurveFloat::OnCurveChanged(InRichCurveEditInfos);
+}
+
+FKeyHandle UAvaEaseCurve::GetStartKeyHandle() const
+{
+	return StartKeyHandle;
+}
+
+FKeyHandle UAvaEaseCurve::GetEndKeyHandle() const
+{
+	return EndKeyHandle;
+}
+
+FRichCurveKey& UAvaEaseCurve::GetStartKey()
+{
+	return FloatCurve.GetKey(StartKeyHandle);
+}
+
+FRichCurveKey& UAvaEaseCurve::GetEndKey()
+{
+	return FloatCurve.GetKey(EndKeyHandle);
+}
+
+float UAvaEaseCurve::GetStartKeyTime() const
+{
+	return FloatCurve.GetKey(StartKeyHandle).Time;
+}
+
+float UAvaEaseCurve::GetStartKeyValue() const
+{
+	return FloatCurve.GetKey(StartKeyHandle).Value;
+}
+
+float UAvaEaseCurve::GetEndKeyTime() const
+{
+	return FloatCurve.GetKey(EndKeyHandle).Time;
+}
+
+float UAvaEaseCurve::GetEndKeyValue() const
+{
+	return FloatCurve.GetKey(EndKeyHandle).Value;
+}
+
+FAvaEaseCurveTangents UAvaEaseCurve::GetTangents() const
+{
+	return FAvaEaseCurveTangents(FloatCurve.GetKey(StartKeyHandle), FloatCurve.GetKey(EndKeyHandle));
+}
+
+void UAvaEaseCurve::SetTangents(const FAvaEaseCurveTangents& InTangents)
+{
+	SetStartTangent(InTangents.Start, InTangents.StartWeight);
+	SetEndTangent(InTangents.End, InTangents.EndWeight);
+}
+
+void UAvaEaseCurve::SetStartTangent(const float InTangent, const float InTangentWeight)
+{
+	FloatCurve.SetKeyTime(StartKeyHandle, 0);
+	FloatCurve.SetKeyValue(StartKeyHandle, 0);
+	FloatCurve.SetKeyInterpMode(StartKeyHandle, ERichCurveInterpMode::RCIM_Cubic);
+	FloatCurve.SetKeyTangentWeightMode(StartKeyHandle, ERichCurveTangentWeightMode::RCTWM_WeightedBoth, false);
+	FloatCurve.SetKeyTangentMode(StartKeyHandle, ERichCurveTangentMode::RCTM_Break, false);
+
+	FRichCurveKey& StartKey = GetStartKey();
+	StartKey.ArriveTangent = 0;
+	StartKey.ArriveTangentWeight = 0;
+	StartKey.LeaveTangent = InTangent;
+	StartKey.LeaveTangentWeight = InTangentWeight;
+}
+
+void UAvaEaseCurve::SetEndTangent(const float InTangent, const float InTangentWeight)
+{
+	FloatCurve.SetKeyTime(EndKeyHandle, 1);
+	FloatCurve.SetKeyValue(EndKeyHandle, 1);
+	FloatCurve.SetKeyInterpMode(EndKeyHandle, ERichCurveInterpMode::RCIM_Cubic);
+	FloatCurve.SetKeyTangentWeightMode(EndKeyHandle, ERichCurveTangentWeightMode::RCTWM_WeightedBoth, false);
+	FloatCurve.SetKeyTangentMode(EndKeyHandle, ERichCurveTangentMode::RCTM_Break, false);
+
+	FRichCurveKey& EndKey = GetEndKey();
+	EndKey.LeaveTangent = 0;
+	EndKey.LeaveTangentWeight = 0;
+	EndKey.ArriveTangent = InTangent;
+	EndKey.ArriveTangentWeight = InTangentWeight;
+}
+
+void UAvaEaseCurve::FlattenOrStraightenTangents(const FKeyHandle InKeyHandle, const bool bInFlattenTangents)
+{
+	FRichCurveKey& Key = FloatCurve.GetKey(InKeyHandle);
+
+	float LeaveTangent = Key.LeaveTangent;
+	float ArriveTangent = Key.ArriveTangent;
+
+	if (bInFlattenTangents)
+	{
+		LeaveTangent = 0.f;
+		ArriveTangent = 0.f;
+	}
+	else
+	{
+		LeaveTangent = (LeaveTangent + ArriveTangent) * 0.5f;
+		ArriveTangent = LeaveTangent;
+	}
+
+	Key.LeaveTangent = LeaveTangent;
+	Key.ArriveTangent = ArriveTangent;
+	if (Key.InterpMode == RCIM_Cubic && Key.TangentMode == RCTM_Auto)
+	{
+		Key.TangentMode = RCTM_User;
+	}
+}
+
+void UAvaEaseCurve::SetInterpMode(const ERichCurveInterpMode InInterpMode, const ERichCurveTangentMode InTangentMode)
+{
+	if (InInterpMode != RCIM_Cubic)
+	{
+		FloatCurve.SetKeyInterpMode(StartKeyHandle, InInterpMode);
+		FloatCurve.SetKeyTangentMode(StartKeyHandle, InTangentMode);
+
+		FloatCurve.SetKeyInterpMode(EndKeyHandle, InInterpMode);
+		FloatCurve.SetKeyTangentMode(EndKeyHandle, InTangentMode);
+	}
+}
+
+void UAvaEaseCurve::BroadcastUpdate()
+{
+	TArray<FRichCurveEditInfo> ChangedCurveEditInfos;
+	ChangedCurveEditInfos.Add(FRichCurveEditInfo(&FloatCurve));
+	OnCurveChanged(ChangedCurveEditInfos);
+}
diff --git a/Engine/Plugins/Experimental/Avalanche/Source/AvalancheSequencer/Private/EaseCurveTool/AvaEaseCurve.h b/Engine/Plugins/Experimental/Avalanche/Source/AvalancheSequencer/Private/EaseCurveTool/AvaEaseCurve.h
new file mode 100644
--- /dev/null
+++ b/Engine/Plugins/Experimental/Avalanche/Source/AvalancheSequencer/Private/EaseCurveTool/AvaEaseCurve.h
@@ -0,0 +1,46 @@
+// Copyright Epic Games, Inc. All Rights Reserved.
+
+#pragma once
+
+#include "AvaEaseCurveTangents.h"
+#include "Curves/CurveFloat.h"
+#include "AvaEaseCurve.generated.h"
+
+UCLASS()
+class UAvaEaseCurve : public UCurveFloat 
+{
+	GENERATED_BODY()
+
+public:
+	UAvaEaseCurve();
+
+	FKeyHandle GetStartKeyHandle() const;
+	FKeyHandle GetEndKeyHandle() const;
+
+	FRichCurveKey& GetStartKey();
+	FRichCurveKey& GetEndKey();
+
+	float GetStartKeyTime() const;
+	float GetStartKeyValue() const;
+	float GetEndKeyTime() const;
+	float GetEndKeyValue() const;
+
+	FAvaEaseCurveTangents GetTangents() const;
+	void SetTangents(const FAvaEaseCurveTangents& InTangents);
+	void SetStartTangent(const float InTangent, const float InTangentWeight);
+	void SetEndTangent(const float InTangent, const float InTangentWeight);
+
+	void FlattenOrStraightenTangents(const FKeyHandle InKeyHandle, const bool bInFlattenTangents);
+
+	void SetInterpMode(const ERichCurveInterpMode InInterpMode, const ERichCurveTangentMode InTangentMode);
+
+	void BroadcastUpdate();
+
+	//~ Begin UCurveFloat
+	virtual void OnCurveChanged(const TArray<FRichCurveEditInfo>& InRichCurveEditInfos) override;
+	//~ End UCurveFloat
+
+protected:
+	FKeyHandle StartKeyHandle;
+	FKeyHandle EndKeyHandle;
+};
diff --git a/Engine/Plugins/Experimental/Avalanche/Source/AvalancheSequencer/Private/EaseCurveTool/AvaEaseCurvePreset.h b/Engine/Plugins/Experimental/Avalanche/Source/AvalancheSequencer/Private/EaseCurveTool/AvaEaseCurvePreset.h
new file mode 100644
--- /dev/null
+++ b/Engine/Plugins/Experimental/Avalanche/Source/AvalancheSequencer/Private/EaseCurveTool/AvaEaseCurvePreset.h
@@ -0,0 +1,42 @@
+// Copyright Epic Games, Inc. All Rights Reserved.
+
+#pragma once
+
+#include "Containers/UnrealString.h"
+#include "EaseCurveTool/AvaEaseCurveTangents.h"
+#include "Templates/SharedPointer.h"
+#include "AvaEaseCurvePreset.generated.h"
+
+USTRUCT()
+struct FAvaEaseCurvePreset
+{
+	GENERATED_BODY()
+
+	FAvaEaseCurvePreset() {}
+	FAvaEaseCurvePreset(const FString& InName, const FString& InCategory, const FAvaEaseCurveTangents& InTangents)
+		: Name(InName), Category(InCategory), Tangents(InTangents)
+	{}
+
+	FORCEINLINE bool operator==(const FAvaEaseCurvePreset& InRhs) const
+	{
+		return Name.Equals(InRhs.Name) && Category.Equals(InRhs.Category);
+	}
+	FORCEINLINE bool operator!=(const FAvaEaseCurvePreset& InRhs) const
+	{
+		return !(*this == InRhs);
+	}
+	FORCEINLINE bool operator<(const FAvaEaseCurvePreset& InRhs) const
+	{
+		// Order from soft to hard based on curve length
+		return Tangents.CalculateCurveLength() < InRhs.Tangents.CalculateCurveLength();
+	}
+
+	UPROPERTY()
+	FString Name;
+
+	UPROPERTY()
+	FString Category;
+
+	UPROPERTY()
+	FAvaEaseCurveTangents Tangents;
+};
diff --git a/Engine/Plugins/Experimental/Avalanche/Source/AvalancheSequencer/Private/EaseCurveTool/AvaEaseCurvePresetDragDropOp.cpp b/Engine/Plugins/Experimental/Avalanche/Source/AvalancheSequencer/Private/EaseCurveTool/AvaEaseCurvePresetDragDropOp.cpp
new file mode 100644
--- /dev/null
+++ b/Engine/Plugins/Experimental/Avalanche/Source/AvalancheSequencer/Private/EaseCurveTool/AvaEaseCurvePresetDragDropOp.cpp
@@ -0,0 +1,52 @@
+// Copyright Epic Games, Inc. All Rights Reserved.
+
+#include "AvaEaseCurvePresetDragDropOp.h"
+#include "EaseCurveTool/AvaEaseCurvePreset.h"
+#include "EaseCurveTool/Widgets/SAvaEaseCurvePresetGroup.h"
+#include "EaseCurveTool/Widgets/SAvaEaseCurvePresetGroupItem.h"
+#include "Styling/AppStyle.h"
+#include "Widgets/Layout/SBorder.h"
+
+FCursorReply FAvaEaseCurvePresetDragDropOperation::OnCursorQuery()
+{
+	return FCursorReply::Cursor(EMouseCursor::GrabHandClosed);
+}
+
+TSharedPtr<SWidget> FAvaEaseCurvePresetDragDropOperation::GetDefaultDecorator() const
+{
+	if (!WidgetWeak.IsValid())
+	{
+		return FDragDropOperation::GetDefaultDecorator();
+	}
+
+	return SNew(SBorder)
+		.Padding(2.f)
+		.BorderImage(FAppStyle::GetBrush(TEXT("Graph.ConnectorFeedback.Border")))
+		.Content()
+		[
+			WidgetWeak.Pin().ToSharedRef()
+		];
+}
+
+void FAvaEaseCurvePresetDragDropOperation::OnDrop(bool bInDropWasHandled, const FPointerEvent& InMouseEvent)
+{
+	FDragDropOperation::OnDrop(bInDropWasHandled, InMouseEvent);
+
+	if (const TSharedPtr<SAvaEaseCurvePresetGroupItem> Widget = WidgetWeak.Pin())
+	{
+		Widget->TriggerEndMove();
+	}
+
+	for (const TWeakPtr<SAvaEaseCurvePresetGroup>& GroupWidgetWeak : HoveredGroupWidgets)
+	{
+		if (const TSharedPtr<SAvaEaseCurvePresetGroup> GroupWidget = GroupWidgetWeak.Pin())
+		{
+			GroupWidget->ResetDragBorder();
+		}
+	}
+}
+
+void FAvaEaseCurvePresetDragDropOperation::AddHoveredGroup(const TSharedRef<SAvaEaseCurvePresetGroup>& InGroupWidget)
+{
+	HoveredGroupWidgets.Add(InGroupWidget);
+}
diff --git a/Engine/Plugins/Experimental/Avalanche/Source/AvalancheSequencer/Private/EaseCurveTool/AvaEaseCurvePresetDragDropOp.h b/Engine/Plugins/Experimental/Avalanche/Source/AvalancheSequencer/Private/EaseCurveTool/AvaEaseCurvePresetDragDropOp.h
new file mode 100644
--- /dev/null
+++ b/Engine/Plugins/Experimental/Avalanche/Source/AvalancheSequencer/Private/EaseCurveTool/AvaEaseCurvePresetDragDropOp.h
@@ -0,0 +1,37 @@
+// Copyright Epic Games, Inc. All Rights Reserved.
+
+#pragma once
+
+#include "Input/DragAndDrop.h"
+#include "Templates/SharedPointer.h"
+
+class SAvaEaseCurvePresetGroup;
+class SAvaEaseCurvePresetGroupItem;
+struct FAvaEaseCurvePreset;
+
+class FAvaEaseCurvePresetDragDropOperation : public FDragDropOperation
+{
+public:
+	DRAG_DROP_OPERATOR_TYPE(FAvaEaseCurvePresetDragDropOperation, FDragDropOperation)
+
+	FAvaEaseCurvePresetDragDropOperation(const TSharedPtr<SAvaEaseCurvePresetGroupItem>& InWidget, const TSharedPtr<FAvaEaseCurvePreset>& InPreset)
+		: WidgetWeak(InWidget), Preset(InPreset)
+	{}
+
+	TSharedPtr<SAvaEaseCurvePresetGroupItem> GetWidget() const { return WidgetWeak.Pin(); }
+	TSharedPtr<FAvaEaseCurvePreset> GetPreset() const { return Preset; }
+
+	//~ Begin FDragDropOperation
+	virtual FCursorReply OnCursorQuery() override;
+	virtual TSharedPtr<SWidget> GetDefaultDecorator() const override;
+	virtual void OnDrop(bool bInDropWasHandled, const FPointerEvent& InMouseEvent) override;
+	//~ End FDragDropOperation
+
+	void AddHoveredGroup(const TSharedRef<SAvaEaseCurvePresetGroup>& InGroupWidget);
+
+protected:
+	TWeakPtr<SAvaEaseCurvePresetGroupItem> WidgetWeak;
+	TSharedPtr<FAvaEaseCurvePreset> Preset;
+
+	TSet<TWeakPtr<SAvaEaseCurvePresetGroup>> HoveredGroupWidgets;
+};
diff --git a/Engine/Plugins/Experimental/Avalanche/Source/AvalancheSequencer/Private/EaseCurveTool/AvaEaseCurveStyle.cpp b/Engine/Plugins/Experimental/Avalanche/Source/AvalancheSequencer/Private/EaseCurveTool/AvaEaseCurveStyle.cpp
new file mode 100644
--- /dev/null
+++ b/Engine/Plugins/Experimental/Avalanche/Source/AvalancheSequencer/Private/EaseCurveTool/AvaEaseCurveStyle.cpp
@@ -0,0 +1,65 @@
+// Copyright Epic Games, Inc. All Rights Reserved.
+
+#include "AvaEaseCurveStyle.h"
+
+#include "Framework/Application/SlateApplication.h"
+#include "Interfaces/IPluginManager.h"
+#include "Brushes/SlateImageBrush.h"
+#include "Styling/AppStyle.h"
+#include "Styling/SlateStyle.h"
+#include "Styling/SlateStyleMacros.h"
+#include "Styling/SlateStyleRegistry.h"
+#include "Styling/SlateTypes.h"
+#include "Styling/StyleColors.h"
+
+FAvaEaseCurveStyle::FAvaEaseCurveStyle()
+	: FSlateStyleSet(TEXT("AvaEaseCurveStyle"))
+{
+	const TSharedPtr<IPlugin> Plugin = IPluginManager::Get().FindPlugin(TEXT("Avalanche"));
+	check(Plugin.IsValid());
+	if (Plugin.IsValid())
+	{
+		SetContentRoot(FPaths::Combine(Plugin->GetBaseDir(), TEXT("Resources")));
+	}
+
+	Set("EditMode.Background", new FSlateRoundedBoxBrush(FStyleColors::Transparent, 4.f, FLinearColor(0.1f, 0.1f, 0.1f, 1.f), 1.f));
+	Set("EditMode.Background.Highlight", new FSlateRoundedBoxBrush(FStyleColors::Transparent, 4.f, FLinearColor(0.6f, 0.6f, 0.6f, 1.f), 1.f));
+	Set("EditMode.Background.Over", new FSlateRoundedBoxBrush(FStyleColors::Transparent, 4.f, FStyleColors::AccentBlue.GetSpecifiedColor(), 1.f));
+
+	const FButtonStyle& SimpleButtonStyle = FAppStyle::GetWidgetStyle<FButtonStyle>(TEXT("SimpleButton"));
+
+	constexpr float ToolButtonPadding = 2.f;
+	Set("ToolButton.Padding", ToolButtonPadding);
+
+	constexpr float ToolButtonImageSize = 12.f;
+	Set("ToolButton.ImageSize", ToolButtonImageSize);
+
+	const FButtonStyle ToolButtonStyle = FButtonStyle(SimpleButtonStyle)
+		.SetHovered(FSlateRoundedBoxBrush(FStyleColors::Hover, 4.0f))
+		.SetPressed(FSlateRoundedBoxBrush(FStyleColors::Hover, 4.0f))
+		.SetNormalPadding(FMargin(ToolButtonPadding))
+		.SetPressedPadding(FMargin(ToolButtonPadding, ToolButtonPadding + (ToolButtonPadding * 0.5f), ToolButtonPadding, ToolButtonPadding - (ToolButtonPadding * 0.5f)));
+	Set("ToolButton", ToolButtonStyle);
+
+	const FButtonStyle ToolButtonNoPadStyle = FButtonStyle(ToolButtonStyle)
+		.SetNormalPadding(FMargin(0.f))
+		.SetPressedPadding(FMargin(0.f));
+	Set("ToolButton.NoPad", ToolButtonNoPadStyle);
+
+	const FCheckBoxStyle& ToggleButtonCheckboxStyle = FAppStyle::GetWidgetStyle<FCheckBoxStyle>(TEXT("ToggleButtonCheckbox"));
+
+	const FCheckBoxStyle ToolToggleButtonStyle = FCheckBoxStyle(ToggleButtonCheckboxStyle)
+		.SetUncheckedHoveredImage(FSlateRoundedBoxBrush(FStyleColors::Hover, 4.0f))
+		.SetUncheckedPressedImage(FSlateRoundedBoxBrush(FStyleColors::Hover, 4.0f))
+		.SetCheckedHoveredImage(FSlateRoundedBoxBrush(FStyleColors::PrimaryHover, 4.0f))
+		.SetCheckedPressedImage(FSlateRoundedBoxBrush(FStyleColors::PrimaryPress, 4.0f))
+		.SetPadding(ToolButtonPadding);
+	Set("ToolToggleButton", ToolToggleButtonStyle);
+
+	FSlateStyleRegistry::RegisterSlateStyle(*this);
+}
+
+FAvaEaseCurveStyle::~FAvaEaseCurveStyle()
+{
+	FSlateStyleRegistry::UnRegisterSlateStyle(*this);
+}
diff --git a/Engine/Plugins/Experimental/Avalanche/Source/AvalancheSequencer/Private/EaseCurveTool/AvaEaseCurveStyle.h b/Engine/Plugins/Experimental/Avalanche/Source/AvalancheSequencer/Private/EaseCurveTool/AvaEaseCurveStyle.h
new file mode 100644
--- /dev/null
+++ b/Engine/Plugins/Experimental/Avalanche/Source/AvalancheSequencer/Private/EaseCurveTool/AvaEaseCurveStyle.h
@@ -0,0 +1,18 @@
+// Copyright Epic Games, Inc. All Rights Reserved.
+
+#pragma once
+
+#include "Styling/SlateStyle.h"
+
+class FAvaEaseCurveStyle final : public FSlateStyleSet
+{
+public:
+	static FAvaEaseCurveStyle& Get()
+	{
+		static FAvaEaseCurveStyle StyleSet;
+		return StyleSet;
+	}
+
+	FAvaEaseCurveStyle();
+	virtual ~FAvaEaseCurveStyle() override;
+};
