[MotionDesign]

diff --git a/Engine/Plugins/Experimental/Avalanche/Source/AvalancheSequencer/Private/EaseCurveTool/AvaEaseCurveSubsystem.cpp b/Engine/Plugins/Experimental/Avalanche/Source/AvalancheSequencer/Private/EaseCurveTool/AvaEaseCurveSubsystem.cpp
new file mode 100644
--- /dev/null
+++ b/Engine/Plugins/Experimental/Avalanche/Source/AvalancheSequencer/Private/EaseCurveTool/AvaEaseCurveSubsystem.cpp
@@ -0,0 +1,620 @@
+// Copyright Epic Games, Inc. All Rights Reserved.
+
+#include "AvaEaseCurveSubsystem.h"
+#include "EaseCurveTool/AvaEaseCurvePreset.h"
+#include "EaseCurveTool/AvaEaseCurveToolSettings.h"
+#include "Editor.h"
+#include "Engine/World.h"
+#include "HAL/FileManager.h"
+#include "Misc/FileHelper.h"
+#include "Serialization/JsonSerializer.h"
+
+DEFINE_LOG_CATEGORY_STATIC(LogAvaEaseCurveSubsystem, Log, All);
+
+namespace UE::EaseCurveTool::Json::Private
+{
+	bool GetPresetFilePath(FString InCategory, const bool bSetNewCategory, FString& OutFilePath)
+	{
+		OutFilePath.Empty();
+
+		const UAvaEaseCurveToolSettings* const Settings = GetDefault<UAvaEaseCurveToolSettings>();
+		check(Settings);
+
+		FString PresetsPath = Settings->GetPresetsPath();
+		PresetsPath = FPaths::ConvertRelativePathToFull(PresetsPath);
+
+		if (!FPaths::DirectoryExists(PresetsPath))
+		{
+			UE_LOG(LogAvaEaseCurveSubsystem, Warning, TEXT("CurveEaseToolJson::GetPresetFilePath() - Configured presets directory does not exist: %s"), *PresetsPath);
+			return false;
+		}
+
+		if (InCategory.IsEmpty())
+		{
+			if (!bSetNewCategory)
+			{
+				return false;
+			}
+
+			const FString EaseCurveTool_NewPresetCategory = Settings->GetNewPresetCategory();
+			if (EaseCurveTool_NewPresetCategory.IsEmpty())
+			{
+				UE_LOG(LogAvaEaseCurveSubsystem, Warning, TEXT("CurveEaseToolJson::GetPresetFilePath() - No preset category and no configured default category"));
+				return false;
+			}
+
+			InCategory = EaseCurveTool_NewPresetCategory;
+		}
+
+		PresetsPath = FPaths::Combine(PresetsPath, InCategory) + TEXT(".json");
+		OutFilePath = PresetsPath;
+
+		return FPaths::FileExists(PresetsPath);
+	}
+
+	bool LoadCurvePresetsJson(const FString& InFilePath, TSharedPtr<FJsonObject>& OutRootObject)
+	{
+		FString FileJsonString;
+		if (!FFileHelper::LoadFileToString(FileJsonString, *InFilePath))
+		{
+			UE_LOG(LogAvaEaseCurveSubsystem, Warning, TEXT("CurveEaseToolJson::LoadCurvePresetsJson() - Unable to load Json file: %s"), *InFilePath);
+			return false;
+		}
+
+		const TSharedPtr<TJsonReader<>> JsonReader = TJsonReaderFactory<>::Create(FileJsonString);
+		OutRootObject = MakeShared<FJsonObject>();
+		if (!FJsonSerializer::Deserialize(JsonReader.ToSharedRef(), OutRootObject))
+		{
+			UE_LOG(LogAvaEaseCurveSubsystem, Warning, TEXT("CurveEaseToolJson::LoadCurvePresetsJson() - Unable to parse file [%s]. Json=[%s]"), *InFilePath, *FileJsonString);
+			return false;
+		}
+
+		return true;
+	}
+
+	bool SaveCurvePresetsJson(const FString& InFilePath, const TSharedRef<FJsonObject>& InRootObject)
+	{
+		FString FileJsonString;
+		const TSharedRef<TJsonWriter<>> Writer = TJsonWriterFactory<>::Create(&FileJsonString);
+		if (!FJsonSerializer::Serialize(InRootObject, Writer))
+		{
+			UE_LOG(LogAvaEaseCurveSubsystem, Warning, TEXT("CurveEaseToolJson::SaveCurvePresetsJson() - Unable to serialize [%s]. Json=[%s]"), *InFilePath, *FileJsonString);
+			return false;
+		}
+
+		if (!FFileHelper::SaveStringToFile(FileJsonString, *InFilePath))
+		{
+			UE_LOG(LogAvaEaseCurveSubsystem, Warning, TEXT("CurveEaseToolJson::SaveCurvePresetsJson() - Unable to save Json file: %s"), *InFilePath);
+			return false;
+		}
+
+		return true;
+	}
+
+	bool DoesPresetCategoryExist(const FString& InCategory)
+	{
+		FString OutFilePath;
+		return GetPresetFilePath(InCategory, false, OutFilePath);
+	}
+
+	FString AddNewPresetCategory()
+	{
+		static const FString DefaultNewCategoryName = TEXT("New Category");
+
+		FString NewCategoryName;
+
+		int32 CurrentNum = 1;
+		FString FilePath;
+
+		while (true)
+		{
+			NewCategoryName = (CurrentNum == 1)
+				? DefaultNewCategoryName 
+				: DefaultNewCategoryName + TEXT(" ") + FString::FromInt(CurrentNum);
+
+			if (!GetPresetFilePath(NewCategoryName, false, FilePath))
+			{
+				break;
+			}
+
+			CurrentNum++;
+		}
+
+		const TSharedPtr<FJsonObject> RootObject = MakeShared<FJsonObject>();
+		if (!SaveCurvePresetsJson(FilePath, RootObject.ToSharedRef()))
+		{
+			return FString();
+		}
+
+		return NewCategoryName;
+	}
+
+	bool RemovePresetCategory(const FString& InCategory)
+	{
+		FString FilePath;
+		if (!GetPresetFilePath(InCategory, false, FilePath))
+		{
+			return false;
+		}
+
+		return IFileManager::Get().Delete(*FilePath);
+	}
+
+	bool RenamePresetCategory(const FString& InCategory, const FString& InNewCategory)
+	{
+		if (!DoesPresetCategoryExist(InCategory))
+		{
+			return false;
+		}
+
+		if (DoesPresetCategoryExist(InNewCategory))
+		{
+			return false;
+		}
+
+		FString OldFilePath;
+		if (!GetPresetFilePath(InCategory, false, OldFilePath))
+		{
+			return false;
+		}
+
+		FString NewFilePath;
+		GetPresetFilePath(InNewCategory, false, NewFilePath);
+
+		TSharedPtr<FJsonObject> RootObject;
+		if (!LoadCurvePresetsJson(OldFilePath, RootObject) || !RootObject.IsValid())
+		{
+			return false;
+		}
+
+		if (!SaveCurvePresetsJson(NewFilePath, RootObject.ToSharedRef()))
+		{
+			return false;
+		}
+
+		return RemovePresetCategory(InCategory);
+	}
+	
+	bool DoesPresetExist(const FString& InCategory, const FString& InPreset)
+	{
+		FString PresetsPath;
+		if (!GetPresetFilePath(InCategory, false, PresetsPath))
+		{
+			UE_LOG(LogAvaEaseCurveSubsystem, Warning, TEXT("CurveEaseToolJson::DoesPresetExist() - Could not find file for preset category: %s"), *InCategory);
+			return false;
+		}
+
+		TSharedPtr<FJsonObject> RootObject;
+		if (!LoadCurvePresetsJson(PresetsPath, RootObject) || !RootObject.IsValid())
+		{
+			UE_LOG(LogAvaEaseCurveSubsystem, Warning, TEXT("CurveEaseToolJson::DoesPresetExist() - Unable to load Json file: %s"), *PresetsPath);
+			return false;
+		}
+
+		return RootObject->Values.Contains(InPreset);
+	}
+
+	bool AddPreset(const FString& InCategory, const FString& InPreset, const FAvaEaseCurveTangents& InTangents)
+	{
+		FString PresetsPath;
+		const bool bCategoryFileExists = GetPresetFilePath(InCategory, true, PresetsPath);
+		if (PresetsPath.IsEmpty())
+		{
+			UE_LOG(LogAvaEaseCurveSubsystem, Warning, TEXT("CurveEaseToolJson::AddPreset() - Could not find file for preset category: %s"), *InCategory);
+			return false;
+		}
+
+		TSharedPtr<FJsonObject> RootObject;
+
+		if (bCategoryFileExists)
+		{
+			if (!UE::EaseCurveTool::Json::Private::LoadCurvePresetsJson(PresetsPath, RootObject) || !RootObject.IsValid())
+			{
+				UE_LOG(LogAvaEaseCurveSubsystem, Warning, TEXT("CurveEaseToolJson::AddPreset() - Unable to load Json file: %s"), *PresetsPath);
+				return false;
+			}
+		}
+		else
+		{
+			RootObject = MakeShared<FJsonObject>();
+		}
+
+		RootObject->Values.Add(InPreset, MakeShared<FJsonValueString>(InTangents.ToJson()));
+
+		if (!SaveCurvePresetsJson(PresetsPath, RootObject.ToSharedRef()))
+		{
+			UE_LOG(LogAvaEaseCurveSubsystem, Warning, TEXT("CurveEaseToolJson::AddPreset() - Unable to save Json file: %s"), *PresetsPath);
+			return false;
+		}
+
+		return true;
+	}
+
+	bool RemovePreset(const FString& InCategory, const FString& InPreset)
+	{
+		FString PresetsPath;
+		if (!GetPresetFilePath(InCategory, false, PresetsPath))
+		{
+			UE_LOG(LogAvaEaseCurveSubsystem, Warning, TEXT("CurveEaseToolJson::RemovePreset() - Could not find file for preset category: %s"), *InCategory);
+			return false;
+		}
+
+		TSharedPtr<FJsonObject> RootObject;
+		if (!LoadCurvePresetsJson(PresetsPath, RootObject) || !RootObject.IsValid())
+		{
+			UE_LOG(LogAvaEaseCurveSubsystem, Warning, TEXT("CurveEaseToolJson::RemovePreset() - Unable to load Json file: %s"), *PresetsPath);
+			return false;
+		}
+
+		RootObject->Values.Remove(InPreset);
+
+		if (!SaveCurvePresetsJson(PresetsPath, RootObject.ToSharedRef()))
+		{
+			UE_LOG(LogAvaEaseCurveSubsystem, Warning, TEXT("CurveEaseToolJson::RemovePreset() - Unable to save Json file: %s"), *PresetsPath);
+			return false;
+		}
+
+		return true;
+	}
+
+	bool RenamePreset(const FString& InCategory, const FString& InPreset, const FString& InNewPreset)
+	{
+		FString PresetsPath;
+		if (!GetPresetFilePath(InCategory, false, PresetsPath))
+		{
+			UE_LOG(LogAvaEaseCurveSubsystem, Warning, TEXT("CurveEaseToolJson::RemovePreset() - Could not find file for preset category: %s"), *InCategory);
+			return false;
+		}
+
+		TSharedPtr<FJsonObject> RootObject;
+		if (!LoadCurvePresetsJson(PresetsPath, RootObject) || !RootObject.IsValid())
+		{
+			UE_LOG(LogAvaEaseCurveSubsystem, Warning, TEXT("CurveEaseToolJson::RemovePreset() - Unable to load Json file: %s"), *PresetsPath);
+			return false;
+		}
+
+		if (const TSharedPtr<FJsonValue>* Value = RootObject->Values.Find(InPreset))
+		{
+			const TSharedPtr<FJsonValue>& ValueRef = *Value;
+			RootObject->Values.Add(InNewPreset, MakeShared<FJsonValueString>(ValueRef->AsString()));
+			RootObject->Values.Remove(InPreset);
+		}
+
+		if (!SaveCurvePresetsJson(PresetsPath, RootObject.ToSharedRef()))
+		{
+			UE_LOG(LogAvaEaseCurveSubsystem, Warning, TEXT("CurveEaseToolJson::RemovePreset() - Unable to save Json file: %s"), *PresetsPath);
+			return false;
+		}
+
+		return true;
+	}
+}
+
+UAvaEaseCurveSubsystem& UAvaEaseCurveSubsystem::Get()
+{
+	check(GEditor);
+	return *GEditor->GetEditorSubsystem<UAvaEaseCurveSubsystem>();
+}
+
+void UAvaEaseCurveSubsystem::Initialize(FSubsystemCollectionBase& InCollection)
+{
+	Super::Initialize(InCollection);
+
+	ReloadPresetsFromJson();
+}
+
+void UAvaEaseCurveSubsystem::ExploreJsonPresetsFolder()
+{
+	const FString EaseCurvePresetsPath = FPaths::ConvertRelativePathToFull(GetDefault<UAvaEaseCurveToolSettings>()->GetPresetsPath());
+	FPlatformProcess::ExploreFolder(*EaseCurvePresetsPath);
+}
+
+void UAvaEaseCurveSubsystem::ReloadPresetsFromJson()
+{
+	const UAvaEaseCurveToolSettings* const EaseCurveToolSettings = GetDefault<UAvaEaseCurveToolSettings>();
+	check(EaseCurveToolSettings);
+	const FString EaseCurvePresetsPath = EaseCurveToolSettings->GetPresetsPath();
+
+	TArray<FString> JsonFiles;
+	IFileManager::Get().FindFilesRecursive(JsonFiles, *EaseCurvePresetsPath, TEXT("*.json"), true, false);
+
+	// Load all Json files found as their own separate category. Category Name = File Name
+	Presets.Empty();
+	for (const FString& File : JsonFiles)
+	{
+		TSharedPtr<FJsonObject> RootObject;
+		if (!UE::EaseCurveTool::Json::Private::LoadCurvePresetsJson(File, RootObject) || !RootObject.IsValid())
+		{
+			UE_LOG(LogAvaEaseCurveSubsystem, Warning, TEXT("UAvaEaseCurveSubsystem::LoadCurvePresetsJson() - Unable to load Json file [%s]"), *File);
+			continue;
+		}
+
+		const FString CategoryName = FPaths::GetBaseFilename(File);
+
+		TArray<TSharedPtr<FAvaEaseCurvePreset>>& PresetCategoryRef = Presets.FindOrAdd(CategoryName);
+
+		for (TMap<FString, TSharedPtr<FJsonValue>>::TConstIterator JsonValueIter = RootObject->Values.CreateConstIterator(); JsonValueIter; ++JsonValueIter)
+		{
+			const FString& PresetName = JsonValueIter.Key();
+			const FString TangentsString = (*JsonValueIter).Value->AsString();
+
+			FAvaEaseCurveTangents Tangents;
+			if (!FAvaEaseCurveTangents::FromString(TangentsString, Tangents))
+			{
+				continue;
+			}
+
+			const bool bExists = PresetCategoryRef.ContainsByPredicate([&PresetName](const TSharedPtr<FAvaEaseCurvePreset>& InPreset)
+				{
+					return InPreset->Name.Equals(PresetName);
+				});
+			if (!bExists)
+			{
+				PresetCategoryRef.Add(MakeShared<FAvaEaseCurvePreset>(PresetName, CategoryName, MoveTemp(Tangents)));
+			}
+		}
+
+		PresetCategoryRef.Sort([](const TSharedPtr<FAvaEaseCurvePreset>& InPresetA, const TSharedPtr<FAvaEaseCurvePreset>& InPresetB)
+			{
+				return *InPresetA < *InPresetB;
+			});
+	}
+}
+
+TArray<FString> UAvaEaseCurveSubsystem::GetEaseCurveCategories() const
+{
+	TArray<FString> OutCategories;
+	Presets.GenerateKeyArray(OutCategories);
+	return OutCategories;
+}
+
+TArray<TSharedPtr<FAvaEaseCurvePreset>> UAvaEaseCurveSubsystem::GetEaseCurvePresets() const
+{
+	TArray<TSharedPtr<FAvaEaseCurvePreset>> OutPresets;
+
+	for (const TPair<FString, TArray<TSharedPtr<FAvaEaseCurvePreset>>>& Preset : Presets)
+	{
+		OutPresets.Append(Preset.Value);
+	}
+
+	return OutPresets;
+}
+
+TArray<TSharedPtr<FAvaEaseCurvePreset>> UAvaEaseCurveSubsystem::GetEaseCurvePresets(const FString& InCategory) const
+{
+	if (!Presets.Contains(InCategory))
+	{
+		return {};
+	}
+	return Presets[InCategory];
+}
+
+bool UAvaEaseCurveSubsystem::DoesPresetExist(const FString& InName) const
+{
+	for (const TPair<FString, TArray<TSharedPtr<FAvaEaseCurvePreset>>>& Preset : Presets)
+	{
+		const bool bExists = Preset.Value.ContainsByPredicate([&InName](const TSharedPtr<FAvaEaseCurvePreset>& InPreset)
+			{
+				return InPreset->Name.Equals(InName);
+			});
+		if (bExists)
+		{
+			return true;
+		}
+	}
+	return false;
+}
+
+bool UAvaEaseCurveSubsystem::DoesPresetExist(const FString& InCategory, const FString& InName) const
+{
+	if (!Presets.Contains(InCategory))
+	{
+		return false;
+	}
+
+	return Presets[InCategory].ContainsByPredicate([&InName](const TSharedPtr<FAvaEaseCurvePreset>& InPreset)
+		{
+			return InPreset->Name.Equals(InName);
+		});
+}
+
+TSharedPtr<FAvaEaseCurvePreset> UAvaEaseCurveSubsystem::AddPreset(FAvaEaseCurvePreset InPreset)
+{
+	if (DoesPresetExist(InPreset.Category, InPreset.Name))
+	{
+		UE_LOG(LogAvaEaseCurveSubsystem, Warning, TEXT("UAvaEaseCurveSubsystem::AddPreset() - Preset already exists: %s"), *InPreset.Name);
+		return nullptr;
+	}
+
+	if (!UE::EaseCurveTool::Json::Private::AddPreset(InPreset.Category, InPreset.Name, InPreset.Tangents))
+	{
+		UE_LOG(LogAvaEaseCurveSubsystem, Warning, TEXT("UAvaEaseCurveSubsystem::AddPreset() - Unable to add preset %s to category %s Json file"), *InPreset.Name, *InPreset.Category);
+		return nullptr;
+	}
+
+	TSharedPtr<FAvaEaseCurvePreset> NewPresetPtr = MakeShared<FAvaEaseCurvePreset>(InPreset);
+	Presets.Add(InPreset.Category, { NewPresetPtr });
+
+	return NewPresetPtr;
+}
+
+TSharedPtr<FAvaEaseCurvePreset> UAvaEaseCurveSubsystem::AddPreset(const FString& InName, const FAvaEaseCurveTangents& InTangents)
+{
+	const FString CategoryName = GetDefault<UAvaEaseCurveToolSettings>()->GetNewPresetCategory();
+	if (CategoryName.IsEmpty())
+	{
+		return nullptr;
+	}
+
+	return AddPreset(FAvaEaseCurvePreset(InName, CategoryName, InTangents));
+}
+
+bool UAvaEaseCurveSubsystem::RemovePreset(const FAvaEaseCurvePreset& InPreset)
+{
+	if (!DoesPresetExist(InPreset.Name))
+	{
+		UE_LOG(LogAvaEaseCurveSubsystem, Warning, TEXT("UAvaEaseCurveSubsystem::RemovePreset() - Preset already exists: %s"), *InPreset.Name);
+		return false;
+	}
+
+	if (!UE::EaseCurveTool::Json::Private::RemovePreset(InPreset.Category, InPreset.Name))
+	{
+		UE_LOG(LogAvaEaseCurveSubsystem, Warning, TEXT("UAvaEaseCurveSubsystem::RemovePreset() - Failed to remove preset %s from category %s Json file"), *InPreset.Name, *InPreset.Category);
+		return false;
+	}
+
+	if (Presets.Contains(InPreset.Category))
+	{
+		TSharedPtr<FAvaEaseCurvePreset> PresetToRemove;
+		for (const TSharedPtr<FAvaEaseCurvePreset>& PresetPtr : Presets[InPreset.Category])
+		{
+			if (PresetPtr->Name.Equals(InPreset.Name, ESearchCase::IgnoreCase))
+			{
+				PresetToRemove = PresetPtr;
+				break;
+			}
+		}
+		if (PresetToRemove.IsValid())
+		{
+			Presets[InPreset.Category].Remove(PresetToRemove);
+		}
+	}
+
+	return true;
+}
+
+bool UAvaEaseCurveSubsystem::RemovePreset(const FString& InCategory, const FString& InName)
+{
+	return RemovePreset(FAvaEaseCurvePreset(InName, InCategory, FAvaEaseCurveTangents()));
+}
+
+bool UAvaEaseCurveSubsystem::ChangePresetCategory(const TSharedPtr<FAvaEaseCurvePreset>& InPreset, const FString& InNewCategory) const
+{
+	if (!Presets.Contains(InPreset->Category) || !Presets.Contains(InNewCategory))
+	{
+		return false;
+	}
+
+	// Verify old category exists and new category does not
+	if (!DoesPresetExist(InPreset->Category, InPreset->Name) || DoesPresetExist(InNewCategory, InPreset->Name))
+	{
+		return false;
+	}
+
+	if (!UAvaEaseCurveSubsystem::Get().RemovePreset(InPreset->Category, InPreset->Name))
+	{
+		return false;
+	}
+
+	InPreset->Category = InNewCategory;
+
+	return UAvaEaseCurveSubsystem::Get().AddPreset(*InPreset).IsValid();
+}
+
+TSharedPtr<FAvaEaseCurvePreset> UAvaEaseCurveSubsystem::FindPreset(const FString& InName)
+{
+	for (const TPair<FString, TArray<TSharedPtr<FAvaEaseCurvePreset>>>& Preset : Presets)
+	{
+		const TSharedPtr<FAvaEaseCurvePreset>* PresetPtr = Preset.Value.FindByPredicate([&InName](const TSharedPtr<FAvaEaseCurvePreset>& InPreset)
+			{
+				return InPreset->Name == InName;
+			});
+		if (PresetPtr)
+		{
+			return *PresetPtr;
+		}
+	}
+	return nullptr;
+}
+
+TSharedPtr<FAvaEaseCurvePreset> UAvaEaseCurveSubsystem::FindPresetByTangents(const FAvaEaseCurveTangents& InTangents)
+{
+	for (const TPair<FString, TArray<TSharedPtr<FAvaEaseCurvePreset>>>& Preset : Presets)
+	{
+		const TSharedPtr<FAvaEaseCurvePreset>* PresetPtr = Preset.Value.FindByPredicate([&InTangents](const TSharedPtr<FAvaEaseCurvePreset>& InPreset)
+			{
+				return InPreset->Tangents == InTangents;
+			});
+		if (PresetPtr)
+		{
+			return *PresetPtr;
+		}
+	}
+	return nullptr;
+}
+
+bool UAvaEaseCurveSubsystem::DoesPresetCategoryExist(const FString& InCategory)
+{
+	return UE::EaseCurveTool::Json::Private::DoesPresetCategoryExist(InCategory);
+}
+
+bool UAvaEaseCurveSubsystem::RenamePresetCategory(const FString& InCategory, const FString& InNewCategoryName)
+{
+	if (!UE::EaseCurveTool::Json::Private::RenamePresetCategory(InCategory, InNewCategoryName))
+	{
+		return false;
+	}
+
+	if (Presets.Remove(InCategory) > 0)
+	{
+		Presets.Add(InNewCategoryName, Presets[InCategory]);
+	}
+
+	return true;
+}
+
+bool UAvaEaseCurveSubsystem::AddNewPresetCategory()
+{
+	const FString NewCategoryName = UE::EaseCurveTool::Json::Private::AddNewPresetCategory();
+
+	if (NewCategoryName.IsEmpty())
+	{
+		return false;
+	}
+
+	Presets.Add(NewCategoryName);
+
+	return true;
+}
+
+bool UAvaEaseCurveSubsystem::RemovePresetCategory(const FString& InCategory)
+{
+	if (!UE::EaseCurveTool::Json::Private::RemovePresetCategory(InCategory))
+	{
+		return false;
+	}
+
+	return Presets.Remove(InCategory) > 0;
+}
+
+bool UAvaEaseCurveSubsystem::RenamePreset(const FString& InCategory, const FString& InPreset, const FString& InNewPresetName)
+{
+	if (!UE::EaseCurveTool::Json::Private::RenamePreset(InCategory, InPreset, InNewPresetName))
+	{
+		UE_LOG(LogAvaEaseCurveSubsystem, Warning, TEXT("UAvaEaseCurveSubsystem::RenamePreset() - Failed to rename preset %s from category %s Json file"), *InPreset, *InCategory);
+		return false;
+	}
+
+	if (Presets.Contains(InCategory))
+	{
+		TSharedPtr<FAvaEaseCurvePreset> PresetToRename;
+		for (const TSharedPtr<FAvaEaseCurvePreset>& PresetPtr : Presets[InCategory])
+		{
+			if (PresetPtr->Name.Equals(InPreset, ESearchCase::IgnoreCase))
+			{
+				PresetToRename = PresetPtr;
+				break;
+			}
+		}
+		if (PresetToRename.IsValid())
+		{
+			Presets[InCategory].Remove(PresetToRename);
+
+			PresetToRename->Name = InNewPresetName;
+
+			Presets[InCategory].Add(PresetToRename);
+		}
+	}
+
+	return true;
+}
diff --git a/Engine/Plugins/Experimental/Avalanche/Source/AvalancheSequencer/Private/EaseCurveTool/AvaEaseCurveSubsystem.h b/Engine/Plugins/Experimental/Avalanche/Source/AvalancheSequencer/Private/EaseCurveTool/AvaEaseCurveSubsystem.h
new file mode 100644
--- /dev/null
+++ b/Engine/Plugins/Experimental/Avalanche/Source/AvalancheSequencer/Private/EaseCurveTool/AvaEaseCurveSubsystem.h
@@ -0,0 +1,58 @@
+// Copyright Epic Games, Inc. All Rights Reserved.
+
+#pragma once
+
+#include "AvaEaseCurveTangents.h"
+#include "Containers/Array.h"
+#include "EditorSubsystem.h"
+#include "Templates/SharedPointer.h"
+#include "AvaEaseCurveSubsystem.generated.h"
+
+struct FAvaEaseCurvePreset;
+
+UCLASS()
+class UAvaEaseCurveSubsystem : public UEditorSubsystem
+{
+	GENERATED_BODY()
+
+public:
+	static UAvaEaseCurveSubsystem& Get();
+
+	static void ExploreJsonPresetsFolder();
+
+	void ReloadPresetsFromJson();
+
+	TArray<FString> GetEaseCurveCategories() const;
+
+	TArray<TSharedPtr<FAvaEaseCurvePreset>> GetEaseCurvePresets() const;
+	TArray<TSharedPtr<FAvaEaseCurvePreset>> GetEaseCurvePresets(const FString& InCategory) const;
+
+	bool DoesPresetExist(const FString& InName) const;
+	bool DoesPresetExist(const FString& InCategory, const FString& InName) const;
+
+	TSharedPtr<FAvaEaseCurvePreset> AddPreset(FAvaEaseCurvePreset InPreset);
+	TSharedPtr<FAvaEaseCurvePreset> AddPreset(const FString& InName, const FAvaEaseCurveTangents& InTangents);
+
+	bool RemovePreset(const FAvaEaseCurvePreset& InPreset);
+	bool RemovePreset(const FString& InCategory, const FString& InName);
+
+	bool ChangePresetCategory(const TSharedPtr<FAvaEaseCurvePreset>& InPreset, const FString& InNewCategory) const;
+
+	TSharedPtr<FAvaEaseCurvePreset> FindPreset(const FString& InName);
+	TSharedPtr<FAvaEaseCurvePreset> FindPresetByTangents(const FAvaEaseCurveTangents& InTangents);
+
+	bool DoesPresetCategoryExist(const FString& InCategory);
+	bool RenamePresetCategory(const FString& InCategory, const FString& InNewCategoryName);
+	bool AddNewPresetCategory();
+	bool RemovePresetCategory(const FString& InCategory);
+
+	bool RenamePreset(const FString& InCategory, const FString& InPreset, const FString& InNewPresetName);
+
+	//~ Begin USubsystem
+	virtual void Initialize(FSubsystemCollectionBase& InCollection) override;
+	//~ End USubsystem
+
+protected:
+	/** Cached presets from Json files. Each Json file represents a category. */
+	TSortedMap<FString, TArray<TSharedPtr<FAvaEaseCurvePreset>>> Presets;
+};
diff --git a/Engine/Plugins/Experimental/Avalanche/Source/AvalancheSequencer/Private/EaseCurveTool/AvaEaseCurveTangents.cpp b/Engine/Plugins/Experimental/Avalanche/Source/AvalancheSequencer/Private/EaseCurveTool/AvaEaseCurveTangents.cpp
new file mode 100644
--- /dev/null
+++ b/Engine/Plugins/Experimental/Avalanche/Source/AvalancheSequencer/Private/EaseCurveTool/AvaEaseCurveTangents.cpp
@@ -0,0 +1,306 @@
+// Copyright Epic Games, Inc. All Rights Reserved.
+
+#include "AvaEaseCurveTangents.h"
+#include "EaseCurveTool/Widgets/SAvaEaseCurveEditor.h"
+#include "Internationalization/Text.h"
+#include "Misc/DefaultValueHelper.h"
+#include "Misc/FrameRate.h"
+#include "Misc/QualifiedFrameTime.h"
+
+#define LOCTEXT_NAMESPACE "AvaEaseCurveTangents"
+
+bool FAvaEaseCurveTangents::FromString(const FString& InString, FAvaEaseCurveTangents& OutTangents)
+{
+	const FString TangentsString = InString.Replace(TEXT(" "), TEXT(""));
+
+	TArray<FString> ValueStrings;
+	if (TangentsString.ParseIntoArray(ValueStrings, TEXT(",")) != 4)
+	{
+		return false;
+	}
+
+	// Convert strings to doubles
+	TArray<double> Values = { 0.0, 0.0, 0.0, 0.0 };
+	for (int32 Index = 0; Index < 4; ++Index)
+	{
+		if (!FDefaultValueHelper::ParseDouble(ValueStrings[Index], Values[Index]))
+		{
+			return false;
+		}
+	}
+
+	// Convert four cubic bezier points to two points/tangents
+	OutTangents.FromCubicBezier(Values);
+
+	return true;
+}
+
+bool FAvaEaseCurveTangents::CanParseString(const FString& InString)
+{
+	FAvaEaseCurveTangents Tangents;
+	return FromString(InString, Tangents);
+}
+
+FNumberFormattingOptions FAvaEaseCurveTangents::DefaultNumberFormattingOptions()
+{
+	FNumberFormattingOptions NumberFormat;
+	NumberFormat.MinimumIntegralDigits = 1;
+	NumberFormat.MinimumFractionalDigits = 2;
+	NumberFormat.MaximumFractionalDigits = 2;
+	NumberFormat.UseGrouping = false;
+	return NumberFormat;
+}
+
+FAvaEaseCurveTangents::FAvaEaseCurveTangents(const FString& InTangentsString)
+{
+	FAvaEaseCurveTangents::FromString(InTangentsString, *this);
+}
+
+FText FAvaEaseCurveTangents::ToDisplayText() const
+{
+	const FNumberFormattingOptions NumberFormat = DefaultNumberFormattingOptions();
+	const FText StartTangentText = FText::AsNumber(Start, &NumberFormat);
+	const FText StartTangentWeightText = FText::AsNumber(StartWeight, &NumberFormat);
+	const FText EndTangentText = FText::AsNumber(End, &NumberFormat);
+	const FText EndTangentWeightText = FText::AsNumber(EndWeight, &NumberFormat);
+	return FText::Format(LOCTEXT("TangentText", "{0}, {1} - {2}, {3}")
+		, { StartTangentText, StartTangentWeightText, EndTangentText, EndTangentWeightText });
+}
+
+FString FAvaEaseCurveTangents::ToDisplayString() const
+{
+	return ToDisplayText().ToString();
+}
+
+FString FAvaEaseCurveTangents::ToJson() const
+{
+	// Convert two points/tangents to four cubic bezier points 
+	TArray<double> CubicBezierPoints;
+	ToCubicBezier(CubicBezierPoints);
+
+	const FNumberFormattingOptions NumberFormat = DefaultNumberFormattingOptions();
+	const FText PointAText = FText::AsNumber(CubicBezierPoints[0], &NumberFormat);
+	const FText PointBText = FText::AsNumber(CubicBezierPoints[1], &NumberFormat);
+	const FText PointCText = FText::AsNumber(CubicBezierPoints[2], &NumberFormat);
+	const FText PointDText = FText::AsNumber(CubicBezierPoints[3], &NumberFormat);
+
+	return FString::Printf(TEXT("%s, %s, %s, %s")
+		, *PointAText.ToString(), *PointBText.ToString(), *PointCText.ToString(), *PointDText.ToString());
+}
+
+FText FAvaEaseCurveTangents::GetStartTangentText() const
+{
+	const FNumberFormattingOptions NumberFormat = DefaultNumberFormattingOptions();
+	const FText StartTangentText = FText::AsNumber(Start, &NumberFormat);
+	const FText StartTangentWeightText = FText::AsNumber(StartWeight, &NumberFormat);
+	return FText::Format(LOCTEXT("StartTangentText", "Start: {0}, {1}"), { StartTangentText, StartTangentWeightText });
+}
+
+FText FAvaEaseCurveTangents::GetEndTangentText() const
+{
+	const FNumberFormattingOptions NumberFormat = DefaultNumberFormattingOptions();
+	const FText EndTangentText = FText::AsNumber(End, &NumberFormat);
+	const FText EndTangentWeightText = FText::AsNumber(EndWeight, &NumberFormat);
+	return FText::Format(LOCTEXT("EndTangentText", "End: {0}, {1}"), { EndTangentText, EndTangentWeightText });
+}
+
+FText FAvaEaseCurveTangents::GetCubicBezierText() const
+{
+	return FText::FromString(ToJson());
+}
+
+bool FAvaEaseCurveTangents::FromCubicBezier(const TArray<double>& InPoints)
+{
+	if (InPoints.Num() != 4)
+	{
+		return false;
+	}
+
+	const FVector2d StartDir(InPoints[0], InPoints[1]);
+	const FVector2d EndDir(1.0 - InPoints[2], 1.0 - InPoints[3]);
+
+	Start = SAvaEaseCurveEditor::CalcTangent(StartDir);
+	End = SAvaEaseCurveEditor::CalcTangent(EndDir);
+
+	StartWeight = StartDir.Size();
+	EndWeight = EndDir.Size();
+
+	return true;
+}
+
+bool FAvaEaseCurveTangents::ToCubicBezier(TArray<double>& OutPoints) const
+{
+	const FVector2d StartDir = SAvaEaseCurveEditor::CalcTangentDir(Start) * StartWeight;
+	const FVector2d EndDir = SAvaEaseCurveEditor::CalcTangentDir(End) * EndWeight;
+
+	OutPoints.Reset(4);
+	OutPoints.Add(StartDir.X);
+	OutPoints.Add(StartDir.Y);
+	OutPoints.Add(1.0 - EndDir.X);
+	OutPoints.Add(1.0 - EndDir.Y);
+
+	return true;
+}
+
+void FAvaEaseCurveTangents::Normalize(const FFrameNumber& InFrameNumber, const double InValue
+	, const FFrameNumber& InNextFrameNumber, const double InNextValue
+	, const FFrameRate& InDisplayRate, const FFrameRate& InTickResolution)
+{
+	// Convert frame time from tick resolution to display rate
+	const FFrameTime FrameTimeDifference = InNextFrameNumber - InFrameNumber;
+	const FQualifiedFrameTime QualifiedFrameTime(FrameTimeDifference, InTickResolution);
+	const FFrameTime ConvertedFrameTime = QualifiedFrameTime.ConvertTo(InDisplayRate);
+
+	// Create time/value range scale factor
+	FVector2d TimeValueRange;
+	TimeValueRange.X = InDisplayRate.AsSeconds(ConvertedFrameTime);
+	TimeValueRange.Y = FMath::Abs(InNextValue - InValue);
+
+	const double ScaleFactor = InTickResolution.AsDecimal();
+
+	// Convert tangent angles to grid coordinates
+	FVector2d StartDir = SAvaEaseCurveEditor::CalcTangentDir(Start * ScaleFactor);
+	FVector2d EndDir = SAvaEaseCurveEditor::CalcTangentDir(End * ScaleFactor);
+
+	StartDir *= StartWeight;
+	EndDir *= EndWeight;
+
+	// Scale down tangent grid coordinates to normalized range of InTimeValueRange
+	auto SafeDivide = [](const double InNumerator, const double InDenominator)
+		{
+			return FMath::IsNearlyZero(InDenominator) ? InNumerator : (InNumerator / InDenominator);
+		};
+	StartDir.X = SafeDivide(StartDir.X, TimeValueRange.X);
+	StartDir.Y = SafeDivide(StartDir.Y, TimeValueRange.Y);
+	EndDir.X = SafeDivide(EndDir.X, TimeValueRange.X);
+	EndDir.Y = SafeDivide(EndDir.Y, TimeValueRange.Y);
+
+	// Set new weights and tangents
+	StartWeight = StartDir.Size();
+	EndWeight = EndDir.Size();
+
+	StartDir.Normalize();
+	EndDir.Normalize();
+
+	Start = SAvaEaseCurveEditor::CalcTangent(StartDir);
+	End = SAvaEaseCurveEditor::CalcTangent(EndDir);
+}
+
+void FAvaEaseCurveTangents::ScaleUp(const FFrameNumber& InFrameNumber, const double InValue
+	, const FFrameNumber& InNextFrameNumber, const double InNextValue
+	, const FFrameRate& InDisplayRate, const FFrameRate& InTickResolution)
+{
+	// Convert frame time from tick resolution to display rate
+	const FFrameTime FrameTimeDifference = InNextFrameNumber - InFrameNumber;
+	const FQualifiedFrameTime QualifiedFrameTime(FrameTimeDifference, InTickResolution);
+	const FFrameTime ConvertedFrameTime = QualifiedFrameTime.ConvertTo(InDisplayRate);
+
+	// Create time/value range scale factor
+	FVector2d TimeValueRange;
+	TimeValueRange.X = InDisplayRate.AsSeconds(ConvertedFrameTime);
+	TimeValueRange.Y = FMath::Abs(InNextValue - InValue);
+
+	FVector2d StartDir = SAvaEaseCurveEditor::CalcTangentDir(Start);
+	FVector2d EndDir = SAvaEaseCurveEditor::CalcTangentDir(End);
+
+	StartDir *= StartWeight;
+	EndDir *= EndWeight;
+
+	StartDir *= TimeValueRange;
+	EndDir *= TimeValueRange;
+
+	// Set new weights
+	StartWeight = StartDir.Size();
+	EndWeight = EndDir.Size();
+
+	// Normalize vector and set new tangents
+	FVector2d NormalizedStartDir = StartDir;
+	FVector2d NormalizedEndDir = EndDir;
+
+	const double ScaleFactor = InTickResolution.AsDecimal();
+	NormalizedStartDir.X *= ScaleFactor;
+	NormalizedEndDir.X *= ScaleFactor;
+
+	NormalizedStartDir.Normalize();
+	NormalizedEndDir.Normalize();
+
+	Start = SAvaEaseCurveEditor::CalcTangent(NormalizedStartDir);
+	End = SAvaEaseCurveEditor::CalcTangent(NormalizedEndDir);
+}
+
+FAvaEaseCurveTangents FAvaEaseCurveTangents::Average(const TArray<FAvaEaseCurveTangents>& InTangentArray)
+{
+	const int32 KeySetTangentsCount = InTangentArray.Num();
+	if (KeySetTangentsCount == 0)
+	{
+		return FAvaEaseCurveTangents();
+	}
+
+	double TotalStartTangent = 0.0;
+	double TotalStartWeight = 0.0;
+	double TotalEndTangent = 0.0;
+	double TotalEndWeight = 0.0;
+
+	for (const FAvaEaseCurveTangents& Tangent : InTangentArray)
+	{
+		TotalStartTangent += Tangent.Start;
+		TotalStartWeight += Tangent.StartWeight;
+		TotalEndTangent += Tangent.End;
+		TotalEndWeight += Tangent.EndWeight;
+	}
+
+	FAvaEaseCurveTangents AverageTangents;
+	AverageTangents.Start = TotalStartTangent / KeySetTangentsCount;
+	AverageTangents.StartWeight = TotalStartWeight / KeySetTangentsCount;
+	AverageTangents.End = TotalEndTangent / KeySetTangentsCount;
+	AverageTangents.EndWeight = TotalEndWeight / KeySetTangentsCount;
+
+	return AverageTangents;
+}
+
+double FAvaEaseCurveTangents::CalculateCurveLength(const int32 SampleCount) const
+{
+	TArray<double> CubicBezierPoints;
+	ToCubicBezier(CubicBezierPoints);
+
+	const TArray<FVector2d> Points = { 
+			FVector2d(0.0, 0.0), 
+			FVector2d(CubicBezierPoints[0], CubicBezierPoints[1]), 
+			FVector2d(CubicBezierPoints[2], CubicBezierPoints[3]), 
+			FVector2d(1.0, 1.0)
+		};
+
+	auto EvaluatePoint = [](const double InTime, double InStart, double InControlA, double InControlB, double InEnd) -> double
+		{
+			return InStart * (1.0 - InTime) * (1.0 - InTime) * (1.0 - InTime)
+				+ 3.0 * InControlA * (1.0 - InTime) * (1.0 - InTime) * InTime
+				+ 3.0 * InControlB * (1.0 - InTime) * InTime * InTime
+				+ InEnd * InTime * InTime * InTime;
+		};
+
+	// Sample points along the curve and use those points to calculate the length of the curve
+	FVector2d SamplePoint;
+	FVector2d PreviousSamplePoint;
+	double Length = 0.0;
+
+	for (int32 Index = 0; Index <= SampleCount; Index++)
+	{
+		const double Time = (double)Index / (double)SampleCount;
+				
+		SamplePoint.X = EvaluatePoint(Time, Points[0].X, Points[1].X, Points[2].X, Points[3].X);
+		SamplePoint.Y = EvaluatePoint(Time, Points[0].Y, Points[1].Y, Points[2].Y, Points[3].Y);
+
+		if (Index > 0)
+		{
+			const double DifferenceX = SamplePoint.X - PreviousSamplePoint.X;
+			const double DifferenceY = SamplePoint.Y - PreviousSamplePoint.Y;
+			Length += FMath::Sqrt(DifferenceX * DifferenceX + DifferenceY * DifferenceY);
+		}
+
+		PreviousSamplePoint = SamplePoint;
+	}
+
+	return Length;
+}
+
+#undef LOCTEXT_NAMESPACE
diff --git a/Engine/Plugins/Experimental/Avalanche/Source/AvalancheSequencer/Private/EaseCurveTool/AvaEaseCurveTangents.h b/Engine/Plugins/Experimental/Avalanche/Source/AvalancheSequencer/Private/EaseCurveTool/AvaEaseCurveTangents.h
new file mode 100644
--- /dev/null
+++ b/Engine/Plugins/Experimental/Avalanche/Source/AvalancheSequencer/Private/EaseCurveTool/AvaEaseCurveTangents.h
@@ -0,0 +1,94 @@
+// Copyright Epic Games, Inc. All Rights Reserved.
+
+#pragma once
+
+#include "Curves/RichCurve.h"
+#include "Channels/MovieSceneDoubleChannel.h"
+#include "AvaEaseCurveTangents.generated.h"
+
+class FString;
+class FText;
+struct FFrameRate;
+
+/** Represents a curve constructed of two tangents and their weights, one for each end of the curve. */
+USTRUCT()
+struct FAvaEaseCurveTangents
+{
+	GENERATED_BODY()
+
+public:
+	static bool FromString(const FString& InString, FAvaEaseCurveTangents& OutTangents);
+	static bool CanParseString(const FString& InString);
+
+	static FNumberFormattingOptions DefaultNumberFormattingOptions();
+
+	static FAvaEaseCurveTangents Average(const TArray<FAvaEaseCurveTangents>& InTangentArray);
+
+	FAvaEaseCurveTangents() {}
+	FAvaEaseCurveTangents(const double InStart, const double InStartWeight, const double InEnd, const double InEndWeight)
+		: Start(InStart), StartWeight(InStartWeight)
+		, End(InEnd), EndWeight(InEndWeight)
+	{}
+	FAvaEaseCurveTangents(const FRichCurveKey& InRichCurveKey)
+		: Start(InRichCurveKey.LeaveTangent), StartWeight(InRichCurveKey.LeaveTangentWeight)
+		, End(InRichCurveKey.ArriveTangent), EndWeight(InRichCurveKey.ArriveTangentWeight)
+	{}
+	FAvaEaseCurveTangents(const FRichCurveKey& InStartRichCurveKey, const FRichCurveKey& InEndRichCurveKey)
+		: Start(InStartRichCurveKey.LeaveTangent), StartWeight(InStartRichCurveKey.LeaveTangentWeight)
+		, End(InEndRichCurveKey.ArriveTangent), EndWeight(InEndRichCurveKey.ArriveTangentWeight)
+	{}
+	FAvaEaseCurveTangents(const FMovieSceneDoubleValue& InMovieSceneDoubleValue)
+		: Start(InMovieSceneDoubleValue.Tangent.LeaveTangent), StartWeight(InMovieSceneDoubleValue.Tangent.LeaveTangentWeight)
+		, End(InMovieSceneDoubleValue.Tangent.ArriveTangent), EndWeight(InMovieSceneDoubleValue.Tangent.ArriveTangentWeight)
+	{}
+	FAvaEaseCurveTangents(const FMovieSceneDoubleValue& InStartMovieSceneDoubleValue, const FMovieSceneDoubleValue& InEndMovieSceneDoubleValue)
+		: Start(InStartMovieSceneDoubleValue.Tangent.LeaveTangent), StartWeight(InStartMovieSceneDoubleValue.Tangent.LeaveTangentWeight)
+		, End(InEndMovieSceneDoubleValue.Tangent.ArriveTangent), EndWeight(InEndMovieSceneDoubleValue.Tangent.ArriveTangentWeight)
+	{}
+	/** Constructor from string consisting of cubic bezier points. Ex. "0.45, 34.0, 0.0, 3.45" */
+	explicit FAvaEaseCurveTangents(const FString& InTangentsString);
+
+	FORCEINLINE bool operator==(const FAvaEaseCurveTangents& InRhs) const
+	{
+		return Start == InRhs.Start && StartWeight == InRhs.StartWeight && End == InRhs.End && EndWeight == InRhs.EndWeight;
+	}
+	FORCEINLINE bool operator!=(const FAvaEaseCurveTangents& InRhs) const
+	{
+		return !(*this == InRhs);
+	}
+
+	FText ToDisplayText() const;
+	FString ToDisplayString() const;
+
+	FString ToJson() const;
+
+	FText GetStartTangentText() const;
+	FText GetEndTangentText() const;
+	FText GetCubicBezierText() const;
+
+	bool FromCubicBezier(const TArray<double>& InPoints);
+	bool ToCubicBezier(TArray<double>& OutPoints) const;
+
+	void Normalize(const FFrameNumber& InFrameNumber, const double InValue
+		, const FFrameNumber& InNextFrameNumber, const double InNextValue
+		, const FFrameRate& InDisplayRate, const FFrameRate& InTickResolution);
+
+	void ScaleUp(const FFrameNumber& InFrameNumber, const double InValue
+		, const FFrameNumber& InNextFrameNumber, const double InNextValue
+		, const FFrameRate& InDisplayRate, const FFrameRate& InTickResolution);
+
+	/** Calculates the length of this curve. The higher the sample count, the higher the calculation precision. */
+	double CalculateCurveLength(const int32 InSampleCount = 10) const;
+
+	UPROPERTY()
+	double Start = 0.0;
+
+	UPROPERTY()
+	double StartWeight = 0.0;
+
+	UPROPERTY()
+	double End = 0.0;
+
+	UPROPERTY()
+	double EndWeight = 0.0;
+};
diff --git a/Engine/Plugins/Experimental/Avalanche/Source/AvalancheSequencer/Private/EaseCurveTool/AvaEaseCurveTool.cpp b/Engine/Plugins/Experimental/Avalanche/Source/AvalancheSequencer/Private/EaseCurveTool/AvaEaseCurveTool.cpp
new file mode 100644
--- /dev/null
+++ b/Engine/Plugins/Experimental/Avalanche/Source/AvalancheSequencer/Private/EaseCurveTool/AvaEaseCurveTool.cpp
@@ -0,0 +1,737 @@
+// Copyright Epic Games, Inc. All Rights Reserved.
+
+#include "AvaEaseCurveTool.h"
+#include "AssetRegistry/AssetRegistryModule.h"
+#include "AssetToolsModule.h"
+#include "AvaSequencer.h"
+#include "CurveEditor.h"
+#include "EaseCurveTool/AvaEaseCurveToolSettings.h"
+#include "EaseCurveTool/Widgets/SAvaEaseCurveEditor.h"
+#include "EaseCurveTool/Widgets/SAvaEaseCurveTool.h"
+#include "Dialogs/DlgPickAssetPath.h"
+#include "Factories/CurveFactory.h"
+#include "Framework/Notifications/NotificationManager.h"
+#include "HAL/PlatformApplicationMisc.h"
+#include "IKeyArea.h"
+#include "ISequencer.h"
+#include "ISettingsModule.h"
+#include "Math/UnrealMathUtility.h"
+#include "MVVM/CurveEditorExtension.h"
+#include "MVVM/Selection/Selection.h"
+#include "MVVM/ViewModels/ChannelModel.h"
+#include "MVVM/ViewModels/SequencerEditorViewModel.h"
+#include "Settings/AvaSequencerSettings.h"
+#include "Widgets/Notifications/SNotificationList.h"
+
+using namespace UE::Sequencer;
+
+#define LOCTEXT_NAMESPACE "CurveEaseTool"
+
+FAvaEaseCurveTool::FAvaEaseCurveTool(const TSharedRef<FAvaSequencer>& InSequencer)
+	: AvaSequencerWeak(InSequencer)
+{
+	EaseCurve = NewObject<UAvaEaseCurve>(GetTransientPackage(), NAME_None, RF_Transactional);
+
+	const TSharedRef<ISequencer> Sequencer = InSequencer->GetSequencer();
+
+	if (const TSharedPtr<FSequencerEditorViewModel> SequencerViewModel = Sequencer->GetViewModel())
+	{
+		SequencerSelectionWeak = SequencerViewModel->GetSelection();
+
+		if (const FCurveEditorExtension* CurveEditorExtension = SequencerViewModel->CastDynamic<FCurveEditorExtension>())
+		{
+			SequencerCurveEditorWeak = CurveEditorExtension->GetCurveEditor();
+			check(SequencerCurveEditorWeak.IsValid());
+		}
+
+		CacheSelectionData(SequencerViewModel->GetSelection().ToSharedRef());
+	}
+
+	UpdateEaseCurveFromSequencerKeySelections();
+}
+
+void FAvaEaseCurveTool::CacheSelectionData(const TSharedRef<FSequencerSelection>& InSequencerSelection)
+{
+	if (!SequencerCurveEditorWeak.IsValid())
+	{
+		return;
+	}
+
+	TSharedPtr<FCurveEditor> SequencerCurveEditor = SequencerCurveEditorWeak.Pin();
+
+	CachedChannelKeyData.Reset();
+	bCachedAllChannelSingleKeySelections = true;
+	bCachedIsLastOnlySelectedKey = false;
+
+	for (const FKeyHandle Key : InSequencerSelection->KeySelection)
+	{
+		if (Key == FKeyHandle::Invalid())
+		{
+			continue;
+		}
+
+		TViewModelPtr<FChannelModel> ChannelModel = InSequencerSelection->KeySelection.GetModelForKey(Key);
+		if (!ChannelModel.IsValid())
+		{
+			continue;
+		}
+
+		TSharedPtr<IKeyArea> KeyArea = ChannelModel->GetKeyArea();
+		if (!KeyArea.IsValid())
+		{
+			continue;
+		}
+
+		FMovieSceneChannelHandle ChannelHandle = KeyArea->GetChannel();
+		if (ChannelHandle.GetChannelTypeName() != FMovieSceneDoubleChannel::StaticStruct()->GetFName())
+		{
+			continue;
+		}
+
+		FCachedKeyData& Entry = CachedChannelKeyData.FindOrAdd(ChannelModel->GetChannelName());
+		Entry.ChannelModel = ChannelModel;
+		Entry.DoubleChannel = static_cast<FMovieSceneDoubleChannel*>(ChannelHandle.Get());
+		Entry.Section = ChannelModel->GetSection();
+		Entry.KeyHandles.Add(Key);
+
+		TotalSelectedKeys++;
+
+		const TArrayView<const FFrameNumber> ChannelTimes = Entry.DoubleChannel->GetTimes();
+		const int32 AllKeyCount = ChannelTimes.Num();
+		const int32 SelectedKeyCount = Entry.KeyHandles.Num();
+
+		if (SelectedKeyCount == 1 && TotalSelectedKeys == 1)
+		{
+			const int32 KeyIndex = Entry.DoubleChannel->GetIndex(Entry.KeyHandles[0]);
+			if (KeyIndex == AllKeyCount - 1)
+			{
+				bCachedIsLastOnlySelectedKey = true;
+			}
+		}
+		else if (SelectedKeyCount > 1)
+		{
+			bCachedAllChannelSingleKeySelections = false;
+		}
+	}
+}
+
+void FAvaEaseCurveTool::AddReferencedObjects(FReferenceCollector& Collector)
+{
+	Collector.AddReferencedObject(EaseCurve);
+}
+
+FString FAvaEaseCurveTool::GetReferencerName() const
+{
+	return TEXT("AvaCurveEaseTool");
+}
+
+TSharedRef<SWidget> FAvaEaseCurveTool::GenerateWidget()
+{
+	if (!CurveEaseToolWidget.IsValid())
+	{
+		CurveEaseToolWidget = SNew(SAvaEaseCurveTool, SharedThis(this))
+			.ToolMode(EMode::DoubleKeyEdit)
+			.Visibility(this, &FAvaEaseCurveTool::GetVisibility)
+			.ToolOperation(this, &FAvaEaseCurveTool::GetToolOperation);
+	}
+	return CurveEaseToolWidget.ToSharedRef();
+}
+
+EVisibility FAvaEaseCurveTool::GetVisibility() const
+{
+	return (TotalSelectedKeys > 0 && !bCachedIsLastOnlySelectedKey) ? EVisibility::Visible : EVisibility::Collapsed;
+}
+
+TObjectPtr<UAvaEaseCurve> FAvaEaseCurveTool::GetToolCurve() const
+{
+	return EaseCurve;
+}
+
+FRichCurve* FAvaEaseCurveTool::GetToolRichCurve() const
+{
+	return &EaseCurve->FloatCurve;
+}
+
+FAvaEaseCurveTangents FAvaEaseCurveTool::GetEaseCurveTangents() const
+{
+	return EaseCurve->GetTangents();
+}
+
+void FAvaEaseCurveTool::SetEaseCurveTangents(const FAvaEaseCurveTangents& InTangents, const bool bInBroadcastUpdate, const bool bInSetSequencerTangents)
+{
+	EaseCurve->SetTangents(InTangents);
+
+	if (bInBroadcastUpdate)
+	{
+		EaseCurve->BroadcastUpdate();
+	}
+
+	if (bInSetSequencerTangents)
+	{
+		SetSequencerKeySelectionTangents(InTangents);
+	}
+}
+
+void FAvaEaseCurveTool::ResetEaseCurveTangents(const bool bInStartTangent, const bool bInEndTangent)
+{
+	FText TransactionText;
+	if (bInStartTangent && !bInEndTangent)
+	{
+		TransactionText = LOCTEXT("ResetStartTangents", "Reset Start Tangents");
+	}
+	else if (!bInStartTangent && bInEndTangent)
+	{
+		TransactionText = LOCTEXT("ResetEndTangents", "Reset End Tangents");
+	}
+	else
+	{
+		TransactionText = LOCTEXT("ResetTangents", "Reset Tangents");
+	}
+	const FScopedTransaction Transaction(TransactionText);
+	EaseCurve->ModifyOwner();
+
+	FAvaEaseCurveTangents NewTangents = GetEaseCurveTangents();
+	if (bInStartTangent)
+	{
+		NewTangents.Start = 0.f;
+		NewTangents.StartWeight = 0.f;
+	}
+	if (bInEndTangent)
+	{
+		NewTangents.End = 0.f;
+		NewTangents.EndWeight = 0.f;
+	}
+	SetEaseCurveTangents(NewTangents, true, true);
+}
+
+void FAvaEaseCurveTool::FlattenOrStraightenTangents(const bool bInStartTangent, const bool bInEndTangent, const bool bInFlattenTangents) const
+{
+	FText TransactionText;
+	if (bInFlattenTangents)
+	{
+		if (bInStartTangent && !bInEndTangent)
+		{
+			TransactionText = LOCTEXT("FlattenStartTangents", "Flatten Start Tangents");
+		}
+		else if (!bInStartTangent && bInEndTangent)
+		{
+			TransactionText = LOCTEXT("FlattenEndTangents", "Flatten End Tangents");
+		}
+		else
+		{
+			TransactionText = LOCTEXT("FlattenTangents", "Flatten Tangents");
+		}
+	}
+	else
+	{
+		if (bInStartTangent && !bInEndTangent)
+		{
+			TransactionText = LOCTEXT("StraightenStartTangents", "Straighten Start Tangents");
+		}
+		else if (!bInStartTangent && bInEndTangent)
+		{
+			TransactionText = LOCTEXT("StraightenEndTangents", "Straighten End Tangents");
+		}
+		else
+		{
+			TransactionText = LOCTEXT("StraightenTangents", "Straighten Tangents");
+		}
+	}
+	const FScopedTransaction Transaction(TransactionText);
+	EaseCurve->ModifyOwner();
+
+	if (bInStartTangent)
+	{
+		EaseCurve->FlattenOrStraightenTangents(EaseCurve->GetStartKeyHandle(), bInFlattenTangents);
+	}
+	if (bInEndTangent)
+	{
+		EaseCurve->FlattenOrStraightenTangents(EaseCurve->GetEndKeyHandle(), bInFlattenTangents);
+	}
+
+	EaseCurve->BroadcastUpdate();
+}
+
+void FAvaEaseCurveTool::SetSequencerKeySelectionTangents(const FAvaEaseCurveTangents& InTangents)
+{
+	if (TotalSelectedKeys == 0)
+	{
+		return;
+	}
+
+	const FFrameRate DisplayRate = GetDisplayRate();
+	const FFrameRate TickResolution = GetTickResolution();
+	const UAvaEaseCurveToolSettings* const EaseCurveToolSettings = GetDefault<UAvaEaseCurveToolSettings>();
+	check(EaseCurveToolSettings);
+	const bool bAutoFlipTangents = EaseCurveToolSettings->GetAutoFlipTangents();
+
+	for (const TPair<FName, FCachedKeyData>& ChannelEntry : CachedChannelKeyData)
+	{
+		for (const FKeyHandle& KeyHandleToEdit : ChannelEntry.Value.KeyHandles)
+		{
+			TMovieSceneChannelData<FMovieSceneDoubleValue> ChannelData = ChannelEntry.Value.DoubleChannel->GetData();
+
+			const int32 KeyIndex = ChannelData.GetIndex(KeyHandleToEdit);
+			if (KeyIndex == INDEX_NONE)
+			{
+				continue;
+			}
+			
+			// InTangents should always be a normalized range, allowing for weights in the range of 0 - 10
+			const TArrayView<FFrameNumber> ChannelTimes = ChannelData.GetTimes();
+			TArrayView<FMovieSceneDoubleValue> ChannelValues = ChannelData.GetValues();
+			const int32 KeyCount = ChannelValues.Num();
+			
+			// If there is a keyframe after this keyframe that we are editing, we check if the that keyframe value is less
+			// than or greater than this keyframe value. If less, flip the tangent (if option is set).
+			int32 NextKeyIndex = INDEX_NONE;
+			bool bIncreasingValue = false;
+			if (KeyIndex + 1 < KeyCount)
+			{
+				NextKeyIndex = KeyIndex + 1;
+				bIncreasingValue = ChannelValues[NextKeyIndex].Value >= ChannelValues[KeyIndex].Value;
+			}
+
+			FAvaEaseCurveTangents ScaledTangents = InTangents;
+			if (bAutoFlipTangents && !bIncreasingValue)
+			{
+				ScaledTangents.Start *= -1.f;
+				ScaledTangents.End *= -1.f;
+			}
+
+			// Scale normalized tangents to time/value range
+			if (NextKeyIndex != INDEX_NONE)
+			{
+				ScaledTangents.ScaleUp(ChannelTimes[KeyIndex], ChannelValues[KeyIndex].Value
+					, ChannelTimes[NextKeyIndex], ChannelValues[NextKeyIndex].Value
+					, DisplayRate, TickResolution);
+			}
+
+			const FScopedTransaction Transaction(LOCTEXT("SetSequencerCurveTangents", "Set Sequencer Curve Tangents"));
+			ChannelEntry.Value.Section->Modify();
+			ChannelEntry.Value.Section->MarkAsChanged();
+
+			// Set this keys leave tangent
+			if (OperationMode == EOperation::Out || OperationMode == EOperation::InOut)
+			{
+				ChannelValues[KeyIndex].InterpMode = ERichCurveInterpMode::RCIM_Cubic;
+				ChannelValues[KeyIndex].Tangent.TangentWeightMode = ERichCurveTangentWeightMode::RCTWM_WeightedBoth;
+				ChannelValues[KeyIndex].TangentMode = ERichCurveTangentMode::RCTM_Break;
+				ChannelValues[KeyIndex].Tangent.LeaveTangent = ScaledTangents.Start;
+				ChannelValues[KeyIndex].Tangent.LeaveTangentWeight = ScaledTangents.StartWeight;
+			}
+			
+			// Set the next keys arrive tangent
+			if (NextKeyIndex != INDEX_NONE && (OperationMode == EOperation::In || OperationMode == EOperation::InOut))
+			{
+				ChannelValues[NextKeyIndex].InterpMode = ERichCurveInterpMode::RCIM_Cubic;
+				ChannelValues[NextKeyIndex].Tangent.TangentWeightMode = ERichCurveTangentWeightMode::RCTWM_WeightedBoth;
+				ChannelValues[NextKeyIndex].TangentMode = ERichCurveTangentMode::RCTM_Break;
+				ChannelValues[NextKeyIndex].Tangent.ArriveTangent = ScaledTangents.End;
+				ChannelValues[NextKeyIndex].Tangent.ArriveTangentWeight = ScaledTangents.EndWeight;
+			}
+		}
+	}
+}
+
+void FAvaEaseCurveTool::UpdateEaseCurveFromSequencerKeySelections()
+{
+	if (TotalSelectedKeys == 0)
+	{
+		return;
+	}
+	
+	const FFrameRate DisplayRate = GetDisplayRate();
+	const FFrameRate TickResolution = GetTickResolution();
+	const UAvaEaseCurveToolSettings* const EaseCurveToolSettings = GetDefault<UAvaEaseCurveToolSettings>();
+	check(EaseCurveToolSettings);
+	const bool bAutoFlipTangents = EaseCurveToolSettings->GetAutoFlipTangents();
+
+	TArray<FAvaEaseCurveTangents> KeySetTangents;
+
+	for (const TPair<FName, FCachedKeyData>& ChannelEntry : CachedChannelKeyData)
+	{
+		for (const FKeyHandle& KeyHandleToEdit : ChannelEntry.Value.KeyHandles)
+		{
+			TMovieSceneChannelData<FMovieSceneDoubleValue> ChannelData = ChannelEntry.Value.DoubleChannel->GetData();
+			const TArrayView<FFrameNumber> ChannelTimes = ChannelData.GetTimes();
+			const TArrayView<const FMovieSceneDoubleValue> ChannelValues = ChannelData.GetValues();
+			const int32 KeyCount = ChannelValues.Num();
+
+			const int32 KeyIndex = ChannelData.GetIndex(KeyHandleToEdit);
+			if (KeyIndex == INDEX_NONE)
+			{
+				continue;
+			}
+
+			if (ChannelValues.IsValidIndex(KeyIndex))
+			{
+				// If there is a keyframe after this keyframe that we are editing, we check if the that keyframe value is less
+				// than or greater than this keyframe value. If less, flip the tangent (if option is set).
+				int32 NextKeyIndex = INDEX_NONE;
+				bool bIncreasingValue = false;
+				if (KeyIndex + 1 < KeyCount)
+				{
+					NextKeyIndex = KeyIndex + 1;
+					bIncreasingValue = ChannelValues[NextKeyIndex].Value >= ChannelValues[KeyIndex].Value;
+				}
+
+				FAvaEaseCurveTangents Tangents;
+				if (ChannelValues.IsValidIndex(NextKeyIndex))
+				{
+					Tangents = FAvaEaseCurveTangents(ChannelValues[KeyIndex], ChannelValues[NextKeyIndex]);
+				}
+				else
+				{
+					Tangents = FAvaEaseCurveTangents(ChannelValues[KeyIndex].Tangent.LeaveTangent, ChannelValues[KeyIndex].Tangent.LeaveTangentWeight, 0.f, 0.f);
+				}
+
+				if (bAutoFlipTangents && !bIncreasingValue)
+				{
+					Tangents.Start *= -1.f;
+					Tangents.End *= -1.f;
+				}
+
+				// Scale time/value to normalized tangent range
+				FAvaEaseCurveTangents ScaledTangents = Tangents;
+				if (NextKeyIndex != INDEX_NONE)
+				{
+					ScaledTangents.Normalize(ChannelTimes[KeyIndex], ChannelValues[KeyIndex].Value
+						, ChannelTimes[NextKeyIndex], ChannelValues[NextKeyIndex].Value
+						, DisplayRate, TickResolution);
+				}
+
+				KeySetTangents.Add(ScaledTangents);
+			}
+		}
+	}
+
+	SetEaseCurveTangents(FAvaEaseCurveTangents::Average(KeySetTangents), true, false);
+}
+
+void FAvaEaseCurveTool::ApplyEaseCurveToSequencerKeySelections()
+{
+	SetSequencerKeySelectionTangents(GetEaseCurveTangents());
+}
+
+UCurveBase* FAvaEaseCurveTool::CreateCurveAsset() const
+{
+	const FAssetToolsModule& AssetToolsModule = FModuleManager::LoadModuleChecked<FAssetToolsModule>(TEXT("AssetTools"));
+
+	FString OutNewPackageName;
+	FString OutNewAssetName;
+	AssetToolsModule.Get().CreateUniqueAssetName(TEXT("/Game/NewCurve"), TEXT(""), OutNewPackageName, OutNewAssetName);
+
+	const TSharedRef<SDlgPickAssetPath> NewAssetDialog =
+		SNew(SDlgPickAssetPath)
+		.Title(LOCTEXT("CreateExternalCurve", "Create External Curve"))
+		.DefaultAssetPath(FText::FromString(OutNewPackageName));
+
+	if (NewAssetDialog->ShowModal() != EAppReturnType::Cancel)
+	{
+		const FString PackageName = NewAssetDialog->GetFullAssetPath().ToString();
+		const FName AssetName = FName(*NewAssetDialog->GetAssetName().ToString());
+
+		UPackage* const Package = CreatePackage(*PackageName);
+		
+		// Create curve object
+		UObject* NewCurveObject = nullptr;
+		
+		if (UCurveFactory* const CurveFactory = Cast<UCurveFactory>(NewObject<UFactory>(GetTransientPackage(), UCurveFactory::StaticClass())))
+		{
+			CurveFactory->CurveClass = UCurveFloat::StaticClass();
+			NewCurveObject = CurveFactory->FactoryCreateNew(CurveFactory->GetSupportedClass(), Package, AssetName, RF_Public | RF_Standalone, nullptr, GWarn);
+		}
+
+		if (NewCurveObject)
+		{
+			UCurveBase* AssetCurve = nullptr;
+			
+			// Copy curve data from current curve to newly create curve
+			UCurveFloat* const DestCurve = CastChecked<UCurveFloat>(NewCurveObject);
+			if (EaseCurve && DestCurve)
+			{
+				DestCurve->bIsEventCurve = false;
+
+				AssetCurve = DestCurve;
+
+				for (auto It(EaseCurve->FloatCurve.GetKeyIterator()); It; ++It)
+				{
+					const FRichCurveKey& Key = *It;
+					const FKeyHandle KeyHandle = DestCurve->FloatCurve.AddKey(Key.Time, Key.Value);
+					DestCurve->FloatCurve.GetKey(KeyHandle) = Key;
+				}
+			}
+
+			FAssetRegistryModule::AssetCreated(NewCurveObject);
+
+			Package->GetOutermost()->MarkPackageDirty();
+
+			return AssetCurve;
+		}
+	}
+
+	return nullptr;
+}
+
+FAvaEaseCurveTool::EOperation FAvaEaseCurveTool::GetToolOperation() const
+{
+	return OperationMode;
+}
+
+void FAvaEaseCurveTool::SetToolOperation(const FAvaEaseCurveTool::EOperation InNewOperation)
+{
+	OperationMode = InNewOperation;
+}
+
+bool FAvaEaseCurveTool::IsToolOperation(const FAvaEaseCurveTool::EOperation InNewOperation) const
+{
+	return OperationMode == InNewOperation;
+}
+
+bool FAvaEaseCurveTool::CanCopyTangentsToClipboard() const
+{
+	return true;
+}
+
+void FAvaEaseCurveTool::CopyTangentsToClipboard()
+{
+	FPlatformApplicationMisc::ClipboardCopy(*EaseCurve->GetTangents().ToDisplayString());
+
+	ShowNotificationMessage(LOCTEXT("EaseCurveToolTangentsCopied", "Ease Curve Tool Tangents Copied!"));
+}
+
+bool FAvaEaseCurveTool::CanPasteTangentsFromClipboard() const
+{
+	FAvaEaseCurveTangents Tangents;
+	return TangentsFromClipboardPaste(Tangents);
+}
+
+void FAvaEaseCurveTool::PasteTangentsFromClipboard()
+{
+	FAvaEaseCurveTangents Tangents;
+	if (TangentsFromClipboardPaste(Tangents))
+	{
+		EaseCurve->SetTangents(Tangents);
+	}
+}
+
+bool FAvaEaseCurveTool::TangentsFromClipboardPaste(FAvaEaseCurveTangents& OutTangents)
+{
+	FString ClipboardString;
+	FPlatformApplicationMisc::ClipboardPaste(ClipboardString);
+
+	// Expects four comma separated cubic bezier points that define the curve
+	return FAvaEaseCurveTangents::FromString(ClipboardString, OutTangents);
+}
+
+bool FAvaEaseCurveTool::IsKeyInterpMode(const ERichCurveInterpMode InInterpMode, const ERichCurveTangentMode InTangentMode)
+{
+	const FKeyHandle StartKeyHandle = EaseCurve->GetStartKeyHandle();
+	return (EaseCurve->FloatCurve.GetKeyInterpMode(StartKeyHandle) == InInterpMode
+		&& EaseCurve->FloatCurve.GetKeyTangentMode(StartKeyHandle) == InTangentMode);
+}
+
+void FAvaEaseCurveTool::SetKeyInterpMode(const ERichCurveInterpMode InInterpMode, const ERichCurveTangentMode InTangentMode)
+{
+	const FKeyHandle StartKeyHandle = EaseCurve->GetStartKeyHandle();
+
+	const FScopedTransaction Transaction(LOCTEXT("CurveEditor_SetInterpolationMode", "Select Interpolation Mode"));
+	EaseCurve->ModifyOwner();
+
+	EaseCurve->FloatCurve.SetKeyInterpMode(StartKeyHandle, InInterpMode);
+	EaseCurve->FloatCurve.SetKeyTangentMode(StartKeyHandle, InTangentMode);
+
+	if (InInterpMode != ERichCurveInterpMode::RCIM_Cubic)
+	{
+		FRichCurveKey& StartKey = EaseCurve->GetStartKey();
+		StartKey.LeaveTangentWeight = 0.f;
+		
+		FRichCurveKey& EndKey = EaseCurve->GetEndKey();
+		EndKey.ArriveTangentWeight = 0.f;
+	}
+
+	TArray<FRichCurveEditInfo> ChangedCurveEditInfos;
+	ChangedCurveEditInfos.Add(FRichCurveEditInfo(&EaseCurve->FloatCurve));
+	EaseCurve->OnCurveChanged(ChangedCurveEditInfos);
+}
+
+void FAvaEaseCurveTool::BeginTransaction(const FText& InDescription) const
+{
+	if (GEditor)
+	{
+		EaseCurve->ModifyOwnerChange();
+
+		GEditor->BeginTransaction(InDescription);
+	}
+}
+
+void FAvaEaseCurveTool::EndTransaction() const
+{
+	if (GEditor)
+	{
+		GEditor->EndTransaction();
+	}
+}
+
+void FAvaEaseCurveTool::UndoAction()
+{
+	if (GEditor && GEditor->UndoTransaction())
+	{
+		UpdateEaseCurveFromSequencerKeySelections();
+	}
+}
+
+void FAvaEaseCurveTool::RedoAction()
+{
+	if (GEditor && GEditor->RedoTransaction())
+	{
+		UpdateEaseCurveFromSequencerKeySelections();
+	}
+}
+
+void FAvaEaseCurveTool::PostUndo(bool bInSuccess)
+{
+	UndoAction();
+}
+
+void FAvaEaseCurveTool::PostRedo(bool bInSuccess)
+{
+	RedoAction();
+}
+
+void FAvaEaseCurveTool::OpenToolSettings() const
+{
+	if (ISettingsModule* SettingsModule = FModuleManager::GetModulePtr<ISettingsModule>("Settings"))
+	{
+		SettingsModule->ShowViewer(TEXT("Editor"), TEXT("Motion Design"), TEXT("Ease Curve Tool"));
+	}
+}
+
+FFrameRate FAvaEaseCurveTool::GetTickResolution() const
+{
+	if (TSharedPtr<FAvaSequencer> AvaSequencer = AvaSequencerWeak.Pin())
+	{
+		return AvaSequencer->GetSequencer()->GetFocusedTickResolution();
+	}
+
+	return FFrameRate();
+}
+
+FFrameRate FAvaEaseCurveTool::GetDisplayRate() const
+{
+	if (TSharedPtr<FAvaSequencer> AvaSequencer = AvaSequencerWeak.Pin())
+	{
+		return AvaSequencer->GetSequencer()->GetFocusedDisplayRate();
+	}
+	
+	// Fallback to using config display rate if tool is being used outside sequencer
+	return GetDefault<UAvaSequencerSettings>()->GetDisplayRate();
+}
+
+void FAvaEaseCurveTool::ShowNotificationMessage(const FText& InMessageText)
+{
+	FNotificationInfo Info(InMessageText);
+	Info.ExpireDuration = 3.0f;
+	FSlateNotificationManager::Get().AddNotification(Info);
+}
+
+void FAvaEaseCurveTool::SelectNextChannelKey()
+{
+	// This command is only available for single key selections
+	if (TotalSelectedKeys != 1 || !SequencerSelectionWeak.IsValid())
+	{
+		return;
+	}
+
+	const TSharedPtr<FSequencerSelection> SequencerSelection = SequencerSelectionWeak.Pin();
+
+	int32 SelectedKeyIndex = INDEX_NONE;
+	FKeyHandle SelectedKeyHandle;
+	FKeyHandle NextKeyHandle;
+	TViewModelPtr<FChannelModel> ChannelModel;
+
+	for (const TPair<FName, FCachedKeyData>& ChannelEntry : CachedChannelKeyData)
+	{
+		if (ChannelEntry.Value.KeyHandles.Num() == 1)
+		{
+			SelectedKeyHandle = ChannelEntry.Value.KeyHandles[0];
+
+			TMovieSceneChannelData<FMovieSceneDoubleValue> ChannelData = ChannelEntry.Value.DoubleChannel->GetData();
+			const TArrayView<const FFrameNumber> KeyTimes = ChannelData.GetTimes();
+
+			SelectedKeyIndex = ChannelData.GetIndex(SelectedKeyHandle);
+			if (KeyTimes.IsValidIndex(SelectedKeyIndex + 1))
+			{
+				NextKeyHandle = ChannelData.GetHandle(SelectedKeyIndex + 1);
+
+				ChannelModel = ChannelEntry.Value.ChannelModel;
+			}
+		}
+		break;
+	}
+
+	if (SelectedKeyIndex == INDEX_NONE
+		|| SelectedKeyHandle == FKeyHandle::Invalid()
+		|| NextKeyHandle == FKeyHandle::Invalid()
+		|| !ChannelModel.IsValid())
+	{
+		return;
+	}
+	
+	SequencerSelection->KeySelection.Deselect(SelectedKeyHandle);
+	SequencerSelection->KeySelection.Select(ChannelModel, NextKeyHandle);
+
+	UpdateEaseCurveFromSequencerKeySelections();
+}
+
+void FAvaEaseCurveTool::SelectPreviousChannelKey()
+{
+	// This command is only available for single key selections
+	if (TotalSelectedKeys != 1 || !SequencerSelectionWeak.IsValid())
+	{
+		return;
+	}
+
+	const TSharedPtr<FSequencerSelection> SequencerSelection = SequencerSelectionWeak.Pin();
+
+	int32 SelectedKeyIndex = INDEX_NONE;
+	FKeyHandle SelectedKeyHandle;
+	FKeyHandle PreviousKeyHandle;
+	TViewModelPtr<FChannelModel> ChannelModel;
+
+	for (const TPair<FName, FCachedKeyData>& ChannelEntry : CachedChannelKeyData)
+	{
+		if (ChannelEntry.Value.KeyHandles.Num() == 1)
+		{
+			SelectedKeyHandle = ChannelEntry.Value.KeyHandles[0];
+
+			TMovieSceneChannelData<FMovieSceneDoubleValue> ChannelData = ChannelEntry.Value.DoubleChannel->GetData();
+			const TArrayView<const FFrameNumber> KeyTimes = ChannelData.GetTimes();
+
+			SelectedKeyIndex = ChannelData.GetIndex(SelectedKeyHandle);
+			if (KeyTimes.IsValidIndex(SelectedKeyIndex - 1))
+			{
+				PreviousKeyHandle = ChannelData.GetHandle(SelectedKeyIndex - 1);
+
+				ChannelModel = ChannelEntry.Value.ChannelModel;
+			}
+		}
+		break;
+	}
+
+	if (SelectedKeyIndex == INDEX_NONE
+		|| SelectedKeyHandle == FKeyHandle::Invalid()
+		|| PreviousKeyHandle == FKeyHandle::Invalid()
+		|| !ChannelModel.IsValid())
+	{
+		return;
+	}
+
+	SequencerSelection->KeySelection.Deselect(SelectedKeyHandle);
+	SequencerSelection->KeySelection.Select(ChannelModel, PreviousKeyHandle);
+
+	UpdateEaseCurveFromSequencerKeySelections();
+}
+
+#undef LOCTEXT_NAMESPACE
diff --git a/Engine/Plugins/Experimental/Avalanche/Source/AvalancheSequencer/Private/EaseCurveTool/AvaEaseCurveTool.h b/Engine/Plugins/Experimental/Avalanche/Source/AvalancheSequencer/Private/EaseCurveTool/AvaEaseCurveTool.h
new file mode 100644
--- /dev/null
+++ b/Engine/Plugins/Experimental/Avalanche/Source/AvalancheSequencer/Private/EaseCurveTool/AvaEaseCurveTool.h
@@ -0,0 +1,160 @@
+// Copyright Epic Games, Inc. All Rights Reserved.
+
+#pragma once
+
+#include "EaseCurveTool/AvaEaseCurve.h"
+#include "Curves/RichCurve.h"
+#include "Curves/KeyHandle.h"
+#include "EaseCurveTool/Widgets/SAvaEaseCurvePreset.h"
+#include "Editor/Transactor.h"
+#include "EditorUndoClient.h"
+#include "ScopedTransaction.h"
+#include "Templates/SharedPointer.h"
+#include "UObject/WeakObjectPtrTemplates.h"
+#include "Widgets/DeclarativeSyntaxSupport.h"
+#include "Widgets/SCompoundWidget.h"
+
+class FAvaSequencer;
+class FCurveEditor;
+class SAvaEaseCurveTool;
+class SCurveEditor;
+class SCurveEditorPanel;
+class SEditableTextBox;
+class UCurveFloat;
+struct FMovieSceneDoubleChannel;
+
+namespace UE::Sequencer
+{
+	class FChannelModel;
+	class FSequencerSelection;
+}
+
+class FAvaEaseCurveTool
+	: public TSharedFromThis<FAvaEaseCurveTool>
+	, public FGCObject
+	, public FSelfRegisteringEditorUndoClient
+{
+public:
+	static void ShowNotificationMessage(const FText& InMessageText);
+
+	/** Returns true if the clipboard paste data contains tangent information. */
+	static bool TangentsFromClipboardPaste(FAvaEaseCurveTangents& OutTangents);
+
+	/**
+	 * Current default and only implemented is DoubleKeyEdit.
+	 */
+	enum class EMode : uint8
+	{
+		/** Edits the selected key's leave tangent and the next key's arrive tangent in the curve editor graph. */
+		DoubleKeyEdit,
+		/** Edits only the selected key.
+		 * The leave tangent in the curve editor graph will set the sequence key arrive tangent.
+		 * The arrive tangent in the curve editor graph will set the sequence key leave tangent. */
+		SingleKeyEdit
+	};
+
+	enum class EOperation : uint8
+	{
+		In,
+		Out,
+		InOut
+	};
+
+	FAvaEaseCurveTool(const TSharedRef<FAvaSequencer>& InSequencer);
+
+	TSharedRef<SWidget> GenerateWidget();
+
+	EVisibility GetVisibility() const;
+
+	TObjectPtr<UAvaEaseCurve> GetToolCurve() const;
+	FRichCurve* GetToolRichCurve() const;
+
+	FAvaEaseCurveTangents GetEaseCurveTangents() const;
+	void SetEaseCurveTangents(const FAvaEaseCurveTangents& InTangents, const bool bBroadcastUpdate, const bool bInSetSequencerTangents);
+	void ResetEaseCurveTangents(const bool bInStartTangent, const bool bInEndTangent);
+
+	void FlattenOrStraightenTangents(const bool bInStartTangent, const bool bInEndTangent, const bool bInFlattenTangents) const;
+
+	FORCEINLINE EOperation GetOperation() const { return OperationMode; }
+	FORCEINLINE void SetOperation(const EOperation InOperation) { OperationMode = InOperation; }
+
+	/** Creates a new external float curve from the internet curve editor curve. */
+	UCurveBase* CreateCurveAsset() const;
+	
+	void SetSequencerKeySelectionTangents(const FAvaEaseCurveTangents& InTangents);
+
+	/** Updates the ease curve graph view based on the active sequencer key selection. */
+	void UpdateEaseCurveFromSequencerKeySelections();
+
+	void ApplyEaseCurveToSequencerKeySelections();
+
+	EOperation GetToolOperation() const;
+	void SetToolOperation(const EOperation InNewOperation);
+	bool IsToolOperation(const EOperation InNewOperation) const;
+
+	bool CanCopyTangentsToClipboard() const;
+	void CopyTangentsToClipboard();
+	bool CanPasteTangentsFromClipboard() const;
+	void PasteTangentsFromClipboard();
+
+	bool IsKeyInterpMode(const ERichCurveInterpMode InInterpMode, const ERichCurveTangentMode InTangentMode);
+	void SetKeyInterpMode(const ERichCurveInterpMode InInterpMode, const ERichCurveTangentMode InTangentMode);
+
+	void BeginTransaction(const FText& InDescription) const;
+	void EndTransaction() const;
+
+	void UndoAction();
+	void RedoAction();
+
+	void OpenToolSettings() const;
+
+	FFrameRate GetTickResolution() const;
+	FFrameRate GetDisplayRate() const;
+
+	void SelectNextChannelKey();
+	void SelectPreviousChannelKey();
+
+	//~ Begin FGCObject
+	virtual void AddReferencedObjects(FReferenceCollector& Collector) override;
+	virtual FString GetReferencerName() const override;
+	//~ End FGCObject
+
+	//~ Begin FEditorUndoClient
+	virtual void PostUndo(bool bInSuccess);
+	virtual void PostRedo(bool bInSuccess);
+	//~ End FEditorUndoClient
+
+protected:
+	TWeakPtr<FAvaSequencer> AvaSequencerWeak;
+	TWeakPtr<UE::Sequencer::FSequencerSelection> SequencerSelectionWeak;
+	TWeakPtr<FCurveEditor> SequencerCurveEditorWeak;
+
+	TObjectPtr<UAvaEaseCurve> EaseCurve;
+
+	EMode ToolMode = EMode::DoubleKeyEdit;
+	EOperation OperationMode = EOperation::InOut;
+
+	TSharedPtr<SAvaEaseCurveTool> CurveEaseToolWidget;
+
+	FSimpleMulticastDelegate::FDelegate LastDelegate;
+	
+	/** Cached data set when a new sequencer selection is made. */
+	struct FCachedKeyData
+	{
+		TSharedPtr<UE::Sequencer::FChannelModel> ChannelModel;
+		UMovieSceneSection* Section = nullptr;
+		FMovieSceneDoubleChannel* DoubleChannel = nullptr;
+		TArray<FKeyHandle> KeyHandles;
+	};
+	TMap<FName, FCachedKeyData> CachedChannelKeyData;
+	int32 TotalSelectedKeys = 0;
+
+	/** Indicates only one selected key on each of the selected channels. */
+	bool bCachedAllChannelSingleKeySelections = false;
+
+	/** Indicates only one key selected and it is the last key of the channel. */
+	bool bCachedIsLastOnlySelectedKey = false;
+
+private:
+	void CacheSelectionData(const TSharedRef<UE::Sequencer::FSequencerSelection>& InSequencerSelection);
+};
diff --git a/Engine/Plugins/Experimental/Avalanche/Source/AvalancheSequencer/Private/EaseCurveTool/AvaEaseCurveToolCommands.cpp b/Engine/Plugins/Experimental/Avalanche/Source/AvalancheSequencer/Private/EaseCurveTool/AvaEaseCurveToolCommands.cpp
new file mode 100644
--- /dev/null
+++ b/Engine/Plugins/Experimental/Avalanche/Source/AvalancheSequencer/Private/EaseCurveTool/AvaEaseCurveToolCommands.cpp
@@ -0,0 +1,225 @@
+// Copyright Epic Games, Inc. All Rights Reserved.
+
+#include "AvaEaseCurveToolCommands.h"
+
+#define LOCTEXT_NAMESPACE "AvaCurveEaseToolCommands"
+
+void FAvaEaseCurveToolCommands::RegisterCommands()
+{
+	FUICommandInfo::MakeCommandInfo(AsShared()
+		, OpenToolSettings
+		, TEXT("OpenToolSettings")
+		, LOCTEXT("OpenToolSettings_Label", "Open Tool Settings...")
+		, LOCTEXT("OpenToolSettings_ToolTip", "Open the editor settings for the the curve ease tool.")
+		, FSlateIcon()
+		, EUserInterfaceActionType::Button
+		, FInputChord());
+
+	FUICommandInfo::MakeCommandInfo(AsShared()
+		, Refresh
+		, TEXT("Refresh")
+		, LOCTEXT("Refresh_Label", "Refresh")
+		, LOCTEXT("Refresh_ToolTip", "Refresh the ease curve editor graph tangents from the selected sequencer key(s).")
+		, FSlateIcon(FAppStyle::Get().GetStyleSetName(), TEXT("Icons.Refresh"))
+		, EUserInterfaceActionType::Button
+		, FInputChord(EKeys::F5));
+
+	FUICommandInfo::MakeCommandInfo(AsShared()
+		, Apply
+		, TEXT("Apply")
+		, LOCTEXT("Apply_Label", "Apply")
+		, LOCTEXT("Apply_ToolTip", "Applies the ease curve to the selected sequencer key(s).")
+		, FSlateIcon(FAppStyle::Get().GetStyleSetName(), TEXT("AssetEditor.Apply"))
+		, EUserInterfaceActionType::Button
+		, FInputChord());
+
+	FUICommandInfo::MakeCommandInfo(AsShared()
+		, ZoomToFit
+		, TEXT("ZoomToFit")
+		, LOCTEXT("ZoomToFit_Label", "Zoom to Fit")
+		, LOCTEXT("ZoomToFit_ToolTip", "Zoom the curve graph view to fit tangent handles.")
+		, FSlateIcon(FAppStyle::Get().GetStyleSetName(), TEXT("CurveEditorTools.SetFocusPlaybackRange"))
+		, EUserInterfaceActionType::Button
+		, FInputChord(EKeys::F));
+
+	FUICommandInfo::MakeCommandInfo(AsShared()
+		, SetOperationToEaseOut
+		, TEXT("SetOperationToEaseOut")
+		, LOCTEXT("SetOperationToEaseOut_Label", "Ease Out")
+		, LOCTEXT("SetOperationToEaseOut_ToolTip", "Set the operation of the Curve Ease Tool to Ease Out only. This will only set the Out (Leave) tangent.")
+		, FSlateIcon(FAppStyle::Get().GetStyleSetName(), TEXT("Icons.ChevronLeft"))
+		, EUserInterfaceActionType::RadioButton
+		, FInputChord(EKeys::O));
+
+	FUICommandInfo::MakeCommandInfo(AsShared()
+		, SetOperationToEaseInOut
+		, TEXT("SetOperationToEaseInOut")
+		, LOCTEXT("SetOperationToEaseInOut_Label", "Ease In Out")
+		, LOCTEXT("SetOperationToEaseInOut_ToolTip", "Set the operation of the Curve Ease Tool to both Ease Out and Ease In. This will set the Out (Leave) and In (Arrive) tangents.")
+		, FSlateIcon(FAppStyle::Get().GetStyleSetName(), TEXT("Icons.BulletPoint"))
+		, EUserInterfaceActionType::RadioButton
+		, FInputChord(EKeys::U));
+
+	FUICommandInfo::MakeCommandInfo(AsShared()
+		, SetOperationToEaseIn
+		, TEXT("SetOperationToEaseIn")
+		, LOCTEXT("SetOperationToEaseIn_Label", "Ease In")
+		, LOCTEXT("SetOperationToEaseIn_ToolTip", "Set the operation of the Curve Ease Tool to Ease In only. This will only set the In (Arrive) tangent.")
+		, FSlateIcon(FAppStyle::Get().GetStyleSetName(), TEXT("Icons.ChevronRight"))
+		, EUserInterfaceActionType::RadioButton
+		, FInputChord(EKeys::I));
+
+	FUICommandInfo::MakeCommandInfo(AsShared()
+		, ToggleGridSnap
+		, TEXT("ToggleGridSnap")
+		, LOCTEXT("ToggleGridSnap_Label", "Snap To Grid")
+		, LOCTEXT("ToggleGridSnap_ToolTip", "Toggle snapping of tangents to grid.")
+		, FSlateIcon(FAppStyle::Get().GetStyleSetName(), TEXT("Icons.Snap"))
+		, EUserInterfaceActionType::Check
+		, FInputChord(EKeys::G));
+
+	FUICommandInfo::MakeCommandInfo(AsShared()
+		, ToggleAutoFlipTangents
+		, TEXT("ToggleAutoFlipTangents")
+		, LOCTEXT("ToggleAutoFlipTangents_Label", "Auto Flip Tangents")
+		, LOCTEXT("ToggleAutoFlipTangents_ToolTip", "Toggle flipping of tangents when sequential key frame curve values are descending.")
+		, FSlateIcon(FAppStyle::Get().GetStyleSetName(), TEXT("GenericCurveEditor.InterpolationCubicAuto.Small"))
+		, EUserInterfaceActionType::Check
+		, FInputChord());
+
+	FUICommandInfo::MakeCommandInfo(AsShared()
+		, ResetTangents
+		, TEXT("ResetBothTangents")
+		, LOCTEXT("ResetBothTangents_Label", "Both Tangents")
+		, LOCTEXT("ResetBothTangents_ToolTip", "Set the In / Out tangent values and weights to 0.")
+		, FSlateIcon(FAppStyle::Get().GetStyleSetName(), TEXT("PropertyWindow.DiffersFromDefault"))
+		, EUserInterfaceActionType::Button
+		, FInputChord(EKeys::R));
+
+	FUICommandInfo::MakeCommandInfo(AsShared()
+		, ResetStartTangent
+		, TEXT("ResetStartTangent")
+		, LOCTEXT("ResetStartTangent_Label", "Start Tangent")
+		, LOCTEXT("ResetStartTangent_ToolTip", "Set the start tangent value and weight to 0.")
+		, FSlateIcon(FAppStyle::Get().GetStyleSetName(), TEXT("PropertyWindow.DiffersFromDefault"))
+		, EUserInterfaceActionType::Button
+		, FInputChord(EKeys::R, EModifierKey::Control));
+
+	FUICommandInfo::MakeCommandInfo(AsShared()
+		, ResetEndTangent
+		, TEXT("ResetEndTangent")
+		, LOCTEXT("ResetEndTangent_Label", "End Tangent")
+		, LOCTEXT("ResetEndTangent_ToolTip", "Set the end tangent value and weight to 0.")
+		, FSlateIcon(FAppStyle::Get().GetStyleSetName(), TEXT("PropertyWindow.DiffersFromDefault"))
+		, EUserInterfaceActionType::Button
+		, FInputChord(EKeys::R, EModifierKey::Alt));
+
+	FUICommandInfo::MakeCommandInfo(AsShared()
+		, FlattenTangents
+		, TEXT("FlattenBothTangents")
+		, LOCTEXT("FlattenBothTangents_Label", "Both Tangents")
+		, LOCTEXT("FlattenBothTangents_ToolTip", "Flatten the start and end tangents.")
+		, FSlateIcon(FAppStyle::Get().GetStyleSetName(), TEXT("GenericCurveEditor.FlattenTangents"))
+		, EUserInterfaceActionType::Button
+		, FInputChord(EKeys::T));
+
+	FUICommandInfo::MakeCommandInfo(AsShared()
+		, FlattenStartTangent
+		, TEXT("FlattenStartTangent")
+		, LOCTEXT("FlattenStartTangent_Label", "Start Tangent")
+		, LOCTEXT("FlattenStartTangent_ToolTip", "Flatten the start tangent.")
+		, FSlateIcon(FAppStyle::Get().GetStyleSetName(), TEXT("GenericCurveEditor.FlattenTangents"))
+		, EUserInterfaceActionType::Button
+		, FInputChord(EKeys::T, EModifierKey::Control));
+
+	FUICommandInfo::MakeCommandInfo(AsShared()
+		, FlattenEndTangent
+		, TEXT("FlattenEndTangent")
+		, LOCTEXT("FlattenEndTangent_Label", "End Tangent")
+		, LOCTEXT("FlattenEndTangent_ToolTip", "Flatten the end tangent.")
+		, FSlateIcon(FAppStyle::Get().GetStyleSetName(), TEXT("GenericCurveEditor.FlattenTangents"))
+		, EUserInterfaceActionType::Button
+		, FInputChord(EKeys::T, EModifierKey::Alt));
+
+	FUICommandInfo::MakeCommandInfo(AsShared()
+		, StraightenTangents
+		, TEXT("StraightenBothTangents")
+		, LOCTEXT("StraightenBothTangents_Label", "Both Tangents")
+		, LOCTEXT("StraightenBothTangents_ToolTip", "Straighten the start and end tangents.")
+		, FSlateIcon(FAppStyle::Get().GetStyleSetName(), TEXT("GenericCurveEditor.StraightenTangents"))
+		, EUserInterfaceActionType::Button
+		, FInputChord(EKeys::S));
+
+	FUICommandInfo::MakeCommandInfo(AsShared()
+		, StraightenStartTangent
+		, TEXT("StraightenStartTangent")
+		, LOCTEXT("StraightenStartTangent_Label", "Start Tangent")
+		, LOCTEXT("StraightenStartTangent_ToolTip", "Straighten the start tangent.")
+		, FSlateIcon(FAppStyle::Get().GetStyleSetName(), TEXT("GenericCurveEditor.StraightenTangents"))
+		, EUserInterfaceActionType::Button
+		, FInputChord(EKeys::S, EModifierKey::Control));
+
+	FUICommandInfo::MakeCommandInfo(AsShared()
+		, StraightenEndTangent
+		, TEXT("StraightenEndTangent")
+		, LOCTEXT("StraightenEndTangent_Label", "End Tangent")
+		, LOCTEXT("StraightenEndTangent_ToolTip", "Straighten the end tangent.")
+		, FSlateIcon(FAppStyle::Get().GetStyleSetName(), TEXT("GenericCurveEditor.StraightenTangents"))
+		, EUserInterfaceActionType::Button
+		, FInputChord(EKeys::S, EModifierKey::Alt));
+
+	FUICommandInfo::MakeCommandInfo(AsShared()
+		, CopyTangents
+		, TEXT("CopyTangents")
+		, LOCTEXT("CopyTangents_Label", "Copy Tangents")
+		, LOCTEXT("CopyTangents_ToolTip", "Copy the set of In / Out tangent values and weights.")
+		, FSlateIcon(FAppStyle::Get().GetStyleSetName(), TEXT("GenericCommands.Copy"))
+		, EUserInterfaceActionType::Button
+		, FInputChord(EKeys::C, EModifierKey::Control));
+
+	FUICommandInfo::MakeCommandInfo(AsShared()
+		, PasteTangents
+		, TEXT("PasteTangents")
+		, LOCTEXT("PasteTangents_Label", "Paste Tangents")
+		, LOCTEXT("PasteTangents_ToolTip", "Paste the set of In / Out tangent values and weights.")
+		, FSlateIcon(FAppStyle::Get().GetStyleSetName(), TEXT("GenericCommands.Paste"))
+		, EUserInterfaceActionType::Button
+		, FInputChord(EKeys::V, EModifierKey::Control));
+
+	FUICommandInfo::MakeCommandInfo(AsShared()
+		, CreateExternalCurveAsset
+		, TEXT("CreateExternalCurve")
+		, LOCTEXT("CreateExternalCurve_Label", "Create External Curve...")
+		, LOCTEXT("CreateExternalCurve_ToolTip", "Create an external curve asset using this internal curve.")
+		, FSlateIcon(FAppStyle::Get().GetStyleSetName(), "TimelineEditor.AddCurveAssetTrack")
+		, EUserInterfaceActionType::Button
+		, FInputChord());
+
+	UI_COMMAND(SetKeyInterpConstant, "Constant", "Constant interpolation", EUserInterfaceActionType::ToggleButton, FInputChord(EKeys::Five));
+	UI_COMMAND(SetKeyInterpLinear, "Linear", "Linear interpolation", EUserInterfaceActionType::ToggleButton, FInputChord(EKeys::Four));
+	UI_COMMAND(SetKeyInterpCubicSmartAuto, "Smart Auto", "Cubic interpolation - Smart Automatic tangents", EUserInterfaceActionType::ToggleButton, FInputChord(EKeys::Zero));
+	UI_COMMAND(SetKeyInterpCubicAuto, "Auto", "Cubic interpolation - Automatic tangents", EUserInterfaceActionType::ToggleButton, FInputChord(EKeys::One));
+	UI_COMMAND(SetKeyInterpCubicUser, "User", "Cubic interpolation - User flat tangents", EUserInterfaceActionType::ToggleButton, FInputChord(EKeys::Two));
+	UI_COMMAND(SetKeyInterpCubicBreak, "Break", "Cubic interpolation - User broken tangents", EUserInterfaceActionType::ToggleButton, FInputChord(EKeys::Three));
+	UI_COMMAND(SetKeyInterpToggleWeighted, "Weighted Tangents", "Toggle weighted tangents for cubic interpolation modes. Only supported on some curve types", EUserInterfaceActionType::ToggleButton, FInputChord(EModifierKey::Control, EKeys::W));
+
+	FUICommandInfo::MakeCommandInfo(AsShared()
+		, SelectNextChannelKey
+		, TEXT("SelectNextChannelKey")
+		, LOCTEXT("SelectNextChannelKey_Label", "Select Next Channel Key")
+		, LOCTEXT("SelectNextChannelKey_ToolTip", "Select the next key on the same channel.")
+		, FSlateIcon(FAppStyle::Get().GetStyleSetName(), TEXT("Wizard.NextIcon"))
+		, EUserInterfaceActionType::Button
+		, FInputChord(EKeys::Period, EModifierKey::Control));
+
+	FUICommandInfo::MakeCommandInfo(AsShared()
+		, SelectPreviousChannelKey
+		, TEXT("SelectsPreviousChannelKey")
+		, LOCTEXT("SelectsPreviousChannelKey_Label", "Selects Previous Channel Key")
+		, LOCTEXT("SelectsPreviousChannelKey_ToolTip", "Select the previous key on the same channel.")
+		, FSlateIcon(FAppStyle::Get().GetStyleSetName(), TEXT("Wizard.BackIcon"))
+		, EUserInterfaceActionType::Button
+		, FInputChord(EKeys::Comma, EModifierKey::Control));
+}
+
+#undef LOCTEXT_NAMESPACE
diff --git a/Engine/Plugins/Experimental/Avalanche/Source/AvalancheSequencer/Private/EaseCurveTool/AvaEaseCurveToolCommands.h b/Engine/Plugins/Experimental/Avalanche/Source/AvalancheSequencer/Private/EaseCurveTool/AvaEaseCurveToolCommands.h
new file mode 100644
--- /dev/null
+++ b/Engine/Plugins/Experimental/Avalanche/Source/AvalancheSequencer/Private/EaseCurveTool/AvaEaseCurveToolCommands.h
@@ -0,0 +1,63 @@
+// Copyright Epic Games, Inc. All Rights Reserved.
+
+#pragma once
+
+#include "Framework/Commands/Commands.h"
+#include "Styling/AppStyle.h"
+
+class FAvaEaseCurveToolCommands : public TCommands<FAvaEaseCurveToolCommands>
+{
+public:
+	FAvaEaseCurveToolCommands()
+		: TCommands<FAvaEaseCurveToolCommands>(TEXT("AvalancheCurveEaseToolCommands")
+			, NSLOCTEXT("AvaCurveEaseToolCommands", "AvaCurveEaseToolCommands", "Motion Design Curve Ease Tool Commands")
+			, NAME_None
+			, FAppStyle::GetAppStyleSetName())
+	{
+	}
+
+	virtual void RegisterCommands() override;
+
+	TSharedPtr<FUICommandInfo> OpenToolSettings;
+
+	TSharedPtr<FUICommandInfo> Refresh;
+	TSharedPtr<FUICommandInfo> Apply;
+
+	TSharedPtr<FUICommandInfo> ZoomToFit;
+
+	TSharedPtr<FUICommandInfo> SetOperationToEaseOut;
+	TSharedPtr<FUICommandInfo> SetOperationToEaseInOut;
+	TSharedPtr<FUICommandInfo> SetOperationToEaseIn;
+
+	TSharedPtr<FUICommandInfo> ToggleGridSnap;
+
+	TSharedPtr<FUICommandInfo> ToggleAutoFlipTangents;
+
+	TSharedPtr<FUICommandInfo> ResetTangents;
+	TSharedPtr<FUICommandInfo> ResetStartTangent;
+	TSharedPtr<FUICommandInfo> ResetEndTangent;
+
+	TSharedPtr<FUICommandInfo> FlattenTangents;
+	TSharedPtr<FUICommandInfo> FlattenStartTangent;
+	TSharedPtr<FUICommandInfo> FlattenEndTangent;
+
+	TSharedPtr<FUICommandInfo> StraightenTangents;
+	TSharedPtr<FUICommandInfo> StraightenStartTangent;
+	TSharedPtr<FUICommandInfo> StraightenEndTangent;
+
+	TSharedPtr<FUICommandInfo> CopyTangents;
+	TSharedPtr<FUICommandInfo> PasteTangents;
+
+	TSharedPtr<FUICommandInfo> CreateExternalCurveAsset;
+
+	TSharedPtr<FUICommandInfo> SetKeyInterpConstant;
+	TSharedPtr<FUICommandInfo> SetKeyInterpLinear;
+	TSharedPtr<FUICommandInfo> SetKeyInterpCubicAuto;
+	TSharedPtr<FUICommandInfo> SetKeyInterpCubicSmartAuto;
+	TSharedPtr<FUICommandInfo> SetKeyInterpCubicUser;
+	TSharedPtr<FUICommandInfo> SetKeyInterpCubicBreak;
+	TSharedPtr<FUICommandInfo> SetKeyInterpToggleWeighted;
+
+	TSharedPtr<FUICommandInfo> SelectNextChannelKey;
+	TSharedPtr<FUICommandInfo> SelectPreviousChannelKey;
+};
diff --git a/Engine/Plugins/Experimental/Avalanche/Source/AvalancheSequencer/Private/EaseCurveTool/AvaEaseCurveToolSettings.cpp b/Engine/Plugins/Experimental/Avalanche/Source/AvalancheSequencer/Private/EaseCurveTool/AvaEaseCurveToolSettings.cpp
new file mode 100644
--- /dev/null
+++ b/Engine/Plugins/Experimental/Avalanche/Source/AvalancheSequencer/Private/EaseCurveTool/AvaEaseCurveToolSettings.cpp
@@ -0,0 +1,13 @@
+// Copyright Epic Games, Inc. All Rights Reserved.
+
+#include "AvaEaseCurveToolSettings.h"
+#include "Misc/Paths.h"
+
+UAvaEaseCurveToolSettings::UAvaEaseCurveToolSettings()
+{
+	CategoryName = TEXT("Motion Design");
+	SectionName = TEXT("Ease Curve Tool");
+
+	PresetsPath = FPaths::ProjectConfigDir() / TEXT("EaseCurves");
+	NewPresetCategory = TEXT("Custom");
+}
diff --git a/Engine/Plugins/Experimental/Avalanche/Source/AvalancheSequencer/Private/EaseCurveTool/AvaEaseCurveToolSettings.h b/Engine/Plugins/Experimental/Avalanche/Source/AvalancheSequencer/Private/EaseCurveTool/AvaEaseCurveToolSettings.h
new file mode 100644
--- /dev/null
+++ b/Engine/Plugins/Experimental/Avalanche/Source/AvalancheSequencer/Private/EaseCurveTool/AvaEaseCurveToolSettings.h
@@ -0,0 +1,67 @@
+// Copyright Epic Games, Inc. All Rights Reserved.
+
+#pragma once
+
+#include "Containers/EnumAsByte.h"
+#include "Engine/DeveloperSettings.h"
+#include "Types/SlateEnums.h"
+#include "UObject/Object.h"
+#include "AvaEaseCurveToolSettings.generated.h"
+
+UCLASS(config=EditorPerProjectUserSettings, meta = (DisplayName = "Ease Curve Tool"))
+class UAvaEaseCurveToolSettings : public UDeveloperSettings
+{
+	GENERATED_BODY()
+
+public:
+	UAvaEaseCurveToolSettings();
+
+	int32 GetGraphSize() const { return GraphSize; }
+	void SetGraphSize(const int32 InNewSize) { GraphSize = InNewSize; }
+
+	EHorizontalAlignment GetGraphHAlign() const { return GraphHAlign.GetValue(); }
+	void SetGraphHAlign(const EHorizontalAlignment InNewHAlign) { GraphHAlign = InNewHAlign; }
+
+	bool GetGridSnap() const { return bGridSnap; }
+	void SetGridSnap(const bool bInCurveEaseToolGridSnap) { bGridSnap = bInCurveEaseToolGridSnap; }
+	void ToggleGridSnap() { bGridSnap = !bGridSnap; }
+
+	int32 GetGridSize() const { return GridSize; }
+	void SetGridSize(const int32 InNewSize) { GridSize = InNewSize; }
+
+	bool GetAutoFlipTangents() const { return bAutoFlipTangents; }
+	void SetAutoFlipTangents(const bool bInAutoFlipTangents) { bAutoFlipTangents = bInAutoFlipTangents; }
+	void ToggleAutoFlipTangents() { bAutoFlipTangents = !bAutoFlipTangents; }
+
+	FString GetPresetsPath() const { return PresetsPath; }
+	FString GetNewPresetCategory() const { return NewPresetCategory; }
+
+private:
+	/** The height of the curve ease tool in the details panel. */
+	UPROPERTY(Config, EditAnywhere, Category = "EaseCurveTool", meta = (DisplayName = "Graph Size", UIMin = 64, ClampMin = 64, UIMax = 256, ClampMax = 256, Delta = 1))
+	int32 GraphSize = 140;
+
+	/** The horizontal alignment of the curve ease tool in the details panel. */
+	UPROPERTY(Config, EditAnywhere, Category = "EaseCurveTool", meta = (DisplayName = "Graph Horizontal Alignment"))
+	TEnumAsByte<EHorizontalAlignment> GraphHAlign = EHorizontalAlignment::HAlign_Fill;
+
+	/** If true, snaps tangents to grid. */
+	UPROPERTY(Config, EditAnywhere, Category = "EaseCurveTool", meta = (DisplayName = "Grid Snap"))
+	bool bGridSnap = false;
+
+	/** The height of the curve ease tool in the details panel. */
+	UPROPERTY(Config, EditAnywhere, Category = "EaseCurveTool", meta = (DisplayName = "Grid Size", UIMin = 4, ClampMin = 4, UIMax = 24, ClampMax = 24, Delta = 1))
+	int32 GridSize = 8;
+
+	/** If true, auto flips tangents when sequential key frame curve values are descending. */
+	UPROPERTY(Config, EditAnywhere, Category = "EaseCurveTool", meta = (DisplayName = "Auto Flip Tangents"))
+	bool bAutoFlipTangents = true;
+
+	/** The location to search for ease curve preset Json files. */
+	UPROPERTY(Config, EditAnywhere, Category = "EaseCurveTool", meta = (DisplayName = "Presets Path"))
+	FString PresetsPath;
+
+	/** The name of the category to place newly created curve presets. */
+	UPROPERTY(Config, EditAnywhere, Category = "EaseCurveTool", meta = (DisplayName = "New Preset Category"))
+	FString NewPresetCategory;
+};
