[MotionDesign]

diff --git a/Engine/Plugins/Experimental/Avalanche/Source/AvalancheSequencer/Private/EaseCurveTool/Widgets/SAvaEaseCurveEditor.cpp b/Engine/Plugins/Experimental/Avalanche/Source/AvalancheSequencer/Private/EaseCurveTool/Widgets/SAvaEaseCurveEditor.cpp
new file mode 100644
--- /dev/null
+++ b/Engine/Plugins/Experimental/Avalanche/Source/AvalancheSequencer/Private/EaseCurveTool/Widgets/SAvaEaseCurveEditor.cpp
@@ -0,0 +1,1259 @@
+// Copyright Epic Games, Inc. All Rights Reserved.
+
+#include "SAvaEaseCurveEditor.h"
+#include "EaseCurveTool/Widgets/SAvaEaseCurvePreview.h"
+#include "EaseCurveTool/Widgets/SAvaEaseCurvePreviewToolTip.h"
+#include "Curves/SimpleCurve.h"
+#include "Editor.h"
+#include "Factories/Factory.h"
+#include "Fonts/SlateFontInfo.h"
+#include "Framework/Application/MenuStack.h"
+#include "Framework/Application/SlateApplication.h"
+#include "Framework/Commands/UIAction.h"
+#include "Framework/Commands/UICommandList.h"
+#include "Framework/MultiBox/MultiBoxBuilder.h"
+#include "HAL/PlatformApplicationMisc.h"
+#include "Layout/WidgetPath.h"
+#include "Rendering/DrawElements.h"
+#include "ScopedTransaction.h"
+#include "Styling/AppStyle.h"
+#include "Styling/CoreStyle.h"
+#include "Styling/SlateTypes.h"
+#include "Styling/StyleColors.h"
+#include "Textures/SlateIcon.h"
+#include "Widgets/Colors/SColorBlock.h"
+#include "Widgets/Layout/SBox.h"
+#include "Widgets/SBoxPanel.h"
+#include "Widgets/SToolTip.h"
+#include "Widgets/Text/STextBlock.h"
+
+#define LOCTEXT_NAMESPACE "SAvaEaseCurveEditor"
+
+namespace UE::EaseCurveTool::Private
+{
+	const static FVector2d CONST_KeySize = FVector2d(7, 7);
+	const static FVector2d CONST_TangentSize = FVector2d(16, 16);
+	const static FVector2d CONST_CurveSize = FVector2d(12, 12);
+
+	constexpr static float CONST_FitMargin = 0.1f;
+	constexpr static float CONST_MinViewRange = 0.5f;
+	constexpr static float CONST_DefaultZoomRange = 1.f;
+
+	constexpr static float CONST_NormalBoundsThickness = 1.f;
+	constexpr static float CONST_TangentHandleLineThickness = 2.f;
+}
+
+using namespace UE::EaseCurveTool::Private;
+
+FVector2d SAvaEaseCurveEditor::CalcTangentDir(const double InTangent)
+{
+	const double Angle = FMath::Atan(InTangent);
+	return FVector2d(FMath::Cos(Angle), FMath::Sin(Angle));
+}
+
+float SAvaEaseCurveEditor::CalcTangent(const FVector2d& InHandleDelta)
+{
+	return InHandleDelta.Y / FMath::Max<double>(InHandleDelta.X, KINDA_SMALL_NUMBER);
+}
+
+void SAvaEaseCurveEditor::Construct(const FArguments& InArgs, const TObjectPtr<UAvaEaseCurve>& InEaseCurve)
+{
+	check(InEaseCurve);
+	EaseCurve = InEaseCurve;
+
+	DataMinInput = InArgs._DataMinInput;
+	DataMaxInput = InArgs._DataMaxInput;
+
+	DesiredSize = InArgs._DesiredSize;
+	DisplayRate = InArgs._DisplayRate;
+
+	bShowInputGridNumbers = InArgs._ShowInputGridNumbers;
+	bShowOutputGridNumbers = InArgs._ShowOutputGridNumbers;
+
+	NormalAreaColor = InArgs._NormalAreaColor;
+	NormalBoundsColor = InArgs._NormalBoundsColor;
+	bGridSnap = InArgs._GridSnap;
+	GridSize = InArgs._GridSize;
+	GridColor = InArgs._GridColor;
+	ExtendedGridColor = InArgs._ExtendedGridColor;
+	CurveThickness = InArgs._CurveThickness;
+	CurveColor = InArgs._CurveColor;
+	Operation = InArgs._Operation;
+	
+	OnTangentsChanged = InArgs._OnTangentsChanged;
+	GetContextMenuContent = InArgs._GetContextMenuContent;
+	OnKeyDownEvent = InArgs._OnKeyDown;
+	OnDragStart = InArgs._OnDragStart;
+	OnDragEnd = InArgs._OnDragEnd;
+
+	// If editor size is set, use it, otherwise, use default value
+	if (DesiredSize.Get().IsZero())
+	{
+		DesiredSize.Set(FVector2D(128));
+	}
+
+	if (GEditor)
+	{
+		GEditor->RegisterForUndo(this);
+	}
+
+	FCoreUObjectDelegates::OnObjectPropertyChanged.AddSP(this, &SAvaEaseCurveEditor::OnObjectPropertyChanged);
+
+	ZoomToFit();
+
+	ChildSlot
+		[
+			SNew(SColorBlock)
+			.Color(FStyleColors::Background.GetSpecifiedColor())
+		];
+}
+
+FText SAvaEaseCurveEditor::GetCurveToolTipInputText() const
+{
+	return CurveToolTipInputText;
+}
+
+FText SAvaEaseCurveEditor::GetCurveToolTipOutputText() const
+{
+	return CurveToolTipOutputText;
+}
+
+SAvaEaseCurveEditor::~SAvaEaseCurveEditor()
+{
+	if (GEditor)
+	{
+		GEditor->UnregisterForUndo(this);
+	}
+
+	FCoreUObjectDelegates::OnObjectPropertyChanged.RemoveAll(this);
+}
+
+FVector2D SAvaEaseCurveEditor::ComputeDesiredSize(const float InLayoutScaleMultiplier) const
+{
+	return DesiredSize.Get();
+}
+
+int32 SAvaEaseCurveEditor::OnPaint(const FPaintArgs& InArgs, const FGeometry& InAllottedGeometry, const FSlateRect& InMyCullingRect
+	, FSlateWindowElementList& OutDrawElements, int32 InLayerId, const FWidgetStyle& InWidgetStyle, bool bInParentEnabled) const
+{
+	InLayerId = SCompoundWidget::OnPaint(InArgs, InAllottedGeometry, InMyCullingRect, OutDrawElements, InLayerId, InWidgetStyle, bInParentEnabled);
+
+	const bool bEnabled = ShouldBeEnabled(bInParentEnabled);
+	const ESlateDrawEffect DrawEffects = bEnabled ? ESlateDrawEffect::None : ESlateDrawEffect::DisabledEffect;
+	const FTrackScaleInfo ScaleInfo(ViewMinInput, ViewMaxInput, ViewMinOutput, ViewMaxOutput, InAllottedGeometry.GetLocalSize());
+
+	if (FMath::IsNearlyEqual(ViewMinInput, ViewMaxInput) || FMath::IsNearlyEqual(ViewMinOutput, ViewMaxOutput))
+	{
+		return InLayerId;
+	}
+	
+	const float ZeroInputX = ScaleInfo.InputToLocalX(0.f);
+	const float ZeroOutputY = ScaleInfo.OutputToLocalY(0.f);
+	const float OneInputX = ScaleInfo.InputToLocalX(1.f);
+	const float OneOutputY = ScaleInfo.OutputToLocalY(1.f);
+
+	// Normalized area background
+	{
+		const float InputSize = OneInputX - ZeroInputX;
+		const float OutputSize = (OneOutputY - ZeroOutputY) * -1;
+		const FSlateBrush* WhiteBrush = FAppStyle::GetBrush(TEXT("WhiteTexture"));
+
+		FSlateDrawElement::MakeBox(OutDrawElements, ++InLayerId,
+			InAllottedGeometry.ToPaintGeometry(FVector2D(InputSize, OutputSize), FSlateLayoutTransform(FVector2D(ZeroInputX, OneOutputY))),
+			WhiteBrush, DrawEffects, NormalAreaColor);
+	}
+
+	InLayerId = PaintGrid(ScaleInfo, InAllottedGeometry, OutDrawElements, ++InLayerId, InMyCullingRect, DrawEffects);
+
+	PaintNormalBounds(ScaleInfo, InAllottedGeometry, OutDrawElements, ++InLayerId, InMyCullingRect, DrawEffects, InWidgetStyle);
+
+	PaintCurve(ScaleInfo, InAllottedGeometry, OutDrawElements, ++InLayerId, InMyCullingRect, DrawEffects, InWidgetStyle);
+
+	InLayerId = PaintKeys(ScaleInfo, InAllottedGeometry, OutDrawElements, ++InLayerId, InMyCullingRect, DrawEffects, InWidgetStyle);
+
+	if (DragState == EDragState::MarqueeSelect)
+	{
+		const FVector2D MarqueTopLeft(FMath::Min(MouseDownLocation.X, MouseMoveLocation.X), FMath::Min(MouseDownLocation.Y, MouseMoveLocation.Y));
+		const FVector2D MarqueBottomRight(FMath::Max(MouseDownLocation.X, MouseMoveLocation.X), FMath::Max(MouseDownLocation.Y, MouseMoveLocation.Y));
+
+		FSlateDrawElement::MakeBox(OutDrawElements, ++InLayerId
+			, InAllottedGeometry.ToPaintGeometry(MarqueBottomRight - MarqueTopLeft, FSlateLayoutTransform(MarqueTopLeft))
+			, FAppStyle::GetBrush(TEXT("MarqueeSelection")));
+	}
+
+	return InLayerId;
+}
+
+int32 SAvaEaseCurveEditor::PaintGrid(const FTrackScaleInfo& InScaleInfo, const FGeometry& InAllottedGeometry, FSlateWindowElementList& OutDrawElements
+	, int32 InLayerId, const FSlateRect& InMyCullingRect, ESlateDrawEffect InDrawEffects) const
+{
+	if (!GridSize.IsSet())
+	{
+		return InLayerId;
+	}
+
+	const float GridStep = 1.f / GridSize.Get();
+
+	// Outside normalized area grid lines
+	{
+		const FVector2D LocalSize = InAllottedGeometry.GetLocalSize();
+		const float ScreenGridStepX = FMath::Abs(InScaleInfo.InputToLocalX(GridStep) - InScaleInfo.InputToLocalX(0.f));
+		const float ScreenGridStepY = FMath::Abs(InScaleInfo.OutputToLocalY(GridStep) - InScaleInfo.OutputToLocalY(0.f));
+
+		// Vertical grid lines
+		const float StartX = -FMath::Fractional(InScaleInfo.LocalXToInput(0.f) / GridStep) * ScreenGridStepX;
+
+		if (ScreenGridStepX >= 1.f)
+		{
+			for (float ScreenX = StartX; ScreenX < LocalSize.X; ScreenX += ScreenGridStepX)
+			{
+				TArray<FVector2D> LinePoints;
+				LinePoints.Add(FVector2D(ScreenX, 0.f));
+				LinePoints.Add(FVector2D(ScreenX, LocalSize.Y));
+
+				FSlateDrawElement::MakeLines(OutDrawElements, InLayerId, InAllottedGeometry.ToPaintGeometry(), LinePoints, InDrawEffects, ExtendedGridColor, false);
+			}
+		}
+
+		// Horizontal grid lines
+		const float StartY = FMath::Fractional(InScaleInfo.LocalYToOutput(0.f) / GridStep) * ScreenGridStepY;
+
+		if (ScreenGridStepY >= 1.f)
+		{
+			for (float ScreenY = StartY; ScreenY < LocalSize.Y; ScreenY += ScreenGridStepY)
+			{
+				TArray<FVector2D> LinePoints;
+				LinePoints.Add(FVector2D(0.f, ScreenY));
+				LinePoints.Add(FVector2D(LocalSize.X, ScreenY));
+
+				FSlateDrawElement::MakeLines(OutDrawElements, InLayerId, InAllottedGeometry.ToPaintGeometry(), LinePoints, InDrawEffects, ExtendedGridColor, false);
+			}
+		}
+	}
+
+	// Normalized area grid
+	const int32 NormalizedLayerId = InLayerId + 1;
+	{
+		// Vertical grid lines
+		for (float X = 0.f; X < 1.f; X += GridStep)
+		{
+			const float ScreenX = InScaleInfo.InputToLocalX(X);
+
+			const float CurveZeroY1 = InScaleInfo.OutputToLocalY(0.f);
+			const float CurveZeroY2 = InScaleInfo.OutputToLocalY(1.f);
+
+			TArray<FVector2D> LinePoints;
+			LinePoints.Add(FVector2D(ScreenX, CurveZeroY1));
+			LinePoints.Add(FVector2D(ScreenX, CurveZeroY2));
+
+			FSlateDrawElement::MakeLines(OutDrawElements, NormalizedLayerId, InAllottedGeometry.ToPaintGeometry(), LinePoints, InDrawEffects, GridColor, false);
+		}
+
+		// Horizontal grid lines
+		for (float Y = 0.f; Y < 1.f; Y += GridStep)
+		{
+			const float ScreenY = InScaleInfo.OutputToLocalY(Y);
+
+			const float CurveZeroX1 = InScaleInfo.InputToLocalX(0.f);
+			const float CurveZeroX2 = InScaleInfo.InputToLocalX(1.f);
+
+			TArray<FVector2D> LinePoints;
+			LinePoints.Add(FVector2D(CurveZeroX1, ScreenY));
+			LinePoints.Add(FVector2D(CurveZeroX2, ScreenY));
+
+			FSlateDrawElement::MakeLines(OutDrawElements, NormalizedLayerId, InAllottedGeometry.ToPaintGeometry(), LinePoints, InDrawEffects, GridColor, false);
+		}
+	}
+
+	return NormalizedLayerId;
+}
+
+void SAvaEaseCurveEditor::PaintNormalBounds(const FTrackScaleInfo& InScaleInfo, const FGeometry& InAllottedGeometry, FSlateWindowElementList& OutDrawElements
+	, int32 InLayerId, const FSlateRect& InMyCullingRect, ESlateDrawEffect InDrawEffects, const FWidgetStyle& InWidgetStyle) const
+{
+	const float ZeroInputX = InScaleInfo.InputToLocalX(0.f);
+	const float OneInputX = InScaleInfo.InputToLocalX(1.f);
+	const float ZeroOutputY = InScaleInfo.OutputToLocalY(0.f);
+	const float OneOutputY = InScaleInfo.OutputToLocalY(1.f);
+
+	// Time = 0 line
+	TArray<FVector2D> ZeroTimeLinePoints;
+	ZeroTimeLinePoints.Add(FVector2D(ZeroInputX, 0.f));
+	ZeroTimeLinePoints.Add(FVector2D(ZeroInputX, InAllottedGeometry.GetLocalSize().Y));
+	FSlateDrawElement::MakeLines(OutDrawElements, InLayerId, InAllottedGeometry.ToPaintGeometry(), ZeroTimeLinePoints, InDrawEffects, NormalBoundsColor, false, CONST_NormalBoundsThickness);
+
+	// Time = 1 line
+	TArray<FVector2D> OneTimeLinePoints;
+	OneTimeLinePoints.Add(FVector2D(OneInputX, 0.f));
+	OneTimeLinePoints.Add(FVector2D(OneInputX, InAllottedGeometry.GetLocalSize().Y));
+	FSlateDrawElement::MakeLines(OutDrawElements, InLayerId, InAllottedGeometry.ToPaintGeometry(), OneTimeLinePoints, InDrawEffects, NormalBoundsColor, false, CONST_NormalBoundsThickness);
+
+	// Value = 0 line
+	TArray<FVector2D> ZeroValueLinePoints;
+	ZeroValueLinePoints.Add(FVector2D(0.f, ZeroOutputY));
+	ZeroValueLinePoints.Add(FVector2D(InAllottedGeometry.GetLocalSize().X, ZeroOutputY));
+	FSlateDrawElement::MakeLines(OutDrawElements, InLayerId, InAllottedGeometry.ToPaintGeometry(), ZeroValueLinePoints, InDrawEffects, NormalBoundsColor, false, CONST_NormalBoundsThickness);
+
+	// Value = 1 line
+	TArray<FVector2D> OneValueLinePoints;
+	OneValueLinePoints.Add(FVector2D(0.f, OneOutputY));
+	OneValueLinePoints.Add(FVector2D(InAllottedGeometry.GetLocalSize().X, OneOutputY));
+	FSlateDrawElement::MakeLines(OutDrawElements, InLayerId, InAllottedGeometry.ToPaintGeometry(), OneValueLinePoints, InDrawEffects, NormalBoundsColor, false, CONST_NormalBoundsThickness);
+}
+
+void SAvaEaseCurveEditor::PaintCurve(const FTrackScaleInfo& InScaleInfo, const FGeometry& InAllottedGeometry, FSlateWindowElementList& OutDrawElements,
+	int32 LayerId, const FSlateRect& MyCullingRect, ESlateDrawEffect DrawEffects, const FWidgetStyle& InWidgetStyle) const
+{
+	const float NumKeys = EaseCurve->FloatCurve.GetNumKeys();
+	
+	if (NumKeys != 2)
+	{
+		return;
+	}
+	
+	TArray<FVector2D> LinePoints;
+	TArray<FLinearColor> LineColors;
+	TArray<FKeyHandle> KeyHandles;
+	TArray<TPair<float, float>> Key_TimeValuePairs;
+	
+	KeyHandles.Reserve(NumKeys);
+	Key_TimeValuePairs.Reserve(NumKeys);
+
+	for (auto It = EaseCurve->FloatCurve.GetKeyHandleIterator(); It; ++It)
+	{
+		const FKeyHandle& KeyHandle = *It;
+
+		KeyHandles.Add(KeyHandle);
+		Key_TimeValuePairs.Emplace(EaseCurve->FloatCurve.GetKeyTimeValuePair(KeyHandle));
+	}
+
+	// Add enclosed segments
+	for (int32 Index = 0; Index < NumKeys - 1; ++Index)
+	{
+		CreateLinesForSegment(EaseCurve->FloatCurve.GetKeyInterpMode(KeyHandles[Index])
+			, Key_TimeValuePairs[Index], Key_TimeValuePairs[Index + 1], LinePoints, LineColors, InScaleInfo);
+
+		FSlateDrawElement::MakeLines(OutDrawElements, LayerId, InAllottedGeometry.ToPaintGeometry()
+			, LinePoints, LineColors, DrawEffects, FLinearColor::White/*Color*/, true, CurveThickness * InAllottedGeometry.Scale);
+
+		LinePoints.Empty();
+		LineColors.Empty();
+	}
+}
+
+void SAvaEaseCurveEditor::CreateLinesForSegment(const ERichCurveInterpMode InInterpMode
+	, const TPair<float, float>& InStartKeyTimeValue, const TPair<float, float>& InEndKeyTimeValue
+	, TArray<FVector2D>& OutLinePoints, TArray<FLinearColor>& OutLineColors
+	, const FTrackScaleInfo& InScaleInfo) const
+{
+	FLinearColor FadedCurveColor = CurveColor;
+	FadedCurveColor.A = 0.2f;
+
+	switch (InInterpMode)
+	{
+	case RCIM_Constant:
+	{
+		//@todo: should really only need 3 points here but something about the line rendering isn't quite behaving as I'd expect, so need extras
+		OutLinePoints.Add(FVector2D(InStartKeyTimeValue.Key, InStartKeyTimeValue.Value));
+		OutLinePoints.Add(FVector2D(InEndKeyTimeValue.Key, InStartKeyTimeValue.Value));
+		OutLinePoints.Add(FVector2D(InEndKeyTimeValue.Key, InStartKeyTimeValue.Value));
+		OutLinePoints.Add(FVector2D(InEndKeyTimeValue.Key, InEndKeyTimeValue.Value));
+		OutLinePoints.Add(FVector2D(InEndKeyTimeValue.Key, InStartKeyTimeValue.Value));
+		break;
+	}
+	case RCIM_Linear:
+	{
+		OutLinePoints.Add(FVector2D(InStartKeyTimeValue.Key, InStartKeyTimeValue.Value));
+		OutLinePoints.Add(FVector2D(InEndKeyTimeValue.Key, InEndKeyTimeValue.Value));
+		break;
+	}
+	case RCIM_Cubic:
+	{
+		// Clamp to screen to avoid massive slowdown when zoomed in
+		const float StartX = FMath::Max(InScaleInfo.InputToLocalX(InStartKeyTimeValue.Key), 0.0f);
+		const float EndX = FMath::Min(InScaleInfo.InputToLocalX(InEndKeyTimeValue.Key), InScaleInfo.WidgetSize.X);
+
+		constexpr float StepSize = 1.0f;
+		const float CurveLengthX = EndX - StartX;
+
+		for (float CurrentX = StartX; CurrentX < EndX; CurrentX += StepSize)
+		{
+			// Add line point
+			const float CurveIn = InScaleInfo.LocalXToInput(FMath::Min(CurrentX, EndX));
+			const float CurveOut = EaseCurve->FloatCurve.Eval(CurveIn);
+			OutLinePoints.Add(FVector2D(CurveIn, CurveOut));
+			
+			switch (Operation.Get(FAvaEaseCurveTool::EOperation::InOut))
+			{
+			case FAvaEaseCurveTool::EOperation::Out:
+			{
+				const float Alpha = CurrentX / CurveLengthX;
+				OutLineColors.Add(FLinearColor::LerpUsingHSV(FadedCurveColor, CurveColor,  1.f - Alpha));
+				break;
+			}
+			case FAvaEaseCurveTool::EOperation::In:
+			{
+				const float Alpha = CurrentX / CurveLengthX;
+				OutLineColors.Add(FLinearColor::LerpUsingHSV(FadedCurveColor, CurveColor, Alpha));
+				break;
+			}
+			case FAvaEaseCurveTool::EOperation::InOut:
+			default:
+			{
+				OutLineColors.Add(CurveColor);
+				break;
+			}
+			}
+		}
+		OutLinePoints.Add(FVector2D(InEndKeyTimeValue.Key, InEndKeyTimeValue.Value));
+		OutLineColors.Add(CurveColor);
+		break;
+	}
+	default:
+		break;
+	}
+
+	// Transform to screen
+	for (auto It = OutLinePoints.CreateIterator(); It; ++It)
+	{
+		FVector2D Vec2D = *It;
+		Vec2D.X = InScaleInfo.InputToLocalX(Vec2D.X);
+		Vec2D.Y = InScaleInfo.OutputToLocalY(Vec2D.Y);
+		*It = Vec2D;
+	}
+}
+
+int32 SAvaEaseCurveEditor::PaintKeys(const FTrackScaleInfo& InScaleInfo, const FGeometry& InAllottedGeometry, FSlateWindowElementList& OutDrawElements
+	, int32 InLayerId, const FSlateRect& InMyCullingRect, ESlateDrawEffect InDrawEffects, const FWidgetStyle& InWidgetStyle) const
+{
+	if (EaseCurve->FloatCurve.GetNumKeys() < 2)
+	{
+		return InLayerId;
+	}
+
+	const int32 LayerId = InLayerId;
+	const int32 SelectedLayerId = InLayerId + 1;
+
+	// Iterate over each key
+	ERichCurveInterpMode LastInterpMode = RCIM_Linear;
+	for (auto It(EaseCurve->FloatCurve.GetKeyHandleIterator()); It; ++It)
+	{
+		FKeyHandle KeyHandle = *It;
+
+		const FVector2D KeyLocation(InScaleInfo.InputToLocalX(EaseCurve->FloatCurve.GetKeyTime(KeyHandle))
+			, InScaleInfo.OutputToLocalY(EaseCurve->FloatCurve.GetKeyValue(KeyHandle)));
+		const FVector2D KeyIconLocation = KeyLocation - (CONST_KeySize / 2);
+		const bool bIsSelected = (KeyHandle == SelectedTangent.KeyHandle);
+		const FSlateBrush* KeyBrush = /*bIsSelected ? FAppStyle::GetBrush("CurveEd.CurveKeySelected") : */ FAppStyle::GetBrush("CurveEd.CurveKey");
+		const int32 LayerToUse = bIsSelected ? SelectedLayerId : LayerId;
+
+		FSlateDrawElement::MakeBox(OutDrawElements, LayerToUse,
+			InAllottedGeometry.ToPaintGeometry(CONST_KeySize, FSlateLayoutTransform(KeyIconLocation)),
+			KeyBrush, InDrawEffects,
+			KeyBrush->GetTint(InWidgetStyle) * InWidgetStyle.GetColorAndOpacityTint());
+
+		// Draw the tangent handles
+		const int32 KeyIndex = EaseCurve->FloatCurve.GetIndexSafe(KeyHandle);
+
+		if (EaseCurve->FloatCurve.GetKeyInterpMode(KeyHandle) == RCIM_Cubic || LastInterpMode == RCIM_Cubic)
+		{
+			if (KeyIndex == 0)
+			{
+				PaintTangentHandle(InAllottedGeometry, OutDrawElements, LayerToUse, InDrawEffects, InWidgetStyle
+					, KeyLocation, GetLeaveTangentScreenLocation(InScaleInfo, KeyHandle), bIsSelected);
+			}
+			else if (KeyIndex == 1)
+			{
+				PaintTangentHandle(InAllottedGeometry, OutDrawElements, LayerToUse, InDrawEffects, InWidgetStyle
+					, KeyLocation, GetArriveTangentScreenLocation(InScaleInfo, KeyHandle), bIsSelected);
+			}
+		}
+		else
+		{
+			PaintTangentHandle(InAllottedGeometry, OutDrawElements, LayerToUse, InDrawEffects, InWidgetStyle
+				, KeyLocation, GetLeaveTangentScreenLocation(InScaleInfo, KeyHandle), bIsSelected);
+
+			PaintTangentHandle(InAllottedGeometry, OutDrawElements, LayerToUse, InDrawEffects, InWidgetStyle
+				, KeyLocation, GetArriveTangentScreenLocation(InScaleInfo, KeyHandle), bIsSelected);
+		}
+
+		LastInterpMode = EaseCurve->FloatCurve.GetKeyInterpMode(KeyHandle);
+	}
+
+	return SelectedLayerId;
+}
+
+int32 SAvaEaseCurveEditor::PaintTangentHandle(const FGeometry& InAllottedGeometry, FSlateWindowElementList& OutDrawElements, int32 InLayerId
+	, ESlateDrawEffect InDrawEffects, const FWidgetStyle& InWidgetStyle, const FVector2D& InKeyLocation, const FVector2D& InTangentLocation, const bool bInSelected) const
+{
+	// Draw line to tangent control handle
+	{
+		const FLinearColor LineColor = FAppStyle::GetColor("CurveEd.TangentColor");
+
+		TArray<FVector2D> LinePoints;
+		LinePoints.Add(FVector2D(InKeyLocation));
+		LinePoints.Add(FVector2D(InTangentLocation));
+
+		FSlateDrawElement::MakeLines(OutDrawElements, InLayerId++, InAllottedGeometry.ToPaintGeometry(), LinePoints, InDrawEffects, LineColor, true, CONST_TangentHandleLineThickness);
+	}
+
+	// Draw tangent handle
+	{
+		const FSlateBrush* ImageBrush = FAppStyle::GetBrush(TEXT("Icons.BulletPoint"));
+		const FVector2D TangentIconLocation = InTangentLocation - (CONST_TangentSize / 2);
+		const FLinearColor HandleColor = bInSelected ? FStyleColors::AccentBlue.GetSpecifiedColor() : FLinearColor::White;
+
+		FSlateDrawElement::MakeBox(OutDrawElements, InLayerId++
+			, InAllottedGeometry.ToPaintGeometry(CONST_TangentSize, FSlateLayoutTransform(TangentIconLocation))
+			, ImageBrush, InDrawEffects, HandleColor * InWidgetStyle.GetColorAndOpacityTint());
+	}
+
+	return InLayerId;
+}
+
+void SAvaEaseCurveEditor::SetRequireFocusToZoom(const bool bInRequireFocusToZoom)
+{
+	bRequireFocusToZoom = bInRequireFocusToZoom;
+}
+
+TOptional<bool> SAvaEaseCurveEditor::OnQueryShowFocus(const EFocusCause InFocusCause) const
+{
+	if (bRequireFocusToZoom)
+	{
+		// Enable showing a focus rectangle when the widget has keyboard focus
+		return TOptional<bool>(true);
+	}
+	return TOptional<bool>();
+}
+
+FReply SAvaEaseCurveEditor::OnMouseButtonDown(const FGeometry& InMyGeometry, const FPointerEvent& InMouseEvent)
+{
+	const FKey EffectingButton = InMouseEvent.GetEffectingButton();
+	
+	OnDragEnd.ExecuteIfBound();
+
+	DragState = EDragState::PreDrag;
+
+	if (EffectingButton == EKeys::LeftMouseButton)
+	{
+		MouseDownLocation = InMyGeometry.AbsoluteToLocal(InMouseEvent.GetScreenSpacePosition());
+
+		if (!HasKeyboardFocus())
+		{
+			FSlateApplication::Get().SetKeyboardFocus(SharedThis(this), EFocusCause::SetDirectly);
+		}
+
+		OnMouseMove(InMyGeometry, InMouseEvent);
+
+		return FReply::Handled().CaptureMouse(SharedThis(this));
+	}
+	else if (EffectingButton == EKeys::RightMouseButton)
+	{
+		if (HitTestCurves(InMyGeometry, InMouseEvent))
+		{
+			FPlatformApplicationMisc::ClipboardCopy(*EaseCurve->GetTangents().ToJson());
+
+			FAvaEaseCurveTool::ShowNotificationMessage(LOCTEXT("EaseCurveToolTangentsCopied", "Ease Curve Tool Tangents Copied!"));
+		}
+
+		return FReply::Handled().CaptureMouse(SharedThis(this));
+	}
+
+	return FReply::Unhandled();
+}
+
+void SAvaEaseCurveEditor::OnMouseCaptureLost(const FCaptureLostEvent& CaptureLostEvent)
+{
+	if (DragState == EDragState::DragTangent)
+	{
+		OnDragEnd.ExecuteIfBound();
+	}
+	
+	DragState = EDragState::None;
+}
+
+FReply SAvaEaseCurveEditor::OnMouseButtonUp(const FGeometry& InMyGeometry, const FPointerEvent& InMouseEvent)
+{
+	if (this->HasMouseCapture())
+	{
+		if (DragState == EDragState::PreDrag)
+		{
+			ProcessClick(InMyGeometry, InMouseEvent);
+		}
+		else
+		{
+			EndDrag(InMyGeometry, InMouseEvent);
+		}
+
+		ClearSelection();
+
+		return FReply::Handled().ReleaseMouseCapture();
+	}
+
+	return FReply::Unhandled();
+}
+
+void ClampViewRangeToDataIfBound(float& NewViewMin, float& NewViewMax, const TAttribute< TOptional<float> >& DataMin, const TAttribute< TOptional<float> >& DataMax, const float ViewRange)
+{
+	const TOptional<float>& Min = DataMin.Get();
+	const TOptional<float>& Max = DataMax.Get();
+	if (Min.IsSet() && NewViewMin < Min.GetValue())
+	{
+		NewViewMin = Min.GetValue();
+		NewViewMax = ViewRange;
+	}
+	else if (Max.IsSet() && NewViewMax > Max.GetValue())
+	{
+		NewViewMin = Max.GetValue() - ViewRange;
+		NewViewMax = Max.GetValue();
+	}
+}
+
+FReply SAvaEaseCurveEditor::OnMouseMove(const FGeometry& InMyGeometry, const FPointerEvent& InMouseEvent)
+{
+	UpdateCurveToolTip(InMyGeometry, InMouseEvent);
+
+	if (this->HasMouseCapture())
+	{
+		if (DragState == EDragState::PreDrag)
+		{
+			TryStartDrag(InMyGeometry, InMouseEvent);
+		}
+		if (DragState != EDragState::None)
+		{
+			ProcessDrag(InMyGeometry, InMouseEvent);
+		}
+
+		MouseMoveLocation = InMyGeometry.AbsoluteToLocal(InMouseEvent.GetScreenSpacePosition());
+
+		return FReply::Handled();
+	}
+
+	return FReply::Unhandled();
+}
+
+void SAvaEaseCurveEditor::UpdateCurveToolTip(const FGeometry& InMyGeometry, const FPointerEvent& InMouseEvent)
+{
+	if (DragState != EDragState::DragTangent && HitTestCurves(InMyGeometry, InMouseEvent))
+	{
+		const FTrackScaleInfo ScaleInfo(ViewMinInput, ViewMaxInput, ViewMinOutput, ViewMaxOutput, InMyGeometry.GetLocalSize());
+		const FVector2D HitPosition = InMyGeometry.AbsoluteToLocal(InMouseEvent.GetScreenSpacePosition());
+		const float Time = ScaleInfo.LocalXToInput(HitPosition.X);
+		const float Value = EaseCurve->FloatCurve.Eval(Time);
+
+		FNumberFormattingOptions FormattingOptions;
+		FormattingOptions.MaximumFractionalDigits = 2;
+		CurveToolTipOutputText = FText::Format(LOCTEXT("CurveToolTipValueFormat", "Value:\t{0}"), FText::AsNumber(Value, &FormattingOptions));
+		CurveToolTipInputText = FText::Format(LOCTEXT("CurveToolTipTimeFormat", "Time:\t{0}"), FText::AsNumber(Time, &FormattingOptions));
+
+		if (!CurveToolTip.IsValid())
+		{
+			const FAvaEaseCurveTangents Tangents = EaseCurve->GetTangents();
+
+			SAvaEaseCurvePreview::FArguments PreviewArgs;
+			PreviewArgs.Tangents(Tangents)
+				.PreviewSize(256.f)
+				.CanExpandPreview(true)
+				.CustomToolTip(false)
+				.Animate(true)
+				.DisplayRate(DisplayRate.Get())
+				.DrawMotionTrails(true);
+
+			const TSharedPtr<SWidget> AdditionalContent = SNew(SVerticalBox)
+				+ SVerticalBox::Slot()
+				.AutoHeight()
+				.Padding(15.f, 5.f, 0.f, 0.f)
+				[
+					SNew(STextBlock)
+					.Text(this, &SAvaEaseCurveEditor::GetCurveToolTipInputText)
+				]
+				+ SVerticalBox::Slot()
+				.AutoHeight()
+				.Padding(15.f, 5.f, 0.f, 15.f)
+				[
+					SNew(STextBlock)
+					.Text(this, &SAvaEaseCurveEditor::GetCurveToolTipOutputText)
+				];
+
+			CurveToolTip = SAvaEaseCurvePreviewToolTip::CreateDefaultToolTip(PreviewArgs, AdditionalContent);
+
+			SetToolTip(CurveToolTip);
+		}
+	}
+	else
+	{
+		CurveToolTip.Reset();
+		SetToolTip(CurveToolTip);
+	}
+}
+
+FReply SAvaEaseCurveEditor::OnMouseWheel(const FGeometry& MyGeometry, const FPointerEvent& MouseEvent)
+{
+	if (!bRequireFocusToZoom || HasKeyboardFocus())
+	{
+		ZoomView(FVector2D(MouseEvent.GetWheelDelta(), MouseEvent.GetWheelDelta()));
+		return FReply::Handled();
+	}
+	return FReply::Unhandled();
+}
+
+void SAvaEaseCurveEditor::ZoomView(const FVector2D& InDelta)
+{
+	const FVector2D ZoomDelta = -0.1f * InDelta;
+
+	const float OutputViewSize = ViewMaxOutput - ViewMinOutput;
+	const float OutputChange = OutputViewSize * ZoomDelta.Y;
+
+	const float NewMinOutput = ViewMinOutput - (OutputChange * 0.5f);
+	const float NewMaxOutput = ViewMaxOutput + (OutputChange * 0.5f);
+
+	SetOutputMinMax(NewMinOutput, NewMaxOutput);
+
+	const float InputViewSize = ViewMaxInput - ViewMinInput;
+	const float InputChange = InputViewSize * ZoomDelta.X;
+
+	const float NewMinInput = ViewMinInput - (InputChange * 0.5f);
+	const float NewMaxInput = ViewMaxInput + (InputChange * 0.5f);
+
+	SetInputMinMax(NewMinInput, NewMaxInput);
+}
+
+FReply SAvaEaseCurveEditor::OnKeyDown(const FGeometry& MyGeometry, const FKeyEvent& InKeyEvent)
+{
+	const FModifierKeysState& ModifierKeys = InKeyEvent.GetModifierKeys();
+
+	FReply Reply = FReply::Unhandled();
+
+	if (ModifierKeys.IsAltDown()) // lock tangent
+	{
+		LockTangents = MovingTangents;
+		Reply = FReply::Handled();
+	}
+	else if (ModifierKeys.IsControlDown()) // lock weight
+	{
+		LockTangents = MovingTangents;
+		Reply = FReply::Handled();
+	}
+	else if (ModifierKeys.IsAltDown() && ModifierKeys.IsControlDown()) // mirror tangents and weights
+	{
+		LockTangents = MovingTangents;
+		Reply = FReply::Handled();
+	}
+
+	if (OnKeyDownEvent.IsBound())
+	{
+		Reply = OnKeyDownEvent.Execute(MyGeometry, InKeyEvent);
+	}
+
+	return Reply;
+}
+
+void SAvaEaseCurveEditor::TryStartDrag(const FGeometry& InMyGeometry, const FPointerEvent& InMouseEvent)
+{
+	const bool bLeftMouseButton = InMouseEvent.IsMouseButtonDown(EKeys::LeftMouseButton);
+	const bool bMiddleMouseButton = InMouseEvent.IsMouseButtonDown(EKeys::MiddleMouseButton);
+	const bool bRightMouseButton = InMouseEvent.IsMouseButtonDown(EKeys::RightMouseButton);
+
+	const FVector2D MousePosition = InMyGeometry.AbsoluteToLocal(InMouseEvent.GetScreenSpacePosition());
+	const FVector2D DragVector = MousePosition - MouseDownLocation;
+
+	if (DragVector.SizeSquared() >= FMath::Square(DragThreshold))
+	{
+		if (bLeftMouseButton)
+		{
+			// Check if we should start dragging a tangent.
+			const FSelectedTangent HitTangent = HitTestCubicTangents(InMyGeometry, InMyGeometry.LocalToAbsolute(MouseDownLocation));
+			if (HitTangent.IsValid())
+			{
+				SelectedTangent = HitTangent;
+
+				OnDragStart.ExecuteIfBound();
+
+				DragState = EDragState::DragTangent;
+
+				PreDragTangents = EaseCurve->GetTangents();
+			}
+			else
+			{
+				// Otherwise if the user left clicked on nothing and start a marquee select.
+				DragState = EDragState::MarqueeSelect;
+			}
+		}
+		else if (bMiddleMouseButton)
+		{
+			DragState = EDragState::Pan;
+		}
+		else if (bRightMouseButton)
+		{
+			DragState = EDragState::Pan;
+		}
+		else
+		{
+			DragState = EDragState::None;
+		}
+	}
+}
+
+void SAvaEaseCurveEditor::ProcessDrag(const FGeometry& InMyGeometry, const FPointerEvent& InMouseEvent)
+{
+	const FTrackScaleInfo ScaleInfo(ViewMinInput, ViewMaxInput, ViewMinOutput, ViewMaxOutput, InMyGeometry.GetLocalSize());
+	const FVector2D ScreenDelta = InMouseEvent.GetCursorDelta();
+
+	FVector2D InputDelta;
+	InputDelta.X = ScreenDelta.X / ScaleInfo.PixelsPerInput;
+	InputDelta.Y = -ScreenDelta.Y / ScaleInfo.PixelsPerOutput;
+	
+	switch (DragState)
+	{
+	case EDragState::DragTangent:
+	{
+		const FVector2D MousePositionScreen = InMyGeometry.AbsoluteToLocal(InMouseEvent.GetScreenSpacePosition());
+		MoveSelectedTangent(ScaleInfo, MousePositionScreen - MouseDownLocation);
+		break;
+	}
+	case EDragState::Pan:
+	{
+		// Output is not clamped.
+		const float NewMinOutput = ViewMinOutput - InputDelta.Y;
+		const float NewMaxOutput = ViewMaxOutput - InputDelta.Y;
+		SetOutputMinMax(NewMinOutput, NewMaxOutput);
+
+		// Input maybe clamped if DataMinInput or DataMaxOutput was set.
+		float NewMinInput = ViewMinInput - InputDelta.X;
+		float NewMaxInput = ViewMaxInput - InputDelta.X;
+		ClampViewRangeToDataIfBound(NewMinInput, NewMaxInput, DataMinInput, DataMaxInput, ScaleInfo.ViewInputRange);
+		SetInputMinMax(NewMinInput, NewMaxInput);
+		break;
+	}
+	case EDragState::Zoom:
+	{
+		const FVector2D Delta = FVector2D(ScreenDelta.X * 0.05f, ScreenDelta.X * 0.05f);
+		ZoomView(Delta);
+		break;
+	}
+	default:
+		break;
+	}
+}
+
+void SAvaEaseCurveEditor::EndDrag(const FGeometry& InMyGeometry, const FPointerEvent& InMouseEvent)
+{
+	if (DragState == EDragState::DragTangent)
+	{
+		OnDragEnd.ExecuteIfBound();
+	}
+	else if (DragState == EDragState::MarqueeSelect)
+	{
+		if (!InMouseEvent.IsControlDown() && !InMouseEvent.IsShiftDown())
+		{
+			ClearSelection();
+		}
+	}
+
+	DragState = EDragState::None;
+}
+
+void SAvaEaseCurveEditor::ProcessClick(const FGeometry& InMyGeometry, const FPointerEvent& InMouseEvent)
+{
+	const FKey EffectingButton = InMouseEvent.GetEffectingButton();
+
+	if (EffectingButton == EKeys::LeftMouseButton)
+	{
+		const FSelectedTangent HitTangent = HitTestCubicTangents(InMyGeometry, InMouseEvent.GetScreenSpacePosition());
+		if (HitTangent.IsValid())
+		{
+			SelectedTangent = HitTangent;
+		}
+	}
+	else if (EffectingButton == EKeys::RightMouseButton)
+	{
+		if (GetContextMenuContent.IsBound() && !HitTestCurves(InMyGeometry, InMouseEvent))
+		{
+			const FWidgetPath WidgetPath = InMouseEvent.GetEventPath() != nullptr ? *InMouseEvent.GetEventPath() : FWidgetPath();
+			FSlateApplication::Get().PushMenu(SharedThis(this), WidgetPath, GetContextMenuContent.Execute()
+				, FSlateApplication::Get().GetCursorPos(), FPopupTransitionEffect(FPopupTransitionEffect::ContextMenu));
+		}
+	}
+}
+
+void SAvaEaseCurveEditor::ClearSelection()
+{
+	SelectedTangent = FSelectedTangent();
+}
+
+void SAvaEaseCurveEditor::SetDefaultOutput(const float MinZoomRange)
+{
+	const float NewMinOutput = (ViewMinOutput - (0.5f * MinZoomRange));
+	const float NewMaxOutput = (ViewMaxOutput + (0.5f * MinZoomRange));
+
+	SetOutputMinMax(NewMinOutput, NewMaxOutput);
+}
+
+void SAvaEaseCurveEditor::ZoomToFit()
+{
+	const FRichCurveKey& StartKey = EaseCurve->GetStartKey();
+	const FRichCurveKey& EndKey = EaseCurve->GetEndKey();
+
+	const FVector2D ArriveTangentDir = CalcTangentDir(EndKey.ArriveTangent) * EndKey.ArriveTangentWeight;
+	const FVector2D LeaveTangentDir = CalcTangentDir(StartKey.LeaveTangent) * StartKey.LeaveTangentWeight;
+
+	const FVector2D StartKeyPosition(EaseCurve->FloatCurve.GetKeyTime(EaseCurve->GetStartKeyHandle()), EaseCurve->FloatCurve.GetKeyValue(EaseCurve->GetStartKeyHandle()));
+	const FVector2D EndKeyPosition(EaseCurve->FloatCurve.GetKeyTime(EaseCurve->GetEndKeyHandle()), EaseCurve->FloatCurve.GetKeyValue(EaseCurve->GetEndKeyHandle()));
+	const FVector2D ArrivePosition = -ArriveTangentDir + EndKeyPosition;
+	const FVector2D LeavePosition = LeaveTangentDir - StartKeyPosition;
+
+	// Fit horizontally
+	{
+		float InMinX = FLT_MAX;
+		float InMaxX = -FLT_MAX;
+
+		float MinTimeX = 0.f;
+		float MaxTimeX = 0.f;
+		EaseCurve->FloatCurve.GetTimeRange(MinTimeX, MaxTimeX);
+
+		InMinX = FMath::Min(MinTimeX, InMinX);
+		InMaxX = FMath::Max(MaxTimeX, InMaxX);
+
+		InMinX = FMath::Min(LeavePosition.X, InMinX);
+		InMaxX = FMath::Max(LeavePosition.X, InMaxX);
+
+		InMinX = FMath::Min(ArrivePosition.X, InMinX);
+		InMaxX = FMath::Max(ArrivePosition.X, InMaxX);
+
+		// Clamp the minimum size
+		float SizeX = InMaxX - InMinX;
+		if (SizeX < CONST_MinViewRange)
+		{
+			InMinX -= 0.5f * CONST_MinViewRange;
+			InMaxX += 0.5f * CONST_MinViewRange;
+			SizeX = InMaxX - InMinX;
+		}
+
+		// Add margin
+		InMinX -= CONST_FitMargin * SizeX;
+		InMaxX += CONST_FitMargin * SizeX;
+
+		SetInputMinMax(InMinX, InMaxX);
+	}
+
+	// Fit vertically
+	{
+		float InMinY = FLT_MAX;
+		float InMaxY = -FLT_MAX;
+
+		float MinValY = 0.f;
+		float MaxValY = 0.f;
+		EaseCurve->FloatCurve.GetValueRange(MinValY, MaxValY);
+
+		InMinY = FMath::Min(MinValY, InMinY);
+		InMaxY = FMath::Max(MaxValY, InMaxY);
+
+		InMinY = FMath::Min(LeavePosition.Y, InMinY);
+		InMaxY = FMath::Max(LeavePosition.Y, InMaxY);
+
+		InMinY = FMath::Min(ArrivePosition.Y, InMinY);
+		InMaxY = FMath::Max(ArrivePosition.Y, InMaxY);
+
+		// If in max and in min is same, then include 0.f
+		if (InMaxY == InMinY)
+		{
+			InMaxY = FMath::Max(InMaxY, 0.f);
+			InMinY = FMath::Min(InMinY, 0.f);
+		}
+
+		// Clamp the minimum size
+		float SizeY = InMaxY - InMinY;
+		if (SizeY < CONST_MinViewRange)
+		{
+			SetDefaultOutput(CONST_MinViewRange);
+			InMinY = ViewMinOutput;
+			InMaxY = ViewMaxOutput;
+			SizeY = InMaxY - InMinY;
+		}
+
+		// Add margin
+		const float NewMinOutputY = (InMinY - CONST_FitMargin * SizeY);
+		const float NewMaxOutputY = (InMaxY + CONST_FitMargin * SizeY);
+
+		SetOutputMinMax(NewMinOutputY, NewMaxOutputY);
+	}
+}
+
+FString SAvaEaseCurveEditor::GetReferencerName() const
+{
+	return TEXT("SAvaEaseCurveEditor");
+}
+
+void SAvaEaseCurveEditor::AddReferencedObjects(FReferenceCollector& Collector)
+{
+	Collector.AddReferencedObject(EaseCurve);
+}
+
+void SAvaEaseCurveEditor::SetInputMinMax(const float InNewMin, const float InNewMax)
+{
+	ViewMinInput = InNewMin;
+	ViewMaxInput = InNewMax;
+}
+
+void SAvaEaseCurveEditor::SetOutputMinMax(const float InNewMin, const float InNewMax)
+{
+	ViewMinOutput = InNewMin;
+	ViewMaxOutput = InNewMax;
+}
+
+bool SAvaEaseCurveEditor::HitTestCurves(const FGeometry& InMyGeometry, const FPointerEvent& InMouseEvent) const
+{
+	const FTrackScaleInfo ScaleInfo(ViewMinInput, ViewMaxInput, ViewMinOutput, ViewMaxOutput, InMyGeometry.GetLocalSize());
+	const FVector2D HitPosition = InMyGeometry.AbsoluteToLocal(InMouseEvent.GetScreenSpacePosition());
+	const float Time = ScaleInfo.LocalXToInput(HitPosition.X);
+	const float KeyScreenY = ScaleInfo.OutputToLocalY(EaseCurve->FloatCurve.Eval(Time));
+
+	if (HitPosition.Y > (KeyScreenY - (0.5f * CONST_CurveSize.Y)) && HitPosition.Y < (KeyScreenY + (0.5f * CONST_CurveSize.Y)))
+	{
+		return true;
+	}
+
+	return false;
+}
+
+SAvaEaseCurveEditor::FSelectedTangent SAvaEaseCurveEditor::HitTestCubicTangents(const FGeometry& InMyGeometry, const FVector2D& InHitScreenPosition) const
+{
+	FSelectedTangent OutTangent;
+
+	const FTrackScaleInfo ScaleInfo(ViewMinInput, ViewMaxInput, ViewMinOutput, ViewMaxOutput, InMyGeometry.GetLocalSize());
+	const FVector2D HitPosition = InMyGeometry.AbsoluteToLocal(InHitScreenPosition);
+
+	auto IsInside = [&HitPosition](const FVector2D& InTangent) -> bool
+		{
+			return HitPosition.Y > (InTangent.Y - (0.5f * CONST_CurveSize.Y)) &&
+				HitPosition.Y < (InTangent.Y + (0.5f * CONST_CurveSize.Y)) &&
+				HitPosition.X > (InTangent.X - (0.5f * CONST_TangentSize.X)) &&
+				HitPosition.X < (InTangent.X + (0.5f * CONST_TangentSize.X));
+		};
+
+	for (auto It(EaseCurve->FloatCurve.GetKeyHandleIterator()); It; ++It)
+	{
+		const FKeyHandle KeyHandle = *It;
+		const int32 KeyIndex = EaseCurve->FloatCurve.GetIndexSafe(KeyHandle);
+		check(KeyIndex != INDEX_NONE)
+
+		if (KeyIndex == 0)
+		{
+			const FVector2D LeaveTangentScreenLocation = GetLeaveTangentScreenLocation(ScaleInfo, KeyHandle);
+			if (IsInside(LeaveTangentScreenLocation))
+			{
+				OutTangent = FSelectedTangent(KeyHandle, false);
+				break;
+			}
+		}
+		else if (KeyIndex == 1)
+		{
+			const FVector2D ArriveTangentScreenLocation = GetArriveTangentScreenLocation(ScaleInfo, KeyHandle);
+			if (IsInside(ArriveTangentScreenLocation))
+			{
+				OutTangent = FSelectedTangent(KeyHandle, true);
+				break;
+			}
+		}
+	}
+
+	return OutTangent;
+}
+
+void SAvaEaseCurveEditor::MoveSelectedTangent(const FTrackScaleInfo& InScaleInfo, const FVector2D& InScreenDelta)
+{
+	if (!SelectedTangent.IsValid())
+	{
+		return;
+	}
+
+	const FModifierKeysState ModifierKeysState = FSlateApplication::Get().GetModifierKeys();
+	const bool bIsAltOnly = ModifierKeysState.IsAltDown() && !ModifierKeysState.IsControlDown() && !ModifierKeysState.IsShiftDown();
+	const bool bIsControlOnly = !ModifierKeysState.IsAltDown() && ModifierKeysState.IsControlDown() && !ModifierKeysState.IsShiftDown();
+	const bool bIsControlAlt = ModifierKeysState.IsAltDown() && ModifierKeysState.IsControlDown();
+
+	const FVector2D SelectedKeyPosition(EaseCurve->FloatCurve.GetKeyTime(SelectedTangent.KeyHandle)
+		, EaseCurve->FloatCurve.GetKeyValue(SelectedTangent.KeyHandle));
+
+	const bool bSnapLocation = (bGridSnap.IsSet() && GridSize.IsSet())
+		? (bGridSnap.Get() && !ModifierKeysState.IsShiftDown()) || (!bGridSnap.Get() && ModifierKeysState.IsShiftDown())
+		: false;
+
+	MovingTangents = PreDragTangents;
+	FAvaEaseCurveTangents NewTangents = MovingTangents;
+
+	if (!SelectedTangent.bIsArrival) // Start / leave tangent
+	{
+		const FVector2D LeaveTangentScreenLocation = InScreenDelta + GetLeaveTangentScreenLocation(InScaleInfo
+			, SelectedKeyPosition, PreDragTangents.Start, PreDragTangents.StartWeight);
+
+		FVector2D NewLeaveDir(InScaleInfo.LocalXToInput(LeaveTangentScreenLocation.X)
+			, InScaleInfo.LocalYToOutput(LeaveTangentScreenLocation.Y));
+		NewLeaveDir.X = FMath::Clamp(NewLeaveDir.X, 0.f, 1.f);
+		NewLeaveDir.Y = FMath::Clamp(NewLeaveDir.Y, -10.f, 10.f);
+		if (bSnapLocation)
+		{
+			NewLeaveDir = SnapLocation(NewLeaveDir);
+		}
+		const FVector2D NewLeaveDirDelta = NewLeaveDir - SelectedKeyPosition;
+		
+		MovingTangents.Start = CalcTangent(NewLeaveDirDelta);
+		MovingTangents.StartWeight = NewLeaveDirDelta.Size();
+
+		if (bIsAltOnly) // lock tangent
+		{
+			NewTangents.Start = LockTangents.Start;
+			NewTangents.StartWeight = MovingTangents.StartWeight;
+		}
+		else if (bIsControlOnly) // lock weight
+		{
+			NewTangents.Start = MovingTangents.Start;
+			NewTangents.StartWeight = LockTangents.StartWeight;
+		}
+		else
+		{
+			NewTangents.Start = MovingTangents.Start;
+			NewTangents.StartWeight = MovingTangents.StartWeight;
+
+			if (bIsControlAlt) // mirror tangents and weights
+			{
+				NewTangents.End = NewTangents.Start;
+				NewTangents.EndWeight = MovingTangents.StartWeight;
+			}
+		}
+	}
+	else // End / arrive tangent
+	{
+		const FVector2D ArriveTangentScreenLocation = InScreenDelta + GetArriveTangentScreenLocation(InScaleInfo
+			, SelectedKeyPosition, PreDragTangents.End, PreDragTangents.EndWeight);
+
+		FVector2D NewArriveDir(InScaleInfo.LocalXToInput(ArriveTangentScreenLocation.X)
+			, InScaleInfo.LocalYToOutput(ArriveTangentScreenLocation.Y));
+		NewArriveDir.X = FMath::Clamp(NewArriveDir.X, 0.f, 1.f);
+		NewArriveDir.Y = FMath::Clamp(NewArriveDir.Y, -10.f, 10.f);
+		if (bSnapLocation)
+		{
+			NewArriveDir = SnapLocation(NewArriveDir);
+		}
+		const FVector2D NewArriveDirDelta = NewArriveDir - SelectedKeyPosition;
+
+		MovingTangents.End = CalcTangent(-1.f * NewArriveDirDelta);
+		MovingTangents.EndWeight = NewArriveDirDelta.Size();
+
+		if (bIsAltOnly) // lock tangent
+		{
+			NewTangents.End = LockTangents.End;
+			NewTangents.EndWeight = MovingTangents.EndWeight;
+		}
+		else if (bIsControlOnly) // lock weight
+		{
+			NewTangents.End = MovingTangents.End;
+			NewTangents.EndWeight = LockTangents.EndWeight;
+		}
+		else
+		{
+			NewTangents.End = MovingTangents.End;
+			NewTangents.EndWeight = MovingTangents.EndWeight;
+
+			if (bIsControlAlt) // mirror tangents and weights
+			{
+				NewTangents.Start = NewTangents.End;
+				NewTangents.StartWeight = MovingTangents.EndWeight;
+			}
+		}
+	}
+
+	OnTangentsChanged.ExecuteIfBound(NewTangents);
+}
+
+FVector2d SAvaEaseCurveEditor::GetArriveTangentScreenLocation(const FTrackScaleInfo& InScaleInfo, const FKeyHandle& InKeyHandle) const
+{
+	const FVector2d KeyPosition(EaseCurve->FloatCurve.GetKeyTime(InKeyHandle), EaseCurve->FloatCurve.GetKeyValue(InKeyHandle));
+	const FRichCurveKey& RichKey = EaseCurve->FloatCurve.GetKey(InKeyHandle);
+	return GetArriveTangentScreenLocation(InScaleInfo, KeyPosition, RichKey.ArriveTangent, RichKey.ArriveTangentWeight);
+}
+
+FVector2d SAvaEaseCurveEditor::GetArriveTangentScreenLocation(const FTrackScaleInfo& InScaleInfo
+	, const FVector2d& InKeyPosition, const float InTangent, const float InWeight) const
+{
+	// Get direction vector from tangent
+	const FVector2d ArriveTangentDir = CalcTangentDir(InTangent);
+	const FVector2d ArrivePosition = -ArriveTangentDir + InKeyPosition;
+
+	FVector2d OutScreenArrivePosition = FVector2D(InScaleInfo.InputToLocalX(ArrivePosition.X), InScaleInfo.OutputToLocalY(ArrivePosition.Y));
+
+	// Convert curve position to screen position
+	const FVector2d KeyScreenPosition = FVector2D(InScaleInfo.InputToLocalX(InKeyPosition.X), InScaleInfo.OutputToLocalY(InKeyPosition.Y));
+
+	const FVector2d DeltaScreenPosition = OutScreenArrivePosition - KeyScreenPosition;
+	OutScreenArrivePosition = KeyScreenPosition + (DeltaScreenPosition * InWeight);
+
+	return OutScreenArrivePosition;
+}
+
+FVector2d SAvaEaseCurveEditor::GetLeaveTangentScreenLocation(const FTrackScaleInfo& InScaleInfo, const FKeyHandle& InKeyHandle) const
+{
+	const FVector2d KeyPosition(EaseCurve->FloatCurve.GetKeyTime(InKeyHandle), EaseCurve->FloatCurve.GetKeyValue(InKeyHandle));
+	const FRichCurveKey& RichKey = EaseCurve->FloatCurve.GetKey(InKeyHandle);
+	return GetLeaveTangentScreenLocation(InScaleInfo, KeyPosition, RichKey.LeaveTangent, RichKey.LeaveTangentWeight);
+}
+
+FVector2D SAvaEaseCurveEditor::GetLeaveTangentScreenLocation(const FTrackScaleInfo& InScaleInfo
+	, const FVector2d& InKeyPosition, const float InTangent, const float InWeight) const
+{
+	// Get direction vector from tangent
+	const FVector2d LeaveTangentDir = CalcTangentDir(InTangent);
+	const FVector2d LeavePosition = LeaveTangentDir + InKeyPosition;
+
+	FVector2d OutScreenLeavePosition = FVector2D(InScaleInfo.InputToLocalX(LeavePosition.X), InScaleInfo.OutputToLocalY(LeavePosition.Y));
+	
+	// Convert curve position to screen position
+	const FVector2d KeyScreenPosition = FVector2D(InScaleInfo.InputToLocalX(InKeyPosition.X), InScaleInfo.OutputToLocalY(InKeyPosition.Y));
+
+	const FVector2d DeltaScreenPosition = OutScreenLeavePosition - KeyScreenPosition;
+	OutScreenLeavePosition = KeyScreenPosition + (DeltaScreenPosition * InWeight);
+
+	return OutScreenLeavePosition;
+}
+
+void SAvaEaseCurveEditor::OnObjectPropertyChanged(UObject* Object, FPropertyChangedEvent& PropertyChangedEvent)
+{
+	if (EaseCurve && EaseCurve->GetOwners().Contains(Object))
+	{
+		ClearSelection();
+	}
+}
+
+void SAvaEaseCurveEditor::PostUndo(const bool bInSuccess)
+{
+	ClearSelection();
+}
+
+void SAvaEaseCurveEditor::PostRedo(const bool bInSuccess)
+{
+	PostUndo(bInSuccess);
+}
+
+FVector2D SAvaEaseCurveEditor::SnapLocation(const FVector2D& InLocation) const
+{
+	if (!GridSize.IsSet())
+	{
+		return InLocation;
+	}
+
+	const double SnapGridSize = GridSize.Get();
+	if (SnapGridSize == 0.f)
+	{
+		return InLocation;
+	}
+
+	const double SnapStep = 1.f / SnapGridSize;
+
+	FVector2d OutLocation;
+	OutLocation.X = FMath::RoundToInt(InLocation.X / SnapStep) * SnapStep;
+	OutLocation.Y = FMath::RoundToInt(InLocation.Y / SnapStep) * SnapStep;
+	
+	return OutLocation;
+}
+
+#undef LOCTEXT_NAMESPACE
diff --git a/Engine/Plugins/Experimental/Avalanche/Source/AvalancheSequencer/Private/EaseCurveTool/Widgets/SAvaEaseCurveEditor.h b/Engine/Plugins/Experimental/Avalanche/Source/AvalancheSequencer/Private/EaseCurveTool/Widgets/SAvaEaseCurveEditor.h
new file mode 100644
--- /dev/null
+++ b/Engine/Plugins/Experimental/Avalanche/Source/AvalancheSequencer/Private/EaseCurveTool/Widgets/SAvaEaseCurveEditor.h
@@ -0,0 +1,278 @@
+// Copyright Epic Games, Inc. All Rights Reserved.
+
+#pragma once
+
+#include "Containers/UnrealString.h"
+#include "EaseCurveTool/AvaEaseCurveTool.h"
+#include "EaseCurveTool/AvaEaseCurveTangents.h"
+#include "Curves/CurveBase.h"
+#include "Curves/KeyHandle.h"
+#include "EditorUndoClient.h"
+#include "Framework/SlateDelegates.h"
+#include "Input/Reply.h"
+#include "Layout/Geometry.h"
+#include "Layout/Visibility.h"
+#include "Misc/Attribute.h"
+#include "SCurveEditor.h"
+#include "Templates/SharedPointer.h"
+#include "UObject/GCObject.h"
+#include "Widgets/DeclarativeSyntaxSupport.h"
+#include "Widgets/SCompoundWidget.h"
+#include "Widgets/SWidget.h"
+
+class IDetailCategoryBuilder;
+class FDetailWidgetRow;
+class SHorizontalBox;
+struct FCurveEditorScreenSpace;
+
+DECLARE_DELEGATE_OneParam(FAvaOnEaseCurveChanged, const FAvaEaseCurveTangents& /*InTangents*/)
+
+class SAvaEaseCurveEditor : 
+	public SCompoundWidget,
+	public FGCObject,
+	public FEditorUndoClient
+{
+public:
+	SLATE_BEGIN_ARGS(SAvaEaseCurveEditor)
+		: _DesiredSize(FVector2D(200.f))
+		, _DisplayRate(FFrameRate(30, 1))
+		, _ShowInputGridNumbers(true)
+		, _ShowOutputGridNumbers(true)
+		, _NormalAreaColor(FLinearColor(1.0f, 1.0f, 1.0f, 0.05f))
+		, _NormalBoundsColor(FLinearColor(0.4f, 0.4f, 0.4f, 1.0f))
+		, _GridSnap(false)
+		, _GridColor(FLinearColor(0.1f, 0.1f, 0.1f, 0.1f))
+		, _ExtendedGridColor(FLinearColor(0.02f, 0.02f, 0.02f, 0.5f))
+		, _CurveThickness(2.5f)
+		, _CurveColor(FLinearColor::White)
+	{
+		_Clipping = EWidgetClipping::ClipToBounds;
+	}
+		SLATE_ATTRIBUTE(TOptional<float>, DataMinInput)
+		SLATE_ATTRIBUTE(TOptional<float>, DataMaxInput)
+
+		SLATE_ATTRIBUTE(FVector2D, DesiredSize)
+		SLATE_ATTRIBUTE(FFrameRate, DisplayRate)
+
+		SLATE_ARGUMENT(bool, ShowInputGridNumbers)
+		SLATE_ARGUMENT(bool, ShowOutputGridNumbers)
+
+		SLATE_ARGUMENT(FLinearColor, NormalAreaColor)
+		SLATE_ARGUMENT(FLinearColor, NormalBoundsColor)
+		SLATE_ATTRIBUTE(bool, GridSnap)
+		SLATE_ATTRIBUTE(int32, GridSize)
+		SLATE_ARGUMENT(FLinearColor, GridColor)
+		SLATE_ARGUMENT(FLinearColor, ExtendedGridColor)
+		SLATE_ARGUMENT(float, CurveThickness)
+		SLATE_ARGUMENT(FLinearColor, CurveColor)
+		SLATE_ATTRIBUTE(FAvaEaseCurveTool::EOperation, Operation)
+
+		SLATE_EVENT(FAvaOnEaseCurveChanged, OnTangentsChanged)
+		SLATE_EVENT(FOnGetContent, GetContextMenuContent)
+		SLATE_EVENT(FOnKeyDown, OnKeyDown)
+		SLATE_EVENT(FSimpleDelegate, OnDragStart)
+		SLATE_EVENT(FSimpleDelegate, OnDragEnd)
+	SLATE_END_ARGS()
+
+	void Construct(const FArguments& InArgs, const TObjectPtr<UAvaEaseCurve>& InEaseCurve);
+
+	virtual ~SAvaEaseCurveEditor() override;
+
+	static FVector2d CalcTangentDir(const double InTangent);
+	static float CalcTangent(const FVector2d& InHandleDelta);
+
+	void ZoomToFit();
+
+	/* Set flag that allows scrolling up/down over the widget from the outside without it handling the scroll wheel event. */
+	void SetRequireFocusToZoom(const bool bInRequireFocusToZoom);
+
+	FKeyHandle GetSelectedKeyHandle() const { return SelectedTangent.KeyHandle; }
+
+protected:
+	/** Represents the different states of a drag operation */
+	enum class EDragState
+	{
+		/** There is no active drag operation. */
+		None,
+		/** The user has clicked a mouse button, but has not moved more then the drag threshold. */
+		PreDrag,
+		/** The user is dragging a selected tangent handle. */
+		DragTangent,
+		/** The user is performing a marquee selection of keys. */
+		MarqueeSelect,
+		/** The user is panning the curve view. */
+		Pan,
+		/** The user is zooming the curve view. */
+		Zoom
+	};
+
+	struct FSelectedTangent
+	{
+		FSelectedTangent()
+			: KeyHandle(FKeyHandle::Invalid()), bIsArrival(false)
+		{}
+		FSelectedTangent(const FKeyHandle InKey, const bool bInIsArrival)
+			: KeyHandle(InKey), bIsArrival(bInIsArrival)
+		{}
+
+		bool operator == (const FSelectedTangent& InOther) const
+		{
+			return (KeyHandle == InOther.KeyHandle) && (bIsArrival == InOther.bIsArrival);
+		}
+
+		bool IsValid() const
+		{
+			return KeyHandle != FKeyHandle::Invalid();
+		}
+
+		FKeyHandle KeyHandle;
+
+		/** Indicates if it is the arrival tangent, or the leave tangent */
+		bool bIsArrival;
+	};
+
+	FVector2d GetArriveTangentScreenLocation(const FTrackScaleInfo& InScaleInfo, const FKeyHandle& InKeyHandle) const;
+	FVector2d GetArriveTangentScreenLocation(const FTrackScaleInfo& InScaleInfo
+		, const FVector2d& InKeyPosition, const float InTangent, const float InWeight) const;
+	
+	FVector2d GetLeaveTangentScreenLocation(const FTrackScaleInfo& InScaleInfo, const FKeyHandle& InKeyHandle) const;
+	FVector2d GetLeaveTangentScreenLocation(const FTrackScaleInfo& InScaleInfo
+		, const FVector2d& InKeyPosition, const float InTangent, const float InWeight) const;
+
+	void ClearSelection();
+
+	int32 PaintGrid(const FTrackScaleInfo& InScaleInfo, const FGeometry& InAllottedGeometry, FSlateWindowElementList& OutDrawElements,
+		int32 InLayerId, const FSlateRect& InMyCullingRect, ESlateDrawEffect InDrawEffects) const;
+
+	void PaintNormalBounds(const FTrackScaleInfo& InScaleInfo, const FGeometry& InAllottedGeometry, FSlateWindowElementList& OutDrawElements
+		, int32 InLayerId, const FSlateRect& InMyCullingRect, ESlateDrawEffect InDrawEffects, const FWidgetStyle& InWidgetStyle) const;
+
+	void PaintCurve(const FTrackScaleInfo& InScaleInfo, const FGeometry& InAllottedGeometry, FSlateWindowElementList& OutDrawElements
+		, int32 LayerId, const FSlateRect& MyCullingRect, ESlateDrawEffect DrawEffects, const FWidgetStyle& InWidgetStyle) const;
+
+	int32 PaintKeys(const FTrackScaleInfo& InScaleInfo, const FGeometry& InAllottedGeometry, FSlateWindowElementList& OutDrawElements
+		, int32 InLayerId, const FSlateRect& InMyCullingRect, ESlateDrawEffect InDrawEffects, const FWidgetStyle& InWidgetStyle) const;
+
+	int32 PaintTangentHandle(const FGeometry& InAllottedGeometry, FSlateWindowElementList& OutDrawElements, int32 InLayerId, ESlateDrawEffect InDrawEffects
+		, const FWidgetStyle& InWidgetStyle, const FVector2D& InKeyLocation, const FVector2D& InTangentLocation, const bool bInSelected) const;
+
+	//~ Begin FEditorUndoClient
+	virtual void PostUndo(const bool bInSuccess) override;
+	virtual void PostRedo(const bool bInSuccess) override;
+	//~ End FEditorUndoClient
+
+	//~ Begin FGCObject
+	virtual FString GetReferencerName() const override;
+	virtual void AddReferencedObjects(FReferenceCollector& Collector) override;
+	//~ End FGCObject
+
+	//~ Begin SWidget
+	virtual bool SupportsKeyboardFocus() const override { return true; }
+	virtual FVector2D ComputeDesiredSize(const float InLayoutScaleMultiplier) const override;
+
+	virtual FReply OnMouseButtonDown(const FGeometry& MyGeometry, const FPointerEvent& MouseEvent) override;
+	virtual FReply OnMouseButtonUp(const FGeometry& MyGeometry, const FPointerEvent& MouseEvent) override;
+	virtual FReply OnMouseMove(const FGeometry& MyGeometry, const FPointerEvent& MouseEvent) override;
+	virtual FReply OnMouseWheel(const FGeometry& MyGeometry, const FPointerEvent& MouseEvent) override;
+	virtual FReply OnKeyDown(const FGeometry& MyGeometry, const FKeyEvent& InKeyEvent) override;
+	virtual void   OnMouseCaptureLost(const FCaptureLostEvent& CaptureLostEvent) override;
+
+	virtual TOptional<bool> OnQueryShowFocus(const EFocusCause InFocusCause) const override;
+
+	virtual int32 OnPaint(const FPaintArgs& InArgs, const FGeometry& InAllottedGeometry, const FSlateRect& InMyCullingRect,
+		FSlateWindowElementList& OutDrawElements, int32 InLayerId, const FWidgetStyle& InWidgetStyle, bool bInParentEnabled) const override;
+	//~ End SWidget
+
+	void TryStartDrag(const FGeometry& MyGeometry, const FPointerEvent& MouseEvent);
+	void ProcessDrag(const FGeometry& MyGeometry, const FPointerEvent& MouseEvent);
+	void EndDrag(const FGeometry& MyGeometry, const FPointerEvent& MouseEvent);
+
+	void ProcessClick(const FGeometry& MyGeometry, const FPointerEvent& MouseEvent);
+
+	void ZoomView(const FVector2D& InDelta);
+
+	/* Generates the line(s) for rendering between KeyIndex and the following key. */
+	void CreateLinesForSegment(const ERichCurveInterpMode InInterpMode
+		, const TPair<float, float>& InStartKeyTimeValue, const TPair<float, float>& InEndKeyTimeValue
+		, TArray<FVector2D>& OutLinePoints, TArray<FLinearColor>& OutLineColors
+		, const FTrackScaleInfo& InScaleInfo) const;
+
+	void MoveSelectedTangent(const FTrackScaleInfo& InScaleInfo, const FVector2D& InScreenDelta);
+
+	FVector2D SnapLocation(const FVector2D& InLocation) const;
+
+	FText GetCurveToolTipInputText() const;
+	FText GetCurveToolTipOutputText() const;
+
+	void UpdateCurveToolTip(const FGeometry& MyGeometry, const FPointerEvent& MouseEvent);
+
+	void OnObjectPropertyChanged(UObject* Object, FPropertyChangedEvent& PropertyChangedEvent);
+
+	/** Set Default output values when range is too small */
+	virtual void SetDefaultOutput(const float MinZoomRange);
+
+	void SetInputMinMax(const float InNewMin, const float InNewMax);
+	void SetOutputMinMax(const float InNewMin, const float InNewMax);
+
+	FSelectedTangent HitTestCubicTangents(const FGeometry& InMyGeometry, const FVector2D& InHitScreenPosition) const;
+	bool HitTestCurves(const FGeometry& InMyGeometry, const FPointerEvent& InMouseEvent) const;
+
+	TObjectPtr<UAvaEaseCurve> EaseCurve;
+
+	/** Minimum/Maximum input of data range  */
+	TAttribute<TOptional<float>> DataMinInput;
+	TAttribute<TOptional<float>> DataMaxInput;
+
+	TAttribute<FVector2D> DesiredSize;
+	TAttribute<FFrameRate> DisplayRate;
+
+	bool bShowInputGridNumbers = false;
+	bool bShowOutputGridNumbers = false;
+
+	FLinearColor NormalAreaColor;
+	FLinearColor NormalBoundsColor;
+	TAttribute<bool> bGridSnap;
+	TAttribute<int32> GridSize;
+	FLinearColor GridColor;
+	FLinearColor ExtendedGridColor;
+	float CurveThickness = 1.f;
+	FLinearColor CurveColor;
+	TAttribute<FAvaEaseCurveTool::EOperation> Operation;
+
+	FAvaOnEaseCurveChanged OnTangentsChanged;
+	FOnGetContent GetContextMenuContent;
+	FOnKeyDown OnKeyDownEvent;
+	FSimpleDelegate OnDragStart;
+	FSimpleDelegate OnDragEnd;
+
+	float ViewMinInput = 0.f;
+	float ViewMaxInput = 1.f;
+	float ViewMinOutput = 0.f;
+	float ViewMaxOutput = 1.f;
+
+	FVector2D MouseDownLocation;
+	FVector2D MouseMoveLocation;
+
+	EDragState DragState = EDragState::None;
+	/** The number of pixels which the mouse must move before a drag operation starts. */
+	float DragThreshold = 4.f;
+	/** A handle to the key which was clicked to start a key drag operation. */
+	FKeyHandle DraggedKeyHandle;
+
+	/** Tangent values at the beginning of a drag operation. */
+	FAvaEaseCurveTangents PreDragTangents;
+
+	TSharedPtr<SToolTip> CurveToolTip;
+
+	FText CurveToolTipInputText;
+	FText CurveToolTipOutputText;
+
+	FAvaEaseCurveTangents MovingTangents;
+	FAvaEaseCurveTangents LockTangents;
+
+	bool bDraggingToolSizeValue = false;
+
+	FSelectedTangent SelectedTangent;
+	
+	bool bRequireFocusToZoom = false;
+};
diff --git a/Engine/Plugins/Experimental/Avalanche/Source/AvalancheSequencer/Private/EaseCurveTool/Widgets/SAvaEaseCurvePreset.cpp b/Engine/Plugins/Experimental/Avalanche/Source/AvalancheSequencer/Private/EaseCurveTool/Widgets/SAvaEaseCurvePreset.cpp
new file mode 100644
--- /dev/null
+++ b/Engine/Plugins/Experimental/Avalanche/Source/AvalancheSequencer/Private/EaseCurveTool/Widgets/SAvaEaseCurvePreset.cpp
@@ -0,0 +1,615 @@
+// Copyright Epic Games, Inc. All Rights Reserved.
+
+#include "SAvaEaseCurvePreset.h"
+#include "EaseCurveTool/AvaEaseCurvePreset.h"
+#include "EaseCurveTool/AvaEaseCurveStyle.h"
+#include "EaseCurveTool/AvaEaseCurveSubsystem.h"
+#include "EaseCurveTool/Widgets/SAvaEaseCurvePreview.h"
+#include "Editor.h"
+#include "DetailLayoutBuilder.h"
+#include "Framework/SlateDelegates.h"
+#include "Internationalization/Text.h"
+#include "Styling/AppStyle.h"
+#include "Styling/StyleColors.h"
+#include "Widgets/Images/SImage.h"
+#include "Widgets/Input/SButton.h"
+#include "Widgets/Input/SCheckBox.h"
+#include "Widgets/Input/SSearchBox.h"
+#include "Widgets/Layout/SScrollBox.h"
+#include "Widgets/Layout/SUniformWrapPanel.h"
+#include "Widgets/Layout/SWidgetSwitcher.h"
+
+#define LOCTEXT_NAMESPACE "SAvaEaseCurvePreset"
+
+void SAvaEaseCurvePreset::Construct(const FArguments& InArgs)
+{
+	DisplayRate = InArgs._DisplayRate;
+	OnPresetChanged = InArgs._OnPresetChanged;
+	OnGetNewPresetTangents = InArgs._OnGetNewPresetTangents;
+
+	ChildSlot
+		[
+			SNew(SWidgetSwitcher)
+			.WidgetIndex_Lambda([this]()
+				{
+					return bIsCreatingNewPreset ? 1 : 0;
+				})
+			+ SWidgetSwitcher::Slot()
+			[
+				SNew(SHorizontalBox)
+				+ SHorizontalBox::Slot()
+				.Expose(ComboBoxSlot)
+				[
+					ConstructPresetComboBox()
+				]
+				+ SHorizontalBox::Slot()
+				.AutoWidth()
+				.VAlign(VAlign_Center)
+				.Padding(2.f, 0.f, 0.f, 0.f)
+				[
+					SNew(SButton)
+					.ButtonStyle(FAvaEaseCurveStyle::Get(), TEXT("ToolButton"))
+					.VAlign(VAlign_Center)
+					.ToolTipText(LOCTEXT("AddNewPresetToolTip", "Save the current ease curve as a new preset"))
+					.IsEnabled_Lambda([this]()
+						{
+							return !SelectedItem.IsValid();
+						})
+					.OnClicked(this, &SAvaEaseCurvePreset::OnCreateNewPresetClick)
+					[
+						SNew(SImage)
+						.DesiredSizeOverride(FVector2D(FAvaEaseCurveStyle::Get().GetFloat(TEXT("ToolButton.ImageSize"))))
+						.Image(FAppStyle::GetBrush(TEXT("Icons.Plus")))
+					]
+				]
+			]
+			+ SWidgetSwitcher::Slot()
+			[
+				SNew(SHorizontalBox)
+				+ SHorizontalBox::Slot()
+				[
+					SAssignNew(NewPresetNameTextBox, SEditableTextBox)
+					.OnKeyDownHandler(this, &SAvaEaseCurvePreset::OnNewPresetKeyDownHandler)
+					.OnTextCommitted(this, &SAvaEaseCurvePreset::OnNewPresetTextCommitted)
+				]
+				+ SHorizontalBox::Slot()
+				.AutoWidth()
+				.Padding(2.f, 0.f, 0.f, 0.f)
+				[
+					SNew(SButton)
+					.ButtonStyle(FAvaEaseCurveStyle::Get(), TEXT("ToolButton"))
+					.VAlign(VAlign_Center)
+					.ToolTipText(LOCTEXT("CancelNewPresetToolTip", "Cancels the current new ease curve preset operation"))
+					.IsEnabled_Lambda([this]()
+						{
+							return !SelectedItem.IsValid();
+						})
+					.OnClicked(this, &SAvaEaseCurvePreset::OnCancelNewPresetClick)
+					[
+						SNew(SImage)
+						.DesiredSizeOverride(FVector2D(FAvaEaseCurveStyle::Get().GetFloat(TEXT("ToolButton.ImageSize"))))
+						.Image(FAppStyle::GetBrush(TEXT("Icons.X")))
+					]
+				]
+			]
+		];
+}
+
+TSharedRef<SWidget> SAvaEaseCurvePreset::ConstructPresetComboBox()
+{
+	return SNew(SComboButton)
+		.OnGetMenuContent(this, &SAvaEaseCurvePreset::GeneratePresetDropdown)
+		.OnMenuOpenChanged_Lambda([this](const bool bInOpening)
+			{
+				bIsInEditMode.Set(false);
+			})
+		.ButtonContent()
+		[
+			GenerateSelectedRowWidget()
+		];
+}
+
+TSharedRef<SWidget> SAvaEaseCurvePreset::GeneratePresetDropdown()
+{
+	static const FVector2D ButtonImageSize = FVector2D(FAvaEaseCurveStyle::Get().GetFloat(TEXT("ToolButton.ImageSize")));
+
+	TSharedRef<SWidget> OutWidget = SNew(SVerticalBox)
+		+ SVerticalBox::Slot()
+		.AutoHeight()
+		.Padding(3.f)
+		[
+			SNew(SHorizontalBox)
+			+ SHorizontalBox::Slot()
+			.VAlign(VAlign_Center)
+			[
+				SNew(SSearchBox)
+				.HintText(LOCTEXT("SearchHintLabel", "Search"))
+				.OnTextChanged(this, &SAvaEaseCurvePreset::OnSearchTextChanged)
+			]
+			+ SHorizontalBox::Slot()
+			.AutoWidth()
+			.VAlign(VAlign_Center)
+			.Padding(3.f, 0.f, 0.f, 0.f)
+			[
+				SNew(SButton)
+				.ButtonStyle(FAvaEaseCurveStyle::Get(), TEXT("ToolButton"))
+				.VAlign(VAlign_Center)
+				.ToolTipText(LOCTEXT("ReloadJsonPresetsToolTip", "Reload ease curve presets from Json files"))
+				.OnClicked(this, &SAvaEaseCurvePreset::ReloadJsonPresets)
+				[
+					SNew(SImage)
+					.DesiredSizeOverride(ButtonImageSize)
+					.Image(FAppStyle::GetBrush(TEXT("Icons.Refresh")))
+				]
+			]
+			+ SHorizontalBox::Slot()
+			.AutoWidth()
+			.VAlign(VAlign_Center)
+			.Padding(3.f, 0.f, 0.f, 0.f)
+			[
+				SNew(SButton)
+				.ButtonStyle(FAvaEaseCurveStyle::Get(), TEXT("ToolButton"))
+				.VAlign(VAlign_Center)
+				.ToolTipText(LOCTEXT("ExploreJsonPresetsFolderToolTip", "Opens the folder location for the Json ease curve presets"))
+				.OnClicked(this, &SAvaEaseCurvePreset::ExploreJsonPresetsFolder)
+				[
+					SNew(SImage)
+					.DesiredSizeOverride(ButtonImageSize)
+					.Image(FAppStyle::GetBrush(TEXT("Icons.FolderOpen")))
+				]
+			]
+			+ SHorizontalBox::Slot()
+			.AutoWidth()
+			.VAlign(VAlign_Center)
+			.Padding(3.f, 0.f, 3.f, 0.f)
+			[
+				SNew(SCheckBox)
+				.Style(FAvaEaseCurveStyle::Get(), TEXT("ToolToggleButton"))
+				.Padding(4.f)
+				.ToolTipText(LOCTEXT("ToggleEditModeToolTip", "Enable editing of ease curve presets and categories"))
+				.IsChecked_Lambda([this]()
+					{
+						return bIsInEditMode.Get(false) ? ECheckBoxState::Checked : ECheckBoxState::Unchecked;
+					})
+				.OnCheckStateChanged(this, &SAvaEaseCurvePreset::ToggleEditMode)
+				[
+					SNew(SImage)
+					.DesiredSizeOverride(ButtonImageSize)
+					.Image(FAppStyle::GetBrush(TEXT("Icons.Edit")))
+				]
+			]
+			+ SHorizontalBox::Slot()
+			.AutoWidth()
+			.VAlign(VAlign_Center)
+			[
+				SNew(SButton)
+				.ButtonStyle(FAvaEaseCurveStyle::Get(), TEXT("ToolButton"))
+				.VAlign(VAlign_Center)
+				.ToolTipText(LOCTEXT("CreateCategoryToolTip", "Creates a new empty category"))
+				.Visibility_Lambda([this]()
+					{
+						return bIsInEditMode.Get(false) ? EVisibility::Visible : EVisibility::Collapsed;
+					})
+				.OnClicked(this, &SAvaEaseCurvePreset::CreateNewCategory)
+				[
+					SNew(SImage)
+					.DesiredSizeOverride(ButtonImageSize)
+					.Image(FAppStyle::GetBrush(TEXT("Icons.Plus")))
+				]
+			]
+		]
+		+ SVerticalBox::Slot()
+		.AutoHeight()
+		.HAlign(HAlign_Fill)
+		.Padding(3.f, 0.f, 3.f, 3.f)
+		[
+			SNew(SBox)
+			.MaxDesiredHeight(960.f)
+			[
+				SNew(SScrollBox)
+				+ SScrollBox::Slot()
+				[
+					SAssignNew(GroupWidgetsParent, SBox)
+				]
+			]
+		];
+
+	RegenerateGroupWrapBox();
+
+	return OutWidget;
+}
+
+void SAvaEaseCurvePreset::UpdateGroupsContent()
+{
+	const UAvaEaseCurveSubsystem& EaseCurveSubsystem = UAvaEaseCurveSubsystem::Get();
+	const TArray<FString>& EaseCurveCategories = EaseCurveSubsystem.GetEaseCurveCategories();
+
+	auto GenerateNoPresetsWidget = [](const FText& InText)
+		{
+			return SNew(SBox)
+				.WidthOverride(300.f)
+				.HeightOverride(200.f)
+				.HAlign(HAlign_Center)
+				.VAlign(VAlign_Center)
+				[
+					SNew(STextBlock)
+					.TextStyle(FAppStyle::Get(), TEXT("HintText"))
+					.Text(InText)
+				];
+		};
+
+	if (EaseCurveCategories.Num() == 0)
+	{
+		GroupWidgetsParent->SetContent(GenerateNoPresetsWidget(LOCTEXT("NoPresetsLabel", "No ease curve presets")));
+		return;
+	}
+
+	if (!SearchText.IsEmpty())
+	{
+		int32 TotalVisiblePresets = 0;
+		for (const TSharedPtr<SAvaEaseCurvePresetGroup>& GroupWidget : GroupWidgets)
+		{
+			TotalVisiblePresets += GroupWidget->GetVisiblePresetCount();
+		}
+		if (TotalVisiblePresets == 0)
+		{
+			GroupWidgetsParent->SetContent(GenerateNoPresetsWidget(LOCTEXT("NoPresetsFoundLabel", "No ease curve presets found")));
+			return;
+		}
+	}
+
+	GroupWidgetsParent->SetContent(GroupWrapBox.ToSharedRef());
+}
+
+void SAvaEaseCurvePreset::RegenerateGroupWrapBox()
+{
+	const UAvaEaseCurveSubsystem& EaseCurveSubsystem = UAvaEaseCurveSubsystem::Get();
+	const TArray<FString>& EaseCurveCategories = EaseCurveSubsystem.GetEaseCurveCategories();
+	const int32 CurvePresetCount = EaseCurveCategories.Num();
+
+	GroupWidgets.Empty(CurvePresetCount);
+
+	GroupWrapBox = SNew(SUniformWrapPanel)
+		.HAlign(HAlign_Center)
+		.SlotPadding(FMargin(2.f, 1.f))
+		.EvenRowDistribution(true)
+		.NumColumnsOverride_Lambda([]()
+			{
+				const int32 EaseCurveCategoryCount = UAvaEaseCurveSubsystem::Get().GetEaseCurveCategories().Num();
+				return FMath::Min(5, EaseCurveCategoryCount);
+			});
+
+	for (const FString& Category : EaseCurveCategories)
+	{
+		TSharedRef<SAvaEaseCurvePresetGroup> NewGroupWidget = SNew(SAvaEaseCurvePresetGroup)
+			.CategoryName(Category)
+			.Presets(EaseCurveSubsystem.GetEaseCurvePresets(Category))
+			.IsEditMode(bIsInEditMode)
+			.DisplayRate(DisplayRate.Get())
+			.OnCategoryDelete(this, &SAvaEaseCurvePreset::HandleCategoryDelete)
+			.OnCategoryRename(this, &SAvaEaseCurvePreset::HandleCategoryRename)
+			.OnPresetDelete(this, &SAvaEaseCurvePreset::HandlePresetDelete)
+			.OnPresetRename(this, &SAvaEaseCurvePreset::HandlePresetRename)
+			.OnBeginPresetMove(this, &SAvaEaseCurvePreset::HandleBeginPresetMove)
+			.OnEndPresetMove(this, &SAvaEaseCurvePreset::HandleEndPresetMove)
+			.OnPresetClick(this, &SAvaEaseCurvePreset::HandlePresetClick);
+
+		GroupWidgets.Add(NewGroupWidget);
+
+		GroupWrapBox->AddSlot()
+			.HAlign(HAlign_Left)
+			[
+				NewGroupWidget
+			];
+	}
+
+	UpdateGroupsContent();
+}
+
+TSharedRef<SWidget> SAvaEaseCurvePreset::GenerateSelectedRowWidget() const
+{
+	if (!SelectedItem.IsValid())
+	{
+		return SNew(STextBlock)
+			.Font(IDetailLayoutBuilder::GetDetailFont())
+			.Text(FText::FromString("Select Preset..."));
+	}
+
+	return SNew(SHorizontalBox)
+		+ SHorizontalBox::Slot()
+		.HAlign(HAlign_Fill)
+		.VAlign(VAlign_Center)
+		[
+			SNew(SHorizontalBox)
+			+ SHorizontalBox::Slot()
+			.HAlign(HAlign_Left)
+			[
+				SNew(STextBlock)
+				.Font(IDetailLayoutBuilder::GetDetailFont())
+				.ColorAndOpacity(FStyleColors::Foreground)
+				.OverflowPolicy(ETextOverflowPolicy::Ellipsis)
+				.Text_Lambda([this]()
+					{
+						return SelectedItem.IsValid() ? FText::FromString(SelectedItem->Name) : FText();
+					})
+			]
+			+ SHorizontalBox::Slot()
+			.HAlign(HAlign_Right)
+			[
+				SNew(STextBlock)
+				.Font(IDetailLayoutBuilder::GetDetailFont())
+				.ColorAndOpacity(FStyleColors::White25)
+				.OverflowPolicy(ETextOverflowPolicy::Ellipsis)
+				.Text_Lambda([this]()
+					{
+						return SelectedItem.IsValid() ? FText::FromString(SelectedItem->Category) : FText();
+					})
+			]
+
+		]
+		+ SHorizontalBox::Slot()
+		.AutoWidth()
+		.HAlign(HAlign_Right)
+		.VAlign(VAlign_Center)
+		.Padding(5.f, 2.f, 4.f, 2.f)
+		[
+			SNew(SBorder)
+			.BorderBackgroundColor(FStyleColors::White25)
+			[
+				SNew(SAvaEaseCurvePreview)
+				.PreviewSize(12.f)
+				.CustomToolTip(true)
+				.DisplayRate(DisplayRate.Get())
+				.Tangents_Lambda([this]()
+					{
+						return SelectedItem.IsValid() ? SelectedItem->Tangents : FAvaEaseCurveTangents();
+					})
+			]
+		];
+}
+
+FReply SAvaEaseCurvePreset::OnCreateNewPresetClick()
+{
+	bIsCreatingNewPreset = true;
+
+	FSlateApplication::Get().SetAllUserFocus(NewPresetNameTextBox);
+
+	return FReply::Handled();
+}
+
+FReply SAvaEaseCurvePreset::OnCancelNewPresetClick()
+{
+	bIsCreatingNewPreset = false;
+	NewPresetNameTextBox->SetText(FText());
+
+	return FReply::Handled();
+}
+
+FReply SAvaEaseCurvePreset::OnDeletePresetClick()
+{
+	if (SelectedItem.IsValid())
+	{
+		if (UAvaEaseCurveSubsystem* EaseCurveSubsystem = GEditor->GetEditorSubsystem<UAvaEaseCurveSubsystem>())
+		{
+			EaseCurveSubsystem->RemovePreset(*SelectedItem);
+		}
+
+		ClearSelection();
+	}
+
+	return FReply::Handled();
+}
+
+FReply SAvaEaseCurvePreset::OnNewPresetKeyDownHandler(const FGeometry& InGeometry, const FKeyEvent& InKeyEvent)
+{
+	if (InKeyEvent.GetKey() == EKeys::Escape)
+	{
+		bIsCreatingNewPreset = false;
+		NewPresetNameTextBox->SetText(FText());
+
+		return FReply::Handled();
+	}
+
+	return FReply::Unhandled();
+}
+
+void SAvaEaseCurvePreset::OnNewPresetTextCommitted(const FText& InNewText, ETextCommit::Type InCommitType)
+{
+	if (InCommitType == ETextCommit::OnEnter && !InNewText.IsEmpty())
+	{
+		FAvaEaseCurveTangents NewTangents;
+		if (OnGetNewPresetTangents.IsBound() && OnGetNewPresetTangents.Execute(NewTangents))
+		{
+			UAvaEaseCurveSubsystem& EaseCurveSubsystem = UAvaEaseCurveSubsystem::Get();
+
+			SelectedItem = EaseCurveSubsystem.AddPreset(InNewText.ToString(), NewTangents);
+
+			ComboBoxSlot->AttachWidget(ConstructPresetComboBox());
+		}
+	}
+
+	bIsCreatingNewPreset = false;
+	NewPresetNameTextBox->SetText(FText());
+}
+
+void SAvaEaseCurvePreset::OnSearchTextChanged(const FText& InSearchText)
+{
+	SearchText = InSearchText;
+
+	for (const TSharedPtr<SAvaEaseCurvePresetGroup>& GroupWidget : GroupWidgets)
+	{
+		GroupWidget->SetSearchText(SearchText);
+	}
+
+	UpdateGroupsContent();
+}
+
+void SAvaEaseCurvePreset::ClearSelection()
+{
+	SelectedItem.Reset();
+	ComboBoxSlot->AttachWidget(ConstructPresetComboBox());
+}
+
+bool SAvaEaseCurvePreset::GetSelectedItem(FAvaEaseCurvePreset& OutPreset) const
+{
+	if (!SelectedItem.IsValid())
+	{
+		return false;
+	}
+
+	OutPreset = *SelectedItem;
+
+	return true;
+}
+
+bool SAvaEaseCurvePreset::SetSelectedItem(const FString& InName)
+{
+	UAvaEaseCurveSubsystem& EaseCurveSubsystem = UAvaEaseCurveSubsystem::Get();
+
+	const TSharedPtr<FAvaEaseCurvePreset>& FoundItem = EaseCurveSubsystem.FindPreset(InName);
+	if (!FoundItem.IsValid())
+	{
+		return false;
+	}
+
+	SelectedItem = FoundItem;
+	ComboBoxSlot->AttachWidget(ConstructPresetComboBox());
+
+	return true;
+}
+
+bool SAvaEaseCurvePreset::SetSelectedItem(const FAvaEaseCurvePreset& InPreset)
+{
+	return SetSelectedItem(InPreset.Name);
+}
+
+bool SAvaEaseCurvePreset::SetSelectedItem(const FAvaEaseCurveTangents& InTangents)
+{
+	UAvaEaseCurveSubsystem& EaseCurveSubsystem = UAvaEaseCurveSubsystem::Get();
+
+	const TSharedPtr<FAvaEaseCurvePreset>& FoundItem = EaseCurveSubsystem.FindPresetByTangents(InTangents);
+	if (!FoundItem.IsValid())
+	{
+		return false;
+	}
+
+	SelectedItem = FoundItem;
+	ComboBoxSlot->AttachWidget(ConstructPresetComboBox());
+
+	return true;
+}
+
+FReply SAvaEaseCurvePreset::ReloadJsonPresets()
+{
+	UAvaEaseCurveSubsystem::Get().ReloadPresetsFromJson();
+
+	RegenerateGroupWrapBox();
+
+	return FReply::Handled();
+}
+
+FReply SAvaEaseCurvePreset::ExploreJsonPresetsFolder()
+{
+	UAvaEaseCurveSubsystem::Get().ExploreJsonPresetsFolder();
+
+	return FReply::Handled();
+}
+
+FReply SAvaEaseCurvePreset::CreateNewCategory()
+{
+	UAvaEaseCurveSubsystem::Get().AddNewPresetCategory();
+
+	ReloadJsonPresets();
+
+	return FReply::Handled();
+}
+
+void SAvaEaseCurvePreset::ToggleEditMode(const ECheckBoxState bInNewState)
+{
+	bIsInEditMode.Set(bInNewState == ECheckBoxState::Checked ? true : false);
+
+	RegenerateGroupWrapBox();
+}
+
+bool SAvaEaseCurvePreset::HandleCategoryDelete(const FString& InCategoryName)
+{
+	if (!UAvaEaseCurveSubsystem::Get().RemovePresetCategory(InCategoryName))
+	{
+		return false;
+	}
+
+	ReloadJsonPresets();
+
+	return true;
+}
+
+bool SAvaEaseCurvePreset::HandleCategoryRename(const FString& InCategoryName, const FString& InNewName)
+{
+	return UAvaEaseCurveSubsystem::Get().RenamePresetCategory(InCategoryName, InNewName);
+}
+
+bool SAvaEaseCurvePreset::HandlePresetDelete(const TSharedPtr<FAvaEaseCurvePreset>& InPreset)
+{
+	return UAvaEaseCurveSubsystem::Get().RemovePreset(*InPreset);
+}
+
+bool SAvaEaseCurvePreset::HandlePresetRename(const TSharedPtr<FAvaEaseCurvePreset>& InPreset, const FString& InNewName)
+{
+	return UAvaEaseCurveSubsystem::Get().RenamePreset(InPreset->Category, InPreset->Name, InNewName);
+}
+
+bool SAvaEaseCurvePreset::HandleBeginPresetMove(const TSharedPtr<FAvaEaseCurvePreset>& InPreset, const FString& InNewCategoryName)
+{
+	for (const TSharedPtr<SAvaEaseCurvePresetGroup>& GroupWidget : GroupWidgets)
+	{
+		if (!GroupWidget->GetCategoryName().Equals(InNewCategoryName))
+		{
+			GroupWidget->NotifyCanDrop(true);
+		}
+	}
+
+	return true;
+}
+
+bool SAvaEaseCurvePreset::HandleEndPresetMove(const TSharedPtr<FAvaEaseCurvePreset>& InPreset, const FString& InNewCategoryName)
+{
+	for (const TSharedPtr<SAvaEaseCurvePresetGroup>& GroupWidget : GroupWidgets)
+	{
+		if (!GroupWidget->GetCategoryName().Equals(InNewCategoryName))
+		{
+			GroupWidget->NotifyCanDrop(false);
+		}
+	}
+
+	if (!InPreset.IsValid() || InPreset->Category.Equals(InNewCategoryName))
+	{
+		return false;
+	}
+
+	if (!UAvaEaseCurveSubsystem::Get().ChangePresetCategory(InPreset, InNewCategoryName))
+	{
+		return false;
+	}
+
+	ReloadJsonPresets();
+
+	return true;
+}
+
+bool SAvaEaseCurvePreset::HandlePresetClick(const TSharedPtr<FAvaEaseCurvePreset>& InPreset)
+{
+	if (ComboBoxSlot)
+	{
+		FSlateApplication::Get().SetAllUserFocus(ComboBoxSlot->GetWidget());
+
+		SetSelectedItem(*InPreset);
+
+		OnPresetChanged.ExecuteIfBound(InPreset);
+
+		return true;
+	}
+
+	return false;
+}
+
+#undef LOCTEXT_NAMESPACE
diff --git a/Engine/Plugins/Experimental/Avalanche/Source/AvalancheSequencer/Private/EaseCurveTool/Widgets/SAvaEaseCurvePreset.h b/Engine/Plugins/Experimental/Avalanche/Source/AvalancheSequencer/Private/EaseCurveTool/Widgets/SAvaEaseCurvePreset.h
new file mode 100644
--- /dev/null
+++ b/Engine/Plugins/Experimental/Avalanche/Source/AvalancheSequencer/Private/EaseCurveTool/Widgets/SAvaEaseCurvePreset.h
@@ -0,0 +1,89 @@
+// Copyright Epic Games, Inc. All Rights Reserved.
+
+#pragma once
+
+#include "EaseCurveTool/AvaEaseCurveTangents.h"
+#include "EaseCurveTool/Widgets/SAvaEaseCurvePresetGroup.h"
+#include "Templates/SharedPointer.h"
+#include "Widgets/Input/SComboBox.h"
+#include "Widgets/SBoxPanel.h"
+#include "Widgets/SCompoundWidget.h"
+
+class FName;
+class FText;
+class SUniformWrapPanel;
+struct FAvaEaseCurvePreset;
+
+DECLARE_DELEGATE_OneParam(FAvaOnPresetChanged, const TSharedPtr<FAvaEaseCurvePreset>&)
+DECLARE_DELEGATE_RetVal_OneParam(bool, FAvaOnGetNewPresetTangents, FAvaEaseCurveTangents& /*InTangents*/)
+
+class SAvaEaseCurvePreset : public SCompoundWidget
+{
+public:
+	SLATE_BEGIN_ARGS(SAvaEaseCurvePreset)
+		: _DisplayRate(FFrameRate(30, 1))
+	{}
+		SLATE_ATTRIBUTE(FFrameRate, DisplayRate)
+		SLATE_EVENT(FAvaOnPresetChanged, OnPresetChanged)
+		SLATE_EVENT(FAvaOnGetNewPresetTangents, OnGetNewPresetTangents)
+	SLATE_END_ARGS()
+
+	void Construct(const FArguments& InArgs);
+
+	void ClearSelection();
+	bool GetSelectedItem(FAvaEaseCurvePreset& OutPreset) const;
+	bool SetSelectedItem(const FString& InName);
+	bool SetSelectedItem(const FAvaEaseCurvePreset& InPreset);
+	bool SetSelectedItem(const FAvaEaseCurveTangents& InTangents);
+
+protected:
+	TSharedRef<SWidget> ConstructPresetComboBox();
+	TSharedRef<SWidget> GeneratePresetDropdown();
+	
+	TSharedRef<SWidget> GenerateSelectedRowWidget() const;
+
+	void UpdateGroupsContent();
+	void RegenerateGroupWrapBox();
+
+	FReply OnCreateNewPresetClick();
+	FReply OnCancelNewPresetClick();
+	FReply OnDeletePresetClick();
+
+	FReply OnNewPresetKeyDownHandler(const FGeometry& InGeometry, const FKeyEvent& InKeyEvent);
+	void OnNewPresetTextCommitted(const FText& InNewText, ETextCommit::Type InCommitType);
+
+	void OnSearchTextChanged(const FText& InSearchText);
+
+	FReply ReloadJsonPresets();
+	FReply ExploreJsonPresetsFolder();
+	FReply CreateNewCategory();
+
+	void ToggleEditMode(const ECheckBoxState bInNewState);
+
+	bool HandleCategoryDelete(const FString& InCategoryName);
+	bool HandleCategoryRename(const FString& InCategoryName, const FString& InNewName);
+	bool HandlePresetDelete(const TSharedPtr<FAvaEaseCurvePreset>& InPreset);
+	bool HandlePresetRename(const TSharedPtr<FAvaEaseCurvePreset>& InPreset, const FString& InNewName);
+	bool HandleBeginPresetMove(const TSharedPtr<FAvaEaseCurvePreset>& InPreset, const FString& InNewCategoryName);
+	bool HandleEndPresetMove(const TSharedPtr<FAvaEaseCurvePreset>& InPreset, const FString& InNewCategoryName);
+	bool HandlePresetClick(const TSharedPtr<FAvaEaseCurvePreset>& InPreset);
+
+	TAttribute<FFrameRate> DisplayRate;
+
+	FAvaOnPresetChanged OnPresetChanged;
+	FAvaOnGetNewPresetTangents OnGetNewPresetTangents;
+
+	TSharedPtr<SBox> GroupWidgetsParent;
+	TSharedPtr<SUniformWrapPanel> GroupWrapBox;
+	TArray<TSharedPtr<SAvaEaseCurvePresetGroup>> GroupWidgets;
+
+	SHorizontalBox::FSlot* ComboBoxSlot = nullptr;
+	TSharedPtr<SEditableTextBox> NewPresetNameTextBox;
+
+	bool bIsCreatingNewPreset = false;
+	TAttribute<bool> bIsInEditMode = false;
+
+	TSharedPtr<FAvaEaseCurvePreset> SelectedItem;
+
+	FText SearchText;
+};
diff --git a/Engine/Plugins/Experimental/Avalanche/Source/AvalancheSequencer/Private/EaseCurveTool/Widgets/SAvaEaseCurvePresetGroup.cpp b/Engine/Plugins/Experimental/Avalanche/Source/AvalancheSequencer/Private/EaseCurveTool/Widgets/SAvaEaseCurvePresetGroup.cpp
new file mode 100644
--- /dev/null
+++ b/Engine/Plugins/Experimental/Avalanche/Source/AvalancheSequencer/Private/EaseCurveTool/Widgets/SAvaEaseCurvePresetGroup.cpp
@@ -0,0 +1,373 @@
+// Copyright Epic Games, Inc. All Rights Reserved.
+
+#include "SAvaEaseCurvePresetGroup.h"
+#include "EaseCurveTool/AvaEaseCurveStyle.h"
+#include "EaseCurveTool/AvaEaseCurvePreset.h"
+#include "EaseCurveTool/AvaEaseCurvePresetDragDropOp.h"
+#include "Internationalization/Text.h"
+#include "Styling/AppStyle.h"
+#include "Styling/StyleColors.h"
+#include "Widgets/Images/SImage.h"
+#include "Widgets/Input/SButton.h"
+#include "Widgets/Input/SEditableTextBox.h"
+#include "Widgets/Layout/SSeparator.h"
+#include "Widgets/Layout/SWidgetSwitcher.h"
+#include "Widgets/SBoxPanel.h"
+#include "Widgets/SOverlay.h"
+#include "Widgets/Views/SListView.h"
+#include "Widgets/Views/STableViewBase.h"
+
+#define LOCTEXT_NAMESPACE "SAvaEaseCurvePresetGroup"
+
+void SAvaEaseCurvePresetGroup::Construct(const FArguments& InArgs)
+{
+	CategoryName = InArgs._CategoryName;
+	Presets = InArgs._Presets;
+	SearchText = InArgs._SearchText;
+	bIsEditMode = InArgs._IsEditMode;
+	DisplayRate = InArgs._DisplayRate;
+	OnCategoryDelete = InArgs._OnCategoryDelete;
+	OnCategoryRename = InArgs._OnCategoryRename;
+	OnPresetDelete = InArgs._OnPresetDelete;
+	OnPresetRename = InArgs._OnPresetRename;
+	OnBeginPresetMove = InArgs._OnBeginPresetMove;
+	OnEndPresetMove = InArgs._OnEndPresetMove;
+	OnPresetClick = InArgs._OnPresetClick;
+
+	ChildSlot
+		[
+			SNew(SOverlay)
+			+ SOverlay::Slot()
+			[
+				SNew(SBorder)
+				.Visibility(this, &SAvaEaseCurvePresetGroup::GetEditModeVisibility)
+				.BorderImage(this, &SAvaEaseCurvePresetGroup::GetBorderImage)
+			]
+			+ SOverlay::Slot()
+			.Padding(1.f)
+			[
+				SNew(SBox)
+				.WidthOverride(140.f)
+				[
+					SNew(SVerticalBox)
+					+ SVerticalBox::Slot()
+					.AutoHeight()
+					.Padding(2.f)
+					[
+						ConstructHeader()
+					]
+					+ SVerticalBox::Slot()
+					.AutoHeight()
+					[
+						SNew(SListView<TSharedPtr<FAvaEaseCurvePreset>>)
+						.ListItemsSource(&VisiblePresets)
+						.SelectionMode_Lambda([this]()
+							{
+								return IsEditMode() ? ESelectionMode::None : ESelectionMode::Single;
+							})
+						.ListViewStyle(&FAppStyle::Get().GetWidgetStyle<FTableViewStyle>(TEXT("PropertyTable.InViewport.ListView")))
+						.OnGenerateRow(this, &SAvaEaseCurvePresetGroup::GeneratePresetWidget)
+					]
+				]
+			]
+			+ SOverlay::Slot()
+			.HAlign(HAlign_Right)
+			.VAlign(VAlign_Top)
+			.Padding(1.f)
+			[
+				SNew(SButton)
+				.ButtonStyle(FAvaEaseCurveStyle::Get(), TEXT("ToolButton.NoPad"))
+				.VAlign(VAlign_Center)
+				.ToolTipText(LOCTEXT("EditModeDeleteTooltip", "Delete this category and the json file associated with it on disk"))
+				.Visibility(this, &SAvaEaseCurvePresetGroup::GetEditModeVisibility)
+				.OnClicked(this, &SAvaEaseCurvePresetGroup::HandleCategoryDelete)
+				[
+					SNew(SImage)
+					.DesiredSizeOverride(FVector2D(10.f))
+					.Image(FAppStyle::GetBrush(TEXT("Icons.Delete")))
+				]
+			]
+		];
+	
+	SetSearchText(SearchText);
+}
+
+TSharedRef<SWidget> SAvaEaseCurvePresetGroup::ConstructHeader()
+{
+	return SNew(SHorizontalBox)
+		+ SHorizontalBox::Slot()
+		.AutoWidth()
+		.VAlign(VAlign_Center)
+		.Padding(3.f, 2.f)
+		[
+			SNew(SBox)
+			.MaxDesiredWidth(180.f)
+			.Clipping(EWidgetClipping::ClipToBoundsAlways)
+			[
+				SNew(SWidgetSwitcher)
+				.WidgetIndex_Lambda([this]()
+					{
+						return IsEditMode() ? 1 : 0;
+					})
+				+ SWidgetSwitcher::Slot()
+				[
+					SNew(STextBlock)
+					.TextStyle(FAppStyle::Get(), TEXT("Menu.Heading"))
+					.OverflowPolicy(ETextOverflowPolicy::Ellipsis)
+					.Text(FText::FromString(CategoryName))
+					.ToolTipText(this, &SAvaEaseCurvePresetGroup::GetPresetNameTooltipText)
+				]
+				+ SWidgetSwitcher::Slot()
+				[
+					SAssignNew(RenameCategoryNameTextBox, SEditableTextBox)
+					.OverflowPolicy(ETextOverflowPolicy::Ellipsis)
+					.Text(FText::FromString(CategoryName))
+					.ToolTipText(this, &SAvaEaseCurvePresetGroup::GetPresetNameTooltipText)
+					.OnTextCommitted(this, &SAvaEaseCurvePresetGroup::HandleCategoryRenameCommitted)
+				]
+			]
+		]
+		+ SHorizontalBox::Slot()
+		.VAlign(VAlign_Center)
+		.Padding(3.f, 0.f)
+		[
+			SNew(SSeparator)
+			.SeparatorImage(FAppStyle::Get().GetBrush(TEXT("Menu.Separator")))
+			.Thickness(1.1f)
+		];
+}
+
+TSharedRef<ITableRow> SAvaEaseCurvePresetGroup::GeneratePresetWidget(TSharedPtr<FAvaEaseCurvePreset> InPreset, const TSharedRef<STableViewBase>& InOwnerTable)
+{
+	TSharedRef<SAvaEaseCurvePresetGroupItem> NewPresetWidget = SNew(SAvaEaseCurvePresetGroupItem, InOwnerTable)
+		.Preset(InPreset)
+		.IsEditMode(bIsEditMode)
+		.DisplayRate(DisplayRate)
+		.OnClick(this, &SAvaEaseCurvePresetGroup::HandlePresetClick)
+		.OnDelete(this, &SAvaEaseCurvePresetGroup::HandlePresetDelete)
+		.OnRename(this, &SAvaEaseCurvePresetGroup::HandlePresetRename)
+		.OnBeginMove(this, &SAvaEaseCurvePresetGroup::HandlePresetBeginMove)
+		.OnEndMove(this, &SAvaEaseCurvePresetGroup::HandlePresetEndMove);
+	PresetWidgetsMap.Add(InPreset, NewPresetWidget);
+
+	return NewPresetWidget;
+}
+
+void SAvaEaseCurvePresetGroup::SetSearchText(const FText& InText)
+{
+	SearchText = InText;
+
+	VisiblePresets.Empty();
+
+	const FString SearchString = SearchText.ToString();
+
+	for (TSharedPtr<FAvaEaseCurvePreset> Preset : Presets)
+	{
+		if (Preset.IsValid() && (Preset->Name.Contains(SearchString) || Preset->Category.Contains(SearchString)))
+		{
+			VisiblePresets.Add(Preset);
+		}
+	}
+
+	const bool bIsVisible = SearchText.IsEmpty() || VisiblePresets.Num() > 0;
+	SetVisibility(bIsVisible ? EVisibility::Visible : EVisibility::Collapsed);
+}
+
+int32 SAvaEaseCurvePresetGroup::GetVisiblePresetCount() const
+{
+	return VisiblePresets.Num();
+}
+
+bool SAvaEaseCurvePresetGroup::IsEditMode() const
+{
+	return bIsEditMode.Get(false);
+}
+
+EVisibility SAvaEaseCurvePresetGroup::GetEditModeVisibility() const
+{
+	return IsEditMode() ? EVisibility::Visible : EVisibility::Collapsed;
+}
+
+const FSlateBrush* SAvaEaseCurvePresetGroup::GetBorderImage() const
+{
+	FName BrushName;
+
+	if (bCanBeDroppedOn)
+	{
+		BrushName = bIsOverDifferentCategory ? TEXT("EditMode.Background.Over") : TEXT("EditMode.Background.Highlight");
+	}
+	else
+	{
+		BrushName = TEXT("EditMode.Background");
+	}
+
+	return FAvaEaseCurveStyle::Get().GetBrush(BrushName);
+}
+
+FText SAvaEaseCurvePresetGroup::GetPresetNameTooltipText() const
+{
+	return FText::Format(LOCTEXT("CategoryTooltipText", "Category: {0}"), FText::FromString(CategoryName));
+}
+
+void SAvaEaseCurvePresetGroup::HandleCategoryRenameCommitted(const FText& InNewText, ETextCommit::Type InCommitType)
+{
+	const FString NewCategoryName = InNewText.ToString();
+
+	if (!NewCategoryName.IsEmpty() && !NewCategoryName.Equals(CategoryName)
+		&& OnCategoryRename.IsBound() && OnCategoryRename.Execute(CategoryName, NewCategoryName))
+	{
+		CategoryName = NewCategoryName;
+	}
+	else
+	{
+		RenameCategoryNameTextBox->SetText(FText::FromString(CategoryName));
+	}
+}
+
+FReply SAvaEaseCurvePresetGroup::HandleCategoryDelete() const
+{
+	if (OnCategoryDelete.IsBound())
+	{
+		OnCategoryDelete.Execute(CategoryName);
+	}
+
+	return FReply::Handled();
+}
+
+bool SAvaEaseCurvePresetGroup::HandlePresetDelete(const TSharedPtr<FAvaEaseCurvePreset>& InPreset)
+{
+	if (OnPresetDelete.IsBound() && OnPresetDelete.Execute(InPreset))
+	{
+		Presets.Remove(InPreset);
+		VisiblePresets.Remove(InPreset);
+
+		return true;
+	}
+
+	return false;
+}
+
+bool SAvaEaseCurvePresetGroup::HandlePresetRename(const TSharedPtr<FAvaEaseCurvePreset>& InPreset, const FString& InNewName)
+{
+	if (OnPresetRename.IsBound() && OnPresetRename.Execute(InPreset, InNewName))
+	{
+		InPreset->Name = InNewName;
+
+		return true;
+	}
+	else
+	{
+		PresetWidgetsMap[InPreset]->SetPreset(InPreset);
+
+		return false;
+	}
+}
+
+bool SAvaEaseCurvePresetGroup::HandlePresetBeginMove(const TSharedPtr<FAvaEaseCurvePreset>& InPreset, const FString& InNewCategoryName) const
+{
+	if (OnBeginPresetMove.IsBound())
+	{
+		return OnBeginPresetMove.Execute(InPreset, InNewCategoryName);
+	}
+
+	return false;
+}
+
+bool SAvaEaseCurvePresetGroup::HandlePresetEndMove(const TSharedPtr<FAvaEaseCurvePreset>& InPreset, const FString& InNewCategoryName) const
+{
+	if (OnEndPresetMove.IsBound() && OnEndPresetMove.Execute(InPreset, InNewCategoryName))
+	{
+		InPreset->Category = InNewCategoryName;
+
+		return true;
+	}
+
+	return false;
+}
+
+bool SAvaEaseCurvePresetGroup::HandlePresetClick(const TSharedPtr<FAvaEaseCurvePreset>& InPreset) const
+{
+	if (OnPresetClick.IsBound())
+	{
+		OnPresetClick.Execute(InPreset);
+
+		return true;
+	}
+
+	return false;
+}
+
+void SAvaEaseCurvePresetGroup::OnDragEnter(const FGeometry& InGeometry, const FDragDropEvent& InDragDropEvent)
+{
+	SCompoundWidget::OnDragEnter(InGeometry, InDragDropEvent);
+
+	const TSharedPtr<FAvaEaseCurvePresetDragDropOperation> Operation =
+		InDragDropEvent.GetOperationAs<FAvaEaseCurvePresetDragDropOperation>();
+
+	if (Operation.IsValid())
+	{
+		bIsOverDifferentCategory = !CategoryName.Equals(Operation->GetPreset()->Category);
+		bCanBeDroppedOn = bIsOverDifferentCategory;
+
+		Operation->AddHoveredGroup(SharedThis(this));
+	}
+}
+
+void SAvaEaseCurvePresetGroup::OnDragLeave(const FDragDropEvent& InDragDropEvent)
+{
+	SCompoundWidget::OnDragLeave(InDragDropEvent);
+
+	const TSharedPtr<FAvaEaseCurvePresetDragDropOperation> Operation =
+		InDragDropEvent.GetOperationAs<FAvaEaseCurvePresetDragDropOperation>();
+
+	if (Operation.IsValid())
+	{
+		bCanBeDroppedOn = !CategoryName.Equals(Operation->GetPreset()->Category);
+		bIsOverDifferentCategory = false;
+	}
+}
+
+FReply SAvaEaseCurvePresetGroup::OnDragOver(const FGeometry& InGeometry, const FDragDropEvent& InDragDropEvent)
+{
+	SCompoundWidget::OnDragOver(InGeometry, InDragDropEvent);
+
+	const TSharedPtr<FAvaEaseCurvePresetDragDropOperation> Operation =
+		InDragDropEvent.GetOperationAs<FAvaEaseCurvePresetDragDropOperation>();
+
+	if (Operation.IsValid() && Operation->GetPreset().IsValid())
+	{
+		bIsOverDifferentCategory = !CategoryName.Equals(Operation->GetPreset()->Category);
+
+		return FReply::Handled();
+	}
+
+	return FReply::Unhandled();
+}
+
+FReply SAvaEaseCurvePresetGroup::OnDrop(const FGeometry& InGeometry, const FDragDropEvent& InDragDropEvent)
+{
+	SCompoundWidget::OnDrop(InGeometry, InDragDropEvent);
+
+	ResetDragBorder();
+
+	const TSharedPtr<FAvaEaseCurvePresetDragDropOperation> Operation =
+		InDragDropEvent.GetOperationAs<FAvaEaseCurvePresetDragDropOperation>();
+
+	if (Operation.IsValid())
+	{
+		if (OnEndPresetMove.IsBound() && OnEndPresetMove.Execute(Operation->GetPreset(), CategoryName))
+		{
+			return FReply::Handled();
+		}
+	}
+
+	return FReply::Unhandled();
+}
+
+void SAvaEaseCurvePresetGroup::ResetDragBorder()
+{
+	bCanBeDroppedOn = false;
+	bIsOverDifferentCategory = false;
+}
+
+#undef LOCTEXT_NAMESPACE
diff --git a/Engine/Plugins/Experimental/Avalanche/Source/AvalancheSequencer/Private/EaseCurveTool/Widgets/SAvaEaseCurvePresetGroup.h b/Engine/Plugins/Experimental/Avalanche/Source/AvalancheSequencer/Private/EaseCurveTool/Widgets/SAvaEaseCurvePresetGroup.h
new file mode 100644
--- /dev/null
+++ b/Engine/Plugins/Experimental/Avalanche/Source/AvalancheSequencer/Private/EaseCurveTool/Widgets/SAvaEaseCurvePresetGroup.h
@@ -0,0 +1,97 @@
+// Copyright Epic Games, Inc. All Rights Reserved.
+
+#pragma once
+
+#include "EaseCurveTool/Widgets/SAvaEaseCurvePresetGroupItem.h"
+#include "Internationalization/Text.h"
+#include "Templates/SharedPointer.h"
+#include "Widgets/SCompoundWidget.h"
+
+class STableViewBase;
+struct FAvaEaseCurvePreset;
+
+DECLARE_DELEGATE_RetVal_OneParam(bool, FAvaEaseCurveCategoryDeleteDelegate, const FString& /*InCategoryName*/)
+DECLARE_DELEGATE_RetVal_TwoParams(bool, FAvaEaseCurveCategoryRenameDelegate, const FString& /*InOldName*/, const FString& /*InNewName*/)
+
+class SAvaEaseCurvePresetGroup : public SCompoundWidget
+{
+public:
+	SLATE_BEGIN_ARGS(SAvaEaseCurvePresetGroup)
+		: _DisplayRate(FFrameRate(30, 1))
+	{}
+		SLATE_ARGUMENT(FString, CategoryName)
+		SLATE_ARGUMENT(TArray<TSharedPtr<FAvaEaseCurvePreset>>, Presets)
+		SLATE_ARGUMENT(FText, SearchText)
+		SLATE_ATTRIBUTE(bool, IsEditMode)
+		SLATE_ARGUMENT(FFrameRate, DisplayRate)
+		SLATE_EVENT(FAvaEaseCurveCategoryDeleteDelegate, OnCategoryDelete)
+		SLATE_EVENT(FAvaEaseCurveCategoryRenameDelegate, OnCategoryRename)
+		SLATE_EVENT(FAvaEaseCurvePresetDelegate, OnPresetDelete)
+		SLATE_EVENT(FAvaEaseCurvePresetRenameDelegate, OnPresetRename)
+		SLATE_EVENT(FAvaEaseCurvePresetMoveDelegate, OnBeginPresetMove)
+		SLATE_EVENT(FAvaEaseCurvePresetMoveDelegate, OnEndPresetMove)
+		SLATE_EVENT(FAvaEaseCurvePresetDelegate, OnPresetClick)
+	SLATE_END_ARGS()
+
+	void Construct(const FArguments& InArgs);
+
+	void SetSearchText(const FText& InText);
+
+	int32 GetVisiblePresetCount() const;
+
+	bool IsEditMode() const;
+
+	void ResetDragBorder();
+
+	void NotifyCanDrop(const bool bInCanDrop) { bCanBeDroppedOn = bInCanDrop; }
+
+	FString GetCategoryName() const { return CategoryName; }
+
+protected:
+	TSharedRef<SWidget> ConstructHeader();
+
+	TSharedRef<ITableRow> GeneratePresetWidget(TSharedPtr<FAvaEaseCurvePreset> InPreset, const TSharedRef<STableViewBase>& InOwnerTable);
+	
+	EVisibility GetEditModeVisibility() const;
+
+	const FSlateBrush* GetBorderImage() const;
+
+	FText GetPresetNameTooltipText() const;
+
+	void HandleCategoryRenameCommitted(const FText& InNewText, ETextCommit::Type InCommitType);
+	FReply HandleCategoryDelete() const;
+
+	bool HandlePresetClick(const TSharedPtr<FAvaEaseCurvePreset>& InPreset) const;
+	bool HandlePresetDelete(const TSharedPtr<FAvaEaseCurvePreset>& InPreset);
+	bool HandlePresetRename(const TSharedPtr<FAvaEaseCurvePreset>& InPreset, const FString& InNewName);
+	bool HandlePresetBeginMove(const TSharedPtr<FAvaEaseCurvePreset>& InPreset, const FString& InNewCategoryName) const;
+	bool HandlePresetEndMove(const TSharedPtr<FAvaEaseCurvePreset>& InPreset, const FString& InNewCategoryName) const;
+
+	//~ Begin SWidget
+	virtual void OnDragEnter(const FGeometry& InGeometry, const FDragDropEvent& InDragDropEvent) override;
+	virtual void OnDragLeave(const FDragDropEvent& InDragDropEvent) override;
+	virtual FReply OnDragOver(const FGeometry& InGeometry, const FDragDropEvent& InDragDropEvent) override;
+	virtual FReply OnDrop(const FGeometry& InGeometry, const FDragDropEvent& InDragDropEvent) override;
+	//~ End SWidget
+
+	FString CategoryName;
+	TArray<TSharedPtr<FAvaEaseCurvePreset>> Presets;
+	FText SearchText;
+	TAttribute<bool> bIsEditMode;
+	FFrameRate DisplayRate;
+	FAvaEaseCurveCategoryDeleteDelegate OnCategoryDelete;
+	FAvaEaseCurveCategoryRenameDelegate OnCategoryRename;
+	FAvaEaseCurvePresetDelegate OnPresetDelete;
+	FAvaEaseCurvePresetRenameDelegate OnPresetRename;
+	FAvaEaseCurvePresetMoveDelegate OnBeginPresetMove;
+	FAvaEaseCurvePresetMoveDelegate OnEndPresetMove;
+	FAvaEaseCurvePresetDelegate OnPresetClick;
+
+	TArray<TSharedPtr<FAvaEaseCurvePreset>> VisiblePresets;
+
+	TSharedPtr<SEditableTextBox> RenameCategoryNameTextBox;
+	TMap<TSharedPtr<FAvaEaseCurvePreset>, TSharedPtr<SAvaEaseCurvePresetGroupItem>> PresetWidgetsMap;
+
+	bool bCanBeDroppedOn = false;
+	bool bIsOverDifferentCategory = false;
+};
