[MotionDesign]

diff --git a/Engine/Plugins/Experimental/Avalanche/Source/AvalancheSequencer/Private/EaseCurveTool/Widgets/SAvaEaseCurvePresetGroupItem.cpp b/Engine/Plugins/Experimental/Avalanche/Source/AvalancheSequencer/Private/EaseCurveTool/Widgets/SAvaEaseCurvePresetGroupItem.cpp
new file mode 100644
--- /dev/null
+++ b/Engine/Plugins/Experimental/Avalanche/Source/AvalancheSequencer/Private/EaseCurveTool/Widgets/SAvaEaseCurvePresetGroupItem.cpp
@@ -0,0 +1,250 @@
+// Copyright Epic Games, Inc. All Rights Reserved.
+
+#include "SAvaEaseCurvePresetGroupItem.h"
+#include "EaseCurveTool/AvaEaseCurveStyle.h"
+#include "EaseCurveTool/AvaEaseCurvePreset.h"
+#include "EaseCurveTool/AvaEaseCurvePresetDragDropOp.h"
+#include "EaseCurveTool/Widgets/SAvaEaseCurvePreview.h"
+#include "Internationalization/Text.h"
+#include "Styling/AppStyle.h"
+#include "Styling/StyleColors.h"
+#include "Widgets/Images/SImage.h"
+#include "Widgets/Input/SButton.h"
+#include "Widgets/Input/SEditableTextBox.h"
+#include "Widgets/Layout/SBorder.h"
+#include "Widgets/Layout/SBox.h"
+#include "Widgets/Layout/SWidgetSwitcher.h"
+#include "Widgets/SBoxPanel.h"
+#include "Widgets/SOverlay.h"
+#include "Widgets/Text/STextBlock.h"
+#include "Widgets/Views/STableViewBase.h"
+
+#define LOCTEXT_NAMESPACE "SAvaEaseCurvePresetGroupItem"
+
+void SAvaEaseCurvePresetGroupItem::Construct(const FArguments& InArgs, const TSharedPtr<STableViewBase>& InOwnerTableView)
+{
+	Preset = InArgs._Preset;
+	bIsEditMode = InArgs._IsEditMode;
+	OnClick = InArgs._OnClick;
+	OnDelete = InArgs._OnDelete;
+	OnRename = InArgs._OnRename;
+	OnBeginMove = InArgs._OnBeginMove;
+	OnEndMove = InArgs._OnEndMove;
+
+	ChildSlot
+		[
+			SNew(SOverlay)
+			+ SOverlay::Slot()
+			[
+				SNew(SBorder)
+				.Visibility(this, &SAvaEaseCurvePresetGroupItem::GetBorderVisibility)
+				.BorderImage(this, &SAvaEaseCurvePresetGroupItem::GetBackgroundImage)
+			]
+			+ SOverlay::Slot()
+			.Padding(1.f)
+			[
+				SNew(SHorizontalBox)
+				+ SHorizontalBox::Slot()
+				.HAlign(HAlign_Left)
+				.VAlign(VAlign_Center)
+				.Padding(3.f, 0.f)
+				[
+					SNew(SBox)
+					.WidthOverride(160)
+					.Clipping(EWidgetClipping::ClipToBoundsAlways)
+					[
+						SNew(SWidgetSwitcher)
+						.WidgetIndex_Lambda([this]()
+							{
+								return IsEditMode() ? 1 : 0;
+							})
+						+ SWidgetSwitcher::Slot()
+						[
+							SNew(STextBlock)
+							.TextStyle(FAppStyle::Get(), TEXT("Menu.Label"))
+							.OverflowPolicy(ETextOverflowPolicy::Ellipsis)
+							.Text(FText::FromString(Preset->Name))
+							.ToolTipText(FText::FromString(Preset->Name))
+						]
+						+ SWidgetSwitcher::Slot()
+						[
+							SNew(SHorizontalBox)
+							+ SHorizontalBox::Slot()
+							.AutoWidth()
+							.VAlign(VAlign_Center)
+							[
+								SNew(SButton)
+								.ButtonStyle(FAvaEaseCurveStyle::Get(), TEXT("ToolButton.NoPad"))
+								.VAlign(VAlign_Center)
+								.ToolTipText(LOCTEXT("EditModeDeleteTooltip", "Delete this category and the json file associated with it on disk"))
+								.Visibility(this, &SAvaEaseCurvePresetGroupItem::GetEditModeVisibility)
+								.OnClicked(this, &SAvaEaseCurvePresetGroupItem::HandleDeleteClick)
+								[
+									SNew(SImage)
+									.DesiredSizeOverride(FVector2D(10.f))
+									.Image(FAppStyle::GetBrush(TEXT("Icons.Delete")))
+								]
+							]
+							+ SHorizontalBox::Slot()
+							.Padding(2.f, 0.f, 0.f, 0.f)
+							[
+								SAssignNew(RenameTextBox, SEditableTextBox)
+								.OverflowPolicy(ETextOverflowPolicy::Ellipsis)
+								.Text(FText::FromString(Preset->Name))
+								.ToolTipText(FText::FromString(Preset->Name))
+								.OnTextCommitted(this, &SAvaEaseCurvePresetGroupItem::HandleRenameTextCommitted)
+							]
+						]
+					]
+				]
+				+ SHorizontalBox::Slot()
+				.AutoWidth()
+				.HAlign(HAlign_Right)
+				.VAlign(VAlign_Center)
+				.Padding(8.f, 0.f, 0.f, 0.f)
+				[
+					SNew(SBorder)
+					.BorderBackgroundColor(FStyleColors::White25)
+					.Padding(2.f)
+					[
+						SNew(SAvaEaseCurvePreview)
+						.PreviewSize(20.f)
+						.CurveThickness(1.5f)
+						.Tangents(Preset->Tangents)
+						.CustomToolTip(true)
+						.UnderCurveColor(FStyleColors::SelectInactive.GetSpecifiedColor())
+						.DisplayRate(InArgs._DisplayRate)
+					]
+				]
+			]
+		];
+
+	STableRow<TSharedPtr<FAvaEaseCurvePreset>>::ConstructInternal(
+		STableRow::FArguments()
+		.Style(&FAppStyle::GetWidgetStyle<FTableRowStyle>(TEXT("ComboBox.Row")))
+		.Padding(5.f)
+		.ShowSelection(true)
+		, InOwnerTableView.ToSharedRef());
+}
+
+void SAvaEaseCurvePresetGroupItem::SetPreset(const TSharedPtr<FAvaEaseCurvePreset>& InPreset)
+{
+	Preset = InPreset;
+}
+
+void SAvaEaseCurvePresetGroupItem::HandlePresetClick() const
+{
+	if (OnClick.IsBound())
+	{
+		OnClick.Execute(Preset);
+	}
+}
+
+bool SAvaEaseCurvePresetGroupItem::IsEditMode() const
+{
+	return bIsEditMode.Get(false);
+}
+
+EVisibility SAvaEaseCurvePresetGroupItem::GetEditModeVisibility() const
+{
+	return IsEditMode() ? EVisibility::Visible : EVisibility::Collapsed;
+}
+
+EVisibility SAvaEaseCurvePresetGroupItem::GetBorderVisibility() const
+{
+	return bIsDragging && IsEditMode() ? EVisibility::Visible : EVisibility::Collapsed;
+}
+
+const FSlateBrush* SAvaEaseCurvePresetGroupItem::GetBackgroundImage() const
+{
+	return bIsDragging ? FAvaEaseCurveStyle::Get().GetBrush(TEXT("EditMode.Background.Over")) : nullptr;
+}
+
+void SAvaEaseCurvePresetGroupItem::HandleRenameTextCommitted(const FText& InNewText, ETextCommit::Type InCommitType) const
+{
+	const FString NewPresetName = InNewText.ToString();
+	
+	if (!NewPresetName.IsEmpty() && !NewPresetName.Equals(Preset->Name)
+		&& OnRename.IsBound() && OnRename.Execute(Preset, NewPresetName))
+	{
+		Preset->Name = NewPresetName;
+	}
+	else
+	{
+		RenameTextBox->SetText(FText::FromString(Preset->Name));
+	}
+}
+
+FReply SAvaEaseCurvePresetGroupItem::HandleDeleteClick() const
+{
+	if (OnDelete.IsBound())
+	{
+		OnDelete.Execute(Preset);
+	}
+
+	return FReply::Handled();
+}
+
+FReply SAvaEaseCurvePresetGroupItem::OnMouseButtonDown(const FGeometry& InGeometry, const FPointerEvent& InMouseEvent)
+{
+	if (IsEditMode())
+	{
+		return FReply::Handled().DetectDrag(SharedThis(this), EKeys::LeftMouseButton);
+	}
+
+	if (OnClick.IsBound())
+	{
+		OnClick.Execute(Preset);
+	}
+
+	return FReply::Handled();
+}
+
+FReply SAvaEaseCurvePresetGroupItem::OnMouseButtonUp(const FGeometry& InGeometry, const FPointerEvent& InMouseEvent)
+{
+	TriggerEndMove();
+
+	return STableRow<TSharedPtr<FAvaEaseCurvePreset>>::OnMouseButtonUp(InGeometry, InMouseEvent);
+}
+
+FReply SAvaEaseCurvePresetGroupItem::OnDragDetected(const FGeometry& MyGeometry, const FPointerEvent& MouseEvent)
+{
+	if (!IsEditMode())
+	{
+		return FReply::Unhandled();
+	}
+
+	TriggerBeginMove();
+
+	const TSharedRef<FAvaEaseCurvePresetDragDropOperation> Operation = MakeShared<FAvaEaseCurvePresetDragDropOperation>(SharedThis(this), Preset);
+	return FReply::Handled().BeginDragDrop(Operation);
+}
+
+FReply SAvaEaseCurvePresetGroupItem::OnDrop(const FGeometry& InGeometry, const FDragDropEvent& InDragDropEvent)
+{
+	bIsDragging = false;
+
+	return STableRow<TSharedPtr<FAvaEaseCurvePreset>>::OnDrop(InGeometry, InDragDropEvent);
+}
+
+void SAvaEaseCurvePresetGroupItem::TriggerBeginMove()
+{
+	if (OnBeginMove.IsBound())
+	{
+		OnBeginMove.Execute(Preset, Preset->Category);
+	}
+
+	bIsDragging = true;
+}
+
+void SAvaEaseCurvePresetGroupItem::TriggerEndMove()
+{
+	if (OnEndMove.IsBound())
+	{
+		OnEndMove.Execute(Preset, Preset->Category);
+	}
+
+	bIsDragging = false;
+}
+
+#undef LOCTEXT_NAMESPACE
diff --git a/Engine/Plugins/Experimental/Avalanche/Source/AvalancheSequencer/Private/EaseCurveTool/Widgets/SAvaEaseCurvePresetGroupItem.h b/Engine/Plugins/Experimental/Avalanche/Source/AvalancheSequencer/Private/EaseCurveTool/Widgets/SAvaEaseCurvePresetGroupItem.h
new file mode 100644
--- /dev/null
+++ b/Engine/Plugins/Experimental/Avalanche/Source/AvalancheSequencer/Private/EaseCurveTool/Widgets/SAvaEaseCurvePresetGroupItem.h
@@ -0,0 +1,74 @@
+// Copyright Epic Games, Inc. All Rights Reserved.
+
+#pragma once
+
+#include "Misc/FrameRate.h"
+#include "Templates/SharedPointer.h"
+#include "Widgets/Views/SListView.h"
+#include "Widgets/Views/STableRow.h"
+
+class FText;
+class SEditableTextBox;
+class STableViewBase;
+struct FAvaEaseCurvePreset;
+
+DECLARE_DELEGATE_RetVal_OneParam(bool, FAvaEaseCurvePresetDelegate, const TSharedPtr<FAvaEaseCurvePreset>&)
+DECLARE_DELEGATE_RetVal_TwoParams(bool, FAvaEaseCurvePresetRenameDelegate, const TSharedPtr<FAvaEaseCurvePreset>& /*InPreset*/, const FString& /*InNewName*/)
+DECLARE_DELEGATE_RetVal_TwoParams(bool, FAvaEaseCurvePresetMoveDelegate, const TSharedPtr<FAvaEaseCurvePreset>& /*InPreset*/, const FString& /*InNewCategoryName*/)
+
+class SAvaEaseCurvePresetGroupItem : public STableRow<TSharedPtr<FAvaEaseCurvePreset>>
+{
+public:
+	SLATE_BEGIN_ARGS(SAvaEaseCurvePresetGroupItem)
+		: _DisplayRate(FFrameRate(30, 1))
+	{}
+		SLATE_ARGUMENT(TSharedPtr<FAvaEaseCurvePreset>, Preset)
+		SLATE_ATTRIBUTE(bool, IsEditMode)
+		SLATE_ARGUMENT(FFrameRate, DisplayRate)
+		SLATE_EVENT(FAvaEaseCurvePresetDelegate, OnClick)
+		SLATE_EVENT(FAvaEaseCurvePresetDelegate, OnDelete)
+		SLATE_EVENT(FAvaEaseCurvePresetRenameDelegate, OnRename)
+		SLATE_EVENT(FAvaEaseCurvePresetMoveDelegate, OnBeginMove)
+		SLATE_EVENT(FAvaEaseCurvePresetMoveDelegate, OnEndMove)
+	SLATE_END_ARGS()
+
+	void Construct(const FArguments& InArgs, const TSharedPtr<STableViewBase>& InOwnerTableView);
+
+	void SetPreset(const TSharedPtr<FAvaEaseCurvePreset>& InPreset);
+
+	bool IsEditMode() const;
+
+	void TriggerBeginMove();
+	void TriggerEndMove();
+
+protected:
+	void HandlePresetClick() const;
+
+	void HandleRenameTextCommitted(const FText& InNewText, ETextCommit::Type InCommitType) const;
+
+	FReply HandleDeleteClick() const;
+
+	EVisibility GetEditModeVisibility() const;
+
+	EVisibility GetBorderVisibility() const;
+	const FSlateBrush* GetBackgroundImage() const;
+
+	//~ Begin SWidget
+	virtual FReply OnMouseButtonDown(const FGeometry& InGeometry, const FPointerEvent& InMouseEvent) override;
+	virtual FReply OnMouseButtonUp(const FGeometry& InGeometry, const FPointerEvent& InMouseEvent) override;
+	virtual FReply OnDragDetected(const FGeometry& InGeometry, const FPointerEvent& InMouseEvent) override;
+	virtual FReply OnDrop(const FGeometry& InGeometry, const FDragDropEvent& InDragDropEvent) override;
+	//~ End SWidget
+
+	TSharedPtr<FAvaEaseCurvePreset> Preset;
+	TAttribute<bool> bIsEditMode;
+	FAvaEaseCurvePresetDelegate OnClick;
+	FAvaEaseCurvePresetDelegate OnDelete;
+	FAvaEaseCurvePresetRenameDelegate OnRename;
+	FAvaEaseCurvePresetMoveDelegate OnBeginMove;
+	FAvaEaseCurvePresetMoveDelegate OnEndMove;
+
+	TSharedPtr<SEditableTextBox> RenameTextBox;
+
+	bool bIsDragging = false;
+};
diff --git a/Engine/Plugins/Experimental/Avalanche/Source/AvalancheSequencer/Private/EaseCurveTool/Widgets/SAvaEaseCurvePreview.cpp b/Engine/Plugins/Experimental/Avalanche/Source/AvalancheSequencer/Private/EaseCurveTool/Widgets/SAvaEaseCurvePreview.cpp
new file mode 100644
--- /dev/null
+++ b/Engine/Plugins/Experimental/Avalanche/Source/AvalancheSequencer/Private/EaseCurveTool/Widgets/SAvaEaseCurvePreview.cpp
@@ -0,0 +1,361 @@
+// Copyright Epic Games, Inc. All Rights Reserved.
+
+#include "SAvaEaseCurvePreview.h"
+#include "EaseCurveTool/AvaEaseCurveTool.h"
+#include "EaseCurveTool/Widgets/SAvaEaseCurvePreviewToolTip.h"
+#include "Curves/CurveEvaluation.h"
+#include "HAL/PlatformApplicationMisc.h"
+#include "Kismet/KismetMathLibrary.h"
+#include "SCurveEditor.h"
+#include "Styling/AppStyle.h"
+#include "Widgets/Accessibility/SlateCoreAccessibleWidgets.h"
+
+#define LOCTEXT_NAMESPACE "SAvaEaseCurvePreview"
+
+void SAvaEaseCurvePreview::Construct(const FArguments& InArgs)
+{ 
+	SetCanTick(InArgs._Animate);
+
+	Tangents = InArgs._Tangents;
+	PreviewSize = InArgs._PreviewSize;
+	CanExpandPreview = InArgs._CanExpandPreview;
+	CurveThickness = InArgs._CurveThickness;
+	CurveColor = InArgs._CurveColor;
+	StraightColor = InArgs._StraightColor;
+	UnderCurveColor = InArgs._UnderCurveColor;
+	bAnimate = InArgs._Animate;
+	IdleAnimationLength = InArgs._IdleAnimationLength;
+	DisplayRate = InArgs._DisplayRate;
+	AnimationColor = InArgs._AnimationColor;
+	AnimationSize = InArgs._AnimationSize;
+	bDrawMotionTrails = InArgs._DrawMotionTrails;
+	MotionTrailColor = InArgs._MotionTrailColor;
+	MotionTrailColorFade = InArgs._MotionTrailColorFade;
+	MotionTrailSize = InArgs._MotionTrailSize;
+	MotionTrailFadeLength = InArgs._MotionTrailFadeLength;
+	MotionTrailOffset = InArgs._MotionTrailOffset;
+
+	if (InArgs._CustomToolTip)
+	{
+		ToolTipWidget = MakeShared<SAvaEaseCurvePreviewToolTip>(InArgs);
+		SetToolTip(ToolTipWidget);
+	}
+	else
+	{
+		SetToolTipText(SAvaEaseCurvePreviewToolTip::GetToolTipText(Tangents.Get(FAvaEaseCurveTangents())));
+	}
+
+	SetTangents(InArgs._Tangents);
+}
+
+FAvaEaseCurveTangents SAvaEaseCurvePreview::GetTangents() const
+{
+	return Tangents.Get(FAvaEaseCurveTangents());
+}
+
+void SAvaEaseCurvePreview::SetTangents(const TAttribute<FAvaEaseCurveTangents>& InTangents)
+{
+	Tangents = InTangents;
+
+	// Save min and max to compute extra screen space if the curve goes beyond normal limits (0 - 1).
+	float MinInputValue = 0.f;
+	float MaxInputValue = 1.f;
+
+	// Cache motion trails
+	const FAvaEaseCurveTangents DrawTangents = Tangents.Get(FAvaEaseCurveTangents());
+	const float SecondsPerFrame = DisplayRate.AsInterval();
+	const int32 FrameCount = DisplayRate.AsFrameNumber(1.f).Value;
+
+	MotionTrails.Empty(FrameCount + 1);
+
+	for (float Time = 0; Time < 1.f; Time += SecondsPerFrame)
+	{
+		const float CurrentValue = WeightedEvalForTwoKeys(DrawTangents, Time);
+
+		FMotionTrail NewMotionTrail;
+		NewMotionTrail.Time = Time;
+		NewMotionTrail.Color = FLinearColor(1.f, 1.f, 1.f, 1.f);
+		NewMotionTrail.Location = FVector2D(CurrentValue, 0.f);
+		MotionTrails.Add(MoveTemp(NewMotionTrail));
+
+		MinInputValue = FMath::Min(CurrentValue, MinInputValue);
+		MaxInputValue = FMath::Max(CurrentValue, MaxInputValue);
+	}
+
+	FMotionTrail LastMotionTrail;
+	LastMotionTrail.Time = 1.f;
+	LastMotionTrail.Color = FLinearColor(1.f, 1.f, 1.f, 1.f);
+	LastMotionTrail.Location = FVector2D(1.f, 0.f);
+	MotionTrails.Add(MoveTemp(LastMotionTrail));
+
+	BelowZeroValue = FMath::Abs(MinInputValue);
+	AboveOneValue = MaxInputValue - 1.f;
+}
+
+FVector2D SAvaEaseCurvePreview::ComputeDesiredSize(const float InLayoutScaleMultiplier) const
+{
+	float AdditionalHeight = CanExpandPreview ? (BelowZeroValue + AboveOneValue) : 0.f;
+	AdditionalHeight = FMath::Min(PreviewSize, AdditionalHeight);
+	return FVector2D(PreviewSize) + FVector2D(AdditionalHeight * PreviewSize);
+}
+
+FReply SAvaEaseCurvePreview::OnMouseButtonDown(const FGeometry& InMyGeometry, const FPointerEvent& InMouseEvent)
+{
+	if (InMouseEvent.GetEffectingButton() == EKeys::RightMouseButton)
+	{
+		FPlatformApplicationMisc::ClipboardCopy(*Tangents.Get(FAvaEaseCurveTangents()).ToJson());
+
+		FAvaEaseCurveTool::ShowNotificationMessage(LOCTEXT("EaseCurveToolTangentsCopied", "Ease Curve Tool Tangents Copied!"));
+
+		return FReply::Handled();
+	}
+
+	return SLeafWidget::OnMouseButtonDown(InMyGeometry, InMouseEvent);
+}
+
+void SAvaEaseCurvePreview::Tick(const FGeometry& InAllottedGeometry, const double InCurrentTime, const float InDeltaTime)
+{
+	SLeafWidget::Tick(InAllottedGeometry, InCurrentTime, InDeltaTime);
+
+	if (IdleTimerHandle.IsValid()) // Fade out all animations
+	{
+		CurrentAnimateTime = 1.f;
+
+		// Fade out animation and motion trails
+		const double IdlePlayTime = FMath::Min(IdleAnimationLength, InCurrentTime - MotionTrailFadeStartTime);
+		const double NewOpacity = UKismetMathLibrary::MapRangeClamped(IdlePlayTime, 0.f, IdleAnimationLength, 1.f, 0.f);
+
+		AnimationColor.A = NewOpacity;
+
+		if (bDrawMotionTrails)
+		{
+			const FLinearColor LerpColor = FLinearColor::LerpUsingHSV(FLinearColor(0.f, 0.f, 0.f, 0.f), AnimationColor, NewOpacity);
+
+			for (FMotionTrail& MotionTrail : MotionTrails)
+			{
+				MotionTrail.Color = LerpColor;
+			}
+		}
+	}
+	else // Animate curve along path
+	{
+		if (CurrentAnimateTime >= 1.f)
+		{
+			CurrentAnimateTime = 1.f;
+			MotionTrailFadeStartTime = InCurrentTime;
+
+			IdleTimerHandle = RegisterActiveTimer(IdleAnimationLength, FWidgetActiveTimerDelegate::CreateLambda([this](double InCurrentTime, float InDeltaTime) -> EActiveTimerReturnType
+				{
+					IdleTimerHandle.Reset();
+
+					CurrentAnimateTime = 0.f;
+
+					if (bDrawMotionTrails)
+					{
+						for (FMotionTrail& MotionTrail : MotionTrails)
+						{
+							MotionTrail.Color = FLinearColor::White;
+						}
+					}
+
+					return EActiveTimerReturnType::Stop;
+				}));
+		}
+		else
+		{
+			CurrentAnimateTime += InDeltaTime;
+			CurrentAnimateTime = FMath::Min(1.f, CurrentAnimateTime);
+
+			const FTrackScaleInfo ScaleInfo(0.f, 1.f, 0.f, 1.f, FVector2D(PreviewSize)/*InAllottedGeometry.GetLocalSize()*/);
+			const FAvaEaseCurveTangents DrawTangents = Tangents.Get(FAvaEaseCurveTangents());
+			const float CurrentValue = WeightedEvalForTwoKeys(DrawTangents, CurrentAnimateTime);
+
+			CurrentAnimateLocation = FVector2D(ScaleInfo.InputToLocalX(CurrentAnimateTime), ScaleInfo.OutputToLocalY(CurrentValue));
+
+			AnimationColor.A = 1.f;
+
+			if (bDrawMotionTrails)
+			{
+				for (FMotionTrail& MotionTrail : MotionTrails)
+				{
+					const float TimeSinceTrail = CurrentAnimateTime - MotionTrail.Time;
+					const float ClampedTimeSinceTrail = FMath::Min(MotionTrailFadeLength, TimeSinceTrail);
+					const float FadeOutAlpha = TimeSinceTrail <= MotionTrailFadeLength
+						? UKismetMathLibrary::MapRangeClamped(ClampedTimeSinceTrail, 0.f, MotionTrailFadeLength, 0.f, 1.f)
+						: 1.f;
+
+					FLinearColor NewColor = FLinearColor::LerpUsingHSV(MotionTrailColor, MotionTrailColorFade, FadeOutAlpha);
+					NewColor.A = MotionTrail.Color.A;
+
+					MotionTrail.Color = NewColor;
+				}
+			}
+		}
+	}
+}
+
+#if WITH_ACCESSIBILITY
+TSharedRef<FSlateAccessibleWidget> SAvaEaseCurvePreview::CreateAccessibleWidget()
+{
+	return MakeShareable<FSlateAccessibleWidget>(new FSlateAccessibleImage(SharedThis(this)));
+}
+#endif
+
+int32 SAvaEaseCurvePreview::OnPaint(const FPaintArgs& InArgs, const FGeometry& InAllottedGeometry, const FSlateRect& InMyCullingRect, FSlateWindowElementList& OutDrawElements, int32 InLayerId, const FWidgetStyle& InWidgetStyle, bool bInParentEnabled) const
+{
+	const bool bEnabled = ShouldBeEnabled(bInParentEnabled);
+	const ESlateDrawEffect DrawEffects = bEnabled ? ESlateDrawEffect::None : ESlateDrawEffect::DisabledEffect;
+	const FVector2D LocalSize = InAllottedGeometry.GetLocalSize();
+	const FTrackScaleInfo ScaleInfo(0.f, 1.f, 0.f, 1.f, FVector2D(PreviewSize));
+
+	const float ZeroOutputY = ScaleInfo.OutputToLocalY(0.f);
+	const float OneOutputY = ScaleInfo.OutputToLocalY(1.f);
+
+	int32 LayerId = InLayerId;
+
+	const FVector2D AdditionalAboveHeight(BelowZeroValue * PreviewSize, AboveOneValue * PreviewSize);
+	const FPaintGeometry AbovePaintGeometry = InAllottedGeometry.ToPaintGeometry(InAllottedGeometry.GetLocalSize(), FSlateLayoutTransform(AdditionalAboveHeight));
+
+	// Draw 0 and 1 line
+	{
+		constexpr FLinearColor NormalLimitLineColor = FLinearColor(0.1f, 0.1f, 0.1f, 1.f);
+		const TArray<FLinearColor> ValueLineColors = { FLinearColor::Transparent, NormalLimitLineColor, NormalLimitLineColor, FLinearColor::Transparent};
+		const float LocalWidth = InAllottedGeometry.GetLocalSize().X;
+		const float OneFifthWidth = LocalWidth * 0.2f;
+
+		// Value = 0 line
+		if (BelowZeroValue > 0.f)
+		{
+			TArray<FVector2D> ZeroValueLinePoints;
+			ZeroValueLinePoints.Add(FVector2D(0.f, ZeroOutputY));
+			ZeroValueLinePoints.Add(FVector2D(OneFifthWidth, ZeroOutputY));
+			ZeroValueLinePoints.Add(FVector2D(LocalWidth - OneFifthWidth, ZeroOutputY));
+			ZeroValueLinePoints.Add(FVector2D(LocalWidth, ZeroOutputY));
+			TArray<FLinearColor> ZeroValueLineColors;
+			FSlateDrawElement::MakeLines(OutDrawElements, InLayerId, AbovePaintGeometry
+				, ZeroValueLinePoints, ValueLineColors, DrawEffects, FLinearColor::White, false, InAllottedGeometry.Scale);
+		}
+
+		// Value = 1 line
+		if (AboveOneValue > 0.f)
+		{
+			TArray<FVector2D> OneValueLinePoints;
+			OneValueLinePoints.Add(FVector2D(0.f, OneOutputY));
+			OneValueLinePoints.Add(FVector2D(OneFifthWidth, OneOutputY));
+			OneValueLinePoints.Add(FVector2D(LocalWidth - OneFifthWidth, OneOutputY));
+			OneValueLinePoints.Add(FVector2D(LocalWidth, OneOutputY));
+			FSlateDrawElement::MakeLines(OutDrawElements, InLayerId, AbovePaintGeometry
+				, OneValueLinePoints, ValueLineColors, DrawEffects, FLinearColor::White, false, InAllottedGeometry.Scale);
+		}
+	}
+
+	TArray<FVector2D> CurvePoints;
+	TArray<FLinearColor> CurveColors;
+	MakeCurveData(InAllottedGeometry, CurvePoints, CurveColors);
+
+	// Draw under curve fill
+	if (UnderCurveColor.IsSet())
+	{
+		for (const FVector2D& Point : CurvePoints)
+		{
+			const TArray<FVector2D> LinePoints = { FVector2D(Point.X, LocalSize.Y), Point };
+			const TArray<FLinearColor> LineColors = { UnderCurveColor.GetValue(), UnderCurveColor.GetValue() };
+
+			FSlateDrawElement::MakeLines(OutDrawElements, LayerId++, AbovePaintGeometry,
+				LinePoints, LineColors, DrawEffects, FLinearColor::White, true, 2.0f * InAllottedGeometry.Scale);
+		}
+	}
+
+	// Draw curve
+	FSlateDrawElement::MakeLines(OutDrawElements, LayerId++, AbovePaintGeometry,
+		CurvePoints, CurveColors, DrawEffects, FLinearColor::White, true, CurveThickness * InAllottedGeometry.Scale);
+
+	// Draw animation
+	if (bAnimate)
+	{
+		const FSlateBrush* ImageBrush = FAppStyle::GetBrush(TEXT("Icons.BulletPoint"));
+
+		// Draw past motion trails
+		if (bDrawMotionTrails)
+		{
+			for (const FMotionTrail& MotionTrail : MotionTrails)
+			{
+				if (MotionTrail.Time >= CurrentAnimateTime)
+				{
+					continue;
+				}
+
+				FVector2D MotionTrailLocation(ScaleInfo.InputToLocalX(MotionTrail.Location.X), ScaleInfo.OutputToLocalY(MotionTrail.Location.Y));
+				MotionTrailLocation -= FVector2D(MotionTrailSize * 0.5f);
+				MotionTrailLocation += MotionTrailOffset;
+				MotionTrailLocation += AdditionalAboveHeight;
+
+				FSlateDrawElement::MakeBox(OutDrawElements, LayerId++
+					, InAllottedGeometry.ToPaintGeometry(FVector2D(MotionTrailSize), FSlateLayoutTransform(MotionTrailLocation))
+					, ImageBrush, DrawEffects, MotionTrail.Color * InWidgetStyle.GetColorAndOpacityTint());
+			}
+		}
+
+		// Draw current animation point
+		const FVector2D TangentIconLocation = CurrentAnimateLocation - (AnimationSize * 0.5f);
+
+		FSlateDrawElement::MakeBox(OutDrawElements, LayerId++
+			, InAllottedGeometry.ToPaintGeometry(FVector2D(AnimationSize), FSlateLayoutTransform(TangentIconLocation + AdditionalAboveHeight))
+			, ImageBrush, DrawEffects, AnimationColor * InWidgetStyle.GetColorAndOpacityTint());
+	}
+
+	return LayerId;
+}
+
+void SAvaEaseCurvePreview::MakeCurveData(const FGeometry& InGeometry, TArray<FVector2D>& OutCurvePoints, TArray<FLinearColor>& OutCurveColors) const
+{
+	const FAvaEaseCurveTangents DrawTangents = Tangents.Get(FAvaEaseCurveTangents());
+	const FTrackScaleInfo ScaleInfo(0.f, 1.f, 0.f, 1.f, FVector2D(PreviewSize)/*InAllottedGeometry.GetLocalSize()InGeometry.GetLocalSize()*/);
+
+	// Clamp to screen to avoid massive slowdown when zoomed in
+	float StartX = FMath::Max(ScaleInfo.InputToLocalX(0.f), 0.0f);
+	const float EndX = FMath::Min(ScaleInfo.InputToLocalX(1.f), ScaleInfo.WidgetSize.X);
+	
+	constexpr float StepSize = 1.0f;
+	float LastOutputY = 0.0f;
+	
+	for (; StartX < EndX; StartX += StepSize)
+	{
+		const float CurveIn = ScaleInfo.LocalXToInput(FMath::Min(StartX, EndX));
+		const float CurveOut = WeightedEvalForTwoKeys(DrawTangents, CurveIn);
+
+		const FVector2D ScreenPoint(ScaleInfo.InputToLocalX(CurveIn), ScaleInfo.OutputToLocalY(CurveOut));
+		OutCurvePoints.Add(ScreenPoint);
+
+		const float NewCurveStrength = CurveOut - LastOutputY;
+		LastOutputY = CurveOut;
+
+		const float Alpha = UKismetMathLibrary::MapRangeClamped(FMath::Abs(NewCurveStrength), 0.f, 0.03f, 0.f, 1.f);
+		OutCurveColors.Add(FLinearColor::LerpUsingHSV(StraightColor, CurveColor, Alpha));
+	}
+}
+
+float SAvaEaseCurvePreview::WeightedEvalForTwoKeys(const FAvaEaseCurveTangents& InTangents, const float InTime) const
+{
+	FRichCurveKey StartKey(0.f, 0.f);
+	StartKey.InterpMode = ERichCurveInterpMode::RCIM_Cubic;
+	StartKey.TangentMode = ERichCurveTangentMode::RCTM_Break;
+	StartKey.TangentWeightMode = ERichCurveTangentWeightMode::RCTWM_WeightedLeave;
+	StartKey.ArriveTangent = 0.f;
+	StartKey.ArriveTangentWeight = 0.f;
+	StartKey.LeaveTangent = InTangents.Start;
+	StartKey.LeaveTangentWeight = InTangents.StartWeight;
+
+	FRichCurveKey EndKey(1.f, 1.f);
+	EndKey.InterpMode = ERichCurveInterpMode::RCIM_Cubic;
+	EndKey.TangentMode = ERichCurveTangentMode::RCTM_Break;
+	EndKey.TangentWeightMode = ERichCurveTangentWeightMode::RCTWM_WeightedArrive;
+	EndKey.ArriveTangent = InTangents.End;
+	EndKey.ArriveTangentWeight = InTangents.EndWeight;
+	EndKey.LeaveTangent = 0.f;
+	EndKey.LeaveTangentWeight = 0.f;
+	
+	return UE::Curves::EvalForTwoKeys(StartKey, EndKey, InTime);
+}
+
+#undef LOCTEXT_NAMESPACE
diff --git a/Engine/Plugins/Experimental/Avalanche/Source/AvalancheSequencer/Private/EaseCurveTool/Widgets/SAvaEaseCurvePreview.h b/Engine/Plugins/Experimental/Avalanche/Source/AvalancheSequencer/Private/EaseCurveTool/Widgets/SAvaEaseCurvePreview.h
new file mode 100644
--- /dev/null
+++ b/Engine/Plugins/Experimental/Avalanche/Source/AvalancheSequencer/Private/EaseCurveTool/Widgets/SAvaEaseCurvePreview.h
@@ -0,0 +1,133 @@
+// Copyright Epic Games, Inc. All Rights Reserved.
+
+#pragma once
+
+#include "EaseCurveTool/Widgets/SAvaEaseCurvePreset.h"
+#include "Misc/Optional.h"
+#include "Styling/SlateTypes.h"
+#include "Styling/StyleColors.h"
+#include "Templates/SharedPointer.h"
+#include "Widgets/DeclarativeSyntaxSupport.h"
+#include "Widgets/SLeafWidget.h"
+
+class FPaintArgs;
+class FSlateWindowElementList;
+class SAvaEaseCurvePreviewToolTip;
+struct FGeometry;
+
+class SAvaEaseCurvePreview : public SLeafWidget
+{
+public:
+	SLATE_BEGIN_ARGS(SAvaEaseCurvePreview)
+		: _Tangents()
+		, _PreviewSize(32.f)
+		, _CanExpandPreview(false)
+		, _CurveThickness(2.f)
+		, _CurveColor(FLinearColor::White)
+		, _StraightColor(FStyleColors::AccentBlue.GetSpecifiedColor())
+		, _CustomToolTip(false)
+		, _Animate(false)
+		, _IdleAnimationLength(1.f)
+		, _DisplayRate(FFrameRate(30, 1))
+		, _AnimationColor(FLinearColor::White)
+		, _AnimationSize(12.f)
+		, _DrawMotionTrails(false)
+		, _MotionTrailColor(FLinearColor(0.5f, 0.5f, 0.5f, 1.f))
+		, _MotionTrailColorFade(FLinearColor(0.15f, 0.15f, 0.15f, 1.f))
+		, _MotionTrailSize(10.f)
+		, _MotionTrailFadeLength(0.2f)
+		, _MotionTrailOffset(FVector2D(0.f, 10.f))
+	{}
+		SLATE_ATTRIBUTE(FAvaEaseCurveTangents, Tangents)
+		SLATE_ARGUMENT(float, PreviewSize)
+		/** If true, will expand the preview size to expand for curves that go outside of the normalized range. */
+		SLATE_ARGUMENT(bool, CanExpandPreview)
+		SLATE_ARGUMENT(float, CurveThickness)
+	
+		SLATE_ARGUMENT(FLinearColor, CurveColor)
+		SLATE_ARGUMENT(FLinearColor, StraightColor)
+		SLATE_ARGUMENT(TOptional<FLinearColor>, UnderCurveColor)
+
+		/** If true, will show a custom tooltip showing an animated preview of the curve. */
+		SLATE_ARGUMENT(bool, CustomToolTip)
+
+		SLATE_ARGUMENT(bool, Animate)
+		SLATE_ARGUMENT(bool, FlatAnimation)
+		SLATE_ARGUMENT(float, IdleAnimationLength)
+		SLATE_ARGUMENT(FFrameRate, DisplayRate)
+		SLATE_ARGUMENT(FLinearColor, AnimationColor)
+		SLATE_ARGUMENT(float, AnimationSize)
+		SLATE_ARGUMENT(bool, DrawMotionTrails)
+		SLATE_ARGUMENT(FLinearColor, MotionTrailColor)
+		SLATE_ARGUMENT(FLinearColor, MotionTrailColorFade)
+		SLATE_ARGUMENT(float, MotionTrailSize)
+		SLATE_ARGUMENT(float, MotionTrailFadeLength)
+		SLATE_ARGUMENT(FVector2D, MotionTrailOffset)
+	SLATE_END_ARGS()
+
+	void Construct(const FArguments& InArgs);
+
+	FAvaEaseCurveTangents GetTangents() const;
+	void SetTangents(const TAttribute<FAvaEaseCurveTangents>& InTangents);
+
+	//~ Begin SWidget
+#if WITH_ACCESSIBILITY
+	virtual TSharedRef<FSlateAccessibleWidget> CreateAccessibleWidget() override;
+#endif
+	virtual int32 OnPaint(const FPaintArgs& InArgs, const FGeometry& InAllottedGeometry, const FSlateRect& InMyCullingRect, FSlateWindowElementList& OutDrawElements, int32 InLayerId, const FWidgetStyle& InWidgetStyle, bool bInParentEnabled) const override;
+	//~ End SWidget
+
+protected:
+	void MakeCurveData(const FGeometry& InGeometry, TArray<FVector2D>& OutCurvePoints, TArray<FLinearColor>& OutCurveColors) const;
+
+	float WeightedEvalForTwoKeys(const FAvaEaseCurveTangents& InTangents, const float InTime) const;
+
+	//~ Begin SWidget
+	virtual FVector2D ComputeDesiredSize(const float InLayoutScaleMultiplier) const override;
+	virtual FReply OnMouseButtonDown(const FGeometry& InMyGeometry, const FPointerEvent& InMouseEvent) override;
+	virtual void Tick(const FGeometry& InAllottedGeometry, const double InCurrentTime, const float InDeltaTime) override;
+	//~ End SWidget
+
+	TAttribute<FAvaEaseCurveTangents> Tangents;
+	float PreviewSize = 0.f;
+	bool CanExpandPreview = false;
+	float CurveThickness = 1.f;
+
+	FLinearColor CurveColor;
+	FLinearColor StraightColor;
+	TOptional<FLinearColor> UnderCurveColor;
+
+	bool bAnimate = false;
+	float IdleAnimationLength = 0.f;
+	FFrameRate DisplayRate;
+	FLinearColor AnimationColor;
+	float AnimationSize = 0.f;
+	bool bDrawMotionTrails = false;
+	FLinearColor MotionTrailColor;
+	FLinearColor MotionTrailColorFade;
+	float MotionTrailSize = 0.f;
+	float MotionTrailFadeLength = 0.f;
+	FVector2D MotionTrailOffset;
+
+	TSharedPtr<SAvaEaseCurvePreviewToolTip> ToolTipWidget;
+
+	/** The amount the curve is below value 0 in curve space. */
+	float BelowZeroValue = 0.f;
+	/** The amount the curve is above value 1 in curve space. */
+	float AboveOneValue = 0.f;
+
+	TSharedPtr<FActiveTimerHandle> IdleTimerHandle;
+
+	double CurrentAnimateTime = 0.f;
+	FVector2D CurrentAnimateLocation = FVector2D::ZeroVector;
+
+	struct FMotionTrail
+	{
+		FVector2D Location = FVector2D::ZeroVector;
+		float Time = 0.f;
+		FLinearColor Color = FLinearColor(1.f, 1.f, 1.f, 1.f);
+	};
+	TArray<FMotionTrail> MotionTrails;
+
+	double MotionTrailFadeStartTime = 0.f;
+};
diff --git a/Engine/Plugins/Experimental/Avalanche/Source/AvalancheSequencer/Private/EaseCurveTool/Widgets/SAvaEaseCurvePreviewToolTip.cpp b/Engine/Plugins/Experimental/Avalanche/Source/AvalancheSequencer/Private/EaseCurveTool/Widgets/SAvaEaseCurvePreviewToolTip.cpp
new file mode 100644
--- /dev/null
+++ b/Engine/Plugins/Experimental/Avalanche/Source/AvalancheSequencer/Private/EaseCurveTool/Widgets/SAvaEaseCurvePreviewToolTip.cpp
@@ -0,0 +1,147 @@
+// Copyright Epic Games, Inc. All Rights Reserved.
+
+#include "SAvaEaseCurvePreviewToolTip.h"
+#include "EaseCurveTool/Widgets/SAvaEaseCurvePreset.h"
+#include "Styling/AppStyle.h"
+#include "Styling/StyleColors.h"
+#include "Widgets/Colors/SColorBlock.h"
+#include "Widgets/SToolTip.h"
+
+#define LOCTEXT_NAMESPACE "SAvaEaseCurvePreviewToolTip"
+
+SAvaEaseCurvePreviewToolTip::SAvaEaseCurvePreviewToolTip(const SAvaEaseCurvePreview::FArguments& InPreviewArgs, const TSharedPtr<SWidget>& InAdditionalContent)
+	: PreviewArgs(InPreviewArgs), AdditionalContent(InAdditionalContent)
+{
+}
+
+FText SAvaEaseCurvePreviewToolTip::GetToolTipText(const FAvaEaseCurveTangents& InTangents)
+{
+	return FText::Format(LOCTEXT("FullToolTipText", "{0}\n\nRight click to copy to Json values"), InTangents.ToDisplayText());
+}
+
+TSharedRef<SToolTip> SAvaEaseCurvePreviewToolTip::CreateDefaultToolTip(const SAvaEaseCurvePreview::FArguments& InPreviewArgs, const TSharedPtr<SWidget>& InAdditionalContent)
+{
+	const TSharedRef<SVerticalBox> ContentWidget = SNew(SVerticalBox)
+		+ SVerticalBox::Slot()
+		.AutoHeight()
+		.Padding(10.f)
+		[
+			SNew(SOverlay)
+			+ SOverlay::Slot()
+			[
+				SNew(SColorBlock)
+				.Color(FStyleColors::Background.GetSpecifiedColor())
+			]
+			+ SOverlay::Slot()
+			.Padding(15.f, 15.f, 15.f, 20.f)
+			[
+				SNew(SAvaEaseCurvePreview)
+				.Tangents(InPreviewArgs._Tangents)
+				.PreviewSize(256.f)
+				.CanExpandPreview(true)
+				.CurveThickness(InPreviewArgs._CurveThickness)
+				.CurveColor(InPreviewArgs._CurveColor)
+				.CustomToolTip(false)
+				.Animate(true)
+				.DisplayRate(InPreviewArgs._DisplayRate)
+				.DrawMotionTrails(true)
+			]
+		]
+		+ SVerticalBox::Slot()
+		.AutoHeight()
+		.Padding(15.f, 2.f, 15.f, 0.f)
+		[
+			SNew(SHorizontalBox)
+			+ SHorizontalBox::Slot()
+			.HAlign(HAlign_Left)
+			[
+				SNew(STextBlock)
+				.ToolTipText(LOCTEXT("StartTangentTooltipText", "Start tangent and weight"))
+				.Text(InPreviewArgs._Tangents.Get(FAvaEaseCurveTangents()).GetStartTangentText())
+			]
+			+ SHorizontalBox::Slot()
+			.HAlign(HAlign_Right)
+			[
+				SNew(STextBlock)
+				.ToolTipText(LOCTEXT("EndTangentTooltipText", "End tangent and weight"))
+				.Text(InPreviewArgs._Tangents.Get(FAvaEaseCurveTangents()).GetEndTangentText())
+			]
+		]
+		+ SVerticalBox::Slot()
+		.AutoHeight()
+		.Padding(15.f, 5.f, 15.f, 0.f)
+		[
+			SNew(SHorizontalBox)
+			+ SHorizontalBox::Slot()
+			.HAlign(HAlign_Left)
+			[
+				SNew(STextBlock)
+				.Text(LOCTEXT("CubicBezierText", "Cubic Bezier Points:"))
+			]
+			+ SHorizontalBox::Slot()
+			.HAlign(HAlign_Right)
+			[
+				SNew(STextBlock)
+				.ToolTipText(LOCTEXT("CubicBezierTooltipText", "Cubic bezier points"))
+				.Text(InPreviewArgs._Tangents.Get(FAvaEaseCurveTangents()).GetCubicBezierText())
+			]
+		]
+		+ SVerticalBox::Slot()
+		.AutoHeight()
+		.Padding(20.f, 10.f, 0.f, 10.f)
+		[
+			SNew(STextBlock)
+			.TextStyle(FAppStyle::Get(), TEXT("HintText"))
+			.Text(LOCTEXT("ToolTipText", "Right click to copy to tangent values"))
+		];
+
+	if (InAdditionalContent.IsValid())
+	{
+		ContentWidget->AddSlot()
+			.Padding(15.f, 10.f, 15.f, 0.f)
+			[
+				InAdditionalContent.ToSharedRef()
+			];
+	}
+
+	return SNew(SToolTip)
+		.TextMargin(1.0f)
+		[
+			ContentWidget
+		];
+}
+
+void SAvaEaseCurvePreviewToolTip::CreateToolTipWidget()
+{
+	if (ToolTipWidget.IsValid())
+	{
+		return;
+	}
+
+	ToolTipWidget = CreateDefaultToolTip(PreviewArgs, AdditionalContent);
+}
+
+TSharedRef<SWidget> SAvaEaseCurvePreviewToolTip::AsWidget()
+{
+	CreateToolTipWidget();
+	return ToolTipWidget.ToSharedRef();
+}
+
+TSharedRef<SWidget> SAvaEaseCurvePreviewToolTip::GetContentWidget()
+{
+	CreateToolTipWidget();
+	return ToolTipWidget->GetContentWidget();
+}
+
+void SAvaEaseCurvePreviewToolTip::SetContentWidget(const TSharedRef<SWidget>& InContentWidget)
+{
+	CreateToolTipWidget();
+	ToolTipWidget->SetContentWidget(InContentWidget);
+}
+
+void SAvaEaseCurvePreviewToolTip::InvalidateWidget()
+{
+	ToolTipWidget.Reset();
+}
+
+#undef LOCTEXT_NAMESPACE
diff --git a/Engine/Plugins/Experimental/Avalanche/Source/AvalancheSequencer/Private/EaseCurveTool/Widgets/SAvaEaseCurvePreviewToolTip.h b/Engine/Plugins/Experimental/Avalanche/Source/AvalancheSequencer/Private/EaseCurveTool/Widgets/SAvaEaseCurvePreviewToolTip.h
new file mode 100644
--- /dev/null
+++ b/Engine/Plugins/Experimental/Avalanche/Source/AvalancheSequencer/Private/EaseCurveTool/Widgets/SAvaEaseCurvePreviewToolTip.h
@@ -0,0 +1,41 @@
+// Copyright Epic Games, Inc. All Rights Reserved.
+
+#pragma once
+
+#include "EaseCurveTool/Widgets/SAvaEaseCurvePreview.h"
+#include "Templates/SharedPointer.h"
+#include "Widgets/IToolTip.h"
+
+class SToolTip;
+class SWidget;
+
+class SAvaEaseCurvePreviewToolTip : public IToolTip
+{
+public:
+	static FText GetToolTipText(const FAvaEaseCurveTangents& InTangents);
+
+	static TSharedRef<SToolTip> CreateDefaultToolTip(const SAvaEaseCurvePreview::FArguments& InPreviewArgs, const TSharedPtr<SWidget>& InAdditionalContent = nullptr);
+
+	SAvaEaseCurvePreviewToolTip(const SAvaEaseCurvePreview::FArguments& InPreviewArgs, const TSharedPtr<SWidget>& InAdditionalContent = nullptr);
+	virtual ~SAvaEaseCurvePreviewToolTip() override {}
+
+protected:
+	void CreateToolTipWidget();
+
+	void InvalidateWidget();
+
+	//~ Begin IToolTip
+	virtual TSharedRef<SWidget> AsWidget() override;
+	virtual TSharedRef<SWidget> GetContentWidget() override;
+	virtual void SetContentWidget(const TSharedRef<SWidget>& InContentWidget) override;
+	virtual bool IsEmpty() const override { return false; }
+	virtual bool IsInteractive() const override { return false; }
+	virtual void OnOpening() override {}
+	virtual void OnClosed() override {}
+	//~ End IToolTip
+
+	SAvaEaseCurvePreview::FArguments PreviewArgs;
+	TSharedPtr<SWidget> AdditionalContent;
+
+	TSharedPtr<SToolTip> ToolTipWidget;
+};
diff --git a/Engine/Plugins/Experimental/Avalanche/Source/AvalancheSequencer/Private/EaseCurveTool/Widgets/SAvaEaseCurveTool.cpp b/Engine/Plugins/Experimental/Avalanche/Source/AvalancheSequencer/Private/EaseCurveTool/Widgets/SAvaEaseCurveTool.cpp
new file mode 100644
--- /dev/null
+++ b/Engine/Plugins/Experimental/Avalanche/Source/AvalancheSequencer/Private/EaseCurveTool/Widgets/SAvaEaseCurveTool.cpp
@@ -0,0 +1,801 @@
+// Copyright Epic Games, Inc. All Rights Reserved.
+
+#include "SAvaEaseCurveTool.h"
+#include "Curves/KeyHandle.h"
+#include "Curves/RichCurve.h"
+#include "EaseCurveTool/AvaEaseCurvePreset.h"
+#include "EaseCurveTool/AvaEaseCurveTool.h"
+#include "EaseCurveTool/AvaEaseCurveToolCommands.h"
+#include "EaseCurveTool/AvaEaseCurveToolSettings.h"
+#include "EaseCurveTool/Widgets/SAvaEaseCurveEditor.h"
+#include "EaseCurveTool/Widgets/SAvaEaseCurvePreset.h"
+#include "Editor.h"
+#include "DetailLayoutBuilder.h"
+#include "Framework/Commands/GenericCommands.h"
+#include "Framework/MultiBox/MultiBoxBuilder.h"
+#include "Math/UnrealMathUtility.h"
+#include "SCurveEditor.h"
+#include "Styling/AppStyle.h"
+#include "Styling/StyleColors.h"
+#include "Widgets/Images/SImage.h"
+#include "Widgets/Input/SButton.h"
+#include "Widgets/Input/SComboBox.h"
+#include "Widgets/Input/SSpinBox.h"
+#include "Widgets/Layout/SScrollBox.h"
+#include "Widgets/Layout/SWrapBox.h"
+
+#define LOCTEXT_NAMESPACE "SAvaEaseCurveTool"
+
+namespace UE::AvaSequencer
+{
+	template<typename NumericType>
+	void ExtractNumericMetadata(FProperty* InProperty
+		, TOptional<NumericType>& OutMinValue
+		, TOptional<NumericType>& OutMaxValue
+		, TOptional<NumericType>& OutSliderMinValue
+		, TOptional<NumericType>& OutSliderMaxValue
+		, NumericType& OutSliderExponent
+		, NumericType& OutDelta
+		, float& OutShiftMultiplier
+		, float& OutCtrlMultiplier
+		, bool& OutSupportDynamicSliderMaxValue
+		, bool& OutSupportDynamicSliderMinValue)
+	{
+		const FString& MetaUIMinString = InProperty->GetMetaData(TEXT("UIMin"));
+		const FString& MetaUIMaxString = InProperty->GetMetaData(TEXT("UIMax"));
+		const FString& SliderExponentString = InProperty->GetMetaData(TEXT("SliderExponent"));
+		const FString& DeltaString = InProperty->GetMetaData(TEXT("Delta"));
+		const FString& ShiftMultiplierString = InProperty->GetMetaData(TEXT("ShiftMultiplier"));
+		const FString& CtrlMultiplierString = InProperty->GetMetaData(TEXT("CtrlMultiplier"));
+		const FString& SupportDynamicSliderMaxValueString = InProperty->GetMetaData(TEXT("SupportDynamicSliderMaxValue"));
+		const FString& SupportDynamicSliderMinValueString = InProperty->GetMetaData(TEXT("SupportDynamicSliderMinValue"));
+		const FString& ClampMinString = InProperty->GetMetaData(TEXT("ClampMin"));
+		const FString& ClampMaxString = InProperty->GetMetaData(TEXT("ClampMax"));
+
+		// If no UIMin/Max was specified then use the clamp string
+		const FString& UIMinString = MetaUIMinString.Len() ? MetaUIMinString : ClampMinString;
+		const FString& UIMaxString = MetaUIMaxString.Len() ? MetaUIMaxString : ClampMaxString;
+
+		NumericType ClampMin = TNumericLimits<NumericType>::Lowest();
+		NumericType ClampMax = TNumericLimits<NumericType>::Max();
+
+		if (!ClampMinString.IsEmpty())
+		{
+			TTypeFromString<NumericType>::FromString(ClampMin, *ClampMinString);
+		}
+
+		if (!ClampMaxString.IsEmpty())
+		{
+			TTypeFromString<NumericType>::FromString(ClampMax, *ClampMaxString);
+		}
+
+		NumericType UIMin = TNumericLimits<NumericType>::Lowest();
+		NumericType UIMax = TNumericLimits<NumericType>::Max();
+		TTypeFromString<NumericType>::FromString(UIMin, *UIMinString);
+		TTypeFromString<NumericType>::FromString(UIMax, *UIMaxString);
+
+		OutSliderExponent = NumericType(1);
+
+		if (SliderExponentString.Len())
+		{
+			TTypeFromString<NumericType>::FromString(OutSliderExponent, *SliderExponentString);
+		}
+
+		OutDelta = NumericType(0);
+
+		if (DeltaString.Len())
+		{
+			TTypeFromString<NumericType>::FromString(OutDelta, *DeltaString);
+		}
+
+		OutShiftMultiplier = 10.f;
+		if (ShiftMultiplierString.Len())
+		{
+			TTypeFromString<float>::FromString(OutShiftMultiplier, *ShiftMultiplierString);
+		}
+
+		OutCtrlMultiplier = 0.1f;
+		if (CtrlMultiplierString.Len())
+		{
+			TTypeFromString<float>::FromString(OutCtrlMultiplier, *CtrlMultiplierString);
+		}
+
+		const NumericType ActualUIMin = FMath::Max(UIMin, ClampMin);
+		const NumericType ActualUIMax = FMath::Min(UIMax, ClampMax);
+
+		OutMinValue = ClampMinString.Len() ? ClampMin : TOptional<NumericType>();
+		OutMaxValue = ClampMaxString.Len() ? ClampMax : TOptional<NumericType>();
+		OutSliderMinValue = (UIMinString.Len()) ? ActualUIMin : TOptional<NumericType>();
+		OutSliderMaxValue = (UIMaxString.Len()) ? ActualUIMax : TOptional<NumericType>();
+
+		OutSupportDynamicSliderMaxValue = SupportDynamicSliderMaxValueString.Len() > 0 && SupportDynamicSliderMaxValueString.ToBool();
+		OutSupportDynamicSliderMinValue = SupportDynamicSliderMinValueString.Len() > 0 && SupportDynamicSliderMinValueString.ToBool();
+	}
+}
+
+void SAvaEaseCurveTool::Construct(const FArguments& InArgs, const TSharedRef<FAvaEaseCurveTool>& InCurveEaseTool)
+{
+	ToolMode = InArgs._ToolMode;
+	ToolOperation = InArgs._ToolOperation;
+
+	CurveEaseTool = InCurveEaseTool;
+
+	ChildSlot
+		[
+			SNew(SVerticalBox)
+			+ SVerticalBox::Slot()
+			.AutoHeight()
+			.HAlign(HAlign_Fill)
+			.Padding(0.f, 1.f, 0.f, 0.f)
+			[
+				SAssignNew(CurvePresetWidget, SAvaEaseCurvePreset)
+				.OnPresetChanged(this, &SAvaEaseCurveTool::OnPresetChanged)
+				.OnGetNewPresetTangents_Lambda([this](FAvaEaseCurveTangents& OutTangents) -> bool
+					{
+						OutTangents = CurveEaseTool->GetEaseCurveTangents();
+						return true;
+					})
+			]
+			+ SVerticalBox::Slot()
+			.AutoHeight()
+			.HAlign(HAlign_Fill)
+			.Padding(0.f, 4.f, 0.f, 0.f)
+			[
+				ConstructCurveEditorPanel()
+			]
+			+ SVerticalBox::Slot()
+			.AutoHeight()
+			.HAlign(HAlign_Fill)
+			.Padding(0.f, 3.f, 0.f, 0.f)
+			[
+				ConstructInputBoxes()
+			]
+		];
+
+	BindCommands();
+
+	if (GEditor)
+	{
+		GEditor->RegisterForUndo(this);
+	}
+}
+
+TSharedRef<SWidget> SAvaEaseCurveTool::ConstructCurveEditorPanel()
+{
+	CurrentGraphSize = GetDefault<UAvaEaseCurveToolSettings>()->GetGraphSize();
+
+	const TSharedRef<FAvaEaseCurveTool> EaseCurveToolRef = CurveEaseTool.ToSharedRef();
+
+	return SNew(SBorder)
+		[
+			SNew(SOverlay)
+			+ SOverlay::Slot()
+			[
+				SAssignNew(CurveEaseEditorWidget, SAvaEaseCurveEditor, CurveEaseTool->GetToolCurve())
+				.DisplayRate(EaseCurveToolRef, &FAvaEaseCurveTool::GetDisplayRate)
+				.Operation(EaseCurveToolRef, &FAvaEaseCurveTool::GetOperation)
+				.DesiredSize_Lambda([this]() -> FVector2D
+					{
+						return FVector2D(CurrentGraphSize);
+					})
+				.OnTangentsChanged(this, &SAvaEaseCurveTool::HandleEditorTangentsChanged)
+				.GridSnap_UObject(GetDefault<UAvaEaseCurveToolSettings>(), &UAvaEaseCurveToolSettings::GetGridSnap)
+				.GridSize_UObject(GetDefault<UAvaEaseCurveToolSettings>(), &UAvaEaseCurveToolSettings::GetGridSize)
+				.GetContextMenuContent(this, &SAvaEaseCurveTool::CreateContextMenuContent)
+				.OnKeyDown(this, &SAvaEaseCurveTool::OnKeyDown)
+				.OnDragStart(this, &SAvaEaseCurveTool::OnEditorDragStart)
+				.OnDragEnd(this, &SAvaEaseCurveTool::OnEditorDragEnd)
+			]
+			+ SOverlay::Slot()
+			.HAlign(HAlign_Fill)
+			.VAlign(VAlign_Top)
+			[
+				SNew(SImage)
+				.Visibility(EVisibility::HitTestInvisible)
+				.ColorAndOpacity(FLinearColor(0.f, 0.f, 0.f, 0.5f))
+				.Image(&(FCoreStyle::Get().GetWidgetStyle<FScrollBoxStyle>("ScrollBox").TopShadowBrush))
+			]
+			+ SOverlay::Slot()
+			.HAlign(HAlign_Fill)
+			.VAlign(VAlign_Bottom)
+			[
+				SNew(SImage)
+				.Visibility(EVisibility::HitTestInvisible)
+				.ColorAndOpacity(FLinearColor(0.f, 0.f, 0.f, 0.5f))
+				.Image(&(FCoreStyle::Get().GetWidgetStyle<FScrollBoxStyle>("ScrollBox").BottomShadowBrush))
+			]
+			+ SOverlay::Slot()
+			.HAlign(HAlign_Left)
+			.VAlign(VAlign_Fill)
+			[
+				SNew(SImage)
+				.Visibility(EVisibility::HitTestInvisible)
+				.ColorAndOpacity(FLinearColor(0.f, 0.f, 0.f, 0.5f))
+				.Image(&(FCoreStyle::Get().GetWidgetStyle<FScrollBoxStyle>("ScrollBox").LeftShadowBrush))
+			]
+			+ SOverlay::Slot()
+			.HAlign(HAlign_Right)
+			.VAlign(VAlign_Fill)
+			[
+				SNew(SImage)
+				.Visibility(EVisibility::HitTestInvisible)
+				.ColorAndOpacity(FLinearColor(0.f, 0.f, 0.f, 0.5f))
+				.Image(&(FCoreStyle::Get().GetWidgetStyle<FScrollBoxStyle>("ScrollBox").RightShadowBrush))
+			]
+		];
+}
+
+TSharedRef<SWidget> SAvaEaseCurveTool::ConstructInputBoxes()
+{
+	constexpr float WrapSize = 120.f;
+
+	constexpr float MinTangent = -180.0f;
+	constexpr float MaxTangent = 180.0f;
+	constexpr float MinWeight = 0.0f;
+	constexpr float MaxWeight = 10.0f;
+
+	return SNew(SWrapBox)
+		.UseAllottedSize(true)
+		.HAlign(HAlign_Center)
+		+ SWrapBox::Slot()
+		.FillLineWhenSizeLessThan(WrapSize)
+		[
+			SNew(SHorizontalBox)
+			+ SHorizontalBox::Slot()
+			.AutoWidth()
+			.VAlign(VAlign_Center)
+			.Padding(0.f, 1.f, 0.f, 0.f)
+			[
+				SNew(SBox)
+				.WidthOverride(26.f)
+				.HAlign(HAlign_Center)
+				[
+					SNew(STextBlock)
+					.Font(IDetailLayoutBuilder::GetDetailFont())
+					.Text(LOCTEXT("OutLabel", "Out"))
+				]
+			]
+			+ SHorizontalBox::Slot()
+			.FillWidth(1.f)
+			.HAlign(HAlign_Fill)
+			.VAlign(VAlign_Center)
+			[
+				ConstructTangentNumBox(LOCTEXT("OutTangentLabel", "T")
+					, LOCTEXT("OutTangentToolTip", "Out Tangent")
+					, TAttribute<float>::CreateSP(this, &SAvaEaseCurveTool::GetStartTangent)
+					, SNumericEntryBox<float>::FOnValueChanged::CreateSP(this, &SAvaEaseCurveTool::OnStartTangentSpinBoxChanged)
+					, MinTangent, MaxTangent)
+			]
+			+ SHorizontalBox::Slot()
+			.FillWidth(1.f)
+			.HAlign(HAlign_Fill)
+			.VAlign(VAlign_Center)
+			[
+				ConstructTangentNumBox(LOCTEXT("OutTangentWeightLabel", "W")
+					, LOCTEXT("OutTangentWeightToolTip", "Out Tangent Weight")
+					, TAttribute<float>::CreateSP(this, &SAvaEaseCurveTool::GetStartTangentWeight)
+					, SNumericEntryBox<float>::FOnValueChanged::CreateSP(this, &SAvaEaseCurveTool::OnStartTangentWeightSpinBoxChanged)
+					, MinWeight, MaxWeight)
+			]
+		]
+		+ SWrapBox::Slot()
+		.FillLineWhenSizeLessThan(WrapSize)
+		[
+			SNew(SHorizontalBox)
+			+ SHorizontalBox::Slot()
+			.AutoWidth()
+			.VAlign(VAlign_Center)
+			.Padding(0.f, 1.f, 0.f, 0.f)
+			[
+				SNew(SBox)
+				.WidthOverride(26.f)
+				.HAlign(HAlign_Center)
+				[
+					SNew(STextBlock)
+					.Font(IDetailLayoutBuilder::GetDetailFont())
+					.Text(LOCTEXT("InLabel", "In"))
+				]
+			]
+			+ SHorizontalBox::Slot()
+			.FillWidth(1.f)
+			.HAlign(HAlign_Fill)
+			.VAlign(VAlign_Center)
+			[
+				ConstructTangentNumBox(LOCTEXT("InTangentLabel", "T")
+					, LOCTEXT("InTangentToolTip", "In Tangent")
+					, TAttribute<float>::CreateSP(this, &SAvaEaseCurveTool::GetEndTangent)
+					, SNumericEntryBox<float>::FOnValueChanged::CreateSP(this, &SAvaEaseCurveTool::OnEndTangentSpinBoxChanged)
+					, MinTangent, MaxTangent)
+			]
+			+ SHorizontalBox::Slot()
+			.FillWidth(1.f)
+			.HAlign(HAlign_Fill)
+			.VAlign(VAlign_Center)
+			[
+				ConstructTangentNumBox(LOCTEXT("InTangentWeightLabel", "W")
+					, LOCTEXT("InTangentWeightToolTip", "In Tangent Weight")
+					, TAttribute<float>::CreateSP(this, &SAvaEaseCurveTool::GetEndTangentWeight)
+					, SNumericEntryBox<float>::FOnValueChanged::CreateSP(this, &SAvaEaseCurveTool::OnEndTangentWeightSpinBoxChanged)
+					, MinWeight, MaxWeight)
+			]
+		];
+}
+
+TSharedRef<SWidget> SAvaEaseCurveTool::ConstructTangentNumBox(const FText& InLabel
+	, const FText& InToolTip
+	, const TAttribute<float>& InValue
+	, const SNumericEntryBox<float>::FOnValueChanged& InOnValueChanged
+	, const TOptional<float>& InMinSliderValue
+	, const TOptional<float>& InMaxSliderValue) const
+{
+	return SNew(SBox)
+		.MaxDesiredWidth(100.f)
+		.HAlign(HAlign_Fill)
+		[
+			SNew(SHorizontalBox)
+			+ SHorizontalBox::Slot()
+			.AutoWidth()
+			.HAlign(HAlign_Right)
+			[
+				SNew(STextBlock)
+				.MinDesiredWidth(8.f)
+				.Margin(FMargin(2.f, 5.f, 2.f, 3.f))
+				.Font(IDetailLayoutBuilder::GetDetailFont())
+				.Text(InLabel)
+				.ToolTipText(InToolTip)
+			]
+			+ SHorizontalBox::Slot()
+			[
+				SNew(SSpinBox<float>)
+				.Font(IDetailLayoutBuilder::GetDetailFont())
+				.MinSliderValue(InMinSliderValue)
+				.MaxSliderValue(InMaxSliderValue)
+				.Delta(0.00001f)
+				.WheelStep(0.001f)
+				.MinFractionalDigits(4)
+				.MaxFractionalDigits(6)
+				.MinDesiredWidth(70.f)
+				.Value(InValue)
+				.OnBeginSliderMovement_Lambda([this]()
+					{
+						CurveEaseTool->BeginTransaction(LOCTEXT("SliderDragStartLabel", "Ease Curve Slider Drag"));
+					})
+				.OnEndSliderMovement_Lambda([this](const float InNewValue)
+					{
+						CurveEaseTool->EndTransaction();
+					})
+				.OnValueChanged_Lambda([this, InOnValueChanged](const float InNewValue)
+					{
+						InOnValueChanged.ExecuteIfBound(InNewValue);
+					})
+				.OnValueCommitted_Lambda([InOnValueChanged](const float InNewValue, ETextCommit::Type InCommitType)
+					{
+						InOnValueChanged.ExecuteIfBound(InNewValue);
+					})
+			]
+		];
+}
+
+void SAvaEaseCurveTool::HandleEditorTangentsChanged(const FAvaEaseCurveTangents& InTangents) const
+{
+	SetTangents(InTangents, true, true, true);
+}
+
+void SAvaEaseCurveTool::OnEditorDragStart() const
+{
+	CurveEaseTool->BeginTransaction(LOCTEXT("EditorDragStartLabel", "Ease Curve Graph Drag"));
+}
+
+void SAvaEaseCurveTool::OnEditorDragEnd() const
+{
+	CurveEaseTool->EndTransaction();
+}
+
+void SAvaEaseCurveTool::SetTangents(const FAvaEaseCurveTangents& InTangents, const bool bInSetEaseCurve, const bool bInBroadcastUpdate, const bool bInSetSequencerTangents) const
+{
+	if (CurvePresetWidget.IsValid())
+	{
+		if (!CurvePresetWidget->SetSelectedItem(InTangents))
+		{
+			CurvePresetWidget->ClearSelection();
+		}
+	}
+
+	// To change the graph UI tangents, we need to change the ease curve object tangents and the graph will reflect.
+	if (bInSetEaseCurve && CurveEaseTool.IsValid())
+	{
+		CurveEaseTool->SetEaseCurveTangents(InTangents, bInBroadcastUpdate, bInSetSequencerTangents);
+	}
+
+	if (CurveEaseEditorWidget.IsValid())
+	{
+		CurveEaseEditorWidget->ZoomToFit();
+	}
+}
+
+float SAvaEaseCurveTool::GetStartTangent() const
+{
+	return CurveEaseTool->GetEaseCurveTangents().Start;
+}
+
+float SAvaEaseCurveTool::GetStartTangentWeight() const
+{
+	return CurveEaseTool->GetEaseCurveTangents().StartWeight;
+}
+
+float SAvaEaseCurveTool::GetEndTangent() const
+{
+	return CurveEaseTool->GetEaseCurveTangents().End;
+}
+
+float SAvaEaseCurveTool::GetEndTangentWeight() const
+{
+	return CurveEaseTool->GetEaseCurveTangents().EndWeight;
+}
+
+void SAvaEaseCurveTool::OnStartTangentSpinBoxChanged(const float InNewValue) const
+{
+	FAvaEaseCurveTangents NewTangents = CurveEaseTool->GetEaseCurveTangents();
+	NewTangents.Start = InNewValue;
+	SetTangents(NewTangents, true, true, true);
+}
+
+void SAvaEaseCurveTool::OnStartTangentWeightSpinBoxChanged(const float InNewValue) const
+{
+	FAvaEaseCurveTangents NewTangents = CurveEaseTool->GetEaseCurveTangents();
+	NewTangents.StartWeight = InNewValue;
+	SetTangents(NewTangents, true, true, true);
+}
+
+void SAvaEaseCurveTool::OnEndTangentSpinBoxChanged(const float InNewValue) const
+{
+	FAvaEaseCurveTangents NewTangents = CurveEaseTool->GetEaseCurveTangents();
+	NewTangents.End = InNewValue;
+	SetTangents(NewTangents, true, true, true);
+}
+
+void SAvaEaseCurveTool::OnEndTangentWeightSpinBoxChanged(const float InNewValue) const
+{
+	FAvaEaseCurveTangents NewTangents = CurveEaseTool->GetEaseCurveTangents();
+	NewTangents.EndWeight = InNewValue;
+	SetTangents(NewTangents, true, true, true);
+}
+
+void SAvaEaseCurveTool::OnPresetChanged(const TSharedPtr<FAvaEaseCurvePreset>& InPreset) const
+{
+	SetTangents(InPreset->Tangents, true, true, true);
+}
+
+void SAvaEaseCurveTool::BindCommands()
+{
+	const FAvaEaseCurveToolCommands& AvaCurveEaseToolCommands = FAvaEaseCurveToolCommands::Get();
+
+	CommandList = MakeShared<FUICommandList>();
+
+	CommandList->MapAction(FGenericCommands::Get().Undo, FExecuteAction::CreateSP(CurveEaseTool.Get(), &FAvaEaseCurveTool::UndoAction));
+	CommandList->MapAction(FGenericCommands::Get().Redo, FExecuteAction::CreateSP(CurveEaseTool.Get(), &FAvaEaseCurveTool::RedoAction));
+
+	CommandList->MapAction(AvaCurveEaseToolCommands.OpenToolSettings, FExecuteAction::CreateSP(CurveEaseTool.Get(), &FAvaEaseCurveTool::OpenToolSettings));
+
+	CommandList->MapAction(AvaCurveEaseToolCommands.Refresh, FExecuteAction::CreateSP(CurveEaseTool.Get(), &FAvaEaseCurveTool::UpdateEaseCurveFromSequencerKeySelections));
+	CommandList->MapAction(AvaCurveEaseToolCommands.Apply, FExecuteAction::CreateSP(CurveEaseTool.Get(), &FAvaEaseCurveTool::ApplyEaseCurveToSequencerKeySelections));
+
+	CommandList->MapAction(AvaCurveEaseToolCommands.ZoomToFit, FExecuteAction::CreateSP(this, &SAvaEaseCurveTool::ZoomToFit));
+	CommandList->MapAction(AvaCurveEaseToolCommands.ToggleGridSnap
+		, FExecuteAction::CreateUObject(GetMutableDefault<UAvaEaseCurveToolSettings>(), &UAvaEaseCurveToolSettings::ToggleGridSnap)
+		, FCanExecuteAction()
+		, FIsActionChecked::CreateUObject(GetDefault<UAvaEaseCurveToolSettings>(), &UAvaEaseCurveToolSettings::GetGridSnap));
+	CommandList->MapAction(AvaCurveEaseToolCommands.ToggleAutoFlipTangents
+		, FExecuteAction::CreateUObject(GetMutableDefault<UAvaEaseCurveToolSettings>(), &UAvaEaseCurveToolSettings::ToggleAutoFlipTangents)
+		, FCanExecuteAction()
+		, FIsActionChecked::CreateUObject(GetDefault<UAvaEaseCurveToolSettings>(), &UAvaEaseCurveToolSettings::GetAutoFlipTangents));
+
+	CommandList->MapAction(AvaCurveEaseToolCommands.SelectNextChannelKey, FExecuteAction::CreateSP(CurveEaseTool.Get(), &FAvaEaseCurveTool::SelectNextChannelKey));
+	CommandList->MapAction(AvaCurveEaseToolCommands.SelectPreviousChannelKey, FExecuteAction::CreateSP(CurveEaseTool.Get(), &FAvaEaseCurveTool::SelectPreviousChannelKey));
+
+	CommandList->MapAction(AvaCurveEaseToolCommands.SetOperationToEaseOut
+		, FExecuteAction::CreateSP(CurveEaseTool.Get(), &FAvaEaseCurveTool::SetToolOperation, FAvaEaseCurveTool::EOperation::Out)
+		, FCanExecuteAction()
+		, FIsActionChecked::CreateSP(CurveEaseTool.Get(), &FAvaEaseCurveTool::IsToolOperation, FAvaEaseCurveTool::EOperation::Out));
+	CommandList->MapAction(AvaCurveEaseToolCommands.SetOperationToEaseInOut
+		, FExecuteAction::CreateSP(CurveEaseTool.Get(), &FAvaEaseCurveTool::SetToolOperation, FAvaEaseCurveTool::EOperation::InOut)
+		, FCanExecuteAction()
+		, FIsActionChecked::CreateSP(CurveEaseTool.Get(), &FAvaEaseCurveTool::IsToolOperation, FAvaEaseCurveTool::EOperation::InOut));
+	CommandList->MapAction(AvaCurveEaseToolCommands.SetOperationToEaseIn
+		, FExecuteAction::CreateSP(CurveEaseTool.Get(), &FAvaEaseCurveTool::SetToolOperation, FAvaEaseCurveTool::EOperation::In)
+		, FCanExecuteAction()
+		, FIsActionChecked::CreateSP(CurveEaseTool.Get(), &FAvaEaseCurveTool::IsToolOperation, FAvaEaseCurveTool::EOperation::In));
+
+	CommandList->MapAction(AvaCurveEaseToolCommands.ResetTangents
+		, FExecuteAction::CreateSP(CurveEaseTool.Get(), &FAvaEaseCurveTool::ResetEaseCurveTangents, true, true));
+	CommandList->MapAction(AvaCurveEaseToolCommands.ResetStartTangent
+		, FExecuteAction::CreateSP(CurveEaseTool.Get(), &FAvaEaseCurveTool::ResetEaseCurveTangents, true, false));
+	CommandList->MapAction(AvaCurveEaseToolCommands.ResetEndTangent
+		, FExecuteAction::CreateSP(CurveEaseTool.Get(), &FAvaEaseCurveTool::ResetEaseCurveTangents, false, true));
+
+	CommandList->MapAction(AvaCurveEaseToolCommands.FlattenTangents
+		, FExecuteAction::CreateSP(CurveEaseTool.Get(), &FAvaEaseCurveTool::FlattenOrStraightenTangents, true, true, true));
+	CommandList->MapAction(AvaCurveEaseToolCommands.FlattenStartTangent
+		, FExecuteAction::CreateSP(CurveEaseTool.Get(), &FAvaEaseCurveTool::FlattenOrStraightenTangents, true, false, true));
+	CommandList->MapAction(AvaCurveEaseToolCommands.FlattenEndTangent
+		, FExecuteAction::CreateSP(CurveEaseTool.Get(), &FAvaEaseCurveTool::FlattenOrStraightenTangents, false, true, true));
+
+	CommandList->MapAction(AvaCurveEaseToolCommands.StraightenTangents
+		, FExecuteAction::CreateSP(CurveEaseTool.Get(), &FAvaEaseCurveTool::FlattenOrStraightenTangents, true, true, false));
+	CommandList->MapAction(AvaCurveEaseToolCommands.StraightenStartTangent
+		, FExecuteAction::CreateSP(CurveEaseTool.Get(), &FAvaEaseCurveTool::FlattenOrStraightenTangents, true, false, false));
+	CommandList->MapAction(AvaCurveEaseToolCommands.StraightenEndTangent
+		, FExecuteAction::CreateSP(CurveEaseTool.Get(), &FAvaEaseCurveTool::FlattenOrStraightenTangents, false, true, false));
+
+	CommandList->MapAction(AvaCurveEaseToolCommands.CopyTangents
+		, FExecuteAction::CreateSP(CurveEaseTool.Get(), &FAvaEaseCurveTool::CopyTangentsToClipboard)
+		, FCanExecuteAction::CreateSP(CurveEaseTool.Get(), &FAvaEaseCurveTool::CanCopyTangentsToClipboard));
+	CommandList->MapAction(AvaCurveEaseToolCommands.PasteTangents
+		, FExecuteAction::CreateSP(CurveEaseTool.Get(), &FAvaEaseCurveTool::PasteTangentsFromClipboard)
+		, FCanExecuteAction::CreateSP(CurveEaseTool.Get(), &FAvaEaseCurveTool::CanPasteTangentsFromClipboard));
+
+	CommandList->MapAction(AvaCurveEaseToolCommands.CreateExternalCurveAsset
+		, FExecuteAction::CreateLambda([this]()
+			{
+				CurveEaseTool->CreateCurveAsset();
+			})
+		, FCanExecuteAction());
+
+	CommandList->MapAction(AvaCurveEaseToolCommands.SetKeyInterpConstant,
+		FExecuteAction::CreateSP(CurveEaseTool.Get(), &FAvaEaseCurveTool::SetKeyInterpMode, RCIM_Constant, RCTM_Auto),
+		FCanExecuteAction(),
+		FIsActionChecked::CreateSP(CurveEaseTool.Get(), &FAvaEaseCurveTool::IsKeyInterpMode, RCIM_Constant, RCTM_Auto));
+	CommandList->MapAction(AvaCurveEaseToolCommands.SetKeyInterpLinear,
+		FExecuteAction::CreateSP(CurveEaseTool.Get(), &FAvaEaseCurveTool::SetKeyInterpMode, RCIM_Linear, RCTM_Auto),
+		FCanExecuteAction(),
+		FIsActionChecked::CreateSP(CurveEaseTool.Get(), &FAvaEaseCurveTool::IsKeyInterpMode, RCIM_Linear, RCTM_Auto));
+	CommandList->MapAction(AvaCurveEaseToolCommands.SetKeyInterpCubicAuto,
+		FExecuteAction::CreateSP(CurveEaseTool.Get(), &FAvaEaseCurveTool::SetKeyInterpMode, RCIM_Cubic, RCTM_Auto),
+		FCanExecuteAction(),
+		FIsActionChecked::CreateSP(CurveEaseTool.Get(), &FAvaEaseCurveTool::IsKeyInterpMode, RCIM_Cubic, RCTM_Auto));
+	CommandList->MapAction(AvaCurveEaseToolCommands.SetKeyInterpCubicSmartAuto,
+		FExecuteAction::CreateSP(CurveEaseTool.Get(), &FAvaEaseCurveTool::SetKeyInterpMode, RCIM_Cubic, RCTM_SmartAuto),
+		FCanExecuteAction(),
+		FIsActionChecked::CreateSP(CurveEaseTool.Get(), &FAvaEaseCurveTool::IsKeyInterpMode, RCIM_Cubic, RCTM_SmartAuto));
+	CommandList->MapAction(AvaCurveEaseToolCommands.SetKeyInterpCubicUser,
+		FExecuteAction::CreateSP(CurveEaseTool.Get(), &FAvaEaseCurveTool::SetKeyInterpMode, RCIM_Cubic, RCTM_User),
+		FCanExecuteAction(),
+		FIsActionChecked::CreateSP(CurveEaseTool.Get(), &FAvaEaseCurveTool::IsKeyInterpMode, RCIM_Cubic, RCTM_User));
+	CommandList->MapAction(AvaCurveEaseToolCommands.SetKeyInterpCubicBreak,
+		FExecuteAction::CreateSP(CurveEaseTool.Get(), &FAvaEaseCurveTool::SetKeyInterpMode, RCIM_Cubic, RCTM_Break),
+		FCanExecuteAction(),
+		FIsActionChecked::CreateSP(CurveEaseTool.Get(), &FAvaEaseCurveTool::IsKeyInterpMode, RCIM_Cubic, RCTM_Break));
+}
+
+TSharedRef<SWidget> SAvaEaseCurveTool::CreateContextMenuContent()
+{
+	const FAvaEaseCurveToolCommands& AvaCurveEaseToolCommands = FAvaEaseCurveToolCommands::Get();
+
+	TSharedPtr<TArray<TSharedPtr<FCurveViewModel>>> CurvesToAddKeysTo = MakeShareable(new TArray<TSharedPtr<FCurveViewModel>>());
+	
+	FMenuBuilder MenuBuilder(/*bInShouldCloseWindowAfterMenuSelection=*/true, CommandList);
+
+	MenuBuilder.BeginSection(TEXT("CurveEditorActions"), LOCTEXT("CurveAction", "Curve Ease Tool Actions"));
+	{
+		MenuBuilder.AddSubMenu(
+			LOCTEXT("SettingsSubMenuLabel", "Settings"),
+			LOCTEXT("SettingsSubMenuToolTip", ""),
+			FNewMenuDelegate::CreateSP(this, &SAvaEaseCurveTool::MakeContextMenuSettings),
+			false,
+			FSlateIcon(FAppStyle::Get().GetStyleSetName(), TEXT("Icons.Toolbar.Settings")));
+		MenuBuilder.AddSeparator();
+		MenuBuilder.AddMenuEntry(AvaCurveEaseToolCommands.CreateExternalCurveAsset);
+		MenuBuilder.AddSeparator();
+		MenuBuilder.AddMenuEntry(AvaCurveEaseToolCommands.CopyTangents);
+		MenuBuilder.AddMenuEntry(AvaCurveEaseToolCommands.PasteTangents);
+		MenuBuilder.AddSeparator();
+		MenuBuilder.AddSubMenu(
+			LOCTEXT("StraightenTangentsSubMenuLabel", "Straighten Tangents"),
+			LOCTEXT("StraightenTangentsSubMenuToolTip", ""),
+			FNewMenuDelegate::CreateLambda([&AvaCurveEaseToolCommands](FMenuBuilder& InMenuBuilder)
+				{
+					InMenuBuilder.AddMenuEntry(AvaCurveEaseToolCommands.StraightenTangents);
+					InMenuBuilder.AddMenuEntry(AvaCurveEaseToolCommands.StraightenStartTangent);
+					InMenuBuilder.AddMenuEntry(AvaCurveEaseToolCommands.StraightenEndTangent);
+				}),
+			false, FSlateIcon(FAppStyle::Get().GetStyleSetName(), TEXT("GenericCurveEditor.StraightenTangents")));
+		MenuBuilder.AddSubMenu(
+			LOCTEXT("FlattenTangentsSubMenuLabel", "Flatten Tangents"),
+			LOCTEXT("FlattenTangentsSubMenuToolTip", ""),
+			FNewMenuDelegate::CreateLambda([&AvaCurveEaseToolCommands](FMenuBuilder& InMenuBuilder)
+				{
+					InMenuBuilder.AddMenuEntry(AvaCurveEaseToolCommands.FlattenTangents);
+					InMenuBuilder.AddMenuEntry(AvaCurveEaseToolCommands.FlattenStartTangent);
+					InMenuBuilder.AddMenuEntry(AvaCurveEaseToolCommands.FlattenEndTangent);
+				}), 
+			false, FSlateIcon(FAppStyle::Get().GetStyleSetName(), TEXT("GenericCurveEditor.FlattenTangents")));
+		MenuBuilder.AddSubMenu(
+			LOCTEXT("ResetTangentsSubMenuLabel", "Reset Tangents"),
+			LOCTEXT("ResetTangentsSubMenuToolTip", ""),
+			FNewMenuDelegate::CreateLambda([&AvaCurveEaseToolCommands](FMenuBuilder& InMenuBuilder)
+				{
+					InMenuBuilder.AddMenuEntry(AvaCurveEaseToolCommands.ResetTangents);
+					InMenuBuilder.AddMenuEntry(AvaCurveEaseToolCommands.ResetStartTangent);
+					InMenuBuilder.AddMenuEntry(AvaCurveEaseToolCommands.ResetEndTangent);
+				}),
+			false, FSlateIcon(FAppStyle::Get().GetStyleSetName(), TEXT("PropertyWindow.DiffersFromDefault")));
+		MenuBuilder.AddSeparator();
+		/** @TODO: Only show these in single edit mode(?)
+		MenuBuilder.AddMenuEntry(AvaCurveEaseToolCommands.SetKeyInterpConstant);
+		MenuBuilder.AddMenuEntry(AvaCurveEaseToolCommands.SetKeyInterpLinear);
+		MenuBuilder.AddMenuEntry(AvaCurveEaseToolCommands.SetKeyInterpCubicAuto);
+		MenuBuilder.AddMenuEntry(AvaCurveEaseToolCommands.SetKeyInterpCubicSmartAuto);
+		MenuBuilder.AddMenuEntry(AvaCurveEaseToolCommands.SetKeyInterpCubicUser);
+		MenuBuilder.AddMenuEntry(AvaCurveEaseToolCommands.SetKeyInterpCubicBreak);*/
+		MenuBuilder.AddSeparator();
+		MenuBuilder.AddMenuEntry(AvaCurveEaseToolCommands.SetOperationToEaseOut);
+		MenuBuilder.AddMenuEntry(AvaCurveEaseToolCommands.SetOperationToEaseInOut);
+		MenuBuilder.AddMenuEntry(AvaCurveEaseToolCommands.SetOperationToEaseIn);
+		MenuBuilder.AddSeparator();
+		MenuBuilder.AddMenuEntry(AvaCurveEaseToolCommands.SelectPreviousChannelKey);
+		MenuBuilder.AddMenuEntry(AvaCurveEaseToolCommands.SelectNextChannelKey);
+		MenuBuilder.AddSeparator();
+		MenuBuilder.AddMenuEntry(AvaCurveEaseToolCommands.ToggleGridSnap);
+		MenuBuilder.AddMenuEntry(AvaCurveEaseToolCommands.ZoomToFit);
+		MenuBuilder.AddSeparator();
+		MenuBuilder.AddMenuEntry(AvaCurveEaseToolCommands.Refresh);
+		MenuBuilder.AddMenuEntry(AvaCurveEaseToolCommands.Apply);
+	}
+	MenuBuilder.EndSection();
+
+	return MenuBuilder.MakeWidget();
+}
+
+void SAvaEaseCurveTool::MakeContextMenuSettings(FMenuBuilder& InMenuBuilder)
+{
+	const FAvaEaseCurveToolCommands& AvaCurveEaseToolCommands = FAvaEaseCurveToolCommands::Get();
+
+	InMenuBuilder.AddMenuEntry(AvaCurveEaseToolCommands.OpenToolSettings);
+	InMenuBuilder.AddSeparator();
+	InMenuBuilder.AddMenuEntry(AvaCurveEaseToolCommands.ToggleAutoFlipTangents);
+	InMenuBuilder.AddSeparator();
+	{
+		FProperty* const GraphSizeProperty = UAvaEaseCurveToolSettings::StaticClass()->FindPropertyByName(TEXT("GraphSize"));
+		check(GraphSizeProperty);
+
+		TOptional<int32> MinValue, MaxValue, SliderMinValue, SliderMaxValue;
+		int32 SliderExponent, Delta;
+		float ShiftMultiplier = 10.f;
+		float CtrlMultiplier = 0.1f;
+		bool SupportDynamicSliderMaxValue = false;
+		bool SupportDynamicSliderMinValue = false;
+		UE::AvaSequencer::ExtractNumericMetadata(GraphSizeProperty
+			, MinValue, MaxValue
+			, SliderMinValue, SliderMaxValue
+			, SliderExponent, Delta
+			, ShiftMultiplier
+			, CtrlMultiplier
+			, SupportDynamicSliderMaxValue, SupportDynamicSliderMinValue);
+
+		const TSharedRef<SNumericEntryBox<int32>> GraphSizeWidget = SNew(SNumericEntryBox<int32>)
+			.Font(FAppStyle::GetFontStyle(TEXT("MenuItem.Font")))
+			.AllowSpin(true)
+			.MinValue(MinValue)
+			.MaxValue(MaxValue)
+			.MinSliderValue(SliderMinValue)
+			.MaxSliderValue(SliderMaxValue)
+			.SliderExponent(SliderExponent)
+			.Delta(Delta)
+			.ShiftMultiplier(ShiftMultiplier)
+			.CtrlMultiplier(CtrlMultiplier)
+			.SupportDynamicSliderMaxValue(SupportDynamicSliderMaxValue)
+			.SupportDynamicSliderMinValue(SupportDynamicSliderMinValue)
+			.Value_Lambda([this]()
+				{
+					return CurrentGraphSize;
+				})
+			.OnValueChanged_Lambda([this](const int32 InNewValue)
+				{
+					CurrentGraphSize = InNewValue;
+				})
+			.OnValueCommitted_Lambda([this](const int32 InNewValue, ETextCommit::Type InCommitType)
+				{
+					CurrentGraphSize = InNewValue;
+
+					UAvaEaseCurveToolSettings* const EaseCurveToolSettings = GetMutableDefault<UAvaEaseCurveToolSettings>();
+					check(EaseCurveToolSettings);
+					EaseCurveToolSettings->SetGraphSize(CurrentGraphSize);
+					EaseCurveToolSettings->SaveConfig();
+				});
+		InMenuBuilder.AddWidget(GraphSizeWidget, LOCTEXT("ToolSizeLabel", "Tool Size"));
+	}
+	{
+		FProperty* const GridSizeProperty = UAvaEaseCurveToolSettings::StaticClass()->FindPropertyByName(TEXT("GridSize"));
+		check(GridSizeProperty);
+
+		TOptional<int32> MinValue, MaxValue, SliderMinValue, SliderMaxValue;
+		int32 SliderExponent, Delta;
+		float ShiftMultiplier = 10.f;
+		float CtrlMultiplier = 0.1f;
+		bool SupportDynamicSliderMaxValue = false;
+		bool SupportDynamicSliderMinValue = false;
+		UE::AvaSequencer::ExtractNumericMetadata(GridSizeProperty
+			, MinValue, MaxValue
+			, SliderMinValue, SliderMaxValue
+			, SliderExponent, Delta
+			, ShiftMultiplier
+			, CtrlMultiplier
+			, SupportDynamicSliderMaxValue, SupportDynamicSliderMinValue);
+
+		auto SetEaseCurveToolGridSize = [](const int32 InNewValue)
+			{
+				UAvaEaseCurveToolSettings* const EaseCurveToolSettings = GetMutableDefault<UAvaEaseCurveToolSettings>();
+				check(EaseCurveToolSettings);
+				EaseCurveToolSettings->SetGridSize(InNewValue);
+				EaseCurveToolSettings->SaveConfig();
+			};
+
+		const TSharedRef<SNumericEntryBox<int32>> GridSizeWidget = SNew(SNumericEntryBox<int32>)
+			.Font(FAppStyle::GetFontStyle(TEXT("MenuItem.Font")))
+			.AllowSpin(true)
+			.MinValue(MinValue)
+			.MaxValue(MaxValue)
+			.MinSliderValue(SliderMinValue)
+			.MaxSliderValue(SliderMaxValue)
+			.SliderExponent(SliderExponent)
+			.Delta(Delta)
+			.ShiftMultiplier(ShiftMultiplier)
+			.CtrlMultiplier(CtrlMultiplier)
+			.SupportDynamicSliderMaxValue(SupportDynamicSliderMaxValue)
+			.SupportDynamicSliderMinValue(SupportDynamicSliderMinValue)
+			.Value_Lambda([this]()
+				{
+					return GetDefault<UAvaEaseCurveToolSettings>()->GetGridSize();
+				})
+			.OnValueChanged_Lambda(SetEaseCurveToolGridSize)
+			.OnValueCommitted_Lambda([SetEaseCurveToolGridSize](const int32 InNewValue, ETextCommit::Type InCommitType)
+				{
+					SetEaseCurveToolGridSize(InNewValue);
+				});
+		InMenuBuilder.AddWidget(GridSizeWidget, LOCTEXT("GridSizeLabel", "Grid Size"));
+	}
+}
+
+void SAvaEaseCurveTool::UndoAction()
+{
+	if (GEditor)
+	{
+		GEditor->UndoTransaction();
+	}
+}
+
+void SAvaEaseCurveTool::RedoAction()
+{
+	if (GEditor)
+	{
+		GEditor->RedoTransaction();
+	}
+}
+
+void SAvaEaseCurveTool::ZoomToFit() const
+{
+	if (CurveEaseEditorWidget.IsValid())
+	{
+		CurveEaseEditorWidget->ZoomToFit();
+	}
+}
+
+FKeyHandle SAvaEaseCurveTool::GetSelectedKeyHandle() const
+{
+	if (CurveEaseEditorWidget.IsValid())
+	{
+		return CurveEaseEditorWidget->GetSelectedKeyHandle();
+	}
+	return FKeyHandle::Invalid();
+}
+
+FReply SAvaEaseCurveTool::OnKeyDown(const FGeometry& MyGeometry, const FKeyEvent& InKeyEvent)
+{
+	if (CommandList->ProcessCommandBindings(InKeyEvent))
+	{
+		return FReply::Handled();
+	}
+	return FReply::Unhandled();
+}
+
+#undef LOCTEXT_NAMESPACE
diff --git a/Engine/Plugins/Experimental/Avalanche/Source/AvalancheSequencer/Private/EaseCurveTool/Widgets/SAvaEaseCurveTool.h b/Engine/Plugins/Experimental/Avalanche/Source/AvalancheSequencer/Private/EaseCurveTool/Widgets/SAvaEaseCurveTool.h
new file mode 100644
--- /dev/null
+++ b/Engine/Plugins/Experimental/Avalanche/Source/AvalancheSequencer/Private/EaseCurveTool/Widgets/SAvaEaseCurveTool.h
@@ -0,0 +1,92 @@
+// Copyright Epic Games, Inc. All Rights Reserved.
+
+#pragma once
+
+#include "EaseCurveTool/AvaEaseCurveTool.h"
+#include "EditorUndoClient.h"
+#include "Templates/SharedPointer.h"
+#include "Widgets/DeclarativeSyntaxSupport.h"
+#include "Widgets/Input/SNumericEntryBox.h"
+#include "Widgets/SCompoundWidget.h"
+
+class FCurveEditor;
+class FText;
+class SAvaEaseCurveEditor;
+class SAvaEaseCurvePreset;
+class UCurveBase;
+struct FKeyHandle;
+struct FRichCurve;
+
+class SAvaEaseCurveTool
+	: public SCompoundWidget
+	, public FEditorUndoClient
+{
+public:
+	SLATE_BEGIN_ARGS(SAvaEaseCurveTool)
+		: _ToolMode(FAvaEaseCurveTool::EMode::DoubleKeyEdit)
+		, _ToolOperation(FAvaEaseCurveTool::EOperation::InOut)
+	{}
+		SLATE_ATTRIBUTE(FAvaEaseCurveTool::EMode, ToolMode)
+		SLATE_ATTRIBUTE(FAvaEaseCurveTool::EOperation, ToolOperation)
+		SLATE_ARGUMENT(FAvaEaseCurveTangents, InitialTangents)
+	SLATE_END_ARGS()
+
+	void Construct(const FArguments& InArgs, const TSharedRef<FAvaEaseCurveTool>& InCurveEaseTool);
+	
+	void SetTangents(const FAvaEaseCurveTangents& InTangents, const bool bInSetEaseCurve, const bool bInBroadcastUpdate, const bool bInSetSequencerTangents) const;
+
+	float GetStartTangent() const;
+	float GetStartTangentWeight() const;
+	float GetEndTangent() const;
+	float GetEndTangentWeight() const;
+
+	FKeyHandle GetSelectedKeyHandle() const;
+
+protected:
+	TSharedRef<SWidget> ConstructCurveEditorPanel();
+	TSharedRef<SWidget> ConstructInputBoxes();
+	TSharedRef<SWidget> ConstructTangentNumBox(const FText& InLabel
+		, const FText& InToolTip
+		, const TAttribute<float>& InValue
+		, const SNumericEntryBox<float>::FOnValueChanged& InOnValueChanged
+		, const TOptional<float>& InMinSliderValue
+		, const TOptional<float>& InMaxSliderValue) const;
+
+	void HandleEditorTangentsChanged(const FAvaEaseCurveTangents& InTangents) const;
+
+	void OnStartTangentSpinBoxChanged(const float InNewValue) const;
+	void OnStartTangentWeightSpinBoxChanged(const float InNewValue) const;
+	void OnEndTangentSpinBoxChanged(const float InNewValue) const;
+	void OnEndTangentWeightSpinBoxChanged(const float InNewValue) const;
+
+	void OnPresetChanged(const TSharedPtr<FAvaEaseCurvePreset>& InPreset) const;
+
+	void BindCommands();
+
+	TSharedRef<SWidget> CreateContextMenuContent();
+	void MakeContextMenuSettings(FMenuBuilder& InMenuBuilder);
+
+	void UndoAction();
+	void RedoAction();
+
+	void ZoomToFit() const;
+
+	void OnEditorDragStart() const;
+	void OnEditorDragEnd() const;
+
+	//~ Begin SWidget
+	virtual FReply OnKeyDown(const FGeometry& MyGeometry, const FKeyEvent& InKeyEvent) override;
+	//~ End SWidget
+
+	TSharedPtr<FUICommandList> CommandList;
+
+	TSharedPtr<FAvaEaseCurveTool> CurveEaseTool;
+
+	TAttribute<FAvaEaseCurveTool::EMode> ToolMode;
+	TAttribute<FAvaEaseCurveTool::EOperation> ToolOperation;
+
+	TSharedPtr<SAvaEaseCurveEditor> CurveEaseEditorWidget;
+	TSharedPtr<SAvaEaseCurvePreset> CurvePresetWidget;
+
+	float CurrentGraphSize = 200.f;
+};
diff --git a/Engine/Plugins/Experimental/Avalanche/Source/AvalancheSequencer/Private/Settings/AvaSequencerSettings.cpp b/Engine/Plugins/Experimental/Avalanche/Source/AvalancheSequencer/Private/Settings/AvaSequencerSettings.cpp
new file mode 100644
--- /dev/null
+++ b/Engine/Plugins/Experimental/Avalanche/Source/AvalancheSequencer/Private/Settings/AvaSequencerSettings.cpp
@@ -0,0 +1,58 @@
+// Copyright Epic Games, Inc. All Rights Reserved.
+
+#include "AvaSequencerSettings.h"
+#include "AvaSequenceDefaultTags.h"
+#include "Misc/FrameRate.h"
+
+namespace UE::AvaSequencer::Private
+{
+	FAvaSequencePreset CreateDefaultPreset(FName InSequenceType, const FAvaTagSoftHandle& InSequenceTag)
+	{
+		FAvaSequencePreset Preset;
+		Preset.PresetName    = InSequenceType;
+		Preset.SequenceLabel = InSequenceType;
+		Preset.SequenceTag   = InSequenceTag;
+		Preset.bEnableLabel  = true;
+		Preset.bEnableTag    = true;
+		return Preset;
+	}
+
+	TArray<FAvaSequencePreset> MakeDefaultSequencePresets()
+	{
+		const FAvaSequenceDefaultTags& DefaultTags = FAvaSequenceDefaultTags::Get();
+
+		// Guids found via the entries that DefaultSequenceTags had for In/Out/Change
+		TArray<FAvaSequencePreset> DefaultPresets =
+			{
+				CreateDefaultPreset(TEXT("In")    , DefaultTags.In),
+				CreateDefaultPreset(TEXT("Out")   , DefaultTags.Out),
+				CreateDefaultPreset(TEXT("Change"), DefaultTags.Change),
+			};
+
+		// Configure Change Marks
+		{
+			FAvaSequencePreset& ChangePreset = DefaultPresets[2];
+			ChangePreset.bEnableMarks = true;
+
+			FAvaMarkSetting& MarkSetting = ChangePreset.Marks.AddDefaulted_GetRef();
+        	MarkSetting.Label = TEXT("A");
+			MarkSetting.FrameNumber = 60;
+		}
+
+		return DefaultPresets;
+	}
+}
+
+UAvaSequencerSettings::UAvaSequencerSettings()
+{
+	CategoryName = TEXT("Motion Design");
+	SectionName = TEXT("Sequencer");
+
+	DisplayRate.FrameRate = FFrameRate(60000, 1001);
+}
+
+TConstArrayView<FAvaSequencePreset> UAvaSequencerSettings::GetDefaultSequencePresets() const
+{
+	static const TArray<FAvaSequencePreset> DefaultPresets = UE::AvaSequencer::Private::MakeDefaultSequencePresets();
+	return DefaultPresets;
+}
diff --git a/Engine/Plugins/Experimental/Avalanche/Source/AvalancheSequencer/Private/Transition/AvaSequenceTransitionCompiler.cpp b/Engine/Plugins/Experimental/Avalanche/Source/AvalancheSequencer/Private/Transition/AvaSequenceTransitionCompiler.cpp
new file mode 100644
--- /dev/null
+++ b/Engine/Plugins/Experimental/Avalanche/Source/AvalancheSequencer/Private/Transition/AvaSequenceTransitionCompiler.cpp
@@ -0,0 +1,270 @@
+// Copyright Epic Games, Inc. All Rights Reserved.
+
+#include "AvaSequenceTransitionCompiler.h"
+#include "AvaTransitionTreeEditorData.h"
+#include "IAvaTransitionEditorModule.h"
+#include "Tasks/AvaTransitionDelayTask.h"
+#include "Tasks/AvaTransitionWaitForLayerTask.h"
+#include "Transition/AvaTransitionInitializeSequence.h"
+#include "Transition/AvaTransitionPlaySequenceTask.h"
+
+#define LOCTEXT_NAMESPACE "AvaSequenceTransitionCompiler"
+
+namespace UE::AvaSequencer::Private
+{
+	struct FPlayState
+	{
+		const UStateTreeState* State = nullptr;
+
+		// Play Tasks pending to be updated. These are removed as their counterpart Init Tasks are added/updated
+		TArray<const FStateTreeEditorNode*, TInlineAllocator<1>> PlayTaskNodes;
+
+		int32 StateIndex = INDEX_NONE;
+	};
+
+	// Returns true if the state is recognized by this compiler to be a next state
+	bool IsNextState(const UStateTreeState& InState)
+	{
+		// State transition must be only 1
+		if (InState.Transitions.Num() != 1)
+		{
+			return false;
+		}
+
+		const FStateTreeTransition& Transition = InState.Transitions[0];
+
+		// Must be enabled, set to trigger on state completed and transition to next selectable with no condition preventing it from transitioning
+		return Transition.bTransitionEnabled
+			&& Transition.Trigger == EStateTreeTransitionTrigger::OnStateCompleted
+			&& Transition.State.LinkType == EStateTreeTransitionType::NextSelectableState
+			&& Transition.Conditions.IsEmpty();
+	}
+
+	// note: This only checks for WaitForLayer/DelayTask. There could be Delayed Transitions,
+	// or Transitions on an Event, etc that this does not account for.
+	// This is only to cover these basic tasks. It's not meant to cover all possible cases
+	bool IsWaitState(const UStateTreeState& InState)
+	{
+		return InState.Tasks.ContainsByPredicate([](const FStateTreeEditorNode& InNode)
+		{
+			const UScriptStruct* ScriptStruct = InNode.Node.GetScriptStruct();
+			if (!ScriptStruct)
+			{
+				return false;
+			}
+			return ScriptStruct->IsChildOf(FAvaTransitionWaitForLayerTask::StaticStruct())
+				|| ScriptStruct->IsChildOf(FAvaTransitionDelayTask::StaticStruct());
+		});
+	}
+
+	void SetInitState(UStateTreeState& InInitState, const FPlayState& InPlayState)
+	{
+		InInitState.Modify();
+
+		check(InPlayState.State);
+
+		// Copy Enter Conditions of Play State (ensure a new ID is generated)
+		InInitState.EnterConditions.Empty(InPlayState.State->EnterConditions.Num());
+		for (const FStateTreeEditorNode& PlayEnterCondition : InPlayState.State->EnterConditions)
+		{
+			FStateTreeEditorNode& InitEnterCondition = InInitState.EnterConditions.Add_GetRef(PlayEnterCondition);
+			InitEnterCondition.ID = FGuid::NewGuid();
+		}
+
+		// Set Transition to Next Selectable. These new Init states all go before a wait state, so it will never be set to "Tree Completed"
+		InInitState.Transitions.Empty(1);
+		InInitState.AddTransition(EStateTreeTransitionTrigger::OnStateCompleted, EStateTreeTransitionType::NextSelectableState);
+
+		InInitState.Tasks.Empty(1);
+		for (const FStateTreeEditorNode* PlayTaskNode : InPlayState.PlayTaskNodes)
+		{
+			if (PlayTaskNode)
+			{
+				const FAvaTransitionPlaySequenceTask& PlayTask = PlayTaskNode->Node.Get<FAvaTransitionPlaySequenceTask>();
+
+				FAvaTransitionInitializeSequence& InitTask = InInitState.AddTask<FAvaTransitionInitializeSequence>().GetNode();
+				InitTask.InitializeTime     = PlayTask.PlaySettings.Start;
+				InitTask.PlayMode           = PlayTask.PlaySettings.PlayMode;
+				InitTask.QueryType          = PlayTask.QueryType;
+				InitTask.SequenceName       = PlayTask.SequenceName;
+				InitTask.SequenceTag        = PlayTask.SequenceTag;
+				InitTask.bPerformExactMatch = PlayTask.bPerformExactMatch;
+			}
+		}
+	}
+}
+
+FAvaSequenceTransitionCompiler& FAvaSequenceTransitionCompiler::Get()
+{
+	static FAvaSequenceTransitionCompiler Instance;
+	return Instance;
+}
+
+FAvaSequenceTransitionCompiler::FAvaSequenceTransitionCompiler()
+	: AutogeneratedStateContext(FGuid(0x66DD5961, 0x12B24F60, 0x93E8E63D, 0x812B7E49), LOCTEXT("ContextLabel", "Motion Design Sequence Transition Compiler"))
+{
+}
+
+void FAvaSequenceTransitionCompiler::Register()
+{
+	IAvaTransitionEditorModule::Get().GetOnCompileTransitionTree().AddRaw(this, &FAvaSequenceTransitionCompiler::OnCompileTransitionTree);
+}
+
+void FAvaSequenceTransitionCompiler::Unregister()
+{
+	if (IAvaTransitionEditorModule::IsLoaded())
+	{
+		IAvaTransitionEditorModule::Get().GetOnCompileTransitionTree().RemoveAll(this);
+	}
+}
+
+void FAvaSequenceTransitionCompiler::OnCompileTransitionTree(UAvaTransitionTreeEditorData& InEditorData)
+{
+	CompileStatesRecursive(InEditorData, InEditorData.SubTrees);
+}
+
+void FAvaSequenceTransitionCompiler::CompileStatesRecursive(UAvaTransitionTreeEditorData& InEditorData, TArray<TObjectPtr<UStateTreeState>>& InStates)
+{
+	using namespace UE::AvaSequencer;
+
+	// Remove invalid States
+	InStates.RemoveAll(
+		[](const UStateTreeState* InState)
+		{
+			return !IsValid(InState);
+		});
+
+	if (InStates.IsEmpty())
+	{
+		return;
+	}
+
+	// Recurse first as anything after this will aim towards adding new states which don't need to be recursively compiled
+	for (UStateTreeState* State : InStates)
+	{
+		CompileStatesRecursive(InEditorData, State->Children);
+	}
+
+	// Do not compile any auto-generated nodes if the conditions are not 'right'... in that all the states must be set to next selectable (except the very last one)
+	// Iterate from first state to the one before last. The Last state is already assumed to be a finishing state
+	for (int32 StateIndex = 0; StateIndex < InStates.Num() - 1; ++StateIndex)
+	{
+		if (!Private::IsNextState(*InStates[StateIndex]))
+		{
+			return;
+		}
+	}
+
+	TArray<Private::FPlayState> PlayStates;
+
+	const UStateTreeState* WaitState = nullptr;
+	int32 WaitStateIndex = INDEX_NONE;
+
+	// Iterate Backwards: Gather the States with Play Sequence Tasks and keep track of the lowest indexed State with Delay/Wait
+	for (int32 StateIndex = InStates.Num() - 1; StateIndex >= 0; --StateIndex)
+	{
+		const UStateTreeState* State = InStates[StateIndex];
+		if (Private::IsWaitState(*State))
+		{
+			WaitState = State;
+			WaitStateIndex = StateIndex;
+		}
+
+		Private::FPlayState PlayState;
+		PlayState.State = State;
+		PlayState.StateIndex = StateIndex;
+
+		for (const FStateTreeEditorNode& Task : State->Tasks)
+		{
+			if (Task.Node.GetPtr<FAvaTransitionPlaySequenceTask>())
+			{
+				PlayState.PlayTaskNodes.Add(&Task);	
+			}
+		}
+
+		if (!PlayState.PlayTaskNodes.IsEmpty())
+		{
+			PlayStates.Add(MoveTemp(PlayState));
+		}
+	}
+
+	// No Wait State found. Nothing to compile. Return
+	if (WaitStateIndex == INDEX_NONE)
+	{
+		return;
+	}
+
+	// Remove all the Play States that come prior to the Wait State Index
+	// (they come before the waiting so they don't need to be initialized)
+	PlayStates.RemoveAll(
+		[WaitStateIndex](const Private::FPlayState& InPlayState)
+		{
+			return InPlayState.StateIndex <= WaitStateIndex;
+		});
+
+	// Since WaitStateIndex is valid, WaitState should also be
+	check(WaitState);
+
+	// Removes all the states from the provided index all the way to the Wait State Index
+	auto CleanupUnusedStates = [this, &InStates, &InEditorData, WaitStateIndex](int32 InStateIndex)
+	{
+		TSet<UStateTreeState*> StatesToRemove;
+
+		while (InStateIndex < WaitStateIndex)
+		{
+			UStateTreeState* UnhandledState = InStates[InStateIndex];
+			if (InEditorData.IsAutogeneratedState(*UnhandledState, AutogeneratedStateContext))
+			{
+				StatesToRemove.Add(UnhandledState);
+			}
+			++InStateIndex;
+		}
+
+		InStates.RemoveAll(
+			[&StatesToRemove](UStateTreeState* InState)
+			{
+				return StatesToRemove.Contains(InState);
+			});
+	
+	};
+
+	// If no play states are present, cleanup all the auto-generated states
+	if (PlayStates.IsEmpty())
+	{
+		CleanupUnusedStates(0);
+		return;
+	}
+
+	// Find existing Initialize Nodes and update them
+	for (int32 StateIndex = 0; StateIndex < WaitStateIndex; ++StateIndex)
+	{
+		UStateTreeState& State = *InStates[StateIndex];
+
+		// Skip if this State is not autogenerated by this Context
+		if (!InEditorData.IsAutogeneratedState(State, AutogeneratedStateContext))
+		{
+			continue;
+		}
+
+		Private::SetInitState(State, PlayStates.Pop());
+
+		// Break as there are no more Play States left. However, there could still be remaining previously generated states that are now unhandled and need to be removed
+		if (PlayStates.IsEmpty())
+		{
+			CleanupUnusedStates(StateIndex + 1);
+			break;
+		}
+	}
+
+	while (!PlayStates.IsEmpty())
+	{
+		UStateTreeState& NewState = InEditorData.CreateState(*WaitState, /*bAfter*/false);
+		InEditorData.SetAutogeneratedState(NewState.ID, AutogeneratedStateContext);
+		Private::SetInitState(NewState, PlayStates.Pop());
+	}
+
+	// todo: at the moment, refresh is called even if all the states are up-to-date
+	InEditorData.GetOnTreeRequestRefresh().Broadcast();
+}
+
+#undef LOCTEXT_NAMESPACE
diff --git a/Engine/Plugins/Experimental/Avalanche/Source/AvalancheSequencer/Private/Transition/AvaSequenceTransitionCompiler.h b/Engine/Plugins/Experimental/Avalanche/Source/AvalancheSequencer/Private/Transition/AvaSequenceTransitionCompiler.h
new file mode 100644
--- /dev/null
+++ b/Engine/Plugins/Experimental/Avalanche/Source/AvalancheSequencer/Private/Transition/AvaSequenceTransitionCompiler.h
@@ -0,0 +1,29 @@
+// Copyright Epic Games, Inc. All Rights Reserved.
+
+#pragma once
+
+#include "Containers/ContainersFwd.h"
+#include "State/AvaTransitionAutogeneratedStateContext.h"
+
+class UAvaTransitionTreeEditorData;
+class UStateTreeState;
+template <typename T> struct TObjectPtr;
+
+class FAvaSequenceTransitionCompiler
+{
+public:
+	static FAvaSequenceTransitionCompiler& Get();
+
+	FAvaSequenceTransitionCompiler();
+
+	void Register();
+
+	void Unregister();
+
+	void OnCompileTransitionTree(UAvaTransitionTreeEditorData& InEditorData);
+
+	void CompileStatesRecursive(UAvaTransitionTreeEditorData& InEditorData, TArray<TObjectPtr<UStateTreeState>>& InStates);
+
+private:
+	const FAvaTransitionAutogeneratedStateContext AutogeneratedStateContext;
+};
diff --git a/Engine/Plugins/Experimental/Avalanche/Source/AvalancheShapes/Private/AvaShapeParametricMaterial.cpp b/Engine/Plugins/Experimental/Avalanche/Source/AvalancheShapes/Private/AvaShapeParametricMaterial.cpp
new file mode 100644
--- /dev/null
+++ b/Engine/Plugins/Experimental/Avalanche/Source/AvalancheShapes/Private/AvaShapeParametricMaterial.cpp
@@ -0,0 +1,447 @@
+// Copyright Epic Games, Inc. All Rights Reserved.
+
+#include "AvaShapeParametricMaterial.h"
+
+#include "Materials/Material.h"
+#include "Materials/MaterialInstance.h"
+#include "Materials/MaterialInstanceDynamic.h"
+
+FAvaShapeParametricMaterial::FOnMaterialChanged FAvaShapeParametricMaterial::OnMaterialChangedDelegate;
+
+FAvaShapeParametricMaterial::FAvaShapeParametricMaterial()
+{
+	bUseTranslucentMaterial         = false;
+	bUseUnlitMaterial               = false;
+
+	Style    = EAvaShapeParametricMaterialStyle::Solid;
+
+	ColorA   = FLinearColor::White;
+	ColorA.A = 1.f;
+
+	ColorB   = FLinearColor::Black;
+	ColorB.A = 1.f;
+
+	GradientOffset = 0.5f;
+
+	LoadDefaultMaterials();
+}
+
+FAvaShapeParametricMaterial::FAvaShapeParametricMaterial(const FAvaShapeParametricMaterial& Other)
+{
+	bUseUnlitMaterial        = Other.bUseUnlitMaterial;
+	bUseTranslucentMaterial  = Other.bUseTranslucentMaterial;
+	ColorA                   = Other.ColorA;
+	ColorB                   = Other.ColorB;
+	GradientOffset           = Other.GradientOffset;
+	Style                    = Other.Style;
+	ActiveInstanceIndex      = Other.ActiveInstanceIndex;
+
+	DefaultMaterials         = Other.DefaultMaterials;
+	InstanceMaterials.SetNum(DefaultMaterials.Num(), EAllowShrinking::No);
+}
+
+FAvaShapeParametricMaterial& FAvaShapeParametricMaterial::operator=(const FAvaShapeParametricMaterial& Other)
+{
+	bUseUnlitMaterial        = Other.bUseUnlitMaterial;
+	bUseTranslucentMaterial  = Other.bUseTranslucentMaterial;
+	ColorA                   = Other.ColorA;
+	ColorB                   = Other.ColorB;
+	GradientOffset           = Other.GradientOffset;
+	Style                    = Other.Style;
+	ActiveInstanceIndex      = Other.ActiveInstanceIndex;
+
+	DefaultMaterials         = Other.DefaultMaterials;
+	InstanceMaterials.SetNum(DefaultMaterials.Num(), EAllowShrinking::No);
+
+	return *this;
+}
+
+UMaterialInterface* FAvaShapeParametricMaterial::GetDefaultMaterial() const
+{
+	const int32 Index = GetActiveInstanceIndex();
+	check(DefaultMaterials.IsValidIndex(Index))
+	return DefaultMaterials[Index];
+}
+
+void FAvaShapeParametricMaterial::LoadDefaultMaterials()
+{
+	static const TCHAR* LitOpaqueMaterialPath = TEXT("/Avalanche/ToolboxResources/M_Toolbox_Opaque_Lit.M_Toolbox_Opaque_Lit");
+	UMaterialInterface* DefaultLitMaterialOpaque = FindObject<UMaterialInterface>(nullptr, LitOpaqueMaterialPath);
+	if (!DefaultLitMaterialOpaque)
+	{
+		DefaultLitMaterialOpaque = LoadObject<UMaterialInterface>(nullptr, LitOpaqueMaterialPath);
+	}
+	DefaultMaterials.Add(DefaultLitMaterialOpaque);
+
+	static const TCHAR* LitTranslucentMaterialPath = TEXT("/Avalanche/ToolboxResources/M_Toolbox_Translucent_Lit.M_Toolbox_Translucent_Lit");
+	UMaterialInterface* DefaultLitMaterialTranslucent = FindObject<UMaterialInterface>(nullptr, LitTranslucentMaterialPath);
+	if (!DefaultLitMaterialTranslucent)
+	{
+		DefaultLitMaterialTranslucent = LoadObject<UMaterialInterface>(nullptr, LitTranslucentMaterialPath);
+	}
+	DefaultMaterials.Add(DefaultLitMaterialTranslucent);
+
+	static const TCHAR* UnlitOpaqueMaterialPath = TEXT("/Avalanche/ToolboxResources/M_Toolbox_Opaque_Unlit.M_Toolbox_Opaque_Unlit");
+	UMaterialInterface* DefaultUnlitMaterialOpaque = FindObject<UMaterialInterface>(nullptr, UnlitOpaqueMaterialPath);
+	if (!DefaultUnlitMaterialOpaque)
+	{
+		DefaultUnlitMaterialOpaque = LoadObject<UMaterialInterface>(nullptr, UnlitOpaqueMaterialPath);
+	}
+	DefaultMaterials.Add(DefaultUnlitMaterialOpaque);
+
+	static const TCHAR* UnlitTranslucentMaterialPath = TEXT("/Avalanche/ToolboxResources/M_Toolbox_Translucent_Unlit.M_Toolbox_Translucent_Unlit");
+	UMaterialInterface* DefaultUnlitMaterialTranslucent = FindObject<UMaterialInterface>(nullptr, UnlitTranslucentMaterialPath);
+	if (!DefaultUnlitMaterialTranslucent)
+	{
+		DefaultUnlitMaterialTranslucent = LoadObject<UMaterialInterface>(nullptr, UnlitTranslucentMaterialPath);
+	}
+	DefaultMaterials.Add(DefaultUnlitMaterialTranslucent);
+
+	InstanceMaterials.SetNum(DefaultMaterials.Num(), EAllowShrinking::No);
+}
+
+void FAvaShapeParametricMaterial::SetMaterialParameterValues(UMaterialInstanceDynamic* InMaterialInstance, bool bInNotifyUpdate) const
+{
+	if (!IsValid(InMaterialInstance))
+	{
+		return;
+	}
+
+	switch (Style)
+	{
+		case EAvaShapeParametricMaterialStyle::Solid:
+			InMaterialInstance->SetScalarParameterValue(TEXT("LinearGradient"), 0.f);
+			InMaterialInstance->SetVectorParameterValue(TEXT("ColorA"), ColorA);
+		break;
+
+		case EAvaShapeParametricMaterialStyle::LinearGradient:
+			InMaterialInstance->SetScalarParameterValue(TEXT("LinearGradient"), 1.f);
+			InMaterialInstance->SetVectorParameterValue(TEXT("ColorA"), ColorA);
+			InMaterialInstance->SetVectorParameterValue(TEXT("ColorB"), ColorB);
+			InMaterialInstance->SetScalarParameterValue(TEXT("Offset"), GradientOffset);
+		break;
+	}
+
+	if (bInNotifyUpdate
+		&& InstanceMaterials.IsValidIndex(ActiveInstanceIndex)
+		&& InstanceMaterials[ActiveInstanceIndex] == InMaterialInstance)
+	{
+		if (OnMaterialParameterChangedDelegate.IsBound())
+		{
+			OnMaterialParameterChangedDelegate.Broadcast();
+		}
+	}
+}
+
+void FAvaShapeParametricMaterial::OnMaterialParameterUpdated()
+{
+	const int32 PreviousActiveIndex = ActiveInstanceIndex;
+
+	for (const TObjectPtr<UMaterialInstanceDynamic>& Material : InstanceMaterials)
+	{
+		SetMaterialParameterValues(Material, true);
+	}
+
+	if (PreviousActiveIndex != GetActiveInstanceIndex())
+	{
+		if (OnMaterialChangedDelegate.IsBound())
+		{
+			OnMaterialChangedDelegate.Broadcast(*this);
+		}
+	}
+
+}
+
+int32 FAvaShapeParametricMaterial::GetActiveInstanceIndex() const
+{
+	FAvaShapeParametricMaterial* MutableThis = const_cast<FAvaShapeParametricMaterial*>(this);
+
+	if (bUseTranslucentMaterial)
+	{
+		MutableThis->ActiveInstanceIndex = bUseUnlitMaterial ? TRANSLUCENT_UNLIT : TRANSLUCENT_LIT;
+	}
+	else
+	{
+		MutableThis->ActiveInstanceIndex = bUseUnlitMaterial ? OPAQUE_UNLIT : OPAQUE_LIT;
+	}
+
+	return MutableThis->ActiveInstanceIndex;
+}
+
+UMaterialInstanceDynamic* FAvaShapeParametricMaterial::CreateMaterialInstance(UObject* InOuter)
+{
+	if (!IsValid(InOuter))
+	{
+		return nullptr;
+	}
+
+	const UMaterialInterface* ParentMaterial = GetDefaultMaterial();
+
+	if (!ParentMaterial)
+	{
+		LoadDefaultMaterials();
+	}
+
+	UMaterialInstanceDynamic* MaterialInstance = UMaterialInstanceDynamic::Create(
+		GetDefaultMaterial(),
+		InOuter
+	);
+
+	const int32 Index = GetActiveInstanceIndex();
+	InstanceMaterials[Index] = MaterialInstance;
+
+	OnMaterialParameterUpdated();
+
+	return MaterialInstance;
+}
+
+bool FAvaShapeParametricMaterial::CopyFromMaterialParameters(UMaterialInstance* InMaterial)
+{
+	if (!IsParametricMaterial(InMaterial, true))
+	{
+		return false;
+	}
+
+	float IsLinearGradient;
+	if (InMaterial->GetScalarParameterValue(FMaterialParameterInfo(TEXT("LinearGradient")), IsLinearGradient))
+	{
+		Style = IsLinearGradient >= 0.5f ? EAvaShapeParametricMaterialStyle::LinearGradient : EAvaShapeParametricMaterialStyle::Solid;
+	}
+
+	FLinearColor InputColorA;
+	if (InMaterial->GetVectorParameterValue(FMaterialParameterInfo(TEXT("ColorA")), InputColorA))
+	{
+		ColorA = InputColorA;
+	}
+
+	FLinearColor InputColorB;
+	if (InMaterial->GetVectorParameterValue(FMaterialParameterInfo(TEXT("ColorB")), InputColorB))
+	{
+		ColorB = InputColorB;
+	}
+
+	float InputOffset;
+	if (InMaterial->GetScalarParameterValue(FMaterialParameterInfo(TEXT("Offset")), InputOffset))
+	{
+		GradientOffset = InputOffset;
+	}
+
+	return true;
+}
+
+bool FAvaShapeParametricMaterial::IsParametricMaterial(UMaterialInterface* InMaterial, const bool bCheckIfDefault) const
+{
+	if (!InMaterial)
+	{
+		return false;
+	}
+
+	auto IsInstance = [this](const UMaterialInterface* Material)
+	{
+		return InstanceMaterials.Contains(Material);
+	};
+
+	auto IsDefault = [this](const UMaterialInterface* Material)
+	{
+		return DefaultMaterials.Contains(Material);
+	};
+
+	auto IsParametric = [IsInstance, IsDefault, bCheckIfDefault](const UMaterialInterface* Material)
+	{
+		return IsInstance(Material) || (bCheckIfDefault && IsDefault(Material));
+	};
+
+	// Direct comparison first
+	if (IsParametric(InMaterial))
+	{
+		return true;
+	}
+
+	// Otherwise check parent
+	const UMaterialInstance* AsMaterialInstance = Cast<UMaterialInstance>(InMaterial);
+	if (IsValid(AsMaterialInstance) && IsValid(AsMaterialInstance->Parent))
+	{
+		UMaterialInterface* Parent = AsMaterialInstance->Parent;
+		while (Parent)
+		{
+			if (IsParametric(Parent))
+			{
+				return true;
+			}
+
+			if (const UMaterialInstance* ParentAsMaterialInstance = Cast<UMaterialInstance>(Parent))
+			{
+				Parent = ParentAsMaterialInstance->Parent;
+				continue;
+			}
+
+			break;
+		}
+	}
+
+	return false;
+}
+
+UMaterialInstanceDynamic* FAvaShapeParametricMaterial::GetMaterial() const
+{
+	const int32 Index = GetActiveInstanceIndex();
+
+	check(InstanceMaterials.IsValidIndex(Index))
+	UMaterialInstanceDynamic* Material = InstanceMaterials[Index];
+
+	SetMaterialParameterValues(Material, false);
+
+	return Material;
+}
+
+UMaterialInstanceDynamic* FAvaShapeParametricMaterial::GetOrCreateMaterial(UObject* InOuter)
+{
+	UMaterialInstanceDynamic* Material = GetMaterial();
+
+	/** If active instance material is not set, create and set it up */
+	if (Material == nullptr)
+	{
+		Material = CreateMaterialInstance(InOuter);
+	}
+
+	return Material;
+}
+
+FAvaShapeParametricMaterial::FOnMaterialParameterChanged& FAvaShapeParametricMaterial::OnMaterialParameterChanged()
+{
+	return OnMaterialParameterChangedDelegate;
+}
+
+FAvaShapeParametricMaterial::FOnMaterialChanged& FAvaShapeParametricMaterial::OnMaterialChanged()
+{
+	return OnMaterialChangedDelegate;
+}
+
+UObject* FAvaShapeParametricMaterial::GetMaterialsOuter() const
+{
+	UObject* Outer = nullptr;
+
+	for (const TObjectPtr<UMaterialInstanceDynamic>& Material : InstanceMaterials)
+	{
+		if (!IsValid(Material))
+		{
+			continue;
+		}
+
+		if (Outer && Outer != Material->GetOuter())
+		{
+			return nullptr;
+		}
+
+		Outer = Material->GetOuter();
+	}
+
+	return Outer;
+}
+
+void FAvaShapeParametricMaterial::SetMaterial(UMaterialInstanceDynamic* InMaterial)
+{
+	if (!IsValid(InMaterial))
+	{
+		return;
+	}
+
+	if (!IsParametricMaterial(InMaterial, false))
+	{
+		return;
+	}
+
+	const int32 ActiveIndex = GetActiveInstanceIndex();
+	InstanceMaterials[ActiveIndex] = InMaterial;
+
+	OnMaterialParameterUpdated();
+}
+
+void FAvaShapeParametricMaterial::SetUseTranslucentMaterial(bool bInUse)
+{
+	if (bInUse == bUseTranslucentMaterial)
+	{
+		return;
+	}
+
+	bUseTranslucentMaterial = bInUse;
+	OnMaterialParameterUpdated();
+}
+
+void FAvaShapeParametricMaterial::SetUseAutoTranslucency(bool bInUse)
+{
+	if (bInUse == bUseAutoTranslucency)
+	{
+		return;
+	}
+
+	bUseAutoTranslucency = bInUse;
+	OnMaterialParameterUpdated();
+}
+
+void FAvaShapeParametricMaterial::SetStyle(EAvaShapeParametricMaterialStyle InStyle)
+{
+	if (InStyle == Style)
+	{
+		return;
+	}
+
+	Style = InStyle;
+	OnMaterialParameterUpdated();
+}
+
+void FAvaShapeParametricMaterial::SetPrimaryColor(const FLinearColor& InColor)
+{
+	if (InColor == ColorA)
+	{
+		return;
+	}
+
+	ColorA = InColor;
+
+	if (bUseAutoTranslucency)
+	{
+		bUseTranslucentMaterial = ColorA.A < 1.f || ColorB.A < 1.f;	
+	}
+	
+	OnMaterialParameterUpdated();
+}
+
+void FAvaShapeParametricMaterial::SetSecondaryColor(const FLinearColor& InColor)
+{
+	if (InColor == ColorB)
+	{
+		return;
+	}
+
+	ColorB = InColor;
+
+	if (bUseAutoTranslucency)
+	{
+		bUseTranslucentMaterial = ColorA.A < 1.f || ColorB.A < 1.f;
+	}
+	
+	OnMaterialParameterUpdated();
+}
+
+void FAvaShapeParametricMaterial::SetGradientOffset(float InOffset)
+{
+	if (InOffset == GradientOffset)
+	{
+		return;
+	}
+
+	GradientOffset = InOffset;
+	OnMaterialParameterUpdated();
+}
+
+void FAvaShapeParametricMaterial::SetUseUnlitMaterial(bool bInUse)
+{
+	if (bInUse == bUseUnlitMaterial)
+	{
+		return;
+	}
+
+	bUseUnlitMaterial = bInUse;
+	OnMaterialParameterUpdated();
+}
diff --git a/Engine/Plugins/Experimental/Avalanche/Source/AvalancheShapes/Private/AvaShapesModule.cpp b/Engine/Plugins/Experimental/Avalanche/Source/AvalancheShapes/Private/AvaShapesModule.cpp
new file mode 100644
--- /dev/null
+++ b/Engine/Plugins/Experimental/Avalanche/Source/AvalancheShapes/Private/AvaShapesModule.cpp
@@ -0,0 +1,5 @@
+// Copyright Epic Games, Inc. All Rights Reserved.
+
+#include "Modules/ModuleManager.h"
+
+IMPLEMENT_MODULE(FDefaultModuleImpl, AvalancheShapes)
