[MotionDesign]

diff --git a/Engine/Plugins/Experimental/Avalanche/Source/AvalancheShapes/Public/AvaShapeParametricMaterial.h b/Engine/Plugins/Experimental/Avalanche/Source/AvalancheShapes/Public/AvaShapeParametricMaterial.h
new file mode 100644
--- /dev/null
+++ b/Engine/Plugins/Experimental/Avalanche/Source/AvalancheShapes/Public/AvaShapeParametricMaterial.h
@@ -0,0 +1,180 @@
+// Copyright Epic Games, Inc. All Rights Reserved.
+
+#pragma once
+
+#include "AvaShapesDefs.h"
+#include "Math/Color.h"
+#include "Math/MathFwd.h"
+#include "UObject/ObjectPtr.h"
+#include "AvaShapeParametricMaterial.generated.h"
+
+class UMaterial;
+class UMaterialInstance;
+class UMaterialInstanceDynamic;
+class UMaterialInterface;
+class UObject;
+struct FAvaShapeParametricMaterial;
+
+USTRUCT(BlueprintType)
+struct AVALANCHESHAPES_API FAvaShapeParametricMaterial
+{
+	GENERATED_BODY()
+
+	static inline constexpr int32 OPAQUE_LIT = 0;
+	static inline constexpr int32 TRANSLUCENT_LIT = 1;
+	static inline constexpr int32 OPAQUE_UNLIT = 2;
+	static inline constexpr int32 TRANSLUCENT_UNLIT = 3;
+
+	DECLARE_MULTICAST_DELEGATE_OneParam(FOnMaterialChanged, FAvaShapeParametricMaterial&)
+	DECLARE_MULTICAST_DELEGATE(FOnMaterialParameterChanged)
+
+	FAvaShapeParametricMaterial();
+
+	FAvaShapeParametricMaterial(const FAvaShapeParametricMaterial& Other);
+	FAvaShapeParametricMaterial& operator=(const FAvaShapeParametricMaterial& Other);
+
+	/** Called when active material parameters are updated */
+	FOnMaterialParameterChanged& OnMaterialParameterChanged();
+
+	/** Called when active material index changed */
+	static FOnMaterialChanged& OnMaterialChanged();
+
+	bool CopyFromMaterialParameters(UMaterialInstance* InMaterial);
+
+	/** Check if input material is a parametric material */
+	bool IsParametricMaterial(UMaterialInterface* InMaterial, const bool bCheckIfDefault = false) const;
+
+	/** Get default parent material currently active */
+	UMaterialInterface* GetDefaultMaterial() const;
+
+	/** Get outer from all instanced materials or null if they differ */
+	UObject* GetMaterialsOuter() const;
+
+	/** Get active up to date material instance */
+	UMaterialInstanceDynamic* GetMaterial() const;
+
+	/** Get active up to date material instance or creates it */
+	UMaterialInstanceDynamic* GetOrCreateMaterial(UObject* InOuter);
+
+	/** Set instanced material currently active */
+	void SetMaterial(UMaterialInstanceDynamic* InMaterial);
+
+	bool GetUseTranslucentMaterial() const
+	{
+		return bUseTranslucentMaterial;
+	}
+
+	void SetUseTranslucentMaterial(bool bInUse);
+
+	bool GetUseAutoTranslucency() const
+	{
+		return bUseAutoTranslucency;
+	}
+
+	void SetUseAutoTranslucency(bool bInUse);
+
+	EAvaShapeParametricMaterialStyle GetStyle() const
+	{
+		return Style;
+	}
+
+	void SetStyle(EAvaShapeParametricMaterialStyle InStyle);
+
+	const FLinearColor& GetPrimaryColor() const
+	{
+		return ColorA;
+	}
+
+	void SetPrimaryColor(const FLinearColor& InColor);
+
+	const FLinearColor& GetSecondaryColor() const
+	{
+		return ColorB;
+	}
+
+	void SetSecondaryColor(const FLinearColor& InColor);
+
+	float GetGradientOffset() const
+	{
+		return GradientOffset;
+	}
+
+	void SetGradientOffset(float InOffset);
+
+	bool GetUseUnlitMaterial() const
+	{
+		return bUseUnlitMaterial;
+	}
+
+	void SetUseUnlitMaterial(bool bInUse);
+
+	/** Set parameter values on a material instance */
+	void SetMaterialParameterValues(UMaterialInstanceDynamic* InMaterialInstance, bool bInNotifyUpdate = false) const;
+
+protected:
+	/** Parent material
+	 * 1. Opaque Lit
+	 * 2. Translucent Lit
+	 * 3. Opaque Unlit
+	 * 4. Translucent Unlit
+	 * ...
+	 */
+	UPROPERTY(Transient)
+	TArray<TObjectPtr<UMaterialInterface>> DefaultMaterials;
+
+	/** Material instances corresponding to the parent material instanced
+	 * 1. Opaque Lit
+	 * 2. Translucent Lit
+	 * 3. Opaque Unlit
+	 * 4. Translucent Unlit
+	 * ...
+	 */
+	UPROPERTY(Transient)
+	TArray<TObjectPtr<UMaterialInstanceDynamic>> InstanceMaterials;
+
+	// whether we use opaque or translucent material, set automatically when opacity is below 1
+	UPROPERTY()
+	bool bUseTranslucentMaterial;
+
+	UPROPERTY(EditInstanceOnly, BlueprintReadWrite, AdvancedDisplay, Category="Material")
+	bool bUseAutoTranslucency = true;
+
+	// Default style for the material
+	UPROPERTY(EditInstanceOnly, BlueprintReadWrite, Category="Material", meta=(DisplayName="Simple Style", AllowPrivateAccess = "true"))
+	EAvaShapeParametricMaterialStyle Style;
+
+	// Primary colour for the material
+	UPROPERTY(EditInstanceOnly, BlueprintReadWrite, Category="Material", meta=(DisplayName="Simple Color", AllowPrivateAccess = "true"))
+	FLinearColor ColorA;
+
+	// Secondary colour for the material
+	UPROPERTY(EditInstanceOnly, BlueprintReadWrite, Category="Material", meta=(DisplayName="Gradient Color", EditCondition="Style == EAvaShapeParametricMaterialStyle::LinearGradient", EditConditionHides, AllowPrivateAccess = "true"))
+	FLinearColor ColorB;
+
+	// Offset for gradient style material
+	UPROPERTY(EditInstanceOnly, BlueprintReadWrite, Category="Material", meta=(ClampMin="0", ClampMax="1", DisplayName="Gradient Offset", EditCondition = "Style == EAvaShapeParametricMaterialStyle::LinearGradient", EditConditionHides, AllowPrivateAccess = "true"))
+	float GradientOffset;
+
+	// whether the material is unlit or default lit
+	UPROPERTY(EditInstanceOnly, BlueprintReadWrite, Category="Material", meta=(DisplayName="Use Unlit Material", AllowPrivateAccess = "true"))
+	bool bUseUnlitMaterial;
+
+	/** Load parents materials to create instance materials */
+	void LoadDefaultMaterials();
+
+	/** Create an instance material based on the current active parent */
+	UMaterialInstanceDynamic* CreateMaterialInstance(UObject* InOuter);
+
+	/** Update parameter values on all instance materials */
+	void OnMaterialParameterUpdated();
+
+	/** Get the active material index */
+	int32 GetActiveInstanceIndex() const;
+
+private:
+	static FOnMaterialChanged OnMaterialChangedDelegate;
+
+	FOnMaterialParameterChanged OnMaterialParameterChangedDelegate;
+
+	int32 ActiveInstanceIndex = INDEX_NONE;
+};
diff --git a/Engine/Plugins/Experimental/Avalanche/Source/AvalancheShapes/Public/AvaShapeUVParameters.h b/Engine/Plugins/Experimental/Avalanche/Source/AvalancheShapes/Public/AvaShapeUVParameters.h
new file mode 100644
--- /dev/null
+++ b/Engine/Plugins/Experimental/Avalanche/Source/AvalancheShapes/Public/AvaShapeUVParameters.h
@@ -0,0 +1,95 @@
+// Copyright Epic Games, Inc. All Rights Reserved.
+
+#pragma once
+
+#include "AvaShapeActor.h"
+#include "AvaShapesDefs.h"
+#include "AvaShapeUVParameters.generated.h"
+
+USTRUCT(BlueprintType)
+struct AVALANCHESHAPES_API FAvaShapeMaterialUVParameters
+{
+	friend class UAvaShapeDynamicMeshBase;
+	friend class FAvaShapeDynamicMeshVisualizer;
+
+	GENERATED_BODY()
+
+	FAvaShapeMaterialUVParameters()
+		: Mode(EAvaShapeUVMode::Uniform)
+		, AnchorPreset(EAvaAnchors::Center)
+		, Anchor(FVector2D(0.5f, 0.5f))
+		, Scale(FVector2D(1.f, 1.f))
+		, Offset(FVector2D::ZeroVector)
+		, Rotation(0.f)
+		, bFlipHorizontal(false)
+		, bFlipVertical(false)
+		, bUVsDirty(false)
+	{}
+
+	EAvaShapeUVMode GetMode() const
+	{
+		return Mode;
+	}
+
+	EAvaAnchors GetAnchorPreset() const
+	{
+		return AnchorPreset;
+	}
+
+	FVector2D GetAnchor() const
+	{
+		return Anchor;
+	}
+
+	FVector2D GetScale() const
+	{
+		return Scale;
+	}
+
+	FVector2D GetOffset() const
+	{
+		return Offset;
+	}
+
+	float GetRotation() const
+	{
+		return Rotation;
+	}
+
+	bool GetFlipHorizontal() const
+	{
+		return bFlipHorizontal;
+	}
+
+	bool GetFlipVertical() const
+	{
+		return bFlipVertical;
+	}
+
+protected:
+	UPROPERTY(EditInstanceOnly, BlueprintReadWrite, Category="Shape", meta=(AllowPrivateAccess = "true"))
+	EAvaShapeUVMode Mode;
+
+	UPROPERTY(EditInstanceOnly, BlueprintReadWrite, Category="Shape", meta=(AllowPrivateAccess = "true"))
+	EAvaAnchors AnchorPreset;
+
+	UPROPERTY(EditInstanceOnly, BlueprintReadWrite, Category="Shape", meta=(ClampMin="0.0", ClampMax="1.0", EditCondition="AnchorPreset == EAvaAnchors::Custom", AllowPrivateAccess = "true"))
+	FVector2D Anchor;
+
+	UPROPERTY(EditInstanceOnly, BlueprintReadWrite, Category="Shape", meta=(ClampMin="0.01", AllowPreserveRatio, Delta="0.05", AllowPrivateAccess = "true"))
+	FVector2D Scale;
+
+	UPROPERTY(EditInstanceOnly, BlueprintReadWrite, Category="Shape", meta=(Delta="0.05", AllowPrivateAccess = "true"))
+	FVector2D Offset;
+
+	UPROPERTY(EditInstanceOnly, BlueprintReadWrite, Category="Shape", meta=(ClampMin="-360.0", ClampMax="360.0", AllowPrivateAccess = "true"))
+	float Rotation;
+
+	UPROPERTY(EditInstanceOnly, BlueprintReadWrite, Category="Shape", meta=(AllowPrivateAccess = "true"))
+	bool bFlipHorizontal;
+
+	UPROPERTY(EditInstanceOnly, BlueprintReadWrite, Category="Shape", meta=(AllowPrivateAccess = "true"))
+	bool bFlipVertical;
+
+	bool bUVsDirty;
+};
\ No newline at end of file
diff --git a/Engine/Plugins/Experimental/Avalanche/Source/AvalancheShapes/Public/DynamicMeshes/AvaShape2DArrowDynMesh.h b/Engine/Plugins/Experimental/Avalanche/Source/AvalancheShapes/Public/DynamicMeshes/AvaShape2DArrowDynMesh.h
new file mode 100644
--- /dev/null
+++ b/Engine/Plugins/Experimental/Avalanche/Source/AvalancheShapes/Public/DynamicMeshes/AvaShape2DArrowDynMesh.h
@@ -0,0 +1,85 @@
+// Copyright Epic Games, Inc. All Rights Reserved.
+
+#pragma once
+
+#include "DynamicMeshes/AvaShape2DDynMeshBase.h"
+#include "AvaShape2DArrowDynMesh.generated.h"
+
+UCLASS(ClassGroup="Shape", BlueprintType, CustomConstructor, Within=AvaShapeActor)
+class AVALANCHESHAPES_API UAvaShape2DArrowDynamicMesh : public UAvaShape2DDynMeshBase
+{
+	GENERATED_BODY()
+
+public:
+	UAvaShape2DArrowDynamicMesh()
+		: UAvaShape2DArrowDynamicMesh(FVector2D(50.f, 50.f))
+	{
+	}
+
+	UAvaShape2DArrowDynamicMesh(const FVector2D& Size2D, const FLinearColor& InVertexColor = FLinearColor::White)
+		: UAvaShape2DDynMeshBase(Size2D, InVertexColor)
+		, RatioArrowLine(0.5f)
+		, RatioLineHeight(0.5f)
+		, RatioArrowY(0.5f)
+		, RatioLineY(0.5f)
+		, bBothSideArrows(false)
+	{
+	}
+
+	static const FString MeshName;
+	virtual const FString& GetMeshName() const override { return MeshName; }
+
+	UFUNCTION()
+	bool SetRatioArrowLine(float InRatio);
+	float GetRatioArrowLine() const { return RatioArrowLine; }
+
+	UFUNCTION()
+	bool SetRatioLineHeight(float InRatio);
+	float GetRatioLineHeight() const { return RatioLineHeight; }
+
+	UFUNCTION()
+	bool SetRatioArrowY(float InRatio);
+	float GetRatioArrowY() const { return RatioArrowY; }
+
+	UFUNCTION()
+	bool SetRatioLineY(float InRatio);
+	float GetRatioLineY() const { return RatioLineY; }
+
+	UFUNCTION()
+	bool SetBothSideArrows(bool bBothSide);
+	bool IsBothSideArrows() const { return bBothSideArrows; }
+
+protected:
+#if WITH_EDITOR
+	virtual void PostEditChangeProperty(FPropertyChangedEvent& PropertyChangedEvent) override;
+#endif
+
+	virtual void OnRatioArrowLineChanged();
+	virtual void OnRatioLineHeightChanged();
+	virtual void OnBothSideArrowsChanged();
+	virtual void OnRatioArrowYChanged();
+	virtual void OnRatioLineYChanged();
+
+	virtual bool IsMeshVisible(int32 MeshIndex) override;
+	virtual bool CreateMesh(FAvaShapeMesh& InMesh) override;
+
+	// represents the ratio for the arrow and the line, 0.6 means 60% arrow and 40% line
+	UPROPERTY(EditInstanceOnly, BlueprintReadWrite, Category="Shape", meta=(ClampMin="0.0", ClampMax="1.0", DisplayName="Ratio Arrow Line", AllowPrivateAccess="true"))
+	float RatioArrowLine;
+
+	// represents the ratio for the line height, 1 means 100% of the height available
+	UPROPERTY(EditInstanceOnly, BlueprintReadWrite, Category="Shape", meta=(ClampMin="0.0", ClampMax="1.0", DisplayName="Ratio Line Height", AllowPrivateAccess="true"))
+	float RatioLineHeight;
+
+	// represents the ratio for the arrow end, 0 means arrow point will be at bottom, 1 means arrow point will be at top
+	UPROPERTY(EditInstanceOnly, BlueprintReadWrite, Category="Shape", meta=(ClampMin="0.0", ClampMax="1.0", DisplayName="Ratio Arrow Y", AllowPrivateAccess="true"))
+	float RatioArrowY;
+
+	// represents the ratio for the arrow end, 0 means arrow point will be at bottom, 1 means arrow point will be at top
+	UPROPERTY(EditInstanceOnly, BlueprintReadWrite, Category="Shape", meta=(ClampMin="0.0", ClampMax="1.0", DisplayName="Ratio Line Y", AllowPrivateAccess="true"))
+	float RatioLineY;
+
+	// whether there should be an arrow on both side, arrows will have same ratio
+	UPROPERTY(EditInstanceOnly, BlueprintReadWrite, Category="Shape", meta=(DisplayName="Both Side Arrows", AllowPrivateAccess="true"))
+	bool bBothSideArrows;
+};
diff --git a/Engine/Plugins/Experimental/Avalanche/Source/AvalancheShapes/Public/DynamicMeshes/AvaShapeChevronDynMesh.h b/Engine/Plugins/Experimental/Avalanche/Source/AvalancheShapes/Public/DynamicMeshes/AvaShapeChevronDynMesh.h
new file mode 100644
--- /dev/null
+++ b/Engine/Plugins/Experimental/Avalanche/Source/AvalancheShapes/Public/DynamicMeshes/AvaShapeChevronDynMesh.h
@@ -0,0 +1,47 @@
+// Copyright Epic Games, Inc. All Rights Reserved.
+
+#pragma once
+
+#include "DynamicMeshes/AvaShape2DDynMeshBase.h"
+#include "AvaShapeChevronDynMesh.generated.h"
+
+UCLASS(ClassGroup="Shape", BlueprintType, CustomConstructor, Within=AvaShapeActor)
+class AVALANCHESHAPES_API UAvaShapeChevronDynamicMesh : public UAvaShape2DDynMeshBase
+{
+	GENERATED_BODY()
+
+public:
+	UAvaShapeChevronDynamicMesh()
+		: UAvaShapeChevronDynamicMesh(FVector2D(50.f, 50.f))
+	{ }
+
+	UAvaShapeChevronDynamicMesh(
+		const FVector2D& Size2D,
+		const FLinearColor& InVertexColor = FLinearColor::White,
+		const float InRatioChevron = 0.5)
+		: UAvaShape2DDynMeshBase(Size2D, InVertexColor)
+		, RatioChevron(InRatioChevron)
+	{
+	}
+
+	static const FString MeshName;
+	virtual const FString& GetMeshName() const override { return MeshName; }
+
+	UFUNCTION()
+	bool SetRatioChevron(float InRatio);
+	float GetRatioChevron() const { return RatioChevron; }
+
+protected:
+#if WITH_EDITOR
+	virtual void PostEditChangeProperty(FPropertyChangedEvent& PropertyChangedEvent) override;
+#endif
+
+	virtual void OnRatioChevronChanged();
+
+	virtual bool IsMeshVisible(int32 MeshIndex) override;
+	virtual bool CreateMesh(FAvaShapeMesh& InMesh) override;
+
+	// represents the ratio for the chevron size
+	UPROPERTY(EditInstanceOnly, BlueprintReadWrite, Category="Shape", meta=(ClampMin="0.0", ClampMax="0.99", DisplayName="Ratio Chevron", AllowPrivateAccess="true"))
+	float RatioChevron;
+};
diff --git a/Engine/Plugins/Experimental/Avalanche/Source/AvalancheShapes/Public/DynamicMeshes/AvaShapeConeDynMesh.h b/Engine/Plugins/Experimental/Avalanche/Source/AvalancheShapes/Public/DynamicMeshes/AvaShapeConeDynMesh.h
new file mode 100644
--- /dev/null
+++ b/Engine/Plugins/Experimental/Avalanche/Source/AvalancheShapes/Public/DynamicMeshes/AvaShapeConeDynMesh.h
@@ -0,0 +1,128 @@
+// Copyright Epic Games, Inc. All Rights Reserved.
+
+#pragma once
+
+#include "AvaShape3DDynMeshBase.h"
+#include "AvaShapeConeDynMesh.generated.h"
+
+UCLASS(ClassGroup="Shape", BlueprintType, CustomConstructor, Within=AvaShapeActor)
+class AVALANCHESHAPES_API UAvaShapeConeDynamicMesh : public UAvaShape3DDynMeshBase
+{
+	GENERATED_BODY()
+
+	friend class FAvaShapeConeDynamicMeshVisualizer;
+
+public:
+	static constexpr uint8 MinNumSides = 3;
+	static constexpr uint8 MaxNumSides = 128;
+
+	static constexpr int32 MESH_INDEX_BOTTOM = 1;
+	static constexpr int32 MESH_INDEX_TOP = 2;
+	static constexpr int32 MESH_INDEX_START = 3;
+	static constexpr int32 MESH_INDEX_END = 4;
+
+	UAvaShapeConeDynamicMesh()
+		: UAvaShapeConeDynamicMesh(FVector(50.f, 50.f, 50.f))
+	{
+	}
+
+	UAvaShapeConeDynamicMesh(
+		const FVector& InSize,
+		const FLinearColor& InVertexColor = FLinearColor::White,
+		float InHeight = 1,
+		float InBaseRadius = 0.5,
+		float InTopRadius = 0,
+		uint8 InNumSides = 32,
+		float InAngleDegree = 360.f,
+		float InStartDegree = 0.f)
+		: UAvaShape3DDynMeshBase(InSize, InVertexColor)
+		, NumSides(InNumSides)
+		, Height(InHeight)
+		, BaseRadius(InBaseRadius)
+		, TopRadius(InTopRadius)
+		, AngleDegree(InAngleDegree)
+		, StartDegree(InStartDegree)
+	{
+	}
+
+	static const FString MeshName;
+	virtual const FString& GetMeshName() const override { return MeshName; };
+
+	UFUNCTION()
+	bool SetNumSides(uint8 InNumSides);
+	uint8 GetNumSides() const { return NumSides; }
+
+	UFUNCTION()
+	float GetBaseRadius() const { return BaseRadius; }
+
+	UFUNCTION()
+	bool SetTopRadius(float InTopRadius);
+	float GetTopRadius() const { return TopRadius; }
+
+	UFUNCTION()
+	float GetHeight() const { return Height; }
+
+	UFUNCTION()
+	bool SetAngleDegree(float InDegree);
+	float GetAngleDegree() const { return AngleDegree; }
+
+	UFUNCTION()
+	bool SetStartDegree(float InDegree);
+	float GetStartDegree() const { return StartDegree; }
+
+protected:
+#if WITH_EDITOR
+	virtual void PostEditChangeProperty(FPropertyChangedEvent& PropertyChangedEvent) override;
+#endif
+
+	virtual void OnHeightChanged();
+	virtual void OnBaseRadiusChanged();
+	virtual void OnTopRadiusChanged();
+	virtual void OnNumSidesChanged();
+	virtual void OnAngleDegreeChanged();
+	virtual void OnStartDegreeChanged();
+
+	virtual bool CreateBaseUVs(FAvaShapeMesh& BaseMesh, FAvaShapeMaterialUVParameters& InUseParams);
+	virtual bool GenerateBaseMeshSections(FAvaShapeMesh& BottomMesh);
+
+	virtual bool CreateTopUVs(FAvaShapeMesh& TopMesh, FAvaShapeMaterialUVParameters& InUseParams);
+	virtual bool GenerateTopMeshSections(FAvaShapeMesh& TopMesh);
+
+	virtual bool CreateConeUVs(FAvaShapeMesh& ConeMesh, FAvaShapeMaterialUVParameters& InUseParams);
+	virtual bool GenerateConeMeshSections(FAvaShapeMesh& ConeMesh);
+
+	virtual bool CreateStartUVs(FAvaShapeMesh& StartMesh, FAvaShapeMaterialUVParameters& InUseParams);
+	virtual bool GenerateStartMeshSections(FAvaShapeMesh& StartMesh);
+
+	virtual bool CreateEndUVs(FAvaShapeMesh& EndMesh, FAvaShapeMaterialUVParameters& InUseParams);
+	virtual bool GenerateEndMeshSections(FAvaShapeMesh& EndMesh);
+
+	virtual void RegisterMeshes() override;
+	virtual bool IsMeshVisible(int32 MeshIndex) override;
+	virtual bool CreateMesh(FAvaShapeMesh& InMesh) override;
+	virtual bool CreateUVs(FAvaShapeMesh& InMesh, FAvaShapeMaterialUVParameters& InParams) override;
+
+	// The number of sides around the base of the cone
+	UPROPERTY(EditInstanceOnly, BlueprintReadWrite, Category="Shape", meta=(ClampMin="3.0", ClampMax="128.0", DisplayName="Sides", AllowPrivateAccess="true"))
+	uint8 NumSides;
+
+	// The height of the cone from the base to the top
+	UPROPERTY()
+	float Height;
+
+	// the radius of the cone base
+	UPROPERTY()
+	float BaseRadius;
+
+	// the ratio for the radius of the cone top
+	UPROPERTY(EditInstanceOnly, BlueprintReadWrite, Category="Shape", meta=(ClampMin="0.0", ClampMax="1.0", DisplayName="Top radius ratio", AllowPrivateAccess="true"))
+	float TopRadius;
+
+	// represents the base angle in degree for the cone
+	UPROPERTY(EditInstanceOnly, BlueprintReadWrite, Category="Shape", meta=(ClampMin="0.0", ClampMax="360.0", DisplayName="Angle degree", AllowPrivateAccess="true"))
+	float AngleDegree;
+
+	// represents the starting angle in degree for the cone
+	UPROPERTY(EditInstanceOnly, BlueprintReadWrite, Category="Shape", meta=(ClampMin="0.0", ClampMax="360.0", DisplayName="Start Degree", AllowPrivateAccess="true"))
+	float StartDegree;
+};
diff --git a/Engine/Plugins/Experimental/Avalanche/Source/AvalancheShapes/Public/DynamicMeshes/AvaShapeDynMeshBase.h b/Engine/Plugins/Experimental/Avalanche/Source/AvalancheShapes/Public/DynamicMeshes/AvaShapeDynMeshBase.h
new file mode 100644
--- /dev/null
+++ b/Engine/Plugins/Experimental/Avalanche/Source/AvalancheShapes/Public/DynamicMeshes/AvaShapeDynMeshBase.h
@@ -0,0 +1,426 @@
+// Copyright Epic Games, Inc. All Rights Reserved.
+
+#pragma once
+
+#include "AvaDefs.h"
+#include "AvaShapeMesh.h"
+#include "AvaShapesDefs.h"
+#include "AvaShapeUVParameters.h"
+#include "Components/ActorComponent.h"
+#include "DynamicMesh/DynamicMesh3.h"
+#include "IAvalancheInteractiveToolsModeDetailsObject.h"
+#include "Viewport/Interaction/AvaSnapPoint.h"
+#include "Viewport/Interaction/IAvaGizmoObject.h"
+
+#include "AvaShapeDynMeshBase.generated.h"
+
+class AAvaShapeActor;
+class UDynamicMeshComponent;
+class UStaticMesh;
+
+DECLARE_MULTICAST_DELEGATE_OneParam(FMaskEnabledDelegate, AActor* /** NewMaskActor */);
+DECLARE_MULTICAST_DELEGATE_OneParam(FMaskDisabledDelegate, AActor* /** OldMaskActor */);
+DECLARE_MULTICAST_DELEGATE_TwoParams(FMaskVisibilityDelegate, const UWorld* /** CurrentWorld */, bool /** bMaskActorVisible */);
+
+UCLASS(ClassGroup="Shape", Abstract, BlueprintType, CustomConstructor, EditInlineNew, DefaultToInstanced)
+class AVALANCHESHAPES_API UAvaShapeDynamicMeshBase
+	: public UActorComponent
+	, public IAvalancheInteractiveToolsModeDetailsObject
+	, public IAvaGizmoObjectInterface
+{
+	GENERATED_BODY()
+
+	friend class FAvaDynMeshBaseDetailCustomization;
+	friend class FAvaMeshesDetailCustomization;
+	friend class FAvaShapeDynamicMeshVisualizer;
+	friend class AAvaShapeActor;
+
+public:
+	static constexpr float MinSizeValue = UE_SMALL_NUMBER;
+	static inline const FVector2D MinSize2D = FVector2D(MinSizeValue);
+	static inline const FVector MinSize3D = FVector(MinSizeValue);
+
+	static constexpr uint8 DefaultSubdivisions = 8;
+	static constexpr uint8 MaxSubdivisions = 64;
+
+	static constexpr int32 MESH_INDEX_NONE = -1;
+	static constexpr int32 MESH_INDEX_ALL = INT_MAX;
+	static constexpr int32 MESH_INDEX_PRIMARY = 0;
+
+	static inline FMaskEnabledDelegate OnMaskEnabled;
+	static inline FMaskDisabledDelegate OnMaskDisabled;
+	static inline FMaskVisibilityDelegate OnMaskVisibility;
+
+	// Converts the Outer of this object to a AAvaShapeActor.
+	AAvaShapeActor* GetShapeActor() const;
+
+	// Gets the Dynamic Mesh component from the Shape actor, this will load it if its nullptr
+	UDynamicMeshComponent* GetShapeMeshComponent() const;
+
+	/** Get the name of the shape */
+	virtual const FString& GetMeshName() const;
+
+	/** Can we change the size of this shape */
+	bool AllowsSizeEditing() const { return bAllowEditSize; }
+
+	// override in child classes, represents the mesh size
+	UFUNCTION()
+	virtual const FVector& GetSize3D() const { return FVector::ZeroVector; }
+
+	UFUNCTION()
+	virtual bool SetSize3D(const FVector& InSize) { return false; }
+
+	// Primary material override
+	bool GetUsePrimaryMaterialEverywhere() const { return bUsePrimaryMaterialEverywhere; }
+	void SetUsePrimaryMaterialEverywhere(bool bInUse);
+
+	/** Checks if shape size is render-able and not minimal */
+	bool IsMeshSizeValid() const;
+
+	/** Checks if the mesh section is valid */
+	bool IsValidMeshIndex(int32 MeshIndex) const;
+
+	UMaterialInterface* GetMaterial(int32 MeshIndex) const;
+	bool SetMaterial(int32 MeshIndex, UMaterialInterface* NewMaterial);
+
+	const FAvaShapeParametricMaterial* GetParametricMaterial(int32 MeshIndex) const;
+	FAvaShapeParametricMaterial* GetParametricMaterial(int32 MeshIndex);
+
+	// Internal materials are not copied over, only settings
+	bool SetParametricMaterial(int32 MeshIndex, const FAvaShapeParametricMaterial& NewMaterialParams);
+
+	// Allows custom uv params per mesh section instead of using the primary one
+	bool SetOverridePrimaryUVParams(int32 MeshIndex, bool bOverride);
+
+	// Helper function to quickly get what you need instead of using GetMeshData()
+	EAvaShapeUVMode GetMaterialUVMode(int32 MeshIndex) const;
+	bool SetMaterialUVMode(int32 MeshIndex, EAvaShapeUVMode InUVMode);
+
+	EAvaAnchors GetMaterialUVAnchorPreset(int32 MeshIndex) const;
+	bool SetMaterialUVAnchorPreset(int32 MeshIndex, EAvaAnchors InUVAnchorPreset);
+
+	float GetMaterialUVRotation(int32 MeshIndex) const;
+	bool SetMaterialUVRotation(int32 MeshIndex, float InUVRotation);
+
+	const FVector2D& GetMaterialUVAnchor(int32 MeshIndex) const;
+	bool SetMaterialUVAnchor(int32 MeshIndex, const FVector2D& InUVAnchor);
+
+	const FVector2D& GetMaterialUVScale(int32 MeshIndex) const;
+	bool SetMaterialUVScale(int32 MeshIndex, const FVector2D& InUVScale);
+
+	const FVector2D& GetMaterialUVOffset(int32 MeshIndex) const;
+	bool SetMaterialUVOffset(int32 MeshIndex, const FVector2D& InUVOffset);
+
+	bool GetMaterialHorizontalFlip(int32 MeshIndex) const;
+	bool SetMaterialHorizontalFlip(int32 MeshIndex, bool InHorizontalFlip);
+
+	bool GetMaterialVerticalFlip(int32 MeshIndex) const;
+	bool SetMaterialVerticalFlip(int32 MeshIndex, bool InVerticalFlip);
+
+	const FAvaShapeMaterialUVParameters* GetInUseMaterialUVParams(int32 MeshIndex) const;
+	FAvaShapeMaterialUVParameters* GetInUseMaterialUVParams(int32 MeshIndex);
+
+	const FAvaShapeMaterialUVParameters* GetMaterialUVParams(int32 MeshIndex) const;
+	bool SetMaterialUVParams(int32 MeshIndex, const FAvaShapeMaterialUVParameters& InParams);
+
+	bool SetMaterialType(int32 MeshIndex, EMaterialType Type);
+	bool IsMaterialType(int32 MeshIndex, EMaterialType Type);
+
+	bool HasMeshRegenWorldLocation() const { return bHasNewMeshRegenWorldLocation; }
+	const FVector& GetMeshRegenWorldLocation() const { return MeshRegenWorldLocation; }
+	void SetMeshRegenWorldLocation(const FVector& NewLocation, bool bImmediateUpdate = false);
+
+	// Generate a list of 3d-space snap points for this shape
+	virtual TArray<FAvaSnapPoint> GetLocalSnapPoints() const;
+	void GetLocalSnapPoints(TArray<FAvaSnapPoint>& Points) const;
+
+	FTransform GetTransform() const;
+
+	/** Clear the dynamic mesh section with a specific index */
+	bool ClearDynamicMeshSection(int32 MeshIndex);
+
+	/** Clear the whole dynamic mesh */
+	bool ClearDynamicMesh();
+
+	/** Converts dynamic mesh to static mesh */
+	bool ExportToStaticMesh(UStaticMesh* DestinationMesh);
+
+	// Gets the bounds of the shape, override this in child classes for custom bounds
+	// Origin is the center of the box, BoxExtent is half the size, pivot is the default location of the pivot for this shape
+	virtual void GetBounds(FVector& Origin, FVector& BoxExtent, FVector& Pivot) const {}
+
+	virtual void OnColorPicked(const FAvaColorChangeData& InNewColorData);
+	virtual FAvaColorChangeData GetActiveColor() const;
+
+	/** Get the registered meshes indexes */
+	TSet<int32> GetMeshesIndexes() const;
+
+	/** Returns the name of the meshes */
+	TArray<FName> GetMeshDataNames() const;
+
+	/** find a registered mesh and gets a pointer to it */
+	FAvaShapeMeshData* GetMeshData(int32 MeshIndex);
+
+	// apply scale, offset, rotation on dynamic mesh section
+	bool ApplyUVsTransform(FAvaShapeMesh& InMesh, FAvaShapeMaterialUVParameters& InParams, FVector2D ShapeSize, FVector2D UVFixOffset, float UVFixRotation = 0.f) const;
+
+protected:
+	static bool SetAlignmentSize(AActor* InActor, const FVector& InSizeMultiplier);
+	static EAvaAnchors GetAnchorFromNumerics(const FVector2D& AnchorNumeric);
+	static FVector2D GetNumericsFromAnchor(EAvaAnchors AnchorEnum);
+
+	//~ Begin UActorComponent
+	virtual void OnRegister() override;
+	virtual void OnComponentCreated() override;
+	virtual void OnComponentDestroyed(bool bDestroyingHierarchy) override;
+	//~ End UActorComponent
+
+	//~ Begin UObject
+	virtual void PostLoad() override;
+#if WITH_EDITOR
+	virtual void PostEditChangeProperty(FPropertyChangedEvent& PropertyChangedEvent) override;
+	virtual void PostEditChangeChainProperty(FPropertyChangedChainEvent& PropertyChangedEvent) override;
+	virtual void PreEditUndo() override;
+	virtual void PostEditUndo() override;
+	virtual void PostEditImport() override;
+	virtual void PostDuplicate(EDuplicateMode::Type DuplicateMode) override;
+	//~ End UObject
+
+	// triggers when we drop an asset (material) onto our actor
+	void OnAssetDropped(UObject* DroppedObj, AActor* TargetActor);
+#endif
+
+	// ~Begin IAvaGizmoObjectInterface
+	virtual void ToggleGizmo_Implementation(const UAvaGizmoComponent* InGizmoComponent, const bool bShowAsGizmo) override;
+	// ~End IAvaGizmoObjectInterface
+
+	EAvaDynamicMeshUpdateState GetMeshUpdateStatus() const { return MeshUpdateStatus; }
+
+	/** Flag all mesh sections dirty to regenerate them */
+	void MarkAllMeshesDirty();
+
+	/** Only update vertices from mesh */
+	void MarkVerticesDirty();
+
+	// When a parametric material is updated
+	void OnParametricMaterialChanged(FAvaShapeParametricMaterial& InMaterial);
+
+	UAvaShapeDynamicMeshBase(const FObjectInitializer& ObjectInitializer)
+		: UAvaShapeDynamicMeshBase(FLinearColor::White)
+	{ }
+
+	UAvaShapeDynamicMeshBase(const FLinearColor& InVertexColor = FLinearColor::White,
+		float InUniformScaledSize = 1.f, bool bInAllowEditSize = true);
+
+	/** Register all meshes once, calls RegisterMeshes() */
+	void SetupMeshes();
+
+	/** override this in child classes, to register new meshes */
+	virtual void RegisterMeshes() {}
+
+	/** Call to register meshes used for this shape, call this inside RegisterMeshes() */
+	bool RegisterMesh(FAvaShapeMeshData& NewMeshData);
+
+	/** Called when a new mesh has been registered */
+	virtual void OnRegisteredMesh(int32 MeshIndex) {}
+
+	/** Called when all the meshes have been registered and the setup is done */
+	virtual void OnRegisteredMeshes() {}
+
+	/** Returns whether the mesh is dirty, called before creating mesh if visible */
+	virtual bool IsMeshDirty(int32 MeshIndex);
+
+	/** override this in child classes, called before creating mesh if visible */
+	virtual bool IsMeshVisible(int32 MeshIndex);
+
+	/** Returns the name of the mesh */
+	FName GetMeshDataName(int32 InMeshIndex) const;
+
+	/** Returns the materials used by every mesh */
+	const TArray<UMaterialInterface*> GetMeshDataMaterials();
+
+	/** Checks to see if attributes are initialized and mesh is ready */
+	bool IsDynamicMeshInitialized() const;
+
+	/** Clears everything and reinitialize the DM to use it properly */
+	void InitializeDynamicMesh();
+
+	EMaterialType& GetMaterialType(int32 MeshIndex);
+	bool SetMaterialDirect(int32 MeshIndex, UMaterialInterface* NewMaterial);
+	FAvaShapeMaterialUVParameters* GetMaterialUVParams(int32 MeshIndex);
+
+	// called once the mesh is done updating
+	virtual void OnMeshUpdateFinished();
+
+	// called once the pixel size has changed, only in editor
+	virtual void OnPixelSizeChanged() {}
+
+	virtual void OnMeshChanged(int32 MeshIndex);
+	virtual void OnVertexColorChanged();
+	virtual void OnMaterialChanged(int32 MaterialIndex);
+	virtual void OnMaterialTypeChanged(int32 MaterialIndex);
+	virtual void OnParametricMaterialChanged(int32 MaterialIndex);
+	virtual void OnUVParamsChanged(int32 MeshIndex);
+	virtual void OnUsesPrimaryUVParamsChanged(int32 MeshIndex);
+
+	virtual bool OnMeshSectionUpdated(FAvaShapeMesh& InMesh);
+
+	// override this in child classes
+	virtual void OnScaledSizeChanged() {}
+
+	// override this in child classes, for 2D and 3D shape to update the scale
+	virtual void OnSizeChanged() {}
+
+	// Clears the mesh sections that are dirty
+	virtual bool ClearMesh();
+
+	// Adds vertex colors to the vertex colour array. Does not empty it.
+	virtual bool CreateColors(FAvaShapeMesh& InMesh);
+
+	// Performs a vertex-colour-only update. Then calls OnMeshUpdate to the PMC
+	virtual bool UpdateColors();
+
+	// Special case when the mesh is scaled, rather than reconfigured.
+	virtual void UpdateVertices();
+
+	// sets the uvs manually on a dynamic mesh section
+	bool ApplyUVsManually(FAvaShapeMesh& InMesh);
+
+	// apply a planar projection on a dynamic mesh section
+	bool ApplyUVsPlanarProjection(FAvaShapeMesh& InMesh, FRotator PlaneRotation, FVector2D PlaneSize);
+
+	// apply a box projection on a dynamic mesh section
+	bool ApplyUVsBoxProjection(FAvaShapeMesh& InMesh, FRotator BoxRotation, FVector BoxSize);
+
+	// override this in child classes when you create UV for a specific mesh
+	virtual bool CreateUVs(FAvaShapeMesh& InMesh, FAvaShapeMaterialUVParameters& InParams);
+
+	// override this in child classes to create the mesh if it's visible
+	virtual bool CreateMesh(FAvaShapeMesh& InMesh) { return true; }
+
+	// checks whether the topology of the mesh has changed (different vertices count and different triangles count than current)
+	bool HasSameTopology(FAvaShapeMesh& InMesh);
+
+	// Creates a section in the dynamic mesh component
+	bool CreateDynamicMesh(FAvaShapeMesh& InMesh);
+
+	// Updates a section in the dynamic mesh component
+	bool UpdateDynamicMesh(FAvaShapeMesh& InMesh);
+
+	// Invalidates a section and runs the mesh update if necessary
+	void InvalidateSection(bool& bInvalidatedSection, bool bRequireUpdate = true, bool bDoAsync = true);
+
+	/** Flag a mesh section dirty to regenerate it */
+	void MarkMeshDirty(int32 MeshIndex);
+
+	// Takes all dirty flags, regenerate whatever is needed and updates the DMC, called by RunUpdate or RequireUpdate
+	void UpdateDynamicMesh();
+
+	// runs the update for this mesh, async by default
+	void RequireUpdate(bool bDoAsync = true);
+
+	// used to check for material changes in slot of DMC to match with our current material settings
+	bool CheckMaterialSlotChanges();
+
+	// Changes the material on the shape component and keeps the MDI up to date.
+	void SetShapeComponentMaterial(int32 InMaterialIndex, UMaterialInterface* InNewMaterial);
+
+	/** Scale the vertices on all sections of the mesh */
+	void ScaleVertices(const FVector& InScale);
+	void ScaleVertices(const FVector2D& InScale);
+
+	// Meshes used for the current shape sections
+	UPROPERTY(EditInstanceOnly, BlueprintReadOnly, Category="Material", meta=(EditFixedOrder, EditFixedSize, DisplayAfter="bUsePrimaryMaterialEverywhere", AllowPrivateAccess = "true"))
+	TMap<int32, FAvaShapeMeshData> MeshDatas;
+
+#if WITH_EDITORONLY_DATA
+	/** Stores select settings to allow restoration when Gizmo mode turned off */
+	UPROPERTY()
+	TMap<int32, FAvaShapeMeshData> NonGizmoMeshData;
+#endif
+
+	// enable mesh size property editing
+	UPROPERTY(EditDefaultsOnly, BlueprintReadOnly, Category="Shape")
+	bool bAllowEditSize;
+
+	// the type of size you want to handle
+	UPROPERTY(EditInstanceOnly, BlueprintReadWrite, Category="Shape", Transient, meta=(DisplayName="Size Type", EditCondition="bAllowEditSize", EditConditionHides, AllowPrivateAccess="true"))
+	ESizeType SizeType = ESizeType::UnrealUnit;
+
+	// Uniform scaled size of the mesh
+	UPROPERTY(EditInstanceOnly, BlueprintReadWrite, Transient, Category="Shape", meta=(ClampMin="0.0", DisplayName="Uniform Scaled Size", Units="times", AllowPrivateAccess="true"))
+	float UniformScaledSize = 1.f;
+
+	UPROPERTY(BlueprintReadWrite, AdvancedDisplay, Category = "Shape", meta=(DisplayName="Material Vertex Color", AllowPrivateAccess="true"))
+	FLinearColor VertexColor;
+
+	// use primary material for every slot available
+	UPROPERTY(EditInstanceOnly, BlueprintReadWrite, Category="Material", meta=(DisplayName="Use Single Material", AllowPrivateAccess="true"))
+	bool bUsePrimaryMaterialEverywhere;
+
+	TArray<FAvaSnapPoint> LocalSnapPoints;
+
+private:
+	/** find a registered mesh and gets a const pointer to it */
+	const FAvaShapeMeshData* GetMeshData(int32 MeshIndex) const;
+
+	FAvaShapeMesh* GetMesh(int32 MeshIndex);
+
+	/** Update render state (materials) */
+	void MarkMeshRenderStateDirty() const;
+
+	/** Generates the dirty mesh sections */
+	bool GenerateMesh();
+
+	/** Generates the uvs for the mesh (all sections) */
+	bool GenerateUV();
+
+	/** Generates the tangents for the mesh (all sections) */
+	bool GenerateTangents();
+
+	/** Generates the normals for the mesh (all sections), needs the uv set on the mesh */
+	bool GenerateNormals();
+
+	/** runs the update for this mesh only if it's required (mesh is not up to date), async by default */
+	void RunUpdate(bool bDoAsync = true);
+
+	/** Restore the existing saved cached mesh on the component */
+	void RestoreCachedMesh();
+
+	/** Save the existing component mesh into a saved cached mesh */
+	void SaveCachedMesh();
+
+	/** Apply primary material everywhere */
+	void OnUsePrimaryMaterialEverywhereChanged();
+
+	// used to quickly restore mesh that has been modified after OnMeshUpdateFinished was called (eg extrude, masking)
+	TOptional<UE::Geometry::FDynamicMesh3> CachedMesh;
+
+	// only change this value from the base class to avoid thread lock, use RunUpdate() or RequireUpdate()
+	std::atomic<EAvaDynamicMeshUpdateState> MeshUpdateStatus = EAvaDynamicMeshUpdateState::UpToDate;
+
+	/** Do not access directly use getter to load it if nullptr */
+	TWeakObjectPtr<UDynamicMeshComponent> CachedComponent;
+
+	UPROPERTY()
+	FVector MeshRegenWorldLocation;
+
+	bool bHasNewMeshRegenWorldLocation;
+
+	// flag to register meshes only once
+	bool bMeshRegistered;
+
+	// whether we are in the tool shape mode and creating this shape
+	bool bInCreateMode;
+
+	/** true if all meshes are marked as dirty, false otherwise, use MarkAllMeshesDirty() instead of modifying this */
+	bool bAllMeshDirty;
+
+	/** true if at least one meshes is marked as dirty, false otherwise, use MarkMeshDirty(MeshIndex) instead of modifying this */
+	bool bAnyMeshDirty;
+
+	// Should only be used when the shape's vertices have changed, but the overall shape has not changed.
+	bool bVerticesDirty;
+
+	bool bColorsDirty;
+};
diff --git a/Engine/Plugins/Experimental/Avalanche/Source/AvalancheShapes/Public/DynamicMeshes/AvaShapeEllipseDynMesh.h b/Engine/Plugins/Experimental/Avalanche/Source/AvalancheShapes/Public/DynamicMeshes/AvaShapeEllipseDynMesh.h
new file mode 100644
--- /dev/null
+++ b/Engine/Plugins/Experimental/Avalanche/Source/AvalancheShapes/Public/DynamicMeshes/AvaShapeEllipseDynMesh.h
@@ -0,0 +1,71 @@
+// Copyright Epic Games, Inc. All Rights Reserved.
+
+#pragma once
+
+#include "DynamicMeshes/AvaShape2DDynMeshBase.h"
+#include "AvaShapeEllipseDynMesh.generated.h"
+
+UCLASS(ClassGroup="Shape", BlueprintType, CustomConstructor, Within=AvaShapeActor)
+class AVALANCHESHAPES_API UAvaShapeEllipseDynamicMesh : public UAvaShape2DDynMeshBase
+{
+	GENERATED_BODY()
+
+	friend class FAvaShapeEllipseDynamicMeshVisualizer;
+
+public:
+	static constexpr uint8 MinNumSides = 3;
+	static constexpr uint8 MaxNumSides = 128;
+
+	UAvaShapeEllipseDynamicMesh()
+		: UAvaShapeEllipseDynamicMesh(FVector2D(50.f, 50.f))
+	{ }
+
+	UAvaShapeEllipseDynamicMesh(const FVector2D& Size2D, const FLinearColor& InVertexColor = FLinearColor::White,
+		uint8 InNumSides = 64, float InAngleDegree = 360.f, float InStartDegree = 0.f)
+		: UAvaShape2DDynMeshBase(Size2D, InVertexColor)
+		, NumSides(InNumSides)
+		, AngleDegree(InAngleDegree)
+		, StartDegree(InStartDegree)
+	{
+		bDoNotRecenterVertices = true;
+	}
+
+	static const FString MeshName;
+	virtual const FString& GetMeshName() const override { return MeshName; }
+
+	UFUNCTION()
+	bool SetNumSides(uint8 InNumSides);
+	uint8 GetNumSides() const { return NumSides; }
+
+	UFUNCTION()
+	bool SetAngleDegree(float InDegree);
+	float GetAngleDegree() const { return AngleDegree; }
+
+	UFUNCTION()
+	bool SetStartDegree(float InDegree);
+	float GetStartDegree() const { return StartDegree; }
+
+protected:
+#if WITH_EDITOR
+	virtual void PostEditChangeProperty(FPropertyChangedEvent& PropertyChangedEvent) override;
+#endif
+
+	virtual void OnNumSidesChanged();
+	virtual void OnAngleDegreeChanged();
+	virtual void OnStartDegreeChanged();
+
+	virtual bool IsMeshVisible(int32 MeshIndex) override;
+	virtual bool CreateMesh(FAvaShapeMesh& InMesh) override;
+
+	// The number of sides for the shape
+	UPROPERTY(EditInstanceOnly, BlueprintReadWrite, Category="Shape", meta=(ClampMin="3.0", ClampMax="128.0", DisplayName="Sides", AllowPrivateAccess="true"))
+	uint8 NumSides;
+
+	// represents the angle in degree for the ellipse
+	UPROPERTY(EditInstanceOnly, BlueprintReadWrite, Category="Shape", meta=(ClampMin="0.0", ClampMax="360.0", DisplayName="Angle degree", AllowPrivateAccess="true"))
+	float AngleDegree;
+
+	// represents the starting angle in degree for the ellipse
+	UPROPERTY(EditInstanceOnly, BlueprintReadWrite, Category="Shape", meta=(ClampMin="0.0", ClampMax="360.0", DisplayName="Start Degree", AllowPrivateAccess="true"))
+	float StartDegree;
+};
diff --git a/Engine/Plugins/Experimental/Avalanche/Source/AvalancheShapes/Public/DynamicMeshes/AvaShapeIrregularPolygonDynMesh.h b/Engine/Plugins/Experimental/Avalanche/Source/AvalancheShapes/Public/DynamicMeshes/AvaShapeIrregularPolygonDynMesh.h
new file mode 100644
--- /dev/null
+++ b/Engine/Plugins/Experimental/Avalanche/Source/AvalancheShapes/Public/DynamicMeshes/AvaShapeIrregularPolygonDynMesh.h
@@ -0,0 +1,154 @@
+// Copyright Epic Games, Inc. All Rights Reserved.
+
+#pragma once
+
+#include "AvaShapeRoundedPolygonDynMesh.h"
+#include "DynamicMeshes/AvaShape2DDynMeshBase.h"
+#include "AvaShapeIrregularPolygonDynMesh.generated.h"
+
+USTRUCT(BlueprintType)
+struct AVALANCHESHAPES_API FAvaShapeRoundedCornerSettings
+{
+	GENERATED_BODY()
+
+	UPROPERTY(EditInstanceOnly, BlueprintReadWrite, Category="Shape", meta=(ClampMin="0.0",ClampMax="1.0"))
+	float BevelSize = 0.f;
+
+	UPROPERTY(EditInstanceOnly, BlueprintReadWrite, Category="Shape", meta=(ClampMax="64.0"))
+	uint8 BevelSubdivisions = UAvaShapeDynamicMeshBase::DefaultSubdivisions;
+};
+
+USTRUCT(BlueprintType)
+struct AVALANCHESHAPES_API FAvaShapeRoundedCorner
+{
+	GENERATED_BODY()
+
+	FAvaShapeRoundedCorner()
+		: FAvaShapeRoundedCorner(FVector2D::ZeroVector)
+	{
+	}
+
+	FAvaShapeRoundedCorner(const FVector2D& InLocation)
+	{
+		Location = InLocation;
+		CornerMetrics = { 0.f, FVector2D::ZeroVector, FVector2D::ZeroVector,
+			FVector2D::ZeroVector, true, FVector2D::ZeroVector, false };
+		Settings = { 0.f, 0 };
+	}
+
+	UPROPERTY(EditInstanceOnly, BlueprintReadWrite, Category="Shape")
+	FVector2D Location;
+
+	UPROPERTY(EditInstanceOnly, BlueprintReadWrite, Category="Shape", meta=(DisplayName="Corner"))
+	FAvaShapeRoundedCornerSettings Settings;
+
+	FAvaShapeRoundedCornerMetrics CornerMetrics;
+};
+
+UCLASS(ClassGroup="Shape", BlueprintType, CustomConstructor, Within=AvaShapeActor)
+class AVALANCHESHAPES_API UAvaShapeIrregularPolygonDynamicMesh : public UAvaShape2DDynMeshBase
+{
+	GENERATED_BODY()
+
+	friend class FAvaShapeIrregularPolygonDynamicMeshVisualizer;
+
+public:
+	static const float MinPointDistance;
+
+	static bool DoLinesIntersect(const FVector2D Origin1, const FVector2D End1, const FVector2D Origin2, const FVector2D End2);
+
+	UAvaShapeIrregularPolygonDynamicMesh(const FObjectInitializer& ObjectInitializer)
+     		: UAvaShapeIrregularPolygonDynamicMesh(FLinearColor::White)
+	{
+    }
+
+	UAvaShapeIrregularPolygonDynamicMesh(const FLinearColor& InVertexColor)
+		: UAvaShape2DDynMeshBase(FVector2D::Zero(), InVertexColor)
+	{
+		bDoNotRecenterVertices = true;
+		bAllowEditSize = false;
+	}
+
+	static const FString MeshName;
+	virtual const FString& GetMeshName() const override { return MeshName; }
+
+	UFUNCTION()
+	bool SetGlobalBevelSize(float InBevelSize);
+	float GetGlobalBevelSize() const { return GlobalBevelSize; }
+
+	UFUNCTION()
+	void SetGlobalBevelSubdivisions(uint8 InBevelSubdivisions);
+	uint8 GetGlobalBevelSubdivisions() const { return GlobalBevelSubdivisions; }
+
+	bool CanAddPoint(const FVector2D& InPoint);
+	bool AddPoint(const FVector2D& InPoint);
+	bool RemovePoint(int32 PointIdx);
+	bool RemoveFirstPoint();
+	bool RemoveLastPoint();
+	bool RemoveAllPoints();
+	const FAvaShapeRoundedCorner& GetPoint(int32 PointIdx) const { return Points[PointIdx]; }
+	const TArray<FAvaShapeRoundedCorner>& GetPoints() const { return Points; }
+	int32 GetNumPoints() const { return Points.Num(); }
+	bool SetPoints(const TArray<FVector2D>& InPoints);
+
+	void RecalculateActorPosition();
+
+protected:
+#if WITH_EDITOR
+	virtual void PreEditChange(FProperty* PropertyAboutToChange) override;
+	virtual void PostEditChangeProperty(FPropertyChangedEvent& PropertyChangedEvent) override;
+#endif
+
+	void RecalculateExtent();
+
+	bool SetLocation(int32 PointIdx, const FVector2D& InPoint);
+	bool SetBevelSize(int32 PointIdx, float InBevelSize);
+	bool SetBevelSubdivisions(int32 PointIdx, uint8 InBevelSubdivisions);
+
+	bool ShiftPoints(const FVector2D& Amount);
+
+	bool DoesLineIntersectBorder(const FVector2D CheckOrigin, const FVector2D CheckEnd) const;
+
+	// Checks to see if this point is too near to other points.
+	bool IsPointTooCloseToAnotherPoint(const FVector2D& InPoint, bool bAllowExactlyOneMatch = false) const;
+	bool IsPointTooCloseToALine(const FVector2D& InPoint) const;
+	bool IsLineTooCloseToAPoint(const FVector2D& Start, const FVector2D& End) const;
+	bool CanBeGenerated() const;
+
+	bool IsLocationInsideShape(const FVector2D& Location);
+
+	void BackupPoints();
+
+	float GetMaxBevelSizeForPoint(int32 PointIdx) const;
+
+	// Breaks the line, adding a new point, in between InPointIdx and InPointIdx+1
+	bool BreakSide(int32 InPointIdx);
+
+	void RestorePoints();
+	bool CheckNewPointsArray();
+	void OnPointsUpdated();
+	void OnLocationsUpdated();
+	void OnBevelSizeChanged();
+	void OnBevelSubdivisionsChanged();
+	void OnGlobalBevelSizeChanged();
+	void OnGlobalBevelSubdivisionsChanged();
+
+	void OnLocationUpdated(int32 InPointIdx);
+	void OnBevelSizeChanged(int32 PointIdx);
+	void OnBevelSubdivisionsChanged(int32 PointIdx);
+
+	FVector ScreenToWorld(const FVector2D& ScreenLocation) const;
+
+	virtual bool ClearMesh() override;
+	virtual bool CreateMesh(FAvaShapeMesh& InMesh) override;
+
+	UPROPERTY(EditInstanceOnly, BlueprintReadWrite, Category="Shape", meta=(ClampMin="0.0",ClampMax="1.0", AllowPrivateAccess="true"))
+	float GlobalBevelSize = 0.f;
+
+	UPROPERTY(EditInstanceOnly, BlueprintReadWrite, Category="Shape", meta=(ClampMax="64.0", AllowPrivateAccess="true"))
+	uint8 GlobalBevelSubdivisions = DefaultSubdivisions;
+
+	UPROPERTY(BlueprintReadWrite, EditInstanceOnly, Category="Shape", meta=(EditFixedOrder, AllowPrivateAccess="true"))
+	TArray<FAvaShapeRoundedCorner> Points;
+	TArray<FAvaShapeRoundedCorner> PreEditPoints;
+};
diff --git a/Engine/Plugins/Experimental/Avalanche/Source/AvalancheShapes/Public/DynamicMeshes/AvaShapeLineDynMesh.h b/Engine/Plugins/Experimental/Avalanche/Source/AvalancheShapes/Public/DynamicMeshes/AvaShapeLineDynMesh.h
new file mode 100644
--- /dev/null
+++ b/Engine/Plugins/Experimental/Avalanche/Source/AvalancheShapes/Public/DynamicMeshes/AvaShapeLineDynMesh.h
@@ -0,0 +1,61 @@
+// Copyright Epic Games, Inc. All Rights Reserved.
+
+#pragma once
+
+#include "AvaShapeRoundedPolygonDynMesh.h"
+#include "AvaShapeLineDynMesh.generated.h"
+
+UCLASS(ClassGroup="Shape", BlueprintType, CustomConstructor, Within=AvaShapeActor)
+class AVALANCHESHAPES_API UAvaShapeLineDynamicMesh : public UAvaShapeRoundedPolygonDynamicMesh
+{
+	GENERATED_BODY()
+
+	friend class FAvaShapeLineDynamicMeshVisualizer;
+
+public:
+	UAvaShapeLineDynamicMesh()
+		: UAvaShapeLineDynamicMesh(FVector2D(50.f, 1.f))
+	{
+	}
+
+	UAvaShapeLineDynamicMesh(const FVector2D& Size2D, const FLinearColor& InVertexColor = FLinearColor::White,
+		float InBevelSize = 0.f, uint8 InBevelSubdivisions = 0, float InLineWidth = 5.f)
+		: UAvaShapeRoundedPolygonDynamicMesh(Size2D, InVertexColor, InBevelSize, InBevelSubdivisions)
+		, LineWidth(InLineWidth)
+		, Vector(FVector2D(50.f, 0.f))
+	{
+		bAllowEditSize = false;
+	}
+
+	static const FString MeshName;
+	virtual const FString& GetMeshName() const override { return MeshName; }
+
+	UFUNCTION()
+	bool SetLineWidth(float InLineWidth);
+	float GetLineWidth() const { return LineWidth; }
+
+	UFUNCTION()
+	bool SetVector(const FVector2D& InVector);
+	const FVector2D& GetVector() const { return Vector; }
+
+protected:
+#if WITH_EDITOR
+	virtual void PostEditChangeProperty(FPropertyChangedEvent& PropertyChangedEvent) override;
+#endif
+
+	void OnLineWidthChanged();
+	void OnVectorChanged();
+
+	void UpdateExtent();
+
+	TArray<FVector2D> GenerateLineVertices() const;
+
+	virtual bool IsMeshVisible(int32 MeshIndex) override;
+	virtual void GenerateBorderVertices(TArray<FVector2D>& BorderVertices) override;
+
+	UPROPERTY(EditInstanceOnly, BlueprintReadWrite, Category="Shape", meta=(ClampMin="0.0", AllowPrivateAccess="true"))
+	float LineWidth;
+
+	UPROPERTY(EditInstanceOnly, BlueprintReadWrite, Category="Shape", meta=(AllowPrivateAccess="true"))
+	FVector2D Vector;
+};
diff --git a/Engine/Plugins/Experimental/Avalanche/Source/AvalancheShapes/Public/DynamicMeshes/AvaShapeNGonDynMesh.h b/Engine/Plugins/Experimental/Avalanche/Source/AvalancheShapes/Public/DynamicMeshes/AvaShapeNGonDynMesh.h
new file mode 100644
--- /dev/null
+++ b/Engine/Plugins/Experimental/Avalanche/Source/AvalancheShapes/Public/DynamicMeshes/AvaShapeNGonDynMesh.h
@@ -0,0 +1,49 @@
+// Copyright Epic Games, Inc. All Rights Reserved.
+
+#pragma once
+
+#include "AvaShapeRoundedPolygonDynMesh.h"
+#include "AvaShapeNGonDynMesh.generated.h"
+
+UCLASS(ClassGroup="Shape", BlueprintType, CustomConstructor, Within=AvaShapeActor)
+class AVALANCHESHAPES_API UAvaShapeNGonDynamicMesh : public UAvaShapeRoundedPolygonDynamicMesh
+{
+	GENERATED_BODY()
+
+	friend class FAvaShapeNGonDynamicMeshVisualizer;
+
+public:
+	static inline constexpr float MinNumSides = 3.f;
+	static inline constexpr float MaxNumSides = 128.f;
+
+	UAvaShapeNGonDynamicMesh()
+		: UAvaShapeNGonDynamicMesh(FVector2D(50.f, 50.f))
+	{
+	}
+
+	UAvaShapeNGonDynamicMesh(const FVector2D& Size2D, const FLinearColor& InVertexColor = FLinearColor::White,
+		float InBevelSize = 0.f, uint8 InBevelSubdivisions = 0, uint8 InNumSides = 5)
+		: UAvaShapeRoundedPolygonDynamicMesh(Size2D, InVertexColor, InBevelSize, InBevelSubdivisions)
+		, NumSides(InNumSides)
+	{
+	}
+
+	static const FString MeshName;
+	virtual const FString& GetMeshName() const override { return MeshName; }
+
+	UFUNCTION()
+	bool SetNumSides(uint8 InNumSides);
+	uint8 GetNumSides() const { return NumSides; }
+
+protected:
+#if WITH_EDITOR
+	virtual void PostEditChangeProperty(FPropertyChangedEvent& PropertyChangedEvent) override;
+#endif
+
+	virtual void OnNumSidesChanged();
+
+	virtual void GenerateBorderVertices(TArray<FVector2D>& BorderVertices) override;
+
+	UPROPERTY(EditInstanceOnly, BlueprintReadWrite, Category="Shape", meta=(ClampMin="3.0", ClampMax="128.0", DisplayName="Sides", AllowPrivateAccess="true"))
+	uint8 NumSides;
+};
diff --git a/Engine/Plugins/Experimental/Avalanche/Source/AvalancheShapes/Public/DynamicMeshes/AvaShapeRectangleDynMesh.h b/Engine/Plugins/Experimental/Avalanche/Source/AvalancheShapes/Public/DynamicMeshes/AvaShapeRectangleDynMesh.h
new file mode 100644
--- /dev/null
+++ b/Engine/Plugins/Experimental/Avalanche/Source/AvalancheShapes/Public/DynamicMeshes/AvaShapeRectangleDynMesh.h
@@ -0,0 +1,206 @@
+// Copyright Epic Games, Inc. All Rights Reserved.
+
+#pragma once
+
+#include "DynamicMeshes/AvaShape2DDynMeshBase.h"
+#include "AvaShapeRectangleDynMesh.generated.h"
+
+USTRUCT(BlueprintType)
+struct AVALANCHESHAPES_API FAvaShapeRectangleCornerSettings
+{
+	GENERATED_BODY()
+
+	bool IsBeveled() const
+	{
+		return Type != EAvaShapeCornerType::Point && BevelSize > 0.f;
+	}
+
+	UPROPERTY(EditInstanceOnly, BlueprintReadWrite, Category="Shape")
+	EAvaShapeCornerType Type = EAvaShapeCornerType::Point;
+
+	UPROPERTY(EditInstanceOnly, BlueprintReadWrite, Category="Shape", meta=(ClampMin="0.0",DisplayName="Size"))
+	float BevelSize = 0.f;
+
+	UPROPERTY(EditInstanceOnly, BlueprintReadWrite, Category="Shape", meta=(ClampMax="64.0", DisplayName="Subdivisions"))
+	uint8 BevelSubdivisions = UAvaShapeDynamicMeshBase::DefaultSubdivisions;
+
+	UPROPERTY()
+	FVector2D CornerPositionCache = FVector2D::ZeroVector;
+};
+
+UCLASS(ClassGroup="Shape", BlueprintType, CustomConstructor, Within=AvaShapeActor)
+class AVALANCHESHAPES_API UAvaShapeRectangleDynamicMesh : public UAvaShape2DDynMeshBase
+{
+	GENERATED_BODY()
+
+	friend class FAvaShapeRectangleDynamicMeshVisualizer;
+
+public:
+	static inline constexpr float MinSlantAngle = -45.f;
+	static inline constexpr float MaxSlantAngle = 45.f;
+	static inline constexpr float CornerMinMargin = 0.1f;
+
+	UAvaShapeRectangleDynamicMesh()
+		: UAvaShapeRectangleDynamicMesh(FVector2D(50.f, 50.f))
+	{
+	}
+
+    explicit UAvaShapeRectangleDynamicMesh(const FVector2D& Size2D, const FLinearColor& InVertexColor = FLinearColor::White)
+		: UAvaShape2DDynMeshBase(Size2D, InVertexColor)
+	{
+	}
+
+	static const FString MeshName;
+	virtual const FString& GetMeshName() const override { return MeshName; }
+
+	UFUNCTION()
+	bool SetHorizontalAlignment(EAvaHorizontalAlignment InHorizontalAlignment);
+	EAvaHorizontalAlignment GetHorizontalAlignment() const { return HorizontalAlignment; }
+
+	UFUNCTION()
+	bool SetVerticalAlignment(EAvaVerticalAlignment InVerticalAlignment);
+	EAvaVerticalAlignment GetVerticalAlignment() const { return VerticalAlignment; }
+
+	UFUNCTION()
+	bool SetLeftSlant(float InSlant);
+	float GetLeftSlant() const { return LeftSlant; }
+
+	UFUNCTION()
+	bool SetRightSlant(float InSlant);
+	float GetRightSlant() const { return RightSlant; }
+
+	UFUNCTION()
+	bool SetGlobalBevelSize(float InBevelSize);
+	float GetGlobalBevelSize() const { return GlobalBevelSize; }
+
+	/** Controls all corners bevel subdivisions */
+	UFUNCTION()
+	void SetGlobalBevelSubdivisions(uint8 InGlobalBevelSubdivisions);
+	uint8 GetGlobalBevelSubdivisions() const { return GlobalBevelSubdivisions; }
+
+	UFUNCTION()
+	void SetTopLeft(const FAvaShapeRectangleCornerSettings& InCornerSettings);
+
+	UFUNCTION()
+	void SetTopRight(const FAvaShapeRectangleCornerSettings& InCornerSettings);
+
+	UFUNCTION()
+	void SetBottomLeft(const FAvaShapeRectangleCornerSettings& InCornerSettings);
+
+	UFUNCTION()
+	void SetBottomRight(const FAvaShapeRectangleCornerSettings& InCornerSettings);
+
+	EAvaShapeCornerType GetTopLeftCornerType() const { return TopLeft.Type; }
+	bool SetTopLeftCornerType(EAvaShapeCornerType InType);
+
+	float GetTopLeftBevelSize() const { return TopLeft.BevelSize; }
+	bool SetTopLeftBevelSize(float InSize);
+
+	/** Bevel subdivisions for top left corner */
+	uint8 GetTopLeftBevelSubdivisions() const { return TopLeft.BevelSubdivisions; }
+	void SetTopLeftBevelSubdivisions(uint8 InBevelSubdivisions);
+
+	EAvaShapeCornerType GetBottomLeftCornerType() const { return BottomLeft.Type; }
+	bool SetBottomLeftCornerType(EAvaShapeCornerType InType);
+
+	float GetBottomLeftBevelSize() const { return BottomLeft.BevelSize; }
+	bool SetBottomLeftBevelSize(float InSize);
+
+	/** Bevel subdivisions for bottom left corner */
+	uint8 GetBottomLeftBevelSubdivisions() const { return BottomLeft.BevelSubdivisions; }
+	void SetBottomLeftBevelSubdivisions(uint8 InBevelSubdivisions);
+
+	EAvaShapeCornerType GetTopRightCornerType() const { return TopRight.Type; }
+	bool SetTopRightCornerType(EAvaShapeCornerType InType);
+
+	float GetTopRightBevelSize() const { return TopRight.BevelSize; }
+	bool SetTopRightBevelSize(float InSize);
+
+	/** Bevel subdivisions for top right corner */
+	uint8 GetTopRightBevelSubdivisions() const { return TopRight.BevelSubdivisions; }
+	void SetTopRightBevelSubdivisions(uint8 InBevelSubdivisions);
+
+	EAvaShapeCornerType GetBottomRightCornerType() const { return BottomRight.Type; }
+	bool SetBottomRightCornerType(EAvaShapeCornerType InType);
+
+	float GetBottomRightBevelSize() const { return BottomRight.BevelSize; }
+	bool SetBottomRightBevelSize(float InSize);
+
+	/** Bevel subdivisions for bottom right corner */
+	uint8 GetBottomRightBevelSubdivisions() const { return BottomRight.BevelSubdivisions; }
+	void SetBottomRightBevelSubdivisions(uint8 InBevelSubdivisions);
+
+protected:
+#if WITH_EDITOR
+	virtual void PostEditChangeProperty(FPropertyChangedEvent& PropertyChangedEvent) override;
+#endif
+
+	float GetMaximumBevelSize() const;
+
+	bool IsSlantAngleValid() const;
+	FVector2D GetValidRangeLeftSlantAngle() const;
+	FVector2D GetValidRangeRightSlantAngle() const;
+	void GetValidSlantAngle(float& OutLeftSlant, float& OutRightSlant) const;
+
+	void OnAlignmentChanged();
+
+	void OnLeftSlantChanged();
+	void OnRightSlantChanged();
+
+	void OnGlobalBevelSizeChanged();
+	void OnGlobalBevelSubdivisionsChanged();
+
+	void OnTopLeftCornerTypeChanged();
+	void OnTopLeftBevelSizeChanged();
+	void OnTopLeftBevelSubdivisionsChanged();
+
+	void OnBottomLeftCornerTypeChanged();
+	void OnBottomLeftBevelSizeChanged();
+	void OnBottomLeftBevelSubdivisionsChanged();
+
+	void OnTopRightCornerTypeChanged();
+	void OnTopRightBevelSizeChanged();
+	void OnTopRightBevelSubdivisionsChanged();
+
+	void OnBottomRightCornerTypeChanged();
+	void OnBottomRightBevelSizeChanged();
+	void OnBottomRightBevelSubdivisionsChanged();
+
+	bool GenerateBaseMeshSections(FAvaShapeMesh& BaseMesh);
+
+	virtual bool CreateMesh(FAvaShapeMesh& InMesh) override;
+
+	virtual void OnSizeChanged() override;
+
+	UPROPERTY(EditInstanceOnly, BlueprintReadWrite, Category="Shape", meta=(AllowPrivateAccess="true"))
+	EAvaHorizontalAlignment HorizontalAlignment = EAvaHorizontalAlignment::Center;
+
+	UPROPERTY(EditInstanceOnly, BlueprintReadWrite, Category="Shape", meta=(AllowPrivateAccess="true"))
+	EAvaVerticalAlignment VerticalAlignment = EAvaVerticalAlignment::Center;
+
+	/** Angle in degrees for the left slant of the rectangle */
+	UPROPERTY(EditInstanceOnly, BlueprintReadWrite, Category="Shape", meta=(ClampMin="-45.0",ClampMax="45.0", AllowPrivateAccess="true"))
+	float LeftSlant = 0.f;
+
+	/** Angle in degrees for the right slant of the rectangle */
+	UPROPERTY(EditInstanceOnly, BlueprintReadWrite, Category="Shape", meta=(ClampMin="-45.0",ClampMax="45.0", AllowPrivateAccess="true"))
+	float RightSlant = 0.f;
+
+	UPROPERTY(EditInstanceOnly, BlueprintReadWrite, Category="Shape", meta=(ClampMin="0.0", AllowPrivateAccess="true"))
+	float GlobalBevelSize = 0.f;
+
+	UPROPERTY(EditInstanceOnly, BlueprintReadWrite, Category="Shape", meta=(ClampMax="128.0", AllowPrivateAccess="true"))
+	uint8 GlobalBevelSubdivisions = DefaultSubdivisions;
+
+	UPROPERTY(EditInstanceOnly, BlueprintReadWrite, Category="Shape", meta=(AllowPrivateAccess="true"))
+	FAvaShapeRectangleCornerSettings TopLeft;
+
+	UPROPERTY(EditInstanceOnly, BlueprintReadWrite, Category="Shape", meta=(AllowPrivateAccess="true"))
+	FAvaShapeRectangleCornerSettings TopRight;
+
+	UPROPERTY(EditInstanceOnly, BlueprintReadWrite, Category="Shape", meta=(AllowPrivateAccess="true"))
+	FAvaShapeRectangleCornerSettings BottomLeft;
+
+	UPROPERTY(EditInstanceOnly, BlueprintReadWrite, Category="Shape", meta=(AllowPrivateAccess="true"))
+	FAvaShapeRectangleCornerSettings BottomRight;
+};
diff --git a/Engine/Plugins/Experimental/Avalanche/Source/AvalancheShapes/Public/DynamicMeshes/AvaShapeRingDynMesh.h b/Engine/Plugins/Experimental/Avalanche/Source/AvalancheShapes/Public/DynamicMeshes/AvaShapeRingDynMesh.h
new file mode 100644
--- /dev/null
+++ b/Engine/Plugins/Experimental/Avalanche/Source/AvalancheShapes/Public/DynamicMeshes/AvaShapeRingDynMesh.h
@@ -0,0 +1,80 @@
+// Copyright Epic Games, Inc. All Rights Reserved.
+
+#pragma once
+
+#include "DynamicMeshes/AvaShape2DDynMeshBase.h"
+#include "AvaShapeRingDynMesh.generated.h"
+
+UCLASS(ClassGroup="Shape", BlueprintType, CustomConstructor, Within=AvaShapeActor)
+class AVALANCHESHAPES_API UAvaShapeRingDynamicMesh : public UAvaShape2DDynMeshBase
+{
+	GENERATED_BODY()
+
+	friend class FAvaShapeRingDynamicMeshVisualizer;
+
+public:
+	static inline constexpr uint8 MinNumSides = 3;
+	static inline constexpr uint8 MaxNumSides = 128;
+
+	UAvaShapeRingDynamicMesh()
+		: UAvaShapeRingDynamicMesh(FVector2D(50.f, 50.f))
+	{
+	}
+
+	UAvaShapeRingDynamicMesh(const FVector2D& Size2D, const FLinearColor& InVertexColor = FLinearColor::White,
+		uint8 InNumSides = 64, float InInnerSize = 0.8f, float InAngleDegree = 360.f, float InStartDegree = 0.f)
+		: UAvaShape2DDynMeshBase(Size2D, InVertexColor)
+		, NumSides(InNumSides)
+		, InnerSize(InInnerSize)
+		, AngleDegree(InAngleDegree)
+		, StartDegree (InStartDegree)
+	{
+		bDoNotRecenterVertices = true;
+	}
+
+	static const FString MeshName;
+	virtual const FString& GetMeshName() const override { return MeshName; }
+
+	UFUNCTION()
+	bool SetNumSides(uint8 InNumSides);
+	uint8 GetNumSides() const { return NumSides; }
+
+	UFUNCTION()
+	bool SetInnerSize(float InInnerSize);
+	float GetInnerSize() const { return InnerSize; }
+
+	UFUNCTION()
+	bool SetAngleDegree(float InDegree);
+	float GetAngleDegree() const { return AngleDegree; }
+
+	UFUNCTION()
+	bool SetStartDegree(float InDegree);
+	float GetStartDegree() const { return StartDegree; }
+
+protected:
+#if WITH_EDITOR
+	virtual void PostEditChangeProperty(FPropertyChangedEvent& PropertyChangedEvent) override;
+#endif
+
+	virtual void OnNumSidesChanged();
+	virtual void OnInnerSizeChanged();
+	virtual void OnAngleDegreeChanged();
+	virtual void OnStartDegreeChanged();
+
+	virtual bool IsMeshVisible(int32 MeshIndex) override;
+	virtual bool CreateMesh(FAvaShapeMesh& InMesh) override;
+
+	UPROPERTY(EditInstanceOnly, BlueprintReadWrite, Category="Shape", meta=(ClampMin="3.0", ClampMax="128.0", DisplayName="Sides", AllowPrivateAccess="true"))
+	uint8 NumSides;
+
+	UPROPERTY(EditInstanceOnly, BlueprintReadWrite, Category="Shape", meta=(ClampMin="0.01", ClampMax="0.99",DisplayName="Inner Size", AllowPrivateAccess="true"))
+	float InnerSize;
+
+	// represents the angle in degree for the ring
+	UPROPERTY(EditInstanceOnly, BlueprintReadWrite, Category="Shape", meta=(ClampMin="0.0", ClampMax="360.0", DisplayName="Angle degree", AllowPrivateAccess="true"))
+	float AngleDegree;
+
+	// represents the starting angle in degree for the ring
+	UPROPERTY(EditInstanceOnly, BlueprintReadWrite, Category="Shape", meta=(ClampMin="0.0", ClampMax="360.0", DisplayName="Start Degree", AllowPrivateAccess="true"))
+	float StartDegree;
+};
diff --git a/Engine/Plugins/Experimental/Avalanche/Source/AvalancheShapes/Public/DynamicMeshes/AvaShapeSphereDynMesh.h b/Engine/Plugins/Experimental/Avalanche/Source/AvalancheShapes/Public/DynamicMeshes/AvaShapeSphereDynMesh.h
new file mode 100644
--- /dev/null
+++ b/Engine/Plugins/Experimental/Avalanche/Source/AvalancheShapes/Public/DynamicMeshes/AvaShapeSphereDynMesh.h
@@ -0,0 +1,133 @@
+// Copyright Epic Games, Inc. All Rights Reserved.
+
+#pragma once
+
+#include "DynamicMeshes/AvaShape3DDynMeshBase.h"
+#include "AvaShapeSphereDynMesh.generated.h"
+
+UCLASS(ClassGroup="Shape", BlueprintType, CustomConstructor, Within=AvaShapeActor)
+class AVALANCHESHAPES_API UAvaShapeSphereDynamicMesh : public UAvaShape3DDynMeshBase
+{
+	GENERATED_BODY()
+
+	friend class FAvaShapeSphereDynamicMeshVisualizer;
+
+public:
+	static constexpr uint8 MinNumSides = 4;
+	static constexpr uint8 MaxNumSides = 128;
+
+	static constexpr int32 MESH_INDEX_TOP = 1;
+	static constexpr int32 MESH_INDEX_BOTTOM = 2;
+	static constexpr int32 MESH_INDEX_START = 3;
+	static constexpr int32 MESH_INDEX_END = 4;
+
+	UAvaShapeSphereDynamicMesh()
+		: UAvaShapeSphereDynamicMesh(FVector(50.f, 50.f, 50.f))
+	{ }
+
+	UAvaShapeSphereDynamicMesh(
+		const FVector& InSize,
+		const FLinearColor& InVertexColor = FLinearColor::White,
+		float InStartLongitude = 0.f,
+		float InEndLongitude = 180.f,
+		float InStartLatitude = 0.f,
+		float InLatitudeDegree = 360.f,
+		uint8 InNumSides = 32,
+		float InRadius = 1.f)
+		: UAvaShape3DDynMeshBase(InSize, InVertexColor)
+		, StartLongitude(InStartLongitude)
+		, PreEditStartLongitude(0)
+		, EndLongitude(InEndLongitude)
+		, PreEditEndLongitude(0)
+		, StartLatitude(InStartLatitude)
+		, LatitudeDegree(InLatitudeDegree)
+		, NumSides(InNumSides)
+		, Radius(InRadius)
+	{
+	}
+
+	static const FString MeshName;
+	virtual const FString& GetMeshName() const override { return MeshName; }
+
+	UFUNCTION()
+	bool SetNumSides(uint8 InNumSides);
+	uint8 GetNumSides() const { return NumSides; }
+
+	UFUNCTION()
+	float GetRadius() const { return Radius; }
+
+	UFUNCTION()
+	bool SetStartLatitude(float InDegree);
+	float GetStartLatitude() const { return StartLatitude; }
+
+	UFUNCTION()
+	bool SetLatitudeDegree(float InDegree);
+	float GetLatitudeDegree() const { return LatitudeDegree; }
+
+	UFUNCTION()
+	bool SetStartLongitude(float InDegree);
+	float GetStartLongitude() const { return StartLongitude; }
+
+	UFUNCTION()
+	bool SetEndLongitude(float InDegree);
+	float GetEndLongitude() const { return EndLongitude; }
+
+protected:
+#if WITH_EDITOR
+	virtual void PreEditChange(FProperty* PropertyAboutToChange) override;
+	virtual void PostEditChangeProperty(FPropertyChangedEvent& PropertyChangedEvent) override;
+#endif
+
+	virtual void OnNumSidesChanged();
+	virtual void OnRadiusChanged();
+	virtual void OnStartLatitudeChanged();
+	virtual void OnLatitudeDegreeChanged();
+	virtual void OnStartLongitudeChanged();
+	virtual void OnEndLongitudeChanged();
+
+	virtual bool CreateBaseUVs(FAvaShapeMesh& BaseMesh, FAvaShapeMaterialUVParameters& InUseParams);
+	virtual bool GenerateBaseMeshSections(FAvaShapeMesh& BaseMesh);
+
+	virtual bool CreateTopUVs(FAvaShapeMesh& TopMesh, FAvaShapeMaterialUVParameters& InUseParams);
+	virtual bool GenerateTopMeshSections(FAvaShapeMesh& TopMesh);
+
+	virtual bool CreateBottomUVs(FAvaShapeMesh& BottomMesh, FAvaShapeMaterialUVParameters& InUseParams);
+	virtual bool GenerateBottomMeshSections(FAvaShapeMesh& BottomMesh);
+
+	virtual bool CreateStartUVs(FAvaShapeMesh& StartMesh, FAvaShapeMaterialUVParameters& InUseParams);
+	virtual bool GenerateStartMeshSections(FAvaShapeMesh& StartMesh);
+
+	virtual bool CreateEndUVs(FAvaShapeMesh& EndMesh, FAvaShapeMaterialUVParameters& InUseParams);
+	virtual bool GenerateEndMeshSections(FAvaShapeMesh& EndMesh);
+
+	virtual void RegisterMeshes() override;
+	virtual bool IsMeshVisible(int32 MeshIndex) override;
+	virtual bool CreateMesh(FAvaShapeMesh& InMesh) override;
+	virtual bool CreateUVs(FAvaShapeMesh& InMesh, FAvaShapeMaterialUVParameters& InParams) override;
+
+	// represents the longitude (Z) angle in degree for the sphere at the start
+	UPROPERTY(EditInstanceOnly, BlueprintReadWrite, Category="Shape", meta=(ClampMin="0.0", ClampMax="180.0", DisplayName="Start Longitude degree", AllowPrivateAccess="true"))
+	float StartLongitude;
+	float PreEditStartLongitude;
+
+	// represents the longitude (Z) angle in degree for the sphere at the end
+	UPROPERTY(EditInstanceOnly, BlueprintReadWrite, Category="Shape", meta=(ClampMin="0.0", ClampMax="180.0", DisplayName="End Longitude degree", AllowPrivateAccess="true"))
+	float EndLongitude;
+	float PreEditEndLongitude;
+
+	// represents the latitude (Y) angle in degree for the sphere at the start
+	UPROPERTY(EditInstanceOnly, BlueprintReadWrite, Category="Shape", meta=(ClampMin="0.0", ClampMax="360.0", DisplayName="Start Latitude degree", AllowPrivateAccess="true"))
+	float StartLatitude;
+
+	// represents the total latitude (Y) angle in degree for the sphere
+	UPROPERTY(EditInstanceOnly, BlueprintReadWrite, Category="Shape", meta=(ClampMin="0.0", ClampMax="360.0", DisplayName="Latitude degree", AllowPrivateAccess="true"))
+	float LatitudeDegree;
+
+	// represents the precision of the sphere mesh
+	UPROPERTY(EditInstanceOnly, BlueprintReadWrite, Category="Shape", meta=(ClampMin="4.0", ClampMax="128.0", DisplayName="Sides", AllowPrivateAccess="true"))
+	uint8 NumSides;
+
+	// represents the radius ratio of the sphere
+	UPROPERTY()
+	float Radius;
+};
diff --git a/Engine/Plugins/Experimental/Avalanche/Source/AvalancheShapes/Public/DynamicMeshes/AvaShapeStarDynMesh.h b/Engine/Plugins/Experimental/Avalanche/Source/AvalancheShapes/Public/DynamicMeshes/AvaShapeStarDynMesh.h
new file mode 100644
--- /dev/null
+++ b/Engine/Plugins/Experimental/Avalanche/Source/AvalancheShapes/Public/DynamicMeshes/AvaShapeStarDynMesh.h
@@ -0,0 +1,61 @@
+// Copyright Epic Games, Inc. All Rights Reserved.
+
+#pragma once
+
+#include "AvaShapeRoundedPolygonDynMesh.h"
+#include "AvaShapeStarDynMesh.generated.h"
+
+UCLASS(ClassGroup="Shape", BlueprintType, CustomConstructor, Within=AvaShapeActor)
+class AVALANCHESHAPES_API UAvaShapeStarDynamicMesh : public UAvaShapeRoundedPolygonDynamicMesh
+{
+	GENERATED_BODY()
+
+	friend class FAvaShapeStarDynamicMeshVisualizer;
+
+public:
+	static inline constexpr uint8 MinNumPoints = 2;
+	static inline constexpr uint8 MaxNumPoints = 128;
+
+	UAvaShapeStarDynamicMesh()
+		: UAvaShapeStarDynamicMesh(FVector2D(50.f, 50.f))
+	{
+	}
+
+	UAvaShapeStarDynamicMesh(const FVector2D& Size2D, const FLinearColor& InVertexColor = FLinearColor::White,
+		float InBevelSize = 0.f, uint8 InBevelSubdivisions = 0, uint8 InNumPoints = 5, float InInnerSize = 0.5f)
+		: UAvaShapeRoundedPolygonDynamicMesh(Size2D, InVertexColor, InBevelSize, InBevelSubdivisions)
+		, NumPoints(InNumPoints)
+		, InnerSize(InInnerSize)
+	{
+	}
+
+	static const FString MeshName;
+	virtual const FString& GetMeshName() const override { return MeshName; }
+
+	UFUNCTION()
+	bool SetNumPoints(uint8 InNumPoints);
+	uint8 GetNumPoints() const { return NumPoints; }
+
+	UFUNCTION()
+	bool SetInnerSize(float InInnerSize);
+	float GetInnerSize() const { return InnerSize; }
+
+protected:
+#if WITH_EDITOR
+	virtual void PostEditChangeProperty(FPropertyChangedEvent& PropertyChangedEvent) override;
+#endif
+
+	virtual void OnNumSidesChanged();
+	virtual void OnInnerSizeChanged();
+
+	virtual void GenerateBorderVertices(TArray<FVector2D>& BorderVertices) override;
+
+	virtual bool UseCenteredVertex() override { return true; }
+	virtual bool IsMeshVisible(int32 MeshIndex) override;
+
+	UPROPERTY(EditInstanceOnly, BlueprintReadWrite, Category="Shape", meta=(ClampMin="2.0", ClampMax="128.0", DisplayName="Points", AllowPrivateAccess="true"))
+	uint8 NumPoints = 0;
+
+	UPROPERTY(EditInstanceOnly, BlueprintReadWrite, Category="Shape", meta=(ClampMin="0.0", ClampMax="0.99", DisplayName="Inner Size Fraction", AllowPrivateAccess="true"))
+	float InnerSize = 0.f;
+};
diff --git a/Engine/Plugins/Experimental/Avalanche/Source/AvalancheShapes/Public/DynamicMeshes/AvaShapeTorusDynMesh.h b/Engine/Plugins/Experimental/Avalanche/Source/AvalancheShapes/Public/DynamicMeshes/AvaShapeTorusDynMesh.h
new file mode 100644
--- /dev/null
+++ b/Engine/Plugins/Experimental/Avalanche/Source/AvalancheShapes/Public/DynamicMeshes/AvaShapeTorusDynMesh.h
@@ -0,0 +1,113 @@
+// Copyright Epic Games, Inc. All Rights Reserved.
+
+#pragma once
+
+#include "AvaShape3DDynMeshBase.h"
+#include "AvaShapeTorusDynMesh.generated.h"
+
+UCLASS(ClassGroup="Shape", BlueprintType, CustomConstructor, Within=AvaShapeActor)
+class AVALANCHESHAPES_API UAvaShapeTorusDynamicMesh : public UAvaShape3DDynMeshBase
+{
+	GENERATED_BODY()
+
+	friend class FAvaShapeTorusDynamicMeshVisualizer;
+
+public:
+	static constexpr uint8 MinNumSlices = 3;
+	static constexpr uint8 MaxNumSlices = 128;
+	static constexpr uint8 MinNumSides = 4;
+	static constexpr uint8 MaxNumSides = 128;
+
+	static constexpr int32 MESH_INDEX_START = 1;
+	static constexpr int32 MESH_INDEX_END = 2;
+
+	UAvaShapeTorusDynamicMesh()
+		: UAvaShapeTorusDynamicMesh(FVector(50.f, 50.f, 50.f))
+	{
+	}
+
+	UAvaShapeTorusDynamicMesh(
+		const FVector& InSize,
+		const FLinearColor& InVertexColor = FLinearColor::White,
+		uint8 InNumSlices = 32,
+		uint8 InNumSides = 32,
+		float InInnerSize = 0.75,
+		float InAngleDegree = 360.f,
+		float InStartDegree = 0.f)
+		: UAvaShape3DDynMeshBase(InSize, InVertexColor)
+		, NumSlices(InNumSlices)
+		, NumSides(InNumSides)
+		, InnerSize(InInnerSize)
+		, AngleDegree(InAngleDegree)
+		, StartDegree(InStartDegree)
+	{
+	}
+
+	static const FString MeshName;
+	virtual const FString& GetMeshName() const override { return MeshName; }
+
+	UFUNCTION()
+	bool SetNumSides(uint8 InNumSides);
+	uint8 GetNumSides() const { return NumSides; }
+
+	UFUNCTION()
+	bool SetNumSlices(uint8 InNumSlices);
+	uint8 GetNumSlices() const { return NumSlices; }
+
+	UFUNCTION()
+	bool SetInnerSize(float InInnerSize);
+	float GetInnerSize() const { return InnerSize; }
+
+	UFUNCTION()
+	bool SetAngleDegree(float InAngleDegree);
+	float GetAngleDegree() const { return AngleDegree; }
+
+	UFUNCTION()
+	bool SetStartDegree(float InStartDegree);
+	float GetStartDegree() const { return StartDegree; }
+
+protected:
+#if WITH_EDITOR
+	virtual void PostEditChangeProperty(FPropertyChangedEvent& PropertyChangedEvent) override;
+#endif
+
+	virtual bool CreateBaseUVs(FAvaShapeMesh& BaseMesh, FAvaShapeMaterialUVParameters& InUseParams);
+	virtual bool GenerateBaseMeshSections(FAvaShapeMesh& BaseMesh);
+
+	virtual bool CreateStartUVs(FAvaShapeMesh& StartMesh, FAvaShapeMaterialUVParameters& InUseParams);
+	virtual bool GenerateStartMeshSections(FAvaShapeMesh& StartMesh);
+
+	virtual bool CreateEndUVs(FAvaShapeMesh& EndMesh, FAvaShapeMaterialUVParameters& InUseParams);
+	virtual bool GenerateEndMeshSections(FAvaShapeMesh& EndMesh);
+
+	virtual void OnNumSidesChanged();
+	virtual void OnNumSlicesChanged();
+	virtual void OnInnerSizeChanged();
+	virtual void OnAngleDegreeChanged();
+	virtual void OnStartDegreeChanged();
+
+	virtual void RegisterMeshes() override;
+	virtual bool IsMeshVisible(int32 MeshIndex) override;
+	virtual bool CreateMesh(FAvaShapeMesh& InMesh) override;
+	virtual bool CreateUVs(FAvaShapeMesh& InMesh, FAvaShapeMaterialUVParameters& InParams) override;
+
+	// represents the number of slices composing the tube
+	UPROPERTY(EditInstanceOnly, BlueprintReadWrite, Category="Shape", meta=(ClampMin="3.0", ClampMax="128.0", DisplayName="Slices", AllowPrivateAccess="true"))
+	uint8 NumSlices;
+
+	// represents the precision of each circle composing a slice
+	UPROPERTY(EditInstanceOnly, BlueprintReadWrite, Category="Shape", meta=(ClampMin="4.0", ClampMax="128.0", DisplayName="Sides", AllowPrivateAccess="true"))
+	uint8 NumSides;
+
+	// represents the size ratio available inside the torus
+	UPROPERTY(EditInstanceOnly, BlueprintReadWrite, Category="Shape", meta=(ClampMin="0.5", ClampMax="0.99", DisplayName="Inner Size", AllowPrivateAccess="true"))
+	float InnerSize;
+
+	// represents the tube angle in degree for the torus
+	UPROPERTY(EditInstanceOnly, BlueprintReadWrite, Category="Shape", meta=(ClampMin="0.0", ClampMax="360.0", DisplayName="Angle degree", AllowPrivateAccess="true"))
+	float AngleDegree;
+
+	// represents the starting angle in degree for the torus
+	UPROPERTY(EditInstanceOnly, BlueprintReadWrite, Category="Shape", meta=(ClampMin="0.0", ClampMax="360.0", DisplayName="Start Degree", AllowPrivateAccess="true"))
+	float StartDegree;
+};
diff --git a/Engine/Plugins/Experimental/Avalanche/Source/AvalancheShapesEditor/Private/AvaShapesEditorStyle.cpp b/Engine/Plugins/Experimental/Avalanche/Source/AvalancheShapesEditor/Private/AvaShapesEditorStyle.cpp
new file mode 100644
--- /dev/null
+++ b/Engine/Plugins/Experimental/Avalanche/Source/AvalancheShapesEditor/Private/AvaShapesEditorStyle.cpp
@@ -0,0 +1,120 @@
+// Copyright Epic Games, Inc. All Rights Reserved.
+
+#include "AvaShapesEditorStyle.h"
+#include "Interfaces/IPluginManager.h"
+#include "Styling/SlateStyleRegistry.h"
+
+#define LOCTEXT_NAMESPACE "AvaShapesEditorStyle"
+
+TSharedPtr<FSlateStyleSet> FAvaShapesEditorStyle::StyleInstance;
+
+void FAvaShapesEditorStyle::Initialize()
+{
+	StyleInstance = Create();
+	FSlateStyleRegistry::RegisterSlateStyle(Get());
+}
+
+void FAvaShapesEditorStyle::Shutdown()
+{
+	FSlateStyleRegistry::UnRegisterSlateStyle(Get());
+	StyleInstance.Reset();
+}
+
+FName FAvaShapesEditorStyle::GetStyleSetName()
+{
+	static FName StyleSetName(TEXT("AvalancheShapesEditor"));
+	return StyleSetName;
+}
+
+const ISlateStyle& FAvaShapesEditorStyle::Get()
+{
+	if (StyleInstance.IsValid() == false)
+	{
+		Initialize();
+	}
+
+	return *StyleInstance;
+}
+
+const FSlateBrush* FAvaShapesEditorStyle::GetBrush(FName PropertyName, const ANSICHAR* Specifier)
+{
+	return StyleInstance->GetBrush(PropertyName, Specifier);
+}
+
+#define IMAGE_BRUSH_SVG(RelativePath, ...) FSlateVectorImageBrush(Style->RootToContentDir(RelativePath, TEXT(".svg")), __VA_ARGS__)
+#define IMAGE_BRUSH(RelativePath, ... ) FSlateImageBrush(Style->RootToContentDir(RelativePath, TEXT(".png")), __VA_ARGS__)
+
+TSharedRef<FSlateStyleSet> FAvaShapesEditorStyle::Create()
+{
+	const FVector2f Icon16x16(16.0f, 16.0f);
+	const FVector2f Icon20x20(20.0f, 20.0f);
+
+	TSharedRef<FSlateStyleSet> Style = MakeShared<FSlateStyleSet>(GetStyleSetName());
+
+	TSharedPtr<IPlugin> Plugin = IPluginManager::Get().FindPlugin(TEXT("Avalanche"));
+	check(Plugin.IsValid());
+	if (Plugin.IsValid())
+	{
+		Style->SetContentRoot(Plugin->GetBaseDir() / TEXT("Resources"));
+	}
+
+	// Path Constants
+	const FString UEContentEditorSlatePath = FPaths::EngineContentDir() / TEXT("Editor/Slate");
+	const FString UEContentEditorSlateIconsPath = UEContentEditorSlatePath / TEXT("Icons");
+
+	// Class Icons
+	Style->Set("ClassIcon.AvaShape2DArrowDynamicMesh", new IMAGE_BRUSH("Icons/ToolboxIcons/arrow", Icon16x16));
+	Style->Set("ClassIcon.AvaShapeChevronDynamicMesh", new IMAGE_BRUSH("Icons/ToolboxIcons/chevron", Icon16x16));
+	Style->Set("ClassIcon.AvaShapeConeDynamicMesh", new IMAGE_BRUSH("Icons/ToolboxIcons/cone", Icon16x16));
+	Style->Set("ClassIcon.AvaShapeCubeDynamicMesh", new IMAGE_BRUSH("Icons/ToolboxIcons/cube", Icon16x16));
+	Style->Set("ClassIcon.AvaShapeRoundedPolygonDynamicMesh", new IMAGE_BRUSH_SVG("Icons/ToolboxIcons/Toolbar_Primitive", Icon16x16));
+	Style->Set("ClassIcon.AvaShapeSphereDynamicMesh", new IMAGE_BRUSH("Icons/ToolboxIcons/sphere", Icon16x16));
+	Style->Set("ClassIcon.AvaShapeTorusDynamicMesh", new IMAGE_BRUSH("Icons/ToolboxIcons/torus", Icon16x16));
+	Style->Set("ClassIcon.AvaShapeStarDynamicMesh", new IMAGE_BRUSH_SVG("Icons/ToolboxIcons/Favorite", Icon16x16));
+	Style->Set("ClassIcon.AvaShapeLineDynamicMesh", new IMAGE_BRUSH_SVG("Icons/ToolboxIcons/minus", Icon16x16));
+	Style->Set("ClassIcon.AvaShapeRectangleDynamicMesh", new IMAGE_BRUSH("Icons/ToolboxIcons/rectangle", Icon16x16));
+	Style->Set("ClassIcon.AvaShapeNGonDynamicMesh", new IMAGE_BRUSH("Icons/ToolboxIcons/regularpolygon", Icon16x16));
+	Style->Set("ClassIcon.AvaShapeIrregularPolygonDynamicMesh", new IMAGE_BRUSH("Icons/ToolboxIcons/irregularpolygon", Icon16x16));
+	Style->Set("ClassIcon.AvaShapeEllipseDynamicMesh", new IMAGE_BRUSH("Icons/ToolboxIcons/circle", Icon16x16));
+	Style->Set("ClassIcon.AvaShapeRingDynamicMesh", new IMAGE_BRUSH("Icons/ToolboxIcons/ring", Icon16x16));
+	Style->Set("ClassIcon.AvaShapeShapeActor", new IMAGE_BRUSH("Icons/ToolboxIcons/regularpolygon", Icon16x16));
+
+	// 2D Tools
+	Style->Set("AvalancheShapesEditor.Tool_Shape_2DArrow",       new IMAGE_BRUSH("Icons/ToolboxIcons/arrow", Icon20x20));
+	Style->Set("AvalancheShapesEditor.Tool_Shape_Chevron",       new IMAGE_BRUSH("Icons/ToolboxIcons/chevron", Icon20x20));
+	Style->Set("AvalancheShapesEditor.Tool_Shape_Ellipse",       new IMAGE_BRUSH("Icons/ToolboxIcons/circle", Icon20x20));
+	Style->Set("AvalancheShapesEditor.Tool_Shape_IrregularPoly", new IMAGE_BRUSH("Icons/ToolboxIcons/irregularpolygon", Icon20x20));
+	Style->Set("AvalancheShapesEditor.Tool_Shape_Line",          new IMAGE_BRUSH_SVG("Icons/ToolboxIcons/minus", Icon20x20));
+	Style->Set("AvalancheShapesEditor.Tool_Shape_NGon",          new IMAGE_BRUSH("Icons/ToolboxIcons/regularpolygon", Icon20x20));
+	Style->Set("AvalancheShapesEditor.Tool_Shape_Rectangle",     new IMAGE_BRUSH("Icons/ToolboxIcons/rectangle", Icon20x20));
+	Style->Set("AvalancheShapesEditor.Tool_Shape_Ring",          new IMAGE_BRUSH("Icons/ToolboxIcons/ring", Icon20x20));
+	Style->Set("AvalancheShapesEditor.Tool_Shape_Star",          new IMAGE_BRUSH_SVG("Icons/ToolboxIcons/Favorite", Icon20x20));
+
+	Style->Set("Tool_Shape_2DArrow",       new IMAGE_BRUSH("Icons/ToolboxIcons/arrow", Icon20x20));
+	Style->Set("Tool_Shape_Chevron",       new IMAGE_BRUSH("Icons/ToolboxIcons/chevron", Icon20x20));
+	Style->Set("Tool_Shape_Ellipse",       new IMAGE_BRUSH("Icons/ToolboxIcons/circle", Icon20x20));
+	Style->Set("Tool_Shape_IrregularPoly", new IMAGE_BRUSH("Icons/ToolboxIcons/irregularpolygon", Icon20x20));
+	Style->Set("Tool_Shape_Line",          new IMAGE_BRUSH_SVG("Icons/ToolboxIcons/minus", Icon20x20));
+	Style->Set("Tool_Shape_NGon",          new IMAGE_BRUSH("Icons/ToolboxIcons/regularpolygon", Icon20x20));
+	Style->Set("Tool_Shape_Rectangle",     new IMAGE_BRUSH("Icons/ToolboxIcons/rectangle", Icon20x20));
+	Style->Set("Tool_Shape_Ring",          new IMAGE_BRUSH("Icons/ToolboxIcons/ring", Icon20x20));
+	Style->Set("Tool_Shape_Star",          new IMAGE_BRUSH_SVG("Icons/ToolboxIcons/Favorite", Icon20x20));
+
+	// 3D Tools
+	Style->Set("AvalancheShapesEditor.Tool_Shape_Cone",   new IMAGE_BRUSH("Icons/ToolboxIcons/cone", Icon20x20));
+	Style->Set("AvalancheShapesEditor.Tool_Shape_Cube",   new IMAGE_BRUSH("Icons/ToolboxIcons/cube", Icon20x20));
+	Style->Set("AvalancheShapesEditor.Tool_Shape_Sphere", new IMAGE_BRUSH("Icons/ToolboxIcons/sphere", Icon20x20));
+	Style->Set("AvalancheShapesEditor.Tool_Shape_Torus",  new IMAGE_BRUSH("Icons/ToolboxIcons/torus", Icon20x20));
+
+	Style->Set("Tool_Shape_Cone",   new IMAGE_BRUSH("Icons/ToolboxIcons/cone", Icon20x20));
+	Style->Set("Tool_Shape_Cube",   new IMAGE_BRUSH("Icons/ToolboxIcons/cube", Icon20x20));
+	Style->Set("Tool_Shape_Sphere", new IMAGE_BRUSH("Icons/ToolboxIcons/sphere", Icon20x20));
+	Style->Set("Tool_Shape_Torus",  new IMAGE_BRUSH("Icons/ToolboxIcons/torus", Icon20x20));
+
+	return Style;
+}
+
+#undef IMAGE_BRUSH_SVG
+#undef IMAGE_BRUSH
+
+#undef LOCTEXT_NAMESPACE
diff --git a/Engine/Plugins/Experimental/Avalanche/Source/AvalancheShapesEditor/Private/AvaShapesEditorStyle.h b/Engine/Plugins/Experimental/Avalanche/Source/AvalancheShapesEditor/Private/AvaShapesEditorStyle.h
new file mode 100644
--- /dev/null
+++ b/Engine/Plugins/Experimental/Avalanche/Source/AvalancheShapesEditor/Private/AvaShapesEditorStyle.h
@@ -0,0 +1,25 @@
+// Copyright Epic Games, Inc. All Rights Reserved.
+
+#pragma once
+
+#include "Styling/ISlateStyle.h"
+#include "Styling/SlateStyle.h"
+
+struct FSlateBrush;
+
+class FAvaShapesEditorStyle
+{
+public:
+	static void Initialize();
+	static void Shutdown();
+
+	static FName GetStyleSetName();
+	static const ISlateStyle& Get();
+
+	static const FSlateBrush* GetBrush(FName PropertyName, const ANSICHAR* Specifier = nullptr);
+
+private:
+	static TSharedPtr<FSlateStyleSet> StyleInstance;
+
+	static TSharedRef<FSlateStyleSet> Create();
+};
diff --git a/Engine/Plugins/Experimental/Avalanche/Source/AvalancheShapesEditor/Private/AvalancheShapesEditorModule.cpp b/Engine/Plugins/Experimental/Avalanche/Source/AvalancheShapesEditor/Private/AvalancheShapesEditorModule.cpp
new file mode 100644
--- /dev/null
+++ b/Engine/Plugins/Experimental/Avalanche/Source/AvalancheShapesEditor/Private/AvalancheShapesEditorModule.cpp
@@ -0,0 +1,227 @@
+// Copyright Epic Games, Inc. All Rights Reserved.
+
+#include "AvalancheShapesEditorModule.h"
+#include "AvaInteractiveToolsDelegates.h"
+#include "AvaShapeActor.h"
+#include "AvaShapesEditorCommands.h"
+#include "AvaShapesEditorStyle.h"
+#include "ComponentVisualizers.h"
+#include "Engine/Texture2D.h"
+#include "IAvalancheComponentVisualizersModule.h"
+#include "IDetailsView.h"
+#include "ISequencerModule.h"
+#include "TrackEditors/AvaShapeRectCornerTrackEditor.h"
+#include "Widgets/AvaViewportColorPickerActorClassRegistry.h"
+
+// Meshes
+#include "DynamicMeshes/AvaShape2DArrowDynMesh.h"
+#include "DynamicMeshes/AvaShapeChevronDynMesh.h"
+#include "DynamicMeshes/AvaShapeConeDynMesh.h"
+#include "DynamicMeshes/AvaShapeCubeDynMesh.h"
+#include "DynamicMeshes/AvaShapeEllipseDynMesh.h"
+#include "DynamicMeshes/AvaShapeIrregularPolygonDynMesh.h"
+#include "DynamicMeshes/AvaShapeLineDynMesh.h"
+#include "DynamicMeshes/AvaShapeNGonDynMesh.h"
+#include "DynamicMeshes/AvaShapeRectangleDynMesh.h"
+#include "DynamicMeshes/AvaShapeRingDynMesh.h"
+#include "DynamicMeshes/AvaShapeSphereDynMesh.h"
+#include "DynamicMeshes/AvaShapeStarDynMesh.h"
+#include "DynamicMeshes/AvaShapeTorusDynMesh.h"
+
+// Tools
+#include "Tools/AvaShapesEditorShapeTool2DArrow.h"
+#include "Tools/AvaShapesEditorShapeToolChevron.h"
+#include "Tools/AvaShapesEditorShapeToolCone.h"
+#include "Tools/AvaShapesEditorShapeToolCube.h"
+#include "Tools/AvaShapesEditorShapeToolEllipse.h"
+#include "Tools/AvaShapesEditorShapeToolIrregularPoly.h"
+#include "Tools/AvaShapesEditorShapeToolLine.h"
+#include "Tools/AvaShapesEditorShapeToolNGon.h"
+#include "Tools/AvaShapesEditorShapeToolRectangle.h"
+#include "Tools/AvaShapesEditorShapeToolRing.h"
+#include "Tools/AvaShapesEditorShapeToolSphere.h"
+#include "Tools/AvaShapesEditorShapeToolStar.h"
+#include "Tools/AvaShapesEditorShapeToolTorus.h"
+
+// Visualizers
+#include "Visualizers/AvaShapeConeDynMeshVis.h"
+#include "Visualizers/AvaShapeCubeDynMeshVis.h"
+#include "Visualizers/AvaShapeEllipseDynMeshVis.h"
+#include "Visualizers/AvaShapeIrregularPolygonDynMeshVis.h"
+#include "Visualizers/AvaShapeLineDynMeshVis.h"
+#include "Visualizers/AvaShapeNGonDynMeshVis.h"
+#include "Visualizers/AvaShapeRectangleDynMeshVis.h"
+#include "Visualizers/AvaShapeRingDynMeshVis.h"
+#include "Visualizers/AvaShapeSphereDynMeshVis.h"
+#include "Visualizers/AvaShapeStarDynMeshVis.h"
+#include "Visualizers/AvaShapeTorusDynMeshVis.h"
+
+#define LOCTEXT_NAMESPACE "AvalancheShapesEditor"
+
+const FName FAvalancheShapesEditorModule::BevelSprite = TEXT("Bevel");
+const FName FAvalancheShapesEditorModule::BreakSideSprite = TEXT("BreakSide");
+const FName FAvalancheShapesEditorModule::ColorSelectionSprite = TEXT("ColorSelection");
+const FName FAvalancheShapesEditorModule::CornerSprite = TEXT("Corner");
+const FName FAvalancheShapesEditorModule::DepthSprite = TEXT("Depth");
+const FName FAvalancheShapesEditorModule::InnerSizeSprite = TEXT("InnerSize");
+const FName FAvalancheShapesEditorModule::LinearGradientSprite = TEXT("LinearGradient");
+const FName FAvalancheShapesEditorModule::NumPointsSprite = TEXT("NumPoints");
+const FName FAvalancheShapesEditorModule::NumSidesSprite = TEXT("NumSides");
+const FName FAvalancheShapesEditorModule::TextMaxHeightSprite = TEXT("TextMaxHeight");
+const FName FAvalancheShapesEditorModule::TextMaxWidthSprite = TEXT("TextMaxWidth");
+const FName FAvalancheShapesEditorModule::TextScaleProportionallySprite = TEXT("TextScaleProportionally");
+const FName FAvalancheShapesEditorModule::SizeSprite = TEXT("Size");
+const FName FAvalancheShapesEditorModule::SlantSprite = TEXT("Slant");
+const FName FAvalancheShapesEditorModule::UVSprite = TEXT("UV");
+
+FAvalancheShapesEditorModule& FAvalancheShapesEditorModule::Get()
+{
+	static const FName ModuleName = TEXT("AvalancheShapesEditor");
+	return FModuleManager::LoadModuleChecked<FAvalancheShapesEditorModule>(ModuleName);
+}
+
+void FAvalancheShapesEditorModule::StartupModule()
+{
+	FAvaShapesEditorStyle::Initialize();
+	FAvaShapesEditorCommands::Register();
+	FAvaInteractiveToolsDelegates::GetRegisterToolsDelegate().AddRaw(this, &FAvalancheShapesEditorModule::RegisterShapeTools);
+	FCoreDelegates::OnPostEngineInit.AddRaw(this, &FAvalancheShapesEditorModule::RegisterVisualizers);
+
+	IAvaComponentVisualizersSettings* AvaVisSettings = IAvalancheComponentVisualizersModule::Get().GetSettings();
+
+	auto RegisterDefaultSprite = [AvaVisSettings](FName InName, const TCHAR* InPath)
+	{
+		UTexture2D* Sprite = LoadObject<UTexture2D>(nullptr, InPath);
+		AvaVisSettings->SetDefaultVisualizerSprite(InName, Sprite);
+	};
+
+	RegisterDefaultSprite(BevelSprite, TEXT("Texture2D'/Avalanche/EditorResources/NewBevelHandle.NewBevelHandle'"));
+	RegisterDefaultSprite(BreakSideSprite, TEXT("Texture2D'/Engine/EditorResources/S_Emitter.S_Emitter'"));
+	RegisterDefaultSprite(ColorSelectionSprite, TEXT("Texture2D'/Engine/EditorResources/S_ReflActorIcon.S_ReflActorIcon'"));
+	RegisterDefaultSprite(CornerSprite, TEXT("Texture2D'/Avalanche/EditorResources/Bevel.Bevel'"));
+	RegisterDefaultSprite(DepthSprite, TEXT("Texture2D'/Engine/EditorResources/S_Terrain.S_Terrain'"));
+	RegisterDefaultSprite(InnerSizeSprite, TEXT("Texture2D'/Engine/EditorResources/S_RadForce.S_RadForce'"));
+	RegisterDefaultSprite(LinearGradientSprite, TEXT("Texture2D'/Avalanche/EditorResources/LinearGradient.LinearGradient'"));
+	RegisterDefaultSprite(NumPointsSprite, TEXT("Texture2D'/Engine/EditorResources/S_Emitter.S_Emitter'"));
+	RegisterDefaultSprite(NumSidesSprite, TEXT("Texture2D'/Engine/EditorResources/S_Emitter.S_Emitter'"));
+	RegisterDefaultSprite(SizeSprite, TEXT("Texture2D'/Avalanche/EditorResources/NewSizeHandle.NewSizeHandle'"));
+	RegisterDefaultSprite(SlantSprite, TEXT("Texture2D'/Avalanche/EditorResources/Slant.Slant'"));
+	RegisterDefaultSprite(TextMaxHeightSprite, TEXT("Texture2D'/Engine/EngineResources/Cursors/SplitterVert.SplitterVert'"));
+	RegisterDefaultSprite(TextMaxWidthSprite, TEXT("Texture2D'/Engine/EngineResources/Cursors/SplitterHorz.SplitterHorz'"));
+	RegisterDefaultSprite(TextScaleProportionallySprite, TEXT("Texture2D'/Engine/EditorResources/S_TextRenderActorIcon.S_TextRenderActorIcon'"));
+	RegisterDefaultSprite(UVSprite, TEXT("Texture2D'/Engine/EditorResources/MatInstActSprite.MatInstActSprite'"));
+
+	ISequencerModule& SequencerModule = FModuleManager::LoadModuleChecked<ISequencerModule>("Sequencer");
+	TrackEditorHandle = SequencerModule.RegisterTrackEditor(FOnCreateTrackEditor::CreateStatic(&FAvaShapeRectCornerTrackEditor::CreateTrackEditor));
+
+	FAvaViewportColorPickerActorClassRegistry::RegisterClassAdapter(AAvaShapeActor::StaticClass(), MakeShared<TAvaViewportColorPickerActorAdapter<AAvaShapeActor>>());
+}
+
+void FAvalancheShapesEditorModule::ShutdownModule()
+{
+	FAvaShapesEditorStyle::Shutdown();
+	FAvaShapesEditorCommands::Unregister();
+	FAvaInteractiveToolsDelegates::GetRegisterToolsDelegate().RemoveAll(this);
+	FCoreDelegates::OnPostEngineInit.RemoveAll(this);
+
+	if (FModuleManager::Get().IsModuleLoaded("Sequencer"))
+	{
+		ISequencerModule& SequencerModule = FModuleManager::LoadModuleChecked<ISequencerModule>("Sequencer");
+		SequencerModule.UnRegisterTrackEditor(TrackEditorHandle);
+		TrackEditorHandle.Reset();
+	}
+}
+
+void FAvalancheShapesEditorModule::RegisterShapeTools(IAvalancheInteractiveToolsModule* InModule)
+{
+	// 2D Tools
+	InModule->RegisterTool(
+		IAvalancheInteractiveToolsModule::Get().CategoryName2D,
+		GetDefault<UAvaShapesEditorShapeToolRectangle>()->GetToolParameters()
+	);
+
+	InModule->RegisterTool(
+		IAvalancheInteractiveToolsModule::Get().CategoryName2D,
+		GetDefault<UAvaShapesEditorShapeToolEllipse>()->GetToolParameters()
+	);
+
+	InModule->RegisterTool(
+		IAvalancheInteractiveToolsModule::Get().CategoryName2D,
+		GetDefault<UAvaShapesEditorShapeToolNGon>()->GetToolParameters()
+	);
+
+	InModule->RegisterTool(
+		IAvalancheInteractiveToolsModule::Get().CategoryName2D,
+		GetDefault<UAvaShapesEditorShapeToolIrregularPoly>()->GetToolParameters()
+	);
+
+	InModule->RegisterTool(
+		IAvalancheInteractiveToolsModule::Get().CategoryName2D,
+		GetDefault<UAvaShapesEditorShapeToolRing>()->GetToolParameters()
+	);
+
+	InModule->RegisterTool(
+		IAvalancheInteractiveToolsModule::Get().CategoryName2D,
+		GetDefault<UAvaShapesEditorShapeToolStar>()->GetToolParameters()
+	);
+
+	InModule->RegisterTool(
+		IAvalancheInteractiveToolsModule::Get().CategoryName2D,
+		GetDefault<UAvaShapesEditorShapeTool2DArrow>()->GetToolParameters()
+	);
+
+	InModule->RegisterTool(
+		IAvalancheInteractiveToolsModule::Get().CategoryName2D,
+		GetDefault<UAvaShapesEditorShapeToolChevron>()->GetToolParameters()
+	);
+
+	InModule->RegisterTool(
+		IAvalancheInteractiveToolsModule::Get().CategoryName2D,
+		GetDefault<UAvaShapesEditorShapeToolLine>()->GetToolParameters()
+	);
+
+	// 3D Tools
+	InModule->RegisterTool(
+		IAvalancheInteractiveToolsModule::Get().CategoryName3D,
+		GetDefault<UAvaShapesEditorShapeToolCube>()->GetToolParameters()
+	);
+
+	InModule->RegisterTool(
+		IAvalancheInteractiveToolsModule::Get().CategoryName3D,
+		GetDefault<UAvaShapesEditorShapeToolSphere>()->GetToolParameters()
+	);
+
+	InModule->RegisterTool(
+		IAvalancheInteractiveToolsModule::Get().CategoryName3D,
+		GetDefault<UAvaShapesEditorShapeToolCone>()->GetToolParameters()
+	);
+
+	InModule->RegisterTool(
+		IAvalancheInteractiveToolsModule::Get().CategoryName3D,
+		GetDefault<UAvaShapesEditorShapeToolTorus>()->GetToolParameters()
+	);
+}
+
+void FAvalancheShapesEditorModule::RegisterVisualizers()
+{
+	// 2D Visualizers
+	IAvalancheComponentVisualizersModule::RegisterComponentVisualizer<UAvaShape2DArrowDynamicMesh, FAvaShape2DDynamicMeshVisualizer>(&Visualizers);
+	IAvalancheComponentVisualizersModule::RegisterComponentVisualizer<UAvaShapeChevronDynamicMesh, FAvaShape2DDynamicMeshVisualizer>(&Visualizers);
+	IAvalancheComponentVisualizersModule::RegisterComponentVisualizer<UAvaShapeEllipseDynamicMesh, FAvaShapeEllipseDynamicMeshVisualizer>(&Visualizers);
+	IAvalancheComponentVisualizersModule::RegisterComponentVisualizer<UAvaShapeIrregularPolygonDynamicMesh, FAvaShapeIrregularPolygonDynamicMeshVisualizer>(&Visualizers);
+	IAvalancheComponentVisualizersModule::RegisterComponentVisualizer<UAvaShapeLineDynamicMesh, FAvaShapeLineDynamicMeshVisualizer>(&Visualizers);
+	IAvalancheComponentVisualizersModule::RegisterComponentVisualizer<UAvaShapeNGonDynamicMesh, FAvaShapeNGonDynamicMeshVisualizer>(&Visualizers);
+	IAvalancheComponentVisualizersModule::RegisterComponentVisualizer<UAvaShapeRectangleDynamicMesh, FAvaShapeRectangleDynamicMeshVisualizer>(&Visualizers);
+	IAvalancheComponentVisualizersModule::RegisterComponentVisualizer<UAvaShapeRingDynamicMesh, FAvaShapeRingDynamicMeshVisualizer>(&Visualizers);
+	IAvalancheComponentVisualizersModule::RegisterComponentVisualizer<UAvaShapeStarDynamicMesh, FAvaShapeStarDynamicMeshVisualizer>(&Visualizers);
+
+	// 3D Visualizers
+	IAvalancheComponentVisualizersModule::RegisterComponentVisualizer<UAvaShapeConeDynamicMesh, FAvaShapeConeDynamicMeshVisualizer>(&Visualizers);
+	IAvalancheComponentVisualizersModule::RegisterComponentVisualizer<UAvaShapeCubeDynamicMesh, FAvaShapeCubeDynamicMeshVisualizer>(&Visualizers);
+	IAvalancheComponentVisualizersModule::RegisterComponentVisualizer<UAvaShapeSphereDynamicMesh, FAvaShapeSphereDynamicMeshVisualizer>(&Visualizers);
+	IAvalancheComponentVisualizersModule::RegisterComponentVisualizer<UAvaShapeTorusDynamicMesh, FAvaShapeTorusDynamicMeshVisualizer>(&Visualizers);
+}
+
+#undef LOCTEXT_NAMESPACE
+
+IMPLEMENT_MODULE(FAvalancheShapesEditorModule, AvalancheShapesEditor)
\ No newline at end of file
diff --git a/Engine/Plugins/Experimental/Avalanche/Source/AvalancheShapesEditor/Private/Tools/AvaShapesEditorShapeTool2DArrow.cpp b/Engine/Plugins/Experimental/Avalanche/Source/AvalancheShapesEditor/Private/Tools/AvaShapesEditorShapeTool2DArrow.cpp
new file mode 100644
--- /dev/null
+++ b/Engine/Plugins/Experimental/Avalanche/Source/AvalancheShapesEditor/Private/Tools/AvaShapesEditorShapeTool2DArrow.cpp
@@ -0,0 +1,44 @@
+// Copyright Epic Games, Inc. All Rights Reserved.
+
+#include "Tools/AvaShapesEditorShapeTool2DArrow.h"
+#include "AvaShapesEditorCommands.h"
+#include "AvalancheShapesEditorModule.h"
+#include "Builders/AvaInteractiveToolsToolBuilder.h"
+#include "DynamicMeshes/AvaShape2DArrowDynMesh.h"
+#include "UObject/ConstructorHelpers.h"
+
+UAvaShapesEditorShapeTool2DArrow::UAvaShapesEditorShapeTool2DArrow()
+{
+	ShapeClass = UAvaShape2DArrowDynamicMesh::StaticClass();
+}
+
+FName UAvaShapesEditorShapeTool2DArrow::GetCategoryName()
+{
+	return IAvalancheInteractiveToolsModule::CategoryName2D;
+}
+
+FAvaInteractiveToolsToolParameters UAvaShapesEditorShapeTool2DArrow::GetToolParameters() const
+{
+	return {
+		FAvaShapesEditorCommands::Get().Tool_Shape_2DArrow,
+		TEXT("Parametric 2D Arrow Tool"),
+		7000,
+		FAvalancheInteractiveToolsCreateBuilder::CreateLambda(
+			[](UEdMode* InEdMode)
+			{
+				return UAvaInteractiveToolsToolBuilder::CreateToolBuilder<UAvaShapesEditorShapeTool2DArrow>(InEdMode);
+			}),
+		nullptr,
+		CreateFactory<UAvaShape2DArrowDynamicMesh>()
+	};
+}
+
+void UAvaShapesEditorShapeTool2DArrow::InitShape(UAvaShapeDynamicMeshBase* InShape) const
+{
+	UAvaShape2DArrowDynamicMesh* Arrow2D = Cast<UAvaShape2DArrowDynamicMesh>(InShape);
+	check(Arrow2D);
+
+	Arrow2D->SetSize2D({DefaultDim, DefaultDim});
+
+	Super::InitShape(InShape);
+}
diff --git a/Engine/Plugins/Experimental/Avalanche/Source/AvalancheShapesEditor/Private/Tools/AvaShapesEditorShapeTool2DArrow.h b/Engine/Plugins/Experimental/Avalanche/Source/AvalancheShapesEditor/Private/Tools/AvaShapesEditorShapeTool2DArrow.h
new file mode 100644
--- /dev/null
+++ b/Engine/Plugins/Experimental/Avalanche/Source/AvalancheShapesEditor/Private/Tools/AvaShapesEditorShapeTool2DArrow.h
@@ -0,0 +1,25 @@
+// Copyright Epic Games, Inc. All Rights Reserved.
+
+#pragma once
+
+#include "Tools/AvaShapesEditorShapeAreaToolBase.h"
+#include "AvaShapesEditorShapeTool2DArrow.generated.h"
+
+UCLASS()
+class UAvaShapesEditorShapeTool2DArrow : public UAvaShapesEditorShapeAreaToolBase
+{
+	GENERATED_BODY()
+
+public:
+	UAvaShapesEditorShapeTool2DArrow();
+
+	//~ Begin UAvaInteractiveToolsToolBase
+	virtual FName GetCategoryName() override;
+	virtual FAvaInteractiveToolsToolParameters GetToolParameters() const override;
+	//~ End UAvaInteractiveToolsToolBase
+
+protected:
+	//~ Begin UAvaShapesEditorShapeToolBase
+	virtual void InitShape(UAvaShapeDynamicMeshBase* InShape) const override;
+	//~ End UAvaShapesEditorShapeToolBase
+};
