[MotionDesign]

diff --git a/Engine/Plugins/Experimental/Avalanche/Source/AvalancheShapesEditor/Private/Tools/AvaShapesEditorShapeToolBase.cpp b/Engine/Plugins/Experimental/Avalanche/Source/AvalancheShapesEditor/Private/Tools/AvaShapesEditorShapeToolBase.cpp
new file mode 100644
--- /dev/null
+++ b/Engine/Plugins/Experimental/Avalanche/Source/AvalancheShapesEditor/Private/Tools/AvaShapesEditorShapeToolBase.cpp
@@ -0,0 +1,113 @@
+// Copyright Epic Games, Inc. All Rights Reserved.
+
+#include "Tools/AvaShapesEditorShapeToolBase.h"
+#include "AvaShapeActor.h"
+#include "AvaInteractiveToolsSettings.h"
+#include "AvaViewportUtils.h"
+#include "DynamicMeshes/AvaShape2DDynMeshBase.h"
+#include "DynamicMeshes/AvaShape3DDynMeshBase.h"
+#include "DynamicMeshes/AvaShapeDynMeshBase.h"
+#include "GameFramework/Actor.h"
+#include "InteractiveToolManager.h"
+#include "ToolContextInterfaces.h"
+
+UAvaShapesEditorShapeToolBase::UAvaShapesEditorShapeToolBase()
+{
+	ActorClass = AAvaShapeActor::StaticClass();
+}
+
+bool UAvaShapesEditorShapeToolBase::OnBegin()
+{
+	if (ShapeClass == nullptr)
+	{
+		return false;
+	}
+
+	return Super::OnBegin();
+}
+
+AActor* UAvaShapesEditorShapeToolBase::SpawnActor(TSubclassOf<AActor> InActorClass, EAvaViewportStatus InViewportStatus,
+	const FVector2f& InViewportPosition, bool bInPreview, FString* InActorLabelOverride) const
+{
+	FString ActorLabelOverride = (InActorLabelOverride && !InActorLabelOverride->IsEmpty())
+		? *InActorLabelOverride
+		: GetActorNameOverride();
+
+	AActor* NewActor = Super::SpawnActor(InActorClass, InViewportStatus, InViewportPosition, bInPreview, &ActorLabelOverride);
+
+	if (AAvaShapeActor* ShapeActor = Cast<AAvaShapeActor>(NewActor))
+	{
+		UAvaShapeDynamicMeshBase* NewShape = NewObject<UAvaShapeDynamicMeshBase>(ShapeActor, ShapeClass);
+		SetShapeSize(ShapeActor, FVector2D(DefaultDim, DefaultDim));
+		InitShape(NewShape);
+
+		ShapeActor->SetDynamicMesh(NewShape);
+
+		SetToolkitSettingsObject(NewShape);
+	}
+
+	return NewActor;
+}
+
+FString UAvaShapesEditorShapeToolBase::GetActorNameOverride() const
+{
+	if (const UAvaShapeDynamicMeshBase* DefaultObject = ShapeClass.GetDefaultObject())
+	{
+		return DefaultObject->GetMeshName();
+	}
+
+	return "";
+}
+
+void UAvaShapesEditorShapeToolBase::InitShape(UAvaShapeDynamicMeshBase* InShape) const
+{
+	InShape->SetParametricMaterial(UAvaShapeDynamicMeshBase::MESH_INDEX_PRIMARY, DefaultMaterialParams);
+}
+
+void UAvaShapesEditorShapeToolBase::SetShapeSize(AAvaShapeActor* InShapeActor, const FVector2D& InShapeSize) const
+{
+	if (!InShapeActor)
+	{
+		return;
+	}
+
+	UAvaShapeDynamicMeshBase* MeshBase = InShapeActor->GetDynamicMesh();
+
+	if (!MeshBase)
+	{
+		return;
+	}
+
+	if (UAvaShape2DDynMeshBase* Mesh2D = Cast<UAvaShape2DDynMeshBase>(MeshBase))
+	{
+		Mesh2D->SetSize2D(InShapeSize);
+	}
+	else if (UAvaShape3DDynMeshBase* Mesh3D = Cast<UAvaShape3DDynMeshBase>(MeshBase))
+	{
+		Mesh3D->SetSize3D({DefaultDim, InShapeSize.X, InShapeSize.Y});
+	}
+	else
+	{
+		checkNoEntry();
+	}
+}
+
+bool UAvaShapesEditorShapeToolBase::UseIdentityLocation() const
+{
+	if (!bPerformingDefaultAction)
+	{
+		return false;
+	}
+
+	if (!ShapeClass || ShapeClass->IsChildOf<UAvaShape2DDynMeshBase>())
+	{
+		return false;
+	}
+
+	return !IsAvalancheViewport();
+}
+
+bool UAvaShapesEditorShapeToolBase::UseIdentityRotation() const
+{
+	return ConditionalIdentityRotation();
+}
diff --git a/Engine/Plugins/Experimental/Avalanche/Source/AvalancheShapesEditor/Private/Tools/AvaShapesEditorShapeToolChevron.cpp b/Engine/Plugins/Experimental/Avalanche/Source/AvalancheShapesEditor/Private/Tools/AvaShapesEditorShapeToolChevron.cpp
new file mode 100644
--- /dev/null
+++ b/Engine/Plugins/Experimental/Avalanche/Source/AvalancheShapesEditor/Private/Tools/AvaShapesEditorShapeToolChevron.cpp
@@ -0,0 +1,44 @@
+// Copyright Epic Games, Inc. All Rights Reserved.
+
+#include "Tools/AvaShapesEditorShapeToolChevron.h"
+#include "AvaShapesEditorCommands.h"
+#include "AvalancheShapesEditorModule.h"
+#include "Builders/AvaInteractiveToolsToolBuilder.h"
+#include "DynamicMeshes/AvaShapeChevronDynMesh.h"
+#include "UObject/ConstructorHelpers.h"
+
+UAvaShapesEditorShapeToolChevron::UAvaShapesEditorShapeToolChevron()
+{
+	ShapeClass = UAvaShapeChevronDynamicMesh::StaticClass();
+}
+
+FName UAvaShapesEditorShapeToolChevron::GetCategoryName()
+{
+	return IAvalancheInteractiveToolsModule::CategoryName2D;
+}
+
+FAvaInteractiveToolsToolParameters UAvaShapesEditorShapeToolChevron::GetToolParameters() const
+{
+	return {
+		FAvaShapesEditorCommands::Get().Tool_Shape_Chevron,
+		TEXT("Parametric Chevron Tool"),
+		8000,
+		FAvalancheInteractiveToolsCreateBuilder::CreateLambda(
+			[](UEdMode* InEdMode)
+			{
+				return UAvaInteractiveToolsToolBuilder::CreateToolBuilder<UAvaShapesEditorShapeToolChevron>(InEdMode);
+			}),
+		nullptr,
+		CreateFactory<UAvaShapeChevronDynamicMesh>()
+	};
+}
+
+void UAvaShapesEditorShapeToolChevron::InitShape(UAvaShapeDynamicMeshBase* InShape) const
+{
+	UAvaShapeChevronDynamicMesh* Chevron = Cast<UAvaShapeChevronDynamicMesh>(InShape);
+	check(Chevron);
+
+	Chevron->SetSize2D({DefaultDim, DefaultDim});
+
+	Super::InitShape(InShape);
+}
diff --git a/Engine/Plugins/Experimental/Avalanche/Source/AvalancheShapesEditor/Private/Tools/AvaShapesEditorShapeToolChevron.h b/Engine/Plugins/Experimental/Avalanche/Source/AvalancheShapesEditor/Private/Tools/AvaShapesEditorShapeToolChevron.h
new file mode 100644
--- /dev/null
+++ b/Engine/Plugins/Experimental/Avalanche/Source/AvalancheShapesEditor/Private/Tools/AvaShapesEditorShapeToolChevron.h
@@ -0,0 +1,25 @@
+// Copyright Epic Games, Inc. All Rights Reserved.
+
+#pragma once
+
+#include "Tools/AvaShapesEditorShapeAreaToolBase.h"
+#include "AvaShapesEditorShapeToolChevron.generated.h"
+
+UCLASS()
+class UAvaShapesEditorShapeToolChevron : public UAvaShapesEditorShapeAreaToolBase
+{
+	GENERATED_BODY()
+
+public:
+	UAvaShapesEditorShapeToolChevron();
+
+	//~ Begin UAvaInteractiveToolsToolBase
+	virtual FName GetCategoryName() override;
+	virtual FAvaInteractiveToolsToolParameters GetToolParameters() const override;
+	//~ End UAvaInteractiveToolsToolBase
+
+protected:
+	//~ Begin UAvaShapesEditorShapeToolBase
+	virtual void InitShape(UAvaShapeDynamicMeshBase* InShape) const override;
+	//~ End UAvaShapesEditorShapeToolBase
+};
diff --git a/Engine/Plugins/Experimental/Avalanche/Source/AvalancheShapesEditor/Private/Tools/AvaShapesEditorShapeToolCone.cpp b/Engine/Plugins/Experimental/Avalanche/Source/AvalancheShapesEditor/Private/Tools/AvaShapesEditorShapeToolCone.cpp
new file mode 100644
--- /dev/null
+++ b/Engine/Plugins/Experimental/Avalanche/Source/AvalancheShapesEditor/Private/Tools/AvaShapesEditorShapeToolCone.cpp
@@ -0,0 +1,44 @@
+// Copyright Epic Games, Inc. All Rights Reserved.
+
+#include "Tools/AvaShapesEditorShapeToolCone.h"
+#include "AvaShapesEditorCommands.h"
+#include "AvalancheShapesEditorModule.h"
+#include "Builders/AvaInteractiveToolsToolBuilder.h"
+#include "DynamicMeshes/AvaShapeConeDynMesh.h"
+#include "UObject/ConstructorHelpers.h"
+
+UAvaShapesEditorShapeToolCone::UAvaShapesEditorShapeToolCone()
+{
+	ShapeClass = UAvaShapeConeDynamicMesh::StaticClass();
+}
+
+FName UAvaShapesEditorShapeToolCone::GetCategoryName()
+{
+	return IAvalancheInteractiveToolsModule::CategoryName3D;
+}
+
+FAvaInteractiveToolsToolParameters UAvaShapesEditorShapeToolCone::GetToolParameters() const
+{
+	return {
+		FAvaShapesEditorCommands::Get().Tool_Shape_Cone,
+		TEXT("Parametric Cone Tool"),
+		3000,
+		FAvalancheInteractiveToolsCreateBuilder::CreateLambda(
+			[](UEdMode* InEdMode)
+			{
+				return UAvaInteractiveToolsToolBuilder::CreateToolBuilder<UAvaShapesEditorShapeToolCone>(InEdMode);
+			}),
+		nullptr,
+		CreateFactory<UAvaShapeConeDynamicMesh>()
+	};
+}
+
+void UAvaShapesEditorShapeToolCone::InitShape(UAvaShapeDynamicMeshBase* InShape) const
+{
+	UAvaShapeConeDynamicMesh* Cone = Cast<UAvaShapeConeDynamicMesh>(InShape);
+	check(Cone);
+
+	Cone->SetSize3D({DefaultDim, DefaultDim, DefaultDim});
+
+	Super::InitShape(InShape);
+}
diff --git a/Engine/Plugins/Experimental/Avalanche/Source/AvalancheShapesEditor/Private/Tools/AvaShapesEditorShapeToolCone.h b/Engine/Plugins/Experimental/Avalanche/Source/AvalancheShapesEditor/Private/Tools/AvaShapesEditorShapeToolCone.h
new file mode 100644
--- /dev/null
+++ b/Engine/Plugins/Experimental/Avalanche/Source/AvalancheShapesEditor/Private/Tools/AvaShapesEditorShapeToolCone.h
@@ -0,0 +1,25 @@
+// Copyright Epic Games, Inc. All Rights Reserved.
+
+#pragma once
+
+#include "Tools/AvaShapesEditorShapeAreaToolBase.h"
+#include "AvaShapesEditorShapeToolCone.generated.h"
+
+UCLASS()
+class UAvaShapesEditorShapeToolCone : public UAvaShapesEditorShapeAreaToolBase
+{
+	GENERATED_BODY()
+
+public:
+	UAvaShapesEditorShapeToolCone();
+
+	//~ Begin UAvaInteractiveToolsToolBase
+	virtual FName GetCategoryName() override;
+	virtual FAvaInteractiveToolsToolParameters GetToolParameters() const override;
+	//~ End UAvaInteractiveToolsToolBase
+
+protected:
+	//~ Begin UAvaShapesEditorShapeToolBase
+	virtual void InitShape(UAvaShapeDynamicMeshBase* InShape) const override;
+	//~ End UAvaShapesEditorShapeToolBase
+};
diff --git a/Engine/Plugins/Experimental/Avalanche/Source/AvalancheShapesEditor/Private/Tools/AvaShapesEditorShapeToolCube.cpp b/Engine/Plugins/Experimental/Avalanche/Source/AvalancheShapesEditor/Private/Tools/AvaShapesEditorShapeToolCube.cpp
new file mode 100644
--- /dev/null
+++ b/Engine/Plugins/Experimental/Avalanche/Source/AvalancheShapesEditor/Private/Tools/AvaShapesEditorShapeToolCube.cpp
@@ -0,0 +1,44 @@
+// Copyright Epic Games, Inc. All Rights Reserved.
+
+#include "Tools/AvaShapesEditorShapeToolCube.h"
+#include "AvaShapesEditorCommands.h"
+#include "AvalancheShapesEditorModule.h"
+#include "Builders/AvaInteractiveToolsToolBuilder.h"
+#include "DynamicMeshes/AvaShapeCubeDynMesh.h"
+#include "UObject/ConstructorHelpers.h"
+
+UAvaShapesEditorShapeToolCube::UAvaShapesEditorShapeToolCube()
+{
+	ShapeClass = UAvaShapeCubeDynamicMesh::StaticClass();
+}
+
+FName UAvaShapesEditorShapeToolCube::GetCategoryName()
+{
+	return IAvalancheInteractiveToolsModule::CategoryName3D;
+}
+
+FAvaInteractiveToolsToolParameters UAvaShapesEditorShapeToolCube::GetToolParameters() const
+{
+	return {
+		FAvaShapesEditorCommands::Get().Tool_Shape_Cube,
+		TEXT("Parametric Cube Tool"),
+		1000,
+		FAvalancheInteractiveToolsCreateBuilder::CreateLambda(
+			[](UEdMode* InEdMode)
+			{
+				return UAvaInteractiveToolsToolBuilder::CreateToolBuilder<UAvaShapesEditorShapeToolCube>(InEdMode);
+			}),
+		nullptr,
+		CreateFactory<UAvaShapeCubeDynamicMesh>()
+	};
+}
+
+void UAvaShapesEditorShapeToolCube::InitShape(UAvaShapeDynamicMeshBase* InShape) const
+{
+	UAvaShapeCubeDynamicMesh* Cube = Cast<UAvaShapeCubeDynamicMesh>(InShape);
+	check(Cube);
+
+	Cube->SetSize3D({DefaultDim, DefaultDim, DefaultDim});
+
+	Super::InitShape(InShape);
+}
diff --git a/Engine/Plugins/Experimental/Avalanche/Source/AvalancheShapesEditor/Private/Tools/AvaShapesEditorShapeToolCube.h b/Engine/Plugins/Experimental/Avalanche/Source/AvalancheShapesEditor/Private/Tools/AvaShapesEditorShapeToolCube.h
new file mode 100644
--- /dev/null
+++ b/Engine/Plugins/Experimental/Avalanche/Source/AvalancheShapesEditor/Private/Tools/AvaShapesEditorShapeToolCube.h
@@ -0,0 +1,25 @@
+// Copyright Epic Games, Inc. All Rights Reserved.
+
+#pragma once
+
+#include "Tools/AvaShapesEditorShapeAreaToolBase.h"
+#include "AvaShapesEditorShapeToolCube.generated.h"
+
+UCLASS()
+class UAvaShapesEditorShapeToolCube : public UAvaShapesEditorShapeAreaToolBase
+{
+	GENERATED_BODY()
+
+public:
+	UAvaShapesEditorShapeToolCube();
+
+	//~ Begin UAvaInteractiveToolsToolBase
+	virtual FName GetCategoryName() override;
+	virtual FAvaInteractiveToolsToolParameters GetToolParameters() const override;
+	//~ End UAvaInteractiveToolsToolBase
+
+protected:
+	//~ Begin UAvaShapesEditorShapeToolBase
+	virtual void InitShape(UAvaShapeDynamicMeshBase* InShape) const override;
+	//~ End UAvaShapesEditorShapeToolBase
+};
diff --git a/Engine/Plugins/Experimental/Avalanche/Source/AvalancheShapesEditor/Private/Tools/AvaShapesEditorShapeToolEllipse.cpp b/Engine/Plugins/Experimental/Avalanche/Source/AvalancheShapesEditor/Private/Tools/AvaShapesEditorShapeToolEllipse.cpp
new file mode 100644
--- /dev/null
+++ b/Engine/Plugins/Experimental/Avalanche/Source/AvalancheShapesEditor/Private/Tools/AvaShapesEditorShapeToolEllipse.cpp
@@ -0,0 +1,44 @@
+// Copyright Epic Games, Inc. All Rights Reserved.
+
+#include "Tools/AvaShapesEditorShapeToolEllipse.h"
+#include "AvaShapesEditorCommands.h"
+#include "AvalancheShapesEditorModule.h"
+#include "Builders/AvaInteractiveToolsToolBuilder.h"
+#include "DynamicMeshes/AvaShapeEllipseDynMesh.h"
+#include "UObject/ConstructorHelpers.h"
+
+UAvaShapesEditorShapeToolEllipse::UAvaShapesEditorShapeToolEllipse()
+{
+	ShapeClass = UAvaShapeEllipseDynamicMesh::StaticClass();
+}
+
+FName UAvaShapesEditorShapeToolEllipse::GetCategoryName()
+{
+	return IAvalancheInteractiveToolsModule::CategoryName2D;
+}
+
+FAvaInteractiveToolsToolParameters UAvaShapesEditorShapeToolEllipse::GetToolParameters() const
+{
+	return {
+		FAvaShapesEditorCommands::Get().Tool_Shape_Ellipse,
+		TEXT("Parametric Ellipse Tool"),
+		2000,
+		FAvalancheInteractiveToolsCreateBuilder::CreateLambda(
+			[](UEdMode* InEdMode)
+			{
+				return UAvaInteractiveToolsToolBuilder::CreateToolBuilder<UAvaShapesEditorShapeToolEllipse>(InEdMode);
+			}),
+		nullptr,
+		CreateFactory<UAvaShapeEllipseDynamicMesh>()
+	};
+}
+
+void UAvaShapesEditorShapeToolEllipse::InitShape(UAvaShapeDynamicMeshBase* InShape) const
+{
+	UAvaShapeEllipseDynamicMesh* Ellipse = Cast<UAvaShapeEllipseDynamicMesh>(InShape);
+	check(Ellipse);
+
+	Ellipse->SetSize2D({DefaultDim, DefaultDim});
+
+	Super::InitShape(InShape);
+}
diff --git a/Engine/Plugins/Experimental/Avalanche/Source/AvalancheShapesEditor/Private/Tools/AvaShapesEditorShapeToolEllipse.h b/Engine/Plugins/Experimental/Avalanche/Source/AvalancheShapesEditor/Private/Tools/AvaShapesEditorShapeToolEllipse.h
new file mode 100644
--- /dev/null
+++ b/Engine/Plugins/Experimental/Avalanche/Source/AvalancheShapesEditor/Private/Tools/AvaShapesEditorShapeToolEllipse.h
@@ -0,0 +1,25 @@
+// Copyright Epic Games, Inc. All Rights Reserved.
+
+#pragma once
+
+#include "Tools/AvaShapesEditorShapeAreaToolBase.h"
+#include "AvaShapesEditorShapeToolEllipse.generated.h"
+
+UCLASS()
+class UAvaShapesEditorShapeToolEllipse : public UAvaShapesEditorShapeAreaToolBase
+{
+	GENERATED_BODY()
+
+public:
+	UAvaShapesEditorShapeToolEllipse();
+
+	//~ Begin UAvaInteractiveToolsToolBase
+	virtual FName GetCategoryName() override;
+	virtual FAvaInteractiveToolsToolParameters GetToolParameters() const override;
+	//~ End UAvaInteractiveToolsToolBase
+
+protected:
+	//~ Begin UAvaShapesEditorShapeToolBase
+	virtual void InitShape(UAvaShapeDynamicMeshBase* InShape) const override;
+	//~ End UAvaShapesEditorShapeToolBase
+};
diff --git a/Engine/Plugins/Experimental/Avalanche/Source/AvalancheShapesEditor/Private/Tools/AvaShapesEditorShapeToolNGon.cpp b/Engine/Plugins/Experimental/Avalanche/Source/AvalancheShapesEditor/Private/Tools/AvaShapesEditorShapeToolNGon.cpp
new file mode 100644
--- /dev/null
+++ b/Engine/Plugins/Experimental/Avalanche/Source/AvalancheShapesEditor/Private/Tools/AvaShapesEditorShapeToolNGon.cpp
@@ -0,0 +1,44 @@
+// Copyright Epic Games, Inc. All Rights Reserved.
+
+#include "Tools/AvaShapesEditorShapeToolNGon.h"
+#include "AvaShapesEditorCommands.h"
+#include "AvalancheShapesEditorModule.h"
+#include "Builders/AvaInteractiveToolsToolBuilder.h"
+#include "DynamicMeshes/AvaShapeNGonDynMesh.h"
+#include "UObject/ConstructorHelpers.h"
+
+UAvaShapesEditorShapeToolNGon::UAvaShapesEditorShapeToolNGon()
+{
+	ShapeClass = UAvaShapeNGonDynamicMesh::StaticClass();
+}
+
+FName UAvaShapesEditorShapeToolNGon::GetCategoryName()
+{
+	return IAvalancheInteractiveToolsModule::CategoryName2D;
+}
+
+FAvaInteractiveToolsToolParameters UAvaShapesEditorShapeToolNGon::GetToolParameters() const
+{
+	return {
+		FAvaShapesEditorCommands::Get().Tool_Shape_NGon,
+		TEXT("Parametric Regular Polygon Tool"),
+		3000,
+		FAvalancheInteractiveToolsCreateBuilder::CreateLambda(
+			[](UEdMode* InEdMode)
+			{
+				return UAvaInteractiveToolsToolBuilder::CreateToolBuilder<UAvaShapesEditorShapeToolNGon>(InEdMode);
+			}),
+		nullptr,
+		CreateFactory<UAvaShapeNGonDynamicMesh>()
+	};
+}
+
+void UAvaShapesEditorShapeToolNGon::InitShape(UAvaShapeDynamicMeshBase* InShape) const
+{
+	UAvaShapeNGonDynamicMesh* NGon = Cast<UAvaShapeNGonDynamicMesh>(InShape);
+	check(NGon);
+
+	NGon->SetSize2D({DefaultDim, DefaultDim});
+
+	Super::InitShape(InShape);
+}
diff --git a/Engine/Plugins/Experimental/Avalanche/Source/AvalancheShapesEditor/Private/Tools/AvaShapesEditorShapeToolNGon.h b/Engine/Plugins/Experimental/Avalanche/Source/AvalancheShapesEditor/Private/Tools/AvaShapesEditorShapeToolNGon.h
new file mode 100644
--- /dev/null
+++ b/Engine/Plugins/Experimental/Avalanche/Source/AvalancheShapesEditor/Private/Tools/AvaShapesEditorShapeToolNGon.h
@@ -0,0 +1,25 @@
+// Copyright Epic Games, Inc. All Rights Reserved.
+
+#pragma once
+
+#include "Tools/AvaShapesEditorShapeAreaToolBase.h"
+#include "AvaShapesEditorShapeToolNGon.generated.h"
+
+UCLASS()
+class UAvaShapesEditorShapeToolNGon : public UAvaShapesEditorShapeAreaToolBase
+{
+	GENERATED_BODY()
+
+public:
+	UAvaShapesEditorShapeToolNGon();
+
+	//~ Begin UAvaInteractiveToolsToolBase
+	virtual FName GetCategoryName() override;
+	virtual FAvaInteractiveToolsToolParameters GetToolParameters() const override;
+	//~ End UAvaInteractiveToolsToolBase
+
+protected:
+	//~ Begin UAvaShapesEditorShapeToolBase
+	virtual void InitShape(UAvaShapeDynamicMeshBase* InShape) const override;
+	//~ End UAvaShapesEditorShapeToolBase
+};
diff --git a/Engine/Plugins/Experimental/Avalanche/Source/AvalancheShapesEditor/Private/Tools/AvaShapesEditorShapeToolRectangle.cpp b/Engine/Plugins/Experimental/Avalanche/Source/AvalancheShapesEditor/Private/Tools/AvaShapesEditorShapeToolRectangle.cpp
new file mode 100644
--- /dev/null
+++ b/Engine/Plugins/Experimental/Avalanche/Source/AvalancheShapesEditor/Private/Tools/AvaShapesEditorShapeToolRectangle.cpp
@@ -0,0 +1,93 @@
+// Copyright Epic Games, Inc. All Rights Reserved.
+
+#include "Tools/AvaShapesEditorShapeToolRectangle.h"
+#include "AvaShapesEditorCommands.h"
+#include "Builders/AvaInteractiveToolsToolBuilder.h"
+#include "DynamicMeshes/AvaShapeRectangleDynMesh.h"
+#include "Tools/AvaInteractiveToolsToolPresetBase.h"
+#include "UObject/ConstructorHelpers.h"
+
+#define LOCTEXT_NAMESPACE "AvaShapesEditorShapeToolRectangle"
+
+class FAvaShapesEditorShapeToolRectanglePreset : public FAvaInteractiveToolsToolPresetBase
+{
+public:
+	FAvaShapesEditorShapeToolRectanglePreset(const FText& InName, const FText& InDescription, const FVector2D& InSize)
+	{
+		Name = InName;
+		Description = InDescription;
+		Size = InSize;
+	}
+
+	virtual void ApplyPreset(AActor* InActor) const override
+	{
+		if (AAvaShapeActor* ShapeActor = Cast<AAvaShapeActor>(InActor))
+		{
+			if (UAvaShapeRectangleDynamicMesh* Rectangle = Cast<UAvaShapeRectangleDynamicMesh>(ShapeActor->GetDynamicMesh()))
+			{
+				Rectangle->SetSize2D(Size);
+			}
+		}
+	}
+
+protected:
+	FVector2D Size;
+};
+
+UAvaShapesEditorShapeToolRectangle::UAvaShapesEditorShapeToolRectangle()
+{
+	ShapeClass = UAvaShapeRectangleDynamicMesh::StaticClass();
+}
+
+FName UAvaShapesEditorShapeToolRectangle::GetCategoryName()
+{
+	return IAvalancheInteractiveToolsModule::CategoryName2D;
+}
+
+FAvaInteractiveToolsToolParameters UAvaShapesEditorShapeToolRectangle::GetToolParameters() const
+{
+	static const TMap<FName, TSharedRef<FAvaInteractiveToolsToolPresetBase>> Presets = {
+			{
+				"Square100",
+				MakeShared<FAvaShapesEditorShapeToolRectanglePreset>(
+					LOCTEXT("100x100", "Square [100 x 100]"),
+					LOCTEXT("100x100.ToolTip", "Create a 100 x 100 square shape."),
+					FVector2D(100, 100)
+				)
+			},
+			{
+				"Rectangle16x9",
+				MakeShared<FAvaShapesEditorShapeToolRectanglePreset>(
+					LOCTEXT("160x90", "Rectangle [160 x 90]"),
+					LOCTEXT("160x90.ToolTip", "Create a 160 x 90 rectangle shape."),
+					FVector2D(160, 90)
+				)
+			}
+	};
+
+	return {
+		FAvaShapesEditorCommands::Get().Tool_Shape_Rectangle,
+		TEXT("Parametric Rectangle Tool"),
+		1000,
+		FAvalancheInteractiveToolsCreateBuilder::CreateLambda(
+			[](UEdMode* InEdMode)
+			{
+				return UAvaInteractiveToolsToolBuilder::CreateToolBuilder<UAvaShapesEditorShapeToolRectangle>(InEdMode);
+			}),
+		nullptr,
+		CreateFactory<UAvaShapeRectangleDynamicMesh>(),
+		Presets
+	};
+}
+
+void UAvaShapesEditorShapeToolRectangle::InitShape(UAvaShapeDynamicMeshBase* InShape) const
+{
+	UAvaShapeRectangleDynamicMesh* Rectangle = Cast<UAvaShapeRectangleDynamicMesh>(InShape);
+	check(Rectangle);
+
+	Rectangle->SetSize2D({DefaultDim, DefaultDim});
+
+	Super::InitShape(InShape);
+}
+
+#undef LOCTEXT_NAMESPACE
diff --git a/Engine/Plugins/Experimental/Avalanche/Source/AvalancheShapesEditor/Private/Tools/AvaShapesEditorShapeToolRectangle.h b/Engine/Plugins/Experimental/Avalanche/Source/AvalancheShapesEditor/Private/Tools/AvaShapesEditorShapeToolRectangle.h
new file mode 100644
--- /dev/null
+++ b/Engine/Plugins/Experimental/Avalanche/Source/AvalancheShapesEditor/Private/Tools/AvaShapesEditorShapeToolRectangle.h
@@ -0,0 +1,25 @@
+// Copyright Epic Games, Inc. All Rights Reserved.
+
+#pragma once
+
+#include "Tools/AvaShapesEditorShapeAreaToolBase.h"
+#include "AvaShapesEditorShapeToolRectangle.generated.h"
+
+UCLASS()
+class UAvaShapesEditorShapeToolRectangle : public UAvaShapesEditorShapeAreaToolBase
+{
+	GENERATED_BODY()
+
+public:
+	UAvaShapesEditorShapeToolRectangle();
+
+	//~ Begin UAvaInteractiveToolsToolBase
+	virtual FName GetCategoryName() override;
+	virtual FAvaInteractiveToolsToolParameters GetToolParameters() const override;
+	//~ End UAvaInteractiveToolsToolBase
+
+protected:
+	//~ Begin UAvaShapesEditorShapeToolBase
+	virtual void InitShape(UAvaShapeDynamicMeshBase* InShape) const override;
+	//~ End UAvaShapesEditorShapeToolBase
+};
diff --git a/Engine/Plugins/Experimental/Avalanche/Source/AvalancheShapesEditor/Private/Tools/AvaShapesEditorShapeToolRing.cpp b/Engine/Plugins/Experimental/Avalanche/Source/AvalancheShapesEditor/Private/Tools/AvaShapesEditorShapeToolRing.cpp
new file mode 100644
--- /dev/null
+++ b/Engine/Plugins/Experimental/Avalanche/Source/AvalancheShapesEditor/Private/Tools/AvaShapesEditorShapeToolRing.cpp
@@ -0,0 +1,44 @@
+// Copyright Epic Games, Inc. All Rights Reserved.
+
+#include "Tools/AvaShapesEditorShapeToolRing.h"
+#include "AvaShapesEditorCommands.h"
+#include "AvalancheShapesEditorModule.h"
+#include "Builders/AvaInteractiveToolsToolBuilder.h"
+#include "DynamicMeshes/AvaShapeRingDynMesh.h"
+#include "UObject/ConstructorHelpers.h"
+
+UAvaShapesEditorShapeToolRing::UAvaShapesEditorShapeToolRing()
+{
+	ShapeClass = UAvaShapeRingDynamicMesh::StaticClass();
+}
+
+FName UAvaShapesEditorShapeToolRing::GetCategoryName()
+{
+	return IAvalancheInteractiveToolsModule::CategoryName2D;
+}
+
+FAvaInteractiveToolsToolParameters UAvaShapesEditorShapeToolRing::GetToolParameters() const
+{
+	return {
+		FAvaShapesEditorCommands::Get().Tool_Shape_Ring,
+		TEXT("Parametric Ring Tool"),
+		5000,
+		FAvalancheInteractiveToolsCreateBuilder::CreateLambda(
+			[](UEdMode* InEdMode)
+			{
+				return UAvaInteractiveToolsToolBuilder::CreateToolBuilder<UAvaShapesEditorShapeToolRing>(InEdMode);
+			}),
+		nullptr,
+		CreateFactory<UAvaShapeRingDynamicMesh>()
+	};
+}
+
+void UAvaShapesEditorShapeToolRing::InitShape(UAvaShapeDynamicMeshBase* InShape) const
+{
+	UAvaShapeRingDynamicMesh* Ring = Cast<UAvaShapeRingDynamicMesh>(InShape);
+	check(Ring);
+
+	Ring->SetSize2D({DefaultDim, DefaultDim});
+
+	Super::InitShape(InShape);
+}
diff --git a/Engine/Plugins/Experimental/Avalanche/Source/AvalancheShapesEditor/Private/Tools/AvaShapesEditorShapeToolRing.h b/Engine/Plugins/Experimental/Avalanche/Source/AvalancheShapesEditor/Private/Tools/AvaShapesEditorShapeToolRing.h
new file mode 100644
--- /dev/null
+++ b/Engine/Plugins/Experimental/Avalanche/Source/AvalancheShapesEditor/Private/Tools/AvaShapesEditorShapeToolRing.h
@@ -0,0 +1,25 @@
+// Copyright Epic Games, Inc. All Rights Reserved.
+
+#pragma once
+
+#include "Tools/AvaShapesEditorShapeAreaToolBase.h"
+#include "AvaShapesEditorShapeToolRing.generated.h"
+
+UCLASS()
+class UAvaShapesEditorShapeToolRing : public UAvaShapesEditorShapeAreaToolBase
+{
+	GENERATED_BODY()
+
+public:
+	UAvaShapesEditorShapeToolRing();
+
+	//~ Begin UAvaInteractiveToolsToolBase
+	virtual FName GetCategoryName() override;
+	virtual FAvaInteractiveToolsToolParameters GetToolParameters() const override;
+	//~ End UAvaInteractiveToolsToolBase
+
+protected:
+	//~ Begin UAvaShapesEditorShapeToolBase
+	virtual void InitShape(UAvaShapeDynamicMeshBase* InShape) const override;
+	//~ End UAvaShapesEditorShapeToolBase
+};
diff --git a/Engine/Plugins/Experimental/Avalanche/Source/AvalancheShapesEditor/Private/Tools/AvaShapesEditorShapeToolSphere.cpp b/Engine/Plugins/Experimental/Avalanche/Source/AvalancheShapesEditor/Private/Tools/AvaShapesEditorShapeToolSphere.cpp
new file mode 100644
--- /dev/null
+++ b/Engine/Plugins/Experimental/Avalanche/Source/AvalancheShapesEditor/Private/Tools/AvaShapesEditorShapeToolSphere.cpp
@@ -0,0 +1,44 @@
+// Copyright Epic Games, Inc. All Rights Reserved.
+
+#include "Tools/AvaShapesEditorShapeToolSphere.h"
+#include "AvaShapesEditorCommands.h"
+#include "AvalancheShapesEditorModule.h"
+#include "Builders/AvaInteractiveToolsToolBuilder.h"
+#include "DynamicMeshes/AvaShapeSphereDynMesh.h"
+#include "UObject/ConstructorHelpers.h"
+
+UAvaShapesEditorShapeToolSphere::UAvaShapesEditorShapeToolSphere()
+{
+	ShapeClass = UAvaShapeSphereDynamicMesh::StaticClass();
+}
+
+FName UAvaShapesEditorShapeToolSphere::GetCategoryName()
+{
+	return IAvalancheInteractiveToolsModule::CategoryName3D;
+}
+
+FAvaInteractiveToolsToolParameters UAvaShapesEditorShapeToolSphere::GetToolParameters() const
+{
+	return {
+		FAvaShapesEditorCommands::Get().Tool_Shape_Sphere,
+		TEXT("Parametric Sphere Tool"),
+		2000,
+		FAvalancheInteractiveToolsCreateBuilder::CreateLambda(
+			[](UEdMode* InEdMode)
+			{
+				return UAvaInteractiveToolsToolBuilder::CreateToolBuilder<UAvaShapesEditorShapeToolSphere>(InEdMode);
+			}),
+		nullptr,
+		CreateFactory<UAvaShapeSphereDynamicMesh>()
+	};
+}
+
+void UAvaShapesEditorShapeToolSphere::InitShape(UAvaShapeDynamicMeshBase* InShape) const
+{
+	UAvaShapeSphereDynamicMesh* Sphere = Cast<UAvaShapeSphereDynamicMesh>(InShape);
+	check(Sphere);
+
+	Sphere->SetSize3D({DefaultDim, DefaultDim, DefaultDim});
+
+	Super::InitShape(InShape);
+}
diff --git a/Engine/Plugins/Experimental/Avalanche/Source/AvalancheShapesEditor/Private/Tools/AvaShapesEditorShapeToolSphere.h b/Engine/Plugins/Experimental/Avalanche/Source/AvalancheShapesEditor/Private/Tools/AvaShapesEditorShapeToolSphere.h
new file mode 100644
--- /dev/null
+++ b/Engine/Plugins/Experimental/Avalanche/Source/AvalancheShapesEditor/Private/Tools/AvaShapesEditorShapeToolSphere.h
@@ -0,0 +1,25 @@
+// Copyright Epic Games, Inc. All Rights Reserved.
+
+#pragma once
+
+#include "Tools/AvaShapesEditorShapeAreaToolBase.h"
+#include "AvaShapesEditorShapeToolSphere.generated.h"
+
+UCLASS()
+class UAvaShapesEditorShapeToolSphere : public UAvaShapesEditorShapeAreaToolBase
+{
+	GENERATED_BODY()
+
+public:
+	UAvaShapesEditorShapeToolSphere();
+
+	//~ Begin UAvaInteractiveToolsToolBase
+	virtual FName GetCategoryName() override;
+	virtual FAvaInteractiveToolsToolParameters GetToolParameters() const override;
+	//~ End UAvaInteractiveToolsToolBase
+
+protected:
+	//~ Begin UAvaShapesEditorShapeToolBase
+	virtual void InitShape(UAvaShapeDynamicMeshBase* InShape) const override;
+	//~ End UAvaShapesEditorShapeToolBase
+};
diff --git a/Engine/Plugins/Experimental/Avalanche/Source/AvalancheShapesEditor/Private/Tools/AvaShapesEditorShapeToolStar.cpp b/Engine/Plugins/Experimental/Avalanche/Source/AvalancheShapesEditor/Private/Tools/AvaShapesEditorShapeToolStar.cpp
new file mode 100644
--- /dev/null
+++ b/Engine/Plugins/Experimental/Avalanche/Source/AvalancheShapesEditor/Private/Tools/AvaShapesEditorShapeToolStar.cpp
@@ -0,0 +1,44 @@
+// Copyright Epic Games, Inc. All Rights Reserved.
+
+#include "Tools/AvaShapesEditorShapeToolStar.h"
+#include "AvaShapesEditorCommands.h"
+#include "AvalancheShapesEditorModule.h"
+#include "Builders/AvaInteractiveToolsToolBuilder.h"
+#include "DynamicMeshes/AvaShapeStarDynMesh.h"
+#include "UObject/ConstructorHelpers.h"
+
+UAvaShapesEditorShapeToolStar::UAvaShapesEditorShapeToolStar()
+{
+	ShapeClass = UAvaShapeStarDynamicMesh::StaticClass();
+}
+
+FName UAvaShapesEditorShapeToolStar::GetCategoryName()
+{
+	return IAvalancheInteractiveToolsModule::CategoryName2D;
+}
+
+FAvaInteractiveToolsToolParameters UAvaShapesEditorShapeToolStar::GetToolParameters() const
+{
+	return {
+		FAvaShapesEditorCommands::Get().Tool_Shape_Star,
+		TEXT("Parametric Star Tool"),
+		6000,
+		FAvalancheInteractiveToolsCreateBuilder::CreateLambda(
+			[](UEdMode* InEdMode)
+			{
+				return UAvaInteractiveToolsToolBuilder::CreateToolBuilder<UAvaShapesEditorShapeToolStar>(InEdMode);
+			}),
+		nullptr,
+		CreateFactory<UAvaShapeStarDynamicMesh>()
+	};
+}
+
+void UAvaShapesEditorShapeToolStar::InitShape(UAvaShapeDynamicMeshBase* InShape) const
+{
+	UAvaShapeStarDynamicMesh* Star = Cast<UAvaShapeStarDynamicMesh>(InShape);
+	check(Star);
+
+	Star->SetSize2D({DefaultDim, DefaultDim});
+
+	Super::InitShape(InShape);
+}
diff --git a/Engine/Plugins/Experimental/Avalanche/Source/AvalancheShapesEditor/Private/Tools/AvaShapesEditorShapeToolStar.h b/Engine/Plugins/Experimental/Avalanche/Source/AvalancheShapesEditor/Private/Tools/AvaShapesEditorShapeToolStar.h
new file mode 100644
--- /dev/null
+++ b/Engine/Plugins/Experimental/Avalanche/Source/AvalancheShapesEditor/Private/Tools/AvaShapesEditorShapeToolStar.h
@@ -0,0 +1,25 @@
+// Copyright Epic Games, Inc. All Rights Reserved.
+
+#pragma once
+
+#include "Tools/AvaShapesEditorShapeAreaToolBase.h"
+#include "AvaShapesEditorShapeToolStar.generated.h"
+
+UCLASS()
+class UAvaShapesEditorShapeToolStar : public UAvaShapesEditorShapeAreaToolBase
+{
+	GENERATED_BODY()
+
+public:
+	UAvaShapesEditorShapeToolStar();
+
+	//~ Begin UAvaInteractiveToolsToolBase
+	virtual FName GetCategoryName() override;
+	virtual FAvaInteractiveToolsToolParameters GetToolParameters() const override;
+	//~ End UAvaInteractiveToolsToolBase
+
+protected:
+	//~ Begin UAvaShapesEditorShapeToolBase
+	virtual void InitShape(UAvaShapeDynamicMeshBase* InShape) const override;
+	//~ End UAvaShapesEditorShapeToolBase
+};
diff --git a/Engine/Plugins/Experimental/Avalanche/Source/AvalancheShapesEditor/Private/Tools/AvaShapesEditorShapeToolTorus.cpp b/Engine/Plugins/Experimental/Avalanche/Source/AvalancheShapesEditor/Private/Tools/AvaShapesEditorShapeToolTorus.cpp
new file mode 100644
--- /dev/null
+++ b/Engine/Plugins/Experimental/Avalanche/Source/AvalancheShapesEditor/Private/Tools/AvaShapesEditorShapeToolTorus.cpp
@@ -0,0 +1,44 @@
+// Copyright Epic Games, Inc. All Rights Reserved.
+
+#include "Tools/AvaShapesEditorShapeToolTorus.h"
+#include "AvaShapesEditorCommands.h"
+#include "AvalancheShapesEditorModule.h"
+#include "Builders/AvaInteractiveToolsToolBuilder.h"
+#include "DynamicMeshes/AvaShapeTorusDynMesh.h"
+#include "UObject/ConstructorHelpers.h"
+
+UAvaShapesEditorShapeToolTorus::UAvaShapesEditorShapeToolTorus()
+{
+	ShapeClass = UAvaShapeTorusDynamicMesh::StaticClass();
+}
+
+FName UAvaShapesEditorShapeToolTorus::GetCategoryName()
+{
+	return IAvalancheInteractiveToolsModule::CategoryName3D;
+}
+
+FAvaInteractiveToolsToolParameters UAvaShapesEditorShapeToolTorus::GetToolParameters() const
+{
+	return {
+		FAvaShapesEditorCommands::Get().Tool_Shape_Torus,
+		TEXT("Parametric Torus Tool"),
+		4000,
+		FAvalancheInteractiveToolsCreateBuilder::CreateLambda(
+			[](UEdMode* InEdMode)
+			{
+				return UAvaInteractiveToolsToolBuilder::CreateToolBuilder<UAvaShapesEditorShapeToolTorus>(InEdMode);
+			}),
+		nullptr,
+		CreateFactory<UAvaShapeTorusDynamicMesh>()
+	};
+}
+
+void UAvaShapesEditorShapeToolTorus::InitShape(UAvaShapeDynamicMeshBase* InShape) const
+{
+	UAvaShapeTorusDynamicMesh* Torus = Cast<UAvaShapeTorusDynamicMesh>(InShape);
+	check(Torus);
+
+	Torus->SetSize3D({DefaultDim, DefaultDim, DefaultDim});
+
+	Super::InitShape(InShape);
+}
diff --git a/Engine/Plugins/Experimental/Avalanche/Source/AvalancheShapesEditor/Private/Tools/AvaShapesEditorShapeToolTorus.h b/Engine/Plugins/Experimental/Avalanche/Source/AvalancheShapesEditor/Private/Tools/AvaShapesEditorShapeToolTorus.h
new file mode 100644
--- /dev/null
+++ b/Engine/Plugins/Experimental/Avalanche/Source/AvalancheShapesEditor/Private/Tools/AvaShapesEditorShapeToolTorus.h
@@ -0,0 +1,25 @@
+// Copyright Epic Games, Inc. All Rights Reserved.
+
+#pragma once
+
+#include "Tools/AvaShapesEditorShapeAreaToolBase.h"
+#include "AvaShapesEditorShapeToolTorus.generated.h"
+
+UCLASS()
+class UAvaShapesEditorShapeToolTorus : public UAvaShapesEditorShapeAreaToolBase
+{
+	GENERATED_BODY()
+
+public:
+	UAvaShapesEditorShapeToolTorus();
+
+	//~ Begin UAvaInteractiveToolsToolBase
+	virtual FName GetCategoryName() override;
+	virtual FAvaInteractiveToolsToolParameters GetToolParameters() const override;
+	//~ End UAvaInteractiveToolsToolBase
+
+protected:
+	//~ Begin UAvaShapesEditorShapeToolBase
+	virtual void InitShape(UAvaShapeDynamicMeshBase* InShape) const override;
+	//~ End UAvaShapesEditorShapeToolBase
+};
diff --git a/Engine/Plugins/Experimental/Avalanche/Source/AvalancheShapesEditor/Public/AvalancheShapesEditorModule.h b/Engine/Plugins/Experimental/Avalanche/Source/AvalancheShapesEditor/Public/AvalancheShapesEditorModule.h
new file mode 100644
--- /dev/null
+++ b/Engine/Plugins/Experimental/Avalanche/Source/AvalancheShapesEditor/Public/AvalancheShapesEditorModule.h
@@ -0,0 +1,51 @@
+// Copyright Epic Games, Inc. All Rights Reserved.
+
+#pragma once
+
+#include "Modules/ModuleInterface.h"
+#include "Containers/Array.h"
+#include "Delegates/IDelegateInstance.h"
+#include "Templates/SharedPointer.h"
+
+class FComponentVisualizer;
+class FUICommandInfo;
+class IAvaDetailsProvider;
+class IAvalancheInteractiveToolsModule;
+class IDetailsView;
+class UEdMode;
+struct FAvaInteractiveToolsToolParameters;
+
+class AVALANCHESHAPESEDITOR_API FAvalancheShapesEditorModule : public IModuleInterface
+{
+public:
+	static const FName BevelSprite;
+	static const FName BreakSideSprite;
+	static const FName ColorSelectionSprite;
+	static const FName CornerSprite;
+	static const FName DepthSprite;
+	static const FName InnerSizeSprite;
+	static const FName LinearGradientSprite;
+	static const FName NumPointsSprite;
+	static const FName NumSidesSprite;
+	static const FName TextMaxHeightSprite;
+	static const FName TextMaxWidthSprite;
+	static const FName TextScaleProportionallySprite;
+	static const FName SizeSprite;
+	static const FName SlantSprite;
+	static const FName UVSprite;
+
+	static FAvalancheShapesEditorModule& Get();
+
+	//~ Begin IModuleInterface interface
+	virtual void StartupModule() override;
+	virtual void ShutdownModule() override;
+	//~ End IModuleInterface interface
+
+private:
+	TArray<TSharedPtr<FComponentVisualizer>> Visualizers;
+
+	void RegisterShapeTools(IAvalancheInteractiveToolsModule* InModule);
+	void RegisterVisualizers();
+	
+	FDelegateHandle TrackEditorHandle;
+};
diff --git a/Engine/Plugins/Experimental/Avalanche/Source/AvalancheTag/Private/AvaTagCollection.cpp b/Engine/Plugins/Experimental/Avalanche/Source/AvalancheTag/Private/AvaTagCollection.cpp
new file mode 100644
--- /dev/null
+++ b/Engine/Plugins/Experimental/Avalanche/Source/AvalancheTag/Private/AvaTagCollection.cpp
@@ -0,0 +1,20 @@
+// Copyright Epic Games, Inc. All Rights Reserved.
+
+#include "AvaTagCollection.h"
+
+const FAvaTag* UAvaTagCollection::GetTag(const FAvaTagId& InTagId) const
+{
+	return Tags.Find(InTagId);
+}
+
+TArray<FAvaTagId> UAvaTagCollection::GetTagIds() const
+{
+	TArray<FAvaTagId> TagIds;
+	Tags.GetKeys(TagIds);
+	return TagIds;
+}
+
+FName UAvaTagCollection::GetTagMapName()
+{
+	return GET_MEMBER_NAME_CHECKED(UAvaTagCollection, Tags);
+}
diff --git a/Engine/Plugins/Experimental/Avalanche/Source/AvalancheTag/Private/AvaTagHandle.cpp b/Engine/Plugins/Experimental/Avalanche/Source/AvalancheTag/Private/AvaTagHandle.cpp
new file mode 100644
--- /dev/null
+++ b/Engine/Plugins/Experimental/Avalanche/Source/AvalancheTag/Private/AvaTagHandle.cpp
@@ -0,0 +1,62 @@
+// Copyright Epic Games, Inc. All Rights Reserved.
+
+#include "AvaTagHandle.h"
+#include "AvaTagCollection.h"
+
+const FAvaTag* FAvaTagHandle::GetTag() const
+{
+	if (::IsValid(Source))
+	{
+		return Source->GetTag(TagId);
+	}
+	return nullptr;
+}
+
+FString FAvaTagHandle::ToString() const
+{
+	if (const FAvaTag* Tag = GetTag())
+	{
+		return Tag->ToString();
+	}
+	return FString();
+}
+
+FString FAvaTagHandle::ToDebugString() const
+{
+	return FString::Printf(TEXT("TagId: %s, Source: %s"), *TagId.ToString(), *GetNameSafe(Source));
+}
+
+FName FAvaTagHandle::ToName() const
+{
+	if (const FAvaTag* Tag = GetTag())
+	{
+		return Tag->TagName;
+	}
+	return NAME_None;
+}
+
+void FAvaTagHandle::PostSerialize(const FArchive& Ar)
+{
+	if (Ar.IsSaving() && TagId.IsValid())
+	{
+		Ar.MarkSearchableName(FAvaTagId::StaticStruct(), *TagId.ToString());
+	}
+}
+
+bool FAvaTagHandle::MatchesTag(const FAvaTagHandle& InOther) const
+{
+	const FAvaTag* ThisTag  = GetTag();
+	const FAvaTag* OtherTag = InOther.GetTag();
+
+	if (ThisTag && OtherTag)
+	{
+		return *ThisTag == *OtherTag;
+	}
+
+	return MatchesExact(InOther);
+}
+
+bool FAvaTagHandle::MatchesExact(const FAvaTagHandle& InOther) const
+{
+	return Source == InOther.Source && TagId == InOther.TagId;
+}
diff --git a/Engine/Plugins/Experimental/Avalanche/Source/AvalancheTag/Public/AvaTag.h b/Engine/Plugins/Experimental/Avalanche/Source/AvalancheTag/Public/AvaTag.h
new file mode 100644
--- /dev/null
+++ b/Engine/Plugins/Experimental/Avalanche/Source/AvalancheTag/Public/AvaTag.h
@@ -0,0 +1,100 @@
+// Copyright Epic Games, Inc. All Rights Reserved.
+
+#pragma once
+
+#include "Containers/UnrealString.h"
+#include "Misc/Guid.h"
+#include "UObject/NameTypes.h"
+#include "AvaTag.generated.h"
+
+/**
+ * Struct for a Tag
+ * NOTE: Do not use as storage type. Prefer using FAvaTagHandle to then retrieve the FAvaTag it references
+ */
+USTRUCT(BlueprintType)
+struct FAvaTag
+{
+	GENERATED_BODY()
+
+	bool IsValid() const
+	{
+		return TagName != NAME_None;
+	}
+
+	bool operator==(const FAvaTag& InTag) const
+	{
+		return TagName == InTag.TagName;
+	}
+
+	friend uint32 GetTypeHash(const FAvaTag& InTag)
+	{
+		return GetTypeHash(InTag.TagName);
+	}
+
+	FString ToString() const
+	{
+		return TagName.ToString();
+	}
+
+	UPROPERTY(EditAnywhere, Category="Tag")
+	FName TagName;
+};
+
+/**
+ * Struct to identify a Tag. This is used by FAvaTagHandle to reference an FAvaTag
+ */
+USTRUCT()
+struct FAvaTagId
+{
+	GENERATED_BODY()
+
+	/** Id is initialized to a zero */
+	FAvaTagId() = default;
+
+	/**
+	 * Force init where Id is initialized to a new guid.
+	 * This is used in TCppStructOps::Construct
+	 */
+	explicit FAvaTagId(EForceInit)
+		: Id(FGuid::NewGuid())
+	{
+	}
+
+	explicit FAvaTagId(const FGuid& InId)
+		: Id(InId)
+	{
+	}
+
+	bool operator==(const FAvaTagId& InOther) const
+	{
+		return Id == InOther.Id;
+	}
+
+	bool IsValid() const
+	{
+		return Id.IsValid();
+	}
+
+	FString ToString() const
+	{
+		return Id.ToString();
+	}
+
+	friend uint32 GetTypeHash(const FAvaTagId& InTagId)
+	{
+		return GetTypeHash(InTagId.Id);
+	}
+
+private:
+	UPROPERTY(EditAnywhere, Category="Tag", meta=(IgnoreForMemberInitializationTest))
+	FGuid Id;
+};
+
+template<>
+struct TStructOpsTypeTraits<FAvaTagId> : TStructOpsTypeTraitsBase2<FAvaTagId>
+{
+	enum
+	{
+		WithNoInitConstructor = true,
+	};
+};
diff --git a/Engine/Plugins/Experimental/Avalanche/Source/AvalancheTag/Public/AvaTagCollection.h b/Engine/Plugins/Experimental/Avalanche/Source/AvalancheTag/Public/AvaTagCollection.h
new file mode 100644
--- /dev/null
+++ b/Engine/Plugins/Experimental/Avalanche/Source/AvalancheTag/Public/AvaTagCollection.h
@@ -0,0 +1,30 @@
+// Copyright Epic Games, Inc. All Rights Reserved.
+
+#pragma once
+
+#include "AvaTag.h"
+#include "Containers/Set.h"
+#include "UObject/Object.h"
+#include "AvaTagCollection.generated.h"
+
+/**
+ * Tag Collection that identifies a tag with an underlying Tag Id Guid
+ * and provides Tag reference capabilities
+ */
+UCLASS()
+class AVALANCHETAG_API UAvaTagCollection : public UObject
+{
+	GENERATED_BODY()
+
+public:
+	const FAvaTag* GetTag(const FAvaTagId& InTagId) const;
+
+	/** Returns the keys of the Tag Map */
+	TArray<FAvaTagId> GetTagIds() const;
+
+	static FName GetTagMapName();
+
+private:
+	UPROPERTY(EditAnywhere, Category="Tag")
+	TMap<FAvaTagId, FAvaTag> Tags;
+};
diff --git a/Engine/Plugins/Experimental/Avalanche/Source/AvalancheTag/Public/AvaTagHandle.h b/Engine/Plugins/Experimental/Avalanche/Source/AvalancheTag/Public/AvaTagHandle.h
new file mode 100644
--- /dev/null
+++ b/Engine/Plugins/Experimental/Avalanche/Source/AvalancheTag/Public/AvaTagHandle.h
@@ -0,0 +1,68 @@
+// Copyright Epic Games, Inc. All Rights Reserved.
+
+#pragma once
+
+#include "AvaTag.h"
+#include "Containers/ContainersFwd.h"
+#include "AvaTagHandle.generated.h"
+
+class UAvaTagCollection;
+
+/**
+ * Handle to a Tag in a particular Source.
+ * This should be used by the UStructs/UObjects to properly reference a particular FAvaTag.
+ */
+USTRUCT()
+struct AVALANCHETAG_API FAvaTagHandle
+{
+	GENERATED_BODY()
+
+	FAvaTagHandle() = default;
+
+	FAvaTagHandle(const UAvaTagCollection* InSource, const FAvaTagId& InTagId)
+		: Source(InSource)
+		, TagId(InTagId)
+	{
+	}
+
+	const FAvaTag* GetTag() const;
+
+	FString ToString() const;
+
+	FString ToDebugString() const;
+
+	FName ToName() const;
+
+	void PostSerialize(const FArchive& Ar);
+
+	/** Returns true if the Tag Handles resolve to same valued FAvaTags, even if the Source or Tag Id is different */
+	bool MatchesTag(const FAvaTagHandle& InOther) const;
+
+	/** Returns true if the Tag Handles is the exact same as the other (Same Source and Tag Id) */
+	bool MatchesExact(const FAvaTagHandle& InOther) const;
+
+	bool IsValid() const
+	{
+		return Source && TagId.IsValid();
+	}
+
+	friend uint32 GetTypeHash(const FAvaTagHandle& InHandle)
+	{
+		return HashCombineFast(GetTypeHash(InHandle.Source), GetTypeHash(InHandle.TagId));
+	}
+
+	UPROPERTY(EditAnywhere, Category="Tag")
+	TObjectPtr<const UAvaTagCollection> Source;
+
+	UPROPERTY()
+	FAvaTagId TagId;
+};
+
+template<>
+struct TStructOpsTypeTraits<FAvaTagHandle> : public TStructOpsTypeTraitsBase2<FAvaTagHandle>
+{
+	enum
+	{
+		WithPostSerialize = true,
+	};
+};
diff --git a/Engine/Plugins/Experimental/Avalanche/Source/AvalancheTag/Public/AvaTagHandleContainer.h b/Engine/Plugins/Experimental/Avalanche/Source/AvalancheTag/Public/AvaTagHandleContainer.h
new file mode 100644
--- /dev/null
+++ b/Engine/Plugins/Experimental/Avalanche/Source/AvalancheTag/Public/AvaTagHandleContainer.h
@@ -0,0 +1,62 @@
+// Copyright Epic Games, Inc. All Rights Reserved.
+
+#pragma once
+
+#include "AvaTag.h"
+#include "Containers/ContainersFwd.h"
+#include "AvaTagHandleContainer.generated.h"
+
+class UAvaTagCollection;
+struct FAvaTagHandle;
+
+/**
+ * Handle to a multiple tags in a particular Source.
+ * This should be used by the UStructs/UObjects to properly reference a multiple FAvaTags.
+ */
+USTRUCT()
+struct AVALANCHETAG_API FAvaTagHandleContainer
+{
+	GENERATED_BODY()
+
+	FAvaTagHandleContainer() = default;
+
+	explicit FAvaTagHandleContainer(const FAvaTagHandle& InTagHandle);
+
+	/** Returns true if the Tag Handles resolve to same valued FAvaTags, even if the Source or Tag Id is different */
+	bool ContainsTag(const FAvaTagHandle& InTagHandle) const;
+
+	/** Returns true if the Tag Handles is the exact same as the other (Same Source and Tag Id) */
+	bool ContainsTagHandle(const FAvaTagHandle& InTagHandle) const;
+
+	FString ToString() const;
+
+	void PostSerialize(const FArchive& Ar);
+
+	bool SerializeFromMismatchedTag(const FPropertyTag& InPropertyTag, FStructuredArchive::FSlot InSlot);
+
+	/** Adds the provided Tag Handle to TagIds. Returns true only if it was added as a new entry to TagIds */
+	bool AddTagHandle(const FAvaTagHandle& InTagHandle);
+
+	/** Removes the provided Tag Handle from TagIds. Returns true only if it existed and was removed from TagIds */
+	bool RemoveTagHandle(const FAvaTagHandle& InTagHandle);
+
+	/** Returns an array of resolved tags through this container's tag ids and source tag collection */
+	TArray<FAvaTag> ResolveTags() const;
+
+	UPROPERTY(EditAnywhere, Category="Tag")
+	TObjectPtr<const UAvaTagCollection> Source;
+
+private:
+	UPROPERTY(EditAnywhere, Category="Tag")
+	TArray<FAvaTagId> TagIds;
+};
+
+template<>
+struct TStructOpsTypeTraits<FAvaTagHandleContainer> : public TStructOpsTypeTraitsBase2<FAvaTagHandleContainer>
+{
+	enum
+	{
+		WithPostSerialize = true,
+		WithStructuredSerializeFromMismatchedTag = true,
+	};
+};
diff --git a/Engine/Plugins/Experimental/Avalanche/Source/AvalancheTagEditor/Private/AvaTagEditorStyle.h b/Engine/Plugins/Experimental/Avalanche/Source/AvalancheTagEditor/Private/AvaTagEditorStyle.h
new file mode 100644
--- /dev/null
+++ b/Engine/Plugins/Experimental/Avalanche/Source/AvalancheTagEditor/Private/AvaTagEditorStyle.h
@@ -0,0 +1,19 @@
+// Copyright Epic Games, Inc. All Rights Reserved.
+
+#pragma once
+
+#include "Styling/SlateStyle.h"
+
+class FAvaTagEditorStyle : public FSlateStyleSet
+{
+public:
+	FAvaTagEditorStyle();
+
+	virtual ~FAvaTagEditorStyle() override;
+
+	static FAvaTagEditorStyle& Get()
+	{
+		static FAvaTagEditorStyle StyleSet;
+		return StyleSet;
+	}
+};
diff --git a/Engine/Plugins/Experimental/Avalanche/Source/AvalancheTagEditor/Private/Customization/AvaTagCollectionCustomization.cpp b/Engine/Plugins/Experimental/Avalanche/Source/AvalancheTagEditor/Private/Customization/AvaTagCollectionCustomization.cpp
new file mode 100644
--- /dev/null
+++ b/Engine/Plugins/Experimental/Avalanche/Source/AvalancheTagEditor/Private/Customization/AvaTagCollectionCustomization.cpp
@@ -0,0 +1,170 @@
+// Copyright Epic Games, Inc. All Rights Reserved.
+
+#include "AvaTagCollectionCustomization.h"
+#include "AvaTagCollection.h"
+#include "DetailCategoryBuilder.h"
+#include "DetailLayoutBuilder.h"
+#include "Editor.h"
+#include "Framework/MultiBox/MultiBoxBuilder.h"
+#include "IDetailChildrenBuilder.h"
+#include "PropertyCustomizationHelpers.h"
+#include "Widgets/Input/SComboButton.h"
+#include "Widgets/SBoxPanel.h"
+
+#define LOCTEXT_NAMESPACE "AvaTagCollectionCustomization"
+
+void FAvaTagCollectionCustomization::CustomizeDetails(IDetailLayoutBuilder& InDetailBuilder)
+{
+	const TSharedPtr<IPropertyHandle> TagMapProperty = InDetailBuilder.GetProperty(UAvaTagCollection::GetTagMapName());
+	TagMapProperty->MarkHiddenByCustomization();
+
+	IDetailCategoryBuilder& Category = InDetailBuilder.EditCategory(TEXT("Tag"));
+
+	TSharedRef<SHorizontalBox> HeaderContentWidget = SNew(SHorizontalBox);
+	HeaderContentWidget->AddSlot()
+		.HAlign(HAlign_Right)
+		.VAlign(VAlign_Center)
+		[
+			TagMapProperty->CreateDefaultPropertyButtonWidgets()
+		];
+
+	Category.HeaderContent(HeaderContentWidget);
+	Category.AddCustomBuilder(MakeShared<FAvaTagMapBuilder>(TagMapProperty.ToSharedRef()), /*bForAdvanced*/false);
+}
+
+FAvaTagMapBuilder::FAvaTagMapBuilder(const TSharedRef<IPropertyHandle>& InTagMapProperty)
+	: MapProperty(InTagMapProperty->AsMap())
+	, BaseProperty(InTagMapProperty)
+{
+	check(MapProperty.IsValid());
+
+	// Delegate for when the number of children in the array changes
+	FSimpleDelegate OnNumChildrenChanged = FSimpleDelegate::CreateRaw(this, &FAvaTagMapBuilder::OnNumChildrenChanged);
+	OnNumElementsChangedHandle = MapProperty->SetOnNumElementsChanged(OnNumChildrenChanged);
+
+	BaseProperty->MarkHiddenByCustomization();
+}
+
+FAvaTagMapBuilder::~FAvaTagMapBuilder()
+{
+	MapProperty->UnregisterOnNumElementsChanged(OnNumElementsChangedHandle);
+}
+
+FName FAvaTagMapBuilder::GetName() const
+{
+	return BaseProperty->GetProperty()->GetFName();
+}
+
+void FAvaTagMapBuilder::GenerateChildContent(IDetailChildrenBuilder& InChildrenBuilder)
+{
+	uint32 NumChildren = 0;
+	BaseProperty->GetNumChildren(NumChildren);
+
+	for (uint32 ChildIndex = 0; ChildIndex < NumChildren; ++ChildIndex)
+	{
+		TSharedPtr<IPropertyHandle> TagHandle = BaseProperty->GetChildHandle(ChildIndex);
+		if (!TagHandle.IsValid())
+		{
+			continue;
+		}
+
+		if (TSharedPtr<IPropertyHandle> TagNameHandle = TagHandle->GetChildHandle(GET_MEMBER_NAME_CHECKED(FAvaTag, TagName)))
+		{
+			InChildrenBuilder.AddProperty(TagHandle.ToSharedRef())
+				.CustomWidget()
+				.WholeRowContent()
+				.HAlign(HAlign_Fill)
+				[
+					SNew(SHorizontalBox)
+					+ SHorizontalBox::Slot()
+					.FillWidth(1.f)
+					[
+						TagNameHandle->CreatePropertyValueWidget()
+					]
+					+ SHorizontalBox::Slot()
+					.AutoWidth()
+					.HAlign(HAlign_Center)
+					.VAlign(VAlign_Center)
+					[
+						CreatePropertyButtonsWidget(TagHandle)
+					]
+				];
+		}
+	}
+}
+
+void FAvaTagMapBuilder::OnNumChildrenChanged()
+{
+	OnRebuildChildren.ExecuteIfBound();
+}
+
+TSharedPtr<IPropertyHandle> FAvaTagMapBuilder::GetPropertyHandle() const
+{
+	return BaseProperty;
+}
+
+void FAvaTagMapBuilder::SetOnRebuildChildren(FSimpleDelegate InOnRebuildChildren)
+{
+	OnRebuildChildren = InOnRebuildChildren;
+}
+
+TSharedRef<SWidget> FAvaTagMapBuilder::CreatePropertyButtonsWidget(TSharedPtr<IPropertyHandle> InTagHandle)
+{
+	FMenuBuilder MenuContentBuilder(true, nullptr, nullptr, true);
+
+	MenuContentBuilder.AddMenuEntry(NSLOCTEXT("PropertyCustomizationHelpers", "DeleteButtonLabel", "Delete")
+		, FText::GetEmpty()
+		, FSlateIcon()
+		, FExecuteAction::CreateSP(this, &FAvaTagMapBuilder::DeleteItem, InTagHandle));
+
+	MenuContentBuilder.AddMenuEntry(LOCTEXT("SearchForReferencesLabel", "Search for References")
+		, FText::GetEmpty()
+		, FSlateIcon()
+		, FExecuteAction::CreateSP(this, &FAvaTagMapBuilder::SearchForReferences, InTagHandle->GetKeyHandle()));
+
+	return SNew(SComboButton)
+		.ComboButtonStyle(FAppStyle::Get(), "SimpleComboButton")
+		.ContentPadding(2)
+		.ForegroundColor(FSlateColor::UseForeground())
+		.HasDownArrow(true)
+		.MenuContent()
+		[
+			MenuContentBuilder.MakeWidget()
+		];
+}
+
+void FAvaTagMapBuilder::DeleteItem(TSharedPtr<IPropertyHandle> InTagHandle)
+{
+	if (InTagHandle.IsValid())
+	{
+		MapProperty->DeleteItem(InTagHandle->GetArrayIndex());	
+	}
+}
+
+void FAvaTagMapBuilder::SearchForReferences(TSharedPtr<IPropertyHandle> InTagIdHandle)
+{
+	if (!FEditorDelegates::OnOpenReferenceViewer.IsBound())
+	{
+		return;
+	}
+
+	if (!InTagIdHandle.IsValid())
+	{
+		return;
+	}
+
+	TArray<const void*> TagIdRawData;
+	InTagIdHandle->AccessRawData(TagIdRawData);
+	if (TagIdRawData.IsEmpty() || !TagIdRawData[0])
+	{
+		return;
+	}
+
+	const FAvaTagId& TagId = *static_cast<const FAvaTagId*>(TagIdRawData[0]);
+
+	TArray<FAssetIdentifier> AssetIdentifiers;
+	AssetIdentifiers.Emplace(FAvaTagId::StaticStruct(), *TagId.ToString());
+	FEditorDelegates::OnOpenReferenceViewer.Broadcast(AssetIdentifiers, FReferenceViewerParams());
+}
+
+#undef LOCTEXT_NAMESPACE
diff --git a/Engine/Plugins/Experimental/Avalanche/Source/AvalancheTagEditor/Private/Customization/AvaTagCollectionCustomization.h b/Engine/Plugins/Experimental/Avalanche/Source/AvalancheTagEditor/Private/Customization/AvaTagCollectionCustomization.h
new file mode 100644
--- /dev/null
+++ b/Engine/Plugins/Experimental/Avalanche/Source/AvalancheTagEditor/Private/Customization/AvaTagCollectionCustomization.h
@@ -0,0 +1,54 @@
+// Copyright Epic Games, Inc. All Rights Reserved.
+
+#pragma once
+
+#include "Delegates/Delegate.h"
+#include "IDetailCustomNodeBuilder.h"
+#include "IDetailCustomization.h"
+#include "Templates/SharedPointer.h"
+
+class IPropertyHandleMap;
+class SWidget;
+
+class FAvaTagCollectionCustomization : public IDetailCustomization
+{
+public:
+	//~ Begin IDetailCustomization
+	virtual void CustomizeDetails(IDetailLayoutBuilder& InDetailBuilder) override;
+	//~ End IDetailCustomization
+};
+
+class FAvaTagMapBuilder : public IDetailCustomNodeBuilder, public TSharedFromThis<FAvaTagMapBuilder>
+{
+public:
+	FAvaTagMapBuilder(const TSharedRef<IPropertyHandle>& InTagMapProperty);
+
+	virtual ~FAvaTagMapBuilder() override;
+
+	//~ Begin IDetailCustomNodeBuilder
+	virtual FName GetName() const override;
+	virtual bool InitiallyCollapsed() const override { return false; }
+	virtual void GenerateHeaderRowContent(FDetailWidgetRow& InNodeRow) override {}
+	virtual void GenerateChildContent(IDetailChildrenBuilder& InChildrenBuilder) override;
+	virtual TSharedPtr<IPropertyHandle> GetPropertyHandle() const override;
+	virtual void SetOnRebuildChildren(FSimpleDelegate InOnRebuildChildren) override;
+	//~ End IDetailCustomNodeBuilder
+
+private:
+	TSharedRef<SWidget> CreatePropertyButtonsWidget(TSharedPtr<IPropertyHandle> InTagHandle);
+
+	void DeleteItem(TSharedPtr<IPropertyHandle> InTagHandle);
+
+	void SearchForReferences(TSharedPtr<IPropertyHandle> InTagIdHandle);
+
+	void OnNumChildrenChanged();
+
+	TSharedPtr<IPropertyHandleMap> MapProperty;
+
+	TSharedPtr<IPropertyHandle> BaseProperty;
+
+	FDelegateHandle OnNumElementsChangedHandle;
+
+	FSimpleDelegate OnRebuildChildren;
+};
+
diff --git a/Engine/Plugins/Experimental/Avalanche/Source/AvalancheTest/AvalancheTest.Build.cs b/Engine/Plugins/Experimental/Avalanche/Source/AvalancheTest/AvalancheTest.Build.cs
new file mode 100644
--- /dev/null
+++ b/Engine/Plugins/Experimental/Avalanche/Source/AvalancheTest/AvalancheTest.Build.cs
@@ -0,0 +1,41 @@
+﻿// Copyright Epic Games, Inc. All Rights Reserved.
+
+using UnrealBuildTool;
+
+public class AvalancheTest : ModuleRules
+{
+    public AvalancheTest(ReadOnlyTargetRules Target) : base(Target)
+    {
+        PCHUsage = ModuleRules.PCHUsageMode.UseExplicitOrSharedPCHs;
+
+        PublicDependencyModuleNames.AddRange(
+            new string[]
+            {
+                "Core",
+                "FunctionalTesting",
+            }
+        );
+
+        PrivateDependencyModuleNames.AddRange(
+            new string[]
+            {
+	            "Avalanche",
+                "CoreUObject",
+                "Engine",
+                "Slate",
+                "SlateCore",
+            }
+        );
+        
+        if (Target.Type == TargetType.Editor)
+        {
+	        PrivateDependencyModuleNames.AddRange(
+		        new string[]
+		        {
+			        "AvalancheEditor",
+			        "LevelEditor",
+			        "UnrealEd",
+		        });
+        }
+    }
+}
diff --git a/Engine/Plugins/Experimental/Avalanche/Source/AvalancheTest/Private/AvaScreenshotTestActor.cpp b/Engine/Plugins/Experimental/Avalanche/Source/AvalancheTest/Private/AvaScreenshotTestActor.cpp
new file mode 100644
--- /dev/null
+++ b/Engine/Plugins/Experimental/Avalanche/Source/AvalancheTest/Private/AvaScreenshotTestActor.cpp
@@ -0,0 +1,176 @@
+﻿// Copyright Epic Games, Inc. All Rights Reserved.
+
+#include "AvaScreenshotTestActor.h"
+
+#include "ActorFolder.h"
+#include "AutomationBlueprintFunctionLibrary.h"
+#include "AvaBlueprint.h"
+#include "AvaTestBlueprintFunctionLibrary.h"
+#include "AvaTestModule.h"
+#include "Camera/CameraActor.h"
+#include "Camera/CameraComponent.h"
+#include "EngineUtils.h"
+#include "Kismet/GameplayStatics.h"
+#include "LevelInstance/LevelInstanceActor.h"
+#include "Tests/AutomationCommon.h"
+#include "UObject/Package.h"
+
+#define LOCTEXT_NAMESPACE "AvaScreenshotTestActor"
+
+namespace UE::Ava::Test::Private
+{
+	/** Finds and returns the first camera in the specified world. */
+	ACameraActor* FindFirstCameraInWorld(const UWorld* InWorld)
+	{
+		for (ACameraActor* Camera : TActorRange<ACameraActor>(InWorld, ACameraActor::StaticClass(), EActorIteratorFlags::AllActors))
+		{
+			return Camera;
+		}
+
+		return nullptr;
+	}
+	
+}
+
+AAvaScreenshotTestActor::AAvaScreenshotTestActor(const FObjectInitializer& ObjectInitializer)
+	: Super(ObjectInitializer)
+{
+	ScreenshotOptions.Resolution = { 1920, 1080 };
+}
+
+void AAvaScreenshotTestActor::GetActorEyesViewPoint(FVector& OutLocation, FRotator& OutRotation) const
+{
+	Super::GetActorEyesViewPoint(OutLocation, OutRotation);
+
+	if (ScreenshotCamera)
+	{
+		FMinimalViewInfo ViewInfo;
+		ScreenshotCamera->GetCameraView(0.0f, ViewInfo);
+
+		OutLocation = ViewInfo.Location;
+		OutRotation = ViewInfo.Rotation;
+	}
+}
+
+void AAvaScreenshotTestActor::SetupTest()
+{
+	const UAvalancheBlueprint* ResolvedBlueprintToTest = BlueprintToTest.LoadSynchronous();
+	const bool bHasValidBP = ::IsValid(ResolvedBlueprintToTest);
+
+	// Create a world (.umap) asset for the Ava BP, overwriting any existing
+	if (bHasValidBP)
+	{
+		UAvaTestBlueprintFunctionLibrary::ExportAvaBlueprintsToWorld({ ResolvedBlueprintToTest });
+	}
+
+	ReceiveSetupTest();
+}
+
+void AAvaScreenshotTestActor::PrepareTest()
+{
+	const UAvalancheBlueprint* ResolvedBlueprintToTest = BlueprintToTest.LoadSynchronous();
+	const bool bHasValidBP = ::IsValid(ResolvedBlueprintToTest);
+	if (bHasValidBP)
+	{
+		if (UWorld* CurrentWorld = GetWorld())
+		{
+			// Spawn a LevelInstance to contain the previously created World for the specified BP
+			ALevelInstance* NewLevelInstanceActor =	CurrentWorld->SpawnActor<ALevelInstance>();
+			if (!NewLevelInstanceActor)
+			{
+				UE_LOG(LogAvaTest, Error, TEXT("There was an error spawning ALevelInstance"));
+			}
+			else
+			{
+				const UPackage* const BlueprintPackage = ResolvedBlueprintToTest->GetPackage();
+				const FString Path = FPackageName::GetLongPackagePath(BlueprintPackage->GetName()) + TEXT("/Exported/") + ResolvedBlueprintToTest->GetName();
+
+				TSoftObjectPtr<UWorld> WorldCreatedByBP(Path);
+				ensure(!WorldCreatedByBP.IsNull());
+
+				// Set the World (umap) of the LevelInstance and wait for it to load
+				NewLevelInstanceActor->SetWorldAsset(WorldCreatedByBP);
+				NewLevelInstanceActor->GetLevelInstanceSubsystem()->BlockLoadLevelInstance(NewLevelInstanceActor);
+				GEngine->BlockTillLevelStreamingCompleted(NewLevelInstanceActor->GetWorld());
+
+				// Flag the LevelInstance for destruction when this test is complete
+				AutoDestroyActors.Add(NewLevelInstanceActor);
+
+				// Find and use the first camera in the newly loaded level, if it exists
+				UWorld* LevelInstanceWorld = CurrentWorld;
+				if (ACameraActor* FirstCameraInLevelInstance = UE::Ava::Test::Private::FindFirstCameraInWorld(LevelInstanceWorld))
+				{
+					DefaultScreenshotCamera = ScreenshotCamera;
+					ScreenshotCamera = FirstCameraInLevelInstance->GetCameraComponent();
+				}
+				else
+				{
+					UE_LOG(LogAvaTest, Display, TEXT("No camera was found in the loaded Ava BP, using default (in Test Actor)."));
+				}
+			}
+		}
+	}
+	
+	// Only set if ScreenshotCamera is different to the Default (and thus the view target should be a different actor)
+	if (DefaultScreenshotCamera && ScreenshotCamera != DefaultScreenshotCamera)
+	{
+		UGameViewportClient* GameViewportClient = AutomationCommon::GetAnyGameViewportClient();
+		check(GameViewportClient);
+		
+		APlayerController* PlayerController = UGameplayStatics::GetPlayerController(GameViewportClient->GetWorld(), 0);
+		if (PlayerController)
+		{
+			if (ACameraActor* CameraActor = Cast<ACameraActor>(ScreenshotCamera->GetOwner()))
+			{
+				// @note: we're only copying the camera's FOV and transform, other settings (ie. lens) are not copied
+				PlayerController->FOV(ScreenshotCamera->FieldOfView);
+				SetActorRelativeTransform(CameraActor->GetActorTransform());	
+			}
+		}
+	}
+
+	Super::PrepareTest();
+}
+
+void AAvaScreenshotTestActor::FinishTest(EFunctionalTestResult TestResult, const FString& Message)
+{
+	// Restore camera
+	ScreenshotCamera = DefaultScreenshotCamera;
+	
+	// Set actors to destroy before calling Super (which actually destroys them)
+	{
+		static const FString ProtectedFolderName = TEXT("DontDelete");
+
+		// Forcibly deletes all actors outside of a reserved folder (see above) on test completion
+		if (bClearLevelOnComplete)
+		{
+			if (const UWorld* World = GetWorld())
+			{
+				TArray<AActor*> ActorsToDestroy;
+				ActorsToDestroy.Reserve(World->GetActorCount());
+
+				for (AActor* const Actor : TActorRange<AActor>(World))
+				{
+					FFolder TopLevelActorFolder = Actor->GetFolder();
+					FString TopLevelActorFolderName = TopLevelActorFolder.IsNone() || !TopLevelActorFolder.GetActorFolder()
+						? TEXT("")
+						: TopLevelActorFolder.GetActorFolder()->GetLabel();
+
+					// Don't destroy an actor in this folder
+					if (TopLevelActorFolderName.Equals(ProtectedFolderName, ESearchCase::IgnoreCase))
+					{
+						continue;
+					}
+
+					ActorsToDestroy.Add(Actor);
+				}
+
+				AutoDestroyActors.Append(ActorsToDestroy);
+			}		
+		}
+	}
+
+	Super::FinishTest(TestResult, Message);
+}
+
+#undef LOCTEXT_NAMESPACE
diff --git a/Engine/Plugins/Experimental/Avalanche/Source/AvalancheTest/Private/AvaTestBlueprintFunctionLibrary.cpp b/Engine/Plugins/Experimental/Avalanche/Source/AvalancheTest/Private/AvaTestBlueprintFunctionLibrary.cpp
new file mode 100644
--- /dev/null
+++ b/Engine/Plugins/Experimental/Avalanche/Source/AvalancheTest/Private/AvaTestBlueprintFunctionLibrary.cpp
@@ -0,0 +1,31 @@
+﻿// Copyright Epic Games, Inc. All Rights Reserved.
+
+#include "AvaTestBlueprintFunctionLibrary.h"
+
+#include "AvaBlueprint.h"
+
+#if WITH_EDITOR
+#include "AvaEditorFunctionLibrary.h"
+#endif
+
+void UAvaTestBlueprintFunctionLibrary::ExportAvaBlueprintsToWorld(const TArray<TSoftObjectPtr<UAvalancheBlueprint>>& InBlueprints)
+{
+#if WITH_EDITOR
+	TArray<TWeakObjectPtr<UAvalancheBlueprint>> WeakBlueprints;
+	Algo::TransformIf(
+		InBlueprints,
+		WeakBlueprints,
+		[](const TSoftObjectPtr<UAvalancheBlueprint>& InBlueprint)
+		{
+			return InBlueprint.IsValid() || InBlueprint.LoadSynchronous() != nullptr;
+		},
+		[](const TSoftObjectPtr<UAvalancheBlueprint>& InBlueprint)
+		{
+			return MakeWeakObjectPtr(InBlueprint.Get());			
+		});
+	
+	FAvaEditorFunctionLibrary::ExportAvaBlueprintsToWorld(WeakBlueprints);
+#else
+	UE_LOG(LogAvaTest, Error, TEXT("ExportAvaBlueprintsToWorld is editor-only"));
+#endif
+}
diff --git a/Engine/Plugins/Experimental/Avalanche/Source/AvalancheTest/Private/AvaTestModule.cpp b/Engine/Plugins/Experimental/Avalanche/Source/AvalancheTest/Private/AvaTestModule.cpp
new file mode 100644
--- /dev/null
+++ b/Engine/Plugins/Experimental/Avalanche/Source/AvalancheTest/Private/AvaTestModule.cpp
@@ -0,0 +1,134 @@
+﻿// Copyright Epic Games, Inc. All Rights Reserved.
+
+#include "AvaTestModule.h"
+
+#include "AvaScreenshotTestActor.h"
+#include "Editor.h"
+#include "EngineUtils.h"
+#include "FunctionalTestBase.h"
+#include "LevelEditorSubsystem.h"
+#include "Modules/ModuleManager.h"
+
+DEFINE_LOG_CATEGORY(LogAvaTest);
+
+#define LOCTEXT_NAMESPACE "AvalancheTest"
+
+namespace UE::AvaTest::Private
+{
+	static constexpr const TCHAR* AutomationTestFrameworkModuleName = TEXT("AutomationWorker");
+
+	// Duplicate of ClientFuncTestPerforming.cpp ParseTestMapInfo
+	static void ParseTestMapInfo(const FString& InParameters, FString& OutMapObjectPath, FString& OutMapPackageName, FString& OutMapTestName)
+	{
+		TArray<FString> ParamArray;
+		InParameters.ParseIntoArray(ParamArray, TEXT(";"), true);
+
+		OutMapObjectPath = ParamArray[0];
+		OutMapPackageName = ParamArray[1];
+		OutMapTestName = (ParamArray.Num() > 2) ? ParamArray[2] : TEXT("");
+	}
+
+	static AAvaScreenshotTestActor* FindActorForTestName(UWorld* InWorld, const FString& InTestName)
+	{
+		check(InWorld);
+
+		for (const TWeakObjectPtr<AAvaScreenshotTestActor> Actor : TActorRange<AAvaScreenshotTestActor>(InWorld))
+		{
+			if (Actor->GetName() == InTestName)
+			{
+				return Actor.Get();
+			}
+		}
+
+		return nullptr;
+	}
+}
+
+void FAvaTestModule::StartupModule()
+{
+	FAutomationTestFramework::Get().OnTestStartEvent.AddRaw(this, &FAvaTestModule::OnTestStart);
+}
+
+void FAvaTestModule::ShutdownModule()
+{
+	if (FModuleManager::Get().IsModuleLoaded(UE::AvaTest::Private::AutomationTestFrameworkModuleName))
+	{
+		FAutomationTestFramework::Get().OnTestStartEvent.RemoveAll(this);
+	}
+}
+
+void FAvaTestModule::OnTestStart(FAutomationTestBase* InTest)
+{
+	check(InTest); // This shouldn't ever happen
+
+#if WITH_EDITOR
+	ensure(!GIsPlayInEditorWorld); // Shouldn't be playing!
+#endif
+
+	// If it's one of our tests, run pre-test setup
+	if (InTest && InTest->GetTestName().StartsWith(TEXT("FClientFunctionalTesting")))
+	{
+#if WITH_EDITOR
+		ULevelEditorSubsystem* LevelEditorSubsystem = GEditor->GetEditorSubsystem<ULevelEditorSubsystem>();
+		const ULevel* CurrentLevel = LevelEditorSubsystem->GetCurrentLevel();
+#endif
+
+		const FFunctionalTestBase* FunctionalTest = static_cast<FFunctionalTestBase*>(InTest);
+
+		TArray<FAutomationTestInfo> TestsToRun;
+		FunctionalTest->GenerateTestNames(TestsToRun);
+
+		for (const FAutomationTestInfo& TestInfo : TestsToRun)
+		{
+			FString MapObjectPath;
+			FString MapPackageName;
+			FString MapTestName; // Usually this is the actual actor name (not outliner label!)
+			UE::AvaTest::Private::ParseTestMapInfo(TestInfo.GetTestParameter(), MapObjectPath, MapPackageName, MapTestName);
+
+			// @note: this MUST occur in-editor
+#if WITH_EDITOR
+			// 1. load map (if not already current)
+
+			// @todo: need to prevent new, unsaved levels from being saved
+			if (!CurrentLevel->GetPathName().IsEmpty())
+			{
+				LevelEditorSubsystem->SaveCurrentLevel();
+			}
+
+			// Only load if not current
+			if (CurrentLevel->GetPathName() != MapObjectPath)
+			{
+				if (!LevelEditorSubsystem->LoadLevel(MapObjectPath))
+				{
+					UE_LOG(LogAvaTest, Error, TEXT("There was an error loading the specified level: %s"), *MapObjectPath);
+					continue;
+				}
+
+				CurrentLevel = LevelEditorSubsystem->GetCurrentLevel();
+			}
+#endif
+
+			// 2. locate actor for testname
+			UWorld* World = CurrentLevel->GetWorld();
+			AAvaScreenshotTestActor* TestActor = UE::AvaTest::Private::FindActorForTestName(World, MapTestName);
+			if (!ensure(TestActor))
+			{
+				UE_LOG(LogAvaTest, Error, TEXT("The TestActor for test \"%s\" was not found"), *MapTestName);
+				continue;
+			}
+			
+			// 3. if actor found, run pre-test func - @todo: support latent/async actions within?
+			TestActor->SetupTest();
+
+#if WITH_EDITOR
+			// 4. save map if dirtied
+			LevelEditorSubsystem->SaveCurrentLevel();
+#endif
+		}
+	}
+}
+
+#undef LOCTEXT_NAMESPACE
+
+IMPLEMENT_MODULE(FAvaTestModule, AvalancheTest)
+
diff --git a/Engine/Plugins/Experimental/Avalanche/Source/AvalancheTest/Private/AvaTestModule.h b/Engine/Plugins/Experimental/Avalanche/Source/AvalancheTest/Private/AvaTestModule.h
new file mode 100644
--- /dev/null
+++ b/Engine/Plugins/Experimental/Avalanche/Source/AvalancheTest/Private/AvaTestModule.h
@@ -0,0 +1,23 @@
+﻿// Copyright Epic Games, Inc. All Rights Reserved.
+
+#pragma once
+
+#include "Logging/LogMacros.h"
+#include "Modules/ModuleInterface.h"
+
+DECLARE_LOG_CATEGORY_EXTERN(LogAvaTest, Log, All);
+
+/**
+ * Avalanche Test Module
+ */
+class FAvaTestModule : public IModuleInterface
+{
+public:
+	//~ Begin IAvaTestModule
+	virtual void StartupModule() override;
+	virtual void ShutdownModule() override;
+	//~ End IAvaTestModule
+
+private:
+	void OnTestStart(FAutomationTestBase* InTest);
+};
diff --git a/Engine/Plugins/Experimental/Avalanche/Source/AvalancheTest/Public/AvaScreenshotTestActor.h b/Engine/Plugins/Experimental/Avalanche/Source/AvalancheTest/Public/AvaScreenshotTestActor.h
new file mode 100644
--- /dev/null
+++ b/Engine/Plugins/Experimental/Avalanche/Source/AvalancheTest/Public/AvaScreenshotTestActor.h
@@ -0,0 +1,52 @@
+﻿// Copyright Epic Games, Inc. All Rights Reserved.
+
+#pragma once
+
+#include "ScreenshotFunctionalTest.h"
+
+#include "AvaScreenshotTestActor.generated.h"
+
+class UAvalancheBlueprint;
+
+/**  */
+UCLASS(Blueprintable, BlueprintType, HideCategories = ("Replication", "Collision", "Physics", "Networking", "HLOD", "LevelInstance", "Cooking"))
+class AVALANCHETEST_API AAvaScreenshotTestActor
+	: public AScreenshotFunctionalTest
+{
+	GENERATED_BODY()
+
+public:
+	AAvaScreenshotTestActor(const FObjectInitializer& ObjectInitializer);
+
+	virtual void GetActorEyesViewPoint(FVector& OutLocation, FRotator& OutRotation) const override;
+
+	/** Optionally specify the Ava BP to test (and spawn in this level). */
+	UPROPERTY(EditAnywhere, Category = "Ava", meta = (EditFixedOrder))
+	TSoftObjectPtr<UAvalancheBlueprint> BlueprintToTest;
+
+	/** If true, tests the alpha channel output (only!). */
+	UPROPERTY(EditAnywhere, Category = "Ava")
+	bool bTestAlphaOnly = false;
+
+	/** If true, all actors outside of the "DontDelete" folder are destroyed on test completion. */
+	UPROPERTY(EditAnywhere, Category = "Ava")
+	bool bClearLevelOnComplete = true;
+
+	/** Optionally implement to perform a pre-play, editor only setup step. Good for asset generation or modification that breaks in PIE. */
+	virtual void SetupTest();
+
+	/** Optionally implement to perform a pre-play, editor only setup step. Good for asset generation or modification that breaks in PIE. */
+	UFUNCTION(BlueprintImplementableEvent, meta = (DisplayName="SetupTest"))
+	void ReceiveSetupTest();
+
+protected:
+	//~ Begin AFunctionalTest
+	virtual void PrepareTest() override;
+	virtual void FinishTest(EFunctionalTestResult TestResult, const FString& Message) override;
+	//~ End AFunctionalTest
+
+private:
+	/** Stores the default camera for restoration after another is temporarily used. */
+	UPROPERTY(Transient)
+	TObjectPtr<class UCameraComponent> DefaultScreenshotCamera;
+};
diff --git a/Engine/Plugins/Experimental/Avalanche/Source/AvalancheTest/Public/AvaTestBlueprintFunctionLibrary.h b/Engine/Plugins/Experimental/Avalanche/Source/AvalancheTest/Public/AvaTestBlueprintFunctionLibrary.h
new file mode 100644
--- /dev/null
+++ b/Engine/Plugins/Experimental/Avalanche/Source/AvalancheTest/Public/AvaTestBlueprintFunctionLibrary.h
@@ -0,0 +1,27 @@
+﻿// Copyright Epic Games, Inc. All Rights Reserved.
+
+#pragma once
+
+#include "Kismet/BlueprintFunctionLibrary.h"
+
+#include "AvaTestBlueprintFunctionLibrary.generated.h"
+
+class UAvalancheBlueprint;
+
+/**
+ * 
+ */
+UCLASS()
+class AVALANCHETEST_API UAvaTestBlueprintFunctionLibrary : public UBlueprintFunctionLibrary
+{
+	GENERATED_BODY()
+
+public:
+	/**
+	 * Creates World Assets based on the Avalanche Blueprints provided
+	 * Wraps FAvaEditorFunctionLibrary::ExportAvaBlueprintsToWorld
+	 * @param InBlueprints the Avalanche Blueprints to export
+	 */
+	UFUNCTION(BlueprintCallable, Category = "Motion Design")
+	static void ExportAvaBlueprintsToWorld(const TArray<TSoftObjectPtr<UAvalancheBlueprint>>& InBlueprints);
+};
