[MotionDesign]

diff --git a/Engine/Plugins/Experimental/Avalanche/Source/AvalancheText/Private/AvaText3DComponent.cpp b/Engine/Plugins/Experimental/Avalanche/Source/AvalancheText/Private/AvaText3DComponent.cpp
new file mode 100644
--- /dev/null
+++ b/Engine/Plugins/Experimental/Avalanche/Source/AvalancheText/Private/AvaText3DComponent.cpp
@@ -0,0 +1,1621 @@
+// Copyright Epic Games, Inc. All Rights Reserved.
+
+#include "AvaText3DComponent.h"
+
+#include "AvaActorUtils.h"
+#include "AvaLog.h"
+#include "Engine/Texture2D.h"
+#include "GameFramework/Actor.h"
+#include "MaterialHub/AvaTextMaterialHub.h"
+#include "Materials/Material.h"
+#include "Materials/MaterialInstanceDynamic.h"
+
+#if WITH_EDITOR
+#include "Editor.h"
+#endif
+
+namespace UE::AvaText::Private
+{
+	void PatchMaterial(
+		UText3DComponent* InComponent
+		, TMemFunPtrType<true, UText3DComponent, UMaterialInterface*()>::Type&& InGetter
+		, TMemFunPtrType<false, UText3DComponent, void(UMaterialInterface*)>::Type&& InSetter
+		, UMaterialInstanceDynamic* InNewMaterial)
+	{
+		if (!InNewMaterial)
+		{
+			return;
+		}
+		
+		if (UMaterialInstance* CurrentMaterialInstance = Cast<UMaterialInstance>(Invoke(InGetter, InComponent)))
+		{
+			InNewMaterial->CopyInterpParameters(CurrentMaterialInstance);
+		}
+
+		Invoke(InSetter, InComponent, InNewMaterial);
+	}
+}
+
+// Sets default values for this component's properties
+UAvaText3DComponent::UAvaText3DComponent()
+{
+	PrimaryComponentTick.bCanEverTick = true;
+	
+	Extrude = 0.0f;
+	Bevel = 0.0f;
+	BevelType = EText3DBevelType::Convex;
+	BevelSegments = 8;
+	
+	HorizontalAlignment = EText3DHorizontalTextAlignment::Left;
+	VerticalAlignment = EText3DVerticalTextAlignment::FirstLine;
+
+	bHasMaxWidth = false;
+	bHasMaxHeight = false;
+	bScaleProportionally = false;
+	
+	Kerning = 0.0f;
+	WordSpacing = 0.0f;
+	LineSpacing = 0.0f;
+	
+	MaxWidth = 100.0f;
+	MaxHeight = 100.0f;
+
+	AvalancheFont = {};
+
+	Tiling = FVector2D(1.0, 1.0);
+	
+	Color = FLinearColor::White;
+	ExtrudeColor = FLinearColor::Gray;
+	BevelColor = FLinearColor::Gray;
+
+	GradientSettings.Direction = EAvaGradientDirection::Vertical;
+	GradientSettings.Smoothness = 0.1;
+
+	bIsUnlit = true;
+	TranslucencyStyle = EAvaTextTranslucency::None;
+
+	bEnforceUpperCase = false;
+
+	MarkForFullRefresh();
+
+#if WITH_EDITOR
+	RegisterOnPropertyChangeFunctions();
+#endif
+}
+
+void UAvaText3DComponent::PostInitProperties()
+{
+	Super::PostInitProperties();
+	
+	SetupParentRootAndCallbacks();
+
+	MarkForFullRefresh();
+}
+
+void UAvaText3DComponent::RetrieveDataFromBaseComponent()
+{
+	RetrieveCurrentColoringStyle();	
+	RetrieveCurrentMaterials();	
+	RetrieveCurrentTextAndFont();
+	RetrieveCurrentAlignment();
+}
+
+void UAvaText3DComponent::PostLoad()
+{
+	Super::PostLoad();
+
+	SetupParentRootAndCallbacks();
+	RetrieveDataFromBaseComponent();
+}
+
+void UAvaText3DComponent::PostDuplicateActions()
+{
+	SetupParentRootAndCallbacks();
+	
+	RetrieveCurrentTextAndFont();
+	
+	MarkForFullRefresh();
+	RefreshText3D();
+}
+
+void UAvaText3DComponent::PostDuplicate(EDuplicateMode::Type DuplicateMode)
+{	
+	Super::PostDuplicate(DuplicateMode);
+
+	PostDuplicateActions();
+}
+
+// overriding PostEditImport to handle Owner duplication
+void UAvaText3DComponent::PostEditImport()
+{
+	Super::PostEditImport();
+	
+	PostDuplicateActions();
+}
+
+#if WITH_EDITOR
+void UAvaText3DComponent::PostEditChangeProperty(FPropertyChangedEvent& PropertyChangedEvent)
+{
+	Super::PostEditChangeProperty(PropertyChangedEvent);
+	
+	if (PropertyChangedEvent.Property)
+	{
+		const FName ChangedPropertyName = PropertyChangedEvent.MemberProperty->GetFName();
+
+		if (PropertyChangeFunctions.Contains(ChangedPropertyName))
+		{
+			PropertyChangeFunctions[ChangedPropertyName](this);
+		}
+	}
+}
+
+void UAvaText3DComponent::PostEditComponentMove(bool bFinished)
+{
+	Super::PostEditComponentMove(bFinished);
+
+	OnEditMoveDelegate.Broadcast();
+	
+	ScheduleTextRefreshOnTicker();
+}
+
+#endif
+
+void UAvaText3DComponent::OnUpdateTransform(EUpdateTransformFlags UpdateTransformFlags, ETeleportType Teleport)
+{
+	Super::OnUpdateTransform(UpdateTransformFlags, Teleport);
+
+	ScheduleTextRefreshOnTicker();
+}
+
+void UAvaText3DComponent::SetupParentRootAndCallbacks()
+{
+	RefreshParentRootComponent();
+	RegisterCallbacks();
+}
+
+void UAvaText3DComponent::OnAttachmentChanged()
+{
+	Super::OnAttachmentChanged();
+
+	SetupParentRootAndCallbacks();
+}
+
+void UAvaText3DComponent::OnRegister()
+{
+	Super::OnRegister();
+
+	SetComponentTickEnabled(false);
+	
+	SetupParentRootAndCallbacks();
+	RefreshText3D();
+}
+
+void UAvaText3DComponent::FormatText(FText& InOutText) const
+{
+	if (bEnforceUpperCase)
+	{
+		InOutText = InOutText.ToUpper();
+	}
+}
+
+void UAvaText3DComponent::OnRootComponentTransformUpdated(USceneComponent* InSceneComponent, EUpdateTransformFlags InUpdateTransformFlags, ETeleportType InTeleportType)
+{
+	ScheduleTextRefreshOnTicker();
+}
+
+void UAvaText3DComponent::SetupGradientMaterial()
+{
+	TexturedTextMaterialInstance = nullptr;
+	
+	if (!IsValid(GradientTextMaterialInstance) || GradientTextMaterialInstance->GetOuter() != this || IsMarkedForMaterialRefresh())
+	{
+		const EAvaTextMaterialFeatures MaterialFeatures = GetMaterialFeaturesFromProperties();
+		const FAvaTextMaterialSettings MaterialSettings(EAvaTextColoringStyle::Gradient, MaterialFeatures);
+		
+		if (UMaterialInstanceDynamic* const GradientMaterial = GetMIDWithSettings(GradientTextMaterialInstance, MaterialSettings))
+		{
+			UMaterialInstanceDynamic* PreviousMaterial = GradientTextMaterialInstance;
+			GradientTextMaterialInstance = GradientMaterial;
+			if (PreviousMaterial)
+			{
+				GradientTextMaterialInstance->CopyInterpParameters(PreviousMaterial);	
+			}
+		
+			RemoveRefreshReason(EAvaTextRefreshReason::MaterialInstancesChange);
+		}
+	}
+
+	if (!GradientTextMaterialInstance)
+	{
+		return;
+	}
+
+	UE::AvaText::Private::PatchMaterial(
+		this
+		, &UText3DComponent::GetFrontMaterial
+		, &UText3DComponent::SetFrontMaterial
+		, GradientTextMaterialInstance);
+
+	UE::AvaText::Private::PatchMaterial(
+		this
+		, &UText3DComponent::GetExtrudeMaterial
+		, &UText3DComponent::SetExtrudeMaterial
+		, GradientTextMaterialInstance);
+
+	UE::AvaText::Private::PatchMaterial(
+		this
+		, &UText3DComponent::GetBevelMaterial
+		, &UText3DComponent::SetBevelMaterial
+		, GradientTextMaterialInstance);
+
+	UE::AvaText::Private::PatchMaterial(
+		this
+		, &UText3DComponent::GetBackMaterial
+		, &UText3DComponent::SetBackMaterial
+		, GradientTextMaterialInstance);
+
+	if (IsMarkedForMaterialParametersRefresh())
+	{
+		RefreshGradientValues();
+		RefreshTranslucencyValues();
+	}
+}
+
+void UAvaText3DComponent::SetupCustomMaterial()
+{
+	TexturedTextMaterialInstance = nullptr;
+	GradientTextMaterialInstance = nullptr;
+	
+	if (CustomMaterial && IsMarkedForMaterialRefresh())
+	{
+		RemoveRefreshReason(EAvaTextRefreshReason::MaterialInstancesChange);
+
+		UE::AvaText::Private::PatchMaterial(
+			this
+			, &UText3DComponent::GetFrontMaterial
+			, &UText3DComponent::SetFrontMaterial
+			, GradientTextMaterialInstance);
+
+		if (UMaterialInstanceDynamic* CustomMaterialAsInstance = Cast<UMaterialInstanceDynamic>(CustomMaterial))
+		{
+			UE::AvaText::Private::PatchMaterial(
+				this
+				, &UText3DComponent::GetFrontMaterial
+				, &UText3DComponent::SetFrontMaterial
+				, CustomMaterialAsInstance);
+
+			UE::AvaText::Private::PatchMaterial(
+				this
+				, &UText3DComponent::GetExtrudeMaterial
+				, &UText3DComponent::SetExtrudeMaterial
+				, CustomMaterialAsInstance);
+
+			UE::AvaText::Private::PatchMaterial(
+				this
+				, &UText3DComponent::GetBevelMaterial
+				, &UText3DComponent::SetBevelMaterial
+				, CustomMaterialAsInstance);
+
+			UE::AvaText::Private::PatchMaterial(
+				this
+				, &UText3DComponent::GetBackMaterial
+				, &UText3DComponent::SetBackMaterial
+				, CustomMaterialAsInstance);
+		}
+		else
+		{
+			SetFrontMaterial(CustomMaterial);
+			SetExtrudeMaterial(CustomMaterial);
+			SetBevelMaterial(CustomMaterial);
+			SetBackMaterial(CustomMaterial);
+		}
+	}
+}
+
+void UAvaText3DComponent::RefreshTranslucencyValues()
+{
+	switch (TranslucencyStyle)
+	{
+		case EAvaTextTranslucency::Translucent:
+			RefreshOpacity();
+			break;
+				
+		case EAvaTextTranslucency::GradientMask:
+			RefreshMaskValues();
+			break;
+			
+		default: ;
+	}
+}
+
+void UAvaText3DComponent::SetupTexturedMaterial()
+{
+	GradientTextMaterialInstance = nullptr;
+	
+	if (!IsValid(TexturedTextMaterialInstance) || TexturedTextMaterialInstance->GetOuter() != this || IsMarkedForMaterialRefresh())
+	{
+		const EAvaTextMaterialFeatures MaterialFeatures = GetMaterialFeaturesFromProperties();
+		const FAvaTextMaterialSettings MaterialSettings(EAvaTextColoringStyle::FromTexture, MaterialFeatures);
+
+		if (UMaterialInstanceDynamic* const TexturedMaterial = GetMIDWithSettings(TexturedTextMaterialInstance, MaterialSettings))
+		{
+			UMaterialInstanceDynamic* PreviousMaterial = TexturedTextMaterialInstance;
+			TexturedTextMaterialInstance = TexturedMaterial;
+			if (PreviousMaterial)
+			{
+				TexturedTextMaterialInstance->CopyInterpParameters(PreviousMaterial);	
+			}
+		}
+		
+		RemoveRefreshReason(EAvaTextRefreshReason::MaterialInstancesChange);
+	}
+
+	if (!TexturedTextMaterialInstance)
+	{
+		return;
+	}
+
+	{
+		UE::AvaText::Private::PatchMaterial(
+			this
+			, &UText3DComponent::GetFrontMaterial
+			, &UText3DComponent::SetFrontMaterial
+			, TexturedTextMaterialInstance);
+
+		UE::AvaText::Private::PatchMaterial(
+			this
+			, &UText3DComponent::GetExtrudeMaterial
+			, &UText3DComponent::SetExtrudeMaterial
+			, TexturedTextMaterialInstance);
+
+		UE::AvaText::Private::PatchMaterial(
+			this
+			, &UText3DComponent::GetBevelMaterial
+			, &UText3DComponent::SetBevelMaterial
+			, TexturedTextMaterialInstance);
+
+		UE::AvaText::Private::PatchMaterial(
+			this
+			, &UText3DComponent::GetBackMaterial
+			, &UText3DComponent::SetBackMaterial
+			, TexturedTextMaterialInstance);
+	}
+
+	UTexture* CurrentTexture;
+	const FMemoryImageMaterialParameterInfo MaterialInfo(UAvaTextMaterialHub::FMaterialStatics::MainTexture_MatParam);
+	TexturedTextMaterialInstance->GetTextureParameterValue(MaterialInfo, CurrentTexture);
+
+	if (MainTexture != CurrentTexture)
+	{
+		TexturedTextMaterialInstance->SetTextureParameterValue(UAvaTextMaterialHub::FMaterialStatics::MainTexture_MatParam, MainTexture);
+	}
+
+	TexturedTextMaterialInstance->SetScalarParameterValue(UAvaTextMaterialHub::FMaterialStatics::TexturedUTiling_MatParam, Tiling.X);
+	TexturedTextMaterialInstance->SetScalarParameterValue(UAvaTextMaterialHub::FMaterialStatics::TexturedVTiling_MatParam, Tiling.Y);
+
+	RefreshTranslucencyValues();
+}
+
+void UAvaText3DComponent::SetupSolidMaterial()
+{
+	TexturedTextMaterialInstance = nullptr;
+	GradientTextMaterialInstance = nullptr;
+
+	if (!IsValid(TextMaterialInstance) || TextMaterialInstance->GetOuter() != this || IsMarkedForMaterialRefresh())
+	{
+		const EAvaTextMaterialFeatures MaterialFeatures = GetMaterialFeaturesFromProperties();
+		const FAvaTextMaterialSettings MaterialSettings(EAvaTextColoringStyle::Solid, MaterialFeatures);
+
+		if (UMaterialInstanceDynamic* const SolidMaterial = GetMIDWithSettings(TextMaterialInstance, MaterialSettings))
+		{
+			UMaterialInstanceDynamic* PreviousMaterial = TextMaterialInstance;
+			TextMaterialInstance = SolidMaterial;
+			if (PreviousMaterial)
+			{
+				TextMaterialInstance->CopyInterpParameters(PreviousMaterial);	
+			}
+
+			PreviousMaterial = ExtrudeTextMaterialInstance;
+			ExtrudeTextMaterialInstance = SolidMaterial;
+			if (PreviousMaterial)
+			{
+				ExtrudeTextMaterialInstance->CopyInterpParameters(PreviousMaterial);
+			}
+
+			PreviousMaterial = BevelTextMaterialInstance;
+			BevelTextMaterialInstance = SolidMaterial;
+			if (PreviousMaterial)
+			{
+				BevelTextMaterialInstance->CopyInterpParameters(PreviousMaterial);
+			}
+		}
+	
+		RemoveRefreshReason(EAvaTextRefreshReason::MaterialInstancesChange);
+	}
+ 
+	{
+		UE::AvaText::Private::PatchMaterial(
+			this
+			, &UText3DComponent::GetFrontMaterial
+			, &UText3DComponent::SetFrontMaterial
+			, TextMaterialInstance);
+
+		UE::AvaText::Private::PatchMaterial(
+			this
+			, &UText3DComponent::GetExtrudeMaterial
+			, &UText3DComponent::SetExtrudeMaterial
+			, ExtrudeTextMaterialInstance);
+
+		UE::AvaText::Private::PatchMaterial(
+			this
+			, &UText3DComponent::GetBevelMaterial
+			, &UText3DComponent::SetBevelMaterial
+			, BevelTextMaterialInstance);
+
+		UE::AvaText::Private::PatchMaterial(
+			this
+			, &UText3DComponent::GetBackMaterial
+			, &UText3DComponent::SetBackMaterial
+			, TextMaterialInstance);
+	}
+
+	if (IsMarkedForMaterialParametersRefresh())
+	{
+		RefreshColor();
+		RefreshExtrudeColor();
+		RefreshBevelColor();
+		
+		RefreshTranslucencyValues();
+	}
+}
+
+void UAvaText3DComponent::SetupTextMaterials(const EAvaTextColoringStyle& InColoringStyle)
+{
+	if (IsMarkedForMaterialRefresh())
+	{
+		switch (InColoringStyle)
+		{
+		case EAvaTextColoringStyle::Invalid:
+			break;
+			
+		case EAvaTextColoringStyle::Solid:
+			SetupSolidMaterial();
+			break;
+
+		case EAvaTextColoringStyle::Gradient:
+			SetupGradientMaterial();
+			break;
+					
+		case EAvaTextColoringStyle::FromTexture:
+			SetupTexturedMaterial();
+			break;
+					
+		case EAvaTextColoringStyle::CustomMaterial:
+			SetupCustomMaterial();
+			break;
+		
+		default:
+			break;
+		}
+	}
+}
+
+void UAvaText3DComponent::SetupTextLayout()
+{
+	if (IsMarkedForLayoutRefresh())
+	{
+		SetKerning(Kerning);
+		SetWordSpacing(WordSpacing);
+		SetLineSpacing(LineSpacing);
+		SetScaleProportionally(bScaleProportionally);		
+		SetMaxWidth(MaxWidth);
+		SetMaxHeight(MaxHeight);
+		SetHasMaxWidth(bHasMaxWidth);
+		SetHasMaxHeight(bHasMaxHeight);
+		SetHorizontalAlignment(HorizontalAlignment);
+		SetVerticalAlignment(VerticalAlignment);
+
+		RemoveRefreshReason(EAvaTextRefreshReason::LayoutChange);
+	}
+}
+
+void UAvaText3DComponent::SetupTextGeometry()
+{
+	if (IsMarkedForGeometryRefresh())
+	{
+		// we are about to setup multiple properties together, let's freeze rebuild
+		SetFreeze(true);
+	
+		SetText(Text);
+
+		if (UFont* UnderlyingFont = AvalancheFont.GetFont())
+		{
+			SetFont(UnderlyingFont);	
+		}
+		
+		SetExtrude(Extrude);
+		
+		SetBevel(Bevel);
+		SetBevelType(BevelType);
+		SetBevelSegments(BevelSegments);
+	
+		//... and finally unfreeze, to trigger pending rebuild
+		SetFreeze(false);
+
+		RemoveRefreshReason(EAvaTextRefreshReason::GeometryChange);
+	}
+}
+
+FVector UAvaText3DComponent::GetGradientDirection() const
+{
+	// Gradient Rot is a [0;1] ratio, representing to [0;360] degrees
+	float GradientRot;
+	
+	if (GradientSettings.Direction == EAvaGradientDirection::Vertical)
+	{
+		GradientRot = UAvaTextMaterialHub::FMaterialStatics::MaterialMaskTopBottomRotationValue;
+	}
+	else if (GradientSettings.Direction == EAvaGradientDirection::Horizontal)
+	{
+		GradientRot = UAvaTextMaterialHub::FMaterialStatics::MaterialMaskLeftToRightRotationValue;
+	}
+	else
+	{
+		GradientRot = GradientSettings.Rotation;
+	}
+	
+	FVector GradientDir = GetUpVector().RotateAngleAxis(-GradientRot * 360, GetForwardVector());
+	
+	/**
+	 * In order to properly map gradient offset along the text surface, text bounds are not normalized (anymore) in the material function creating the gradient.
+	 * Therefore, we need to remap the gradient direction, taking into account the current text actor bounds.
+	 *
+	 * (material function: '/Avalanche/Text3DResources/MaterialFunctions/MF_LocalGradientMask.MF_LocalGradientMask')
+	 */
+	{
+		const FVector GradientDirFixer = FVector(1.0f, LocalBoundsExtent.Z, LocalBoundsExtent.Y);
+		GradientDir*=GradientDirFixer;
+	}
+
+	GradientDir.Normalize();
+
+	return GradientDir;
+}
+
+void UAvaText3DComponent::RefreshText3D()
+{
+	SetupTextGeometry();
+	SetupTextLayout();
+	SetupTextMaterials(ColoringStyle);
+	
+	RefreshGradientValues();
+	RefreshMaskValues();
+	RefreshMaterialGeometryParameterValues();
+
+	ClearRefreshReasons();
+	MarkRenderDirty();
+}
+
+void UAvaText3DComponent::MarkRenderDirty()
+{
+	MarkRenderStateDirty();
+	MarkRenderInstancesDirty();
+}
+
+void UAvaText3DComponent::RegisterCallbacks()
+{
+	if (!OnTextGeometryGeneratedDelegate.IsValid())
+	{
+		OnTextGeometryGeneratedDelegate = OnTextGenerated().AddUObject(this, &UAvaText3DComponent::OnTextGeometryGenerated);
+	}
+	
+	if (ParentRootComponent.IsValid() && !OnTransformUpdatedDelegate.IsValid())
+	{
+		OnTransformUpdatedDelegate = ParentRootComponent->TransformUpdated.AddUObject(this, &UAvaText3DComponent::OnRootComponentTransformUpdated);
+	}
+}
+
+void UAvaText3DComponent::RefreshParentRootComponent()
+{
+	if (GetOwner())
+	{
+		ParentRootComponent = GetOwner()->GetRootComponent();
+	}
+}
+
+void UAvaText3DComponent::SetIsUnlit(bool bInIsUnlit)
+{
+	if (bIsUnlit != bInIsUnlit)
+	{
+		bIsUnlit = bInIsUnlit;
+		RefreshUnlit();
+	}
+}
+
+void UAvaText3DComponent::SetEnforceUpperCase(bool bInEnforceUpperCase)
+{
+	if (bEnforceUpperCase == bInEnforceUpperCase)
+	{
+		return;
+	}
+	
+	bEnforceUpperCase = bInEnforceUpperCase;
+	RefreshFormatting();
+}
+
+void UAvaText3DComponent::SetAvalancheFont(const FAvaFont& InAvalancheFont)
+{
+	if (AvalancheFont == InAvalancheFont)
+	{
+		return;
+	}
+	AvalancheFont = InAvalancheFont;
+	RefreshGeometry();
+}
+
+void UAvaText3DComponent::SetMainTexture(UTexture2D* InMainTexture)
+{
+	if (MainTexture != InMainTexture)
+	{
+		MainTexture = InMainTexture;
+		RefreshMaterialInstances();
+	}
+}
+
+void UAvaText3DComponent::SetTiling(const FVector2D InTiling)
+{
+	if (Tiling != InTiling)
+	{
+		Tiling = InTiling;
+		RefreshMaterialInstances();
+	}
+}
+
+void UAvaText3DComponent::SetCustomMaterial(UMaterialInterface* InCustomMaterial)
+{
+	if (CustomMaterial != InCustomMaterial)
+	{
+		CustomMaterial = InCustomMaterial;
+		RefreshMaterialInstances();
+	}
+}
+
+void UAvaText3DComponent::SetTranslucencyStyle(EAvaTextTranslucency InTranslucencyStyle)
+{
+	if (TranslucencyStyle != InTranslucencyStyle)
+	{
+		TranslucencyStyle = InTranslucencyStyle;
+		RefreshMaterialInstances();
+	}
+}
+
+void UAvaText3DComponent::SetOpacity(float InOpacity)
+{
+	if (Opacity != InOpacity)
+	{
+		Opacity = InOpacity;
+		RefreshOpacity();
+	}
+}
+
+void UAvaText3DComponent::SetColoringStyle(const EAvaTextColoringStyle& InColoringStyle)
+{
+	if (ColoringStyle != InColoringStyle)
+	{
+		ColoringStyle = InColoringStyle;	
+		RefreshMaterialInstances();
+	}
+}
+
+void UAvaText3DComponent::SetColor(const FLinearColor& InColor)
+{
+	if (Color != InColor)
+	{
+		Color = InColor;
+		RefreshColor();
+	}
+}
+
+void UAvaText3DComponent::SetExtrudeColor(const FLinearColor& InColor)
+{
+	if (ExtrudeColor != InColor)
+	{
+		ExtrudeColor = InColor;
+		RefreshExtrudeColor();	
+	}
+}
+
+void UAvaText3DComponent::SetBevelColor(const FLinearColor& InColor)
+{
+	if (BevelColor != InColor)
+	{
+		BevelColor = InColor;
+		RefreshBevelColor();	
+	}
+}
+
+void UAvaText3DComponent::SetGradientSettings(const FAvaLinearGradientSettings& InGradientSettings)
+{
+	if (GradientSettings == InGradientSettings)
+	{
+		return;
+	}
+
+	GradientSettings = InGradientSettings;
+	RefreshGradient();	
+}
+
+void UAvaText3DComponent::SetGradientColors(const FLinearColor& InColorA, const FLinearColor& InColorB)
+{
+	if (GradientSettings.ColorA == InColorA && GradientSettings.ColorB == InColorB)
+	{
+		return;
+	}
+	
+	GradientSettings.ColorA = InColorA;
+	GradientSettings.ColorB = InColorB;
+	RefreshGradient();
+}
+
+void UAvaText3DComponent::SetMaskOrientation(EAvaMaterialMaskOrientation InMaskOrientation)
+{
+	if (MaskOrientation != InMaskOrientation)
+	{
+		MaskOrientation = InMaskOrientation;
+		RefreshMask();
+	}
+}
+
+void UAvaText3DComponent::SetMaskSmoothness(float InMaskSmoothness)
+{
+	if (MaskSmoothness != InMaskSmoothness)
+	{
+		MaskSmoothness = InMaskSmoothness;
+		RefreshMask();	
+	}
+}
+
+void UAvaText3DComponent::SetMaskOffset(float InMaskOffset)
+{
+	if (MaskOffset != InMaskOffset)
+	{
+		MaskOffset = InMaskOffset;
+		RefreshMask();	
+	}
+}
+
+void UAvaText3DComponent::SetMaskRotation(float InMaskRotation)
+{
+	if (MaskRotation != InMaskRotation)
+	{
+		MaskRotation = InMaskRotation;
+		RefreshMask();	
+	}
+}
+
+void UAvaText3DComponent::SetAlignment(FAvaTextAlignment InAlignment)
+{
+	if (Alignment == InAlignment)
+	{
+		return;
+	}
+
+	Alignment = InAlignment;
+	RefreshAlignment();	
+}
+
+void UAvaText3DComponent::RefreshOpacity() const
+{
+	if (TranslucencyStyle == EAvaTextTranslucency::Translucent)
+	{
+		// Set for any/all materials assigned
+		ForEachMID([this](UMaterialInstanceDynamic* InMID)
+		{
+			InMID->SetScalarParameterValue(UAvaTextMaterialHub::FMaterialStatics::Opacity_MatParam, Opacity);
+		});
+
+		// ... and set for the default materials
+		switch (ColoringStyle)
+		{
+		case EAvaTextColoringStyle::Solid:
+			if (TextMaterialInstance)
+			{
+				TextMaterialInstance->SetScalarParameterValue(UAvaTextMaterialHub::FMaterialStatics::Opacity_MatParam, Opacity);
+			}
+		
+			if (ExtrudeTextMaterialInstance)
+			{
+				ExtrudeTextMaterialInstance->SetScalarParameterValue(UAvaTextMaterialHub::FMaterialStatics::Opacity_MatParam, Opacity);
+			}
+
+			if (BevelTextMaterialInstance)
+			{
+				BevelTextMaterialInstance->SetScalarParameterValue(UAvaTextMaterialHub::FMaterialStatics::Opacity_MatParam, Opacity);
+			}
+			break;
+				
+		case EAvaTextColoringStyle::Gradient:
+			if (GradientTextMaterialInstance)
+			{
+				GradientTextMaterialInstance->SetScalarParameterValue(UAvaTextMaterialHub::FMaterialStatics::Opacity_MatParam, Opacity);
+			}
+			break;
+				
+		case EAvaTextColoringStyle::FromTexture:
+			if (TexturedTextMaterialInstance)
+			{
+				TexturedTextMaterialInstance->SetScalarParameterValue(UAvaTextMaterialHub::FMaterialStatics::Opacity_MatParam, Opacity);
+			}
+			break;
+				
+		case EAvaTextColoringStyle::CustomMaterial:
+			break;
+			
+		default: ;
+		}
+	}
+}
+
+void UAvaText3DComponent::RefreshColor() const
+{
+	if (UMaterialInstanceDynamic* MID = Cast<UMaterialInstanceDynamic>(GetFrontMaterial()))
+	{
+		MID->SetVectorParameterValue(UAvaTextMaterialHub::FMaterialStatics::SolidColor_MatParam, Color);
+	}
+
+	if (TextMaterialInstance)
+	{
+		TextMaterialInstance->SetVectorParameterValue(UAvaTextMaterialHub::FMaterialStatics::SolidColor_MatParam, Color);
+	}
+}
+
+void UAvaText3DComponent::RefreshExtrudeColor() const
+{
+	if (UMaterialInstanceDynamic* MID = Cast<UMaterialInstanceDynamic>(GetExtrudeMaterial()))
+	{
+		MID->SetVectorParameterValue(UAvaTextMaterialHub::FMaterialStatics::SolidColor_MatParam, ExtrudeColor);
+	}
+
+	if (ExtrudeTextMaterialInstance)
+	{
+		ExtrudeTextMaterialInstance->SetVectorParameterValue(UAvaTextMaterialHub::FMaterialStatics::SolidColor_MatParam, Color);
+	}
+}
+
+void UAvaText3DComponent::RefreshBevelColor() const
+{
+	if (UMaterialInstanceDynamic* MID = Cast<UMaterialInstanceDynamic>(GetBevelMaterial()))
+	{
+		MID->SetVectorParameterValue(UAvaTextMaterialHub::FMaterialStatics::SolidColor_MatParam, BevelColor);
+	}
+
+	if (BevelTextMaterialInstance)
+	{
+		BevelTextMaterialInstance->SetVectorParameterValue(UAvaTextMaterialHub::FMaterialStatics::SolidColor_MatParam, Color);
+	}
+}
+
+void UAvaText3DComponent::ScheduleTextRefreshOnTicker()
+{
+	if (!TickerHandle.IsValid())
+	{
+		TWeakObjectPtr<UAvaText3DComponent> WeakThis(this);
+		FTSTicker& Ticker = FTSTicker::GetCoreTicker();
+		
+		TickerHandle = Ticker.AddTicker(FTickerDelegate::CreateLambda(
+			[WeakThis](float InDeltaTime)->bool
+			{
+				if (UAvaText3DComponent* This = WeakThis.Get())
+				{
+					if (!UE::IsSavingPackage(This))
+					{
+						This->RefreshStoredBoundsValues();
+						This->MarkForMaterialParametersRefresh();
+						This->RefreshText3D();
+					}
+
+					This->TickerHandle.Reset();
+				}
+
+				// Executes only once
+				return false;
+			}));
+	}
+}
+
+#if WITH_EDITOR
+void UAvaText3DComponent::RegisterOnPropertyChangeFunctions()
+{
+	if (PropertyChangeFunctions.IsEmpty())
+	{
+		PropertyChangeFunctions.Add(GET_MEMBER_NAME_CHECKED(UAvaText3DComponent, Color), &UAvaText3DComponent::RefreshColor);
+		PropertyChangeFunctions.Add(GET_MEMBER_NAME_CHECKED(UAvaText3DComponent, ExtrudeColor), &UAvaText3DComponent::RefreshExtrudeColor);
+		PropertyChangeFunctions.Add(GET_MEMBER_NAME_CHECKED(UAvaText3DComponent, BevelColor), &UAvaText3DComponent::RefreshBevelColor);
+		PropertyChangeFunctions.Add(GET_MEMBER_NAME_CHECKED(UAvaText3DComponent, bIsUnlit), &UAvaText3DComponent::RefreshUnlit);
+		PropertyChangeFunctions.Add(GET_MEMBER_NAME_CHECKED(UAvaText3DComponent, Opacity), &UAvaText3DComponent::RefreshOpacity);
+		PropertyChangeFunctions.Add(GET_MEMBER_NAME_CHECKED(UAvaText3DComponent, GradientSettings), &UAvaText3DComponent::RefreshGradient);
+
+		PropertyChangeFunctions.Add(GET_MEMBER_NAME_CHECKED(UAvaText3DComponent, Extrude), &UAvaText3DComponent::RefreshGeometry);
+		PropertyChangeFunctions.Add(GET_MEMBER_NAME_CHECKED(UAvaText3DComponent, Bevel), &UAvaText3DComponent::RefreshGeometry);
+		PropertyChangeFunctions.Add(GET_MEMBER_NAME_CHECKED(UAvaText3DComponent, AvalancheFont), &UAvaText3DComponent::RefreshGeometry);
+		PropertyChangeFunctions.Add(GET_MEMBER_NAME_CHECKED(UAvaText3DComponent, Font), &UAvaText3DComponent::RefreshAvalancheFont);
+		PropertyChangeFunctions.Add(GET_MEMBER_NAME_CHECKED(UAvaText3DComponent, Text), &UAvaText3DComponent::RefreshGeometry);
+		PropertyChangeFunctions.Add(GET_MEMBER_NAME_CHECKED(UAvaText3DComponent, bEnforceUpperCase), &UAvaText3DComponent::RefreshFormatting);
+
+		PropertyChangeFunctions.Add(GET_MEMBER_NAME_CHECKED(UAvaText3DComponent, CustomMaterial), &UAvaText3DComponent::RefreshMaterialInstances);
+		PropertyChangeFunctions.Add(GET_MEMBER_NAME_CHECKED(UAvaText3DComponent, ColoringStyle), &UAvaText3DComponent::RefreshMaterialInstances);
+		PropertyChangeFunctions.Add(GET_MEMBER_NAME_CHECKED(UAvaText3DComponent, MainTexture), &UAvaText3DComponent::RefreshMaterialInstances);
+		PropertyChangeFunctions.Add(GET_MEMBER_NAME_CHECKED(UAvaText3DComponent, Tiling), &UAvaText3DComponent::RefreshMaterialInstances);
+		PropertyChangeFunctions.Add(GET_MEMBER_NAME_CHECKED(UAvaText3DComponent, TranslucencyStyle), &UAvaText3DComponent::RefreshMaterialInstances);
+		
+		PropertyChangeFunctions.Add(GET_MEMBER_NAME_CHECKED(UAvaText3DComponent, MaskOrientation), &UAvaText3DComponent::RefreshMask);
+		PropertyChangeFunctions.Add(GET_MEMBER_NAME_CHECKED(UAvaText3DComponent, MaskSmoothness), &UAvaText3DComponent::RefreshMask);
+		PropertyChangeFunctions.Add(GET_MEMBER_NAME_CHECKED(UAvaText3DComponent, MaskOffset), &UAvaText3DComponent::RefreshMask);
+		PropertyChangeFunctions.Add(GET_MEMBER_NAME_CHECKED(UAvaText3DComponent, MaskRotation), &UAvaText3DComponent::RefreshMask);
+
+		PropertyChangeFunctions.Add(GET_MEMBER_NAME_CHECKED(UAvaText3DComponent, Alignment), &UAvaText3DComponent::RefreshAlignment);
+		PropertyChangeFunctions.Add(GET_MEMBER_NAME_CHECKED(UAvaText3DComponent, HorizontalAlignment), &UAvaText3DComponent::RefreshHorizontalAlignment);
+		PropertyChangeFunctions.Add(GET_MEMBER_NAME_CHECKED(UAvaText3DComponent, VerticalAlignment), &UAvaText3DComponent::RefreshVerticalAlignment);
+		PropertyChangeFunctions.Add(GET_MEMBER_NAME_CHECKED(UAvaText3DComponent, bHasMaxWidth), &UAvaText3DComponent::RefreshLayout);
+		PropertyChangeFunctions.Add(GET_MEMBER_NAME_CHECKED(UAvaText3DComponent, MaxWidth), &UAvaText3DComponent::RefreshLayout);
+		PropertyChangeFunctions.Add(GET_MEMBER_NAME_CHECKED(UAvaText3DComponent, bHasMaxHeight), &UAvaText3DComponent::RefreshLayout);
+		PropertyChangeFunctions.Add(GET_MEMBER_NAME_CHECKED(UAvaText3DComponent, MaxHeight), &UAvaText3DComponent::RefreshLayout);
+		PropertyChangeFunctions.Add(GET_MEMBER_NAME_CHECKED(UAvaText3DComponent, bScaleProportionally), &UAvaText3DComponent::RefreshLayout);
+		PropertyChangeFunctions.Add(GET_MEMBER_NAME_CHECKED(UAvaText3DComponent, Kerning), &UAvaText3DComponent::RefreshLayout);
+		PropertyChangeFunctions.Add(GET_MEMBER_NAME_CHECKED(UAvaText3DComponent, WordSpacing), &UAvaText3DComponent::RefreshLayout);
+		PropertyChangeFunctions.Add(GET_MEMBER_NAME_CHECKED(UAvaText3DComponent, LineSpacing), &UAvaText3DComponent::RefreshLayout);
+	}
+}
+#endif
+
+void UAvaText3DComponent::RefreshAlignment()
+{
+	HorizontalAlignment = Alignment.HorizontalAlignment;
+	VerticalAlignment = Alignment.VerticalAlignment;
+
+	ScheduleTextRefreshOnTicker();
+	RefreshLayout();
+}
+
+void UAvaText3DComponent::RefreshAvalancheFont()
+{
+	if (IsValid(Font) && Font != AvalancheFont.GetFont())
+	{
+		// todo: this leads to proper update of the AvalancheFont, and of the text component geometry
+		// nonetheless, custom FAvalancheFont widget is not updating properly yet
+		AvalancheFont.InitFromFont(Font);
+	}
+}
+
+void UAvaText3DComponent::RefreshHorizontalAlignment()
+{
+	Alignment.HorizontalAlignment = HorizontalAlignment;
+	RefreshLayout();
+}
+
+void UAvaText3DComponent::RefreshVerticalAlignment()
+{
+	Alignment.VerticalAlignment = VerticalAlignment;
+	RefreshLayout();
+}
+
+void UAvaText3DComponent::RefreshUnlit()
+{
+	MarkForMaterialsFullRefresh();	
+	RefreshText3D();
+}
+
+void UAvaText3DComponent::RefreshMasked()
+{
+	MarkForMaterialsFullRefresh();
+	RefreshText3D();
+}
+
+void UAvaText3DComponent::RefreshMaterialInstances()
+{
+	MarkForMaterialsFullRefresh();
+	RefreshText3D();
+}
+
+void UAvaText3DComponent::RefreshGeometry()
+{
+	MarkForGeometryRefresh();
+	RefreshText3D();
+}
+
+void UAvaText3DComponent::RefreshLayout()
+{
+	MarkForLayoutRefresh();
+	RefreshText3D();
+}
+
+void UAvaText3DComponent::RefreshGradientValues()
+{
+	if (IsMarkedForMaterialRefresh() || IsMarkedForMaterialParametersRefresh())
+	{
+		if (IsValid(GradientTextMaterialInstance))
+		{
+			GradientTextMaterialInstance->SetScalarParameterValue(UAvaTextMaterialHub::FMaterialStatics::GradientOffset_MatParam, GradientSettings.Offset);
+			GradientTextMaterialInstance->SetScalarParameterValue(UAvaTextMaterialHub::FMaterialStatics::GradientSmoothness_MatParam, GradientSettings.Smoothness);
+			GradientTextMaterialInstance->SetVectorParameterValue(UAvaTextMaterialHub::FMaterialStatics::GradientColorA_MatParam, GradientSettings.ColorA);
+			GradientTextMaterialInstance->SetVectorParameterValue(UAvaTextMaterialHub::FMaterialStatics::GradientColorB_MatParam, GradientSettings.ColorB);
+
+			float GradientRotation = UAvaTextMaterialHub::FMaterialStatics::MaterialMaskTopBottomRotationValue;
+			
+			switch (GradientSettings.Direction)
+			{
+			case EAvaGradientDirection::None:
+				break;
+
+			case EAvaGradientDirection::Vertical:				
+					GradientRotation = UAvaTextMaterialHub::FMaterialStatics::MaterialMaskTopBottomRotationValue;
+					break;				
+
+			case EAvaGradientDirection::Horizontal:
+					GradientRotation = UAvaTextMaterialHub::FMaterialStatics::MaterialMaskLeftToRightRotationValue;
+					break;
+
+			case EAvaGradientDirection::Custom:
+					GradientRotation = GradientSettings.Rotation;
+					break;
+			}
+
+			GradientTextMaterialInstance->SetScalarParameterValue(UAvaTextMaterialHub::FMaterialStatics::GradientRotation_MatParam, GradientRotation);
+		}
+
+		// scaling factor required to "normalize" origin and local bounds to properly update the material
+		const FVector TextScaleFactor = GetTextScale() * GetComponentScale();
+		SetMaterialGeometryParameterValues(GradientTextMaterialInstance, TextScaleFactor);
+
+		RemoveRefreshReason(EAvaTextRefreshReason::MaterialParametersChange);
+	}
+}
+
+void UAvaText3DComponent::RefreshGradient()
+{
+	MarkForMaterialParametersRefresh();
+	RefreshGradientValues();
+}
+
+void UAvaText3DComponent::RefreshMaskValues()
+{
+	if (IsMarkedForMaterialRefresh() || IsMarkedForMaterialParametersRefresh())
+	{
+		// Set for any/all materials assigned
+		ForEachMID([this](UMaterialInstanceDynamic* InMID)
+		{
+			SetMaskedMaterialValues(InMID);
+		});
+
+		// ... and set for the default materials
+		switch (ColoringStyle)
+		{
+		case EAvaTextColoringStyle::Invalid:
+			break;
+	
+		case EAvaTextColoringStyle::Solid:
+			SetMaskedMaterialValues(TextMaterialInstance);
+			SetMaskedMaterialValues(ExtrudeTextMaterialInstance);
+			SetMaskedMaterialValues(BevelTextMaterialInstance);
+			break;
+		
+		case EAvaTextColoringStyle::Gradient:
+			SetMaskedMaterialValues(GradientTextMaterialInstance);
+			break;
+		
+		case EAvaTextColoringStyle::FromTexture:
+			SetMaskedMaterialValues(TexturedTextMaterialInstance);
+			break;
+		
+		case EAvaTextColoringStyle::CustomMaterial:
+			break;
+		}
+		
+		RemoveRefreshReason(EAvaTextRefreshReason::MaterialParametersChange);
+	}
+}
+
+void UAvaText3DComponent::RefreshMask()
+{
+	MarkForMaterialParametersRefresh();
+	RefreshMaskValues();
+}
+
+void UAvaText3DComponent::SetMaterialGeometryParameterValues(UMaterialInstanceDynamic* InMaterial, const FVector& TextScaleFactor) const
+{
+	if (!InMaterial)
+	{
+		return;
+	}
+	
+	// we use the bounds minimum point, because the material needs the bbox origin in the the lower left corner
+	const FVector BoundsLowerLeftCornerLocalScaled = LocalBoundsLowerLeftCorner/TextScaleFactor;
+	const FVector BoundsSizeScaled = LocalBoundsExtent/TextScaleFactor;
+
+	InMaterial->SetVectorParameterValue(UAvaTextMaterialHub::FMaterialStatics::BoundsOrigin_MatParam, BoundsLowerLeftCornerLocalScaled);
+	InMaterial->SetVectorParameterValue(UAvaTextMaterialHub::FMaterialStatics::BoundsSize_MatParam, BoundsSizeScaled);
+	
+	FVector TextPositionParameter = GetComponentLocation();
+	
+	// in case this component is not the root, let's take into account also relative location
+	if (ParentRootComponent != this)
+	{
+		TextPositionParameter -= GetRelativeLocation();
+	}
+	
+	InMaterial->SetVectorParameterValue(UAvaTextMaterialHub::FMaterialStatics::TextPosition_MatParam, TextPositionParameter);
+}
+
+void UAvaText3DComponent::RefreshMaterialGeometryParameterValues()
+{
+	// scaling factor required to "normalize" origin and local bounds to properly update the material
+	const FVector TextScaleFactor = GetTextScale() * GetComponentScale();
+	
+	switch (ColoringStyle)
+	{
+	case EAvaTextColoringStyle::Invalid:
+		break;
+	
+	case EAvaTextColoringStyle::Solid:
+		SetMaterialGeometryParameterValues(TextMaterialInstance, TextScaleFactor);
+		SetMaterialGeometryParameterValues(ExtrudeTextMaterialInstance, TextScaleFactor);
+		SetMaterialGeometryParameterValues(BevelTextMaterialInstance, TextScaleFactor);
+		break;
+	
+	case EAvaTextColoringStyle::Gradient:
+		SetMaterialGeometryParameterValues(GradientTextMaterialInstance, TextScaleFactor);
+		break;
+	
+	case EAvaTextColoringStyle::FromTexture:
+		SetMaterialGeometryParameterValues(TexturedTextMaterialInstance, TextScaleFactor);
+		break;
+	
+	case EAvaTextColoringStyle::CustomMaterial:
+		break;
+	}
+}
+
+void UAvaText3DComponent::RefreshFormatting()
+{
+	TriggerInternalRebuild(EText3DModifyFlags::Geometry);
+}
+
+void UAvaText3DComponent::SetMaskedMaterialValues(UMaterialInstanceDynamic* InMaterial)
+{
+	if (InMaterial)
+	{
+		InMaterial->SetScalarParameterValue(UAvaTextMaterialHub::FMaterialStatics::MaskOffset_MatParam, MaskOffset);
+
+		float GradientRotation = UAvaTextMaterialHub::FMaterialStatics::MaterialMaskLeftToRightRotationValue;
+		
+		switch (MaskOrientation)
+		{
+		case EAvaMaterialMaskOrientation::LeftRight:
+			GradientRotation = UAvaTextMaterialHub::FMaterialStatics::MaterialMaskLeftToRightRotationValue;
+			break;
+
+		case EAvaMaterialMaskOrientation::RightLeft:
+			GradientRotation = UAvaTextMaterialHub::FMaterialStatics::MaterialMaskRightToLeftRotationValue;
+			break;
+
+		case EAvaMaterialMaskOrientation::Custom:
+			GradientRotation = MaskRotation;
+			break;
+		}
+
+		InMaterial->SetScalarParameterValue(UAvaTextMaterialHub::FMaterialStatics::MaskRotation_MatParam, GradientRotation);
+		InMaterial->SetScalarParameterValue(UAvaTextMaterialHub::FMaterialStatics::MaskSmoothness_MatParam, MaskSmoothness);
+
+		RefreshMaterialGeometryParameterValues();
+	}
+}
+
+UAvaText3DComponent::FOnGetMaterialWithSettings& UAvaText3DComponent::GetMaterialProviderDelegate()
+{
+	return OnGetMaterialWithSettingsDelegate;
+}
+
+void UAvaText3DComponent::ForEachMID(TUniqueFunction<void(UMaterialInstanceDynamic* InMID)>&& InFunc) const
+{
+	auto ExecuteIfMID = [&InFunc](UMaterialInterface* InMaterial)
+	{
+		if (UMaterialInstanceDynamic* MID = Cast<UMaterialInstanceDynamic>(InMaterial))
+		{
+			InFunc(MID);
+		}
+	};
+	
+	ExecuteIfMID(GetFrontMaterial());
+	ExecuteIfMID(GetBevelMaterial());
+	ExecuteIfMID(GetExtrudeMaterial());
+	ExecuteIfMID(GetBackMaterial());
+}
+
+void UAvaText3DComponent::RefreshStoredBoundsValues()
+{
+	if (const AActor* ParentActor = GetOwner())
+	{
+		// local bounds (used for gradient or masked materials)
+		const FBox LocalBounds = FAvaActorUtils::GetActorLocalBoundingBox(ParentActor);
+		LocalBoundsLowerLeftCorner = LocalBounds.Min;
+		LocalBoundsExtent = LocalBounds.GetExtent();
+	}
+}
+
+UMaterialInstanceDynamic* UAvaText3DComponent::GetMIDWithSettings(
+	const UMaterialInterface* InPreviousMaterial
+	, const FAvaTextMaterialSettings& InSettings)
+{
+	UMaterialInterface* Result = nullptr;
+	if (OnGetMaterialWithSettingsDelegate.IsBound())
+	{
+		Result = OnGetMaterialWithSettingsDelegate.Execute(InPreviousMaterial, InSettings);
+	}
+
+	if (!Result)
+	{
+		Result = UAvaTextMaterialHub::GetMaterial(InSettings);
+	}
+
+	if (!::IsValid(Result))
+	{
+		UE_LOG(LogAva, Warning, TEXT("Material could not be found or was invalid for the provided settings."));
+		return nullptr;
+	}
+
+	if (UMaterialInstanceDynamic* ResultAsMID = Cast<UMaterialInstanceDynamic>(Result))
+	{
+		return ResultAsMID;
+	}
+	
+	return UMaterialInstanceDynamic::Create(Result, this);
+}
+
+void UAvaText3DComponent::OnTextGeneratedActions()
+{
+	RefreshStoredBoundsValues();
+
+	// bounding box has been updated, so we need to also update gradient and mask values
+	RefreshGradient();
+	RefreshMask();
+}
+
+void UAvaText3DComponent::OnTextGeometryGenerated()
+{
+	OnTextGeneratedActions();
+
+#if WITH_EDITOR
+	if (GEditor)
+	{
+		GEditor->BroadcastOnComponentTransformChanged(this, ETeleportType::None);
+	}
+#endif
+
+	TransformUpdated.Broadcast(this, EUpdateTransformFlags::None, ETeleportType::None);
+}
+
+void UAvaText3DComponent::MarkForGeometryRefresh()
+{
+	EnumAddFlags(TextRefreshReason, EAvaTextRefreshReason::GeometryChange);
+}
+
+bool UAvaText3DComponent::IsMarkedForGeometryRefresh() const
+{
+	return EnumHasAnyFlags(TextRefreshReason, EAvaTextRefreshReason::GeometryChange);
+}
+
+void UAvaText3DComponent::RemoveRefreshReason(EAvaTextRefreshReason InRefreshReasonToRemove)
+{
+	EnumRemoveFlags(TextRefreshReason, InRefreshReasonToRemove);
+}
+
+void UAvaText3DComponent::MarkForLayoutRefresh()
+{
+	EnumAddFlags(TextRefreshReason, EAvaTextRefreshReason::LayoutChange);
+}
+
+bool UAvaText3DComponent::IsMarkedForLayoutRefresh() const
+{
+	return EnumHasAnyFlags(TextRefreshReason, EAvaTextRefreshReason::LayoutChange);
+}
+
+void UAvaText3DComponent::MarkForMaterialsFullRefresh()
+{
+	EnumAddFlags(TextRefreshReason, EAvaTextRefreshReason::MaterialInstancesChange);
+	EnumAddFlags(TextRefreshReason, EAvaTextRefreshReason::MaterialParametersChange);
+}
+
+bool UAvaText3DComponent::IsMarkedForMaterialRefresh() const
+{
+	return EnumHasAnyFlags(TextRefreshReason, EAvaTextRefreshReason::MaterialInstancesChange);
+}
+
+void UAvaText3DComponent::MarkForMaterialParametersRefresh()
+{
+	EnumAddFlags(TextRefreshReason, EAvaTextRefreshReason::MaterialParametersChange);
+}
+
+bool UAvaText3DComponent::IsMarkedForMaterialParametersRefresh() const
+{
+	return EnumHasAnyFlags(TextRefreshReason, EAvaTextRefreshReason::MaterialParametersChange);
+}
+
+void UAvaText3DComponent::MarkForFullRefresh()
+{
+	EnumAddFlags(TextRefreshReason, EAvaTextRefreshReason::FullRefresh);
+}
+
+void UAvaText3DComponent::ClearRefreshReasons()
+{
+	EnumAddFlags(TextRefreshReason, EAvaTextRefreshReason::NoRefresh);
+}
+
+void UAvaText3DComponent::RetrieveTranslucentMaterialParameters()
+{
+	switch (TranslucencyStyle)
+	{
+	case EAvaTextTranslucency::Translucent:
+		RetrieveSimpleTranslucentMaterialParameters(TextMaterialInstance);
+		break;
+				
+	case EAvaTextTranslucency::GradientMask:
+		RetrieveMaskedMaterialParameters(TextMaterialInstance);
+		break;
+				
+	default: ;
+	}
+}
+
+void UAvaText3DComponent::RetrieveSolidMaterialInstances()
+{
+	if (FrontMaterial)
+	{
+		TextMaterialInstance = Cast<UMaterialInstanceDynamic>(FrontMaterial);
+		if (TextMaterialInstance)
+		{
+			TextMaterialInstance->GetVectorParameterValue(UAvaTextMaterialHub::FMaterialStatics::SolidColor_MatParam, Color);
+
+			RetrieveTranslucentMaterialParameters();
+		}
+	}
+
+	if (ExtrudeMaterial)
+	{
+		ExtrudeTextMaterialInstance = Cast<UMaterialInstanceDynamic>(ExtrudeMaterial);
+		if (ExtrudeTextMaterialInstance)
+		{
+			ExtrudeTextMaterialInstance->GetVectorParameterValue(UAvaTextMaterialHub::FMaterialStatics::SolidColor_MatParam, ExtrudeColor);
+		}
+	}
+
+	if (BevelMaterial)
+	{
+		BevelTextMaterialInstance = Cast<UMaterialInstanceDynamic>(BevelMaterial);
+		if (BevelTextMaterialInstance)
+		{
+			BevelTextMaterialInstance->GetVectorParameterValue(UAvaTextMaterialHub::FMaterialStatics::SolidColor_MatParam, BevelColor);
+		}
+	}
+}
+
+void UAvaText3DComponent::RetrieveGradientMaterialInstance()
+{
+	if (!FrontMaterial)
+	{
+		return;
+	}
+	
+	GradientTextMaterialInstance = Cast<UMaterialInstanceDynamic>(FrontMaterial);
+
+	if (GradientTextMaterialInstance)
+	{
+		GradientTextMaterialInstance->GetVectorParameterValue(UAvaTextMaterialHub::FMaterialStatics::GradientColorA_MatParam, GradientSettings.ColorA);
+		GradientTextMaterialInstance->GetVectorParameterValue(UAvaTextMaterialHub::FMaterialStatics::GradientColorB_MatParam, GradientSettings.ColorB);
+
+		GradientTextMaterialInstance->GetScalarParameterValue(UAvaTextMaterialHub::FMaterialStatics::GradientOffset_MatParam, GradientSettings.Offset);
+		GradientTextMaterialInstance->GetScalarParameterValue(UAvaTextMaterialHub::FMaterialStatics::GradientSmoothness_MatParam, GradientSettings.Smoothness);
+
+		const bool bRotationRetrieved = GradientTextMaterialInstance->GetScalarParameterValue(UAvaTextMaterialHub::FMaterialStatics::GradientRotation_MatParam, GradientSettings.Rotation);
+		
+		if (bRotationRetrieved)
+		{
+			if (GradientSettings.Rotation == UAvaTextMaterialHub::FMaterialStatics::MaterialMaskTopBottomRotationValue)
+			{
+				GradientSettings.Direction = EAvaGradientDirection::Vertical;
+			}
+			else if (GradientSettings.Rotation == UAvaTextMaterialHub::FMaterialStatics::MaterialMaskLeftToRightRotationValue)
+			{
+				GradientSettings.Direction = EAvaGradientDirection::Horizontal;
+			}
+			else
+			{
+				GradientSettings.Direction = EAvaGradientDirection::Custom;
+			}
+		}
+
+		RetrieveTranslucentMaterialParameters();
+	}
+}
+
+void UAvaText3DComponent::RetrieveTexturedMaterialInstance()
+{
+	if (!FrontMaterial)
+	{
+		return;
+	}
+	
+	TexturedTextMaterialInstance = Cast<UMaterialInstanceDynamic>(FrontMaterial);
+
+	if (TexturedTextMaterialInstance)
+	{
+		float Tiling_X;
+		float Tiling_Y;
+		TexturedTextMaterialInstance->GetScalarParameterValue(UAvaTextMaterialHub::FMaterialStatics::TexturedUTiling_MatParam, Tiling_X);
+		TexturedTextMaterialInstance->GetScalarParameterValue(UAvaTextMaterialHub::FMaterialStatics::TexturedVTiling_MatParam, Tiling_Y);
+		
+		Tiling = FVector2D(Tiling_X, Tiling_Y);
+		
+		UTexture* CurrentTexture;
+		TexturedTextMaterialInstance->GetTextureParameterValue(UAvaTextMaterialHub::FMaterialStatics::MainTexture_MatParam, CurrentTexture);
+
+		if (CurrentTexture)
+		{
+			MainTexture = Cast<UTexture2D>(CurrentTexture);
+		}
+		
+		RetrieveTranslucentMaterialParameters();
+	}
+}
+
+void UAvaText3DComponent::RetrieveMaskedMaterialParameters(const UMaterialInstanceDynamic* InMaterialInstance)
+{
+	if (!InMaterialInstance)
+	{
+		return;
+	}
+	
+	InMaterialInstance->GetScalarParameterValue(UAvaTextMaterialHub::FMaterialStatics::MaskSmoothness_MatParam, MaskSmoothness);
+	InMaterialInstance->GetScalarParameterValue(UAvaTextMaterialHub::FMaterialStatics::MaskOffset_MatParam, MaskOffset);
+	InMaterialInstance->GetScalarParameterValue(UAvaTextMaterialHub::FMaterialStatics::MaskRotation_MatParam, MaskRotation);
+
+	// in case the material was saved with the deprecated value, let's fix the rotation value
+	if (MaskRotation == UAvaTextMaterialHub::FMaterialStatics::MaterialMaskRightToLeftRotationValue_DEPRECATED)
+	{
+		MaskRotation = UAvaTextMaterialHub::FMaterialStatics::MaterialMaskRightToLeftRotationValue;
+	}
+
+	if (MaskRotation == UAvaTextMaterialHub::FMaterialStatics::MaterialMaskLeftToRightRotationValue)
+	{
+		MaskOrientation = EAvaMaterialMaskOrientation::LeftRight;
+	}
+	else if (MaskRotation == UAvaTextMaterialHub::FMaterialStatics::MaterialMaskRightToLeftRotationValue)
+	{
+		MaskOrientation = EAvaMaterialMaskOrientation::RightLeft;
+	}
+	else
+	{
+		MaskOrientation = EAvaMaterialMaskOrientation::Custom;
+	}
+}
+
+void UAvaText3DComponent::RetrieveSimpleTranslucentMaterialParameters(const UMaterialInstanceDynamic* InMaterialInstance)
+{
+	if (!InMaterialInstance)
+	{
+		return;
+	}
+	
+	InMaterialInstance->GetScalarParameterValue(UAvaTextMaterialHub::FMaterialStatics::Opacity_MatParam, Opacity);
+}
+
+void UAvaText3DComponent::RetrieveCustomMaterialReference()
+{
+	if (!FrontMaterial)
+	{
+		return;
+	}
+	
+	CustomMaterial = FrontMaterial;
+}
+
+void UAvaText3DComponent::RetrieveCurrentMaterials()
+{
+	switch (ColoringStyle)
+	{
+		case EAvaTextColoringStyle::Invalid:
+			break;
+
+		case EAvaTextColoringStyle::Solid:
+			RetrieveSolidMaterialInstances();
+			break;
+
+		case EAvaTextColoringStyle::Gradient:
+			RetrieveGradientMaterialInstance();
+			break;
+
+		case EAvaTextColoringStyle::FromTexture: 
+			RetrieveTexturedMaterialInstance();
+			break;
+
+		case EAvaTextColoringStyle::CustomMaterial: 
+			RetrieveCustomMaterialReference();
+			break;
+	}
+}
+
+void UAvaText3DComponent::RetrieveCurrentColoringStyle()
+{
+	if (!FrontMaterial)
+	{
+		return;
+	}
+	
+	if (UMaterialInterface* CurrentFrontMaterial = FrontMaterial->GetMaterial())
+	{
+		if (const FAvaTextMaterialSettings* MaterialSettings = UAvaTextMaterialHub::GetSettingsFromMaterial(CurrentFrontMaterial))
+		{
+			ColoringStyle = MaterialSettings->ColoringStyle;
+			
+			if (EnumHasAnyFlags(MaterialSettings->MaterialFeatures, EAvaTextMaterialFeatures::GradientMask))
+			{
+				TranslucencyStyle = EAvaTextTranslucency::GradientMask;
+			}
+
+			if (EnumHasAnyFlags(MaterialSettings->MaterialFeatures, EAvaTextMaterialFeatures::Translucent))
+			{
+				TranslucencyStyle = EAvaTextTranslucency::Translucent;
+			}
+			
+			bIsUnlit = EnumHasAnyFlags(MaterialSettings->MaterialFeatures, EAvaTextMaterialFeatures::Unlit);
+		}		
+		else // if there is no match, best guess is custom material
+		{
+			ColoringStyle = EAvaTextColoringStyle::CustomMaterial;
+			CustomMaterial = CurrentFrontMaterial;
+			TranslucencyStyle = EAvaTextTranslucency::None;
+		}
+	}
+}
+
+void UAvaText3DComponent::RetrieveCurrentTextAndFont()
+{
+	// if default had to be loaded (e.g. font asset missing), then we regenerate geometry
+	if (AvalancheFont.IsDefaultFont())
+	{
+		if (IsValid(Font) && Font != AvalancheFont.GetFont())
+		{
+			AvalancheFont.InitFromFont(Font);
+		}
+		
+		MarkForGeometryRefresh();
+	}
+	
+	if (!IsValid(Font))
+	{
+		MarkForGeometryRefresh();
+	}
+}
+
+void UAvaText3DComponent::RetrieveCurrentAlignment()
+{
+	Alignment.HorizontalAlignment = HorizontalAlignment;
+	Alignment.VerticalAlignment = VerticalAlignment;
+}
+
+EAvaTextMaterialFeatures UAvaText3DComponent::GetMaterialFeaturesFromProperties() const
+{
+	EAvaTextMaterialFeatures MaterialFeatures = EAvaTextMaterialFeatures::None;
+	
+	if (TranslucencyStyle == EAvaTextTranslucency::GradientMask)
+	{
+		MaterialFeatures |= EAvaTextMaterialFeatures::GradientMask;
+	}
+	else if (TranslucencyStyle == EAvaTextTranslucency::Translucent)
+	{
+		MaterialFeatures |= EAvaTextMaterialFeatures::Translucent;
+	}
+
+	MaterialFeatures |= bIsUnlit ? EAvaTextMaterialFeatures::Unlit : EAvaTextMaterialFeatures::None;
+
+	return MaterialFeatures;
+}
diff --git a/Engine/Plugins/Experimental/Avalanche/Source/AvalancheText/Private/AvaTextActor.cpp b/Engine/Plugins/Experimental/Avalanche/Source/AvalancheText/Private/AvaTextActor.cpp
new file mode 100644
--- /dev/null
+++ b/Engine/Plugins/Experimental/Avalanche/Source/AvalancheText/Private/AvaTextActor.cpp
@@ -0,0 +1,69 @@
+// Copyright Epic Games, Inc. All Rights Reserved.
+
+#include "AvaTextActor.h"
+#include "AvaDefs.h"
+#include "AvaTextCharacterTransform.h"
+
+AAvaTextActor::AAvaTextActor()
+{
+	Text3DComponent = CreateDefaultSubobject<UAvaText3DComponent>(TEXT("Text3DComponent"));
+	RootComponent = Text3DComponent;
+	
+	Text3DCharacterTransform = CreateDefaultSubobject<UAvaTextCharacterTransform>(TEXT("Text3DCharacterTransform"));
+	Text3DCharacterTransform->SetupAttachment(RootComponent);
+}
+
+FAvaColorChangeData AAvaTextActor::GetColorData() const
+{
+	FAvaColorChangeData ColorData;
+	ColorData.ColorStyle = EAvaColorStyle::None;
+
+	if (const UAvaText3DComponent* const AvaTextComponent = Cast<UAvaText3DComponent>(Text3DComponent))
+	{
+		switch (AvaTextComponent->GetColoringStyle())
+		{
+		case EAvaTextColoringStyle::Solid:
+			ColorData.ColorStyle = EAvaColorStyle::Solid;
+			ColorData.PrimaryColor = AvaTextComponent->GetColor();
+			ColorData.SecondaryColor = AvaTextComponent->GetColor();
+			break;
+
+		case EAvaTextColoringStyle::Gradient:
+			ColorData.ColorStyle = EAvaColorStyle::LinearGradient;
+			ColorData.PrimaryColor = AvaTextComponent->GetGradientSettings().ColorA;
+			ColorData.SecondaryColor = AvaTextComponent->GetGradientSettings().ColorB;
+			break;
+
+		default:
+			break;
+		}
+		
+		ColorData.bIsUnlit = AvaTextComponent->GetIsUnlit();
+	}
+
+	return ColorData;
+}
+
+void AAvaTextActor::SetColorData(const FAvaColorChangeData& NewColorData)
+{
+	if (UAvaText3DComponent* const AvaTextComponent = Cast<UAvaText3DComponent>(Text3DComponent))
+	{
+		switch (NewColorData.ColorStyle)
+		{
+			case EAvaColorStyle::Solid:
+				AvaTextComponent->SetColoringStyle(EAvaTextColoringStyle::Solid);
+				AvaTextComponent->SetColor(NewColorData.PrimaryColor);
+				break;
+
+			case EAvaColorStyle::LinearGradient:
+				AvaTextComponent->SetColoringStyle(EAvaTextColoringStyle::Gradient);
+				AvaTextComponent->SetGradientColors(NewColorData.PrimaryColor, NewColorData.SecondaryColor);
+				break;
+
+			default:
+				break;
+		}
+
+		AvaTextComponent->SetIsUnlit(NewColorData.bIsUnlit);
+	}
+}
diff --git a/Engine/Plugins/Experimental/Avalanche/Source/AvalancheText/Private/AvaTextModule.cpp b/Engine/Plugins/Experimental/Avalanche/Source/AvalancheText/Private/AvaTextModule.cpp
new file mode 100644
--- /dev/null
+++ b/Engine/Plugins/Experimental/Avalanche/Source/AvalancheText/Private/AvaTextModule.cpp
@@ -0,0 +1,5 @@
+// Copyright Epic Games, Inc. All Rights Reserved.
+
+#include "Modules/ModuleManager.h"
+
+IMPLEMENT_MODULE(FDefaultModuleImpl, AvalancheText)
diff --git a/Engine/Plugins/Experimental/Avalanche/Source/AvalancheText/Private/Font/AvaFont.cpp b/Engine/Plugins/Experimental/Avalanche/Source/AvalancheText/Private/Font/AvaFont.cpp
new file mode 100644
--- /dev/null
+++ b/Engine/Plugins/Experimental/Avalanche/Source/AvalancheText/Private/Font/AvaFont.cpp
@@ -0,0 +1,442 @@
+// Copyright Epic Games, Inc. All Rights Reserved.
+
+#include "Font/AvaFont.h"
+#include "AssetRegistry/AssetData.h"
+#include "AssetRegistry/AssetRegistryModule.h"
+
+#define LOCTEXT_NAMESPACE "AvaFont"
+
+FAvaFont::FAvaDefaultFontObjects FAvaFont::DefaultFontObjects;
+
+UFont* FAvaFont::GetDefaultFont()
+{
+	if (IsRunningDedicatedServer())
+	{
+		return nullptr;
+	}
+
+	if (!DefaultFontObjects.AvaDefaultFont)
+	{
+		const FAssetRegistryModule& AssetRegistryModule = FModuleManager::LoadModuleChecked<FAssetRegistryModule>("AssetRegistry");
+
+		// Trying to load Roboto font
+		const FSoftObjectPath RobotoPath(TEXT("/Script/Engine.Font'/Engine/EngineFonts/Roboto.Roboto'"));
+		UFont* DefaultFont = Cast<UFont>(AssetRegistryModule.Get().GetAssetByObjectPath(RobotoPath).GetAsset());
+
+		if (!DefaultFont)
+		{
+			// no Roboto was found, this is not expected. Trying to load an avalanche font
+			TArray<FAssetData> AssetDataArray;
+			// avalanche fonts path
+			const FName Path = TEXT("/Game/SystemFonts/Fonts");
+			AssetRegistryModule.Get().GetAssetsByPath(Path, AssetDataArray);
+
+			// no avalanche font available, trying to load the first engine font available
+			if (AssetDataArray.IsEmpty())
+			{
+				const UClass* Class = UFont::StaticClass();
+				const FTopLevelAssetPath AssetPath(Class->GetPathName());
+				AssetRegistryModule.Get().GetAssetsByClass(AssetPath, AssetDataArray);
+			}
+
+			for (const FAssetData& AssetData : AssetDataArray)
+			{
+				DefaultFont = CastChecked<UFont>(AssetData.GetAsset());
+
+				if (DefaultFont)
+				{
+					break;
+				}
+			}
+		}
+
+		if (ensureMsgf(DefaultFont, TEXT("AvalancheFont: cannot load any font to be used as default.")))
+		{
+			DefaultFontObjects.AvaDefaultFont = DefaultFont;
+		}
+	}
+
+	return DefaultFontObjects.AvaDefaultFont.Get();
+}
+
+FString FAvaFont::GenerateAvalancheFontFormattedString(const FString& InFontName, const FString& InFontObjectPathName)
+{
+	// e.g. (Property1=Value1,Property2=Value2,Property3=Value3,...)
+
+	FString FormattedString = TEXT("(");
+	FormattedString += TEXT("CurrentFont=None");
+	FormattedString += TEXT(",AvalancheFontObject=/Script/Avalanche.AvalancheFontObject'") + InFontObjectPathName + TEXT("'");
+	FormattedString += TEXT(",FontName=\"") + InFontName  + TEXT("\"");
+	FormattedString += TEXT(")");
+
+	return FormattedString;
+}
+
+bool FAvaFont::GenerateAvalancheFontFormattedString(const UAvaFontObject* InFontObject, FString& OutFormattedString)
+{
+	if (!InFontObject)
+	{
+		return false;
+	}
+
+	OutFormattedString = GenerateAvalancheFontFormattedString(InFontObject->GetFontName(), InFontObject->GetPathName());
+	return true;
+}
+
+bool FAvaFont::AreSameFont(const FAvaFont* InFontA, const FAvaFont* InFontB)
+{
+	if (!InFontA || !InFontB)
+	{
+		return false;
+	}
+
+	return *InFontA == *InFontB;
+}
+
+void FAvaFont::InitDefaults()
+{
+	bIsFavorite = false;
+	FontAssetState = EFontAssetState::DefaultFont;
+}
+
+FAvaFont::FAvaFont()
+{
+	InitDefaults();
+	SetFontObject(GetDefaultAvaFontObject());
+}
+
+FAvaFont::FAvaFont(UAvaFontObject* InFontObject)
+{
+	InitDefaults();
+
+	if (InFontObject)
+	{
+		SetFontObject(InFontObject);
+	}
+	else
+	{
+		SetFontObject(GetDefaultAvaFontObject());
+	}
+}
+
+UFont* FAvaFont::GetFont()
+{
+	// make sure FontAssetState value is up to date
+	RefreshAssetState();
+
+	if (FontAssetState == EFontAssetState::SelectedFont)
+	{
+		EnsureUsingCurrentVersion();
+
+		if (AvalancheFontObject)
+		{
+			return AvalancheFontObject->GetFont();
+		}
+	}
+
+	// if this is the default font, or previously referenced resource is not available, return default font
+	return GetDefaultFont();
+}
+
+EAvaFontSource FAvaFont::GetFontSource() const
+{
+	if (AvalancheFontObject)
+	{
+		return AvalancheFontObject->GetSource();
+	}
+
+	return EAvaFontSource::Invalid;
+}
+
+FName FAvaFont::GetFontName() const
+{
+	return FName(GetFontNameAsString());
+}
+
+FString FAvaFont::GetFontNameAsString() const
+{
+	if (AvalancheFontObject)
+	{
+		return AvalancheFontObject->GetFontName();
+	}
+
+	return FontName;
+}
+
+FText FAvaFont::GetFontNameAsText() const
+{
+	return FText::FromName(GetFontName());
+}
+
+bool FAvaFont::IsFavorite() const
+{
+	return bIsFavorite;
+}
+
+bool FAvaFont::IsDefaultFont() const
+{
+	if (AvalancheFontObject && AvalancheFontObject->GetFont())
+	{
+		return AvalancheFontObject->GetFont() == GetDefaultFont();
+	}
+
+	return true;
+}
+
+bool FAvaFont::IsFallbackFont() const
+{
+	return FontAssetState == EFontAssetState::FallbackFont;
+}
+
+bool FAvaFont::IsMonospaced() const
+{
+	if (AvalancheFontObject)
+	{
+		return AvalancheFontObject->IsMonospaced();
+	}
+
+	return false;
+}
+
+bool FAvaFont::IsBold() const
+{
+	if (AvalancheFontObject)
+	{
+		return AvalancheFontObject->IsBold();
+	}
+
+	return false;
+}
+
+bool FAvaFont::IsItalic() const
+{
+	if (AvalancheFontObject)
+	{
+		return AvalancheFontObject->IsItalic();
+	}
+
+	return false;
+}
+
+void FAvaFont::SetFavorite(const bool bFavorite)
+{
+	bIsFavorite = bFavorite;
+}
+
+void FAvaFont::EnsureUsingCurrentVersion()
+{
+	if (CurrentFont_DEPRECATED)
+	{
+		AvalancheFontObject = NewObject<UAvaFontObject>();
+		AvalancheFontObject->InitProjectFont(CurrentFont_DEPRECATED, GetFontName().ToString());
+		CurrentFont_DEPRECATED = nullptr;
+
+		RefreshName();
+	}
+}
+
+bool FAvaFont::HasValidFont() const
+{
+	const FCompositeFont* CompositeFont = nullptr;
+
+	if (CurrentFont_DEPRECATED)
+	{
+		CompositeFont = CurrentFont_DEPRECATED->GetCompositeFont();
+	}
+	else if (AvalancheFontObject && AvalancheFontObject->GetFont())
+	{
+		CompositeFont = AvalancheFontObject->GetFont()->GetCompositeFont();
+	}
+
+	if (CompositeFont)
+	{
+		if (!CompositeFont->DefaultTypeface.Fonts.IsEmpty())
+		{
+			if (CompositeFont->DefaultTypeface.Fonts[0].Font.GetFontFaceAsset())
+			{
+				return true;
+			}
+		}
+	}
+
+	return false;
+}
+
+void FAvaFont::SetFontObject(UAvaFontObject* InFontObject)
+{
+	CurrentFont_DEPRECATED = nullptr;
+	AvalancheFontObject = InFontObject;
+
+	RefreshName();
+	RefreshAssetState();
+}
+
+void FAvaFont::InitFromFont(UFont* InFont)
+{
+	if (InFont)
+	{
+		FString Name;
+		UE::Avalanche::FontUtilities::Public::GetFontName(InFont, Name);
+
+		AvalancheFontObject = NewObject<UAvaFontObject>();
+		AvalancheFontObject->InitProjectFont(InFont, Name);
+
+		RefreshName();
+	}
+}
+
+void FAvaFont::PostSerialize(const FArchive& Ar)
+{
+	if (Ar.IsLoading())
+	{
+		RefreshName();
+
+		if ((AvalancheFontObject && AvalancheFontObject->GetFont()) || CurrentFont_DEPRECATED)
+		{
+			FontAssetState = EFontAssetState::SelectedFont;
+		}
+		else
+		{
+			/*
+			 * AvalancheFontObject might be invalid due to a previously existing issue.
+			 * That would cause the Editor only version of AvalancheFontObject to be serialized, instead of the proper one.
+			 * For assets using the font saved like that, the AvalancheFontObject had the wrong Outer, which worked fine in Editor, but not for Game/Runtime.
+			 * The following code tries to find a font based on the font name, and then to create a new AvalancheFontObject, with the proper Outer
+			 */
+
+			bool bFontRecoverySuccess = false;
+
+			TArray<FProperty*> OutProperties;
+			Ar.GetSerializedPropertyChain(OutProperties);
+
+			// this should just contain the FAvaFont
+			for (const FProperty* Property : OutProperties)
+			{
+				// let's get the Outer we need to create the AvalancheFontObject
+				if (UObject* const Outer = Property->GetOwner<UObject>())
+				{
+					// try to get a font just using the name (this might fail!)
+					if (UFont* const Font = GetFontByName(FontName))
+					{
+						AvalancheFontObject = nullptr;
+
+						// we create and assign a new UAvaFontObject with the proper Outer
+						AvalancheFontObject = NewObject<UAvaFontObject>(Outer, GetFontName(), RF_Public | RF_Standalone);
+						AvalancheFontObject->InitProjectFont(Font, FontName);
+						bFontRecoverySuccess = true;
+						break;
+					}
+				}
+			}
+
+			// if we managed to retrieve the font, just act as if it was properly loaded
+			if (bFontRecoverySuccess)
+			{
+				FontAssetState = EFontAssetState::SelectedFont;
+			}
+			else
+			{
+				MissingFontName = FontName;
+				FontAssetState = EFontAssetState::FallbackFont;
+			}
+		}
+	}
+}
+
+UAvaFontObject* FAvaFont::GetDefaultAvaFontObject()
+{
+	if (!DefaultFontObjects.AvaDefaultFontObject)
+	{
+		if (UFont* DefaultFont = GetDefaultFont())
+		{
+			DefaultFontObjects.AvaDefaultFontObject = NewObject<UAvaFontObject>(DefaultFont);
+			DefaultFontObjects.AvaDefaultFontObject->InitProjectFont(DefaultFont, DefaultFont->GetName());
+		}
+	}
+
+	return DefaultFontObjects.AvaDefaultFontObject.Get();
+}
+
+void FAvaFont::RefreshName()
+{
+	if (AvalancheFontObject)
+	{
+		const FString& FontObjectName = AvalancheFontObject->GetFontName();
+
+		if (FontName != FontObjectName)
+		{
+			FontName = FontObjectName;
+		}
+	}
+}
+
+void FAvaFont::RefreshAssetState()
+{
+	if (!IsDefaultFont())
+	{
+		FontAssetState = EFontAssetState::SelectedFont;
+	}
+}
+
+UFont* FAvaFont::GetFontByName(const FString& InFontName)
+{
+	const FAssetRegistryModule& AssetRegistryModule = FModuleManager::LoadModuleChecked<FAssetRegistryModule>("AssetRegistry");
+	// no Roboto was found, this is not expected. Trying to load an avalanche font
+	TArray<FAssetData> AssetDataArray;
+
+	// no avalanche font available, trying to load the first engine font available
+	if (AssetDataArray.IsEmpty())
+	{
+		const UClass* Class = UFont::StaticClass();
+		const FTopLevelAssetPath AssetPath(Class->GetPathName());
+		AssetRegistryModule.Get().GetAssetsByClass(AssetPath, AssetDataArray);
+	}
+
+	for (const FAssetData& AssetData : AssetDataArray)
+	{
+		if (UFont* const CurrFont = CastChecked<UFont>(AssetData.GetAsset()))
+		{
+			FString CurrFontName;
+			UE::Avalanche::FontUtilities::Public::GetFontName(CurrFont, CurrFontName);
+
+			if (CurrFontName == InFontName)
+			{
+				return CurrFont;
+			}
+		}
+	}
+
+	return nullptr;
+}
+
+// note: this function used to be in the AvalancheEditor module, since it was not needed at Runtime before 
+void UE::Avalanche::FontUtilities::Public::GetFontName(const UFont* InFont, FString& OutFontName)
+{
+	if (IsValid(InFont))
+	{
+		FString FontAssetName;
+		FString FontImportName = InFont->ImportOptions.FontName;
+
+		if (InFont->GetFName() == NAME_None)
+		{
+			FontAssetName = InFont->LegacyFontName.ToString();
+		}
+		else
+		{
+			FontAssetName = InFont->GetName();
+		}
+
+		// Roboto fonts are actually from the Arial family and their import name is "Arial", so we try to list them as well
+		// this will likely lead to missing spaces in their names
+		if (FontAssetName.Contains(TEXT("Roboto")) || FontImportName == TEXT("Arial"))
+		{
+			OutFontName = FontAssetName;
+		}
+		else
+		{
+			OutFontName = FontImportName;
+		}
+	}
+}
+
+#undef LOCTEXT_NAMESPACE
diff --git a/Engine/Plugins/Experimental/Avalanche/Source/AvalancheText/Private/MaterialHub/AvaTextMaterialHub.cpp b/Engine/Plugins/Experimental/Avalanche/Source/AvalancheText/Private/MaterialHub/AvaTextMaterialHub.cpp
new file mode 100644
--- /dev/null
+++ b/Engine/Plugins/Experimental/Avalanche/Source/AvalancheText/Private/MaterialHub/AvaTextMaterialHub.cpp
@@ -0,0 +1,159 @@
+// Copyright Epic Games, Inc. All Rights Reserved.
+
+#include "MaterialHub/AvaTextMaterialHub.h"
+#include "AvaTextDefs.h"
+
+UAvaTextMaterialHub* UAvaTextMaterialHub::Get()
+{
+	UAvaTextMaterialHub* MaterialHub = GetMutableDefault<UAvaTextMaterialHub>();
+
+	static bool bInitialized = false;
+	
+	if (!bInitialized)
+	{
+		// Make sure to load config, in case custom values are specified by the user
+		MaterialHub->LoadConfig();
+
+		// Initializes the internal maps used to access/lookup materials
+		MaterialHub->SetupInternalMaps();
+
+		bInitialized = true;
+	}
+
+	return MaterialHub;
+}
+
+UAvaTextMaterialHub::UAvaTextMaterialHub()
+{
+	SolidMaterial                    = TEXT("/Avalanche/Text3DResources/Materials/Solid/M_DefaultText3DMaterial.M_DefaultText3DMaterial");
+	SolidTranslucentMaterial         = TEXT("/Avalanche/Text3DResources/Materials/Solid/M_DefaultText3DMaterial_Translucent.M_DefaultText3DMaterial_Translucent");
+	SolidTranslucentUnlitMaterial    = TEXT("/Avalanche/Text3DResources/Materials/Solid/M_DefaultText3DMaterial_Translucent_Unlit.M_DefaultText3DMaterial_Translucent_Unlit");
+	SolidMaskedMaterial              = TEXT("/Avalanche/Text3DResources/Materials/Solid/M_DefaultText3DMaterial_Masked.M_DefaultText3DMaterial_Masked");
+	SolidUnlitMaterial               = TEXT("/Avalanche/Text3DResources/Materials/Solid/M_DefaultText3DMaterial_Unlit.M_DefaultText3DMaterial_Unlit");
+	SolidUnlitMaskedTextMaterial     = TEXT("/Avalanche/Text3DResources/Materials/Solid/M_DefaultText3DMaterial_Unlit_Masked.M_DefaultText3DMaterial_Unlit_Masked");
+	
+	GradientMaterial                 = TEXT("/Avalanche/Text3DResources/Materials/Gradient/M_GradientText3DMaterial.M_GradientText3DMaterial");
+	GradientTranslucentMaterial      = TEXT("/Avalanche/Text3DResources/Materials/Gradient/M_GradientText3DMaterial_Translucent.M_GradientText3DMaterial_Translucent");
+	GradientTranslucentUnlitMaterial = TEXT("/Avalanche/Text3DResources/Materials/Gradient/M_GradientText3DMaterial_Translucent_Unlit.M_GradientText3DMaterial_Translucent_Unlit");
+	GradientMaskedMaterial           = TEXT("/Avalanche/Text3DResources/Materials/Gradient/M_GradientText3DMaterial_Masked.M_GradientText3DMaterial_Masked");
+	GradientUnlitMaterial            = TEXT("/Avalanche/Text3DResources/Materials/Gradient/M_GradientText3DMaterial_Unlit.M_GradientText3DMaterial_Unlit");
+	GradientUnlitMaskedMaterial      = TEXT("/Avalanche/Text3DResources/Materials/Gradient/M_GradientText3DMaterial_Unlit_Masked.M_GradientText3DMaterial_Unlit_Masked");
+
+	TexturedMaterial                 = TEXT("/Avalanche/Text3DResources/Materials/Textured/M_TexturedText3DMaterial.M_TexturedText3DMaterial");
+	TexturedTranslucentMaterial      = TEXT("/Avalanche/Text3DResources/Materials/Textured/M_TexturedText3DMaterial_Translucent.M_TexturedText3DMaterial_Translucent");
+	TexturedTranslucentUnlitMaterial = TEXT("/Avalanche/Text3DResources/Materials/Textured/M_TexturedText3DMaterial_Translucent_Unlit.M_TexturedText3DMaterial_Translucent_Unlit");
+	TexturedMaskedMaterial           = TEXT("/Avalanche/Text3DResources/Materials/Textured/M_TexturedText3DMaterial_Masked.M_TexturedText3DMaterial_Masked");
+	TexturedUnlitMaterial            = TEXT("/Avalanche/Text3DResources/Materials/Textured/M_TexturedText3DMaterial_Unlit.M_TexturedText3DMaterial_Unlit");
+	TexturedUnlitMaskedMaterial      = TEXT("/Avalanche/Text3DResources/Materials/Textured/M_TexturedText3DMaterial_Unlit_Masked.M_TexturedText3DMaterial_Unlit_Masked");
+}
+
+UMaterialInterface* UAvaTextMaterialHub::GetMaterial(FAvaTextMaterialSettings InMaterialSettings)
+{
+	return Get()->GetMaterialInternal(InMaterialSettings);
+}
+
+FAvaTextMaterialSettings* UAvaTextMaterialHub::GetSettingsFromMaterial(const UMaterialInterface* InMaterial)
+{
+	return Get()->GetSettingsFromMaterialInternal(InMaterial);
+}
+
+void UAvaTextMaterialHub::SetupInternalMaps()
+{
+	AddMaterial(SolidMaterial,
+			FAvaTextMaterialSettings(EAvaTextColoringStyle::Solid, EAvaTextMaterialFeatures::None));
+
+	AddMaterial(SolidTranslucentMaterial,
+			FAvaTextMaterialSettings(EAvaTextColoringStyle::Solid, EAvaTextMaterialFeatures::Translucent));
+
+	AddMaterial(SolidTranslucentUnlitMaterial,
+			FAvaTextMaterialSettings(EAvaTextColoringStyle::Solid, EAvaTextMaterialFeatures::Translucent | EAvaTextMaterialFeatures::Unlit));
+		
+	AddMaterial(SolidMaskedMaterial,
+		FAvaTextMaterialSettings(EAvaTextColoringStyle::Solid, EAvaTextMaterialFeatures::GradientMask));
+
+	AddMaterial(SolidUnlitMaterial,
+		FAvaTextMaterialSettings(EAvaTextColoringStyle::Solid, EAvaTextMaterialFeatures::Unlit));
+	
+	AddMaterial(SolidUnlitMaskedTextMaterial,
+		FAvaTextMaterialSettings(EAvaTextColoringStyle::Solid, EAvaTextMaterialFeatures::GradientMask | EAvaTextMaterialFeatures::Unlit));
+	
+	AddMaterial(GradientMaterial,
+		FAvaTextMaterialSettings(EAvaTextColoringStyle::Gradient, EAvaTextMaterialFeatures::None));
+
+	AddMaterial(GradientTranslucentMaterial,
+			FAvaTextMaterialSettings(EAvaTextColoringStyle::Gradient, EAvaTextMaterialFeatures::Translucent));
+
+	AddMaterial(GradientTranslucentUnlitMaterial,
+			FAvaTextMaterialSettings(EAvaTextColoringStyle::Gradient, EAvaTextMaterialFeatures::Translucent | EAvaTextMaterialFeatures::Unlit));
+	
+	AddMaterial(GradientMaskedMaterial,
+		FAvaTextMaterialSettings(EAvaTextColoringStyle::Gradient, EAvaTextMaterialFeatures::GradientMask));
+	
+	AddMaterial(GradientUnlitMaterial,
+		FAvaTextMaterialSettings(EAvaTextColoringStyle::Gradient, EAvaTextMaterialFeatures::Unlit));
+	
+	AddMaterial(GradientUnlitMaskedMaterial,
+		FAvaTextMaterialSettings(EAvaTextColoringStyle::Gradient, EAvaTextMaterialFeatures::GradientMask | EAvaTextMaterialFeatures::Unlit));
+	
+	AddMaterial(TexturedMaterial,
+		FAvaTextMaterialSettings(EAvaTextColoringStyle::FromTexture, EAvaTextMaterialFeatures::None));
+
+	AddMaterial(TexturedTranslucentMaterial,
+			FAvaTextMaterialSettings(EAvaTextColoringStyle::FromTexture, EAvaTextMaterialFeatures::Translucent));
+
+	AddMaterial(TexturedTranslucentUnlitMaterial,
+			FAvaTextMaterialSettings(EAvaTextColoringStyle::FromTexture, EAvaTextMaterialFeatures::Translucent | EAvaTextMaterialFeatures::Unlit));
+	
+	AddMaterial(TexturedMaskedMaterial,
+		FAvaTextMaterialSettings(EAvaTextColoringStyle::FromTexture, EAvaTextMaterialFeatures::GradientMask));
+	
+	AddMaterial(TexturedUnlitMaterial,
+		FAvaTextMaterialSettings(EAvaTextColoringStyle::FromTexture, EAvaTextMaterialFeatures::Unlit));
+	
+	AddMaterial(TexturedUnlitMaskedMaterial,
+		FAvaTextMaterialSettings(EAvaTextColoringStyle::FromTexture, EAvaTextMaterialFeatures::GradientMask | EAvaTextMaterialFeatures::Unlit));
+}
+
+void UAvaTextMaterialHub::AddMaterial(
+	const FSoftObjectPath& InMaterialSoftObjectPath
+	, FAvaTextMaterialSettings InMaterialSettings)
+{
+	MaterialsMap.Add(InMaterialSettings, InMaterialSoftObjectPath);
+
+	FAvaTextMaterialId MaterialId(InMaterialSoftObjectPath);
+	FAvaTextMaterialData MaterialData(InMaterialSoftObjectPath, InMaterialSettings);
+
+	MaterialDataMap.Add(MoveTemp(MaterialId), MoveTemp(MaterialData));
+}
+
+UMaterialInterface* UAvaTextMaterialHub::GetMaterialInternal(FAvaTextMaterialSettings InMaterialSettings) const
+{
+	if (MaterialsMap.Contains(InMaterialSettings))
+	{
+		if (const FSoftObjectPath* MaterialMatch = MaterialsMap.Find(InMaterialSettings))
+		{
+			return Cast<UMaterialInterface>(MaterialMatch->TryLoad());
+		}
+	}
+	
+	return nullptr;
+}
+
+FAvaTextMaterialSettings* UAvaTextMaterialHub::GetSettingsFromMaterialInternal(const UMaterialInterface* InMaterial)
+{
+	if (!InMaterial)
+	{
+		return nullptr;
+	}
+	
+	const FAvaTextMaterialId CurrentMaterialId(InMaterial);
+
+	// look for current material in static material data map, to see if it's a known one
+	if (FAvaTextMaterialData* MaterialDataMatch = MaterialDataMap.Find(CurrentMaterialId))
+	{
+		// retrieve material info
+		return &MaterialDataMatch->MaterialSettings;
+	}
+	
+	return nullptr;
+}
diff --git a/Engine/Plugins/Experimental/Avalanche/Source/AvalancheText/Public/AvaText3DComponent.h b/Engine/Plugins/Experimental/Avalanche/Source/AvalancheText/Public/AvaText3DComponent.h
new file mode 100644
--- /dev/null
+++ b/Engine/Plugins/Experimental/Avalanche/Source/AvalancheText/Public/AvaText3DComponent.h
@@ -0,0 +1,352 @@
+// Copyright Epic Games, Inc. All Rights Reserved.
+
+#pragma once
+
+#include "AvaTextDefs.h"
+#include "Containers/Ticker.h"
+#include "Font/AvaFont.h"
+#include "Text3DComponent.h"
+#include "AvaText3DComponent.generated.h"
+
+class UMaterial;
+class UMaterialInstanceDynamic;
+
+enum class EAvaTextRefreshReason : uint8
+{
+	NoRefresh                 = 0,
+	GeometryChange            = 1 << 0,
+	LayoutChange              = 1 << 1,
+	MaterialInstancesChange   = 1 << 2,
+	MaterialParametersChange  = 1 << 3,
+	FullRefresh               = 0xFF
+};
+ENUM_CLASS_FLAGS(EAvaTextRefreshReason)
+
+UCLASS(ClassGroup="Text3D", PrioritizeCategories=("Text", "Layout", "Geometry", "Style", "Materials"), meta=(BlueprintSpawnableComponent))
+class AVALANCHETEXT_API UAvaText3DComponent : public UText3DComponent
+{
+	friend class FAvaTextVisualizer;
+	friend class FAvaToolboxTextVisualizer;
+	friend class FAvaTextComponentCustomization;
+
+	GENERATED_BODY()
+
+public:
+	UAvaText3DComponent();
+
+	UFUNCTION(BlueprintSetter)
+	void SetEnforceUpperCase(bool bInEnforceUpperCase);
+	bool IsEnforcingUpperCase() const { return bEnforceUpperCase; }
+	
+	UFUNCTION(BlueprintSetter)
+	void SetAvalancheFont(const FAvaFont& InAvalancheFont);
+	const FAvaFont& GetAvalancheFont() const { return AvalancheFont; }
+
+	UFUNCTION(BlueprintSetter)
+	void SetAlignment(FAvaTextAlignment InAlignment);
+	const FAvaTextAlignment& GetAlignment() const { return Alignment; }
+
+	UFUNCTION(BlueprintSetter)
+	void SetColoringStyle(const EAvaTextColoringStyle& InColoringStyle);
+	EAvaTextColoringStyle GetColoringStyle() const { return ColoringStyle; }
+
+	UFUNCTION(BlueprintSetter)
+	void SetColor(const FLinearColor& InColor);
+	const FLinearColor& GetColor() const { return Color; }
+
+	UFUNCTION(BlueprintSetter)
+	void SetExtrudeColor(const FLinearColor& InColor);
+	const FLinearColor& GetExtrudeColor() const { return ExtrudeColor; }
+
+	UFUNCTION(BlueprintSetter)
+	void SetBevelColor(const FLinearColor& InColor);
+	const FLinearColor& GetBevelColor() const { return BevelColor; }
+
+	UFUNCTION(BlueprintSetter)
+	void SetGradientSettings(const FAvaLinearGradientSettings& InGradientSettings);
+	const FAvaLinearGradientSettings& GetGradientSettings() const { return GradientSettings; }
+
+	UFUNCTION(BlueprintSetter)
+	void SetGradientColors(const FLinearColor& InColorA, const FLinearColor& InColorB);
+
+	UFUNCTION(BlueprintSetter)
+	void SetMainTexture(UTexture2D* InMainTexture);
+	UTexture2D* GetMainTexture() const { return MainTexture; }
+
+	UFUNCTION(BlueprintSetter)
+	void SetTiling(const FVector2D InTiling);
+	const FVector2D& GetTiling() const { return Tiling; }
+
+	UFUNCTION(BlueprintSetter)
+	void SetCustomMaterial(UMaterialInterface* InCustomMaterial);
+	UMaterialInterface* GetCustomMaterial() const { return CustomMaterial; }
+
+	UFUNCTION(BlueprintSetter)
+	void SetTranslucencyStyle(EAvaTextTranslucency InTranslucencyStyle);
+	EAvaTextTranslucency GetTranslucencyStyle() const { return TranslucencyStyle; }
+
+	UFUNCTION(BlueprintSetter)
+	void SetOpacity(float InOpacity);
+	float GetOpacity() const { return Opacity; }
+	
+	UFUNCTION(BlueprintSetter)
+	void SetMaskOrientation(EAvaMaterialMaskOrientation InMaskOrientation);
+	EAvaMaterialMaskOrientation GetMaskOrientation() const { return MaskOrientation; }
+
+	UFUNCTION(BlueprintSetter)
+	void SetMaskSmoothness(float InMaskSmoothness);
+	float GetMaskSmoothness() const { return MaskSmoothness; }
+
+	UFUNCTION(BlueprintSetter)
+	void SetMaskOffset(float InMaskOffset);
+	float GetMaskOffset() const { return MaskOffset; }
+
+	UFUNCTION(BlueprintSetter)
+	void SetMaskRotation(float InMaskRotation);
+	float GetMaskRotation() const { return MaskRotation; }
+
+	UFUNCTION(BlueprintSetter)
+	void SetIsUnlit(bool bInIsUnlit);
+	bool GetIsUnlit() const { return bIsUnlit; }
+
+#if WITH_EDITOR
+	static void RegisterOnPropertyChangeFunctions();
+#endif
+
+	void RefreshOpacity() const;
+	
+	void RefreshColor() const;
+	void RefreshExtrudeColor() const;
+	void RefreshBevelColor() const;
+
+	void RefreshAlignment();
+	void RefreshAvalancheFont();
+	void RefreshHorizontalAlignment();
+	void RefreshVerticalAlignment();
+	void RefreshUnlit();
+	void RefreshMasked();
+	void RefreshMaterialInstances();
+	void RefreshGeometry();
+	void RefreshLayout();
+	void RefreshGradientValues();
+	void RefreshGradient();
+	void RefreshMaskValues();
+	void RefreshMask();
+	void RefreshMaterialGeometryParameterValues();
+
+	void RefreshFormatting();
+
+	void SetMaterialGeometryParameterValues(UMaterialInstanceDynamic* InMaterial, const FVector& TextScaleFactor) const;
+	void SetMaskedMaterialValues(UMaterialInstanceDynamic* InMaterial);
+	using FOnGetMaterialWithSettings = TDelegate<UMaterialInterface*(const UMaterialInterface* InPreviousMaterial, const FAvaTextMaterialSettings& InSettings)>;
+	FOnGetMaterialWithSettings& GetMaterialProviderDelegate();
+	void ForEachMID(TUniqueFunction<void(UMaterialInstanceDynamic* InMID)>&& InFunc) const;
+
+	FVector GetGradientDirection() const;
+
+#if WITH_EDITOR
+	DECLARE_MULTICAST_DELEGATE(FOnEditMoveDelegate);
+	FOnEditMoveDelegate& OnEditMove() { return OnEditMoveDelegate; } // used by visualizer
+
+	virtual void PostEditChangeProperty(FPropertyChangedEvent& PropertyChangedEvent) override;
+	virtual void PostEditComponentMove(bool bFinished) override;
+#endif
+
+	//~ Begin USceneComponent
+	virtual void OnUpdateTransform(EUpdateTransformFlags UpdateTransformFlags, ETeleportType Teleport) override;
+	virtual void OnAttachmentChanged() override;
+	virtual void OnRegister() override;
+	//~ End USceneComponent
+
+	virtual void FormatText(FText& InOutText) const override;
+
+private:
+	UPROPERTY(EditAnywhere, Getter="IsEnforcingUpperCase", Setter="SetEnforceUpperCase", Category="Text", meta=(DisplayAfter="Text", AllowPrivateAccess="true"))
+	bool bEnforceUpperCase = false;
+	
+	UPROPERTY(EditAnywhere, Getter="GetAvalancheFont", Setter="SetAvalancheFont", Category="Text", meta=(DisplayAfter="Text", AllowPrivateAccess="true"))
+	FAvaFont AvalancheFont;
+
+	UPROPERTY(EditAnywhere, Getter="GetAlignment", Setter="SetAlignment", Category="Layout", meta=(DisplayAfter="VerticalAlignment", AllowPrivateAccess="true"))
+	FAvaTextAlignment Alignment;
+
+	UPROPERTY(EditAnywhere, Getter="GetColoringStyle", Setter="SetColoringStyle", Category="Style", meta=(DisplayAfter="Text", AllowPrivateAccess="true"))
+	EAvaTextColoringStyle ColoringStyle = EAvaTextColoringStyle::Solid;
+
+	UPROPERTY(EditAnywhere, Getter="GetColor", Setter="SetColor", Category="Style", meta=(EditCondition="ColoringStyle == EAvaTextColoringStyle::Solid", EditConditionHides, DisplayAfter="Text", AllowPrivateAccess="true"))
+	FLinearColor Color = FLinearColor::White;
+
+	UPROPERTY(EditAnywhere, Getter="GetExtrudeColor", Setter="SetExtrudeColor", Category="Style", meta=(EditCondition="ColoringStyle == EAvaTextColoringStyle::Solid && Extrude > 0.0", EditConditionHides, DisplayAfter="Text", AllowPrivateAccess="true"))
+	FLinearColor ExtrudeColor = FLinearColor::Gray;
+
+	UPROPERTY(EditAnywhere, Getter="GetBevelColor", Setter="SetBevelColor", Category="Style", meta=(EditCondition="ColoringStyle == EAvaTextColoringStyle::Solid && Bevel > 0.0", EditConditionHides, DisplayAfter="Text", AllowPrivateAccess="true"))
+	FLinearColor BevelColor = FLinearColor::Gray;
+
+	UPROPERTY(EditAnywhere, Getter="GetGradientSettings", Setter="SetGradientSettings", Category="Style", meta=(EditCondition="ColoringStyle == EAvaTextColoringStyle::Gradient", EditConditionHides, DisplayAfter="Text", AllowPrivateAccess="true"))
+	FAvaLinearGradientSettings GradientSettings;
+
+	UPROPERTY(EditAnywhere, Getter="GetMainTexture", Setter="SetMainTexture", Category="Style", meta=(EditCondition="ColoringStyle == EAvaTextColoringStyle::FromTexture", EditConditionHides, DisplayAfter="Text", AllowPrivateAccess="true"))
+	TObjectPtr<UTexture2D> MainTexture;
+
+	UPROPERTY(EditAnywhere, Getter="GetTiling", Setter="SetTiling", Category="Style", meta=(Delta=0.1, EditCondition="ColoringStyle == EAvaTextColoringStyle::FromTexture", EditConditionHides, DisplayAfter="Text", AllowPrivateAccess="true"))
+	FVector2D Tiling;
+
+	UPROPERTY(EditAnywhere, Getter="GetCustomMaterial", Setter="SetCustomMaterial", Category="Style", meta=(EditCondition="ColoringStyle == EAvaTextColoringStyle::CustomMaterial", EditConditionHides, DisplayAfter="Text", AllowPrivateAccess="true"))
+	TObjectPtr<UMaterialInterface> CustomMaterial = nullptr;
+	
+	UPROPERTY(EditAnywhere, Getter="GetTranslucencyStyle", Setter="SetTranslucencyStyle", Category="Style", meta=(EditCondition="ColoringStyle != EAvaTextColoringStyle::CustomMaterial", EditConditionHides,DisplayAfter="Text", AllowPrivateAccess="true"))
+	EAvaTextTranslucency TranslucencyStyle = EAvaTextTranslucency::None;
+
+	UPROPERTY(EditAnywhere, Getter="GetOpacity", Setter="SetOpacity", Category="Style", meta=(ClampMin=0.0, ClampMax=1.0, EditCondition="TranslucencyStyle == EAvaTextTranslucency::Translucent && ColoringStyle != EAvaTextColoringStyle::CustomMaterial", EditConditionHides, DisplayAfter="Text", AllowPrivateAccess="true"))
+	float Opacity = 0.75f;
+
+	UPROPERTY(EditAnywhere, Getter="GetMaskOrientation", Setter="SetMaskOrientation", Category="Style", meta=(EditCondition="TranslucencyStyle == EAvaTextTranslucency::GradientMask", EditConditionHides, DisplayAfter="Text", AllowPrivateAccess="true"))
+	EAvaMaterialMaskOrientation MaskOrientation = EAvaMaterialMaskOrientation::LeftRight;
+
+	UPROPERTY(EditAnywhere, Getter="GetMaskSmoothness", Setter="SetMaskSmoothness", Category="Style", meta=(ClampMin=0.0, ClampMax=1.0, EditCondition="TranslucencyStyle == EAvaTextTranslucency::GradientMask", EditConditionHides, DisplayAfter="Text", AllowPrivateAccess="true"))
+	float MaskSmoothness = 0.1f;
+
+	UPROPERTY(EditAnywhere, Getter="GetMaskOffset", Setter="SetMaskOffset", Category="Style", meta=(ClampMin=0.0, ClampMax=1.0, EditCondition="TranslucencyStyle == EAvaTextTranslucency::GradientMask", EditConditionHides, DisplayAfter="Text", AllowPrivateAccess="true"))
+	float MaskOffset = 0.5f;
+
+	UPROPERTY(EditAnywhere, Getter="GetMaskRotation", Setter="SetMaskRotation", Category="Style", meta=(ClampMin=0.0, ClampMax=1.0, EditCondition="MaskOrientation == EAvaMaterialMaskOrientation::Custom && TranslucencyStyle == EAvaTextTranslucency::GradientMask", EditConditionHides, DisplayAfter="Text", AllowPrivateAccess="true"))
+	float MaskRotation = 0.0f;
+
+	UPROPERTY(EditAnywhere, Getter="GetIsUnlit", Setter="SetIsUnlit", Category="Style", meta=(EditCondition="ColoringStyle != EAvaTextColoringStyle::Custom", EditConditionHides, DisplayAfter="Text", AllowPrivateAccess="true"))
+	bool bIsUnlit = true;
+
+protected:
+	void RefreshText3D();
+	void MarkRenderDirty();
+
+	void MarkForGeometryRefresh();
+	bool IsMarkedForGeometryRefresh() const;
+
+	void MarkForLayoutRefresh();
+	bool IsMarkedForLayoutRefresh() const;
+
+	/** Mark this Text3D component so that at the next RefreshText3D call it will refresh both materials and their parameters*/
+	void MarkForMaterialsFullRefresh();
+	bool IsMarkedForMaterialRefresh() const;
+
+	void MarkForMaterialParametersRefresh();
+	bool IsMarkedForMaterialParametersRefresh() const;
+
+	void RemoveRefreshReason(EAvaTextRefreshReason InRefreshReasonToRemove);
+
+	void MarkForFullRefresh();
+	void ClearRefreshReasons();
+
+	void RegisterCallbacks();
+	void RefreshParentRootComponent();
+	void SetupParentRootAndCallbacks();
+
+	void SetupTextMaterials(const EAvaTextColoringStyle& InColoringStyle);
+	void SetupTextGeometry();
+	void SetupTextLayout();
+
+	void SetupSolidMaterial();
+	void SetupGradientMaterial();
+	void SetupCustomMaterial();
+	void RefreshTranslucencyValues();
+	void SetupTexturedMaterial();
+
+	void RefreshStoredBoundsValues();
+
+	//~ Begin UObject
+	virtual void PostLoad() override;
+	virtual void PostDuplicate(EDuplicateMode::Type DuplicateMode) override;
+	virtual void PostInitProperties() override;
+	virtual void PostEditImport() override;
+	//~ Begin UObject
+
+	void PostDuplicateActions();
+
+	/** Will schedule a Text Refresh to be executed by Core Ticker */
+	void ScheduleTextRefreshOnTicker();
+
+	TWeakObjectPtr<USceneComponent> ParentRootComponent;
+
+	/**
+	 * UMaterialInstanceDynamic references are transient.
+	 * We need to retrieve them from material when loading assets using this component, in order to avoid having to create the material again.
+	 * Refer to UAvaText3DComponent::RetrieveCurrentMaterials().
+	 *
+	 * Additionally, this helps when loading old versions of the AAvaTextActor, which was hosting all the custom properties now handled by this component (UAvaText3DComponent)
+	 */
+
+	UPROPERTY(Transient)
+	TObjectPtr<UMaterialInstanceDynamic> TextMaterialInstance;
+
+	UPROPERTY(Transient)
+	TObjectPtr<UMaterialInstanceDynamic> ExtrudeTextMaterialInstance;
+
+	UPROPERTY(Transient)
+	TObjectPtr<UMaterialInstanceDynamic> BevelTextMaterialInstance;
+
+	UPROPERTY(Transient)
+	TObjectPtr<UMaterialInstanceDynamic> TexturedTextMaterialInstance;
+
+	UPROPERTY(Transient)
+	TObjectPtr<UMaterialInstanceDynamic> GradientTextMaterialInstance;
+
+	FOnGetMaterialWithSettings OnGetMaterialWithSettingsDelegate;
+	UMaterialInstanceDynamic* GetMIDWithSettings(const UMaterialInterface* InPreviousMaterial, const FAvaTextMaterialSettings& InSettings);
+
+	void OnTextGeneratedActions();
+	void OnTextGeometryGenerated();
+
+	void OnRootComponentTransformUpdated(USceneComponent* InSceneComponent, EUpdateTransformFlags InUpdateTransformFlags, ETeleportType InTeleportType);
+
+	/**
+	 * This function can be used to obtain implicit data from the base UText3DComponent.
+	 * 
+	 * Previously, the logic from UAvaText3DComponent was handled by AAvaTextActor, which is now mostly a wrapper for UAvaText3DComponent.
+	 * Therefore, previously saved AAvaTextActor assets will have UText3DComponent materials compatible with UAvaText3DComponent.
+	 * This function extracts implicit information from those materials, and updates UAvaText3DComponent properties accordingly.
+	 * Furthermore, materials themselves will be retrieved, so that this component won't have to store them as well.
+	 * 
+	 * Additionally, also text, font and layout information will be retrieved, to avoid potential data loss from older assets.
+	 */
+	void RetrieveDataFromBaseComponent();
+
+	/** Retrieves font and text from the values stored in the component */
+	void RetrieveCurrentTextAndFont();
+
+	/** Retrieves alignment information from the values stored in the component */
+	void RetrieveCurrentAlignment();
+
+	/** Retrieves materials from the component, according to the current settings*/
+	void RetrieveCurrentMaterials();
+
+	void RetrieveCurrentColoringStyle();
+	void RetrieveSolidMaterialInstances();
+	void RetrieveGradientMaterialInstance();
+	void RetrieveTexturedMaterialInstance();
+	void RetrieveCustomMaterialReference();
+
+	void RetrieveTranslucentMaterialParameters();
+	void RetrieveMaskedMaterialParameters(const UMaterialInstanceDynamic* InMaterialInstance);
+	void RetrieveSimpleTranslucentMaterialParameters(const UMaterialInstanceDynamic* InMaterialInstance);
+
+	EAvaTextMaterialFeatures GetMaterialFeaturesFromProperties() const;
+
+#if WITH_EDITOR
+	FOnEditMoveDelegate OnEditMoveDelegate;
+#endif
+
+	/** Used to filter which refreshes need be performed when calling RefreshText3D */
+	EAvaTextRefreshReason TextRefreshReason;
+
+	FVector LocalBoundsLowerLeftCorner;
+	FVector LocalBoundsExtent;
+
+	FDelegateHandle OnTextGeometryGeneratedDelegate;
+	FDelegateHandle OnTransformUpdatedDelegate;
+	FTSTicker::FDelegateHandle TickerHandle;
+
+private:
+#if WITH_EDITOR
+	static inline TMap<FName, TFunction<void(UAvaText3DComponent*)>> PropertyChangeFunctions;
+#endif
+};
diff --git a/Engine/Plugins/Experimental/Avalanche/Source/AvalancheText/Public/AvaTextActor.h b/Engine/Plugins/Experimental/Avalanche/Source/AvalancheText/Public/AvaTextActor.h
new file mode 100644
--- /dev/null
+++ b/Engine/Plugins/Experimental/Avalanche/Source/AvalancheText/Public/AvaTextActor.h
@@ -0,0 +1,32 @@
+// Copyright Epic Games, Inc. All Rights Reserved.
+
+#pragma once
+
+#include "AvaText3DComponent.h"
+#include "GameFramework/Actor.h"
+#include "AvaTextActor.generated.h"
+
+class UAvaTextCharacterTransform;
+struct FAvaColorChangeData;
+
+UCLASS(ClassGroup = (Text3D), DisplayName = "Motion Design Text", meta = (ComponentWrapperClass))
+class AVALANCHETEXT_API AAvaTextActor : public AActor
+{
+	GENERATED_BODY()
+
+public:
+	AAvaTextActor();
+
+	/** Returns Text3D SubObject **/
+	UText3DComponent* GetText3DComponent() const { return Text3DComponent; }
+
+	FAvaColorChangeData GetColorData() const;
+	void SetColorData(const FAvaColorChangeData& InColorData);
+
+private:
+	UPROPERTY(Category = "Text", VisibleAnywhere, BlueprintReadOnly, meta = (AllowPrivateAccess = "true"))
+	TObjectPtr<UText3DComponent> Text3DComponent;
+
+	UPROPERTY(Category = "Layout", EditAnywhere, BlueprintReadWrite, meta = (AllowPrivateAccess = "true"))
+	TObjectPtr<UAvaTextCharacterTransform> Text3DCharacterTransform;
+};
diff --git a/Engine/Plugins/Experimental/Avalanche/Source/AvalancheText/Public/AvaTextCharacterTransform.h b/Engine/Plugins/Experimental/Avalanche/Source/AvalancheText/Public/AvaTextCharacterTransform.h
new file mode 100644
--- /dev/null
+++ b/Engine/Plugins/Experimental/Avalanche/Source/AvalancheText/Public/AvaTextCharacterTransform.h
@@ -0,0 +1,16 @@
+// Copyright Epic Games, Inc. All Rights Reserved.
+
+#pragma once
+
+#include "Text3DCharacterTransform.h"
+
+#include "AvaTextCharacterTransform.generated.h"
+
+/**
+ * This Component is kept for compatibility with old AAvaTextActor assets
+ */
+UCLASS(ClassGroup=(Text3D), HideCategories = (Collision, Tags, Activation, Cooking, Rendering, Physics, Mobility, LOD, AssetUserData, Navigation), meta = (BlueprintSpawnableComponent))
+class AVALANCHETEXT_API UAvaTextCharacterTransform : public UText3DCharacterTransform
+{
+	GENERATED_BODY()
+};
diff --git a/Engine/Plugins/Experimental/Avalanche/Source/AvalancheText/Public/AvaTextDefs.h b/Engine/Plugins/Experimental/Avalanche/Source/AvalancheText/Public/AvaTextDefs.h
new file mode 100644
--- /dev/null
+++ b/Engine/Plugins/Experimental/Avalanche/Source/AvalancheText/Public/AvaTextDefs.h
@@ -0,0 +1,213 @@
+﻿// Copyright Epic Games, Inc. All Rights Reserved.
+
+#pragma once
+
+#include "Text3DComponent.h"
+#include "AvaTextDefs.generated.h"
+
+UENUM()
+enum class EAvaOutlineType : uint8 
+{
+	Invalid     UMETA(Hidden),
+	None,
+	AddOutline,
+	OutlineOnly
+};
+
+UENUM()
+enum class EAvaTextColoringStyle : uint8 
+{
+	Invalid        UMETA(Hidden),
+	Solid,
+	Gradient,
+	FromTexture,
+	CustomMaterial
+};
+
+UENUM()
+enum class EAvaTextTranslucency : uint8 
+{
+	Invalid        UMETA(Hidden),
+	None,
+	Translucent,
+	GradientMask,
+};
+
+UENUM()
+enum class EAvaMaterialMaskOrientation : uint8 
+{
+	LeftRight UMETA(DisplayName = "Left to Right"),
+	RightLeft UMETA(DisplayName = "Right to Left"),
+	Custom
+};
+
+USTRUCT(BlueprintType)
+struct FAvaTextAlignment
+{
+	GENERATED_BODY()
+    
+	UPROPERTY(EditAnywhere, BlueprintReadWrite, Category = "Motion Design")
+	EText3DHorizontalTextAlignment HorizontalAlignment = EText3DHorizontalTextAlignment::Left;
+
+	UPROPERTY(EditAnywhere, BlueprintReadWrite, Category = "Motion Design")
+	EText3DVerticalTextAlignment VerticalAlignment = EText3DVerticalTextAlignment::FirstLine;
+	
+	bool operator==(const FAvaTextAlignment& Other) const
+	{
+		return HorizontalAlignment == Other.HorizontalAlignment && VerticalAlignment == Other.VerticalAlignment;
+	}
+};
+
+UENUM()
+enum class EAvaGradientDirection : uint8 
+{
+	None        UMETA(Hidden),
+	Vertical,
+	Horizontal,
+	Custom,
+};
+
+USTRUCT(BlueprintType)
+struct FAvaLinearGradientSettings
+{
+	GENERATED_BODY()
+
+	UPROPERTY(EditAnywhere, BlueprintReadWrite, Category = "Motion Design")
+	EAvaGradientDirection Direction = EAvaGradientDirection::Vertical;
+	
+	UPROPERTY(EditAnywhere, BlueprintReadWrite, Category = "Motion Design")
+	FLinearColor ColorA = FLinearColor::White;
+
+	UPROPERTY(EditAnywhere, BlueprintReadWrite, Category = "Motion Design")
+	FLinearColor ColorB = FLinearColor::Black;
+
+	UPROPERTY(EditAnywhere, BlueprintReadWrite, Category = "Motion Design", meta = (ClampMin = 0.0, ClampMax = 1.0))
+	float Smoothness = 0.1f;
+	
+	UPROPERTY(EditAnywhere, BlueprintReadWrite, Category = "Motion Design", meta = (ClampMin = 0.0, ClampMax = 1.0))
+	float Offset = 0.5f;
+
+	UPROPERTY(EditAnywhere, BlueprintReadWrite, Category = "Motion Design", meta = (ClampMin = 0.0, ClampMax = 1.0))
+	float Rotation = 0.0f;
+
+	bool operator==(const FAvaLinearGradientSettings& Other) const
+	{
+		return
+			Direction == Other.Direction &&
+			ColorA == Other.ColorA &&
+			ColorB == Other.ColorB &&
+			Smoothness == Other.Smoothness &&
+			Offset == Other.Offset &&
+			Rotation == Other.Rotation;
+	}
+};
+
+/**
+ * TODO: use this for UAvaText3DComponent in place of separate mask settings
+ * Main Blocker from using this now is that it needs custom tracks/sections/property track editor to be keyable for Sequencer
+ */
+USTRUCT(BlueprintType)
+struct FAvaMaterialMaskSettings
+{
+	GENERATED_BODY()
+	
+	UPROPERTY(EditAnywhere, BlueprintReadWrite, Category = "Text")
+	EAvaMaterialMaskOrientation Orientation = EAvaMaterialMaskOrientation::LeftRight;
+	
+	UPROPERTY(Interp, EditAnywhere, BlueprintReadWrite, Category = "Text", meta = (ClampMin = 0.0, ClampMax = 1.0))
+	float Smoothness = 0.1f;
+	
+	UPROPERTY(Interp, EditAnywhere, BlueprintReadWrite, Category = "Text", meta = (ClampMin = 0.0, ClampMax = 1.0))
+	float Offset = 0.5f;
+
+	UPROPERTY(Interp, EditAnywhere, BlueprintReadWrite, Category = "Text", meta = (ClampMin = 0.0, ClampMax = 1.0, EditCondition = "Orientation==EAvaMaterialMaskOrientation::Custom", EditConditionHides))
+	float Rotation = 0.0f;
+
+	bool operator==(const FAvaMaterialMaskSettings& Other) const
+	{
+		return
+			Orientation == Other.Orientation &&
+			Smoothness == Other.Smoothness &&
+			Offset == Other.Offset &&
+			Rotation == Other.Rotation;
+	}
+};
+
+UENUM()
+enum class EAvaTextLength : uint8 
+{
+	Unlimited,
+	CharacterCount UMETA(DisplayName = "Maximum Character Count")
+};
+
+UENUM()
+enum class EAvaTextCase : uint8 
+{
+	Regular,
+	UpperCase,
+	LowerCase
+};
+
+USTRUCT(BlueprintType)
+struct FAvaTextField
+{
+	GENERATED_BODY()
+	
+	UPROPERTY(EditAnywhere, BlueprintReadWrite, Category = "Motion Design", meta = (MultiLine = true))
+	FText Text;
+
+	UPROPERTY(EditAnywhere, BlueprintReadWrite, Category = "Motion Design")
+	EAvaTextLength MaximumLength = EAvaTextLength::Unlimited;
+
+	UPROPERTY(EditAnywhere, BlueprintReadWrite, Category = "Motion Design", meta = (EditCondition = "MaximumLength==EAvaTextLength::CharacterCount", EditConditionHides))
+	int32 Count = 20;
+
+	UPROPERTY(EditAnywhere, BlueprintReadWrite, Category = "Motion Design")
+	EAvaTextCase TextCase = EAvaTextCase::Regular;
+};
+
+/**
+ * Flags to list the desired features of a Text Material
+ */
+enum class EAvaTextMaterialFeatures : uint8
+{
+	None         = 0,
+	GradientMask = 1 << 0,
+	Unlit        = 1 << 1,
+	Translucent  = 1 << 2
+};
+ENUM_CLASS_FLAGS(EAvaTextMaterialFeatures)
+
+struct FAvaTextMaterialSettings
+{
+	/** The coloring style used for the material (e.g. Solid, Gradient, Textured) */
+	EAvaTextColoringStyle ColoringStyle = EAvaTextColoringStyle::Solid;
+
+	/** The features of the material (e.g. Masked, Unlit) */
+	EAvaTextMaterialFeatures MaterialFeatures = EAvaTextMaterialFeatures::Unlit;
+
+	FAvaTextMaterialSettings() = default;
+	
+	FAvaTextMaterialSettings(const EAvaTextColoringStyle InColoringStyle, const EAvaTextMaterialFeatures InMaterialFeatures)
+		: ColoringStyle(InColoringStyle)
+		, MaterialFeatures(InMaterialFeatures)
+	{
+	}
+
+	bool operator==(const FAvaTextMaterialSettings& Other) const
+	{
+		return
+			ColoringStyle == Other.ColoringStyle &&
+			MaterialFeatures == Other.MaterialFeatures;
+	}
+
+	bool operator!=(const FAvaTextMaterialSettings& Other) const
+	{
+		return !(*this == Other);
+	}
+	
+	friend uint32 GetTypeHash(const FAvaTextMaterialSettings& InMaterialSettings)
+	{
+		return HashCombineFast(GetTypeHash(InMaterialSettings.ColoringStyle), GetTypeHash(InMaterialSettings.MaterialFeatures));
+	}
+};
