[MotionDesign]

diff --git a/Engine/Plugins/Experimental/Avalanche/Source/AvalancheText/Public/Font/AvaFont.h b/Engine/Plugins/Experimental/Avalanche/Source/AvalancheText/Public/Font/AvaFont.h
new file mode 100644
--- /dev/null
+++ b/Engine/Plugins/Experimental/Avalanche/Source/AvalancheText/Public/Font/AvaFont.h
@@ -0,0 +1,244 @@
+// Copyright Epic Games, Inc. All Rights Reserved.
+
+#pragma once
+
+#include "AvaFontObject.h"
+#include "Engine/Font.h"
+#include "UObject/StrongObjectPtr.h"
+
+#include "AvaFont.generated.h"
+
+USTRUCT(BlueprintType)
+struct AVALANCHETEXT_API FAvaFont
+{
+	friend class UAvaText3DComponent;
+
+	GENERATED_BODY()
+
+	/**
+	 * Returns a default font (Roboto). The font will be loaded from asset if needed.
+	 * In the unlikely case in which Roboto is not available, the first available font will be returned
+	 */
+	static UFont* GetDefaultFont();
+
+	/**
+	 * Returns a formatted string representing a FAvalancheFont for the specified FontObjectPathName and FontName
+	 * Useful e.g. when setting the value of multiple FAvalancheFonts using IPropertyHandle->SetPerObjectValues
+	 * @param InFontName {The name of the font handled by UAvalancheFontObject}
+	 * @param InFontObjectPathName {The path of the UAvalancheFontObject}
+	 * @return {the formatted string}
+	 */
+	static FString GenerateAvalancheFontFormattedString(const FString& InFontName, const FString& InFontObjectPathName);
+
+	/**
+	 * Returns a formatted string representing a FAvalancheFont for the specified UAvalancheFontObject
+	 * Useful e.g. when setting the value of multiple FAvalancheFonts using IPropertyHandle->SetPerObjectValues
+	 * @param InFontObject {The UAvalancheFontObject used to generated the FAvalancheFont formatted string}
+	 * @param OutFormattedString {will contain the formatted string}
+	 * @return {true if successful}
+	 */
+	static bool GenerateAvalancheFontFormattedString(const UAvaFontObject* InFontObject, FString& OutFormattedString);
+
+	/**
+	 * Checks if the specified Ava Fonts are the same
+	 * @param InFontA First font for comparison
+	 * @param InFontB Second font for comparison
+	 * @return true if Fonts are the same
+	 */
+	static bool AreSameFont(const FAvaFont* InFontA, const FAvaFont* InFontB);
+
+	/** Default constructor: calling GetFont() right after the constructor will return the default font */
+	FAvaFont();
+
+	/** this FAvalancheFont will be initialized using the provided UAvalancheFontObject */
+	FAvaFont(UAvaFontObject* InFontObject);
+
+	/** Returns the currently stored font, or default font if none available. Preferred getter to just get the UFont and use it */
+	UFont* GetFont();
+
+	/** Returns current font name, as FName */
+	FName GetFontName() const;
+
+	/** Returns current font name as FString */
+	FString GetFontNameAsString() const;
+
+	/** Returns current font name as FText */
+	FText GetFontNameAsText() const;
+
+	/** Returns true if this FAvalancheFont is marked as favorite by the user */
+	bool IsFavorite() const;
+
+	/**
+	 * Checks if this FAvalancheFont is referencing the default font. This will be happening also in fallback state.
+	 * Use IsFallBackFont() if you explicitly need to know if this font is default because of a missing asset
+	 * @return true if referenced font is the default one
+	 */
+	bool IsDefaultFont() const;
+
+	/**
+	 * Checks if this FAvalancheFont is in fallback state.
+	 * This happens when the font referenced by this FAvalancheFont is not available after de-serialization. This likely means that a UFont asset is missing.
+	 * In fallback state, a FAvalancheFont will return the default font when GetFont() is called.
+	 * @return true if this FAvalancheFont is in fallback state
+	 */
+	bool IsFallbackFont() const;
+
+	/** Returns true if font is monospaced */
+	bool IsMonospaced() const;
+
+	/** Returns true if font is bold */
+	bool IsBold() const;
+
+	/** Returns true if font is italic */
+	bool IsItalic() const;
+
+	/**
+	 * Mark/unmark this font as favorite
+	 * @param bFavorite {the flag to mark the font as favorite or not. True means favorite.}
+	 */
+	void SetFavorite(const bool bFavorite);
+
+	/**
+	 * Calling this function will null the now deprecated UFont CurrentFont reference, and use that UFont to initialize the currently used UAvalancheFontObject
+	 */
+	void EnsureUsingCurrentVersion();
+
+	/**
+	 * In case this FAvalancheFont is in fallback state, this call can return the name of the missing font asset, if available.
+	 * @return the name of the missing font asset.
+	 */
+	const FString& GetMissingFontName() const { return MissingFontName; }
+
+	/**
+	 * Checks if the currently referenced font asset is valid. Also internally checks if there's a composite font available within the UFont.
+	 * @return true if font is valid
+	 */
+	bool HasValidFont() const;
+
+	/**
+	 * Returns the currently stored UAvalancheFontObject. This is useful when dealing with font load/import.
+	 * If only the current UFont is needed, use GetFont().
+	 * @return Current font object. 
+	 */
+	UAvaFontObject* GetFontObject() const { return AvalancheFontObject; }
+
+	/**
+	 * Directly sets the current UAvalancheFontObject, and then updates this FAvalancheFont accordingly.
+	 * If this FAvalancheFont is from past versions, the internal deprecated CurrentFont field will be set to nullptr.
+	 * @param InFontObject {The FontObject which will be assigned to this FAvalancheFont}
+	 */
+	void SetFontObject(UAvaFontObject* InFontObject);
+
+	/* Returns current font source (System, Project or Invalid) */
+	EAvaFontSource GetFontSource() const;
+
+	/** Fonts are considered equal if they share the same font and font name */
+	FORCEINLINE bool operator==(const FAvaFont& Other) const
+	{
+		if (CurrentFont_DEPRECATED)
+		{
+			return (CurrentFont_DEPRECATED == Other.CurrentFont_DEPRECATED) && (GetFontName() == Other.GetFontName());
+		}
+
+		const UAvaFontObject* MyFontObject = GetFontObject();
+		const UAvaFontObject* OtherFontObject = Other.GetFontObject();
+
+		if (MyFontObject && OtherFontObject && MyFontObject->GetFont() != OtherFontObject->GetFont())
+		{
+			return false;
+		}
+
+		return (GetFontName() == Other.GetFontName());
+	}
+
+	FORCEINLINE bool operator!=(const FAvaFont& Other) const
+	{
+		return !(*this == Other);
+	}
+
+	/**
+	 * Called after serialization of FAvalancheFont.
+	 * In case of missing font assets, the FAvalancheFont state will be marked as "FallbackFont"
+	 */
+	void PostSerialize(const FArchive& Ar);
+
+private:
+	/** A struct to hold default values which we want to prevent from being GC'd */
+	struct FAvaDefaultFontObjects : public FGCObject
+	{
+		TObjectPtr<UFont> AvaDefaultFont;
+		TObjectPtr<UAvaFontObject> AvaDefaultFontObject;
+
+		//~ Begin FGCObject
+		virtual void AddReferencedObjects(FReferenceCollector& OutCollector) override
+		{
+			OutCollector.AddReferencedObject(AvaDefaultFont);
+			OutCollector.AddReferencedObject(AvaDefaultFontObject);
+		}
+		virtual FString GetReferencerName() const override
+		{
+			return TEXT("FAvaDefaultFontObjects");
+		}
+		//~ End FGCObject
+	};
+
+	enum EFontAssetState
+	{
+		DefaultFont,   // Font is meant to be the default one
+		SelectedFont,  // Font is not the default one and has been explicitly set, either by serialization or UI
+		FallbackFont   // Referenced font asset is missing
+	};
+
+	static UAvaFontObject* GetDefaultAvaFontObject();
+
+	static UFont* GetFontByName(const FString& InFontName);
+
+	void InitDefaults();
+
+	void RefreshName();
+
+	void RefreshAssetState();
+
+	void InitFromFont(UFont* InFont);
+
+	static FAvaDefaultFontObjects DefaultFontObjects;
+
+	/**
+	 * Deprecated - used to reference the font used by FAvalancheFont.
+	 * Older assets still have the font stored in this field. Calling GetFont() will move the font to the currently used AvalancheFontObject property
+	 */
+	UPROPERTY(Transient, meta=(DeprecatedProperty, DeprecationMessage="AvalancheFontObject is the property currently handling the font."))
+	TObjectPtr<UFont> CurrentFont_DEPRECATED;
+
+	/** Stores the UAvalancheFontObject holding all the information about the font currently used by FAvalancheFont */
+	UPROPERTY()
+	TObjectPtr<UAvaFontObject> AvalancheFontObject;
+
+	/** Updated when AvalancheFontObject is set. Mainly used when loading FAvaFont, to help looking for a font in case the current font is null */
+	UPROPERTY()
+	FString FontName;
+
+	/** The current font asset state, used internally to track the current state of FAvalancheFont with respect to its referenced font asset */
+	EFontAssetState FontAssetState;
+
+	/** FAvalancheFonts can be marked as favorites by the user */
+	bool bIsFavorite;
+
+	/** When a font is missing and its name is available, it will be stored here. Old assets version cannot retrieve the name of the missing font */
+	FString MissingFontName;
+};
+
+/** We need to intercept FAvalancheFont serialization to detect if there's a missing font when loading */
+template<> struct TStructOpsTypeTraits<FAvaFont> : public TStructOpsTypeTraitsBase2<FAvaFont>
+{
+	enum
+	{
+		/** void PostSerialize(const FArchive& Ar) in FAvalancheFont */
+		WithPostSerialize = true 
+	};
+};
+
+namespace UE::Avalanche::FontUtilities::Public
+{
+	AVALANCHETEXT_API void GetFontName(const UFont* InFont, FString& OutFontName);
+}
diff --git a/Engine/Plugins/Experimental/Avalanche/Source/AvalancheText/Public/MaterialHub/AvaTextMaterialHub.h b/Engine/Plugins/Experimental/Avalanche/Source/AvalancheText/Public/MaterialHub/AvaTextMaterialHub.h
new file mode 100644
--- /dev/null
+++ b/Engine/Plugins/Experimental/Avalanche/Source/AvalancheText/Public/MaterialHub/AvaTextMaterialHub.h
@@ -0,0 +1,254 @@
+// Copyright Epic Games, Inc. All Rights Reserved.
+
+#pragma once
+
+#include "AvaTextDefs.h"
+#include "Materials/Material.h"
+#include "UObject/SoftObjectPath.h"
+#include "AvaTextMaterialHub.generated.h"
+
+/**
+ * Associates Material Settings with a Material Asset (through its path).
+ */
+struct FAvaTextMaterialData
+{
+	/** Material asset is stored as FSoftObjectPath*/
+	FSoftObjectPath MaterialAssetPath;
+
+	/** Settings describing this material */
+	FAvaTextMaterialSettings MaterialSettings;
+
+	FAvaTextMaterialData(const FSoftObjectPath& InSoftObjectPath, FAvaTextMaterialSettings InMaterialSettings)
+		: MaterialAssetPath(InSoftObjectPath)
+		, MaterialSettings(InMaterialSettings)
+	{
+	}
+
+	/**
+	 * Tries to load and return the material
+	 * 
+	 * @return: The material asset, nullptr if internal TryLoad() fails.
+	 */
+	UMaterialInterface* GetMaterial() const
+	{
+		return Cast<UMaterialInterface>(MaterialAssetPath.TryLoad());
+	}
+};
+
+/**
+ * Keeps a reference to a Material asset path.
+ * Provides functionality to compare materials according to the needs of UAvaText3DComponent materials
+ */
+class FAvaTextMaterialId
+{
+public:
+	/** Default constructor does not compute CachedHash, so the instance is not valid. */
+	FAvaTextMaterialId() = default;
+	
+	FAvaTextMaterialId(const FSoftObjectPath& InSoftObjectPath)
+		: MaterialAssetPath(InSoftObjectPath)
+	{
+		CalculateTypeHash();
+	}
+
+	/**
+	 * Comparison is successful if material is the same, or if their name is the same
+	 * This allows to have a successful comparison even for assets saved with different material versions.
+	 */
+	bool operator==(const FAvaTextMaterialId& Other) const
+	{		
+		return
+			MaterialAssetPath == Other.MaterialAssetPath &&
+			CachedHash == Other.CachedHash;	
+	}
+
+	bool operator!=(const FAvaTextMaterialId& Other) const
+	{
+		return !(*this == Other);
+	}
+
+	FAvaTextMaterialId(const FAvaTextMaterialId& Other)
+	{
+		*this = Other;
+	}
+
+	FAvaTextMaterialId(FAvaTextMaterialId&& Other) noexcept
+	{
+		*this = MoveTemp(Other);
+	}
+
+	FAvaTextMaterialId& operator=(const FAvaTextMaterialId& Other)
+{
+		MaterialAssetPath = Other.MaterialAssetPath;
+		CalculateTypeHash();		
+		
+		return *this;
+	}
+
+	FAvaTextMaterialId& operator=(FAvaTextMaterialId&& Other) noexcept
+	{
+		Swap(*this, Other);
+		CalculateTypeHash();
+		
+		return *this;
+	}
+
+	friend uint32 GetTypeHash(const FAvaTextMaterialId& InMaterialId)
+	{
+		return InMaterialId.CachedHash;
+	}
+
+private:
+	void CalculateTypeHash()
+	{
+		CachedHash = GetTypeHash(MaterialAssetPath);
+	}
+	
+	FSoftObjectPath MaterialAssetPath;
+	
+	uint32 CachedHash;
+};
+
+/**
+ * The UAvaTextMaterialHub class can be used to:
+ * 
+ * - Retrieve a Text3D Material with the specified FAvaTextMaterialSettings
+ * - Retrieve FAvaTextMaterialSettings for a given material
+ */
+UCLASS(Config=Editor)
+class AVALANCHETEXT_API UAvaTextMaterialHub : public UObject
+{
+	GENERATED_BODY()
+	
+public:
+	/**
+	 * Gets a material with the specified settings.
+	 *
+	 * @param InMaterialSettings: Settings for the desired material (coloring style, features)
+	 *
+	 * @return A material matching the desired settings
+	 */
+	static UMaterialInterface* GetMaterial(FAvaTextMaterialSettings InMaterialSettings);
+
+	/**
+	 * Can be used to verify if the provided material has the same type as one of the materials handled by FAvaTextMaterials.
+	 * Outputs information about its settings
+	 *
+	 * @param InMaterial: The material to get the settings from
+	 * @return If lookup is successful, returns the FAvaTextMaterialSettings struct with the current parameters of the looked up material
+	 */
+	static FAvaTextMaterialSettings* GetSettingsFromMaterial(const UMaterialInterface* InMaterial);
+
+	/** Static variables for material handling */
+	struct FMaterialStatics
+	{
+		/** Rotation values are in the [0.0;1.0] domain */
+		inline static float MaterialMaskTopBottomRotationValue   = 0.0f;
+		inline static float MaterialMaskLeftToRightRotationValue = 0.75f;
+		inline static float MaterialMaskRightToLeftRotationValue = 0.25f;
+		
+		/**
+		 * Due to a typo, the value of 0.25f above got written as 2.25f
+		 * This value is kept since it might be needed when retrieving rotation parameter from materials saved prior to fixing it
+		 */
+		inline static float MaterialMaskRightToLeftRotationValue_DEPRECATED = 2.25f;
+
+		/** make sure the following names match the ones of their respective material parameter */
+		inline static FName SolidColor_MatParam         = TEXT("Color");
+		inline static FName Opacity_MatParam            = TEXT("Opacity");
+
+		inline static FName MainTexture_MatParam        = TEXT("MainTexture");
+		inline static FName TexturedUTiling_MatParam    = TEXT("U_Tiling");
+		inline static FName TexturedVTiling_MatParam    = TEXT("V_Tiling");
+
+		inline static FName BoundsOrigin_MatParam       = TEXT("TextBoundsOrigin");
+		inline static FName BoundsSize_MatParam         = TEXT("TextBoundsSize");
+		inline static FName TextPosition_MatParam       = TEXT("TextActorPosition");
+
+		inline static FName GradientOffset_MatParam     = TEXT("GradientOffset");
+		inline static FName GradientColorA_MatParam     = TEXT("GradientColorA");
+		inline static FName GradientColorB_MatParam     = TEXT("GradientColorB");
+		inline static FName GradientRotation_MatParam   = TEXT("GradientRotation");
+		inline static FName GradientSmoothness_MatParam = TEXT("GradientSmoothness");
+		
+		inline static FName MaskRotation_MatParam       = TEXT("MaskRotationDegrees");
+		inline static FName MaskOffset_MatParam         = TEXT("MaskOffset");
+		inline static FName MaskSmoothness_MatParam     = TEXT("MaskSmoothness");
+	};
+	
+private:
+	/** Creates, initializes and returns an UAvaTextMaterialHub instance */
+	static UAvaTextMaterialHub* Get();
+	
+	UAvaTextMaterialHub();
+
+	void SetupInternalMaps();
+	void AddMaterial(const FSoftObjectPath& InMaterialSoftObjectPath, FAvaTextMaterialSettings InMaterialSettings);
+
+	UMaterialInterface* GetMaterialInternal(FAvaTextMaterialSettings InMaterialSettings) const;
+	FAvaTextMaterialSettings* GetSettingsFromMaterialInternal(const UMaterialInterface* InMaterial);
+	
+	UPROPERTY(Config)
+	FSoftObjectPath SolidMaterial;
+
+	UPROPERTY(Config)
+	FSoftObjectPath SolidTranslucentMaterial;
+
+	UPROPERTY(Config)
+	FSoftObjectPath SolidTranslucentUnlitMaterial;
+
+	UPROPERTY(Config)
+	FSoftObjectPath SolidMaskedMaterial;
+
+	UPROPERTY(Config)
+	FSoftObjectPath SolidUnlitMaterial;
+
+	UPROPERTY(Config)
+	FSoftObjectPath SolidUnlitMaskedTextMaterial;
+
+	UPROPERTY(Config)
+	FSoftObjectPath GradientMaterial;
+
+	UPROPERTY(Config)
+	FSoftObjectPath GradientTranslucentMaterial;
+
+	UPROPERTY(Config)
+	FSoftObjectPath GradientTranslucentUnlitMaterial;
+
+	UPROPERTY(Config)
+	FSoftObjectPath GradientMaskedMaterial;
+
+	UPROPERTY(Config)
+	FSoftObjectPath GradientUnlitMaterial;
+
+	UPROPERTY(Config)
+	FSoftObjectPath GradientUnlitMaskedMaterial;
+	
+	UPROPERTY(Config)
+	FSoftObjectPath TexturedMaterial;
+
+	UPROPERTY(Config)
+	FSoftObjectPath TexturedTranslucentMaterial;
+
+	UPROPERTY(Config)
+	FSoftObjectPath TexturedTranslucentUnlitMaterial;
+
+	UPROPERTY(Config)
+	FSoftObjectPath TexturedMaskedMaterial;
+
+	UPROPERTY(Config)
+	FSoftObjectPath TexturedUnlitMaterial;
+
+	UPROPERTY(Config)
+	FSoftObjectPath TexturedUnlitMaskedMaterial;
+
+	/** Material paths, indexed by their settings. Used to easily retrieve a UMaterial based on a set of requirements */
+	TMap<FAvaTextMaterialSettings, FSoftObjectPath> MaterialsMap;
+
+	/**
+	 * MaterialData, indexed by their FAvaTextMaterialId, which is SoftObjectPath-based.
+	 * This allows to retrieve material settings starting from a material asset.
+	 * Useful when loading an AvaTextComponent, in order to retrieve material parameters and settings based on its original asset/type.
+	 */
+	TMap<FAvaTextMaterialId, FAvaTextMaterialData> MaterialDataMap;
+};
diff --git a/Engine/Plugins/Experimental/Avalanche/Source/AvalancheTextEditor/AvalancheTextEditor.Build.cs b/Engine/Plugins/Experimental/Avalanche/Source/AvalancheTextEditor/AvalancheTextEditor.Build.cs
new file mode 100644
--- /dev/null
+++ b/Engine/Plugins/Experimental/Avalanche/Source/AvalancheTextEditor/AvalancheTextEditor.Build.cs
@@ -0,0 +1,63 @@
+// Copyright Epic Games, Inc. All Rights Reserved.
+
+using System.IO;
+using UnrealBuildTool;
+
+public class AvalancheTextEditor : ModuleRules
+{
+	public AvalancheTextEditor(ReadOnlyTargetRules Target) : base(Target)
+	{
+		PCHUsage = ModuleRules.PCHUsageMode.UseExplicitOrSharedPCHs;
+
+		PrivateDependencyModuleNames.AddRange(
+			new string[]
+			{
+				"Avalanche",
+				"AvalancheComponentVisualizers",
+				"AvalancheInteractiveTools",
+				"AvalancheLevelViewport",
+				"AvalancheShapesEditor",
+				"AvalancheText",
+				"ComponentVisualizers",
+				"Core",
+				"CoreUObject",
+				"DynamicMaterialEditor",
+				"EditorScriptingUtilities",
+				"EditorSubsystem",
+				"Engine",
+				"FreeType2",
+				"InputCore",
+				"InteractiveToolsFramework",
+				"Projects",
+				"PropertyEditor",
+				"Slate",
+				"SlateCore",
+				"Text3D",
+				"UnrealEd",
+			}
+		);
+
+		if (Target.Platform.IsInGroup(UnrealPlatformGroup.Windows))
+		{
+			PrivateDependencyModuleNames.Add("DirectX");
+
+			PublicAdditionalLibraries.AddRange(new string[] {
+				Path.Combine(Target.WindowsPlatform.DirectXLibDir, "dwrite.lib")
+			});
+		}
+
+		// Needed for FreeType/Font functionality
+		AddEngineThirdPartyPrivateStaticDependencies(Target,
+			"zlib",
+			"UElibPNG"
+		);
+
+		if (Target.Type == TargetRules.TargetType.Editor)
+		{
+			PrivateDependencyModuleNames.AddRange(new string[]
+			{
+				"AvalancheEditorCore"
+			});
+		}
+	}
+}
diff --git a/Engine/Plugins/Experimental/Avalanche/Source/AvalancheTextEditor/Private/AvaTextEditorModule.cpp b/Engine/Plugins/Experimental/Avalanche/Source/AvalancheTextEditor/Private/AvaTextEditorModule.cpp
new file mode 100644
--- /dev/null
+++ b/Engine/Plugins/Experimental/Avalanche/Source/AvalancheTextEditor/Private/AvaTextEditorModule.cpp
@@ -0,0 +1,156 @@
+// Copyright Epic Games, Inc. All Rights Reserved.
+
+#include "AvaTextEditorModule.h"
+#include "AvaInteractiveToolsDelegates.h"
+#include "AvaText3DComponent.h"
+#include "AvaTextActor.h"
+#include "AvaTextDefs.h"
+#include "AvaTextEditorCommands.h"
+#include "AvaTextEditorStyle.h"
+#include "DetailsView/AvaLinearGradientSettingsCustomization.h"
+#include "DetailsView/AvaTextAlignmentCustomization.h"
+#include "DetailsView/AvaTextComponentCustomization.h"
+#include "DetailsView/AvaTextFieldCustomization.h"
+#include "DynamicMaterialEditorModule.h"
+#include "Font/AvaFontDetailsCustomization.h"
+#include "Font/AvaFontManagerSubsystem.h"
+#include "Framework/Application/SlateApplication.h"
+#include "IAvalancheComponentVisualizersModule.h"
+#include "Misc/CoreDelegates.h"
+#include "Modules/ModuleManager.h"
+#include "PropertyEditorModule.h"
+#include "Tool/AvaTextActorTool.h"
+#include "Visualizer/AvaTextVisualizer.h"
+#include "Widgets/AvaViewportColorPickerActorClassRegistry.h"
+
+void FAvaTextEditorModule::StartupModule()
+{
+	FCoreDelegates::OnPostEngineInit.AddRaw(this, &FAvaTextEditorModule::PostEngineInit);
+
+	FAvaTextEditorStyle::Initialize();
+	FAvaTextEditorCommands::Register();
+	RegisterCustomLayouts();
+	FAvaInteractiveToolsDelegates::GetRegisterToolsDelegate().AddRaw(this, &FAvaTextEditorModule::RegisterTools);
+
+	RegisterDynamicMaterialPropertyGenerator();
+
+	FAvaViewportColorPickerActorClassRegistry::RegisterClassAdapter(AAvaTextActor::StaticClass(), MakeShared<TAvaViewportColorPickerActorAdapter<AAvaTextActor>>());
+}
+
+void FAvaTextEditorModule::ShutdownModule()
+{
+	FCoreDelegates::OnPostEngineInit.RemoveAll(this);
+
+	FAvaTextEditorStyle::Shutdown();
+	FAvaTextEditorCommands::Unregister();
+
+	if (UObjectInitialized() && !IsEngineExitRequested())
+	{
+		UnregisterCustomLayouts();
+	}
+
+	FAvaInteractiveToolsDelegates::GetRegisterToolsDelegate().RemoveAll(this);
+}
+
+void FAvaTextEditorModule::PostEngineInit()
+{
+	if (FSlateApplication::IsInitialized())
+	{
+		RegisterComponentVisualizers();
+	}
+}
+
+void FAvaTextEditorModule::RegisterTools(IAvalancheInteractiveToolsModule* InModule)
+{
+	// 2D Tools
+	InModule->RegisterTool(
+		IAvalancheInteractiveToolsModule::Get().CategoryNameActor,
+		GetDefault<UAvaTextActorTool>()->GetToolParameters()
+	);
+}
+
+void FAvaTextEditorModule::RegisterComponentVisualizers()
+{
+	IAvalancheComponentVisualizersModule& AvaComponentVisualizerModule = IAvalancheComponentVisualizersModule::Get();
+
+	AvaComponentVisualizerModule.RegisterComponentVisualizer<UAvaText3DComponent, FAvaTextVisualizer>(&Visualizers);
+}
+
+void FAvaTextEditorModule::RegisterCustomLayouts()
+{
+	static FName PropertyEditor("PropertyEditor");
+
+	FPropertyEditorModule& PropertyModule = FModuleManager::GetModuleChecked<FPropertyEditorModule>(PropertyEditor);
+
+	// Text
+	PropertyModule.RegisterCustomPropertyTypeLayout(FAvaFont::StaticStruct()->GetFName()
+		, FOnGetPropertyTypeCustomizationInstance::CreateStatic(&FAvaFontDetailsCustomization::MakeInstance));
+
+	PropertyModule.RegisterCustomPropertyTypeLayout(FAvaTextAlignment::StaticStruct()->GetFName()
+		, FOnGetPropertyTypeCustomizationInstance::CreateStatic(&FAvaTextAlignmentCustomization::MakeInstance));
+
+	PropertyModule.RegisterCustomPropertyTypeLayout(FAvaLinearGradientSettings::StaticStruct()->GetFName()
+		, FOnGetPropertyTypeCustomizationInstance::CreateStatic(&FAvaLinearGradientSettingsCustomization::MakeInstance));
+
+	PropertyModule.RegisterCustomPropertyTypeLayout(FAvaTextField::StaticStruct()->GetFName()
+		, FOnGetPropertyTypeCustomizationInstance::CreateStatic(&FAvaTextFieldCustomization::MakeInstance));
+
+	PropertyModule.RegisterCustomClassLayout(UAvaText3DComponent::StaticClass()->GetFName()
+		, FOnGetDetailCustomizationInstance::CreateStatic(&FAvaTextComponentCustomization::MakeInstance));
+}
+
+void FAvaTextEditorModule::UnregisterCustomLayouts()
+{
+	static FName PropertyEditor("PropertyEditor");
+
+	if (FModuleManager::Get().IsModuleLoaded(PropertyEditor))
+	{
+		FPropertyEditorModule& PropertyModule = FModuleManager::GetModuleChecked<FPropertyEditorModule>(PropertyEditor);
+
+		// Text
+		PropertyModule.UnregisterCustomPropertyTypeLayout(FAvaFont::StaticStruct()->GetFName());
+		PropertyModule.UnregisterCustomPropertyTypeLayout(FAvaTextAlignment::StaticStruct()->GetFName());
+		PropertyModule.UnregisterCustomPropertyTypeLayout(FAvaLinearGradientSettings::StaticStruct()->GetFName());
+		PropertyModule.UnregisterCustomPropertyTypeLayout(FAvaTextField::StaticStruct()->GetFName());
+	}
+}
+
+void FAvaTextEditorModule::RegisterDynamicMaterialPropertyGenerator()
+{
+	FDynamicMaterialEditorModule::RegisterCustomMaterialPropertyGenerator(
+		AAvaTextActor::StaticClass(),
+		FDMGetObjectMaterialPropertiesDelegate::CreateLambda(
+			[](UObject* InObject)
+			{
+				static TArray<FProperty*> Text3DProperties = {
+					UText3DComponent::StaticClass()->FindPropertyByName("FrontMaterial"),
+					UText3DComponent::StaticClass()->FindPropertyByName("BevelMaterial"),
+					UText3DComponent::StaticClass()->FindPropertyByName("ExtrudeMaterial"),
+					UText3DComponent::StaticClass()->FindPropertyByName("BackMaterial")
+				};
+
+				TArray<FDMObjectMaterialProperty> Properties;
+
+				if (AAvaTextActor* AvaTextActor = Cast<AAvaTextActor>(InObject))
+				{
+					if (UText3DComponent* TextComponent = AvaTextActor->GetText3DComponent())
+					{
+						Properties.Reserve(Text3DProperties.Num());
+
+						for (FProperty* Property : Text3DProperties)
+						{
+							if (Property)
+							{
+								Properties.Emplace(TextComponent, Property);
+							}
+						}
+					}
+				}
+
+				return Properties;
+			}
+		)
+	);
+}
+
+IMPLEMENT_MODULE(FAvaTextEditorModule, AvalancheTextEditorModule)
diff --git a/Engine/Plugins/Experimental/Avalanche/Source/AvalancheTextEditor/Private/AvaTextEditorStyle.cpp b/Engine/Plugins/Experimental/Avalanche/Source/AvalancheTextEditor/Private/AvaTextEditorStyle.cpp
new file mode 100644
--- /dev/null
+++ b/Engine/Plugins/Experimental/Avalanche/Source/AvalancheTextEditor/Private/AvaTextEditorStyle.cpp
@@ -0,0 +1,76 @@
+// Copyright Epic Games, Inc. All Rights Reserved.
+
+#include "AvaTextEditorStyle.h"
+
+#include "Brushes/SlateImageBrush.h"
+#include "AvaTextDefs.h"
+#include "Interfaces/IPluginManager.h"
+#include "SlateOptMacros.h"
+#include "Styling/SlateStyle.h"
+#include "Styling/SlateStyleRegistry.h"
+
+TSharedPtr<FSlateStyleSet> FAvaTextEditorStyle::StyleInstance = nullptr;
+
+void FAvaTextEditorStyle::Initialize()
+{
+	if (!StyleInstance.IsValid())
+	{
+		StyleInstance = Create();
+		FSlateStyleRegistry::RegisterSlateStyle(*StyleInstance);
+	}
+}
+
+void FAvaTextEditorStyle::Shutdown()
+{
+	FSlateStyleRegistry::UnRegisterSlateStyle(*StyleInstance);
+	ensure(StyleInstance.IsUnique());
+	StyleInstance.Reset();
+}
+
+FName FAvaTextEditorStyle::GetStyleSetName()
+{
+	static FName StyleSetName(TEXT("AvaTextEditor"));
+	return StyleSetName;
+}
+
+const ISlateStyle& FAvaTextEditorStyle::Get()
+{
+	return *StyleInstance;
+}
+
+const FSlateBrush* FAvaTextEditorStyle::GetBrush(FName PropertyName, const ANSICHAR* Specifier)
+{
+	return StyleInstance->GetBrush(PropertyName, Specifier);
+}
+
+BEGIN_SLATE_FUNCTION_BUILD_OPTIMIZATION
+
+#define IMAGE_BRUSH_SVG(RelativePath, ...) FSlateVectorImageBrush(Style->RootToContentDir(RelativePath, TEXT(".svg")), __VA_ARGS__)
+#define IMAGE_BRUSH(RelativePath, ... ) FSlateImageBrush(Style->RootToContentDir(RelativePath, TEXT(".png")), __VA_ARGS__)
+
+TSharedRef<FSlateStyleSet> FAvaTextEditorStyle::Create()
+{
+	TSharedRef<FSlateStyleSet> Style = MakeShared<FSlateStyleSet>("AvaTextEditor");
+
+	TSharedPtr<IPlugin> Plugin = IPluginManager::Get().FindPlugin(TEXT("Avalanche"));
+	check(Plugin.IsValid());
+	if (Plugin.IsValid())
+	{
+		Style->SetContentRoot(Plugin->GetBaseDir() / TEXT("Resources"));
+	}
+
+	const FVector2D Icon16x16(16.0f, 16.0f);
+	const FVector2D Icon20x20(20.0f, 20.0f);
+	const FVector2D Icon22x22(22.0f, 22.0f);
+	
+	Style->Set("ClassIcon.AvaTextActor", new IMAGE_BRUSH("Icons/ToolboxIcons/3d-text", Icon16x16));
+	Style->Set("AvaTextEditor.Tool_Actor_Text3D", new IMAGE_BRUSH_SVG("Icons/ToolboxIcons/Toolbar_Text", Icon20x20));
+	Style->Set("Tool_Actor_Text3D", new IMAGE_BRUSH_SVG("Icons/ToolboxIcons/Toolbar_Text", Icon20x20));
+
+	return Style;
+}
+
+#undef IMAGE_BRUSH_SVG
+#undef IMAGE_BRUSH
+
+END_SLATE_FUNCTION_BUILD_OPTIMIZATION
diff --git a/Engine/Plugins/Experimental/Avalanche/Source/AvalancheTextEditor/Private/AvaTextEditorStyle.h b/Engine/Plugins/Experimental/Avalanche/Source/AvalancheTextEditor/Private/AvaTextEditorStyle.h
new file mode 100644
--- /dev/null
+++ b/Engine/Plugins/Experimental/Avalanche/Source/AvalancheTextEditor/Private/AvaTextEditorStyle.h
@@ -0,0 +1,23 @@
+// Copyright Epic Games, Inc. All Rights Reserved.
+
+#pragma once
+
+#include "Styling/ISlateStyle.h"
+#include "Styling/SlateStyle.h"
+
+class FAvaTextEditorStyle
+{
+public:
+	static void Initialize();
+	static void Shutdown();
+
+	static FName GetStyleSetName();
+	static const ISlateStyle& Get();
+
+	static const FSlateBrush* GetBrush(FName PropertyName, const ANSICHAR* Specifier = nullptr);
+
+private:
+	static TSharedPtr<FSlateStyleSet> StyleInstance;
+
+	static TSharedRef<FSlateStyleSet> Create();
+};
diff --git a/Engine/Plugins/Experimental/Avalanche/Source/AvalancheTextEditor/Private/DetailsView/AvaLinearGradientSettingsCustomization.cpp b/Engine/Plugins/Experimental/Avalanche/Source/AvalancheTextEditor/Private/DetailsView/AvaLinearGradientSettingsCustomization.cpp
new file mode 100644
--- /dev/null
+++ b/Engine/Plugins/Experimental/Avalanche/Source/AvalancheTextEditor/Private/DetailsView/AvaLinearGradientSettingsCustomization.cpp
@@ -0,0 +1,227 @@
+// Copyright Epic Games, Inc. All Rights Reserved.
+
+#include "AvaLinearGradientSettingsCustomization.h"
+
+#include "AvaTextDefs.h"
+#include "DetailWidgetRow.h"
+#include "IDetailChildrenBuilder.h"
+#include "SResetToDefaultPropertyEditor.h"
+#include "Widgets/Layout/SBox.h"
+#include "Widgets/Layout/SGridPanel.h"
+#include "Widgets/Text/STextBlock.h"
+
+#define LOCTEXT_NAMESPACE "AvaLinearGradientSettingsCustomization"
+
+TSharedRef<IPropertyTypeCustomization> FAvaLinearGradientSettingsCustomization::MakeInstance()
+{
+	return MakeShared<FAvaLinearGradientSettingsCustomization>();
+}
+
+void FAvaLinearGradientSettingsCustomization::CustomizeHeader(TSharedRef<IPropertyHandle> StructPropertyHandle,
+	FDetailWidgetRow& HeaderRow, IPropertyTypeCustomizationUtils& StructCustomizationUtils)
+{
+	TSharedPtr<IPropertyHandle> ColorA_Handle = StructPropertyHandle->GetChildHandle(
+		GET_MEMBER_NAME_CHECKED(FAvaLinearGradientSettings, ColorA));
+	
+	HeaderRow.NameContent()[StructPropertyHandle->CreatePropertyNameWidget()]
+	.ValueContent()
+	[
+		// we will fill this out in customize children, since we can only do it there because we will need the StructBuilder
+		SAssignNew(HeaderContentBox, SBox)
+	];
+}
+
+void FAvaLinearGradientSettingsCustomization::CustomizeChildren(TSharedRef<IPropertyHandle> StructPropertyHandle,
+	IDetailChildrenBuilder& StructBuilder, IPropertyTypeCustomizationUtils& StructCustomizationUtils)
+{
+	const TSharedPtr<IPropertyHandle> ColorA_Handle = StructPropertyHandle->GetChildHandle(
+		GET_MEMBER_NAME_CHECKED(FAvaLinearGradientSettings, ColorA));
+	
+	const TSharedPtr<IPropertyHandle> ColorB_Handle = StructPropertyHandle->GetChildHandle(
+		GET_MEMBER_NAME_CHECKED(FAvaLinearGradientSettings, ColorB));
+
+	const TSharedPtr<IPropertyHandle> Offset_Handle = StructPropertyHandle->GetChildHandle(
+		GET_MEMBER_NAME_CHECKED(FAvaLinearGradientSettings, Offset));
+
+	const TSharedPtr<IPropertyHandle> Smoothness_Handle = StructPropertyHandle->GetChildHandle(
+		GET_MEMBER_NAME_CHECKED(FAvaLinearGradientSettings, Smoothness));
+
+	const TSharedPtr<IPropertyHandle> Rotation_Handle = StructPropertyHandle->GetChildHandle(
+		GET_MEMBER_NAME_CHECKED(FAvaLinearGradientSettings, Rotation));
+	
+	Direction_Handle = StructPropertyHandle->GetChildHandle(
+		GET_MEMBER_NAME_CHECKED(FAvaLinearGradientSettings, Direction));
+	
+	auto GetResetButton = [](TSharedPtr<IPropertyHandle> PropertyHandle)
+	{
+		return
+			SNew(SBox)
+			.VAlign(VAlign_Center)
+			.HAlign(HAlign_Center)
+			.Padding(5.0f)
+			[
+				SNew(SResetToDefaultPropertyEditor, PropertyHandle)
+			];
+			
+	};
+	
+	// we need StructBuilder in order to create some widgets in the HeaderContentBox, so we fill it up here in CustomizeChildren
+	
+	HeaderContentBox
+	->SetContent(
+		SNew(SBox)
+		.HAlign(HAlign_Fill)
+		.Padding(FMargin(3.f, 10.f, 3.f, 10.f))
+		.Content()
+		[
+			SNew(SBox)
+			.HAlign(HAlign_Fill)
+			[
+				SNew(SGridPanel)
+				.FillColumn(1, 1.0)
+				.FillColumn(3, 1.0)
+				+ SGridPanel::Slot(0, 0)
+				.HAlign(HAlign_Fill)
+				.VAlign(VAlign_Center)
+				[
+					SNew(STextBlock)
+					.Font(FAppStyle::GetFontStyle("PropertyWindow.NormalFont"))
+					.Text(LOCTEXT("ColorsFieldName", "Colors"))
+				]
+				+ SGridPanel::Slot(1, 0)
+				.HAlign(HAlign_Fill)
+				[
+					StructBuilder.GenerateStructValueWidget(ColorA_Handle.ToSharedRef())
+				]
+				+ SGridPanel::Slot(2, 0)
+				.HAlign(HAlign_Fill)
+				[
+					GetResetButton(ColorA_Handle)
+				]
+				+ SGridPanel::Slot(3, 0)
+				.HAlign(HAlign_Fill)
+				[
+					StructBuilder.GenerateStructValueWidget(ColorB_Handle.ToSharedRef())
+				]
+				+ SGridPanel::Slot(4, 0)
+				.HAlign(HAlign_Fill)
+				[
+					GetResetButton(ColorB_Handle)
+				]
+				+ SGridPanel::Slot(0, 1)
+				.HAlign(HAlign_Fill)
+				.Padding(0, 0, 5.0, 0)
+				[
+					Offset_Handle->CreatePropertyNameWidget()
+				]
+				+ SGridPanel::Slot(1, 1)
+				.HAlign(HAlign_Fill)
+				.ColumnSpan(3)
+				[					
+					Offset_Handle->CreatePropertyValueWidget()
+				]
+				+ SGridPanel::Slot(4, 1)
+				[
+					GetResetButton(Offset_Handle)
+				]
+				+ SGridPanel::Slot(0, 2)
+				.HAlign(HAlign_Fill)
+				.Padding(0, 0, 5.0, 0)
+				[
+					Smoothness_Handle->CreatePropertyNameWidget()
+				]
+				+ SGridPanel::Slot(1, 2)
+				.HAlign(HAlign_Fill)
+				.ColumnSpan(3)
+				[					
+					Smoothness_Handle->CreatePropertyValueWidget()
+				]
+				+ SGridPanel::Slot(4, 2)
+				.HAlign(HAlign_Fill)
+				[
+					GetResetButton(Smoothness_Handle)
+				]
+				+ SGridPanel::Slot(0, 3)
+				.HAlign(HAlign_Fill)
+				.Padding(0, 0, 5.0, 0)
+				[
+					Direction_Handle->CreatePropertyNameWidget()
+				]
+				+ SGridPanel::Slot(1, 3)
+				.HAlign(HAlign_Fill)
+				.ColumnSpan(3)
+				[					
+					Direction_Handle->CreatePropertyValueWidget()					
+				]
+				+ SGridPanel::Slot(4, 3)
+				.HAlign(HAlign_Fill)
+				[
+					GetResetButton(Direction_Handle)
+				]
+				+ SGridPanel::Slot(0, 4)
+				.HAlign(HAlign_Fill)
+				[
+					SNew(SBox)
+					.HAlign(HAlign_Fill)
+					.Padding(0, 0, 5.0, 0)
+					.Visibility(MakeAttributeLambda([this]{return IsCustomDirectionEnabled() ? EVisibility::Visible : EVisibility::Collapsed;}))
+					[
+						Rotation_Handle->CreatePropertyNameWidget()
+					]
+				]
+				+ SGridPanel::Slot(1, 4)
+				.HAlign(HAlign_Fill)
+				.ColumnSpan(3)
+				[
+					SNew(SBox)
+					.HAlign(HAlign_Fill)
+					.Visibility(MakeAttributeLambda([this]{return IsCustomDirectionEnabled() ? EVisibility::Visible : EVisibility::Collapsed;}))
+					[
+						Rotation_Handle->CreatePropertyValueWidget()
+					]
+				]
+				+ SGridPanel::Slot(4, 4)
+				.HAlign(HAlign_Fill)
+				[	
+					SNew(SBox)
+					.HAlign(HAlign_Fill)
+					.Visibility(MakeAttributeLambda([this]{return IsCustomDirectionEnabled() ? EVisibility::Visible : EVisibility::Collapsed;}))
+					[
+						GetResetButton(Rotation_Handle)
+					]
+				]
+			]
+		]
+	);
+}
+
+bool FAvaLinearGradientSettingsCustomization::IsCustomDirectionEnabled() const
+{
+	uint8 DirectionAsUint8;
+	Direction_Handle->GetValue(DirectionAsUint8);
+	const EAvaGradientDirection CurrDirectionSetting = static_cast<EAvaGradientDirection>(DirectionAsUint8);
+	
+	return CurrDirectionSetting == EAvaGradientDirection::Custom;
+}
+
+FLinearColor FAvaLinearGradientSettingsCustomization::GetLinearColorFromProperty(const TSharedPtr<IPropertyHandle>& InColorPropertyHandle)
+{
+	// Default to full alpha in case the alpha component is disabled.
+	FLinearColor OutColor;
+
+	FString StringValue;
+	const FPropertyAccess::Result Result = InColorPropertyHandle->GetValueAsFormattedString(StringValue);
+
+	if (Result == FPropertyAccess::Success)
+	{		
+		OutColor.InitFromString(StringValue);
+	}
+	else if (Result == FPropertyAccess::MultipleValues)
+	{
+		OutColor = FLinearColor::White;
+	}
+
+	return OutColor;
+}
+
+#undef LOCTEXT_NAMESPACE
diff --git a/Engine/Plugins/Experimental/Avalanche/Source/AvalancheTextEditor/Private/DetailsView/AvaLinearGradientSettingsCustomization.h b/Engine/Plugins/Experimental/Avalanche/Source/AvalancheTextEditor/Private/DetailsView/AvaLinearGradientSettingsCustomization.h
new file mode 100644
--- /dev/null
+++ b/Engine/Plugins/Experimental/Avalanche/Source/AvalancheTextEditor/Private/DetailsView/AvaLinearGradientSettingsCustomization.h
@@ -0,0 +1,35 @@
+// Copyright Epic Games, Inc. All Rights Reserved.
+
+#pragma once
+
+#include "CoreMinimal.h"
+#include "IPropertyTypeCustomization.h"
+
+class SBox;
+
+class FAvaLinearGradientSettingsCustomization : public IPropertyTypeCustomization
+{
+public:
+
+	static TSharedRef<IPropertyTypeCustomization> MakeInstance();
+
+	// BEGIN IPropertyTypeCustomization interface
+	
+	virtual void CustomizeHeader(TSharedRef<IPropertyHandle> StructPropertyHandle,
+		class FDetailWidgetRow& HeaderRow,
+		IPropertyTypeCustomizationUtils& StructCustomizationUtils) override;
+	
+	virtual void CustomizeChildren(TSharedRef<IPropertyHandle> StructPropertyHandle,
+		class IDetailChildrenBuilder& StructBuilder,
+		IPropertyTypeCustomizationUtils& StructCustomizationUtils) override;
+	bool IsCustomDirectionEnabled() const;
+
+	// END IPropertyTypeCustomization interface
+
+	static FLinearColor GetLinearColorFromProperty(const TSharedPtr<IPropertyHandle>& InColorPropertyHandle);
+
+private:
+	
+	TSharedPtr<IPropertyHandle> Direction_Handle;
+	TSharedPtr<SBox> HeaderContentBox;
+};
diff --git a/Engine/Plugins/Experimental/Avalanche/Source/AvalancheTextEditor/Private/DetailsView/AvaTextAlignmentCustomization.cpp b/Engine/Plugins/Experimental/Avalanche/Source/AvalancheTextEditor/Private/DetailsView/AvaTextAlignmentCustomization.cpp
new file mode 100644
--- /dev/null
+++ b/Engine/Plugins/Experimental/Avalanche/Source/AvalancheTextEditor/Private/DetailsView/AvaTextAlignmentCustomization.cpp
@@ -0,0 +1,61 @@
+// Copyright Epic Games, Inc. All Rights Reserved.
+
+#include "AvaTextAlignmentCustomization.h"
+
+#include "AvaTextDefs.h"
+#include "DetailWidgetRow.h"
+#include "Slate/SAvaTextAlignmentWidget.h"
+#include "Widgets/Layout/SBox.h"
+
+#define LOCTEXT_NAMESPACE "FAvaTextAlignmentCustomization"
+
+TSharedRef<IPropertyTypeCustomization> FAvaTextAlignmentCustomization::MakeInstance()
+{
+	return MakeShared<FAvaTextAlignmentCustomization>();
+}
+
+void FAvaTextAlignmentCustomization::CustomizeHeader(
+	TSharedRef<IPropertyHandle> StructPropertyHandle
+	, FDetailWidgetRow& HeaderRow
+	, IPropertyTypeCustomizationUtils& StructCustomizationUtils)
+{
+	AlignmentPropertyHandle = StructPropertyHandle;
+
+	HeaderRow.NameContent()[AlignmentPropertyHandle->CreatePropertyNameWidget()]
+	.ValueContent()
+	[
+		SNew(SBox)
+		.ToolTipText(LOCTEXT("AlignTextStructTooltip", "Align Text"))
+		[
+			SNew(SAvaTextAlignmentWidget)				
+			.TextAlignmentPropertyHandle(AlignmentPropertyHandle)
+			.OnHorizontalAlignmentChanged(this, &FAvaTextAlignmentCustomization::OnHorizontalAlignmentChanged)
+			.OnVerticalAlignmentChanged(this, &FAvaTextAlignmentCustomization::OnVerticalAlignmentChanged)
+		]
+	];
+}
+
+void FAvaTextAlignmentCustomization::CustomizeChildren(
+	TSharedRef<IPropertyHandle> StructPropertyHandle
+	, IDetailChildrenBuilder& StructBuilder
+	, IPropertyTypeCustomizationUtils& StructCustomizationUtils)
+{
+}
+
+void FAvaTextAlignmentCustomization::OnHorizontalAlignmentChanged(EText3DHorizontalTextAlignment InHorizontalAlignment) const
+{
+	const TSharedPtr<IPropertyHandle> HorizontalAlignmentPropertyHandle = AlignmentPropertyHandle->GetChildHandle(GET_MEMBER_NAME_CHECKED(FAvaTextAlignment, HorizontalAlignment));
+	check(HorizontalAlignmentPropertyHandle.IsValid());
+
+	HorizontalAlignmentPropertyHandle->SetValue(static_cast<uint8>(InHorizontalAlignment));
+}
+
+void FAvaTextAlignmentCustomization::OnVerticalAlignmentChanged(EText3DVerticalTextAlignment InVerticalAlignment) const
+{
+	const TSharedPtr<IPropertyHandle> VerticalAlignmentPropertyHandle = AlignmentPropertyHandle->GetChildHandle(GET_MEMBER_NAME_CHECKED(FAvaTextAlignment, VerticalAlignment));
+	check(VerticalAlignmentPropertyHandle.IsValid());
+
+	VerticalAlignmentPropertyHandle->SetValue(static_cast<uint8>(InVerticalAlignment));
+}
+
+#undef LOCTEXT_NAMESPACE
diff --git a/Engine/Plugins/Experimental/Avalanche/Source/AvalancheTextEditor/Private/DetailsView/AvaTextAlignmentCustomization.h b/Engine/Plugins/Experimental/Avalanche/Source/AvalancheTextEditor/Private/DetailsView/AvaTextAlignmentCustomization.h
new file mode 100644
--- /dev/null
+++ b/Engine/Plugins/Experimental/Avalanche/Source/AvalancheTextEditor/Private/DetailsView/AvaTextAlignmentCustomization.h
@@ -0,0 +1,31 @@
+// Copyright Epic Games, Inc. All Rights Reserved.
+
+#pragma once
+
+#include "IPropertyTypeCustomization.h"
+#include "Templates/SharedPointer.h"
+
+enum class EText3DHorizontalTextAlignment : uint8;
+enum class EText3DVerticalTextAlignment : uint8;
+
+class FAvaTextAlignmentCustomization : public IPropertyTypeCustomization
+{
+public:
+	static TSharedRef<IPropertyTypeCustomization> MakeInstance();
+
+	// ~ Begin IPropertyTypeCustomization interface
+	virtual void CustomizeHeader(TSharedRef<IPropertyHandle> StructPropertyHandle,
+         class FDetailWidgetRow& HeaderRow,
+         IPropertyTypeCustomizationUtils& StructCustomizationUtils) override;
+	
+	virtual void CustomizeChildren(TSharedRef<IPropertyHandle> StructPropertyHandle,
+		class IDetailChildrenBuilder& StructBuilder,
+		IPropertyTypeCustomizationUtils& StructCustomizationUtils) override;
+	// ~ End IPropertyTypeCustomization interface
+
+protected:
+	TSharedPtr<IPropertyHandle> AlignmentPropertyHandle;
+
+	void OnHorizontalAlignmentChanged(EText3DHorizontalTextAlignment InHorizontalAlignment) const;
+	void OnVerticalAlignmentChanged(EText3DVerticalTextAlignment InVerticalAlignment) const;
+};
diff --git a/Engine/Plugins/Experimental/Avalanche/Source/AvalancheTextEditor/Private/DetailsView/AvaTextComponentCustomization.cpp b/Engine/Plugins/Experimental/Avalanche/Source/AvalancheTextEditor/Private/DetailsView/AvaTextComponentCustomization.cpp
new file mode 100644
--- /dev/null
+++ b/Engine/Plugins/Experimental/Avalanche/Source/AvalancheTextEditor/Private/DetailsView/AvaTextComponentCustomization.cpp
@@ -0,0 +1,25 @@
+// Copyright Epic Games, Inc. All Rights Reserved.
+
+#include "AvaTextComponentCustomization.h"
+
+#include "AvaText3DComponent.h"
+#include "DetailCategoryBuilder.h"
+#include "DetailLayoutBuilder.h"
+
+#define LOCTEXT_NAMESPACE "AvaTextComponentCustomization"
+
+void FAvaTextComponentCustomization::CustomizeDetails(IDetailLayoutBuilder& DetailBuilder)
+{
+	// UAvaText3DComponent adds 
+	//
+	// - FAvaFont property for fonts (additional functionalities for font selection)
+	// - FAvaTextAlignment Alignment for both Horizontal and Vertical alignment (custom widget for alignment)
+	//
+	// Therefore, to avoid showing duplicate properties with the same meaning, we hide the base class ones, as follows:
+	
+	DetailBuilder.HideProperty(GET_MEMBER_NAME_CHECKED(UAvaText3DComponent, Font), UText3DComponent::StaticClass());
+	DetailBuilder.HideProperty(GET_MEMBER_NAME_CHECKED(UAvaText3DComponent, HorizontalAlignment), UText3DComponent::StaticClass());
+	DetailBuilder.HideProperty(GET_MEMBER_NAME_CHECKED(UAvaText3DComponent, VerticalAlignment), UText3DComponent::StaticClass());
+}
+
+#undef LOCTEXT_NAMESPACE
diff --git a/Engine/Plugins/Experimental/Avalanche/Source/AvalancheTextEditor/Private/DetailsView/AvaTextComponentCustomization.h b/Engine/Plugins/Experimental/Avalanche/Source/AvalancheTextEditor/Private/DetailsView/AvaTextComponentCustomization.h
new file mode 100644
--- /dev/null
+++ b/Engine/Plugins/Experimental/Avalanche/Source/AvalancheTextEditor/Private/DetailsView/AvaTextComponentCustomization.h
@@ -0,0 +1,20 @@
+﻿// Copyright Epic Games, Inc. All Rights Reserved.
+
+#pragma once
+
+#include "IDetailCustomization.h"
+
+class FAvaTextComponentCustomization : public IDetailCustomization
+{
+public:
+	static TSharedRef<IDetailCustomization> MakeInstance()
+	{
+		return MakeShared<FAvaTextComponentCustomization>();
+	}
+
+	FAvaTextComponentCustomization(){};
+	
+	//~ Begin IDetailCustomization
+	virtual void CustomizeDetails(IDetailLayoutBuilder& DetailBuilder) override;
+	//~ End IDetailCustomization
+};
diff --git a/Engine/Plugins/Experimental/Avalanche/Source/AvalancheTextEditor/Private/DetailsView/AvaTextFieldCustomization.cpp b/Engine/Plugins/Experimental/Avalanche/Source/AvalancheTextEditor/Private/DetailsView/AvaTextFieldCustomization.cpp
new file mode 100644
--- /dev/null
+++ b/Engine/Plugins/Experimental/Avalanche/Source/AvalancheTextEditor/Private/DetailsView/AvaTextFieldCustomization.cpp
@@ -0,0 +1,301 @@
+// Copyright Epic Games, Inc. All Rights Reserved.
+
+#include "AvaTextFieldCustomization.h"
+
+#include "AvaTextDefs.h"
+#include "DetailWidgetRow.h"
+#include "IDetailChildrenBuilder.h"
+#include "SResetToDefaultPropertyEditor.h"
+#include "Widgets/Input/SMultiLineEditableTextBox.h"
+#include "Widgets/Layout/SBox.h"
+#include "Widgets/Text/STextBlock.h"
+
+#define LOCTEXT_NAMESPACE "FAvaTextFieldCustomization"
+
+TSharedRef<IPropertyTypeCustomization> FAvaTextFieldCustomization::MakeInstance()
+{
+	return MakeShared<FAvaTextFieldCustomization>();
+}
+
+void FAvaTextFieldCustomization::CustomizeHeader(TSharedRef<IPropertyHandle> StructPropertyHandle,
+    FDetailWidgetRow& HeaderRow, IPropertyTypeCustomizationUtils& StructCustomizationUtils)
+{
+	HeaderRow.ForceAutoExpansion = true;
+}
+
+void FAvaTextFieldCustomization::CustomizeChildren(TSharedRef<IPropertyHandle> StructPropertyHandle,
+	IDetailChildrenBuilder& StructBuilder, IPropertyTypeCustomizationUtils& StructCustomizationUtils)
+{
+	Text_Handle = StructPropertyHandle->GetChildHandle(GET_MEMBER_NAME_CHECKED(FAvaTextField, Text));
+	MaximumLengthEnum_Handle = StructPropertyHandle->GetChildHandle(GET_MEMBER_NAME_CHECKED(FAvaTextField, MaximumLength));
+	MaxTextCount_Handle = StructPropertyHandle->GetChildHandle(GET_MEMBER_NAME_CHECKED(FAvaTextField, Count));
+	TextCase_Handle = StructPropertyHandle->GetChildHandle(GET_MEMBER_NAME_CHECKED(FAvaTextField, TextCase));
+	
+	Text_Handle->GetValue(DefaultText);
+
+	bTextFieldIsValid = true;
+	
+	auto GetResetButton = [](TSharedPtr<IPropertyHandle> PropertyHandle)
+	{
+		return
+			SNew(SBox)
+			.VAlign(VAlign_Center)
+			.HAlign(HAlign_Center)
+			.Padding(5.0f)
+			[
+				SNew(SResetToDefaultPropertyEditor, PropertyHandle)
+			];			
+	};
+
+	// text row
+	StructBuilder.AddProperty(Text_Handle.ToSharedRef())
+		.CustomWidget()
+		.NameContent()
+		[
+			SNew(SBox)
+			.HAlign(HAlign_Fill)
+			.VAlign(VAlign_Center)
+			.Padding(0, 0, 5.0, 0)
+			[
+				SNew(STextBlock)
+				.Font(FAppStyle::GetFontStyle("PropertyWindow.NormalFont"))
+				.Text(LOCTEXT("Text", "Text"))
+			]
+		]
+		.ValueContent()
+		.HAlign(HAlign_Fill)
+		[
+			SNew(SBox)
+			.HAlign(HAlign_Fill)
+			[
+				SNew(SHorizontalBox)
+				+ SHorizontalBox::Slot()
+				.FillWidth(1.f)
+				[
+					SAssignNew(MultilineTextBox, SMultiLineEditableTextBox)
+					.Text(DefaultText)
+					.OnKeyDownHandler(this, &FAvaTextFieldCustomization::OnKeyDownInTextField)
+					.OnVerifyTextChanged(this, &FAvaTextFieldCustomization::OnVerifyTextChanged)
+					.OnTextCommitted(this, &FAvaTextFieldCustomization::OnTextCommitted)			
+				]
+			]
+		];
+
+	// text length row
+	StructBuilder.AddProperty(MaximumLengthEnum_Handle.ToSharedRef()).CustomWidget()
+		.ValueContent()
+		.HAlign(HAlign_Fill)
+		[
+			SNew(SHorizontalBox)
+			+SHorizontalBox::Slot()
+			.FillWidth(0.3f)
+			[
+				SNew(SBox)
+				.HAlign(HAlign_Fill)
+				.VAlign(VAlign_Center)
+				.Padding(0, 0, 5.0, 0)
+				[
+					SNew(STextBlock)
+					.Font(FAppStyle::GetFontStyle("PropertyWindow.NormalFont"))
+					.Text(LOCTEXT("TextLength", "Length"))
+				]
+			]
+			+ SHorizontalBox::Slot()
+			.FillWidth(0.7f)
+			[
+				SNew(SHorizontalBox)
+				+SHorizontalBox::Slot()
+				.HAlign(HAlign_Fill)
+				[
+					MaximumLengthEnum_Handle->CreatePropertyValueWidget()
+				]
+				+ SHorizontalBox::Slot()
+				.AutoWidth()
+				[
+					SNew(SHorizontalBox)
+					.Visibility(MakeAttributeLambda([this]{return IsCharacterLimitEnabled() ? EVisibility::Visible : EVisibility::Collapsed;}))
+					+ SHorizontalBox::Slot()
+					.HAlign(HAlign_Fill)
+					[
+						MaxTextCount_Handle->CreatePropertyValueWidget()
+					]
+					+ SHorizontalBox::Slot()
+					.AutoWidth()
+					[
+						GetResetButton(MaxTextCount_Handle)
+					]
+				]
+			]
+		];
+
+	// text case row
+	StructBuilder.AddProperty(TextCase_Handle.ToSharedRef()).CustomWidget()
+		.ValueContent()
+		.HAlign(HAlign_Fill)
+		[
+			SNew(SHorizontalBox)
+			+ SHorizontalBox::Slot()
+			.FillWidth(0.3f)
+			[
+				SNew(SBox)
+				.HAlign(HAlign_Fill)
+				.VAlign(VAlign_Center)
+				.Padding(0, 0, 5.0, 0)
+				[
+					SNew(STextBlock)
+					.Font(FAppStyle::GetFontStyle("PropertyWindow.NormalFont"))
+					.Text(LOCTEXT("TextCase", "Text Case"))
+				]
+			]
+			+ SHorizontalBox::Slot()
+			.FillWidth(0.7f)
+			[
+				TextCase_Handle->CreatePropertyValueWidget()
+			]
+		];
+
+	const FSimpleDelegate OnTextLengthPropertiesChangedDelegate = FSimpleDelegate::CreateSP(this, &FAvaTextFieldCustomization::OnMaximumLengthEnumChanged);
+	MaximumLengthEnum_Handle->SetOnPropertyValueChanged(OnTextLengthPropertiesChangedDelegate);
+	MaxTextCount_Handle->SetOnPropertyValueChanged(OnTextLengthPropertiesChangedDelegate);
+
+	const FSimpleDelegate OnTextFieldResetDelegate = FSimpleDelegate::CreateSP(this, &FAvaTextFieldCustomization::ResetText);	
+	StructPropertyHandle->SetOnPropertyResetToDefault(OnTextFieldResetDelegate);
+	Text_Handle->SetOnPropertyResetToDefault(OnTextFieldResetDelegate);
+}
+
+void FAvaTextFieldCustomization::UpdateTextFieldValidity(const FText& InText)
+{
+	bTextFieldIsValid = true;
+	
+	if (GetCurrentTextLengthType() == EAvaTextLength::CharacterCount)
+	{
+		int32 MaxTextCount;
+
+		const FPropertyAccess::Result Result = MaxTextCount_Handle->GetValue(MaxTextCount);
+
+		if (Result == FPropertyAccess::Success)
+		{
+			if (InText.ToString().Len() > MaxTextCount)
+			{				
+				bTextFieldIsValid = false;
+			}
+		}
+	}
+}
+
+bool FAvaTextFieldCustomization::OnVerifyTextChanged(const FText& InText, FText& OutErrorText)
+{
+	bTextFieldIsValid = true;	
+	
+	if (GetCurrentTextLengthType() == EAvaTextLength::CharacterCount)
+	{
+		int32 MaxTextCount;
+
+		const FPropertyAccess::Result Result = MaxTextCount_Handle->GetValue(MaxTextCount);
+
+		if (Result == FPropertyAccess::Success)
+		{
+			if (InText.ToString().Len() > MaxTextCount)
+			{
+				OutErrorText = LOCTEXT("TextLengthError", "Text character limit exceeded.");
+				bTextFieldIsValid = false;
+				return false;
+			}
+		}
+	}
+	
+	return true;
+}
+
+void FAvaTextFieldCustomization::OnTextCommitted(const FText& InText, ETextCommit::Type InCommitType) const
+{
+	if (bTextFieldIsValid)
+	{
+		Text_Handle->SetValue(InText);
+	}
+	else
+	{
+		int32 MaxTextCount;
+		const FPropertyAccess::Result Result = MaxTextCount_Handle->GetValue(MaxTextCount);
+		
+		FString TrimmedText = InText.ToString();
+
+		const int32 TextLength = TrimmedText.Len();
+
+		// just to be safe
+		if (TextLength > MaxTextCount)
+		{
+			TrimmedText.RemoveAt(MaxTextCount, TextLength - MaxTextCount);
+			Text_Handle->SetValue(TrimmedText);
+		}
+	}
+}
+
+FReply FAvaTextFieldCustomization::OnKeyDownInTextField(const FGeometry& InGeometry, const FKeyEvent& InKeyEvent)
+{
+	if (!InKeyEvent.GetModifierKeys().IsShiftDown())
+	{		
+		if (InKeyEvent.GetKey() == EKeys::Enter && MultilineTextBox->HasKeyboardFocus())
+		{
+			// Clear focus
+			return FReply::Handled().SetUserFocus(MultilineTextBox->GetParentWidget().ToSharedRef(), EFocusCause::Cleared);
+		}
+	}
+	
+	return FReply::Unhandled();
+}
+
+void FAvaTextFieldCustomization::Refresh(const FText& InText)
+{
+	// let's simulate a text field update
+	FText OutError;
+	OnVerifyTextChanged(InText, OutError);
+
+	if (!bTextFieldIsValid)
+	{
+		MultilineTextBox->SetError(OutError);
+	}
+	else
+	{
+		MultilineTextBox->SetError(FText::GetEmpty());
+	}
+	
+	OnTextCommitted(InText, ETextCommit::Default);
+}
+
+void FAvaTextFieldCustomization::OnMaximumLengthEnumChanged()
+{
+	Refresh(MultilineTextBox->GetText());
+}
+
+void FAvaTextFieldCustomization::ResetText()
+{
+	FText Text;
+	Text_Handle->GetValue(Text);
+	MultilineTextBox->SetText(Text);
+	
+	Refresh(Text);
+}
+
+EAvaTextLength FAvaTextFieldCustomization::GetCurrentTextLengthType() const
+{
+	if (MaximumLengthEnum_Handle.IsValid())
+	{
+		uint8 EnumAsUint;
+		const FPropertyAccess::Result Result = MaximumLengthEnum_Handle->GetValue(EnumAsUint);
+
+		if (Result == FPropertyAccess::Success)
+		{
+			return static_cast<EAvaTextLength>(EnumAsUint);
+		}
+	}
+
+	return EAvaTextLength::Unlimited;	
+}
+
+bool FAvaTextFieldCustomization::IsCharacterLimitEnabled() const
+{
+	return GetCurrentTextLengthType() == EAvaTextLength::CharacterCount;
+}
+
+#undef LOCTEXT_NAMESPACE
diff --git a/Engine/Plugins/Experimental/Avalanche/Source/AvalancheTextEditor/Private/DetailsView/AvaTextFieldCustomization.h b/Engine/Plugins/Experimental/Avalanche/Source/AvalancheTextEditor/Private/DetailsView/AvaTextFieldCustomization.h
new file mode 100644
--- /dev/null
+++ b/Engine/Plugins/Experimental/Avalanche/Source/AvalancheTextEditor/Private/DetailsView/AvaTextFieldCustomization.h
@@ -0,0 +1,49 @@
+// Copyright Epic Games, Inc. All Rights Reserved.
+
+#pragma once
+
+#include "IPropertyTypeCustomization.h"
+#include "Input/Reply.h"
+
+enum class EAvaTextLength : uint8;
+class SEditableTextBox;
+class SMultiLineEditableTextBox;
+
+class FAvaTextFieldCustomization : public IPropertyTypeCustomization
+{
+public:
+	static TSharedRef<IPropertyTypeCustomization> MakeInstance();
+	
+	// BEGIN IPropertyTypeCustomization interface
+	virtual void CustomizeHeader(TSharedRef<IPropertyHandle> StructPropertyHandle,
+	                             class FDetailWidgetRow& HeaderRow,
+	                             IPropertyTypeCustomizationUtils& StructCustomizationUtils) override;
+	
+	virtual void CustomizeChildren(TSharedRef<IPropertyHandle> StructPropertyHandle,
+	                               class IDetailChildrenBuilder& StructBuilder,
+	                               IPropertyTypeCustomizationUtils& StructCustomizationUtils) override;
+	// END IPropertyTypeCustomization interface
+	
+private:
+	bool IsCharacterLimitEnabled() const;	
+	EAvaTextLength GetCurrentTextLengthType() const;
+	void UpdateTextFieldValidity(const FText& InText);	
+	
+	void Refresh(const FText& InText);
+	void ResetText();
+	
+	void OnTextCommitted(const FText& InText, ETextCommit::Type InCommitType) const;
+	bool OnVerifyTextChanged(const FText& InText, FText& OutErrorText);
+	FReply OnKeyDownInTextField(const FGeometry& Geometry, const FKeyEvent& KeyEvent);
+	void OnMaximumLengthEnumChanged();
+
+	TSharedPtr<SMultiLineEditableTextBox> MultilineTextBox;
+	
+	TSharedPtr<IPropertyHandle> Text_Handle;
+	TSharedPtr<IPropertyHandle> MaximumLengthEnum_Handle;
+	TSharedPtr<IPropertyHandle> MaxTextCount_Handle;
+	TSharedPtr<IPropertyHandle> TextCase_Handle;
+	
+	FText DefaultText;
+	bool bTextFieldIsValid = true;
+};
diff --git a/Engine/Plugins/Experimental/Avalanche/Source/AvalancheTextEditor/Private/Font/AvaFontDetailsCustomization.cpp b/Engine/Plugins/Experimental/Avalanche/Source/AvalancheTextEditor/Private/Font/AvaFontDetailsCustomization.cpp
new file mode 100644
--- /dev/null
+++ b/Engine/Plugins/Experimental/Avalanche/Source/AvalancheTextEditor/Private/Font/AvaFontDetailsCustomization.cpp
@@ -0,0 +1,730 @@
+// Copyright Epic Games, Inc. All Rights Reserved.
+
+#include "Font/AvaFontDetailsCustomization.h"
+#include "AvaFontManagerSubsystem.h"
+#include "AvaFontView.h"
+#include "DetailWidgetRow.h"
+#include "Editor.h"
+#include "Engine/Font.h"
+#include "Font/AvaFont.h"
+#include "Framework/MultiBox/MultiBoxBuilder.h"
+#include "Slate/SAvaFontField.h"
+#include "Slate/SAvaFontSelector.h"
+#include "Widgets/Images/SImage.h"
+#include "Widgets/Input/SButton.h"
+#include "Widgets/Layout/SBox.h"
+#include "Widgets/Text/STextBlock.h"
+
+#define LOCTEXT_NAMESPACE "AvaFontDetailsCustomization"
+
+TSharedRef<IPropertyTypeCustomization> FAvaFontDetailsCustomization::MakeInstance()
+{
+	// Create the instance and return a SharedRef
+	return MakeShared<FAvaFontDetailsCustomization>();
+}
+
+FAvaFontDetailsCustomization::~FAvaFontDetailsCustomization()
+{
+	UnregisterFontManagerCallbacks();
+}
+
+void FAvaFontDetailsCustomization::CustomizeHeader(TSharedRef<IPropertyHandle> FontStructPropertyHandle
+	, FDetailWidgetRow& HeaderRow
+	, IPropertyTypeCustomizationUtils& StructCustomizationUtils)
+{
+	AvaFontPropertyHandle = FontStructPropertyHandle;
+
+	UAvaFontManagerSubsystem* FontManagerSubsystem = UAvaFontManagerSubsystem::Get();
+	if (!ensureMsgf(FontManagerSubsystem, TEXT("Font Manager is unavailable!")))
+	{
+		return;
+	}
+
+	FontManagerSubsystemWeak = TWeakObjectPtr<UAvaFontManagerSubsystem>(FontManagerSubsystem);
+
+	// Refresh also works as a main initialization method for available options and current font
+	// Calling it before FontSelector since it sets up data needed for its initialization
+	Refresh();
+
+	SAssignNew(MissingFontText, STextBlock)
+		.Justification(ETextJustify::Right)
+		.Visibility(EVisibility::Collapsed);
+
+	SAssignNew(ComboBoxContent, SBox)
+		.HAlign(HAlign_Fill);
+
+	bool bIsShowMonospacedActive = false;
+	bool bIsShowBoldActive = false;
+	bool bIsShowItalicActive = false;
+
+	if (const UAvaFontConfig* FontManagerConfig = FontManagerSubsystem->GetFontManagerConfig())
+	{
+		bIsShowMonospacedActive = FontManagerConfig->bShowOnlyMonospaced;
+		bIsShowBoldActive = FontManagerConfig->bShowOnlyBold;
+		bIsShowItalicActive = FontManagerConfig->bShowOnlyItalic;
+	}
+
+	SAssignNew(FontSelector, SAvaFontSelector)
+		.ShowMonospacedFontsOnly(bIsShowMonospacedActive)
+		.ShowBoldFontsOnly(bIsShowBoldActive)
+		.ShowItalicFontsOnly(bIsShowItalicActive)
+		.OptionsSource(&Options)
+		.InitiallySelectedItem(SelectedOption)
+		.AvalancheFontPropertyHandle(AvaFontPropertyHandle)
+		.OnGenerateWidget(this, &FAvaFontDetailsCustomization::HandleGenerateWidget)
+		.OnSelectionChanged(this, &FAvaFontDetailsCustomization::HandleSelectionChanged)
+		[
+			ComboBoxContent.ToSharedRef()
+		];
+
+	HeaderRow.NameContent()[FontStructPropertyHandle->CreatePropertyNameWidget()]
+	.ValueContent()
+	[
+		SNew(SVerticalBox)
+		+ SVerticalBox::Slot()
+		[
+			SNew(SHorizontalBox)
+			+ SHorizontalBox::Slot()
+			.HAlign(HAlign_Fill)
+			.FillWidth(1.f)
+			[
+				FontSelector.ToSharedRef()
+			]
+			+ SHorizontalBox::Slot()
+			.AutoWidth()
+			.HAlign(HAlign_Right)
+			[
+				MakeImportButton()
+			]
+			+SHorizontalBox::Slot()
+			.AutoWidth()
+			[
+				MakeBrowseButton()
+			]
+			+SHorizontalBox::Slot()
+			.AutoWidth()
+			[
+				MakeRefreshFontButton()
+			]
+		]
+		+ SVerticalBox::Slot()
+		.VAlign(VAlign_Center)
+		.HAlign(HAlign_Left)
+		.AutoHeight()
+		[
+			SNew(SBox)
+			.VAlign(VAlign_Center)
+			[
+				MissingFontText.ToSharedRef()
+			]
+		]
+	];
+
+	RegisterFontManagerCallbacks();
+	FontStructPropertyHandle->SetOnPropertyResetToDefault(FSimpleDelegate::CreateSP(this, &FAvaFontDetailsCustomization::SetDefaultValue));
+
+	RefreshComboboxContent();
+	RefreshMissingFontWidget();
+}
+
+void FAvaFontDetailsCustomization::CustomizeChildren(TSharedRef<IPropertyHandle> PropertyHandle,
+	IDetailChildrenBuilder& ChildBuilder, IPropertyTypeCustomizationUtils& CustomizationUtils)
+{
+}
+
+void FAvaFontDetailsCustomization::RefreshSelectedFont()
+{
+	const UAvaFontManagerSubsystem* FontManagerSubsystem = FontManagerSubsystemWeak.Get();
+	if (!FontManagerSubsystem)
+	{
+		return;
+	}
+
+	FPropertyAccess::Result AccessResult;
+	const TSharedPtr<FAvaFontView> CurrAvaFont = GetCurrentAvaFontView(AccessResult);
+
+	bool bLoadDefault = false;
+
+	if (AccessResult == FPropertyAccess::Fail)
+	{
+		bLoadDefault = true;
+	}
+	else
+	{
+		FAvaMultiFontSelectionData MultiFontSelectionData;
+		FontManagerSubsystem->GetMultipleSelectionInformation(AvaFontPropertyHandle, MultiFontSelectionData);
+
+		if (MultiFontSelectionData.bMultipleValues && ComboBoxContent)
+		{
+			ComboBoxContent->SetContent(GenerateMultipleSelectionField());
+			bImportEnabled = false;
+			bRefreshFontEnabled = false;
+			return;
+		}
+
+		// If there's a font already selected, update the option ptr accordingly
+		if (CurrAvaFont)
+		{
+			const UAvaFontObject* CurrFontObj = CurrAvaFont->GetFontObject();
+
+			if (CurrFontObj && CurrFontObj->GetFont())
+			{
+				TSharedPtr<FAvaFontView> Option;
+
+				const bool bCurrFontObjectMatchesOption = Options.ContainsByPredicate([&Option, CurrFontObj](const TSharedPtr<FAvaFontView>& AvaFontValue)
+				{
+					if (AvaFontValue.IsValid() && AvaFontValue->HasValidFont())
+					{
+						const bool bContains = AvaFontValue->GetFontObject()->GetFontName() == CurrFontObj->GetFontName();
+
+						if (bContains)
+						{
+							Option = AvaFontValue;
+						}
+
+						return bContains;
+					}
+
+					return false;
+				});
+
+				if (bCurrFontObjectMatchesOption) // This bool should also guard Option validity
+				{
+					HandleSelectionChanged(Option, ESelectInfo::Direct);
+				}
+			}
+			else
+			{
+				// Font is not available, try to load a default option
+				bLoadDefault = true;
+			}
+		}
+	}
+
+	if (bLoadDefault)
+	{
+		SetDefaultValue();
+	}
+}
+
+void FAvaFontDetailsCustomization::RefreshLocalFontOptions()
+{
+	UAvaFontManagerSubsystem* FontManagerSubsystem = FontManagerSubsystemWeak.Get();
+	if (!FontManagerSubsystem)
+	{
+		return;
+	}
+
+	Options = FontManagerSubsystem->GetFontOptions();
+
+	if (!Options.IsEmpty())
+	{
+		const int32 DefaultOptionIndex = FontManagerSubsystem->GetDefaultFontIndex();
+
+		if (DefaultOptionIndex >= 0)
+		{
+			DefaultOption = Options[DefaultOptionIndex];
+		}
+		else
+		{
+			DefaultOption = Options[0];
+		}
+	}
+}
+
+void FAvaFontDetailsCustomization::Refresh()
+{
+	RefreshLocalFontOptions();
+	RefreshSelectedFont();
+
+	if (FontSelector)
+	{
+		FontSelector->RefreshOptions();
+	}
+}
+
+void FAvaFontDetailsCustomization::SetDefaultValue()
+{
+	if (DefaultOption.IsValid())
+	{
+		HandleSelectionChanged(DefaultOption, ESelectInfo::Direct);
+	}
+}
+
+FReply FAvaFontDetailsCustomization::ImportButtonClicked()
+{
+	FPropertyAccess::Result AccessResult;
+
+	if (const TSharedPtr<FAvaFontView> CurrAvaFont = GetCurrentAvaFontView(AccessResult))
+	{
+		const UAvaFontObject* CurrFontObj = CurrAvaFont->GetFontObject();
+
+		UAvaFontManagerSubsystem* FontManagerSubsystem = FontManagerSubsystemWeak.Get();
+
+		if (IsValid(CurrFontObj) && FontManagerSubsystem)
+		{
+			if (FontManagerSubsystem->ImportFont(AvaFontPropertyHandle))
+			{
+				AvaFontPropertyHandle->NotifyPostChange(EPropertyChangeType::ValueSet);
+				Refresh();
+			}
+		}
+	}
+
+	return FReply::Handled();
+}
+
+TSharedRef<SWidget> FAvaFontDetailsCustomization::HandleGenerateWidget(const TSharedPtr<FAvaFontView>& InItem)
+{
+	return SNew(SAvaFontField)
+			.OnAvaFontFieldModified(this, &FAvaFontDetailsCustomization::OnFontFieldUpdated)
+			.AvalancheFont(InItem);
+}
+
+TSharedRef<SWidget> FAvaFontDetailsCustomization::GenerateMultipleSelectionField() const
+{
+	return SNew(STextBlock)
+			.Text(FText::FromString("Multiple Values"));
+}
+
+FReply FAvaFontDetailsCustomization::OnBrowseToAssetClicked()
+{
+	if (AvaFontPropertyHandle.IsValid())
+	{
+		FPropertyAccess::Result AccessResult;
+		if (const TSharedPtr<FAvaFontView> CurrentAvaFont = GetCurrentAvaFontView(AccessResult))
+		{
+			UFont* CurrentFont = CurrentAvaFont->GetFont();
+
+			const TArray<UObject*> ObjectsToFocus {CurrentFont};
+			GEditor->SyncBrowserToObjects(ObjectsToFocus);
+		}
+	}
+
+	return FReply::Handled();
+}
+
+FReply FAvaFontDetailsCustomization::OnRefreshFontClicked()
+{
+	if (AvaFontPropertyHandle.IsValid())
+	{
+		FPropertyAccess::Result AccessResult;
+		if (const TSharedPtr<FAvaFontView> CurrentAvaFont = GetCurrentAvaFontView(AccessResult))
+		{
+			UFont* CurrentFont = CurrentAvaFont->GetFont();
+
+			if (UAvaFontManagerSubsystem* FontManagerSubsystem = UAvaFontManagerSubsystem::Get())
+			{
+				FontManagerSubsystem->RefreshSystemFont(CurrentFont);
+			}
+		}
+	}
+
+	return FReply::Handled();
+}
+
+TSharedRef<SWidget> FAvaFontDetailsCustomization::MakeBrowseButton()
+{
+	const FSlateBrush* BrowseIcon = FAppStyle::Get().GetBrush("Icons.BrowseContent");
+
+	return
+		SNew(SBox)
+		.HAlign(HAlign_Center)
+		.VAlign(VAlign_Center)
+		.ToolTipText(LOCTEXT( "BrowseButtonToolTipText", "Browse to Font asset in Content Browser"))
+		[
+			SNew(SButton)
+			.ButtonStyle(FAppStyle::Get(), "SimpleButton")
+			.OnClicked(this, &FAvaFontDetailsCustomization::OnBrowseToAssetClicked)
+			.ContentPadding(0)
+			.Visibility_Lambda([this](){ return !bImportEnabled ? EVisibility::Visible : EVisibility::Collapsed; })
+			[
+				SNew(SImage)
+				.Image(BrowseIcon)
+				.ColorAndOpacity(FSlateColor::UseForeground())
+			]
+		];
+}
+
+TSharedRef<SWidget> FAvaFontDetailsCustomization::MakeRefreshFontButton()
+{
+	const FSlateBrush* BrowseIcon = FAppStyle::Get().GetBrush("EditorViewport.RotateMode");
+
+	return
+		SNew(SBox)
+		.HAlign(HAlign_Center)
+		.VAlign(VAlign_Center)
+		.ToolTipText(LOCTEXT( "RefreshFontButtonToolTipText", "Refresh font"))
+		[
+			SNew(SButton)
+			.ButtonStyle(FAppStyle::Get(), "SimpleButton")
+			.OnClicked(this, &FAvaFontDetailsCustomization::OnRefreshFontClicked)
+			.ContentPadding(0)
+			.Visibility_Lambda([this](){ return bRefreshFontEnabled ? EVisibility::Visible : EVisibility::Collapsed; })
+			[
+				SNew(SImage)
+				.Image(BrowseIcon)
+				.ColorAndOpacity(FSlateColor::UseForeground())
+			]
+		];
+}
+
+TSharedRef<SWidget> FAvaFontDetailsCustomization::MakeImportButton()
+{
+	return
+		SNew(SBox)
+		.HAlign(HAlign_Center)
+		.VAlign(VAlign_Center)
+		.WidthOverride(22)
+		.HeightOverride(22)
+		.ToolTipText(LOCTEXT("ImportSystemFont", "Import system font to project"))
+		[
+			SNew(SButton)
+			.ButtonStyle(FAppStyle::Get(), "SimpleButton" )
+			.OnClicked(this, &FAvaFontDetailsCustomization::ImportButtonClicked)
+			.ContentPadding(0)
+			.Visibility_Lambda([this](){ return bImportEnabled ? EVisibility::Visible : EVisibility::Collapsed; })
+			[
+				SNew(SImage)
+				.Image(FAppStyle::GetBrush("Icons.Import"))
+				.ColorAndOpacity(FSlateColor::UseForeground())
+			]
+		];
+}
+
+void FAvaFontDetailsCustomization::UpdateOrGenerateWidget(const TSharedPtr<FAvaFontView>& InItem)
+{
+	const UAvaFontManagerSubsystem* FontManagerSubsystem = FontManagerSubsystemWeak.Get();
+	if (!FontManagerSubsystem)
+	{
+		return;
+	}
+
+	if (ComboBoxContent)
+	{
+		FAvaMultiFontSelectionData MultiFontSelectionData;
+		FontManagerSubsystem->GetMultipleSelectionInformation(AvaFontPropertyHandle, MultiFontSelectionData);
+
+		if (MultiFontSelectionData.bMultipleValues)
+		{
+			ComboBoxContent->SetContent(GenerateMultipleSelectionField());
+			bImportEnabled = false;
+			bRefreshFontEnabled = false;
+		}
+		else
+		{
+			const TSharedRef<SAvaFontField> FontFieldWidget = StaticCastSharedRef<SAvaFontField>(HandleGenerateWidget(InItem));
+			FontFieldWidget->Select();
+			ComboBoxContent->SetContent(FontFieldWidget);
+		}
+	}
+}
+
+void FAvaFontDetailsCustomization::RefreshComboboxContent()
+{
+	if (ComboBoxContent.IsValid())
+	{
+		UpdateOrGenerateWidget(SelectedOption);
+	}
+}
+
+void FAvaFontDetailsCustomization::HandleSelectionChanged(const TSharedPtr<FAvaFontView> InItem, ESelectInfo::Type InSelectionType)
+{
+	// Sometimes when removing an item from favorites, this gets triggered, but the item is not there anymore,
+	// so we need to check if the callback Item argument is actually valid
+	if (InItem && InItem != SelectedOption)
+	{
+		SelectedOption = InItem;
+
+		if (InItem->GetFontObject())
+		{
+			UAvaFontObject* const FontObjectSlateAttribute = InItem->GetFontObject();
+
+			if (IsValid(FontObjectSlateAttribute))
+			{
+				UpdateFontObjectProperty(FontObjectSlateAttribute);
+				UpdateButtonsStatus(SelectedOption);
+
+				UAvaFontManagerSubsystem* FontManagerSubsystem = FontManagerSubsystemWeak.Get();
+                if (!FontManagerSubsystem)
+                {
+                	return;
+                }
+
+				if (SelectedOption.IsValid())
+				{
+					// If the newly selected font obj is not yet imported, let's mark it for auto save
+					if (SelectedOption->GetFontSource() == EAvaFontSource::System)
+					{
+						FontManagerSubsystem->MarkFontForAutoImport(AvaFontPropertyHandle);
+					}
+				}
+			}
+		}
+
+		// When the selection changes we always generate another widget to represent the content area of the combobox.
+		RefreshComboboxContent();
+
+		RefreshMissingFontWidget();
+	}
+}
+
+void FAvaFontDetailsCustomization::OnFontFieldUpdated(const TSharedPtr<FAvaFontView>& InItem) const
+{
+	if (FontSelector)
+	{
+		FontSelector->RefreshOptions();
+	}
+}
+
+void FAvaFontDetailsCustomization::UpdateFontObjectProperty(UAvaFontObject* InFontObject)
+{
+	check(InFontObject);
+
+	if (!AvaFontPropertyHandle)
+	{
+		return;
+	}
+
+	TArray<UObject*> OuterObjects;
+	AvaFontPropertyHandle->GetOuterObjects(OuterObjects);
+
+	FPropertyAccess::Result AccessResult;
+	const TSharedPtr<FAvaFontView> CurrAvaFontView = GetCurrentAvaFontView(AccessResult);
+
+	if (AccessResult == FPropertyAccess::Success || !CurrAvaFontView)
+	{
+		bool bFontObjNeedsUpdate;
+
+		const UAvaFontObject* CurrFontObjectPtr = nullptr;
+		if (CurrAvaFontView)
+		{
+			CurrFontObjectPtr = CurrAvaFontView->GetFontObject();
+		}
+
+		if (!CurrFontObjectPtr)
+		{
+			bFontObjNeedsUpdate = true;
+		}
+		else
+		{
+			bFontObjNeedsUpdate = CurrFontObjectPtr->GetFontName() != InFontObject->GetFontName();
+		}
+
+		if (bFontObjNeedsUpdate)
+		{
+			if (!OuterObjects.IsEmpty())
+			{
+				for (int32 Index = 0; Index < OuterObjects.Num(); Index++)
+				{
+					if (UObject* Object = OuterObjects[Index])
+					{
+						UAvaFontObject* NewFontObj = NewObject<UAvaFontObject>(Object, InFontObject->GetFName(), RF_Public);
+						NewFontObj->InitFromFontObject(InFontObject);
+
+						FString ValueString;
+						FAvaFont::GenerateAvalancheFontFormattedString(NewFontObj, ValueString);
+
+						AvaFontPropertyHandle->SetPerObjectValue(Index, ValueString);
+					}
+				}
+			}
+			else
+			{
+				// No direct outer(s), so we try to find the closest one (e.g. this happens for ava fonts in RC Panel)
+
+				UAvaFontObject* SourceFontObject = nullptr;
+				if (FProperty* Property = AvaFontPropertyHandle->GetProperty())
+				{
+					if (const FProperty* OwnerProperty = Property->GetOwnerProperty())
+					{
+						if (OwnerProperty->Owner.IsUObject())
+						{
+							if (UObject* Outer = OwnerProperty->Owner.ToUObject())
+							{
+								SourceFontObject = NewObject<UAvaFontObject>(Outer, InFontObject->GetFName(), RF_Public);
+								SourceFontObject->InitFromFontObject(InFontObject);
+							}
+						}
+					}
+				}
+
+				// If we couldn't reach an outer, let's use the input font object, which might be transient
+				if (!SourceFontObject)
+				{
+					SourceFontObject = InFontObject;
+				}
+
+				FString ValueString;
+				bool bStringSuccess = FAvaFont::GenerateAvalancheFontFormattedString(SourceFontObject, ValueString);
+
+				if (bStringSuccess)
+				{
+					AvaFontPropertyHandle->SetPerObjectValues({ValueString});
+				}
+			}
+
+			AvaFontPropertyHandle->NotifyPostChange(EPropertyChangeType::ValueSet);
+		}
+	}
+	// MultipleValues needs to be handled differently, since we want to act on all selected objects
+	else if (AccessResult == FPropertyAccess::MultipleValues)
+	{
+		AccessResult = FPropertyAccess::MultipleValues;
+
+		const UAvaFontManagerSubsystem* FontManagerSubsystem = FontManagerSubsystemWeak.Get();
+		if (!FontManagerSubsystem)
+		{
+			return;
+		}
+
+		FAvaMultiFontSelectionData MultiFontSelectionData;
+		FontManagerSubsystem->GetMultipleSelectionInformation(AvaFontPropertyHandle, MultiFontSelectionData);
+
+		const bool bFontObjNeedsUpdate = InFontObject->GetFontName() != MultiFontSelectionData.FirstSelectionFontName;
+
+		if (bFontObjNeedsUpdate)
+		{
+			TArray<FString> PerObjectValues;
+
+			for (UObject* const OuterObject : OuterObjects)
+			{
+				UAvaFontObject* NewFontObj = NewObject<UAvaFontObject>(OuterObject, InFontObject->GetFName(), RF_Public);
+				NewFontObj->InitFromFontObject(InFontObject);
+
+				FString ValueString;
+				FAvaFont::GenerateAvalancheFontFormattedString(NewFontObj, ValueString);
+
+				PerObjectValues.Add(ValueString);
+			}
+
+			AvaFontPropertyHandle->SetPerObjectValues(PerObjectValues);
+		}
+	}
+}
+
+TSharedPtr<FAvaFontView> FAvaFontDetailsCustomization::GetCurrentAvaFontView(FPropertyAccess::Result& OutAccessResult)
+{
+	UAvaFontManagerSubsystem* FontManagerSubsystem = FontManagerSubsystemWeak.Get();
+	if (!FontManagerSubsystem)
+	{
+		return nullptr;
+	}
+
+	CurrentAvaFontValue = FontManagerSubsystem->GetFontViewFromPropertyHandle(AvaFontPropertyHandle, OutAccessResult);
+	return CurrentAvaFontValue;
+}
+
+FAvaFont* FAvaFontDetailsCustomization::GetCurrentAvaFont(FPropertyAccess::Result& OutAccessResult) const
+{
+	const UAvaFontManagerSubsystem* FontManagerSubsystem = FontManagerSubsystemWeak.Get();
+	if (!FontManagerSubsystem)
+	{
+		return nullptr;
+	}
+
+	return FontManagerSubsystem->GetFontFromPropertyHandle(AvaFontPropertyHandle, OutAccessResult);
+}
+
+void FAvaFontDetailsCustomization::UpdateButtonsStatus(const TSharedPtr<FAvaFontView>& InAvaFont)
+{
+	if (InAvaFont.IsValid() && InAvaFont->GetFontSource() == EAvaFontSource::System)
+	{
+		bImportEnabled = true;
+		bRefreshFontEnabled = false;
+	}
+	else
+	{
+		bImportEnabled = false;
+		if (const UAvaFontManagerSubsystem* FontManagerSubsystem = FontManagerSubsystemWeak.Get())
+		{
+			bRefreshFontEnabled = FontManagerSubsystem->IsFontAvailableOnLocalOS(InAvaFont->GetFont());
+		}
+	}
+}
+
+void FAvaFontDetailsCustomization::OnProjectFontCreated(const UFont* InFont)
+{
+	Refresh();
+}
+
+void FAvaFontDetailsCustomization::OnProjectFontDeleted(const UFont* InFont)
+{
+	FPropertyAccess::Result AccessResult;
+	if (const TSharedPtr<FAvaFontView> CurrAvaFont = GetCurrentAvaFontView(AccessResult))
+	{
+		const UAvaFontObject* CurrFontObj = CurrAvaFont->GetFontObject();
+
+		if (CurrFontObj && CurrFontObj->GetFont() == InFont)
+		{
+			SetDefaultValue();
+
+			if (AvaFontPropertyHandle)
+			{
+				AvaFontPropertyHandle->NotifyPostChange(EPropertyChangeType::ValueSet);
+			}
+		}
+	}
+
+	Refresh();
+}
+
+void FAvaFontDetailsCustomization::OnSystemFontsUpdated()
+{
+	Refresh();
+
+	if (FontSelector.IsValid())
+	{
+		FontSelector->RefreshOptions();
+	}
+}
+
+void FAvaFontDetailsCustomization::RegisterFontManagerCallbacks()
+{
+	UAvaFontManagerSubsystem* FontManagerSubsystem = FontManagerSubsystemWeak.Get();
+	if (!FontManagerSubsystem)
+	{
+		return;
+	}
+
+	FontManagerSubsystem->OnProjectFontCreated().AddRaw(this, &FAvaFontDetailsCustomization::OnProjectFontCreated);
+	FontManagerSubsystem->OnProjectFontDeleted().AddRaw(this, &FAvaFontDetailsCustomization::OnProjectFontDeleted);
+	FontManagerSubsystem->OnSystemFontsUpdated().AddRaw(this, &FAvaFontDetailsCustomization::OnSystemFontsUpdated);
+}
+
+void FAvaFontDetailsCustomization::UnregisterFontManagerCallbacks() const
+{
+	UAvaFontManagerSubsystem* FontManagerSubsystem = FontManagerSubsystemWeak.Get();
+	if (!FontManagerSubsystem)
+	{
+		return;
+	}
+
+	FontManagerSubsystem->OnProjectFontCreated().RemoveAll(this);
+	FontManagerSubsystem->OnProjectFontDeleted().RemoveAll(this);
+	FontManagerSubsystem->OnSystemFontsUpdated().RemoveAll(this);
+}
+
+void FAvaFontDetailsCustomization::RefreshMissingFontWidget() const
+{
+	if (SelectedOption && MissingFontText)
+	{
+		// Check if the font is in fallback state - this happens when a font asset could not be loaded (e.g. missing from disk)
+		if (SelectedOption->IsFallbackFont())
+		{
+			MissingFontText->SetVisibility(EVisibility::Visible);
+
+			const FText FontName = FText::FromString(*SelectedOption->GetMissingFontName());
+			const FText Message = FText::Format(LOCTEXT("MissingFontMessage", "Failed to load font {0}"), FontName);
+
+			MissingFontText->SetText(Message);
+			MissingFontText->SetColorAndOpacity(FLinearColor::Red);
+		}
+		else
+		{
+			MissingFontText->SetVisibility(EVisibility::Collapsed);
+		}
+	}
+}
+
+#undef LOCTEXT_NAMESPACE
diff --git a/Engine/Plugins/Experimental/Avalanche/Source/AvalancheTextEditor/Private/Font/AvaFontDetailsCustomization.h b/Engine/Plugins/Experimental/Avalanche/Source/AvalancheTextEditor/Private/Font/AvaFontDetailsCustomization.h
new file mode 100644
--- /dev/null
+++ b/Engine/Plugins/Experimental/Avalanche/Source/AvalancheTextEditor/Private/Font/AvaFontDetailsCustomization.h
@@ -0,0 +1,124 @@
+// Copyright Epic Games, Inc. All Rights Reserved.
+
+#pragma once
+
+#include "Font/AvaFontObject.h"
+#include "IPropertyTypeCustomization.h"
+#include "Input/Reply.h"
+#include "PropertyEditorModule.h"
+#include "UObject/StrongObjectPtr.h"
+
+class FAvaFontView;
+class FUICommandList;
+class SAvaFontField;
+class SAvaFontSelector;
+class SBox;
+class SButton;
+class STextBlock;
+class UAvaFontManagerSubsystem;
+class UFactory;
+class UFont;
+
+struct FAvaFont;
+struct FAvaMultiFontSelectionData;
+
+class FAvaFontDetailsCustomization : public IPropertyTypeCustomization
+{
+public:
+	static TSharedRef<IPropertyTypeCustomization> MakeInstance();
+
+	virtual ~FAvaFontDetailsCustomization() override;
+
+	//~ Begin IPropertyTypeCustomization
+	virtual void CustomizeHeader(TSharedRef<IPropertyHandle> FontStructPropertyHandle
+		, FDetailWidgetRow& HeaderRow
+		, IPropertyTypeCustomizationUtils& StructCustomizationUtils) override;
+
+	virtual void CustomizeChildren(TSharedRef<IPropertyHandle> PropertyHandle
+		, IDetailChildrenBuilder& ChildBuilder
+		, IPropertyTypeCustomizationUtils& CustomizationUtils) override;
+	//~ End IPropertyTypeCustomization
+
+private:
+	void RefreshSelectedFont();
+
+	void RefreshLocalFontOptions();
+
+	void Refresh();
+
+	void SetDefaultValue();
+
+	FReply ImportButtonClicked();
+
+	TSharedRef<SWidget> HandleGenerateWidget(const TSharedPtr<FAvaFontView>& InItem);
+
+	TSharedRef<SWidget> GenerateMultipleSelectionField() const;
+
+	FReply OnBrowseToAssetClicked();
+	
+	FReply OnRefreshFontClicked();
+
+	TSharedRef<SWidget> MakeBrowseButton();
+
+	TSharedRef<SWidget> MakeRefreshFontButton();
+
+	TSharedRef<SWidget> MakeImportButton();
+
+	void UpdateOrGenerateWidget(const TSharedPtr<FAvaFontView>& InItem);
+
+	void HandleSelectionChanged(const TSharedPtr<FAvaFontView> InItem, ESelectInfo::Type InSelectionType);
+
+	void OnFontFieldUpdated(const TSharedPtr<FAvaFontView>& InItem) const;
+
+	void UpdateFontObjectProperty(UAvaFontObject* InFontObject);
+
+	TSharedPtr<FAvaFontView> GetCurrentAvaFontView(FPropertyAccess::Result& OutAccessResult);
+
+	FAvaFont* GetCurrentAvaFont(FPropertyAccess::Result& OutAccessResult) const;
+
+	void UpdateButtonsStatus(const TSharedPtr<FAvaFontView>& InAvaFont);
+
+	void OnProjectFontCreated(const UFont* InFont);
+
+	void OnProjectFontDeleted(const UFont* InObject);
+
+	void OnSystemFontsUpdated();
+
+	void RegisterFontManagerCallbacks();
+
+	void UnregisterFontManagerCallbacks() const;
+
+	void RefreshMissingFontWidget() const;
+
+	void RefreshComboboxContent();
+
+	/** The true objects bound to the Slate combobox. */
+	TArray<TSharedPtr<FAvaFontView>> Options;
+
+	/** A shared pointer to the underlying slate combobox */
+	TSharedPtr<SAvaFontSelector> FontSelector;
+
+	/** A shared pointer to a container that holds the combobox content that is selected */
+	TSharedPtr<SBox> ComboBoxContent;
+
+	/** A shared pointer to the current selected font (points to font in font manager list) */
+	TSharedPtr<FAvaFontView> SelectedOption;
+
+	/** Reference to default font, used also as fallback */
+	TSharedPtr<FAvaFontView> DefaultOption;
+
+	/** Stores a reference to the font currently handled by this customization PropertyHandle */
+	TSharedPtr<FAvaFontView> CurrentAvaFontValue;
+
+	TSharedPtr<IPropertyHandle> AvaFontPropertyHandle;
+
+	TSharedPtr<STextBlock> MissingFontText;
+
+	/** Import is enabled for system-only fonts (they are not assets) */
+	bool bImportEnabled = true;
+
+	/** Refresh is enabled for fonts which are available on the local OS */
+	bool bRefreshFontEnabled = true;
+
+	TWeakObjectPtr<UAvaFontManagerSubsystem> FontManagerSubsystemWeak;
+};
