[MotionDesign]

diff --git a/Engine/Plugins/Experimental/Avalanche/Source/AvalancheTextEditor/Private/Font/AvaFontManagerSubsystem.cpp b/Engine/Plugins/Experimental/Avalanche/Source/AvalancheTextEditor/Private/Font/AvaFontManagerSubsystem.cpp
new file mode 100644
--- /dev/null
+++ b/Engine/Plugins/Experimental/Avalanche/Source/AvalancheTextEditor/Private/Font/AvaFontManagerSubsystem.cpp
@@ -0,0 +1,1269 @@
+// Copyright Epic Games, Inc. All Rights Reserved.
+
+#include "Font/AvaFontManagerSubsystem.h"
+#include "AssetRegistry/AssetRegistryModule.h"
+#include "AutomatedAssetImportData.h"
+#include "AvaFontView.h"
+#include "Editor.h"
+#include "Factories/FontFileImportFactory.h"
+#include "Font/AvaFont.h"
+#include "Font/AvaFontObject.h"
+#include "HAL/PlatformFileManager.h"
+#include "Misc/ScopedSlowTask.h"
+#include "PropertyHandle.h"
+#include "SystemFontLoading.h"
+#include "UObject/ObjectSaveContext.h"
+#include "UObject/SavePackage.h"
+
+/** ThirdParty */
+#include "EditorAssetLibrary.h"
+#include "freetype/freetype.h"
+
+DEFINE_LOG_CATEGORY(LogAvaFont);
+
+#define LOCTEXT_NAMESPACE "AvaFontManagerSubsystem"
+
+UAvaFontManagerSubsystem* UAvaFontManagerSubsystem::Get()
+{
+	if (!GEditor)
+	{
+		return nullptr;
+	}
+
+	return GEditor->GetEditorSubsystem<UAvaFontManagerSubsystem>();
+}
+
+void UAvaFontManagerSubsystem::Initialize(FSubsystemCollectionBase& Collection)
+{
+	Initialize();
+
+	const FAssetRegistryModule& AssetRegistryModule = FModuleManager::LoadModuleChecked<FAssetRegistryModule>("AssetRegistry");
+
+	// Initialize project fonts as soon as AssetRegistry module is ready for that
+	// Project fonts will be found, compared to System fonts to avoid duplicates, and then added to the list including both system and project fonts
+	AssetRegistryModule.Get().OnFilesLoaded().AddUObject(this, &UAvaFontManagerSubsystem::InitializeProjectFonts);
+
+	Super::Initialize(Collection);
+}
+
+void UAvaFontManagerSubsystem::Deinitialize()
+{
+	UnregisterDefaultAutoImportCallbacks();
+	UnregisterAssetsCallbacks();
+
+	ClearFontsData();
+
+	CleanFreeTypeLibrary();
+
+	Super::Deinitialize();
+}
+
+void UAvaFontManagerSubsystem::ClearFontsData()
+{
+	OSFontsMap.Empty();
+	ProjectFontsMap.Empty();
+	AvalancheFontsMap.Empty();
+	AvalancheFontsOptions.Empty();
+	FontsToImportOnSave.Empty();
+}
+
+void UAvaFontManagerSubsystem::RegisterDefaultAutoImportCallbacks()
+{
+	// auto saving non imported fonts which are in use in level editor
+	FEditorDelegates::PreSaveWorldWithContext.AddUObject(this, &UAvaFontManagerSubsystem::FontsAutoImportPreSave);
+}
+
+void UAvaFontManagerSubsystem::UnregisterDefaultAutoImportCallbacks() const
+{
+	FEditorDelegates::PreSaveWorldWithContext.RemoveAll(this);
+	FEditorDelegates::PostSaveWorldWithContext.RemoveAll(this);
+}
+
+void UAvaFontManagerSubsystem::Initialize()
+{
+	UnregisterAssetsCallbacks();
+
+	LoadFavorites();
+
+	// init OS fonts only, project fonts can be loaded lately with InitializeProjectFonts()
+	InitializeOSFonts();
+
+	RefreshAvalancheFontsMap();
+	RefreshAvalancheFontsOptions();
+
+	RegisterAssetsCallbacks();
+
+	RegisterDefaultAutoImportCallbacks();
+}
+
+bool UAvaFontManagerSubsystem::IsSupportedFontFile(const FString& InFontFilePath)
+{
+	const FString FileName = FPaths::GetCleanFilename(InFontFilePath);
+	const FString Extension = FPaths::GetExtension(FileName).ToLower();
+
+	return Extension == TEXT("ttf") || Extension == TEXT("otf");
+}
+
+void UAvaFontManagerSubsystem::OnAssetCreated(UObject* InObject)
+{
+	if (!IsValid(InObject))
+	{
+		return;
+	}
+
+	if (InObject->GetClass()->IsChildOf<UFont>())
+	{
+		UFont* NewFont = Cast<UFont>(InObject);
+
+		if (IsValid(NewFont))
+		{
+			CreateProjectFont(NewFont);
+		}
+
+		RefreshAvalancheFontsMap();
+		RefreshAvalancheFontsOptions();
+
+		OnProjectFontCreatedDelegate.Broadcast(NewFont);
+	}
+}
+
+void UAvaFontManagerSubsystem::OnAssetDeleted(UObject* InObject)
+{
+	if (!IsValid(InObject))
+	{
+		return;
+	}
+
+	if (InObject->GetClass()->IsChildOf<UFont>())
+	{
+		const UFont* const FontToBeDeleted = Cast<UFont>(InObject);
+
+		if (IsValid(FontToBeDeleted))
+		{
+			OnProjectFontDeletedDelegate.Broadcast(FontToBeDeleted);
+			RemoveProjectFont(FontToBeDeleted);
+		}
+
+		// this should allow the assets to still point to the same font, in case the deleted font is a system one
+		LoadOSFonts();
+
+		RefreshAvalancheFontsMap();
+		RefreshAvalancheFontsOptions();
+	}
+}
+
+void UAvaFontManagerSubsystem::LoadFavorites()
+{
+	if (!IsValid(AvaFontManagerConfig))
+	{
+		static const FString PackageName = GetFontConfigPackageName();
+		UPackage* NewPackage = CreatePackage(*PackageName);
+		NewPackage->SetFlags(RF_Transient);
+		NewPackage->AddToRoot();
+		AvaFontManagerConfig = NewObject<UAvaFontConfig>(NewPackage, "AvalancheFontConfig", RF_Transient | RF_Transactional | RF_Standalone);
+	}
+}
+
+void UAvaFontManagerSubsystem::SanitizeString(FString& OutSanitizedName)
+{
+	// spaces would be also handled by code below, but spaces used to be removed
+	// so let's to this anyway in order to avoid any potential issues with new vs. old imported fonts
+
+	const TCHAR* InvalidObjectChar = INVALID_OBJECTNAME_CHARACTERS;
+	while (*InvalidObjectChar)
+	{
+		OutSanitizedName.ReplaceCharInline(*InvalidObjectChar, TCHAR(' '), ESearchCase::CaseSensitive);
+		++InvalidObjectChar;
+	}
+
+	const TCHAR* InvalidPackageChar = INVALID_LONGPACKAGE_CHARACTERS;
+	while (*InvalidPackageChar)
+	{
+		OutSanitizedName.ReplaceCharInline(*InvalidPackageChar, TCHAR(' '), ESearchCase::CaseSensitive);
+		++InvalidPackageChar;
+	}
+
+	OutSanitizedName.RemoveSpacesInline();
+}
+
+bool UAvaFontManagerSubsystem::ImportAsAsset(UAvaFontObject* InFontToImport)
+{
+	if (IsValid(InFontToImport))
+	{
+		if (InFontToImport->GetSource() == EAvaFontSource::System)
+		{
+			const FString& FontToImportName = InFontToImport->GetFontName();
+
+			// system font name might include invalid characters, let's handle them
+			FString SanitizedFontFamilyName(FontToImportName);
+			SanitizeString(SanitizedFontFamilyName);
+
+			const FString FontPackageName = GetImportFontPackageNameRoot() + TEXT("/Fonts/") + SanitizedFontFamilyName;
+
+			UPackage* FontPackage = CreatePackage(*FontPackageName);
+
+			if (IsValid(FontPackage))
+			{
+				const UFont* Font = InFontToImport->GetFont();
+				if (!IsValid(Font))
+				{
+					return false;
+				}
+
+				const FString FontAssetFileName = FPackageName::LongPackageNameToFilename(FontPackage->GetPathName(), FPackageName::GetAssetPackageExtension());
+
+				FontPackage->AddToRoot();
+				FontPackage->FullyLoad();
+
+				UFont* NewFont = NewObject<UFont>(FontPackage, *Font->GetName(), RF_Public | RF_Standalone);
+				NewFont->ImportOptions.FontName = FontToImportName;
+				NewFont->LegacyFontName = FName(FontToImportName);
+				NewFont->FontCacheType = EFontCacheType::Runtime;
+
+				SetupFontFamilyTypefaces(NewFont, InFontToImport->GetFontParameters());
+
+				FAssetRegistryModule::AssetCreated(NewFont);
+				FontPackage->SetDirtyFlag(true);
+
+				InFontToImport->InitProjectFont(NewFont, InFontToImport->GetFontName());
+
+				if (!FPlatformFileManager::Get().GetPlatformFile().IsReadOnly(*FontAssetFileName))
+				{
+					FSavePackageArgs FontSavePackageArgs;
+					FontSavePackageArgs.TopLevelFlags = RF_Public | RF_Standalone;
+					UPackage::SavePackage(FontPackage, NewFont, *FontAssetFileName, FontSavePackageArgs);
+				}
+
+				if (OSFontsMap.Contains(FontToImportName))
+				{
+					OSFontsMap.Remove(FontToImportName);
+
+					// Creates a new AvaFontObject and updates the ProjectsFontsMap as well
+					CreateProjectFont(NewFont);
+				}
+
+				RefreshAvalancheFontsMap();
+				RefreshAvalancheFontsOptions();
+				OnSystemFontsUpdatedDelegate.Broadcast();
+
+				return true;
+			}
+		}
+	}
+
+	return false;
+}
+
+bool UAvaFontManagerSubsystem::ImportFont(const TSharedPtr<IPropertyHandle>& InAvaFontPropertyHandle)
+{
+	bool bSuccess = false;
+
+	if (const TSharedPtr<FAvaFontView>& AvaFont = GetFontViewFromPropertyHandle(InAvaFontPropertyHandle))
+	{
+		// we don't want to trigger automatic behavior for manually imported fonts
+		UnregisterAssetsCallbacks();
+
+		bSuccess = ImportAsAsset(AvaFont->GetFontObject());
+
+		if (bSuccess)
+		{
+			ForceFontAutoImportUnmark(InAvaFontPropertyHandle);
+			InAvaFontPropertyHandle->NotifyPostChange(EPropertyChangeType::ValueSet);
+		}
+
+		// register callbacks again
+		RegisterAssetsCallbacks();
+	}
+
+	return bSuccess;
+}
+
+void UAvaFontManagerSubsystem::RefreshSystemFont(UFont* InFont) const
+{
+	if (!InFont)
+	{
+		return;
+	}
+
+	const FString& FontFamilyName = InFont->GetName();
+
+	FScopedSlowTask DeleteReimportTypefaceSlowTask(1
+		, FText::Format(LOCTEXT("DeleteAndReimportTypeface", "Refreshing {0} font.")
+		, FText::FromString(FontFamilyName)));
+
+	DeleteReimportTypefaceSlowTask.MakeDialog();
+	DeleteReimportTypefaceSlowTask.EnterProgressFrame();
+
+	// We cannot update fonts which are not available on this system
+	if (!FontsInfoMap.Contains(FontFamilyName))
+	{
+		return;
+	}
+
+	// We remove old Font Faces
+	for (const FTypefaceEntry& TypefaceEntry : InFont->CompositeFont.DefaultTypeface.Fonts)
+	{
+		if (const UFontFace* FontFace = Cast<UFontFace>(TypefaceEntry.Font.GetFontFaceAsset()))
+		{
+			UEditorAssetLibrary::DeleteAsset(FontFace->GetPathName());
+		}
+	}
+
+	// We empty font faces from font
+	InFont->CompositeFont.DefaultTypeface.Fonts.Reset();
+
+	// Setup font typefaces
+	SetupFontFamilyTypefaces(InFont, FontsInfoMap[FontFamilyName]);
+
+	if (UPackage* FontPackage = InFont->GetPackage())
+	{
+		FontPackage->SetDirtyFlag(true);
+
+		const FString FontAssetFileName = FPackageName::LongPackageNameToFilename(FontPackage->GetPathName(), FPackageName::GetAssetPackageExtension());
+
+		if (!FPlatformFileManager::Get().GetPlatformFile().IsReadOnly(*FontAssetFileName))
+		{
+			FSavePackageArgs SavePackageArgs;
+        	SavePackageArgs.TopLevelFlags = RF_Public | RF_Standalone;
+			UPackage::SavePackage(FontPackage, InFont, *FontAssetFileName, SavePackageArgs);
+		}
+	}
+}
+
+bool UAvaFontManagerSubsystem::IsFontAvailableOnLocalOS(const UFont* InFont) const
+{
+	if (!InFont)
+	{
+		return false;
+	}
+
+	const FString& FontFamilyName = InFont->GetName();
+
+	// We cannot update fonts which are not available on this system
+	return FontsInfoMap.Contains(FontFamilyName);
+}
+
+UFontFace* UAvaFontManagerSubsystem::GetFallbackFontFace()
+{
+	if (!IsValid(DefaultFallbackFontFace))
+	{
+		DefaultFallbackFontFace = LoadObject<UFontFace>(nullptr, GetFallbackFontFaceReference());
+	}
+
+	return DefaultFallbackFontFace;
+}
+
+TArray<FString> UAvaFontManagerSubsystem::GetSystemFontNames() const
+{
+	TArray<FString> OutFontsNames;
+
+	FontsInfoMap.GetKeys(OutFontsNames);
+
+	return OutFontsNames;
+}
+
+TArray<FString> UAvaFontManagerSubsystem::GetSystemFontTypefaces(const FString& InFontFamilyName) const
+{
+	TArray<FString> OutTypefaces;
+	if (FontsInfoMap.Contains(InFontFamilyName))
+	{
+		OutTypefaces = FontsInfoMap[InFontFamilyName].GetFontFaceNames();
+	}
+
+	return OutTypefaces;
+}
+
+void UAvaFontManagerSubsystem::SetupFontFamilyTypefaces(UFont* InFont, const FSystemFontsRetrieveParams& InFontParams)
+{
+	const FString& FontFamilyName = InFontParams.FontFamilyName;
+
+	const bool bIsTempAsset = InFont->GetPackage()->GetName().Contains(GetTempFontPackageNameRoot());
+	const FString PackagePath = bIsTempAsset ? GetTempFontPackageNameRoot() : GetImportFontPackageNameRoot();
+
+	UPackage* ImportPackage = CreatePackage(*PackagePath);
+	ImportPackage->AddToRoot();
+
+	UFontFileImportFactory* FontFaceFactory = NewObject<UFontFileImportFactory>();
+	FontFaceFactory->AddToRoot();
+
+	const UAutomatedAssetImportData* AutomatedAssetImportData = NewObject<UAutomatedAssetImportData>();
+	FontFaceFactory->SetAutomatedAssetImportData(AutomatedAssetImportData);
+
+	constexpr EObjectFlags ImportFontAssetFlags = RF_Public | RF_Standalone;
+	constexpr EObjectFlags TempFontAssetFlags = ImportFontAssetFlags | RF_Transient;
+
+	const EObjectFlags NewFontFaceFlags = bIsTempAsset ? TempFontAssetFlags : ImportFontAssetFlags;
+
+	for (int32 Index = 0; Index < InFontParams.GetFontFaceNames().Num(); Index++)
+	{
+		const FString& FontPath = InFontParams.GetFontFacePaths()[Index];
+		const FString& FontFaceName = InFontParams.GetFontFaceNames()[Index];
+		if (!IsSupportedFontFile(FontPath))
+		{
+			continue;
+		}
+
+		FString FontFaceAssetName = FontFamilyName + "_" + FontFaceName;
+		SanitizeString(FontFaceAssetName);
+		UPackage* FontFacePackage = CreatePackage(*(ImportPackage->GetName() + TEXT("/FontFaces/") + FontFaceAssetName));
+		FontFacePackage->AddToRoot();
+		FontFacePackage->FullyLoad();
+
+		bool bCanceled = false;
+		UObject* NewFontFaceAsObject = FontFaceFactory->ImportObject(UFontFace::StaticClass(), FontFacePackage, *FontFaceAssetName, NewFontFaceFlags, FontPath, TEXT(""), bCanceled);
+		UFontFace* NewFontFace = Cast<UFontFace>(NewFontFaceAsObject);
+
+		if (!NewFontFace)
+		{
+			continue;
+		}
+
+		UPackage* StyleFontFacePackage = CreatePackage(*FontFacePackage->GetName());
+		StyleFontFacePackage->AddToRoot();
+		StyleFontFacePackage->FullyLoad();
+
+		if (!bIsTempAsset)
+		{
+			FAssetRegistryModule::AssetCreated(NewFontFace);
+			StyleFontFacePackage->SetDirtyFlag(true);
+
+			const FString StyleFontFaceAssetFileName = FPackageName::LongPackageNameToFilename(StyleFontFacePackage->GetPathName(), FPackageName::GetAssetPackageExtension());
+
+			if (!FPlatformFileManager::Get().GetPlatformFile().IsReadOnly(*StyleFontFaceAssetFileName))
+			{
+				FSavePackageArgs SavePackageArgs;
+				SavePackageArgs.TopLevelFlags = RF_Public | RF_Standalone;
+				UPackage::SavePackage(StyleFontFacePackage, NewFontFace, *StyleFontFaceAssetFileName, SavePackageArgs);
+			}
+		}
+
+		if (NewFontFace)
+		{
+			FTypefaceEntry& Entry = InFont->CompositeFont.DefaultTypeface.Fonts[InFont->CompositeFont.DefaultTypeface.Fonts.AddDefaulted()];
+			Entry.Name = FName(FName::NameToDisplayString(FontFaceName, false));
+			Entry.Font = FFontData(NewFontFace);
+		}
+	}
+
+	FontFaceFactory->RemoveFromRoot();
+}
+
+bool UAvaFontManagerSubsystem::Exec(class UWorld* InWorld, const TCHAR* InCmd, FOutputDevice& InAr)
+{
+	if (FParse::Command(&InCmd, TEXT("fontmanager")))
+	{
+		if (FParse::Command(&InCmd, TEXT("list")))
+		{
+			UE::Avalanche::Private::Fonts::ListAvailableFontFiles();
+			return true;
+		}
+		else if (FParse::Command(&InCmd, TEXT("refresh")))
+		{
+			if (!HasAnyFlags(RF_ClassDefaultObject))
+			{
+				RefreshAllSystemFonts();
+				return true;
+			}
+		}
+	}
+
+	return false;
+}
+
+void UAvaFontManagerSubsystem::CombineAllFonts()
+{
+	AvalancheFontsMap.Empty();
+	AvalancheFontsMap.Append(ProjectFontsMap);
+	AvalancheFontsMap.Append(OSFontsMap);
+}
+
+void UAvaFontManagerSubsystem::UnregisterAssetsCallbacks() const
+{
+	if (FModuleManager::Get().IsModuleLoaded("AssetRegistry"))
+	{
+		const FAssetRegistryModule& AssetRegistryModule = FModuleManager::GetModuleChecked<FAssetRegistryModule>("AssetRegistry");
+
+		if (AssetRegistryModule.IsValid())
+		{
+			AssetRegistryModule.Get().OnInMemoryAssetCreated().RemoveAll(this);
+			AssetRegistryModule.Get().OnInMemoryAssetDeleted().RemoveAll(this);
+		}
+	}
+}
+
+void UAvaFontManagerSubsystem::RegisterAssetsCallbacks()
+{
+	const FAssetRegistryModule& AssetRegistryModule = FModuleManager::LoadModuleChecked<FAssetRegistryModule>("AssetRegistry");
+
+	if (AssetRegistryModule.IsValid())
+	{
+		AssetRegistryModule.Get().OnInMemoryAssetCreated().AddUObject(this, &UAvaFontManagerSubsystem::OnAssetCreated);
+		AssetRegistryModule.Get().OnInMemoryAssetDeleted().AddUObject(this, &UAvaFontManagerSubsystem::OnAssetDeleted);
+	}
+}
+
+void UAvaFontManagerSubsystem::InitializeProjectFonts()
+{
+	ProjectFontsMap.Empty();
+
+	LoadProjectFonts();
+
+	RemoveProjectFontsFromOSFonts();
+}
+
+void UAvaFontManagerSubsystem::InitializeOSFonts()
+{
+	OSFontsMap.Empty();
+	LoadOSFonts();
+}
+
+void UAvaFontManagerSubsystem::LoadProjectFonts()
+{
+	const FAssetRegistryModule& AssetRegistryModule = FModuleManager::LoadModuleChecked<FAssetRegistryModule>("AssetRegistry");
+
+	if (!AssetRegistryModule.IsValid())
+	{
+		return;
+	}
+
+	TArray<FAssetData> AssetData;
+	const UClass* Class = UFont::StaticClass();
+
+	const FTopLevelAssetPath AssetPath(Class->GetPathName());
+	AssetRegistryModule.Get().GetAssetsByClass(AssetPath, AssetData);
+
+	// we don't want to go through temp system fonts, so we filter them based on their Package Path since we know where we put them
+	TArray<FAssetData> AssetDataArrayForPackageName = AssetData.FilterByPredicate([&](const FAssetData& AssetDataItem) -> bool
+		{
+			const FName PathToExcludeTemp = *(GetTempFontPackageNameRoot()/TEXT("Fonts"));
+			return (AssetDataItem.PackagePath != PathToExcludeTemp);
+		});
+
+	const FString& TempPackageName = UAvaFontManagerSubsystem::GetFontObjPackageNameRoot();
+	UPackage* ProjectFontsPackage = CreatePackage(*TempPackageName);
+	ProjectFontsPackage->AddToRoot();
+
+	TArray<FString> CurrFontAssetsNamesArray;
+
+	// check if there are new fonts to be added from Project Assets
+	for (const FAssetData& Elem : AssetDataArrayForPackageName)
+	{
+		UFont* CurrFont = Cast<UFont>(Elem.GetAsset());
+
+		if (IsValid(CurrFont))
+		{
+			if (const UAvaFontObject* const NewAvalancheFontObj = CreateProjectFont(CurrFont, ProjectFontsPackage))
+			{
+				CurrFontAssetsNamesArray.Add(NewAvalancheFontObj->GetFontName());
+				CurrFont->GetPackage()->FullyLoad();
+			}
+		}
+	}
+}
+void UAvaFontManagerSubsystem::RefreshAllSystemFonts() const
+{
+	const FAssetRegistryModule& AssetRegistryModule = FModuleManager::LoadModuleChecked<FAssetRegistryModule>("AssetRegistry");
+
+	if (!AssetRegistryModule.IsValid())
+	{
+		return;
+	}
+
+	TArray<FAssetData> AssetData;
+	const UClass* Class = UFont::StaticClass();
+
+	const FTopLevelAssetPath AssetPath(Class->GetPathName());
+	AssetRegistryModule.Get().GetAssetsByClass(AssetPath, AssetData);
+
+	// We want to go through imported system fonts, so we filter them based on their Package Path since we know where we put them
+	TArray<FAssetData> AssetDataArrayForPackageName = AssetData.FilterByPredicate([&](const FAssetData& AssetDataItem) -> bool
+		{
+			const FName PathToExcludeTemp = *(GetImportFontPackageNameRoot()/TEXT("Fonts"));
+			return (AssetDataItem.PackagePath == PathToExcludeTemp);
+		});
+		TArray<FString> CurrFontAssetsNamesArray;
+
+	if (AssetDataArrayForPackageName.IsEmpty())
+	{
+		return;
+	}
+
+	if (FontsInfoMap.IsEmpty())
+	{
+		return;
+	}
+
+	FScopedSlowTask FixTypefacesSlowTask(AssetDataArrayForPackageName.Num(), LOCTEXT("FixingTypefaces", "Refreshing typefaces for imported system fonts"));
+	FixTypefacesSlowTask.MakeDialog();
+
+	// Check if there are new fonts to be added from Project Assets
+	for (const FAssetData& AssetDataItem : AssetDataArrayForPackageName)
+	{
+		FixTypefacesSlowTask.EnterProgressFrame();
+
+		UFont* Font = Cast<UFont>(AssetDataItem.GetAsset());
+		if (!Font)
+		{
+			continue;
+		}
+
+		RefreshSystemFont(Font);
+	}
+}
+
+bool UAvaFontManagerSubsystem::RemoveProjectFont(const UFont* InSourceFont)
+{
+	if (IsValid(InSourceFont))
+	{
+		FString SourceFontName;
+		GetSanitizedFontName(InSourceFont, SourceFontName);
+
+		if (ProjectFontsMap.Contains(SourceFontName))
+		{
+			if (UAvaFontObject* AvaFontObj = ProjectFontsMap.FindAndRemoveChecked(SourceFontName))
+			{
+				AvaFontObj->Invalidate();
+			}
+
+			return true;
+		}
+	}
+	return false;
+}
+
+UAvaFontObject* UAvaFontManagerSubsystem::CreateProjectFont(UFont* InSourceFont, UPackage* InProjectFontsPackage)
+{
+	if (IsValid(InSourceFont))
+	{
+		if (!InSourceFont->ImportOptions.bUseDistanceFieldAlpha) // ignoring distance field fonts
+		{
+			if (InSourceFont->GetCompositeFont())
+			{
+				FString SanitizedFontName;
+				FString FontName;
+				GetFontName(InSourceFont, FontName);
+				GetSanitizedFontName(InSourceFont, SanitizedFontName);
+
+				if (ProjectFontsMap.Contains(SanitizedFontName))
+				{
+					return nullptr;
+				}
+
+				UPackage* ProjectFontsPackage = InProjectFontsPackage;
+				if (!IsValid(ProjectFontsPackage))
+				{
+					const FString TempPackageName = UAvaFontManagerSubsystem::GetFontObjPackageNameRoot();
+					ProjectFontsPackage = CreatePackage(*TempPackageName);
+					ProjectFontsPackage->AddToRoot();
+				}
+
+				UAvaFontObject* NewAvalancheFont = NewObject<UAvaFontObject>(ProjectFontsPackage, *SanitizedFontName, RF_Public | RF_Standalone | RF_Transient);
+				NewAvalancheFont->AddToRoot();
+
+				NewAvalancheFont->InitProjectFont(InSourceFont, FontName);
+				UAvaFontManagerSubsystem::SetupMetrics(NewAvalancheFont);
+
+				ProjectFontsMap.Add(SanitizedFontName, NewAvalancheFont);
+
+				return NewAvalancheFont;
+			}
+			else
+			{
+				UE_LOG(LogAvaFont, Warning, TEXT("No composite font found for font %s"), *InSourceFont->GetName());
+			}
+		}
+	}
+
+	return nullptr;
+}
+
+void UAvaFontManagerSubsystem::LoadOSFonts()
+{
+	UE::Avalanche::Private::Fonts::GetSystemFontInfo(FontsInfoMap);
+
+	if (FontsInfoMap.IsEmpty())
+	{
+		return;
+	}
+
+	UPackage* OSFontsPackage = CreatePackage(*GetTempFontPackageNameRoot());
+	OSFontsPackage->AddToRoot();
+
+	for (const TPair<FString, FSystemFontsRetrieveParams>& FontParamsPair : FontsInfoMap)
+	{
+		const FSystemFontsRetrieveParams& FontParams = FontParamsPair.Value;
+		ImportSystemFontFamily(FontParams, OSFontsPackage);
+	}
+
+	OnSystemFontsUpdatedDelegate.Broadcast();
+}
+
+void UAvaFontManagerSubsystem::RemoveProjectFontsFromOSFonts()
+{
+	for (const TPair<FString, TObjectPtr<UAvaFontObject>>& ProjectFontPair : ProjectFontsMap)
+	{
+		const FString& FontName = ProjectFontPair.Key;
+		if (OSFontsMap.Contains(FontName))
+		{
+			if (const TObjectPtr<UAvaFontObject> FontToMove = OSFontsMap[FontName])
+			{
+				OSFontsMap.Remove(FontName);
+				FontToMove->SwitchToProjectFont();
+			}
+		}
+	}
+
+	RefreshAvalancheFontsMap();
+	RefreshAvalancheFontsOptions();
+	OnSystemFontsUpdatedDelegate.Broadcast();
+}
+
+void UAvaFontManagerSubsystem::ImportSystemFontFamily(const FSystemFontsRetrieveParams& InFontParams, UPackage* InImportPackage)
+{
+	UPackage* ImportPackage;
+
+	if (IsValid(InImportPackage))
+	{
+		ImportPackage = InImportPackage;
+	}
+	else
+	{
+		ImportPackage = CreatePackage(*GetTempFontPackageNameRoot());
+		ImportPackage->AddToRoot();
+	}
+
+	if (InFontParams.GetFontFaceNames().IsEmpty())
+	{
+		return;
+	}
+
+	const FString& FontFamilyName = InFontParams.FontFamilyName;
+
+	// this font is already in the map, just ignore and return
+	if (OSFontsMap.Contains(FontFamilyName) || ProjectFontsMap.Contains(FontFamilyName))
+	{
+		return;
+	}
+
+	FString SanitizedFontFamilyName(FontFamilyName);
+	SanitizeString(SanitizedFontFamilyName);
+	UPackage* FontPackage = CreatePackage(*(ImportPackage->GetName() + TEXT("/Fonts/") + SanitizedFontFamilyName));
+	FontPackage->AddToRoot();
+
+	constexpr EObjectFlags TempFontAssetFlags = RF_Public | RF_Standalone | RF_Transient;
+
+	UFont* NewFont = NewObject<UFont>(FontPackage, *SanitizedFontFamilyName, TempFontAssetFlags);
+	NewFont->AddToRoot();
+	NewFont->ImportOptions.FontName = FontFamilyName;
+	NewFont->LegacyFontName = FName(FontFamilyName);
+	NewFont->FontCacheType = EFontCacheType::Runtime;
+
+	SetupFontFamilyTypefaces(NewFont, InFontParams);
+	UAvaFontObject* NewAvalancheFont = NewObject<UAvaFontObject>(ImportPackage,  *SanitizedFontFamilyName, TempFontAssetFlags);
+	NewAvalancheFont->AddToRoot();
+
+	FTypefaceEntry& FallbackTypefaceEntry = NewFont->CompositeFont.FallbackTypeface.Typeface.Fonts[NewFont->CompositeFont.FallbackTypeface.Typeface.Fonts.AddDefaulted()];
+	FallbackTypefaceEntry.Font = FFontData(GetFallbackFontFace());
+	FallbackTypefaceEntry.Name = FName("Regular");
+
+	NewAvalancheFont->InitSystemFont(InFontParams, NewFont);
+	UAvaFontManagerSubsystem::SetupMetrics(NewAvalancheFont);
+	OSFontsMap.Add(SanitizedFontFamilyName, NewAvalancheFont);
+}
+
+void UAvaFontManagerSubsystem::SortFontsMap()
+{
+	AvalancheFontsMap.KeySort([](const FString& A, const FString& B)
+	{
+		return A < B;
+	});
+}
+
+void UAvaFontManagerSubsystem::RefreshAvalancheFontsMap()
+{
+	CombineAllFonts();
+	SortFontsMap();
+}
+
+TConstArrayView<const TSharedPtr<FAvaFontView>> UAvaFontManagerSubsystem::GetFontOptions()
+{
+	return AvalancheFontsOptions;
+}
+
+TSharedPtr<FAvaFontView> UAvaFontManagerSubsystem::GetFontViewFromName(const FString& InName)
+{
+	TSharedPtr<FAvaFontView>* AvalancheFontView = AvalancheFontsOptions.FindByPredicate([&InName](const TSharedPtr<FAvaFontView>& AvaFontValue)
+	{
+		if (AvaFontValue.IsValid() && AvaFontValue->HasValidFont())
+		{
+			return AvaFontValue->GetFontNameAsString().Equals(InName);
+		}
+
+		return false;
+	});
+
+	if (AvalancheFontView)
+	{
+		return *AvalancheFontView;
+	}
+
+	return nullptr;
+}
+
+TSharedPtr<FAvaFontView> UAvaFontManagerSubsystem::GetFontViewFromPropertyHandle(const TSharedPtr<IPropertyHandle>& InFontPropertyHandle, FPropertyAccess::Result& OutAccessResult)
+{
+	if (!InFontPropertyHandle)
+	{
+		OutAccessResult = FPropertyAccess::Fail;
+		return nullptr;
+	}
+
+	FString FontName;
+
+	void* FontRawData;
+	OutAccessResult = InFontPropertyHandle->GetValueData(FontRawData);
+
+	if (OutAccessResult == FPropertyAccess::Fail)
+	{
+		return nullptr;
+	}
+
+	if (OutAccessResult == FPropertyAccess::Success)
+	{
+		const FAvaFont* ActualFontRaw = static_cast<FAvaFont*>(FontRawData);
+		if (!ActualFontRaw)
+		{
+			OutAccessResult = FPropertyAccess::Fail;
+			return nullptr;
+		}
+
+		FontName = ActualFontRaw->GetFontNameAsString();
+	}
+	else if (OutAccessResult == FPropertyAccess::MultipleValues)
+	{
+		OutAccessResult = FPropertyAccess::MultipleValues;
+		FAvaMultiFontSelectionData MultiFontSelectionData;
+		GetMultipleSelectionInformation(InFontPropertyHandle, MultiFontSelectionData);
+		if (!MultiFontSelectionData.bMultipleValues)
+		{
+			FontName = MultiFontSelectionData.FirstSelectionFontName;
+		}
+	}
+
+	if (FontName.IsEmpty())
+	{
+		return nullptr;
+	}
+
+	TSharedPtr<FAvaFontView> FontView = GetFontViewFromName(FontName);
+
+	// If result is Multiple Values, we don't consider it a Fail
+	if (OutAccessResult != FPropertyAccess::MultipleValues)
+	{
+		// If font is nullptr, it's a Fail
+		if (!FontView)
+		{
+			OutAccessResult = FPropertyAccess::Fail;
+		}
+	}
+
+	return FontView;
+}
+
+FAvaFont* UAvaFontManagerSubsystem::GetFontFromPropertyHandle(const TSharedPtr<IPropertyHandle>& InFontPropertyHandle, FPropertyAccess::Result& OutAccessResult)
+{
+	if (!InFontPropertyHandle)
+	{
+		OutAccessResult = FPropertyAccess::Fail;
+		return nullptr;
+	}
+
+	FAvaFont* AvaFont = nullptr;
+
+	FString FontName;
+
+	void* FontRawData;
+	OutAccessResult = InFontPropertyHandle->GetValueData(FontRawData);
+
+	if (OutAccessResult == FPropertyAccess::Fail)
+	{
+		return nullptr;
+	}
+
+	if (OutAccessResult == FPropertyAccess::Success)
+	{
+		FAvaFont* ActualFontRaw = static_cast<FAvaFont*>(FontRawData);
+		if (!ActualFontRaw)
+		{
+			OutAccessResult = FPropertyAccess::Fail;
+			return nullptr;
+		}
+
+		return ActualFontRaw;
+	}
+	else if (OutAccessResult == FPropertyAccess::MultipleValues)
+	{
+		OutAccessResult = FPropertyAccess::MultipleValues;
+		FAvaMultiFontSelectionData MultiFontSelectionData;
+		GetMultipleSelectionInformation(InFontPropertyHandle, MultiFontSelectionData);
+		if (!MultiFontSelectionData.bMultipleValues)
+		{
+			FontName = MultiFontSelectionData.FirstSelectionFontName;
+		}
+	}
+
+	if (FontName.IsEmpty())
+	{
+		return nullptr;
+	}
+
+	return AvaFont;
+}
+
+void UAvaFontManagerSubsystem::GetMultipleSelectionInformation(const TSharedPtr<IPropertyHandle>& InPropertyHandle
+	, FAvaMultiFontSelectionData& OutMultiFontSelectionData)
+{
+	bool bMultipleObjectsSelected = false;
+	bool bMultipleValues = false;
+	FString FirstSelectionFontName = TEXT("");
+
+	if (InPropertyHandle)
+	{
+		if (InPropertyHandle->GetNumOuterObjects() > 1)
+		{
+			bMultipleObjectsSelected = true;
+
+			TArray<FString> ObjectValues;
+			InPropertyHandle->GetPerObjectValues(ObjectValues);
+
+			if (!ObjectValues.IsEmpty())
+			{
+				const FString FirstObjectValueString = ObjectValues[0];
+
+				if (FParse::Value(*FirstObjectValueString, TEXT("FontName="), FirstSelectionFontName))
+				{
+					for (int32 ObjectIndex = 1; ObjectIndex < ObjectValues.Num(); ++ObjectIndex)
+					{
+						const FString CurrentValueString = ObjectValues[ObjectIndex];
+						FString CurrentValueFontName;
+
+						if (FParse::Value(*CurrentValueString, TEXT("FontName="), CurrentValueFontName))
+						{
+							if (FirstSelectionFontName != CurrentValueFontName)
+							{
+								bMultipleValues = true;
+								break;
+							}
+						}
+					}
+				}
+			}
+		}
+	}
+
+	OutMultiFontSelectionData.bMultipleObjectsSelected = bMultipleObjectsSelected;
+	OutMultiFontSelectionData.bMultipleValues = bMultipleValues;
+	OutMultiFontSelectionData.FirstSelectionFontName = FirstSelectionFontName;
+}
+
+void UAvaFontManagerSubsystem::GetFontName(const UFont* InFont, FString& OutFontName)
+{
+	// this function used to be defined here, but it has been moved to the Runtime module for easier maintenance
+	UE::Avalanche::FontUtilities::Public::GetFontName(InFont, OutFontName);
+}
+
+void UAvaFontManagerSubsystem::GetSanitizedFontName(const UFont* InFont, FString& OutFontName)
+{
+	if (InFont)
+	{
+		GetFontName(InFont, OutFontName);
+		SanitizeString(OutFontName);
+	}
+}
+
+void UAvaFontManagerSubsystem::RefreshAvalancheFontsOptions()
+{
+	AvalancheFontsOptions.Empty();
+
+	int32 Count = 0;
+
+	for (const TPair<FString, TObjectPtr<UAvaFontObject>>& Elem : AvalancheFontsMap)
+	{
+		FString ElemFontName = Elem.Key;
+		const TObjectPtr<UAvaFontObject>& FontObject = Elem.Value;
+
+		if (IsValid(FontObject->GetFont()))
+		{
+			FontObject->GetFont()->GetPackage()->FullyLoad();
+			FAvaFont FontOption = FAvaFont(FontObject.Get());
+			const bool bFontIsFavorite = IsFavoriteFont(FontOption.GetFontNameAsString());
+			FontOption.SetFavorite(bFontIsFavorite);
+			AvalancheFontsOptions.Add(FAvaFontView::Make(FontOption));
+
+			if (FontObject->GetFont() == FAvaFont::GetDefaultFont())
+			{
+				DefaultFontIndex = Count; 
+			}
+		}
+
+		Count++;
+	}
+}
+
+void UAvaFontManagerSubsystem::AddFavorite(const FString& InFontName)
+{
+	// does something only if needed
+	LoadFavorites();
+
+	if (IsValid(AvaFontManagerConfig))
+	{
+		AvaFontManagerConfig->AddFavoriteFont(InFontName);
+	}
+}
+
+void UAvaFontManagerSubsystem::RemoveFavorite(const FString& InFontName)
+{
+	// does something only if needed
+	LoadFavorites();
+
+	if (IsValid(AvaFontManagerConfig))
+	{
+		AvaFontManagerConfig->RemoveFavoriteFont(InFontName);
+	}
+}
+
+bool UAvaFontManagerSubsystem::IsFavoriteFont(const FString& InFontName)
+{
+	// does something only if needed
+	LoadFavorites();
+
+	if (IsValid(AvaFontManagerConfig))
+	{
+		return AvaFontManagerConfig->IsFavoriteFont(InFontName);
+	}
+
+	return false;
+}
+
+TSharedPtr<FAvaFontView> UAvaFontManagerSubsystem::GetFontViewFromPropertyHandle(const TSharedPtr<IPropertyHandle>& InFontPropertyHandle)
+{
+	FPropertyAccess::Result AccessResult;
+	return GetFontViewFromPropertyHandle(InFontPropertyHandle, AccessResult);
+}
+
+void UAvaFontManagerSubsystem::MarkFontForAutoImport(const TSharedPtr<IPropertyHandle>& InFontToImportOnSave)
+{
+	if (const TSharedPtr<FAvaFontView>& AvaFont = GetFontViewFromPropertyHandle(InFontToImportOnSave))
+	{
+		FontsToImportOnSave.Add(AvaFont, InFontToImportOnSave);
+
+		TArray<UObject*> OuterObjects;
+		InFontToImportOnSave->GetOuterObjects(OuterObjects);
+	}
+}
+
+void UAvaFontManagerSubsystem::UnmarkFontFromAutoImport(const TSharedPtr<IPropertyHandle>& InFontToUnregister)
+{
+	if (const TSharedPtr<FAvaFontView>& AvaFont = GetFontViewFromPropertyHandle(InFontToUnregister))
+	{
+		if (InFontToUnregister)
+		{
+			if (FontsToImportOnSave.Contains(AvaFont))
+			{
+				FontsToImportOnSave.Remove(AvaFont);
+			}
+		}
+	}
+}
+
+void UAvaFontManagerSubsystem::ForceFontAutoImportUnmark(const TSharedPtr<IPropertyHandle>& InFontToUnregister)
+{
+	if (const TSharedPtr<FAvaFontView>& AvaFont = GetFontViewFromPropertyHandle(InFontToUnregister))
+	{
+		if (FontsToImportOnSave.Contains(AvaFont))
+		{
+			FontsToImportOnSave.Remove(AvaFont);
+		}
+	}
+}
+
+void UAvaFontManagerSubsystem::FontsAutoImport()
+{
+	UnregisterAssetsCallbacks();
+
+	for (const TPair<TSharedPtr<FAvaFontView>, TSharedPtr<IPropertyHandle>>& Pair : FontsToImportOnSave)
+	{
+		if (const TSharedPtr<FAvaFontView>& AvaFont = Pair.Key)
+		{
+			UAvaFontObject* FontObjForImport = AvaFont->GetFontObject();
+
+			if (IsValid(FontObjForImport))
+			{
+				if (ImportAsAsset(FontObjForImport))
+				{
+					Pair.Value->NotifyPostChange(EPropertyChangeType::ValueSet);
+				}
+			}
+		}
+	}
+
+	// not needed anymore
+	FontsToImportOnSave.Empty();
+
+	RegisterAssetsCallbacks();
+}
+
+void UAvaFontManagerSubsystem::FontsAutoImportPreSave(UWorld* InWorld, FObjectPreSaveContext InObjectPreSaveContext)
+{
+	FontsAutoImport();
+}
+
+UAvaFontConfig* UAvaFontManagerSubsystem::GetFontManagerConfig()
+{
+	return AvaFontManagerConfig;
+}
+
+bool UAvaFontManagerSubsystem::IsFontMonospaced(const UFont* const InFont)
+{
+	if (InFont)
+	{
+		int32 Height;
+
+		int32 SpaceWidth;
+		InFont->GetStringHeightAndWidth(TEXT(" "), Height, SpaceWidth);
+
+		int32 LWidth;
+		InFont->GetStringHeightAndWidth(TEXT("l"), Height, LWidth);
+
+		int32 WWidth;
+		InFont->GetStringHeightAndWidth(TEXT("W"), Height, WWidth);
+
+		if ((SpaceWidth == LWidth) && (SpaceWidth == WWidth))
+		{
+			return true;
+		}
+	}
+
+	return false;
+}
+
+void UAvaFontManagerSubsystem::SetupMetrics(UAvaFontObject* InFontObject)
+{
+	if (!InFontObject)
+	{
+		return;
+	}
+
+	if (const UFont* const Font = InFontObject->GetFont())
+	{
+		const FCompositeFont* const CompositeFont = Font->GetCompositeFont();
+		if (!CompositeFont || CompositeFont->DefaultTypeface.Fonts.Num() == 0)
+		{
+			return;
+		}
+
+		bool bItalic = false;
+		bool bBold = false;
+
+		for (const FTypefaceEntry& TypefaceEntry : CompositeFont->DefaultTypeface.Fonts)
+		{
+			const FFontFaceDataConstPtr FaceData = TypefaceEntry.Font.GetFontFaceData();
+
+			FT_Face FreeTypeFace;
+
+			if (FaceData.IsValid() && FaceData->HasData() && FaceData->GetData().Num() > 0)
+			{
+				const TArray<uint8> Data = FaceData->GetData();
+
+				FT_New_Memory_Face(GetFreeTypeLibrary(), Data.GetData(), Data.Num(), 0, &FreeTypeFace);
+
+				if (FreeTypeFace)
+				{
+					bItalic = FreeTypeFace->style_flags & FT_STYLE_FLAG_ITALIC;
+					bBold = FreeTypeFace->style_flags & FT_STYLE_FLAG_BOLD;
+				}
+			}
+		}
+
+		FAvaSystemFontMetrics Metrics;
+		Metrics.bIsMonospaced = IsFontMonospaced(Font);
+		Metrics.bIsItalic = bItalic;
+		Metrics.bIsBold = bBold;
+
+		InFontObject->SetMetrics(Metrics);
+	}
+}
+
+FT_Library UAvaFontManagerSubsystem::GetFreeTypeLibrary()
+{
+	InitFreeTypeLibrary();
+	return FreeTypeLib;
+}
+
+void UAvaFontManagerSubsystem::InitFreeTypeLibrary()
+{
+	if (!FreeTypeLib)
+	{
+		FT_Init_FreeType(&FreeTypeLib);
+	}
+}
+
+void UAvaFontManagerSubsystem::CleanFreeTypeLibrary()
+{
+	if (FreeTypeLib)
+	{
+		FT_Done_FreeType(FreeTypeLib);
+		FreeTypeLib = nullptr;
+	}
+}
+
+//////// UAvalancheFontConfig ////////
+
+bool UAvaFontConfig::IsFavoriteFont(const FString& InFontName) const
+{
+	return FavoriteFonts.Contains(InFontName);
+}
+
+bool UAvaFontConfig::RemoveFavoriteFont(const FString& InFontName)
+{
+	if (IsFavoriteFont(InFontName))
+	{
+		FavoriteFonts.Remove(InFontName);
+		SaveConfig();
+		return true;
+	}
+
+	return false;
+}
+
+bool UAvaFontConfig::AddFavoriteFont(const FString& InFontName)
+{
+	if (!IsFavoriteFont(InFontName))
+	{
+		FavoriteFonts.AddUnique(InFontName);
+		SaveConfig();
+		return true;
+	}
+
+	return false;
+}
+
+void UAvaFontConfig::ToggleShowMonospaced()
+{
+	bShowOnlyMonospaced = !bShowOnlyMonospaced;
+	SaveConfig();
+}
+
+void UAvaFontConfig::ToggleShowBold()
+{
+	bShowOnlyBold = !bShowOnlyBold;
+	SaveConfig();
+}
+
+void UAvaFontConfig::ToggleShowItalic()
+{
+	bShowOnlyItalic = !bShowOnlyItalic;
+	SaveConfig();
+}
+
+#undef LOCTEXT_NAMESPACE
diff --git a/Engine/Plugins/Experimental/Avalanche/Source/AvalancheTextEditor/Private/Font/AvaFontManagerSubsystem.h b/Engine/Plugins/Experimental/Avalanche/Source/AvalancheTextEditor/Private/Font/AvaFontManagerSubsystem.h
new file mode 100644
--- /dev/null
+++ b/Engine/Plugins/Experimental/Avalanche/Source/AvalancheTextEditor/Private/Font/AvaFontManagerSubsystem.h
@@ -0,0 +1,252 @@
+// Copyright Epic Games, Inc. All Rights Reserved.
+
+#pragma once
+
+#include "EditorSubsystem.h"
+#include "Engine/Font.h"
+#include "Engine/FontFace.h"
+#include "Font/AvaFontObject.h"
+#include "Misc/CoreMisc.h"
+#include "PropertyEditorModule.h"
+
+THIRD_PARTY_INCLUDES_START
+#include "ft2build.h"
+#include FT_FREETYPE_H
+THIRD_PARTY_INCLUDES_END
+
+#include "AvaFontManagerSubsystem.generated.h"
+
+class FAvaFontView;
+class IPropertyHandle;
+class UAvaFontObject;
+class UFactory;
+struct FAvaFont;
+
+DECLARE_LOG_CATEGORY_EXTERN(LogAvaFont, Log, All);
+
+/** Facilitator struct to handle multiple selection */
+struct FAvaMultiFontSelectionData
+{
+	/** true if multiple objects are selected */
+	bool bMultipleObjectsSelected;
+
+	/** true if selection has multiple values - implies bAreMultipleItemsSelected is true */
+	bool bMultipleValues;
+
+	/** the name of the first selected font. Useful when all items reference the same font */
+	FString FirstSelectionFontName;
+};
+
+UCLASS(Config=Editor)
+class UAvaFontConfig : public UObject
+{
+	GENERATED_BODY()
+
+public:
+	UPROPERTY(Config, BlueprintReadOnly, Category = "Motion Design")
+	TArray<FString> FavoriteFonts;
+
+	UPROPERTY(Config, BlueprintReadOnly, Category = "Motion Design")
+	bool bShowOnlyMonospaced;
+
+	UPROPERTY(Config, BlueprintReadOnly, Category = "Motion Design")
+	bool bShowOnlyBold;
+
+	UPROPERTY(Config, BlueprintReadOnly, Category = "Motion Design")
+	bool bShowOnlyItalic;
+
+	bool IsFavoriteFont(const FString& InFontName) const;
+	bool RemoveFavoriteFont(const FString& InFontName);
+	bool AddFavoriteFont(const FString& InFontName);
+
+	void ToggleShowMonospaced();
+	void ToggleShowBold();
+	void ToggleShowItalic();
+};
+
+DECLARE_MULTICAST_DELEGATE_OneParam(FOnProjectFontsChange, const UFont*)
+DECLARE_MULTICAST_DELEGATE(FOnSystemFontsUpdated)
+
+UCLASS()
+class UAvaFontManagerSubsystem : public UEditorSubsystem, public FSelfRegisteringExec
+{
+	friend class FAvaTextEditorModule;
+
+	GENERATED_BODY()
+public:
+	/**
+	 * Gets the singleton instance of the fonts manager.
+	 * @return The FontImportManager instance.
+	 */
+	static UAvaFontManagerSubsystem* Get();
+
+	/** Is this font .ttf or .otf? */
+	static bool IsSupportedFontFile(const FString& InFontFilePath);
+
+	/**
+	 * Retrieve information about multiple selection
+	 * @param InPropertyHandle The property handle to check
+	 * @param OutMultiFontSelectionData Will store selection data
+	 */
+	static void GetMultipleSelectionInformation(const TSharedPtr<IPropertyHandle>& InPropertyHandle, FAvaMultiFontSelectionData& OutMultiFontSelectionData);
+
+	/**
+	 * Returns a raw pointer to the FAvaFont property handled by the specified property handle (if valid).
+	 * @param InFontPropertyHandle A property handle dealing with a property of type FAvaFont
+	 * @param OutAccessResult Result of the access operation
+	 * @return The handled FAvaFont, as a pointer.
+	 */
+	static FAvaFont* GetFontFromPropertyHandle(const TSharedPtr<IPropertyHandle>& InFontPropertyHandle, FPropertyAccess::Result& OutAccessResult);
+
+	static void SanitizeString(FString& OutSanitizedName);
+
+	// ~Begin UEditorSubsystem
+	virtual void Initialize(FSubsystemCollectionBase& Collection) override;
+	virtual void Deinitialize() override;
+	// ~End UEditorSubsystem
+
+	// ~Begin FSelfRegisteringExec
+	virtual bool Exec(class UWorld* InWorld, const TCHAR* InCmd, FOutputDevice& InAr) override;
+	// ~End FSelfRegisteringExec
+
+	void MarkFontForAutoImport(const TSharedPtr<IPropertyHandle>& InFontToImportOnSave);
+	void AddFavorite(const FString& InFontName);
+	void RemoveFavorite(const FString& InFontName);
+
+	UAvaFontConfig* GetFontManagerConfig();
+
+	TConstArrayView<const TSharedPtr<FAvaFontView>> GetFontOptions();
+
+	int32 GetDefaultFontIndex() const { return DefaultFontIndex; }
+
+	/** Given a Property Handle, tries to retrieve and return a font view from AvalancheFontsOptions */
+	TSharedPtr<FAvaFontView> GetFontViewFromPropertyHandle(const TSharedPtr<IPropertyHandle>& InFontPropertyHandle, FPropertyAccess::Result& OutAccessResult);
+
+	FOnProjectFontsChange& OnProjectFontCreated() { return OnProjectFontCreatedDelegate; }
+	FOnProjectFontsChange& OnProjectFontDeleted() { return OnProjectFontDeletedDelegate; }
+	FOnSystemFontsUpdated& OnSystemFontsUpdated() { return OnSystemFontsUpdatedDelegate; }
+	/**
+	 * Import the font contained in the UAvalancheFontObject as .uasset, if yet to be imported
+	 * @param InAvaFontPropertyHandle a property handle to the FAvalancheFont referencing the UFont to be saved
+	 * @return success
+	 */
+	bool ImportFont(const TSharedPtr<IPropertyHandle>& InAvaFontPropertyHandle);
+
+	void RefreshSystemFont(UFont* InFont) const;
+
+	bool IsFontAvailableOnLocalOS(const UFont* InFont) const;
+
+private:
+	static void SetupFontFamilyTypefaces(UFont* InFont, const FSystemFontsRetrieveParams& InFontParams);
+
+	/**
+	 * Initializes OS fonts and Avalanche fonts data structures
+	 * Since this is called before project assets are available, it will NOT initialize project fonts.
+	 * That is done later by InitializeProjectFonts function.
+	 */
+	void Initialize();
+
+	/**
+	 * Adds project fonts (uassets) to the list of fonts handled by the font manager
+	 * It will also handle duplicates, e.g. system fonts which have already been imported
+	 * Note: make sure to call this when AssetRegistry module is ready.
+	 */
+	void InitializeProjectFonts();
+
+	bool IsFavoriteFont(const FString& InFontName);
+	bool ImportAsAsset(UAvaFontObject* InFontToImport);
+	bool RemoveProjectFont(const UFont* InSourceFont);
+
+	UAvaFontObject* CreateProjectFont(UFont* InSourceFont, UPackage* InProjectFontsPackage = nullptr);
+
+	/** Given a font name, tries to retrieve and return a font from AvalancheFontsOptions */
+	TSharedPtr<FAvaFontView> GetFontViewFromName(const FString& InName);
+
+	/** Given a Property Handle, tries to retrieve and return a font from AvalancheFontsOptions as a raw pointer */
+	TSharedPtr<FAvaFontView> GetFontViewFromPropertyHandle(const TSharedPtr<IPropertyHandle>& InFontPropertyHandle);
+
+	void ClearFontsData();
+	void LoadFavorites();
+	void OnAssetCreated(UObject* InObject);
+	void OnAssetDeleted(UObject* InObject);
+	void LoadOSFonts();
+	void LoadProjectFonts();
+	void RemoveProjectFontsFromOSFonts();
+	void RefreshAvalancheFontsMap();
+	void RefreshAvalancheFontsOptions();
+	void SortFontsMap();
+	void ImportSystemFontFamily(const FSystemFontsRetrieveParams& InFontParams, UPackage* InImportPackage);
+	void UnmarkFontFromAutoImport(const TSharedPtr<IPropertyHandle>& InFontToUnregister);
+	void ForceFontAutoImportUnmark(const TSharedPtr<IPropertyHandle>& InFontToUnregister);
+	void FontsAutoImport();
+	void FontsAutoImportPreSave(UWorld* InWorld, FObjectPreSaveContext InObjectPreSaveContext);
+	void RegisterDefaultAutoImportCallbacks();
+	void UnregisterDefaultAutoImportCallbacks() const;
+	void RefreshAllSystemFonts() const;
+
+	static void GetFontName(const UFont* InFont, FString& OutFontName);
+	static void GetSanitizedFontName(const UFont* InFont, FString& OutFontName);
+	static void SetupMetrics(UAvaFontObject* InFontObject);
+
+	static bool IsFontMonospaced(const UFont* const InFont);
+	static FString GetImportFontPackageNameRoot() { return TEXT("/Game/SystemFonts"); }
+	static FString GetTempFontPackageNameRoot() { return TEXT("/Temp/SystemFonts"); }
+	static FString GetFontObjPackageNameRoot() { return TEXT("/Temp/SystemFonts/FontObjs"); }
+	static FString GetFontConfigPackageName() { return TEXT("/Temp/AvalancheEditor/AvalancheFontConfig"); }
+
+	static void InitFreeTypeLibrary();
+	static void CleanFreeTypeLibrary();
+	static FT_Library GetFreeTypeLibrary();
+
+	static const TCHAR* GetFallbackFontFaceReference() { return FallbackFontReference; }
+
+	/** Used internally to load system fonts for the first time */
+	void InitializeOSFonts();
+
+	/** updates the entire fonts list by combining OS and Project fonts */
+	void CombineAllFonts();
+
+	void RegisterAssetsCallbacks();
+	void UnregisterAssetsCallbacks() const;
+
+	UFontFace* GetFallbackFontFace();
+
+	TArray<FString> GetSystemFontNames() const;
+	TArray<FString> GetSystemFontTypefaces(const FString& InFontFamilyName) const;
+	inline static FT_Library FreeTypeLib;
+
+	UPROPERTY(Transient)
+	TObjectPtr<UFontFace> DefaultFallbackFontFace;
+
+	/** used for custom config - e.g. favorite fonts */
+	UPROPERTY(Transient)
+	TObjectPtr<UAvaFontConfig> AvaFontManagerConfig;
+
+	/** Fonts available from the OS */
+	UPROPERTY(Transient)
+	TMap<FString, TObjectPtr<UAvaFontObject>> OSFontsMap;
+
+	/** Fonts available from current project */
+	UPROPERTY(Transient)
+	TMap<FString, TObjectPtr<UAvaFontObject>> ProjectFontsMap;
+
+	/** OS + current project fonts combined */
+	UPROPERTY(Transient)
+	TMap<FString, TObjectPtr<UAvaFontObject>> AvalancheFontsMap;
+
+	/** Caching font options here, so that details panel can get them right away */
+	TArray<TSharedPtr<FAvaFontView>> AvalancheFontsOptions;
+
+	/** Used to allow auto font import on ava bp and level editor save  */
+	TMap<TSharedPtr<FAvaFontView>, TSharedPtr<IPropertyHandle>> FontsToImportOnSave;
+
+	TMap<FString, FSystemFontsRetrieveParams> FontsInfoMap;
+
+	int32 DefaultFontIndex = -1;
+
+	FOnProjectFontsChange OnProjectFontCreatedDelegate;
+	FOnProjectFontsChange OnProjectFontDeletedDelegate;
+	FOnSystemFontsUpdated OnSystemFontsUpdatedDelegate;
+
+	static constexpr TCHAR FallbackFontReference[] = TEXT("/Engine/EngineFonts/Faces/DroidSansFallback.DroidSansFallback");
+};
diff --git a/Engine/Plugins/Experimental/Avalanche/Source/AvalancheTextEditor/Private/Font/AvaFontSelectorCommands.cpp b/Engine/Plugins/Experimental/Avalanche/Source/AvalancheTextEditor/Private/Font/AvaFontSelectorCommands.cpp
new file mode 100644
--- /dev/null
+++ b/Engine/Plugins/Experimental/Avalanche/Source/AvalancheTextEditor/Private/Font/AvaFontSelectorCommands.cpp
@@ -0,0 +1,28 @@
+// Copyright Epic Games, Inc. All Rights Reserved.
+
+#include "AvaFontSelectorCommands.h"
+
+#define LOCTEXT_NAMESPACE "AvaFontFilteringCommands"
+
+void FAvaFontSelectorCommands::RegisterCommands()
+{
+	UI_COMMAND(ShowMonospacedFonts
+		, "Monospaced"
+		, "Filter monospaced fonts"
+		, EUserInterfaceActionType::ToggleButton
+		, FInputChord());
+
+	UI_COMMAND(ShowBoldFonts
+		, "Bold"
+		, "Filter fonts with bold typefaces available"
+		, EUserInterfaceActionType::ToggleButton
+		, FInputChord());
+
+	UI_COMMAND(ShowItalicFonts
+		, "Italic"
+		, "Filter fonts with italic typefaces available"
+		, EUserInterfaceActionType::ToggleButton
+		, FInputChord());
+}
+
+#undef LOCTEXT_NAMESPACE
diff --git a/Engine/Plugins/Experimental/Avalanche/Source/AvalancheTextEditor/Private/Font/AvaFontSelectorCommands.h b/Engine/Plugins/Experimental/Avalanche/Source/AvalancheTextEditor/Private/Font/AvaFontSelectorCommands.h
new file mode 100644
--- /dev/null
+++ b/Engine/Plugins/Experimental/Avalanche/Source/AvalancheTextEditor/Private/Font/AvaFontSelectorCommands.h
@@ -0,0 +1,25 @@
+// Copyright Epic Games, Inc. All Rights Reserved.
+
+#pragma once
+
+#include "Framework/Commands/Commands.h"
+#include "Styling/AppStyle.h"
+
+class FAvaFontSelectorCommands
+	: public TCommands<FAvaFontSelectorCommands>
+{
+public:
+	FAvaFontSelectorCommands()
+		: TCommands<FAvaFontSelectorCommands>(TEXT("AvalancheFontSelector")
+		, NSLOCTEXT("AvaFontSelectorCommands", "AvalancheFontSelector", "Motion Design Font Selector")
+		, NAME_None
+		, FAppStyle::GetAppStyleSetName())
+	{
+	}
+
+	virtual void RegisterCommands() override;
+
+	TSharedPtr<FUICommandInfo> ShowMonospacedFonts;
+	TSharedPtr<FUICommandInfo> ShowBoldFonts;
+	TSharedPtr<FUICommandInfo> ShowItalicFonts;
+};
\ No newline at end of file
diff --git a/Engine/Plugins/Experimental/Avalanche/Source/AvalancheTextEditor/Private/Font/AvaFontView.cpp b/Engine/Plugins/Experimental/Avalanche/Source/AvalancheTextEditor/Private/Font/AvaFontView.cpp
new file mode 100644
--- /dev/null
+++ b/Engine/Plugins/Experimental/Avalanche/Source/AvalancheTextEditor/Private/Font/AvaFontView.cpp
@@ -0,0 +1,121 @@
+﻿// Copyright Epic Games, Inc. All Rights Reserved.
+
+#include "AvaFontView.h"
+#include "Font/AvaFont.h"
+
+TSharedRef<FAvaFontView> FAvaFontView::Make(const FAvaFont& InFont)
+{
+	return MakeShareable(new FAvaFontView(InFont));
+}
+
+bool FAvaFontView::IsAvaFontValid() const
+{
+	return AvaFontObject.IsValid();
+}
+
+FAvaFontView::FAvaFontView(const FAvaFont& InFont)
+{
+	SetFavorite(InFont.IsFavorite());
+	SetFontObject(InFont.GetFontObject());
+
+	bIsFallbackFont = InFont.IsFallbackFont();
+	if (bIsFallbackFont)
+	{
+		MissingFontName = InFont.GetMissingFontName();
+	}
+}
+
+UFont* FAvaFontView::GetFont() const
+{
+	if (!IsAvaFontValid())
+	{
+		return nullptr;
+	}
+
+	return AvaFontObject->GetFont();
+}
+
+bool FAvaFontView::HasValidFont() const
+{
+	if (!IsAvaFontValid())
+	{
+		return false;
+	}
+
+	return AvaFontObject->HasValidFont();
+}
+
+bool FAvaFontView::IsMonospaced() const
+{
+	if (!IsAvaFontValid())
+	{
+		return false;
+	}
+
+	return AvaFontObject->IsMonospaced();
+}
+
+bool FAvaFontView::IsBold() const
+{
+	if (!IsAvaFontValid())
+	{
+		return false;
+	}
+
+	return AvaFontObject->IsBold();
+}
+
+bool FAvaFontView::IsItalic() const
+{
+	if (!IsAvaFontValid())
+	{
+		return false;
+	}
+
+	return AvaFontObject->IsItalic();
+}
+
+FName FAvaFontView::GetFontName() const
+{
+	return FName(GetFontNameAsString());
+}
+
+FString FAvaFontView::GetFontNameAsString() const
+{
+	if (!IsAvaFontValid())
+	{
+		return TEXT("");
+	}
+
+	return AvaFontObject->GetFontName();
+}
+
+FText FAvaFontView::GetFontNameAsText() const
+{
+	return FText::FromString(GetFontNameAsString());
+}
+
+EAvaFontSource FAvaFontView::GetFontSource() const
+{
+	if (IsAvaFontValid())
+	{
+		return AvaFontObject->GetSource();
+	}
+
+	return EAvaFontSource::Invalid;
+}
+
+void FAvaFontView::SetFavorite(bool bInIsFavorite)
+{
+	bIsFavorite = bInIsFavorite;
+}
+
+void FAvaFontView::SetFontObject(UAvaFontObject* InFontObject)
+{
+	if (!InFontObject)
+	{
+		return;
+	}
+
+	AvaFontObject = TStrongObjectPtr<UAvaFontObject>(InFontObject);
+}
diff --git a/Engine/Plugins/Experimental/Avalanche/Source/AvalancheTextEditor/Private/Font/AvaFontView.h b/Engine/Plugins/Experimental/Avalanche/Source/AvalancheTextEditor/Private/Font/AvaFontView.h
new file mode 100644
--- /dev/null
+++ b/Engine/Plugins/Experimental/Avalanche/Source/AvalancheTextEditor/Private/Font/AvaFontView.h
@@ -0,0 +1,77 @@
+﻿// Copyright Epic Games, Inc. All Rights Reserved.
+
+#pragma once
+
+#include "Font/AvaFontObject.h"
+#include "UObject/StrongObjectPtr.h"
+
+class UAvaFontObject;
+struct FAvaFont;
+
+class FAvaFontView
+{
+	friend class SAvaFontField;
+
+public:
+	static TSharedRef<FAvaFontView> Make(const FAvaFont& InFont);
+
+	UAvaFontObject* GetFontObject() const { return AvaFontObject.Get(); }
+	UFont* GetFont() const;
+
+	bool IsFallbackFont() const { return bIsFallbackFont; }
+	bool IsFavorite() const { return bIsFavorite; }
+
+	bool HasValidFont() const;
+	bool IsMonospaced() const;
+	bool IsBold() const;
+	bool IsItalic() const;
+	FName GetFontName() const;
+	FString GetFontNameAsString() const;
+	FText GetFontNameAsText() const;
+	EAvaFontSource GetFontSource() const;
+	const FString& GetMissingFontName() const { return MissingFontName; }
+
+	bool operator==(const FAvaFontView& Other) const
+	{
+		const UAvaFontObject* MyFontObject = GetFontObject();
+		const UAvaFontObject* OtherFontObject = Other.GetFontObject();
+
+		if (MyFontObject && OtherFontObject && MyFontObject->GetFont() != OtherFontObject->GetFont())
+		{
+			return false;
+		}
+
+		return (GetFontName() == Other.GetFontName());
+	}
+
+	bool operator!=(const FAvaFontView& Other) const
+	{
+		return !(*this == Other);
+	}
+
+	/** Alphabetical ordering */
+	bool operator<(const FAvaFontView& Other) const
+	{
+		return GetFontName().LexicalLess(Other.GetFontName());
+	}
+
+	bool IsAvaFontValid() const;
+
+protected:
+	FAvaFontView()
+		: bIsFavorite(false)
+		, bIsFallbackFont(false)
+		, AvaFontObject(nullptr)
+	{
+	}
+
+	FAvaFontView(const FAvaFont& InFont);
+
+	void SetFontObject(UAvaFontObject* InFontObject);
+	void SetFavorite(bool bInIsFavorite);
+
+	bool bIsFavorite;
+	bool bIsFallbackFont;
+	TStrongObjectPtr<UAvaFontObject> AvaFontObject;
+	FString MissingFontName;
+};
diff --git a/Engine/Plugins/Experimental/Avalanche/Source/AvalancheTextEditor/Private/Platforms/Linux/LinuxSystemFontLoading.cpp b/Engine/Plugins/Experimental/Avalanche/Source/AvalancheTextEditor/Private/Platforms/Linux/LinuxSystemFontLoading.cpp
new file mode 100644
--- /dev/null
+++ b/Engine/Plugins/Experimental/Avalanche/Source/AvalancheTextEditor/Private/Platforms/Linux/LinuxSystemFontLoading.cpp
@@ -0,0 +1,14 @@
+// Copyright Epic Games, Inc. All Rights Reserved.
+
+#include "LinuxSystemFontLoading.h"
+
+namespace UE::Avalanche::Private::Fonts
+{
+	void GetSystemFontInfo(TMap<FString, FSystemFontsRetrieveParams>& OutFontsInfo)
+	{
+	}
+
+	void ListAvailableFontFiles()
+	{
+	}
+}
diff --git a/Engine/Plugins/Experimental/Avalanche/Source/AvalancheTextEditor/Private/Platforms/Linux/LinuxSystemFontLoading.h b/Engine/Plugins/Experimental/Avalanche/Source/AvalancheTextEditor/Private/Platforms/Linux/LinuxSystemFontLoading.h
new file mode 100644
--- /dev/null
+++ b/Engine/Plugins/Experimental/Avalanche/Source/AvalancheTextEditor/Private/Platforms/Linux/LinuxSystemFontLoading.h
@@ -0,0 +1,16 @@
+// Copyright Epic Games, Inc. All Rights Reserved.
+
+#pragma once
+
+#include "Containers/Array.h"
+#include "Containers/UnrealString.h"
+
+struct FSystemFontsRetrieveParams;
+
+//todo: yet to be implemented
+namespace UE::Avalanche::Private::Fonts
+{
+	void GetSystemFontInfo(TMap<FString, FSystemFontsRetrieveParams>& OutFontsInfo);
+
+	void ListAvailableFontFiles();
+}
diff --git a/Engine/Plugins/Experimental/Avalanche/Source/AvalancheTextEditor/Private/Platforms/Mac/MacSystemFontLoading.cpp b/Engine/Plugins/Experimental/Avalanche/Source/AvalancheTextEditor/Private/Platforms/Mac/MacSystemFontLoading.cpp
new file mode 100644
--- /dev/null
+++ b/Engine/Plugins/Experimental/Avalanche/Source/AvalancheTextEditor/Private/Platforms/Mac/MacSystemFontLoading.cpp
@@ -0,0 +1,14 @@
+// Copyright Epic Games, Inc. All Rights Reserved.
+
+#include "MacSystemFontLoading.h"
+
+namespace UE::Avalanche::Private::Fonts
+{
+	void GetSystemFontInfo(TMap<FString, FSystemFontsRetrieveParams>& OutFontsInfo)
+	{
+	}
+
+	void ListAvailableFontFiles()
+	{
+	}
+}
diff --git a/Engine/Plugins/Experimental/Avalanche/Source/AvalancheTextEditor/Private/Platforms/Mac/MacSystemFontLoading.h b/Engine/Plugins/Experimental/Avalanche/Source/AvalancheTextEditor/Private/Platforms/Mac/MacSystemFontLoading.h
new file mode 100644
--- /dev/null
+++ b/Engine/Plugins/Experimental/Avalanche/Source/AvalancheTextEditor/Private/Platforms/Mac/MacSystemFontLoading.h
@@ -0,0 +1,16 @@
+// Copyright Epic Games, Inc. All Rights Reserved.
+
+#pragma once
+
+#include "Containers/Array.h"
+#include "Containers/UnrealString.h"
+
+struct FSystemFontsRetrieveParams;
+
+//todo: yet to be implemented
+namespace UE::Avalanche::Private::Fonts
+{
+	void GetSystemFontInfo(TMap<FString, FSystemFontsRetrieveParams>& OutFontsInfo);
+
+	void ListAvailableFontFiles();
+}
diff --git a/Engine/Plugins/Experimental/Avalanche/Source/AvalancheTextEditor/Private/Platforms/Windows/WindowsSystemFontLoading.h b/Engine/Plugins/Experimental/Avalanche/Source/AvalancheTextEditor/Private/Platforms/Windows/WindowsSystemFontLoading.h
new file mode 100644
--- /dev/null
+++ b/Engine/Plugins/Experimental/Avalanche/Source/AvalancheTextEditor/Private/Platforms/Windows/WindowsSystemFontLoading.h
@@ -0,0 +1,15 @@
+// Copyright Epic Games, Inc. All Rights Reserved.
+
+#pragma once
+
+#include "Containers/ContainersFwd.h"
+#include "Containers/UnrealString.h"
+
+struct FSystemFontsRetrieveParams;
+
+namespace UE::Avalanche::Private::Fonts
+{
+	void GetSystemFontInfo(TMap<FString, FSystemFontsRetrieveParams>& OutFontsInfo);
+
+	void ListAvailableFontFiles();
+}
diff --git a/Engine/Plugins/Experimental/Avalanche/Source/AvalancheTextEditor/Private/Slate/AvaFontSearchSettingsMenu.cpp b/Engine/Plugins/Experimental/Avalanche/Source/AvalancheTextEditor/Private/Slate/AvaFontSearchSettingsMenu.cpp
new file mode 100644
--- /dev/null
+++ b/Engine/Plugins/Experimental/Avalanche/Source/AvalancheTextEditor/Private/Slate/AvaFontSearchSettingsMenu.cpp
@@ -0,0 +1,153 @@
+// Copyright Epic Games, Inc. All Rights Reserved.
+
+#include "AvaFontSearchSettingsMenu.h"
+#include "Font/AvaFontManagerSubsystem.h"
+#include "Font/AvaFontSelectorCommands.h"
+#include "Framework/MultiBox/MultiBoxBuilder.h"
+#include "SAvaFontSelector.h"
+#include "SlateOptMacros.h"
+#include "Widgets/Images/SImage.h"
+
+#define LOCTEXT_NAMESPACE "AvaFontSearchSettingsMenu"
+
+BEGIN_SLATE_FUNCTION_BUILD_OPTIMIZATION
+void SAvaFontSearchSettingsMenu::Construct(const FArguments& InArgs)
+{
+	FontSelector = InArgs._AvaFontSelector;
+
+	check(FontSelector.IsValid());
+
+	UAvaFontManagerSubsystem* FontManagerSubsystem = UAvaFontManagerSubsystem::Get();
+	if (!FontManagerSubsystem)
+	{
+		return;
+	}
+
+	if (const UAvaFontConfig* FontManagerConfig = FontManagerSubsystem->GetFontManagerConfig())
+	{
+		bIsShowMonospacedActive = FontManagerConfig->bShowOnlyMonospaced;
+		bIsShowBoldActive = FontManagerConfig->bShowOnlyBold;
+		bIsShowItalicActive = FontManagerConfig->bShowOnlyItalic;
+	}
+
+	BindCommands();
+
+	ChildSlot
+	[
+		SNew(SComboButton)
+		.CollapseMenuOnParentFocus(true)
+		.ComboButtonStyle(&FAppStyle::Get().GetWidgetStyle<FComboButtonStyle>("SimpleComboButton"))
+		.OnGetMenuContent(this, &SAvaFontSearchSettingsMenu::CreateSettingsMenu)
+		.HasDownArrow(false)
+		.ButtonContent()
+		[
+			SNew(SHorizontalBox)
+			+ SHorizontalBox::Slot()
+			.AutoWidth()
+			.VAlign(VAlign_Center)
+			.Padding(4.0f, 0.0f)
+			[
+				SNew(SImage)
+				.ColorAndOpacity(FSlateColor::UseForeground())
+				.Image(FAppStyle::Get().GetBrush("Icons.Settings"))
+			]
+		]
+	];
+}
+END_SLATE_FUNCTION_BUILD_OPTIMIZATION
+
+TSharedRef<SWidget> SAvaFontSearchSettingsMenu::CreateSettingsMenu() const
+{
+	const FAvaFontSelectorCommands& FontSelectorCommands = FAvaFontSelectorCommands::Get();
+
+	constexpr bool bCloseAfterSelection = false;
+	FMenuBuilder MenuBuilder(bCloseAfterSelection, CommandList);
+	MenuBuilder.BeginSection(TEXT("FontsFilteringSettings"), LOCTEXT("FontsFiltering", "Fonts Filtering"));
+	{
+		MenuBuilder.AddMenuEntry(FontSelectorCommands.ShowMonospacedFonts);
+		MenuBuilder.AddMenuEntry(FontSelectorCommands.ShowBoldFonts);
+		MenuBuilder.AddMenuEntry(FontSelectorCommands.ShowItalicFonts);
+	}
+	MenuBuilder.EndSection();
+
+	return MenuBuilder.MakeWidget();
+}
+
+void SAvaFontSearchSettingsMenu::BindCommands()
+{
+	FAvaFontSelectorCommands::Register();
+	const FAvaFontSelectorCommands& FontSelectorCommands = FAvaFontSelectorCommands::Get();
+
+	CommandList = MakeShared<FUICommandList>();
+
+	// Monospaced filter
+	CommandList->MapAction(FontSelectorCommands.ShowMonospacedFonts
+		, FExecuteAction::CreateSP(this, &SAvaFontSearchSettingsMenu::ShowMonospacedToggle_Execute)
+		, FCanExecuteAction::CreateSP(this, &SAvaFontSearchSettingsMenu::ShowMonospacedToggle_CanExecute)
+		, FIsActionChecked::CreateSP(this, &SAvaFontSearchSettingsMenu::ShowMonospacedToggle_IsChecked));
+
+	// Bold filter
+	CommandList->MapAction(FontSelectorCommands.ShowBoldFonts
+		, FExecuteAction::CreateSP(this, &SAvaFontSearchSettingsMenu::ShowBoldToggle_Execute)
+		, FCanExecuteAction::CreateSP(this, &SAvaFontSearchSettingsMenu::ShowBoldToggle_CanExecute)
+		, FIsActionChecked::CreateSP(this, &SAvaFontSearchSettingsMenu::ShowBoldToggle_IsChecked));
+
+	// Italic filter
+	CommandList->MapAction(FontSelectorCommands.ShowItalicFonts
+		, FExecuteAction::CreateSP(this, &SAvaFontSearchSettingsMenu::ShowItalicToggle_Execute)
+		, FCanExecuteAction::CreateSP(this, &SAvaFontSearchSettingsMenu::ShowItalicToggle_CanExecute)
+		, FIsActionChecked::CreateSP(this, &SAvaFontSearchSettingsMenu::ShowItalicToggle_IsChecked));
+}
+
+void SAvaFontSearchSettingsMenu::ShowMonospacedToggle_Execute()
+{
+	UAvaFontManagerSubsystem* FontManagerSubsystem = UAvaFontManagerSubsystem::Get();
+	if (!FontManagerSubsystem)
+	{
+		return;
+	}
+
+	bIsShowMonospacedActive = !bIsShowMonospacedActive;
+	FontManagerSubsystem->GetFontManagerConfig()->ToggleShowMonospaced();
+
+	if (FontSelector)
+	{
+		FontSelector->SetShowMonospacedOnly(bIsShowMonospacedActive);
+	}
+}
+
+void SAvaFontSearchSettingsMenu::ShowBoldToggle_Execute()
+{
+	UAvaFontManagerSubsystem* FontManagerSubsystem = UAvaFontManagerSubsystem::Get();
+	if (!FontManagerSubsystem)
+	{
+		return;
+	}
+
+	bIsShowBoldActive = !bIsShowBoldActive;
+	FontManagerSubsystem->GetFontManagerConfig()->ToggleShowBold();
+
+	if (FontSelector)
+	{
+		FontSelector->SetShowBoldOnly(bIsShowBoldActive);
+	}
+}
+
+void SAvaFontSearchSettingsMenu::ShowItalicToggle_Execute()
+{
+	UAvaFontManagerSubsystem* FontManagerSubsystem = UAvaFontManagerSubsystem::Get();
+	if (!FontManagerSubsystem)
+	{
+		return;
+	}
+
+	bIsShowItalicActive = !bIsShowItalicActive;
+	FontManagerSubsystem->GetFontManagerConfig()->ToggleShowItalic();
+
+	if (FontSelector)
+	{
+		FontSelector->SetShowItalicOnly(bIsShowItalicActive);
+	}
+}
+
+#undef LOCTEXT_NAMESPACE
diff --git a/Engine/Plugins/Experimental/Avalanche/Source/AvalancheTextEditor/Private/Slate/AvaFontSearchSettingsMenu.h b/Engine/Plugins/Experimental/Avalanche/Source/AvalancheTextEditor/Private/Slate/AvaFontSearchSettingsMenu.h
new file mode 100644
--- /dev/null
+++ b/Engine/Plugins/Experimental/Avalanche/Source/AvalancheTextEditor/Private/Slate/AvaFontSearchSettingsMenu.h
@@ -0,0 +1,49 @@
+// Copyright Epic Games, Inc. All Rights Reserved.
+
+#pragma once
+
+#include "Widgets/SCompoundWidget.h"
+
+class FUICommandList;
+class SAvaFontSelector;
+
+class SAvaFontSearchSettingsMenu : public SCompoundWidget
+{
+public:
+	SLATE_BEGIN_ARGS(SAvaFontSearchSettingsMenu)
+		: _AvaFontSelector()
+		{
+		}
+
+		SLATE_ARGUMENT(TSharedPtr<SAvaFontSelector>, AvaFontSelector)
+	SLATE_END_ARGS()
+
+	void Construct(const FArguments& InArgs);
+
+protected:
+	TSharedRef<SWidget> CreateSettingsMenu() const;
+
+	void BindCommands();
+
+	void ShowMonospacedToggle_Execute();
+	void ShowBoldToggle_Execute();
+	void ShowItalicToggle_Execute();
+
+	bool ShowMonospacedToggle_CanExecute() { return true; }
+	bool ShowBoldToggle_CanExecute() { return true; }
+	bool ShowItalicToggle_CanExecute() { return true; }
+
+	bool ShowMonospacedToggle_IsChecked() const { return bIsShowMonospacedActive; }
+	bool ShowBoldToggle_IsChecked() const { return bIsShowBoldActive; }
+	bool ShowItalicToggle_IsChecked() const { return bIsShowItalicActive; }
+
+	TSharedPtr<FUICommandList> CommandList;
+
+	TSharedPtr<SAvaFontSelector> FontSelector;
+
+	bool bIsShowMonospacedActive = false;
+
+	bool bIsShowBoldActive = false;
+
+	bool bIsShowItalicActive = false;
+};
diff --git a/Engine/Plugins/Experimental/Avalanche/Source/AvalancheTextEditor/Private/Slate/SAvaFontField.cpp b/Engine/Plugins/Experimental/Avalanche/Source/AvalancheTextEditor/Private/Slate/SAvaFontField.cpp
new file mode 100644
--- /dev/null
+++ b/Engine/Plugins/Experimental/Avalanche/Source/AvalancheTextEditor/Private/Slate/SAvaFontField.cpp
@@ -0,0 +1,241 @@
+// Copyright Epic Games, Inc. All Rights Reserved.
+
+#include "SAvaFontField.h"
+#include "Engine/Font.h"
+#include "Font/AvaFontManagerSubsystem.h"
+#include "Font/AvaFontView.h"
+#include "SlateOptMacros.h"
+#include "Styling/StyleColors.h"
+#include "Widgets/Images/SImage.h"
+#include "Widgets/Input/SButton.h"
+#include "Widgets/Layout/SBox.h"
+#include "Widgets/Layout/SScaleBox.h"
+#include "Widgets/Text/STextBlock.h"
+
+namespace UE::AvalancheTextEditor::Private
+{
+	static constexpr float DefaultFontSize = 12.0f;
+}
+
+#define LOCTEXT_NAMESPACE "AvaFontField"
+
+BEGIN_SLATE_FUNCTION_BUILD_OPTIMIZATION
+
+void SAvaFontField::Construct(const FArguments& InArgs)
+{
+	AvalancheFont = InArgs._AvalancheFont;
+	OnAvaFontFieldModified = InArgs._OnAvaFontFieldModified;
+
+	const FVector2D Icon16x16(16.0f, 16.0f);
+
+	ChildSlot
+	[
+		SNew(SBox)
+		.MinDesiredWidth(400.0f)
+		.HAlign(HAlign_Fill)
+		.VAlign(VAlign_Fill)
+		.Clipping(EWidgetClipping::ClipToBoundsAlways) // we don't want list elements to render outside the box
+		.Content()
+		[
+			SNew(SHorizontalBox)
+			+ SHorizontalBox::Slot()
+			.AutoWidth()
+			[
+				SNew(SScaleBox)
+				.Stretch(EStretch::ScaleToFit)
+				.HAlign(HAlign_Fill)
+				.VAlign(VAlign_Center)
+				[
+					SNew(SButton)
+					.ButtonStyle( FAppStyle::Get(), "SimpleButton" )
+					.OnClicked(this, &SAvaFontField::OnToggleFavoriteClicked)
+					.ToolTipText(LOCTEXT("ToggleFavorite", "Mark Font as Favorite"))
+					.Visibility(this, &SAvaFontField::GetFavoriteVisibility)
+					[
+						SNew(SImage)
+						.ColorAndOpacity(this, &SAvaFontField::GetToggleFavoriteColor)
+						.Image(FAppStyle::GetBrush("Icons.Star"))
+						.DesiredSizeOverride(Icon16x16)
+					]
+				]
+			]
+			+ SHorizontalBox::Slot()
+			.VAlign(VAlign_Center)
+			.HAlign(HAlign_Fill)
+			[
+				SAssignNew(LeftFontNameText, STextBlock)
+				.Justification(ETextJustify::Left)
+			]
+			+ SHorizontalBox::Slot()
+			.VAlign(VAlign_Center)
+			.HAlign(HAlign_Right)
+			.AutoWidth()
+			[
+				SAssignNew(RightFontNameText, STextBlock)
+				.Justification(ETextJustify::Right)
+			]
+			+ SHorizontalBox::Slot()
+			.VAlign(VAlign_Center)
+			.HAlign(HAlign_Right)
+			.AutoWidth()
+			[
+				SNew(SBox)
+				.HAlign(HAlign_Center)
+				.VAlign(VAlign_Center)
+				.ToolTipText(LOCTEXT( "FontLocallyAvailable", "This font is available on the local system"))
+				.Visibility(this, &SAvaFontField::GetLocallyAvailableIconVisibility)
+				[
+					SNew(SImage)
+					.Image(FAppStyle::Get().GetBrush("Symbols.Check"))
+					.ColorAndOpacity(FSlateColor::UseForeground())
+				]
+			]
+		]
+	];
+
+	UpdateFont(AvalancheFont.Get());
+}
+END_SLATE_FUNCTION_BUILD_OPTIMIZATION
+
+void SAvaFontField::UpdateFont(const TSharedPtr<FAvaFontView>& InAvaFont)
+{
+	if (!InAvaFont.IsValid())
+	{
+		return;
+	}
+
+	AvalancheFont.Set(InAvaFont);
+
+	if (const TSharedPtr<FAvaFontView>& FontViewPtr = AvalancheFont.Get())
+	{
+		const UAvaFontObject* const CurrFontObj = FontViewPtr->GetFontObject();
+		const UFont* const CurrFont = FontViewPtr->GetFont();
+		if (!CurrFont || !IsValid(CurrFont))
+		{
+			return;
+		}
+
+		FontName = FontViewPtr->GetFontNameAsText();
+
+		if (LeftFontNameText.IsValid())
+		{
+			FText TextDisplay = FontName;
+			if (FontViewPtr->GetFontSource() == EAvaFontSource::System)
+			{
+				TextDisplay = FText::FromString(FontViewPtr->GetFontNameAsString() + TEXT("*"));
+			}
+
+			LeftFontNameText->SetText(TextDisplay);
+
+			if (CurrFont)
+			{
+				if (CurrFont->GetCompositeFont())
+				{
+					LeftFontNameText->SetFont(FSlateFontInfo(CurrFont, UE::AvalancheTextEditor::Private::DefaultFontSize));
+				}
+			}
+		}
+
+		if (RightFontNameText.IsValid())
+		{
+			if (CurrFontObj)
+			{
+				const FString AlternateText = CurrFontObj->GetAlternateText();
+
+				if (!AlternateText.IsEmpty())
+				{
+					RightFontNameText->SetVisibility(EVisibility::Visible);
+					const FText TextDisplay = FText::FromString(CurrFontObj->GetAlternateText());
+					RightFontNameText->SetText(TextDisplay);
+
+					if (CurrFont->GetCompositeFont())
+					{
+						RightFontNameText->SetFont(FSlateFontInfo(CurrFont, UE::AvalancheTextEditor::Private::DefaultFontSize));
+					}
+				}
+			}
+			else
+			{
+				RightFontNameText->SetVisibility(EVisibility::Collapsed);
+			}
+		}
+	}
+}
+
+ECheckBoxState SAvaFontField::GetFavoriteState() const
+{
+	if (const TSharedPtr<FAvaFontView>& FontViewPtr = AvalancheFont.Get())
+	{
+		if (FontViewPtr->IsFavorite())
+		{
+			return ECheckBoxState::Checked;
+		}
+	}
+
+	return ECheckBoxState::Unchecked;
+}
+
+FSlateColor SAvaFontField::GetToggleFavoriteColor() const
+{
+	if (const TSharedPtr<FAvaFontView>& FontViewPtr = AvalancheFont.Get())
+	{
+		if (FontViewPtr->IsFavorite())
+		{
+			return FSlateColor(EStyleColor::AccentBlue);
+		}
+	}
+
+	return FSlateColor(EStyleColor::Foreground);
+}
+
+FReply SAvaFontField::OnToggleFavoriteClicked()
+{
+	UAvaFontManagerSubsystem* FontManagerSubsystem = UAvaFontManagerSubsystem::Get();
+	if (!FontManagerSubsystem)
+	{
+		FReply::Handled();
+	}
+
+	if (const TSharedPtr<FAvaFontView>& FontViewPtr = AvalancheFont.Get())
+	{
+		const bool bIsNowFavorite = !FontViewPtr->IsFavorite();
+
+		FontViewPtr->SetFavorite(bIsNowFavorite);
+
+		const FString FavFontName = FontViewPtr->GetFontNameAsString();
+
+		if (bIsNowFavorite)
+		{
+			FontManagerSubsystem->AddFavorite(*FontViewPtr->GetFontNameAsString());
+		}
+		else
+		{
+			FontManagerSubsystem->RemoveFavorite(*FontViewPtr->GetFontNameAsString());
+		}
+
+		OnAvaFontFieldModified.ExecuteIfBound(FontViewPtr);
+	}
+
+	return FReply::Handled();
+}
+
+EVisibility SAvaFontField::GetFavoriteVisibility() const
+{
+	return IsHovered() ? EVisibility::Visible : EVisibility::Hidden;
+}
+
+EVisibility SAvaFontField::GetLocallyAvailableIconVisibility() const
+{
+	bool bIsAvailable = false;
+	if (const UAvaFontManagerSubsystem* FontManagerSubsystem = UAvaFontManagerSubsystem::Get())
+	{
+		if (const TSharedPtr<FAvaFontView>& FontViewPtr = AvalancheFont.Get())
+		{
+			bIsAvailable = FontManagerSubsystem->IsFontAvailableOnLocalOS(FontViewPtr->GetFont());
+		}
+	}
+
+	return bIsAvailable && !IsSelected() ? EVisibility::Visible : EVisibility::Collapsed;
+}
+
+#undef LOCTEXT_NAMESPACE
diff --git a/Engine/Plugins/Experimental/Avalanche/Source/AvalancheTextEditor/Private/Slate/SAvaFontField.h b/Engine/Plugins/Experimental/Avalanche/Source/AvalancheTextEditor/Private/Slate/SAvaFontField.h
new file mode 100644
--- /dev/null
+++ b/Engine/Plugins/Experimental/Avalanche/Source/AvalancheTextEditor/Private/Slate/SAvaFontField.h
@@ -0,0 +1,57 @@
+// Copyright Epic Games, Inc. All Rights Reserved.
+
+#pragma once
+
+#include "Widgets/SCompoundWidget.h"
+
+class FAvaFontView;
+class SAvaFontSelector;
+class STextBlock;
+enum class ECheckBoxState : uint8;
+
+/**
+ * Represents a font selection field 
+ */
+DECLARE_DELEGATE_OneParam(FOnAvaFontFieldModified, const TSharedPtr<FAvaFontView>&)
+
+class SAvaFontField : public SCompoundWidget
+{
+public:
+	SLATE_BEGIN_ARGS(SAvaFontField)
+		:_AvalancheFont(),
+		_OnAvaFontFieldModified()
+		{
+		}
+
+		SLATE_ATTRIBUTE(TSharedPtr<FAvaFontView>, AvalancheFont)
+		SLATE_EVENT(FOnAvaFontFieldModified, OnAvaFontFieldModified)
+
+	SLATE_END_ARGS()
+
+	void Construct(const FArguments& InArgs);
+	void UpdateFont(const TSharedPtr<FAvaFontView>& InAvaFont);
+	void Select() { bIsSelected = true; }
+	void Deselect() { bIsSelected = false; }
+
+protected:
+	ECheckBoxState GetFavoriteState() const;
+	FSlateColor GetToggleFavoriteColor() const;
+	FReply OnToggleFavoriteClicked();
+	EVisibility GetFavoriteVisibility() const;
+	EVisibility GetLocallyAvailableIconVisibility() const;
+
+	bool IsSelected() const { return bIsSelected; }
+
+	FOnAvaFontFieldModified OnAvaFontFieldModified;
+
+	FText FontName;
+
+	TSharedPtr<STextBlock> LeftFontNameText;
+	TSharedPtr<STextBlock> RightFontNameText;
+
+	TAttribute<TSharedPtr<FAvaFontView>> AvalancheFont;
+
+	TSharedPtr<FSlateBrush> CheckBoxBg;
+
+	bool bIsSelected = false;
+};
diff --git a/Engine/Plugins/Experimental/Avalanche/Source/AvalancheTextEditor/Private/Slate/SAvaFontSelector.cpp b/Engine/Plugins/Experimental/Avalanche/Source/AvalancheTextEditor/Private/Slate/SAvaFontSelector.cpp
new file mode 100644
--- /dev/null
+++ b/Engine/Plugins/Experimental/Avalanche/Source/AvalancheTextEditor/Private/Slate/SAvaFontSelector.cpp
@@ -0,0 +1,780 @@
+// Copyright Epic Games, Inc. All Rights Reserved.
+
+#include "SAvaFontSelector.h"
+#include "AvaFontSearchSettingsMenu.h"
+#include "Font/AvaFontManagerSubsystem.h"
+#include "Font/AvaFontView.h"
+#include "Framework/Application/SlateApplication.h"
+#include "Framework/Application/SlateUser.h"
+#include "Framework/SlateDelegates.h"
+#include "Input/Events.h"
+#include "Input/Reply.h"
+#include "Misc/TextFilter.h"
+#include "PropertyHandle.h"
+#include "Styling/AppStyle.h"
+#include "Styling/SlateColor.h"
+#include "Styling/SlateTypes.h"
+#include "Styling/StyleColors.h"
+#include "Widgets/Input/SSearchBox.h"
+#include "Widgets/Layout/SSeparator.h"
+#include "Widgets/SWidget.h"
+#include "Widgets/Text/STextBlock.h"
+#include "Widgets/Views/SListView.h"
+#include "Widgets/Views/STableRow.h"
+#include "Widgets/Views/STableViewBase.h"
+
+#define LOCTEXT_NAMESPACE "AvaFontSelector"
+
+class SFontSelectorRow : public STableRow<TSharedPtr<FAvaFontView>>
+{
+public:
+	SLATE_BEGIN_ARGS(SFontSelectorRow)
+		: _Style(&FAppStyle::Get().GetWidgetStyle<FTableRowStyle>("ComboBox.Row"))
+		, _Content()
+		, _Padding(FMargin(0))
+	{}
+	SLATE_STYLE_ARGUMENT(FTableRowStyle, Style)
+	SLATE_DEFAULT_SLOT(FArguments, Content)
+	SLATE_ATTRIBUTE(FMargin, Padding)
+SLATE_END_ARGS()
+
+	void Construct(const FArguments& InArgs, const TSharedRef<STableViewBase>& InOwnerTable);
+
+	//~ Begin SWidget
+	virtual FReply OnMouseButtonDown(const FGeometry& MyGeometry, const FPointerEvent& MouseEvent) override;
+	//~ End SWidget
+};
+
+
+class SFontSelectorSeparator: public SCompoundWidget
+{
+public:
+	SLATE_BEGIN_ARGS(SFontSelectorSeparator)
+		:_Text("Separator")
+	{
+	}
+	SLATE_ARGUMENT(FString, Text)
+
+	SLATE_END_ARGS()
+
+	/** Constructs this widget with InArgs */
+	void Construct(const FArguments& InArgs)
+	{
+		Text = InArgs._Text;
+
+		ChildSlot
+		[
+			SNew(SBox)
+			.VAlign(VAlign_Fill)
+			.HAlign(HAlign_Fill)
+			[
+				SNew(STextBlock)
+				.Text(FText::FromString(InArgs._Text))
+			]
+		];
+	}
+
+private:
+	FString Text;
+};
+
+void SFontSelectorRow::Construct(const FArguments& InArgs, const TSharedRef<STableViewBase>& InOwnerTable)
+{
+	STableRow<TSharedPtr<FAvaFontView>>::Construct(
+			typename STableRow<TSharedPtr<FAvaFontView>>::FArguments()
+			.Style(InArgs._Style)
+			.Padding(InArgs._Padding)
+			.Content()
+			[
+				InArgs._Content.Widget
+			]
+			, InOwnerTable
+		);
+}
+
+// Handle case where user clicks on an existing selected item
+FReply SFontSelectorRow::OnMouseButtonDown(const FGeometry& MyGeometry, const FPointerEvent& MouseEvent)
+{
+	if (MouseEvent.GetEffectingButton() == EKeys::LeftMouseButton)
+	{
+		const TSharedPtr<ITypedTableView<TSharedPtr<FAvaFontView>>> OwnerWidget = this->OwnerTablePtr.Pin();
+
+		const TSharedPtr<FAvaFontView>* MyItem = OwnerWidget->Private_ItemFromWidget(this);
+		const bool bIsSelected = OwnerWidget->Private_IsItemSelected(*MyItem);
+
+		if (bIsSelected)
+		{
+			// Reselect content to ensure selection is taken
+			OwnerWidget->Private_SignalSelectionChanged(ESelectInfo::Direct);
+			return FReply::Handled();
+		}
+	}
+	return STableRow<TSharedPtr<FAvaFontView>>::OnMouseButtonDown(MyGeometry, MouseEvent);
+}
+
+void SAvaFontSelector::Construct(const FArguments& InArgs)
+{
+	check(InArgs._ComboBoxStyle)
+
+	ItemStyle = InArgs._ItemStyle;
+	MenuRowPadding = InArgs._ComboBoxStyle->MenuRowPadding;
+
+	// Work out which values we should use based on whether we were given an override, or should use the style's version
+	const FComboButtonStyle& OurComboButtonStyle = InArgs._ComboBoxStyle->ComboButtonStyle;
+	const FButtonStyle* const OurButtonStyle = InArgs._ButtonStyle ? InArgs._ButtonStyle : &OurComboButtonStyle.ButtonStyle;
+
+	OnComboBoxOpening = InArgs._OnComboBoxOpening;
+	OnSelectionChanged = InArgs._OnSelectionChanged;
+	OnGenerateWidget = InArgs._OnGenerateWidget;
+	EnableGamepadNavigationMode = InArgs._EnableGamepadNavigationMode;
+	bControllerInputCaptured = false;
+
+	bShowMonospacedOnly = InArgs._ShowMonospacedFontsOnly;
+	bShowBoldOnly = InArgs._ShowBoldFontsOnly;
+	bShowItalicOnly = InArgs._ShowItalicFontsOnly;
+
+	AvalancheFontPropertyHandlePtr = InArgs._AvalancheFontPropertyHandle.Get();
+
+	FontsOptionsSource = InArgs._OptionsSource;
+	CustomScrollbar = InArgs._CustomScrollbar;
+
+	SearchBoxFilter = MakeShared<FontViewTextFilter>(
+		FontViewTextFilter::FItemToStringArray::CreateSP(this, &SAvaFontSelector::TransformElementToString)
+		);
+
+	bOptionsNeedRefresh = true;
+
+	const TSharedPtr<SWidget> ComboBoxMenuContent = SNew(SBorder)
+		.Padding(2.0f)
+		.BorderBackgroundColor(FSlateColor(EStyleColor::AccentBlue))
+		[
+			SNew(SVerticalBox)
+			+ SVerticalBox::Slot()
+			.Padding(7.0f, 6.0f)
+			.AutoHeight()
+			[
+				SNew(SHorizontalBox)
+				+ SHorizontalBox::Slot()
+				.FillWidth(1.0f)
+				[
+					SAssignNew(SearchBox, SSearchBox)
+					.OnTextChanged(this, &SAvaFontSelector::OnFilterTextChanged)
+					.OnKeyDownHandler(this, &SAvaFontSelector::OnSearchFieldKeyDown)
+				]
+				+ SHorizontalBox::Slot()
+				.AutoWidth()
+				[
+					SNew(SAvaFontSearchSettingsMenu)
+					.AvaFontSelector(SharedThis(this))
+				]
+			]
+			// FAVORITE FONTS SECTION //
+			+ SVerticalBox::Slot()
+			.VAlign(VAlign_Fill)
+			.HAlign(HAlign_Center)
+			.Padding(2.0f)
+			[
+				SAssignNew(FavoritesSeparator, SFontSelectorSeparator)
+				.Text("Favorites")
+			]
+			+ SVerticalBox::Slot()
+			.AutoHeight()
+			[
+				SNew(SBox)
+				.MaxDesiredHeight(InArgs._MaxListHeight)
+				.VAlign(VAlign_Fill)
+				[
+					SAssignNew(FavoriteFontsListView, SComboListType)
+					.ListItemsSource(&FavoriteFontsOptions)
+					.OnGenerateRow(this, &SAvaFontSelector::GenerateMenuItemRow)
+					.OnSelectionChanged(this, &SAvaFontSelector::OnSelectionChanged_Internal)
+					.SelectionMode(ESelectionMode::Single)
+					.ExternalScrollbar(InArgs._CustomScrollbar)
+				]
+			]
+
+			// IMPORTED FONTS SECTION //
+			+ SVerticalBox::Slot()
+			.VAlign(VAlign_Fill)
+			.HAlign(HAlign_Center)
+			.Padding(2.0f)
+			[
+				SAssignNew(ImportedSeparator, SFontSelectorSeparator)
+				.Text("Project Fonts")
+			]
+			+ SVerticalBox::Slot()
+			.AutoHeight()
+			[
+				SNew(SBox)
+				.MaxDesiredHeight(InArgs._MaxListHeight)
+				.VAlign(VAlign_Fill)
+				[
+					SAssignNew(this->ImportedFontsListView, SComboListType)
+					.ListItemsSource(&ImportedFontsOptions)
+					.OnGenerateRow(this, &SAvaFontSelector::GenerateMenuItemRow)
+					.OnSelectionChanged(this, &SAvaFontSelector::OnSelectionChanged_Internal)
+					.SelectionMode(ESelectionMode::Single)
+					.ExternalScrollbar(InArgs._CustomScrollbar)
+				]
+			]
+
+			// AVAILABLE SYSTEM FONTS SECTION //
+			+ SVerticalBox::Slot()
+			.VAlign(VAlign_Fill)
+			.HAlign(HAlign_Center)
+			.Padding(2.0f)
+			[
+				SAssignNew(AllFontsSeparator, SFontSelectorSeparator)
+				.Text("System Fonts")
+			]
+			+ SVerticalBox::Slot()
+			.AutoHeight()
+			[
+				SNew(SBox)
+				.MaxDesiredHeight(InArgs._MaxListHeight)
+				.VAlign(VAlign_Fill)
+				[
+					SAssignNew(FontsListView, SComboListType)
+					.ListItemsSource(&FilteredFonts)
+					.OnGenerateRow(this, &SAvaFontSelector::GenerateMenuItemRow)
+					.OnSelectionChanged(this, &SAvaFontSelector::OnSelectionChanged_Internal)
+					.SelectionMode(ESelectionMode::Single)
+					.ExternalScrollbar(InArgs._CustomScrollbar)
+				]
+			]
+		];
+
+	// Set up content
+	TSharedPtr<SWidget> ButtonContent = InArgs._Content.Widget;
+	if (InArgs._Content.Widget == SNullWidget::NullWidget)
+	{
+		 SAssignNew(ButtonContent, STextBlock)
+		.Text(NSLOCTEXT("SAvaFontSelector", "ContentWarning", "No Content Provided"))
+		.ColorAndOpacity(FLinearColor::Red);
+	}
+
+	SComboButton::Construct(SComboButton::FArguments()
+		.ComboButtonStyle(&OurComboButtonStyle)
+		.ButtonStyle(OurButtonStyle)
+		.Method(InArgs._Method)
+		.ButtonContent()
+		[
+			ButtonContent.ToSharedRef()
+		]
+		.MenuContent()
+		[
+			ComboBoxMenuContent.ToSharedRef()
+		]
+		.HasDownArrow(InArgs._HasDownArrow)
+		.ContentPadding(InArgs._ContentPadding)
+		.ForegroundColor(InArgs._ForegroundColor)
+		.OnMenuOpenChanged(this, &SAvaFontSelector::OnMenuOpenChanged)
+		.IsFocusable(InArgs._IsFocusable)
+		.CollapseMenuOnParentFocus(InArgs._CollapseMenuOnParentFocus)
+		.ToolTipText(this, &SAvaFontSelector::GetCurrentFontTooltipText)
+	);
+
+	SetMenuContentWidgetToFocus(SearchBox);
+
+	// Need to establish the selected item at point of construction so its available for querying
+	// NB: If you need a selection to fire use SetItemSelection rather than setting an InitiallySelectedItem
+	SelectedItem = InArgs._InitiallySelectedItem;
+	if (ListTypeTraits::IsPtrValid(SelectedItem))
+	{
+		const TSharedPtr<FAvaFontView> ValidatedItem = ListTypeTraits::NullableItemTypeConvertToItemType(SelectedItem);
+		FontsListView->Private_SetItemSelection(ValidatedItem, true);
+		FontsListView->RequestScrollIntoView(ValidatedItem, 0);
+	}
+}
+
+void SAvaFontSelector::ClearSelection() const
+{
+	FontsListView->ClearSelection();
+	FavoriteFontsListView->ClearSelection();
+	ImportedFontsListView->ClearSelection();
+}
+
+void SAvaFontSelector::SetSelectedItem(const NullableOptionType& InSelectedItem) const
+{
+	if (ListTypeTraits::IsPtrValid(InSelectedItem))
+	{
+		const TSharedPtr<FAvaFontView> InSelected = ListTypeTraits::NullableItemTypeConvertToItemType(InSelectedItem);
+		FontsListView->SetSelection(InSelected);
+	}
+	else
+	{
+		ClearSelection();
+	}
+}
+
+SAvaFontSelector::NullableOptionType SAvaFontSelector::GetSelectedItem()
+{
+	return SelectedItem;
+}
+
+void SAvaFontSelector::UpdateFilteredFonts()
+{
+	// Based on search results, update the filtered fonts list to shown to the user
+	// Other lists will update consequently
+
+	FilteredFonts.Empty();
+
+	for (const TSharedPtr<FAvaFontView>& FontView : *FontsOptionsSource)
+	{
+		if (!FontView->HasValidFont())
+		{
+			continue;
+		}
+
+		if (bShowMonospacedOnly && !FontView->IsMonospaced())
+		{
+			continue;
+		}
+
+		if (bShowBoldOnly && !FontView->IsBold())
+		{
+			continue;
+		}
+
+		if (bShowItalicOnly && !FontView->IsItalic())
+		{
+			continue;
+		}
+
+		if (SearchBoxFilter->PassesFilter(FontView))
+		{
+			FilteredFonts.Add(FontView);
+		}
+	}
+}
+
+void SAvaFontSelector::UpdateFavoriteFonts()
+{
+	TArray<TSharedPtr<FAvaFontView>> Favorites;
+	for (TSharedPtr<FAvaFontView> Font : FilteredFonts)
+	{
+		if (Font->IsFavorite())
+		{
+			Favorites.AddUnique(Font);
+		}
+	}
+
+	FavoriteFontsOptions.Empty();
+	FavoriteFontsOptions = Favorites;
+}
+
+void SAvaFontSelector::UpdateImportedFonts()
+{
+	TArray<TSharedPtr<FAvaFontView>> Imported;
+	for (TSharedPtr<FAvaFontView> Font : FilteredFonts)
+	{
+		if (Font->GetFontSource() == EAvaFontSource::Project)
+		{
+			Imported.AddUnique(Font);
+		}
+	}
+
+	ImportedFontsOptions.Empty();
+	ImportedFontsOptions = Imported;
+	UpdateProjectFonts();
+
+	ImportedFontsOptions.Sort([](const TSharedPtr<FAvaFontView>& A, const TSharedPtr<FAvaFontView>& B)
+	{
+		return *A.Get() < *B.Get();
+	});
+}
+
+void SAvaFontSelector::UpdateProjectFonts()
+{
+	TArray<TSharedPtr<FAvaFontView>> EngineFonts;
+	for (const TSharedPtr<FAvaFontView>& Font : FilteredFonts)
+	{
+		if (Font->GetFontSource() == EAvaFontSource::Project)
+		{
+			if (!ImportedFontsOptions.Contains(Font))
+			{
+				EngineFonts.AddUnique(Font);
+			}
+		}
+	}
+	EngineFontOptions.Empty();
+	EngineFontOptions = EngineFonts;
+	ImportedFontsOptions.Append(EngineFontOptions);
+}
+
+void SAvaFontSelector::UpdateSeparatorsVisibility() const
+{
+	if (FilteredFonts.IsEmpty())
+	{
+		AllFontsSeparator->SetVisibility(EVisibility::Collapsed);
+	}
+	else
+	{
+		AllFontsSeparator->SetVisibility(EVisibility::Visible);
+	}
+
+	if (FavoriteFontsOptions.IsEmpty())
+	{
+		FavoritesSeparator->SetVisibility(EVisibility::Collapsed);
+	}
+	else
+	{
+		FavoritesSeparator->SetVisibility(EVisibility::Visible);
+	}
+
+	if (ImportedFontsOptions.IsEmpty())
+	{
+		ImportedSeparator->SetVisibility(EVisibility::Collapsed);
+	}
+	else
+	{
+		ImportedSeparator->SetVisibility(EVisibility::Visible);
+	}
+}
+
+void SAvaFontSelector::RefreshOptions()
+{
+	UpdateFilteredFonts();
+
+	UpdateFavoriteFonts();
+	FavoriteFontsListView->RequestListRefresh();
+
+	UpdateImportedFonts();
+	ImportedFontsListView->RequestListRefresh();
+
+	// Remove duplicates from list
+	TArray<TSharedPtr<FAvaFontView>> FontsToRemove;
+	for (const TSharedPtr<FAvaFontView>& Font : FilteredFonts)
+	{
+		if (ImportedFontsOptions.Contains(Font))
+		{
+			FontsToRemove.Add(Font);
+		}
+	}
+
+	for (const TSharedPtr<FAvaFontView>& Font : FontsToRemove)
+	{
+		if (FilteredFonts.Contains(Font))
+		{
+			FilteredFonts.Remove(Font);
+		}
+	}
+
+	FontsListView->RequestListRefresh();
+
+	// We might want to hide some separators, in case their list is empty
+	UpdateSeparatorsVisibility();
+}
+
+void SAvaFontSelector::SetShowMonospacedOnly(bool bMonospacedOnly)
+{
+	bShowMonospacedOnly = bMonospacedOnly;
+
+	// Trigger a refresh
+	bOptionsNeedRefresh = true;
+}
+
+void SAvaFontSelector::SetShowBoldOnly(bool bBoldOnly)
+{
+	bShowBoldOnly = bBoldOnly;
+
+	// Trigger a refresh
+	bOptionsNeedRefresh = true;
+}
+
+void SAvaFontSelector::SetShowItalicOnly(bool bItalicOnly)
+{
+	bShowItalicOnly = bItalicOnly;
+
+	// Trigger a refresh
+	bOptionsNeedRefresh = true;
+}
+
+FReply SAvaFontSelector::OnSearchFieldKeyDown(const FGeometry& MyGeometry, const FKeyEvent& InKeyEvent)
+{
+	if (InKeyEvent.GetKey() == EKeys::Escape)
+	{
+		FSlateApplication::Get().DismissAllMenus();
+		return FReply::Handled();
+	}
+	else if (InKeyEvent.GetKey() == EKeys::Enter)
+	{
+		if (!FilteredFonts.IsEmpty())
+		{
+			OnSelectionChanged_Internal(FilteredFonts[0], ESelectInfo::Type::Direct);
+		}
+
+		return FReply::Handled();
+	}
+
+	return FReply::Unhandled();
+}
+
+FReply SAvaFontSelector::OnKeyDown(const FGeometry& MyGeometry, const FKeyEvent& InKeyEvent)
+{
+	if (IsInteractable())
+	{
+		const EUINavigationAction NavAction = FSlateApplication::Get().GetNavigationActionFromKey(InKeyEvent);
+		const EUINavigation NavDirection = FSlateApplication::Get().GetNavigationDirectionFromKey(InKeyEvent);
+		if (EnableGamepadNavigationMode)
+		{
+			// The controller's bottom face button must be pressed once to begin manipulating the combobox's value.
+			// Navigation away from the widget is prevented until the button has been pressed again or focus is lost.
+			if (NavAction == EUINavigationAction::Accept)
+			{
+				if (bControllerInputCaptured == false)
+				{
+					// Begin capturing controller input and open the ListView
+					bControllerInputCaptured = true;
+					OnComboBoxOpening.ExecuteIfBound();
+					return SComboButton::OnButtonClicked();
+				}
+				else
+				{
+					// Set selection to the selected item on the list and close
+					bControllerInputCaptured = false;
+
+					// Re-select first selected item, just in case it was selected by navigation previously
+					TArray<TSharedPtr<FAvaFontView>> SelectedItems = FontsListView->GetSelectedItems();
+					if (!SelectedItems.IsEmpty())
+					{
+						OnSelectionChanged_Internal(SelectedItems[0], ESelectInfo::Direct);
+					}
+
+					// Set focus back to FontSelector
+					FReply Reply = FReply::Handled();
+					Reply.SetUserFocus(this->AsShared(), EFocusCause::SetDirectly);
+					return Reply;
+				}
+
+			}
+			else if (NavAction == EUINavigationAction::Back || InKeyEvent.GetKey() == EKeys::BackSpace)
+			{
+				const bool bWasInputCaptured = bControllerInputCaptured;
+
+				OnMenuOpenChanged(false);
+				if (bWasInputCaptured)
+				{
+					return FReply::Handled();
+				}
+			}
+			else
+			{
+				if (bControllerInputCaptured)
+				{
+					return FReply::Handled();
+				}
+			}
+		}
+		else
+		{
+			if (NavDirection == EUINavigation::Up)
+			{
+				const NullableOptionType NullableSelected = GetSelectedItem();
+				if (ListTypeTraits::IsPtrValid(NullableSelected))
+				{
+					const TSharedPtr<FAvaFontView> ActuallySelected = ListTypeTraits::NullableItemTypeConvertToItemType(NullableSelected);
+					const int32 SelectionIndex = FontsOptionsSource->Find(ActuallySelected);
+					if (SelectionIndex >= 1)
+					{
+						// Select an item on the prev row
+						SetSelectedItem((*FontsOptionsSource)[SelectionIndex - 1]);
+					}
+				}
+
+				return FReply::Handled();
+			}
+			else if (NavDirection == EUINavigation::Down)
+			{
+				const NullableOptionType NullableSelected = GetSelectedItem();
+				if (ListTypeTraits::IsPtrValid(NullableSelected))
+				{
+					const TSharedPtr<FAvaFontView> ActuallySelected = ListTypeTraits::NullableItemTypeConvertToItemType(NullableSelected);
+					const int32 SelectionIndex = FontsOptionsSource->Find(ActuallySelected);
+					if (SelectionIndex < FontsOptionsSource->Num() - 1)
+					{
+						// Select an item on the next row
+						SetSelectedItem((*FontsOptionsSource)[SelectionIndex + 1]);
+					}
+				}
+
+				return FReply::Handled();
+			}
+
+			return SComboButton::OnKeyDown(MyGeometry, InKeyEvent);
+		}
+	}
+	return SWidget::OnKeyDown(MyGeometry, InKeyEvent);
+}
+
+FText SAvaFontSelector::GetCurrentFontTooltipText() const
+{
+	if (!FontsListView->GetSelectedItems().IsEmpty())
+	{
+		const TSharedPtr<FAvaFontView> Selected = FontsListView->GetSelectedItems()[0];
+		return Selected->GetFontNameAsText();
+	}
+
+	return LOCTEXT("FontSelector_NoFont_Tooltip", "No Font is currently selected.");
+}
+
+TSharedRef<ITableRow> SAvaFontSelector::GenerateMenuItemRow(TSharedPtr<FAvaFontView> InItem, const TSharedRef<STableViewBase>& OwnerTable) const
+{
+	if (OnGenerateWidget.IsBound() && InItem.IsValid())
+	{
+		return SNew(SFontSelectorRow, OwnerTable)
+			.Style(ItemStyle)
+			.Padding(MenuRowPadding)
+			[
+				OnGenerateWidget.Execute(InItem)
+			];
+	}
+	else
+	{
+		return SNew(SFontSelectorRow, OwnerTable)
+			[
+				SNew(STextBlock).Text(NSLOCTEXT("SlateCore", "ComboBoxMissingOnGenerateWidgetMethod", "Please provide a .OnGenerateWidget() handler."))
+			];
+	}
+}
+
+void SAvaFontSelector::OnMenuOpenChanged(bool bOpen)
+{
+	if (bOpen == false)
+	{
+		bControllerInputCaptured = false;
+
+		if (ListTypeTraits::IsPtrValid(SelectedItem))
+		{
+			// Ensure the ListView selection is set back to the last committed selection
+			const TSharedPtr<FAvaFontView> ActuallySelected = ListTypeTraits::NullableItemTypeConvertToItemType(SelectedItem);
+
+			FontsListView->SetSelection(ActuallySelected, ESelectInfo::OnNavigation);
+			FontsListView->RequestScrollIntoView(ActuallySelected, 0);
+		}
+
+		// Set focus back to FontSelector for users focusing the ListView that just closed
+		TSharedRef<SWidget> ThisRef = AsShared();
+		FSlateApplication::Get().ForEachUser([&ThisRef](FSlateUser& User)
+		{
+			if (User.HasFocusedDescendants(ThisRef))
+			{
+				User.SetFocus(ThisRef, EFocusCause::SetDirectly);
+			}
+		});
+	}
+}
+
+void SAvaFontSelector::OnSelectionChanged_Internal(NullableOptionType ProposedSelection, ESelectInfo::Type SelectInfo)
+{
+	if (!ProposedSelection)
+	{
+		return;
+	}
+
+	// Ensure that the proposed selection is different
+	if (SelectInfo != ESelectInfo::OnNavigation)
+	{
+		// Ensure that the proposed selection is different from selected
+		if (ProposedSelection != SelectedItem)
+		{
+			SelectedItem = ProposedSelection;
+			OnSelectionChanged.ExecuteIfBound(ProposedSelection, SelectInfo);
+		}
+
+		// Close combo even if user reselected item
+		SetIsOpen(false);
+	}
+}
+
+FReply SAvaFontSelector::OnButtonClicked()
+{
+	// If user clicked to close the combo menu
+	if (IsOpen())
+	{
+		// Re-select first selected item, just in case it was selected by navigation previously
+		TArray<TSharedPtr<FAvaFontView>> SelectedItems = FontsListView->GetSelectedItems();
+		if (!SelectedItems.IsEmpty())
+		{
+			OnSelectionChanged_Internal(SelectedItems[0], ESelectInfo::Direct);
+		}
+	}
+	else
+	{
+		OnComboBoxOpening.ExecuteIfBound();
+	}
+
+	return SComboButton::OnButtonClicked();
+}
+
+void SAvaFontSelector::OnFilterTextChanged(const FText& Text)
+{
+	SearchBoxFilter->SetRawFilterText(Text);
+	bOptionsNeedRefresh = true;
+}
+
+void SAvaFontSelector::TransformElementToString(TSharedPtr<FAvaFontView> InElement, TArray<FString>& OutStrings)
+{
+	if (InElement)
+	{
+		OutStrings.Add(InElement->GetFontNameAsString());
+	}
+}
+
+FText SAvaFontSelector::GetFilterHighlightText() const
+{
+	if (SearchBoxFilter)
+	{
+		return SearchBoxFilter->GetRawFilterText();
+	}
+
+	return {};
+}
+
+void SAvaFontSelector::RefreshUIFacingFont()
+{
+	UAvaFontManagerSubsystem* FontManagerSubsystem = UAvaFontManagerSubsystem::Get();
+	if (!FontManagerSubsystem)
+	{
+		return;
+	}
+
+	if (!AvalancheFontPropertyHandlePtr)
+	{
+		return;
+	}
+
+	const FAvaFontView* UIFacingFont = SelectedItem.Get();
+	if (!UIFacingFont)
+	{
+		return;
+	}
+
+	FPropertyAccess::Result AccessResult;
+	const TSharedPtr<FAvaFontView> ActualFontPtr = FontManagerSubsystem->GetFontViewFromPropertyHandle(AvalancheFontPropertyHandlePtr, AccessResult);
+	if (!ActualFontPtr)
+	{
+		return;
+	}
+
+	// Only perform the selection action if the fonts are different
+	if (ActualFontPtr.Get() != UIFacingFont)
+	{
+		OnSelectionChanged_Internal(ActualFontPtr, ESelectInfo::Direct);
+	}
+}
+
+void SAvaFontSelector::Tick(const FGeometry& AllottedGeometry, const double InCurrentTime, const float InDeltaTime)
+{
+	SMenuAnchor::Tick(AllottedGeometry, InCurrentTime, InDeltaTime);
+
+	if (bOptionsNeedRefresh)
+	{
+		bOptionsNeedRefresh = false;
+		RefreshOptions();
+	}
+
+	RefreshUIFacingFont();
+}
+
+#undef LOCTEXT_NAMESPACE
