[MotionDesign]

diff --git a/Engine/Plugins/Experimental/Avalanche/Source/AvalancheTextEditor/Private/Slate/SAvaFontSelector.h b/Engine/Plugins/Experimental/Avalanche/Source/AvalancheTextEditor/Private/Slate/SAvaFontSelector.h
new file mode 100644
--- /dev/null
+++ b/Engine/Plugins/Experimental/Avalanche/Source/AvalancheTextEditor/Private/Slate/SAvaFontSelector.h
@@ -0,0 +1,288 @@
+// Copyright Epic Games, Inc. All Rights Reserved.
+
+#pragma once
+
+#include "Framework/Views/TableViewTypeTraits.h"
+#include "Layout/Margin.h"
+#include "Widgets/Input/SComboButton.h"
+
+class FAvaFontView;
+class FReply;
+class IPropertyHandle;
+class SFontSelectorSeparator;
+class SScrollBar;
+class SSearchBox;
+class STableViewBase;
+class SWidget;
+template <typename ArgumentType> class TSlateDelegates;
+template <typename ItemType> class SListView;
+template <typename ItemType> class TTextFilter;
+
+#if WITH_ACCESSIBILITY
+
+#endif
+
+DECLARE_DELEGATE(FOnAvaComboBoxOpening)
+
+class SAvaFontSelector : public SComboButton
+{
+public:
+	typedef TListTypeTraits<TSharedPtr<FAvaFontView>> ListTypeTraits;
+	typedef TListTypeTraits<TSharedPtr<FAvaFontView>>::NullableType NullableOptionType;
+
+	/** Type of list used for showing menu options. */
+	typedef SListView<TSharedPtr<FAvaFontView>> SComboListType;
+
+	/** Type of text filter used for FAvaFontView */
+	typedef TTextFilter<TSharedPtr<FAvaFontView>> FontViewTextFilter;
+
+	/** Delegate types used to generate widgets that represent Options */
+	typedef TSlateDelegates<const TSharedPtr<FAvaFontView>&>::FOnGenerateWidget FOnGenerateWidget;
+	typedef TSlateDelegates<NullableOptionType>::FOnSelectionChanged FOnSelectionChanged;
+
+	SLATE_BEGIN_ARGS(SAvaFontSelector)
+		: _Content()
+		, _ComboBoxStyle(&FAppStyle::Get().GetWidgetStyle< FComboBoxStyle >("ComboBox"))
+		, _ButtonStyle(nullptr)
+		, _ItemStyle(&FAppStyle::Get().GetWidgetStyle< FTableRowStyle >("ComboBox.Row"))
+		, _ContentPadding(_ComboBoxStyle->ContentPadding)
+		, _ForegroundColor(FSlateColor::UseStyle())
+		, _AvalancheFontPropertyHandle()
+		, _OptionsSource()
+		, _OnSelectionChanged()
+		, _OnGenerateWidget()
+		, _InitiallySelectedItem(ListTypeTraits::MakeNullPtr())
+		, _Method()
+		, _MaxListHeight(200.0f)
+		, _HasDownArrow(true)
+		, _EnableGamepadNavigationMode(false)
+		, _IsFocusable(true)
+		, _ShowMonospacedFontsOnly(false)
+		, _ShowBoldFontsOnly(false)
+		, _ShowItalicFontsOnly(false)
+		{}
+
+		/** Slot for this button's content (optional) */
+		SLATE_DEFAULT_SLOT(FArguments, Content)
+
+		SLATE_STYLE_ARGUMENT(FComboBoxStyle, ComboBoxStyle)
+
+		/** The visual style of the button part of the combo box (overrides ComboBoxStyle) */
+		SLATE_STYLE_ARGUMENT(FButtonStyle, ButtonStyle)
+
+		SLATE_STYLE_ARGUMENT(FTableRowStyle, ItemStyle)
+
+		SLATE_ATTRIBUTE(FMargin, ContentPadding)
+		SLATE_ATTRIBUTE(FSlateColor, ForegroundColor)
+
+		/** The Property Handle for the font being customized */
+		SLATE_ATTRIBUTE(TSharedPtr<IPropertyHandle>, AvalancheFontPropertyHandle)
+
+		/** Options source for the fonts list available to this font selector widget */
+		SLATE_ARGUMENT(const TArray<TSharedPtr<FAvaFontView>>*, OptionsSource)
+
+		/** Event called whenever the current font selection is changed */
+		SLATE_EVENT(FOnSelectionChanged, OnSelectionChanged)
+
+		SLATE_EVENT(FOnGenerateWidget, OnGenerateWidget)
+
+		/** Called when combo box is opened, before list is actually created */
+		SLATE_EVENT(FOnAvaComboBoxOpening, OnComboBoxOpening)
+
+		/** The custom scrollbar to use in the ListView */
+		SLATE_ARGUMENT(TSharedPtr<SScrollBar>, CustomScrollbar)
+
+		/** The option that should be selected when the combo box is first created */
+		SLATE_ARGUMENT(NullableOptionType, InitiallySelectedItem)
+
+		SLATE_ARGUMENT(TOptional<EPopupMethod>, Method)
+
+		/** The max height of the combo box menu */
+		SLATE_ARGUMENT(float, MaxListHeight)
+
+		/**
+		 * When false, the down arrow is not generated and it is up to the API consumer
+		 * to make their own visual hint that this is a drop down.
+		 */
+		SLATE_ARGUMENT(bool, HasDownArrow)
+
+		/**
+		 *  When false, directional keys will change the selection. When true, FontSelector
+		 *	must be activated and will only capture arrow input while activated.
+		*/
+		SLATE_ARGUMENT(bool, EnableGamepadNavigationMode)
+
+		/** When true, allows the combo box to receive keyboard focus */
+		SLATE_ARGUMENT(bool, IsFocusable)
+
+		/** True if this combo's menu should be collapsed when our parent receives focus, false (default) otherwise */
+		SLATE_ARGUMENT(bool, CollapseMenuOnParentFocus)
+
+		/** When true, will only show monospaced fonts */
+		SLATE_ARGUMENT(bool, ShowMonospacedFontsOnly)
+
+		/** When true, will only show bold fonts */
+		SLATE_ARGUMENT(bool, ShowBoldFontsOnly)
+
+		/** When true, will only show italic fonts */
+		SLATE_ARGUMENT(bool, ShowItalicFontsOnly)
+
+	SLATE_END_ARGS()
+
+	/**
+	 * Construct the widget from a declaration
+	 *
+	 * @param InArgs   Declaration from which to construct the combo box
+	 */
+	void Construct(const FArguments& InArgs);
+
+	SAvaFontSelector()
+		: ItemStyle(nullptr), EnableGamepadNavigationMode(false)
+		, bControllerInputCaptured(false), FontsOptionsSource(nullptr)
+		, bOptionsNeedRefresh(false), bShowMonospacedOnly(false)
+		, bShowBoldOnly(false), bShowItalicOnly(false)
+	{
+	}
+
+	/** 
+	 * Requests a list refresh after updating options 
+	 * Call SetSelectedItem to update the selected item if required
+	 * @see SetSelectedItem
+	 */
+	void RefreshOptions();
+
+	void SetShowMonospacedOnly(bool bMonospacedOnly);
+	void SetShowBoldOnly(bool bBoldOnly);
+	void SetShowItalicOnly(bool bItalicOnly);
+
+protected:
+	FText GetFilterHighlightText() const;
+
+	/** Clears the current selection from all fonts lists*/
+	void ClearSelection() const;
+
+	/** Updates favorites list */
+	void UpdateFavoriteFonts();
+
+	/** Updates imported fonts list */
+	void UpdateImportedFonts();
+
+	/** Updates project fonts list */
+	void UpdateProjectFonts();
+
+	void UpdateSeparatorsVisibility() const;
+
+	/** Create the row's display name using the delegate. */
+	void TransformElementToString(TSharedPtr<FAvaFontView> InElement, TArray<FString>& OutStrings);
+
+	/** Filter text changed handler. */
+	void OnFilterTextChanged(const FText& Text);
+
+	/** Sets the currently selected item */
+	void SetSelectedItem(const NullableOptionType& InSelectedItem) const;
+
+	/** Updates the list of filtered fonts */
+	void UpdateFilteredFonts();
+
+	/**
+	 * Checks if the Font shown to the user as select matches the one handled by the Property Handler.
+	 * If not, the widgets will update accordingly.
+	 */
+	void RefreshUIFacingFont();
+
+	/** @return the item currently selected by the combo box. */
+	NullableOptionType GetSelectedItem();
+
+	//~ Begin SWidget
+	virtual void Tick(const FGeometry& AllottedGeometry, const double InCurrentTime, const float InDeltaTime) override;
+	virtual FReply OnKeyDown(const FGeometry& MyGeometry, const FKeyEvent& InKeyEvent) override;
+	//~ End SWidget
+
+	/** Handles key down events in filter search (e.g. Enter, Escape) */
+	FReply OnSearchFieldKeyDown(const FGeometry& MyGeometry, const FKeyEvent& InKeyEvent);
+
+	virtual bool SupportsKeyboardFocus() const override { return bIsFocusable; }
+
+	virtual bool IsInteractable() const override { return IsEnabled(); }
+
+	FText GetCurrentFontTooltipText() const;
+
+private:
+	/** Generate a row for the InItem in the combo box's list (passed in as OwnerTable). Do this by calling the user-specified OnGenerateWidget */
+	TSharedRef<ITableRow> GenerateMenuItemRow(TSharedPtr<FAvaFontView> InItem, const TSharedRef<STableViewBase>& OwnerTable) const;
+
+	/** Called if the menu is closed */
+	void OnMenuOpenChanged(bool bOpen);
+
+	/** Invoked when the selection in the list changes */
+	void OnSelectionChanged_Internal(NullableOptionType ProposedSelection, ESelectInfo::Type SelectInfo);
+
+	/** Handle clicking on the content menu */
+	virtual FReply OnButtonClicked() override;
+
+	/** The item style to use. */
+	const FTableRowStyle* ItemStyle;
+
+	/** The padding around each menu row */
+	FMargin MenuRowPadding;
+
+private:
+	/** Delegate that is invoked when the selected item in the combo box changes */
+	FOnSelectionChanged OnSelectionChanged;
+
+	/** The item currently selected in the combo box */
+	NullableOptionType SelectedItem;
+
+	/** The ListView that we pop up; visualized the available options. */
+	TSharedPtr<SComboListType> FontsListView;
+
+	/** The ListView that we pop up; visualized the available options. */
+	TSharedPtr<SComboListType> FavoriteFontsListView;
+
+	/** The ListView that we pop up; visualized the available options. */
+	TSharedPtr<SComboListType> ImportedFontsListView;
+
+	/** The Scrollbar used in the ListView. */
+	TSharedPtr<SScrollBar> CustomScrollbar;
+
+	/** Delegate to invoke before the combo box is opening. */
+	FOnAvaComboBoxOpening OnComboBoxOpening;
+
+	/** Delegate to invoke when we need to visualize an option as a widget. */
+	FOnGenerateWidget OnGenerateWidget;
+
+	/** Use activate button to toggle ListView when enabled. */
+	bool EnableGamepadNavigationMode;
+
+	/**
+	 * Holds a flag indicating whether a controller/keyboard is manipulating the combobox's value.
+	 * When true, navigation away from the widget is prevented until a new value has been accepted or canceled.
+	 */
+	bool bControllerInputCaptured;
+
+	const TArray<TSharedPtr<FAvaFontView>>* FontsOptionsSource;
+
+	TArray<TSharedPtr<FAvaFontView>> FavoriteFontsOptions;
+	TArray<TSharedPtr<FAvaFontView>> ImportedFontsOptions;
+	TArray<TSharedPtr<FAvaFontView>> EngineFontOptions;
+
+	TSharedPtr<SFontSelectorSeparator> FavoritesSeparator;
+	TSharedPtr<SFontSelectorSeparator> ImportedSeparator;
+	TSharedPtr<SFontSelectorSeparator> AllFontsSeparator;
+
+	TSharedPtr<SSearchBox> SearchBox;
+
+	/** Holds the search box filter. */
+	TSharedPtr<FontViewTextFilter> SearchBoxFilter;
+
+	/** Whether the list should be refreshed on the next tick. */
+	bool bOptionsNeedRefresh;
+
+	bool bShowMonospacedOnly;
+	bool bShowBoldOnly;
+	bool bShowItalicOnly;
+
+	TArray<TSharedPtr<FAvaFontView>> FilteredFonts;
+
+	TSharedPtr<IPropertyHandle> AvalancheFontPropertyHandlePtr;
+};
diff --git a/Engine/Plugins/Experimental/Avalanche/Source/AvalancheTextEditor/Private/Slate/SAvaTextAlignmentWidget.cpp b/Engine/Plugins/Experimental/Avalanche/Source/AvalancheTextEditor/Private/Slate/SAvaTextAlignmentWidget.cpp
new file mode 100644
--- /dev/null
+++ b/Engine/Plugins/Experimental/Avalanche/Source/AvalancheTextEditor/Private/Slate/SAvaTextAlignmentWidget.cpp
@@ -0,0 +1,191 @@
+// Copyright Epic Games, Inc. All Rights Reserved.
+
+#include "SAvaTextAlignmentWidget.h"
+
+#include "AvaTextDefs.h"
+#include "AvaTextEditorStyle.h"
+#include "PropertyHandle.h"
+#include "SlateOptMacros.h"
+#include "Styling/SlateStyleRegistry.h"
+#include "Styling/StyleColors.h"
+#include "Text3DComponent.h"
+#include "Widgets/Colors/SColorBlock.h"
+#include "Widgets/Images/SImage.h"
+#include "Widgets/Input/SButton.h"
+#include "Widgets/Layout/SBox.h"
+#include "Widgets/Layout/SGridPanel.h"
+#include "Widgets/Text/STextBlock.h"
+
+#define LOCTEXT_NAMESPACE "SAvaTextAlignmentWidget"
+
+BEGIN_SLATE_FUNCTION_BUILD_OPTIMIZATION
+
+void SAvaTextAlignmentWidget::Construct(const FArguments& InArgs)
+{
+	TextAlignmentPropertyHandle = InArgs._TextAlignmentPropertyHandle;
+	OnHorizontalAlignmentChanged = InArgs._OnHorizontalAlignmentChanged;
+	OnVerticalAlignmentChanged = InArgs._OnVerticalAlignmentChanged;
+
+	const SGridPanel::Layer CellBordersLayer(0);
+	const SGridPanel::Layer BackgroundLayer(1);
+	const SGridPanel::Layer ForegroundLayer(3);
+	
+	const FLinearColor CellBorderColor = FStyleColors::Foldout.GetSpecifiedColor();
+	const FLinearColor BackgroundColor = FStyleColors::Panel.GetSpecifiedColor();
+
+	HorizontalAlignmentPropertyHandle = TextAlignmentPropertyHandle.Get()->GetChildHandle(GET_MEMBER_NAME_CHECKED(FAvaTextAlignment, HorizontalAlignment));
+	VerticalAlignmentPropertyHandle = TextAlignmentPropertyHandle.Get()->GetChildHandle(GET_MEMBER_NAME_CHECKED(FAvaTextAlignment, VerticalAlignment));
+	
+	ChildSlot
+	[
+		SNew(SBox)
+		.Padding(FMargin(3.f, 10.f, 3.f, 10.f))
+		[
+			SNew(SGridPanel)
+			+ SGridPanel::Slot(0, 0, CellBordersLayer)
+			.RowSpan(2)
+			.ColumnSpan(5)
+			[
+				SNew(SColorBlock)
+				.Color(CellBorderColor)
+			]
+			+ SGridPanel::Slot(1, 0, ForegroundLayer)
+			[
+				GetHorizontalAlignmentButton(HorizontalLeftButton, EText3DHorizontalTextAlignment::Left, "AvalancheIcons.Alignment.Left", LOCTEXT("AlignSelectedTextLeft", "Align Text to the Left"))
+			]
+			+ SGridPanel::Slot(2, 0, ForegroundLayer)
+			[
+				GetHorizontalAlignmentButton(HorizontalCenterButton, EText3DHorizontalTextAlignment::Center, "AvalancheIcons.Alignment.Center_Y", LOCTEXT("AlignSelectedTextCenter", "Align Text to Center"))
+			]
+			+ SGridPanel::Slot(3, 0, ForegroundLayer)
+			[
+				GetHorizontalAlignmentButton(HorizontalRightButton, EText3DHorizontalTextAlignment::Right, "AvalancheIcons.Alignment.Right", LOCTEXT("AlignSelectedTextRight", "Align Text to the Right"))
+			]
+			+ SGridPanel::Slot(0, 0, BackgroundLayer) // color 1st empty slot
+			[
+				SNew(SColorBlock)
+				.Color(BackgroundColor)
+			]
+			+ SGridPanel::Slot(0, 1, ForegroundLayer)
+			[
+				GetVerticalAlignmentButton(VerticalFirstLineButton, EText3DVerticalTextAlignment::FirstLine, "AvalancheIcons.Alignment.Top", LOCTEXT("AlignSelectedTextFirstLine", "Align Text to First Line"))
+			]
+			+ SGridPanel::Slot(1, 1, ForegroundLayer)
+			[
+				GetVerticalAlignmentButton(VerticalTopButton, EText3DVerticalTextAlignment::Top, "AvalancheIcons.Alignment.Top", LOCTEXT("AlignSelectedTextTop", "Align Text to Top"))
+			]
+			+ SGridPanel::Slot(2, 1, ForegroundLayer)
+			[
+				GetVerticalAlignmentButton(VerticalCenterButton, EText3DVerticalTextAlignment::Center, "AvalancheIcons.Alignment.Center_Z", LOCTEXT("AlignSelectedTextCenter", "Align Text to Center"))
+			]
+			+ SGridPanel::Slot(3, 1, ForegroundLayer)
+			[
+				GetVerticalAlignmentButton(VerticalBottomButton, EText3DVerticalTextAlignment::Bottom, "AvalancheIcons.Alignment.Bottom", LOCTEXT("AlignSelectedTextRight", "Align Text to Bottom"))
+			]
+		]
+	];
+
+	OnAlignmentPropertyChanged();
+
+	TextAlignmentPropertyHandle.Get()->SetOnPropertyValueChanged(FSimpleDelegate::CreateSP(this, &SAvaTextAlignmentWidget::OnAlignmentPropertyChanged));
+}
+
+END_SLATE_FUNCTION_BUILD_OPTIMIZATION
+
+TSharedRef<SButton> SAvaTextAlignmentWidget::GetHorizontalAlignmentButton(TSharedPtr<SButton>& OutButton, EText3DHorizontalTextAlignment InHorizontalAlignment, FName Image, FText Tooltip)
+{
+	const ISlateStyle* AvalancheEditorStyle = FSlateStyleRegistry::FindSlateStyle("AvalancheEditor");
+
+	OutButton = SNew(SButton)
+		.ButtonStyle(&FAppStyle::Get().GetWidgetStyle<FButtonStyle>("SimpleButton"))
+		.ContentPadding(FMargin(2.f))
+		.OnClicked(this, &SAvaTextAlignmentWidget::ApplyHorizontalAlignment, InHorizontalAlignment)
+		.ToolTipText(Tooltip)
+		[
+			SNew(SImage)
+			.Image(AvalancheEditorStyle ? AvalancheEditorStyle->GetBrush(Image) : nullptr)
+			.DesiredSizeOverride(FVector2D(16.f, 16.f))
+			.ColorAndOpacity(this, &SAvaTextAlignmentWidget::ApplyHorizontalAlignmentButtonColorAndOpacity, InHorizontalAlignment)
+		];
+
+	return OutButton.ToSharedRef();
+}
+
+TSharedRef<SButton> SAvaTextAlignmentWidget::GetVerticalAlignmentButton(TSharedPtr<SButton>& OutButton, EText3DVerticalTextAlignment InVerticalAlignment, FName Image, FText Tooltip)
+{
+	const ISlateStyle* AvalancheEditorStyle = FSlateStyleRegistry::FindSlateStyle("AvalancheEditor");
+
+	OutButton = SNew(SButton)
+		.ButtonStyle(&FAppStyle::Get().GetWidgetStyle<FButtonStyle>("SimpleButton"))
+		.ContentPadding(FMargin(2.f))
+		.OnClicked(this, &SAvaTextAlignmentWidget::ApplyVerticalAlignment, InVerticalAlignment)
+		.ToolTipText(Tooltip)
+		[
+			SNew(SImage)
+			.Image(AvalancheEditorStyle ? AvalancheEditorStyle->GetBrush(Image) : nullptr)
+			.DesiredSizeOverride(FVector2D(16.f, 16.f))
+			.ColorAndOpacity(this, &SAvaTextAlignmentWidget::ApplyVerticalAlignmentButtonColorAndOpacity, InVerticalAlignment)
+		];
+
+	return OutButton.ToSharedRef();
+}
+
+FReply SAvaTextAlignmentWidget::ApplyHorizontalAlignment(EText3DHorizontalTextAlignment InHorizontalAlignment) const
+{
+	OnHorizontalAlignmentChanged.ExecuteIfBound(InHorizontalAlignment);
+
+	return FReply::Handled();
+}
+
+FReply SAvaTextAlignmentWidget::ApplyVerticalAlignment(EText3DVerticalTextAlignment InVerticalAlignment) const
+{
+	OnVerticalAlignmentChanged.ExecuteIfBound(InVerticalAlignment);
+
+	return FReply::Handled();
+}
+
+FSlateColor SAvaTextAlignmentWidget::ApplyHorizontalAlignmentButtonColorAndOpacity(EText3DHorizontalTextAlignment InHorizontalAlignment) const
+{
+	return InHorizontalAlignment == GetCurrentHorizontalAlignment() ?
+		FSlateColor(EStyleColor::AccentBlue).GetSpecifiedColor() :
+		FSlateColor(EStyleColor::Foreground).GetSpecifiedColor();
+}
+
+FSlateColor SAvaTextAlignmentWidget::ApplyVerticalAlignmentButtonColorAndOpacity(EText3DVerticalTextAlignment InVerticalAlignment) const
+{
+	return InVerticalAlignment == GetCurrentVerticalAlignment() ?
+		FSlateColor(EStyleColor::AccentBlue).GetSpecifiedColor() :
+		FSlateColor(EStyleColor::Foreground).GetSpecifiedColor();
+}
+
+void SAvaTextAlignmentWidget::OnAlignmentPropertyChanged() const
+{
+	OnHorizontalAlignmentChanged.ExecuteIfBound(GetCurrentHorizontalAlignment());
+	OnVerticalAlignmentChanged.ExecuteIfBound(GetCurrentVerticalAlignment());
+}
+
+EText3DHorizontalTextAlignment SAvaTextAlignmentWidget::GetCurrentHorizontalAlignment() const
+{
+	if (HorizontalAlignmentPropertyHandle)
+	{
+		uint8 HorizontalAlignAsUint8;
+		HorizontalAlignmentPropertyHandle->GetValue(HorizontalAlignAsUint8);
+		return static_cast<EText3DHorizontalTextAlignment>(HorizontalAlignAsUint8);
+	}
+
+	return EText3DHorizontalTextAlignment::Left;
+}
+
+EText3DVerticalTextAlignment SAvaTextAlignmentWidget::GetCurrentVerticalAlignment() const
+{
+	if (VerticalAlignmentPropertyHandle)
+	{
+		uint8 VerticalAlignAsUint8;
+		VerticalAlignmentPropertyHandle->GetValue(VerticalAlignAsUint8);
+		return static_cast<EText3DVerticalTextAlignment>(VerticalAlignAsUint8);
+	}
+
+	return EText3DVerticalTextAlignment::FirstLine;
+}
+
+#undef LOCTEXT_NAMESPACE
diff --git a/Engine/Plugins/Experimental/Avalanche/Source/AvalancheTextEditor/Private/Slate/SAvaTextAlignmentWidget.h b/Engine/Plugins/Experimental/Avalanche/Source/AvalancheTextEditor/Private/Slate/SAvaTextAlignmentWidget.h
new file mode 100644
--- /dev/null
+++ b/Engine/Plugins/Experimental/Avalanche/Source/AvalancheTextEditor/Private/Slate/SAvaTextAlignmentWidget.h
@@ -0,0 +1,65 @@
+// Copyright Epic Games, Inc. All Rights Reserved.
+
+#pragma once
+
+#include "Widgets/SCompoundWidget.h"
+
+enum class EText3DVerticalTextAlignment : uint8;
+enum class EText3DHorizontalTextAlignment : uint8;
+class IPropertyHandle;
+class SButton;
+
+DECLARE_DELEGATE_OneParam(FOnAvaTextHorizontalAlignmentChanged, EText3DHorizontalTextAlignment)
+DECLARE_DELEGATE_OneParam(FOnAvaTextVerticalAlignmentChanged, EText3DVerticalTextAlignment)
+
+class SAvaTextAlignmentWidget : public SCompoundWidget
+{
+public:
+	SLATE_BEGIN_ARGS(SAvaTextAlignmentWidget)
+		:_TextAlignmentPropertyHandle(),
+		_OnHorizontalAlignmentChanged(),
+		_OnVerticalAlignmentChanged()
+		{
+		}
+	SLATE_ATTRIBUTE(TSharedPtr<IPropertyHandle>, TextAlignmentPropertyHandle)
+	SLATE_EVENT(FOnAvaTextHorizontalAlignmentChanged, OnHorizontalAlignmentChanged)
+	SLATE_EVENT(FOnAvaTextVerticalAlignmentChanged, OnVerticalAlignmentChanged)
+
+	SLATE_END_ARGS()
+
+	TAttribute<TSharedPtr<IPropertyHandle>> TextAlignmentPropertyHandle;
+
+	void OnAlignmentPropertyChanged() const;
+	
+	/** Constructs this widget with InArgs */
+	void Construct(const FArguments& InArgs);
+
+protected:
+	FOnAvaTextHorizontalAlignmentChanged OnHorizontalAlignmentChanged;
+	FOnAvaTextVerticalAlignmentChanged OnVerticalAlignmentChanged;
+
+	TSharedRef<SButton> GetHorizontalAlignmentButton(TSharedPtr<SButton>& OutButton, EText3DHorizontalTextAlignment InHorizontalAlignment, FName Image, FText Tooltip);
+	TSharedRef<SButton> GetVerticalAlignmentButton(TSharedPtr<SButton>& OutButton, EText3DVerticalTextAlignment InVerticalAlignment, FName Image, FText Tooltip);
+	
+	FReply ApplyHorizontalAlignment(EText3DHorizontalTextAlignment InHorizontalAlignment) const;
+	FReply ApplyVerticalAlignment(EText3DVerticalTextAlignment InVerticalAlignment) const;
+
+	FSlateColor ApplyHorizontalAlignmentButtonColorAndOpacity(EText3DHorizontalTextAlignment InHorizontalTextAlignment) const;
+	FSlateColor ApplyVerticalAlignmentButtonColorAndOpacity(EText3DVerticalTextAlignment InVerticalAlignment) const;
+
+	EText3DHorizontalTextAlignment GetCurrentHorizontalAlignment() const;
+	EText3DVerticalTextAlignment GetCurrentVerticalAlignment() const;
+
+	TSharedPtr<SButton> HorizontalLeftButton;
+	TSharedPtr<SButton> HorizontalCenterButton;
+	TSharedPtr<SButton> HorizontalRightButton;
+
+	TSharedPtr<SButton> VerticalFirstLineButton;
+	TSharedPtr<SButton> VerticalTopButton;
+	TSharedPtr<SButton> VerticalCenterButton;
+	TSharedPtr<SButton> VerticalBottomButton;
+
+	TSharedPtr<IPropertyHandle> HorizontalAlignmentPropertyHandle;
+	TSharedPtr<IPropertyHandle> VerticalAlignmentPropertyHandle;
+};
+
diff --git a/Engine/Plugins/Experimental/Avalanche/Source/AvalancheTextEditor/Private/SystemFontLoading.h b/Engine/Plugins/Experimental/Avalanche/Source/AvalancheTextEditor/Private/SystemFontLoading.h
new file mode 100644
--- /dev/null
+++ b/Engine/Plugins/Experimental/Avalanche/Source/AvalancheTextEditor/Private/SystemFontLoading.h
@@ -0,0 +1,7 @@
+// Copyright Epic Games, Inc. All Rights Reserved.
+
+#pragma once
+
+#include "CoreTypes.h"
+
+#include COMPILED_PLATFORM_HEADER_WITH_PREFIX(Platforms, SystemFontLoading.h)
diff --git a/Engine/Plugins/Experimental/Avalanche/Source/AvalancheTextEditor/Private/Tool/AvaTextActorTool.cpp b/Engine/Plugins/Experimental/Avalanche/Source/AvalancheTextEditor/Private/Tool/AvaTextActorTool.cpp
new file mode 100644
--- /dev/null
+++ b/Engine/Plugins/Experimental/Avalanche/Source/AvalancheTextEditor/Private/Tool/AvaTextActorTool.cpp
@@ -0,0 +1,38 @@
+// Copyright Epic Games, Inc. All Rights Reserved.
+
+#include "Tool/AvaTextActorTool.h"
+#include "AvaInteractiveToolsSettings.h"
+#include "AvaTextActor.h"
+#include "AvaTextEditorCommands.h"
+#include "Builders/AvaInteractiveToolsToolBuilder.h"
+#include "IAvalancheComponentVisualizersModule.h"
+
+UAvaTextActorTool::UAvaTextActorTool()
+{
+	ActorClass = AAvaTextActor::StaticClass();
+}
+
+bool UAvaTextActorTool::UseIdentityRotation() const
+{
+	return ConditionalIdentityRotation();
+}
+
+FName UAvaTextActorTool::GetCategoryName()
+{
+	return IAvalancheInteractiveToolsModule::CategoryNameActor;
+}
+
+FAvaInteractiveToolsToolParameters UAvaTextActorTool::GetToolParameters() const
+{
+	return {
+		FAvaTextEditorCommands::Get().Tool_Actor_Text3D,
+		TEXT("Text Actor Tool"),
+		2000,
+		FAvalancheInteractiveToolsCreateBuilder::CreateLambda(
+			[](UEdMode* InEdMode)
+			{
+				return UAvaInteractiveToolsToolBuilder::CreateToolBuilder<UAvaTextActorTool>(InEdMode);
+			}),
+		ActorClass
+	};
+}
diff --git a/Engine/Plugins/Experimental/Avalanche/Source/AvalancheTextEditor/Private/Tool/AvaTextActorTool.h b/Engine/Plugins/Experimental/Avalanche/Source/AvalancheTextEditor/Private/Tool/AvaTextActorTool.h
new file mode 100644
--- /dev/null
+++ b/Engine/Plugins/Experimental/Avalanche/Source/AvalancheTextEditor/Private/Tool/AvaTextActorTool.h
@@ -0,0 +1,22 @@
+// Copyright Epic Games, Inc. All Rights Reserved.
+
+#pragma once
+
+#include "Tools/AvaInteractiveToolsActorPointToolBase.h"
+#include "AvaTextActorTool.generated.h"
+
+UCLASS()
+class UAvaTextActorTool : public UAvaInteractiveToolsActorPointToolBase
+{
+	GENERATED_BODY()
+
+public:
+	UAvaTextActorTool();
+
+	//~ Begin UAvaInteractiveToolsToolBase
+	virtual bool UseIdentityLocation() const override { return false; }
+	virtual bool UseIdentityRotation() const override;
+	virtual FName GetCategoryName() override;
+	virtual FAvaInteractiveToolsToolParameters GetToolParameters() const override;
+	//~ End UAvaInteractiveToolsToolBase
+};
diff --git a/Engine/Plugins/Experimental/Avalanche/Source/AvalancheTextEditor/Private/Visualizer/AvaTextVisualizer.cpp b/Engine/Plugins/Experimental/Avalanche/Source/AvalancheTextEditor/Private/Visualizer/AvaTextVisualizer.cpp
new file mode 100644
--- /dev/null
+++ b/Engine/Plugins/Experimental/Avalanche/Source/AvalancheTextEditor/Private/Visualizer/AvaTextVisualizer.cpp
@@ -0,0 +1,1408 @@
+// Copyright Epic Games, Inc. All Rights Reserved.
+
+#include "AvaTextVisualizer.h"
+#include "AvaField.h"
+#include "AvaText3DComponent.h"
+#include "AvalancheShapesEditorModule.h"
+#include "EditorViewportClient.h"
+#include "Engine/Texture2D.h"
+#include "IAvalancheComponentVisualizersModule.h"
+#include "LevelEditor/AvaLevelEditorUtils.h"
+#include "Modules/ModuleManager.h"
+#include "PropertyEditorModule.h"
+#include "ScopedTransaction.h"
+#include "Text3DComponent.h"
+#include "TextureResource.h"
+
+IMPLEMENT_HIT_PROXY(HAvaTextMaxTextHeightProxy, HAvaHitProxy)
+IMPLEMENT_HIT_PROXY(HAvaTextMaxTextHeightHandleProxy, HAvaHitProxy)
+IMPLEMENT_HIT_PROXY(HAvaTextMaxTextWidthProxy, HAvaHitProxy)
+IMPLEMENT_HIT_PROXY(HAvaTextMaxTextWidthHandleProxy, HAvaHitProxy)
+IMPLEMENT_HIT_PROXY(HAvaTextScaleProportionallyProxy, HAvaHitProxy)
+
+IMPLEMENT_HIT_PROXY(HAvaTextEditGradientProxy, HAvaHitProxy)
+IMPLEMENT_HIT_PROXY(HAvaTextGradientLineStartHandleProxy, HAvaHitProxy)
+IMPLEMENT_HIT_PROXY(HAvaTextGradientLineEndHandleProxy, HAvaHitProxy)
+IMPLEMENT_HIT_PROXY(HAvaTextGradientCenterHandleProxy, HAvaHitProxy)
+IMPLEMENT_HIT_PROXY(HAvaTextGradientSmoothnessHandleProxy, HAvaHitProxy)
+
+namespace UE::Avalanche::Private
+{
+	struct FTextVisualizerStatics
+	{
+		inline static constexpr float GradientRotHandleScale = 0.005f;
+		inline static constexpr float GradientOffsetHandleScale = 0.01f;
+		inline static constexpr float GradientSmoothnessHandleScale = 0.1f;
+		inline static constexpr float GradientHandleMaxLength = 50.0f;
+		inline static constexpr float GradientSmoothnessHandleOffset = -10;
+	};
+}
+
+#define LOCTEXT_NAMESPACE "AvaTextVisualizer"
+
+FAvaTextVisualizer::FAvaTextVisualizer()
+	: FAvaVisualizerBase()
+{
+	using namespace UE::AvaCore;
+
+	HasMaxHeightProperty        = GetProperty<UAvaText3DComponent>(GET_MEMBER_NAME_CHECKED(UAvaText3DComponent, bHasMaxHeight));
+	HasMaxWidthProperty         = GetProperty<UAvaText3DComponent>(GET_MEMBER_NAME_CHECKED(UAvaText3DComponent, bHasMaxWidth));
+	MaxHeightProperty           = GetProperty<UAvaText3DComponent>(GET_MEMBER_NAME_CHECKED(UAvaText3DComponent, MaxHeight));
+	MaxWidthProperty            = GetProperty<UAvaText3DComponent>(GET_MEMBER_NAME_CHECKED(UAvaText3DComponent, MaxWidth));
+	ScaleProportionallyProperty = GetProperty<UAvaText3DComponent>(GET_MEMBER_NAME_CHECKED(UAvaText3DComponent, bScaleProportionally));
+	GradientProperty            = GetProperty<UAvaText3DComponent>(GET_MEMBER_NAME_CHECKED(UAvaText3DComponent, GradientSettings));
+
+	ResetEditingFlags();
+	bShowGradientControls = false;
+}
+
+UActorComponent* FAvaTextVisualizer::GetEditedComponent() const
+{
+	return TextComponent.Get();
+}
+
+TMap<UObject*, TArray<FProperty*>> FAvaTextVisualizer::GatherEditableProperties(UObject* InObject) const
+{
+	if (UAvaText3DComponent* Component = Cast<UAvaText3DComponent>(InObject))
+	{
+		return {{
+			Component,
+			{
+				HasMaxHeightProperty,
+				MaxWidthProperty,
+				HasMaxHeightProperty,
+				MaxHeightProperty,
+				ScaleProportionallyProperty,
+				GradientProperty
+			}
+		}};
+	}
+
+	return {};
+}
+
+bool FAvaTextVisualizer::VisProxyHandleClick(FEditorViewportClient* InViewportClient, HComponentVisProxy* VisProxy, const FViewportClick& Click)
+{
+	if (Click.GetKey() != EKeys::LeftMouseButton || !VisProxy)
+	{
+		EndEditing();
+		return Super::VisProxyHandleClick(InViewportClient, VisProxy, Click);
+	}	
+	
+	if (VisProxy->IsA(HAvaTextMaxTextHeightProxy::StaticGetType()))
+	{
+		if (const HAvaTextMaxTextHeightProxy* const MaxHeightProxy = static_cast<HAvaTextMaxTextHeightProxy*>(VisProxy))
+		{
+			EndEditing();
+
+			if (UAvaText3DComponent* const AvaTextComponent = Cast<UAvaText3DComponent>(const_cast<UActorComponent*>(MaxHeightProxy->Component.Get())))
+			{
+				FScopedTransaction Transaction(LOCTEXT("ToggleText3DHasMaxHeight", "Toggle Text 3D Has Max Height"));
+
+				if (AActor* const Text3DActor = AvaTextComponent->GetOwner())
+				{
+					Text3DActor->Modify();
+					AvaTextComponent->Modify();
+
+					// If max height is off, we will toggle it on later.
+					if (!AvaTextComponent->HasMaxHeight() && AvaTextComponent->GetMaxHeight() == GetDefault<UAvaText3DComponent>()->GetMaxHeight())
+					{
+						// Set a sensible default value
+						Bounds = GetComponentBounds(AvaTextComponent);
+						float NewMaxHeight = (Bounds.Max.Z - Bounds.Min.Z) * 1.25f;
+						AvaTextComponent->SetMaxHeight(NewMaxHeight);
+					}					
+
+					AvaTextComponent->SetHasMaxHeight(!AvaTextComponent->HasMaxHeight());
+				}
+			}
+
+			return true;
+		}
+	}
+	else if (VisProxy->IsA(HAvaTextMaxTextHeightHandleProxy::StaticGetType()))
+	{
+		if (const HAvaTextMaxTextHeightHandleProxy* const MaxHeightHandleProxy = static_cast<HAvaTextMaxTextHeightHandleProxy*>(VisProxy))
+		{
+			EndEditing();
+
+			if (UAvaText3DComponent* AvaTextComponent = Cast<UAvaText3DComponent>(const_cast<UActorComponent*>(MaxHeightHandleProxy->Component.Get())))
+			{
+				TextComponent = AvaTextComponent;
+				bEditingHeight = true;
+				StartEditing(InViewportClient, AvaTextComponent);
+			}
+
+			return true;
+		}
+	}
+
+	if (VisProxy->IsA(HAvaTextMaxTextWidthProxy::StaticGetType()))
+	{
+		if (const HAvaTextMaxTextWidthProxy* const MaxWidthProxy = static_cast<HAvaTextMaxTextWidthProxy*>(VisProxy))
+		{
+			EndEditing();
+
+			if (UAvaText3DComponent* const AvaTextComponent = Cast<UAvaText3DComponent>(const_cast<UActorComponent*>(MaxWidthProxy->Component.Get())))
+			{
+				FScopedTransaction Transaction(LOCTEXT("ToggleText3DHasMaxWidth", "Toggle Text 3D Has Max Width"));
+
+				if (AActor* const Text3DActor = AvaTextComponent->GetOwner())
+				{
+					Text3DActor->Modify();
+					AvaTextComponent->Modify();
+
+					// If max width is off, we will toggle it on later.
+					if (!AvaTextComponent->HasMaxWidth() && AvaTextComponent->GetMaxWidth() == GetDefault<UAvaText3DComponent>()->GetMaxWidth())
+					{
+						// Set a sensible default value
+						Bounds = GetComponentBounds(AvaTextComponent);
+						float NewMaxWidth = (Bounds.Max.Y - Bounds.Min.Y) * 1.25f;
+						AvaTextComponent->SetMaxWidth(NewMaxWidth);
+					}
+
+					AvaTextComponent->SetHasMaxWidth(!AvaTextComponent->HasMaxWidth());
+				}
+			}
+
+			return true;
+		}
+	}
+	else if (VisProxy->IsA(HAvaTextMaxTextWidthHandleProxy::StaticGetType()))
+	{
+		if (const HAvaTextMaxTextWidthHandleProxy* const MaxWidthHandleProxy = static_cast<HAvaTextMaxTextWidthHandleProxy*>(VisProxy))
+		{
+			EndEditing();
+
+			if (UAvaText3DComponent* const AvaTextComponent = Cast<UAvaText3DComponent>(const_cast<UActorComponent*>(MaxWidthHandleProxy->Component.Get())))
+			{
+				TextComponent = AvaTextComponent;
+				bEditingWidth = true;
+				StartEditing(InViewportClient, AvaTextComponent);
+			}
+
+			return true;
+		}
+	}
+	else if (VisProxy->IsA(HAvaTextScaleProportionallyProxy::StaticGetType()))
+	{
+		if (const HAvaTextScaleProportionallyProxy* const ScaleProportionallyProxy = static_cast<HAvaTextScaleProportionallyProxy*>(VisProxy))
+		{
+			EndEditing();
+
+			if (UAvaText3DComponent* const AvaTextComponent = Cast<UAvaText3DComponent>(const_cast<UActorComponent*>(ScaleProportionallyProxy->Component.Get())))
+			{
+				FScopedTransaction Transaction(LOCTEXT("ToggleText3DScaleProportionally", "Toggle Text 3D Scale Proportionally"));
+
+				if (AActor* const Text3DActor = AvaTextComponent->GetOwner())
+				{
+					Text3DActor->Modify();
+					AvaTextComponent->Modify();
+					AvaTextComponent->SetScaleProportionally(!AvaTextComponent->ScalesProportionally());
+				}
+			}
+			
+			return true;
+		}
+	}
+
+	// this one just toggles gradient control on/off, in case the actor has it
+	if (VisProxy->IsA(HAvaTextEditGradientProxy::StaticGetType()))
+	{
+		EndEditing();			
+		bShowGradientControls = !bShowGradientControls;
+
+		return true;
+	}
+	else if (VisProxy->IsA(HAvaTextGradientLineStartHandleProxy::StaticGetType()))
+	{
+		if (const HAvaTextGradientLineStartHandleProxy* const GradientStartHandleProxy = static_cast<HAvaTextGradientLineStartHandleProxy*>(VisProxy))
+		{
+			EndEditing();
+
+			if (UAvaText3DComponent* AvaTextComponent = Cast<UAvaText3DComponent>(const_cast<UActorComponent*>(GradientStartHandleProxy->Component.Get())))
+			{
+				TextComponent = AvaTextComponent;
+				bEditingGradientRotation_StartHandle = true;
+				GradientEditBeginLocation_StartHandle = GetGradientStartHandleLocation(AvaTextComponent);
+				GradientEditBeginLocation_Center = GetGradientCenterHandleLocation(AvaTextComponent);
+				StartEditing(InViewportClient, AvaTextComponent);
+			}
+
+			return true;
+		}
+	}
+	else if (VisProxy->IsA(HAvaTextGradientLineEndHandleProxy::StaticGetType()))
+	{
+		if (const HAvaTextGradientLineEndHandleProxy* const GradientEndHandleProxy = static_cast<HAvaTextGradientLineEndHandleProxy*>(VisProxy))
+		{
+			EndEditing();
+
+			if (UAvaText3DComponent* AvaTextComponent = Cast<UAvaText3DComponent>(const_cast<UActorComponent*>(GradientEndHandleProxy->Component.Get())))
+			{
+				TextComponent = AvaTextComponent;
+				bEditingGradientRotation_EndHandle = true;
+				GradientEditBeginLocation_EndHandle = GetGradientEndHandleLocation(AvaTextComponent);
+				GradientEditBeginLocation_Center = GetGradientCenterHandleLocation(AvaTextComponent);
+				StartEditing(InViewportClient, AvaTextComponent);
+			}
+
+			return true;
+		}
+	}
+	else if (VisProxy->IsA(HAvaTextGradientCenterHandleProxy::StaticGetType()))
+	{
+		if (const HAvaTextGradientCenterHandleProxy* const GradientCenterHandleProxy = static_cast<HAvaTextGradientCenterHandleProxy*>(VisProxy))
+		{
+			EndEditing();
+
+			if (UAvaText3DComponent* AvaTextComponent = Cast<UAvaText3DComponent>(const_cast<UActorComponent*>(GradientCenterHandleProxy->Component.Get())))
+			{
+				TextComponent = AvaTextComponent;
+				bEditingGradientOffset = true;
+				StartEditing(InViewportClient, AvaTextComponent);
+			}
+
+			return true;
+		}
+	}
+	else if (VisProxy->IsA(HAvaTextGradientSmoothnessHandleProxy::StaticGetType()))
+	{
+		if (const HAvaTextGradientSmoothnessHandleProxy* const GradientSmoothnessHandleProxy = static_cast<HAvaTextGradientSmoothnessHandleProxy*>(VisProxy))
+		{
+			EndEditing();
+
+			if (UAvaText3DComponent* AvaTextComponent = Cast<UAvaText3DComponent>(const_cast<UActorComponent*>(GradientSmoothnessHandleProxy->Component.Get())))
+			{
+				TextComponent = AvaTextComponent;
+				bEditingGradientSmoothness = true;
+				StartEditing(InViewportClient, AvaTextComponent);
+				return true;
+			}			
+		}
+	}
+
+	return false;
+}
+
+FVector FAvaTextVisualizer::GetGradientSmoothnessHandleLocation(const UAvaText3DComponent* InAvaTextComponent) const
+{
+	using namespace UE::Avalanche::Private;
+	
+	return GetGradientCenterHandleLocation(InAvaTextComponent) + FVector(0, FTextVisualizerStatics::GradientSmoothnessHandleOffset, 0);
+}
+
+bool FAvaTextVisualizer::GetWidgetLocation(const FEditorViewportClient* ViewportClient, FVector& OutLocation) const
+{
+	if (const UAvaText3DComponent* const AvaTextComponent = Cast<UAvaText3DComponent>(GetEditedComponent()))
+	{
+		if (bEditingWidth)
+		{
+			OutLocation = GetWidthHandleLocation(AvaTextComponent);
+			return true;
+		}
+
+		if (bEditingHeight)
+		{
+			OutLocation = GetHeightHandleLocation(AvaTextComponent);
+			return true;
+		}
+
+		if (bEditingGradientRotation_StartHandle)
+		{
+			OutLocation = GetGradientStartHandleLocation(AvaTextComponent);
+			return true;
+		}
+		
+		if (bEditingGradientRotation_EndHandle)
+		{
+			OutLocation = GetGradientEndHandleLocation(AvaTextComponent);
+			return true;
+		}
+		
+		if (bEditingGradientOffset)
+		{
+			OutLocation = GetGradientCenterHandleLocation(AvaTextComponent);
+			return true;
+		}
+
+		if (bEditingGradientSmoothness)
+		{
+			OutLocation = GetGradientSmoothnessHandleLocation(AvaTextComponent);
+			return true;
+		}		
+	}
+
+	return Super::GetWidgetLocation(ViewportClient, OutLocation);
+}
+
+bool FAvaTextVisualizer::GetWidgetMode(const FEditorViewportClient* ViewportClient, UE::Widget::EWidgetMode& Mode) const
+{
+	if (bEditingWidth || bEditingHeight)
+	{
+		Mode = UE::Widget::WM_Translate;
+		return true;
+	}
+
+	if (bEditingGradientRotation_EndHandle || bEditingGradientRotation_StartHandle)
+	{
+		Mode = UE::Widget::WM_Translate;
+		return true;
+	}
+
+	if (bEditingGradientOffset)
+	{
+		Mode = UE::Widget::WM_Translate;
+		return true;
+	}
+
+	if (bEditingGradientSmoothness)
+	{
+		Mode = UE::Widget::WM_Scale;
+		return true;
+	}
+
+	return Super::GetWidgetMode(ViewportClient, Mode);
+}
+
+bool FAvaTextVisualizer::GetWidgetAxisList(const FEditorViewportClient* ViewportClient, UE::Widget::EWidgetMode WidgetMode, 
+	EAxisList::Type& AxisList) const
+{
+	if (bEditingWidth)
+	{
+		AxisList = EAxisList::Y;
+		return true;
+	}
+
+	if (bEditingHeight)
+	{
+		AxisList = EAxisList::Z;
+		return true;
+	}
+
+	if (bEditingGradientRotation_EndHandle || bEditingGradientRotation_StartHandle)
+	{
+		AxisList = EAxisList::Type::Screen;
+		return true;
+	}
+
+	if (bEditingGradientOffset)
+	{
+		AxisList = EAxisList::Type::Screen;
+		return true;
+	}
+
+	if (bEditingGradientSmoothness)
+	{
+		AxisList = EAxisList::Type::Z;
+		return true;
+	}
+
+	return Super::GetWidgetAxisList(ViewportClient, WidgetMode, AxisList);
+}
+
+bool FAvaTextVisualizer::GetWidgetAxisListDragOverride(const FEditorViewportClient* ViewportClient,
+	UE::Widget::EWidgetMode WidgetMode, EAxisList::Type& AxisList) const
+{
+	if (bEditingGradientOffset)
+	{
+		AxisList = EAxisList::YZ;
+		return true;
+	}
+
+	if (bEditingGradientRotation_EndHandle || bEditingGradientRotation_StartHandle)
+	{
+		AxisList = EAxisList::YZ;
+		return true;
+	}
+
+	return Super::GetWidgetAxisListDragOverride(ViewportClient, WidgetMode, AxisList);
+}
+
+bool FAvaTextVisualizer::HandleInputDeltaInternal(FEditorViewportClient* InViewportClient, FViewport* InViewport, const FVector& InAccumulatedTranslation,
+	const FRotator& InAccumulatedRotation, const FVector& InAccumulatedScale)
+{
+	using namespace UE::Avalanche::Private;
+	
+	if (bEditingWidth)
+	{
+		if (UAvaText3DComponent* const AvaTextComponent = Cast<UAvaText3DComponent>(GetEditedComponent()))
+		{
+			if (GetViewportWidgetMode(InViewportClient) == UE::Widget::WM_Translate)
+			{
+				if (GetViewportWidgetAxisList(InViewportClient) & EAxisList::Y)
+				{
+					float MaxWidth = InitialMaxWidth;
+
+					switch (AvaTextComponent->GetHorizontalAlignment())
+					{
+						case EText3DHorizontalTextAlignment::Left:
+							MaxWidth += InAccumulatedTranslation.Y;
+							break;
+
+						case EText3DHorizontalTextAlignment::Center:
+							MaxWidth += InAccumulatedTranslation.Y * 2.f;
+							break;
+
+						case EText3DHorizontalTextAlignment::Right:
+							MaxWidth -= InAccumulatedTranslation.Y;
+							break;
+					}
+
+					MaxWidth = FMath::Max(1.f, MaxWidth);
+
+					if (MaxWidth != InitialMaxWidth)
+					{
+						bHasBeenModified = true;
+					}
+
+					if (AActor* const Text3DActor = AvaTextComponent->GetOwner())
+					{
+						Text3DActor->Modify();
+						AvaTextComponent->Modify();
+						AvaTextComponent->SetMaxWidth(MaxWidth);
+					}
+
+					return true;
+				}
+			}
+		}
+
+		EndEditing();
+		return true;
+	}
+
+	if (bEditingHeight)
+	{
+		if (UAvaText3DComponent* const AvaTextComponent = Cast<UAvaText3DComponent>(GetEditedComponent()))
+		{
+			if (GetViewportWidgetMode(InViewportClient) == UE::Widget::WM_Translate)
+			{
+				if (GetViewportWidgetAxisList(InViewportClient) & EAxisList::Z)
+				{
+					float MaxHeight = InitialMaxHeight;
+
+					switch (AvaTextComponent->GetVerticalAlignment())
+					{
+						case EText3DVerticalTextAlignment::FirstLine:
+							MaxHeight += InAccumulatedTranslation.Z / TopHeightFraction;
+							break;
+
+						case EText3DVerticalTextAlignment::Top:
+							MaxHeight -= InAccumulatedTranslation.Z;
+							break;
+
+						case EText3DVerticalTextAlignment::Center:
+							MaxHeight += InAccumulatedTranslation.Z * 2.f;
+							break;
+
+						case EText3DVerticalTextAlignment::Bottom:
+							MaxHeight += InAccumulatedTranslation.Z;
+							break;
+					}
+
+					MaxHeight = FMath::Max(1.f, MaxHeight);
+
+					if (MaxHeight != InitialMaxHeight)
+					{
+						bHasBeenModified = true;
+					}
+
+					if (AActor* const Text3DActor = AvaTextComponent->GetOwner())
+					{
+						Text3DActor->Modify();
+						AvaTextComponent->Modify();
+						AvaTextComponent->SetMaxHeight(MaxHeight);
+					}
+
+					return true;
+				}
+			}
+		}
+
+		EndEditing();
+		return true;
+	}
+
+	if (bEditingGradientRotation_EndHandle || bEditingGradientRotation_StartHandle)
+	{
+		if (UAvaText3DComponent* const AvaTextComponent = Cast<UAvaText3DComponent>(GetEditedComponent()))
+		{
+			if (GetViewportWidgetMode(InViewportClient) == UE::Widget::WM_Translate)
+			{
+				if (GetViewportWidgetAxisList(InViewportClient) & EAxisList::Z)
+				{
+					float GradientRot = InitialGradientRotation;
+					float RotIncrZ = InAccumulatedTranslation.Z * FTextVisualizerStatics::GradientRotHandleScale;
+
+					FVector HandleToCenter;
+
+					if (bEditingGradientRotation_EndHandle)
+					{
+					    HandleToCenter = GradientEditBeginLocation_EndHandle - GradientEditBeginLocation_Center;
+					}
+					else if (bEditingGradientRotation_StartHandle)
+					{
+					    HandleToCenter = GradientEditBeginLocation_StartHandle - GradientEditBeginLocation_Center;
+					}
+					
+					if (HandleToCenter.Y < 0)
+                    {
+                        RotIncrZ *= -1;
+                    }
+					
+					GradientRot -= RotIncrZ;
+					
+					if (GradientRot != InitialGradientRotation)
+					{
+						bHasBeenModified = true;
+					}
+
+					if (AActor* const Text3DActor = AvaTextComponent->GetOwner())
+					{
+						Text3DActor->Modify();
+						AvaTextComponent->Modify();
+
+						if (AvaTextComponent->GradientSettings.Direction != EAvaGradientDirection::Custom)
+						{
+							if (AvaTextComponent->GradientSettings.Direction == EAvaGradientDirection::Vertical)
+							{
+								AvaTextComponent->GradientSettings.Rotation = 0;
+							}
+							else if (AvaTextComponent->GradientSettings.Direction == EAvaGradientDirection::Horizontal)
+							{
+								AvaTextComponent->GradientSettings.Rotation = 0.75;
+							}
+
+							InitialGradientRotation = AvaTextComponent->GradientSettings.Rotation;
+							
+							AvaTextComponent->GradientSettings.Direction = EAvaGradientDirection::Custom;
+
+							const FPropertyEditorModule& PropertyModule = FModuleManager::LoadModuleChecked<FPropertyEditorModule>("PropertyEditor");
+
+							for (const FName& DetailViewId : FAvaLevelEditorUtils::GetDetailsViewNames())
+							{
+								if (const TSharedPtr<IDetailsView> DetailView = PropertyModule.FindDetailView(DetailViewId))
+								{
+									DetailView->ForceRefresh();
+								}
+							}
+						}
+						else
+						{
+							AvaTextComponent->GradientSettings.Rotation = GradientRot;
+						}
+
+						AvaTextComponent->RefreshGradient();
+					}
+					return true;
+				}
+			}
+		}
+		
+		EndEditing();
+		return true;
+	}
+
+	if (bEditingGradientOffset)
+	{
+		if (UAvaText3DComponent* const AvaTextComponent = Cast<UAvaText3DComponent>(GetEditedComponent()))
+		{
+			if (GetViewportWidgetMode(InViewportClient) == UE::Widget::WM_Translate)
+			{
+				if (GetViewportWidgetAxisList(InViewportClient) & EAxisList::YZ)
+				{
+					if (AActor* const Text3DActor = AvaTextComponent->GetOwner())
+					{
+						Text3DActor->Modify();
+						AvaTextComponent->Modify();
+
+						float GradientOffset = InitialGradientOffset;
+						
+						const float OffsetIncr = InAccumulatedTranslation.Length() * FTextVisualizerStatics::GradientOffsetHandleScale;
+
+						FVector GradientDir = Text3DActor->GetActorUpVector().RotateAngleAxis(-AvaTextComponent->GradientSettings.Rotation * 360, Text3DActor->GetActorForwardVector());
+						GradientDir.Normalize();
+
+						const float DeltaDotGradDir = InAccumulatedTranslation.Dot(GradientDir);
+						
+						GradientOffset -= OffsetIncr * FMath::Sign(DeltaDotGradDir);
+						GradientOffset = FMath::Clamp(GradientOffset, 0.0, 1.0);
+
+						if (GradientOffset != InitialGradientOffset)
+						{
+							bHasBeenModified = true;
+						}
+						
+						AvaTextComponent->GradientSettings.Offset = GradientOffset;
+						AvaTextComponent->RefreshGradient();
+						return true;
+					}
+				}
+			}
+		}
+
+		EndEditing();
+		return true;
+	}
+
+	if (bEditingGradientSmoothness)
+	{
+		if (UAvaText3DComponent* const AvaTextComponent = Cast<UAvaText3DComponent>(GetEditedComponent()))
+		{
+			if (GetViewportWidgetMode(InViewportClient) == UE::Widget::WM_Scale)
+			{
+				if (GetViewportWidgetAxisList(InViewportClient) & EAxisList::Z)
+				{
+					if (AActor* const Text3DActor = AvaTextComponent->GetOwner())
+					{
+						Text3DActor->Modify();
+						AvaTextComponent->Modify();
+
+						float GradientSmoothness = InitialGradientSmoothness;
+						
+						const float SmoothnessIncr = InAccumulatedScale.Z * FTextVisualizerStatics::GradientSmoothnessHandleScale;
+						
+						GradientSmoothness += SmoothnessIncr;
+						GradientSmoothness = FMath::Clamp(GradientSmoothness, 0.0, 1.0);
+
+						if (GradientSmoothness != InitialGradientSmoothness)
+						{
+							bHasBeenModified = true;
+						}
+
+						AvaTextComponent->GradientSettings.Smoothness = GradientSmoothness;
+						AvaTextComponent->RefreshGradient();
+						return true;
+					}
+				}
+			}
+		}
+
+		EndEditing();
+		return true;
+	}
+
+	return Super::HandleInputDeltaInternal(InViewportClient, InViewport, InAccumulatedTranslation, InAccumulatedRotation, InAccumulatedScale);
+}
+
+void FAvaTextVisualizer::TrackingStopped(FEditorViewportClient* InViewportClient, bool bInDidMove)
+{
+	if (bEditingGradientRotation_EndHandle || bEditingGradientRotation_StartHandle)
+	{
+		if (const UAvaText3DComponent* const AvaTextComponent = Cast<UAvaText3DComponent>(GetEditedComponent()))
+		{
+			// handles positions need to be refreshed every time we release the mouse, so that clockwise vs. counterclockwise interaction works properly
+			GradientEditBeginLocation_StartHandle = GetGradientStartHandleLocation(AvaTextComponent);
+			GradientEditBeginLocation_EndHandle = GetGradientEndHandleLocation(AvaTextComponent);
+			GradientEditBeginLocation_Center = GetGradientCenterHandleLocation(AvaTextComponent);
+		}
+	}
+
+	FAvaVisualizerBase::TrackingStopped(InViewportClient, bInDidMove);
+}
+
+bool FAvaTextVisualizer::ResetValue(FEditorViewportClient* InViewportClient, HHitProxy* HitProxy)
+{
+	if (!HitProxy->IsA(HAvaTextMaxTextHeightProxy::StaticGetType())
+		&& !HitProxy->IsA(HAvaTextMaxTextWidthProxy::StaticGetType())
+		&& !HitProxy->IsA(HAvaTextScaleProportionallyProxy::StaticGetType())
+		&& !HitProxy->IsA(HAvaTextEditGradientProxy::StaticGetType()))
+	{
+		return Super::ResetValue(InViewportClient, HitProxy);
+	}
+
+	UAvaText3DComponent* AvaTextComponent = TextComponent.Get();
+
+	if (!AvaTextComponent)
+	{
+		return Super::ResetValue(InViewportClient, HitProxy);
+	}
+
+	AvaTextComponent->SetFlags(RF_Transactional);
+
+	if (HitProxy->IsA(HAvaTextMaxTextWidthProxy::StaticGetType()))
+	{
+		FScopedTransaction Transaction(NSLOCTEXT("AvaText3DVisualizer", "VisualizerResetValue", "Visualizer Reset Value"));
+		AvaTextComponent->Modify();
+
+		AvaTextComponent->SetHasMaxWidth(false);
+		AvaTextComponent->SetMaxWidth(GetDefault<UAvaText3DComponent>()->MaxWidth);
+
+		FAvaVisualizerBase::NotifyPropertyModified(AvaTextComponent, MaxWidthProperty, EPropertyChangeType::ValueSet);
+	}
+	else if (HitProxy->IsA(HAvaTextMaxTextHeightProxy::StaticGetType()))
+	{
+		FScopedTransaction Transaction(NSLOCTEXT("AvaTextVisualizer", "VisualizerResetValue", "Visualizer Reset Value"));
+		AvaTextComponent->Modify();
+
+		AvaTextComponent->SetHasMaxHeight(false);
+		AvaTextComponent->SetMaxHeight(GetDefault<UAvaText3DComponent>()->MaxHeight);
+
+		FAvaVisualizerBase::NotifyPropertyModified(AvaTextComponent, MaxHeightProperty, EPropertyChangeType::ValueSet);
+	}
+	else if (HitProxy->IsA(HAvaTextScaleProportionallyProxy::StaticGetType()))
+	{
+		FScopedTransaction Transaction(NSLOCTEXT("AvaTextVisualizer", "VisualizerResetValue", "Visualizer Reset Value"));
+		AvaTextComponent->Modify();
+
+		AvaTextComponent->SetScaleProportionally(false);
+
+		FAvaVisualizerBase::NotifyPropertyModified(AvaTextComponent, ScaleProportionallyProperty, EPropertyChangeType::ValueSet);
+	}
+
+	return true;
+}
+
+bool FAvaTextVisualizer::IsEditing() const
+{
+	if (bEditingWidth || bEditingHeight ||
+		bEditingGradientRotation_StartHandle || bEditingGradientRotation_EndHandle || bEditingGradientOffset || bEditingGradientSmoothness)
+	{
+		return true;
+	}
+
+	return Super::IsEditing();
+}
+
+void FAvaTextVisualizer::ResetEditingFlags()
+{
+	bEditingWidth = false;
+	bEditingHeight = false;	
+	bEditingGradientRotation_StartHandle = false;
+	bEditingGradientRotation_EndHandle = false;
+	bEditingGradientOffset = false;
+	bEditingGradientSmoothness = false;
+}
+
+
+void FAvaTextVisualizer::EndEditing()
+{
+	Super::EndEditing();
+	
+	ResetEditingFlags();
+}
+
+FVector FAvaTextVisualizer::GetWidthHandleLocation(const UText3DComponent* InText3DComp) const
+{
+	const FTransform Transform = InText3DComp->GetComponentTransform();
+	const FBox BoundsMax = GetBoundsMax(InText3DComp);
+	const float CoordZ = (BoundsMax.Min.Z + BoundsMax.Max.Z) * 0.5f;
+	FVector IconLocation;
+
+	switch (InText3DComp->GetHorizontalAlignment())
+	{
+		default:
+		case EText3DHorizontalTextAlignment::Left:
+			IconLocation = Transform.TransformPositionNoScale(FVector(0.f, BoundsMax.Max.Y + 10.f, CoordZ));
+			break;
+
+		case EText3DHorizontalTextAlignment::Center:
+			IconLocation = Transform.TransformPositionNoScale(FVector(0.f, BoundsMax.Max.Y + 10.f, CoordZ));
+			break;
+
+		case EText3DHorizontalTextAlignment::Right:
+			IconLocation = Transform.TransformPositionNoScale(FVector(0.f, BoundsMax.Min.Y - 10.f, CoordZ));
+			break;
+	}
+
+	return IconLocation;
+}
+
+FVector FAvaTextVisualizer::GetHeightHandleLocation(const UText3DComponent* InText3DComp) const
+{
+	const FTransform Transform = InText3DComp->GetComponentTransform();
+	const FBox BoundsMax = GetBoundsMax(InText3DComp);
+	const float CoordY = (InText3DComp->GetHorizontalAlignment() != EText3DHorizontalTextAlignment::Right)
+		? BoundsMax.Min.Y - 10.f
+		: BoundsMax.Max.Y + 10.f;
+
+	FVector IconLocation;
+
+	switch (InText3DComp->GetVerticalAlignment())
+	{
+		case EText3DVerticalTextAlignment::FirstLine:
+			IconLocation = Transform.TransformPositionNoScale(FVector(0.f, CoordY, BoundsMax.Max.Z));
+			break;
+
+		case EText3DVerticalTextAlignment::Top:
+			IconLocation = Transform.TransformPositionNoScale(FVector(0.f, CoordY, BoundsMax.Min.Z));
+			break;
+
+		case EText3DVerticalTextAlignment::Center:
+			IconLocation = Transform.TransformPositionNoScale(FVector(0.f, CoordY, BoundsMax.Max.Z));
+			break;
+
+		case EText3DVerticalTextAlignment::Bottom:
+			IconLocation = Transform.TransformPositionNoScale(FVector(0.f, CoordY, BoundsMax.Max.Z));
+			break;
+	}
+
+	return IconLocation;
+}
+
+void FAvaTextVisualizer::GetTextActorGradientControlsLocations(const UAvaText3DComponent* InAvaTextComponent, FVector& OutGradientCenterLocation, FVector& OutGradientStartLocation, FVector& OutGradientEndLocation) const
+{
+	using namespace UE::Avalanche::Private;
+	if (IsValid(InAvaTextComponent))
+	{
+		if (const AActor* const Text3DActor = InAvaTextComponent->GetOwner())
+		{
+			const AActor* TextAsActor = Cast<AActor>(Text3DActor);
+			check(TextAsActor);
+
+			FVector BoundsOrigin;
+			FVector BoundsExtent;
+			
+			Text3DActor->GetActorBounds(false, BoundsOrigin, BoundsExtent);
+			const float GradientLineHalfLength = FMath::Min(BoundsExtent.Length(), FTextVisualizerStatics::GradientHandleMaxLength);
+			FVector GradientCenterLocation = BoundsOrigin;
+
+			const FVector GradientDir = InAvaTextComponent->GetGradientDirection();
+			OutGradientStartLocation = GradientCenterLocation - GradientDir * GradientLineHalfLength;
+			OutGradientEndLocation = GradientCenterLocation + GradientDir * GradientLineHalfLength;	
+
+			// we need to move gradient center along its direction, base on offset amount
+			const float GradientAsNormalizedPos = InAvaTextComponent->GradientSettings.Offset * 2.0 - 1.0;
+
+			GradientCenterLocation -= GradientDir * GradientAsNormalizedPos * GradientLineHalfLength;
+			OutGradientCenterLocation = GradientCenterLocation;
+		}
+	}
+}
+
+FVector FAvaTextVisualizer::GetGradientEndHandleLocation(const UAvaText3DComponent* InAvaTextComponent) const
+{
+	if (IsValid(InAvaTextComponent))
+	{
+		FVector GradientCenter;
+		FVector GradientStart;
+		FVector GradientEnd;
+		
+		GetTextActorGradientControlsLocations(TextComponent.Get(), GradientCenter, GradientStart, GradientEnd);
+
+		return GradientEnd;
+	}
+
+	return FVector::Zero();
+}
+
+FVector FAvaTextVisualizer::GetGradientCenterHandleLocation(const UAvaText3DComponent* InAvaTextComponent) const
+{
+	if (IsValid(InAvaTextComponent))
+	{
+		FVector GradientCenter;
+		FVector GradientStart;
+		FVector GradientEnd;
+		
+		GetTextActorGradientControlsLocations(TextComponent.Get(), GradientCenter, GradientStart, GradientEnd);
+
+		return GradientCenter;
+	}
+
+	return FVector::Zero();
+}
+
+FVector FAvaTextVisualizer::GetGradientStartHandleLocation(const UAvaText3DComponent* InAvaTextComponent) const
+{
+	if (IsValid(InAvaTextComponent))
+	{
+		FVector GradientCenter;
+		FVector GradientStart;
+		FVector GradientEnd;
+		
+		GetTextActorGradientControlsLocations(TextComponent.Get(), GradientCenter, GradientStart, GradientEnd);
+
+		return GradientStart;
+	}
+
+	return FVector::Zero();
+}
+
+void FAvaTextVisualizer::StoreTextMetrics(const UText3DComponent* InText3DComp)
+{
+	Bounds = GetComponentBounds(InText3DComp);
+	TArray<FString> Lines;
+	InText3DComp->GetText().ToString().ParseIntoArrayLines(Lines);
+	LineCount = Lines.Num();
+	LineHeight = (Bounds.GetSize().Z - (InText3DComp->GetLineSpacing() * (static_cast<float>(LineCount) - 1.f))) / static_cast<float>(LineCount);
+}
+
+void FAvaTextVisualizer::DrawVisualizationNotEditing(const UActorComponent* InComponent, const FSceneView* InView, 
+	FPrimitiveDrawInterface* InPDI, int32& InOutIconIndex)
+{
+	Super::DrawVisualizationNotEditing(InComponent, InView, InPDI, InOutIconIndex);
+
+	const UAvaText3DComponent* const AvaTextComponent = Cast<UAvaText3DComponent>(InComponent);
+
+	if (!AvaTextComponent)
+	{
+		return;
+	}
+	
+	StoreTextMetrics(AvaTextComponent);
+
+	DrawMaxTextSizeVisualization(AvaTextComponent, InView, InPDI);
+
+	DrawMaxTextWidthButton(AvaTextComponent, InView, InPDI, InOutIconIndex, AvaTextComponent->HasMaxWidth() ? Enabled : Disabled);
+	++InOutIconIndex;
+
+	DrawMaxTextHeightButton(AvaTextComponent, InView, InPDI, InOutIconIndex, AvaTextComponent->HasMaxHeight() ? Enabled : Disabled);
+	++InOutIconIndex;
+
+	DrawScaleProportionallyButton(AvaTextComponent, InView, InPDI, InOutIconIndex, AvaTextComponent->ScalesProportionally() ? Enabled : Disabled);
+	++InOutIconIndex;
+
+	if (AvaTextComponent->HasMaxWidth())
+	{
+		DrawMaxTextWidthHandle(AvaTextComponent, InView, InPDI, Inactive);
+	}
+
+	if (AvaTextComponent->HasMaxHeight())
+	{
+		DrawMaxTextHeightHandle(AvaTextComponent, InView, InPDI, Inactive);
+	}
+	
+	if (AvaTextComponent->GetColoringStyle() == EAvaTextColoringStyle::Gradient)
+	{
+		DrawEditGradientButton(AvaTextComponent, InView, InPDI, InOutIconIndex, bShowGradientControls ? Enabled : Disabled);
+
+		if (bShowGradientControls)
+		{
+			DrawGradientHandles(AvaTextComponent, InView, InPDI);
+		}
+	}
+}
+
+void FAvaTextVisualizer::DrawVisualizationEditing(const UActorComponent* InComponent, const FSceneView* InView, 
+	FPrimitiveDrawInterface* InPDI, int32& InOutIconIndex)
+{
+	Super::DrawVisualizationNotEditing(InComponent, InView, InPDI, InOutIconIndex);
+
+	const UAvaText3DComponent* const AvaTextComponent = Cast<UAvaText3DComponent>(InComponent);
+
+	if (!AvaTextComponent)
+	{
+		return;
+	}
+
+	StoreTextMetrics(AvaTextComponent);
+
+	DrawMaxTextSizeVisualization(AvaTextComponent, InView, InPDI);
+
+	DrawMaxTextWidthButton(AvaTextComponent, InView, InPDI, InOutIconIndex, AvaTextComponent->HasMaxWidth() ? Enabled : Disabled);
+	++InOutIconIndex;
+
+	DrawMaxTextHeightButton(AvaTextComponent, InView, InPDI, InOutIconIndex, AvaTextComponent->HasMaxHeight() ? Enabled : Disabled);
+	++InOutIconIndex;
+
+	DrawScaleProportionallyButton(AvaTextComponent, InView, InPDI, InOutIconIndex, AvaTextComponent->ScalesProportionally() ? Enabled : Disabled);
+	++InOutIconIndex;
+
+	if (AvaTextComponent->HasMaxWidth() && !bEditingWidth)
+	{
+		DrawMaxTextWidthHandle(AvaTextComponent, InView, InPDI, Inactive);
+	}
+
+	if (AvaTextComponent->HasMaxHeight() && !bEditingHeight)
+	{
+		DrawMaxTextHeightHandle(AvaTextComponent, InView, InPDI, Inactive);
+	}
+	
+	if (AvaTextComponent->GetColoringStyle() == EAvaTextColoringStyle::Gradient)
+	{
+		DrawEditGradientButton(AvaTextComponent, InView, InPDI, InOutIconIndex, bShowGradientControls ? Enabled : Disabled);
+		++InOutIconIndex;
+
+		if (bShowGradientControls)
+		{
+			DrawGradientHandles(AvaTextComponent, InView, InPDI);
+		}
+	}
+}
+
+void FAvaTextVisualizer::StoreInitialValues()
+{
+	Super::StoreInitialValues();
+
+	if (!TextComponent.IsValid())
+	{
+		return;
+	}
+
+	bInitialMaxWidthEnabled = TextComponent->HasMaxHeight();
+	bInitialMaxHeightEnabled = TextComponent->HasMaxHeight();
+	InitialMaxWidth = TextComponent->GetMaxWidth();
+	InitialMaxHeight = TextComponent->GetMaxHeight();
+	bInitialScaleProportionally = TextComponent->ScalesProportionally();
+	
+	InitialGradientRotation = TextComponent->GetGradientSettings().Rotation;
+	InitialGradientOffset = TextComponent->GetGradientSettings().Offset;
+	InitialGradientSmoothness = TextComponent->GetGradientSettings().Smoothness;
+}
+
+void FAvaTextVisualizer::DrawMaxTextSizeVisualization(const UText3DComponent* InTextComponent, const FSceneView* InView, 
+	FPrimitiveDrawInterface* InPDI) const
+{
+	if (!InTextComponent || (!InTextComponent->HasMaxWidth() && !InTextComponent->HasMaxHeight()))
+	{
+		return;
+	}
+
+	const FTransform Transform = InTextComponent->GetComponentTransform();
+	const FBox BoundsMax = GetBoundsMax(InTextComponent);
+	const FVector TopLeft = Transform.TransformPositionNoScale(FVector(0.f, BoundsMax.Min.Y, BoundsMax.Max.Z));
+	const FVector TopRight = Transform.TransformPositionNoScale(FVector(0.f, BoundsMax.Max.Y, BoundsMax.Max.Z));
+	const FVector BottomLeft = Transform.TransformPositionNoScale(FVector(0.f, BoundsMax.Min.Y, BoundsMax.Min.Z));
+	const FVector BottomRight = Transform.TransformPositionNoScale(FVector(0.f, BoundsMax.Max.Y, BoundsMax.Min.Z));
+
+	InPDI->DrawLine(TopLeft, TopRight, FLinearColor::Yellow, SDPG_World, 0.5f);
+	InPDI->DrawLine(TopRight, BottomRight, FLinearColor::Yellow, SDPG_World, 0.5f);
+	InPDI->DrawLine(BottomRight, BottomLeft, FLinearColor::Yellow, SDPG_World, 0.5f);
+	InPDI->DrawLine(BottomLeft, TopLeft, FLinearColor::Yellow, SDPG_World, 0.5f);
+}
+
+bool FAvaTextVisualizer::DrawGradientCenterHandle(const UText3DComponent* InTextComponent, const FSceneView* InView, FPrimitiveDrawInterface* InPDI, FVector InGradientCenter)
+{
+	static constexpr float BaseSize = 1.f;
+
+	UTexture2D* UVSprite = IAvalancheComponentVisualizersModule::Get().GetSettings()->GetVisualizerSprite(FAvalancheShapesEditorModule::SizeSprite);
+
+	if (!UVSprite || !UVSprite->GetResource())
+	{
+		return true;
+	}
+
+	const FVector IconLocation = InGradientCenter;
+	const float IconSize = BaseSize * GetIconSizeScale(InView, IconLocation);
+
+	InPDI->SetHitProxy(new HAvaTextGradientCenterHandleProxy(InTextComponent));
+	InPDI->DrawSprite(InGradientCenter, IconSize, IconSize, UVSprite->GetResource(), FLinearColor::White,
+	                SDPG_Foreground, 0, 0, 0, 0, SE_BLEND_Opaque);
+	InPDI->SetHitProxy(nullptr);		
+	return false;
+}
+
+bool FAvaTextVisualizer::DrawGradientSmoothnessHandle(const UText3DComponent* InTextComponent, const FSceneView* InView, FPrimitiveDrawInterface* InPDI, FVector InHandleLocation)
+{
+	static constexpr float BaseSize = 1.5f;
+		
+	const UTexture2D* const UVSprite = IAvalancheComponentVisualizersModule::Get().GetSettings()->GetVisualizerSprite(FAvalancheShapesEditorModule::UVSprite);
+
+	if (!UVSprite || !UVSprite->GetResource())
+	{
+		return true;
+	}
+
+	const FVector IconLocation = InHandleLocation;
+	const float IconSize = BaseSize * GetIconSizeScale(InView, IconLocation);
+
+	InPDI->SetHitProxy(new HAvaTextGradientSmoothnessHandleProxy(InTextComponent));
+	InPDI->DrawSprite(InHandleLocation, IconSize, IconSize, UVSprite->GetResource(), FLinearColor::White,
+					SDPG_Foreground, 0, 0, 0, 0, SE_BLEND_Opaque);
+	InPDI->SetHitProxy(nullptr);		
+	return false;
+}
+
+void FAvaTextVisualizer::DrawGradientHandles(const UAvaText3DComponent* InAvaTextComponent, const FSceneView* InView, FPrimitiveDrawInterface* InPDI) const
+{
+	using namespace UE::Avalanche::Private;
+	if (!InAvaTextComponent)
+	{
+		return;
+	}
+
+	FVector GradientCenter;
+	FVector GradientLineStart;
+	FVector GradientLineEnd;
+
+	GetTextActorGradientControlsLocations(InAvaTextComponent, GradientCenter, GradientLineStart, GradientLineEnd);
+	
+	InPDI->DrawLine(GradientLineStart, GradientLineEnd, FLinearColor::White,SDPG_Foreground );
+	
+	DrawGradientLineStartHandle(InAvaTextComponent, InView, InPDI, GradientLineStart, InAvaTextComponent->GetGradientSettings().ColorB);
+	DrawGradientLineEndHandle(InAvaTextComponent, InView, InPDI, GradientLineEnd, InAvaTextComponent->GetGradientSettings().ColorA);
+	DrawGradientCenterHandle(InAvaTextComponent, InView, InPDI, GradientCenter);
+	DrawGradientSmoothnessHandle(InAvaTextComponent, InView, InPDI, GradientCenter + FVector(0, FTextVisualizerStatics::GradientSmoothnessHandleOffset, 0));
+	
+}
+
+FBox FAvaTextVisualizer::GetBoundsMax(const UText3DComponent* InText3DComp) const
+{
+	const FVector Scale = InText3DComp->GetComponentTransform().GetScale3D();
+	FBox BoundsMax = Bounds;
+
+	if (InText3DComp->HasMaxWidth())
+	{
+		switch (InText3DComp->GetHorizontalAlignment())
+		{
+			case EText3DHorizontalTextAlignment::Left:
+				BoundsMax.Max.Y = InText3DComp->GetMaxWidth() * Scale.Y;
+				break;
+
+			case EText3DHorizontalTextAlignment::Center:
+				BoundsMax.Min.Y = InText3DComp->GetMaxWidth() * -0.5f * Scale.Y;
+				BoundsMax.Max.Y = InText3DComp->GetMaxWidth() * 0.5f * Scale.Y;
+				break;
+
+			case EText3DHorizontalTextAlignment::Right:
+				BoundsMax.Min.Y = InText3DComp->GetMaxWidth() * -Scale.Y;
+				break;
+		}
+	}
+
+	if (InText3DComp->HasMaxHeight())
+	{
+		switch (InText3DComp->GetVerticalAlignment())
+		{
+			case EText3DVerticalTextAlignment::FirstLine:
+				if (LineCount <= 1)
+				{
+					BoundsMax.Max.Z = FMath::Max(BoundsMax.Max.Z, InText3DComp->GetMaxHeight() * TopHeightFraction * Scale.Z);
+				}
+				else
+				{
+					BoundsMax.Max.Z = FMath::Max(BoundsMax.Max.Z, LineHeight * TopHeightFraction * Scale.Z * InText3DComp->GetMaxHeight() / Bounds.GetSize().Z);
+				}
+				break;
+
+			case EText3DVerticalTextAlignment::Top:
+				BoundsMax.Min.Z = -InText3DComp->GetMaxHeight() * Scale.Z;
+				break;
+
+			case EText3DVerticalTextAlignment::Center:
+				BoundsMax.Min.Z = InText3DComp->GetMaxHeight() / -2.f * Scale.Z;
+				BoundsMax.Max.Z = InText3DComp->GetMaxHeight() / 2.f * Scale.Z;
+				break;
+
+			case EText3DVerticalTextAlignment::Bottom:
+				BoundsMax.Max.Z = InText3DComp->GetMaxHeight() * Scale.Z;
+				break;
+		}
+	}
+
+	return BoundsMax;
+}
+
+void FAvaTextVisualizer::DrawMaxTextWidthButton(const UText3DComponent* InTextComponent, const FSceneView* InView, 
+	FPrimitiveDrawInterface* InPDI, int32 InIconIndex, const FLinearColor& InColor) const
+{
+	if (!InTextComponent)
+	{
+		return;
+	}
+
+	UTexture2D* UVSprite = IAvalancheComponentVisualizersModule::Get().GetSettings()->GetVisualizerSprite(FAvalancheShapesEditorModule::TextMaxWidthSprite);
+
+	if (!UVSprite || !UVSprite->GetResource())
+	{
+		return;
+	}
+
+	FVector IconLocation;
+	float IconSize;
+	GetIconMetrics(InView, InIconIndex, IconLocation, IconSize);
+
+	InPDI->SetHitProxy(new HAvaTextMaxTextWidthProxy(InTextComponent));
+	InPDI->DrawSprite(IconLocation, IconSize, IconSize, UVSprite->GetResource(), InColor,
+		SDPG_Foreground, 0, 0, 0, 0, SE_BLEND_Opaque);
+	InPDI->SetHitProxy(nullptr);
+}
+
+void FAvaTextVisualizer::DrawMaxTextHeightButton(const UText3DComponent* InTextComponent, const FSceneView* InView, 
+	FPrimitiveDrawInterface* InPDI, int32 InIconIndex, const FLinearColor& InColor) const
+{
+	if (!InTextComponent)
+	{
+		return;
+	}
+
+	const UTexture2D* const UVSprite = IAvalancheComponentVisualizersModule::Get().GetSettings()->GetVisualizerSprite(FAvalancheShapesEditorModule::TextMaxHeightSprite);
+
+	if (!UVSprite || !UVSprite->GetResource())
+	{
+		return;
+	}
+
+	FVector IconLocation;
+	float IconSize;
+	GetIconMetrics(InView, InIconIndex, IconLocation, IconSize);
+
+	InPDI->SetHitProxy(new HAvaTextMaxTextHeightProxy(InTextComponent));
+	InPDI->DrawSprite(IconLocation, IconSize, IconSize, UVSprite->GetResource(), InColor,
+		SDPG_Foreground, 0, 0, 0, 0, SE_BLEND_Opaque);
+	InPDI->SetHitProxy(nullptr);
+}
+
+void FAvaTextVisualizer::DrawScaleProportionallyButton(const UText3DComponent* InTextComponent, const FSceneView* InView, 
+	FPrimitiveDrawInterface* InPDI, int32 InIconIndex, const FLinearColor& InColor) const
+{
+	if (!InTextComponent)
+	{
+		return;
+	}
+
+	const UTexture2D* const UVSprite = IAvalancheComponentVisualizersModule::Get().GetSettings()->GetVisualizerSprite(FAvalancheShapesEditorModule::TextScaleProportionallySprite);
+
+	if (!UVSprite || !UVSprite->GetResource())
+	{
+		return;
+	}
+
+	FVector IconLocation;
+	float IconSize;
+	GetIconMetrics(InView, InIconIndex, IconLocation, IconSize);
+
+	InPDI->SetHitProxy(new HAvaTextScaleProportionallyProxy(InTextComponent));
+	InPDI->DrawSprite(IconLocation, IconSize, IconSize, UVSprite->GetResource(), InColor,
+		SDPG_Foreground, 0, 0, 0, 0, SE_BLEND_Opaque);
+	InPDI->SetHitProxy(nullptr);
+}
+
+void FAvaTextVisualizer::DrawEditGradientButton(
+	const UText3DComponent* InTextComponent
+	, const FSceneView* InView
+	, FPrimitiveDrawInterface* InPDI
+	, int32 InIconIndex
+	, const FLinearColor& InColor) const
+{
+	if (!InTextComponent)
+	{
+		return;
+	}
+
+	// todo: proper sprite
+	const UTexture2D* const UVSprite = IAvalancheComponentVisualizersModule::Get().GetSettings()->GetVisualizerSprite(FAvalancheShapesEditorModule::LinearGradientSprite);
+
+	if (!UVSprite || !UVSprite->GetResource())
+	{
+		return;
+	}
+
+	FVector IconLocation;
+	float IconSize;
+	GetIconMetrics(InView, InIconIndex, IconLocation, IconSize);
+	InPDI->SetHitProxy(new HAvaTextEditGradientProxy(InTextComponent));
+	InPDI->DrawSprite(IconLocation, IconSize, IconSize, UVSprite->GetResource(), InColor,
+		SDPG_Foreground, 0, 0, 0, 0, SE_BLEND_Opaque);
+	InPDI->SetHitProxy(nullptr);
+}
+
+void FAvaTextVisualizer::DrawMaxTextWidthHandle(
+	const UText3DComponent* InTextComponent
+	, const FSceneView* InView
+	, FPrimitiveDrawInterface* InPDI
+	, const FLinearColor& InColor) const
+{
+	static constexpr float BaseSize = 1.f;
+
+	if (!InTextComponent)
+	{
+		return;
+	}
+
+	const UTexture2D* const UVSprite = IAvalancheComponentVisualizersModule::Get().GetSettings()->GetVisualizerSprite(FAvalancheShapesEditorModule::SizeSprite);
+
+	if (!UVSprite || !UVSprite->GetResource())
+	{
+		return;
+	}
+
+	const FVector IconLocation = GetWidthHandleLocation(InTextComponent);
+	const float IconSize = BaseSize * GetIconSizeScale(InView, IconLocation);
+
+	InPDI->SetHitProxy(new HAvaTextMaxTextWidthHandleProxy(InTextComponent));
+	InPDI->DrawSprite(IconLocation, IconSize, IconSize, UVSprite->GetResource(), InColor,
+		SDPG_Foreground, 0, 0, 0, 0, SE_BLEND_Opaque);
+	InPDI->SetHitProxy(nullptr);
+}
+
+void FAvaTextVisualizer::DrawMaxTextHeightHandle(const UText3DComponent* InTextComponent, const FSceneView* InView, FPrimitiveDrawInterface* InPDI,
+	const FLinearColor& InColor) const
+{
+	static constexpr  float BaseSize = 1.f;
+
+	if (!InTextComponent)
+	{
+		return;
+	}
+
+	const UTexture2D* const UVSprite = IAvalancheComponentVisualizersModule::Get().GetSettings()->GetVisualizerSprite(FAvalancheShapesEditorModule::SizeSprite);
+
+	if (!UVSprite || !UVSprite->GetResource())
+	{
+		return;
+	}
+
+	const FVector IconLocation = GetHeightHandleLocation(InTextComponent);
+	const float IconSize = BaseSize * GetIconSizeScale(InView, IconLocation);
+
+	InPDI->SetHitProxy(new HAvaTextMaxTextHeightHandleProxy(InTextComponent));
+	InPDI->DrawSprite(IconLocation, IconSize, IconSize, UVSprite->GetResource(), InColor,
+		SDPG_Foreground, 0, 0, 0, 0, SE_BLEND_Opaque);
+	InPDI->SetHitProxy(nullptr);
+}
+
+void FAvaTextVisualizer::DrawGradientLineStartHandle(const UAvaText3DComponent* InAvaTextComponent, const FSceneView* InView, 
+	FPrimitiveDrawInterface* InPDI, const FVector& InLocation, const FLinearColor& InColor) const
+{
+	static constexpr float BaseSize = 1.5f;
+
+	if (!InAvaTextComponent)
+	{
+		return;
+	}
+
+	const UTexture2D* const UVSprite = IAvalancheComponentVisualizersModule::Get().GetSettings()->GetVisualizerSprite(FAvalancheShapesEditorModule::ColorSelectionSprite);
+
+	if (!UVSprite || !UVSprite->GetResource())
+	{
+		return;
+	}
+	
+	const FVector IconLocation = InLocation;
+	const float IconSize = BaseSize * GetIconSizeScale(InView, IconLocation);
+
+	// in case gradient center handle is on top of this handle, or almost on top of it, don't enable interaction
+	if (FVector::Distance(InLocation, GetGradientCenterHandleLocation(InAvaTextComponent)) > BaseSize/2.0)
+	{
+		InPDI->SetHitProxy(new HAvaTextGradientLineStartHandleProxy(InAvaTextComponent));
+	}
+	
+	InPDI->DrawSprite(IconLocation, IconSize, IconSize, UVSprite->GetResource(), InColor,
+		SDPG_Foreground, 0, 0, 0, 0, SE_BLEND_Opaque);
+	InPDI->SetHitProxy(nullptr);
+}
+
+void FAvaTextVisualizer::DrawGradientLineEndHandle(const UAvaText3DComponent* InAvaTextComponent, const FSceneView* InView, 
+	FPrimitiveDrawInterface* InPDI, const FVector& InLocation, const FLinearColor& InColor) const
+{
+	static constexpr float BaseSize = 1.5f;
+
+	if (!InAvaTextComponent)
+	{
+		return;
+	}
+
+	const UTexture2D* const UVSprite = IAvalancheComponentVisualizersModule::Get().GetSettings()->GetVisualizerSprite(FAvalancheShapesEditorModule::ColorSelectionSprite);
+
+	if (!UVSprite || !UVSprite->GetResource())
+	{
+		return;
+	}
+
+	const FVector IconLocation = InLocation;
+	const float IconSize = BaseSize * GetIconSizeScale(InView, IconLocation);
+
+	// in case gradient center handle is on top of this handle, or almost on top of it, don't enable interaction
+	if (FVector::Distance(InLocation, GetGradientCenterHandleLocation(InAvaTextComponent)) > BaseSize/2.0)
+	{
+		InPDI->SetHitProxy(new HAvaTextGradientLineEndHandleProxy(InAvaTextComponent));
+	}
+	
+	InPDI->DrawSprite(IconLocation, IconSize, IconSize, UVSprite->GetResource(), InColor,
+		SDPG_Foreground, 0, 0, 0, 0, SE_BLEND_Opaque);
+	InPDI->SetHitProxy(nullptr);
+}
+
+#undef LOCTEXT_NAMESPACE
diff --git a/Engine/Plugins/Experimental/Avalanche/Source/AvalancheTransition/Private/AvaTransitionModule.cpp b/Engine/Plugins/Experimental/Avalanche/Source/AvalancheTransition/Private/AvaTransitionModule.cpp
new file mode 100644
--- /dev/null
+++ b/Engine/Plugins/Experimental/Avalanche/Source/AvalancheTransition/Private/AvaTransitionModule.cpp
@@ -0,0 +1,18 @@
+// Copyright Epic Games, Inc. All Rights Reserved.
+
+#include "IAvaTransitionModule.h"
+#include "AvaTransitionLog.h"
+
+DEFINE_LOG_CATEGORY(LogAvaTransition);
+
+class FAvaTransitionModule : public IAvaTransitionModule
+{
+	virtual FOnValidateTransitionTree& GetOnValidateTransitionTree()
+	{
+		return OnValidateStateTree;
+	}
+
+	FOnValidateTransitionTree OnValidateStateTree;
+};
+
+IMPLEMENT_MODULE(FAvaTransitionModule, AvalancheTransition)
diff --git a/Engine/Plugins/Experimental/Avalanche/Source/AvalancheTransition/Private/Behavior/AvaTransitionBehaviorInstance.cpp b/Engine/Plugins/Experimental/Avalanche/Source/AvalancheTransition/Private/Behavior/AvaTransitionBehaviorInstance.cpp
new file mode 100644
--- /dev/null
+++ b/Engine/Plugins/Experimental/Avalanche/Source/AvalancheTransition/Private/Behavior/AvaTransitionBehaviorInstance.cpp
@@ -0,0 +1,298 @@
+// Copyright Epic Games, Inc. All Rights Reserved.
+
+#include "Behavior/AvaTransitionBehaviorInstance.h"
+#include "AvaTransitionLog.h"
+#include "AvaTransitionTree.h"
+#include "Behavior/IAvaTransitionBehavior.h"
+#include "Engine/Level.h"
+#include "Engine/World.h"
+#include "Execution/AvaTransitionExecutionContext.h"
+#include "StateTreeExecutionContext.h"
+#include "StateTreeReference.h"
+#include "Subsystems/WorldSubsystem.h"
+
+FAvaTransitionBehaviorInstance& FAvaTransitionBehaviorInstance::SetBehavior(IAvaTransitionBehavior* InBehavior)
+{
+	BehaviorWeak = InBehavior;
+	return *this;
+}
+
+bool FAvaTransitionBehaviorInstance::IsEnabled() const
+{
+	const IAvaTransitionBehavior* TransitionBehavior = GetBehavior();
+	if (!TransitionBehavior)
+	{
+		return false;
+	}
+
+	const UAvaTransitionTree* TransitionTree = TransitionBehavior->GetTransitionTree();
+	if (!TransitionTree)
+	{
+		return false;
+	}
+
+	return TransitionTree->IsEnabled();
+}
+
+IAvaTransitionBehavior* FAvaTransitionBehaviorInstance::GetBehavior() const
+{
+	return BehaviorWeak.Get();
+}
+
+FAvaTagHandle FAvaTransitionBehaviorInstance::GetTransitionLayer() const
+{
+	return TransitionContext.GetTransitionLayer();
+}
+
+EAvaTransitionType FAvaTransitionBehaviorInstance::GetTransitionType() const
+{
+	return TransitionContext.GetTransitionType();
+}
+
+bool FAvaTransitionBehaviorInstance::IsRunning() const
+{
+	return RunStatus == EStateTreeRunStatus::Running;
+}
+
+const FAvaTransitionContext& FAvaTransitionBehaviorInstance::GetTransitionContext() const
+{
+	return TransitionContext;
+}
+
+FAvaTransitionContext& FAvaTransitionBehaviorInstance::GetTransitionContext()
+{
+	return TransitionContext;
+}
+
+void FAvaTransitionBehaviorInstance::SetTransitionType(EAvaTransitionType InTransitionType)
+{
+	TransitionContext.TransitionType = InTransitionType;
+}
+
+bool FAvaTransitionBehaviorInstance::Setup()
+{
+	RunStatus = EStateTreeRunStatus::Unset;
+	TOptional<FAvaTransitionExecutionContext> Context = UpdateContext();
+	return Context.IsSet();
+}
+
+void FAvaTransitionBehaviorInstance::Start()
+{
+	if (RunStatus != EStateTreeRunStatus::Unset)
+	{
+		return;
+	}
+
+	// If this Instance is not Enabled for Transition, finish immediately
+	if (!IsEnabled())
+	{
+		RunStatus = EStateTreeRunStatus::Succeeded;
+		return;
+	}
+
+	if (TOptional<FAvaTransitionExecutionContext> Context = UpdateContext())
+	{
+		const IAvaTransitionBehavior* Behavior = GetBehavior();
+		check(Behavior); // Assume that we have valid behavior if we were able to create the context.
+		const FStateTreeReference& StateTreeReference = Behavior->GetStateTreeReference();
+
+		RunStatus = Context->Start(&StateTreeReference.GetParameters());
+	}
+	else
+	{
+		RunStatus = EStateTreeRunStatus::Failed;
+	}
+
+	ConditionallyStop();
+}
+
+void FAvaTransitionBehaviorInstance::Tick(float InDeltaSeconds)
+{
+	if (RunStatus != EStateTreeRunStatus::Running)
+	{
+		return;
+	}
+
+	if (TOptional<FAvaTransitionExecutionContext> Context = UpdateContext())
+	{
+		RunStatus = Context->Tick(InDeltaSeconds);
+	}
+	else
+	{
+		RunStatus = EStateTreeRunStatus::Failed;
+	}
+
+	ConditionallyStop();
+}
+
+void FAvaTransitionBehaviorInstance::Stop()
+{
+	if (TOptional<FAvaTransitionExecutionContext> Context = UpdateContext())
+	{
+		RunStatus = Context->Stop();
+	}
+	else
+	{
+		RunStatus = EStateTreeRunStatus::Stopped;
+	}
+}
+
+void FAvaTransitionBehaviorInstance::SetOverrideLayer(const FAvaTagHandle& InOverrideLayer)
+{
+	OverrideLayer = InOverrideLayer;
+}
+
+void FAvaTransitionBehaviorInstance::SetLogContext(const FString& InContext)
+{
+	LogContext = InContext;
+}
+
+void FAvaTransitionBehaviorInstance::ConditionallyStop()
+{
+	// if the run status was set in Start/Tick with anything other than running (e.g. Succeeded/Failed/Stop) it should call stop
+	if (RunStatus != EStateTreeRunStatus::Running)
+	{
+		Stop();
+	}
+}
+
+TOptional<FAvaTransitionExecutionContext> FAvaTransitionBehaviorInstance::UpdateContext()
+{
+	if (ValidateTransitionScene())
+	{
+		IAvaTransitionBehavior* Behavior = GetBehavior();
+		UpdateTransitionLayers(Behavior);
+		return MakeContext(Behavior);
+	}
+	return TOptional<FAvaTransitionExecutionContext>();
+}
+
+bool FAvaTransitionBehaviorInstance::ValidateTransitionScene()
+{
+	if (!TransitionSceneOwner.IsValid())
+	{
+		//todo: support having a behavior instance run without a Transition Scene instanced
+		UE_LOG(LogAvaTransition, Error
+			, TEXT("SetContextRequirements failed for '%s'. Transition Scene Owner is invalid!")
+			, *LogContext);
+		return false;
+	}
+
+	FAvaTransitionScene* TransitionScene = TransitionContext.GetTransitionScene();
+	if (!TransitionScene)
+	{
+		UE_LOG(LogAvaTransition, Error
+			, TEXT("SetContextRequirements failed for '%s'. Transition Scene is null")
+			, *LogContext);
+		return false;
+	}
+
+	return true;
+}
+
+void FAvaTransitionBehaviorInstance::UpdateTransitionLayers(const IAvaTransitionBehavior* InBehavior)
+{
+	FAvaTransitionScene& TransitionScene = *TransitionContext.GetTransitionScene();
+
+	const UAvaTransitionTree* TransitionTree = InBehavior ? InBehavior->GetTransitionTree() : nullptr;
+	if (TransitionTree)
+	{
+		TransitionContext.TransitionLayer = TransitionTree->GetTransitionLayer();
+	}
+	else
+	{
+		TransitionContext.TransitionLayer = FAvaTagHandle();
+	}
+
+	// Give opportunity for the Scene to override the layer (if not already overriden)
+	if (!OverrideLayer.IsSet())
+	{
+		FAvaTagHandle TagHandle;
+		TransitionScene.GetOverrideTransitionLayer(TagHandle);
+		if (TagHandle.IsValid())
+		{
+			OverrideLayer = TagHandle;
+		}
+	}
+
+	if (OverrideLayer.IsSet())
+	{
+		TransitionContext.TransitionLayer = *OverrideLayer;
+	}
+}
+
+TOptional<FAvaTransitionExecutionContext> FAvaTransitionBehaviorInstance::MakeContext(IAvaTransitionBehavior* InBehavior)
+{
+	FAvaTransitionScene* TransitionScene = TransitionContext.GetTransitionScene();
+	if (!InBehavior || !TransitionScene)
+	{
+		return TOptional<FAvaTransitionExecutionContext>();
+	}
+
+	const ULevel* Level = TransitionScene->GetLevel();
+
+	const UWorld* World = Level ? Level->OwningWorld : nullptr;
+
+	if (!IsValid(World))
+	{
+		UE_LOG(LogAvaTransition, Error
+			, TEXT("SetContextRequirements failed for '%s'. World is invalid")
+			, *LogContext);
+		return TOptional<FAvaTransitionExecutionContext>();
+	}
+
+	const FStateTreeReference& StateTreeReference = InBehavior->GetStateTreeReference();
+
+	FAvaTransitionExecutionContext Context(InBehavior->AsUObject(), *StateTreeReference.GetStateTree(), InstanceData);
+	if (!Context.IsValid())
+	{
+		return TOptional<FAvaTransitionExecutionContext>();
+	}
+
+	Context.SetCollectExternalDataCallback(FOnCollectStateTreeExternalData::CreateLambda(
+		[World, &TransitionContext = TransitionContext]
+		(const FStateTreeExecutionContext& Context, const UStateTree* StateTree, TArrayView<const FStateTreeExternalDataDesc> ExternalDescs, TArrayView<FStateTreeDataView> OutDataViews)
+		{
+			check(ExternalDescs.Num() == OutDataViews.Num());
+			for (int32 Index = 0; Index < ExternalDescs.Num(); Index++)
+			{
+				const FStateTreeExternalDataDesc& Desc = ExternalDescs[Index];
+				if (!Desc.Struct)
+				{
+					continue;
+				}
+
+				if (Desc.Struct->IsChildOf(FAvaTransitionContext::StaticStruct()))
+				{
+					OutDataViews[Index] = FStructView::Make(TransitionContext);
+				}
+				else if (Desc.Struct->IsChildOf(UWorldSubsystem::StaticClass()))
+				{
+					UWorldSubsystem* Subsystem = World->GetSubsystemBase(Cast<UClass>(const_cast<UStruct*>(Desc.Struct.Get())));
+					OutDataViews[Index] = FStateTreeDataView(Subsystem);
+				}
+			}
+				
+			return true;
+		})
+	);
+	
+	if (!Context.AreContextDataViewsValid())
+	{
+		UE_LOG(LogAvaTransition, Error
+			, TEXT("SetContextRequirements failed for '%s'. Missing external data requirements. StateTree will not update.")
+			, *LogContext);
+
+		return TOptional<FAvaTransitionExecutionContext>();
+	}
+
+	// Default Scene Description will be just the Transition Type
+	FString SceneDescription = *UEnum::GetDisplayValueAsText(TransitionContext.GetTransitionType()).ToString();
+
+	// Opportunity for the Transition Scene Implementation to update or set its own Description
+	TransitionScene->UpdateSceneDescription(SceneDescription);
+
+	Context.SetSceneDescription(MoveTemp(SceneDescription));
+
+	return Context;
+}
diff --git a/Engine/Plugins/Experimental/Avalanche/Source/AvalancheTransition/Private/Conditions/AvaTransitionTypeMatchCondition.cpp b/Engine/Plugins/Experimental/Avalanche/Source/AvalancheTransition/Private/Conditions/AvaTransitionTypeMatchCondition.cpp
new file mode 100644
--- /dev/null
+++ b/Engine/Plugins/Experimental/Avalanche/Source/AvalancheTransition/Private/Conditions/AvaTransitionTypeMatchCondition.cpp
@@ -0,0 +1,21 @@
+// Copyright Epic Games, Inc. All Rights Reserved.
+
+#include "Conditions/AvaTransitionTypeMatchCondition.h"
+#include "AvaTransitionContext.h"
+#include "StateTreeExecutionContext.h"
+
+#define LOCTEXT_NAMESPACE "AvaTransitionTypeMatchCondition"
+
+FText FAvaTransitionTypeMatchCondition::GenerateDescription() const
+{
+	return FText::Format(LOCTEXT("ConditionDescription", "transitioning {0}")
+		, UEnum::GetDisplayValueAsText(TransitionType).ToLower());
+}
+
+bool FAvaTransitionTypeMatchCondition::TestCondition(FStateTreeExecutionContext& InContext) const
+{
+	const FAvaTransitionContext& TransitionContext = InContext.GetExternalData(TransitionContextHandle);
+	return TransitionContext.GetTransitionType() == TransitionType;
+}
+
+#undef LOCTEXT_NAMESPACE
diff --git a/Engine/Plugins/Experimental/Avalanche/Source/AvalancheTransition/Private/Execution/AvaTransitionExecutionContext.cpp b/Engine/Plugins/Experimental/Avalanche/Source/AvalancheTransition/Private/Execution/AvaTransitionExecutionContext.cpp
new file mode 100644
--- /dev/null
+++ b/Engine/Plugins/Experimental/Avalanche/Source/AvalancheTransition/Private/Execution/AvaTransitionExecutionContext.cpp
@@ -0,0 +1,22 @@
+// Copyright Epic Games, Inc. All Rights Reserved.
+
+#include "AvaTransitionExecutionContext.h"
+
+FAvaTransitionExecutionContext::FAvaTransitionExecutionContext(UObject& InOwner, const UStateTree& InStateTree, FStateTreeInstanceData& InInstanceData)
+	: FStateTreeExecutionContext(InOwner, InStateTree, InInstanceData)
+{
+}
+
+void FAvaTransitionExecutionContext::SetSceneDescription(FString&& InSceneDescription)
+{
+	SceneDescription = MoveTemp(InSceneDescription);
+}
+
+FString FAvaTransitionExecutionContext::GetInstanceDescription() const
+{
+	if (!SceneDescription.IsEmpty())
+	{
+		return SceneDescription;
+	}
+	return FStateTreeExecutionContext::GetInstanceDescription();
+}
diff --git a/Engine/Plugins/Experimental/Avalanche/Source/AvalancheTransition/Private/Execution/AvaTransitionExecutionContext.h b/Engine/Plugins/Experimental/Avalanche/Source/AvalancheTransition/Private/Execution/AvaTransitionExecutionContext.h
new file mode 100644
--- /dev/null
+++ b/Engine/Plugins/Experimental/Avalanche/Source/AvalancheTransition/Private/Execution/AvaTransitionExecutionContext.h
@@ -0,0 +1,21 @@
+// Copyright Epic Games, Inc. All Rights Reserved.
+
+#pragma once
+
+#include "Containers/UnrealString.h"
+#include "StateTreeExecutionContext.h"
+
+struct FAvaTransitionExecutionContext : FStateTreeExecutionContext
+{
+	FAvaTransitionExecutionContext(UObject& InOwner, const UStateTree& InStateTree, FStateTreeInstanceData& InInstanceData);
+
+	void SetSceneDescription(FString&& InSceneDescription);
+
+protected:
+	//~ Begin FStateTreeExecutionContext
+	virtual FString GetInstanceDescription() const override;
+	//~ End FStateTreeExecutionContext
+
+private:
+	FString SceneDescription;
+};
diff --git a/Engine/Plugins/Experimental/Avalanche/Source/AvalancheTransition/Public/AvaTransitionLayerUtils.h b/Engine/Plugins/Experimental/Avalanche/Source/AvalancheTransition/Public/AvaTransitionLayerUtils.h
new file mode 100644
--- /dev/null
+++ b/Engine/Plugins/Experimental/Avalanche/Source/AvalancheTransition/Public/AvaTransitionLayerUtils.h
@@ -0,0 +1,28 @@
+// Copyright Epic Games, Inc. All Rights Reserved.
+
+#pragma once
+
+#include "Containers/ContainersFwd.h"
+#include "UObject/NameTypes.h"
+
+class UAvaTransitionSubsystem;
+enum class EAvaTransitionLayerCompareType : uint8;
+struct FAvaTagHandle;
+struct FAvaTagHandleContainer;
+struct FAvaTransitionBehaviorInstance;
+struct FAvaTransitionContext;
+struct FAvaTransitionLayerComparator;
+
+struct AVALANCHETRANSITION_API FAvaTransitionLayerUtils
+{
+	/** Gets all the Behavior Instances that match the Comparator */
+	static TArray<const FAvaTransitionBehaviorInstance*> QueryBehaviorInstances(UAvaTransitionSubsystem& InTransitionSubsystem, const FAvaTransitionLayerComparator& InComparator);
+
+	/** Builds a Comparator for the given Context (and optionally Layer), excluding the provided Transition Context (assumes it's the transition context calling this) */
+	static FAvaTransitionLayerComparator BuildComparator(const FAvaTransitionContext& InTransitionContext, EAvaTransitionLayerCompareType InCompareType, const FAvaTagHandle& InSpecificLayer);
+
+	/** Builds a Comparator for the given Context and specific layers, excluding the provided Transition Context (assumes it's the transition context calling this) */
+	static FAvaTransitionLayerComparator BuildComparator(const FAvaTransitionContext& InTransitionContext, EAvaTransitionLayerCompareType InCompareType, const FAvaTagHandleContainer& InSpecificLayers);
+
+	static FText GetLayerQueryText(EAvaTransitionLayerCompareType InLayerType, FName InSpecificLayerName);
+};
diff --git a/Engine/Plugins/Experimental/Avalanche/Source/AvalancheTransition/Public/AvaTransitionTree.h b/Engine/Plugins/Experimental/Avalanche/Source/AvalancheTransition/Public/AvaTransitionTree.h
new file mode 100644
--- /dev/null
+++ b/Engine/Plugins/Experimental/Avalanche/Source/AvalancheTransition/Public/AvaTransitionTree.h
@@ -0,0 +1,34 @@
+// Copyright Epic Games, Inc. All Rights Reserved.
+
+#pragma once
+
+#include "AvaTagHandle.h"
+#include "StateTree.h"
+#include "AvaTransitionTree.generated.h"
+
+UCLASS(DisplayName = "Motion Design Transition Tree")
+class AVALANCHETRANSITION_API UAvaTransitionTree : public UStateTree
+{
+	GENERATED_BODY()
+
+public:
+	FAvaTagHandle GetTransitionLayer() const;
+
+	void SetTransitionLayer(FAvaTagHandle InTransitionLayer);
+
+	bool IsEnabled() const;
+
+	void SetEnabled(bool bInEnabled);
+
+private:
+	/** The Layer this Transition Logic Tree deals with */
+	UPROPERTY()
+	FAvaTagHandle TransitionLayer;
+
+	/**
+	 * Determines whether this Transition Logic is enabled, by default.
+	 * Can be overriden by a Transition Instance to force the logic to run regardless
+	 */
+	UPROPERTY()
+	bool bEnabled = true;
+};
diff --git a/Engine/Plugins/Experimental/Avalanche/Source/AvalancheTransitionEditor/AvalancheTransitionEditor.Build.cs b/Engine/Plugins/Experimental/Avalanche/Source/AvalancheTransitionEditor/AvalancheTransitionEditor.Build.cs
new file mode 100644
--- /dev/null
+++ b/Engine/Plugins/Experimental/Avalanche/Source/AvalancheTransitionEditor/AvalancheTransitionEditor.Build.cs
@@ -0,0 +1,61 @@
+// Copyright Epic Games, Inc. All Rights Reserved.
+
+using UnrealBuildTool;
+
+public class AvalancheTransitionEditor : ModuleRules
+{
+    public AvalancheTransitionEditor(ReadOnlyTargetRules Target) : base(Target)
+    {
+        PCHUsage = ModuleRules.PCHUsageMode.UseExplicitOrSharedPCHs;
+
+        PublicDependencyModuleNames.AddRange(
+            new string[]
+            {
+                "AvalancheTag",
+                "Core",
+                "CoreUObject",
+                "StateTreeEditorModule",
+            }
+        );
+
+        PrivateDependencyModuleNames.AddRange(
+            new string[]
+            {
+                "ApplicationCore",
+                "AssetDefinition",
+                "AssetTools",
+                "AvalancheCore",
+                "AvalancheTransition",
+                "EditorStyle",
+                "Engine",
+                "InputCore",
+                "MessageLog",
+                "Projects",
+                "PropertyEditor",
+                "Slate",
+                "SlateCore",
+                "StateTreeModule",
+                "StructUtils",
+                "ToolMenus",
+                "UnrealEd",
+            }
+        );
+
+        if (Target.Platform == UnrealTargetPlatform.Win64 && (Target.Configuration != UnrealTargetConfiguration.Shipping || Target.bBuildEditor))
+        {
+            PrivateDefinitions.Add("WITH_STATETREE_DEBUGGER=1");
+            PrivateDependencyModuleNames.AddRange(
+                new string[]
+                {
+                    "TraceAnalysis",
+                    "TraceLog",
+                    "TraceServices",
+                }
+            );
+        }
+        else
+        {
+            PrivateDefinitions.Add("WITH_STATETREE_DEBUGGER=0");
+        }
+    }
+}
diff --git a/Engine/Plugins/Experimental/Avalanche/Source/AvalancheTransitionEditor/Private/AvaTransitionEditorModule.cpp b/Engine/Plugins/Experimental/Avalanche/Source/AvalancheTransitionEditor/Private/AvaTransitionEditorModule.cpp
new file mode 100644
--- /dev/null
+++ b/Engine/Plugins/Experimental/Avalanche/Source/AvalancheTransitionEditor/Private/AvaTransitionEditorModule.cpp
@@ -0,0 +1,158 @@
+// Copyright Epic Games, Inc. All Rights Reserved.
+
+#include "AvaTransitionEditorModule.h"
+#include "AvaTransitionCommands.h"
+#include "AvaTransitionEditorLog.h"
+#include "AvaTransitionEditorStyle.h"
+#include "AvaTransitionEditorUtils.h"
+#include "AvaTransitionTree.h"
+#include "AvaTransitionTreeEditorData.h"
+#include "AvaTransitionTreeSchema.h"
+#include "Behavior/IAvaTransitionBehavior.h"
+#include "IAvaTransitionModule.h"
+#include "StateTreeDelegates.h"
+#include "ToolMenu.h"
+#include "ToolMenuSection.h"
+#include "ViewModels/AvaTransitionEditorViewModel.h"
+#include "Widgets/Layout/SBox.h"
+
+#if WITH_STATETREE_DEBUGGER
+#include "Debugger/AvaTransitionTraceModule.h"
+#endif
+
+#define LOCTEXT_NAMESPACE "AvaTransitionEditorModule"
+
+DEFINE_LOG_CATEGORY(LogAvaEditorTransition)
+
+IMPLEMENT_MODULE(FAvaTransitionEditorModule, AvalancheTransitionEditor)
+
+void FAvaTransitionEditorModule::StartupModule()
+{
+	// Trigger Registration of the Style on Module Startup
+	FAvaTransitionEditorStyle::Get();
+
+#if WITH_STATETREE_DEBUGGER
+	FAvaTransitionTraceModule::Startup();
+#endif
+
+	FAvaTransitionEditorCommands::Register();
+
+	IAvaTransitionModule::Get().GetOnValidateTransitionTree().BindRaw(this, &FAvaTransitionEditorModule::ValidateStateTree);
+
+	OnPostCompileHandle = UE::StateTree::Delegates::OnPostCompile.AddRaw(this, &FAvaTransitionEditorModule::OnPostCompile);
+}
+
+void FAvaTransitionEditorModule::ShutdownModule()
+{
+	FAvaTransitionEditorCommands::Unregister();
+
+#if WITH_STATETREE_DEBUGGER
+	FAvaTransitionTraceModule::Shutdown();
+#endif
+
+	IAvaTransitionModule::Get().GetOnValidateTransitionTree().Unbind();
+
+	UE::StateTree::Delegates::OnPostCompile.Remove(OnPostCompileHandle);
+	OnPostCompileHandle.Reset();
+}
+
+IAvaTransitionEditorModule::FOnBuildDefaultTransitionTree& FAvaTransitionEditorModule::GetOnBuildDefaultTransitionTree()
+{
+	return OnBuildDefaultTransitionTree;
+}
+
+IAvaTransitionEditorModule::FOnCompileTransitionTree& FAvaTransitionEditorModule::GetOnCompileTransitionTree()
+{
+	return OnCompileTransitionTree;
+}
+
+void FAvaTransitionEditorModule::GenerateTransitionTreeOptionsMenu(UToolMenu* InMenu, IAvaTransitionBehavior* InTransitionBehavior)
+{
+	if (!InMenu || !InTransitionBehavior)
+	{
+		return;
+	}
+
+	UAvaTransitionTree* TransitionTree = InTransitionBehavior->GetTransitionTree();
+	if (!TransitionTree)
+	{
+		return;
+	}
+
+	TWeakObjectPtr<UAvaTransitionTree> TransitionTreeWeak = TransitionTree;
+
+	FToolMenuSection& GeneralSection = InMenu->FindOrAddSection(TEXT("GeneralSection"), LOCTEXT("GeneralSectionLabel", "General"));
+
+	GeneralSection.AddMenuEntry(TEXT("TransitionTreeEnabled")
+		, LOCTEXT("TransitionTreeEnabledLabel", "Enabled")
+		, LOCTEXT("TransitionTreeEnabledTooltip", "Determines whether Transition Tree is Enabled")
+		, TAttribute<FSlateIcon>()
+		, FUIAction(FExecuteAction::CreateStatic(&UE::AvaTransitionEditor::ToggleTransitionTreeEnabled, TransitionTreeWeak)
+			, FCanExecuteAction()
+			, FIsActionChecked::CreateStatic(&UE::AvaTransitionEditor::IsTransitionTreeEnabled, TransitionTreeWeak))
+		, EUserInterfaceActionType::ToggleButton);
+
+	if (TSharedPtr<SWidget> LayerPicker = UE::AvaTransitionEditor::CreateTransitionLayerPicker(Cast<UAvaTransitionTreeEditorData>(TransitionTree->EditorData)))
+	{
+		FToolMenuSection& LayerSection = InMenu->FindOrAddSection(TEXT("LayerSection"), LOCTEXT("LayerSectionLabel", "Layer"));
+
+		LayerSection.AddEntry(FToolMenuEntry::InitWidget(TEXT("TransitionLayerPicker")
+			, SNew(SBox)
+				.HeightOverride(40.f)
+				[
+					LayerPicker.ToSharedRef()
+				]
+			, FText::GetEmpty()
+			, /*bNoIndent*/true
+			, /*bSearchable*/false
+			, /*bNoPadding*/true));
+	}
+}
+
+void FAvaTransitionEditorModule::ValidateStateTree(UAvaTransitionTree* InTransitionTree)
+{
+	// Return early if Editor Data is already valid
+	if (!InTransitionTree || InTransitionTree->EditorData)
+	{
+		return;
+	}
+
+	UAvaTransitionTreeEditorData* const EditorData = NewObject<UAvaTransitionTreeEditorData>(InTransitionTree, NAME_None, RF_Transactional);
+	check(EditorData);
+
+	EditorData->Schema = NewObject<UAvaTransitionTreeSchema>(EditorData);
+
+	InTransitionTree->EditorData = EditorData;
+
+	// Build the Default Tree
+	if (OnBuildDefaultTransitionTree.IsBound())
+	{
+		OnBuildDefaultTransitionTree.Execute(*EditorData);
+	}
+	else
+	{
+		EditorData->AddRootState();	
+	}
+
+	// Use Editor View Model to Compile
+	// todo: improve Compiler class so that an Editor View Model isn't needed
+	TSharedRef<FAvaTransitionEditorViewModel> EditorViewModel = MakeShared<FAvaTransitionEditorViewModel>(InTransitionTree, /*EditorToolkit*/nullptr);
+	EditorViewModel->Initialize(/*Parent*/nullptr);
+	EditorViewModel->Compile();
+}
+
+void FAvaTransitionEditorModule::OnPostCompile(const UStateTree& InStateTree)
+{
+	// TODO: Currently State Tree Compiler does not have any hooks to compile additional data,
+	// so this is where we inject the additional properties
+	if (const UAvaTransitionTree* TransitionTree = Cast<UAvaTransitionTree>(&InStateTree))
+	{
+		UAvaTransitionTree* TransitionTreeMutable = const_cast<UAvaTransitionTree*>(TransitionTree);
+
+		const UAvaTransitionTreeEditorData* EditorData = CastChecked<UAvaTransitionTreeEditorData>(InStateTree.EditorData);
+
+		TransitionTreeMutable->SetTransitionLayer(EditorData->GetTransitionLayer());
+	}
+}
+
+#undef LOCTEXT_NAMESPACE
diff --git a/Engine/Plugins/Experimental/Avalanche/Source/AvalancheTransitionEditor/Private/AvaTransitionEditorStyle.cpp b/Engine/Plugins/Experimental/Avalanche/Source/AvalancheTransitionEditor/Private/AvaTransitionEditorStyle.cpp
new file mode 100644
--- /dev/null
+++ b/Engine/Plugins/Experimental/Avalanche/Source/AvalancheTransitionEditor/Private/AvaTransitionEditorStyle.cpp
@@ -0,0 +1,59 @@
+// Copyright Epic Games, Inc. All Rights Reserved.
+
+#include "AvaTransitionEditorStyle.h"
+#include "AvaTransitionTree.h"
+#include "Interfaces/IPluginManager.h"
+#include "Misc/Paths.h"
+#include "Styling/SlateStyleMacros.h"
+#include "Styling/SlateStyleRegistry.h"
+
+FAvaTransitionEditorStyle::FAvaTransitionEditorStyle()
+	: FSlateStyleSet(TEXT("AvaTransitionEditorStyle"))
+{
+	const FVector2f Icon16(16.f);
+	const FVector2f Icon20(20.f);
+	const FVector2f Icon64(64.f);
+
+	ContentRootDir     = FPaths::EngineContentDir() / TEXT("Editor/Slate");
+	CoreContentRootDir = FPaths::EngineContentDir() / TEXT("Slate");
+
+	// Re-use the Behavior Tree Icon
+	Set("ClassIcon.AvaTransitionTree"     , new IMAGE_BRUSH("Icons/AssetIcons/BehaviorTree_16x", Icon16));
+	Set("ClassThumbnail.AvaTransitionTree", new IMAGE_BRUSH("Icons/AssetIcons/BehaviorTree_64x", Icon64));
+
+	Set("Throbber.CircleChunk", new CORE_IMAGE_BRUSH("Common/Throbber_Piece", FVector2f(1.5f)));
+
+	// Editor Commands
+	Set("AvaTransitionEditorCommands.AddSiblingState"       , new CORE_IMAGE_BRUSH_SVG("Starship/Common/plus"     , Icon20));
+	Set("AvaTransitionEditorCommands.AddChildState"         , new CORE_IMAGE_BRUSH_SVG("Starship/Common/plus"     , Icon20));
+	Set("AvaTransitionEditorCommands.ImportTransitionTree"  , new CORE_IMAGE_BRUSH_SVG("Starship/Common/import_20", Icon20));
+	Set("AvaTransitionEditorCommands.ReimportTransitionTree", new CORE_IMAGE_BRUSH_SVG("Starship/Common/import_20", Icon20));
+	Set("AvaTransitionEditorCommands.ExportTransitionTree"  , new CORE_IMAGE_BRUSH_SVG("Starship/Common/export_20", Icon20));
+	Set("AvaTransitionEditorCommands.ToggleDebug"           , new IMAGE_BRUSH_SVG("Starship/Common/Debug"         , Icon20));
+
+	const TSharedPtr<IPlugin> Plugin = IPluginManager::Get().FindPlugin(UE_PLUGIN_NAME);
+	check(Plugin.IsValid());
+
+	ContentRootDir = (FPaths::Combine(Plugin->GetBaseDir(), TEXT("Resources")));
+
+	Set("DebugIndicatorBorder", new BOX_BRUSH("Images/NamespaceBorder", FMargin(0.25f)));
+
+	FSlateStyleRegistry::RegisterSlateStyle(*this);
+}
+
+FAvaTransitionEditorStyle::~FAvaTransitionEditorStyle()
+{
+	FSlateStyleRegistry::UnRegisterSlateStyle(*this);
+}
+
+FLinearColor FAvaTransitionEditorStyle::LerpColorSRGB(const FLinearColor& InA, const FLinearColor& InB, float InAlpha)
+{
+	const FColor A = InA.ToFColorSRGB();
+	const FColor B = InB.ToFColorSRGB();
+
+	return FLinearColor(FColor(
+		static_cast<uint8>(FMath::RoundToInt(static_cast<float>(A.R) * (1.f - InAlpha) + static_cast<float>(B.R) * InAlpha)),
+		static_cast<uint8>(FMath::RoundToInt(static_cast<float>(A.G) * (1.f - InAlpha) + static_cast<float>(B.G) * InAlpha)),
+		static_cast<uint8>(FMath::RoundToInt(static_cast<float>(A.B) * (1.f - InAlpha) + static_cast<float>(B.B) * InAlpha)),
+		static_cast<uint8>(FMath::RoundToInt(static_cast<float>(A.A) * (1.f - InAlpha) + static_cast<float>(B.A) * InAlpha))));
+}
