[MotionDesign]

diff --git a/Engine/Plugins/Experimental/Avalanche/Source/AvalancheTransitionEditor/Private/AvaTransitionEditorUtils.h b/Engine/Plugins/Experimental/Avalanche/Source/AvalancheTransitionEditor/Private/AvaTransitionEditorUtils.h
new file mode 100644
--- /dev/null
+++ b/Engine/Plugins/Experimental/Avalanche/Source/AvalancheTransitionEditor/Private/AvaTransitionEditorUtils.h
@@ -0,0 +1,22 @@
+// Copyright Epic Games, Inc. All Rights Reserved.
+
+#pragma once
+
+#include "Behavior/IAvaTransitionBehavior.h"
+#include "Templates/SharedPointerFwd.h"
+#include "UObject/WeakObjectPtr.h"
+
+class SWidget;
+class UAvaTransitionTree;
+class UAvaTransitionTreeEditorData;
+
+namespace UE::AvaTransitionEditor
+{
+	TSharedPtr<SWidget> CreateTransitionLayerPicker(UAvaTransitionTreeEditorData* InEditorData);
+
+	bool PickTransitionTreeAsset(const FText& InDialogTitle, UAvaTransitionTree*& OutTransitionTree);
+
+	void ToggleTransitionTreeEnabled(TWeakObjectPtr<UAvaTransitionTree> InTransitionTreeWeak);
+
+	bool IsTransitionTreeEnabled(TWeakObjectPtr<UAvaTransitionTree> InTransitionTreeWeak);
+}
diff --git a/Engine/Plugins/Experimental/Avalanche/Source/AvalancheTransitionEditor/Private/Compiler/AvaTransitionCompiler.cpp b/Engine/Plugins/Experimental/Avalanche/Source/AvalancheTransitionEditor/Private/Compiler/AvaTransitionCompiler.cpp
new file mode 100644
--- /dev/null
+++ b/Engine/Plugins/Experimental/Avalanche/Source/AvalancheTransitionEditor/Private/Compiler/AvaTransitionCompiler.cpp
@@ -0,0 +1,211 @@
+// Copyright Epic Games, Inc. All Rights Reserved.
+
+#include "AvaTransitionCompiler.h"
+#include "AvaTransitionCommands.h"
+#include "AvaTransitionEditor.h"
+#include "AvaTransitionSelection.h"
+#include "AvaTransitionTree.h"
+#include "AvaTransitionTreeEditorData.h"
+#include "Extensions/IAvaTransitionObjectExtension.h"
+#include "FileHelpers.h"
+#include "IAvaTransitionEditorModule.h"
+#include "IMessageLogListing.h"
+#include "MessageLogModule.h"
+#include "Misc/UObjectToken.h"
+#include "StateTreeCompiler.h"
+#include "StateTreeCompilerLog.h"
+#include "StateTreeDelegates.h"
+#include "StateTreeEditorSettings.h"
+#include "StateTreeState.h"
+#include "TabFactories/AvaTransitionCompilerResultsTabFactory.h"
+#include "ViewModels/AvaTransitionEditorViewModel.h"
+#include "ViewModels/AvaTransitionViewModelSharedData.h"
+#include "ViewModels/Registry/AvaTransitionViewModelRegistryCollection.h"
+
+#define LOCTEXT_NAMESPACE "AvaTransitionCompiler"
+
+FAvaTransitionCompiler::FAvaTransitionCompiler(FAvaTransitionEditorViewModel& InOwner)
+	: Owner(InOwner)
+{
+	FMessageLogModule& MessageLogModule = FModuleManager::LoadModuleChecked<FMessageLogModule>("MessageLog");
+
+	// Show Pages so that user is never allowed to clear log messages
+	FMessageLogInitializationOptions LogOptions;
+	LogOptions.bShowPages   = false;
+	LogOptions.bShowFilters = false;
+	LogOptions.bAllowClear  = false;
+	LogOptions.MaxPageCount = 1;
+
+	CompilerResultsListing = MessageLogModule.CreateLogListing("AvaTransitionTreeCompiler", LogOptions);
+	CompilerResultsWidget  = MessageLogModule.CreateLogListingWidget(CompilerResultsListing.ToSharedRef());
+
+	CompilerResultsListing->OnMessageTokenClicked().AddRaw(this, &FAvaTransitionCompiler::OnMessageTokenClicked);
+}
+
+bool FAvaTransitionCompiler::Compile()
+{
+	UAvaTransitionTree* TransitionTree = Owner.GetTransitionTree();
+	if (!TransitionTree)
+	{
+		return false;
+	}
+
+	if (UAvaTransitionTreeEditorData* EditorData = Owner.GetEditorData())
+	{
+		// Update Tree's Transition Layer
+		TransitionTree->SetTransitionLayer(EditorData->GetTransitionLayer());
+
+		const EAvaTransitionEditorMode EditorMode = Owner.GetSharedData()->GetEditorMode();
+
+		// Only allow extensible compilation when it's not in Advanced Mode
+		if (EditorMode != EAvaTransitionEditorMode::Advanced)
+		{
+			// Remove all invalid Sub Trees, if any
+			EditorData->SubTrees.RemoveAll(
+				[](const UStateTreeState* InState)
+				{
+					return !IsValid(InState);
+				});
+
+			IAvaTransitionEditorModule::Get().GetOnCompileTransitionTree().Broadcast(*EditorData);
+		}
+	}
+
+	Owner.UpdateTree();
+
+	FStateTreeCompilerLog Log;
+	FStateTreeCompiler Compiler(Log);
+
+	bLastCompileSucceeded = Compiler.Compile(*TransitionTree);
+
+	if (CompilerResultsListing.IsValid())
+	{
+		CompilerResultsListing->ClearMessages();
+		Log.AppendToLog(CompilerResultsListing.Get());
+	}
+
+	if (bLastCompileSucceeded)
+	{
+		// Success
+		TransitionTree->LastCompiledEditorDataHash = Owner.GetEditorDataHash();
+		UE::StateTree::Delegates::OnPostCompile.Broadcast(*TransitionTree);
+	}
+	else
+	{
+		// Make sure not to leave stale data on failed compile.
+		TransitionTree->ResetCompiled();
+		TransitionTree->LastCompiledEditorDataHash = 0;
+
+		if (TSharedPtr<FAvaTransitionEditor> Editor = Owner.GetEditor())
+		{
+			if (TSharedPtr<FTabManager> TabManager = Editor->GetTabManager())
+			{
+				TabManager->TryInvokeTab(FAvaTransitionCompilerResultsTabFactory::TabId);
+			}
+		}
+	}
+
+	const UStateTreeEditorSettings* Settings = GetMutableDefault<UStateTreeEditorSettings>();
+
+	const bool bShouldSaveOnCompile = Settings->SaveOnCompile == EStateTreeSaveOnCompile::Always
+		|| (Settings->SaveOnCompile == EStateTreeSaveOnCompile::SuccessOnly && bLastCompileSucceeded);
+
+	if (bShouldSaveOnCompile)
+	{
+		FEditorFileUtils::PromptForCheckoutAndSave({ TransitionTree->GetPackage() }
+			, /*bCheckDirty*/true
+			, /*bPromptToSave*/false);
+	}
+
+	return true;
+}
+
+FSlateIcon FAvaTransitionCompiler::GetCompileStatusIcon() const
+{
+	static const FName CompileStatusBackground("Blueprint.CompileStatus.Background");
+	static const FName CompileStatusUnknown("Blueprint.CompileStatus.Overlay.Unknown");
+	static const FName CompileStatusError("Blueprint.CompileStatus.Overlay.Error");
+	static const FName CompileStatusGood("Blueprint.CompileStatus.Overlay.Good");
+	static const FName CompileStatusWarning("Blueprint.CompileStatus.Overlay.Warning");
+
+	UAvaTransitionTree* TransitionTree = Owner.GetTransitionTree();
+	if (!TransitionTree)
+	{
+		return FSlateIcon(FAppStyle::GetAppStyleSetName(), CompileStatusBackground, NAME_None, CompileStatusUnknown);
+	}
+
+	const uint32 EditorDataHash = Owner.GetEditorDataHash();
+
+	const bool bCompiledDataResetDuringLoad = TransitionTree->LastCompiledEditorDataHash == EditorDataHash && !TransitionTree->IsReadyToRun();
+
+	if (!bLastCompileSucceeded || bCompiledDataResetDuringLoad)
+	{
+		return FSlateIcon(FAppStyle::GetAppStyleSetName(), CompileStatusBackground, NAME_None, CompileStatusError);
+	}
+
+	if (TransitionTree->LastCompiledEditorDataHash != EditorDataHash)
+	{
+		return FSlateIcon(FAppStyle::GetAppStyleSetName(), CompileStatusBackground, NAME_None, CompileStatusUnknown);
+	}
+
+	return FSlateIcon(FAppStyle::GetAppStyleSetName(), CompileStatusBackground, NAME_None, CompileStatusGood);
+}
+
+TSharedRef<SWidget> FAvaTransitionCompiler::GetCompilerResultsWidget() const
+{
+	return CompilerResultsWidget.ToSharedRef();
+}
+
+void FAvaTransitionCompiler::SetSaveOnCompile(EStateTreeSaveOnCompile InSaveOnCompileType)
+{
+	UStateTreeEditorSettings* Settings = GetMutableDefault<UStateTreeEditorSettings>();
+	check(Settings);
+	Settings->SaveOnCompile = InSaveOnCompileType;
+	Settings->SaveConfig();
+}
+
+bool FAvaTransitionCompiler::HasSaveOnCompile(EStateTreeSaveOnCompile InSaveOnCompileType)
+{
+	const UStateTreeEditorSettings* Settings = GetDefault<UStateTreeEditorSettings>();
+	check(Settings);
+	return Settings->SaveOnCompile == InSaveOnCompileType;
+}
+
+void FAvaTransitionCompiler::GenerateCompileOptionsMenu(UToolMenu* InMenu)
+{
+	auto MakeSaveOnCompileMenu = [](UToolMenu* InSubMenu)
+	{
+		FToolMenuSection& Section = InSubMenu->AddSection("CompileOptions");
+
+		const FAvaTransitionEditorCommands& Commands = FAvaTransitionEditorCommands::Get();
+		Section.AddMenuEntry(Commands.SaveOnCompile_Never);
+		Section.AddMenuEntry(Commands.SaveOnCompile_SuccessOnly);
+		Section.AddMenuEntry(Commands.SaveOnCompile_Always);
+	};
+
+	FToolMenuSection& Section = InMenu->AddSection(TEXT("CompileOptions"));
+	Section.AddSubMenu("SaveOnCompile"
+		, LOCTEXT("SaveOnCompileSubMenu", "Save on Compile")
+		, LOCTEXT("SaveOnCompileSubMenu_ToolTip", "Determines how the StateTree is saved whenever you compile it.")
+		, FNewToolMenuDelegate::CreateLambda(MakeSaveOnCompileMenu));
+}
+
+void FAvaTransitionCompiler::OnMessageTokenClicked(const TSharedRef<IMessageToken>& InMessageToken)
+{
+	if (InMessageToken->GetType() == EMessageToken::Object)
+	{
+		const TSharedRef<FUObjectToken> ObjectToken = StaticCastSharedRef<FUObjectToken>(InMessageToken);
+
+		if (UStateTreeState* State = Cast<UStateTreeState>(ObjectToken->GetObject().Get()))
+		{
+			TSharedRef<FAvaTransitionViewModelRegistryCollection> RegistryCollection = Owner.GetSharedData()->GetRegistryCollection();
+
+			if (TSharedPtr<FAvaTransitionViewModel> FoundViewModel = RegistryCollection->FindViewModel(State))
+			{
+				Owner.GetSelection()->SetSelectedItems({ FoundViewModel });
+			}
+		}
+	}
+}
+
+#undef LOCTEXT_NAMESPACE
diff --git a/Engine/Plugins/Experimental/Avalanche/Source/AvalancheTransitionEditor/Private/Compiler/AvaTransitionCompiler.h b/Engine/Plugins/Experimental/Avalanche/Source/AvalancheTransitionEditor/Private/Compiler/AvaTransitionCompiler.h
new file mode 100644
--- /dev/null
+++ b/Engine/Plugins/Experimental/Avalanche/Source/AvalancheTransitionEditor/Private/Compiler/AvaTransitionCompiler.h
@@ -0,0 +1,42 @@
+// Copyright Epic Games, Inc. All Rights Reserved.
+
+#pragma once
+
+#include "Templates/SharedPointer.h"
+
+class FAvaTransitionEditorViewModel;
+class IMessageLogListing;
+class IMessageToken;
+class SWidget;
+class UToolMenu;
+enum class EStateTreeSaveOnCompile : uint8;
+struct FSlateIcon;
+
+class FAvaTransitionCompiler
+{
+public:
+	FAvaTransitionCompiler(FAvaTransitionEditorViewModel& InOwner);
+
+	bool Compile();
+
+	FSlateIcon GetCompileStatusIcon() const;
+
+	TSharedRef<SWidget> GetCompilerResultsWidget() const;
+
+	static void SetSaveOnCompile(EStateTreeSaveOnCompile InSaveOnCompileType);
+
+	static bool HasSaveOnCompile(EStateTreeSaveOnCompile InSaveOnCompileType);
+
+	static void GenerateCompileOptionsMenu(UToolMenu* InMenu);
+
+private:
+	void OnMessageTokenClicked(const TSharedRef<IMessageToken>& InMessageToken);
+
+	FAvaTransitionEditorViewModel& Owner;
+
+	TSharedPtr<IMessageLogListing> CompilerResultsListing;
+
+	TSharedPtr<SWidget> CompilerResultsWidget;
+
+	bool bLastCompileSucceeded = true;
+};
diff --git a/Engine/Plugins/Experimental/Avalanche/Source/AvalancheTransitionEditor/Private/Debugger/AvaTransitionTraceAnalyzer.cpp b/Engine/Plugins/Experimental/Avalanche/Source/AvalancheTransitionEditor/Private/Debugger/AvaTransitionTraceAnalyzer.cpp
new file mode 100644
--- /dev/null
+++ b/Engine/Plugins/Experimental/Avalanche/Source/AvalancheTransitionEditor/Private/Debugger/AvaTransitionTraceAnalyzer.cpp
@@ -0,0 +1,237 @@
+// Copyright Epic Games, Inc. All Rights Reserved.
+
+#if WITH_STATETREE_DEBUGGER
+
+#include "AvaTransitionTraceAnalyzer.h"
+#include "AvaTransitionDebugger.h"
+#include "AvaTransitionEditorLog.h"
+#include "Containers/Ticker.h"
+#include "Debugger/StateTreeDebuggerTypes.h"
+#include "StateTree.h"
+#include "StateTreeExecutionTypes.h"
+#include "StateTreeModule.h"
+#include "TraceServices/Model/AnalysisSession.h"
+
+namespace UE::AvaTransitionEditor::Private
+{
+	EStateTreeTraceEventType GetEventType(const Trace::IAnalyzer::FOnEventContext& InContext)
+	{
+		return InContext.EventData.GetValue<EStateTreeTraceEventType>("EventType");
+	}
+
+	FStateTreeIndex16 GetAssetDebugId(const Trace::IAnalyzer::FOnEventContext& InContext)
+	{
+		return FStateTreeIndex16(InContext.EventData.GetValue<uint16>("AssetDebugId"));
+	}
+
+	FStateTreeInstanceDebugId GetInstanceDebugId(const Trace::IAnalyzer::FOnEventContext& InContext)
+	{
+		return FStateTreeInstanceDebugId(InContext.EventData.GetValue<uint32>("InstanceId"), InContext.EventData.GetValue<uint32>("InstanceSerial"));
+	}
+}
+
+TArray<FAvaTransitionDebugger*> FAvaTransitionTraceAnalyzer::ActiveDebuggers;
+
+const TArray<FAvaTransitionTraceAnalyzer::FEvent> FAvaTransitionTraceAnalyzer::Events =
+{
+	{ "AssetDebugIdEvent", &FAvaTransitionTraceAnalyzer::OnAssetDebugIdEvent },
+	{ "InstanceEvent"    , &FAvaTransitionTraceAnalyzer::OnInstanceEvent },
+	{ "StateEvent"       , &FAvaTransitionTraceAnalyzer::OnStateEvent },
+};
+
+void FAvaTransitionTraceAnalyzer::RegisterDebugger(FAvaTransitionDebugger* InDebugger)
+{
+	ActiveDebuggers.Add(InDebugger);
+
+	if (IStateTreeModule::IsAvailable())
+	{
+		int32 TraceId = 0;
+        IStateTreeModule::Get().StartTraces(TraceId);
+	}
+}
+
+void FAvaTransitionTraceAnalyzer::UnregisterDebugger(FAvaTransitionDebugger* InDebugger)
+{
+	bool bRemovedDebugger = ActiveDebuggers.Remove(InDebugger) > 0;
+
+	// Stop Traces only if the passed in Debugger was the last of the Active Debuggers
+	if (bRemovedDebugger && ActiveDebuggers.IsEmpty() && IStateTreeModule::IsAvailable())
+	{
+		IStateTreeModule::Get().StopTraces();
+	}
+}
+
+void FAvaTransitionTraceAnalyzer::OnAnalysisBegin(const FOnAnalysisContext& InContext)
+{
+	for (int32 RouteId = 0; RouteId < Events.Num(); ++RouteId)
+	{
+		InContext.InterfaceBuilder.RouteEvent(RouteId, "StateTreeDebugger", Events[RouteId].Name);
+	}
+}
+
+bool FAvaTransitionTraceAnalyzer::OnEvent(uint16 InRouteId, EStyle InStyle, const FOnEventContext& InContext)
+{
+	if (Events.IsValidIndex(InRouteId))
+	{
+		::Invoke(Events[InRouteId].Handler, this, InContext);
+	}
+	return true;
+}
+
+void FAvaTransitionTraceAnalyzer::ForEachDebugger(const TCHAR* InDebugName, TFunction<void(FAvaTransitionDebugger&)>&& InCallable)
+{
+	ExecuteOnGameThread(InDebugName,
+		[Callable = MoveTemp(InCallable)]
+		{
+			for (FAvaTransitionDebugger* Debugger : ActiveDebuggers)
+			{
+				check(Debugger);
+				Callable(*Debugger);
+			}
+		});
+}
+
+UStateTree* FAvaTransitionTraceAnalyzer::FindStateTree(const FStateTreeInstanceDebugId& InDebugInstanceId) const
+{
+	const FStateTreeIndex16* AssetDebugId = DebugInstances.Find(InDebugInstanceId);
+	if (!AssetDebugId)
+	{
+		return nullptr;
+	}
+
+	const TWeakObjectPtr<UStateTree>* DebugAsset = DebugAssets.Find(*AssetDebugId);
+	if (!DebugAsset)
+	{
+		return nullptr;
+	}
+
+	return DebugAsset->Get();
+}
+
+void FAvaTransitionTraceAnalyzer::OnAssetDebugIdEvent(const FOnEventContext& InContext)
+{
+	FString TreeName, TreePathName;
+	InContext.EventData.GetString("TreeName", TreeName);
+	InContext.EventData.GetString("TreePath", TreePathName);
+
+	UE_LOG(LogAvaEditorTransition, Verbose, TEXT("OnAssetDebugIdEvent --- Attempting to find tree %s in path %s"), *TreeName, *TreePathName);
+
+	UStateTree* StateTree;
+	{
+		FGCScopeGuard Guard;
+		StateTree = FindObject<UStateTree>(nullptr, *TreePathName);
+	}
+
+	const uint32 CompiledDataHash = InContext.EventData.GetValue<uint32>("CompiledDataHash");
+	if (StateTree && StateTree->LastCompiledEditorDataHash == CompiledDataHash)
+	{
+		FStateTreeIndex16 AssetDebugId = UE::AvaTransitionEditor::Private::GetAssetDebugId(InContext);
+		DebugAssets.Add(AssetDebugId, StateTree);
+	}
+	else
+	{
+		UE_LOG(LogAvaEditorTransition, Warning, TEXT("Unable to find StateTree asset: %s : %s"), *TreePathName, *TreeName);
+	}
+}
+
+void FAvaTransitionTraceAnalyzer::OnInstanceEvent(const FOnEventContext& InContext)
+{
+	EStateTreeTraceEventType EventType = UE::AvaTransitionEditor::Private::GetEventType(InContext);
+
+	switch (EventType)
+	{
+	case EStateTreeTraceEventType::Push:
+		OnTreeInstanceStarted(InContext);
+		break;
+	case EStateTreeTraceEventType::Pop:
+		OnTreeInstanceStopped(InContext);
+		break;
+	}
+}
+
+void FAvaTransitionTraceAnalyzer::OnTreeInstanceStarted(const FOnEventContext& InContext)
+{
+	using namespace UE::AvaTransitionEditor;
+
+	FStateTreeInstanceDebugId InstanceId = UE::AvaTransitionEditor::Private::GetInstanceDebugId(InContext);
+	FStateTreeIndex16 AssetId = UE::AvaTransitionEditor::Private::GetAssetDebugId(InContext);
+
+	DebugInstances.Add(InstanceId, AssetId);
+
+	FString InstanceName;
+	InContext.EventData.GetString("InstanceName", InstanceName);
+
+	if (TWeakObjectPtr<UStateTree>* DebugAsset = DebugAssets.Find(AssetId))
+	{
+		ForEachDebugger(UE_SOURCE_LOCATION,
+			[StateTreeWeak = *DebugAsset, InstanceId, InstanceName](FAvaTransitionDebugger& InDebugger)
+			{
+				if (UStateTree* StateTree = StateTreeWeak.Get())
+				{
+					InDebugger.OnTreeInstanceStarted(*StateTree, InstanceId, InstanceName);
+				}
+			});
+	}
+}
+
+void FAvaTransitionTraceAnalyzer::OnTreeInstanceStopped(const FOnEventContext& InContext)
+{
+	FStateTreeInstanceDebugId InstanceDebugId = UE::AvaTransitionEditor::Private::GetInstanceDebugId(InContext);
+
+	DebugInstances.Remove(InstanceDebugId);
+
+	ForEachDebugger(UE_SOURCE_LOCATION,
+		[InstanceDebugId](FAvaTransitionDebugger& InDebugger)
+		{
+			InDebugger.OnTreeInstanceStopped(InstanceDebugId);
+		});
+}
+
+void FAvaTransitionTraceAnalyzer::OnStateEvent(const FOnEventContext& InContext)
+{
+	EStateTreeTraceEventType EventType = UE::AvaTransitionEditor::Private::GetEventType(InContext);
+
+	// Ignore other events for now. Only concerned with enter/exit
+	if (EventType != EStateTreeTraceEventType::OnEntered && EventType != EStateTreeTraceEventType::OnExited)
+	{
+		return;
+	}
+
+	FStateTreeInstanceDebugId InstanceDebugId = UE::AvaTransitionEditor::Private::GetInstanceDebugId(InContext);
+	if (!InstanceDebugId.IsValid())
+	{
+		return;
+	}
+
+	UStateTree* StateTree = FindStateTree(InstanceDebugId);
+	if (!StateTree)
+	{
+		return;
+	}
+
+	FStateTreeStateHandle StateHandle(InContext.EventData.GetValue<uint16>("StateIndex"));
+
+	FGuid StateId = StateTree->GetStateIdFromHandle(StateHandle);
+	if (!StateId.IsValid())
+	{
+		UE_LOG(LogAvaEditorTransition, Warning, TEXT("Could not find id of state with index %d"), StateHandle.Index);
+		return;
+	}
+
+	ForEachDebugger(UE_SOURCE_LOCATION,
+		[EventType, StateId, InstanceDebugId](FAvaTransitionDebugger& InDebugger)
+		{
+			switch (EventType)
+			{
+			case EStateTreeTraceEventType::OnEntered:
+				InDebugger.OnNodeEntered(StateId, InstanceDebugId);
+				break;
+
+			case EStateTreeTraceEventType::OnExited:
+				InDebugger.OnNodeExited(StateId, InstanceDebugId);
+				break;
+			}
+		});
+}
+
+#endif // WITH_STATETREE_DEBUGGER
diff --git a/Engine/Plugins/Experimental/Avalanche/Source/AvalancheTransitionEditor/Private/Debugger/AvaTransitionTraceAnalyzer.h b/Engine/Plugins/Experimental/Avalanche/Source/AvalancheTransitionEditor/Private/Debugger/AvaTransitionTraceAnalyzer.h
new file mode 100644
--- /dev/null
+++ b/Engine/Plugins/Experimental/Avalanche/Source/AvalancheTransitionEditor/Private/Debugger/AvaTransitionTraceAnalyzer.h
@@ -0,0 +1,58 @@
+// Copyright Epic Games, Inc. All Rights Reserved.
+
+#pragma once
+
+#if WITH_STATETREE_DEBUGGER
+
+#include "StateTreeExecutionTypes.h"
+#include "StateTreeIndexTypes.h"
+#include "Trace/Analyzer.h"
+#include "UObject/WeakObjectPtr.h"
+
+class FAvaTransitionDebugger;
+class UStateTree;
+
+class FAvaTransitionTraceAnalyzer : public UE::Trace::IAnalyzer
+{
+	struct FEvent
+	{
+		const ANSICHAR* Name;
+		void(FAvaTransitionTraceAnalyzer::*Handler)(const FOnEventContext&);
+	};
+
+public:
+	static void RegisterDebugger(FAvaTransitionDebugger* InDebugger);
+	static void UnregisterDebugger(FAvaTransitionDebugger* InDebugger);
+
+	//~ Begin IAnalyzer
+	virtual void OnAnalysisBegin(const FOnAnalysisContext& InContext) override;
+	virtual bool OnEvent(uint16 InRouteId, EStyle InStyle, const FOnEventContext& InContext) override;
+	//~ End IAnalyzer
+
+private:
+	static TArray<FAvaTransitionDebugger*> ActiveDebuggers;
+
+	static const TArray<FEvent> Events;
+
+	void ForEachDebugger(const TCHAR* InDebugName, TFunction<void(FAvaTransitionDebugger&)>&& InCallable);
+
+	UStateTree* FindStateTree(const FStateTreeInstanceDebugId& InDebugInstanceId) const;
+
+	void OnAssetDebugIdEvent(const FOnEventContext& InContext);
+
+	void OnInstanceEvent(const FOnEventContext& InContext);
+
+	void OnTreeInstanceStarted(const FOnEventContext& InContext);
+
+	void OnTreeInstanceStopped(const FOnEventContext& InContext);
+
+	void OnStateEvent(const FOnEventContext& InContext);
+
+	/** Mapping of the Asset Debug Id its State Tree */
+	TMap<FStateTreeIndex16, TWeakObjectPtr<UStateTree>> DebugAssets;
+
+	/** Mapping of the Instance Debug Id its Asset Debug Id */
+	TMap<FStateTreeInstanceDebugId, FStateTreeIndex16> DebugInstances;
+};
+
+#endif //WITH_STATETREE_DEBUGGER
diff --git a/Engine/Plugins/Experimental/Avalanche/Source/AvalancheTransitionEditor/Private/Debugger/AvaTransitionTraceModule.cpp b/Engine/Plugins/Experimental/Avalanche/Source/AvalancheTransitionEditor/Private/Debugger/AvaTransitionTraceModule.cpp
new file mode 100644
--- /dev/null
+++ b/Engine/Plugins/Experimental/Avalanche/Source/AvalancheTransitionEditor/Private/Debugger/AvaTransitionTraceModule.cpp
@@ -0,0 +1,37 @@
+// Copyright Epic Games, Inc. All Rights Reserved.
+
+#if WITH_STATETREE_DEBUGGER
+
+#include "AvaTransitionTraceModule.h"
+#include "AvaTransitionTraceAnalyzer.h"
+#include "Features/IModularFeatures.h"
+#include "TraceServices/Model/AnalysisSession.h"
+
+FAvaTransitionTraceModule* FAvaTransitionTraceModule::Get()
+{
+	static FAvaTransitionTraceModule TraceModule;
+	return &TraceModule;
+}
+
+void FAvaTransitionTraceModule::Startup()
+{
+	IModularFeatures::Get().RegisterModularFeature(TraceServices::ModuleFeatureName, Get());
+}
+
+void FAvaTransitionTraceModule::Shutdown()
+{
+	IModularFeatures::Get().UnregisterModularFeature(TraceServices::ModuleFeatureName, Get());
+}
+
+void FAvaTransitionTraceModule::GetModuleInfo(TraceServices::FModuleInfo& OutModuleInfo)
+{
+	OutModuleInfo.Name = TEXT("TraceModule_AvaTransition");
+	OutModuleInfo.DisplayName = TEXT("Motion Design Transition Trace");
+}
+
+void FAvaTransitionTraceModule::OnAnalysisBegin(TraceServices::IAnalysisSession& InSession)
+{
+	InSession.AddAnalyzer(new FAvaTransitionTraceAnalyzer);
+}
+
+#endif // WITH_STATETREE_DEBUGGER
diff --git a/Engine/Plugins/Experimental/Avalanche/Source/AvalancheTransitionEditor/Private/Debugger/AvaTransitionTraceModule.h b/Engine/Plugins/Experimental/Avalanche/Source/AvalancheTransitionEditor/Private/Debugger/AvaTransitionTraceModule.h
new file mode 100644
--- /dev/null
+++ b/Engine/Plugins/Experimental/Avalanche/Source/AvalancheTransitionEditor/Private/Debugger/AvaTransitionTraceModule.h
@@ -0,0 +1,25 @@
+// Copyright Epic Games, Inc. All Rights Reserved.
+
+#pragma once
+
+#if WITH_STATETREE_DEBUGGER
+
+#include "Trace/StoreClient.h"
+#include "TraceServices/ModuleService.h"
+
+class FAvaTransitionTraceModule : public TraceServices::IModule
+{
+	static FAvaTransitionTraceModule* Get();
+
+public:
+	static void Startup();
+	static void Shutdown();
+
+private:
+	//~ Begin IModule
+	virtual void GetModuleInfo(TraceServices::FModuleInfo& OutModuleInfo) override;
+	virtual void OnAnalysisBegin(TraceServices::IAnalysisSession& InSession) override;
+	//~ End IModule
+};
+
+#endif // WITH_STATETREE_DEBUGGER
diff --git a/Engine/Plugins/Experimental/Avalanche/Source/AvalancheTransitionEditor/Private/Serialization/AvaTransitionStateExportData.h b/Engine/Plugins/Experimental/Avalanche/Source/AvalancheTransitionEditor/Private/Serialization/AvaTransitionStateExportData.h
new file mode 100644
--- /dev/null
+++ b/Engine/Plugins/Experimental/Avalanche/Source/AvalancheTransitionEditor/Private/Serialization/AvaTransitionStateExportData.h
@@ -0,0 +1,23 @@
+// Copyright Epic Games, Inc. All Rights Reserved.
+
+#pragma once
+
+#include "Containers/Array.h"
+#include "State/AvaTransitionAutogeneratedStateContext.h"
+#include "StateTreePropertyBindings.h"
+#include "UObject/Object.h"
+#include "AvaTransitionStateExportData.generated.h"
+
+UCLASS(Hidden)
+class UAvaTransitionStateExportData : public UObject
+{
+	GENERATED_BODY()
+
+public:
+	UPROPERTY()
+	TArray<FStateTreePropertyPathBinding> Bindings;
+
+	/** Map of an autogenerated state id to their context information */
+	UPROPERTY()
+	TMap<FGuid, FAvaTransitionAutogeneratedStateContext> AutogeneratedStates;
+};
diff --git a/Engine/Plugins/Experimental/Avalanche/Source/AvalancheTransitionEditor/Private/TabFactories/AvaTransitionTreeDetailsTabFactory.cpp b/Engine/Plugins/Experimental/Avalanche/Source/AvalancheTransitionEditor/Private/TabFactories/AvaTransitionTreeDetailsTabFactory.cpp
new file mode 100644
--- /dev/null
+++ b/Engine/Plugins/Experimental/Avalanche/Source/AvalancheTransitionEditor/Private/TabFactories/AvaTransitionTreeDetailsTabFactory.cpp
@@ -0,0 +1,69 @@
+// Copyright Epic Games, Inc. All Rights Reserved.
+
+#include "AvaTransitionTreeDetailsTabFactory.h"
+#include "AvaTransitionTreeEditorData.h"
+#include "AvaTypeSharedPointer.h"
+#include "Customizations/AvaTransitionTreeEditorDataCustomization.h"
+#include "IDetailsView.h"
+#include "Modules/ModuleManager.h"
+#include "PropertyEditorModule.h"
+#include "ViewModels/AvaTransitionEditorViewModel.h"
+#include "ViewModels/AvaTransitionViewModelSharedData.h"
+#include "Widgets/SNullWidget.h"
+
+#define LOCTEXT_NAMESPACE "AvaTransitionTreeDetailsTabFactory"
+
+const FName FAvaTransitionTreeDetailsTabFactory::TabId(TEXT("AvaTransitionTreeDetails"));
+
+FAvaTransitionTreeDetailsTabFactory::FAvaTransitionTreeDetailsTabFactory(const TSharedRef<FAvaTransitionEditor>& InEditor)
+	: FAvaTransitionTabFactory(TabId, InEditor)
+{
+	TabIcon             = FSlateIcon(FAppStyle::GetAppStyleSetName(), "LevelEditor.Tabs.Details");
+	TabLabel            = LOCTEXT("TabLabel", "State Tree");
+	ViewMenuTooltip     = LOCTEXT("ViewMenuTooltip", "State Tree Details");
+	ViewMenuDescription = LOCTEXT("ViewMenuDescription", "State Tree Details");
+	bIsSingleton        = true;
+	ReadOnlyBehavior    = ETabReadOnlyBehavior::Custom;
+}
+
+TSharedRef<SWidget> FAvaTransitionTreeDetailsTabFactory::CreateTabBody(const FWorkflowTabSpawnInfo& InInfo) const
+{
+	TSharedPtr<FAvaTransitionEditor> Editor = GetEditor();
+	if (!ensure(Editor.IsValid()))
+	{
+		return SNullWidget::NullWidget;
+	}
+
+	TSharedPtr<FAvaTransitionEditorViewModel> EditorViewModel = Editor->GetEditorViewModel();
+	if (!ensure(EditorViewModel.IsValid()))
+	{
+		return SNullWidget::NullWidget;
+	}
+
+	UAvaTransitionTreeEditorData* EditorData = EditorViewModel->GetEditorData();
+	if (!ensure(EditorData))
+	{
+		return SNullWidget::NullWidget;
+	}
+
+	FDetailsViewArgs DetailsViewArgs;
+	DetailsViewArgs.NameAreaSettings = FDetailsViewArgs::HideNameArea;
+
+	FPropertyEditorModule& PropertyEditorModule = FModuleManager::GetModuleChecked<FPropertyEditorModule>("PropertyEditor");
+	TSharedRef<IDetailsView> DetailsView = PropertyEditorModule.CreateDetailView(DetailsViewArgs);
+
+	DetailsView->RegisterInstancedCustomPropertyLayout(UStateTreeEditorData::StaticClass()
+		, FOnGetDetailCustomizationInstance::CreateStatic(&FAvaTransitionTreeEditorDataCustomization::MakeInstance));
+
+	// Read-only
+	if (EditorViewModel->GetSharedData()->IsReadOnly())
+	{
+		DetailsView->SetIsPropertyEditingEnabledDelegate(FIsPropertyEditingEnabled::CreateLambda([]{ return false; }));
+	}
+
+	DetailsView->SetObject(EditorData);
+
+	return DetailsView;
+}
+
+#undef LOCTEXT_NAMESPACE
diff --git a/Engine/Plugins/Experimental/Avalanche/Source/AvalancheTransitionEditor/Private/ViewModels/AvaTransitionEditorViewModel.cpp b/Engine/Plugins/Experimental/Avalanche/Source/AvalancheTransitionEditor/Private/ViewModels/AvaTransitionEditorViewModel.cpp
new file mode 100644
--- /dev/null
+++ b/Engine/Plugins/Experimental/Avalanche/Source/AvalancheTransitionEditor/Private/ViewModels/AvaTransitionEditorViewModel.cpp
@@ -0,0 +1,456 @@
+// Copyright Epic Games, Inc. All Rights Reserved.
+
+#include "AvaTransitionEditorViewModel.h"
+#include "Actions/AvaTransitionStateActions.h"
+#include "Actions/AvaTransitionTreeActions.h"
+#include "AvaTransitionCommands.h"
+#include "AvaTransitionEditor.h"
+#include "AvaTransitionEditorLog.h"
+#include "AvaTransitionEditorUtils.h"
+#include "AvaTransitionSelection.h"
+#include "AvaTransitionTree.h"
+#include "AvaTransitionTreeEditorData.h"
+#include "AvaTransitionViewModelChildren.h"
+#include "AvaTransitionViewModelSharedData.h"
+#include "Menu/AvaTransitionToolbar.h"
+#include "Menu/AvaTransitionTreeContextMenu.h"
+#include "Serialization/ArchiveObjectCrc32.h"
+#include "State/AvaTransitionStateViewModel.h"
+#include "StateTreeDelegates.h"
+#include "StateTreeEditorSettings.h"
+#include "StateTreeTaskBase.h"
+#include "Views/SAvaTransitionTreeView.h"
+#include "Widgets/Layout/SScrollBar.h"
+#include "Widgets/Layout/SScrollBox.h"
+
+#define LOCTEXT_NAMESPACE "AvaTransitionEditorViewModel"
+
+FAvaTransitionEditorViewModel::FAvaTransitionEditorViewModel(UAvaTransitionTree* InTransitionTree, const TSharedPtr<FAvaTransitionEditor>& InEditor)
+	: Toolbar(MakeShared<FAvaTransitionToolbar>(*this))
+	, ContextMenu(MakeShared<FAvaTransitionTreeContextMenu>(*this))
+	, TransitionTreeWeak(InTransitionTree)
+	, EditorWeak(InEditor)
+	, Compiler(*this)
+	, CommandList(MakeShared<FUICommandList>())
+{
+}
+
+FAvaTransitionEditorViewModel::~FAvaTransitionEditorViewModel()
+{
+	UnbindDelegates();
+}
+
+void FAvaTransitionEditorViewModel::BindCommands(const TSharedRef<FUICommandList>& InCommandList)
+{
+	Actions =
+		{
+			MakeShared<FAvaTransitionTreeActions>(*this),
+			MakeShared<FAvaTransitionStateActions>(*this),
+		};
+
+	InCommandList->Append(CommandList);
+
+	for (const TSharedRef<FAvaTransitionActions>& Action : Actions)
+	{
+		Action->BindCommands(CommandList);
+	}
+}
+
+const FAvaTransitionCompiler& FAvaTransitionEditorViewModel::GetCompiler() const
+{
+	return Compiler;
+}
+
+bool FAvaTransitionEditorViewModel::CanCompile() const
+{
+	UAvaTransitionTree* TransitionTree = GetTransitionTree();
+
+	bool bReadOnly = GetSharedData()->IsReadOnly();
+
+	return !bReadOnly
+		&& !GEditor->IsPlayingSessionInEditor()
+		&& TransitionTree;
+}
+
+void FAvaTransitionEditorViewModel::Compile()
+{
+	if (CanCompile())
+	{
+		UpdateTree();
+		Compiler.Compile();
+	}
+}
+
+UAvaTransitionTree* FAvaTransitionEditorViewModel::GetTransitionTree() const
+{
+	return TransitionTreeWeak.Get();
+}
+
+UAvaTransitionTreeEditorData* FAvaTransitionEditorViewModel::GetEditorData() const
+{
+	return EditorDataWeak.Get();
+}
+
+bool FAvaTransitionEditorViewModel::UpdateEditorData(bool bInCreateIfNotFound)
+{
+	UAvaTransitionTree* TransitionTree = GetTransitionTree();
+	if (!TransitionTree)
+	{
+		return false;
+	}
+
+	bool bCreatedNewEditorData = false;
+
+	UAvaTransitionTreeEditorData* EditorData = Cast<UAvaTransitionTreeEditorData>(TransitionTree->EditorData);
+	if (!EditorData && bInCreateIfNotFound)
+	{
+		EditorData = NewObject<UAvaTransitionTreeEditorData>(TransitionTree, NAME_None, RF_Transactional);
+		EditorData->AddRootState();
+		TransitionTree->EditorData = EditorData;
+		bCreatedNewEditorData = true;
+	}
+
+	EditorDataWeak = EditorData;
+	return bCreatedNewEditorData;
+}
+
+void FAvaTransitionEditorViewModel::UpdateTree()
+{
+	UAvaTransitionTree* TransitionTree = GetTransitionTree();
+	if (!TransitionTree)
+	{
+		return;
+	}
+
+	ValidateTree(TransitionTree->GetName());
+	EditorDataHash = CalculateTreeHash(*TransitionTree);
+}
+
+TSharedPtr<FAvaTransitionEditor> FAvaTransitionEditorViewModel::GetEditor() const
+{
+	return EditorWeak.Pin();
+}
+
+TSharedRef<FAvaTransitionSelection> FAvaTransitionEditorViewModel::GetSelection() const
+{
+	return GetSharedData()->GetSelection();
+}
+
+void FAvaTransitionEditorViewModel::OnInitialize()
+{
+	FAvaTransitionViewModel::OnInitialize();
+
+	BindDelegates();
+
+	if (UpdateEditorData(/*bCreateIfNotFound*/true))
+	{
+		Compile();
+	}
+
+	TSharedRef<FAvaTransitionEditorViewModel> This = SharedThis(this);
+
+	TSharedRef<FAvaTransitionViewModelSharedData> ViewModelSharedData = GetSharedData();
+
+	ViewModelSharedData->Initialize(This);
+
+	if (TSharedPtr<FAvaTransitionEditor> Editor = GetEditor())
+	{
+		ViewModelSharedData->SetReadOnly(Editor->IsReadOnly());
+	}
+
+	TreeView = SNew(SAvaTransitionTreeView, This);
+}
+
+void FAvaTransitionEditorViewModel::PostRefresh()
+{
+	FAvaTransitionViewModel::PostRefresh();
+	RefreshTreeView();
+	UpdateTree();
+}
+
+void FAvaTransitionEditorViewModel::RefreshTreeView()
+{
+	if (TreeView.IsValid())
+	{
+		TreeView->Refresh();
+	}
+}
+
+void FAvaTransitionEditorViewModel::GatherChildren(FAvaTransitionViewModelChildren& OutChildren)
+{
+	UAvaTransitionTreeEditorData* EditorData = GetEditorData();
+	if (!EditorData)
+	{
+		return;
+	}
+
+	OutChildren.Reserve(EditorData->SubTrees.Num());
+
+	for (UStateTreeState* State : EditorData->SubTrees)
+	{
+		OutChildren.Add<FAvaTransitionStateViewModel>(State);
+	}
+}
+
+void FAvaTransitionEditorViewModel::PostRedo(bool bInSuccess)
+{
+	UpdateEditorData();
+	Refresh();
+}
+
+void FAvaTransitionEditorViewModel::PostUndo(bool bInSuccess)
+{
+	UpdateEditorData();
+	Refresh();
+}
+
+TSharedRef<SWidget> FAvaTransitionEditorViewModel::CreateWidget()
+{
+	check(TreeView.IsValid());
+	return SNew(SVerticalBox)
+		+ SVerticalBox::Slot()
+		.AutoHeight()
+		[
+			Toolbar->GenerateTreeToolbarWidget()
+		]
+		+ SVerticalBox::Slot()
+		.FillHeight(1.f)
+		[
+			SNew(SHorizontalBox)
+			+ SHorizontalBox::Slot()
+			.FillWidth(1.f)
+			[
+				SNew(SScrollBox)
+				.Orientation(Orient_Horizontal)
+				+ SScrollBox::Slot()
+				.FillSize(1.f)
+				[
+					TreeView.ToSharedRef()
+				]
+			]
+			+ SHorizontalBox::Slot()
+			.AutoWidth()
+			[
+				TreeView->GetVerticalScrollbar()
+			]
+		];
+}
+
+UObject* FAvaTransitionEditorViewModel::GetObject() const
+{
+	return GetEditorData();
+}
+
+void FAvaTransitionEditorViewModel::BindDelegates()
+{
+	UnbindDelegates();
+
+	if (UAvaTransitionTreeEditorData* EditorData = GetEditorData())
+	{
+		EditorData->GetOnTreeRequestRefresh().AddSP(this, &FAvaTransitionEditorViewModel::Refresh);
+	}
+
+	UE::StateTree::Delegates::OnIdentifierChanged.AddSP(this, &FAvaTransitionEditorViewModel::OnIdentifierChanged);
+	UE::StateTree::Delegates::OnSchemaChanged.AddSP(this, &FAvaTransitionEditorViewModel::OnSchemaChanged);
+	UE::StateTree::Delegates::OnParametersChanged.AddSP(this, &FAvaTransitionEditorViewModel::OnParametersChanged);
+	UE::StateTree::Delegates::OnStateParametersChanged.AddSP(this, &FAvaTransitionEditorViewModel::OnStateParametersChanged);
+}
+
+void FAvaTransitionEditorViewModel::UnbindDelegates()
+{
+	if (UAvaTransitionTreeEditorData* EditorData = GetEditorData())
+	{
+		EditorData->GetOnTreeRequestRefresh().RemoveAll(this);
+	}
+
+	UE::StateTree::Delegates::OnIdentifierChanged.RemoveAll(this);
+	UE::StateTree::Delegates::OnSchemaChanged.RemoveAll(this);
+	UE::StateTree::Delegates::OnParametersChanged.RemoveAll(this);
+	UE::StateTree::Delegates::OnStateParametersChanged.RemoveAll(this);
+}
+
+uint32 FAvaTransitionEditorViewModel::CalculateTreeHash(const UAvaTransitionTree& InTree) const
+{
+	if (!InTree.EditorData)
+	{
+		return 0;
+	}
+
+	static const FName MD_ExcludeFromHash(TEXT("ExcludeFromHash"));
+
+	class : public FArchiveObjectCrc32
+	{
+		virtual bool ShouldSkipProperty(const FProperty* InProperty) const override
+		{
+			return !InProperty
+				|| FArchiveObjectCrc32::ShouldSkipProperty(InProperty)
+				|| InProperty->HasAllPropertyFlags(CPF_Transient)
+				|| InProperty->HasMetaData(MD_ExcludeFromHash);
+		}
+	} Archive;
+
+	return Archive.Crc32(InTree.EditorData, 0);
+}
+
+void FAvaTransitionEditorViewModel::ValidateTree(const FString& InTreeDebugName)
+{
+	UAvaTransitionTreeEditorData* EditorData = GetEditorData();
+	if (!EditorData || !EditorData->Schema)
+	{
+		return;
+	}
+
+	EditorData->ReparentStates();
+
+	// Clear evaluators if not allowed.
+	if (!EditorData->Evaluators.IsEmpty() && !EditorData->Schema->AllowEvaluators())
+	{
+		UE_LOG(LogAvaEditorTransition, Warning
+			, TEXT("%s: Resetting Evaluators due to current schema restrictions.")
+			, *InTreeDebugName);
+
+		EditorData->Evaluators.Reset();
+	}
+
+	// Apply Schema Rules to each State 
+	EditorData->VisitHierarchy([&InTreeDebugName, EditorData](UStateTreeState& State, UStateTreeState*)
+	{
+		// Clear enter conditions if not allowed.
+		if (!State.EnterConditions.IsEmpty() && !EditorData->Schema->AllowEnterConditions())
+		{
+			UE_LOG(LogAvaEditorTransition, Warning
+				, TEXT("%s: Resetting Enter Conditions in state %s due to current schema restrictions.")
+				, *InTreeDebugName
+				, *State.GetName());
+
+			State.EnterConditions.Reset();
+		}
+
+		// Keep single and many tasks based on what is allowed.
+		if (!EditorData->Schema->AllowMultipleTasks())
+		{
+			if (!State.Tasks.IsEmpty())
+			{
+				State.Tasks.Reset();
+				UE_LOG(LogAvaEditorTransition, Warning
+					, TEXT("%s: Resetting Tasks in state %s due to current schema restrictions.")
+					, *InTreeDebugName
+					, *State.GetName());
+			}
+
+			// Task name is the same as state name.
+			if (FStateTreeTaskBase* Task = State.SingleTask.Node.GetMutablePtr<FStateTreeTaskBase>())
+			{
+				Task->Name = State.Name;
+			}
+		}
+		else
+		{
+			if (State.SingleTask.Node.IsValid())
+			{
+				State.SingleTask.Reset();
+				UE_LOG(LogAvaEditorTransition, Warning
+					, TEXT("%s: Resetting Single Task in state %s due to current schema restrictions.")
+					, *InTreeDebugName
+					, *State.GetName());
+			}
+		}
+
+		return EStateTreeVisitor::Continue;
+	});
+
+	// Remove unused Bindings
+	{
+		TMap<FGuid, const FStateTreeDataView> AllStructValues;
+		EditorData->GetAllStructValues(AllStructValues);
+		EditorData->GetPropertyEditorBindings()->RemoveUnusedBindings(AllStructValues);
+	}
+
+	// Validate Linked States
+	{
+		// Make sure all state links are valid and update the names if needed.
+		// Create ID to state name map.
+		TMap<FGuid, FName> IdToName;
+
+		EditorData->VisitHierarchy([&IdToName](const UStateTreeState& State, UStateTreeState* /*ParentState*/)
+		{
+			IdToName.Add(State.ID, State.Name);
+			return EStateTreeVisitor::Continue;
+		});
+
+		static auto FixChangedStateLinkName = [](FStateTreeStateLink& StateLink, const TMap<FGuid, FName>& IDToName)
+		{
+			if (StateLink.ID.IsValid())
+			{
+				const FName* Name = IDToName.Find(StateLink.ID);
+				if (Name == nullptr)
+				{
+					// Missing link, we'll show these in the UI
+					return false;
+				}
+				if (StateLink.Name != *Name)
+				{
+					// Name changed, fix!
+					StateLink.Name = *Name;
+					return true;
+				}
+			}
+			return false;
+		};
+
+		// Fix changed names.
+		EditorData->VisitHierarchy([&IdToName](UStateTreeState& State, UStateTreeState* /*ParentState*/)
+		{
+			if (State.Type == EStateTreeStateType::Linked)
+			{
+				FixChangedStateLinkName(State.LinkedSubtree, IdToName);
+			}
+					
+			for (FStateTreeTransition& Transition : State.Transitions)
+			{
+				FixChangedStateLinkName(Transition.State, IdToName);
+			}
+
+			return EStateTreeVisitor::Continue;
+		});
+	}
+
+	// Update Linked State Parameters
+	EditorData->VisitHierarchy([](UStateTreeState& State, UStateTreeState* /*ParentState*/)
+	{
+		if (State.Type == EStateTreeStateType::Linked)
+		{
+			State.UpdateParametersFromLinkedSubtree();
+		}
+		return EStateTreeVisitor::Continue;
+	});
+}
+
+void FAvaTransitionEditorViewModel::OnIdentifierChanged(const UStateTree& InStateTree)
+{
+	if (&InStateTree == GetTransitionTree())
+	{
+		UpdateTree();
+	}
+}
+
+void FAvaTransitionEditorViewModel::OnSchemaChanged(const UStateTree& InStateTree)
+{
+	if (&InStateTree == GetTransitionTree())
+	{
+		UpdateTree();
+		// todo: notify asset change externally
+	}
+}
+
+void FAvaTransitionEditorViewModel::OnParametersChanged(const UStateTree& InStateTree)
+{
+	// todo
+}
+
+void FAvaTransitionEditorViewModel::OnStateParametersChanged(const UStateTree& InStateTree, const FGuid InGuid)
+{
+	// todo
+}
+
+#undef LOCTEXT_NAMESPACE
diff --git a/Engine/Plugins/Experimental/Avalanche/Source/AvalancheTransitionEditor/Public/AvaTransitionTreeEditorData.h b/Engine/Plugins/Experimental/Avalanche/Source/AvalancheTransitionEditor/Public/AvaTransitionTreeEditorData.h
new file mode 100644
--- /dev/null
+++ b/Engine/Plugins/Experimental/Avalanche/Source/AvalancheTransitionEditor/Public/AvaTransitionTreeEditorData.h
@@ -0,0 +1,85 @@
+// Copyright Epic Games, Inc. All Rights Reserved.
+
+#pragma once
+
+#include "AvaTagHandle.h"
+#include "Containers/Map.h"
+#include "Delegates/Delegate.h"
+#include "Misc/Guid.h"
+#include "State/AvaTransitionAutogeneratedStateContext.h"
+#include "State/AvaTransitionStateMetadata.h"
+#include "StateTreeEditorData.h"
+#include "AvaTransitionTreeEditorData.generated.h"
+
+class UAvaTransitionTreeEditorData;
+class UStateTreeState;
+
+UCLASS(MinimalAPI, HideCategories=(Common))
+class UAvaTransitionTreeEditorData : public UStateTreeEditorData
+{
+	GENERATED_BODY()
+
+public:
+	AVALANCHETRANSITIONEDITOR_API UStateTreeState& CreateState(const UStateTreeState& InSiblingState, bool bInAfter);
+
+	FAvaTagHandle GetTransitionLayer() const
+	{
+		return TransitionLayer;
+	}
+
+	void SetTransitionLayer(const FAvaTagHandle& InTransitionLayer)
+	{
+		TransitionLayer = InTransitionLayer;
+	}
+
+	static FName GetTransitionLayerPropertyName()
+	{
+		return GET_MEMBER_NAME_CHECKED(UAvaTransitionTreeEditorData, TransitionLayer);
+	}
+
+	const FAvaTransitionStateMetadata* FindStateMetadata(const FGuid& InStateId) const
+	{
+		return StateMetadata.Find(InStateId);
+	}
+
+	FAvaTransitionStateMetadata& FindOrAddStateMetadata(const FGuid& InStateId)
+	{
+		return StateMetadata.FindOrAdd(InStateId);
+	}
+
+	const FAvaTransitionAutogeneratedStateContext* FindAutogeneratedStateContext(const FGuid& InStateId) const
+	{
+		return AutogeneratedStates.Find(InStateId);
+	}
+
+	bool IsAutogeneratedState(const UStateTreeState& InState, const FAvaTransitionAutogeneratedStateContext& InContext) const
+	{
+		const FAvaTransitionAutogeneratedStateContext* AutogeneratedContext = FindAutogeneratedStateContext(InState.ID);
+		return AutogeneratedContext && AutogeneratedContext->Matches(InContext);
+	}
+
+	void SetAutogeneratedState(const FGuid& InStateId, const FAvaTransitionAutogeneratedStateContext& InContext)
+	{
+		AutogeneratedStates.Add(InStateId, InContext);
+	}
+
+	FSimpleMulticastDelegate& GetOnTreeRequestRefresh()
+	{
+		return OnTreeRequestRefresh;
+	}
+
+private:
+	/** The Layer this Transition Logic Tree deals with */
+	UPROPERTY(EditAnywhere, Category="Transition Logic")
+	FAvaTagHandle TransitionLayer;
+
+	/** Map of a state's id to its metadata */
+	UPROPERTY()
+	TMap<FGuid, FAvaTransitionStateMetadata> StateMetadata;
+
+	/** Map of an autogenerated state's id to their context information */
+	UPROPERTY()
+	TMap<FGuid, FAvaTransitionAutogeneratedStateContext> AutogeneratedStates;
+
+	FSimpleMulticastDelegate OnTreeRequestRefresh;
+};
diff --git a/Engine/Plugins/Experimental/Avalanche/Source/AvalancheTransitionEditor/Public/State/AvaTransitionAutogeneratedStateContext.h b/Engine/Plugins/Experimental/Avalanche/Source/AvalancheTransitionEditor/Public/State/AvaTransitionAutogeneratedStateContext.h
new file mode 100644
--- /dev/null
+++ b/Engine/Plugins/Experimental/Avalanche/Source/AvalancheTransitionEditor/Public/State/AvaTransitionAutogeneratedStateContext.h
@@ -0,0 +1,41 @@
+// Copyright Epic Games, Inc. All Rights Reserved.
+
+#pragma once
+
+#include "Internationalization/Text.h"
+#include "Misc/Guid.h"
+#include "AvaTransitionAutogeneratedStateContext.generated.h"
+
+/** Context of how a particular autogenerated state was created */
+USTRUCT()
+struct FAvaTransitionAutogeneratedStateContext
+{
+	GENERATED_BODY()
+
+	FAvaTransitionAutogeneratedStateContext() = default;
+
+	explicit FAvaTransitionAutogeneratedStateContext(const FGuid& InKey, const FText& InLabel)
+		: Key(InKey)
+		, Label(InLabel)
+	{
+	}
+
+	bool Matches(const FAvaTransitionAutogeneratedStateContext& InOther) const
+	{
+		return Key == InOther.Key;
+	}
+
+	const FText& GetLabel() const
+	{
+		return Label;
+	}
+
+private:
+	/** Key of the Context that autogenerated the state */
+	UPROPERTY()
+	FGuid Key;
+
+	/** Label Display of the Context that autogenerated the state */
+	UPROPERTY()
+	FText Label;
+};
diff --git a/Engine/Plugins/Experimental/Avalanche/Source/AvalancheViewport/Private/Visualizers/AvaViewportBackplateVisualizer.cpp b/Engine/Plugins/Experimental/Avalanche/Source/AvalancheViewport/Private/Visualizers/AvaViewportBackplateVisualizer.cpp
new file mode 100644
--- /dev/null
+++ b/Engine/Plugins/Experimental/Avalanche/Source/AvalancheViewport/Private/Visualizers/AvaViewportBackplateVisualizer.cpp
@@ -0,0 +1,206 @@
+// Copyright Epic Games, Inc. All Rights Reserved.
+
+#include "Visualizers/AvaViewportBackplateVisualizer.h"
+#include "AvaViewportPostProcessManager.h"
+#include "AvaViewportSettings.h"
+#include "AvaViewportUtils.h"
+#include "AvaVisibleArea.h"
+#include "Engine/Texture.h"
+#include "Materials/Material.h"
+#include "Materials/MaterialInstanceDynamic.h"
+#include "UObject/Package.h"
+#include "Viewport/Interaction/AvaViewportPostProcessInfo.h"
+#include "Viewport/Interaction/IAvaViewportDataProvider.h"
+#include "Viewport/Interaction/IAvaViewportDataProxy.h"
+#include "ViewportClient/IAvaViewportClient.h"
+
+#define LOCTEXT_NAMESPACE "AvaViewportBackplateVisualizer"
+
+namespace UE::AvalancheViewport::Private
+{
+	const FString BackplateReferencerName = FString(TEXT("AvaViewportBackplateVisualizer"));
+	const FName TextureObjectName = FName(TEXT("TextureObject"));
+	const FName TextureOffsetName = FName(TEXT("TextureOffset"));
+	const FName TextureScaleName = FName(TEXT("TextureScale"));
+}
+
+FAvaViewportBackplateVisualizer::FAvaViewportBackplateVisualizer(TSharedRef<IAvaViewportClient> InAvaViewportClient)
+	: FAvaViewportPostProcessVisualizer(InAvaViewportClient)
+{
+	bRequiresTonemapperSetting = true;
+
+	TextureOffset = FVector::ZeroVector;
+	TextureScale = FVector::ZeroVector;
+
+	const UAvaViewportSettings* ViewportSettings = GetDefault<UAvaViewportSettings>();
+
+	if (!ViewportSettings)
+	{
+		return;
+	}
+
+	UMaterial* BackplateMaterial = ViewportSettings->ViewportBackplateMaterial.LoadSynchronous();
+
+	if (!BackplateMaterial)
+	{
+		return;
+	}
+
+	PostProcessBaseMaterial = BackplateMaterial;
+	PostProcessMaterial = UMaterialInstanceDynamic::Create(BackplateMaterial, GetTransientPackage());
+}
+
+UTexture* FAvaViewportBackplateVisualizer::GetTexture() const
+{
+	return Texture;
+}
+
+void FAvaViewportBackplateVisualizer::SetTexture(UTexture* InTexture)
+{
+	if (Texture == InTexture)
+	{
+		return;
+	}
+
+	SetTextureInternal(InTexture);
+
+	UpdatePostProcessInfo();
+	UpdatePostProcessMaterial();
+}
+
+void FAvaViewportBackplateVisualizer::AddReferencedObjects(FReferenceCollector& InCollector)
+{
+	Super::AddReferencedObjects(InCollector);
+
+	if (Texture)
+	{
+		InCollector.AddReferencedObject(Texture);
+	}
+}
+
+FString FAvaViewportBackplateVisualizer::GetReferencerName() const
+{
+	return UE::AvalancheViewport::Private::BackplateReferencerName;
+}
+
+void FAvaViewportBackplateVisualizer::UpdateForViewport(const FAvaVisibleArea& InVisibleArea, const FVector2f& InWidgetSize, 
+	const FVector2f& InCameraOffset)
+{
+	if (FMath::IsNearlyZero(PostProcessOpacity) || !Texture || !PostProcessMaterial)
+	{
+		return;
+	}
+
+	if (!InVisibleArea.IsValid())
+	{
+		return;
+	}
+
+	if (!FAvaViewportUtils::IsValidViewportSize(InWidgetSize))
+	{
+		return;
+	}
+
+	const FVector2f ImageSize = {Texture->GetSurfaceWidth(), Texture->GetSurfaceHeight()};
+
+	if (!FAvaViewportUtils::IsValidViewportSize(ImageSize))
+	{
+		return;
+	}
+
+	using namespace UE::AvalancheViewport::Private;
+
+	const float ImageAspectRatio = ImageSize.X / ImageSize.Y;
+	const float WidgetAspectRatio = InWidgetSize.X / InWidgetSize.Y;
+	const float ViewportAspectRatio = InVisibleArea.AbsoluteSize.X / InVisibleArea.AbsoluteSize.Y;
+	const FVector2f WidgetBasedScale = InVisibleArea.AbsoluteSize / InWidgetSize;
+
+	const FVector2f Scale = WidgetBasedScale / InVisibleArea.GetVisibleAreaFraction();
+
+	if (!FMath::IsNearlyEqual(TextureScale.X, Scale.X)
+		|| !FMath::IsNearlyEqual(TextureScale.Y, Scale.Y))
+	{
+		TextureScale.X = Scale.X;
+		TextureScale.Y = Scale.Y;
+		TextureOffset.Z = 0.f;
+		PostProcessMaterial->SetVectorParameterValue(TextureScaleName, TextureScale);
+	}
+
+	FVector2f Offset = ((InWidgetSize - InVisibleArea.AbsoluteSize) * 0.5f)
+		+ InCameraOffset / InVisibleArea.GetVisibleAreaFraction();
+
+	if (!FMath::IsNearlyEqual(WidgetAspectRatio, ViewportAspectRatio))
+	{
+		if (WidgetAspectRatio > ViewportAspectRatio)
+		{
+			Offset.X -= (Scale.Y - 1) * InVisibleArea.AbsoluteSize.X * 0.5f;
+			Offset.Y -= (Scale.Y - 1) * InVisibleArea.AbsoluteSize.Y * 0.5f;
+		}
+		else
+		{
+			Offset.X -= (Scale.X - 1) * InVisibleArea.AbsoluteSize.X * 0.5f;
+			Offset.Y -= (Scale.X - 1) * InVisibleArea.AbsoluteSize.Y * 0.5f;
+		}
+	}		
+
+	if (!FMath::IsNearlyEqual(TextureOffset.X, Offset.X)
+		|| !FMath::IsNearlyEqual(TextureOffset.Y, Offset.Y))
+	{
+		TextureOffset.X = Offset.X;
+		TextureOffset.Y = Offset.Y;
+		TextureOffset.Z = 0.f;
+		PostProcessMaterial->SetVectorParameterValue(TextureOffsetName, TextureOffset);
+	}
+}
+
+void FAvaViewportBackplateVisualizer::LoadPostProcessInfo(const FAvaViewportPostProcessInfo& InPostProcessInfo)
+{
+	Super::LoadPostProcessInfo(InPostProcessInfo);
+
+	SetTextureInternal(InPostProcessInfo.Texture.LoadSynchronous());
+}
+
+void FAvaViewportBackplateVisualizer::UpdatePostProcessInfo(FAvaViewportPostProcessInfo& InPostProcessInfo) const
+{
+	Super::UpdatePostProcessInfo(InPostProcessInfo);
+
+	InPostProcessInfo.Texture = Texture;
+}
+
+void FAvaViewportBackplateVisualizer::UpdatePostProcessMaterial()
+{
+	if (!PostProcessMaterial)
+	{
+		return;
+	}
+
+	Super::UpdatePostProcessMaterial();
+
+	using namespace UE::AvalancheViewport::Private;
+
+	PostProcessMaterial->SetTextureParameterValue(TextureObjectName, Texture);
+}
+
+bool FAvaViewportBackplateVisualizer::SetupPostProcessSettings(FPostProcessSettings& InPostProcessSettings) const
+{
+	if (!IsValid(Texture))
+	{
+		return false;
+	}
+
+	return FAvaViewportPostProcessVisualizer::SetupPostProcessSettings(InPostProcessSettings);
+}
+
+void FAvaViewportBackplateVisualizer::SetTextureInternal(UTexture* InTexture)
+{
+	if (!IsValid(InTexture))
+	{
+		Texture = nullptr;
+	}
+	else
+	{
+		Texture = InTexture;
+	}
+}
+
+#undef LOCTEXT_NAMESPACE
diff --git a/Engine/Plugins/Experimental/Avalanche/Source/AvalancheViewport/Private/Visualizers/AvaViewportBoundingBoxVisualizer.h b/Engine/Plugins/Experimental/Avalanche/Source/AvalancheViewport/Private/Visualizers/AvaViewportBoundingBoxVisualizer.h
new file mode 100644
--- /dev/null
+++ b/Engine/Plugins/Experimental/Avalanche/Source/AvalancheViewport/Private/Visualizers/AvaViewportBoundingBoxVisualizer.h
@@ -0,0 +1,31 @@
+// Copyright Epic Games, Inc. All Rights Reserved.
+
+#pragma once
+
+#include "AvaType.h"
+#include "Visualizers/IAvaViewportBoundingBoxVisualizerInterface.h"
+
+class AActor;
+
+class FAvaLevelViewportBoundingBoxVisualizer : public IAvaViewportBoundingBoxVisualizerInterface
+{
+public:
+	UE_AVA_INHERITS(FAvaLevelViewportBoundingBoxVisualizer, IAvaViewportBoundingBoxVisualizerInterface)
+
+	FAvaLevelViewportBoundingBoxVisualizer();
+	virtual ~FAvaLevelViewportBoundingBoxVisualizer() = default;
+
+	//~ Begin IAvaViewportBoundingBoxVisualizerInterface
+	virtual void Draw(UAvaSelectionProviderSubsystem& InSelectionProvider, UAvaBoundsProviderSubsystem& InBoundsProvider, FPrimitiveDrawInterface& InPDI) override;
+	virtual EAvaViewportBoundingBoxOptimizationState GetOptimizationState() const override { return OptimizationState; }
+	virtual void ResetOptimizationState() override;
+	//~ End IAvaViewportBoundingBoxVisualizerInterface
+
+protected:
+	EAvaViewportBoundingBoxOptimizationState OptimizationState;
+
+	void SetOptimizationRenderNothing(double InTaskTimeTaken);
+	void SetOptimizationRenderSelectedActors(double InTaskTimeTaken);
+	void SetOptimizationRenderSelectedActorsAndChildren(double InTaskTimeTaken);
+	void SetOptimizationRenderSelectionBounds();
+};
diff --git a/Engine/Plugins/Experimental/Avalanche/Source/AvalancheViewport/Private/Visualizers/AvaViewportChannelVisualizer.cpp b/Engine/Plugins/Experimental/Avalanche/Source/AvalancheViewport/Private/Visualizers/AvaViewportChannelVisualizer.cpp
new file mode 100644
--- /dev/null
+++ b/Engine/Plugins/Experimental/Avalanche/Source/AvalancheViewport/Private/Visualizers/AvaViewportChannelVisualizer.cpp
@@ -0,0 +1,44 @@
+// Copyright Epic Games, Inc. All Rights Reserved.
+
+#include "Visualizers/AvaViewportChannelVisualizer.h"
+#include "AvaViewportPostProcessManager.h"
+#include "Containers/Map.h"
+#include "Materials/Material.h"
+#include "Materials/MaterialInstanceDynamic.h"
+#include "UObject/SoftObjectPtr.h"
+#include "UObject/Package.h"
+#include "Viewport/Interaction/AvaViewportPostProcessInfo.h"
+
+namespace UE::AvalancheViewport::Private
+{
+	const FString ChannelReferencerName = FString(TEXT("AvaViewportChannelVisualizer"));
+
+	const TMap<EAvaViewportPostProcessType, TSoftObjectPtr<UMaterial>> ChannelAssets = {
+		{ EAvaViewportPostProcessType::RedChannel,   TSoftObjectPtr<UMaterial>(FSoftObjectPath("/Script/Engine.Material'/Avalanche/EditorResources/M_AvalancheChannelVisualizer_Red.M_AvalancheChannelVisualizer_Red'")) },
+		{ EAvaViewportPostProcessType::GreenChannel, TSoftObjectPtr<UMaterial>(FSoftObjectPath("/Script/Engine.Material'/Avalanche/EditorResources/M_AvalancheChannelVisualizer_Green.M_AvalancheChannelVisualizer_Green'")) },
+		{ EAvaViewportPostProcessType::BlueChannel,  TSoftObjectPtr<UMaterial>(FSoftObjectPath("/Script/Engine.Material'/Avalanche/EditorResources/M_AvalancheChannelVisualizer_Blue.M_AvalancheChannelVisualizer_Blue'")) },
+		{ EAvaViewportPostProcessType::AlphaChannel, TSoftObjectPtr<UMaterial>(FSoftObjectPath("/Script/Engine.Material'/Avalanche/EditorResources/M_AvalancheChannelVisualizer_Alpha.M_AvalancheChannelVisualizer_Alpha'")) }
+	};
+}
+
+FAvaViewportChannelVisualizer::FAvaViewportChannelVisualizer(TSharedRef<IAvaViewportClient> InAvaViewportClient, EAvaViewportPostProcessType InChannel)
+	: FAvaViewportPostProcessVisualizer(InAvaViewportClient)
+{
+	bRequiresTonemapperSetting = (InChannel == EAvaViewportPostProcessType::AlphaChannel);
+
+	using namespace UE::AvalancheViewport::Private;
+
+	if (const TSoftObjectPtr<UMaterial>* BaseMaterialPtr = ChannelAssets.Find(InChannel))
+	{
+		if (UMaterial* BaseMaterial = BaseMaterialPtr->LoadSynchronous())
+		{
+			PostProcessBaseMaterial = BaseMaterial;
+			PostProcessMaterial = UMaterialInstanceDynamic::Create(BaseMaterial, GetTransientPackage());
+		}
+	}
+}
+
+FString FAvaViewportChannelVisualizer::GetReferencerName() const
+{
+	return UE::AvalancheViewport::Private::ChannelReferencerName;
+}
diff --git a/Engine/Plugins/Experimental/Avalanche/Source/AvalancheViewport/Public/Interaction/AvaCameraZoomController.h b/Engine/Plugins/Experimental/Avalanche/Source/AvalancheViewport/Public/Interaction/AvaCameraZoomController.h
new file mode 100644
--- /dev/null
+++ b/Engine/Plugins/Experimental/Avalanche/Source/AvalancheViewport/Public/Interaction/AvaCameraZoomController.h
@@ -0,0 +1,96 @@
+// Copyright Epic Games, Inc. All Rights Reserved.
+
+#pragma once
+
+#include "AvaVisibleArea.h"
+#include "Math/Vector2D.h"
+#include "Templates/SharedPointer.h"
+
+class IAvaViewportClient;
+struct FAvaVisibleArea;
+
+class AVALANCHEVIEWPORT_API FAvaCameraZoomController
+{
+public:
+	static bool IsCameraZoomPossible();
+
+	FAvaCameraZoomController(TSharedRef<IAvaViewportClient> InAvaViewportClient, float InFallbackFOV);
+
+	TSharedPtr<IAvaViewportClient> GetViewportClient() const { return AvaViewportClientWeak.Pin(); }
+
+	uint8 GetZoomLevel() const { return ZoomLevel; }
+	void SetZoomLevel(uint8 InZoomLevel);
+	bool IsZoomed() const;
+
+	float GetFOVPerStep() const;
+
+	void ZoomIn();
+	void ZoomInCursor();
+
+	/* Zooms in maintaining the current PanOffsetFraction. Uses the absolute screen position. */
+	void ZoomInAroundPoint(const FVector2f& InScreenPosition);
+
+	/* Zooms in maintaining the position of the given viewport position. Uses the absolute screen position. */
+	void ZoomInRelativePoint(const FVector2f& InViewportPosition);
+
+	void ZoomOut();
+	void ZoomOutCursor();
+
+	/* Zooms out maintaining the current PanOffsetFraction. Uses the absolute screen position. */
+	void ZoomOutAroundPoint(const FVector2f& OutScreenPosition);
+
+	/* Zooms out maintaining the position of the given viewport position. Uses the absolute screen position. */
+	void ZoomOutRelativePoint(const FVector2f& OutViewportPosition);
+
+	void PanLeft();
+	void PanRight();
+	void PanUp();
+	void PanDown();
+
+	void FrameActor();
+
+	void Reset();
+
+	const FVector2f& GetPanOffsetFraction() const { return PanOffsetFraction; }
+	void SetPanOffsetFraction(const FVector2f& InOffsetFraction);
+	void PanAdjust(const FVector2f& InDirection);
+
+	/** Adjusts zoom pan based on current zoom settings. */
+	void PanAdjustZoomed(const FVector2f& InZoomedDirection);
+
+	void CenterOnPoint(const FVector2f& InPoint);
+	void CenterOnBox(const FBox& InBoundingBox, const FTransform& InBoxTransform);
+
+	bool IsPanning() const { return bIsPanning; }
+	void StartPanning();
+	void EndPanning();
+
+	float GetFallbackFOV() const { return FallbackFOV; }
+	float GetDefaultFOV() const;
+	float GetFOV() const;
+	FVector2f GetCameraProjectionOffset() const;
+
+	const FAvaVisibleArea& GetCachedVisibleArea() const;
+	const FAvaVisibleArea& GetCachedZoomedVisibleArea() const;
+
+	void UpdateVisibleAreas();
+
+protected:
+	TWeakPtr<IAvaViewportClient> AvaViewportClientWeak;
+
+	const float FallbackFOV;
+
+	uint8 ZoomLevel;
+
+	FVector2f PanOffsetFraction;
+
+	bool bIsPanning;
+
+	FAvaVisibleArea CachedVisibleArea;
+	FAvaVisibleArea CachedZoomedVisibleArea;
+
+	void ZoomIn_Internal();
+	void ZoomOut_Internal();
+
+	void InvalidateViewport();
+};
\ No newline at end of file
diff --git a/Engine/Plugins/Experimental/Avalanche/Source/AvalancheViewport/Public/Visualizers/AvaViewportBackplateVisualizer.h b/Engine/Plugins/Experimental/Avalanche/Source/AvalancheViewport/Public/Visualizers/AvaViewportBackplateVisualizer.h
new file mode 100644
--- /dev/null
+++ b/Engine/Plugins/Experimental/Avalanche/Source/AvalancheViewport/Public/Visualizers/AvaViewportBackplateVisualizer.h
@@ -0,0 +1,50 @@
+// Copyright Epic Games, Inc. All Rights Reserved.
+
+#pragma once
+
+#include "Visualizers/AvaViewportPostProcessVisualizer.h"
+#include "Math/Vector.h"
+
+class UTexture;
+struct FAvaViewportPostProcessInfo;
+
+class AVALANCHEVIEWPORT_API FAvaViewportBackplateVisualizer : public FAvaViewportPostProcessVisualizer
+{
+public:
+	UE_AVA_INHERITS_WITH_SUPER(FAvaViewportBackplateVisualizer, FAvaViewportPostProcessVisualizer)
+
+	FAvaViewportBackplateVisualizer(TSharedRef<IAvaViewportClient> InAvaViewportClient);
+	virtual ~FAvaViewportBackplateVisualizer() override = default;
+
+	UTexture* GetTexture() const;
+	void SetTexture(UTexture* InTexture);
+
+	//~ Begin FAvaViewportPostProcessVisualizer
+	virtual void UpdateForViewport(const FAvaVisibleArea& InVisibleArea, const FVector2f& InWidgetSize, const FVector2f& InCameraOffset) override;
+	//~ End FAvaViewportPostProcessVisualizer
+
+	//~ Begin FGCObject
+	virtual void AddReferencedObjects(FReferenceCollector& InCollector) override;
+	virtual FString GetReferencerName() const override;
+	//~ End FGCObject
+
+protected:
+	TObjectPtr<UTexture> Texture;
+
+	/** Offset from the top-left of the viewport where the texture should start in pixel space. */
+	FVector TextureOffset;
+
+	/** Scale of the texture. Larger values mean a larger texture. */
+	FVector TextureScale;
+
+	using FAvaViewportPostProcessVisualizer::UpdatePostProcessInfo;
+
+	void SetTextureInternal(UTexture* InTexture);
+
+	//~ Begin FAvaViewportPostProcessVisualizer
+	virtual void LoadPostProcessInfo(const FAvaViewportPostProcessInfo& InPostProcessInfo) override;
+	virtual void UpdatePostProcessInfo(FAvaViewportPostProcessInfo& InPostProcessInfo) const override;
+	virtual void UpdatePostProcessMaterial() override;
+	virtual bool SetupPostProcessSettings(FPostProcessSettings& InPostProcessSettings) const override;
+	//~ End FAvaViewportPostProcessVisualizer
+};
\ No newline at end of file
diff --git a/Engine/Plugins/Experimental/Avalanche/Source/AvalancheViewport/Public/Visualizers/IAvaViewportBoundingBoxVisualizerInterface.h b/Engine/Plugins/Experimental/Avalanche/Source/AvalancheViewport/Public/Visualizers/IAvaViewportBoundingBoxVisualizerInterface.h
new file mode 100644
--- /dev/null
+++ b/Engine/Plugins/Experimental/Avalanche/Source/AvalancheViewport/Public/Visualizers/IAvaViewportBoundingBoxVisualizerInterface.h
@@ -0,0 +1,34 @@
+// Copyright Epic Games, Inc. All Rights Reserved.
+
+#pragma once
+
+#include "AvaType.h"
+#include "Templates/SharedPointer.h"
+
+class AActor;
+class FPrimitiveDrawInterface;
+class UAvaBoundsProviderSubsystem;
+class UAvaSelectionProviderSubsystem;
+
+enum class EAvaViewportBoundingBoxOptimizationState : uint8
+{
+	RenderNothing,
+	RenderSelectedActors,
+	RenderSelectedActorsAndChildren,
+	RenderSelectionBounds
+};
+
+class IAvaViewportBoundingBoxVisualizerInterface : public IAvaTypeCastable
+{
+public:
+	UE_AVA_INHERITS(IAvaViewportBoundingBoxVisualizerInterface, IAvaTypeCastable)
+
+	virtual void Draw(UAvaSelectionProviderSubsystem& InSelectionProvider, UAvaBoundsProviderSubsystem& InBoundsProvider, FPrimitiveDrawInterface& InPDI) = 0;
+	virtual EAvaViewportBoundingBoxOptimizationState GetOptimizationState() const = 0;
+	virtual void ResetOptimizationState() = 0;
+};
+
+struct AVALANCHEVIEWPORT_API FAvaViewportBoundingBoxVisualizerProvider
+{
+	static TSharedRef<IAvaViewportBoundingBoxVisualizerInterface> CreateVisualizer();
+};
diff --git a/Engine/Plugins/Experimental/Avalanche/Templates/TP_ME_BroadcastBP/TP_ME_BroadcastBP.uproject b/Engine/Plugins/Experimental/Avalanche/Templates/TP_ME_BroadcastBP/TP_ME_BroadcastBP.uproject
new file mode 100644
--- /dev/null
+++ b/Engine/Plugins/Experimental/Avalanche/Templates/TP_ME_BroadcastBP/TP_ME_BroadcastBP.uproject
@@ -0,0 +1,78 @@
+{
+	"FileVersion": 3,
+	"EngineAssociation": "",
+	"Category": "",
+	"Description": "",
+	"Plugins": [
+		{
+			"Name": "RemoteControlProtocolDMX",
+			"Enabled": true
+		},
+		{
+			"Name": "RemoteControlProtocolMIDI",
+			"Enabled": true
+		},
+		{
+			"Name": "RemoteControlProtocolOSC",
+			"Enabled": true
+		},
+		{
+			"Name": "DatasmithImporter",
+			"Enabled": true
+		},
+		{
+			"Name": "ColorCorrectRegions",
+			"Enabled": true
+		},
+		{
+			"Name": "RemoteControlWebInterface",
+			"Enabled": true,
+			"MarketplaceURL": "com.epicgames.launcher://ue/marketplace/content/beb39fbe47e74bb1a8098e2b2c9c40b2"
+		},
+		{
+			"Name": "VirtualProductionUtilities",
+			"Enabled": true
+		},
+		{
+			"Name": "RenderGrid",
+			"Enabled": true
+		},
+		{
+			"Name": "ElectraPlayer",
+			"Enabled": true
+		},
+		{
+			"Name": "ElectraUtil",
+			"Enabled": true
+		},
+		{
+			"Name": "BlackmagicMedia",
+			"Enabled": true,
+			"MarketplaceURL": "com.epicgames.launcher://ue/marketplace/content/433d28bdfb764ec5b2e60c4f45d5cb5c",
+			"SupportedTargetPlatforms": [
+				"Win64",
+				"Linux"
+			]
+		},
+		{
+			"Name": "AjaMedia",
+			"Enabled": true,
+			"MarketplaceURL": "com.epicgames.launcher://ue/marketplace/content/91088fb17fb04b6b9eda41aaf8fda21e",
+			"SupportedTargetPlatforms": [
+				"Win64"
+			]
+		},
+		{
+			"Name": "StormSyncAvalancheBridge",
+			"Enabled": true
+		},
+		{
+			"Name": "WebSocketMessaging",
+			"Enabled": true
+		},
+		{
+			"Name": "StormSync",
+			"Enabled": true
+		}
+	]
+}
