[Motion Design] For AvalancheMedia and AvalancheMediaEditor, renaming "Avalanche" to "Motion Design" in logs, menus and tooltips and comments that end up in UI.

diff --git a/Engine/Plugins/Experimental/Avalanche/Source/AvalancheMedia/Private/AvaMediaModule.cpp b/Engine/Plugins/Experimental/Avalanche/Source/AvalancheMedia/Private/AvaMediaModule.cpp
--- a/Engine/Plugins/Experimental/Avalanche/Source/AvalancheMedia/Private/AvaMediaModule.cpp
+++ b/Engine/Plugins/Experimental/Avalanche/Source/AvalancheMedia/Private/AvaMediaModule.cpp
@@ -119,7 +119,7 @@ void FAvaMediaModule::StartupModule()
 				));
 	ConsoleCmds.Add(IConsoleManager::Get().RegisterConsoleCommand(
 				TEXT("AvaStat"),
-				TEXT("Enable engine performance statistics. Same as 'stat' command but will affect Avalanche Playback outputs and propagate to connected servers."),
+				TEXT("Enable engine performance statistics. Same as 'stat' command but will affect Motion Design Playback outputs and propagate to connected servers."),
 				FConsoleCommandWithArgsDelegate::CreateRaw(this, &FAvaMediaModule::HandleStatCommand),
 				ECVF_Default
 				));
diff --git a/Engine/Plugins/Experimental/Avalanche/Source/AvalancheMedia/Private/OutputDevices/AvaDisplayMediaOutput.cpp b/Engine/Plugins/Experimental/Avalanche/Source/AvalancheMedia/Private/OutputDevices/AvaDisplayMediaOutput.cpp
--- a/Engine/Plugins/Experimental/Avalanche/Source/AvalancheMedia/Private/OutputDevices/AvaDisplayMediaOutput.cpp
+++ b/Engine/Plugins/Experimental/Avalanche/Source/AvalancheMedia/Private/OutputDevices/AvaDisplayMediaOutput.cpp
@@ -38,7 +38,7 @@ UMediaCapture* UAvaDisplayMediaOutput::CreateMediaCaptureImpl()
 	UMediaCapture* Result = NewObject<UAvaDisplayMediaCapture>();
 	if (Result)
 	{
-		UE_LOG(LogAvaDisplayMedia, Log, TEXT("Created Avalanche Display Media Capture"));
+		UE_LOG(LogAvaDisplayMedia, Log, TEXT("Created Motion Design Display Media Capture"));
 		Result->SetMediaOutput(this);
 	}
 	return Result;
diff --git a/Engine/Plugins/Experimental/Avalanche/Source/AvalancheMedia/Private/OutputDevices/AvaRenderTargetMediaOutput.cpp b/Engine/Plugins/Experimental/Avalanche/Source/AvalancheMedia/Private/OutputDevices/AvaRenderTargetMediaOutput.cpp
--- a/Engine/Plugins/Experimental/Avalanche/Source/AvalancheMedia/Private/OutputDevices/AvaRenderTargetMediaOutput.cpp
+++ b/Engine/Plugins/Experimental/Avalanche/Source/AvalancheMedia/Private/OutputDevices/AvaRenderTargetMediaOutput.cpp
@@ -60,7 +60,7 @@ UMediaCapture* UAvaRenderTargetMediaOutput::CreateMediaCaptureImpl()
 	UMediaCapture* Result = NewObject<UAvaRenderTargetMediaCapture>();
 	if (Result)
 	{
-		UE_LOG(LogAvaRenderTargetMedia, Log, TEXT("Created Avalanche Render Target Media Capture"));
+		UE_LOG(LogAvaRenderTargetMedia, Log, TEXT("Created Motion Design Render Target Media Capture"));
 		Result->SetMediaOutput(this);
 	}
 	return Result;
diff --git a/Engine/Plugins/Experimental/Avalanche/Source/AvalancheMedia/Private/Playback/AvaMediaPlayableGroup.cpp b/Engine/Plugins/Experimental/Avalanche/Source/AvalancheMedia/Private/Playback/AvaMediaPlayableGroup.cpp
--- a/Engine/Plugins/Experimental/Avalanche/Source/AvalancheMedia/Private/Playback/AvaMediaPlayableGroup.cpp
+++ b/Engine/Plugins/Experimental/Avalanche/Source/AvalancheMedia/Private/Playback/AvaMediaPlayableGroup.cpp
@@ -84,7 +84,7 @@ UPackage* UAvaMediaPlayableGroup::MakeInstancePackage(const FString& InInstanceP
 	else
 	{
 		// Note: The outer will fallback to GEngine in that case.
-		UE_LOG(LogAvalanchePlayable, Error, TEXT("Unable to create package \"%s\" for Avalanche Game Instance."), *InInstancePackageName);
+		UE_LOG(LogAvalanchePlayable, Error, TEXT("Unable to create package \"%s\" for Motion Design Game Instance."), *InInstancePackageName);
 	}
 	return InstancePackage;
 }
diff --git a/Engine/Plugins/Experimental/Avalanche/Source/AvalancheMedia/Private/Playback/AvaMediaPlaybackManager.cpp b/Engine/Plugins/Experimental/Avalanche/Source/AvalancheMedia/Private/Playback/AvaMediaPlaybackManager.cpp
--- a/Engine/Plugins/Experimental/Avalanche/Source/AvalancheMedia/Private/Playback/AvaMediaPlaybackManager.cpp
+++ b/Engine/Plugins/Experimental/Avalanche/Source/AvalancheMedia/Private/Playback/AvaMediaPlaybackManager.cpp
@@ -398,7 +398,7 @@ UAvalanchePlayback*  FAvaMediaPlaybackManager::LoadPlaybackObject(const FSoftObj
 			}
 
 			UE_LOG(LogAvaPlaybackManager, Error,
-				TEXT("Asset \"%s\" in package \"%s\" is not a supported Avalanche playback asset (\"%s\")."),
+				TEXT("Asset \"%s\" in package \"%s\" is not a supported Motion Design playback asset (\"%s\")."),
 				*AssetName, *PackageName, *FoundObject->GetClass()->GetFullName());
 		}
 		else
@@ -723,7 +723,7 @@ bool FAvaMediaPlaybackManager::HandleStatCommand(const TArray<FString>& InArgs)
 		}
 		if (!bFound)
 		{
-			UE_LOG(LogAvaPlaybackManager, Warning, TEXT("No Active Avalanche Game Instances found to apply the stat '%s' command on."), *InArgs[0]);
+			UE_LOG(LogAvaPlaybackManager, Warning, TEXT("No Active Motion Design Game Instances found to apply the stat '%s' command on."), *InArgs[0]);
 		}
 	}
 	
diff --git a/Engine/Plugins/Experimental/Avalanche/Source/AvalancheMedia/Private/Playback/Nodes/AvaPlaybackNodePlayer.cpp b/Engine/Plugins/Experimental/Avalanche/Source/AvalancheMedia/Private/Playback/Nodes/AvaPlaybackNodePlayer.cpp
--- a/Engine/Plugins/Experimental/Avalanche/Source/AvalancheMedia/Private/Playback/Nodes/AvaPlaybackNodePlayer.cpp
+++ b/Engine/Plugins/Experimental/Avalanche/Source/AvalancheMedia/Private/Playback/Nodes/AvaPlaybackNodePlayer.cpp
@@ -77,7 +77,7 @@ FText UAvaPlaybackNodePlayer::GetNodeDisplayNameText() const
 
 FText UAvaPlaybackNodePlayer::GetNodeTooltipText() const
 {
-	return LOCTEXT("PlayerNode_ToolTip", "Plays the given Avalanche Asset");
+	return LOCTEXT("PlayerNode_ToolTip", "Plays the given Motion Design Asset");
 }
 
 #if WITH_EDITOR
diff --git a/Engine/Plugins/Experimental/Avalanche/Source/AvalancheMedia/Private/Playback/Nodes/Events/Actions/AvaPlaybackNode_PlayAnim.cpp b/Engine/Plugins/Experimental/Avalanche/Source/AvalancheMedia/Private/Playback/Nodes/Events/Actions/AvaPlaybackNode_PlayAnim.cpp
--- a/Engine/Plugins/Experimental/Avalanche/Source/AvalancheMedia/Private/Playback/Nodes/Events/Actions/AvaPlaybackNode_PlayAnim.cpp
+++ b/Engine/Plugins/Experimental/Avalanche/Source/AvalancheMedia/Private/Playback/Nodes/Events/Actions/AvaPlaybackNode_PlayAnim.cpp
@@ -55,7 +55,7 @@ FText UAvaPlaybackNode_PlayAnim::GetNodeDisplayNameText() const
 
 FText UAvaPlaybackNode_PlayAnim::GetNodeTooltipText() const
 {
-	return LOCTEXT("PlayAnimNode_Tooltip", "Plays an Animation from the Avalanche Asset");
+	return LOCTEXT("PlayAnimNode_Tooltip", "Plays an Animation from the Motion Design Asset");
 }
 
 void UAvaPlaybackNode_PlayAnim::OnEventTriggered(const FAvaPlaybackEventParameters& InEventParameters)
diff --git a/Engine/Plugins/Experimental/Avalanche/Source/AvalancheMedia/Private/Playback/Nodes/Events/Actions/AvaPlaybackNode_PreloadPlayer.cpp b/Engine/Plugins/Experimental/Avalanche/Source/AvalancheMedia/Private/Playback/Nodes/Events/Actions/AvaPlaybackNode_PreloadPlayer.cpp
--- a/Engine/Plugins/Experimental/Avalanche/Source/AvalancheMedia/Private/Playback/Nodes/Events/Actions/AvaPlaybackNode_PreloadPlayer.cpp
+++ b/Engine/Plugins/Experimental/Avalanche/Source/AvalancheMedia/Private/Playback/Nodes/Events/Actions/AvaPlaybackNode_PreloadPlayer.cpp
@@ -13,7 +13,7 @@ FText UAvaPlaybackNode_PreloadPlayer::GetNodeDisplayNameText() const
 
 FText UAvaPlaybackNode_PreloadPlayer::GetNodeTooltipText() const
 {
-	return LOCTEXT("PreloadPlayerNode_Tooltip", "Loads the Avalanche World Preemptively and Begins Play on that World");
+	return LOCTEXT("PreloadPlayerNode_Tooltip", "Loads the Motion Design World Preemptively and Begins Play on that World");
 }
 
 void UAvaPlaybackNode_PreloadPlayer::OnEventTriggered(const FAvaPlaybackEventParameters& InEventParameters)
diff --git a/Engine/Plugins/Experimental/Avalanche/Source/AvalancheMedia/Private/Playback/Playables/AvalancheLevelStreamingPlayable.cpp b/Engine/Plugins/Experimental/Avalanche/Source/AvalancheMedia/Private/Playback/Playables/AvalancheLevelStreamingPlayable.cpp
--- a/Engine/Plugins/Experimental/Avalanche/Source/AvalancheMedia/Private/Playback/Playables/AvalancheLevelStreamingPlayable.cpp
+++ b/Engine/Plugins/Experimental/Avalanche/Source/AvalancheMedia/Private/Playback/Playables/AvalancheLevelStreamingPlayable.cpp
@@ -505,7 +505,7 @@ void UAvalancheLevelStreamingPlayable::ResolveAvalancheScene(const ULevel* InLev
 		}
 		else
 		{
-			UE_LOG(LogAvalanchePlayable, Error, TEXT("Loaded level \"%s\" is not an Avalanche level."), *LevelStreaming->PackageNameToLoad.ToString());
+			UE_LOG(LogAvalanchePlayable, Error, TEXT("Loaded level \"%s\" is not an Motion Design level."), *LevelStreaming->PackageNameToLoad.ToString());
 		}
 	}
 }
diff --git a/Engine/Plugins/Experimental/Avalanche/Source/AvalancheMedia/Private/Playlist/AvalancheManagedInstanceBlueprint.cpp b/Engine/Plugins/Experimental/Avalanche/Source/AvalancheMedia/Private/Playlist/AvalancheManagedInstanceBlueprint.cpp
--- a/Engine/Plugins/Experimental/Avalanche/Source/AvalancheMedia/Private/Playlist/AvalancheManagedInstanceBlueprint.cpp
+++ b/Engine/Plugins/Experimental/Avalanche/Source/AvalancheMedia/Private/Playlist/AvalancheManagedInstanceBlueprint.cpp
@@ -47,7 +47,7 @@ FAvalancheManagedInstanceBlueprint::FAvalancheManagedInstanceBlueprint(FAvalanch
 	SourceAvalancheBlueprint = LoadAvalancheBlueprint(InAssetPath);
 	if (!SourceAvalancheBlueprint)
 	{
-		UE_LOG(LogAvaMedia, Error, TEXT("Unable to load Source Avalanche Blueprint: %s"), *InAssetPath.ToString());
+		UE_LOG(LogAvaMedia, Error, TEXT("Unable to load Source Motion Design Blueprint: %s"), *InAssetPath.ToString());
 		return;
 	}
 	
@@ -60,7 +60,7 @@ FAvalancheManagedInstanceBlueprint::FAvalancheManagedInstanceBlueprint(FAvalanch
 	ManagedAvalancheBlueprintPackage = FAvalancheManagedInstanceUtils::MakeManagedInstancePackage(InAssetPath);
 	if (!ManagedAvalancheBlueprintPackage)
 	{
-		UE_LOG(LogAvaMedia, Error, TEXT("Unable to create a Managed Avalanche Blueprint Package for %s"), *InAssetPath.ToString());
+		UE_LOG(LogAvaMedia, Error, TEXT("Unable to create a Managed Motion Design Blueprint Package for %s"), *InAssetPath.ToString());
 		return;
 	}
 
@@ -68,7 +68,7 @@ FAvalancheManagedInstanceBlueprint::FAvalancheManagedInstanceBlueprint(FAvalanch
 	ManagedAvalancheBlueprint = Cast<UAvalancheBlueprint>(StaticDuplicateObject(SourceAvalancheBlueprint.Get(), ManagedAvalancheBlueprintPackage.Get(), NAME_None, RF_NoFlags, nullptr, EDuplicateMode::PIE));
 	if (!ManagedAvalancheBlueprint)
 	{
-		UE_LOG(LogAvaMedia, Error, TEXT("Unable to duplicate Source Avalanche Blueprint: %s"), *InAssetPath.ToString());
+		UE_LOG(LogAvaMedia, Error, TEXT("Unable to duplicate Source Motion Design Blueprint: %s"), *InAssetPath.ToString());
 		return;
 	}
 
diff --git a/Engine/Plugins/Experimental/Avalanche/Source/AvalancheMedia/Private/Playlist/AvalancheManagedInstanceCache.cpp b/Engine/Plugins/Experimental/Avalanche/Source/AvalancheMedia/Private/Playlist/AvalancheManagedInstanceCache.cpp
--- a/Engine/Plugins/Experimental/Avalanche/Source/AvalancheMedia/Private/Playlist/AvalancheManagedInstanceCache.cpp
+++ b/Engine/Plugins/Experimental/Avalanche/Source/AvalancheMedia/Private/Playlist/AvalancheManagedInstanceCache.cpp
@@ -220,7 +220,7 @@ void FAvalancheManagedInstanceCache::OnPackageModified(const FName& InPackageNam
 			if (InAssetPath.GetLongPackageFName() == InPackageName)
 			{
 				UE_LOG(LogAvaMedia, Log,
-					TEXT("Managed Avalanche Blueprint Cache: Package \"%s\" being touched caused asset \"%s\" to be invalidated."),
+					TEXT("Managed Motion Design Blueprint Cache: Package \"%s\" being touched caused asset \"%s\" to be invalidated."),
 					*InPackageName.ToString(), *InAssetPath.ToString());
 				return true;
 			}
diff --git a/Engine/Plugins/Experimental/Avalanche/Source/AvalancheMedia/Private/Playlist/AvalancheManagedInstanceLevel.cpp b/Engine/Plugins/Experimental/Avalanche/Source/AvalancheMedia/Private/Playlist/AvalancheManagedInstanceLevel.cpp
--- a/Engine/Plugins/Experimental/Avalanche/Source/AvalancheMedia/Private/Playlist/AvalancheManagedInstanceLevel.cpp
+++ b/Engine/Plugins/Experimental/Avalanche/Source/AvalancheMedia/Private/Playlist/AvalancheManagedInstanceLevel.cpp
@@ -35,7 +35,7 @@ FAvalancheManagedInstanceLevel::FAvalancheManagedInstanceLevel(FAvalancheManaged
 	UWorld* SourceAvalancheLevel = LoadAvalancheLevel(InAssetPath);
 	if (!SourceAvalancheLevel)
 	{
-		UE_LOG(LogAvaMedia, Error, TEXT("Unable to load Source Avalanche Level: %s"), *InAssetPath.ToString());
+		UE_LOG(LogAvaMedia, Error, TEXT("Unable to load Source Motion Design Level: %s"), *InAssetPath.ToString());
 		return;
 	}
 
@@ -48,7 +48,7 @@ FAvalancheManagedInstanceLevel::FAvalancheManagedInstanceLevel(FAvalancheManaged
 	ManagedAvalancheLevelPackage = FAvalancheManagedInstanceUtils::MakeManagedInstancePackage(InAssetPath);
 	if (!ManagedAvalancheLevelPackage)
 	{
-		UE_LOG(LogAvaMedia, Error, TEXT("Unable to create a Managed Avalanche Level Package for %s"), *InAssetPath.ToString());
+		UE_LOG(LogAvaMedia, Error, TEXT("Unable to create a Managed Motion Design Level Package for %s"), *InAssetPath.ToString());
 		return;
 	}
 
@@ -56,7 +56,7 @@ FAvalancheManagedInstanceLevel::FAvalancheManagedInstanceLevel(FAvalancheManaged
 	ManagedAvalancheLevel = Cast<UWorld>(StaticDuplicateObject(SourceAvalancheLevel, ManagedAvalancheLevelPackage.Get(), NAME_None, RF_NoFlags, nullptr, EDuplicateMode::PIE));
 	if (!ManagedAvalancheLevel)
 	{
-		UE_LOG(LogAvaMedia, Error, TEXT("Unable to duplicate Source Avalanche Level: %s"), *InAssetPath.ToString());
+		UE_LOG(LogAvaMedia, Error, TEXT("Unable to duplicate Source Motion Design Level: %s"), *InAssetPath.ToString());
 		return;
 	}
 
diff --git a/Engine/Plugins/Experimental/Avalanche/Source/AvalancheMedia/Private/Playlist/AvalancheManagedInstanceUtils.cpp b/Engine/Plugins/Experimental/Avalanche/Source/AvalancheMedia/Private/Playlist/AvalancheManagedInstanceUtils.cpp
--- a/Engine/Plugins/Experimental/Avalanche/Source/AvalancheMedia/Private/Playlist/AvalancheManagedInstanceUtils.cpp
+++ b/Engine/Plugins/Experimental/Avalanche/Source/AvalancheMedia/Private/Playlist/AvalancheManagedInstanceUtils.cpp
@@ -31,7 +31,7 @@ UPackage* FAvalancheManagedInstanceUtils::MakeManagedInstancePackage(const FSoft
 		FString UniqueInstancePackageName = MakeUniqueObjectName( nullptr, UPackage::StaticClass(), FName(*InstancePackageName)).ToString();
 		
 		UE_LOG(LogAvaMedia, Warning, 
-			TEXT("Package \"%s\" for Avalanche Managed Instance already exists. Will use \"%s\" instead."),
+			TEXT("Package \"%s\" for Motion Design Managed Instance already exists. Will use \"%s\" instead."),
 			*InstancePackageName, *UniqueInstancePackageName);
 
 		InstancePackageName = UniqueInstancePackageName;
@@ -46,7 +46,7 @@ UPackage* FAvalancheManagedInstanceUtils::MakeManagedInstancePackage(const FSoft
 	{
 		// Note: The outer will fallback to GEngine in that case.
 		UE_LOG(LogAvaMedia, Error, 
-			TEXT("Unable to create package \"%s\" for Avalanche Managed Instance."),
+			TEXT("Unable to create package \"%s\" for Motion Design Managed Instance."),
 			*InstancePackageName);
 	}
 	return ManagedAvalancheBlueprintPackage;
diff --git a/Engine/Plugins/Experimental/Avalanche/Source/AvalancheMedia/Public/AvaMediaDefines.h b/Engine/Plugins/Experimental/Avalanche/Source/AvalancheMedia/Public/AvaMediaDefines.h
--- a/Engine/Plugins/Experimental/Avalanche/Source/AvalancheMedia/Public/AvaMediaDefines.h
+++ b/Engine/Plugins/Experimental/Avalanche/Source/AvalancheMedia/Public/AvaMediaDefines.h
@@ -29,7 +29,7 @@ enum class EAvaChannelState : uint8
  * Some additional restrictions are applied according to channel type:
  * - preview channels must only have outputs local to the process. "Remote" previews are not supported.
  * - [backend] playback request type (program or preview) must match with the channel type. This is a safety
- *	 net for any other extended code paths that are not in the Avalanche plugin.
+ *	 net for any other extended code paths that are not in the Motion Design plugin.
  */
 UENUM()
 enum class EAvaBroadcastChannelType : uint8
@@ -232,7 +232,7 @@ ENUM_CLASS_FLAGS(EAvaPageChanges);
  *
  *	This is also sometimes referred to as the "source" asset, by contrast to the
  *	runtime (or managed) transient asset for runtime playback. It will also refer to
- *	the source Avalanche blueprint, rather than a playback object.
+ *	the source Motion Design blueprint, rather than a playback object.
  */
 UENUM()
 enum class EAvaMediaPlaybackAssetStatus
@@ -243,7 +243,7 @@ enum class EAvaMediaPlaybackAssetStatus
 	Missing,
 	/**
 	 * Missing asset dependencies
-	 * Note: An Avalanche asset can run even with missing dependencies.
+	 * Note: An Motion Design asset can run even with missing dependencies.
 	 */
 	 MissingDependencies,
 	 /** Asset is out of date from compare with remote. */
diff --git a/Engine/Plugins/Experimental/Avalanche/Source/AvalancheMedia/Public/AvalancheMediaSettings.h b/Engine/Plugins/Experimental/Avalanche/Source/AvalancheMedia/Public/AvalancheMediaSettings.h
--- a/Engine/Plugins/Experimental/Avalanche/Source/AvalancheMedia/Public/AvalancheMediaSettings.h
+++ b/Engine/Plugins/Experimental/Avalanche/Source/AvalancheMedia/Public/AvalancheMediaSettings.h
@@ -183,10 +183,10 @@ public:
 	FAvalancheInstanceSettings AvalancheInstanceSettings;
 
 	/**
-	 * Maximum cached Managed Avalanche Blueprints used for playlist editor's page details.
+	 * Maximum cached Managed Motion Design Blueprints used for playlist editor's page details.
 	 * A value of 0 indicate the cache will grow without limit.
 	 */
-	UPROPERTY(Config, EditAnywhere, Category = "Managed Avalanche Instance Cache", meta = (DisplayName = "Maximum Cache Size"))
+	UPROPERTY(Config, EditAnywhere, Category = "Managed Motion Design Instance Cache", meta = (DisplayName = "Maximum Cache Size"))
 	int32 ManagedAvalancheInstanceCacheMaximumSize = 20;
 	
 	/** Whether web server is started automatically. */
diff --git a/Engine/Plugins/Experimental/Avalanche/Source/AvalancheMedia/Public/Playback/AvaMediaPlaybackMessages.h b/Engine/Plugins/Experimental/Avalanche/Source/AvalancheMedia/Public/Playback/AvaMediaPlaybackMessages.h
--- a/Engine/Plugins/Experimental/Avalanche/Source/AvalancheMedia/Public/Playback/AvaMediaPlaybackMessages.h
+++ b/Engine/Plugins/Experimental/Avalanche/Source/AvalancheMedia/Public/Playback/AvaMediaPlaybackMessages.h
@@ -205,7 +205,7 @@ struct AVALANCHEMEDIA_API FAvaDeviceProviderDataRequest : public FAvaMediaPlayba
 };
 
 /**
- *	Request for the client to replicate it's Avalanche instance settings to the server.
+ *	Request for the client to replicate it's Motion Design instance settings to the server.
  **/
 USTRUCT()
 struct AVALANCHEMEDIA_API FAvalancheInstanceSettingsUpdate : public FAvaMediaPlaybackClientMessageBase
@@ -310,7 +310,7 @@ struct AVALANCHEMEDIA_API FAvaMediaPlaybackCommand
 	UPROPERTY()
 	FSoftObjectPath AssetPath;
 
-	/** Channel to Play, only considered when the Asset in question is an Avalanche Playable (Blueprint or Level). */
+	/** Channel to Play, only considered when the Asset in question is an Motion Design Playable (Blueprint or Level). */
 	UPROPERTY()
 	FString ChannelName;
 	
@@ -413,7 +413,7 @@ struct FAvaMediaAnimPlaybackRequest : public FAvaMediaPlaybackClientMessageBase
 	UPROPERTY()
 	FSoftObjectPath AssetPath;
 	
-	/** Channel to Play, only considered when the Asset in question is an Avalanche Blueprint */
+	/** Channel to Play, only considered when the Asset in question is an Motion Design Blueprint */
 	UPROPERTY()
 	FString ChannelName;
 	
@@ -459,7 +459,7 @@ struct FAvaMediaRemoteControlUpdateRequest : public FAvaMediaPlaybackClientMessa
 	UPROPERTY()
 	FSoftObjectPath AssetPath;
 	
-	/** Channel to Play, only considered when the Asset in question is an Avalanche Blueprint. */
+	/** Channel to Play, only considered when the Asset in question is an Motion Design Blueprint. */
 	UPROPERTY()
 	FString ChannelName;
 	
diff --git a/Engine/Plugins/Experimental/Avalanche/Source/AvalancheMedia/Public/Playlist/AvalanchePlaylist.h b/Engine/Plugins/Experimental/Avalanche/Source/AvalancheMedia/Public/Playlist/AvalanchePlaylist.h
--- a/Engine/Plugins/Experimental/Avalanche/Source/AvalancheMedia/Public/Playlist/AvalanchePlaylist.h
+++ b/Engine/Plugins/Experimental/Avalanche/Source/AvalancheMedia/Public/Playlist/AvalanchePlaylist.h
@@ -252,8 +252,8 @@ public:
  * This class is a container for what could be described as a "show" for broadcast purposes.
  *
  * It goes beyond a simple playlist of items. It contains the following:
- * - a list of Avalanche Template Pages (or just Templates).
- * - a list of Avalanche Instanced Pages (or just Pages).
+ * - a list of Motion Design Template Pages (or just Templates).
+ * - a list of Motion Design Instanced Pages (or just Pages).
  * - a list of page views (or just Views).
  *
  * Workflow:
@@ -281,7 +281,7 @@ public:
  * pages could be grouped in either of the page list or page views.
  * Other applications support page grouping to emulate MOS's hierarchy.
  * In the MOS/NCS hierarchies: Rundown -> Stories/Segments -> Parts -> Pieces/Items
- * Although full emulation of MOS schema may not be necessary within the Avalanche playback framework.
+ * Although full emulation of MOS schema may not be necessary within the Motion Design playback framework.
  *
  */
 UCLASS(NotBlueprintable, BlueprintType)
diff --git a/Engine/Plugins/Experimental/Avalanche/Source/AvalancheMedia/Public/Playlist/AvalanchePlaylistComponent.h b/Engine/Plugins/Experimental/Avalanche/Source/AvalancheMedia/Public/Playlist/AvalanchePlaylistComponent.h
--- a/Engine/Plugins/Experimental/Avalanche/Source/AvalancheMedia/Public/Playlist/AvalanchePlaylistComponent.h
+++ b/Engine/Plugins/Experimental/Avalanche/Source/AvalancheMedia/Public/Playlist/AvalanchePlaylistComponent.h
@@ -9,7 +9,7 @@ class UAvalanchePlaylist;
 
 /**
  * Add this actor component to blueprint actor to expose the API to control an
- * Avalanche playlist in game.
+ * Motion Design playlist in game.
  */
 UCLASS(BlueprintType, Blueprintable, ClassGroup = "Motion Design Media", meta = (BlueprintSpawnableComponent))
 class AVALANCHEMEDIA_API UAvalanchePlaylistComponent : public UActorComponent
diff --git a/Engine/Plugins/Experimental/Avalanche/Source/AvalancheMediaEditor/Private/AvaMediaEditorModule.cpp b/Engine/Plugins/Experimental/Avalanche/Source/AvalancheMediaEditor/Private/AvaMediaEditorModule.cpp
--- a/Engine/Plugins/Experimental/Avalanche/Source/AvalancheMediaEditor/Private/AvaMediaEditorModule.cpp
+++ b/Engine/Plugins/Experimental/Avalanche/Source/AvalancheMediaEditor/Private/AvaMediaEditorModule.cpp
@@ -263,7 +263,7 @@ void FAvaMediaEditorModule::AddEditorToolbarButtons()
 	FToolMenuEntry OpenBroadcastButtonEntry = FToolMenuEntry::InitToolBarButton(TEXT("OpenBroadcastToolbarButton")
 		, FExecuteAction::CreateStatic(&FAvaBroadcastEditor::OpenBroadcastEditor)
 		, LOCTEXT("OpenBroadcast_Title", "Broadcast")
-		, LOCTEXT("OpenBroadcast_Tooltip", "Opens the Avalanche Broadcast Editor Window")
+		, LOCTEXT("OpenBroadcast_Tooltip", "Opens the Motion Design Broadcast Editor Window")
 		, TAttribute<FSlateIcon>::Create([]() { return IAvaMediaEditorModule::Get().GetToolbarBroadcastButtonIcon(); })
 	);
 	OpenBroadcastButtonEntry.StyleNameOverride = TEXT("CalloutToolbar"); // Display Labels
diff --git a/Engine/Plugins/Experimental/Avalanche/Source/AvalancheMediaEditor/Private/Broadcast/AvaBroadcastEditor.cpp b/Engine/Plugins/Experimental/Avalanche/Source/AvalancheMediaEditor/Private/Broadcast/AvaBroadcastEditor.cpp
--- a/Engine/Plugins/Experimental/Avalanche/Source/AvalancheMediaEditor/Private/Broadcast/AvaBroadcastEditor.cpp
+++ b/Engine/Plugins/Experimental/Avalanche/Source/AvalancheMediaEditor/Private/Broadcast/AvaBroadcastEditor.cpp
@@ -137,6 +137,11 @@ FText FAvaBroadcastEditor::GetBaseToolkitName() const
 	return LOCTEXT("BroadcastAppLabel", "Motion Design Broadcast Editor");
 }
 
+FText FAvaBroadcastEditor::GetToolkitName() const
+{
+	return LOCTEXT("BroadcastAppName", "Broadcast");
+}
+
 FText FAvaBroadcastEditor::GetToolkitToolTipText() const
 {
 	if (IAvaMediaModule::Get().IsMediaPlaybackServerStarted())
diff --git a/Engine/Plugins/Experimental/Avalanche/Source/AvalancheMediaEditor/Private/Playback/AvaPlaybackEditor.cpp b/Engine/Plugins/Experimental/Avalanche/Source/AvalancheMediaEditor/Private/Playback/AvaPlaybackEditor.cpp
--- a/Engine/Plugins/Experimental/Avalanche/Source/AvalancheMediaEditor/Private/Playback/AvaPlaybackEditor.cpp
+++ b/Engine/Plugins/Experimental/Avalanche/Source/AvalancheMediaEditor/Private/Playback/AvaPlaybackEditor.cpp
@@ -138,7 +138,7 @@ void FAvaPlaybackEditor::FillPlayToolBar(FToolBarBuilder& ToolBarBuilder)
 		ToolBarBuilder.AddToolBarButton(FExecuteAction::CreateStatic(&FAvaBroadcastEditor::OpenBroadcastEditor)
 		, NAME_None
 		, LOCTEXT("Broadcast_Label", "Broadcast")
-		, LOCTEXT("Broadcast_ToolTip", "Opens the Avalanche Broadcast Editor Window")
+		, LOCTEXT("Broadcast_ToolTip", "Opens the Motion Design Broadcast Editor Window")
 		, TAttribute<FSlateIcon>::Create([]() { return IAvaMediaEditorModule::Get().GetToolbarBroadcastButtonIcon(); }));
 	}
 	ToolBarBuilder.EndSection();
diff --git a/Engine/Plugins/Experimental/Avalanche/Source/AvalancheMediaEditor/Private/Playlist/AvaPlaylistCommands.cpp b/Engine/Plugins/Experimental/Avalanche/Source/AvalancheMediaEditor/Private/Playlist/AvaPlaylistCommands.cpp
--- a/Engine/Plugins/Experimental/Avalanche/Source/AvalancheMediaEditor/Private/Playlist/AvaPlaylistCommands.cpp
+++ b/Engine/Plugins/Experimental/Avalanche/Source/AvalancheMediaEditor/Private/Playlist/AvaPlaylistCommands.cpp
@@ -44,7 +44,7 @@ void FAvaPlaylistCommands::RegisterCommands()
 
 	UI_COMMAND(EditPageSource
 		, "Edit Scene"
-		, "Edit source Avalanche Asset of the selected page"
+		, "Edit source Motion Design Asset of the selected page"
 		, EUserInterfaceActionType::Button
 		, FInputChord());
 
diff --git a/Engine/Plugins/Experimental/Avalanche/Source/AvalancheMediaEditor/Private/Playlist/AvaPlaylistEditor.cpp b/Engine/Plugins/Experimental/Avalanche/Source/AvalancheMediaEditor/Private/Playlist/AvaPlaylistEditor.cpp
--- a/Engine/Plugins/Experimental/Avalanche/Source/AvalancheMediaEditor/Private/Playlist/AvaPlaylistEditor.cpp
+++ b/Engine/Plugins/Experimental/Avalanche/Source/AvalancheMediaEditor/Private/Playlist/AvaPlaylistEditor.cpp
@@ -943,7 +943,7 @@ void FAvaPlaylistEditor::FillPageToolBar(FToolBarBuilder& OutToolBarBuilder)
 		OutToolBarBuilder.AddToolBarButton(FExecuteAction::CreateStatic(&FAvaBroadcastEditor::OpenBroadcastEditor)
 			, NAME_None
 			, LOCTEXT("Broadcast_Label", "Broadcast")
-			, LOCTEXT("Broadcast_ToolTip", "Opens the Avalanche Broadcast Editor Window")
+			, LOCTEXT("Broadcast_ToolTip", "Opens the Motion Design Broadcast Editor Window")
 			, TAttribute<FSlateIcon>::Create([]() { return IAvaMediaEditorModule::Get().GetToolbarBroadcastButtonIcon(); }));
 		OutToolBarBuilder.AddComboButton(
 			FUIAction(FExecuteAction()
diff --git a/Engine/Plugins/Experimental/Avalanche/Source/AvalancheMediaEditor/Private/Playlist/AvaPlaylistEditorUtils.cpp b/Engine/Plugins/Experimental/Avalanche/Source/AvalancheMediaEditor/Private/Playlist/AvaPlaylistEditorUtils.cpp
--- a/Engine/Plugins/Experimental/Avalanche/Source/AvalancheMediaEditor/Private/Playlist/AvaPlaylistEditorUtils.cpp
+++ b/Engine/Plugins/Experimental/Avalanche/Source/AvalancheMediaEditor/Private/Playlist/AvaPlaylistEditorUtils.cpp
@@ -273,7 +273,7 @@ TArray<FAvalanchePage> UE::AvaPlaylistEditor::Utils::DeserializePagesFromJson(co
 		}
 		else
 		{
-			UE_LOG(LogAvaPlaylist, Warning, TEXT("Unable to convert Page Entry Json Object to Avalanche Page Struct"));
+			UE_LOG(LogAvaPlaylist, Warning, TEXT("Unable to convert Page Entry Json Object to Motion Design Page Struct"));
 		}
 	}
 	return Pages;
diff --git a/Engine/Plugins/Experimental/Avalanche/Source/AvalancheMediaEditor/Private/Playlist/Pages/PageViews/AvaPageViewImpl.cpp b/Engine/Plugins/Experimental/Avalanche/Source/AvalancheMediaEditor/Private/Playlist/Pages/PageViews/AvaPageViewImpl.cpp
--- a/Engine/Plugins/Experimental/Avalanche/Source/AvalancheMediaEditor/Private/Playlist/Pages/PageViews/AvaPageViewImpl.cpp
+++ b/Engine/Plugins/Experimental/Avalanche/Source/AvalancheMediaEditor/Private/Playlist/Pages/PageViews/AvaPageViewImpl.cpp
@@ -135,7 +135,7 @@ void FAvaPageViewImpl::OnObjectChanged(const FAssetData& InAssetData)
 		SetPageSelection(EAvaPageViewSelectionChangeType::ReplaceSelection);
 	}
 
-	PerformWorkOnPages(LOCTEXT("SetAvaBlueprint", "Set Avalanche Blueprint"),
+	PerformWorkOnPages(LOCTEXT("SetAvaBlueprint", "Set Motion Design Blueprint"),
 		[this, &InAssetData](FAvalanchePage& InPage)->bool
 		{
 			if (!InPage.UpdateAvalancheAsset(InAssetData.GetSoftObjectPath())) 
diff --git a/Engine/Plugins/Experimental/Avalanche/Source/AvalancheMediaEditor/Private/Playlist/Pages/Slate/SAvaPageList.cpp b/Engine/Plugins/Experimental/Avalanche/Source/AvalancheMediaEditor/Private/Playlist/Pages/Slate/SAvaPageList.cpp
--- a/Engine/Plugins/Experimental/Avalanche/Source/AvalancheMediaEditor/Private/Playlist/Pages/Slate/SAvaPageList.cpp
+++ b/Engine/Plugins/Experimental/Avalanche/Source/AvalancheMediaEditor/Private/Playlist/Pages/Slate/SAvaPageList.cpp
@@ -866,7 +866,7 @@ namespace UE::AvaPageList::Private
 				if (bInLogInfo)
 				{
 					UE_LOG(LogAvaMediaEditor, Display,
-						TEXT("Level Asset \"%s\" is not an Avalanche Scene (Asset Tag not found)."),
+						TEXT("Level Asset \"%s\" is not an Motion Design Scene (Asset Tag not found)."),
 						*InAssetData.GetSoftObjectPath().ToString());
 				}
 				return false;
@@ -877,7 +877,7 @@ namespace UE::AvaPageList::Private
 				if (bInLogInfo)
 				{
 					UE_LOG(LogAvaMediaEditor, Display,
-						TEXT("Level Asset \"%s\" is an Avalanche Scene but not enabled."),
+						TEXT("Level Asset \"%s\" is an Motion Design Scene but not enabled."),
 						*InAssetData.GetSoftObjectPath().ToString());
 				}
 				return false;						
@@ -994,7 +994,7 @@ bool SAvaPageList::HandleDropEvent(const FDragDropEvent& InDragDropEvent, EItemD
 		const TArray<FSoftObjectPath> AvaAssets = FilterAvaAssetPaths(AssetDragDropOp->GetAssets());
 		const TArray<FSoftObjectPath> Playlists = FilterPlaylistPaths(AssetDragDropOp->GetAssets(), FSoftObjectPath(GetPlaylist()));
 
-		FScopedTransaction DropTransaction(LOCTEXT("DropTransaction", "Drop Avalanche Assets onto Page List"));
+		FScopedTransaction DropTransaction(LOCTEXT("DropTransaction", "Drop Motion Design Assets onto Page List"));
 
 		bool bIsHandled = false;
 		if (!Playlists.IsEmpty())
diff --git a/Engine/Plugins/Experimental/Avalanche/Source/AvalancheMediaEditor/Private/Playlist/Preview/SAvaPagePreview.cpp b/Engine/Plugins/Experimental/Avalanche/Source/AvalancheMediaEditor/Private/Playlist/Preview/SAvaPagePreview.cpp
--- a/Engine/Plugins/Experimental/Avalanche/Source/AvalancheMediaEditor/Private/Playlist/Preview/SAvaPagePreview.cpp
+++ b/Engine/Plugins/Experimental/Avalanche/Source/AvalancheMediaEditor/Private/Playlist/Preview/SAvaPagePreview.cpp
@@ -198,7 +198,7 @@ TSharedRef<SWidget> SAvaPagePreview::OnGenerateSettingsMenu()
 
 	MenuBuilder.AddMenuEntry(
 		LOCTEXT("Settings", "Settings"),
-		LOCTEXT("Settings_Tooltip", "Opens the Avalanche Media Settings."),
+		LOCTEXT("Settings_Tooltip", "Opens the Motion Design Media Settings."),
 		FSlateIcon(),
 		FUIAction(
 			FExecuteAction::CreateSP(this, &SAvaPagePreview::HandleSettingsActionExecute)
diff --git a/Engine/Plugins/Experimental/Avalanche/Source/AvalancheMediaEditor/Public/Broadcast/AvaBroadcastEditor.h b/Engine/Plugins/Experimental/Avalanche/Source/AvalancheMediaEditor/Public/Broadcast/AvaBroadcastEditor.h
--- a/Engine/Plugins/Experimental/Avalanche/Source/AvalancheMediaEditor/Public/Broadcast/AvaBroadcastEditor.h
+++ b/Engine/Plugins/Experimental/Avalanche/Source/AvalancheMediaEditor/Public/Broadcast/AvaBroadcastEditor.h
@@ -45,6 +45,7 @@ protected:
 	//~ Begin IToolkit Interface
 	virtual FName GetToolkitFName() const override;
 	virtual FText GetBaseToolkitName() const override;
+	virtual FText GetToolkitName() const override;
 	virtual FText GetToolkitToolTipText() const override;
 	virtual FString GetWorldCentricTabPrefix() const override;
 	virtual FLinearColor GetWorldCentricTabColorScale() const override;
