[Avalanche] Change ease curve menus to use UToolMenu.

diff --git a/Engine/Plugins/Experimental/Avalanche/Source/AvalancheSequencer/Private/EaseCurveTool/Widgets/SAvaEaseCurveTool.cpp b/Engine/Plugins/Experimental/Avalanche/Source/AvalancheSequencer/Private/EaseCurveTool/Widgets/SAvaEaseCurveTool.cpp
--- a/Engine/Plugins/Experimental/Avalanche/Source/AvalancheSequencer/Private/EaseCurveTool/Widgets/SAvaEaseCurveTool.cpp
+++ b/Engine/Plugins/Experimental/Avalanche/Source/AvalancheSequencer/Private/EaseCurveTool/Widgets/SAvaEaseCurveTool.cpp
@@ -17,12 +17,14 @@
 #include "SCurveEditor.h"
 #include "Styling/AppStyle.h"
 #include "Styling/StyleColors.h"
+#include "ToolMenus.h"
 #include "Widgets/Images/SImage.h"
 #include "Widgets/Input/SButton.h"
 #include "Widgets/Input/SComboBox.h"
 #include "Widgets/Input/SSpinBox.h"
 #include "Widgets/Layout/SScrollBox.h"
 #include "Widgets/Layout/SWrapBox.h"
+#include "Widgets/SNullWidget.h"
 
 #define LOCTEXT_NAMESPACE "SAvaEaseCurveTool"
 
@@ -470,69 +472,85 @@ void SAvaEaseCurveTool::BindCommands()
 {
 	const FAvaEaseCurveToolCommands& AvaCurveEaseToolCommands = FAvaEaseCurveToolCommands::Get();
 
+	FAvaEaseCurveTool* const CurveEaseToolRawPtr = CurveEaseTool.Get();
+
 	CommandList = MakeShared<FUICommandList>();
 
-	CommandList->MapAction(FGenericCommands::Get().Undo, FExecuteAction::CreateSP(CurveEaseTool.Get(), &FAvaEaseCurveTool::UndoAction));
-	CommandList->MapAction(FGenericCommands::Get().Redo, FExecuteAction::CreateSP(CurveEaseTool.Get(), &FAvaEaseCurveTool::RedoAction));
+	CommandList->MapAction(FGenericCommands::Get().Undo, FExecuteAction::CreateSP(CurveEaseToolRawPtr, &FAvaEaseCurveTool::UndoAction));
+	
+	CommandList->MapAction(FGenericCommands::Get().Redo, FExecuteAction::CreateSP(CurveEaseToolRawPtr, &FAvaEaseCurveTool::RedoAction));
 
-	CommandList->MapAction(AvaCurveEaseToolCommands.OpenToolSettings, FExecuteAction::CreateSP(CurveEaseTool.Get(), &FAvaEaseCurveTool::OpenToolSettings));
+	CommandList->MapAction(AvaCurveEaseToolCommands.OpenToolSettings, FExecuteAction::CreateSP(CurveEaseToolRawPtr, &FAvaEaseCurveTool::OpenToolSettings));
 
-	CommandList->MapAction(AvaCurveEaseToolCommands.Refresh, FExecuteAction::CreateSP(CurveEaseTool.Get(), &FAvaEaseCurveTool::UpdateEaseCurveFromSequencerKeySelections));
-	CommandList->MapAction(AvaCurveEaseToolCommands.Apply, FExecuteAction::CreateSP(CurveEaseTool.Get(), &FAvaEaseCurveTool::ApplyEaseCurveToSequencerKeySelections));
+	CommandList->MapAction(AvaCurveEaseToolCommands.Refresh, FExecuteAction::CreateSP(CurveEaseToolRawPtr, &FAvaEaseCurveTool::UpdateEaseCurveFromSequencerKeySelections));
+	
+	CommandList->MapAction(AvaCurveEaseToolCommands.Apply, FExecuteAction::CreateSP(CurveEaseToolRawPtr, &FAvaEaseCurveTool::ApplyEaseCurveToSequencerKeySelections));
 
 	CommandList->MapAction(AvaCurveEaseToolCommands.ZoomToFit, FExecuteAction::CreateSP(this, &SAvaEaseCurveTool::ZoomToFit));
+
 	CommandList->MapAction(AvaCurveEaseToolCommands.ToggleGridSnap
 		, FExecuteAction::CreateUObject(GetMutableDefault<UAvaEaseCurveToolSettings>(), &UAvaEaseCurveToolSettings::ToggleGridSnap)
 		, FCanExecuteAction()
 		, FIsActionChecked::CreateUObject(GetDefault<UAvaEaseCurveToolSettings>(), &UAvaEaseCurveToolSettings::GetGridSnap));
+
 	CommandList->MapAction(AvaCurveEaseToolCommands.ToggleAutoFlipTangents
 		, FExecuteAction::CreateUObject(GetMutableDefault<UAvaEaseCurveToolSettings>(), &UAvaEaseCurveToolSettings::ToggleAutoFlipTangents)
 		, FCanExecuteAction()
 		, FIsActionChecked::CreateUObject(GetDefault<UAvaEaseCurveToolSettings>(), &UAvaEaseCurveToolSettings::GetAutoFlipTangents));
 
-	CommandList->MapAction(AvaCurveEaseToolCommands.SelectNextChannelKey, FExecuteAction::CreateSP(CurveEaseTool.Get(), &FAvaEaseCurveTool::SelectNextChannelKey));
-	CommandList->MapAction(AvaCurveEaseToolCommands.SelectPreviousChannelKey, FExecuteAction::CreateSP(CurveEaseTool.Get(), &FAvaEaseCurveTool::SelectPreviousChannelKey));
+	CommandList->MapAction(AvaCurveEaseToolCommands.SelectNextChannelKey, FExecuteAction::CreateSP(CurveEaseToolRawPtr, &FAvaEaseCurveTool::SelectNextChannelKey));
+	
+	CommandList->MapAction(AvaCurveEaseToolCommands.SelectPreviousChannelKey, FExecuteAction::CreateSP(CurveEaseToolRawPtr, &FAvaEaseCurveTool::SelectPreviousChannelKey));
 
 	CommandList->MapAction(AvaCurveEaseToolCommands.SetOperationToEaseOut
-		, FExecuteAction::CreateSP(CurveEaseTool.Get(), &FAvaEaseCurveTool::SetToolOperation, FAvaEaseCurveTool::EOperation::Out)
+		, FExecuteAction::CreateSP(CurveEaseToolRawPtr, &FAvaEaseCurveTool::SetToolOperation, FAvaEaseCurveTool::EOperation::Out)
 		, FCanExecuteAction()
-		, FIsActionChecked::CreateSP(CurveEaseTool.Get(), &FAvaEaseCurveTool::IsToolOperation, FAvaEaseCurveTool::EOperation::Out));
+		, FIsActionChecked::CreateSP(CurveEaseToolRawPtr, &FAvaEaseCurveTool::IsToolOperation, FAvaEaseCurveTool::EOperation::Out));
+
 	CommandList->MapAction(AvaCurveEaseToolCommands.SetOperationToEaseInOut
-		, FExecuteAction::CreateSP(CurveEaseTool.Get(), &FAvaEaseCurveTool::SetToolOperation, FAvaEaseCurveTool::EOperation::InOut)
+		, FExecuteAction::CreateSP(CurveEaseToolRawPtr, &FAvaEaseCurveTool::SetToolOperation, FAvaEaseCurveTool::EOperation::InOut)
 		, FCanExecuteAction()
-		, FIsActionChecked::CreateSP(CurveEaseTool.Get(), &FAvaEaseCurveTool::IsToolOperation, FAvaEaseCurveTool::EOperation::InOut));
+		, FIsActionChecked::CreateSP(CurveEaseToolRawPtr, &FAvaEaseCurveTool::IsToolOperation, FAvaEaseCurveTool::EOperation::InOut));
+
 	CommandList->MapAction(AvaCurveEaseToolCommands.SetOperationToEaseIn
-		, FExecuteAction::CreateSP(CurveEaseTool.Get(), &FAvaEaseCurveTool::SetToolOperation, FAvaEaseCurveTool::EOperation::In)
+		, FExecuteAction::CreateSP(CurveEaseToolRawPtr, &FAvaEaseCurveTool::SetToolOperation, FAvaEaseCurveTool::EOperation::In)
 		, FCanExecuteAction()
-		, FIsActionChecked::CreateSP(CurveEaseTool.Get(), &FAvaEaseCurveTool::IsToolOperation, FAvaEaseCurveTool::EOperation::In));
+		, FIsActionChecked::CreateSP(CurveEaseToolRawPtr, &FAvaEaseCurveTool::IsToolOperation, FAvaEaseCurveTool::EOperation::In));
 
 	CommandList->MapAction(AvaCurveEaseToolCommands.ResetTangents
-		, FExecuteAction::CreateSP(CurveEaseTool.Get(), &FAvaEaseCurveTool::ResetEaseCurveTangents, true, true));
+		, FExecuteAction::CreateSP(CurveEaseToolRawPtr, &FAvaEaseCurveTool::ResetEaseCurveTangents, true, true));
+
 	CommandList->MapAction(AvaCurveEaseToolCommands.ResetStartTangent
-		, FExecuteAction::CreateSP(CurveEaseTool.Get(), &FAvaEaseCurveTool::ResetEaseCurveTangents, true, false));
+		, FExecuteAction::CreateSP(CurveEaseToolRawPtr, &FAvaEaseCurveTool::ResetEaseCurveTangents, true, false));
+
 	CommandList->MapAction(AvaCurveEaseToolCommands.ResetEndTangent
-		, FExecuteAction::CreateSP(CurveEaseTool.Get(), &FAvaEaseCurveTool::ResetEaseCurveTangents, false, true));
+		, FExecuteAction::CreateSP(CurveEaseToolRawPtr, &FAvaEaseCurveTool::ResetEaseCurveTangents, false, true));
 
 	CommandList->MapAction(AvaCurveEaseToolCommands.FlattenTangents
-		, FExecuteAction::CreateSP(CurveEaseTool.Get(), &FAvaEaseCurveTool::FlattenOrStraightenTangents, true, true, true));
+		, FExecuteAction::CreateSP(CurveEaseToolRawPtr, &FAvaEaseCurveTool::FlattenOrStraightenTangents, true, true, true));
+
 	CommandList->MapAction(AvaCurveEaseToolCommands.FlattenStartTangent
-		, FExecuteAction::CreateSP(CurveEaseTool.Get(), &FAvaEaseCurveTool::FlattenOrStraightenTangents, true, false, true));
+		, FExecuteAction::CreateSP(CurveEaseToolRawPtr, &FAvaEaseCurveTool::FlattenOrStraightenTangents, true, false, true));
+
 	CommandList->MapAction(AvaCurveEaseToolCommands.FlattenEndTangent
-		, FExecuteAction::CreateSP(CurveEaseTool.Get(), &FAvaEaseCurveTool::FlattenOrStraightenTangents, false, true, true));
+		, FExecuteAction::CreateSP(CurveEaseToolRawPtr, &FAvaEaseCurveTool::FlattenOrStraightenTangents, false, true, true));
 
 	CommandList->MapAction(AvaCurveEaseToolCommands.StraightenTangents
-		, FExecuteAction::CreateSP(CurveEaseTool.Get(), &FAvaEaseCurveTool::FlattenOrStraightenTangents, true, true, false));
+		, FExecuteAction::CreateSP(CurveEaseToolRawPtr, &FAvaEaseCurveTool::FlattenOrStraightenTangents, true, true, false));
+
 	CommandList->MapAction(AvaCurveEaseToolCommands.StraightenStartTangent
-		, FExecuteAction::CreateSP(CurveEaseTool.Get(), &FAvaEaseCurveTool::FlattenOrStraightenTangents, true, false, false));
+		, FExecuteAction::CreateSP(CurveEaseToolRawPtr, &FAvaEaseCurveTool::FlattenOrStraightenTangents, true, false, false));
+
 	CommandList->MapAction(AvaCurveEaseToolCommands.StraightenEndTangent
-		, FExecuteAction::CreateSP(CurveEaseTool.Get(), &FAvaEaseCurveTool::FlattenOrStraightenTangents, false, true, false));
+		, FExecuteAction::CreateSP(CurveEaseToolRawPtr, &FAvaEaseCurveTool::FlattenOrStraightenTangents, false, true, false));
 
 	CommandList->MapAction(AvaCurveEaseToolCommands.CopyTangents
-		, FExecuteAction::CreateSP(CurveEaseTool.Get(), &FAvaEaseCurveTool::CopyTangentsToClipboard)
-		, FCanExecuteAction::CreateSP(CurveEaseTool.Get(), &FAvaEaseCurveTool::CanCopyTangentsToClipboard));
+		, FExecuteAction::CreateSP(CurveEaseToolRawPtr, &FAvaEaseCurveTool::CopyTangentsToClipboard)
+		, FCanExecuteAction::CreateSP(CurveEaseToolRawPtr, &FAvaEaseCurveTool::CanCopyTangentsToClipboard));
+
 	CommandList->MapAction(AvaCurveEaseToolCommands.PasteTangents
-		, FExecuteAction::CreateSP(CurveEaseTool.Get(), &FAvaEaseCurveTool::PasteTangentsFromClipboard)
-		, FCanExecuteAction::CreateSP(CurveEaseTool.Get(), &FAvaEaseCurveTool::CanPasteTangentsFromClipboard));
+		, FExecuteAction::CreateSP(CurveEaseToolRawPtr, &FAvaEaseCurveTool::PasteTangentsFromClipboard)
+		, FCanExecuteAction::CreateSP(CurveEaseToolRawPtr, &FAvaEaseCurveTool::CanPasteTangentsFromClipboard));
 
 	CommandList->MapAction(AvaCurveEaseToolCommands.CreateExternalCurveAsset
 		, FExecuteAction::CreateLambda([this]()
@@ -542,54 +560,74 @@ void SAvaEaseCurveTool::BindCommands()
 		, FCanExecuteAction());
 
 	CommandList->MapAction(AvaCurveEaseToolCommands.SetKeyInterpConstant,
-		FExecuteAction::CreateSP(CurveEaseTool.Get(), &FAvaEaseCurveTool::SetKeyInterpMode, RCIM_Constant, RCTM_Auto),
+		FExecuteAction::CreateSP(CurveEaseToolRawPtr, &FAvaEaseCurveTool::SetKeyInterpMode, RCIM_Constant, RCTM_Auto),
 		FCanExecuteAction(),
-		FIsActionChecked::CreateSP(CurveEaseTool.Get(), &FAvaEaseCurveTool::IsKeyInterpMode, RCIM_Constant, RCTM_Auto));
+		FIsActionChecked::CreateSP(CurveEaseToolRawPtr, &FAvaEaseCurveTool::IsKeyInterpMode, RCIM_Constant, RCTM_Auto));
+
 	CommandList->MapAction(AvaCurveEaseToolCommands.SetKeyInterpLinear,
-		FExecuteAction::CreateSP(CurveEaseTool.Get(), &FAvaEaseCurveTool::SetKeyInterpMode, RCIM_Linear, RCTM_Auto),
+		FExecuteAction::CreateSP(CurveEaseToolRawPtr, &FAvaEaseCurveTool::SetKeyInterpMode, RCIM_Linear, RCTM_Auto),
 		FCanExecuteAction(),
-		FIsActionChecked::CreateSP(CurveEaseTool.Get(), &FAvaEaseCurveTool::IsKeyInterpMode, RCIM_Linear, RCTM_Auto));
+		FIsActionChecked::CreateSP(CurveEaseToolRawPtr, &FAvaEaseCurveTool::IsKeyInterpMode, RCIM_Linear, RCTM_Auto));
+
 	CommandList->MapAction(AvaCurveEaseToolCommands.SetKeyInterpCubicAuto,
-		FExecuteAction::CreateSP(CurveEaseTool.Get(), &FAvaEaseCurveTool::SetKeyInterpMode, RCIM_Cubic, RCTM_Auto),
+		FExecuteAction::CreateSP(CurveEaseToolRawPtr, &FAvaEaseCurveTool::SetKeyInterpMode, RCIM_Cubic, RCTM_Auto),
 		FCanExecuteAction(),
-		FIsActionChecked::CreateSP(CurveEaseTool.Get(), &FAvaEaseCurveTool::IsKeyInterpMode, RCIM_Cubic, RCTM_Auto));
+		FIsActionChecked::CreateSP(CurveEaseToolRawPtr, &FAvaEaseCurveTool::IsKeyInterpMode, RCIM_Cubic, RCTM_Auto));
+
 	CommandList->MapAction(AvaCurveEaseToolCommands.SetKeyInterpCubicSmartAuto,
-		FExecuteAction::CreateSP(CurveEaseTool.Get(), &FAvaEaseCurveTool::SetKeyInterpMode, RCIM_Cubic, RCTM_SmartAuto),
+		FExecuteAction::CreateSP(CurveEaseToolRawPtr, &FAvaEaseCurveTool::SetKeyInterpMode, RCIM_Cubic, RCTM_SmartAuto),
 		FCanExecuteAction(),
-		FIsActionChecked::CreateSP(CurveEaseTool.Get(), &FAvaEaseCurveTool::IsKeyInterpMode, RCIM_Cubic, RCTM_SmartAuto));
+		FIsActionChecked::CreateSP(CurveEaseToolRawPtr, &FAvaEaseCurveTool::IsKeyInterpMode, RCIM_Cubic, RCTM_SmartAuto));
+
 	CommandList->MapAction(AvaCurveEaseToolCommands.SetKeyInterpCubicUser,
-		FExecuteAction::CreateSP(CurveEaseTool.Get(), &FAvaEaseCurveTool::SetKeyInterpMode, RCIM_Cubic, RCTM_User),
+		FExecuteAction::CreateSP(CurveEaseToolRawPtr, &FAvaEaseCurveTool::SetKeyInterpMode, RCIM_Cubic, RCTM_User),
 		FCanExecuteAction(),
-		FIsActionChecked::CreateSP(CurveEaseTool.Get(), &FAvaEaseCurveTool::IsKeyInterpMode, RCIM_Cubic, RCTM_User));
+		FIsActionChecked::CreateSP(CurveEaseToolRawPtr, &FAvaEaseCurveTool::IsKeyInterpMode, RCIM_Cubic, RCTM_User));
+
 	CommandList->MapAction(AvaCurveEaseToolCommands.SetKeyInterpCubicBreak,
-		FExecuteAction::CreateSP(CurveEaseTool.Get(), &FAvaEaseCurveTool::SetKeyInterpMode, RCIM_Cubic, RCTM_Break),
+		FExecuteAction::CreateSP(CurveEaseToolRawPtr, &FAvaEaseCurveTool::SetKeyInterpMode, RCIM_Cubic, RCTM_Break),
 		FCanExecuteAction(),
-		FIsActionChecked::CreateSP(CurveEaseTool.Get(), &FAvaEaseCurveTool::IsKeyInterpMode, RCIM_Cubic, RCTM_Break));
+		FIsActionChecked::CreateSP(CurveEaseToolRawPtr, &FAvaEaseCurveTool::IsKeyInterpMode, RCIM_Cubic, RCTM_Break));
 }
 
 TSharedRef<SWidget> SAvaEaseCurveTool::CreateContextMenuContent()
 {
-	const FAvaEaseCurveToolCommands& AvaCurveEaseToolCommands = FAvaEaseCurveToolCommands::Get();
+	UToolMenus* ToolMenus = UToolMenus::Get();
+	if (!ToolMenus)
+	{
+		return SNullWidget::NullWidget;
+	}
 
-	TSharedPtr<TArray<TSharedPtr<FCurveViewModel>>> CurvesToAddKeysTo = MakeShareable(new TArray<TSharedPtr<FCurveViewModel>>());
-	
-	FMenuBuilder MenuBuilder(/*bInShouldCloseWindowAfterMenuSelection=*/true, CommandList);
+	constexpr const TCHAR* MenuName = TEXT("AvaEaseCurveToolMenu");
 
-	MenuBuilder.BeginSection(TEXT("CurveEditorActions"), LOCTEXT("CurveAction", "Curve Ease Tool Actions"));
+	if (!ToolMenus->IsMenuRegistered(MenuName))
 	{
-		MenuBuilder.AddSubMenu(
+		UToolMenu* const ToolMenu = ToolMenus->RegisterMenu(MenuName, NAME_None, EMultiBoxType::Menu);
+
+		FToolMenuSection& Section = ToolMenu->FindOrAddSection(TEXT("EaseCurveTool"), LOCTEXT("EaseCurveToolActions", "Curve Ease Tool Actions"));
+
+		const FAvaEaseCurveToolCommands& AvaCurveEaseToolCommands = FAvaEaseCurveToolCommands::Get();
+
+		Section.AddSubMenu(TEXT("Settings"),
 			LOCTEXT("SettingsSubMenuLabel", "Settings"),
 			LOCTEXT("SettingsSubMenuToolTip", ""),
 			FNewMenuDelegate::CreateSP(this, &SAvaEaseCurveTool::MakeContextMenuSettings),
 			false,
 			FSlateIcon(FAppStyle::Get().GetStyleSetName(), TEXT("Icons.Toolbar.Settings")));
-		MenuBuilder.AddSeparator();
-		MenuBuilder.AddMenuEntry(AvaCurveEaseToolCommands.CreateExternalCurveAsset);
-		MenuBuilder.AddSeparator();
-		MenuBuilder.AddMenuEntry(AvaCurveEaseToolCommands.CopyTangents);
-		MenuBuilder.AddMenuEntry(AvaCurveEaseToolCommands.PasteTangents);
-		MenuBuilder.AddSeparator();
-		MenuBuilder.AddSubMenu(
+
+		Section.AddSeparator(NAME_None);
+
+		Section.AddMenuEntry(AvaCurveEaseToolCommands.CreateExternalCurveAsset);
+
+		Section.AddSeparator(NAME_None);
+
+		Section.AddMenuEntry(AvaCurveEaseToolCommands.CopyTangents);
+
+		Section.AddMenuEntry(AvaCurveEaseToolCommands.PasteTangents);
+
+		Section.AddSeparator(NAME_None);
+
+		Section.AddSubMenu(TEXT("StraightenTangents"),
 			LOCTEXT("StraightenTangentsSubMenuLabel", "Straighten Tangents"),
 			LOCTEXT("StraightenTangentsSubMenuToolTip", ""),
 			FNewMenuDelegate::CreateLambda([&AvaCurveEaseToolCommands](FMenuBuilder& InMenuBuilder)
@@ -599,7 +637,8 @@ TSharedRef<SWidget> SAvaEaseCurveTool::CreateContextMenuContent()
 					InMenuBuilder.AddMenuEntry(AvaCurveEaseToolCommands.StraightenEndTangent);
 				}),
 			false, FSlateIcon(FAppStyle::Get().GetStyleSetName(), TEXT("GenericCurveEditor.StraightenTangents")));
-		MenuBuilder.AddSubMenu(
+
+		Section.AddSubMenu(TEXT("FlattenTangents"),
 			LOCTEXT("FlattenTangentsSubMenuLabel", "Flatten Tangents"),
 			LOCTEXT("FlattenTangentsSubMenuToolTip", ""),
 			FNewMenuDelegate::CreateLambda([&AvaCurveEaseToolCommands](FMenuBuilder& InMenuBuilder)
@@ -609,7 +648,8 @@ TSharedRef<SWidget> SAvaEaseCurveTool::CreateContextMenuContent()
 					InMenuBuilder.AddMenuEntry(AvaCurveEaseToolCommands.FlattenEndTangent);
 				}), 
 			false, FSlateIcon(FAppStyle::Get().GetStyleSetName(), TEXT("GenericCurveEditor.FlattenTangents")));
-		MenuBuilder.AddSubMenu(
+
+		Section.AddSubMenu(TEXT("ResetTangents"),
 			LOCTEXT("ResetTangentsSubMenuLabel", "Reset Tangents"),
 			LOCTEXT("ResetTangentsSubMenuToolTip", ""),
 			FNewMenuDelegate::CreateLambda([&AvaCurveEaseToolCommands](FMenuBuilder& InMenuBuilder)
@@ -619,31 +659,50 @@ TSharedRef<SWidget> SAvaEaseCurveTool::CreateContextMenuContent()
 					InMenuBuilder.AddMenuEntry(AvaCurveEaseToolCommands.ResetEndTangent);
 				}),
 			false, FSlateIcon(FAppStyle::Get().GetStyleSetName(), TEXT("PropertyWindow.DiffersFromDefault")));
-		MenuBuilder.AddSeparator();
+
+		Section.AddSeparator(NAME_None);
+
 		/** @TODO: Only show these in single edit mode(?)
-		MenuBuilder.AddMenuEntry(AvaCurveEaseToolCommands.SetKeyInterpConstant);
-		MenuBuilder.AddMenuEntry(AvaCurveEaseToolCommands.SetKeyInterpLinear);
-		MenuBuilder.AddMenuEntry(AvaCurveEaseToolCommands.SetKeyInterpCubicAuto);
-		MenuBuilder.AddMenuEntry(AvaCurveEaseToolCommands.SetKeyInterpCubicSmartAuto);
-		MenuBuilder.AddMenuEntry(AvaCurveEaseToolCommands.SetKeyInterpCubicUser);
-		MenuBuilder.AddMenuEntry(AvaCurveEaseToolCommands.SetKeyInterpCubicBreak);*/
-		MenuBuilder.AddSeparator();
-		MenuBuilder.AddMenuEntry(AvaCurveEaseToolCommands.SetOperationToEaseOut);
-		MenuBuilder.AddMenuEntry(AvaCurveEaseToolCommands.SetOperationToEaseInOut);
-		MenuBuilder.AddMenuEntry(AvaCurveEaseToolCommands.SetOperationToEaseIn);
-		MenuBuilder.AddSeparator();
-		MenuBuilder.AddMenuEntry(AvaCurveEaseToolCommands.SelectPreviousChannelKey);
-		MenuBuilder.AddMenuEntry(AvaCurveEaseToolCommands.SelectNextChannelKey);
-		MenuBuilder.AddSeparator();
-		MenuBuilder.AddMenuEntry(AvaCurveEaseToolCommands.ToggleGridSnap);
-		MenuBuilder.AddMenuEntry(AvaCurveEaseToolCommands.ZoomToFit);
-		MenuBuilder.AddSeparator();
-		MenuBuilder.AddMenuEntry(AvaCurveEaseToolCommands.Refresh);
-		MenuBuilder.AddMenuEntry(AvaCurveEaseToolCommands.Apply);
+		Section.AddMenuEntry(AvaCurveEaseToolCommands.SetKeyInterpConstant);
+
+		Section.AddMenuEntry(AvaCurveEaseToolCommands.SetKeyInterpLinear);
+
+		Section.AddMenuEntry(AvaCurveEaseToolCommands.SetKeyInterpCubicAuto);
+
+		Section.AddMenuEntry(AvaCurveEaseToolCommands.SetKeyInterpCubicSmartAuto);
+
+		Section.AddMenuEntry(AvaCurveEaseToolCommands.SetKeyInterpCubicUser);
+
+		Section.AddMenuEntry(AvaCurveEaseToolCommands.SetKeyInterpCubicBreak);
+
+		Section.AddSeparator(NAME_None);*/
+
+		Section.AddMenuEntry(AvaCurveEaseToolCommands.SetOperationToEaseOut);
+
+		Section.AddMenuEntry(AvaCurveEaseToolCommands.SetOperationToEaseInOut);
+
+		Section.AddMenuEntry(AvaCurveEaseToolCommands.SetOperationToEaseIn);
+
+		Section.AddSeparator(NAME_None);
+
+		Section.AddMenuEntry(AvaCurveEaseToolCommands.SelectPreviousChannelKey);
+
+		Section.AddMenuEntry(AvaCurveEaseToolCommands.SelectNextChannelKey);
+
+		Section.AddSeparator(NAME_None);
+
+		Section.AddMenuEntry(AvaCurveEaseToolCommands.ToggleGridSnap);
+
+		Section.AddMenuEntry(AvaCurveEaseToolCommands.ZoomToFit);
+
+		Section.AddSeparator(NAME_None);
+
+		Section.AddMenuEntry(AvaCurveEaseToolCommands.Refresh);
+
+		Section.AddMenuEntry(AvaCurveEaseToolCommands.Apply);
 	}
-	MenuBuilder.EndSection();
 
-	return MenuBuilder.MakeWidget();
+	return ToolMenus->GenerateWidget(MenuName, FToolMenuContext(CommandList));
 }
 
 void SAvaEaseCurveTool::MakeContextMenuSettings(FMenuBuilder& InMenuBuilder)
