[Motion Design] Change UAvalancheMaskSettings category and section name to be consistent with all other settings.

diff --git a/Engine/Plugins/Experimental/Avalanche/Source/AvalancheMask/Private/AvaMaskSettings.cpp b/Engine/Plugins/Experimental/Avalanche/Source/AvalancheMask/Private/AvaMaskSettings.cpp
--- a/Engine/Plugins/Experimental/Avalanche/Source/AvalancheMask/Private/AvaMaskSettings.cpp
+++ b/Engine/Plugins/Experimental/Avalanche/Source/AvalancheMask/Private/AvaMaskSettings.cpp
@@ -13,8 +13,8 @@
 
 UAvalancheMaskSettings::UAvalancheMaskSettings()
 {
-    CategoryName = TEXT("Plugins");
-    SectionName = TEXT("AvalancheMask");
+	CategoryName = TEXT("Motion Design");
+	SectionName = TEXT("Mask");
 }
 
 #if WITH_EDITORONLY_DATA
diff --git a/Engine/Plugins/Experimental/Avalanche/Source/AvalancheMask/Private/AvaMaskSettings.h b/Engine/Plugins/Experimental/Avalanche/Source/AvalancheMask/Private/AvaMaskSettings.h
--- a/Engine/Plugins/Experimental/Avalanche/Source/AvalancheMask/Private/AvaMaskSettings.h
+++ b/Engine/Plugins/Experimental/Avalanche/Source/AvalancheMask/Private/AvaMaskSettings.h
@@ -14,7 +14,7 @@ class UMaterialFunctionInterface;
 #endif
 
 /** Settings for Avalanche Mask */
-UCLASS(Config = Engine)
+UCLASS(Config = Engine, meta = (DisplayName = "Mask"))
 class UAvalancheMaskSettings
     : public UDeveloperSettings
 {
