[Motion Design] Checking if SlateApplication is available when initializing Font Manager Subsystem

Font manager subsystem initialization is creating FontFace objects which will end up calling PECP on UFontFace object, thus trying to access FSlateApplication::Get, which might not be available in headless setups.
Adding checks for FSlateApplication::IsInitialized to avoid creating fonts when not necessary.
For safety, also adding the check in other places which might end up indirectly calling UFontFace::PECP.

diff --git a/Engine/Plugins/Experimental/Avalanche/Source/AvalancheTextEditor/Private/Font/AvaFontManagerSubsystem.cpp b/Engine/Plugins/Experimental/Avalanche/Source/AvalancheTextEditor/Private/Font/AvaFontManagerSubsystem.cpp
--- a/Engine/Plugins/Experimental/Avalanche/Source/AvalancheTextEditor/Private/Font/AvaFontManagerSubsystem.cpp
+++ b/Engine/Plugins/Experimental/Avalanche/Source/AvalancheTextEditor/Private/Font/AvaFontManagerSubsystem.cpp
@@ -8,6 +8,7 @@
 #include "Factories/FontFileImportFactory.h"
 #include "Font/AvaFont.h"
 #include "Font/AvaFontObject.h"
+#include "Framework/Application/SlateApplication.h"
 #include "HAL/PlatformFileManager.h"
 #include "Misc/ScopedSlowTask.h"
 #include "PropertyHandle.h"
@@ -35,6 +36,13 @@ UAvaFontManagerSubsystem* UAvaFontManagerSubsystem::Get()
 
 void UAvaFontManagerSubsystem::Initialize(FSubsystemCollectionBase& Collection)
 {
+	// Initialize call below ends up calling PECP on newly created UFontFace objects
+	// which thus trying to access FSlateApplication::Get() which might be invalid for headless setups
+	if (!FSlateApplication::IsInitialized())
+	{
+		return;
+	}
+
 	Initialize();
 
 	const FAssetRegistryModule& AssetRegistryModule = FModuleManager::LoadModuleChecked<FAssetRegistryModule>("AssetRegistry");
@@ -43,7 +51,6 @@ void UAvaFontManagerSubsystem::Initialize(FSubsystemCollectionBase& Collection)
 	// Project fonts will be found, compared to System fonts to avoid duplicates, and then added to the list including both system and project fonts
 	AssetRegistryModule.Get().OnFilesLoaded().AddUObject(this, &UAvaFontManagerSubsystem::InitializeProjectFonts);
 
-	Super::Initialize(Collection);
 }
 
 void UAvaFontManagerSubsystem::Deinitialize()
@@ -188,6 +195,13 @@ void UAvaFontManagerSubsystem::SanitizeString(FString& OutSanitizedName)
 
 bool UAvaFontManagerSubsystem::ImportAsAsset(UAvaFontObject* InFontToImport)
 {
+	// Importing ends up calling PECP on newly created UFontFace objects
+	// which thus trying to access FSlateApplication::Get() which might be invalid for headless setups
+	if (!FSlateApplication::IsInitialized())
+	{
+		return false;
+	}
+
 	if (IsValid(InFontToImport))
 	{
 		if (InFontToImport->GetSource() == EAvaFontSource::System)
@@ -375,6 +389,11 @@ TArray<FString> UAvaFontManagerSubsystem::GetSystemFontTypefaces(const FString&
 
 void UAvaFontManagerSubsystem::SetupFontFamilyTypefaces(UFont* InFont, const FSystemFontsRetrieveParams& InFontParams)
 {
+	if (!FSlateApplication::IsInitialized())
+	{
+		return;
+	}
+
 	const FString& FontFamilyName = InFontParams.FontFamilyName;
 
 	const bool bIsTempAsset = InFont->GetPackage()->GetName().Contains(GetTempFontPackageNameRoot());
@@ -564,7 +583,9 @@ void UAvaFontManagerSubsystem::RefreshAllSystemFonts() const
 {
 	const FAssetRegistryModule& AssetRegistryModule = FModuleManager::LoadModuleChecked<FAssetRegistryModule>("AssetRegistry");
 
-	if (!AssetRegistryModule.IsValid())
+	// Refreshing might end up calling PECP on newly created UFontFace objects
+	// which thus trying to access FSlateApplication::Get() which might be invalid for headless setups
+	if (!AssetRegistryModule.IsValid() || !FSlateApplication::IsInitialized())
 	{
 		return;
 	}
