Motion Design
- Removed references to "Avalanche" in LOCTEXT macros.

diff --git a/Engine/Plugins/Experimental/Avalanche/Source/Avalanche/Private/AvaBlueprint.cpp b/Engine/Plugins/Experimental/Avalanche/Source/Avalanche/Private/AvaBlueprint.cpp
--- a/Engine/Plugins/Experimental/Avalanche/Source/Avalanche/Private/AvaBlueprint.cpp
+++ b/Engine/Plugins/Experimental/Avalanche/Source/Avalanche/Private/AvaBlueprint.cpp
@@ -275,7 +275,7 @@ void UAvalancheBlueprint::SaveAvalancheWorld()
 	TGuardValue SaveGuard(bSavingWorld, true);
 
 #if WITH_EDITOR
-	FScopedSlowTask SaveWorldTask(2.f, LOCTEXT("SaveAvalancheWorld", "Saving Avalanche World"));
+	FScopedSlowTask SaveWorldTask(2.f, LOCTEXT("SaveMotionDesignWorld", "Saving Motion Design World"));
 	SaveWorldTask.MakeDialogDelayed(1.f);	
 #endif
 	
@@ -346,7 +346,7 @@ void UAvalancheBlueprint::LoadAvalancheWorld(UWorld* InOverrideWorld, TOptional<
 	TGuardValue LoadGuard(bLoadingWorld, true);
 	
 #if WITH_EDITOR
-	FScopedSlowTask LoadWorld(WorldData.ActorData.Num(), LOCTEXT("LoadWorld", "Loading Avalanche World"));
+	FScopedSlowTask LoadWorld(WorldData.ActorData.Num(), LOCTEXT("LoadWorld", "Loading Motion Design World"));
 	LoadWorld.MakeDialogDelayed(1.f, false);
 #endif
 	
diff --git a/Engine/Plugins/Experimental/Avalanche/Source/AvalancheEditor/Private/AssetTypeActions/AssetTypeActions_AvaBlueprint.cpp b/Engine/Plugins/Experimental/Avalanche/Source/AvalancheEditor/Private/AssetTypeActions/AssetTypeActions_AvaBlueprint.cpp
--- a/Engine/Plugins/Experimental/Avalanche/Source/AvalancheEditor/Private/AssetTypeActions/AssetTypeActions_AvaBlueprint.cpp
+++ b/Engine/Plugins/Experimental/Avalanche/Source/AvalancheEditor/Private/AssetTypeActions/AssetTypeActions_AvaBlueprint.cpp
@@ -49,12 +49,12 @@ void FAssetTypeActions_AvaBlueprint::PerformAssetDiff(UObject* Asset1, UObject*
 	// times we're comparing two completely separate assets altogether)
 	const bool bIsSingleAsset = (NewBlueprint->GetName() == OldBlueprint->GetName());
 
-	FText WindowTitle = LOCTEXT("NamelessAvalancheBlueprintDiff", "Motion Design Blueprint Diff");
+	FText WindowTitle = LOCTEXT("NamelessMotionDesignBlueprintDiff", "Motion Design Blueprint Diff");
 	// if we're diffing one asset against itself
 	if (bIsSingleAsset)
 	{
 		// identify the assumed single asset in the window's title
-		WindowTitle = FText::Format(LOCTEXT("AvalancheBlueprintDiff", "{0} - Motion Design Blueprint Diff"), FText::FromString(NewBlueprint->GetName()));
+		WindowTitle = FText::Format(LOCTEXT("MotionDesignBlueprintDiff", "{0} - Motion Design Blueprint Diff"), FText::FromString(NewBlueprint->GetName()));
 	}
 
 	SBlueprintDiff::CreateDiffWindow(WindowTitle, OldBlueprint, NewBlueprint, OldRevision, NewRevision);
diff --git a/Engine/Plugins/Experimental/Avalanche/Source/AvalancheEditor/Private/AvaEditorCommands.cpp b/Engine/Plugins/Experimental/Avalanche/Source/AvalancheEditor/Private/AvaEditorCommands.cpp
--- a/Engine/Plugins/Experimental/Avalanche/Source/AvalancheEditor/Private/AvaEditorCommands.cpp
+++ b/Engine/Plugins/Experimental/Avalanche/Source/AvalancheEditor/Private/AvaEditorCommands.cpp
@@ -10,7 +10,7 @@
 
 FAvaEditorCommands::FAvaEditorCommands()
 	: TCommands<FAvaEditorCommands>(TEXT("AvalancheEditor")
-		, LOCTEXT("AvalancheEditor", "Motion Design Editor")
+		, LOCTEXT("MotionDesignEditor", "Motion Design Editor")
 		, NAME_None
 		, FAvaEditorStyle::Get().GetStyleSetName()
 	)
diff --git a/Engine/Plugins/Experimental/Avalanche/Source/AvalancheEditor/Private/AvaEditorModule.cpp b/Engine/Plugins/Experimental/Avalanche/Source/AvalancheEditor/Private/AvaEditorModule.cpp
--- a/Engine/Plugins/Experimental/Avalanche/Source/AvalancheEditor/Private/AvaEditorModule.cpp
+++ b/Engine/Plugins/Experimental/Avalanche/Source/AvalancheEditor/Private/AvaEditorModule.cpp
@@ -183,7 +183,7 @@ void FAvaEditorModule::RegisterAssetTools()
 	IAssetTools& AssetTools = FAssetToolsModule::GetModule().Get();
 
 	AssetTools.RegisterAdvancedAssetCategory(AvalancheCategoryName
-		, LOCTEXT("AvalancheCategoryName", "Motion Design"));
+		, LOCTEXT("MotionDesignCategoryName", "Motion Design"));
 
 	TSharedRef<IAssetTypeActions> AvaBlueprintActions = MakeShared<FAssetTypeActions_AvaBlueprint>();
 	AssetTools.RegisterAssetTypeActions(AvaBlueprintActions);
diff --git a/Engine/Plugins/Experimental/Avalanche/Source/AvalancheEditor/Private/Palettes/AvaPaletteTabSpawner.cpp b/Engine/Plugins/Experimental/Avalanche/Source/AvalancheEditor/Private/Palettes/AvaPaletteTabSpawner.cpp
--- a/Engine/Plugins/Experimental/Avalanche/Source/AvalancheEditor/Private/Palettes/AvaPaletteTabSpawner.cpp
+++ b/Engine/Plugins/Experimental/Avalanche/Source/AvalancheEditor/Private/Palettes/AvaPaletteTabSpawner.cpp
@@ -70,7 +70,7 @@ FTabSpawnerEntry& FAvaPaletteTabSpawner::RegisterTabSpawner(const TSharedRef<FTa
 	{
 		GroupItem = InWorkspaceMenu->AddGroup(
 			PalettesGroupName
-			, LOCTEXT("SubMenuTooltip", "Open an Avalanche palette")
+			, LOCTEXT("SubMenuTooltip", "Open a Motion Design palette")
 			, FSlateIcon(FSlateIcon(FAppStyle::GetAppStyleSetName(), "Icons.FolderClosed"))
 		);
 	}
diff --git a/Engine/Plugins/Experimental/Avalanche/Source/AvalancheEditor/Private/Sequencer/AvaLevelSequencerExtension.cpp b/Engine/Plugins/Experimental/Avalanche/Source/AvalancheEditor/Private/Sequencer/AvaLevelSequencerExtension.cpp
--- a/Engine/Plugins/Experimental/Avalanche/Source/AvalancheEditor/Private/Sequencer/AvaLevelSequencerExtension.cpp
+++ b/Engine/Plugins/Experimental/Avalanche/Source/AvalancheEditor/Private/Sequencer/AvaLevelSequencerExtension.cpp
@@ -277,7 +277,7 @@ void FAvaLevelSequencerExtension::AddCinematicsToolbarExtension()
 	}
 
 	FToolMenuSection& AvaSequenceSection = ToolMenu->AddSection(UE::AvaEditor::Private::AvaSequenceSectionName
-		, LOCTEXT("AvaSequenceSectionLabel", "Edit Avalanche Sequence")
+		, LOCTEXT("AvaSequenceSectionLabel", "Edit Motion Design Sequence")
 		, FToolMenuInsert(UE::AvaEditor::Private::CinematicExtensionName, EToolMenuInsertType::Before));
 
 	AvaSequenceSection.AddDynamicEntry("AvaSequenceEntry", FNewToolMenuSectionDelegate::CreateSPLambda(this,
diff --git a/Engine/Plugins/Experimental/Avalanche/Source/AvalancheEditorCore/Private/LevelEditor/AvaLevelEditorToolbar.cpp b/Engine/Plugins/Experimental/Avalanche/Source/AvalancheEditorCore/Private/LevelEditor/AvaLevelEditorToolbar.cpp
--- a/Engine/Plugins/Experimental/Avalanche/Source/AvalancheEditorCore/Private/LevelEditor/AvaLevelEditorToolbar.cpp
+++ b/Engine/Plugins/Experimental/Avalanche/Source/AvalancheEditorCore/Private/LevelEditor/AvaLevelEditorToolbar.cpp
@@ -44,9 +44,9 @@ namespace UE::AvaEditorCore::Private
 	{
 		const EAppReturnType::Type Reply = FMessageDialog::Open(EAppMsgCategory::Warning
 			, EAppMsgType::YesNo
-			, LOCTEXT("DeactivateSceneConfirmationDesc", "Are you sure you want to deactivate the Avalanche Interface? \n"
+			, LOCTEXT("DeactivateSceneConfirmationDesc", "Are you sure you want to deactivate the Motion Design Interface? \n"
 				"Note: The scene data will be kept.")
-			, LOCTEXT("DeactivateSceneConfirmationTitle", "Deactivate Avalanche Interface"));
+			, LOCTEXT("DeactivateSceneConfirmationTitle", "Deactivate Motion Design Interface"));
 
 		if (Reply != EAppReturnType::Yes)
 		{
@@ -133,7 +133,7 @@ void FAvaLevelEditorToolbar::ExtendLevelEditorToolbar(const TSharedRef<IAvaEdito
 		const FAvaLevelEditorCommands& AvaLevelEditorCommands = FAvaLevelEditorCommands::Get();
 		FToolMenuSection& Section = LevelEditorToolbar->FindOrAddSection(Private::SceneSectionName);
 
-		const FText AvalancheLabel = LOCTEXT("AvalancheScene", "Motion Design");
+		const FText AvalancheLabel = LOCTEXT("MotionDesignScene", "Motion Design");
 
 		auto AddEntry = [this, &Section, AvalancheLabel](const TSharedPtr<FUICommandInfo>& InCommand)
 		{
diff --git a/Engine/Plugins/Experimental/Avalanche/Source/AvalancheEffectorsEditor/Private/AvaEffectorsEditorCommands.cpp b/Engine/Plugins/Experimental/Avalanche/Source/AvalancheEffectorsEditor/Private/AvaEffectorsEditorCommands.cpp
--- a/Engine/Plugins/Experimental/Avalanche/Source/AvalancheEffectorsEditor/Private/AvaEffectorsEditorCommands.cpp
+++ b/Engine/Plugins/Experimental/Avalanche/Source/AvalancheEffectorsEditor/Private/AvaEffectorsEditorCommands.cpp
@@ -8,7 +8,7 @@
 FAvaEffectorsEditorCommands::FAvaEffectorsEditorCommands()
 	: TCommands<FAvaEffectorsEditorCommands>(
 		TEXT("AvalancheEffectorsEditor")
-		, LOCTEXT("AvalancheEffectorsEditor", "Motion Design Effects Editor")
+		, LOCTEXT("MotionDesignEffectorsEditor", "Motion Design Effects Editor")
 		, NAME_None
 		, FAvaEffectorsEditorStyle::Get().GetStyleSetName()
 	)
diff --git a/Engine/Plugins/Experimental/Avalanche/Source/AvalancheInteractiveTools/Private/AvaInteractiveToolsCommands.cpp b/Engine/Plugins/Experimental/Avalanche/Source/AvalancheInteractiveTools/Private/AvaInteractiveToolsCommands.cpp
--- a/Engine/Plugins/Experimental/Avalanche/Source/AvalancheInteractiveTools/Private/AvaInteractiveToolsCommands.cpp
+++ b/Engine/Plugins/Experimental/Avalanche/Source/AvalancheInteractiveTools/Private/AvaInteractiveToolsCommands.cpp
@@ -8,7 +8,7 @@
 FAvaInteractiveToolsCommands::FAvaInteractiveToolsCommands()
 	: TCommands<FAvaInteractiveToolsCommands>(
 		TEXT("AvalancheInteractiveTools")
-		, LOCTEXT("AvalancheInteractiveTools", "Motion Design Interactive Tools")
+		, LOCTEXT("MotionDesignInteractiveTools", "Motion Design Interactive Tools")
 		, NAME_None
 		, FAvaInteractiveToolsStyle::GetStyleSetName()
 	)
diff --git a/Engine/Plugins/Experimental/Avalanche/Source/AvalancheInteractiveTools/Private/Tools/AvaInteractiveToolsToolBase.cpp b/Engine/Plugins/Experimental/Avalanche/Source/AvalancheInteractiveTools/Private/Tools/AvaInteractiveToolsToolBase.cpp
--- a/Engine/Plugins/Experimental/Avalanche/Source/AvalancheInteractiveTools/Private/Tools/AvaInteractiveToolsToolBase.cpp
+++ b/Engine/Plugins/Experimental/Avalanche/Source/AvalancheInteractiveTools/Private/Tools/AvaInteractiveToolsToolBase.cpp
@@ -577,7 +577,7 @@ void UAvaInteractiveToolsToolBase::OnComplete()
 
 void UAvaInteractiveToolsToolBase::BeginTransaction()
 {
-	GetToolManager()->BeginUndoTransaction(LOCTEXT("AvalancheInteractiveToolsTool", "Motion Design Interactive Tools Box Tool"));
+	GetToolManager()->BeginUndoTransaction(LOCTEXT("MotionDesignInteractiveToolsTool", "Motion Design Interactive Tools Box Tool"));
 }
 
 void UAvaInteractiveToolsToolBase::EndTransaction()
diff --git a/Engine/Plugins/Experimental/Avalanche/Source/AvalancheLevelViewport/Private/AvaLevelViewportCommands.cpp b/Engine/Plugins/Experimental/Avalanche/Source/AvalancheLevelViewport/Private/AvaLevelViewportCommands.cpp
--- a/Engine/Plugins/Experimental/Avalanche/Source/AvalancheLevelViewport/Private/AvaLevelViewportCommands.cpp
+++ b/Engine/Plugins/Experimental/Avalanche/Source/AvalancheLevelViewport/Private/AvaLevelViewportCommands.cpp
@@ -8,7 +8,7 @@
 FAvaLevelViewportCommands::FAvaLevelViewportCommands()
 	: TCommands<FAvaLevelViewportCommands>(
 		TEXT("AvalancheLevelViewport")
-		, LOCTEXT("AvalancheLevelViewport", "Motion Design Level Viewport")
+		, LOCTEXT("MotionDesignLevelViewport", "Motion Design Level Viewport")
 		, NAME_None
 		, FAvaLevelViewportStyle::GetStyleSetName()
 	)
diff --git a/Engine/Plugins/Experimental/Avalanche/Source/AvalancheLevelViewport/Private/SAvaLevelViewport_Menus.cpp b/Engine/Plugins/Experimental/Avalanche/Source/AvalancheLevelViewport/Private/SAvaLevelViewport_Menus.cpp
--- a/Engine/Plugins/Experimental/Avalanche/Source/AvalancheLevelViewport/Private/SAvaLevelViewport_Menus.cpp
+++ b/Engine/Plugins/Experimental/Avalanche/Source/AvalancheLevelViewport/Private/SAvaLevelViewport_Menus.cpp
@@ -875,7 +875,7 @@ void SAvaLevelViewport::AddVisualizerEntries(UToolMenu* InMenu)
 				[
 					SNew(SSpinBox<int32>)
 					.Style(&FAppStyle::Get(), "Menu.SpinBox")
-					.ToolTipText(LOCTEXT("VisualizerIconSizeTooltip", "Screen size guide for Avalanche component visualizer icons."))
+					.ToolTipText(LOCTEXT("VisualizerIconSizeTooltip", "Screen size guide for Motion Design component visualizer icons."))
 					.MinValue(1)
 					.MinSliderValue(1)
 					.MaxValue(32)
diff --git a/Engine/Plugins/Experimental/Avalanche/Source/AvalancheMaskEditor/Private/AvaMaskEditorCommands.h b/Engine/Plugins/Experimental/Avalanche/Source/AvalancheMaskEditor/Private/AvaMaskEditorCommands.h
--- a/Engine/Plugins/Experimental/Avalanche/Source/AvalancheMaskEditor/Private/AvaMaskEditorCommands.h
+++ b/Engine/Plugins/Experimental/Avalanche/Source/AvalancheMaskEditor/Private/AvaMaskEditorCommands.h
@@ -1,17 +1,19 @@
-﻿// Copyright Epic Games, Inc. All Rights Reserved.
+// Copyright Epic Games, Inc. All Rights Reserved.
 
 #pragma once
 
 #include "AvaMaskEditorStyle.h"
 #include "Framework/Commands/Commands.h"
 
+#define LOCTEXT_NAMESPACE "AvaMaskEditorCommands"
+
 class FAvaMaskEditorCommands : public TCommands<FAvaMaskEditorCommands>
 {
 public:
 	FAvaMaskEditorCommands()
 		: TCommands<FAvaMaskEditorCommands>(
 			TEXT("AvalancheMaskEditor"),
-			NSLOCTEXT("Contexts", "AvalancheMaskEditor", "Motion Design Masking"),
+			LOCTEXT("MotionDesignMaskEditor", "Motion Design Masking"),
 			NAME_None,
 			FAvalancheMaskEditorStyle::GetStyleSetName())
 	{
@@ -25,3 +27,5 @@ public:
 	TSharedPtr<FUICommandInfo> ToggleIsolateMask;
 	TSharedPtr<FUICommandInfo> ToggleEnableMask;
 };
+
+#undef LOCTEXT_NAMESPACE
diff --git a/Engine/Plugins/Experimental/Avalanche/Source/AvalancheMaskEditor/Private/AvaMaskEditorMode.cpp b/Engine/Plugins/Experimental/Avalanche/Source/AvalancheMaskEditor/Private/AvaMaskEditorMode.cpp
--- a/Engine/Plugins/Experimental/Avalanche/Source/AvalancheMaskEditor/Private/AvaMaskEditorMode.cpp
+++ b/Engine/Plugins/Experimental/Avalanche/Source/AvalancheMaskEditor/Private/AvaMaskEditorMode.cpp
@@ -1,4 +1,4 @@
-﻿// Copyright Epic Games, Inc. All Rights Reserved.
+// Copyright Epic Games, Inc. All Rights Reserved.
 
 #include "AvaMaskEditorMode.h"
 
@@ -34,7 +34,7 @@ const FEditorModeID UAvalancheMaskEditorMode::EM_AvalancheMaskEditorModeId(UE::A
 UAvalancheMaskEditorMode::UAvalancheMaskEditorMode()
 {
 	Info = FEditorModeInfo(UAvalancheMaskEditorMode::EM_AvalancheMaskEditorModeId,
-		LOCTEXT("AvalancheMaskEditorModeName", "AvalancheMaskEditorMode"),
+		LOCTEXT("MotionDesignMaskEditorModeName", "MotionDesign Mask Editor Mode"),
 		FSlateIcon(),
 		false);
 }
diff --git a/Engine/Plugins/Experimental/Avalanche/Source/AvalancheMedia/Public/AvaMediaDevice.h b/Engine/Plugins/Experimental/Avalanche/Source/AvalancheMedia/Public/AvaMediaDevice.h
--- a/Engine/Plugins/Experimental/Avalanche/Source/AvalancheMedia/Public/AvaMediaDevice.h
+++ b/Engine/Plugins/Experimental/Avalanche/Source/AvalancheMedia/Public/AvaMediaDevice.h
@@ -5,6 +5,8 @@
 #include "CoreMinimal.h"
 #include "MediaIOCoreDefinitions.h"
 
+#define LOCTEXT_NAMESPACE "AvaMediaDevice"
+
 struct FAvaMediaDevice
 {
 	FAvaMediaDevice(const FMediaIODevice& InDevice, const FName& InDeviceProviderName, const FString& InServerName)
@@ -28,7 +30,7 @@ struct FAvaMediaDevice
 	
 	FText GetDisplayNameText() const
 	{
-		return FText::Format(NSLOCTEXT("AvalancheMedia", "AvaMediaDeviceDisplayName", "{0}")
+		return FText::Format(LOCTEXT("AvaMediaDeviceDisplayName", "{0}")
 			, FText::FromName(Device.DeviceName));
 	}
 
@@ -54,3 +56,5 @@ protected:
 	/** Name of the server hosting this device. */
 	FString ServerName;
 };
+
+#undef LOCTEXT_NAMESPACE
diff --git a/Engine/Plugins/Experimental/Avalanche/Source/AvalancheMediaEditor/Private/Playback/AssetDefinition_AvalanchePlayback.cpp b/Engine/Plugins/Experimental/Avalanche/Source/AvalancheMediaEditor/Private/Playback/AssetDefinition_AvalanchePlayback.cpp
--- a/Engine/Plugins/Experimental/Avalanche/Source/AvalancheMediaEditor/Private/Playback/AssetDefinition_AvalanchePlayback.cpp
+++ b/Engine/Plugins/Experimental/Avalanche/Source/AvalancheMediaEditor/Private/Playback/AssetDefinition_AvalanchePlayback.cpp
@@ -135,24 +135,24 @@ namespace MenuExtension_AvalanchePlayback
 			Section.AddDynamicEntry(NAME_None, FNewToolMenuSectionDelegate::CreateLambda([](FToolMenuSection& InSection)
 			{
 				{
-					const TAttribute<FText> Label = LOCTEXT("AvalanchePlayback_Play", "Play");
-					const TAttribute<FText> ToolTip = LOCTEXT("AvalanchePlayback_PlayTooltip", "Play");
+					const TAttribute<FText> Label = LOCTEXT("MotionDesignPlayback_Play", "Play");
+					const TAttribute<FText> ToolTip = LOCTEXT("MotionDesignPlayback_PlayTooltip", "Play");
 					const FSlateIcon Icon = FSlateIcon(FAppStyle::GetAppStyleSetName(), "GenericPlay");
 
 					FToolUIAction UIAction;
 					UIAction.ExecuteAction = FToolMenuExecuteAction::CreateStatic(&ExecutePlay);
 					UIAction.CanExecuteAction = FToolMenuCanExecuteAction::CreateStatic(&CanExecutePlay);
-					InSection.AddMenuEntry("AvalanchePlayback_Play", Label, ToolTip, Icon, UIAction);
+					InSection.AddMenuEntry("MotionDesignPlayback_Play", Label, ToolTip, Icon, UIAction);
 				}
 				{
-					const TAttribute<FText> Label = LOCTEXT("AvalanchePlayback_Stop", "Stop");
-					const TAttribute<FText> ToolTip = LOCTEXT("AvalanchePlayback_Stop", "Stop");
+					const TAttribute<FText> Label = LOCTEXT("MotionDesignPlayback_Stop", "Stop");
+					const TAttribute<FText> ToolTip = LOCTEXT("MotionDesignPlayback_Stop", "Stop");
 					const FSlateIcon Icon = FSlateIcon(FAppStyle::GetAppStyleSetName(), "GenericPause");
 
 					FToolUIAction UIAction;
 					UIAction.ExecuteAction = FToolMenuExecuteAction::CreateStatic(&ExecuteStop);
 					UIAction.CanExecuteAction = FToolMenuCanExecuteAction::CreateStatic(&CanExecuteStop);
-					InSection.AddMenuEntry("AvalanchePlayback_Stop", Label, ToolTip, Icon, UIAction);
+					InSection.AddMenuEntry("MotionDesignPlayback_Stop", Label, ToolTip, Icon, UIAction);
 				}
 			}));
 		}));
diff --git a/Engine/Plugins/Experimental/Avalanche/Source/AvalancheOutliner/Private/AvaOutlinerCommands.h b/Engine/Plugins/Experimental/Avalanche/Source/AvalancheOutliner/Private/AvaOutlinerCommands.h
--- a/Engine/Plugins/Experimental/Avalanche/Source/AvalancheOutliner/Private/AvaOutlinerCommands.h
+++ b/Engine/Plugins/Experimental/Avalanche/Source/AvalancheOutliner/Private/AvaOutlinerCommands.h
@@ -5,13 +5,15 @@
 #include "Framework/Commands/Commands.h"
 #include "Styling/AppStyle.h"
 
+#define LOCTEXT_NAMESPACE "AvaOutlinerCommands"
+
 class FAvaOutlinerCommands
 	: public TCommands<FAvaOutlinerCommands>
 {
 public:
 	FAvaOutlinerCommands()
 		: TCommands<FAvaOutlinerCommands>(TEXT("AvalancheOutliner")
-		, NSLOCTEXT("AvaOutlinerCommands", "AvalancheOutliner", "Motion Design Outliner")
+		, LOCTEXT("MotionDesignOutliner", "Motion Design Outliner")
 		, NAME_None
 		, FAppStyle::GetAppStyleSetName())
 	{
@@ -47,3 +49,5 @@ public:
 
 	TSharedPtr<FUICommandInfo> ToggleAutoExpandToSelection;
 };
+
+#undef LOCTEXT_NAMESPACE
diff --git a/Engine/Plugins/Experimental/Avalanche/Source/AvalancheOutliner/Private/Slate/SAvaOutliner.cpp b/Engine/Plugins/Experimental/Avalanche/Source/AvalancheOutliner/Private/Slate/SAvaOutliner.cpp
--- a/Engine/Plugins/Experimental/Avalanche/Source/AvalancheOutliner/Private/Slate/SAvaOutliner.cpp
+++ b/Engine/Plugins/Experimental/Avalanche/Source/AvalancheOutliner/Private/Slate/SAvaOutliner.cpp
@@ -556,7 +556,7 @@ void SAvaOutliner::CreateCustomTextFilterWindow(const FText& InText)
 	CustomTextFilterData.FilterString = InText;
 
 	const TSharedPtr<SWindow> NewTextFilterWindow = SNew(SWindow)
-		.Title(LOCTEXT("CreateCustomTextAvaFilterWindow", "Create Custom Avalanche Filter"))
+		.Title(LOCTEXT("CreateCustomTextAvaFilterWindow", "Create Custom Motion Design Filter"))
 		.HasCloseButton(true)
 		.SupportsMaximize(false)
 		.SupportsMinimize(false)
diff --git a/Engine/Plugins/Experimental/Avalanche/Source/AvalancheShapesEditor/Private/AvaShapesEditorCommands.cpp b/Engine/Plugins/Experimental/Avalanche/Source/AvalancheShapesEditor/Private/AvaShapesEditorCommands.cpp
--- a/Engine/Plugins/Experimental/Avalanche/Source/AvalancheShapesEditor/Private/AvaShapesEditorCommands.cpp
+++ b/Engine/Plugins/Experimental/Avalanche/Source/AvalancheShapesEditor/Private/AvaShapesEditorCommands.cpp
@@ -8,7 +8,7 @@
 FAvaShapesEditorCommands::FAvaShapesEditorCommands()
 	: TCommands<FAvaShapesEditorCommands>(
 		TEXT("AvalancheShapesEditor")
-		, LOCTEXT("AvalancheShapesEditor", "Motion Design Shapes Editor")
+		, LOCTEXT("MotionDesignShapesEditor", "Motion Design Shapes Editor")
 		, NAME_None
 		, FAvaShapesEditorStyle::GetStyleSetName()
 	)
diff --git a/Engine/Plugins/Experimental/Avalanche/Source/AvalancheTextEditor/Private/Font/AvaFontSelectorCommands.h b/Engine/Plugins/Experimental/Avalanche/Source/AvalancheTextEditor/Private/Font/AvaFontSelectorCommands.h
--- a/Engine/Plugins/Experimental/Avalanche/Source/AvalancheTextEditor/Private/Font/AvaFontSelectorCommands.h
+++ b/Engine/Plugins/Experimental/Avalanche/Source/AvalancheTextEditor/Private/Font/AvaFontSelectorCommands.h
@@ -5,13 +5,15 @@
 #include "Framework/Commands/Commands.h"
 #include "Styling/AppStyle.h"
 
+#define LOCTEXT_NAMESPACE "AvaFontSelectorCommands"
+
 class FAvaFontSelectorCommands
 	: public TCommands<FAvaFontSelectorCommands>
 {
 public:
 	FAvaFontSelectorCommands()
 		: TCommands<FAvaFontSelectorCommands>(TEXT("AvalancheFontSelector")
-		, NSLOCTEXT("AvaFontSelectorCommands", "AvalancheFontSelector", "Motion Design Font Selector")
+		, LOCTEXT("MotionDesignFontSelector", "Motion Design Font Selector")
 		, NAME_None
 		, FAppStyle::GetAppStyleSetName())
 	{
@@ -22,4 +24,6 @@ public:
 	TSharedPtr<FUICommandInfo> ShowMonospacedFonts;
 	TSharedPtr<FUICommandInfo> ShowBoldFonts;
 	TSharedPtr<FUICommandInfo> ShowItalicFonts;
-};
\ No newline at end of file
+};
+
+#undef LOCTEXT_NAMESPACE
