Motion Design
- The dynamic shape base component will now remove a material designer asset if switched from material designer mode to asset mode.

diff --git a/Engine/Plugins/Experimental/Avalanche/Source/AvalancheShapes/Private/DynamicMeshes/AvaShapeDynMeshBase.cpp b/Engine/Plugins/Experimental/Avalanche/Source/AvalancheShapes/Private/DynamicMeshes/AvaShapeDynMeshBase.cpp
--- a/Engine/Plugins/Experimental/Avalanche/Source/AvalancheShapes/Private/DynamicMeshes/AvaShapeDynMeshBase.cpp
+++ b/Engine/Plugins/Experimental/Avalanche/Source/AvalancheShapes/Private/DynamicMeshes/AvaShapeDynMeshBase.cpp
@@ -1143,8 +1143,17 @@ void UAvaShapeDynamicMeshBase::OnMaterialTypeChanged(int32 MaterialIndex)
 		switch (MeshData->MaterialType)
 		{
 			case(EMaterialType::Asset):
-				SetMaterialDirect(MaterialIndex, GetMaterial(MaterialIndex));
+			{
+				UMaterialInterface* CurrentMaterial = GetMaterial(MaterialIndex);
+
+				if (!IsValid(CurrentMaterial) || CurrentMaterial->IsA<UDynamicMaterialInstance>())
+				{
+					CurrentMaterial = nullptr;
+				}
+
+				SetMaterialDirect(MaterialIndex, CurrentMaterial);
 				break;
+			}
 
 			case(EMaterialType::Parametric):
 				SetMaterial(MaterialIndex, nullptr);
