MotionDesign
- Cloner fix crash due to event related to attached actors attachment update

diff --git a/Engine/Plugins/Experimental/Avalanche/Source/AvalancheEffectors/Private/Cloner/AvaClonerComponent.cpp b/Engine/Plugins/Experimental/Avalanche/Source/AvalancheEffectors/Private/Cloner/AvaClonerComponent.cpp
--- a/Engine/Plugins/Experimental/Avalanche/Source/AvalancheEffectors/Private/Cloner/AvaClonerComponent.cpp
+++ b/Engine/Plugins/Experimental/Avalanche/Source/AvalancheEffectors/Private/Cloner/AvaClonerComponent.cpp
@@ -65,8 +65,6 @@ UAvaClonerComponent::UAvaClonerComponent()
 		FCoreUObjectDelegates::OnObjectPropertyChanged.RemoveAll(this);
 		FCoreUObjectDelegates::OnObjectPropertyChanged.AddUObject(this, &UAvaClonerComponent::OnActorPropertyChanged);
 #endif
-
-		USceneComponent::MarkRenderStateDirtyEvent.AddUObject(this, &UAvaClonerComponent::OnRenderStateDirty);
 	}
 }
 
