Motion Design
- Added confirmation dialogs for guide presets when: overwriting with save as, replacing with the save button and removing..

diff --git a/Engine/Plugins/Experimental/Avalanche/Source/AvalancheLevelViewport/Private/SAvaLevelViewport_Menus.cpp b/Engine/Plugins/Experimental/Avalanche/Source/AvalancheLevelViewport/Private/SAvaLevelViewport_Menus.cpp
--- a/Engine/Plugins/Experimental/Avalanche/Source/AvalancheLevelViewport/Private/SAvaLevelViewport_Menus.cpp
+++ b/Engine/Plugins/Experimental/Avalanche/Source/AvalancheLevelViewport/Private/SAvaLevelViewport_Menus.cpp
@@ -5,14 +5,16 @@
 #include "AvaLevelViewportStyle.h"
 #include "AvaLevelViewportToolbarContext.h"
 #include "AvaViewportDataSubsystem.h"
+#include "AvaViewportGuideInfo.h"
 #include "AvaViewportVirtualSizeEnums.h"
 #include "Camera/CameraActor.h"
-#include "EngineUtils.h"
 #include "Engine/Texture.h"
+#include "EngineUtils.h"
 #include "IAvaComponentVisualizersSettings.h"
 #include "IAvalancheComponentVisualizersModule.h"
 #include "LevelEditor.h"
 #include "LevelEditor/AvaLevelEditorUtils.h"
+#include "Misc/MessageDialog.h"
 #include "PropertyCustomizationHelpers.h"
 #include "ScopedTransaction.h"
 #include "SEditorViewportToolBarButton.h"
@@ -1317,6 +1319,24 @@ void SAvaLevelViewport::ExecuteSaveAsGuidePreset(const FToolMenuContext& InConte
 		return;
 	}
 
+	TArray<FAvaViewportGuideInfo> GuidesTemp;
+
+	if (DataSubsystem->GetGuidePresetProvider().LoadGuidePreset(PresetName, GuidesTemp, FVector2f(1.f, 1.f)))
+	{
+		static const FText ReplaceGuidePresetFormat = LOCTEXT("AlreadyExistsGuidePresetFormat", "Guide preset already exists. Are you sure you want to replace:\n\n{0}");
+
+		const EAppReturnType::Type Response = FMessageDialog::Open(
+			EAppMsgType::OkCancel,
+			FText::Format(ReplaceGuidePresetFormat, FText::FromString(PresetName)),
+			LOCTEXT("ConfirmReplace", "Confirm Replace")
+		);
+
+		if (Response != EAppReturnType::Ok)
+		{
+			return;
+		}
+	}
+
 	DataSubsystem->GetGuidePresetProvider().SaveGuidePreset(PresetName, ViewportData->GuideData, GetVirtualSize());
 }
 
@@ -1457,12 +1477,23 @@ FReply SAvaLevelViewport::ExecuteReplaceGuidePreset(FString InPresetName)
 		return FReply::Handled();
 	}
 
-	DataSubsystem->GetGuidePresetProvider().SaveGuidePreset(
-		InPresetName,
-		ViewportData->GuideData,
-		GetVirtualSize()
+	static const FText ReplaceGuidePresetFormat = LOCTEXT("ReplaceGuidePresetFormat", "Are you sure you want to replace guide preset:\n\n{0}");
+
+	const EAppReturnType::Type Response = FMessageDialog::Open(
+		EAppMsgType::OkCancel,
+		FText::Format(ReplaceGuidePresetFormat, FText::FromString(InPresetName)),
+		LOCTEXT("ConfirmReplace", "Confirm Replace")
 	);
 
+	if (Response == EAppReturnType::Ok)
+	{
+		DataSubsystem->GetGuidePresetProvider().SaveGuidePreset(
+			InPresetName,
+			ViewportData->GuideData,
+			GetVirtualSize()
+		);
+	}
+
 	return FReply::Handled();
 }
 
@@ -1477,7 +1508,18 @@ FReply SAvaLevelViewport::ExecuteRemoveGuidePreset(FString InPresetName)
 		return FReply::Handled();
 	}
 
-	DataSubsystem->GetGuidePresetProvider().RemoveGuidePreset(InPresetName);
+	static const FText RemoveGuidePresetFormat = LOCTEXT("RemoveGuidePresetFormat", "Are you sure you want to remove guide preset:\n\n{0}");
+
+	const EAppReturnType::Type Response = FMessageDialog::Open(
+		EAppMsgType::OkCancel,
+		FText::Format(RemoveGuidePresetFormat, FText::FromString(InPresetName)),
+		LOCTEXT("ConfirmRemoval", "Confirm Removal")
+	);
+
+	if (Response == EAppReturnType::Ok)
+	{
+		DataSubsystem->GetGuidePresetProvider().RemoveGuidePreset(InPresetName);
+	}
 
 	return FReply::Handled();
 }
