Avalanche Transition:
- Fix issue with Expanded State always set to true

diff --git a/Engine/Plugins/Experimental/Avalanche/Source/AvalancheTransitionEditor/Private/Extensions/IAvaTransitionTreeRowExtension.h b/Engine/Plugins/Experimental/Avalanche/Source/AvalancheTransitionEditor/Private/Extensions/IAvaTransitionTreeRowExtension.h
--- a/Engine/Plugins/Experimental/Avalanche/Source/AvalancheTransitionEditor/Private/Extensions/IAvaTransitionTreeRowExtension.h
+++ b/Engine/Plugins/Experimental/Avalanche/Source/AvalancheTransitionEditor/Private/Extensions/IAvaTransitionTreeRowExtension.h
@@ -26,5 +26,5 @@ public:
 	virtual bool IsExpanded() const = 0;
 
 	/** Sets the Expansion state of the View Model */
-	virtual void SetExpanded(bool bIsExpanded) = 0;
+	virtual void SetExpanded(bool bInIsExpanded) = 0;
 };
diff --git a/Engine/Plugins/Experimental/Avalanche/Source/AvalancheTransitionEditor/Private/ViewModels/State/AvaTransitionStateViewModel.cpp b/Engine/Plugins/Experimental/Avalanche/Source/AvalancheTransitionEditor/Private/ViewModels/State/AvaTransitionStateViewModel.cpp
--- a/Engine/Plugins/Experimental/Avalanche/Source/AvalancheTransitionEditor/Private/ViewModels/State/AvaTransitionStateViewModel.cpp
+++ b/Engine/Plugins/Experimental/Avalanche/Source/AvalancheTransitionEditor/Private/ViewModels/State/AvaTransitionStateViewModel.cpp
@@ -335,15 +335,17 @@ bool FAvaTransitionStateViewModel::IsExpanded() const
 	return State && State->bExpanded;
 }
 
-void FAvaTransitionStateViewModel::SetExpanded(bool bIsExpanded)
+void FAvaTransitionStateViewModel::SetExpanded(bool bInIsExpanded)
 {
-	if (UStateTreeState* State = GetState())
+	UStateTreeState* State = GetState();
+
+	if (State && State->bExpanded != bInIsExpanded)
 	{
 		if (GUndo)
 		{
 			State->Modify();
 		}
-		State->bExpanded = true;
+		State->bExpanded = bInIsExpanded;
 	}
 }
 
diff --git a/Engine/Plugins/Experimental/Avalanche/Source/AvalancheTransitionEditor/Private/ViewModels/State/AvaTransitionStateViewModel.h b/Engine/Plugins/Experimental/Avalanche/Source/AvalancheTransitionEditor/Private/ViewModels/State/AvaTransitionStateViewModel.h
--- a/Engine/Plugins/Experimental/Avalanche/Source/AvalancheTransitionEditor/Private/ViewModels/State/AvaTransitionStateViewModel.h
+++ b/Engine/Plugins/Experimental/Avalanche/Source/AvalancheTransitionEditor/Private/ViewModels/State/AvaTransitionStateViewModel.h
@@ -146,7 +146,7 @@ public:
 	virtual bool CanGenerateRow() const override;
 	virtual TSharedRef<ITableRow> GenerateRow(const TSharedRef<STableViewBase>& InOwningTableView) override;
 	virtual bool IsExpanded() const override;
-	virtual void SetExpanded(bool bIsExpanded) override;
+	virtual void SetExpanded(bool bInIsExpanded) override;
 	//~ End IAvaTransitionTreeRowExtension
 
 	//~ Begin IAvaTransitionObjectExtension
