Motion Designer
- Changed default keybind for toggling guides to ctrl+o

diff --git a/Engine/Plugins/Experimental/Avalanche/Source/AvalancheLevelViewport/Private/AvaLevelViewportCommands.cpp b/Engine/Plugins/Experimental/Avalanche/Source/AvalancheLevelViewport/Private/AvaLevelViewportCommands.cpp
--- a/Engine/Plugins/Experimental/Avalanche/Source/AvalancheLevelViewport/Private/AvaLevelViewportCommands.cpp
+++ b/Engine/Plugins/Experimental/Avalanche/Source/AvalancheLevelViewport/Private/AvaLevelViewportCommands.cpp
@@ -270,7 +270,7 @@ void FAvaLevelViewportCommands::RegisterGuideCommands()
 		, "Toggle Guides"
 		, "Toggles the overlay visibility."
 		, EUserInterfaceActionType::ToggleButton
-		, FInputChord(EKeys::O, EModifierKey::Shift));
+		, FInputChord(EKeys::O, EModifierKey::Control));
 
 	UI_COMMAND(AddGuideHorizontal
 		, "Add Horizontal Guide"
