Motion Designer
- Renamed AvalancheActor and associated structs to use the Ava prefix instead of Avalanche.
- Renamed AvalancheArchive and associated serialization classes to use the Ava prefix instead of Avalanche.

diff --git a/Engine/Plugins/Experimental/Avalanche/Config/DefaultAvalanche.ini b/Engine/Plugins/Experimental/Avalanche/Config/DefaultAvalanche.ini
--- a/Engine/Plugins/Experimental/Avalanche/Config/DefaultAvalanche.ini
+++ b/Engine/Plugins/Experimental/Avalanche/Config/DefaultAvalanche.ini
@@ -149,5 +149,16 @@
 +PropertyRedirects=(OldName="/Script/AvalancheEffectors.AvaClonerActor.StepScaleDelta",NewName="/Script/AvalancheEffectors.AvaClonerActor.DeltaStepScale")
 +PropertyRedirects=(OldName="/Script/AvalancheEffectors.AvaEffectorActor.Strength",NewName="/Script/AvalancheEffectors.AvaEffectorActor.LocationStrength")
 
++ClassRedirects=(OldName="/Script/Avalanche.AvalancheActor",NewName="/Script/Avalanche.AvaActor")
++StructRedirects=(OldName="/Script/Avalanche.AvalancheFileVersionInfo",NewName="/Script/Avalanche.AvaFileVersionInfo")
++StructRedirects=(OldName="/Script/Avalanche.AvalancheEngineVersionInfo",NewName="/Script/Avalanche.AvaEngineVersionInfo")
++StructRedirects=(OldName="/Script/Avalanche.AvalancheCustomVersionInfo",NewName="/Script/Avalanche.AvaCustomVersionInfo")
++StructRedirects=(OldName="/Script/Avalanche.AvalancheVersionInfo",NewName="/Script/Avalanche.AvaVersionInfo")
++StructRedirects=(OldName="/Script/Avalanche.AvalancheWorldData",NewName="/Script/Avalanche.AvaWorldData")
++StructRedirects=(OldName="/Script/Avalanche.AvalancheSubObjectData",NewName="/Script/Avalanche.AvaSubObjectData")
++StructRedirects=(OldName="/Script/Avalanche.AvalancheObjectData",NewName="/Script/Avalanche.AvaObjectData")
++StructRedirects=(OldName="/Script/Avalanche.AvalancheComponentData",NewName="/Script/Avalanche.AvaComponentData")
++StructRedirects=(OldName="/Script/Avalanche.AvalancheActorData",NewName="/Script/Avalanche.AvaActorData")
+
 [/Script/Avalanche.AvaAnimation]
 DefaultCompletionMode=KeepState
diff --git a/Engine/Plugins/Experimental/Avalanche/Source/Avalanche/Private/Archive/AvaArchive.cpp b/Engine/Plugins/Experimental/Avalanche/Source/Avalanche/Private/Archive/AvaArchive.cpp
new file mode 100644
--- /dev/null
+++ b/Engine/Plugins/Experimental/Avalanche/Source/Avalanche/Private/Archive/AvaArchive.cpp
@@ -0,0 +1,175 @@
+// Copyright Epic Games, Inc. All Rights Reserved.
+
+#include "AvaArchive.h"
+#include "AvaBlueprint.h"
+#include "AvaBlueprint_Serialize.h"
+#include "HAL/UnrealMemory.h"
+#include "UObject/PropertyPortFlags.h"
+#include "UObject/UnrealType.h"
+
+#define LOCTEXT_NAMESPACE "AvaArchive"
+
+DEFINE_LOG_CATEGORY_STATIC(LogAvaArchive, Log, All);
+
+FAvaArchive::FAvaArchive(FAvaWorldData& InWorldData, FAvaObjectData& InObjectData, UObject* InSerializedObject, bool bIsLoading)
+	: WorldData(InWorldData)
+	, ObjectData(InObjectData)
+	, SerializedObject(InSerializedObject)
+	, ExcludedPropertyFlags(CPF_BlueprintAssignable | CPF_Transient | CPF_Deprecated)
+{
+	Super::SetWantBinaryPropertySerialization(false);
+	Super::SetIsTransacting(false);
+	Super::SetIsPersistent(true);
+
+	ArNoDelta = true;
+
+	if (bIsLoading)
+	{
+		// Serialize properties that were valid in a previous version and are deprecated now. PostSerialize is responsible to migrate the data.
+		ExcludedPropertyFlags &= ~CPF_Deprecated;
+		SetPortFlags(PPF_UseDeprecatedProperties);
+
+		Super::SetIsLoading(true);
+		Super::SetIsSaving(false);
+
+		WorldData.VersionInfo.ApplyToArchive(*this);
+	}
+	else
+	{
+		Super::SetIsLoading(false);
+		Super::SetIsSaving(true);
+	}
+}
+
+FString FAvaArchive::GetArchiveName() const
+{
+	return TEXT("FAvaArchive");
+}
+
+int64 FAvaArchive::TotalSize()
+{
+	return ObjectData.SerializedData.Num();
+}
+
+int64 FAvaArchive::Tell()
+{
+	return DataIndex;
+}
+
+void FAvaArchive::Seek(int64 InPos)
+{
+	checkSlow(InPos <= TotalSize());
+	DataIndex = InPos;
+}
+
+bool FAvaArchive::ShouldSkipProperty(const FProperty* InProperty) const
+{
+	check(InProperty);
+
+	if (InProperty->HasAnyPropertyFlags(ExcludedPropertyFlags))
+	{
+		return true;
+	}
+
+	return false;
+}
+
+FArchive& FAvaArchive::operator<<(FName& Value)
+{
+	if (IsLoading())
+	{
+		FAvaObjectIndex NameIndex;
+		*this << NameIndex.Index;
+
+		if (!ensureAlwaysMsgf(WorldData.SerializedNames.IsValidIndex(NameIndex.Index), TEXT("Data appears to be corrupted")))
+		{
+			SetError();
+			return *this;
+		}
+		Value = WorldData.SerializedNames[NameIndex.Index];
+	}
+	else
+	{
+		FAvaObjectIndex NameIndex;
+		if (const FAvaObjectIndex* ExistingIndex = WorldData.NameIndexMap.Find(Value))
+		{
+			NameIndex = *ExistingIndex;
+		}
+		else
+		{
+			NameIndex.Index = WorldData.SerializedNames.Add(Value);
+			WorldData.NameIndexMap.Add(Value, NameIndex);
+		}
+
+		*this << NameIndex.Index;
+	}
+	return *this;
+}
+
+FArchive& FAvaArchive::operator<<(UObject*& Value)
+{
+	if (IsLoading())
+	{
+		FAvaObjectIndex ReferencedIndex;
+		*this << ReferencedIndex.Index;
+
+		if (!ensureAlwaysMsgf(WorldData.SerializedObjectReferences.IsValidIndex(ReferencedIndex.Index)
+			, TEXT("Data appears to be corrupted")))
+		{
+			SetError();
+			return *this;
+		}
+
+		const FSoftObjectPath& ObjectPath = WorldData.SerializedObjectReferences[ReferencedIndex.Index];
+		if (ObjectPath.IsNull())
+		{
+			Value = nullptr;
+			return *this;
+		}
+
+		Value = ResolveObjectDependency(ReferencedIndex);
+	}
+	else
+	{
+		FAvaObjectIndex ReferenceIndex = FAvaBlueprint_Serialize::AddObjectDependency(WorldData, Value);
+		*this << ReferenceIndex.Index;
+	}
+
+	return *this;
+}
+
+void FAvaArchive::Serialize(void* Data, int64 Length)
+{
+	if (Length <= 0)
+	{
+		return;
+	}
+
+	if (IsLoading())
+	{
+		if (!ensure(DataIndex + Length <= TotalSize()))
+		{
+			UE_LOG(LogAvaArchive, Error, TEXT("Unable to read %d bytes at index %d (Archive size: %d), missing %d bytes."),
+				Length, DataIndex, TotalSize(), DataIndex + Length - TotalSize());
+			SetError();
+			return;
+		}
+
+		FMemory::Memcpy(Data, &ObjectData.SerializedData[DataIndex], Length);
+		DataIndex += Length;
+	}
+	else
+	{
+		const int64 RequiredEndIndex = DataIndex + Length;
+		const int32 ToAlloc  = RequiredEndIndex - TotalSize();
+		if (ToAlloc > 0)
+		{
+			ObjectData.SerializedData.AddUninitialized(ToAlloc);
+		}
+
+		FMemory::Memcpy(&ObjectData.SerializedData[DataIndex], Data, Length);
+		DataIndex = RequiredEndIndex;
+	}
+}
+
+#undef LOCTEXT_NAMESPACE
diff --git a/Engine/Plugins/Experimental/Avalanche/Source/Avalanche/Private/Archive/AvaArchive.h b/Engine/Plugins/Experimental/Avalanche/Source/Avalanche/Private/Archive/AvaArchive.h
new file mode 100644
--- /dev/null
+++ b/Engine/Plugins/Experimental/Avalanche/Source/Avalanche/Private/Archive/AvaArchive.h
@@ -0,0 +1,50 @@
+// Copyright Epic Games, Inc. All Rights Reserved.
+
+#pragma once
+
+#include "Data/AvaDataDefines.h"
+#include "Serialization/ArchiveUObject.h"
+#include "UObject/ObjectMacros.h"
+
+struct FAvaObjectData;
+struct FAvaWorldData;
+class UObject;
+
+class FAvaArchive : public FArchiveUObject
+{
+	using Super = FArchiveUObject;
+	
+public:
+	
+	FAvaArchive(FAvaWorldData& InWorldData
+		, FAvaObjectData& InObjectData
+		, UObject* InSerializedObject
+		, bool bIsLoading);
+
+	// FArchive Interface
+	virtual FString GetArchiveName() const override;
+	virtual int64 TotalSize() override;
+	virtual int64 Tell() override;
+	virtual void Seek(int64 InPos) override;
+	virtual bool ShouldSkipProperty(const FProperty* InProperty) const override;
+	virtual FArchive& operator<<(FName& Value) override;
+	virtual FArchive& operator<<(UObject*& Value) override;
+	virtual void Serialize(void* Data, int64 Length) override;
+	// ~FArchive Interface
+
+	/* Allocates and serializes an object dependency, or gets the object, if it already exists. */
+	virtual UObject* ResolveObjectDependency(FAvaObjectIndex ObjectIndex) const = 0;
+	
+protected:
+
+	FAvaWorldData& WorldData;
+
+	FAvaObjectData& ObjectData;
+
+	UObject* SerializedObject;
+
+	/** Flags that would cause the Property to be excluded from the Archive */
+	EPropertyFlags ExcludedPropertyFlags = EPropertyFlags::CPF_None;
+
+	int64 DataIndex = 0;
+};
diff --git a/Engine/Plugins/Experimental/Avalanche/Source/Avalanche/Private/Archive/AvaReader.cpp b/Engine/Plugins/Experimental/Avalanche/Source/Avalanche/Private/Archive/AvaReader.cpp
new file mode 100644
--- /dev/null
+++ b/Engine/Plugins/Experimental/Avalanche/Source/Avalanche/Private/Archive/AvaReader.cpp
@@ -0,0 +1,15 @@
+// Copyright Epic Games, Inc. All Rights Reserved.
+
+#include "AvaReader.h"
+#include "AvaBlueprint_Serialize.h"
+
+FAvaReader::FAvaReader(FAvaWorldData& InWorldData, FAvaObjectData& InObjectData, UObject* InObject)
+	: Super(InWorldData, InObjectData, InObject, true)
+{
+	InObject->Serialize(*this);
+}
+
+UObject* FAvaReader::ResolveObjectDependency(FAvaObjectIndex ObjectIndex) const
+{
+	return FAvaBlueprint_Serialize::ResolveObjectDependency(WorldData, ObjectIndex);
+}
diff --git a/Engine/Plugins/Experimental/Avalanche/Source/Avalanche/Private/Archive/AvaReader.h b/Engine/Plugins/Experimental/Avalanche/Source/Avalanche/Private/Archive/AvaReader.h
new file mode 100644
--- /dev/null
+++ b/Engine/Plugins/Experimental/Avalanche/Source/Avalanche/Private/Archive/AvaReader.h
@@ -0,0 +1,18 @@
+// Copyright Epic Games, Inc. All Rights Reserved.
+
+#pragma once
+
+#include "AvaArchive.h"
+
+class FAvaReader : public FAvaArchive
+{
+	using Super = FAvaArchive;
+	
+public:
+	
+	FAvaReader(FAvaWorldData& InWorldData, FAvaObjectData& InObjectData, UObject* InObject);
+
+	// FAvaArchive Interface
+	virtual UObject* ResolveObjectDependency(FAvaObjectIndex ObjectIndex) const override;
+	// ~FAvaArchive Interface
+};
diff --git a/Engine/Plugins/Experimental/Avalanche/Source/Avalanche/Private/Archive/AvaWriter.cpp b/Engine/Plugins/Experimental/Avalanche/Source/Avalanche/Private/Archive/AvaWriter.cpp
new file mode 100644
--- /dev/null
+++ b/Engine/Plugins/Experimental/Avalanche/Source/Avalanche/Private/Archive/AvaWriter.cpp
@@ -0,0 +1,11 @@
+// Copyright Epic Games, Inc. All Rights Reserved.
+
+#include "AvaWriter.h"
+#include "Data/AvaObjectData.h"
+
+FAvaWriter::FAvaWriter(FAvaWorldData& InWorldData, FAvaObjectData& InObjectData, UObject* InObject)
+	: Super(InWorldData, InObjectData, InObject, false)
+{
+	InObject->Serialize(*this);
+	InObjectData.ObjectFlags = InObject->GetFlags();
+}
diff --git a/Engine/Plugins/Experimental/Avalanche/Source/Avalanche/Private/Archive/AvaWriter.h b/Engine/Plugins/Experimental/Avalanche/Source/Avalanche/Private/Archive/AvaWriter.h
new file mode 100644
--- /dev/null
+++ b/Engine/Plugins/Experimental/Avalanche/Source/Avalanche/Private/Archive/AvaWriter.h
@@ -0,0 +1,23 @@
+// Copyright Epic Games, Inc. All Rights Reserved.
+
+#pragma once
+
+#include "AvaArchive.h"
+#include "Misc/AssertionMacros.h"
+
+class FAvaWriter : public FAvaArchive
+{
+	using Super = FAvaArchive;
+
+public:
+	
+	FAvaWriter(FAvaWorldData& InWorldData, FAvaObjectData& InObjectData, UObject* InObject);
+
+	// FAvaArchive Interface
+	virtual UObject* ResolveObjectDependency(FAvaObjectIndex ObjectIndex) const override
+	{
+		checkNoEntry();
+		return nullptr;
+	}
+	// ~FAvaArchive Interface
+};
diff --git a/Engine/Plugins/Experimental/Avalanche/Source/Avalanche/Private/Archive/AvalancheArchive.cpp b/Engine/Plugins/Experimental/Avalanche/Source/Avalanche/Private/Archive/AvalancheArchive.cpp
deleted file mode 100644
diff --git a/Engine/Plugins/Experimental/Avalanche/Source/Avalanche/Private/Archive/AvalancheArchive.h b/Engine/Plugins/Experimental/Avalanche/Source/Avalanche/Private/Archive/AvalancheArchive.h
deleted file mode 100644
diff --git a/Engine/Plugins/Experimental/Avalanche/Source/Avalanche/Private/Archive/AvalancheReader.cpp b/Engine/Plugins/Experimental/Avalanche/Source/Avalanche/Private/Archive/AvalancheReader.cpp
deleted file mode 100644
diff --git a/Engine/Plugins/Experimental/Avalanche/Source/Avalanche/Private/Archive/AvalancheReader.h b/Engine/Plugins/Experimental/Avalanche/Source/Avalanche/Private/Archive/AvalancheReader.h
deleted file mode 100644
diff --git a/Engine/Plugins/Experimental/Avalanche/Source/Avalanche/Private/Archive/AvalancheWriter.cpp b/Engine/Plugins/Experimental/Avalanche/Source/Avalanche/Private/Archive/AvalancheWriter.cpp
deleted file mode 100644
diff --git a/Engine/Plugins/Experimental/Avalanche/Source/Avalanche/Private/Archive/AvalancheWriter.h b/Engine/Plugins/Experimental/Avalanche/Source/Avalanche/Private/Archive/AvalancheWriter.h
deleted file mode 100644
diff --git a/Engine/Plugins/Experimental/Avalanche/Source/Avalanche/Private/AvaActor.cpp b/Engine/Plugins/Experimental/Avalanche/Source/Avalanche/Private/AvaActor.cpp
--- a/Engine/Plugins/Experimental/Avalanche/Source/Avalanche/Private/AvaActor.cpp
+++ b/Engine/Plugins/Experimental/Avalanche/Source/Avalanche/Private/AvaActor.cpp
@@ -4,11 +4,11 @@
 #include "AvaBlueprintGeneratedClass.h"
 #include "AvaSequencePlayer.h"
 
-AAvalancheActor::AAvalancheActor()
+AAvaActor::AAvaActor()
 {
 }
 
-bool AAvalancheActor::UpdateGeneratedClass(bool bForceUpdate)
+bool AAvaActor::UpdateGeneratedClass(bool bForceUpdate)
 {
 	if ((bGeneratedClassUpdated && !bForceUpdate) || HasAnyFlags(RF_ClassDefaultObject))
 	{
@@ -25,13 +25,20 @@ bool AAvalancheActor::UpdateGeneratedClass(bool bForceUpdate)
 	return true;
 }
 
-void AAvalancheActor::PostActorCreated()
+void AAvaActor::PostActorCreated()
 {
 	Super::PostActorCreated();
 	UpdateGeneratedClass();
 }
 
-void AAvalancheActor::PostDuplicate(EDuplicateMode::Type DuplicateMode)
+#if WITH_EDITOR
+FString AAvaActor::GetDefaultActorLabel() const
+{
+	return TEXT("Motion Design Actor");
+}
+#endif
+
+void AAvaActor::PostDuplicate(EDuplicateMode::Type DuplicateMode)
 {
 	Super::PostDuplicate(DuplicateMode);
 	UpdateGeneratedClass();
diff --git a/Engine/Plugins/Experimental/Avalanche/Source/Avalanche/Private/AvaBlueprint.cpp b/Engine/Plugins/Experimental/Avalanche/Source/Avalanche/Private/AvaBlueprint.cpp
--- a/Engine/Plugins/Experimental/Avalanche/Source/Avalanche/Private/AvaBlueprint.cpp
+++ b/Engine/Plugins/Experimental/Avalanche/Source/Avalanche/Private/AvaBlueprint.cpp
@@ -81,7 +81,7 @@ UWorld* UAvalancheBlueprint::GetAvalancheWorld() const
 	return World;
 }
 
-AAvalancheActor* UAvalancheBlueprint::GetPlaceholderActor() const
+AAvaActor* UAvalancheBlueprint::GetPlaceholderActor() const
 {
 #if WITH_EDITOR
 	// If Placeholder Actor is Invalid, or Worlds do not match, try to find the Placeholder Actor
@@ -90,7 +90,7 @@ AAvalancheActor* UAvalancheBlueprint::GetPlaceholderActor() const
 		UAvalancheBlueprint* const MutableThis = const_cast<UAvalancheBlueprint*>(this);
 		MutableThis->PlaceholderActor = nullptr;
 
-		for (AAvalancheActor* const Actor : TActorRange<AAvalancheActor>(World))
+		for (AAvaActor* const Actor : TActorRange<AAvaActor>(World))
 		{
 			if (Actor->GetClass()->ClassGeneratedBy == this)
 			{
@@ -152,10 +152,10 @@ void UAvalancheBlueprint::DestroyPlaceholderActor()
 
 void UAvalancheBlueprint::UpdatePlaceholderActor(bool bInForceFullUpdate)
 {
-	AAvalancheActor* AvalancheActor = GetPlaceholderActor();
+	AAvaActor* AvaActor = GetPlaceholderActor();
 
 	// if Actor is already valid and we are not doing a full update, skip
-	if (AvalancheActor && !bInForceFullUpdate)
+	if (AvaActor && !bInForceFullUpdate)
 	{
 		return;
 	}
@@ -169,13 +169,13 @@ void UAvalancheBlueprint::UpdatePlaceholderActor(bool bInForceFullUpdate)
 #endif
 
 	// Spawn Placeholder Actor if current is invalid, or we're forcing full update
-	if (!AvalancheActor || bInForceFullUpdate)
+	if (!AvaActor || bInForceFullUpdate)
 	{
 		// Destroy the previous actor instance
 		DestroyPlaceholderActor();
 
 		// Spawn a new placeholder actor based on the Blueprint's generated class if it's Actor-based
-		if (ensureMsgf(GeneratedClass && GeneratedClass->IsChildOf(AAvalancheActor::StaticClass())
+		if (ensureMsgf(GeneratedClass && GeneratedClass->IsChildOf(AAvaActor::StaticClass())
 			, TEXT("Generated Class (%s) is not based on Avalanche Actor."), GeneratedClass ? *GeneratedClass->GetName() : TEXT("invalid")))
 		{
 			// Spawn an Actor based on the Blueprint's generated class
@@ -189,15 +189,15 @@ void UAvalancheBlueprint::UpdatePlaceholderActor(bool bInForceFullUpdate)
 #if WITH_EDITOR
 				FMakeClassSpawnableOnScope TemporarilySpawnable(GeneratedClass);
 #endif
-				AvalancheActor   = World->SpawnActor<AAvalancheActor>(GeneratedClass, FVector::ZeroVector, FRotator::ZeroRotator, SpawnInfo);
-				PlaceholderActor = AvalancheActor;
+				AvaActor   = World->SpawnActor<AAvaActor>(GeneratedClass, FVector::ZeroVector, FRotator::ZeroRotator, SpawnInfo);
+				PlaceholderActor = AvaActor;
 			}
-			check(AvalancheActor);
+			check(AvaActor);
 
 			// Ensure that the actor is invisible as Placeholder
-			AvalancheActor->SetHidden(true);
+			AvaActor->SetHidden(true);
 #if WITH_EDITOR
-			AvalancheActor->SetIsTemporarilyHiddenInEditor(true);
+			AvaActor->SetIsTemporarilyHiddenInEditor(true);
 #endif
 
 			// Prevent any audio from playing as a result of spawning
@@ -210,17 +210,17 @@ void UAvalancheBlueprint::UpdatePlaceholderActor(bool bInForceFullUpdate)
 			// Set the reference to the preview actor for component editing purposes
 			if (SimpleConstructionScript)
 			{
-				SimpleConstructionScript->SetComponentEditorActorInstance(AvalancheActor);
+				SimpleConstructionScript->SetComponentEditorActorInstance(AvaActor);
 			}
 #endif
 		}
 	}
 	else
 	{
-		check(AvalancheActor);
-		AvalancheActor->ReregisterAllComponents();
+		check(AvaActor);
+		AvaActor->ReregisterAllComponents();
 #if WITH_EDITOR
-		AvalancheActor->RerunConstructionScripts();
+		AvaActor->RerunConstructionScripts();
 #endif
 	}
 
@@ -303,7 +303,7 @@ void UAvalancheBlueprint::SaveAvalancheWorld()
 		}
 	}
 
-	FAvalancheWorldData NewWorldData;
+	FAvaWorldData NewWorldData;
 	
 #if WITH_EDITOR	
 	SaveWorldTask.EnterProgressFrame(1.f);
@@ -356,10 +356,10 @@ void UAvalancheBlueprint::LoadAvalancheWorld(UWorld* InOverrideWorld, TOptional<
 	const bool bIsEditorWorld = WorldData.World->IsEditorWorld();
 	
 	//1st Pass: Allocate/Spawn the Actors. Serialization is done in separate pass so Object References can be resolved correctly.
-	for (const TPair<FSoftObjectPath, FAvalancheActorData>& Pair : WorldData.ActorData)
+	for (const TPair<FSoftObjectPath, FAvaActorData>& Pair : WorldData.ActorData)
 	{
 		const FSoftObjectPath& ActorPath = Pair.Key;
-		const FAvalancheActorData& ActorData = Pair.Value;
+		const FAvaActorData& ActorData = Pair.Value;
 		
 		UClass* ActorClass = ActorData.ActorClass.TryLoadClass<AActor>();
 		if (!ActorClass)
@@ -427,14 +427,14 @@ void UAvalancheBlueprint::LoadAvalancheWorld(UWorld* InOverrideWorld, TOptional<
 	}
 	
 	//2nd Pass: Serialization 
-	for (const TPair<FSoftObjectPath, FAvalancheActorData>& Pair : WorldData.ActorData)
+	for (const TPair<FSoftObjectPath, FAvaActorData>& Pair : WorldData.ActorData)
 	{
 #if WITH_EDITOR
 		LoadWorld.EnterProgressFrame();
 #endif
 
 		const FSoftObjectPath& ActorPath = Pair.Key;
-		const FAvalancheActorData& ActorData = Pair.Value;
+		const FAvaActorData& ActorData = Pair.Value;
 		
 		if (AActor* const * const RecreatedActor = RecreatedActors.Find(ActorPath))
 		{
@@ -446,11 +446,11 @@ void UAvalancheBlueprint::LoadAvalancheWorld(UWorld* InOverrideWorld, TOptional<
 				, *Actor->GetPathName());
 
 			//Allocate Components
-			for (const TPair<FAvaObjectIndex, FAvalancheComponentData>& ComponentPair : ActorData.ComponentData)
+			for (const TPair<FAvaObjectIndex, FAvaComponentData>& ComponentPair : ActorData.ComponentData)
 			{
 				const FAvaObjectIndex& ReferenceIndex = ComponentPair.Key;
-				const FAvalancheComponentData& ComponentData = ComponentPair.Value;
-				FAvalancheSubObjectData* SubObjectData = WorldData.SubObjects.Find(ReferenceIndex);
+				const FAvaComponentData& ComponentData = ComponentPair.Value;
+				FAvaSubObjectData* SubObjectData = WorldData.SubObjects.Find(ReferenceIndex);
 				
 				if (ensure(SubObjectData))
 				{
diff --git a/Engine/Plugins/Experimental/Avalanche/Source/Avalanche/Private/AvaBlueprintGeneratedClass.cpp b/Engine/Plugins/Experimental/Avalanche/Source/Avalanche/Private/AvaBlueprintGeneratedClass.cpp
--- a/Engine/Plugins/Experimental/Avalanche/Source/Avalanche/Private/AvaBlueprintGeneratedClass.cpp
+++ b/Engine/Plugins/Experimental/Avalanche/Source/Avalanche/Private/AvaBlueprintGeneratedClass.cpp
@@ -4,7 +4,7 @@
 #include "AvaActor.h"
 #include "AvaSequence.h"
 
-void UAvaBlueprintGeneratedClass::UpdateProperties(AAvalancheActor* InActor)
+void UAvaBlueprintGeneratedClass::UpdateProperties(AAvaActor* InActor)
 {
 	UAvaBlueprintGeneratedClass* const GeneratedClass = Cast<UAvaBlueprintGeneratedClass>(InActor->GetClass());
 
diff --git a/Engine/Plugins/Experimental/Avalanche/Source/Avalanche/Private/AvaBlueprint_Serialize.cpp b/Engine/Plugins/Experimental/Avalanche/Source/Avalanche/Private/AvaBlueprint_Serialize.cpp
--- a/Engine/Plugins/Experimental/Avalanche/Source/Avalanche/Private/AvaBlueprint_Serialize.cpp
+++ b/Engine/Plugins/Experimental/Avalanche/Source/Avalanche/Private/AvaBlueprint_Serialize.cpp
@@ -2,19 +2,20 @@
 
 #include "AvaBlueprint_Serialize.h"
 #include "AudioDevice.h"
+#include "AvaLog.h"
 #include "EngineUtils.h"
 #include "MovieScene.h"
 #include "AI/NavigationSystemConfig.h"
 #include "Algo/AllOf.h"
 #include "AvaSequence.h"
-#include "Archive/AvalancheReader.h"
-#include "Archive/AvalancheWriter.h"
-#include "Data/AvalancheActorData.h"
-#include "Data/AvalancheComponentData.h"
-#include "Data/AvalancheDataDefines.h"
-#include "Data/AvalancheObjectData.h"
-#include "Data/AvalancheSubObjectData.h"
-#include "Data/AvalancheWorldData.h"
+#include "Archive/AvaReader.h"
+#include "Archive/AvaWriter.h"
+#include "Data/AvaActorData.h"
+#include "Data/AvaComponentData.h"
+#include "Data/AvaDataDefines.h"
+#include "Data/AvaObjectData.h"
+#include "Data/AvaSubObjectData.h"
+#include "Data/AvaWorldData.h"
 #include "GameFramework/PhysicsVolume.h"
 #include "Engine/BrushBuilder.h"
 #include "Misc/ScopedSlowTask.h"
@@ -50,7 +51,7 @@ FSoftObjectPath FAvaBlueprint_Serialize::SetActorInPath(AActor* NewActor, const
 	return FSoftObjectPath(PathToNewActor.GetAssetPath(), PathToNewActor.GetSubPathString() + PathAfterOriginalActor);
 }
 
-UObject* FAvaBlueprint_Serialize::ResolveObjectDependency(FAvalancheWorldData& WorldData, FAvaObjectIndex ObjectIndex)
+UObject* FAvaBlueprint_Serialize::ResolveObjectDependency(FAvaWorldData& WorldData, FAvaObjectIndex ObjectIndex)
 {
 	if (!ensure(WorldData.SerializedObjectReferences.IsValidIndex(ObjectIndex.Index)))
 	{
@@ -85,10 +86,10 @@ UObject* FAvaBlueprint_Serialize::ResolveObjectDependency(FAvalancheWorldData& W
 	{
 		const FSoftObjectPath& OriginalActorPath = *PathToActor;
 
-		FAvalancheActorData* const ActorData = WorldData.ActorData.Find(OriginalActorPath);
+		FAvaActorData* const ActorData = WorldData.ActorData.Find(OriginalActorPath);
 		if (!ActorData)
 		{
-			UE_LOG(LogTemp, Warning, TEXT("No save data found for actor %s"), *OriginalActorPath.ToString());
+			UE_LOG(LogAvaBlueprint_Serialize , Warning, TEXT("No save data found for actor %s"), *OriginalActorPath.ToString());
 			return nullptr;
 		}
 
@@ -99,7 +100,7 @@ UObject* FAvaBlueprint_Serialize::ResolveObjectDependency(FAvalancheWorldData& W
 			UClass* TargetClass = SoftClassPath.TryLoadClass<AActor>();
 			if (!TargetClass)
 			{
-				UE_LOG(LogTemp
+				UE_LOG(LogAvaBlueprint_Serialize 
 					, Error
 					, TEXT("Unknown class %s. Mostly likely it is referencing a class that was deleted.")
 					, *SoftClassPath.ToString());
@@ -129,11 +130,11 @@ UObject* FAvaBlueprint_Serialize::ResolveObjectDependency(FAvalancheWorldData& W
 				ActorWeak->SetIsTemporarilyHiddenInEditor(true);
 #endif
 
-				for (const TPair<FAvaObjectIndex, FAvalancheComponentData>& Pair : ActorData->ComponentData)
+				for (const TPair<FAvaObjectIndex, FAvaComponentData>& Pair : ActorData->ComponentData)
 				{
 					const FAvaObjectIndex& ReferenceIndex = Pair.Key;
 
-					FAvalancheSubObjectData* SubObjectData = WorldData.SubObjects.Find(ReferenceIndex);
+					FAvaSubObjectData* SubObjectData = WorldData.SubObjects.Find(ReferenceIndex);
 					if (ensure(SubObjectData))
 					{
 						const FSoftObjectPath& ComponentPath = WorldData.SerializedObjectReferences[ReferenceIndex.Index];
@@ -167,7 +168,7 @@ UObject* FAvaBlueprint_Serialize::ResolveObjectDependency(FAvalancheWorldData& W
 	const FSoftObjectPath PathToSubObject = SetActorInPath(SnapshotActor, OriginalObjectPath);
 	if (UObject* ExistingSubObject = PathToSubObject.ResolveObject())
 	{
-		FAvalancheSubObjectData* SubObjectData = WorldData.SubObjects.Find(ObjectIndex.Index);
+		FAvaSubObjectData* SubObjectData = WorldData.SubObjects.Find(ObjectIndex.Index);
 		if (!SubObjectData || SubObjectData->bWasSkippedClass)
 		{
 			return ExistingSubObject;
@@ -175,7 +176,7 @@ UObject* FAvaBlueprint_Serialize::ResolveObjectDependency(FAvalancheWorldData& W
 
 		if (SubObjectData->Class.ResolveClass() != ExistingSubObject->GetClass())
 		{
-			UE_LOG(LogTemp, Warning
+			UE_LOG(LogAvaBlueprint_Serialize , Warning
 				, TEXT("Skipping serialisation of subobject because classes are different. Object '%s' will not contain values saved.")
 				, *ExistingSubObject->GetName());
 			return ExistingSubObject;
@@ -186,7 +187,7 @@ UObject* FAvaBlueprint_Serialize::ResolveObjectDependency(FAvalancheWorldData& W
 		if (!SubobjectCache.IsValid())
 		{
 			SubobjectCache = ExistingSubObject;
-			FAvalancheReader Reader(WorldData, *SubObjectData, ExistingSubObject);
+			FAvaReader Reader(WorldData, *SubObjectData, ExistingSubObject);
 			if (Reader.IsError())
 			{
 				UE_LOG(LogAvaBlueprint_Serialize, Error,
@@ -215,10 +216,10 @@ FString FAvaBlueprint_Serialize::ExtractLastSubObjectName(const FSoftObjectPath&
 	return SubPathString.RightChop(LastDotIndex + 1);
 }
 
-UObject* FAvaBlueprint_Serialize::CreateSubObject(FAvalancheWorldData& WorldData, FAvaObjectIndex ObjectIndex,
+UObject* FAvaBlueprint_Serialize::CreateSubObject(FAvaWorldData& WorldData, FAvaObjectIndex ObjectIndex,
 	const FSoftObjectPath& PathToSubObject)
 {
-	FAvalancheSubObjectData* SubObjectData = WorldData.SubObjects.Find(ObjectIndex.Index);
+	FAvaSubObjectData* SubObjectData = WorldData.SubObjects.Find(ObjectIndex.Index);
 	if (!SubObjectData || SubObjectData->bWasSkippedClass)
 	{
 		return nullptr;
@@ -235,7 +236,7 @@ UObject* FAvaBlueprint_Serialize::CreateSubObject(FAvalancheWorldData& WorldData
 	const UClass* Class = SubObjectData->Class.TryLoadClass<UObject>();
 	if (!Class)
 	{
-		UE_LOG(LogTemp, Warning, TEXT("Class '%s' not found. Maybe it was removed?"), *SubObjectData->Class.ToString());
+		UE_LOG(LogAvaBlueprint_Serialize , Warning, TEXT("Class '%s' not found. Maybe it was removed?"), *SubObjectData->Class.ToString());
 		return nullptr;
 	}
 
@@ -244,7 +245,7 @@ UObject* FAvaBlueprint_Serialize::CreateSubObject(FAvalancheWorldData& WorldData
 	UObject* SubObjectOuter = ResolveObjectDependency(WorldData, OuterIndex);
 	if (!SubObjectOuter)
 	{
-		UE_LOG(LogTemp, Warning
+		UE_LOG(LogAvaBlueprint_Serialize , Warning
 			, TEXT("Failed to create '%s' because its outer could not be created.")
 			, *PathToSubObject.ToString());
 		return nullptr;
@@ -266,7 +267,7 @@ UObject* FAvaBlueprint_Serialize::CreateSubObject(FAvalancheWorldData& WorldData
 	}();
 	if (bSubObjectNameIsTaken)
 	{
-		UE_LOG(LogTemp
+		UE_LOG(LogAvaBlueprint_Serialize 
 			, Warning, TEXT("Failed to create '%s' because subobject name was already taken.")
 			, *PathToSubObject.ToString());
 		return nullptr;
@@ -278,7 +279,7 @@ UObject* FAvaBlueprint_Serialize::CreateSubObject(FAvalancheWorldData& WorldData
 	                                                     FindOrAdd(WorldData.SerializedObjectReferences[ObjectIndex.Index]);
 	CachedSubObject = SubObject;
 
-	FAvalancheReader Reader(WorldData, *SubObjectData, SubObject);
+	FAvaReader Reader(WorldData, *SubObjectData, SubObject);
 	if (Reader.IsError())
 	{
 		UE_LOG(LogAvaBlueprint_Serialize, Error,
@@ -334,9 +335,9 @@ UActorComponent* FAvaBlueprint_Serialize::FindMatchingComponent(AActor* ActorToS
 	return nullptr;
 }
 
-void FAvaBlueprint_Serialize::SaveObject(FAvalancheWorldData& WorldData, FAvalancheObjectData& ObjectData, UObject* Object)
+void FAvaBlueprint_Serialize::SaveObject(FAvaWorldData& WorldData, FAvaObjectData& ObjectData, UObject* Object)
 {
-	FAvalancheWriter Writer(WorldData, ObjectData, Object);
+	FAvaWriter Writer(WorldData, ObjectData, Object);
 
 	for (const FCustomVersion& CustomVersion : Writer.GetCustomVersions().GetAllVersions())
 	{
@@ -344,9 +345,9 @@ void FAvaBlueprint_Serialize::SaveObject(FAvalancheWorldData& WorldData, FAvalan
 	}
 }
 
-void FAvaBlueprint_Serialize::LoadObject(FAvalancheWorldData& WorldData, FAvalancheObjectData& ObjectData, UObject* Object)
+void FAvaBlueprint_Serialize::LoadObject(FAvaWorldData& WorldData, FAvaObjectData& ObjectData, UObject* Object)
 {
-	FAvalancheReader Reader(WorldData, ObjectData, Object);
+	FAvaReader Reader(WorldData, ObjectData, Object);
 	if (Reader.IsError())
 	{
 		UE_LOG(LogAvaBlueprint_Serialize, Error,
@@ -436,7 +437,7 @@ bool FAvaBlueprint_Serialize::ShouldSaveSubObject(UObject* InSubObject)
 	return bIsSupportedClass && !InSubObject->HasAnyFlags(RF_Transient);
 }
 
-FAvaObjectIndex FAvaBlueprint_Serialize::AddOrFindObjectReference(FAvalancheWorldData& WorldData, const FSoftObjectPath& ObjectPath)
+FAvaObjectIndex FAvaBlueprint_Serialize::AddOrFindObjectReference(FAvaWorldData& WorldData, const FSoftObjectPath& ObjectPath)
 {
 	if (const FAvaObjectIndex* ExistingIndex = WorldData.ObjectReferenceIndexMap.Find(ObjectPath))
 	{
@@ -474,12 +475,12 @@ TOptional<FSoftObjectPath> FAvaBlueprint_Serialize::ExtractActorFromPath(const F
 		       FSoftObjectPath(InObjectPath.GetAssetPath(), SubPathString.Left(DotAfterActorNameIndex));
 }
 
-UActorComponent* FAvaBlueprint_Serialize::FindOrAllocateComponent(const FAvalancheWorldData& InWorldData
+UActorComponent* FAvaBlueprint_Serialize::FindOrAllocateComponent(const FAvaWorldData& InWorldData
 	, AActor* InRecreatedActor
-	, const FAvalancheActorData& InActorData
+	, const FAvaActorData& InActorData
 	, const FSoftObjectPath& InComponentPath
-	, FAvalancheSubObjectData& SubObjectData
-	, const FAvalancheComponentData& ComponentData)
+	, FAvaSubObjectData& SubObjectData
+	, const FAvaComponentData& ComponentData)
 {
 	UActorComponent* OutComponent = FindMatchingComponent(InRecreatedActor, InComponentPath);
 
@@ -487,15 +488,15 @@ UActorComponent* FAvaBlueprint_Serialize::FindOrAllocateComponent(const FAvalanc
 	{
 		bool bIsOwnedByComponent;
 		const FSoftObjectPath& ComponentOuterPath = InWorldData.SerializedObjectReferences[SubObjectData.OuterIndex.Index];
-		TOptional<TNonNullPtr<const FAvalancheActorData>> ComponentOuterData;
+		TOptional<TNonNullPtr<const FAvaActorData>> ComponentOuterData;
 
 		//Find Saved Actor Data Using Object Path
 		if (const TOptional<FSoftObjectPath> PathToActor = ExtractActorFromPath(ComponentOuterPath, bIsOwnedByComponent))
 		{
-			const FAvalancheActorData* const Result = InWorldData.ActorData.Find(*PathToActor);
+			const FAvaActorData* const Result = InWorldData.ActorData.Find(*PathToActor);
 
 			UE_CLOG(Result == nullptr
-				, LogTemp
+				, LogAvaBlueprint_Serialize 
 				, Warning
 				, TEXT(
 					"Path %s looks like an actor path but no data was saved for it. Maybe it was a reference to an auto-generated actor, e.g. a brush or volume present in all worlds by default?"
@@ -504,7 +505,7 @@ UActorComponent* FAvaBlueprint_Serialize::FindOrAllocateComponent(const FAvalanc
 
 			if (Result)
 			{
-				ComponentOuterData = TOptional<TNonNullPtr<const FAvalancheActorData>>(Result);
+				ComponentOuterData = TOptional<TNonNullPtr<const FAvaActorData>>(Result);
 			}
 		}
 
@@ -528,7 +529,7 @@ UActorComponent* FAvaBlueprint_Serialize::FindOrAllocateComponent(const FAvalanc
 		if (bIsOwnedByComponent)
 		{
 			bool bFoundPath = false;
-			for (const TPair<FAvaObjectIndex, FAvalancheComponentData>& ComponentPair : InActorData.ComponentData)
+			for (const TPair<FAvaObjectIndex, FAvaComponentData>& ComponentPair : InActorData.ComponentData)
 			{
 				const FSoftObjectPath& SavedComponentPath = InWorldData.SerializedObjectReferences[ComponentPair.Key.Index];
 				if (SavedComponentPath == ComponentOuterPath)
@@ -536,7 +537,7 @@ UActorComponent* FAvaBlueprint_Serialize::FindOrAllocateComponent(const FAvalanc
 					bFoundPath = true;
 
 					const FAvaObjectIndex ReferenceIndex                    = ComponentPair.Key;
-					const FAvalancheSubObjectData* const FoundSubObjectData = InWorldData.SubObjects.Find(ReferenceIndex);
+					const FAvaSubObjectData* const FoundSubObjectData = InWorldData.SubObjects.Find(ReferenceIndex);
 
 					if (ensure(FoundSubObjectData))
 					{
@@ -551,7 +552,7 @@ UActorComponent* FAvaBlueprint_Serialize::FindOrAllocateComponent(const FAvalanc
 			}
 
 			UE_CLOG(bFoundPath
-				, LogTemp
+				, LogAvaBlueprint_Serialize 
 				, Error
 				, TEXT("Failed to find outer for component %s")
 				, *ComponentOuterPath.ToString());
@@ -563,8 +564,9 @@ UActorComponent* FAvaBlueprint_Serialize::FindOrAllocateComponent(const FAvalanc
 			{
 				const UClass* ComponentClass = SubObjectData.Class.TryLoadClass<UActorComponent>();
 
-				if (!ensureMsgf(ComponentClass, TEXT("Component Class is invalid for Outer %s"), *ComponentOuter->GetName()))
+				if (!ComponentClass)
 				{
+					UE_LOG(LogAvaBlueprint_Serialize, Warning, TEXT("Component Class is invalid for Outer %s"), *ComponentOuter->GetName());
 					return nullptr;
 				}
 				
@@ -620,12 +622,12 @@ UActorComponent* FAvaBlueprint_Serialize::FindOrAllocateComponent(const FAvalanc
 	return OutComponent;
 }
 
-void FAvaBlueprint_Serialize::AddSubObjectDependency(FAvalancheWorldData& WorldData, UObject* ReferenceFromOriginalObject,
+void FAvaBlueprint_Serialize::AddSubObjectDependency(FAvaWorldData& WorldData, UObject* ReferenceFromOriginalObject,
 	FAvaObjectIndex Index)
 {
 	if (!ShouldSaveSubObject(ReferenceFromOriginalObject))
 	{
-		WorldData.SubObjects.Add(Index, FAvalancheSubObjectData::MakeSkippedSubObjectData());
+		WorldData.SubObjects.Add(Index, FAvaSubObjectData::MakeSkippedSubObjectData());
 		return;
 	}
 
@@ -644,7 +646,7 @@ void FAvaBlueprint_Serialize::AddSubObjectDependency(FAvalancheWorldData& WorldD
 		}
 
 		// Important: first allocate SubobjectData on stack...
-		FAvalancheSubObjectData SubObjectData;
+		FAvaSubObjectData SubObjectData;
 		const UClass* Class      = ReferenceFromOriginalObject->GetClass();
 		SubObjectData.Class      = Class;
 		SubObjectData.OuterIndex = AddObjectDependency(WorldData, ReferenceFromOriginalObject->GetOuter());
@@ -656,7 +658,7 @@ void FAvaBlueprint_Serialize::AddSubObjectDependency(FAvalancheWorldData& WorldD
 	}
 }
 
-FAvaObjectIndex FAvaBlueprint_Serialize::AddObjectDependency(FAvalancheWorldData& WorldData
+FAvaObjectIndex FAvaBlueprint_Serialize::AddObjectDependency(FAvaWorldData& WorldData
 	, UObject* ReferenceFromOriginalObject
 	, bool bCheckWhetherSubObject)
 {
@@ -669,21 +671,21 @@ FAvaObjectIndex FAvaBlueprint_Serialize::AddObjectDependency(FAvalancheWorldData
 	return OutIndex;
 }
 
-void FAvaBlueprint_Serialize::LoadActor(AActor* InActor, const FSoftObjectPath& InActorPath, FAvalancheWorldData& WorldData)
+void FAvaBlueprint_Serialize::LoadActor(AActor* InActor, const FSoftObjectPath& InActorPath, FAvaWorldData& WorldData)
 {
 	check(InActor);
 
-	UE_LOG(LogTemp
+	UE_LOG(LogAvaBlueprint_Serialize 
 		, Verbose
 		, TEXT("========== Deserialize Actor %s ==========")
 		, *InActorPath.ToString());
 
-	FAvalancheActorData& ActorData = WorldData.ActorData[InActorPath];
+	FAvaActorData& ActorData = WorldData.ActorData[InActorPath];
 
 	LoadObject(WorldData, ActorData, InActor);
 
 #if WITH_EDITOR
-	UE_LOG(LogTemp
+	UE_LOG(LogAvaBlueprint_Serialize 
 		, Verbose
 		, TEXT("ActorLabel is \"%s\" for \"%s\" (editor object path \"%s\")")
 		, *InActor->GetActorLabel()
@@ -692,7 +694,7 @@ void FAvaBlueprint_Serialize::LoadActor(AActor* InActor, const FSoftObjectPath&
 #endif
 
 	//Deserialize Components
-	for (const TPair<FAvaObjectIndex, FAvalancheComponentData>& ComponentPair : ActorData.ComponentData)
+	for (const TPair<FAvaObjectIndex, FAvaComponentData>& ComponentPair : ActorData.ComponentData)
 	{
 		const EComponentCreationMethod CreationMethod = ComponentPair.Value.CreationMethod;
 
@@ -704,7 +706,7 @@ void FAvaBlueprint_Serialize::LoadActor(AActor* InActor, const FSoftObjectPath&
 
 			if (UActorComponent* const ComponentToRestore = FindMatchingComponent(InActor, OriginalComponentPath))
 			{
-				FAvalancheSubObjectData& SubObjectData = WorldData.SubObjects[SubObjectIndex];
+				FAvaSubObjectData& SubObjectData = WorldData.SubObjects[SubObjectIndex];
 				LoadObject(WorldData, SubObjectData, ComponentToRestore);
 				if (UPrimitiveComponent* PrimitiveComponent = Cast<UPrimitiveComponent>(ComponentToRestore))
 				{
@@ -730,28 +732,28 @@ void FAvaBlueprint_Serialize::LoadActor(AActor* InActor, const FSoftObjectPath&
 	}
 }
 
-TOptional<FAvalancheComponentData> FAvaBlueprint_Serialize::SaveComponent(UActorComponent* InComponent)
+TOptional<FAvaComponentData> FAvaBlueprint_Serialize::SaveComponent(UActorComponent* InComponent)
 {
 	check(InComponent);
 
 	if (InComponent->CreationMethod == EComponentCreationMethod::UserConstructionScript)
 	{
-		UE_LOG(LogTemp, Warning
+		UE_LOG(LogAvaBlueprint_Serialize , Warning
 			, TEXT("Components created dynamically in the construction script are not supported (%s). Skipping...")
 			, *InComponent->GetPathName());
-		return TOptional<FAvalancheComponentData>{};
+		return TOptional<FAvaComponentData>{};
 	}
 
-	FAvalancheComponentData OutComponentData;
+	FAvaComponentData OutComponentData;
 	OutComponentData.CreationMethod = InComponent->CreationMethod;
 	return OutComponentData;
 }
 
-FAvalancheActorData FAvaBlueprint_Serialize::SaveActor(AActor* InActor, FAvalancheWorldData& WorldData)
+FAvaActorData FAvaBlueprint_Serialize::SaveActor(AActor* InActor, FAvaWorldData& WorldData)
 {
 	check(InActor);
 
-	FAvalancheActorData OutActorData;
+	FAvaActorData OutActorData;
 
 	const UClass* const ActorClass = InActor->GetClass();
 	OutActorData.ActorClass        = ActorClass;
@@ -769,7 +771,7 @@ FAvalancheActorData FAvaBlueprint_Serialize::SaveActor(AActor* InActor, FAvalanc
 	{
 		if (ShouldSaveComponent(Component))
 		{
-			if (TOptional<FAvalancheComponentData> SerializedComponentData = SaveComponent(Component))
+			if (TOptional<FAvaComponentData> SerializedComponentData = SaveComponent(Component))
 			{
 				const FAvaObjectIndex ComponentIndex = AddObjectDependency(WorldData, Component);
 				OutActorData.ComponentData.Add(ComponentIndex, *SerializedComponentData);
diff --git a/Engine/Plugins/Experimental/Avalanche/Source/Avalanche/Private/AvaBlueprint_Serialize.h b/Engine/Plugins/Experimental/Avalanche/Source/Avalanche/Private/AvaBlueprint_Serialize.h
--- a/Engine/Plugins/Experimental/Avalanche/Source/Avalanche/Private/AvaBlueprint_Serialize.h
+++ b/Engine/Plugins/Experimental/Avalanche/Source/Avalanche/Private/AvaBlueprint_Serialize.h
@@ -2,18 +2,18 @@
 
 #pragma once
 
-#include "Data/AvalancheDataDefines.h"
+#include "Data/AvaDataDefines.h"
 
 class AActor;
 class FString;
 class UActorComponent;
 class UObject;
 class UWorld;
-struct FAvalancheActorData;
-struct FAvalancheComponentData;
-struct FAvalancheObjectData;
-struct FAvalancheSubObjectData;
-struct FAvalancheWorldData;
+struct FAvaActorData;
+struct FAvaComponentData;
+struct FAvaObjectData;
+struct FAvaSubObjectData;
+struct FAvaWorldData;
 struct FSoftObjectPath;
 template<typename OptionalType> struct TOptional;
 
@@ -25,37 +25,37 @@ struct FAvaBlueprint_Serialize
 	 */
 	static FSoftObjectPath SetActorInPath(AActor* NewActor, const FSoftObjectPath& OriginalObjectPath);
 
-	static UObject* ResolveObjectDependency(FAvalancheWorldData& WorldData, FAvaObjectIndex ObjectIndex);
+	static UObject* ResolveObjectDependency(FAvaWorldData& WorldData, FAvaObjectIndex ObjectIndex);
 
 	static FString ExtractLastSubObjectName(const FSoftObjectPath& ObjectPath);
 
-	static UObject* CreateSubObject(FAvalancheWorldData& WorldData, FAvaObjectIndex ObjectIndex, const FSoftObjectPath& PathToSubObject);
+	static UObject* CreateSubObject(FAvaWorldData& WorldData, FAvaObjectIndex ObjectIndex, const FSoftObjectPath& PathToSubObject);
 
 	static FString ExtractRootToLeafComponentPath(const FSoftObjectPath& ComponentPath);
 	static UActorComponent* FindMatchingComponent(AActor* ActorToSearchOn, const FSoftObjectPath& ComponentPath);
 
-	static void SaveObject(FAvalancheWorldData& WorldData, FAvalancheObjectData& ObjectData, UObject* Object);
-	static void LoadObject(FAvalancheWorldData& WorldData, FAvalancheObjectData& ObjectData, UObject* Object);
+	static void SaveObject(FAvaWorldData& WorldData, FAvaObjectData& ObjectData, UObject* Object);
+	static void LoadObject(FAvaWorldData& WorldData, FAvaObjectData& ObjectData, UObject* Object);
 
 	static bool ShouldSaveActor(const UWorld* InWorld, const AActor* InActor);
 	static bool ShouldSaveComponent(const UActorComponent* InActorComponent);
 	static bool ShouldSaveSubObject(UObject* InSubObject);
 
-	static UActorComponent* FindOrAllocateComponent(const FAvalancheWorldData& InWorldData
+	static UActorComponent* FindOrAllocateComponent(const FAvaWorldData& InWorldData
 		, AActor* InRecreatedActor
-		, const FAvalancheActorData& InActorData
+		, const FAvaActorData& InActorData
 		, const FSoftObjectPath& InComponentPath
-		, FAvalancheSubObjectData& SubObjectData
-		, const FAvalancheComponentData& ComponentData);
+		, FAvaSubObjectData& SubObjectData
+		, const FAvaComponentData& ComponentData);
 
-	static FAvaObjectIndex AddOrFindObjectReference(FAvalancheWorldData& WorldData, const FSoftObjectPath& ObjectPath);
+	static FAvaObjectIndex AddOrFindObjectReference(FAvaWorldData& WorldData, const FSoftObjectPath& ObjectPath);
 	static TOptional<FSoftObjectPath> ExtractActorFromPath(const FSoftObjectPath& InObjectPath, bool& bIsPathToActorSubObject);
-	static void AddSubObjectDependency(FAvalancheWorldData& WorldData, UObject* ReferenceFromOriginalObject, FAvaObjectIndex Index);
-	static FAvaObjectIndex AddObjectDependency(FAvalancheWorldData& WorldData
+	static void AddSubObjectDependency(FAvaWorldData& WorldData, UObject* ReferenceFromOriginalObject, FAvaObjectIndex Index);
+	static FAvaObjectIndex AddObjectDependency(FAvaWorldData& WorldData
 		, UObject* ReferenceFromOriginalObject
 		, bool bCheckWhetherSubObject = true);
 
-	static void LoadActor(AActor* InActor, const FSoftObjectPath& InActorPath, FAvalancheWorldData& WorldData);
-	static TOptional<FAvalancheComponentData> SaveComponent(UActorComponent* InComponent);
-	static FAvalancheActorData SaveActor(AActor* InActor, FAvalancheWorldData& WorldData);
+	static void LoadActor(AActor* InActor, const FSoftObjectPath& InActorPath, FAvaWorldData& WorldData);
+	static TOptional<FAvaComponentData> SaveComponent(UActorComponent* InComponent);
+	static FAvaActorData SaveActor(AActor* InActor, FAvaWorldData& WorldData);
 };
diff --git a/Engine/Plugins/Experimental/Avalanche/Source/Avalanche/Private/AvaScene.cpp b/Engine/Plugins/Experimental/Avalanche/Source/Avalanche/Private/AvaScene.cpp
--- a/Engine/Plugins/Experimental/Avalanche/Source/Avalanche/Private/AvaScene.cpp
+++ b/Engine/Plugins/Experimental/Avalanche/Source/Avalanche/Private/AvaScene.cpp
@@ -70,7 +70,7 @@ AAvaScene* AAvaScene::CreateFromBlueprint(ULevel* InNewLevel, UAvalancheBlueprin
 	// Load Blueprint World Actors into the new Level 
 	SlowTask.EnterProgressFrame(80, LOCTEXT("CopyingActors", "Copying Blueprint Actors"));
 	{
-		auto InitActor = [](AActor& InActor, const FAvalancheActorData& InActorData)
+		auto InitActor = [](AActor& InActor, const FAvaActorData& InActorData)
 		{
 			InActor.bHiddenEd = InActorData.bEditorVisibility;
 		};
diff --git a/Engine/Plugins/Experimental/Avalanche/Source/Avalanche/Private/Data/AvalancheVersionInfo.cpp b/Engine/Plugins/Experimental/Avalanche/Source/Avalanche/Private/Data/AvalancheVersionInfo.cpp
--- a/Engine/Plugins/Experimental/Avalanche/Source/Avalanche/Private/Data/AvalancheVersionInfo.cpp
+++ b/Engine/Plugins/Experimental/Avalanche/Source/Avalanche/Private/Data/AvalancheVersionInfo.cpp
@@ -1,18 +1,18 @@
 // Copyright Epic Games, Inc. All Rights Reserved.
 
-#include "Data/AvalancheVersionInfo.h"
+#include "Data/AvaVersionInfo.h"
 #include "Misc/EngineVersion.h"
 #include "Serialization/CustomVersion.h"
 #include "UObject/ObjectVersion.h"
 
-void FAvalancheFileVersionInfo::UpdateToLatest()
+void FAvaFileVersionInfo::UpdateToLatest()
 {
 	FileVersionUE4      = GPackageFileUEVersion.FileVersionUE4;
 	FileVersionUE5      = GPackageFileUEVersion.FileVersionUE5;
 	FileVersionLicensee = GPackageFileLicenseeUEVersion;
 }
 
-void FAvalancheEngineVersionInfo::UpdateToLatest()
+void FAvaEngineVersionInfo::UpdateToLatest()
 {
 	const FEngineVersion& EngineVersion(FEngineVersion::Current());
 	Major      = EngineVersion.GetMajor();
@@ -21,14 +21,14 @@ void FAvalancheEngineVersionInfo::UpdateToLatest()
 	Changelist = EngineVersion.GetChangelist();
 }
 
-void FAvalancheCustomVersionInfo::SetCustomVersion(const FCustomVersion& InVersion)
+void FAvaCustomVersionInfo::SetCustomVersion(const FCustomVersion& InVersion)
 {
 	FriendlyName = InVersion.GetFriendlyName();
 	Key          = InVersion.Key;
 	Version      = InVersion.Version;
 }
 
-void FAvalancheVersionInfo::UpdateToLatest()
+void FAvaVersionInfo::UpdateToLatest()
 {
 	FileVersion.UpdateToLatest();
 	EngineVersion.UpdateToLatest();
@@ -57,7 +57,7 @@ void FAvalancheVersionInfo::UpdateToLatest()
 	}
 }
 
-void FAvalancheVersionInfo::ApplyToArchive(FArchive& Archive) const
+void FAvaVersionInfo::ApplyToArchive(FArchive& Archive) const
 {
 	//If CustomVersions is Empty, it's pre-versioning
 	if (CustomVersions.IsEmpty())
@@ -74,7 +74,7 @@ void FAvalancheVersionInfo::ApplyToArchive(FArchive& Archive) const
 	Archive.SetEngineVer(FEngineVersionBase(EngineVersion.Major, EngineVersion.Minor, EngineVersion.Patch, EngineVersion.Changelist));
 
 	FCustomVersionContainer CustomVersionContainer;
-	for (const FAvalancheCustomVersionInfo& CustomVersion : CustomVersions)
+	for (const FAvaCustomVersionInfo& CustomVersion : CustomVersions)
 	{
 		CustomVersionContainer.SetVersion(CustomVersion.Key
 			, CustomVersion.Version
diff --git a/Engine/Plugins/Experimental/Avalanche/Source/Avalanche/Private/Data/AvalancheWorldData.cpp b/Engine/Plugins/Experimental/Avalanche/Source/Avalanche/Private/Data/AvalancheWorldData.cpp
--- a/Engine/Plugins/Experimental/Avalanche/Source/Avalanche/Private/Data/AvalancheWorldData.cpp
+++ b/Engine/Plugins/Experimental/Avalanche/Source/Avalanche/Private/Data/AvalancheWorldData.cpp
@@ -1,13 +1,13 @@
 // Copyright Epic Games, Inc. All Rights Reserved.
 
-#include "Data/AvalancheWorldData.h"
+#include "Data/AvaWorldData.h"
 
-void FAvalancheWorldData::FinalizeSave()
+void FAvaWorldData::FinalizeSave()
 {
 	VersionInfo.UpdateToLatest();
 }
 
-void FAvalancheWorldData::ResetTransientData()
+void FAvaWorldData::ResetTransientData()
 {
 	World.Reset();
 	NameIndexMap.Reset();
diff --git a/Engine/Plugins/Experimental/Avalanche/Source/Avalanche/Public/AvaActor.h b/Engine/Plugins/Experimental/Avalanche/Source/Avalanche/Public/AvaActor.h
--- a/Engine/Plugins/Experimental/Avalanche/Source/Avalanche/Public/AvaActor.h
+++ b/Engine/Plugins/Experimental/Avalanche/Source/Avalanche/Public/AvaActor.h
@@ -7,12 +7,12 @@
 
 /** Avalanche Actor */
 UCLASS(MinimalAPI, DisplayName = "Motion Design Actor")
-class AAvalancheActor : public AActor
+class AAvaActor : public AActor
 {
 	GENERATED_BODY()
 
 public:
-	AAvalancheActor();
+	AAvaActor();
 
 	void SetIsPlaceholder(bool bInIsPlaceholder = false) { bIsPlaceholder = bInIsPlaceholder; }
 
@@ -22,6 +22,10 @@ public:
 
 	//~ Begin AActor
 	virtual void PostActorCreated() override;
+
+#if WITH_EDITOR
+	virtual FString GetDefaultActorLabel() const override;
+#endif
 	//~ End AActor
 
 	//~ Begin UObject
diff --git a/Engine/Plugins/Experimental/Avalanche/Source/Avalanche/Public/AvaBlueprint.h b/Engine/Plugins/Experimental/Avalanche/Source/Avalanche/Public/AvaBlueprint.h
--- a/Engine/Plugins/Experimental/Avalanche/Source/Avalanche/Public/AvaBlueprint.h
+++ b/Engine/Plugins/Experimental/Avalanche/Source/Avalanche/Public/AvaBlueprint.h
@@ -4,7 +4,7 @@
 
 #include "AvaSceneTree.h"
 #include "AvaSequenceShared.h"
-#include "Data/AvalancheWorldData.h"
+#include "Data/AvaWorldData.h"
 #include "Engine/Blueprint.h"
 #include "IAvaSceneInterface.h"
 #include "IAvaSequenceProvider.h"
@@ -12,7 +12,7 @@
 #include "Viewport/AvaViewportQualitySettings.h"
 #include "AvaBlueprint.generated.h"
 
-class AAvalancheActor;
+class AAvaActor;
 class IAvaSequencePlaybackObject;
 class ISequencer;
 class UAvaSequence;
@@ -55,7 +55,7 @@ public:
 
 	UWorld* GetAvalancheWorld() const;
 
-	AAvalancheActor* GetPlaceholderActor() const;
+	AAvaActor* GetPlaceholderActor() const;
 	IAvaSequencePlaybackObject* GetScenePlayback() const;
 
 	void DestroyPlaceholderActor();
@@ -68,7 +68,7 @@ public:
 	 */
 	void SaveAvalancheWorld();
 
-	using FInitActorFunctionRef = TFunctionRef<void(AActor&, const FAvalancheActorData&)>;
+	using FInitActorFunctionRef = TFunctionRef<void(AActor&, const FAvaActorData&)>;
 	void LoadAvalancheWorld(UWorld* InOverrideWorld = nullptr, TOptional<FInitActorFunctionRef> InInitActorFunction = {});
 
 	bool IsSavingAvalancheWorld() const { return bSavingWorld; }
@@ -154,14 +154,14 @@ private:
 	 * Placeholder should be null when an BP Avalanche Actor is spawned into the World as the Actor to use should be the spawned one instead of this Placeholder
 	 */
 	UPROPERTY(Transient, DuplicateTransient, TextExportTransient)
-	TObjectPtr<AAvalancheActor> PlaceholderActor;
+	TObjectPtr<AAvaActor> PlaceholderActor;
 
 	/** The Base Playback Scene that is always present to Play Animations */
 	UPROPERTY()
 	TScriptInterface<IAvaSequencePlaybackObject> PlaybackObject;
 
 	UPROPERTY()
-	FAvalancheWorldData WorldData;
+	FAvaWorldData WorldData;
 
 	UPROPERTY()
 	FAvaSceneTree SceneTree;
diff --git a/Engine/Plugins/Experimental/Avalanche/Source/Avalanche/Public/AvaBlueprintGeneratedClass.h b/Engine/Plugins/Experimental/Avalanche/Source/Avalanche/Public/AvaBlueprintGeneratedClass.h
--- a/Engine/Plugins/Experimental/Avalanche/Source/Avalanche/Public/AvaBlueprintGeneratedClass.h
+++ b/Engine/Plugins/Experimental/Avalanche/Source/Avalanche/Public/AvaBlueprintGeneratedClass.h
@@ -9,7 +9,7 @@
 
 #include "AvaBlueprintGeneratedClass.generated.h"
 
-class AAvalancheActor;
+class AAvaActor;
 class UAvaSequence;
 
 UCLASS(MinimalAPI, DisplayName = "Motion Design Blueprint Generated Class")
@@ -18,7 +18,7 @@ class UAvaBlueprintGeneratedClass : public UBlueprintGeneratedClass
 	GENERATED_BODY()
 
 public:
-	void UpdateProperties(AAvalancheActor* InActor);
+	void UpdateProperties(AAvaActor* InActor);
 
 	UPROPERTY()
 	TArray<TObjectPtr<UAvaSequence>> Animations;
diff --git a/Engine/Plugins/Experimental/Avalanche/Source/Avalanche/Public/Data/AvaActorData.h b/Engine/Plugins/Experimental/Avalanche/Source/Avalanche/Public/Data/AvaActorData.h
new file mode 100644
--- /dev/null
+++ b/Engine/Plugins/Experimental/Avalanche/Source/Avalanche/Public/Data/AvaActorData.h
@@ -0,0 +1,37 @@
+// Copyright Epic Games, Inc. All Rights Reserved.
+
+#pragma once
+
+#include "AvaComponentData.h"
+#include "AvaDataDefines.h"
+#include "AvaObjectData.h"
+#include "Containers/Array.h"
+#include "Containers/Map.h"
+#include "Containers/UnrealString.h"
+#include "UObject/SoftObjectPath.h"
+#include "AvaActorData.generated.h"
+
+USTRUCT()
+struct FAvaActorData : public FAvaObjectData
+{
+	GENERATED_BODY()
+
+#if WITH_EDITORONLY_DATA
+	UPROPERTY()
+	FString ActorLabel;
+
+	UPROPERTY()
+	bool bEditorVisibility = false;
+#endif
+	
+	UPROPERTY()
+	FSoftClassPath ActorClass;
+
+	/** Additional Info for Components */
+	UPROPERTY()
+	TMap<FAvaObjectIndex, FAvaComponentData> ComponentData;
+
+	/** Indices to all Non-Component SubObjects */
+	UPROPERTY()
+	TArray<FAvaObjectIndex> OwnedSubobjects;
+};
diff --git a/Engine/Plugins/Experimental/Avalanche/Source/Avalanche/Public/Data/AvaComponentData.h b/Engine/Plugins/Experimental/Avalanche/Source/Avalanche/Public/Data/AvaComponentData.h
new file mode 100644
--- /dev/null
+++ b/Engine/Plugins/Experimental/Avalanche/Source/Avalanche/Public/Data/AvaComponentData.h
@@ -0,0 +1,16 @@
+// Copyright Epic Games, Inc. All Rights Reserved.
+
+#pragma once
+
+#include "ComponentInstanceDataCache.h"
+#include "AvaComponentData.generated.h"
+
+USTRUCT()
+struct FAvaComponentData
+{
+	GENERATED_BODY()
+
+	/** Describes how the component was created */
+	UPROPERTY()
+	EComponentCreationMethod CreationMethod{};
+};
diff --git a/Engine/Plugins/Experimental/Avalanche/Source/Avalanche/Public/Data/AvaDataDefines.h b/Engine/Plugins/Experimental/Avalanche/Source/Avalanche/Public/Data/AvaDataDefines.h
new file mode 100644
--- /dev/null
+++ b/Engine/Plugins/Experimental/Avalanche/Source/Avalanche/Public/Data/AvaDataDefines.h
@@ -0,0 +1,28 @@
+// Copyright Epic Games, Inc. All Rights Reserved.
+
+#pragma once
+
+#include "AvaDataDefines.generated.h"
+
+USTRUCT()
+struct FAvaObjectIndex
+{
+	GENERATED_BODY()
+
+	FAvaObjectIndex() {}
+
+	FAvaObjectIndex(int32 Index) : Index(Index) {}
+
+	bool operator==(const FAvaObjectIndex& Other) const
+	{
+		return Index == Other.Index;
+	}
+	
+	friend uint32 GetTypeHash(const FAvaObjectIndex& Value)
+    {
+    	return GetTypeHash(Value.Index);
+    }
+
+	UPROPERTY()
+	int32 Index = INDEX_NONE;
+};
diff --git a/Engine/Plugins/Experimental/Avalanche/Source/Avalanche/Public/Data/AvaObjectData.h b/Engine/Plugins/Experimental/Avalanche/Source/Avalanche/Public/Data/AvaObjectData.h
new file mode 100644
--- /dev/null
+++ b/Engine/Plugins/Experimental/Avalanche/Source/Avalanche/Public/Data/AvaObjectData.h
@@ -0,0 +1,21 @@
+// Copyright Epic Games, Inc. All Rights Reserved.
+
+#pragma once
+
+#include "Containers/Array.h"
+#include "UObject/ObjectMacros.h"
+#include "AvaObjectData.generated.h"
+
+USTRUCT()
+struct FAvaObjectData
+{
+	GENERATED_BODY()
+
+	UPROPERTY()
+	TArray<uint8> SerializedData;
+
+	UPROPERTY()
+	uint64 ObjectFlags{};
+
+	EObjectFlags GetObjectFlags() const { return static_cast<EObjectFlags>(ObjectFlags); }
+};
diff --git a/Engine/Plugins/Experimental/Avalanche/Source/Avalanche/Public/Data/AvaSubObjectData.h b/Engine/Plugins/Experimental/Avalanche/Source/Avalanche/Public/Data/AvaSubObjectData.h
new file mode 100644
--- /dev/null
+++ b/Engine/Plugins/Experimental/Avalanche/Source/Avalanche/Public/Data/AvaSubObjectData.h
@@ -0,0 +1,32 @@
+// Copyright Epic Games, Inc. All Rights Reserved.
+
+#pragma once
+
+#include "AvaDataDefines.h"
+#include "AvaObjectData.h"
+#include "UObject/SoftObjectPath.h"
+#include "AvaSubObjectData.generated.h"
+
+USTRUCT()
+struct FAvaSubObjectData : public FAvaObjectData
+{
+	GENERATED_BODY()
+
+	static FAvaSubObjectData MakeSkippedSubObjectData()
+	{
+		FAvaSubObjectData OutData;
+		OutData.bWasSkippedClass = true;
+		return OutData;
+	}
+
+	/** Index to FAvaWorldData::SerializedObjectReferences */
+	UPROPERTY()
+	FAvaObjectIndex OuterIndex;
+
+	UPROPERTY()
+	FSoftClassPath Class;
+
+	/** Whether this class was marked as unsupported when the snapshot was taken */
+	UPROPERTY()
+	bool bWasSkippedClass = false;
+};
diff --git a/Engine/Plugins/Experimental/Avalanche/Source/Avalanche/Public/Data/AvaVersionInfo.h b/Engine/Plugins/Experimental/Avalanche/Source/Avalanche/Public/Data/AvaVersionInfo.h
new file mode 100644
--- /dev/null
+++ b/Engine/Plugins/Experimental/Avalanche/Source/Avalanche/Public/Data/AvaVersionInfo.h
@@ -0,0 +1,122 @@
+// Copyright Epic Games, Inc. All Rights Reserved.
+
+#pragma once
+
+#include "Misc/Guid.h"
+#include "UObject/NameTypes.h"
+#include "AvaVersionInfo.generated.h"
+
+class FArchive;
+struct FCustomVersion;
+
+/** Holds file version information */
+USTRUCT()
+struct FAvaFileVersionInfo
+{
+	GENERATED_BODY()
+
+	FAvaFileVersionInfo()
+		: FileVersionUE4(0)
+		, FileVersionUE5(0)
+		, FileVersionLicensee(0)
+	{
+	}
+	
+	void UpdateToLatest();
+
+	UPROPERTY()
+	int32 FileVersionUE4;
+
+	UPROPERTY()
+	int32 FileVersionUE5;
+	
+	/* Licensee file version */
+	UPROPERTY()
+	int32 FileVersionLicensee;
+};
+
+/** Holds engine version information */
+USTRUCT()
+struct FAvaEngineVersionInfo
+{
+	GENERATED_BODY()
+
+	FAvaEngineVersionInfo()
+		: Major(0)
+		, Minor(0)
+		, Patch(0)
+		, Changelist(0)
+	{
+	}
+	
+	void UpdateToLatest();
+
+	/** Major version number */
+	UPROPERTY()
+	uint16 Major;
+
+	/** Minor version number */
+	UPROPERTY()
+	uint16 Minor;
+
+	/** Patch version number */
+	UPROPERTY()
+	uint16 Patch;
+
+	/** Changelist number. This is used to arbitrate when Major/Minor/Patch version numbers match */
+	UPROPERTY()
+	uint32 Changelist;
+};
+
+/** Holds custom version information */
+USTRUCT()
+struct FAvaCustomVersionInfo
+{
+	GENERATED_BODY()
+
+	FAvaCustomVersionInfo()
+		: Version(0)
+	{
+	}
+	
+	void SetCustomVersion(const FCustomVersion& InVersion);
+	
+	/** Unique custom key */
+	UPROPERTY()
+	FGuid Key;
+
+	/** Custom version */
+	UPROPERTY()
+	int32 Version;
+
+	/** Friendly name of the version */
+	UPROPERTY()
+	FName FriendlyName;
+};
+
+/** Holds version information for a session */
+USTRUCT()
+struct FAvaVersionInfo
+{
+	GENERATED_BODY()
+
+	/** Update all Versions (File, Engine, Custom) to Latest */
+	void UpdateToLatest();
+
+	void ApplyToArchive(FArchive& Archive) const;
+
+	/** File version info */
+	UPROPERTY()
+	FAvaFileVersionInfo FileVersion;
+
+	/** Engine version info */
+	UPROPERTY()
+	FAvaEngineVersionInfo EngineVersion;
+
+	/** Custom version info */
+	UPROPERTY()
+	TArray<FAvaCustomVersionInfo> CustomVersions;
+
+	/** Keep track of the custom versions used by the assets while saving. */
+	TSet<FGuid> UsedCustomVersions;
+};
\ No newline at end of file
diff --git a/Engine/Plugins/Experimental/Avalanche/Source/Avalanche/Public/Data/AvaWorldData.h b/Engine/Plugins/Experimental/Avalanche/Source/Avalanche/Public/Data/AvaWorldData.h
new file mode 100644
--- /dev/null
+++ b/Engine/Plugins/Experimental/Avalanche/Source/Avalanche/Public/Data/AvaWorldData.h
@@ -0,0 +1,85 @@
+// Copyright Epic Games, Inc. All Rights Reserved.
+
+#pragma once
+
+#include "AvaActorData.h"
+#include "AvaSubObjectData.h"
+#include "AvaVersionInfo.h"
+#include "Containers/Array.h"
+#include "Containers/Map.h"
+#include "UObject/NameTypes.h"
+#include "UObject/SoftObjectPath.h"
+#include "UObject/WeakObjectPtr.h"
+#include "AvaWorldData.generated.h"
+
+class AActor;
+class UObject;
+class UWorld;
+
+USTRUCT()
+struct FAvaWorldData
+{
+	GENERATED_BODY()
+	
+	/**
+	 *	Called after saving the world to perform final processing.
+	 */
+	void FinalizeSave();
+	
+	/**
+	 * Resets the Transient Members to zero, since DuplicateTransient only works on UCLASS not USTRUCT.
+	 */
+	void ResetTransientData();
+	
+	/**
+	* Holds the Version Info this World Data was last serialized with
+	*/
+	UPROPERTY()
+	FAvaVersionInfo VersionInfo;
+
+	/**
+	* Transient World used to load the World Data in
+	*/
+	UPROPERTY(Transient)
+	TWeakObjectPtr<UWorld> World;
+
+	/**
+	* Holds serialized actor data.
+	* Maps the original actor's path to its serialized data.
+	*/
+	UPROPERTY()
+	TMap<FSoftObjectPath, FAvaActorData> ActorData;
+
+	UPROPERTY()
+	TArray<FName> SerializedNames;
+
+	/* 
+	 * List of all Object References in the World
+	 * This is an Untranslated version so we ned to change the World Name Part (when it was serialized)
+	 * of the Path to the New World Name
+	 */
+	UPROPERTY()
+	TArray<FSoftObjectPath> SerializedObjectReferences;
+
+	/*
+	 * Key: Index to SerializedObjectReferences List
+	 * Value: SubObject Information for the Object Reference
+	 * (This Container effectively only contains SubObjects so there will be indices not found here)
+	 */
+	UPROPERTY()
+	TMap<FAvaObjectIndex, FAvaSubObjectData> SubObjects;
+
+	/** Binds every entry in SerializedNames to its index to speed up look up */
+	UPROPERTY(Transient)
+	TMap<FName, FAvaObjectIndex> NameIndexMap;
+
+	/** Binds every entry in SerializedObjectReferences to its index to speed up look up */
+	UPROPERTY(Transient)
+	TMap<FSoftObjectPath, FAvaObjectIndex> ObjectReferenceIndexMap;
+
+	UPROPERTY(Transient)
+	TMap<FSoftObjectPath, TWeakObjectPtr<AActor>> CachedActors;
+
+	UPROPERTY(Transient)
+	TMap<FSoftObjectPath, TWeakObjectPtr<UObject>> CachedSubObjects;
+};
diff --git a/Engine/Plugins/Experimental/Avalanche/Source/Avalanche/Public/Data/AvalancheActorData.h b/Engine/Plugins/Experimental/Avalanche/Source/Avalanche/Public/Data/AvalancheActorData.h
deleted file mode 100644
diff --git a/Engine/Plugins/Experimental/Avalanche/Source/Avalanche/Public/Data/AvalancheComponentData.h b/Engine/Plugins/Experimental/Avalanche/Source/Avalanche/Public/Data/AvalancheComponentData.h
deleted file mode 100644
diff --git a/Engine/Plugins/Experimental/Avalanche/Source/Avalanche/Public/Data/AvalancheDataDefines.h b/Engine/Plugins/Experimental/Avalanche/Source/Avalanche/Public/Data/AvalancheDataDefines.h
deleted file mode 100644
diff --git a/Engine/Plugins/Experimental/Avalanche/Source/Avalanche/Public/Data/AvalancheObjectData.h b/Engine/Plugins/Experimental/Avalanche/Source/Avalanche/Public/Data/AvalancheObjectData.h
deleted file mode 100644
diff --git a/Engine/Plugins/Experimental/Avalanche/Source/Avalanche/Public/Data/AvalancheSubObjectData.h b/Engine/Plugins/Experimental/Avalanche/Source/Avalanche/Public/Data/AvalancheSubObjectData.h
deleted file mode 100644
diff --git a/Engine/Plugins/Experimental/Avalanche/Source/Avalanche/Public/Data/AvalancheVersionInfo.h b/Engine/Plugins/Experimental/Avalanche/Source/Avalanche/Public/Data/AvalancheVersionInfo.h
deleted file mode 100644
diff --git a/Engine/Plugins/Experimental/Avalanche/Source/Avalanche/Public/Data/AvalancheWorldData.h b/Engine/Plugins/Experimental/Avalanche/Source/Avalanche/Public/Data/AvalancheWorldData.h
deleted file mode 100644
diff --git a/Engine/Plugins/Experimental/Avalanche/Source/AvalancheEditor/Private/AvaEditorStyle.cpp b/Engine/Plugins/Experimental/Avalanche/Source/AvalancheEditor/Private/AvaEditorStyle.cpp
--- a/Engine/Plugins/Experimental/Avalanche/Source/AvalancheEditor/Private/AvaEditorStyle.cpp
+++ b/Engine/Plugins/Experimental/Avalanche/Source/AvalancheEditor/Private/AvaEditorStyle.cpp
@@ -69,9 +69,9 @@ void FAvaEditorStyle::Init()
 	SetCoreContentRoot(FPaths::EngineContentDir() / TEXT("Editor/Slate"));
 
 	// Custom Class Icons
-	Set("ClassThumbnail.AvalancheActor", new IMAGE_BRUSH("Icons/Icon64", Icon64));
+	Set("ClassThumbnail.AvaActor", new IMAGE_BRUSH("Icons/Icon64", Icon64));
 	Set("ClassThumbnail.AvalancheBlueprint", new IMAGE_BRUSH("Icons/Icon64", Icon64));
-	Set("ClassIcon.AvalancheActor", new IMAGE_BRUSH("Icons/Icon16", Icon16));
+	Set("ClassIcon.AvaActor", new IMAGE_BRUSH("Icons/Icon16", Icon16));
 	Set("ClassIcon.AvalancheBlueprint", new IMAGE_BRUSH("Icons/Icon16", Icon16));
 	Set("ClassIcon.AvaNullActor", new CORE_IMAGE_BRUSH(TEXT("Icons/SequencerIcons/icon_Sequencer_Move_24x"), Icon16));
 	Set("ClassIcon.AvaToolboxStarDynamicMesh", new IMAGE_BRUSH_SVG("Icons/ToolboxIcons/Favorite", Icon16));
