[Motion Design] Fix the issue of using a scene rig as a page.

There is a clear distinction between a level instance (playable) that is a page and one from a rig that is a sub-playable shared across many pages.

The situation becomes ambiguous if we allow playables to be used both as a page directly and as a rig. It is preferable to keep both type of instance distinct.

diff --git a/Engine/Plugins/Experimental/Avalanche/Source/AvalancheMedia/Private/Playback/AvaMediaPlayableGroup.cpp b/Engine/Plugins/Experimental/Avalanche/Source/AvalancheMedia/Private/Playback/AvaMediaPlayableGroup.cpp
--- a/Engine/Plugins/Experimental/Avalanche/Source/AvalancheMedia/Private/Playback/AvaMediaPlayableGroup.cpp
+++ b/Engine/Plugins/Experimental/Avalanche/Source/AvalancheMedia/Private/Playback/AvaMediaPlayableGroup.cpp
@@ -171,17 +171,16 @@ bool UAvaMediaPlayableGroup::HasPlayingPlayables() const
 	return false;
 }
 
-UAvalanchePlayable* UAvaMediaPlayableGroup::FindFirstPlayableBySourceAssetPath(const FSoftObjectPath& InSourceAssetPath) const
+void UAvaMediaPlayableGroup::FindPlayablesBySourceAssetPath(const FSoftObjectPath& InSourceAssetPath, TArray<UAvalanchePlayable*>& OutFoundPlayables) const
 {
 	for (const TObjectKey<UAvalanchePlayable>& PlayableKey : Playables)
 	{
 		UAvalanchePlayable* Playable = PlayableKey.ResolveObjectPtr();
 		if (Playable && Playable->GetSourceAssetPath() == InSourceAssetPath)
 		{
-			return Playable;
+			OutFoundPlayables.Add(Playable);
 		}
 	}
-	return nullptr;
 }
 
 void UAvaMediaPlayableGroup::RegisterPlayableTransition(UAvalanchePlayableTransition* InPlayableTransition)
diff --git a/Engine/Plugins/Experimental/Avalanche/Source/AvalancheMedia/Private/Playback/Playables/AvalancheLevelStreamingPlayable.cpp b/Engine/Plugins/Experimental/Avalanche/Source/AvalancheMedia/Private/Playback/Playables/AvalancheLevelStreamingPlayable.cpp
--- a/Engine/Plugins/Experimental/Avalanche/Source/AvalancheMedia/Private/Playback/Playables/AvalancheLevelStreamingPlayable.cpp
+++ b/Engine/Plugins/Experimental/Avalanche/Source/AvalancheMedia/Private/Playback/Playables/AvalancheLevelStreamingPlayable.cpp
@@ -677,13 +677,31 @@ void UAvalancheLevelStreamingPlayable::GetOrLoadSubPlayable(const ULevelStreamin
 	}
 
 	const FSoftObjectPath SourceAssetPath = InLevelStreaming->GetWorldAsset().ToSoftObjectPath();
+
+	// Check already loaded sub-playables
+	for (const TObjectPtr<UAvalancheLevelStreamingPlayable>& SubPlayable : SubPlayables)
+	{
+		if (SubPlayable->GetSourceAssetPath() == SourceAssetPath)
+		{
+			return; // Already loaded.
+		}
+	}
 	
-	// Assume sub playables are shared, i.e. unique instance per group. For now.
-	// Need to see with Brad how to determine instancing scope (i.e. global vs local).	
-	if (UAvalancheLevelStreamingPlayable* ExistingPlayable = Cast<UAvalancheLevelStreamingPlayable>(PlayableGroup->FindFirstPlayableBySourceAssetPath(SourceAssetPath)))
+	// For now, sub playables are shared globally, i.e. unique instance per group.
+	// Todo: We could support instancing scope (i.e. global vs local). Would require additional asset/scene info.
+	TArray<UAvalanchePlayable*> FoundPlayables;
+	PlayableGroup->FindPlayablesBySourceAssetPath(SourceAssetPath, FoundPlayables);
+	for (UAvalanchePlayable* FoundPlayable : FoundPlayables)
 	{
-		AddSubPlayable(ExistingPlayable);
-		return;
+		if (UAvalancheLevelStreamingPlayable* ExistingPlayable = Cast<UAvalancheLevelStreamingPlayable>(FoundPlayable))
+		{
+			// Only use the existing playable if it is a sub-playable already.
+			if (ExistingPlayable->HasParentPlayables())
+			{
+				AddSubPlayable(ExistingPlayable);
+				return;
+			}
+		}
 	}
 
 	if (UAvalancheLevelStreamingPlayable* NewPlayable = CreateSubPlayable(PlayableGroup, SourceAssetPath))
diff --git a/Engine/Plugins/Experimental/Avalanche/Source/AvalancheMedia/Public/Playback/AvaMediaPlayableGroup.h b/Engine/Plugins/Experimental/Avalanche/Source/AvalancheMedia/Public/Playback/AvaMediaPlayableGroup.h
--- a/Engine/Plugins/Experimental/Avalanche/Source/AvalancheMedia/Public/Playback/AvaMediaPlayableGroup.h
+++ b/Engine/Plugins/Experimental/Avalanche/Source/AvalancheMedia/Public/Playback/AvaMediaPlayableGroup.h
@@ -66,7 +66,8 @@ public:
 	/** Returns true if there are any valid registered playables that are currently playing. */
 	bool HasPlayingPlayables() const;
 
-	UAvalanchePlayable* FindFirstPlayableBySourceAssetPath(const FSoftObjectPath& InSourceAssetPath) const;
+	/** Finds all the playables that are instances of the given source asset. */
+	void FindPlayablesBySourceAssetPath(const FSoftObjectPath& InSourceAssetPath, TArray<UAvalanchePlayable*>& OutFoundPlayables) const;
 
 	void RegisterPlayableTransition(UAvalanchePlayableTransition* InPlayableTransition);
 	
