Avalanche Transition:
- Moved Tree Utility functions out of EditorViewModel and into EditorUtils
- Decoupled Compiler from EditorViewModel so that others outside the EditorViewModel can use the compiler without needing an EditorViewModel instance (as in the case of the Editor Module)

diff --git a/Engine/Plugins/Experimental/Avalanche/Source/AvalancheTransitionEditor/Private/AvaTransitionEditorModule.cpp b/Engine/Plugins/Experimental/Avalanche/Source/AvalancheTransitionEditor/Private/AvaTransitionEditorModule.cpp
--- a/Engine/Plugins/Experimental/Avalanche/Source/AvalancheTransitionEditor/Private/AvaTransitionEditorModule.cpp
+++ b/Engine/Plugins/Experimental/Avalanche/Source/AvalancheTransitionEditor/Private/AvaTransitionEditorModule.cpp
@@ -2,6 +2,7 @@
 
 #include "AvaTransitionEditorModule.h"
 #include "AvaTransitionCommands.h"
+#include "AvaTransitionEditorEnums.h"
 #include "AvaTransitionEditorLog.h"
 #include "AvaTransitionEditorStyle.h"
 #include "AvaTransitionEditorUtils.h"
@@ -134,11 +135,10 @@ void FAvaTransitionEditorModule::ValidateStateTree(UAvaTransitionTree* InTransit
 		EditorData->AddRootState();	
 	}
 
-	// Use Editor View Model to Compile
-	// todo: improve Compiler class so that an Editor View Model isn't needed
-	TSharedRef<FAvaTransitionEditorViewModel> EditorViewModel = MakeShared<FAvaTransitionEditorViewModel>(InTransitionTree, /*EditorToolkit*/nullptr);
-	EditorViewModel->Initialize(/*Parent*/nullptr);
-	EditorViewModel->Compile();
+	/** Compile in Advanced Mode here so that no new nodes are generated from outside */
+	FAvaTransitionCompiler Compiler;
+	Compiler.SetTransitionTree(InTransitionTree);
+	Compiler.Compile(EAvaTransitionEditorMode::Advanced);
 }
 
 void FAvaTransitionEditorModule::OnPostCompile(const UStateTree& InStateTree)
diff --git a/Engine/Plugins/Experimental/Avalanche/Source/AvalancheTransitionEditor/Private/AvaTransitionEditorUtils.cpp b/Engine/Plugins/Experimental/Avalanche/Source/AvalancheTransitionEditor/Private/AvaTransitionEditorUtils.cpp
--- a/Engine/Plugins/Experimental/Avalanche/Source/AvalancheTransitionEditor/Private/AvaTransitionEditorUtils.cpp
+++ b/Engine/Plugins/Experimental/Avalanche/Source/AvalancheTransitionEditor/Private/AvaTransitionEditorUtils.cpp
@@ -1,6 +1,7 @@
 // Copyright Epic Games, Inc. All Rights Reserved.
 
 #include "AvaTransitionEditorUtils.h"
+#include "AvaTransitionEditorLog.h"
 #include "AvaTransitionTree.h"
 #include "AvaTransitionTreeEditorData.h"
 #include "Behavior/IAvaTransitionBehavior.h"
@@ -10,6 +11,8 @@
 #include "Modules/ModuleManager.h"
 #include "PropertyEditorModule.h"
 #include "ScopedTransaction.h"
+#include "Serialization/ArchiveObjectCrc32.h"
+#include "StateTreeTaskBase.h"
 #include "Widgets/Layout/SBox.h"
 #include "Widgets/Layout/SScaleBox.h"
 
@@ -56,6 +59,165 @@ TSharedPtr<SWidget> CreateTransitionLayerPicker(UAvaTransitionTreeEditorData* In
 	return nullptr;
 }
 
+void ValidateTree(UAvaTransitionTree& InTransitionTree)
+{
+	UStateTreeEditorData* EditorData = Cast<UStateTreeEditorData>(InTransitionTree.EditorData);
+	if (!EditorData || !EditorData->Schema)
+	{
+		return;
+	}
+
+	const FString TreeDebugName = InTransitionTree.GetName();
+
+	EditorData->ReparentStates();
+
+	// Clear evaluators if not allowed.
+	if (!EditorData->Evaluators.IsEmpty() && !EditorData->Schema->AllowEvaluators())
+	{
+		UE_LOG(LogAvaEditorTransition, Warning
+			, TEXT("%s: Resetting Evaluators due to current schema restrictions.")
+			, *TreeDebugName);
+
+		EditorData->Evaluators.Reset();
+	}
+
+	// Apply Schema Rules to each State 
+	EditorData->VisitHierarchy([&TreeDebugName, EditorData](UStateTreeState& State, UStateTreeState*)
+	{
+		// Clear enter conditions if not allowed.
+		if (!State.EnterConditions.IsEmpty() && !EditorData->Schema->AllowEnterConditions())
+		{
+			UE_LOG(LogAvaEditorTransition, Warning
+				, TEXT("%s: Resetting Enter Conditions in state %s due to current schema restrictions.")
+				, *TreeDebugName
+				, *State.GetName());
+
+			State.EnterConditions.Reset();
+		}
+
+		// Keep single and many tasks based on what is allowed.
+		if (!EditorData->Schema->AllowMultipleTasks())
+		{
+			if (!State.Tasks.IsEmpty())
+			{
+				State.Tasks.Reset();
+				UE_LOG(LogAvaEditorTransition, Warning
+					, TEXT("%s: Resetting Tasks in state %s due to current schema restrictions.")
+					, *TreeDebugName
+					, *State.GetName());
+			}
+
+			// Task name is the same as state name.
+			if (FStateTreeTaskBase* Task = State.SingleTask.Node.GetMutablePtr<FStateTreeTaskBase>())
+			{
+				Task->Name = State.Name;
+			}
+		}
+		else
+		{
+			if (State.SingleTask.Node.IsValid())
+			{
+				State.SingleTask.Reset();
+				UE_LOG(LogAvaEditorTransition, Warning
+					, TEXT("%s: Resetting Single Task in state %s due to current schema restrictions.")
+					, *TreeDebugName
+					, *State.GetName());
+			}
+		}
+
+		return EStateTreeVisitor::Continue;
+	});
+
+	// Remove unused Bindings
+	{
+		TMap<FGuid, const FStateTreeDataView> AllStructValues;
+		EditorData->GetAllStructValues(AllStructValues);
+		EditorData->GetPropertyEditorBindings()->RemoveUnusedBindings(AllStructValues);
+	}
+
+	// Validate Linked States
+	{
+		// Make sure all state links are valid and update the names if needed.
+		// Create ID to state name map.
+		TMap<FGuid, FName> IdToName;
+
+		EditorData->VisitHierarchy([&IdToName](const UStateTreeState& State, UStateTreeState* /*ParentState*/)
+		{
+			IdToName.Add(State.ID, State.Name);
+			return EStateTreeVisitor::Continue;
+		});
+
+		static auto FixChangedStateLinkName = [](FStateTreeStateLink& StateLink, const TMap<FGuid, FName>& IDToName)
+		{
+			if (StateLink.ID.IsValid())
+			{
+				const FName* Name = IDToName.Find(StateLink.ID);
+				if (Name == nullptr)
+				{
+					// Missing link, we'll show these in the UI
+					return false;
+				}
+				if (StateLink.Name != *Name)
+				{
+					// Name changed, fix!
+					StateLink.Name = *Name;
+					return true;
+				}
+			}
+			return false;
+		};
+
+		// Fix changed names.
+		EditorData->VisitHierarchy([&IdToName](UStateTreeState& State, UStateTreeState* /*ParentState*/)
+		{
+			if (State.Type == EStateTreeStateType::Linked)
+			{
+				FixChangedStateLinkName(State.LinkedSubtree, IdToName);
+			}
+					
+			for (FStateTreeTransition& Transition : State.Transitions)
+			{
+				FixChangedStateLinkName(Transition.State, IdToName);
+			}
+
+			return EStateTreeVisitor::Continue;
+		});
+	}
+
+	// Update Linked State Parameters
+	EditorData->VisitHierarchy([](UStateTreeState& State, UStateTreeState* /*ParentState*/)
+	{
+		if (State.Type == EStateTreeStateType::Linked)
+		{
+			State.UpdateParametersFromLinkedSubtree();
+		}
+		return EStateTreeVisitor::Continue;
+	});
+}
+
+uint32 CalculateTreeHash(UAvaTransitionTree& InTransitionTree)
+{
+	if (!InTransitionTree.EditorData)
+	{
+		return 0;
+	}
+
+	static const FName MD_ExcludeFromHash(TEXT("ExcludeFromHash"));
+
+	class : public FArchiveObjectCrc32
+	{
+		virtual bool ShouldSkipProperty(const FProperty* InProperty) const override
+		{
+			return !InProperty
+				|| FArchiveObjectCrc32::ShouldSkipProperty(InProperty)
+				|| InProperty->HasAllPropertyFlags(CPF_Transient)
+				|| InProperty->HasMetaData(MD_ExcludeFromHash);
+		}
+	} Archive;
+
+	return Archive.Crc32(InTransitionTree.EditorData, 0);
+}
+
 bool PickTransitionTreeAsset(const FText& InDialogTitle, UAvaTransitionTree*& OutTransitionTree)
 {
 	FOpenAssetDialogConfig SelectAssetConfig;
diff --git a/Engine/Plugins/Experimental/Avalanche/Source/AvalancheTransitionEditor/Private/AvaTransitionEditorUtils.h b/Engine/Plugins/Experimental/Avalanche/Source/AvalancheTransitionEditor/Private/AvaTransitionEditorUtils.h
--- a/Engine/Plugins/Experimental/Avalanche/Source/AvalancheTransitionEditor/Private/AvaTransitionEditorUtils.h
+++ b/Engine/Plugins/Experimental/Avalanche/Source/AvalancheTransitionEditor/Private/AvaTransitionEditorUtils.h
@@ -14,6 +14,10 @@ namespace UE::AvaTransitionEditor
 {
 	TSharedPtr<SWidget> CreateTransitionLayerPicker(UAvaTransitionTreeEditorData* InEditorData);
 
+	void ValidateTree(UAvaTransitionTree& InTransitionTree);
+
+	uint32 CalculateTreeHash(UAvaTransitionTree& InTransitionTree);
+
 	bool PickTransitionTreeAsset(const FText& InDialogTitle, UAvaTransitionTree*& OutTransitionTree);
 
 	void ToggleTransitionTreeEnabled(TWeakObjectPtr<UAvaTransitionTree> InTransitionTreeWeak);
diff --git a/Engine/Plugins/Experimental/Avalanche/Source/AvalancheTransitionEditor/Private/Compiler/AvaTransitionCompiler.cpp b/Engine/Plugins/Experimental/Avalanche/Source/AvalancheTransitionEditor/Private/Compiler/AvaTransitionCompiler.cpp
--- a/Engine/Plugins/Experimental/Avalanche/Source/AvalancheTransitionEditor/Private/Compiler/AvaTransitionCompiler.cpp
+++ b/Engine/Plugins/Experimental/Avalanche/Source/AvalancheTransitionEditor/Private/Compiler/AvaTransitionCompiler.cpp
@@ -3,6 +3,7 @@
 #include "AvaTransitionCompiler.h"
 #include "AvaTransitionCommands.h"
 #include "AvaTransitionEditor.h"
+#include "AvaTransitionEditorUtils.h"
 #include "AvaTransitionSelection.h"
 #include "AvaTransitionTree.h"
 #include "AvaTransitionTreeEditorData.h"
@@ -17,48 +18,64 @@
 #include "StateTreeDelegates.h"
 #include "StateTreeEditorSettings.h"
 #include "StateTreeState.h"
-#include "TabFactories/AvaTransitionCompilerResultsTabFactory.h"
 #include "ViewModels/AvaTransitionEditorViewModel.h"
 #include "ViewModels/AvaTransitionViewModelSharedData.h"
 #include "ViewModels/Registry/AvaTransitionViewModelRegistryCollection.h"
 
 #define LOCTEXT_NAMESPACE "AvaTransitionCompiler"
 
-FAvaTransitionCompiler::FAvaTransitionCompiler(FAvaTransitionEditorViewModel& InOwner)
-	: Owner(InOwner)
+namespace UE::AvaTransitionEditor::Private
 {
-	FMessageLogModule& MessageLogModule = FModuleManager::LoadModuleChecked<FMessageLogModule>("MessageLog");
+	TSharedRef<IMessageLogListing> CreateCompilerResultsListing()
+	{
+		FMessageLogModule& MessageLogModule = FModuleManager::LoadModuleChecked<FMessageLogModule>("MessageLog");
+
+		// Show Pages so that user is never allowed to clear log messages
+		FMessageLogInitializationOptions LogOptions;
+		LogOptions.bShowPages   = false;
+		LogOptions.bShowFilters = false;
+		LogOptions.bAllowClear  = false;
+		LogOptions.MaxPageCount = 1;
+
+		return MessageLogModule.CreateLogListing("AvaTransitionTreeCompiler", LogOptions);
+	}
+}
 
-	// Show Pages so that user is never allowed to clear log messages
-	FMessageLogInitializationOptions LogOptions;
-	LogOptions.bShowPages   = false;
-	LogOptions.bShowFilters = false;
-	LogOptions.bAllowClear  = false;
-	LogOptions.MaxPageCount = 1;
+FAvaTransitionCompiler::FAvaTransitionCompiler()
+	: CompilerResultsListing(UE::AvaTransitionEditor::Private::CreateCompilerResultsListing())
+{
+}
+
+void FAvaTransitionCompiler::SetTransitionTree(UAvaTransitionTree* InTransitionTree)
+{
+	TransitionTreeWeak = InTransitionTree;
+}
 
-	CompilerResultsListing = MessageLogModule.CreateLogListing("AvaTransitionTreeCompiler", LogOptions);
-	CompilerResultsWidget  = MessageLogModule.CreateLogListingWidget(CompilerResultsListing.ToSharedRef());
+IMessageLogListing& FAvaTransitionCompiler::GetCompilerResultsListing()
+{
+	return CompilerResultsListing.Get();
+}
 
-	CompilerResultsListing->OnMessageTokenClicked().AddRaw(this, &FAvaTransitionCompiler::OnMessageTokenClicked);
+FSimpleDelegate& FAvaTransitionCompiler::GetOnCompileFailed()
+{
+	return OnCompileFail;
 }
 
-bool FAvaTransitionCompiler::Compile()
+bool FAvaTransitionCompiler::Compile(EAvaTransitionEditorMode InCompileMode)
 {
-	UAvaTransitionTree* TransitionTree = Owner.GetTransitionTree();
+	UAvaTransitionTree* TransitionTree = TransitionTreeWeak.Get();
 	if (!TransitionTree)
 	{
 		return false;
 	}
 
-	if (UAvaTransitionTreeEditorData* EditorData = Owner.GetEditorData())
+	if (UAvaTransitionTreeEditorData* EditorData = Cast<UAvaTransitionTreeEditorData>(TransitionTree->EditorData))
 	{
 		// Update Tree's Transition Layer
 		TransitionTree->SetTransitionLayer(EditorData->GetTransitionLayer());
 
-		const EAvaTransitionEditorMode EditorMode = Owner.GetSharedData()->GetEditorMode();
-
 		// Only allow extensible compilation when it's not in Advanced Mode
-		if (EditorMode != EAvaTransitionEditorMode::Advanced)
+		if (InCompileMode != EAvaTransitionEditorMode::Advanced)
 		{
 			// Remove all invalid Sub Trees, if any
 			EditorData->SubTrees.RemoveAll(
@@ -71,23 +88,20 @@ bool FAvaTransitionCompiler::Compile()
 		}
 	}
 
-	Owner.UpdateTree();
+	UpdateTree();
 
 	FStateTreeCompilerLog Log;
 	FStateTreeCompiler Compiler(Log);
 
 	bLastCompileSucceeded = Compiler.Compile(*TransitionTree);
 
-	if (CompilerResultsListing.IsValid())
-	{
-		CompilerResultsListing->ClearMessages();
-		Log.AppendToLog(CompilerResultsListing.Get());
-	}
+	CompilerResultsListing->ClearMessages();
+	Log.AppendToLog(&CompilerResultsListing.Get());
 
 	if (bLastCompileSucceeded)
 	{
 		// Success
-		TransitionTree->LastCompiledEditorDataHash = Owner.GetEditorDataHash();
+		TransitionTree->LastCompiledEditorDataHash = EditorDataHash;
 		UE::StateTree::Delegates::OnPostCompile.Broadcast(*TransitionTree);
 	}
 	else
@@ -95,14 +109,7 @@ bool FAvaTransitionCompiler::Compile()
 		// Make sure not to leave stale data on failed compile.
 		TransitionTree->ResetCompiled();
 		TransitionTree->LastCompiledEditorDataHash = 0;
-
-		if (TSharedPtr<FAvaTransitionEditor> Editor = Owner.GetEditor())
-		{
-			if (TSharedPtr<FTabManager> TabManager = Editor->GetTabManager())
-			{
-				TabManager->TryInvokeTab(FAvaTransitionCompilerResultsTabFactory::TabId);
-			}
-		}
+		OnCompileFail.ExecuteIfBound();
 	}
 
 	const UStateTreeEditorSettings* Settings = GetMutableDefault<UStateTreeEditorSettings>();
@@ -120,6 +127,19 @@ bool FAvaTransitionCompiler::Compile()
 	return true;
 }
 
+void FAvaTransitionCompiler::UpdateTree()
+{
+	UAvaTransitionTree* TransitionTree = TransitionTreeWeak.Get();
+	if (!TransitionTree)
+	{
+		return;
+	}
+
+	UE::AvaTransitionEditor::ValidateTree(*TransitionTree);
+
+	EditorDataHash = UE::AvaTransitionEditor::CalculateTreeHash(*TransitionTree);
+}
+
 FSlateIcon FAvaTransitionCompiler::GetCompileStatusIcon() const
 {
 	static const FName CompileStatusBackground("Blueprint.CompileStatus.Background");
@@ -128,14 +148,12 @@ FSlateIcon FAvaTransitionCompiler::GetCompileStatusIcon() const
 	static const FName CompileStatusGood("Blueprint.CompileStatus.Overlay.Good");
 	static const FName CompileStatusWarning("Blueprint.CompileStatus.Overlay.Warning");
 
-	UAvaTransitionTree* TransitionTree = Owner.GetTransitionTree();
+	UAvaTransitionTree* TransitionTree = TransitionTreeWeak.Get();
 	if (!TransitionTree)
 	{
 		return FSlateIcon(FAppStyle::GetAppStyleSetName(), CompileStatusBackground, NAME_None, CompileStatusUnknown);
 	}
 
-	const uint32 EditorDataHash = Owner.GetEditorDataHash();
-
 	const bool bCompiledDataResetDuringLoad = TransitionTree->LastCompiledEditorDataHash == EditorDataHash && !TransitionTree->IsReadyToRun();
 
 	if (!bLastCompileSucceeded || bCompiledDataResetDuringLoad)
@@ -151,9 +169,10 @@ FSlateIcon FAvaTransitionCompiler::GetCompileStatusIcon() const
 	return FSlateIcon(FAppStyle::GetAppStyleSetName(), CompileStatusBackground, NAME_None, CompileStatusGood);
 }
 
-TSharedRef<SWidget> FAvaTransitionCompiler::GetCompilerResultsWidget() const
+TSharedRef<SWidget> FAvaTransitionCompiler::CreateCompilerResultsWidget() const
 {
-	return CompilerResultsWidget.ToSharedRef();
+	FMessageLogModule& MessageLogModule = FModuleManager::LoadModuleChecked<FMessageLogModule>("MessageLog");
+	return MessageLogModule.CreateLogListingWidget(CompilerResultsListing);
 }
 
 void FAvaTransitionCompiler::SetSaveOnCompile(EStateTreeSaveOnCompile InSaveOnCompileType)
@@ -190,22 +209,4 @@ void FAvaTransitionCompiler::GenerateCompileOptionsMenu(UToolMenu* InMenu)
 		, FNewToolMenuDelegate::CreateLambda(MakeSaveOnCompileMenu));
 }
 
-void FAvaTransitionCompiler::OnMessageTokenClicked(const TSharedRef<IMessageToken>& InMessageToken)
-{
-	if (InMessageToken->GetType() == EMessageToken::Object)
-	{
-		const TSharedRef<FUObjectToken> ObjectToken = StaticCastSharedRef<FUObjectToken>(InMessageToken);
-
-		if (UStateTreeState* State = Cast<UStateTreeState>(ObjectToken->GetObject().Get()))
-		{
-			TSharedRef<FAvaTransitionViewModelRegistryCollection> RegistryCollection = Owner.GetSharedData()->GetRegistryCollection();
-
-			if (TSharedPtr<FAvaTransitionViewModel> FoundViewModel = RegistryCollection->FindViewModel(State))
-			{
-				Owner.GetSelection()->SetSelectedItems({ FoundViewModel });
-			}
-		}
-	}
-}
-
 #undef LOCTEXT_NAMESPACE
diff --git a/Engine/Plugins/Experimental/Avalanche/Source/AvalancheTransitionEditor/Private/Compiler/AvaTransitionCompiler.h b/Engine/Plugins/Experimental/Avalanche/Source/AvalancheTransitionEditor/Private/Compiler/AvaTransitionCompiler.h
--- a/Engine/Plugins/Experimental/Avalanche/Source/AvalancheTransitionEditor/Private/Compiler/AvaTransitionCompiler.h
+++ b/Engine/Plugins/Experimental/Avalanche/Source/AvalancheTransitionEditor/Private/Compiler/AvaTransitionCompiler.h
@@ -2,26 +2,36 @@
 
 #pragma once
 
+#include "Delegates/Delegate.h"
 #include "Templates/SharedPointer.h"
+#include "UObject/WeakObjectPtr.h"
 
-class FAvaTransitionEditorViewModel;
 class IMessageLogListing;
-class IMessageToken;
 class SWidget;
+class UAvaTransitionTree;
 class UToolMenu;
+enum class EAvaTransitionEditorMode : uint8;
 enum class EStateTreeSaveOnCompile : uint8;
 struct FSlateIcon;
 
 class FAvaTransitionCompiler
 {
 public:
-	FAvaTransitionCompiler(FAvaTransitionEditorViewModel& InOwner);
+	FAvaTransitionCompiler();
 
-	bool Compile();
+	void SetTransitionTree(UAvaTransitionTree* InTransitionTree);
+
+	IMessageLogListing& GetCompilerResultsListing();
+
+	FSimpleDelegate& GetOnCompileFailed();
+
+	bool Compile(EAvaTransitionEditorMode InCompileMode);
+
+	void UpdateTree();
 
 	FSlateIcon GetCompileStatusIcon() const;
 
-	TSharedRef<SWidget> GetCompilerResultsWidget() const;
+	TSharedRef<SWidget> CreateCompilerResultsWidget() const;
 
 	static void SetSaveOnCompile(EStateTreeSaveOnCompile InSaveOnCompileType);
 
@@ -30,13 +40,13 @@ public:
 	static void GenerateCompileOptionsMenu(UToolMenu* InMenu);
 
 private:
-	void OnMessageTokenClicked(const TSharedRef<IMessageToken>& InMessageToken);
+	TWeakObjectPtr<UAvaTransitionTree> TransitionTreeWeak;
 
-	FAvaTransitionEditorViewModel& Owner;
+	TSharedRef<IMessageLogListing> CompilerResultsListing;
 
-	TSharedPtr<IMessageLogListing> CompilerResultsListing;
+	FSimpleDelegate OnCompileFail;
 
-	TSharedPtr<SWidget> CompilerResultsWidget;
+	uint32 EditorDataHash = 0;
 
 	bool bLastCompileSucceeded = true;
 };
diff --git a/Engine/Plugins/Experimental/Avalanche/Source/AvalancheTransitionEditor/Private/TabFactories/AvaTransitionCompilerResultsTabFactory.cpp b/Engine/Plugins/Experimental/Avalanche/Source/AvalancheTransitionEditor/Private/TabFactories/AvaTransitionCompilerResultsTabFactory.cpp
--- a/Engine/Plugins/Experimental/Avalanche/Source/AvalancheTransitionEditor/Private/TabFactories/AvaTransitionCompilerResultsTabFactory.cpp
+++ b/Engine/Plugins/Experimental/Avalanche/Source/AvalancheTransitionEditor/Private/TabFactories/AvaTransitionCompilerResultsTabFactory.cpp
@@ -34,7 +34,7 @@ TSharedRef<SWidget> FAvaTransitionCompilerResultsTabFactory::CreateTabBody(const
 		return SNullWidget::NullWidget;
 	}
 
-	return EditorViewModel->GetCompiler().GetCompilerResultsWidget();
+	return EditorViewModel->GetCompiler().CreateCompilerResultsWidget();
 }
 
 #undef LOCTEXT_NAMESPACE
diff --git a/Engine/Plugins/Experimental/Avalanche/Source/AvalancheTransitionEditor/Private/ViewModels/AvaTransitionEditorViewModel.cpp b/Engine/Plugins/Experimental/Avalanche/Source/AvalancheTransitionEditor/Private/ViewModels/AvaTransitionEditorViewModel.cpp
--- a/Engine/Plugins/Experimental/Avalanche/Source/AvalancheTransitionEditor/Private/ViewModels/AvaTransitionEditorViewModel.cpp
+++ b/Engine/Plugins/Experimental/Avalanche/Source/AvalancheTransitionEditor/Private/ViewModels/AvaTransitionEditorViewModel.cpp
@@ -12,13 +12,15 @@
 #include "AvaTransitionTreeEditorData.h"
 #include "AvaTransitionViewModelChildren.h"
 #include "AvaTransitionViewModelSharedData.h"
+#include "IMessageLogListing.h"
 #include "Menu/AvaTransitionToolbar.h"
 #include "Menu/AvaTransitionTreeContextMenu.h"
-#include "Serialization/ArchiveObjectCrc32.h"
+#include "Misc/UObjectToken.h"
+#include "Registry/AvaTransitionViewModelRegistryCollection.h"
 #include "State/AvaTransitionStateViewModel.h"
 #include "StateTreeDelegates.h"
 #include "StateTreeEditorSettings.h"
-#include "StateTreeTaskBase.h"
+#include "TabFactories/AvaTransitionCompilerResultsTabFactory.h"
 #include "Views/SAvaTransitionTreeView.h"
 #include "Widgets/Layout/SScrollBar.h"
 #include "Widgets/Layout/SScrollBox.h"
@@ -30,9 +32,9 @@ FAvaTransitionEditorViewModel::FAvaTransitionEditorViewModel(UAvaTransitionTree*
 	, ContextMenu(MakeShared<FAvaTransitionTreeContextMenu>(*this))
 	, TransitionTreeWeak(InTransitionTree)
 	, EditorWeak(InEditor)
-	, Compiler(*this)
 	, CommandList(MakeShared<FUICommandList>())
 {
+	Compiler.SetTransitionTree(InTransitionTree);
 }
 
 FAvaTransitionEditorViewModel::~FAvaTransitionEditorViewModel()
@@ -77,7 +79,7 @@ void FAvaTransitionEditorViewModel::Compile()
 	if (CanCompile())
 	{
 		UpdateTree();
-		Compiler.Compile();
+		Compiler.Compile(GetSharedData()->GetEditorMode());
 	}
 }
 
@@ -116,14 +118,7 @@ bool FAvaTransitionEditorViewModel::UpdateEditorData(bool bInCreateIfNotFound)
 
 void FAvaTransitionEditorViewModel::UpdateTree()
 {
-	UAvaTransitionTree* TransitionTree = GetTransitionTree();
-	if (!TransitionTree)
-	{
-		return;
-	}
-
-	ValidateTree(TransitionTree->GetName());
-	EditorDataHash = CalculateTreeHash(*TransitionTree);
+	Compiler.UpdateTree();
 }
 
 TSharedPtr<FAvaTransitionEditor> FAvaTransitionEditorViewModel::GetEditor() const
@@ -250,10 +245,11 @@ void FAvaTransitionEditorViewModel::BindDelegates()
 		EditorData->GetOnTreeRequestRefresh().AddSP(this, &FAvaTransitionEditorViewModel::Refresh);
 	}
 
+	Compiler.GetCompilerResultsListing().OnMessageTokenClicked().AddSP(this, &FAvaTransitionEditorViewModel::OnMessageTokenClicked);
+	Compiler.GetOnCompileFailed().BindSP(this, &FAvaTransitionEditorViewModel::OnCompileFailed);
+
 	UE::StateTree::Delegates::OnIdentifierChanged.AddSP(this, &FAvaTransitionEditorViewModel::OnIdentifierChanged);
 	UE::StateTree::Delegates::OnSchemaChanged.AddSP(this, &FAvaTransitionEditorViewModel::OnSchemaChanged);
-	UE::StateTree::Delegates::OnParametersChanged.AddSP(this, &FAvaTransitionEditorViewModel::OnParametersChanged);
-	UE::StateTree::Delegates::OnStateParametersChanged.AddSP(this, &FAvaTransitionEditorViewModel::OnStateParametersChanged);
 }
 
 void FAvaTransitionEditorViewModel::UnbindDelegates()
@@ -263,194 +259,55 @@ void FAvaTransitionEditorViewModel::UnbindDelegates()
 		EditorData->GetOnTreeRequestRefresh().RemoveAll(this);
 	}
 
+	Compiler.GetCompilerResultsListing().OnMessageTokenClicked().RemoveAll(this);
+
 	UE::StateTree::Delegates::OnIdentifierChanged.RemoveAll(this);
 	UE::StateTree::Delegates::OnSchemaChanged.RemoveAll(this);
-	UE::StateTree::Delegates::OnParametersChanged.RemoveAll(this);
-	UE::StateTree::Delegates::OnStateParametersChanged.RemoveAll(this);
 }
 
-uint32 FAvaTransitionEditorViewModel::CalculateTreeHash(const UAvaTransitionTree& InTree) const
+void FAvaTransitionEditorViewModel::OnIdentifierChanged(const UStateTree& InStateTree)
 {
-	if (!InTree.EditorData)
+	if (&InStateTree == GetTransitionTree())
 	{
-		return 0;
+		UpdateTree();
 	}
-
-	static const FName MD_ExcludeFromHash(TEXT("ExcludeFromHash"));
-
-	class : public FArchiveObjectCrc32
-	{
-		virtual bool ShouldSkipProperty(const FProperty* InProperty) const override
-		{
-			return !InProperty
-				|| FArchiveObjectCrc32::ShouldSkipProperty(InProperty)
-				|| InProperty->HasAllPropertyFlags(CPF_Transient)
-				|| InProperty->HasMetaData(MD_ExcludeFromHash);
-		}
-	} Archive;
-
-	return Archive.Crc32(InTree.EditorData, 0);
 }
 
-void FAvaTransitionEditorViewModel::ValidateTree(const FString& InTreeDebugName)
+void FAvaTransitionEditorViewModel::OnSchemaChanged(const UStateTree& InStateTree)
 {
-	UAvaTransitionTreeEditorData* EditorData = GetEditorData();
-	if (!EditorData || !EditorData->Schema)
-	{
-		return;
-	}
-
-	EditorData->ReparentStates();
-
-	// Clear evaluators if not allowed.
-	if (!EditorData->Evaluators.IsEmpty() && !EditorData->Schema->AllowEvaluators())
+	if (&InStateTree == GetTransitionTree())
 	{
-		UE_LOG(LogAvaEditorTransition, Warning
-			, TEXT("%s: Resetting Evaluators due to current schema restrictions.")
-			, *InTreeDebugName);
-
-		EditorData->Evaluators.Reset();
+		UpdateTree();
 	}
+}
 
-	// Apply Schema Rules to each State 
-	EditorData->VisitHierarchy([&InTreeDebugName, EditorData](UStateTreeState& State, UStateTreeState*)
+void FAvaTransitionEditorViewModel::OnCompileFailed()
+{
+	if (TSharedPtr<FAvaTransitionEditor> Editor = GetEditor())
 	{
-		// Clear enter conditions if not allowed.
-		if (!State.EnterConditions.IsEmpty() && !EditorData->Schema->AllowEnterConditions())
-		{
-			UE_LOG(LogAvaEditorTransition, Warning
-				, TEXT("%s: Resetting Enter Conditions in state %s due to current schema restrictions.")
-				, *InTreeDebugName
-				, *State.GetName());
-
-			State.EnterConditions.Reset();
-		}
-
-		// Keep single and many tasks based on what is allowed.
-		if (!EditorData->Schema->AllowMultipleTasks())
-		{
-			if (!State.Tasks.IsEmpty())
-			{
-				State.Tasks.Reset();
-				UE_LOG(LogAvaEditorTransition, Warning
-					, TEXT("%s: Resetting Tasks in state %s due to current schema restrictions.")
-					, *InTreeDebugName
-					, *State.GetName());
-			}
-
-			// Task name is the same as state name.
-			if (FStateTreeTaskBase* Task = State.SingleTask.Node.GetMutablePtr<FStateTreeTaskBase>())
-			{
-				Task->Name = State.Name;
-			}
-		}
-		else
+		if (TSharedPtr<FTabManager> TabManager = Editor->GetTabManager())
 		{
-			if (State.SingleTask.Node.IsValid())
-			{
-				State.SingleTask.Reset();
-				UE_LOG(LogAvaEditorTransition, Warning
-					, TEXT("%s: Resetting Single Task in state %s due to current schema restrictions.")
-					, *InTreeDebugName
-					, *State.GetName());
-			}
+			TabManager->TryInvokeTab(FAvaTransitionCompilerResultsTabFactory::TabId);
 		}
-
-		return EStateTreeVisitor::Continue;
-	});
-
-	// Remove unused Bindings
-	{
-		TMap<FGuid, const FStateTreeDataView> AllStructValues;
-		EditorData->GetAllStructValues(AllStructValues);
-		EditorData->GetPropertyEditorBindings()->RemoveUnusedBindings(AllStructValues);
-	}
-
-	// Validate Linked States
-	{
-		// Make sure all state links are valid and update the names if needed.
-		// Create ID to state name map.
-		TMap<FGuid, FName> IdToName;
-
-		EditorData->VisitHierarchy([&IdToName](const UStateTreeState& State, UStateTreeState* /*ParentState*/)
-		{
-			IdToName.Add(State.ID, State.Name);
-			return EStateTreeVisitor::Continue;
-		});
-
-		static auto FixChangedStateLinkName = [](FStateTreeStateLink& StateLink, const TMap<FGuid, FName>& IDToName)
-		{
-			if (StateLink.ID.IsValid())
-			{
-				const FName* Name = IDToName.Find(StateLink.ID);
-				if (Name == nullptr)
-				{
-					// Missing link, we'll show these in the UI
-					return false;
-				}
-				if (StateLink.Name != *Name)
-				{
-					// Name changed, fix!
-					StateLink.Name = *Name;
-					return true;
-				}
-			}
-			return false;
-		};
-
-		// Fix changed names.
-		EditorData->VisitHierarchy([&IdToName](UStateTreeState& State, UStateTreeState* /*ParentState*/)
-		{
-			if (State.Type == EStateTreeStateType::Linked)
-			{
-				FixChangedStateLinkName(State.LinkedSubtree, IdToName);
-			}
-					
-			for (FStateTreeTransition& Transition : State.Transitions)
-			{
-				FixChangedStateLinkName(Transition.State, IdToName);
-			}
-
-			return EStateTreeVisitor::Continue;
-		});
 	}
-
-	// Update Linked State Parameters
-	EditorData->VisitHierarchy([](UStateTreeState& State, UStateTreeState* /*ParentState*/)
-	{
-		if (State.Type == EStateTreeStateType::Linked)
-		{
-			State.UpdateParametersFromLinkedSubtree();
-		}
-		return EStateTreeVisitor::Continue;
-	});
 }
 
-void FAvaTransitionEditorViewModel::OnIdentifierChanged(const UStateTree& InStateTree)
+void FAvaTransitionEditorViewModel::OnMessageTokenClicked(const TSharedRef<IMessageToken>& InMessageToken)
 {
-	if (&InStateTree == GetTransitionTree())
+	if (InMessageToken->GetType() != EMessageToken::Object)
 	{
-		UpdateTree();
+		return;
 	}
-}
 
-void FAvaTransitionEditorViewModel::OnSchemaChanged(const UStateTree& InStateTree)
-{
-	if (&InStateTree == GetTransitionTree())
+	TSharedRef<FAvaTransitionViewModelSharedData> ViewModelSharedData = GetSharedData();
+
+	if (UStateTreeState* State = Cast<UStateTreeState>(StaticCastSharedRef<FUObjectToken>(InMessageToken)->GetObject().Get()))
 	{
-		UpdateTree();
-		// todo: notify asset change externally
+		if (TSharedPtr<FAvaTransitionViewModel> FoundViewModel = ViewModelSharedData->GetRegistryCollection()->FindViewModel(State))
+		{
+			ViewModelSharedData->GetSelection()->SetSelectedItems({ FoundViewModel });
+		}
 	}
 }
 
-void FAvaTransitionEditorViewModel::OnParametersChanged(const UStateTree& InStateTree)
-{
-	// todo
-}
-
-void FAvaTransitionEditorViewModel::OnStateParametersChanged(const UStateTree& InStateTree, const FGuid InGuid)
-{
-	// todo
-}
-
 #undef LOCTEXT_NAMESPACE
diff --git a/Engine/Plugins/Experimental/Avalanche/Source/AvalancheTransitionEditor/Private/ViewModels/AvaTransitionEditorViewModel.h b/Engine/Plugins/Experimental/Avalanche/Source/AvalancheTransitionEditor/Private/ViewModels/AvaTransitionEditorViewModel.h
--- a/Engine/Plugins/Experimental/Avalanche/Source/AvalancheTransitionEditor/Private/ViewModels/AvaTransitionEditorViewModel.h
+++ b/Engine/Plugins/Experimental/Avalanche/Source/AvalancheTransitionEditor/Private/ViewModels/AvaTransitionEditorViewModel.h
@@ -16,6 +16,7 @@ class FAvaTransitionToolbar;
 class FAvaTransitionTreeContextMenu;
 class FUICommandList;
 class IAvaTransitionSelectableExtension;
+class IMessageToken;
 class SAvaTransitionTreeView;
 class UAvaTransitionTree;
 class UAvaTransitionTreeEditorData;
@@ -46,11 +47,6 @@ public:
 
 	void Compile();
 
-	uint32 GetEditorDataHash() const
-	{
-		return EditorDataHash;
-	}
-
 	UAvaTransitionTree* GetTransitionTree() const;
 
 	UAvaTransitionTreeEditorData* GetEditorData() const;
@@ -100,19 +96,14 @@ private:
 
 	void UnbindDelegates();
 
-	uint32 CalculateTreeHash(const UAvaTransitionTree& InTree) const;
-
-	void ValidateTree(const FString& InTreeDebugName);
-
 	void OnIdentifierChanged(const UStateTree& InStateTree);
 
 	void OnSchemaChanged(const UStateTree& InStateTree);
 
-	void OnParametersChanged(const UStateTree& InStateTree);
+	void OnCompileFailed();
 
-	void OnStateParametersChanged(const UStateTree& InStateTree, const FGuid InGuid);
+	void OnMessageTokenClicked(const TSharedRef<IMessageToken>& InMessageToken);
 
-private:
 	TSharedRef<FAvaTransitionToolbar> Toolbar;
 
 	TSharedRef<FAvaTransitionTreeContextMenu> ContextMenu;
@@ -123,8 +114,6 @@ private:
 
 	TWeakPtr<FAvaTransitionEditor> EditorWeak;
 
-	uint32 EditorDataHash = 0;
-
 	FAvaTransitionCompiler Compiler;
 
 	TSharedRef<FUICommandList> CommandList;
