Motion Design Transition:
- Picking TL Layer from Level Editor Menu to auto-compile the Transition Tree

diff --git a/Engine/Plugins/Experimental/Avalanche/Source/AvalancheTransitionEditor/Private/AvaTransitionEditorModule.cpp b/Engine/Plugins/Experimental/Avalanche/Source/AvalancheTransitionEditor/Private/AvaTransitionEditorModule.cpp
--- a/Engine/Plugins/Experimental/Avalanche/Source/AvalancheTransitionEditor/Private/AvaTransitionEditorModule.cpp
+++ b/Engine/Plugins/Experimental/Avalanche/Source/AvalancheTransitionEditor/Private/AvaTransitionEditorModule.cpp
@@ -93,7 +93,7 @@ void FAvaTransitionEditorModule::GenerateTransitionTreeOptionsMenu(UToolMenu* In
 			, FIsActionChecked::CreateStatic(&UE::AvaTransitionEditor::IsTransitionTreeEnabled, TransitionTreeWeak))
 		, EUserInterfaceActionType::ToggleButton);
 
-	if (TSharedPtr<SWidget> LayerPicker = UE::AvaTransitionEditor::CreateTransitionLayerPicker(Cast<UAvaTransitionTreeEditorData>(TransitionTree->EditorData)))
+	if (TSharedPtr<SWidget> LayerPicker = UE::AvaTransitionEditor::CreateTransitionLayerPicker(Cast<UAvaTransitionTreeEditorData>(TransitionTree->EditorData), /*bInCompileOnLayerPicked*/true))
 	{
 		FToolMenuSection& LayerSection = InMenu->FindOrAddSection(TEXT("LayerSection"), LOCTEXT("LayerSectionLabel", "Layer"));
 
@@ -135,7 +135,7 @@ void FAvaTransitionEditorModule::ValidateStateTree(UAvaTransitionTree* InTransit
 		EditorData->AddRootState();	
 	}
 
-	/** Compile in Advanced Mode here so that no new nodes are generated from outside */
+	// Compile in Advanced Mode here so that no new nodes are generated from outside
 	FAvaTransitionCompiler Compiler;
 	Compiler.SetTransitionTree(InTransitionTree);
 	Compiler.Compile(EAvaTransitionEditorMode::Advanced);
diff --git a/Engine/Plugins/Experimental/Avalanche/Source/AvalancheTransitionEditor/Private/AvaTransitionEditorUtils.cpp b/Engine/Plugins/Experimental/Avalanche/Source/AvalancheTransitionEditor/Private/AvaTransitionEditorUtils.cpp
--- a/Engine/Plugins/Experimental/Avalanche/Source/AvalancheTransitionEditor/Private/AvaTransitionEditorUtils.cpp
+++ b/Engine/Plugins/Experimental/Avalanche/Source/AvalancheTransitionEditor/Private/AvaTransitionEditorUtils.cpp
@@ -1,10 +1,12 @@
 // Copyright Epic Games, Inc. All Rights Reserved.
 
 #include "AvaTransitionEditorUtils.h"
+#include "AvaTransitionEditorEnums.h"
 #include "AvaTransitionEditorLog.h"
 #include "AvaTransitionTree.h"
 #include "AvaTransitionTreeEditorData.h"
 #include "Behavior/IAvaTransitionBehavior.h"
+#include "Compiler/AvaTransitionCompiler.h"
 #include "ContentBrowserModule.h"
 #include "IContentBrowserSingleton.h"
 #include "ISinglePropertyView.h"
@@ -21,7 +23,7 @@
 namespace UE::AvaTransitionEditor
 {
 
-TSharedPtr<SWidget> CreateTransitionLayerPicker(UAvaTransitionTreeEditorData* InEditorData)
+TSharedPtr<SWidget> CreateTransitionLayerPicker(UAvaTransitionTreeEditorData* InEditorData, bool bInCompileOnLayerPicked)
 {
 	if (!InEditorData)
 	{
@@ -37,6 +39,23 @@ TSharedPtr<SWidget> CreateTransitionLayerPicker(UAvaTransitionTreeEditorData* In
 		, UAvaTransitionTreeEditorData::GetTransitionLayerPropertyName()
 		, SinglePropertyParams);
 
+	if (bInCompileOnLayerPicked)
+	{
+		if (UAvaTransitionTree* TransitionTree = InEditorData->GetTypedOuter<UAvaTransitionTree>())
+		{
+			FSimpleDelegate OnLayerPicked = FSimpleDelegate::CreateWeakLambda(TransitionTree,
+				[TransitionTree]
+				{
+					// Compile in Advanced Mode here so that no new nodes are generated from outside
+					FAvaTransitionCompiler Compiler;
+					Compiler.SetTransitionTree(TransitionTree);
+					Compiler.Compile(EAvaTransitionEditorMode::Advanced);
+				});
+
+			PropertyView->SetOnPropertyValueChanged(OnLayerPicked);
+		}		
+	}
+
 	if (PropertyView.IsValid())
 	{
 		return SNew(SBox)
diff --git a/Engine/Plugins/Experimental/Avalanche/Source/AvalancheTransitionEditor/Private/AvaTransitionEditorUtils.h b/Engine/Plugins/Experimental/Avalanche/Source/AvalancheTransitionEditor/Private/AvaTransitionEditorUtils.h
--- a/Engine/Plugins/Experimental/Avalanche/Source/AvalancheTransitionEditor/Private/AvaTransitionEditorUtils.h
+++ b/Engine/Plugins/Experimental/Avalanche/Source/AvalancheTransitionEditor/Private/AvaTransitionEditorUtils.h
@@ -12,7 +12,7 @@ class UAvaTransitionTreeEditorData;
 
 namespace UE::AvaTransitionEditor
 {
-	TSharedPtr<SWidget> CreateTransitionLayerPicker(UAvaTransitionTreeEditorData* InEditorData);
+	TSharedPtr<SWidget> CreateTransitionLayerPicker(UAvaTransitionTreeEditorData* InEditorData, bool bInCompileOnLayerPicked);
 
 	void ValidateTree(UAvaTransitionTree& InTransitionTree);
 
diff --git a/Engine/Plugins/Experimental/Avalanche/Source/AvalancheTransitionEditor/Private/Menu/AvaTransitionToolbar.cpp b/Engine/Plugins/Experimental/Avalanche/Source/AvalancheTransitionEditor/Private/Menu/AvaTransitionToolbar.cpp
--- a/Engine/Plugins/Experimental/Avalanche/Source/AvalancheTransitionEditor/Private/Menu/AvaTransitionToolbar.cpp
+++ b/Engine/Plugins/Experimental/Avalanche/Source/AvalancheTransitionEditor/Private/Menu/AvaTransitionToolbar.cpp
@@ -51,7 +51,7 @@ void FAvaTransitionToolbar::ExtendEditorToolbar(UToolMenu* InToolbarMenu)
 	CompileOptions.StyleNameOverride = "CalloutToolbar";
 	CompileOptions.ToolBarData.bSimpleComboBox = true;
 
-	if (TSharedPtr<SWidget> LayerPicker = UE::AvaTransitionEditor::CreateTransitionLayerPicker(Owner.GetEditorData()))
+	if (TSharedPtr<SWidget> LayerPicker = UE::AvaTransitionEditor::CreateTransitionLayerPicker(Owner.GetEditorData(), /*bInCompileOnLayerPicked*/false))
 	{
 		LayerPicker->SetEnabled(!bReadOnly);
 
