[Motion Design] Removing "Avalanche" from new asset names.

#rnx

diff --git a/Engine/Plugins/Experimental/Avalanche/Source/AvalancheMediaEditor/Private/Playback/AvaPlaybackFactory.cpp b/Engine/Plugins/Experimental/Avalanche/Source/AvalancheMediaEditor/Private/Playback/AvaPlaybackFactory.cpp
--- a/Engine/Plugins/Experimental/Avalanche/Source/AvalancheMediaEditor/Private/Playback/AvaPlaybackFactory.cpp
+++ b/Engine/Plugins/Experimental/Avalanche/Source/AvalancheMediaEditor/Private/Playback/AvaPlaybackFactory.cpp
@@ -63,4 +63,9 @@ UClass* UAvaPlaybackFactory::ResolveSupportedClass()
 	return Super::ResolveSupportedClass();
 }
 
+FString UAvaPlaybackFactory::GetDefaultNewAssetName() const
+{
+	return TEXT("NewPlaybackGraph");
+}
+
 #undef LOCTEXT_NAMESPACE
diff --git a/Engine/Plugins/Experimental/Avalanche/Source/AvalancheMediaEditor/Private/Playlist/AvaPlaylistFactory.cpp b/Engine/Plugins/Experimental/Avalanche/Source/AvalancheMediaEditor/Private/Playlist/AvaPlaylistFactory.cpp
--- a/Engine/Plugins/Experimental/Avalanche/Source/AvalancheMediaEditor/Private/Playlist/AvaPlaylistFactory.cpp
+++ b/Engine/Plugins/Experimental/Avalanche/Source/AvalancheMediaEditor/Private/Playlist/AvaPlaylistFactory.cpp
@@ -59,3 +59,9 @@ UClass* UAvaPlaylistFactory::ResolveSupportedClass()
 {
 	return Super::ResolveSupportedClass();
 }
+
+FString UAvaPlaylistFactory::GetDefaultNewAssetName() const
+{
+	return TEXT("NewRundown");
+}
+
diff --git a/Engine/Plugins/Experimental/Avalanche/Source/AvalancheMediaEditor/Private/Playlist/AvaRundownMacroCollectionFactory.cpp b/Engine/Plugins/Experimental/Avalanche/Source/AvalancheMediaEditor/Private/Playlist/AvaRundownMacroCollectionFactory.cpp
--- a/Engine/Plugins/Experimental/Avalanche/Source/AvalancheMediaEditor/Private/Playlist/AvaRundownMacroCollectionFactory.cpp
+++ b/Engine/Plugins/Experimental/Avalanche/Source/AvalancheMediaEditor/Private/Playlist/AvaRundownMacroCollectionFactory.cpp
@@ -28,4 +28,9 @@ UObject* UAvaRundownMacroCollectionFactory::FactoryCreateNew(UClass* Class, UObj
 		return NewObject<UAvaRundownMacroCollection>(InParent, Name, Flags);
 	}
 	return nullptr;
-}
\ No newline at end of file
+}
+
+FString UAvaRundownMacroCollectionFactory::GetDefaultNewAssetName() const
+{
+	return TEXT("NewRundownMacros");
+}
diff --git a/Engine/Plugins/Experimental/Avalanche/Source/AvalancheMediaEditor/Private/Playlist/AvaRundownMacroCollectionFactory.h b/Engine/Plugins/Experimental/Avalanche/Source/AvalancheMediaEditor/Private/Playlist/AvaRundownMacroCollectionFactory.h
--- a/Engine/Plugins/Experimental/Avalanche/Source/AvalancheMediaEditor/Private/Playlist/AvaRundownMacroCollectionFactory.h
+++ b/Engine/Plugins/Experimental/Avalanche/Source/AvalancheMediaEditor/Private/Playlist/AvaRundownMacroCollectionFactory.h
@@ -15,5 +15,6 @@ class UAvaRundownMacroCollectionFactory : public UFactory
 	//~ Begin UFactory Interface
 	virtual uint32 GetMenuCategories() const override;
 	virtual UObject* FactoryCreateNew(UClass* Class, UObject* InParent, FName Name, EObjectFlags Flags, UObject* Context, FFeedbackContext* Warn) override;
+	virtual FString GetDefaultNewAssetName() const override;
 	//~ End UFactory Interface
 };
diff --git a/Engine/Plugins/Experimental/Avalanche/Source/AvalancheMediaEditor/Public/Playback/AvaPlaybackFactory.h b/Engine/Plugins/Experimental/Avalanche/Source/AvalancheMediaEditor/Public/Playback/AvaPlaybackFactory.h
--- a/Engine/Plugins/Experimental/Avalanche/Source/AvalancheMediaEditor/Public/Playback/AvaPlaybackFactory.h
+++ b/Engine/Plugins/Experimental/Avalanche/Source/AvalancheMediaEditor/Public/Playback/AvaPlaybackFactory.h
@@ -26,5 +26,6 @@ protected:
 	virtual UObject* FactoryCreateNew(UClass* Class, UObject* InParent, FName Name, EObjectFlags Flags, UObject* Context, FFeedbackContext* Warn) override;
 	virtual bool DoesSupportClass(UClass* Class) override;
 	virtual UClass* ResolveSupportedClass() override;
+	virtual FString GetDefaultNewAssetName() const override;
 	//~ Begin UFactory Interface
 };
diff --git a/Engine/Plugins/Experimental/Avalanche/Source/AvalancheMediaEditor/Public/Playlist/AvaPlaylistFactory.h b/Engine/Plugins/Experimental/Avalanche/Source/AvalancheMediaEditor/Public/Playlist/AvaPlaylistFactory.h
--- a/Engine/Plugins/Experimental/Avalanche/Source/AvalancheMediaEditor/Public/Playlist/AvaPlaylistFactory.h
+++ b/Engine/Plugins/Experimental/Avalanche/Source/AvalancheMediaEditor/Public/Playlist/AvaPlaylistFactory.h
@@ -26,5 +26,6 @@ protected:
 	virtual UObject* FactoryCreateNew(UClass* Class, UObject* InParent, FName Name, EObjectFlags Flags, UObject* Context, FFeedbackContext* Warn) override;
 	virtual bool DoesSupportClass(UClass* Class) override;
 	virtual UClass* ResolveSupportedClass() override;
+	virtual FString GetDefaultNewAssetName() const override;
 	//~ Begin UFactory Interface
 };
