Motion Design loctext fixes

diff --git a/Engine/Plugins/Experimental/Avalanche/Source/AvalancheLevelViewport/Private/SAvaLevelViewport_Menus.cpp b/Engine/Plugins/Experimental/Avalanche/Source/AvalancheLevelViewport/Private/SAvaLevelViewport_Menus.cpp
--- a/Engine/Plugins/Experimental/Avalanche/Source/AvalancheLevelViewport/Private/SAvaLevelViewport_Menus.cpp
+++ b/Engine/Plugins/Experimental/Avalanche/Source/AvalancheLevelViewport/Private/SAvaLevelViewport_Menus.cpp
@@ -396,8 +396,8 @@ void SAvaLevelViewport::AddVirtualSizeSizeSettings(FToolMenuSection& InSection)
 	const FAvaLevelViewportCommands& LevelViewportCommands = FAvaLevelViewportCommands::Get();
 
 	InSection.AddMenuEntry("UseUnlockedAspectRatio", LevelViewportCommands.VirtualSizeAspectRatioUnlocked, LOCTEXT("VirtualSizeUnlockedAspectRatio", "Free Aspect Ratio"));
-	InSection.AddMenuEntry("UseLockedAspectRatio", LevelViewportCommands.VirtualSizeAspectRatioLocked, LOCTEXT("VirtualSizeUnlockedAspectRatio", "Locked Aspect Ratio"));
-	InSection.AddMenuEntry("UseLockedToCameraAspectRatio", LevelViewportCommands.VirtualSizeAspectRatioLockedToCamera, LOCTEXT("VirtualSizeUnlockedAspectRatio", "Use Camera Aspect Ratio"));
+	InSection.AddMenuEntry("UseLockedAspectRatio", LevelViewportCommands.VirtualSizeAspectRatioLocked, LOCTEXT("VirtualSizeLockedAspectRatio", "Locked Aspect Ratio"));
+	InSection.AddMenuEntry("UseLockedToCameraAspectRatio", LevelViewportCommands.VirtualSizeAspectRatioLockedToCamera, LOCTEXT("VirtualSizeLockedToCameraAspectRatio", "Use Camera Aspect Ratio"));
 
 	TSharedRef<SWidget> ResolutionWidget = SNew(SHorizontalBox)
 			
diff --git a/Engine/Plugins/Experimental/Avalanche/Source/AvalancheLevelViewport/Private/Widgets/SAvaLevelViewportActorAlignmentMenu.cpp b/Engine/Plugins/Experimental/Avalanche/Source/AvalancheLevelViewport/Private/Widgets/SAvaLevelViewportActorAlignmentMenu.cpp
--- a/Engine/Plugins/Experimental/Avalanche/Source/AvalancheLevelViewport/Private/Widgets/SAvaLevelViewportActorAlignmentMenu.cpp
+++ b/Engine/Plugins/Experimental/Avalanche/Source/AvalancheLevelViewport/Private/Widgets/SAvaLevelViewportActorAlignmentMenu.cpp
@@ -632,7 +632,7 @@ FReply SAvaLevelViewportActorAlignmentMenu::OnSizeToScreenClicked(bool bInStretc
 		return FReply::Handled();
 	}
 
-	FScopedTransaction Transaction(LOCTEXT("SizeToScreen", "Size Actor to Screen"));
+	FScopedTransaction Transaction(LOCTEXT("SizeToScreenTransaction", "Size Actor to Screen"));
 
 	TArray<TWeakObjectPtr<const AActor>> SelectedActorsWeak = GetSelectedActors(EditorViewportClient);
 
@@ -684,7 +684,7 @@ FReply SAvaLevelViewportActorAlignmentMenu::OnFitToScreenClicked()
 		return FReply::Handled();
 	}
 
-	FScopedTransaction Transaction(LOCTEXT("FitToScreen", "Fit Actor to Screen"));
+	FScopedTransaction Transaction(LOCTEXT("FitToScreenTransaction", "Fit Actor to Screen"));
 
 	const bool bUseNearestAxis = FSlateApplication::Get().GetModifierKeys().IsShiftDown();
 
diff --git a/Engine/Plugins/Experimental/Avalanche/Source/AvalancheLevelViewport/Private/Widgets/SAvaLevelViewportActorColorMenu.cpp b/Engine/Plugins/Experimental/Avalanche/Source/AvalancheLevelViewport/Private/Widgets/SAvaLevelViewportActorColorMenu.cpp
--- a/Engine/Plugins/Experimental/Avalanche/Source/AvalancheLevelViewport/Private/Widgets/SAvaLevelViewportActorColorMenu.cpp
+++ b/Engine/Plugins/Experimental/Avalanche/Source/AvalancheLevelViewport/Private/Widgets/SAvaLevelViewportActorColorMenu.cpp
@@ -644,11 +644,11 @@ void SAvaLevelViewportActorColorMenu::LoadColorThemesFromIni()
 	if (ColorThemes.Num() == 0)
 	{
 		FAvaColorTheme* NewTheme = FindTheme("Default", true);
-		NewTheme->GetColors().Add({FLinearColor::White, LOCTEXT("White", "White")});
-		NewTheme->GetColors().Add({FLinearColor::Black, LOCTEXT("Black", "Black")});
-		NewTheme->GetColors().Add({FLinearColor::Red, LOCTEXT("Red", "Red")});
-		NewTheme->GetColors().Add({FLinearColor::Green, LOCTEXT("Green", "Green")});
-		NewTheme->GetColors().Add({FLinearColor::Blue, LOCTEXT("Blue", "Blue")});
+		NewTheme->GetColors().Add({FLinearColor::White, LOCTEXT("WhiteColor", "White")});
+		NewTheme->GetColors().Add({FLinearColor::Black, LOCTEXT("BlackColor", "Black")});
+		NewTheme->GetColors().Add({FLinearColor::Red, LOCTEXT("RedColor", "Red")});
+		NewTheme->GetColors().Add({FLinearColor::Green, LOCTEXT("GreenColor", "Green")});
+		NewTheme->GetColors().Add({FLinearColor::Blue, LOCTEXT("BlueColor", "Blue")});
 		SaveColorThemesToIni();		
 	}
 
@@ -1222,7 +1222,7 @@ TSharedRef<SWidget> SAvaLevelViewportActorColorMenu::BuildRenameColorMenu(int32
 	{
 		ColorMenuBuilder.BeginSection("ColorOptions", LOCTEXT("ColorOptions", "Color Options"));
 
-		ColorMenuBuilder.AddEditableText(LOCTEXT("Name","Name"), LOCTEXT("NameTooltip", "The name of the color."),
+		ColorMenuBuilder.AddEditableText(LOCTEXT("NameLabel", "Name"), LOCTEXT("NameTooltip", "The name of the color."),
 			FSlateIcon(), ThemeColors[InColorIdx].Label, 
 			FOnTextCommitted::CreateLambda([Theme, ColorWidget, InColorIdx, CurrentColor](const FText& NewLabel, ETextCommit::Type CommitType)
 				{
@@ -1245,7 +1245,7 @@ TSharedRef<SWidget> SAvaLevelViewportActorColorMenu::BuildRenameColorMenu(int32
 					ColorWidget->SetToolTipText(NewLabel);
 				}));
 
-		ColorMenuBuilder.AddMenuEntry(LOCTEXT("Remove", "Remove"), LOCTEXT("RemoveTooltip", "Remove this color from the palette."), FSlateIcon(),
+		ColorMenuBuilder.AddMenuEntry(LOCTEXT("RemovePaletteColorLabel", "Remove"), LOCTEXT("RemovePaletteColorTooltip", "Remove this color from the palette."), FSlateIcon(),
 			FUIAction(FExecuteAction::CreateLambda([this, Theme, InColorIdx, CurrentColor]()
 				{
 					TArray<FAvaColorInfo>& ThemeColors = Theme->GetColors();
@@ -1278,6 +1278,9 @@ TSharedRef<SWidget> SAvaLevelViewportActorColorMenu::BuildThemesMenu()
 	constexpr bool bInShouldCloseWindowAfterMenuSelection = true;
 	FMenuBuilder ThemeMenuBuilder(bInShouldCloseWindowAfterMenuSelection, nullptr);
 	{
+		static FText RemoveThemeLabelText = LOCTEXT("RemoveThemeLabel", "Remove");
+		static FText RemoveThemeTooltipText = LOCTEXT("RemoveThemeTooltip", "Remove this theme.");
+		
 		if (ColorThemes.IsValidIndex(ActiveThemeIndex))
 		{
 			ThemeMenuBuilder.BeginSection("CurrentTheme", LOCTEXT("CurrentTheme", "Current Theme"));
@@ -1303,10 +1306,10 @@ TSharedRef<SWidget> SAvaLevelViewportActorColorMenu::BuildThemesMenu()
 						SaveColorThemesToIni();
 					})
 				),				
-				LOCTEXT("Name", "Name")
+				LOCTEXT("ThemeNameLabel", "Name")
 			);
 
-			ThemeMenuBuilder.AddMenuEntry(LOCTEXT("Remove", "Remove"), LOCTEXT("RemoveTooltip", "Remove this theme."), FSlateIcon(),
+			ThemeMenuBuilder.AddMenuEntry(RemoveThemeLabelText, RemoveThemeTooltipText, FSlateIcon(),
 				FUIAction(FExecuteAction::CreateLambda([this]()
 					{
 						if (!ColorThemes.IsValidIndex(ActiveThemeIndex))
@@ -1343,7 +1346,7 @@ TSharedRef<SWidget> SAvaLevelViewportActorColorMenu::BuildThemesMenu()
 					{
 						SubMenuBuilder.BeginSection("ThemeOptions", LOCTEXT("Options", "Options"));
 
-						SubMenuBuilder.AddMenuEntry(LOCTEXT("Remove", "Remove"), LOCTEXT("RemoveThemeTooltip", "Remove this theme."), FSlateIcon(),
+						SubMenuBuilder.AddMenuEntry(RemoveThemeLabelText, RemoveThemeTooltipText, FSlateIcon(),
 							FUIAction(FExecuteAction::CreateLambda([this, ThemeIdx]()
 								{
 									if (!ColorThemes.IsValidIndex(ThemeIdx))
@@ -1410,7 +1413,7 @@ TSharedRef<SWidget> SAvaLevelViewportActorColorMenu::BuildThemesMenu()
 										],
 									FNewMenuDelegate::CreateLambda([this, ThemeIdx, ColorIdx](FMenuBuilder& SubMenuBuilder)
 										{
-											SubMenuBuilder.AddMenuEntry(LOCTEXT("Remove", "Remove"), LOCTEXT("RemoveColorTooltip", "Remove this color."), FSlateIcon(),
+											SubMenuBuilder.AddMenuEntry(LOCTEXT("RemoveColorLabel", "Remove"), LOCTEXT("RemoveColorTooltip", "Remove this color."), FSlateIcon(),
 												FUIAction(FExecuteAction::CreateLambda([this, ThemeIdx, ColorIdx]()
 													{
 														if (!ColorThemes.IsValidIndex(ThemeIdx))
@@ -1463,7 +1466,7 @@ TSharedRef<SWidget> SAvaLevelViewportActorColorMenu::BuildThemesMenu()
 
 		ThemeMenuBuilder.EndSection();
 
-		ThemeMenuBuilder.BeginSection("NewTheme", LOCTEXT("NewTheme", "New Theme"));
+		ThemeMenuBuilder.BeginSection("NewTheme", LOCTEXT("NewThemeSectionHeading", "New Theme"));
 
 		ThemeMenuBuilder.AddWidget(
 			SNew(SEditableTextBox)
@@ -1485,7 +1488,7 @@ TSharedRef<SWidget> SAvaLevelViewportActorColorMenu::BuildThemesMenu()
 					}
 				})
 			),
-			LOCTEXT("Name", "Name")
+			LOCTEXT("NewThemeLabel", "Name")
 		);
 
 		ThemeMenuBuilder.EndSection();
diff --git a/Engine/Plugins/Experimental/Avalanche/Source/AvalancheLevelViewport/Private/Widgets/SAvaLevelViewportStatusBarButtons.cpp b/Engine/Plugins/Experimental/Avalanche/Source/AvalancheLevelViewport/Private/Widgets/SAvaLevelViewportStatusBarButtons.cpp
--- a/Engine/Plugins/Experimental/Avalanche/Source/AvalancheLevelViewport/Private/Widgets/SAvaLevelViewportStatusBarButtons.cpp
+++ b/Engine/Plugins/Experimental/Avalanche/Source/AvalancheLevelViewport/Private/Widgets/SAvaLevelViewportStatusBarButtons.cpp
@@ -233,7 +233,7 @@ void SAvaLevelViewportStatusBarButtons::PopulateActorButtons(TSharedPtr<SHorizon
 		.Padding(ViewportStatusBarButton::Padding)
 		[
 			ViewportStatusBarButton::MakeMenuButton(
-				LOCTEXT("ActorAlign", "Actor Color"),
+				LOCTEXT("ActorColor", "Actor Color"),
 				FOnGetContent::CreateSP(this, &SAvaLevelViewportStatusBarButtons::GetActorColorMenuContent),
 				FAppStyle::Get().GetBrush(TEXT("ColorPicker.Mode")),
 				FSlateColor(FLinearColor(1.0f, 1.0f, 1.0f, 1.0f))
diff --git a/Engine/Plugins/Experimental/Avalanche/Source/AvalancheLevelViewport/Private/Widgets/SAvaViewportInfo.cpp b/Engine/Plugins/Experimental/Avalanche/Source/AvalancheLevelViewport/Private/Widgets/SAvaViewportInfo.cpp
--- a/Engine/Plugins/Experimental/Avalanche/Source/AvalancheLevelViewport/Private/Widgets/SAvaViewportInfo.cpp
+++ b/Engine/Plugins/Experimental/Avalanche/Source/AvalancheLevelViewport/Private/Widgets/SAvaViewportInfo.cpp
@@ -1,6 +1,7 @@
 // Copyright Epic Games, Inc. All Rights Reserved.
 
 #include "SAvaViewportInfo.h"
+
 #include "AvaViewportUtils.h"
 #include "AvaVisibleArea.h"
 #include "EditorModeManager.h"
@@ -13,6 +14,25 @@
 
 #define LOCTEXT_NAMESPACE "SAvaViewportInfo"
 
+namespace UE::AvaLevelViewport::Private
+{
+	const FText& GetEmptyXYDisplayText()
+	{
+		static const FText EmptyXY = LOCTEXT("NoXY", "-");
+		return EmptyXY;
+	}
+
+	FText MakeXYDisplayText_By(float InX, float InY)
+	{
+		return FText::Format(LOCTEXT("XYDash", "{0} x {1}"), FMath::RoundToInt(InX), FMath::RoundToInt(InY));
+	}
+
+	FText MakeXYDisplayText_Comma(float InX, float InY)
+	{
+		return FText::Format(LOCTEXT("XYComma", "{0}, {1}"), FMath::RoundToInt(InX), FMath::RoundToInt(InY));
+	}
+}
+
 TSharedRef<SAvaViewportInfo> SAvaViewportInfo::CreateInstance(const TSharedRef<IToolkitHost>& InToolkitHost)
 {
 	return SNew(SAvaViewportInfo, InToolkitHost);
@@ -237,62 +257,42 @@ FVector2f SAvaViewportInfo::GetMouseLocationOnViewport() const
 
 FText SAvaViewportInfo::GetViewportSize() const
 {
-	static const FText Default = LOCTEXT("NoSize", "-");
-
 	const FVector2f& ViewportSize = GetViewportSizeForActiveViewport();
 
 	if (FAvaViewportUtils::IsValidViewportSize(ViewportSize))
 	{
-		return FText::Format(
-			LOCTEXT("Size", "{0} x {1}"), 
-			FMath::RoundToInt(ViewportSize.X), 
-			FMath::RoundToInt(ViewportSize.Y)
-		);
+		return UE::AvaLevelViewport::Private::MakeXYDisplayText_By(FMath::RoundToInt(ViewportSize.X), FMath::RoundToInt(ViewportSize.Y));
 	}
 
-	return Default;
+	return UE::AvaLevelViewport::Private::GetEmptyXYDisplayText();
 }
 
 FText SAvaViewportInfo::GetVirtualViewportSize() const
 {
-	static const FText Default = LOCTEXT("NoSize", "-");
-
 	const FIntPoint VirtualSize = GetVirtualSizeForActiveViewport();
 
 	if (FAvaViewportUtils::IsValidViewportSize(VirtualSize))
 	{
-		return FText::Format(
-			LOCTEXT("Size", "{0} x {1}"),
-			VirtualSize.X,
-			VirtualSize.Y
-		);
+		return UE::AvaLevelViewport::Private::MakeXYDisplayText_By(VirtualSize.X, VirtualSize.Y);
 	}
 
-	return Default;
+	return UE::AvaLevelViewport::Private::GetEmptyXYDisplayText();
 }
 
 FText SAvaViewportInfo::GetViewportVisibleAreaSize() const
 {
-	static const FText Default = LOCTEXT("NoSize", "-");
-
 	const FAvaVisibleArea VisibleArea = GetVisibleAreaForActiveViewport();
 
 	if (VisibleArea.IsValid() && VisibleArea.IsZoomedView())
 	{
-		return FText::Format(
-			LOCTEXT("Size", "{0} x {1}"),
-			FMath::RoundToInt(VisibleArea.VisibleSize.X),
-			FMath::RoundToInt(VisibleArea.VisibleSize.Y)
-		);
+		return UE::AvaLevelViewport::Private::MakeXYDisplayText_By(VisibleArea.VisibleSize.X, VisibleArea.VisibleSize.Y);
 	}
 
-	return Default;
+	return UE::AvaLevelViewport::Private::GetEmptyXYDisplayText();
 }
 
 FText SAvaViewportInfo::GetCanvasVisibleAreaSize() const
 {
-	static const FText Default = LOCTEXT("NoSize", "-");
-
 	const FAvaVisibleArea VisibleArea = GetVisibleAreaForActiveViewport();
 
 	if (VisibleArea.IsValid() && VisibleArea.IsZoomedView())
@@ -303,39 +303,27 @@ FText SAvaViewportInfo::GetCanvasVisibleAreaSize() const
 		{
 			const FVector2f VisibleSize = VisibleArea.VisibleSize * static_cast<float>(CanvasSize.X) / VisibleArea.AbsoluteSize.X;
 
-			return FText::Format(
-				LOCTEXT("Size", "{0} x {1}"),
-				FMath::RoundToInt(VisibleSize.X),
-				FMath::RoundToInt(VisibleSize.Y)
-			);
+			return UE::AvaLevelViewport::Private::MakeXYDisplayText_By(VisibleSize.X, VisibleSize.Y);
 		}
 	}
 
-	return Default;
+	return UE::AvaLevelViewport::Private::GetEmptyXYDisplayText();
 }
 
 FText SAvaViewportInfo::GetViewportZoomOffset() const
 {
-	static const FText Default = LOCTEXT("NoCoordinates", "-");
-
 	const FAvaVisibleArea VisibleArea = GetVisibleAreaForActiveViewport();
 
 	if (VisibleArea.IsValid() && VisibleArea.IsOffset())
 	{
-		return FText::Format(
-			LOCTEXT("Coordinates", "{0}, {1}"),
-			FMath::RoundToInt(VisibleArea.Offset.X),
-			FMath::RoundToInt(VisibleArea.Offset.Y)
-		);
+		return UE::AvaLevelViewport::Private::MakeXYDisplayText_Comma(VisibleArea.Offset.X, VisibleArea.Offset.Y);
 	}
 
-	return Default;
+	return UE::AvaLevelViewport::Private::GetEmptyXYDisplayText();
 }
 
 FText SAvaViewportInfo::GetCanvasZoomOffset() const
 {
-	static const FText Default = LOCTEXT("NoCoordinates", "-");
-
 	const FAvaVisibleArea VisibleArea = GetVisibleAreaForActiveViewport();
 
 	if (VisibleArea.IsValid() && VisibleArea.IsOffset())
@@ -346,39 +334,29 @@ FText SAvaViewportInfo::GetCanvasZoomOffset() const
 		{
 			const FVector2f CanvasOffset = VisibleArea.Offset * static_cast<float>(CanvasSize.X) / VisibleArea.AbsoluteSize.X;
 
-			return FText::Format(
-				LOCTEXT("Coordinates", "{0}, {1}"),
-				FMath::RoundToInt(CanvasOffset.X),
-				FMath::RoundToInt(CanvasOffset.Y)
-			);
+			return UE::AvaLevelViewport::Private::MakeXYDisplayText_Comma(CanvasOffset.X, CanvasOffset.Y);
 		}
 	}
 
-	return Default;
+	return UE::AvaLevelViewport::Private::GetEmptyXYDisplayText();
 }
 
 FText SAvaViewportInfo::GetMouseLocation() const
 {
-	static const FText Default = LOCTEXT("NoCoordinates", "-");
-
 	const FVector2f MouseLocation = GetMouseLocationOnViewport();
 
 	// -1 is the invalid value
 	if (MouseLocation.X < 0 || MouseLocation.Y < 0)
 	{
-		return Default;
+		return UE::AvaLevelViewport::Private::GetEmptyXYDisplayText();
 	}
 
-	return FText::Format(
-		LOCTEXT("Coordinates", "{0}, {1}"),
-		FMath::RoundToInt(MouseLocation.X),
-		FMath::RoundToInt(MouseLocation.Y)
-	);
+	return UE::AvaLevelViewport::Private::MakeXYDisplayText_Comma(MouseLocation.X, MouseLocation.Y);
 }
 
 FText SAvaViewportInfo::GetVirtualMouseLocation() const
 {
-	static const FText Default = LOCTEXT("NoCoordinates", "-");
+	static const FText& Default = UE::AvaLevelViewport::Private::GetEmptyXYDisplayText();
 
 	FVector2f MouseLocation = GetMouseLocationOnViewport();
 
@@ -404,17 +382,11 @@ FText SAvaViewportInfo::GetVirtualMouseLocation() const
 
 	MouseLocation *= static_cast<float>(CanvasSize.X) / ViewportSize.X;
 
-	return FText::Format(
-		LOCTEXT("Coordinates", "{0}, {1}"),
-		FMath::RoundToInt(MouseLocation.X),
-		FMath::RoundToInt(MouseLocation.Y)
-	);
+	return UE::AvaLevelViewport::Private::MakeXYDisplayText_Comma(MouseLocation.X, MouseLocation.Y);
 }
 
 FText SAvaViewportInfo::GetZoomLevel() const
 {
-	static const FText Default = LOCTEXT("No value", "-");
-
 	const FAvaVisibleArea VisibleArea = GetVisibleAreaForActiveViewport();
 
 	if (VisibleArea.IsValid())
@@ -427,7 +399,7 @@ FText SAvaViewportInfo::GetZoomLevel() const
 		);
 	}
 
-	return Default;
+	return UE::AvaLevelViewport::Private::GetEmptyXYDisplayText();
 }
 
 #undef LOCTEXT_NAMESPACE
diff --git a/Engine/Plugins/Experimental/Avalanche/Source/AvalancheMedia/Private/OutputDevices/AvaDisplayDeviceProvider.cpp b/Engine/Plugins/Experimental/Avalanche/Source/AvalancheMedia/Private/OutputDevices/AvaDisplayDeviceProvider.cpp
--- a/Engine/Plugins/Experimental/Avalanche/Source/AvalancheMedia/Private/OutputDevices/AvaDisplayDeviceProvider.cpp
+++ b/Engine/Plugins/Experimental/Avalanche/Source/AvalancheMedia/Private/OutputDevices/AvaDisplayDeviceProvider.cpp
@@ -172,7 +172,7 @@ FText FAvaDisplayDeviceProvider::ToText(const FMediaIOConfiguration& InConfigura
 {
 	if (bInIsAutoDetected)
 	{
-		return FText::Format(LOCTEXT("FMediaIOConfigurationToText", "{0} - {1} [device{2}/auto]")
+		return FText::Format(LOCTEXT("FMediaIOAutoConfigurationToText", "{0} - {1} [device{2}/auto]")
 				, InConfiguration.bIsInput ? LOCTEXT("In", "In") : LOCTEXT("Out", "Out")
 				, FText::FromName(InConfiguration.MediaConnection.Device.DeviceName)
 				, FText::AsNumber(InConfiguration.MediaConnection.Device.DeviceIdentifier)
diff --git a/Engine/Plugins/Experimental/Avalanche/Source/AvalancheMediaEditor/Private/Broadcast/AvaBroadcastEditor.cpp b/Engine/Plugins/Experimental/Avalanche/Source/AvalancheMediaEditor/Private/Broadcast/AvaBroadcastEditor.cpp
--- a/Engine/Plugins/Experimental/Avalanche/Source/AvalancheMediaEditor/Private/Broadcast/AvaBroadcastEditor.cpp
+++ b/Engine/Plugins/Experimental/Avalanche/Source/AvalancheMediaEditor/Private/Broadcast/AvaBroadcastEditor.cpp
@@ -495,9 +495,9 @@ FText FAvaBroadcastEditor::GetStopPlaybackClientTooltip()
 {
 	if (UAvalancheMediaSettings::Get().bAutoStartPlaybackClient)
 	{
-		return LOCTEXT("PlaybackClientStop_ToolTip", "Stop Playback Client (auto started)");
+		return LOCTEXT("PlaybackClientStopAuto_ToolTip", "Stop Playback Client (auto started)");
 	}
-	return LOCTEXT("PlaybackClientStop_ToolTip", "Stop Playback Client (manually started)");
+	return LOCTEXT("PlaybackClientStopManual_ToolTip", "Stop Playback Client (manually started)");
 }
 
 void FAvaBroadcastEditor::StartPlaybackClientAction()
@@ -535,9 +535,9 @@ FText FAvaBroadcastEditor::GetLaunchLocalServerTooltip()
 {
 	if (IAvaMediaModule::Get().IsMediaPlaybackClientStarted())
 	{
-		return LOCTEXT("LaunchLocalServer_ToolTip", "Launches Game Mode Local Server Process");
+		return LOCTEXT("LaunchLocalServerAuto_ToolTip", "Launches Game Mode Local Server Process");
 	}
-	return LOCTEXT("LaunchLocalServer_ToolTip", "Launches Game Mode Local Server Process (disabled, start client first)");
+	return LOCTEXT("LaunchLocalServerManual_ToolTip", "Launches Game Mode Local Server Process (disabled, start client first)");
 }
 
 #undef LOCTEXT_NAMESPACE
diff --git a/Engine/Plugins/Experimental/Avalanche/Source/AvalancheMediaEditor/Private/Broadcast/ChannelGrid/AvaOutputTileItem.cpp b/Engine/Plugins/Experimental/Avalanche/Source/AvalancheMediaEditor/Private/Broadcast/ChannelGrid/AvaOutputTileItem.cpp
--- a/Engine/Plugins/Experimental/Avalanche/Source/AvalancheMediaEditor/Private/Broadcast/ChannelGrid/AvaOutputTileItem.cpp
+++ b/Engine/Plugins/Experimental/Avalanche/Source/AvalancheMediaEditor/Private/Broadcast/ChannelGrid/AvaOutputTileItem.cpp
@@ -190,19 +190,19 @@ void FAvaOutputTileItem::UpdateInfo()
 	switch(OutputState)
 	{
 	case EAvaMediaOutputState::Offline:
-		MediaOutputStatusText = LOCTEXT("MediaOutput_NoIssues", "Offline");
+		MediaOutputStatusText = LOCTEXT("MediaOutput_Offline", "Offline");
 		MediaOutputStatusBrush = FAvaMediaEditorStyle::Get().GetBrush("AvalancheMediaEditor.MediaOutputOffline");
 		break;
 	case EAvaMediaOutputState::Idle:
-		MediaOutputStatusText = LOCTEXT("MediaOutput_NoIssues", "Idle");
+		MediaOutputStatusText = LOCTEXT("MediaOutput_Idle", "Idle");
 		MediaOutputStatusBrush = FAvaMediaEditorStyle::Get().GetBrush("AvalancheMediaEditor.MediaOutputIdle");
 		break;
 	case EAvaMediaOutputState::Preparing:
-		MediaOutputStatusText = LOCTEXT("MediaOutput_NoIssues", "Preparing");
+		MediaOutputStatusText = LOCTEXT("MediaOutput_Preparing", "Preparing");
 		MediaOutputStatusBrush = FAvaMediaEditorStyle::Get().GetBrush("AvalancheMediaEditor.MediaOutputPreparing");
 		break;
 	case EAvaMediaOutputState::Live:
-		MediaOutputStatusText = LOCTEXT("MediaOutput_NoIssues", "Live");
+		MediaOutputStatusText = LOCTEXT("MediaOutput_Live", "Live");
 		if (Severity == EAvaMediaIssueSeverity::Errors || Severity == EAvaMediaIssueSeverity::Warnings)
 		{
 			// If severity is warning or error, add a secondary icon (yellow or red exclamation mark)
@@ -215,6 +215,7 @@ void FAvaOutputTileItem::UpdateInfo()
 		}
 		break;
 	case EAvaMediaOutputState::Error:
+	default:
 		MediaOutputStatusText = LOCTEXT("MediaOutput_Errors", "Error(s)");
 		MediaOutputStatusBrush = FAvaMediaEditorStyle::Get().GetBrush("AvalancheMediaEditor.MediaOutputError");
 		break;
diff --git a/Engine/Plugins/Experimental/Avalanche/Source/AvalancheMediaEditor/Private/Playlist/Pages/Slate/SAvaPageList.cpp b/Engine/Plugins/Experimental/Avalanche/Source/AvalancheMediaEditor/Private/Playlist/Pages/Slate/SAvaPageList.cpp
--- a/Engine/Plugins/Experimental/Avalanche/Source/AvalancheMediaEditor/Private/Playlist/Pages/Slate/SAvaPageList.cpp
+++ b/Engine/Plugins/Experimental/Avalanche/Source/AvalancheMediaEditor/Private/Playlist/Pages/Slate/SAvaPageList.cpp
@@ -994,7 +994,7 @@ bool SAvaPageList::HandleDropEvent(const FDragDropEvent& InDragDropEvent, EItemD
 		const TArray<FSoftObjectPath> AvaAssets = FilterAvaAssetPaths(AssetDragDropOp->GetAssets());
 		const TArray<FSoftObjectPath> Playlists = FilterPlaylistPaths(AssetDragDropOp->GetAssets(), FSoftObjectPath(GetPlaylist()));
 
-		FScopedTransaction DropTransaction(LOCTEXT("DropTransaction", "Drop Motion Design Assets onto Page List"));
+		FScopedTransaction DropTransaction(LOCTEXT("AssetDropTransaction", "Drop Motion Design Assets onto Page List"));
 
 		bool bIsHandled = false;
 		if (!Playlists.IsEmpty())
@@ -1019,7 +1019,7 @@ bool SAvaPageList::HandleDropEvent(const FDragDropEvent& InDragDropEvent, EItemD
 
 	if (const TSharedPtr<FExternalDragOperation> ExternalDragDropOp = InDragDropEvent.GetOperationAs<FExternalDragOperation>())
 	{
-		FScopedTransaction DropTransaction(LOCTEXT("DropTransaction", "Drop External Asset onto Page List"));
+		FScopedTransaction DropTransaction(LOCTEXT("ExternalDropTransaction", "Drop External Asset onto Page List"));
 
 		return ExternalDragDropOp->HasFiles() && HandleDropExternalFiles(ExternalDragDropOp->GetFiles(), InDropZone, InItem);
 	}
diff --git a/Engine/Plugins/Experimental/Avalanche/Source/AvalancheModifiers/Private/Modifiers/AvaDynamicMeshConverterModifier.cpp b/Engine/Plugins/Experimental/Avalanche/Source/AvalancheModifiers/Private/Modifiers/AvaDynamicMeshConverterModifier.cpp
--- a/Engine/Plugins/Experimental/Avalanche/Source/AvalancheModifiers/Private/Modifiers/AvaDynamicMeshConverterModifier.cpp
+++ b/Engine/Plugins/Experimental/Avalanche/Source/AvalancheModifiers/Private/Modifiers/AvaDynamicMeshConverterModifier.cpp
@@ -314,7 +314,6 @@ void UAvaDynamicMeshConverterModifier::ConvertToStaticMeshAsset()
 		FAssetRegistryModule::AssetCreated(DestinationMesh);
 	}
 }
-#undef LOCTEXT_NAMESPACE
 #endif
 
 void UAvaDynamicMeshConverterModifier::SetSourceActorWeak(const TWeakObjectPtr<AActor>& InActor)
diff --git a/Engine/Plugins/Experimental/Avalanche/Source/AvalancheOutliner/Private/AvaOutlinerView.cpp b/Engine/Plugins/Experimental/Avalanche/Source/AvalancheOutliner/Private/AvaOutlinerView.cpp
--- a/Engine/Plugins/Experimental/Avalanche/Source/AvalancheOutliner/Private/AvaOutlinerView.cpp
+++ b/Engine/Plugins/Experimental/Avalanche/Source/AvalancheOutliner/Private/AvaOutlinerView.cpp
@@ -1654,7 +1654,7 @@ void FAvaOutlinerView::PopulateItemContextMenu(UToolMenu* InToolMenu)
 
 	// Generic Actions
 	{
-		FToolMenuSection& GenericActionsSection = InToolMenu->AddSection(TEXT("GenericActions"), LOCTEXT("OutlinerActionsHeader", "Generic Actions"));
+		FToolMenuSection& GenericActionsSection = InToolMenu->AddSection(TEXT("GenericActions"), LOCTEXT("OutlinerGenericActionsHeader", "Generic Actions"));
 		
 		const FGenericCommands& GenericCommands = FGenericCommands::Get();
 		GenericActionsSection.AddMenuEntry(GenericCommands.Cut);
diff --git a/Engine/Plugins/Experimental/Avalanche/Source/AvalancheSequencer/Private/Commands/Stagger/SAvaSequencerStaggerSettings.cpp b/Engine/Plugins/Experimental/Avalanche/Source/AvalancheSequencer/Private/Commands/Stagger/SAvaSequencerStaggerSettings.cpp
--- a/Engine/Plugins/Experimental/Avalanche/Source/AvalancheSequencer/Private/Commands/Stagger/SAvaSequencerStaggerSettings.cpp
+++ b/Engine/Plugins/Experimental/Avalanche/Source/AvalancheSequencer/Private/Commands/Stagger/SAvaSequencerStaggerSettings.cpp
@@ -152,7 +152,7 @@ TSharedRef<SWidget> SAvaSequencerStaggerSettings::ConstructStartPositionRadioGro
 			.VAlign(VAlign_Center)
 			[
 				SNew(STextBlock)
-				.Text(LOCTEXT("ShiftFromLabel", "First:"))
+				.Text(LOCTEXT("ShiftFirstLabel", "First:"))
 			]
 			+ SHorizontalBox::Slot()
 			.FillWidth(0.65f)
diff --git a/Engine/Plugins/Experimental/Avalanche/Source/AvalancheTextEditor/Private/Slate/SAvaTextAlignmentWidget.cpp b/Engine/Plugins/Experimental/Avalanche/Source/AvalancheTextEditor/Private/Slate/SAvaTextAlignmentWidget.cpp
--- a/Engine/Plugins/Experimental/Avalanche/Source/AvalancheTextEditor/Private/Slate/SAvaTextAlignmentWidget.cpp
+++ b/Engine/Plugins/Experimental/Avalanche/Source/AvalancheTextEditor/Private/Slate/SAvaTextAlignmentWidget.cpp
@@ -80,7 +80,7 @@ void SAvaTextAlignmentWidget::Construct(const FArguments& InArgs)
 			]
 			+ SGridPanel::Slot(3, 1, ForegroundLayer)
 			[
-				GetVerticalAlignmentButton(VerticalBottomButton, EText3DVerticalTextAlignment::Bottom, "AvalancheIcons.Alignment.Bottom", LOCTEXT("AlignSelectedTextRight", "Align Text to Bottom"))
+				GetVerticalAlignmentButton(VerticalBottomButton, EText3DVerticalTextAlignment::Bottom, "AvalancheIcons.Alignment.Bottom", LOCTEXT("AlignSelectedTextBottom", "Align Text to Bottom"))
 			]
 		]
 	];
