MotionDesign
- Fix global opacity modifier issue when hitting 0 with sequencer and TL in playlist

diff --git a/Engine/Plugins/Experimental/Avalanche/Source/AvalancheModifiers/Private/Modifiers/AvaGlobalOpacityModifier.cpp b/Engine/Plugins/Experimental/Avalanche/Source/AvalancheModifiers/Private/Modifiers/AvaGlobalOpacityModifier.cpp
--- a/Engine/Plugins/Experimental/Avalanche/Source/AvalancheModifiers/Private/Modifiers/AvaGlobalOpacityModifier.cpp
+++ b/Engine/Plugins/Experimental/Avalanche/Source/AvalancheModifiers/Private/Modifiers/AvaGlobalOpacityModifier.cpp
@@ -24,7 +24,7 @@ void UAvaGlobalOpacityModifier::PostEditChangeProperty(FPropertyChangedEvent& Pr
 	Super::PostEditChangeProperty(PropertyChangedEvent);
 
 	const FName MemberName = PropertyChangedEvent.GetMemberPropertyName();
-	
+
 	static const FName GlobalOpacityName = GET_MEMBER_NAME_CHECKED(UAvaGlobalOpacityModifier, GlobalOpacity);
 
 	if (MemberName == GlobalOpacityName)
@@ -36,7 +36,7 @@ void UAvaGlobalOpacityModifier::PostEditChangeProperty(FPropertyChangedEvent& Pr
 
 void UAvaGlobalOpacityModifier::SetGlobalOpacity(float InOpacity)
 {
-	if (GlobalOpacity == InOpacity)
+	if (FMath::IsNearlyEqual(GlobalOpacity, InOpacity))
 	{
 		return;
 	}
@@ -47,12 +47,11 @@ void UAvaGlobalOpacityModifier::SetGlobalOpacity(float InOpacity)
 
 void UAvaGlobalOpacityModifier::OnGlobalOpacityChanged()
 {
-	GlobalOpacity = FMath::Clamp<float>(GlobalOpacity, 0.f, 1.f);
-	if (float* GlobalOpacityPtr = MaterialParameters.ScalarParameters.Find(UAvaGlobalOpacityModifier::MaterialDesignerGlobalOpacityValueName))
-	{
-		(*GlobalOpacityPtr) = GlobalOpacity;
-		OnMaterialParametersChanged();
-	}
+	GlobalOpacity = FMath::Clamp<float>(GlobalOpacity, UE_SMALL_NUMBER * 2, 1.f);
+
+	float& GlobalOpacityRef = MaterialParameters.ScalarParameters.FindChecked(UAvaGlobalOpacityModifier::MaterialDesignerGlobalOpacityValueName);
+	GlobalOpacityRef = GlobalOpacity;
+	OnMaterialParametersChanged();
 }
 
 void UAvaGlobalOpacityModifier::OnActorMaterialAdded(UMaterialInstanceDynamic* InAdded)
@@ -77,7 +76,7 @@ void UAvaGlobalOpacityModifier::OnActorMaterialAdded(UMaterialInstanceDynamic* I
 void UAvaGlobalOpacityModifier::OnActorMaterialRemoved(UMaterialInstanceDynamic* InRemoved)
 {
 	Super::OnActorMaterialRemoved(InRemoved);
-	
+
 #if WITH_EDITOR
 	if (UDynamicMaterialInstance* MDI = Cast<UDynamicMaterialInstance>(InRemoved))
 	{
@@ -94,9 +93,14 @@ void UAvaGlobalOpacityModifier::OnActorMaterialRemoved(UMaterialInstanceDynamic*
 
 void UAvaGlobalOpacityModifier::OnDynamicMaterialValueChanged(UDMMaterialComponent* InComponent, EDMUpdateType InUpdateType)
 {
+	if (InUpdateType != EDMUpdateType::Value)
+	{
+		return;
+	}
+
 	if (const UDMMaterialValueFloat1* FloatValue = Cast<UDMMaterialValueFloat1>(InComponent))
 	{
-		if (FloatValue->GetValue() != GlobalOpacity)
+		if (!FMath::IsNearlyEqual(FloatValue->GetValue(), GlobalOpacity))
 		{
 			MarkModifierDirty();
 		}
