Avalanche Transition:
- Trivial fix of swapping the start/end parameters of ChangeOut and ChangeIn for a default TL

diff --git a/Engine/Plugins/Experimental/Avalanche/Source/AvalancheEditor/Private/Transition/AvaTransitionExtension.cpp b/Engine/Plugins/Experimental/Avalanche/Source/AvalancheEditor/Private/Transition/AvaTransitionExtension.cpp
--- a/Engine/Plugins/Experimental/Avalanche/Source/AvalancheEditor/Private/Transition/AvaTransitionExtension.cpp
+++ b/Engine/Plugins/Experimental/Avalanche/Source/AvalancheEditor/Private/Transition/AvaTransitionExtension.cpp
@@ -170,18 +170,18 @@ void FAvaTransitionExtension::BuildDefaultTransitionTree(UAvaTransitionTreeEdito
 		const FAvaTagHandle ChangeTag   = DefaultTags.Change.MakeTagHandle();
 		const FAvaSequenceTime MarkTime = FAvaSequenceTime(TEXT("A"));
 
-		// Change out
+		// Change out (0 to A)
 		{
 			UStateTreeState& ChangeOut = State.AddChildState(TEXT("Change Out"));
 			AddTransitionCondition(ChangeOut , EAvaTransitionType::Out);
 			SetStateCompletedResult(ChangeOut, EStateTreeTransitionType::Succeeded);
 
 			FAvaSequencePlayParams& PlaySettings = AddPlayTask(ChangeOut, ChangeTag);
-			PlaySettings.Start = MarkTime;
-			PlaySettings.End   = FAvaSequenceTime(-1);
+			PlaySettings.Start = FAvaSequenceTime(0);
+            PlaySettings.End   = MarkTime;
 		}
 
-		// Change in
+		// Change in (A to End)
 		{
 			UStateTreeState& ChangeIn = State.AddChildState(TEXT("Change In"));
 			AddTransitionCondition(ChangeIn , EAvaTransitionType::In);
@@ -196,8 +196,8 @@ void FAvaTransitionExtension::BuildDefaultTransitionTree(UAvaTransitionTreeEdito
 			SetStateCompletedResult(PlayChangeInState, EStateTreeTransitionType::Succeeded);
 
 			FAvaSequencePlayParams& PlaySettings = AddPlayTask(PlayChangeInState, ChangeTag);
-			PlaySettings.Start = FAvaSequenceTime(0);
-			PlaySettings.End   = MarkTime;
+			PlaySettings.Start = MarkTime;
+			PlaySettings.End   = FAvaSequenceTime(-1);
 		}
 	}
 
