Motion Design
- Changed "Avalanche" prefix to "Ava" for AvaGameInstance, AvaGameViewportClient and AvaInstanceSettings.

diff --git a/Engine/Plugins/Experimental/Avalanche/Config/DefaultAvalanche.ini b/Engine/Plugins/Experimental/Avalanche/Config/DefaultAvalanche.ini
--- a/Engine/Plugins/Experimental/Avalanche/Config/DefaultAvalanche.ini
+++ b/Engine/Plugins/Experimental/Avalanche/Config/DefaultAvalanche.ini
@@ -160,5 +160,8 @@
 +StructRedirects=(OldName="/Script/Avalanche.AvalancheComponentData",NewName="/Script/Avalanche.AvaComponentData")
 +StructRedirects=(OldName="/Script/Avalanche.AvalancheActorData",NewName="/Script/Avalanche.AvaActorData")
 
++StructRedirects=(OldName="/Script/Avalanche.AvalancheInstanceSettings",NewName="/Script/Avalanche.AvaInstanceSettings")
++PropertyRedirects=(OldName="/Script/AvalancheMedia.AvalancheMediaSettings.AvalancheInstanceSettings",NewName="/Script/AvalancheMedia.AvalancheMediaSettings.AvaInstanceSettings")
+
 [/Script/Avalanche.AvaAnimation]
 DefaultCompletionMode=KeepState
diff --git a/Engine/Plugins/Experimental/Avalanche/Source/Avalanche/Private/AvaActorUtils.cpp b/Engine/Plugins/Experimental/Avalanche/Source/Avalanche/Private/AvaActorUtils.cpp
--- a/Engine/Plugins/Experimental/Avalanche/Source/Avalanche/Private/AvaActorUtils.cpp
+++ b/Engine/Plugins/Experimental/Avalanche/Source/Avalanche/Private/AvaActorUtils.cpp
@@ -8,7 +8,7 @@
 #include "Components/SceneComponent.h"
 #include "Engine/World.h"
 #include "EngineUtils.h"
-#include "Framework/AvalancheGameInstance.h"
+#include "Framework/AvaGameInstance.h"
 #include "GameFramework/Actor.h"
 #include "Math/MathFwd.h"
 #include "Math/OrientedBox.h"
diff --git a/Engine/Plugins/Experimental/Avalanche/Source/Avalanche/Private/Framework/AvaGameInstance.cpp b/Engine/Plugins/Experimental/Avalanche/Source/Avalanche/Private/Framework/AvaGameInstance.cpp
new file mode 100644
--- /dev/null
+++ b/Engine/Plugins/Experimental/Avalanche/Source/Avalanche/Private/Framework/AvaGameInstance.cpp
@@ -0,0 +1,510 @@
+// Copyright Epic Games, Inc. All Rights Reserved.
+
+#include "Framework/AvaGameInstance.h"
+#include "AudioDevice.h"
+#include "AvaBlueprint.h"
+#include "AvaRemoteControlRebind.h"
+#include "Components/ReflectionCaptureComponent.h"
+#include "Components/SkyLightComponent.h"
+#include "Engine/Console.h"
+#include "Engine/Engine.h"
+#include "Engine/LocalPlayer.h"
+#include "Engine/TextureRenderTarget2D.h"
+#include "Framework/Application/SlateApplication.h"
+#include "Misc/TimeGuard.h"
+#include "Slate/SceneViewport.h"
+#include "Viewport/AvaCameraManager.h"
+
+namespace UE::AvaGameInstance::Private
+{
+	/* Frame number marked where a synchronous asset loading occured. */
+	static uint64 MarkedSynchronousAssetLoadingFrame = 0;
+
+	/* Frame number marked where delta seconds are ignored. */
+	static uint64 MarkedIgnoreDeltaSecondsFrame = 0;
+
+	static void UpdateFrameMarkers()
+	{
+		// If the last frame had sync asset loading, mark current frame to ignore delta seconds.
+		if (MarkedSynchronousAssetLoadingFrame != 0 && GFrameCounter > MarkedSynchronousAssetLoadingFrame)
+		{
+			MarkedIgnoreDeltaSecondsFrame = GFrameCounter;
+			MarkedSynchronousAssetLoadingFrame = 0;
+		}
+
+		// Reset marker if stale.
+		if (MarkedIgnoreDeltaSecondsFrame != 0 && GFrameCounter > MarkedIgnoreDeltaSecondsFrame)
+		{
+			MarkedIgnoreDeltaSecondsFrame = 0;
+		}
+	}
+	
+	static void MarkSynchronousAssetLoadingThisFrame()
+	{
+		UpdateFrameMarkers();
+		MarkedSynchronousAssetLoadingFrame = GFrameCounter;
+	}
+	
+	static bool ShouldIgnoreDeltaSecondsForCurrentFrame()
+	{
+		UpdateFrameMarkers();
+		return MarkedIgnoreDeltaSecondsFrame != 0 && GFrameCounter == MarkedIgnoreDeltaSecondsFrame;
+	}
+}
+
+UAvaGameInstance::FOnAvaGameInstanceEvent UAvaGameInstance::OnEndPlay;
+UAvaGameInstance::FOnAvaGameInstanceEvent UAvaGameInstance::OnRenderTargetReady;
+
+UAvaGameInstance* UAvaGameInstance::Create(UObject* InOuter, const TSoftObjectPtr<UAvalancheBlueprint>& InBlueprintTemplate)
+{
+	UAvaGameInstance* GameInstance = NewObject<UAvaGameInstance>(InOuter ? InOuter : GEngine);
+	GameInstance->CreateWorld();
+	GameInstance->LoadMotionDesignBlueprint(InBlueprintTemplate);
+	return GameInstance;
+}
+
+UAvaGameInstance* UAvaGameInstance::Create(UObject* InOuter)
+{
+	UAvaGameInstance* GameInstance = NewObject<UAvaGameInstance>(InOuter ? InOuter : GEngine);
+
+	// We need to call Init() (and corresponding Shutdown() later) to get SubsystemCollection to initialize.
+	// Not doing it because it might have side effects.
+	//GameInstance->Init();
+
+	GameInstance->CreateWorld();
+	return GameInstance;
+}
+
+bool UAvaGameInstance::CreateWorld()
+{
+	if (bWorldCreated)
+	{
+		return false;
+	}
+	
+	TRACE_CPUPROFILER_EVENT_SCOPE(UAvaGameInstance::CreateWorld);
+
+	if (!EnginePreExitHandle.IsValid())
+	{
+		EnginePreExitHandle = FCoreDelegates::OnEnginePreExit.AddUObject(this, &UAvaGameInstance::OnEnginePreExit);
+	}
+	
+	constexpr EWorldType::Type WorldType = EWorldType::Game;
+	
+	{
+		TRACE_CPUPROFILER_EVENT_SCOPE(UAvaGameInstance::LoadInstance::InitWorld);
+		
+		const FName MotionDesignWorldName = MakeUniqueObjectName(GetOuter(), UWorld::StaticClass(), TEXT("AvaGameInstanceWorld"));
+		
+		PlayWorld = NewObject<UWorld>(this, MotionDesignWorldName, RF_Transient);
+		check(PlayWorld);
+		PlayWorld->WorldType = WorldType;
+		PlayWorld->InitializeNewWorld(UWorld::InitializationValues()
+			.AllowAudioPlayback(true)
+			.CreatePhysicsScene(true)
+			.RequiresHitProxies(false)
+			.CreateNavigation(true)
+			.CreateAISystem(false)	// Disabling AI System until supported in Motion Design.
+			.ShouldSimulatePhysics(true)
+			.SetTransactional(false)
+			//TODO: World Settings
+		);
+
+		// Disallow the Engine to tick this World as this Game Instance will be in charge of ticking it
+		PlayWorld->SetShouldTick(false);
+	}
+
+	WorldContext = &GEngine->CreateNewWorldContext(WorldType);
+	check(WorldContext);
+	WorldContext->OwningGameInstance = this;
+	WorldContext->SetCurrentWorld(PlayWorld.Get());
+
+	PlayWorld->SetGameInstance(this);
+	PlayWorld->SetGameMode(FURL());
+
+	bWorldCreated = true;
+	return true;
+}
+
+bool UAvaGameInstance::LoadMotionDesignBlueprint(const TSoftObjectPtr<UAvalancheBlueprint>& InBlueprint)
+{
+	if (ManagedBlueprint && SourceBlueprint == InBlueprint)
+	{
+		return false;
+	}
+
+	if (!bWorldCreated)
+	{
+		CreateWorld();
+	}
+
+	UAvalancheBlueprint* Source;
+	{
+		TRACE_CPUPROFILER_EVENT_SCOPE(UAvaGameInstance::LoadAsset::LoadSourceAvaBp);
+		Source = InBlueprint.LoadSynchronous();
+		check(Source);
+	}
+
+#if WITH_EDITORONLY_DATA
+	FGuardValue_Bitfield(Source->bDuplicatingReadOnly, true);
+#endif
+
+	{
+		TRACE_CPUPROFILER_EVENT_SCOPE(UAvaGameInstance::LoadAsset::DupAvaBp);
+		const FName AvaInstanceName = MakeUniqueObjectName(this, UAvalancheBlueprint::StaticClass(), Source->GetFName());
+		ManagedBlueprint = Cast<UAvalancheBlueprint>(StaticDuplicateObject(Source, this, AvaInstanceName, RF_Transient));
+	}
+	
+	{
+		TRACE_CPUPROFILER_EVENT_SCOPE(UAvaGameInstance::LoadAsset::LoadAvaBpWorld);
+		ManagedBlueprint->SetAvalancheWorld(PlayWorld.Get());
+		ManagedBlueprint->LoadAvalancheWorld();
+	}
+	
+	ManagedBlueprint->UpdatePlaceholderActor();
+
+	{
+		TRACE_CPUPROFILER_EVENT_SCOPE(UAvaGameInstance::LoadAsset::RebindRCP);
+		ManagedBlueprint->RegisterRemoteControlPreset();
+		FAvaRemoteControlRebind::RebindUnboundEntities(ManagedBlueprint->GetRemoteControlPreset());
+	}
+	
+	MarkSynchronousAssetLoadingThisFrame();
+
+	SourceBlueprint = InBlueprint;
+	return true;
+}
+
+bool UAvaGameInstance::BeginPlayWorld(const FAvaInstancePlaySettings& InWorldPlaySettings)
+{
+	// Make sure we don't have pending unload or stop requests left over in the game instance. 
+	CancelWorldRequests();
+	
+	if (bWorldPlaying)
+	{
+		return false;
+	}
+
+	TRACE_CPUPROFILER_EVENT_SCOPE(UAvaGameInstance::BeginPlay);
+	TRACE_BOOKMARK(TEXT("UAvaGameInstance::BeginPlay"));
+	
+	ViewportClient = NewObject<UAvaGameViewportClient>(GEngine, NAME_None, RF_Transient);
+	check(ViewportClient);
+
+	// Note: UGameViewportClient::Init ignores "bCreateNewAudioDevice" parameter and always create a device for the world. 
+	ViewportClient->Init(*WorldContext, this);
+	ViewportClient->SetRenderTarget(InWorldPlaySettings.RenderTarget);
+	ViewportClient->bIsPlayInEditorViewport = false;
+	FAvaViewportQualitySettings QualitySettingsMutable = InWorldPlaySettings.QualitySettings; 
+	QualitySettingsMutable.Apply(ViewportClient->EngineShowFlags);
+	
+#if ALLOW_CONSOLE
+	// Create the viewport's console.
+	ViewportClient->ViewportConsole = NewObject<UConsole>(ViewportClient.Get(), GEngine->ConsoleClass);
+	// register console to get all log messages
+	GLog->AddOutputDevice(ViewportClient->ViewportConsole);
+#endif
+
+	TSharedPtr<SViewport> ViewportWidget;
+	Viewport = MakeShareable(ViewportClient->CreateGameViewport(ViewportWidget));
+	Viewport->SetInitialSize(InWorldPlaySettings.ViewportSize);
+
+	PlayWorld->InitializeActorsForPlay(PlayWorld->URL);
+	PlayWorld->BeginPlay();
+	
+	FCoreDelegates::OnEndFrame.AddUObject(this, &UAvaGameInstance::OnEndFrameTick);
+
+	bWorldPlaying = true;
+	PlayingChannelName = InWorldPlaySettings.ChannelName;
+	return true;
+}
+
+void UAvaGameInstance::RequestEndPlayWorld(bool bForceImmediate)
+{
+	if (bWorldPlaying)
+	{
+		if (bForceImmediate)
+		{
+			EndPlayWorld();
+		}
+		else
+		{
+			bRequestEndPlayWorld = true;
+		}
+	}
+}
+
+void UAvaGameInstance::BeginPlayMotionDesignBlueprint(const TSoftObjectPtr<UAvalancheBlueprint>& InBlueprint, const FAvaInstancePlaySettings& InWorldPlaySettings)
+{
+	// World States: Created -> Playing
+	// Asset States: Unloaded -> Loading -> Loaded -> Playing (Visible)
+	const FSoftObjectPath& SourceAssetPath = InBlueprint.ToSoftObjectPath();
+
+	if (!bWorldPlaying)
+	{
+		BeginPlayWorld(InWorldPlaySettings);
+	}
+	else
+	{
+		UpdateRenderTarget(InWorldPlaySettings.RenderTarget);
+		UpdateSceneViewportSize(InWorldPlaySettings.ViewportSize);
+	}
+
+	// Check if the asset is loaded.
+	if (SourceBlueprint != InBlueprint)
+	{
+		LoadMotionDesignBlueprint(InBlueprint);
+	}	
+	
+	if (ManagedBlueprint && SourceBlueprint == InBlueprint)
+	{
+		// Verify if this can be changed on the fly.
+		ManagedBlueprint->GetViewportQualitySettings().Apply(ViewportClient->EngineShowFlags);
+	
+		// Old code using ava camera manager. To retire.
+		constexpr bool bIsCanvasController = false;
+		ViewportClient->GetCameraManager()->Init(ManagedBlueprint->GetPlaybackObject(), bIsCanvasController);
+#if WITH_EDITOR
+		ViewportClient->GetCameraManager()->SetDefaultViewTarget(PlayWorld, ManagedBlueprint->GetStartupCameraName());
+#endif
+	}
+}
+
+void UAvaGameInstance::RequestUnloadWorld(bool bForceImmediate)
+{
+	if (bWorldCreated)
+	{
+		if (bForceImmediate)
+		{
+			UnloadWorld();
+		}
+		else
+		{
+			bRequestUnloadWorld = true;
+		}
+	}
+}
+
+void UAvaGameInstance::CancelWorldRequests()
+{
+	bRequestEndPlayWorld = false;
+	bRequestUnloadWorld = false;
+}
+
+void UAvaGameInstance::UpdateRenderTarget(UTextureRenderTarget2D* InRenderTarget)
+{
+	if (ViewportClient)
+	{
+		ViewportClient->SetRenderTarget(InRenderTarget);
+	}
+}
+
+UTextureRenderTarget2D* UAvaGameInstance::GetRenderTarget() const
+{
+	if (ViewportClient)
+	{
+		return ViewportClient->GetRenderTarget();
+	}
+	return nullptr;
+}
+
+void UAvaGameInstance::UpdateSceneViewportSize(const FIntPoint& InViewportSize)
+{
+	if (Viewport.IsValid() && Viewport->GetSize() != InViewportSize)
+	{
+		// Note: calling either SetViewportSize or SetFixedViewportSize doesn't work
+		// because ViewportWidget is null in this case. And we can't call ResizeViewport
+		// either because it is private. Calling the only function we can call.		
+		Viewport->UpdateViewportRHI(false, InViewportSize.X, InViewportSize.Y, EWindowMode::Type::Windowed, PF_Unknown);
+	}
+}
+
+void UAvaGameInstance::MarkSynchronousAssetLoadingThisFrame()
+{
+	UE::AvaGameInstance::Private::MarkSynchronousAssetLoadingThisFrame();
+}
+
+void UAvaGameInstance::Tick(float DeltaSeconds)
+{
+	if (bIsTicking)
+	{
+		return;
+	}
+	TGuardValue TickGuard(bIsTicking, true);
+
+	if (!bWorldPlaying)
+	{
+		return;
+	}
+
+	if (!PlayWorld)
+	{
+		EndPlayWorld();
+		return;
+	}
+	
+	{
+		QUICK_SCOPE_CYCLE_COUNTER(STAT_MotionDesignGame_Tick_SetDropDetail);
+		ViewportClient->SetDropDetail(DeltaSeconds);
+	}
+	
+	/*
+	 * Things Removed Here because they're currently unneeded and will possibly never be
+	 * TickWorldTravel
+	 */
+
+	{
+		SCOPE_TIME_GUARD(TEXT("UAvaGameInstance::Tick - World Tick"));
+		PlayWorld->Tick(ELevelTick::LEVELTICK_All, DeltaSeconds);
+	}
+
+	//Removed check for Dedicated Server & Commandlet
+	{
+		QUICK_SCOPE_CYCLE_COUNTER(STAT_UGameEngine_Tick_CheckCaptures);
+		
+		// Only update reflection captures in game once all 'always loaded' levels have been loaded
+		// This won't work with actual level streaming though
+		if (PlayWorld->AreAlwaysLoadedLevelsLoaded())
+		{
+			// Update sky light first because it's considered direct lighting, sky diffuse will be visible in reflection capture indirect specular
+			USkyLightComponent::UpdateSkyCaptureContents(PlayWorld.Get());
+			UReflectionCaptureComponent::UpdateReflectionCaptureContents(PlayWorld.Get());
+		}
+	}
+
+	//Rest of the stuff between above Lighting Update and GameViewport Tick removed
+
+	{
+		SCOPE_TIME_GUARD(TEXT("UAvaGameInstance::Tick - TickViewport"));
+		QUICK_SCOPE_CYCLE_COUNTER(STAT_MotionDesignGame_Tick_TickViewport);
+		ViewportClient->Tick(DeltaSeconds);
+	}
+	
+	{
+		Viewport->Draw();
+
+		if (!bIsRenderTargetReady)
+		{
+			if (!RenderTargetFence.IsValid())
+			{
+				RenderTargetFence = MakeUnique<FRenderCommandFence>();
+				RenderTargetFence->BeginFence();
+			}
+			else
+			{
+				bIsRenderTargetReady = RenderTargetFence->IsFenceComplete();
+				OnRenderTargetReady.Broadcast(this, PlayingChannelName);
+			}
+		}
+	}
+
+	if (bRequestEndPlayWorld)
+	{
+		EndPlayWorld();
+	}
+	if (bRequestUnloadWorld)
+	{
+		UnloadWorld();
+	}
+}
+
+void UAvaGameInstance::UnloadWorld()
+{
+	if (bWorldPlaying)
+	{
+		EndPlayWorld();
+	}
+
+	bWorldCreated = false;
+	bRequestUnloadWorld = false;
+
+	if (IsValid(ManagedBlueprint.Get()))
+	{
+		ManagedBlueprint->UnregisterRemoteControlPreset();
+	}
+	SourceBlueprint.Reset();
+	ManagedBlueprint = nullptr;
+	
+	if (PlayWorld)
+	{
+		if (FAudioDeviceHandle AudioDevice = PlayWorld->GetAudioDevice())
+		{
+			AudioDevice->Flush(PlayWorld, false);
+		}
+
+		// Need to do this before destroying the world context apparently.
+		PlayWorld->SetShouldForceUnloadStreamingLevels(true);
+		PlayWorld->FlushLevelStreaming();
+
+		GEngine->DestroyWorldContext(PlayWorld.Get());
+		PlayWorld->DestroyWorld(true);
+	}
+
+	PlayWorld = nullptr;
+	WorldContext = nullptr;
+}
+
+void UAvaGameInstance::EndPlayWorld()
+{
+	bWorldPlaying = false;
+	bRequestEndPlayWorld = false;
+
+	FCoreDelegates::OnEndFrame.RemoveAll(this);
+
+	OnEndPlay.Broadcast(this, PlayingChannelName);
+	PlayingChannelName = NAME_None;
+
+	Viewport.Reset();
+	RenderTargetFence.Reset();
+	bIsRenderTargetReady = false;
+
+#if ALLOW_CONSOLE
+	if (ViewportClient)
+	{
+		GLog->RemoveOutputDevice(ViewportClient->ViewportConsole);
+	}
+#endif
+
+	ViewportClient = nullptr;
+}
+
+void UAvaGameInstance::OnEndFrameTick()
+{
+	double DeltaSeconds = FApp::GetDeltaTime();
+
+	if (UE::AvaGameInstance::Private::ShouldIgnoreDeltaSecondsForCurrentFrame())
+	{
+		DeltaSeconds = LastDeltaSeconds;
+	}
+	else
+	{
+		LastDeltaSeconds = DeltaSeconds;
+	}
+	
+	Tick(DeltaSeconds);
+}
+
+void UAvaGameInstance::OnEnginePreExit()
+{
+	if (PlayWorld && PlayWorld->GetAudioDevice())
+	{
+		// Temporarily set the GIsRequestingExit to false so that the DLLs like XAudio2Dll are not unloaded
+		// for ref: FMixerPlatformXAudio2::TeardownHardware unloads the dll if Engine Exit is requested
+		PRAGMA_DISABLE_DEPRECATION_WARNINGS
+		TGuardValue<bool> RequestExitGuard(GIsRequestingExit, false);
+		PRAGMA_ENABLE_DEPRECATION_WARNINGS
+		const FAudioDeviceHandle EmptyHandle;
+		PlayWorld->SetAudioDevice(EmptyHandle);
+	}
+}
+
+void UAvaGameInstance::BeginDestroy()
+{
+	FCoreDelegates::OnEnginePreExit.Remove(EnginePreExitHandle);
+	EnginePreExitHandle.Reset();
+	
+	EndPlayWorld();
+	UnloadWorld();
+	
+	Super::BeginDestroy();
+}
diff --git a/Engine/Plugins/Experimental/Avalanche/Source/Avalanche/Private/Framework/AvaGameViewportClient.cpp b/Engine/Plugins/Experimental/Avalanche/Source/Avalanche/Private/Framework/AvaGameViewportClient.cpp
new file mode 100644
--- /dev/null
+++ b/Engine/Plugins/Experimental/Avalanche/Source/Avalanche/Private/Framework/AvaGameViewportClient.cpp
@@ -0,0 +1,536 @@
+// Copyright Epic Games, Inc. All Rights Reserved.
+
+#include "Framework/AvaGameViewportClient.h"
+#include "Camera/CameraActor.h"
+#include "Camera/CameraComponent.h"
+#include "Camera/CameraPhotography.h"
+#include "Components/LineBatchComponent.h"
+#include "ContentStreaming.h"
+#include "DynamicResolutionState.h"
+#include "Engine/Canvas.h"
+#include "Engine/LocalPlayer.h"
+#include "Engine/TextureRenderTarget2D.h"
+#include "EngineModule.h"
+#include "EngineUtils.h"
+#include "Framework/Application/SlateApplication.h"
+#include "GameFramework/PlayerController.h"
+#include "IAvaModule.h"
+#include "IHeadMountedDisplay.h"
+#include "IXRCamera.h"
+#include "IXRTrackingSystem.h"
+#include "LegacyScreenPercentageDriver.h"
+#include "Misc/EngineVersionComparison.h"
+#include "SceneViewExtension.h"
+#include "TextureResource.h"
+#include "UObject/Package.h"
+#include "UObject/UObjectGlobals.h"
+
+/** Util to find named canvas in transient package, and create if not found */
+static UCanvas* GetCanvasByName(FName CanvasName)
+{
+	// Cache to avoid FString/FName conversions/compares
+	static TMap<FName, UCanvas*> CanvasMap;
+	UCanvas** FoundCanvas = CanvasMap.Find(CanvasName);
+	if (!FoundCanvas)
+	{
+		UCanvas* CanvasObject = FindObject<UCanvas>(GetTransientPackage(), *CanvasName.ToString(), false);
+		if (!CanvasObject)
+		{
+			CanvasObject = NewObject<UCanvas>(GetTransientPackage(), CanvasName);
+			CanvasObject->AddToRoot();
+		}
+
+		CanvasMap.Add(CanvasName, CanvasObject);
+		return CanvasObject;
+	}
+
+	return *FoundCanvas;
+}
+
+UAvaGameViewportClient::UAvaGameViewportClient()
+	: CameraManager(MakeShared<FAvaCameraManager>())
+{
+}
+
+void UAvaGameViewportClient::Draw(FViewport* InViewport, FCanvas* Canvas)
+{
+	// Override the Canvas Render Target with ours
+	if (RenderTarget.IsValid())
+	{
+		Canvas->SetRenderTarget_GameThread(RenderTarget->GameThread_GetRenderTargetResource());
+	}
+	
+	// Allow HMD to modify the view later, just before rendering
+	const bool bStereoRendering = GEngine->IsStereoscopic3D(InViewport);
+	
+	// Create a temporary canvas if there isn't already one.
+	static FName CanvasObjectName(TEXT("CanvasObject"));
+	UCanvas* CanvasObject = GetCanvasByName(CanvasObjectName);
+	CanvasObject->Canvas = Canvas;
+
+	// Create temp debug canvas object
+	FCanvas* DebugCanvas = InViewport->GetDebugCanvas();
+	FIntPoint DebugCanvasSize = InViewport->GetSizeXY();
+	if (bStereoRendering && GEngine->XRSystem.IsValid() && GEngine->XRSystem->GetHMDDevice())
+	{
+		DebugCanvasSize = GEngine->XRSystem->GetHMDDevice()->GetIdealDebugCanvasRenderTargetSize();
+	}
+
+	static FName DebugCanvasObjectName(TEXT("DebugCanvasObject"));
+	UCanvas* DebugCanvasObject = GetCanvasByName(DebugCanvasObjectName);
+	DebugCanvasObject->Init(DebugCanvasSize.X, DebugCanvasSize.Y, nullptr, DebugCanvas);
+	if (DebugCanvas)
+	{
+		DebugCanvas->SetScaledToRenderTarget(bStereoRendering);
+		DebugCanvas->SetStereoRendering(bStereoRendering);
+	}
+	Canvas->SetScaledToRenderTarget(bStereoRendering);
+	Canvas->SetStereoRendering(bStereoRendering);
+	
+	if (World == nullptr)
+	{
+		return;
+	}
+
+	constexpr bool bCaptureNeedsSceneColor = false;
+	constexpr ESceneCaptureSource SceneCaptureSource = ESceneCaptureSource::SCS_FinalColorLDR;
+	constexpr ESceneCaptureCompositeMode SceneCaptureCompositeMode = ESceneCaptureCompositeMode::SCCM_Overwrite;
+
+	/** When enabled, the scene capture will composite into the render target instead of overwriting its contents. */
+
+	//TODO: World->GetTime() seems to return zeros because we have not Begun Play...
+	FGameTime Time = FGameTime::GetTimeSinceAppStart();
+
+	// Setup a FSceneViewFamily/FSceneView for the viewport.
+	FSceneViewFamilyContext ViewFamilyContext(FSceneViewFamily::ConstructionValues(Canvas->GetRenderTarget()
+		, World->Scene
+		, EngineShowFlags)
+		.SetRealtimeUpdate(true)
+		.SetResolveScene(!bCaptureNeedsSceneColor)
+		.SetTime(Time));
+
+	ViewFamilyContext.SceneCaptureSource = SceneCaptureSource;
+	ViewFamilyContext.SceneCaptureCompositeMode = SceneCaptureCompositeMode;
+	ViewFamilyContext.DebugDPIScale = GetDPIScale();
+	ViewFamilyContext.EngineShowFlags = EngineShowFlags;
+
+#if WITH_EDITOR
+	if (GIsEditor)
+	{
+		// Force enable view family show flag for HighDPI derived's screen percentage.
+		ViewFamilyContext.EngineShowFlags.ScreenPercentage = true;
+	}
+#endif
+
+	if (!bStereoRendering)
+	{
+		// stereo is enabled, as many HMDs require this for proper visuals
+		ViewFamilyContext.EngineShowFlags.SetScreenPercentage(false);
+	}
+
+	ViewFamilyContext.ViewExtensions = GEngine->ViewExtensions->GatherActiveExtensions(FSceneViewExtensionContext(InViewport));
+	
+	for (const TSharedRef<ISceneViewExtension>& ViewExt : ViewFamilyContext.ViewExtensions)
+	{
+		ViewExt->SetupViewFamily(ViewFamilyContext);
+	}
+
+	if (bStereoRendering && GEngine->XRSystem.IsValid() && GEngine->XRSystem->GetHMDDevice())
+	{
+		// Allow HMD to modify screen settings
+		GEngine->XRSystem->GetHMDDevice()->UpdateScreenSettings(InViewport);
+	}
+
+	ViewFamilyContext.ViewMode = EViewModeIndex::VMI_Lit;
+	EngineShowFlagOverride(ESFIM_Game, ViewFamilyContext.ViewMode, ViewFamilyContext.EngineShowFlags, false);
+
+	bool bFinalScreenPercentageShowFlag;
+	bool bUsesDynamicResolution = false;
+	// Setup the screen percentage and upscaling method for the view family.
+	{
+		checkf(ViewFamilyContext.GetScreenPercentageInterface() == nullptr,
+			TEXT("Some code has tried to set up an alien screen percentage driver, that could be wrong if not supported very well by the RHI."));
+
+		// Force screen percentage show flag to be turned off if not supported.
+		if (!ViewFamilyContext.SupportsScreenPercentage())
+		{
+			ViewFamilyContext.EngineShowFlags.ScreenPercentage = false;
+		}
+
+		// Set up secondary resolution fraction for the view family.
+		if (!bStereoRendering && ViewFamilyContext.SupportsScreenPercentage())
+		{
+			// Automatically compute secondary resolution fraction from DPI.
+			ViewFamilyContext.SecondaryViewFraction = GetDPIDerivedResolutionFraction();
+		}
+
+		bFinalScreenPercentageShowFlag = ViewFamilyContext.EngineShowFlags.ScreenPercentage;
+	}
+
+	TArray<FSceneView*> Views;
+	{
+		const int32 NumViews = bStereoRendering ? GEngine->StereoRenderingDevice->GetDesiredNumberOfViews(bStereoRendering) : 1;
+		for (int32 ViewIndex = 0; ViewIndex < NumViews; ++ViewIndex)
+		{
+			FSceneView* const View = CalcSceneView(&ViewFamilyContext
+				, InViewport
+				, nullptr
+				, bStereoRendering ? ViewIndex : INDEX_NONE);
+
+			if (View)
+			{
+				Views.Add(View);
+
+				View->CameraConstrainedViewRect = View->UnscaledViewRect;
+
+				// Add view information for resource streaming. Allow up to 5X boost for small FOV.
+				const float StreamingScale = 1.f / FMath::Clamp<float>(View->LODDistanceFactor, .2f, 1.f);
+
+				IStreamingManager::Get().AddViewInformation(View->ViewMatrices.GetViewOrigin()
+					, View->UnscaledViewRect.Width()
+					, View->UnscaledViewRect.Width() * View->ViewMatrices.GetProjectionMatrix().M[0][0]
+					, StreamingScale);
+
+				World->ViewLocationsRenderedLastFrame.Add(View->ViewMatrices.GetViewOrigin());
+			}
+		}
+	}
+
+	//FinalizeViews(&ViewFamily); TODO Possible Interface
+
+	// Update level streaming.
+	World->UpdateLevelStreaming();
+
+	Canvas->Clear(FLinearColor::Transparent);
+
+	// If a screen percentage interface was not set by one of the view extension, then set the legacy one.
+	if (ViewFamilyContext.GetScreenPercentageInterface() == nullptr)
+	{
+		float GlobalResolutionFraction = 1.0f;
+
+		if (ViewFamilyContext.EngineShowFlags.ScreenPercentage && !bDisableWorldRendering && ViewFamilyContext.Views.Num() > 0)
+		{
+			// Get global view fraction.
+			FStaticResolutionFractionHeuristic StaticHeuristic;
+			StaticHeuristic.Settings.PullRunTimeRenderingSettings(GetViewStatusForScreenPercentage());
+			StaticHeuristic.PullViewFamilyRenderingSettings(ViewFamilyContext);
+			StaticHeuristic.DPIScale = GetDPIScale();
+
+			GlobalResolutionFraction = StaticHeuristic.ResolveResolutionFraction();
+		}
+
+		ViewFamilyContext.SetScreenPercentageInterface(new FLegacyScreenPercentageDriver(
+			ViewFamilyContext, GlobalResolutionFraction));
+	}
+
+	check(ViewFamilyContext.GetScreenPercentageInterface() != nullptr);
+
+	// Make sure the engine show flag for screen percentage is still what it was when setting up the screen percentage interface
+	ViewFamilyContext.EngineShowFlags.ScreenPercentage = bFinalScreenPercentageShowFlag;
+
+	if (bStereoRendering && bUsesDynamicResolution)
+	{
+		// Change screen percentage method to raw output when doing dynamic resolution with VR if not using TAA upsample.
+		for (FSceneView* const View : Views)
+		{
+			if (View->PrimaryScreenPercentageMethod == EPrimaryScreenPercentageMethod::SpatialUpscale)
+			{
+				View->PrimaryScreenPercentageMethod = EPrimaryScreenPercentageMethod::RawOutput;
+			}
+		}
+	}
+
+	ViewFamilyContext.bIsHDR = InViewport->IsHDRViewport();
+
+	if (!bDisableWorldRendering && FSlateApplication::Get().GetPlatformApplication()->IsAllowedToRender()) //-V560
+	{
+		for (const TSharedRef<ISceneViewExtension>& ViewExt : ViewFamilyContext.ViewExtensions)
+		{
+			for (FSceneView* const View : Views)
+			{
+				ViewExt->SetupView(ViewFamilyContext, *View);
+			}
+		}
+		GetRendererModule().BeginRenderingViewFamily(Canvas, &ViewFamilyContext);
+	}
+	else
+	{
+		GetRendererModule().PerFrameCleanupIfSkipRenderer();
+	}
+
+	// Beyond this point, only UI rendering independent from dynamic resolution.
+	GEngine->EmitDynamicResolutionEvent(EDynamicResolutionStateEvent::EndDynamicResolutionRendering);
+
+	// Remove temporary debug lines.
+	if (World->LineBatcher)
+	{
+		World->LineBatcher->Flush();
+	}
+	if (World->ForegroundLineBatcher)
+	{
+		World->ForegroundLineBatcher->Flush();
+	}
+
+	// Render Stats HUD in the main canvas so that it gets captured
+	// and is displayed in the broadcast channel's outputs.
+	if (!Views.IsEmpty() && IAvaModule::Get().ShouldShowRuntimeStats())
+	{
+		DrawStatsHUD( World, InViewport, Canvas, nullptr, DebugProperties, Views[0]->ViewLocation, Views[0]->ViewRotation);
+	}
+	
+	// Ensure canvas has been flushed before rendering UI
+	Canvas->Flush_GameThread();
+
+	CameraManager->ResetCameraCut();
+}
+
+bool UAvaGameViewportClient::IsStatEnabled(const FString& InName) const
+{
+	// The IAvaModule holds the runtime stats. We want them persistent across all viewports.
+	return IAvaModule::Get().IsRuntimeStatEnabled(InName);
+}
+
+void UAvaGameViewportClient::SetCameraCutThisFrame()
+{
+	bCameraCutThisFrame = true;
+}
+
+void UAvaGameViewportClient::SetRenderTarget(UTextureRenderTarget2D* InRenderTarget)
+{
+	RenderTarget = InRenderTarget;
+}
+
+bool UAvaGameViewportClient::CalcSceneViewInitOptions(FSceneViewInitOptions& OutInitOptions
+	, FViewport* InViewport
+	, FViewElementDrawer* ViewDrawer
+	, int32 StereoViewIndex)
+{
+	QUICK_SCOPE_CYCLE_COUNTER(STAT_AvaMediaCalcSceneViewInitOptions);
+	if (InViewport == nullptr || !World)
+	{
+		return false;
+	}
+
+	// get the projection data
+	if (GetProjectionData(InViewport, /*inout*/ OutInitOptions, StereoViewIndex) == false)
+	{
+		// Return false if this we didn't get back the info we needed
+		return false;
+	}
+
+	// return if we have an invalid view rect
+	if (!OutInitOptions.IsValidViewRectangle())
+	{
+		return false;
+	}
+
+	// Was there a camera cut this frame?
+	OutInitOptions.bInCameraCut = CameraManager->ShouldEnableCameraCut() || bCameraCutThisFrame;
+	bCameraCutThisFrame = false;
+
+	if (GEngine->StereoRenderingDevice.IsValid())
+	{
+		OutInitOptions.StereoPass = GEngine->StereoRenderingDevice->GetViewPassForIndex(StereoViewIndex != INDEX_NONE, StereoViewIndex);
+	}
+
+	const uint32 ViewIndex = StereoViewIndex != INDEX_NONE ? StereoViewIndex : 0;
+
+	// Make sure the ViewStates array has enough elements for the given ViewIndex.
+	{
+		const int32 RequiredViewStates = (ViewIndex + 1) - ViewStates.Num();
+
+		if (RequiredViewStates > 0)
+		{
+			ViewStates.AddDefaulted(RequiredViewStates);
+		}
+	}
+
+	// Allocate the current ViewState if necessary
+	if (ViewStates[ViewIndex].GetReference() == nullptr)
+	{
+#if UE_VERSION_OLDER_THAN(5, 3, 0)
+		const ERHIFeatureLevel::Type FeatureLevel = World ? World->FeatureLevel.GetValue() : GMaxRHIFeatureLevel;
+#else
+		const ERHIFeatureLevel::Type FeatureLevel = World ? World->GetFeatureLevel() : GMaxRHIFeatureLevel;
+#endif
+		ViewStates[ViewIndex].Allocate(FeatureLevel);
+	}
+
+	OutInitOptions.SceneViewStateInterface = ViewStates[ViewIndex].GetReference();
+	OutInitOptions.ViewActor = CameraManager->GetViewTarget();
+	OutInitOptions.PlayerIndex = 0;
+	OutInitOptions.ViewElementDrawer = ViewDrawer;
+	OutInitOptions.BackgroundColor = FLinearColor::Transparent;
+	OutInitOptions.LODDistanceFactor = 1.f;
+	OutInitOptions.StereoViewIndex = StereoViewIndex;
+	OutInitOptions.WorldToMetersScale = World->GetWorldSettings()->WorldToMeters;
+	OutInitOptions.CursorPos = InViewport->HasMouseCapture() ? FIntPoint(-1, -1) : FIntPoint(InViewport->GetMouseX(), InViewport->GetMouseY());
+	OutInitOptions.OriginOffsetThisFrame = World->OriginOffsetThisFrame;
+	return true;
+}
+
+bool UAvaGameViewportClient::GetProjectionData(FViewport* InViewport
+	, FSceneViewProjectionData& ProjectionData
+	, int32 StereoViewIndex) const
+{
+		// If the actor
+	if ((InViewport == nullptr) || (InViewport->GetSizeXY().X == 0) || (InViewport->GetSizeXY().Y == 0))
+	{
+		return false;
+	}
+
+	const FVector2D Origin(0.f, 0.f);
+	const FVector2D Size(1.f, 1.f);
+
+	int32 X = FMath::TruncToInt(Origin.X * InViewport->GetSizeXY().X);
+	int32 Y = FMath::TruncToInt(Origin.Y * InViewport->GetSizeXY().Y);
+
+	X += InViewport->GetInitialPositionXY().X;
+	Y += InViewport->GetInitialPositionXY().Y;
+
+	uint32 SizeX = FMath::TruncToInt(Size.X * InViewport->GetSizeXY().X);
+	uint32 SizeY = FMath::TruncToInt(Size.Y * InViewport->GetSizeXY().Y);
+
+	FIntRect UnconstrainedRectangle = FIntRect(X, Y, X + SizeX, Y + SizeY);
+
+	ProjectionData.SetViewRectangle(UnconstrainedRectangle);
+
+	// Get the viewpoint.
+	FMinimalViewInfo ViewInfo;
+	CameraManager->SetCameraViewInfo(ViewInfo);
+
+	// If stereo rendering is enabled, update the size and offset appropriately for this pass
+	const bool bNeedStereo = StereoViewIndex != INDEX_NONE && GEngine->IsStereoscopic3D();
+	const bool bIsHeadTrackingAllowed = GEngine->XRSystem.IsValid()
+		&& (World != nullptr
+			? GEngine->XRSystem->IsHeadTrackingAllowedForWorld(*World)
+			: GEngine->XRSystem->IsHeadTrackingAllowed());
+
+	if (bNeedStereo)
+	{
+		GEngine->StereoRenderingDevice->AdjustViewRect(StereoViewIndex, X, Y, SizeX, SizeY);
+	}
+
+    FVector StereoViewLocation = ViewInfo.Location;
+    if (bNeedStereo || bIsHeadTrackingAllowed)
+    {
+	    TSharedPtr<IXRCamera> XRCamera = GEngine->XRSystem.IsValid()
+    		? GEngine->XRSystem->GetXRCamera()
+    		: nullptr;
+
+		if (XRCamera.IsValid())
+		{
+			const bool bHasActiveCamera = !!CameraManager->GetCachedCameraComponent();
+			XRCamera->UseImplicitHMDPosition(bHasActiveCamera);
+		}
+
+		if (GEngine->StereoRenderingDevice.IsValid())
+		{
+			GEngine->StereoRenderingDevice->CalculateStereoViewOffset(StereoViewIndex, ViewInfo.Rotation, World->GetWorldSettings()->WorldToMeters, StereoViewLocation);
+		}
+    }
+
+	// Create the view matrix
+	ProjectionData.ViewOrigin = StereoViewLocation;
+	ProjectionData.ViewRotationMatrix = FInverseRotationMatrix(ViewInfo.Rotation) * FMatrix(
+		FPlane(0,	0,	1,	0),
+		FPlane(1,	0,	0,	0),
+		FPlane(0,	1,	0,	0),
+		FPlane(0,	0,	0,	1));
+
+	// @todo viewext this use case needs to be revisited
+	if (!bNeedStereo)
+	{
+		// Create the projection matrix (and possibly constrain the view rectangle)
+		FMinimalViewInfo::CalculateProjectionMatrixGivenView(ViewInfo
+			, GetDefault<ULocalPlayer>()->AspectRatioAxisConstraint
+			, InViewport
+			, ProjectionData);
+
+		for (const TSharedRef<ISceneViewExtension>& ViewExt : GEngine->ViewExtensions->GatherActiveExtensions(FSceneViewExtensionContext(InViewport)))
+        {
+			ViewExt->SetupViewProjectionMatrix(ProjectionData);
+		};
+	}
+	else
+	{
+		// Let the stereoscopic rendering device handle creating its own projection matrix, as needed
+		ProjectionData.ProjectionMatrix = GEngine->StereoRenderingDevice->GetStereoProjectionMatrix(StereoViewIndex);
+
+		// calculate the out rect
+		ProjectionData.SetViewRectangle(FIntRect(X, Y, X + SizeX, Y + SizeY));
+	}
+	return true;
+}
+
+FSceneView* UAvaGameViewportClient::CalcSceneView(FSceneViewFamily* ViewFamily
+	, FViewport* InViewport
+	, FViewElementDrawer* ViewDrawer
+	, int32 StereoViewIndex)
+{
+	//QUICK_SCOPE_CYCLE_COUNTER(STAT_AvaMediaCalcSceneView);
+	FSceneViewInitOptions ViewInitOptions;
+	if (!CalcSceneViewInitOptions(ViewInitOptions, InViewport, ViewDrawer, StereoViewIndex))
+	{
+		return nullptr;
+	}
+
+	// Get the viewpoint...technically doing this twice
+	// but it makes GetProjectionData better
+	FMinimalViewInfo ViewInfo;
+	CameraManager->SetCameraViewInfo(ViewInfo);
+
+	ViewInitOptions.ViewLocation = ViewInfo.Location;
+	ViewInitOptions.ViewRotation = ViewInfo.Rotation;
+	ViewInitOptions.bUseFieldOfViewForLOD = ViewInfo.bUseFieldOfViewForLOD;
+	ViewInitOptions.FOV = ViewInfo.FOV;
+	ViewInitOptions.DesiredFOV = ViewInfo.DesiredFOV;
+
+	// Fill out the rest of the view init options
+	ViewInitOptions.ViewFamily = ViewFamily;
+
+	FSceneView* const View = new FSceneView(ViewInitOptions);
+
+	//Pass on the previous view transform from the view info (probably provided by the camera if set)
+	View->PreviousViewTransform = ViewInfo.PreviousViewTransform;
+
+	ViewFamily->Views.Add(View);
+
+	{
+		View->StartFinalPostprocessSettings(ViewInfo.Location);
+
+		//Post Process Setting
+		//TODO?
+
+		//	CAMERA OVERRIDE
+		View->OverridePostProcessSettings(ViewInfo.PostProcessSettings, ViewInfo.PostProcessBlendWeight);
+		CameraManager->ApplyExtraPostProcessBlends(View);
+		FCameraPhotographyManager::Get().UpdatePostProcessing(View->FinalPostProcessSettings);
+
+		if (GEngine->StereoRenderingDevice.IsValid())
+		{
+			FPostProcessSettings StereoDeviceOverridePostProcessingSettings;
+			float BlendWeight = 1.0f;
+
+			const bool StereoSettingsAvailable = GEngine->StereoRenderingDevice->OverrideFinalPostprocessSettings(&StereoDeviceOverridePostProcessingSettings
+				, View->StereoPass
+				, View->StereoViewIndex
+				, BlendWeight);
+
+			if (StereoSettingsAvailable)
+			{
+				View->OverridePostProcessSettings(StereoDeviceOverridePostProcessingSettings, BlendWeight);
+			}
+		}
+
+
+		View->EndFinalPostprocessSettings(ViewInitOptions);
+	}
+
+	for (int32 ViewExt = 0; ViewExt < ViewFamily->ViewExtensions.Num(); ++ViewExt)
+	{
+		ViewFamily->ViewExtensions[ViewExt]->SetupView(*ViewFamily, *View);
+	}
+	
+	return View;
+}
diff --git a/Engine/Plugins/Experimental/Avalanche/Source/Avalanche/Private/Framework/AvalancheGameInstance.cpp b/Engine/Plugins/Experimental/Avalanche/Source/Avalanche/Private/Framework/AvalancheGameInstance.cpp
deleted file mode 100644
diff --git a/Engine/Plugins/Experimental/Avalanche/Source/Avalanche/Private/Framework/AvalancheGameViewportClient.cpp b/Engine/Plugins/Experimental/Avalanche/Source/Avalanche/Private/Framework/AvalancheGameViewportClient.cpp
deleted file mode 100644
diff --git a/Engine/Plugins/Experimental/Avalanche/Source/Avalanche/Public/Framework/AvaGameInstance.h b/Engine/Plugins/Experimental/Avalanche/Source/Avalanche/Public/Framework/AvaGameInstance.h
new file mode 100644
--- /dev/null
+++ b/Engine/Plugins/Experimental/Avalanche/Source/Avalanche/Public/Framework/AvaGameInstance.h
@@ -0,0 +1,187 @@
+// Copyright Epic Games, Inc. All Rights Reserved.
+
+#pragma once
+
+#include "AvaGameViewportClient.h"
+#include "Engine/GameInstance.h"
+#include "Engine/World.h"
+#include "AvaGameInstance.generated.h"
+
+class FRenderCommandFence;
+class FSceneViewport;
+class UAvalancheBlueprint;
+class UTextureRenderTarget2D;
+struct FAvaViewportQualitySettings;
+struct FAvaInstanceSettings;
+
+struct FAvaInstancePlaySettings
+{
+	const FAvaInstanceSettings& Settings;
+	FName ChannelName;
+	UTextureRenderTarget2D* RenderTarget;
+	FIntPoint ViewportSize;
+	const FAvaViewportQualitySettings& QualitySettings;
+};
+
+UCLASS(DisplayName = "Motion Design Game Instance")
+class AVALANCHE_API UAvaGameInstance : public UGameInstance
+{
+	GENERATED_BODY()
+
+public:
+	/**
+	 * Creates a new Motion Design Game instance from the given Motion Design Template Asset.
+	 * Supported asset types: UMotion DesignBlueprint, UWorld (Level).
+	 *
+	 * Remark: For Motion Design Blueprints, the instanced embedded RemoteControlPreset will be registered
+	 * under the name of the outer Package Name. It is thus required that each instance be contained in
+	 * a unique package to avoid stomping the registered RCPs.
+	 */
+	static UAvaGameInstance* Create(UObject* InOuter);
+
+	// Legacy - will create and load a Motion Design blueprint.
+	static UAvaGameInstance* Create(UObject* InOuter, const TSoftObjectPtr<UAvalancheBlueprint>& InBlueprintTemplate);
+	
+	/** Create the game world */
+	bool CreateWorld();
+
+	bool IsWorldCreated() const { return bWorldCreated; }
+
+	UWorld* GetPlayWorld() const { return PlayWorld.Get(); }
+
+	// Legacy
+	bool LoadMotionDesignBlueprint(const TSoftObjectPtr<UAvalancheBlueprint>& InBlueprint);
+
+	// Legacy
+	void BeginPlayMotionDesignBlueprint(const TSoftObjectPtr<UAvalancheBlueprint>& InBlueprint, const FAvaInstancePlaySettings& InWorldPlaySettings);
+	
+	bool BeginPlayWorld(const FAvaInstancePlaySettings& InWorldPlaySettings);
+
+	/**
+	 * The end play is normally requested to be done on the next Tick(). This is to
+	 * avoid having the world destroyed while within the Tick() of that world.
+	 * However, when shutting down, we need to force the end play to be done immediately because
+	 * there will not a be a next Tick().
+	 */
+	void RequestEndPlayWorld(bool bForceImmediate);
+	
+	void RequestUnloadWorld(bool bForceImmediate);
+
+	bool IsWorldPlaying() const { return bWorldPlaying; }
+
+	/**
+	 * Cancel any pending EndPlayWorld or UnloadWorld requests. 
+	 */
+	void CancelWorldRequests();
+
+	/**
+	* The output channels can be resized during playback, if this happens, the
+	* scene viewport needs to be updated to reflect the change.
+	*/
+	void UpdateSceneViewportSize(const FIntPoint& InViewportSize);
+
+	void UpdateRenderTarget(UTextureRenderTarget2D* InRenderTarget);
+
+	UTextureRenderTarget2D* GetRenderTarget() const;
+
+	/**
+	 * Returns true if the render target has been rendered and is ready to be used for output.
+	 */
+	bool IsRenderTargetReady() const { return bIsRenderTargetReady;}
+
+	/**
+	 * Reset the flag so the OnRenderTargetReady event can be called again on the next render.
+	 */
+	void ResetRenderTargetReady() { bIsRenderTargetReady = false; }
+
+	/**
+	 * @brief Get the currently playing scene viewport.
+	 * @remark Only valid within a BeginPlay()/EndPlay() scope. Null otherwise.
+	 */
+	TSharedPtr<FSceneViewport> GetSceneViewport() const { return Viewport; }
+
+	/**
+	 * @brief Get the currently playing game viewport client.
+	 * @remark Only valid within a BeginPlay()/EndPlay() scope. Null otherwise.
+	 */
+	UAvaGameViewportClient* GetAvaGameViewportClient() const { return ViewportClient; }
+	
+	/*
+	 * @return The Motion Design Blueprint used in play. Can be null if called outside play
+	 */
+	UAvalancheBlueprint* GetMotionDesignBlueprint() const { return ManagedBlueprint; }
+
+	DECLARE_MULTICAST_DELEGATE_TwoParams(FOnAvaGameInstanceEvent, UAvaGameInstance* /*InGameInstance*/, FName /*InChannelName*/);
+
+	/** Event called on EndPlayWorld(). */
+	static FOnAvaGameInstanceEvent& GetOnEndPlay() { return OnEndPlay; }
+
+	/** Event called when the render target becomes ready. */
+	static FOnAvaGameInstanceEvent& GetOnRenderTargetReady() { return OnRenderTargetReady; }
+
+	/*
+	 * Mark the current frame as a frame in which an asset was loaded synchronously.
+	 * When an asset has been loaded synchronously (i.e. Motion Design Blueprint), it will cause a
+	 * large delta seconds to occur on the next tick. By marking that particular frame it can be
+	 * used on the next tick to clamp the delta seconds to avoid animations skipping by the
+	 * amount of the sync load time.
+	 */
+	void MarkSynchronousAssetLoadingThisFrame();
+
+protected:
+	void Tick(float DeltaSeconds);
+
+	void UnloadWorld();
+	void EndPlayWorld();
+	void OnEndFrameTick();
+	void OnEnginePreExit();
+	
+	//~ Begin UObject
+	virtual void BeginDestroy() override;
+	//~ End UObject
+
+	/*
+	 * A Soft Ptr to the Source Motion Design (i.e. the actual asset).
+	 * It's a Soft Ptr because it can become stale if it's only a Weak Object Ptr
+	 * and we want to reload it if that happens
+	 */
+	TSoftObjectPtr<UAvalancheBlueprint> SourceBlueprint;
+
+	/*
+	 * The Managed Motion Design Blueprint. A Duplication of the Source Blueprint when Begin Play was called.
+	 * This shouldn't be GCd while in play, as it contains Animation and other data used during play.
+	 */
+	UPROPERTY(Transient)
+	TObjectPtr<UAvalancheBlueprint> ManagedBlueprint;
+
+	UPROPERTY(Transient)
+	TObjectPtr<UWorld> PlayWorld;
+
+	UPROPERTY(Transient)
+	TObjectPtr<UAvaGameViewportClient> ViewportClient;
+
+	TSharedPtr<FSceneViewport> Viewport;
+
+	/** Channel this game instance is playing on. */
+	FName PlayingChannelName;
+
+	bool bWorldCreated = false;
+	bool bRequestUnloadWorld = false;
+	bool bWorldPlaying = false;
+	bool bRequestEndPlayWorld = false;
+
+	/** Tick reentrancy guard. */
+	bool bIsTicking = false;
+	
+	double LastDeltaSeconds = 0.0;
+
+	/** Render command fence to ensure the render target has been rendered. */
+	TUniquePtr<FRenderCommandFence> RenderTargetFence;
+	bool bIsRenderTargetReady = false;
+
+	/** Handle to delegate for EnginePreExit event. */
+	FDelegateHandle EnginePreExitHandle;
+
+	static FOnAvaGameInstanceEvent OnEndPlay;
+	static FOnAvaGameInstanceEvent OnRenderTargetReady;
+};
diff --git a/Engine/Plugins/Experimental/Avalanche/Source/Avalanche/Public/Framework/AvaGameViewportClient.h b/Engine/Plugins/Experimental/Avalanche/Source/Avalanche/Public/Framework/AvaGameViewportClient.h
new file mode 100644
--- /dev/null
+++ b/Engine/Plugins/Experimental/Avalanche/Source/Avalanche/Public/Framework/AvaGameViewportClient.h
@@ -0,0 +1,61 @@
+// Copyright Epic Games, Inc. All Rights Reserved.
+
+#pragma once
+
+#include "Engine/GameViewportClient.h"
+#include "Viewport/AvaCameraManager.h"
+#include "SceneTypes.h"
+#include "AvaGameViewportClient.generated.h"
+
+class AAvaViewportCameraActor;
+class FPrimitiveComponentId;
+class FViewElementDrawer;
+class UTextureRenderTarget2D;
+struct FMinimalViewInfo;
+struct FSceneViewInitOptions;
+struct FSceneViewProjectionData;
+
+UCLASS(DisplayName = "Motion Design Game Viewport Client")
+class AVALANCHE_API UAvaGameViewportClient : public UGameViewportClient
+{
+	GENERATED_BODY()
+	
+public:
+	UAvaGameViewportClient();
+
+	//~ Begin FViewportClient
+	virtual void Draw(FViewport* Viewport,FCanvas* Canvas) override;
+	virtual bool IsStatEnabled(const FString& InName) const override;
+	//~ End FViewportClient
+
+	void SetCameraCutThisFrame();
+
+	void SetRenderTarget(UTextureRenderTarget2D* InRenderTarget);
+
+	UTextureRenderTarget2D* GetRenderTarget() const { return RenderTarget.Get(); }
+
+	TSharedRef<FAvaCameraManager> GetCameraManager() const { return CameraManager; }
+
+private:
+	bool CalcSceneViewInitOptions(FSceneViewInitOptions& OutInitOptions
+		, FViewport* InViewport
+		, FViewElementDrawer* ViewDrawer = nullptr
+		, int32 StereoViewIndex = INDEX_NONE);
+
+	bool GetProjectionData(FViewport* InViewport
+		, FSceneViewProjectionData& ProjectionData
+		, int32 StereoViewIndex = INDEX_NONE) const;
+
+	FSceneView* CalcSceneView(FSceneViewFamily* ViewFamily
+		, FViewport* InViewport
+		, FViewElementDrawer* ViewDrawer = nullptr
+		, int32 StereoViewIndex = INDEX_NONE);
+
+	TArray<FSceneViewStateReference> ViewStates;
+
+	TWeakObjectPtr<UTextureRenderTarget2D> RenderTarget;
+
+	TSharedRef<FAvaCameraManager> CameraManager;
+
+	bool bCameraCutThisFrame = false;
+};
diff --git a/Engine/Plugins/Experimental/Avalanche/Source/Avalanche/Public/Framework/AvaInstanceSettings.h b/Engine/Plugins/Experimental/Avalanche/Source/Avalanche/Public/Framework/AvaInstanceSettings.h
new file mode 100644
--- /dev/null
+++ b/Engine/Plugins/Experimental/Avalanche/Source/Avalanche/Public/Framework/AvaInstanceSettings.h
@@ -0,0 +1,18 @@
+// Copyright Epic Games, Inc. All Rights Reserved.
+
+#pragma once
+
+#include "AvaInstanceSettings.generated.h"
+
+/**
+ * Settings applied when instancing a Motion Design blueprint for playback.
+ */
+USTRUCT()
+struct FAvaInstanceSettings
+{
+	GENERATED_BODY()
+
+	/** Enable loading dependent levels as sub-playables. */
+	UPROPERTY(config, EditAnywhere, Category = Settings)
+	bool bEnableLoadSubPlayables = true;
+};
diff --git a/Engine/Plugins/Experimental/Avalanche/Source/Avalanche/Public/Framework/AvalancheGameInstance.h b/Engine/Plugins/Experimental/Avalanche/Source/Avalanche/Public/Framework/AvalancheGameInstance.h
deleted file mode 100644
diff --git a/Engine/Plugins/Experimental/Avalanche/Source/Avalanche/Public/Framework/AvalancheGameViewportClient.h b/Engine/Plugins/Experimental/Avalanche/Source/Avalanche/Public/Framework/AvalancheGameViewportClient.h
deleted file mode 100644
diff --git a/Engine/Plugins/Experimental/Avalanche/Source/Avalanche/Public/Framework/AvalancheInstanceSettings.h b/Engine/Plugins/Experimental/Avalanche/Source/Avalanche/Public/Framework/AvalancheInstanceSettings.h
deleted file mode 100644
diff --git a/Engine/Plugins/Experimental/Avalanche/Source/AvalancheMedia/Private/AvaMediaModule.cpp b/Engine/Plugins/Experimental/Avalanche/Source/AvalancheMedia/Private/AvaMediaModule.cpp
--- a/Engine/Plugins/Experimental/Avalanche/Source/AvalancheMedia/Private/AvaMediaModule.cpp
+++ b/Engine/Plugins/Experimental/Avalanche/Source/AvalancheMedia/Private/AvaMediaModule.cpp
@@ -1,14 +1,14 @@
 // Copyright Epic Games, Inc. All Rights Reserved.
 
 #include "AvaMediaModule.h"
-
 #include "Application/ThrottleManager.h"
-#include "AvaMediaStyle.h"
 #include "AvalancheMediaSettings.h"
+#include "AvaMediaStyle.h"
 #include "IMediaIOCoreModule.h"
 #include "Interfaces/IPluginManager.h"
 #include "Misc/CommandLine.h"
 #include "Misc/CoreDelegates.h"
+#include "OutputDevices/AvaRenderTargetMediaUtils.h"
 #include "Playback/AvaMediaPlaybackClientDelegates.h"
 #include "ShaderCore.h"
 
@@ -37,11 +37,11 @@ void FAvaMediaModule::StartupModule()
 	using namespace UE::AvaMediaModule::Private;
 	
 	FAvaMediaStyle::Initialize();
-	const TSharedPtr<IPlugin> Plugin = IPluginManager::Get().FindPlugin(TEXT("Avalanche"));
+	const TSharedPtr<IPlugin> Plugin = IPluginManager::Get().FindPlugin(UE_PLUGIN_NAME);
 	check(Plugin.IsValid());
 
 	const FString PluginShaderDir = FPaths::Combine(Plugin->GetBaseDir(), TEXT("Shaders"));
-	AddShaderSourceDirectoryMapping(TEXT("/Plugin/Avalanche"), PluginShaderDir);
+	AddShaderSourceDirectoryMapping(UE::AvaRenderTargetMediaUtils::VirtualShaderMountPoint, PluginShaderDir);
 
 	IMediaIOCoreModule::Get().RegisterDeviceProvider(&AvaDisplayDeviceProvider);
 
@@ -372,18 +372,18 @@ const IAvalancheBroadcastSettings& FAvaMediaModule::GetBroadcastSettings() const
 	return BroadcastSettingsBridge;
 }
 
-const FAvalancheInstanceSettings& FAvaMediaModule::GetAvalancheInstanceSettings() const
+const FAvaInstanceSettings& FAvaMediaModule::GetAvaInstanceSettings() const
 {
 	// if the server is enabled, fetch the settings from the currently connected client.
 	if (AvaMediaPlaybackServer.IsValid())
 	{
-		if (const FAvalancheInstanceSettings* SettingsFromClient = AvaMediaPlaybackServer->GetAvalancheInstanceSettings())
+		if (const FAvaInstanceSettings* SettingsFromClient = AvaMediaPlaybackServer->GetAvaInstanceSettings())
 		{
 			return *SettingsFromClient;
 		}
 	}
 	// Return the local settings.
-	return UAvalancheMediaSettings::Get().AvalancheInstanceSettings;
+	return UAvalancheMediaSettings::Get().AvaInstanceSettings;
 }
 
 FAvaMediaPlaybackManager& FAvaMediaModule::GetLocalPlaybackManager() const
diff --git a/Engine/Plugins/Experimental/Avalanche/Source/AvalancheMedia/Private/AvaMediaModule.h b/Engine/Plugins/Experimental/Avalanche/Source/AvalancheMedia/Private/AvaMediaModule.h
--- a/Engine/Plugins/Experimental/Avalanche/Source/AvalancheMedia/Private/AvaMediaModule.h
+++ b/Engine/Plugins/Experimental/Avalanche/Source/AvalancheMedia/Private/AvaMediaModule.h
@@ -41,7 +41,7 @@ public:
 	virtual void StopGameModeLocalPlaybackServer() override;
 	virtual bool IsGameModeLocalPlaybackServerLaunched() const override;
 	virtual const IAvalancheBroadcastSettings& GetBroadcastSettings() const override;
-	virtual const FAvalancheInstanceSettings& GetAvalancheInstanceSettings() const override;
+	virtual const FAvaInstanceSettings& GetAvaInstanceSettings() const override;
 	virtual FAvaMediaPlaybackManager& GetLocalPlaybackManager() const override;
 	virtual FAvalancheManagedInstanceCache& GetManagedInstanceCache() const override;
 	virtual bool IsAvaMediaSyncProviderFeatureAvailable() const override;
diff --git a/Engine/Plugins/Experimental/Avalanche/Source/AvalancheMedia/Private/OutputDevices/AvaRenderTargetMediaUtils.cpp b/Engine/Plugins/Experimental/Avalanche/Source/AvalancheMedia/Private/OutputDevices/AvaRenderTargetMediaUtils.cpp
--- a/Engine/Plugins/Experimental/Avalanche/Source/AvalancheMedia/Private/OutputDevices/AvaRenderTargetMediaUtils.cpp
+++ b/Engine/Plugins/Experimental/Avalanche/Source/AvalancheMedia/Private/OutputDevices/AvaRenderTargetMediaUtils.cpp
@@ -113,7 +113,9 @@ namespace UE::AvaRenderTargetMediaUtils::Private
 			return Parameters;
 		}
 	};
-	IMPLEMENT_SHADER_TYPE(, FAvaRGBGammaConvertPS, TEXT("/Plugin/Avalanche/AvaRGBGammaConvert.usf"), TEXT("MainPS"), SF_Pixel);
+
+	static const FString GammaConvertShaderPath = FString(VirtualShaderMountPoint) + TEXT("/AvaRGBGammaConvert.usf");
+	IMPLEMENT_SHADER_TYPE(, FAvaRGBGammaConvertPS,  *GammaConvertShaderPath, TEXT("MainPS"), SF_Pixel);
 }
 
 void UE::AvaRenderTargetMediaUtils::ConvertTextureRGBGamma(FRHICommandListImmediate& InRHICmdList, FTextureRHIRef InSourceTexture, FTextureRHIRef InDestTarget, bool bInSrgbToLinear, float InGamma)
diff --git a/Engine/Plugins/Experimental/Avalanche/Source/AvalancheMedia/Private/Playback/AvaMediaPlayableGroup.cpp b/Engine/Plugins/Experimental/Avalanche/Source/AvalancheMedia/Private/Playback/AvaMediaPlayableGroup.cpp
--- a/Engine/Plugins/Experimental/Avalanche/Source/AvalancheMedia/Private/Playback/AvaMediaPlayableGroup.cpp
+++ b/Engine/Plugins/Experimental/Avalanche/Source/AvalancheMedia/Private/Playback/AvaMediaPlayableGroup.cpp
@@ -3,7 +3,7 @@
 #include "Playback/AvaMediaPlayableGroup.h"
 
 #include "Engine/ViewportStatsSubsystem.h"
-#include "Framework/AvalancheGameInstance.h"
+#include "Framework/AvaGameInstance.h"
 #include "Playback/AvaMediaPlayableGroupManager.h"
 #include "Playback/AvaMediaPlaybackUtils.h"
 #include "Playback/AvalanchePlayable.h"
@@ -119,7 +119,7 @@ UAvaMediaPlayableGroup* UAvaMediaPlayableGroup::MakePlayableGroup(UObject* InOut
 			GameInstanceGroup->GameInstancePackage = MakeGameInstancePackage(InPlayableGroupInfo.SourceAssetPath, InPlayableGroupInfo.ChannelName);
 		}
 		
-		GameInstanceGroup->GameInstance = UAvalancheGameInstance::Create(GameInstanceGroup->GameInstancePackage);
+		GameInstanceGroup->GameInstance = UAvaGameInstance::Create(GameInstanceGroup->GameInstancePackage);
 	}
 	return GameInstanceGroup;
 }
@@ -296,7 +296,7 @@ bool UAvaMediaPlayableGroup::ConditionalCreateWorld()
 	return bWorldWasCreated;
 }
 
-bool UAvaMediaPlayableGroup::ConditionalBeginPlay(const FAvalancheInstancePlaySettings& InWorldPlaySettings)
+bool UAvaMediaPlayableGroup::ConditionalBeginPlay(const FAvaInstancePlaySettings& InWorldPlaySettings)
 {
 	bool bHasBegunPlay = false;
 	if (!GameInstance)
@@ -381,7 +381,7 @@ void UAvaMediaPlayableGroup::QueueCameraCut()
 {
 	if (GameInstance)
 	{
-		if (UAvalancheGameViewportClient* GameViewportClient = GameInstance->GetAvalancheGameViewportClient())
+		if (UAvaGameViewportClient* GameViewportClient = GameInstance->GetAvaGameViewportClient())
 		{
 			GameViewportClient->SetCameraCutThisFrame();
 		}
@@ -604,7 +604,7 @@ bool UAvaMediaPlayableGroup::DisplayTransitions(FText& OutText, FLinearColor& Ou
 	return false;
 }
 
-bool UAvaMediaRemoteProxyPlayableGroup::ConditionalBeginPlay(const FAvalancheInstancePlaySettings& InWorldPlaySettings)
+bool UAvaMediaRemoteProxyPlayableGroup::ConditionalBeginPlay(const FAvaInstancePlaySettings& InWorldPlaySettings)
 {
 	if (!bIsPlaying)
 	{
diff --git a/Engine/Plugins/Experimental/Avalanche/Source/AvalancheMedia/Private/Playback/AvaMediaPlayableGroupManager.cpp b/Engine/Plugins/Experimental/Avalanche/Source/AvalancheMedia/Private/Playback/AvaMediaPlayableGroupManager.cpp
--- a/Engine/Plugins/Experimental/Avalanche/Source/AvalancheMedia/Private/Playback/AvaMediaPlayableGroupManager.cpp
+++ b/Engine/Plugins/Experimental/Avalanche/Source/AvalancheMedia/Private/Playback/AvaMediaPlayableGroupManager.cpp
@@ -3,7 +3,7 @@
 #include "Playback/AvaMediaPlayableGroupManager.h"
 
 #include "AvalancheBroadcast.h"
-#include "Framework/AvalancheGameInstance.h"
+#include "Framework/AvaGameInstance.h"
 #include "Misc/CoreDelegates.h"
 #include "Misc/TimeGuard.h"
 #include "OutputDevices/AvaRenderTargetMediaUtils.h"
@@ -59,9 +59,9 @@ void UAvaMediaChannelPlayableGroupManager::BeginDestroy()
 
 void UAvaMediaPlayableGroupManager::Init()
 {
-	if (!UAvalancheGameInstance::GetOnEndPlay().IsBoundToObject(this))
+	if (!UAvaGameInstance::GetOnEndPlay().IsBoundToObject(this))
 	{
-		UAvalancheGameInstance::GetOnEndPlay().AddUObject(this, &UAvaMediaPlayableGroupManager::OnGameInstanceEndPlay);
+		UAvaGameInstance::GetOnEndPlay().AddUObject(this, &UAvaMediaPlayableGroupManager::OnGameInstanceEndPlay);
 	}
 }
 
@@ -74,7 +74,7 @@ void UAvaMediaPlayableGroupManager::Shutdown()
 			ChannelManager->Shutdown();
 		}
 	}
-	UAvalancheGameInstance::GetOnEndPlay().RemoveAll(this);
+	UAvaGameInstance::GetOnEndPlay().RemoveAll(this);
 }
 
 void UAvaMediaPlayableGroupManager::Tick(double InDeltaSeconds)
@@ -134,7 +134,7 @@ void UAvaMediaPlayableGroupManager::BeginDestroy()
 	Super::BeginDestroy();
 }
 
-void UAvaMediaPlayableGroupManager::OnGameInstanceEndPlay(UAvalancheGameInstance* InGameInstance, FName InChannelName)
+void UAvaMediaPlayableGroupManager::OnGameInstanceEndPlay(UAvaGameInstance* InGameInstance, FName InChannelName)
 {
 	FAvaOutputChannel& Channel = UAvalancheBroadcast::Get().GetCurrentProfile().GetChannelMutable(InChannelName);
 	if (!Channel.IsValidChannel())
diff --git a/Engine/Plugins/Experimental/Avalanche/Source/AvalancheMedia/Private/Playback/AvaMediaPlaybackClient.cpp b/Engine/Plugins/Experimental/Avalanche/Source/AvalancheMedia/Private/Playback/AvaMediaPlaybackClient.cpp
--- a/Engine/Plugins/Experimental/Avalanche/Source/AvalancheMedia/Private/Playback/AvaMediaPlaybackClient.cpp
+++ b/Engine/Plugins/Experimental/Avalanche/Source/AvalancheMedia/Private/Playback/AvaMediaPlaybackClient.cpp
@@ -1519,7 +1519,7 @@ void FAvaMediaPlaybackClient::OnAvaMediaSettingsChanged(UObject*, struct FProper
 {
 	ApplyAvaMediaSettings();
 	SendBroadcastSettingsUpdate(AllServerAddresses);
-	SendAvalancheInstanceSettingsUpdate(AllServerAddresses);
+	SendAvaInstanceSettingsUpdate(AllServerAddresses);
 }
 
 void FAvaMediaPlaybackClient::OnPreSavePackage(UPackage* InPackage, FObjectPreSaveContext InObjectSaveContext)
@@ -1650,10 +1650,10 @@ void FAvaMediaPlaybackClient::SendBroadcastSettingsUpdate(const TArray<FMessageA
 	SendRequest(BroadcastSettingsUpdate, InRecipients, EMessageFlags::Reliable);
 }
 
-void FAvaMediaPlaybackClient::SendAvalancheInstanceSettingsUpdate(const TArray<FMessageAddress>& InRecipients)
+void FAvaMediaPlaybackClient::SendAvaInstanceSettingsUpdate(const TArray<FMessageAddress>& InRecipients)
 {
-	FAvalancheInstanceSettingsUpdate* InstanceSettingsUpdate = FMessageEndpoint::MakeMessage<FAvalancheInstanceSettingsUpdate>();
-	InstanceSettingsUpdate->InstanceSettings = UAvalancheMediaSettings::Get().AvalancheInstanceSettings;
+	FAvaInstanceSettingsUpdate* InstanceSettingsUpdate = FMessageEndpoint::MakeMessage<FAvaInstanceSettingsUpdate>();
+	InstanceSettingsUpdate->InstanceSettings = UAvalancheMediaSettings::Get().AvaInstanceSettings;
 	SendRequest(InstanceSettingsUpdate, InRecipients, EMessageFlags::Reliable);
 }
 
@@ -1693,7 +1693,7 @@ void FAvaMediaPlaybackClient::SendClientInfo(const FMessageAddress& InRecipient)
 	
 	SendUserDataUpdate({InRecipient});
 	SendBroadcastSettingsUpdate({InRecipient});
-	SendAvalancheInstanceSettingsUpdate({InRecipient});
+	SendAvaInstanceSettingsUpdate({InRecipient});
 	SendStatCommand(FString(), true, {InRecipient});	// Send empty stat command, will just send current states.
 }
 
diff --git a/Engine/Plugins/Experimental/Avalanche/Source/AvalancheMedia/Private/Playback/AvaMediaPlaybackManager.cpp b/Engine/Plugins/Experimental/Avalanche/Source/AvalancheMedia/Private/Playback/AvaMediaPlaybackManager.cpp
--- a/Engine/Plugins/Experimental/Avalanche/Source/AvalancheMedia/Private/Playback/AvaMediaPlaybackManager.cpp
+++ b/Engine/Plugins/Experimental/Avalanche/Source/AvalancheMedia/Private/Playback/AvaMediaPlaybackManager.cpp
@@ -7,7 +7,7 @@
 #include "AvaBlueprint.h"
 #include "AvalancheBroadcast.h"
 #include "Engine/Engine.h"
-#include "Framework/AvalancheGameInstance.h"
+#include "Framework/AvaGameInstance.h"
 #include "IAvaMediaModule.h"
 #include "IAvaModule.h"
 #include "Playback/AvaMediaPlayableGroup.h"
@@ -703,12 +703,12 @@ bool FAvaMediaPlaybackManager::HandleStatCommand(const TArray<FString>& InArgs)
 		{
 			ForAllPlaybackInstances([&ViewportClient](FAvaMediaPlaybackInstance& InPlaybackInstance)
 			{
-				const TArray<UAvalancheGameInstance*> ActiveGameInstances = InPlaybackInstance.Playback->GetActiveGameInstances();
-				for (const UAvalancheGameInstance* GameInstance : ActiveGameInstances)
+				const TArray<UAvaGameInstance*> ActiveGameInstances = InPlaybackInstance.Playback->GetActiveGameInstances();
+				for (const UAvaGameInstance* GameInstance : ActiveGameInstances)
 				{
-					if (IsValid(GameInstance->GetWorld()) && IsValid(GameInstance->GetAvalancheGameViewportClient()))
+					if (IsValid(GameInstance->GetWorld()) && IsValid(GameInstance->GetAvaGameViewportClient()))
 					{
-						ViewportClient = GameInstance->GetAvalancheGameViewportClient();
+						ViewportClient = GameInstance->GetAvaGameViewportClient();
 						break;						
 					}
 				}
diff --git a/Engine/Plugins/Experimental/Avalanche/Source/AvalancheMedia/Private/Playback/AvaMediaPlaybackServer.cpp b/Engine/Plugins/Experimental/Avalanche/Source/AvalancheMedia/Private/Playback/AvaMediaPlaybackServer.cpp
--- a/Engine/Plugins/Experimental/Avalanche/Source/AvalancheMedia/Private/Playback/AvaMediaPlaybackServer.cpp
+++ b/Engine/Plugins/Experimental/Avalanche/Source/AvalancheMedia/Private/Playback/AvaMediaPlaybackServer.cpp
@@ -128,7 +128,7 @@ void FAvaMediaPlaybackServer::Init(const FString& InAssignedServerName)
 	.Handling<FAvaMediaStatCommand>(this, &FAvaMediaPlaybackServer::HandleStatCommand)
 	.Handling<FAvaDeviceProviderDataRequest>(this, &FAvaMediaPlaybackServer::HandleDeviceProviderDataRequest)
 	.Handling<FAvaMediaUpdateClientInfo>(this, &FAvaMediaPlaybackServer::HandleUpdateClientInfo)
-	.Handling<FAvalancheInstanceSettingsUpdate>(this, &FAvaMediaPlaybackServer::HandleAvalancheInstanceSettingsUpdate)
+	.Handling<FAvaInstanceSettingsUpdate>(this, &FAvaMediaPlaybackServer::HandleAvaInstanceSettingsUpdate)
 	.Handling<FAvaMediaPlaybackPackageEvent>(this, &FAvaMediaPlaybackServer::HandlePackageEvent)
 	.Handling<FAvaMediaPlaybackAssetStatusRequest>(this, &FAvaMediaPlaybackServer::HandlePlaybackAssetStatusRequest)
 	.Handling<FAvaMediaPlaybackRequest>(this, &FAvaMediaPlaybackServer::HandlePlaybackRequest)
@@ -372,13 +372,13 @@ const IAvalancheBroadcastSettings* FAvaMediaPlaybackServer::GetBroadcastSettings
 	return nullptr;
 }
 
-const FAvalancheInstanceSettings* FAvaMediaPlaybackServer::GetAvalancheInstanceSettings() const
+const FAvaInstanceSettings* FAvaMediaPlaybackServer::GetAvaInstanceSettings() const
 {
 	// Returns the first client we have.
 	// Todo: In case we have multiple clients, we will need a smarter way to handle this.
 	for (const TPair<FString, TSharedPtr<FClientInfo>>& Client : Clients)
 	{
-		return &Client.Value->AvalancheInstanceSettings;
+		return &Client.Value->AvaInstanceSettings;
 	}
 	return nullptr;
 }
@@ -489,10 +489,10 @@ void FAvaMediaPlaybackServer::HandleUpdateClientInfo(const FAvaMediaUpdateClient
 	ClientInfo.MediaSyncManager->SetEnable(bShouldEnableSyncManager);
 }
 
-void FAvaMediaPlaybackServer::HandleAvalancheInstanceSettingsUpdate(const FAvalancheInstanceSettingsUpdate& InMessage, const TSharedRef<IMessageContext, ESPMode::ThreadSafe>& InContext)
+void FAvaMediaPlaybackServer::HandleAvaInstanceSettingsUpdate(const FAvaInstanceSettingsUpdate& InMessage, const TSharedRef<IMessageContext, ESPMode::ThreadSafe>& InContext)
 {
 	FClientInfo& ClientInfo = GetOrCreateClientInfo(InMessage.ClientName, InContext->GetSender());
-	ClientInfo.AvalancheInstanceSettings = InMessage.InstanceSettings;
+	ClientInfo.AvaInstanceSettings = InMessage.InstanceSettings;
 	
 	UE_LOG(LogAvaPlaybackServer, Verbose, TEXT("Received new instance settings from client \"%s\"."), *InMessage.ClientName);
 }
diff --git a/Engine/Plugins/Experimental/Avalanche/Source/AvalancheMedia/Private/Playback/AvalanchePlayable.cpp b/Engine/Plugins/Experimental/Avalanche/Source/AvalancheMedia/Private/Playback/AvalanchePlayable.cpp
--- a/Engine/Plugins/Experimental/Avalanche/Source/AvalancheMedia/Private/Playback/AvalanchePlayable.cpp
+++ b/Engine/Plugins/Experimental/Avalanche/Source/AvalancheMedia/Private/Playback/AvalanchePlayable.cpp
@@ -258,7 +258,7 @@ EAvalanchePlayableCommandResult UAvalanchePlayable::UpdateRemoteControlCommand(c
 	return EAvalanchePlayableCommandResult::Executed;
 }
 
-void UAvalanchePlayable::BeginPlay(const FAvalancheInstancePlaySettings& InWorldPlaySettings)
+void UAvalanchePlayable::BeginPlay(const FAvaInstancePlaySettings& InWorldPlaySettings)
 {
 	if (!PlayableGroup)
 	{
diff --git a/Engine/Plugins/Experimental/Avalanche/Source/AvalancheMedia/Private/Playback/AvalanchePlayback.cpp b/Engine/Plugins/Experimental/Avalanche/Source/AvalancheMedia/Private/Playback/AvalanchePlayback.cpp
--- a/Engine/Plugins/Experimental/Avalanche/Source/AvalancheMedia/Private/Playback/AvalanchePlayback.cpp
+++ b/Engine/Plugins/Experimental/Avalanche/Source/AvalancheMedia/Private/Playback/AvalanchePlayback.cpp
@@ -7,7 +7,7 @@
 #include "AvalancheBroadcast.h"
 #include "AvalancheMediaSettings.h"
 #include "Engine/TextureRenderTarget2D.h"
-#include "Framework/AvalancheGameInstance.h"
+#include "Framework/AvaGameInstance.h"
 #include "IAvaMediaModule.h"
 #include "IAvalancheBroadcastSettings.h"
 #include "Misc/CoreDelegates.h"
@@ -227,9 +227,9 @@ bool UAvalanchePlayback::UnloadAndRemovePlayable(const FSoftObjectPath& InSource
 	return UnloadAndRemovePlayable(FindPlayable(InSourceAssetPath, InChannelName), InSourceAssetPath, InChannelName, bInForceImmediate);
 }
 
-TArray<UAvalancheGameInstance*> UAvalanchePlayback::GetActiveGameInstances() const
+TArray<UAvaGameInstance*> UAvalanchePlayback::GetActiveGameInstances() const
 {
-	TArray<UAvalancheGameInstance*> OutGameInstances;
+	TArray<UAvaGameInstance*> OutGameInstances;
 
 	if (IsPlaying())
 	{
@@ -242,7 +242,7 @@ TArray<UAvalancheGameInstance*> UAvalanchePlayback::GetActiveGameInstances() con
 				const UAvalanchePlayable* Playable = FindPlayable(AvalancheAsset.ToSoftObjectPath(), ChannelName);
 				if (Playable && Playable->GetPlayableGroup())
 				{
-					if (UAvalancheGameInstance* AvaGameInstance = Cast<UAvalancheGameInstance>(Playable->GetPlayableGroup()->GetGameInstance()))
+					if (UAvaGameInstance* AvaGameInstance = Cast<UAvaGameInstance>(Playable->GetPlayableGroup()->GetGameInstance()))
 					{
 						OutGameInstances.AddUnique(AvaGameInstance);
 					}
@@ -740,9 +740,9 @@ void UAvalanchePlayback::OnChannelChanged(const FAvaOutputChannel& InChannel, EA
 			});
 			for (const UAvaMediaPlayableGroup* PlayableGroup : PlayableGroups)
 			{
-				if (const UAvalancheGameInstance* GameInstance = Cast<UAvalancheGameInstance>(PlayableGroup->GetGameInstance()))
+				if (const UAvaGameInstance* GameInstance = Cast<UAvaGameInstance>(PlayableGroup->GetGameInstance()))
 				{
-					if (UAvalancheGameViewportClient* ViewportClient = GameInstance->GetAvalancheGameViewportClient())
+					if (UAvaGameViewportClient* ViewportClient = GameInstance->GetAvaGameViewportClient())
 					{
 						FAvaViewportQualitySettings QualitySettingsMutable = InChannel.GetViewportQualitySettings(); 
 						QualitySettingsMutable.Apply(ViewportClient->EngineShowFlags);
@@ -830,8 +830,8 @@ bool UAvalanchePlayback::RefreshPlayback(const FAvaSoftAssetPtr& InSourceAsset,
 			? AvaRenderTargetUtils::GetRenderTargetSize(RenderTarget)
 			: InChannel.DetermineRenderTargetSize();
 
-		const FAvalancheInstancePlaySettings WorldPlaySettings =
-		{ IAvaMediaModule::Get().GetAvalancheInstanceSettings(), InChannelName, RenderTarget, ViewportSize, InChannel.GetViewportQualitySettings() };
+		const FAvaInstancePlaySettings WorldPlaySettings =
+		{ IAvaMediaModule::Get().GetAvaInstanceSettings(), InChannelName, RenderTarget, ViewportSize, InChannel.GetViewportQualitySettings() };
 
 		Playable->BeginPlay(WorldPlaySettings);
 
@@ -891,8 +891,8 @@ bool UAvalanchePlayback::RefreshPreview(const FAvaSoftAssetPtr& InSourceAsset, F
 	const FIntPoint ViewportSize = AvaRenderTargetUtils::GetRenderTargetSize(RenderTarget);
 
 	static const FAvaViewportQualitySettings DefaultPreviewQualitySettings;
-	const FAvalancheInstancePlaySettings WorldPlaySettings =
-		{ IAvaMediaModule::Get().GetAvalancheInstanceSettings(), InChannelName, RenderTarget, ViewportSize, DefaultPreviewQualitySettings };
+	const FAvaInstancePlaySettings WorldPlaySettings =
+		{ IAvaMediaModule::Get().GetAvaInstanceSettings(), InChannelName, RenderTarget, ViewportSize, DefaultPreviewQualitySettings };
 
 	Playable->BeginPlay(WorldPlaySettings);
 	
diff --git a/Engine/Plugins/Experimental/Avalanche/Source/AvalancheMedia/Private/Playback/Nodes/AvaPlaybackNodeBlueprintPlayer.cpp b/Engine/Plugins/Experimental/Avalanche/Source/AvalancheMedia/Private/Playback/Nodes/AvaPlaybackNodeBlueprintPlayer.cpp
--- a/Engine/Plugins/Experimental/Avalanche/Source/AvalancheMedia/Private/Playback/Nodes/AvaPlaybackNodeBlueprintPlayer.cpp
+++ b/Engine/Plugins/Experimental/Avalanche/Source/AvalancheMedia/Private/Playback/Nodes/AvaPlaybackNodeBlueprintPlayer.cpp
@@ -2,7 +2,7 @@
 
 #include "Playback/Nodes/AvaPlaybackNodeBlueprintPlayer.h"
 #include "AvaBlueprint.h"
-#include "Framework/AvalancheGameInstance.h"
+#include "Framework/AvaGameInstance.h"
 #include "Http/AvaMediaHttpServer.h"
 #include "Playback/AvalanchePlayback.h"
 
diff --git a/Engine/Plugins/Experimental/Avalanche/Source/AvalancheMedia/Private/Playback/Playables/AvalancheBlueprintPlayable.cpp b/Engine/Plugins/Experimental/Avalanche/Source/AvalancheMedia/Private/Playback/Playables/AvalancheBlueprintPlayable.cpp
--- a/Engine/Plugins/Experimental/Avalanche/Source/AvalancheMedia/Private/Playback/Playables/AvalancheBlueprintPlayable.cpp
+++ b/Engine/Plugins/Experimental/Avalanche/Source/AvalancheMedia/Private/Playback/Playables/AvalancheBlueprintPlayable.cpp
@@ -1,7 +1,6 @@
 // Copyright Epic Games, Inc. All Rights Reserved.
 
 #include "Playback/Playables/AvalancheBlueprintPlayable.h"
-
 #include "AvaBlueprint.h"
 #include "AvaRemoteControlRebind.h"
 #include "AvaSequencePlaybackObject.h"
@@ -95,13 +94,13 @@ void UAvalancheBlueprintPlayable::OnPlay()
 		return;
 	}
 
-	const UAvalancheGameInstance* AvalancheGameInstance = Cast<UAvalancheGameInstance>(GetPlayableGroup()->GetGameInstance());
-	if (!AvalancheGameInstance)
+	const UAvaGameInstance* AvaGameInstance = Cast<UAvaGameInstance>(GetPlayableGroup()->GetGameInstance());
+	if (!AvaGameInstance)
 	{
 		return;
 	}
 	
-	UAvalancheGameViewportClient* ViewportClient = AvalancheGameInstance->GetAvalancheGameViewportClient();
+	UAvaGameViewportClient* ViewportClient = AvaGameInstance->GetAvaGameViewportClient();
 
 	// Verify if this can be changed on the fly.
 	ManagedAvalancheBlueprint->GetViewportQualitySettings().Apply(ViewportClient->EngineShowFlags);
@@ -132,8 +131,8 @@ bool UAvalancheBlueprintPlayable::LoadAvalancheBlueprintInternal(const TSoftObje
 		return false;
 	}
 
-	UAvalancheGameInstance* AvalancheGameInstance = Cast<UAvalancheGameInstance>(GetPlayableGroup()->GetGameInstance());
-	if (!AvalancheGameInstance)
+	UAvaGameInstance* AvaGameInstance = Cast<UAvaGameInstance>(GetPlayableGroup()->GetGameInstance());
+	if (!AvaGameInstance)
 	{
 		return false;
 	}
@@ -151,8 +150,8 @@ bool UAvalancheBlueprintPlayable::LoadAvalancheBlueprintInternal(const TSoftObje
 
 	{
 		TRACE_CPUPROFILER_EVENT_SCOPE(UAvalancheBlueprintPlayable::LoadAsset::DupAvaBp);
-		const FName AvaInstanceName = MakeUniqueObjectName(AvalancheGameInstance, UAvalancheBlueprint::StaticClass(), Source->GetFName());
-		ManagedAvalancheBlueprint = Cast<UAvalancheBlueprint>(StaticDuplicateObject(Source, AvalancheGameInstance, AvaInstanceName, RF_Transient));
+		const FName AvaInstanceName = MakeUniqueObjectName(AvaGameInstance, UAvalancheBlueprint::StaticClass(), Source->GetFName());
+		ManagedAvalancheBlueprint = Cast<UAvalancheBlueprint>(StaticDuplicateObject(Source, AvaGameInstance, AvaInstanceName, RF_Transient));
 	}
 
 	if (!ManagedAvalancheBlueprint)
@@ -162,7 +161,7 @@ bool UAvalancheBlueprintPlayable::LoadAvalancheBlueprintInternal(const TSoftObje
 	
 	{
 		TRACE_CPUPROFILER_EVENT_SCOPE(UAvalancheBlueprintPlayable::LoadAsset::LoadAvaBpWorld);
-		ManagedAvalancheBlueprint->SetAvalancheWorld(AvalancheGameInstance->GetPlayWorld());
+		ManagedAvalancheBlueprint->SetAvalancheWorld(AvaGameInstance->GetPlayWorld());
 		ManagedAvalancheBlueprint->LoadAvalancheWorld();
 	}
 
@@ -174,7 +173,7 @@ bool UAvalancheBlueprintPlayable::LoadAvalancheBlueprintInternal(const TSoftObje
 		FAvaRemoteControlRebind::RebindUnboundEntities(ManagedAvalancheBlueprint->GetRemoteControlPreset());
 	}
 
-	AvalancheGameInstance->MarkSynchronousAssetLoadingThisFrame();
+	AvaGameInstance->MarkSynchronousAssetLoadingThisFrame();
 	
 	SourceAvalancheBlueprint = InSourceAvalancheBlueprint;
 	return true;
diff --git a/Engine/Plugins/Experimental/Avalanche/Source/AvalancheMedia/Private/Playback/Playables/AvalancheLevelStreamingPlayable.cpp b/Engine/Plugins/Experimental/Avalanche/Source/AvalancheMedia/Private/Playback/Playables/AvalancheLevelStreamingPlayable.cpp
--- a/Engine/Plugins/Experimental/Avalanche/Source/AvalancheMedia/Private/Playback/Playables/AvalancheLevelStreamingPlayable.cpp
+++ b/Engine/Plugins/Experimental/Avalanche/Source/AvalancheMedia/Private/Playback/Playables/AvalancheLevelStreamingPlayable.cpp
@@ -10,7 +10,7 @@
 #include "Engine/LevelStreamingDynamic.h"
 #include "Engine/LocalPlayer.h"
 #include "Framework/AvaSoftAssetPtr.h"
-#include "Framework/AvalancheInstanceSettings.h"
+#include "Framework/AvaInstanceSettings.h"
 #include "IAvaMediaModule.h"
 #include "Playback/AvaMediaPlayableGroup.h"
 #include "Playback/AvaMediaPlayableGroupManager.h"
@@ -204,7 +204,7 @@ bool UAvalancheLevelStreamingPlayable::LoadAsset(const FAvaSoftAssetPtr& InAvala
 	// Ensure world is created. Does nothing if already created.
 	PlayableGroup->ConditionalCreateWorld();
 
-	const FAvalancheInstanceSettings& PlaybackInstanceSettings = IAvaMediaModule::Get().GetAvalancheInstanceSettings();
+	const FAvaInstanceSettings& PlaybackInstanceSettings = IAvaMediaModule::Get().GetAvaInstanceSettings();
 	bLoadSubPlayables = PlaybackInstanceSettings.bEnableLoadSubPlayables;
 	
 	check(InAvalancheSourceAsset.GetAssetType() == EAvalancheAssetType::World);
@@ -314,8 +314,8 @@ bool UAvalancheLevelStreamingPlayable::ApplyCamera()
 		else
 		{
 			// If there is no player controller, fallback to avalanche camera manager, if the game instance has it.
-			const UAvalancheGameInstance* AvalancheGameInstance = Cast<UAvalancheGameInstance>(GetPlayableGroup()->GetGameInstance());			
-			if (const UAvalancheGameViewportClient* ViewportClient = AvalancheGameInstance ? AvalancheGameInstance->GetAvalancheGameViewportClient() : nullptr)
+			const UAvaGameInstance* AvaGameInstance = Cast<UAvaGameInstance>(GetPlayableGroup()->GetGameInstance());			
+			if (const UAvaGameViewportClient* ViewportClient = AvaGameInstance ? AvaGameInstance->GetAvaGameViewportClient() : nullptr)
 			{
 				// Old code using ava camera manager. To retire.
 				constexpr bool bIsCanvasController = false;
diff --git a/Engine/Plugins/Experimental/Avalanche/Source/AvalancheMedia/Private/Playlist/AvalancheManagedInstanceBlueprint.cpp b/Engine/Plugins/Experimental/Avalanche/Source/AvalancheMedia/Private/Playlist/AvalancheManagedInstanceBlueprint.cpp
--- a/Engine/Plugins/Experimental/Avalanche/Source/AvalancheMedia/Private/Playlist/AvalancheManagedInstanceBlueprint.cpp
+++ b/Engine/Plugins/Experimental/Avalanche/Source/AvalancheMedia/Private/Playlist/AvalancheManagedInstanceBlueprint.cpp
@@ -93,7 +93,7 @@ FAvalancheManagedInstanceBlueprint::~FAvalancheManagedInstanceBlueprint()
 		ManagedAvalancheBlueprintPackage->ClearDirtyFlag();
 	}
 
-	// Simplified version of UAvalancheGameInstance::EndPlay()
+	// Simplified version of UAvaGameInstance::EndPlay()
 	if (ManagedWorld)
 	{
 		GEngine->DestroyWorldContext(ManagedWorld.Get());
diff --git a/Engine/Plugins/Experimental/Avalanche/Source/AvalancheMedia/Public/AvalancheMediaSettings.h b/Engine/Plugins/Experimental/Avalanche/Source/AvalancheMedia/Public/AvalancheMediaSettings.h
--- a/Engine/Plugins/Experimental/Avalanche/Source/AvalancheMedia/Public/AvalancheMediaSettings.h
+++ b/Engine/Plugins/Experimental/Avalanche/Source/AvalancheMedia/Public/AvalancheMediaSettings.h
@@ -4,7 +4,7 @@
 
 #include "CoreMinimal.h"
 #include "Engine/DeveloperSettings.h"
-#include "Framework/AvalancheInstanceSettings.h"
+#include "Framework/AvaInstanceSettings.h"
 #include "Logging/LogVerbosity.h"
 #include "PixelFormat.h"
 #include "AvalancheMediaSettings.generated.h"
@@ -180,7 +180,7 @@ public:
 	bool bKeepPagesLoaded = false;
 
 	UPROPERTY(Config, EditAnywhere, Category = "Playback Manager")
-	FAvalancheInstanceSettings AvalancheInstanceSettings;
+	FAvaInstanceSettings AvaInstanceSettings;
 
 	/**
 	 * Maximum cached Managed Motion Design Blueprints used for rundown editor's page details.
diff --git a/Engine/Plugins/Experimental/Avalanche/Source/AvalancheMedia/Public/IAvaMediaModule.h b/Engine/Plugins/Experimental/Avalanche/Source/AvalancheMedia/Public/IAvaMediaModule.h
--- a/Engine/Plugins/Experimental/Avalanche/Source/AvalancheMedia/Public/IAvaMediaModule.h
+++ b/Engine/Plugins/Experimental/Avalanche/Source/AvalancheMedia/Public/IAvaMediaModule.h
@@ -18,7 +18,7 @@ class IAvaMediaSyncProvider;
 class IAvalancheBroadcastSettings;
 class IMediaIOCoreDeviceProvider;
 class UWorld;
-struct FAvalancheInstanceSettings;
+struct FAvaInstanceSettings;
 struct FMediaIOOutputConfiguration;
 
 /** Maps one to one with the editor's map changed type (for now). */
@@ -110,9 +110,9 @@ public:
 	 * Remark: lifetime of the returned reference is not guaranteed beyond the current call context.
 	 * If the settings are replicated from a client, it could get deleted if the client disconnects.
 	 * If the use of the settings is deferred, the caller must make a local copy of the settings or
-	 * call GetAvalancheInstanceSettings() in the deferred call instead.
+	 * call GetAvaInstanceSettings() in the deferred call instead.
 	 */
-	virtual const FAvalancheInstanceSettings& GetAvalancheInstanceSettings() const = 0;
+	virtual const FAvaInstanceSettings& GetAvaInstanceSettings() const = 0;
 	
 	/**
 	 *	This is the backend for playing avalanche blueprints locally.
diff --git a/Engine/Plugins/Experimental/Avalanche/Source/AvalancheMedia/Public/OutputDevices/AvaRenderTargetMediaUtils.h b/Engine/Plugins/Experimental/Avalanche/Source/AvalancheMedia/Public/OutputDevices/AvaRenderTargetMediaUtils.h
--- a/Engine/Plugins/Experimental/Avalanche/Source/AvalancheMedia/Public/OutputDevices/AvaRenderTargetMediaUtils.h
+++ b/Engine/Plugins/Experimental/Avalanche/Source/AvalancheMedia/Public/OutputDevices/AvaRenderTargetMediaUtils.h
@@ -9,6 +9,8 @@ class UTextureRenderTarget2D;
 
 namespace UE::AvaRenderTargetMediaUtils
 {
+	static inline const TCHAR* VirtualShaderMountPoint = TEXT("/Plugin/MotionDesign");;
+
 	/**
 	 * Clears the given render target to it's specified clear color.
 	 * @remark called from the main thread.
diff --git a/Engine/Plugins/Experimental/Avalanche/Source/AvalancheMedia/Public/Playback/AvaMediaPlayableGroup.h b/Engine/Plugins/Experimental/Avalanche/Source/AvalancheMedia/Public/Playback/AvaMediaPlayableGroup.h
--- a/Engine/Plugins/Experimental/Avalanche/Source/AvalancheMedia/Public/Playback/AvaMediaPlayableGroup.h
+++ b/Engine/Plugins/Experimental/Avalanche/Source/AvalancheMedia/Public/Playback/AvaMediaPlayableGroup.h
@@ -2,7 +2,7 @@
 
 #pragma once
 
-#include "Framework/AvalancheGameInstance.h"
+#include "Framework/AvaGameInstance.h"
 #include "Playback/IAvaPlayableVisibilityConstraint.h"
 #include "UObject/Object.h"
 #include "UObject/ObjectKey.h"
@@ -21,7 +21,7 @@ class UTextureRenderTarget2D;
  * It tracks and manage the game instance and playables state.
  *
  * The design goal of this class is to remove all the playback management
- * from UAvalancheGameInstance and move it to the playable framework. This
+ * from UAvaGameInstance and move it to the playable framework. This
  * should allow us to hook the playback framework to any game instance, including PIE
  * so it can work with any work flow (editor, PIE, game, nDisplay, etc).
  *
@@ -88,7 +88,7 @@ public:
 	 * Begin playing the game instance's world if it wasn't already.
 	 * @return true if the BeginPlay was done (i.e. on the state transition only), false otherwise.
 	 */
-	virtual bool ConditionalBeginPlay(const FAvalancheInstancePlaySettings& InWorldPlaySettings);
+	virtual bool ConditionalBeginPlay(const FAvaInstancePlaySettings& InWorldPlaySettings);
 
 	virtual void RequestEndPlayWorld(bool bInForceImmediate);
 
@@ -157,7 +157,7 @@ public:
 	TObjectPtr<UTextureRenderTarget2D> RenderTarget;	// Optional
 	
 	UPROPERTY(Transient)
-	TObjectPtr<UAvalancheGameInstance> GameInstance;
+	TObjectPtr<UAvaGameInstance> GameInstance;
 
 	UPROPERTY(Transient)
 	TObjectPtr<UPackage> GameInstancePackage;
@@ -216,7 +216,7 @@ class AVALANCHEMEDIA_API UAvaMediaRemoteProxyPlayableGroup : public UAvaMediaPla
 	GENERATED_BODY()
 	
 public:
-	virtual bool ConditionalBeginPlay(const FAvalancheInstancePlaySettings& InWorldPlaySettings) override;
+	virtual bool ConditionalBeginPlay(const FAvaInstancePlaySettings& InWorldPlaySettings) override;
 	virtual void RequestEndPlayWorld(bool bInForceImmediate) override;
 	
 	bool bIsPlaying = false;
diff --git a/Engine/Plugins/Experimental/Avalanche/Source/AvalancheMedia/Public/Playback/AvaMediaPlayableGroupManager.h b/Engine/Plugins/Experimental/Avalanche/Source/AvalancheMedia/Public/Playback/AvaMediaPlayableGroupManager.h
--- a/Engine/Plugins/Experimental/Avalanche/Source/AvalancheMedia/Public/Playback/AvaMediaPlayableGroupManager.h
+++ b/Engine/Plugins/Experimental/Avalanche/Source/AvalancheMedia/Public/Playback/AvaMediaPlayableGroupManager.h
@@ -10,7 +10,7 @@
 
 class UAvaMediaPlayableGroup;
 class UAvaMediaPlayableGroupManager;
-class UAvalancheGameInstance;
+class UAvaGameInstance;
 
 /**
  * Manager for the shared playable groups per channel.
@@ -96,7 +96,7 @@ protected:
 	virtual void BeginDestroy() override;
 	//~ End UObject
 
-	void OnGameInstanceEndPlay(UAvalancheGameInstance* InGameInstance, FName InChannelName);
+	void OnGameInstanceEndPlay(UAvaGameInstance* InGameInstance, FName InChannelName);
 
 	void UpdateLevelStreaming();
 	
diff --git a/Engine/Plugins/Experimental/Avalanche/Source/AvalancheMedia/Public/Playback/AvaMediaPlaybackClient.h b/Engine/Plugins/Experimental/Avalanche/Source/AvalancheMedia/Public/Playback/AvaMediaPlaybackClient.h
--- a/Engine/Plugins/Experimental/Avalanche/Source/AvalancheMedia/Public/Playback/AvaMediaPlaybackClient.h
+++ b/Engine/Plugins/Experimental/Avalanche/Source/AvalancheMedia/Public/Playback/AvaMediaPlaybackClient.h
@@ -132,7 +132,7 @@ protected:
 	void PublishPlaybackPing(const FDateTime& InCurrentTime, bool bInAutoPing);
 	void SendUserDataUpdate(const TArray<FMessageAddress>& InRecipients);
 	void SendBroadcastSettingsUpdate(const TArray<FMessageAddress>& InRecipients);
-	void SendAvalancheInstanceSettingsUpdate(const TArray<FMessageAddress>& InRecipients);
+	void SendAvaInstanceSettingsUpdate(const TArray<FMessageAddress>& InRecipients);
 	void SendBroadcastChannelSettingsUpdate(const TArray<FMessageAddress>& InRecipients, const FAvaOutputChannel& InChannel);
 	void SendPackageEvent(const TArray<FMessageAddress>& InRecipients, const FName& InPackageName, EAvaMediaPlaybackPackageEvent InEvent);
 	void SendStatCommand(const FString& InCommand, bool bInBroadcastLocalState, const TArray<FMessageAddress>& InRecipients);
diff --git a/Engine/Plugins/Experimental/Avalanche/Source/AvalancheMedia/Public/Playback/AvaMediaPlaybackMessages.h b/Engine/Plugins/Experimental/Avalanche/Source/AvalancheMedia/Public/Playback/AvaMediaPlaybackMessages.h
--- a/Engine/Plugins/Experimental/Avalanche/Source/AvalancheMedia/Public/Playback/AvaMediaPlaybackMessages.h
+++ b/Engine/Plugins/Experimental/Avalanche/Source/AvalancheMedia/Public/Playback/AvaMediaPlaybackMessages.h
@@ -4,7 +4,7 @@
 
 #include "AvaMediaDefines.h"
 #include "Channel/AvaMediaOutputInfo.h"
-#include "Framework/AvalancheInstanceSettings.h"
+#include "Framework/AvaInstanceSettings.h"
 #include "PixelFormat.h"
 #include "Playback/AvalancheRemoteControlValues.h"
 #include "Playback/Nodes/Events/Actions/AvalancheAnimations.h"
@@ -208,12 +208,12 @@ struct AVALANCHEMEDIA_API FAvaDeviceProviderDataRequest : public FAvaMediaPlayba
  *	Request for the client to replicate it's Motion Design instance settings to the server.
  **/
 USTRUCT()
-struct AVALANCHEMEDIA_API FAvalancheInstanceSettingsUpdate : public FAvaMediaPlaybackClientMessageBase
+struct AVALANCHEMEDIA_API FAvaInstanceSettingsUpdate : public FAvaMediaPlaybackClientMessageBase
 {
 	GENERATED_BODY()
 
 	UPROPERTY()
-	FAvalancheInstanceSettings InstanceSettings;
+	FAvaInstanceSettings InstanceSettings;
 };
 
 UENUM()
diff --git a/Engine/Plugins/Experimental/Avalanche/Source/AvalancheMedia/Public/Playback/AvaMediaPlaybackServer.h b/Engine/Plugins/Experimental/Avalanche/Source/AvalancheMedia/Public/Playback/AvaMediaPlaybackServer.h
--- a/Engine/Plugins/Experimental/Avalanche/Source/AvalancheMedia/Public/Playback/AvaMediaPlaybackServer.h
+++ b/Engine/Plugins/Experimental/Avalanche/Source/AvalancheMedia/Public/Playback/AvaMediaPlaybackServer.h
@@ -2,7 +2,7 @@
 
 #pragma once
 
-#include "Framework/AvalancheInstanceSettings.h"
+#include "Framework/AvaInstanceSettings.h"
 #include "IAvalancheBroadcastSettings.h"
 #include "IMessageContext.h"
 #include "MessageEndpoint.h"
@@ -99,7 +99,7 @@ public:
 	 * Access Avalanche Instance settings replicated from connected client(s).
 	 * Will return nullptr if no clients are connected.
 	 */
-	const FAvalancheInstanceSettings* GetAvalancheInstanceSettings() const;
+	const FAvaInstanceSettings* GetAvaInstanceSettings() const;
 	
 	/** Access the server's playback manager. */
 	const FAvaMediaPlaybackManager& GetPlaybackManager() const { check(Manager); return *Manager; }
@@ -129,7 +129,7 @@ public:
 	void HandleStatCommand(const FAvaMediaStatCommand& InMessage, const TSharedRef<IMessageContext, ESPMode::ThreadSafe>& InContext);
 	void HandleDeviceProviderDataRequest(const FAvaDeviceProviderDataRequest& InMessage, const TSharedRef<IMessageContext, ESPMode::ThreadSafe>& InContext);
 	void HandleUpdateClientInfo(const FAvaMediaUpdateClientInfo& InMessage, const TSharedRef<IMessageContext, ESPMode::ThreadSafe>& InContext);
-	void HandleAvalancheInstanceSettingsUpdate(const FAvalancheInstanceSettingsUpdate& InMessage, const TSharedRef<IMessageContext, ESPMode::ThreadSafe>& InContext);
+	void HandleAvaInstanceSettingsUpdate(const FAvaInstanceSettingsUpdate& InMessage, const TSharedRef<IMessageContext, ESPMode::ThreadSafe>& InContext);
 	void HandlePackageEvent(const FAvaMediaPlaybackPackageEvent& InMessage, const TSharedRef<IMessageContext, ESPMode::ThreadSafe>& InContext);
 	void HandlePlaybackAssetStatusRequest(const FAvaMediaPlaybackAssetStatusRequest& InMessage, const TSharedRef<IMessageContext, ESPMode::ThreadSafe>& InContext);
 	void HandlePlaybackRequest(const FAvaMediaPlaybackRequest& InMessage, const TSharedRef<IMessageContext, ESPMode::ThreadSafe>& InContext);
@@ -310,7 +310,7 @@ private:
 		uint32 ProcessId = 0;
 		TMap<FString, FString> UserDataEntries;
 		FClientBroadcastSettings BroadcastSettings;
-		FAvalancheInstanceSettings AvalancheInstanceSettings;
+		FAvaInstanceSettings AvaInstanceSettings;
 
 		bool bClientInfoReceived = false;
 	
diff --git a/Engine/Plugins/Experimental/Avalanche/Source/AvalancheMedia/Public/Playback/AvalanchePlayable.h b/Engine/Plugins/Experimental/Avalanche/Source/AvalancheMedia/Public/Playback/AvalanchePlayable.h
--- a/Engine/Plugins/Experimental/Avalanche/Source/AvalancheMedia/Public/Playback/AvalanchePlayable.h
+++ b/Engine/Plugins/Experimental/Avalanche/Source/AvalancheMedia/Public/Playback/AvalanchePlayable.h
@@ -20,7 +20,7 @@ class UAvaSequencePlayer;
 class UAvalanchePlayableTransition;
 class URemoteControlPreset;
 struct FAvaSoftAssetPtr;
-struct FAvalancheInstancePlaySettings;
+struct FAvaInstancePlaySettings;
 struct FAvalancheRemoteControlValues;
 
 AVALANCHEMEDIA_API DECLARE_LOG_CATEGORY_EXTERN(LogAvalanchePlayable, Log, All);
@@ -124,7 +124,7 @@ public:
 	 * @param InWorldPlaySettings World settings (render target) to render this asset with.
 	 * @remark This doesn't trigger the animations.
 	 */
-	void BeginPlay(const FAvalancheInstancePlaySettings& InWorldPlaySettings);
+	void BeginPlay(const FAvaInstancePlaySettings& InWorldPlaySettings);
 
 	/**
 	 * @brief Make this asset not play anymore (will end up hidden)
diff --git a/Engine/Plugins/Experimental/Avalanche/Source/AvalancheMedia/Public/Playback/AvalanchePlayback.h b/Engine/Plugins/Experimental/Avalanche/Source/AvalancheMedia/Public/Playback/AvalanchePlayback.h
--- a/Engine/Plugins/Experimental/Avalanche/Source/AvalancheMedia/Public/Playback/AvalanchePlayback.h
+++ b/Engine/Plugins/Experimental/Avalanche/Source/AvalancheMedia/Public/Playback/AvalanchePlayback.h
@@ -21,7 +21,7 @@ class UAvaMediaPlayableGroupManager;
 class UAvaPlaybackNode;
 class UAvaPlaybackNodeRoot;
 class UAvalancheBroadcast;
-class UAvalancheGameInstance;
+class UAvaGameInstance;
 class UAvalanchePlayable;
 class UAvalanchePlayableTransition;
 class UEdGraph;
@@ -187,7 +187,7 @@ public:
 	bool UnloadAndRemovePlayable(UAvalanchePlayable* InPlayable, const FSoftObjectPath& InSourceAssetPath, const FName& InChannelName, bool bInForceImmediate);
 	bool UnloadAndRemovePlayable(const FSoftObjectPath& InSourceAssetPath, const FName& InChannelName, bool bInForceImmediate);
 
-	TArray<UAvalancheGameInstance*> GetActiveGameInstances() const;
+	TArray<UAvaGameInstance*> GetActiveGameInstances() const;
 	
 	/**
 	 *	Determines if the graph contains a node with the given source avalanche blueprint.
diff --git a/Engine/Plugins/Experimental/Avalanche/Source/AvalancheModifiers/Private/Modifiers/AvaTranslucentPriorityModifier.cpp b/Engine/Plugins/Experimental/Avalanche/Source/AvalancheModifiers/Private/Modifiers/AvaTranslucentPriorityModifier.cpp
--- a/Engine/Plugins/Experimental/Avalanche/Source/AvalancheModifiers/Private/Modifiers/AvaTranslucentPriorityModifier.cpp
+++ b/Engine/Plugins/Experimental/Avalanche/Source/AvalancheModifiers/Private/Modifiers/AvaTranslucentPriorityModifier.cpp
@@ -7,7 +7,7 @@
 #include "Camera/CameraComponent.h"
 #include "Components/PrimitiveComponent.h"
 #include "EngineUtils.h"
-#include "Framework/AvalancheGameInstance.h"
+#include "Framework/AvaGameInstance.h"
 #include "Modifiers/ActorModifierCoreBase.h"
 #include "Shared/AvaTranslucentPriorityModifierShared.h"
 
diff --git a/Engine/Plugins/Experimental/Avalanche/Source/AvalancheModifiers/Private/Shared/AvaTranslucentPriorityModifierShared.cpp b/Engine/Plugins/Experimental/Avalanche/Source/AvalancheModifiers/Private/Shared/AvaTranslucentPriorityModifierShared.cpp
--- a/Engine/Plugins/Experimental/Avalanche/Source/AvalancheModifiers/Private/Shared/AvaTranslucentPriorityModifierShared.cpp
+++ b/Engine/Plugins/Experimental/Avalanche/Source/AvalancheModifiers/Private/Shared/AvaTranslucentPriorityModifierShared.cpp
@@ -6,7 +6,7 @@
 #include "AvaOutlinerUtils.h"
 #include "Camera/CameraActor.h"
 #include "Camera/CameraComponent.h"
-#include "Framework/AvalancheGameInstance.h"
+#include "Framework/AvaGameInstance.h"
 
 #if WITH_EDITOR
 #include "AvaOutliner.h"
