Motion Design
- Changed "Avalanche" prefix to "Ava" on AvalancheInteractiveToolsModeDetailsObject, AvalancheInteractiveToolsModeDetailsObjectProvider, AvalancheOutlinerSettings, AvalancheMaskSettings, AvalancheMaskEditorMode, AvalancheEditorSettings.
- Changed "Ava" to "MotionDesign" in AvaTestBlueprintFunctionLibrary->ExportAvaBlueprintsToWorld and fixed the tooltip.
- Fixed tooltips / comments in AvaSplineActor, AvaBaseModifier, AvalancheEditorSettings->DefaultViewportQualitySettings .
- Remove "Motion Design" from the user-facing text for Null Actors and Null Components.
- Fixed AvaText3DComponent and AvaFont and changed "Avalanche" to "Ava" or "MotonDesign" where appropriate.

diff --git a/Engine/Plugins/Experimental/Avalanche/Config/DefaultAvalanche.ini b/Engine/Plugins/Experimental/Avalanche/Config/DefaultAvalanche.ini
--- a/Engine/Plugins/Experimental/Avalanche/Config/DefaultAvalanche.ini
+++ b/Engine/Plugins/Experimental/Avalanche/Config/DefaultAvalanche.ini
@@ -163,5 +163,14 @@
 +StructRedirects=(OldName="/Script/Avalanche.AvalancheInstanceSettings",NewName="/Script/Avalanche.AvaInstanceSettings")
 +PropertyRedirects=(OldName="/Script/AvalancheMedia.AvalancheMediaSettings.AvalancheInstanceSettings",NewName="/Script/AvalancheMedia.AvalancheMediaSettings.AvaInstanceSettings")
 
++FunctionRedirects=(OldName="/Script/AvalancheText.AvaText3DComponent.SetAvalancheFont",NewName="/Script/AvalancheText.AvaText3DComponent.SetMotionDesignFont")
++PropertyRedirects=(OldName="/Script/AvalancheText.AvaText3DComponent.AvalancheFont",NewName="/Script/AvalancheText.AvaText3DComponent.MotionDesignFont")
++PropertyRedirects=(OldName="/Script/AvalancheText.AvaFont.AvalancheFontObject",NewName="/Script/AvalancheText.AvaFont.MotionDesignFontObject")
++ClassRedirects=(OldName="/Script/AvalancheOutliner.AvalancheOutlinerSettings",NewName="/Script/AvalancheOutliner.AvaOutlinerSettings")
++ClassRedirects=(OldName="/Script/AvalancheEditor.AvalancheEditorSettings",NewName="/Script/AvalancheEditor.AvaEditorSettings")
++ClassRedirects=(OldName="/Script/AvalancheMask.AvalancheMaskSettings",NewName="/Script/AvalancheMask.AvaMaskSettings")
++ClassRedirects=(OldName="/Script/AvalancheMaskEditor.AvalancheMaskEditorMode",NewName="/Script/AvalancheMaskEditor.AvaMaskEditorMode")
++FunctionRedirects=(OldName="/Script/AvalancheTest.AvaTestBlueprintFunctionLibrary.ExportAvaBlueprintsToWorld",NewName="/Script/AvalancheTest.AvaTestBlueprintFunctionLibrary.ExportMotionDesignBlueprintsToWorld")
+
 [/Script/Avalanche.AvaAnimation]
 DefaultCompletionMode=KeepState
diff --git a/Engine/Plugins/Experimental/Avalanche/Source/Avalanche/Internal/Framework/AvaSplineActor.h b/Engine/Plugins/Experimental/Avalanche/Source/Avalanche/Internal/Framework/AvaSplineActor.h
--- a/Engine/Plugins/Experimental/Avalanche/Source/Avalanche/Internal/Framework/AvaSplineActor.h
+++ b/Engine/Plugins/Experimental/Avalanche/Source/Avalanche/Internal/Framework/AvaSplineActor.h
@@ -8,7 +8,7 @@
 class USplineComponent;
 
 /**
- * Avalanche Spline Actor
+ * Motion Design Spline Actor
  */
 UCLASS(MinimalAPI, DisplayName = "Motion Design Spline Actor")
 class AAvaSplineActor : public AActor
diff --git a/Engine/Plugins/Experimental/Avalanche/Source/Avalanche/Private/Framework/AvaNullActor.cpp b/Engine/Plugins/Experimental/Avalanche/Source/Avalanche/Private/Framework/AvaNullActor.cpp
--- a/Engine/Plugins/Experimental/Avalanche/Source/Avalanche/Private/Framework/AvaNullActor.cpp
+++ b/Engine/Plugins/Experimental/Avalanche/Source/Avalanche/Private/Framework/AvaNullActor.cpp
@@ -1,7 +1,6 @@
 // Copyright Epic Games, Inc. All Rights Reserved.
 
 #include "Framework/AvaNullActor.h"
-#include "Components/BillboardComponent.h"
 #include "Framework/AvaNullComponent.h"
 #include "UObject/ConstructorHelpers.h"
 
@@ -14,3 +13,10 @@ AAvaNullActor::AAvaNullActor()
 	NullComponent = CreateDefaultSubobject<UAvaNullComponent>(TEXT("NullComponent"));
 	SetRootComponent(NullComponent);
 }
+
+#if WITH_EDITOR
+FString AAvaNullActor::GetDefaultActorLabel() const
+{
+	return TEXT("Null Actor");
+}
+#endif
diff --git a/Engine/Plugins/Experimental/Avalanche/Source/Avalanche/Public/Framework/AvaNullActor.h b/Engine/Plugins/Experimental/Avalanche/Source/Avalanche/Public/Framework/AvaNullActor.h
--- a/Engine/Plugins/Experimental/Avalanche/Source/Avalanche/Public/Framework/AvaNullActor.h
+++ b/Engine/Plugins/Experimental/Avalanche/Source/Avalanche/Public/Framework/AvaNullActor.h
@@ -9,9 +9,9 @@
 class UAvaNullComponent;
 
 /**
- * Avalanche Null Actor (Empty Group)
+ * Null Actor (Empty Group)
  */
-UCLASS(MinimalAPI, DisplayName = "Motion Design Null Actor")
+UCLASS(MinimalAPI, DisplayName = "Null Actor")
 class AAvaNullActor : public AActor
 {
 	GENERATED_BODY()
@@ -21,7 +21,13 @@ public:
 
 	AAvaNullActor();
 
+#if WITH_EDITOR
+	//~ Begin AActor
+	virtual FString GetDefaultActorLabel() const override;
+	//~ End AActor
+#endif
+
 protected:
-	UPROPERTY(VisibleAnywhere, BlueprintReadOnly, Category = "Motion Design")
+	UPROPERTY(VisibleAnywhere, BlueprintReadOnly, Category = "Components")
 	TObjectPtr<UAvaNullComponent> NullComponent;
 };
diff --git a/Engine/Plugins/Experimental/Avalanche/Source/Avalanche/Public/Framework/AvaNullComponent.h b/Engine/Plugins/Experimental/Avalanche/Source/Avalanche/Public/Framework/AvaNullComponent.h
--- a/Engine/Plugins/Experimental/Avalanche/Source/Avalanche/Public/Framework/AvaNullComponent.h
+++ b/Engine/Plugins/Experimental/Avalanche/Source/Avalanche/Public/Framework/AvaNullComponent.h
@@ -5,7 +5,7 @@
 #include "Components/SceneComponent.h"
 #include "AvaNullComponent.generated.h"
 
-UCLASS(MinimalAPI, DisplayName = "Motion Design Null Component")
+UCLASS(MinimalAPI, DisplayName = "Null Component")
 class UAvaNullComponent : public USceneComponent
 {
 	GENERATED_BODY()
diff --git a/Engine/Plugins/Experimental/Avalanche/Source/Avalanche/Public/Viewport/AvaCineCameraActor.h b/Engine/Plugins/Experimental/Avalanche/Source/Avalanche/Public/Viewport/AvaCineCameraActor.h
--- a/Engine/Plugins/Experimental/Avalanche/Source/Avalanche/Public/Viewport/AvaCineCameraActor.h
+++ b/Engine/Plugins/Experimental/Avalanche/Source/Avalanche/Public/Viewport/AvaCineCameraActor.h
@@ -1,4 +1,4 @@
-﻿// Copyright Epic Games, Inc. All Rights Reserved.
+// Copyright Epic Games, Inc. All Rights Reserved.
 
 #pragma once
 
@@ -16,7 +16,7 @@
 UCLASS(DisplayName = "Motion Design Cine Camera Actor")
 class AVALANCHE_API AAvaCineCameraActor : public ACineCameraActor
 {
-	friend class UAvalancheEditorSettings;
+	friend class UAvaEditorSettings;
 	
 	GENERATED_BODY()
 
@@ -42,7 +42,7 @@ public:
 			// smaller camera mesh, todo: this will be substituted with a custom visualizer, or other solution, to visualize frustum + direction
 			CineCameraComp->SetWorldScale3D(FVector(0.5f));
 
-			// using default ManualFocusDistance value from Default CineCameraComponent, since that is handled by UAvalancheEditorSettings
+			// using default ManualFocusDistance value from Default CineCameraComponent, since that is handled by UAvaEditorSettings
 			CineCameraComp->FocusSettings.ManualFocusDistance = DefaultCameraDistance;
 
 			FVector CameraPosition = FVector::ZeroVector;
diff --git a/Engine/Plugins/Experimental/Avalanche/Source/AvalancheEditor/Private/AvaEditorActorUtils.cpp b/Engine/Plugins/Experimental/Avalanche/Source/AvalancheEditor/Private/AvaEditorActorUtils.cpp
--- a/Engine/Plugins/Experimental/Avalanche/Source/AvalancheEditor/Private/AvaEditorActorUtils.cpp
+++ b/Engine/Plugins/Experimental/Avalanche/Source/AvalancheEditor/Private/AvaEditorActorUtils.cpp
@@ -1,7 +1,7 @@
 // Copyright Epic Games, Inc. All Rights Reserved.
 
 #include "AvaEditorActorUtils.h"
-#include "AvalancheEditorSettings.h"
+#include "AvaEditorSettings.h"
 #include "Components/SceneComponent.h"
 #include "GameFramework/Actor.h"
 
@@ -9,7 +9,7 @@
 
 void FAvaEditorActorUtils::GetActorsToEdit(TArray<AActor*>& InOutSelectedActors)
 {
-	if (GetDefault<UAvalancheEditorSettings>()->bAutoIncludeAttachedActorsInEditActions)
+	if (GetDefault<UAvaEditorSettings>()->bAutoIncludeAttachedActorsInEditActions)
 	{
 		TSet<AActor*> SelectedActors(InOutSelectedActors);
 		TArray<AActor*> ActorsRemaining(InOutSelectedActors);
diff --git a/Engine/Plugins/Experimental/Avalanche/Source/AvalancheEditor/Private/AvaEditorSettings.cpp b/Engine/Plugins/Experimental/Avalanche/Source/AvalancheEditor/Private/AvaEditorSettings.cpp
new file mode 100644
--- /dev/null
+++ b/Engine/Plugins/Experimental/Avalanche/Source/AvalancheEditor/Private/AvaEditorSettings.cpp
@@ -0,0 +1,75 @@
+// Copyright Epic Games, Inc. All Rights Reserved.
+
+#include "AvaEditorSettings.h"
+#include "ISettingsModule.h"
+#include "Modules/ModuleManager.h"
+#include "Viewport/AvaCineCameraActor.h"
+
+#define LOCTEXT_NAMESPACE "AvaEditorSettings"
+
+UAvaEditorSettings::UAvaEditorSettings()
+{
+	CategoryName = TEXT("Motion Design");
+	SectionName = TEXT("Editor");
+}
+
+UAvaEditorSettings* UAvaEditorSettings::Get()
+{
+	UAvaEditorSettings* DefaultSettings = GetMutableDefault<UAvaEditorSettings>();
+	static bool bInitialized = false;
+	if (!bInitialized)
+	{
+		bInitialized = true;
+		DefaultSettings->SetFlags(RF_Transactional);
+	}
+	return DefaultSettings;
+}
+
+void UAvaEditorSettings::OpenEditorSettingsWindow() const
+{
+	static ISettingsModule& SettingsModule = FModuleManager::LoadModuleChecked<ISettingsModule>(TEXT("Settings"));
+	SettingsModule.ShowViewer(GetContainerName(), CategoryName, SectionName);
+}
+
+void UAvaEditorSettings::PostInitProperties()
+{
+	Super::PostInitProperties();
+
+	if (const FName PresetNameNoLumen = TEXT("No Lumen");
+		!ViewportQualityPresets.Contains(PresetNameNoLumen))
+	{
+		ViewportQualityPresets.Add(PresetNameNoLumen, FAvaViewportQualitySettings::Preset(PresetNameNoLumen));
+	}
+
+	if (const FName PresetNameReduced = TEXT("Reduced");
+		!ViewportQualityPresets.Contains(PresetNameReduced))
+	{
+		ViewportQualityPresets.Add(PresetNameReduced, FAvaViewportQualitySettings::Preset(PresetNameReduced));
+	}
+
+	// Verifying integrity and sorting will ensure any new/removed entries are handled correctly.
+	DefaultViewportQualitySettings.VerifyIntegrity();
+	DefaultViewportQualitySettings.SortFeaturesByDisplayText();
+
+	for (TPair<FName, FAvaViewportQualitySettings>& Preset : ViewportQualityPresets)
+	{
+		FAvaViewportQualitySettings& Settings = Preset.Value;
+		Settings.VerifyIntegrity();
+		Settings.SortFeaturesByDisplayText();
+	}
+}
+
+void UAvaEditorSettings::PostEditChangeProperty(FPropertyChangedEvent& PropertyChangedEvent)
+{
+	Super::PostEditChangeProperty(PropertyChangedEvent);
+
+	OnChanged.Broadcast(this, PropertyChangedEvent.GetPropertyName());
+}
+
+// todo: when the reset to defaults issue will be solved, this function could be used from both PostInitProperties and PECP to update AAvaCineCameraActor defaults
+void UAvaEditorSettings::UpdateAvaCineCameraDefaults() const
+{
+	AAvaCineCameraActor::SetDefaultCameraDistance(CameraDistance);
+}
+
+#undef LOCTEXT_NAMESPACE
diff --git a/Engine/Plugins/Experimental/Avalanche/Source/AvalancheEditor/Private/AvaEditorSettings.h b/Engine/Plugins/Experimental/Avalanche/Source/AvalancheEditor/Private/AvaEditorSettings.h
new file mode 100644
--- /dev/null
+++ b/Engine/Plugins/Experimental/Avalanche/Source/AvalancheEditor/Private/AvaEditorSettings.h
@@ -0,0 +1,97 @@
+// Copyright Epic Games, Inc. All Rights Reserved.
+
+#pragma once
+
+#include "AvaDefs.h"
+#include "CoreMinimal.h"
+#include "Engine/DeveloperSettings.h"
+#include "Viewport/AvaViewportQualitySettings.h"
+#include "AvaEditorSettings.generated.h"
+
+USTRUCT()
+struct FAvaPaletteSettings
+{
+	GENERATED_BODY()
+
+	UPROPERTY(Config, VisibleAnywhere, Category = "Interface")
+	FName Name;
+
+	// Name is the TMap key
+	UPROPERTY(Config, EditAnywhere, Category = "Interface")
+	bool bEnablePalette = true;
+
+	UPROPERTY(Config)
+	bool bExpanded = false;
+};
+
+USTRUCT()
+struct FAvaPaletteTabSettings
+{
+	GENERATED_BODY()
+
+	UPROPERTY(Config, EditAnywhere, EditFixedSize, Category = "Interface", meta = (TitleProperty = "Name", FullyExpand))
+	TArray<FAvaPaletteSettings> Palettes;
+};
+
+/**
+ * Motion Design Editor Settings
+ */
+UCLASS(Config=EditorPerProjectUserSettings, meta = (DisplayName = "Editor"))
+class UAvaEditorSettings : public UDeveloperSettings
+{
+	GENERATED_BODY()
+
+public:
+	UAvaEditorSettings();
+
+	virtual ~UAvaEditorSettings() override = default;
+
+	static UAvaEditorSettings* Get();
+
+	/** Whether to Automatically Include the Attached Actors when performing Edit Actions such as Cut, Copy, Duplicate. */
+	UPROPERTY(Config, EditAnywhere, BlueprintReadWrite, Category = "Behavior")
+	bool bAutoIncludeAttachedActorsInEditActions = true;
+
+	/** When Grouping Actors with a Null Actor, whether to keep the relative transform of these Actors */
+	UPROPERTY(Config, EditAnywhere, BlueprintReadWrite, Category = "Behavior")
+	bool bKeepRelativeTransformWhenGrouping = false;
+
+	/*
+	 * Distance from the camera that new actors are created via the toolbox or drag and drop.
+	 * Also sets the distance from the origin that new Camera Preview Viewport cameras are created.
+	 */
+	UPROPERTY(Config, EditAnywhere, BlueprintReadWrite, Category = "Behavior")
+	float CameraDistance = 500.0f;
+
+	/**  */
+	UPROPERTY(Config, EditAnywhere, BlueprintReadWrite, Category = "Toolbox")
+	TArray<FString> FavoriteTools;
+
+	/**  */
+	UPROPERTY(Config, EditAnywhere, EditFixedSize, Category = "Interface", meta = (EditFixedOrder, ShowOnlyInnerProperties))
+	TArray<FAvaPaletteTabSettings> PaletteTabs;
+
+	/** Default viewport quality settings for all newly created Motion Design blueprints. */
+	UPROPERTY(Config, EditAnywhere, BlueprintReadWrite, Category = "Quality")
+	FAvaViewportQualitySettings DefaultViewportQualitySettings = FAvaViewportQualitySettings(true);
+
+	/** Viewport quality settings user presets. */
+	UPROPERTY(Config, EditAnywhere, BlueprintReadWrite, Category = "Quality")
+	TMap<FName, FAvaViewportQualitySettings> ViewportQualityPresets;
+
+	DECLARE_MULTICAST_DELEGATE_TwoParams(FOnSettingsChanged, const UAvaEditorSettings* InSettings, FName InSetting)
+	FOnSettingsChanged OnChanged;
+
+	//~ Begin UObject
+	virtual void PostInitProperties() override;
+	virtual void PostEditChangeProperty(FPropertyChangedEvent& PropertyChangedEvent) override;
+	//~ End UObject
+
+	void OpenEditorSettingsWindow() const;
+
+	/**
+	 * AvaCineCamera uses CameraDistance property to setup manual focus in its default object.
+	 * This function updates that default value to match the one from AvaEditorSettings
+	 */
+	void UpdateAvaCineCameraDefaults() const;
+};
diff --git a/Engine/Plugins/Experimental/Avalanche/Source/AvalancheEditor/Private/AvaEditorViewportUtils.cpp b/Engine/Plugins/Experimental/Avalanche/Source/AvalancheEditor/Private/AvaEditorViewportUtils.cpp
--- a/Engine/Plugins/Experimental/Avalanche/Source/AvalancheEditor/Private/AvaEditorViewportUtils.cpp
+++ b/Engine/Plugins/Experimental/Avalanche/Source/AvalancheEditor/Private/AvaEditorViewportUtils.cpp
@@ -1,7 +1,7 @@
 // Copyright Epic Games, Inc. All Rights Reserved.
 
 #include "AvaEditorViewportUtils.h"
-#include "AvalancheEditorSettings.h"
+#include "AvaEditorSettings.h"
 #include "Templates/SharedPointer.h"
 #include "ViewportClient/IAvaViewportClient.h"
 
@@ -14,7 +14,7 @@ bool FAvaEditorViewportUtils::MeshSizeToPixelSize(const TSharedRef<IAvaViewportC
 		return false;
 	}
 
-	const FVector2D FrustumSize = InViewportClient->GetFrustumSizeAtDistance(UAvalancheEditorSettings::Get()->CameraDistance);
+	const FVector2D FrustumSize = InViewportClient->GetFrustumSizeAtDistance(UAvaEditorSettings::Get()->CameraDistance);
 
 	if (FrustumSize.X <= 0)
 	{
@@ -35,7 +35,7 @@ bool FAvaEditorViewportUtils::PixelSizeToMeshSize(const TSharedRef<IAvaViewportC
 		return false;
 	}
 
-	const FVector2D FrustumSize = InViewportClient->GetFrustumSizeAtDistance(UAvalancheEditorSettings::Get()->CameraDistance);
+	const FVector2D FrustumSize = InViewportClient->GetFrustumSizeAtDistance(UAvaEditorSettings::Get()->CameraDistance);
 
 	if (FrustumSize.X <= 0)
 	{
diff --git a/Engine/Plugins/Experimental/Avalanche/Source/AvalancheEditor/Private/AvalancheEditorSettings.cpp b/Engine/Plugins/Experimental/Avalanche/Source/AvalancheEditor/Private/AvalancheEditorSettings.cpp
deleted file mode 100644
diff --git a/Engine/Plugins/Experimental/Avalanche/Source/AvalancheEditor/Private/AvalancheEditorSettings.h b/Engine/Plugins/Experimental/Avalanche/Source/AvalancheEditor/Private/AvalancheEditorSettings.h
deleted file mode 100644
diff --git a/Engine/Plugins/Experimental/Avalanche/Source/AvalancheEditor/Private/Outliner/AvaOutlinerExtension.cpp b/Engine/Plugins/Experimental/Avalanche/Source/AvalancheEditor/Private/Outliner/AvaOutlinerExtension.cpp
--- a/Engine/Plugins/Experimental/Avalanche/Source/AvalancheEditor/Private/Outliner/AvaOutlinerExtension.cpp
+++ b/Engine/Plugins/Experimental/Avalanche/Source/AvalancheEditor/Private/Outliner/AvaOutlinerExtension.cpp
@@ -4,7 +4,7 @@
 #include "Animation/SkeletalMeshActor.h"
 #include "AvaEditorCommands.h"
 #include "AvaEditorModule.h"
-#include "AvalancheEditorSettings.h"
+#include "AvaEditorSettings.h"
 #include "AvaOutliner.h"
 #include "AvaOutlinerSubsystem.h"
 #include "AvaOutlinerTabSpawner.h"
@@ -416,7 +416,7 @@ void FAvaOutlinerExtension::NotifyOutlinerItemLockChanged(const FAvaOutlinerItem
 
 const FAttachmentTransformRules& FAvaOutlinerExtension::GetTransformRule(bool bIsPrimaryTransformRule) const
 {
-	const bool bKeepRelativeWhenGrouping = UAvalancheEditorSettings::Get()->bKeepRelativeTransformWhenGrouping;
+	const bool bKeepRelativeWhenGrouping = UAvaEditorSettings::Get()->bKeepRelativeTransformWhenGrouping;
 	const bool bReturnRelative = bIsPrimaryTransformRule == bKeepRelativeWhenGrouping;
 	return bReturnRelative ? FAttachmentTransformRules::KeepRelativeTransform : FAttachmentTransformRules::KeepWorldTransform;
 }
@@ -435,7 +435,7 @@ void FAvaOutlinerExtension::GroupSelection()
 	FTransform SpawnTransform;
 	TOptional<FAttachmentTransformRules> TransformRules;
 
-	if (UAvalancheEditorSettings::Get()->bKeepRelativeTransformWhenGrouping)
+	if (UAvaEditorSettings::Get()->bKeepRelativeTransformWhenGrouping)
 	{
 		SpawnTransform = FTransform::Identity;
 		TransformRules = FAttachmentTransformRules::KeepRelativeTransform;
diff --git a/Engine/Plugins/Experimental/Avalanche/Source/AvalancheEditor/Private/Scene/AvaSceneDefaults.cpp b/Engine/Plugins/Experimental/Avalanche/Source/AvalancheEditor/Private/Scene/AvaSceneDefaults.cpp
--- a/Engine/Plugins/Experimental/Avalanche/Source/AvalancheEditor/Private/Scene/AvaSceneDefaults.cpp
+++ b/Engine/Plugins/Experimental/Avalanche/Source/AvalancheEditor/Private/Scene/AvaSceneDefaults.cpp
@@ -3,7 +3,7 @@
 #include "Scene/AvaSceneDefaults.h"
 #include "ActorFactories/ActorFactory.h"
 #include "AssetViewerSettings.h"
-#include "AvalancheEditorSettings.h"
+#include "AvaEditorSettings.h"
 #include "Builders/CubeBuilder.h"
 #include "Components/DirectionalLightComponent.h"
 #include "Components/SkyLightComponent.h"
@@ -198,7 +198,7 @@ namespace UE::AvalancheEditor::Private
 
 		FTransform PostProcessVolumeTransform = FTransform::Identity;
 
-		float PostProcessXOffset = UAvalancheEditorSettings::Get()->CameraDistance;
+		float PostProcessXOffset = UAvaEditorSettings::Get()->CameraDistance;
 
 		if (const UCubeBuilder* const DefaultCubeBuilder = Cast<UCubeBuilder>(UCubeBuilder::StaticClass()->GetDefaultObject()))
 		{
@@ -222,13 +222,13 @@ namespace UE::AvalancheEditor::Private
 		if (AActor* DefaultRoot = GetActorFromResponses(InActorResponses, DefaultSceneActorNames::DefaultRoot))
 		{
 			FTransform Transform = FTransform::Identity;
-			Transform.SetLocation(FVector(-UAvalancheEditorSettings::Get()->CameraDistance, 0, 0));
+			Transform.SetLocation(FVector(-UAvaEditorSettings::Get()->CameraDistance, 0, 0));
 
 			Camera->AttachToActor(DefaultRoot, FAttachmentTransformRules::KeepWorldTransform);
 			Camera->SetActorRelativeTransform(Transform);
 		}
 
-		Camera->Configure(UAvalancheEditorSettings::Get()->CameraDistance);
+		Camera->Configure(UAvaEditorSettings::Get()->CameraDistance);
 	}
 
 	TMap<FName, TSharedRef<FAvaSceneDefaultActorResponse>> GeneratorInitialActorResponses(UWorld* InWorld)
diff --git a/Engine/Plugins/Experimental/Avalanche/Source/AvalancheInteractiveTools/Private/Tools/AvaInteractiveToolsToolBase.cpp b/Engine/Plugins/Experimental/Avalanche/Source/AvalancheInteractiveTools/Private/Tools/AvaInteractiveToolsToolBase.cpp
--- a/Engine/Plugins/Experimental/Avalanche/Source/AvalancheInteractiveTools/Private/Tools/AvaInteractiveToolsToolBase.cpp
+++ b/Engine/Plugins/Experimental/Avalanche/Source/AvalancheInteractiveTools/Private/Tools/AvaInteractiveToolsToolBase.cpp
@@ -12,8 +12,8 @@
 #include "EdMode/AvaInteractiveToolsEdMode.h"
 #include "Engine/World.h"
 #include "Framework/Application/SlateApplication.h"
-#include "IAvalancheInteractiveToolsModeDetailsObject.h"
-#include "IAvalancheInteractiveToolsModeDetailsObjectProvider.h"
+#include "IAvaInteractiveToolsModeDetailsObject.h"
+#include "IAvaInteractiveToolsModeDetailsObjectProvider.h"
 #include "InteractiveToolManager.h"
 #include "InteractiveToolsContext.h"
 #include "Planners/AvaInteractiveToolsToolViewportAreaPlanner.h"
@@ -49,13 +49,13 @@ UObject* UAvaInteractiveToolsToolBase::GetDetailsObjectFromActor(AActor* InActor
 		return nullptr;
 	}
 
-	if (InActor->Implements<UAvalancheInteractiveToolsModeDetailsObject>())
+	if (InActor->Implements<UAvaInteractiveToolsModeDetailsObject>())
 	{
 		return InActor;
 	}
-	else if (InActor->Implements<UAvalancheInteractiveToolsModeDetailsObjectProvider>())
+	else if (InActor->Implements<UAvaInteractiveToolsModeDetailsObjectProvider>())
 	{
-		return IAvalancheInteractiveToolsModeDetailsObjectProvider::Execute_GetModeDetailsObject(InActor);
+		return IAvaInteractiveToolsModeDetailsObjectProvider::Execute_GetModeDetailsObject(InActor);
 	}
 
 	TArray<UActorComponent*> Components;
@@ -63,14 +63,14 @@ UObject* UAvaInteractiveToolsToolBase::GetDetailsObjectFromActor(AActor* InActor
 
 	for (UActorComponent* Component : Components)
 	{
-		if (Component->Implements<UAvalancheInteractiveToolsModeDetailsObject>())
+		if (Component->Implements<UAvaInteractiveToolsModeDetailsObject>())
 		{
 			return Component;
 			break;
 		}
-		else if (Component->Implements<UAvalancheInteractiveToolsModeDetailsObjectProvider>())
+		else if (Component->Implements<UAvaInteractiveToolsModeDetailsObjectProvider>())
 		{
-			if (UObject* DetailsObject = IAvalancheInteractiveToolsModeDetailsObjectProvider::Execute_GetModeDetailsObject(Component))
+			if (UObject* DetailsObject = IAvaInteractiveToolsModeDetailsObjectProvider::Execute_GetModeDetailsObject(Component))
 			{
 				return DetailsObject;
 			}
diff --git a/Engine/Plugins/Experimental/Avalanche/Source/AvalancheInteractiveTools/Public/Tools/AvaInteractiveToolsToolBase.h b/Engine/Plugins/Experimental/Avalanche/Source/AvalancheInteractiveTools/Public/Tools/AvaInteractiveToolsToolBase.h
--- a/Engine/Plugins/Experimental/Avalanche/Source/AvalancheInteractiveTools/Public/Tools/AvaInteractiveToolsToolBase.h
+++ b/Engine/Plugins/Experimental/Avalanche/Source/AvalancheInteractiveTools/Public/Tools/AvaInteractiveToolsToolBase.h
@@ -48,8 +48,8 @@ public:
 	static constexpr int32 BID_Cancel = 1;
 
 	/**
-	 * Checks the given actor and its components against UAvalancheInteractiveToolsModeDetailsObject
-	 * and UAvalancheInteractiveToolsModeDetailsObjectProvider.
+	 * Checks the given actor and its components against UAvaInteractiveToolsModeDetailsObject
+	 * and UAvaInteractiveToolsModeDetailsObjectProvider.
 	 */
 	static UObject* GetDetailsObjectFromActor(AActor* InActor);
 
diff --git a/Engine/Plugins/Experimental/Avalanche/Source/AvalancheInteractiveToolsRuntime/Public/IAvaInteractiveToolsModeDetailsObject.h b/Engine/Plugins/Experimental/Avalanche/Source/AvalancheInteractiveToolsRuntime/Public/IAvaInteractiveToolsModeDetailsObject.h
new file mode 100644
--- /dev/null
+++ b/Engine/Plugins/Experimental/Avalanche/Source/AvalancheInteractiveToolsRuntime/Public/IAvaInteractiveToolsModeDetailsObject.h
@@ -0,0 +1,17 @@
+// Copyright Epic Games, Inc. All Rights Reserved.
+
+#pragma once
+
+#include "UObject/Interface.h"
+#include "IAvaInteractiveToolsModeDetailsObject.generated.h"
+
+UINTERFACE(MinimalAPI)
+class UAvaInteractiveToolsModeDetailsObject : public UInterface
+{
+	GENERATED_BODY()
+};
+
+class IAvaInteractiveToolsModeDetailsObject
+{
+	GENERATED_BODY()
+};
diff --git a/Engine/Plugins/Experimental/Avalanche/Source/AvalancheInteractiveToolsRuntime/Public/IAvaInteractiveToolsModeDetailsObjectProvider.h b/Engine/Plugins/Experimental/Avalanche/Source/AvalancheInteractiveToolsRuntime/Public/IAvaInteractiveToolsModeDetailsObjectProvider.h
new file mode 100644
--- /dev/null
+++ b/Engine/Plugins/Experimental/Avalanche/Source/AvalancheInteractiveToolsRuntime/Public/IAvaInteractiveToolsModeDetailsObjectProvider.h
@@ -0,0 +1,25 @@
+// Copyright Epic Games, Inc. All Rights Reserved.
+
+#pragma once
+
+#include "UObject/Interface.h"
+#include "IAvaInteractiveToolsModeDetailsObjectProvider.generated.h"
+
+class UObject;
+
+UINTERFACE(MinimalAPI)
+class UAvaInteractiveToolsModeDetailsObjectProvider : public UInterface
+{
+	GENERATED_BODY()
+};
+
+class AVALANCHEINTERACTIVETOOLSRUNTIME_API IAvaInteractiveToolsModeDetailsObjectProvider
+{
+	GENERATED_BODY()
+
+public:
+	UFUNCTION(BlueprintNativeEvent, CallInEditor, Category = "Motion Design Interactive Tools")
+	UObject* GetModeDetailsObject() const;
+
+	virtual UObject* GetModeDetailsObject_Implementation() const PURE_VIRTUAL(IAvaInteractiveToolsModeDetailsObjectProvider::GetModeDetailsObject, return nullptr;)
+};
diff --git a/Engine/Plugins/Experimental/Avalanche/Source/AvalancheInteractiveToolsRuntime/Public/IAvalancheInteractiveToolsModeDetailsObject.h b/Engine/Plugins/Experimental/Avalanche/Source/AvalancheInteractiveToolsRuntime/Public/IAvalancheInteractiveToolsModeDetailsObject.h
deleted file mode 100644
diff --git a/Engine/Plugins/Experimental/Avalanche/Source/AvalancheInteractiveToolsRuntime/Public/IAvalancheInteractiveToolsModeDetailsObjectProvider.h b/Engine/Plugins/Experimental/Avalanche/Source/AvalancheInteractiveToolsRuntime/Public/IAvalancheInteractiveToolsModeDetailsObjectProvider.h
deleted file mode 100644
diff --git a/Engine/Plugins/Experimental/Avalanche/Source/AvalancheMask/Private/AvaMaskSettings.cpp b/Engine/Plugins/Experimental/Avalanche/Source/AvalancheMask/Private/AvaMaskSettings.cpp
--- a/Engine/Plugins/Experimental/Avalanche/Source/AvalancheMask/Private/AvaMaskSettings.cpp
+++ b/Engine/Plugins/Experimental/Avalanche/Source/AvalancheMask/Private/AvaMaskSettings.cpp
@@ -11,14 +11,14 @@
 #include "Materials/MaterialFunctionInterface.h"
 #endif
 
-UAvalancheMaskSettings::UAvalancheMaskSettings()
+UAvaMaskSettings::UAvaMaskSettings()
 {
 	CategoryName = TEXT("Motion Design");
 	SectionName = TEXT("Mask");
 }
 
 #if WITH_EDITORONLY_DATA
-UMaterialFunctionInterface* UAvalancheMaskSettings::GetMaterialFunction()
+UMaterialFunctionInterface* UAvaMaskSettings::GetMaterialFunction()
 {
 	if (UMaterialFunctionInterface* StrongMaterialFunction = MaterialFunction.Get())
 	{
@@ -31,7 +31,7 @@ UMaterialFunctionInterface* UAvalancheMaskSettings::GetMaterialFunction()
 	return DefaultMaterialFunction;
 }
 
-UMaterialFunctionInterface* UAvalancheMaskSettings::GetDefaultMaterialFunction()
+UMaterialFunctionInterface* UAvaMaskSettings::GetDefaultMaterialFunction()
 {
 	// /Script/Engine.MaterialFunction'/GeometryMask/GeometryMask/MF_ApplyGeometryMask.MF_ApplyGeometryMask'
 	// const FString AssetPath = FString(TEXT("/GeometryMask")) / TEXT("GeometryMask") / TEXT("MF_ApplyGeometryMask");
diff --git a/Engine/Plugins/Experimental/Avalanche/Source/AvalancheMask/Private/AvaMaskSettings.h b/Engine/Plugins/Experimental/Avalanche/Source/AvalancheMask/Private/AvaMaskSettings.h
--- a/Engine/Plugins/Experimental/Avalanche/Source/AvalancheMask/Private/AvaMaskSettings.h
+++ b/Engine/Plugins/Experimental/Avalanche/Source/AvalancheMask/Private/AvaMaskSettings.h
@@ -13,15 +13,15 @@
 class UMaterialFunctionInterface;
 #endif
 
-/** Settings for Avalanche Mask */
+/** Settings for Motion Design Mask */
 UCLASS(Config = Engine, meta = (DisplayName = "Mask"))
-class UAvalancheMaskSettings
+class UAvaMaskSettings
     : public UDeveloperSettings
 {
     GENERATED_BODY()
 
 public:
-    UAvalancheMaskSettings();
+    UAvaMaskSettings();
 
 #if WITH_EDITORONLY_DATA
 	/** Material Function to use to expect or add to a material. */
diff --git a/Engine/Plugins/Experimental/Avalanche/Source/AvalancheMask/Private/Handling/AvaMaskDesignedMaterialHandle.cpp b/Engine/Plugins/Experimental/Avalanche/Source/AvalancheMask/Private/Handling/AvaMaskDesignedMaterialHandle.cpp
--- a/Engine/Plugins/Experimental/Avalanche/Source/AvalancheMask/Private/Handling/AvaMaskDesignedMaterialHandle.cpp
+++ b/Engine/Plugins/Experimental/Avalanche/Source/AvalancheMask/Private/Handling/AvaMaskDesignedMaterialHandle.cpp
@@ -214,7 +214,7 @@ bool FAvaMaskDesignedMaterialHandle::ApplyModifiedState(
 				, InModifiedParameters.MaterialParameters.InnerFeatherRadius);
 
 #if WITH_EDITOR
-			UMaterialFunctionInterface* MaterialFunctionToUse = GetMutableDefault<UAvalancheMaskSettings>()->GetMaterialFunction();
+			UMaterialFunctionInterface* MaterialFunctionToUse = GetMutableDefault<UAvaMaskSettings>()->GetMaterialFunction();
 			SetOutputProcessor(MaterialFunctionToUse);
 #endif
 
@@ -327,7 +327,7 @@ void FAvaMaskDesignedMaterialHandle::SetOutputProcessor(UMaterialFunctionInterfa
 	{
 		if (!::IsValid(InMaterialFunction))
 		{
-			InMaterialFunction = GetMutableDefault<UAvalancheMaskSettings>()->GetMaterialFunction();
+			InMaterialFunction = GetMutableDefault<UAvaMaskSettings>()->GetMaterialFunction();
 		}
 
 		if (UDynamicMaterialModel* Model = DesignedMaterial->GetMaterialModel())
diff --git a/Engine/Plugins/Experimental/Avalanche/Source/AvalancheMaskEditor/Private/AvaMaskEditorMode.cpp b/Engine/Plugins/Experimental/Avalanche/Source/AvalancheMaskEditor/Private/AvaMaskEditorMode.cpp
--- a/Engine/Plugins/Experimental/Avalanche/Source/AvalancheMaskEditor/Private/AvaMaskEditorMode.cpp
+++ b/Engine/Plugins/Experimental/Avalanche/Source/AvalancheMaskEditor/Private/AvaMaskEditorMode.cpp
@@ -29,31 +29,31 @@
 
 #define LOCTEXT_NAMESPACE "AvalancheMaskEditorMode"
 
-const FEditorModeID UAvalancheMaskEditorMode::EM_AvalancheMaskEditorModeId(UE::AvalancheMaskEditor::AvalancheMaskEditorModeName);
+const FEditorModeID UAvaMaskEditorMode::EM_MotionDesignMaskEditorModeId(UE::AvalancheMaskEditor::MotionDesignMaskEditorModeName);
 
-UAvalancheMaskEditorMode::UAvalancheMaskEditorMode()
+UAvaMaskEditorMode::UAvaMaskEditorMode()
 {
-	Info = FEditorModeInfo(UAvalancheMaskEditorMode::EM_AvalancheMaskEditorModeId,
+	Info = FEditorModeInfo(UAvaMaskEditorMode::EM_MotionDesignMaskEditorModeId,
 		LOCTEXT("MotionDesignMaskEditorModeName", "MotionDesign Mask Editor Mode"),
 		FSlateIcon(),
 		false);
 }
 
-void UAvalancheMaskEditorMode::Enter()
+void UAvaMaskEditorMode::Enter()
 {
 	UEdMode::Enter();
 
 	if (GEditor && GEngine)
 	{
 		UTypedElementSelectionSet* SelectionSet = GetModeManager()->GetSelectedActors()->GetElementSelectionSet();
-		SelectionSet->OnChanged().AddUObject(this, &UAvalancheMaskEditorMode::OnSelectionChanged);
+		SelectionSet->OnChanged().AddUObject(this, &UAvaMaskEditorMode::OnSelectionChanged);
 		
 		WeakLastSelectedActor = SelectionSet->GetTopSelectedObject<AActor>();
 		WeakActorSelectionSet = SelectionSet;
 		
-		OnActorSpawnedHandle = GetWorld()->AddOnActorSpawnedHandler(FOnActorSpawned::FDelegate::CreateUObject(this, &UAvalancheMaskEditorMode::OnActorSpawned));
+		OnActorSpawnedHandle = GetWorld()->AddOnActorSpawnedHandler(FOnActorSpawned::FDelegate::CreateUObject(this, &UAvaMaskEditorMode::OnActorSpawned));
 
-		GEngine->OnLevelActorAdded().AddUObject(this, &UAvalancheMaskEditorMode::OnActorSpawned);
+		GEngine->OnLevelActorAdded().AddUObject(this, &UAvaMaskEditorMode::OnActorSpawned);
 	}
 
 	for (AActor* Actor : TActorRange<AActor>(GetWorld()))
@@ -153,7 +153,7 @@ void UAvalancheMaskEditorMode::Enter()
 							.VAlign(VAlign_Center)
 							.OnClicked_Lambda([EditorModeManager = GetModeManager()]() 
 							{
-								EditorModeManager->DeactivateMode(UAvalancheMaskEditorMode::EM_AvalancheMaskEditorModeId);						
+								EditorModeManager->DeactivateMode(UAvaMaskEditorMode::EM_MotionDesignMaskEditorModeId);						
 								return FReply::Handled(); 
 							})
 						]
@@ -184,7 +184,7 @@ void UAvalancheMaskEditorMode::Enter()
 	}
 }
 
-void UAvalancheMaskEditorMode::Exit()
+void UAvaMaskEditorMode::Exit()
 {
 	if (GEditor && GEngine)
 	{
@@ -226,27 +226,27 @@ void UAvalancheMaskEditorMode::Exit()
 	UEdMode::Exit();
 }
 
-bool UAvalancheMaskEditorMode::UsesToolkits() const
+bool UAvaMaskEditorMode::UsesToolkits() const
 {
 	return false;
 }
 
-void UAvalancheMaskEditorMode::CreateToolkit()
+void UAvaMaskEditorMode::CreateToolkit()
 {
 	// Toolkit = MakeShared<FAvalancheMaskEditorModeToolkit>();
 }
 
-void UAvalancheMaskEditorMode::ModeTick(float DeltaTime)
+void UAvaMaskEditorMode::ModeTick(float DeltaTime)
 {
 	Super::ModeTick(DeltaTime);
 }
 
-bool UAvalancheMaskEditorMode::IsCompatibleWith(FEditorModeID OtherModeID) const
+bool UAvaMaskEditorMode::IsCompatibleWith(FEditorModeID OtherModeID) const
 {
 	return true;
 }
 
-AActor* UAvalancheMaskEditorMode::GetActorToParentTo() const
+AActor* UAvaMaskEditorMode::GetActorToParentTo() const
 {
 	if (AActor* LastSelectedActor = WeakLastSelectedActor.Get())
 	{
@@ -262,7 +262,7 @@ AActor* UAvalancheMaskEditorMode::GetActorToParentTo() const
 	return nullptr;
 }
 
-void UAvalancheMaskEditorMode::OnSelectionChanged(const UTypedElementSelectionSet* InSelectionSet)
+void UAvaMaskEditorMode::OnSelectionChanged(const UTypedElementSelectionSet* InSelectionSet)
 {
 	if (InSelectionSet)
 	{
@@ -311,7 +311,7 @@ void UAvalancheMaskEditorMode::OnSelectionChanged(const UTypedElementSelectionSe
 	}
 }
 
-UGeometryMaskCanvas* UAvalancheMaskEditorMode::GetCanvasReferencedByActor(const AActor* InActor)
+UGeometryMaskCanvas* UAvaMaskEditorMode::GetCanvasReferencedByActor(const AActor* InActor)
 {
 	FName CanvasName = NAME_None;
 	if (const IGeometryMaskWriteInterface* WriteComponent = Cast<IGeometryMaskWriteInterface>(InActor->FindComponentByInterface<UGeometryMaskWriteInterface>()))
@@ -333,7 +333,7 @@ UGeometryMaskCanvas* UAvalancheMaskEditorMode::GetCanvasReferencedByActor(const
 	return nullptr;
 }
 
-void UAvalancheMaskEditorMode::UpdatePreviewWidget()
+void UAvaMaskEditorMode::UpdatePreviewWidget()
 {
 	if (CanvasPreviewWidget.IsValid())
 	{
@@ -342,7 +342,7 @@ void UAvalancheMaskEditorMode::UpdatePreviewWidget()
 	}
 }
 
-void UAvalancheMaskEditorMode::OnActorSpawned(AActor* InActor)
+void UAvaMaskEditorMode::OnActorSpawned(AActor* InActor)
 {
 	if (CanMaskSelected(WeakLastSelectedActor.Get()))
 	{
@@ -353,7 +353,7 @@ void UAvalancheMaskEditorMode::OnActorSpawned(AActor* InActor)
 	}
 }
 
-bool UAvalancheMaskEditorMode::AddMaskToSelected(const TArray<AActor*>& InMaskingActors)
+bool UAvaMaskEditorMode::AddMaskToSelected(const TArray<AActor*>& InMaskingActors)
 {
 	// @note: This assumes CanMaskSelected passed
 	
@@ -399,7 +399,7 @@ bool UAvalancheMaskEditorMode::AddMaskToSelected(const TArray<AActor*>& InMaskin
 	return true;
 }
 
-UAvaMask2DBaseModifier* UAvalancheMaskEditorMode::FindOrAddMaskModifier(AActor* InActor)
+UAvaMask2DBaseModifier* UAvaMaskEditorMode::FindOrAddMaskModifier(AActor* InActor)
 {
 	if (const UActorModifierCoreSubsystem* ModifierSubsystem = UActorModifierCoreSubsystem::Get())
 	{
@@ -440,7 +440,7 @@ UAvaMask2DBaseModifier* UAvalancheMaskEditorMode::FindOrAddMaskModifier(AActor*
 	return nullptr;
 }
 
-bool UAvalancheMaskEditorMode::CanMaskSelected(AActor* InSelectedActor)
+bool UAvaMaskEditorMode::CanMaskSelected(AActor* InSelectedActor)
 {
 	auto CanMaskActor = [](AActor* InActor)
 	{
diff --git a/Engine/Plugins/Experimental/Avalanche/Source/AvalancheMaskEditor/Private/AvaMaskEditorMode.h b/Engine/Plugins/Experimental/Avalanche/Source/AvalancheMaskEditor/Private/AvaMaskEditorMode.h
--- a/Engine/Plugins/Experimental/Avalanche/Source/AvalancheMaskEditor/Private/AvaMaskEditorMode.h
+++ b/Engine/Plugins/Experimental/Avalanche/Source/AvalancheMaskEditor/Private/AvaMaskEditorMode.h
@@ -1,4 +1,4 @@
-﻿// Copyright Epic Games, Inc. All Rights Reserved.
+// Copyright Epic Games, Inc. All Rights Reserved.
 
 #pragma once
 
@@ -14,15 +14,15 @@ class UAvaMask2DBaseModifier;
 class UTypedElementSelectionSet;
 
 UCLASS(MinimalAPI)
-class UAvalancheMaskEditorMode : public UEdMode
+class UAvaMaskEditorMode : public UEdMode
 {
 	GENERATED_BODY()
 	
 public:
-	static const FEditorModeID EM_AvalancheMaskEditorModeId;
+	static const FEditorModeID EM_MotionDesignMaskEditorModeId;
 
-	UAvalancheMaskEditorMode();
-	virtual ~UAvalancheMaskEditorMode() override = default;
+	UAvaMaskEditorMode();
+	virtual ~UAvaMaskEditorMode() override = default;
 
 	// Begin UEdMode
 	virtual void Enter() override;
diff --git a/Engine/Plugins/Experimental/Avalanche/Source/AvalancheMaskEditor/Private/AvaMaskEditorModule.cpp b/Engine/Plugins/Experimental/Avalanche/Source/AvalancheMaskEditor/Private/AvaMaskEditorModule.cpp
--- a/Engine/Plugins/Experimental/Avalanche/Source/AvalancheMaskEditor/Private/AvaMaskEditorModule.cpp
+++ b/Engine/Plugins/Experimental/Avalanche/Source/AvalancheMaskEditor/Private/AvaMaskEditorModule.cpp
@@ -40,7 +40,7 @@ TSharedPtr<FUICommandList> FAvalancheMaskEditorModule::GetCommandList() const
 
 void FAvalancheMaskEditorModule::ToggleEditorMode()
 {
-	GLevelEditorModeTools().ActivateMode(UAvalancheMaskEditorMode::EM_AvalancheMaskEditorModeId, true);
+	GLevelEditorModeTools().ActivateMode(UAvaMaskEditorMode::EM_MotionDesignMaskEditorModeId, true);
 }
 
 void FAvalancheMaskEditorModule::RegisterMenus()
diff --git a/Engine/Plugins/Experimental/Avalanche/Source/AvalancheMaskEditor/Private/AvaMaskEditorSVE.cpp b/Engine/Plugins/Experimental/Avalanche/Source/AvalancheMaskEditor/Private/AvaMaskEditorSVE.cpp
--- a/Engine/Plugins/Experimental/Avalanche/Source/AvalancheMaskEditor/Private/AvaMaskEditorSVE.cpp
+++ b/Engine/Plugins/Experimental/Avalanche/Source/AvalancheMaskEditor/Private/AvaMaskEditorSVE.cpp
@@ -78,7 +78,7 @@ bool FAvaMaskSceneViewExtension::IsActiveThisFrame_Internal(const FSceneViewExte
 	if (bIsActive)
 	{
 		// @todo: use events instead, this is expensive to check every tick
-		bIsActive = GLevelEditorModeTools().IsModeActive(UAvalancheMaskEditorMode::EM_AvalancheMaskEditorModeId);
+		bIsActive = GLevelEditorModeTools().IsModeActive(UAvaMaskEditorMode::EM_MotionDesignMaskEditorModeId);
 	}
 
 	return bIsActive;
diff --git a/Engine/Plugins/Experimental/Avalanche/Source/AvalancheMaskEditor/Private/AvaMaskEditorSubsystem.cpp b/Engine/Plugins/Experimental/Avalanche/Source/AvalancheMaskEditor/Private/AvaMaskEditorSubsystem.cpp
--- a/Engine/Plugins/Experimental/Avalanche/Source/AvalancheMaskEditor/Private/AvaMaskEditorSubsystem.cpp
+++ b/Engine/Plugins/Experimental/Avalanche/Source/AvalancheMaskEditor/Private/AvaMaskEditorSubsystem.cpp
@@ -64,12 +64,12 @@ void UAvaMaskEditorSubsystem::Deinitialize()
 
 void UAvaMaskEditorSubsystem::ToggleEditorMode()
 {
-	GLevelEditorModeTools().ActivateMode(UAvalancheMaskEditorMode::EM_AvalancheMaskEditorModeId, true);
+	GLevelEditorModeTools().ActivateMode(UAvaMaskEditorMode::EM_MotionDesignMaskEditorModeId, true);
 }
 
 bool UAvaMaskEditorSubsystem::IsEditorModeActive()
 {
-	return GLevelEditorModeTools().IsModeActive(UAvalancheMaskEditorMode::EM_AvalancheMaskEditorModeId);
+	return GLevelEditorModeTools().IsModeActive(UAvaMaskEditorMode::EM_MotionDesignMaskEditorModeId);
 }
 
 void UAvaMaskEditorSubsystem::ToggleShowAllMasks()
diff --git a/Engine/Plugins/Experimental/Avalanche/Source/AvalancheMaskEditor/Public/IAvaMaskEditor.h b/Engine/Plugins/Experimental/Avalanche/Source/AvalancheMaskEditor/Public/IAvaMaskEditor.h
--- a/Engine/Plugins/Experimental/Avalanche/Source/AvalancheMaskEditor/Public/IAvaMaskEditor.h
+++ b/Engine/Plugins/Experimental/Avalanche/Source/AvalancheMaskEditor/Public/IAvaMaskEditor.h
@@ -6,5 +6,5 @@
 
 namespace UE::AvalancheMaskEditor
 {
-	const static FName AvalancheMaskEditorModeName("EditMode.AvalancheMask");
+	const static FName MotionDesignMaskEditorModeName("EditMode.MotionDesignMask");
 }
diff --git a/Engine/Plugins/Experimental/Avalanche/Source/AvalancheModifiers/Public/Modifiers/AvaBaseModifier.h b/Engine/Plugins/Experimental/Avalanche/Source/AvalancheModifiers/Public/Modifiers/AvaBaseModifier.h
--- a/Engine/Plugins/Experimental/Avalanche/Source/AvalancheModifiers/Public/Modifiers/AvaBaseModifier.h
+++ b/Engine/Plugins/Experimental/Avalanche/Source/AvalancheModifiers/Public/Modifiers/AvaBaseModifier.h
@@ -5,7 +5,7 @@
 #include "Modifiers/ActorModifierCoreBase.h"
 #include "AvaBaseModifier.generated.h"
 
-/** This is the base class for all avalanche modifiers */
+/** This is the base class for all Motion Design modifiers */
 UCLASS(Abstract)
 class AVALANCHEMODIFIERS_API UAvaBaseModifier : public UActorModifierCoreBase
 {
diff --git a/Engine/Plugins/Experimental/Avalanche/Source/AvalancheOutliner/Private/AvaOutliner.cpp b/Engine/Plugins/Experimental/Avalanche/Source/AvalancheOutliner/Private/AvaOutliner.cpp
--- a/Engine/Plugins/Experimental/Avalanche/Source/AvalancheOutliner/Private/AvaOutliner.cpp
+++ b/Engine/Plugins/Experimental/Avalanche/Source/AvalancheOutliner/Private/AvaOutliner.cpp
@@ -654,7 +654,7 @@ void FAvaOutliner::RemoveItemColor(const FAvaOutlinerItemPtr& InItem)
 
 const TMap<FName, FLinearColor>& FAvaOutliner::GetColorMap() const
 {
-	return UAvalancheOutlinerSettings::Get()->GetColorMap();
+	return UAvaOutlinerSettings::Get()->GetColorMap();
 }
 
 void FAvaOutliner::NotifyItemIdChanged(const FAvaOutlinerItemId& OldId, const FAvaOutlinerItemPtr& InItem)
diff --git a/Engine/Plugins/Experimental/Avalanche/Source/AvalancheOutliner/Private/AvaOutlinerSettings.cpp b/Engine/Plugins/Experimental/Avalanche/Source/AvalancheOutliner/Private/AvaOutlinerSettings.cpp
--- a/Engine/Plugins/Experimental/Avalanche/Source/AvalancheOutliner/Private/AvaOutlinerSettings.cpp
+++ b/Engine/Plugins/Experimental/Avalanche/Source/AvalancheOutliner/Private/AvaOutlinerSettings.cpp
@@ -4,7 +4,7 @@
 #include "Styling/StyleColors.h"
 #include "UObject/PropertyIterator.h"
 
-UAvalancheOutlinerSettings::UAvalancheOutlinerSettings()
+UAvaOutlinerSettings::UAvaOutlinerSettings()
 {
 	CategoryName = TEXT("Motion Design");
 	SectionName  = TEXT("Outliner");
@@ -27,9 +27,9 @@ UAvalancheOutlinerSettings::UAvalancheOutlinerSettings()
 	ItemColorMap.Emplace(TEXT("White"),  FStyleColors::AccentWhite.GetSpecifiedColor());
 }
 
-UAvalancheOutlinerSettings* UAvalancheOutlinerSettings::Get()
+UAvaOutlinerSettings* UAvaOutlinerSettings::Get()
 {
-	UAvalancheOutlinerSettings* DefaultSettings = GetMutableDefault<UAvalancheOutlinerSettings>();
+	UAvaOutlinerSettings* DefaultSettings = GetMutableDefault<UAvaOutlinerSettings>();
 	static bool bInitialized = false;
 	if (!bInitialized)
 	{
@@ -39,12 +39,12 @@ UAvalancheOutlinerSettings* UAvalancheOutlinerSettings::Get()
 	return DefaultSettings;
 }
 
-FName UAvalancheOutlinerSettings::GetCustomItemTypeFiltersName()
+FName UAvaOutlinerSettings::GetCustomItemTypeFiltersName()
 {
-	return GET_MEMBER_NAME_CHECKED(UAvalancheOutlinerSettings, CustomItemTypeFilters);
+	return GET_MEMBER_NAME_CHECKED(UAvaOutlinerSettings, CustomItemTypeFilters);
 }
 
-bool UAvalancheOutlinerSettings::AddCustomItemTypeFilter(FName InKey, FAvaOutlinerItemTypeFilterData& InFilter)
+bool UAvaOutlinerSettings::AddCustomItemTypeFilter(FName InKey, FAvaOutlinerItemTypeFilterData& InFilter)
 {
 	if (CustomItemTypeFilters.Contains(InKey))
 	{
@@ -55,16 +55,16 @@ bool UAvalancheOutlinerSettings::AddCustomItemTypeFilter(FName InKey, FAvaOutlin
 	return true;
 }
 
-void UAvalancheOutlinerSettings::PostInitProperties()
+void UAvaOutlinerSettings::PostInitProperties()
 {
 	Super::PostInitProperties();
 	HideItemColorMapAlphaChannel();
 }
 
-void UAvalancheOutlinerSettings::PostEditChangeProperty(FPropertyChangedEvent& PropertyChangedEvent)
+void UAvaOutlinerSettings::PostEditChangeProperty(FPropertyChangedEvent& PropertyChangedEvent)
 {
 	Super::PostEditChangeProperty(PropertyChangedEvent);
-	if (PropertyChangedEvent.GetPropertyName() == GET_MEMBER_NAME_CHECKED(UAvalancheOutlinerSettings, ItemColorMap))
+	if (PropertyChangedEvent.GetPropertyName() == GET_MEMBER_NAME_CHECKED(UAvaOutlinerSettings, ItemColorMap))
 	{
 		//Force all Item Color Map entries to have Alpha value of 1
 		for (TPair<FName, FLinearColor>& Pair : ItemColorMap)
@@ -74,13 +74,13 @@ void UAvalancheOutlinerSettings::PostEditChangeProperty(FPropertyChangedEvent& P
 	}
 }
 
-void UAvalancheOutlinerSettings::HideItemColorMapAlphaChannel() const
+void UAvaOutlinerSettings::HideItemColorMapAlphaChannel() const
 {
 	const FMapProperty* ItemColorMapProperty = nullptr;
 	for (TFieldIterator<FProperty> PropertyIt(StaticClass()); PropertyIt; ++PropertyIt)
 	{
 		FProperty* const Property = CastField<FProperty>(*PropertyIt);
-		if (Property->IsA<FMapProperty>() && Property->GetFName() == GET_MEMBER_NAME_CHECKED(UAvalancheOutlinerSettings, ItemColorMap))
+		if (Property->IsA<FMapProperty>() && Property->GetFName() == GET_MEMBER_NAME_CHECKED(UAvaOutlinerSettings, ItemColorMap))
 		{
 			ItemColorMapProperty = CastField<FMapProperty>(Property);
 			break;
diff --git a/Engine/Plugins/Experimental/Avalanche/Source/AvalancheOutliner/Private/AvaOutlinerView.cpp b/Engine/Plugins/Experimental/Avalanche/Source/AvalancheOutliner/Private/AvaOutlinerView.cpp
--- a/Engine/Plugins/Experimental/Avalanche/Source/AvalancheOutliner/Private/AvaOutlinerView.cpp
+++ b/Engine/Plugins/Experimental/Avalanche/Source/AvalancheOutliner/Private/AvaOutlinerView.cpp
@@ -49,7 +49,7 @@ FAvaOutlinerView::~FAvaOutlinerView()
 
 	if (UObjectInitialized())
 	{
-		if (UAvalancheOutlinerSettings* const OutlinerSettings = UAvalancheOutlinerSettings::Get())
+		if (UAvaOutlinerSettings* const OutlinerSettings = UAvaOutlinerSettings::Get())
 		{
 			OutlinerSettings->OnSettingChanged().RemoveAll(this);
 		}
@@ -77,7 +77,7 @@ void FAvaOutlinerView::Init(const TSharedRef<FAvaOutliner>& InOutliner, bool bCr
 	CreateColumns();
 
 	InOutliner->GetProvider().ExtendOutlinerItemFilters(ItemFilters);
-	UAvalancheOutlinerSettings::Get()->OnSettingChanged().AddSP(this, &FAvaOutlinerView::OnOutlinerSettingsChanged);
+	UAvaOutlinerSettings::Get()->OnSettingChanged().AddSP(this, &FAvaOutlinerView::OnOutlinerSettingsChanged);
 	UpdateCustomFilters();
 
 	TextFilter->OnChanged().AddSP(this, &FAvaOutlinerView::OnFilterChanged);
@@ -158,7 +158,7 @@ void FAvaOutlinerView::PostLoad()
 void FAvaOutlinerView::OnOutlinerSettingsChanged(UObject* InObject, FPropertyChangedEvent& InPropertyChangedEvent)
 {
 	const FName PropertyName = InPropertyChangedEvent.GetMemberPropertyName();
-	if (PropertyName == UAvalancheOutlinerSettings::GetCustomItemTypeFiltersName())
+	if (PropertyName == UAvaOutlinerSettings::GetCustomItemTypeFiltersName())
 	{
 		UpdateCustomFilters();
 	}
@@ -167,7 +167,7 @@ void FAvaOutlinerView::OnOutlinerSettingsChanged(UObject* InObject, FPropertyCha
 
 void FAvaOutlinerView::UpdateCustomFilters()
 {
-	UAvalancheOutlinerSettings* const OutlinerSettings = UAvalancheOutlinerSettings::Get();
+	UAvaOutlinerSettings* const OutlinerSettings = UAvaOutlinerSettings::Get();
 	check(OutlinerSettings);
 		
 	const TMap<FName, FAvaOutlinerItemTypeFilterData>& CustomFilterMap = OutlinerSettings->GetCustomItemTypeFilters();
diff --git a/Engine/Plugins/Experimental/Avalanche/Source/AvalancheOutliner/Private/Columns/Slate/SAvaOutlinerLock.cpp b/Engine/Plugins/Experimental/Avalanche/Source/AvalancheOutliner/Private/Columns/Slate/SAvaOutlinerLock.cpp
--- a/Engine/Plugins/Experimental/Avalanche/Source/AvalancheOutliner/Private/Columns/Slate/SAvaOutlinerLock.cpp
+++ b/Engine/Plugins/Experimental/Avalanche/Source/AvalancheOutliner/Private/Columns/Slate/SAvaOutlinerLock.cpp
@@ -83,7 +83,7 @@ FSlateColor SAvaOutlinerLock::GetForegroundColor() const
 	const bool bIsItemHovered = RowWeak.IsValid()
 		&& RowWeak.Pin()->IsHovered();
 
-	const bool bAlwaysShowLock = UAvalancheOutlinerSettings::Get()->ShouldAlwaysShowLockState();
+	const bool bAlwaysShowLock = UAvaOutlinerSettings::Get()->ShouldAlwaysShowLockState();
 
 	// we can hide the brush if Settings for Always Showing State is OFF
 	// and Item is not locked while also not being selected nor hovered
diff --git a/Engine/Plugins/Experimental/Avalanche/Source/AvalancheOutliner/Private/Columns/Slate/SAvaOutlinerVisibility.cpp b/Engine/Plugins/Experimental/Avalanche/Source/AvalancheOutliner/Private/Columns/Slate/SAvaOutlinerVisibility.cpp
--- a/Engine/Plugins/Experimental/Avalanche/Source/AvalancheOutliner/Private/Columns/Slate/SAvaOutlinerVisibility.cpp
+++ b/Engine/Plugins/Experimental/Avalanche/Source/AvalancheOutliner/Private/Columns/Slate/SAvaOutlinerVisibility.cpp
@@ -201,7 +201,7 @@ FSlateColor SAvaOutlinerVisibility::GetForegroundColor() const
 	const bool bIsItemHovered = RowWeak.IsValid()
 		&& RowWeak.Pin()->IsHovered();
 
-	const bool bAlwaysShowVisibility = UAvalancheOutlinerSettings::Get()->ShouldAlwaysShowVisibilityState();
+	const bool bAlwaysShowVisibility = UAvaOutlinerSettings::Get()->ShouldAlwaysShowVisibilityState();
 
 	// we can hide the brush if Settings for Always Showing State is OFF
 	// and Item is Visible while also not being selected nor hovered
diff --git a/Engine/Plugins/Experimental/Avalanche/Source/AvalancheOutliner/Private/Data/AvaOutlinerSaveState.cpp b/Engine/Plugins/Experimental/Avalanche/Source/AvalancheOutliner/Private/Data/AvaOutlinerSaveState.cpp
--- a/Engine/Plugins/Experimental/Avalanche/Source/AvalancheOutliner/Private/Data/AvaOutlinerSaveState.cpp
+++ b/Engine/Plugins/Experimental/Avalanche/Source/AvalancheOutliner/Private/Data/AvaOutlinerSaveState.cpp
@@ -17,7 +17,7 @@
 
 FAvaOutlinerViewSaveState::FAvaOutlinerViewSaveState()
 {
-	const UAvalancheOutlinerSettings* const OutlinerSettings = UAvalancheOutlinerSettings::Get();
+	const UAvaOutlinerSettings* const OutlinerSettings = UAvaOutlinerSettings::Get();
 	check(OutlinerSettings);
 	
 	bUseMutedHierarchy     = OutlinerSettings->ShouldUseMutedHierarchy();
diff --git a/Engine/Plugins/Experimental/Avalanche/Source/AvalancheOutliner/Private/Filters/Suggestions/AvaFilterColorSuggestionFactory.cpp b/Engine/Plugins/Experimental/Avalanche/Source/AvalancheOutliner/Private/Filters/Suggestions/AvaFilterColorSuggestionFactory.cpp
--- a/Engine/Plugins/Experimental/Avalanche/Source/AvalancheOutliner/Private/Filters/Suggestions/AvaFilterColorSuggestionFactory.cpp
+++ b/Engine/Plugins/Experimental/Avalanche/Source/AvalancheOutliner/Private/Filters/Suggestions/AvaFilterColorSuggestionFactory.cpp
@@ -13,7 +13,7 @@ void FAvaFilterColorSuggestionFactory::AddSuggestion(const TSharedRef<FAvaFilter
 	TSet<FString> ColorCache;
 	const FText ColorCategoryLabel = LOCTEXT("ColorCategoryLabel", "Ava-Outliner-Colors");
 
-	if (const UAvalancheOutlinerSettings* Settings = UAvalancheOutlinerSettings::Get())
+	if (const UAvaOutlinerSettings* Settings = UAvaOutlinerSettings::Get())
 	{
 		for (const TTuple<FName, FLinearColor> ColorPair : Settings->GetColorMap())
 		{
diff --git a/Engine/Plugins/Experimental/Avalanche/Source/AvalancheOutliner/Private/Slate/SAvaOutliner.cpp b/Engine/Plugins/Experimental/Avalanche/Source/AvalancheOutliner/Private/Slate/SAvaOutliner.cpp
--- a/Engine/Plugins/Experimental/Avalanche/Source/AvalancheOutliner/Private/Slate/SAvaOutliner.cpp
+++ b/Engine/Plugins/Experimental/Avalanche/Source/AvalancheOutliner/Private/Slate/SAvaOutliner.cpp
@@ -510,7 +510,7 @@ void SAvaOutliner::OnSaveSearchButtonClicked(const FText& InText)
 
 void SAvaOutliner::OnCreateCustomTextFilter(const FCustomTextFilterData& InFilterData, bool ApplyFilter) const
 {
-	if (UAvalancheOutlinerSettings* OutlinerSettings = UAvalancheOutlinerSettings::Get())
+	if (UAvaOutlinerSettings* OutlinerSettings = UAvaOutlinerSettings::Get())
 	{
 		FAvaOutlinerItemTypeFilterData NewFilter = FAvaOutlinerItemTypeFilterData();
 		NewFilter.SetFilterText(InFilterData.FilterString);
diff --git a/Engine/Plugins/Experimental/Avalanche/Source/AvalancheOutliner/Public/AvaOutlinerSettings.h b/Engine/Plugins/Experimental/Avalanche/Source/AvalancheOutliner/Public/AvaOutlinerSettings.h
--- a/Engine/Plugins/Experimental/Avalanche/Source/AvalancheOutliner/Public/AvaOutlinerSettings.h
+++ b/Engine/Plugins/Experimental/Avalanche/Source/AvalancheOutliner/Public/AvaOutlinerSettings.h
@@ -7,14 +7,14 @@
 #include "AvaOutlinerSettings.generated.h"
 
 UCLASS(config=EditorPerProjectUserSettings, meta=(DisplayName="Outliner"))
-class AVALANCHEOUTLINER_API UAvalancheOutlinerSettings : public UDeveloperSettings
+class AVALANCHEOUTLINER_API UAvaOutlinerSettings : public UDeveloperSettings
 {
 	GENERATED_BODY()
 
 public:
-	UAvalancheOutlinerSettings();
+	UAvaOutlinerSettings();
 
-	static UAvalancheOutlinerSettings* Get();
+	static UAvaOutlinerSettings* Get();
 
 	const TMap<FName, FLinearColor>& GetColorMap() const { return ItemColorMap; }
 
diff --git a/Engine/Plugins/Experimental/Avalanche/Source/AvalancheOutliner/Public/AvaOutlinerView.h b/Engine/Plugins/Experimental/Avalanche/Source/AvalancheOutliner/Public/AvaOutlinerView.h
--- a/Engine/Plugins/Experimental/Avalanche/Source/AvalancheOutliner/Public/AvaOutlinerView.h
+++ b/Engine/Plugins/Experimental/Avalanche/Source/AvalancheOutliner/Public/AvaOutlinerView.h
@@ -62,11 +62,11 @@ public:
 
 	void PostLoad();
 	
-	/** Called when the UAvalancheOutlinerSettings has a property change */
+	/** Called when the UAvaOutlinerSettings has a property change */
 	void OnOutlinerSettingsChanged(UObject* InObject, FPropertyChangedEvent& InPropertyChangedEvent);
 
 	/**
-	 * Retrieves the latest list of Custom Filters specified in UAvalancheOutlinerSettings
+	 * Retrieves the latest list of Custom Filters specified in UAvaOutlinerSettings
 	 * and recreates the filter instances for these
 	 */
 	void UpdateCustomFilters();
@@ -429,7 +429,7 @@ private:
 	/** Native-constructed Item Filters in the Outliner */
 	TArray<TSharedPtr<IAvaOutlinerItemFilter>> ItemFilters;
 
-	/** Custom Item Filters constructed from UAvalancheOutlinerSettings */
+	/** Custom Item Filters constructed from UAvaOutlinerSettings */
 	TArray<TSharedPtr<IAvaOutlinerItemFilter>> CustomItemFilters;
 
 	/** Delegate called when the Custom Item Filters have been modified, as these can be instanced after */
diff --git a/Engine/Plugins/Experimental/Avalanche/Source/AvalancheShapes/Public/AvaShapeActor.h b/Engine/Plugins/Experimental/Avalanche/Source/AvalancheShapes/Public/AvaShapeActor.h
--- a/Engine/Plugins/Experimental/Avalanche/Source/AvalancheShapes/Public/AvaShapeActor.h
+++ b/Engine/Plugins/Experimental/Avalanche/Source/AvalancheShapes/Public/AvaShapeActor.h
@@ -3,7 +3,7 @@
 #pragma once
 
 #include "GameFramework/Actor.h"
-#include "IAvalancheInteractiveToolsModeDetailsObjectProvider.h"
+#include "IAvaInteractiveToolsModeDetailsObjectProvider.h"
 #include "Viewport/Interaction/IAvaSnapPointGenerator.h"
 #include "AvaShapeActor.generated.h"
 
@@ -14,7 +14,7 @@ UCLASS()
 class AVALANCHESHAPES_API AAvaShapeActor
 	: public AActor
 	, public IAvaSnapPointGenerator
-	, public IAvalancheInteractiveToolsModeDetailsObjectProvider
+	, public IAvaInteractiveToolsModeDetailsObjectProvider
 {
 	GENERATED_BODY()
 
@@ -41,9 +41,9 @@ public:
 	virtual TArray<FAvaSnapPoint> GetLocalSnapPoints() const override;
 	//~ End IAvaSnapPointGenerator
 
-	//~ Begin IAvalancheInteractiveToolsModeDetailsObjectProvider
+	//~ Begin IAvaInteractiveToolsModeDetailsObjectProvider
 	virtual UObject* GetModeDetailsObject_Implementation() const override;
-	//~ End IAvalancheInteractiveToolsModeDetailsObjectProvider
+	//~ End IAvaInteractiveToolsModeDetailsObjectProvider
 
 protected:
 	static const FName ShapeComponentName;
diff --git a/Engine/Plugins/Experimental/Avalanche/Source/AvalancheShapes/Public/DynamicMeshes/AvaShapeDynMeshBase.h b/Engine/Plugins/Experimental/Avalanche/Source/AvalancheShapes/Public/DynamicMeshes/AvaShapeDynMeshBase.h
--- a/Engine/Plugins/Experimental/Avalanche/Source/AvalancheShapes/Public/DynamicMeshes/AvaShapeDynMeshBase.h
+++ b/Engine/Plugins/Experimental/Avalanche/Source/AvalancheShapes/Public/DynamicMeshes/AvaShapeDynMeshBase.h
@@ -8,7 +8,7 @@
 #include "AvaShapeUVParameters.h"
 #include "Components/ActorComponent.h"
 #include "DynamicMesh/DynamicMesh3.h"
-#include "IAvalancheInteractiveToolsModeDetailsObject.h"
+#include "IAvaInteractiveToolsModeDetailsObject.h"
 #include "Viewport/Interaction/AvaSnapPoint.h"
 #include "Viewport/Interaction/IAvaGizmoObject.h"
 
@@ -25,7 +25,7 @@ DECLARE_MULTICAST_DELEGATE_TwoParams(FMaskVisibilityDelegate, const UWorld* /**
 UCLASS(ClassGroup="Shape", Abstract, BlueprintType, CustomConstructor, EditInlineNew, DefaultToInstanced)
 class AVALANCHESHAPES_API UAvaShapeDynamicMeshBase
 	: public UActorComponent
-	, public IAvalancheInteractiveToolsModeDetailsObject
+	, public IAvaInteractiveToolsModeDetailsObject
 	, public IAvaGizmoObjectInterface
 {
 	GENERATED_BODY()
diff --git a/Engine/Plugins/Experimental/Avalanche/Source/AvalancheTest/Private/AvaScreenshotTestActor.cpp b/Engine/Plugins/Experimental/Avalanche/Source/AvalancheTest/Private/AvaScreenshotTestActor.cpp
--- a/Engine/Plugins/Experimental/Avalanche/Source/AvalancheTest/Private/AvaScreenshotTestActor.cpp
+++ b/Engine/Plugins/Experimental/Avalanche/Source/AvalancheTest/Private/AvaScreenshotTestActor.cpp
@@ -60,7 +60,7 @@ void AAvaScreenshotTestActor::SetupTest()
 	// Create a world (.umap) asset for the Ava BP, overwriting any existing
 	if (bHasValidBP)
 	{
-		UAvaTestBlueprintFunctionLibrary::ExportAvaBlueprintsToWorld({ ResolvedBlueprintToTest });
+		UAvaTestBlueprintFunctionLibrary::ExportMotionDesignBlueprintsToWorld({ ResolvedBlueprintToTest });
 	}
 
 	ReceiveSetupTest();
diff --git a/Engine/Plugins/Experimental/Avalanche/Source/AvalancheTest/Private/AvaTestBlueprintFunctionLibrary.cpp b/Engine/Plugins/Experimental/Avalanche/Source/AvalancheTest/Private/AvaTestBlueprintFunctionLibrary.cpp
--- a/Engine/Plugins/Experimental/Avalanche/Source/AvalancheTest/Private/AvaTestBlueprintFunctionLibrary.cpp
+++ b/Engine/Plugins/Experimental/Avalanche/Source/AvalancheTest/Private/AvaTestBlueprintFunctionLibrary.cpp
@@ -1,4 +1,4 @@
-﻿// Copyright Epic Games, Inc. All Rights Reserved.
+// Copyright Epic Games, Inc. All Rights Reserved.
 
 #include "AvaTestBlueprintFunctionLibrary.h"
 
@@ -8,7 +8,7 @@
 #include "AvaEditorFunctionLibrary.h"
 #endif
 
-void UAvaTestBlueprintFunctionLibrary::ExportAvaBlueprintsToWorld(const TArray<TSoftObjectPtr<UAvalancheBlueprint>>& InBlueprints)
+void UAvaTestBlueprintFunctionLibrary::ExportMotionDesignBlueprintsToWorld(const TArray<TSoftObjectPtr<UAvalancheBlueprint>>& InBlueprints)
 {
 #if WITH_EDITOR
 	TArray<TWeakObjectPtr<UAvalancheBlueprint>> WeakBlueprints;
@@ -26,6 +26,6 @@ void UAvaTestBlueprintFunctionLibrary::ExportAvaBlueprintsToWorld(const TArray<T
 	
 	FAvaEditorFunctionLibrary::ExportAvaBlueprintsToWorld(WeakBlueprints);
 #else
-	UE_LOG(LogAvaTest, Error, TEXT("ExportAvaBlueprintsToWorld is editor-only"));
+	UE_LOG(LogAvaTest, Error, TEXT("ExportMotionDesignBlueprintsToWorld is editor-only"));
 #endif
 }
diff --git a/Engine/Plugins/Experimental/Avalanche/Source/AvalancheTest/Public/AvaTestBlueprintFunctionLibrary.h b/Engine/Plugins/Experimental/Avalanche/Source/AvalancheTest/Public/AvaTestBlueprintFunctionLibrary.h
--- a/Engine/Plugins/Experimental/Avalanche/Source/AvalancheTest/Public/AvaTestBlueprintFunctionLibrary.h
+++ b/Engine/Plugins/Experimental/Avalanche/Source/AvalancheTest/Public/AvaTestBlueprintFunctionLibrary.h
@@ -1,4 +1,4 @@
-﻿// Copyright Epic Games, Inc. All Rights Reserved.
+// Copyright Epic Games, Inc. All Rights Reserved.
 
 #pragma once
 
@@ -18,10 +18,10 @@ class AVALANCHETEST_API UAvaTestBlueprintFunctionLibrary : public UBlueprintFunc
 
 public:
 	/**
-	 * Creates World Assets based on the Avalanche Blueprints provided
-	 * Wraps FAvaEditorFunctionLibrary::ExportAvaBlueprintsToWorld
-	 * @param InBlueprints the Avalanche Blueprints to export
+	 * Creates World Assets based on the Motion Design Blueprints provided
+	 * Wraps FAvaEditorFunctionLibrary::ExportMotionDesignBlueprintsToWorld
+	 * @param InBlueprints the Motion Design Blueprints to export
 	 */
 	UFUNCTION(BlueprintCallable, Category = "Motion Design")
-	static void ExportAvaBlueprintsToWorld(const TArray<TSoftObjectPtr<UAvalancheBlueprint>>& InBlueprints);
+	static void ExportMotionDesignBlueprintsToWorld(const TArray<TSoftObjectPtr<UAvalancheBlueprint>>& InBlueprints);
 };
diff --git a/Engine/Plugins/Experimental/Avalanche/Source/AvalancheText/Private/AvaText3DComponent.cpp b/Engine/Plugins/Experimental/Avalanche/Source/AvalancheText/Private/AvaText3DComponent.cpp
--- a/Engine/Plugins/Experimental/Avalanche/Source/AvalancheText/Private/AvaText3DComponent.cpp
+++ b/Engine/Plugins/Experimental/Avalanche/Source/AvalancheText/Private/AvaText3DComponent.cpp
@@ -1,7 +1,6 @@
 // Copyright Epic Games, Inc. All Rights Reserved.
 
 #include "AvaText3DComponent.h"
-
 #include "AvaActorUtils.h"
 #include "AvaLog.h"
 #include "Engine/Texture2D.h"
@@ -60,7 +59,7 @@ UAvaText3DComponent::UAvaText3DComponent()
 	MaxWidth = 100.0f;
 	MaxHeight = 100.0f;
 
-	AvalancheFont = {};
+	MotionDesignFont = {};
 
 	Tiling = FVector2D(1.0, 1.0);
 	
@@ -528,7 +527,7 @@ void UAvaText3DComponent::SetupTextGeometry()
 	
 		SetText(Text);
 
-		if (UFont* UnderlyingFont = AvalancheFont.GetFont())
+		if (UFont* UnderlyingFont = MotionDesignFont.GetFont())
 		{
 			SetFont(UnderlyingFont);	
 		}
@@ -643,13 +642,13 @@ void UAvaText3DComponent::SetEnforceUpperCase(bool bInEnforceUpperCase)
 	RefreshFormatting();
 }
 
-void UAvaText3DComponent::SetAvalancheFont(const FAvaFont& InAvalancheFont)
+void UAvaText3DComponent::SetMotionDesignFont(const FAvaFont& InFont)
 {
-	if (AvalancheFont == InAvalancheFont)
+	if (MotionDesignFont == InFont)
 	{
 		return;
 	}
-	AvalancheFont = InAvalancheFont;
+	MotionDesignFont = InFont;
 	RefreshGeometry();
 }
 
@@ -937,8 +936,8 @@ void UAvaText3DComponent::RegisterOnPropertyChangeFunctions()
 
 		PropertyChangeFunctions.Add(GET_MEMBER_NAME_CHECKED(UAvaText3DComponent, Extrude), &UAvaText3DComponent::RefreshGeometry);
 		PropertyChangeFunctions.Add(GET_MEMBER_NAME_CHECKED(UAvaText3DComponent, Bevel), &UAvaText3DComponent::RefreshGeometry);
-		PropertyChangeFunctions.Add(GET_MEMBER_NAME_CHECKED(UAvaText3DComponent, AvalancheFont), &UAvaText3DComponent::RefreshGeometry);
-		PropertyChangeFunctions.Add(GET_MEMBER_NAME_CHECKED(UAvaText3DComponent, Font), &UAvaText3DComponent::RefreshAvalancheFont);
+		PropertyChangeFunctions.Add(GET_MEMBER_NAME_CHECKED(UAvaText3DComponent, MotionDesignFont), &UAvaText3DComponent::RefreshGeometry);
+		PropertyChangeFunctions.Add(GET_MEMBER_NAME_CHECKED(UAvaText3DComponent, Font), &UAvaText3DComponent::RefreshMotionDesignFont);
 		PropertyChangeFunctions.Add(GET_MEMBER_NAME_CHECKED(UAvaText3DComponent, Text), &UAvaText3DComponent::RefreshGeometry);
 		PropertyChangeFunctions.Add(GET_MEMBER_NAME_CHECKED(UAvaText3DComponent, bEnforceUpperCase), &UAvaText3DComponent::RefreshFormatting);
 
@@ -977,13 +976,13 @@ void UAvaText3DComponent::RefreshAlignment()
 	RefreshLayout();
 }
 
-void UAvaText3DComponent::RefreshAvalancheFont()
+void UAvaText3DComponent::RefreshMotionDesignFont()
 {
-	if (IsValid(Font) && Font != AvalancheFont.GetFont())
+	if (IsValid(Font) && Font != MotionDesignFont.GetFont())
 	{
-		// todo: this leads to proper update of the AvalancheFont, and of the text component geometry
-		// nonetheless, custom FAvalancheFont widget is not updating properly yet
-		AvalancheFont.InitFromFont(Font);
+		// todo: this leads to proper update of the MotionDesignFont, and of the text component geometry
+		// nonetheless, custom FAvaFont widget is not updating properly yet
+		MotionDesignFont.InitFromFont(Font);
 	}
 }
 
@@ -1580,11 +1579,11 @@ void UAvaText3DComponent::RetrieveCurrentColoringStyle()
 void UAvaText3DComponent::RetrieveCurrentTextAndFont()
 {
 	// if default had to be loaded (e.g. font asset missing), then we regenerate geometry
-	if (AvalancheFont.IsDefaultFont())
+	if (MotionDesignFont.IsDefaultFont())
 	{
-		if (IsValid(Font) && Font != AvalancheFont.GetFont())
+		if (IsValid(Font) && Font != MotionDesignFont.GetFont())
 		{
-			AvalancheFont.InitFromFont(Font);
+			MotionDesignFont.InitFromFont(Font);
 		}
 		
 		MarkForGeometryRefresh();
diff --git a/Engine/Plugins/Experimental/Avalanche/Source/AvalancheText/Private/Font/AvaFont.cpp b/Engine/Plugins/Experimental/Avalanche/Source/AvalancheText/Private/Font/AvaFont.cpp
--- a/Engine/Plugins/Experimental/Avalanche/Source/AvalancheText/Private/Font/AvaFont.cpp
+++ b/Engine/Plugins/Experimental/Avalanche/Source/AvalancheText/Private/Font/AvaFont.cpp
@@ -59,27 +59,27 @@ UFont* FAvaFont::GetDefaultFont()
 	return DefaultFontObjects.AvaDefaultFont.Get();
 }
 
-FString FAvaFont::GenerateAvalancheFontFormattedString(const FString& InFontName, const FString& InFontObjectPathName)
+FString FAvaFont::GenerateFontFormattedString(const FString& InFontName, const FString& InFontObjectPathName)
 {
 	// e.g. (Property1=Value1,Property2=Value2,Property3=Value3,...)
 
 	FString FormattedString = TEXT("(");
 	FormattedString += TEXT("CurrentFont=None");
-	FormattedString += TEXT(",AvalancheFontObject=/Script/Avalanche.AvalancheFontObject'") + InFontObjectPathName + TEXT("'");
+	FormattedString += TEXT(",MotionDesignFontObject=/Script/AvalancheText.AvaFontObject'") + InFontObjectPathName + TEXT("'");
 	FormattedString += TEXT(",FontName=\"") + InFontName  + TEXT("\"");
 	FormattedString += TEXT(")");
 
 	return FormattedString;
 }
 
-bool FAvaFont::GenerateAvalancheFontFormattedString(const UAvaFontObject* InFontObject, FString& OutFormattedString)
+bool FAvaFont::GenerateFontFormattedString(const UAvaFontObject* InFontObject, FString& OutFormattedString)
 {
 	if (!InFontObject)
 	{
 		return false;
 	}
 
-	OutFormattedString = GenerateAvalancheFontFormattedString(InFontObject->GetFontName(), InFontObject->GetPathName());
+	OutFormattedString = GenerateFontFormattedString(InFontObject->GetFontName(), InFontObject->GetPathName());
 	return true;
 }
 
@@ -128,9 +128,9 @@ UFont* FAvaFont::GetFont()
 	{
 		EnsureUsingCurrentVersion();
 
-		if (AvalancheFontObject)
+		if (MotionDesignFontObject)
 		{
-			return AvalancheFontObject->GetFont();
+			return MotionDesignFontObject->GetFont();
 		}
 	}
 
@@ -140,9 +140,9 @@ UFont* FAvaFont::GetFont()
 
 EAvaFontSource FAvaFont::GetFontSource() const
 {
-	if (AvalancheFontObject)
+	if (MotionDesignFontObject)
 	{
-		return AvalancheFontObject->GetSource();
+		return MotionDesignFontObject->GetSource();
 	}
 
 	return EAvaFontSource::Invalid;
@@ -155,9 +155,9 @@ FName FAvaFont::GetFontName() const
 
 FString FAvaFont::GetFontNameAsString() const
 {
-	if (AvalancheFontObject)
+	if (MotionDesignFontObject)
 	{
-		return AvalancheFontObject->GetFontName();
+		return MotionDesignFontObject->GetFontName();
 	}
 
 	return FontName;
@@ -175,9 +175,9 @@ bool FAvaFont::IsFavorite() const
 
 bool FAvaFont::IsDefaultFont() const
 {
-	if (AvalancheFontObject && AvalancheFontObject->GetFont())
+	if (MotionDesignFontObject && MotionDesignFontObject->GetFont())
 	{
-		return AvalancheFontObject->GetFont() == GetDefaultFont();
+		return MotionDesignFontObject->GetFont() == GetDefaultFont();
 	}
 
 	return true;
@@ -190,9 +190,9 @@ bool FAvaFont::IsFallbackFont() const
 
 bool FAvaFont::IsMonospaced() const
 {
-	if (AvalancheFontObject)
+	if (MotionDesignFontObject)
 	{
-		return AvalancheFontObject->IsMonospaced();
+		return MotionDesignFontObject->IsMonospaced();
 	}
 
 	return false;
@@ -200,9 +200,9 @@ bool FAvaFont::IsMonospaced() const
 
 bool FAvaFont::IsBold() const
 {
-	if (AvalancheFontObject)
+	if (MotionDesignFontObject)
 	{
-		return AvalancheFontObject->IsBold();
+		return MotionDesignFontObject->IsBold();
 	}
 
 	return false;
@@ -210,9 +210,9 @@ bool FAvaFont::IsBold() const
 
 bool FAvaFont::IsItalic() const
 {
-	if (AvalancheFontObject)
+	if (MotionDesignFontObject)
 	{
-		return AvalancheFontObject->IsItalic();
+		return MotionDesignFontObject->IsItalic();
 	}
 
 	return false;
@@ -227,8 +227,8 @@ void FAvaFont::EnsureUsingCurrentVersion()
 {
 	if (CurrentFont_DEPRECATED)
 	{
-		AvalancheFontObject = NewObject<UAvaFontObject>();
-		AvalancheFontObject->InitProjectFont(CurrentFont_DEPRECATED, GetFontName().ToString());
+		MotionDesignFontObject = NewObject<UAvaFontObject>();
+		MotionDesignFontObject->InitProjectFont(CurrentFont_DEPRECATED, GetFontName().ToString());
 		CurrentFont_DEPRECATED = nullptr;
 
 		RefreshName();
@@ -243,9 +243,9 @@ bool FAvaFont::HasValidFont() const
 	{
 		CompositeFont = CurrentFont_DEPRECATED->GetCompositeFont();
 	}
-	else if (AvalancheFontObject && AvalancheFontObject->GetFont())
+	else if (MotionDesignFontObject && MotionDesignFontObject->GetFont())
 	{
-		CompositeFont = AvalancheFontObject->GetFont()->GetCompositeFont();
+		CompositeFont = MotionDesignFontObject->GetFont()->GetCompositeFont();
 	}
 
 	if (CompositeFont)
@@ -265,7 +265,7 @@ bool FAvaFont::HasValidFont() const
 void FAvaFont::SetFontObject(UAvaFontObject* InFontObject)
 {
 	CurrentFont_DEPRECATED = nullptr;
-	AvalancheFontObject = InFontObject;
+	MotionDesignFontObject = InFontObject;
 
 	RefreshName();
 	RefreshAssetState();
@@ -278,8 +278,8 @@ void FAvaFont::InitFromFont(UFont* InFont)
 		FString Name;
 		UE::Avalanche::FontUtilities::Public::GetFontName(InFont, Name);
 
-		AvalancheFontObject = NewObject<UAvaFontObject>();
-		AvalancheFontObject->InitProjectFont(InFont, Name);
+		MotionDesignFontObject = NewObject<UAvaFontObject>();
+		MotionDesignFontObject->InitProjectFont(InFont, Name);
 
 		RefreshName();
 	}
@@ -291,17 +291,17 @@ void FAvaFont::PostSerialize(const FArchive& Ar)
 	{
 		RefreshName();
 
-		if ((AvalancheFontObject && AvalancheFontObject->GetFont()) || CurrentFont_DEPRECATED)
+		if ((MotionDesignFontObject && MotionDesignFontObject->GetFont()) || CurrentFont_DEPRECATED)
 		{
 			FontAssetState = EFontAssetState::SelectedFont;
 		}
 		else
 		{
 			/*
-			 * AvalancheFontObject might be invalid due to a previously existing issue.
-			 * That would cause the Editor only version of AvalancheFontObject to be serialized, instead of the proper one.
-			 * For assets using the font saved like that, the AvalancheFontObject had the wrong Outer, which worked fine in Editor, but not for Game/Runtime.
-			 * The following code tries to find a font based on the font name, and then to create a new AvalancheFontObject, with the proper Outer
+			 * MotionDesignFontObject might be invalid due to a previously existing issue.
+			 * That would cause the Editor only version of MotionDesignFontObject to be serialized, instead of the proper one.
+			 * For assets using the font saved like that, the MotionDesignFontObject had the wrong Outer, which worked fine in Editor, but not for Game/Runtime.
+			 * The following code tries to find a font based on the font name, and then to create a new MotionDesignFontObject, with the proper Outer
 			 */
 
 			bool bFontRecoverySuccess = false;
@@ -312,17 +312,17 @@ void FAvaFont::PostSerialize(const FArchive& Ar)
 			// this should just contain the FAvaFont
 			for (const FProperty* Property : OutProperties)
 			{
-				// let's get the Outer we need to create the AvalancheFontObject
+				// let's get the Outer we need to create the MotionDesignFontObject
 				if (UObject* const Outer = Property->GetOwner<UObject>())
 				{
 					// try to get a font just using the name (this might fail!)
 					if (UFont* const Font = GetFontByName(FontName))
 					{
-						AvalancheFontObject = nullptr;
+						MotionDesignFontObject = nullptr;
 
 						// we create and assign a new UAvaFontObject with the proper Outer
-						AvalancheFontObject = NewObject<UAvaFontObject>(Outer, GetFontName(), RF_Public | RF_Standalone);
-						AvalancheFontObject->InitProjectFont(Font, FontName);
+						MotionDesignFontObject = NewObject<UAvaFontObject>(Outer, GetFontName(), RF_Public | RF_Standalone);
+						MotionDesignFontObject->InitProjectFont(Font, FontName);
 						bFontRecoverySuccess = true;
 						break;
 					}
@@ -359,9 +359,9 @@ UAvaFontObject* FAvaFont::GetDefaultAvaFontObject()
 
 void FAvaFont::RefreshName()
 {
-	if (AvalancheFontObject)
+	if (MotionDesignFontObject)
 	{
-		const FString& FontObjectName = AvalancheFontObject->GetFontName();
+		const FString& FontObjectName = MotionDesignFontObject->GetFontName();
 
 		if (FontName != FontObjectName)
 		{
diff --git a/Engine/Plugins/Experimental/Avalanche/Source/AvalancheText/Public/AvaText3DComponent.h b/Engine/Plugins/Experimental/Avalanche/Source/AvalancheText/Public/AvaText3DComponent.h
--- a/Engine/Plugins/Experimental/Avalanche/Source/AvalancheText/Public/AvaText3DComponent.h
+++ b/Engine/Plugins/Experimental/Avalanche/Source/AvalancheText/Public/AvaText3DComponent.h
@@ -39,8 +39,8 @@ public:
 	bool IsEnforcingUpperCase() const { return bEnforceUpperCase; }
 	
 	UFUNCTION(BlueprintSetter)
-	void SetAvalancheFont(const FAvaFont& InAvalancheFont);
-	const FAvaFont& GetAvalancheFont() const { return AvalancheFont; }
+	void SetMotionDesignFont(const FAvaFont& InFont);
+	const FAvaFont& GetMotionDesignFont() const { return MotionDesignFont; }
 
 	UFUNCTION(BlueprintSetter)
 	void SetAlignment(FAvaTextAlignment InAlignment);
@@ -120,7 +120,7 @@ public:
 	void RefreshBevelColor() const;
 
 	void RefreshAlignment();
-	void RefreshAvalancheFont();
+	void RefreshMotionDesignFont();
 	void RefreshHorizontalAlignment();
 	void RefreshVerticalAlignment();
 	void RefreshUnlit();
@@ -164,8 +164,8 @@ private:
 	UPROPERTY(EditAnywhere, Getter="IsEnforcingUpperCase", Setter="SetEnforceUpperCase", Category="Text", meta=(DisplayAfter="Text", AllowPrivateAccess="true"))
 	bool bEnforceUpperCase = false;
 	
-	UPROPERTY(EditAnywhere, Getter="GetAvalancheFont", Setter="SetAvalancheFont", Category="Text", meta=(DisplayAfter="Text", AllowPrivateAccess="true"))
-	FAvaFont AvalancheFont;
+	UPROPERTY(EditAnywhere, Getter="GetMotionDesignFont", Setter="SetMotionDesignFont", Category="Text", meta=(DisplayAfter="Text", AllowPrivateAccess="true"))
+	FAvaFont MotionDesignFont;
 
 	UPROPERTY(EditAnywhere, Getter="GetAlignment", Setter="SetAlignment", Category="Layout", meta=(DisplayAfter="VerticalAlignment", AllowPrivateAccess="true"))
 	FAvaTextAlignment Alignment;
diff --git a/Engine/Plugins/Experimental/Avalanche/Source/AvalancheText/Public/Font/AvaFont.h b/Engine/Plugins/Experimental/Avalanche/Source/AvalancheText/Public/Font/AvaFont.h
--- a/Engine/Plugins/Experimental/Avalanche/Source/AvalancheText/Public/Font/AvaFont.h
+++ b/Engine/Plugins/Experimental/Avalanche/Source/AvalancheText/Public/Font/AvaFont.h
@@ -8,7 +8,7 @@
 
 #include "AvaFont.generated.h"
 
-USTRUCT(BlueprintType)
+USTRUCT(BlueprintType, meta=(DisplayName="Motion Design Font"))
 struct AVALANCHETEXT_API FAvaFont
 {
 	friend class UAvaText3DComponent;
@@ -22,22 +22,22 @@ struct AVALANCHETEXT_API FAvaFont
 	static UFont* GetDefaultFont();
 
 	/**
-	 * Returns a formatted string representing a FAvalancheFont for the specified FontObjectPathName and FontName
-	 * Useful e.g. when setting the value of multiple FAvalancheFonts using IPropertyHandle->SetPerObjectValues
-	 * @param InFontName {The name of the font handled by UAvalancheFontObject}
-	 * @param InFontObjectPathName {The path of the UAvalancheFontObject}
+	 * Returns a formatted string representing a FAvaFont for the specified FontObjectPathName and FontName
+	 * Useful e.g. when setting the value of multiple FAvaFonts using IPropertyHandle->SetPerObjectValues
+	 * @param InFontName {The name of the font handled by UAvaFontObject}
+	 * @param InFontObjectPathName {The path of the UAvaFontObject}
 	 * @return {the formatted string}
 	 */
-	static FString GenerateAvalancheFontFormattedString(const FString& InFontName, const FString& InFontObjectPathName);
+	static FString GenerateFontFormattedString(const FString& InFontName, const FString& InFontObjectPathName);
 
 	/**
-	 * Returns a formatted string representing a FAvalancheFont for the specified UAvalancheFontObject
-	 * Useful e.g. when setting the value of multiple FAvalancheFonts using IPropertyHandle->SetPerObjectValues
-	 * @param InFontObject {The UAvalancheFontObject used to generated the FAvalancheFont formatted string}
+	 * Returns a formatted string representing a FAvaFont for the specified UAvaFontObject
+	 * Useful e.g. when setting the value of multiple FAvaFonts using IPropertyHandle->SetPerObjectValues
+	 * @param InFontObject {The UAvaFontObject used to generated the FAvaFont formatted string}
 	 * @param OutFormattedString {will contain the formatted string}
 	 * @return {true if successful}
 	 */
-	static bool GenerateAvalancheFontFormattedString(const UAvaFontObject* InFontObject, FString& OutFormattedString);
+	static bool GenerateFontFormattedString(const UAvaFontObject* InFontObject, FString& OutFormattedString);
 
 	/**
 	 * Checks if the specified Ava Fonts are the same
@@ -50,7 +50,7 @@ struct AVALANCHETEXT_API FAvaFont
 	/** Default constructor: calling GetFont() right after the constructor will return the default font */
 	FAvaFont();
 
-	/** this FAvalancheFont will be initialized using the provided UAvalancheFontObject */
+	/** this FAvaFont will be initialized using the provided UAvaFontObject */
 	FAvaFont(UAvaFontObject* InFontObject);
 
 	/** Returns the currently stored font, or default font if none available. Preferred getter to just get the UFont and use it */
@@ -65,21 +65,21 @@ struct AVALANCHETEXT_API FAvaFont
 	/** Returns current font name as FText */
 	FText GetFontNameAsText() const;
 
-	/** Returns true if this FAvalancheFont is marked as favorite by the user */
+	/** Returns true if this FAvaFont is marked as favorite by the user */
 	bool IsFavorite() const;
 
 	/**
-	 * Checks if this FAvalancheFont is referencing the default font. This will be happening also in fallback state.
+	 * Checks if this FAvaFont is referencing the default font. This will be happening also in fallback state.
 	 * Use IsFallBackFont() if you explicitly need to know if this font is default because of a missing asset
 	 * @return true if referenced font is the default one
 	 */
 	bool IsDefaultFont() const;
 
 	/**
-	 * Checks if this FAvalancheFont is in fallback state.
-	 * This happens when the font referenced by this FAvalancheFont is not available after de-serialization. This likely means that a UFont asset is missing.
-	 * In fallback state, a FAvalancheFont will return the default font when GetFont() is called.
-	 * @return true if this FAvalancheFont is in fallback state
+	 * Checks if this FAvaFont is in fallback state.
+	 * This happens when the font referenced by this FAvaFont is not available after de-serialization. This likely means that a UFont asset is missing.
+	 * In fallback state, a FAvaFont will return the default font when GetFont() is called.
+	 * @return true if this FAvaFont is in fallback state
 	 */
 	bool IsFallbackFont() const;
 
@@ -99,12 +99,12 @@ struct AVALANCHETEXT_API FAvaFont
 	void SetFavorite(const bool bFavorite);
 
 	/**
-	 * Calling this function will null the now deprecated UFont CurrentFont reference, and use that UFont to initialize the currently used UAvalancheFontObject
+	 * Calling this function will null the now deprecated UFont CurrentFont reference, and use that UFont to initialize the currently used UAvaFontObject
 	 */
 	void EnsureUsingCurrentVersion();
 
 	/**
-	 * In case this FAvalancheFont is in fallback state, this call can return the name of the missing font asset, if available.
+	 * In case this FAvaFont is in fallback state, this call can return the name of the missing font asset, if available.
 	 * @return the name of the missing font asset.
 	 */
 	const FString& GetMissingFontName() const { return MissingFontName; }
@@ -116,16 +116,16 @@ struct AVALANCHETEXT_API FAvaFont
 	bool HasValidFont() const;
 
 	/**
-	 * Returns the currently stored UAvalancheFontObject. This is useful when dealing with font load/import.
+	 * Returns the currently stored UAvaFontObject. This is useful when dealing with font load/import.
 	 * If only the current UFont is needed, use GetFont().
 	 * @return Current font object. 
 	 */
-	UAvaFontObject* GetFontObject() const { return AvalancheFontObject; }
+	UAvaFontObject* GetFontObject() const { return MotionDesignFontObject; }
 
 	/**
-	 * Directly sets the current UAvalancheFontObject, and then updates this FAvalancheFont accordingly.
-	 * If this FAvalancheFont is from past versions, the internal deprecated CurrentFont field will be set to nullptr.
-	 * @param InFontObject {The FontObject which will be assigned to this FAvalancheFont}
+	 * Directly sets the current UAvaFontObject, and then updates this FAvaFont accordingly.
+	 * If this FAvaFont is from past versions, the internal deprecated CurrentFont field will be set to nullptr.
+	 * @param InFontObject {The FontObject which will be assigned to this FAvaFont}
 	 */
 	void SetFontObject(UAvaFontObject* InFontObject);
 
@@ -157,8 +157,8 @@ struct AVALANCHETEXT_API FAvaFont
 	}
 
 	/**
-	 * Called after serialization of FAvalancheFont.
-	 * In case of missing font assets, the FAvalancheFont state will be marked as "FallbackFont"
+	 * Called after serialization of FAvaFont.
+	 * In case of missing font assets, the FAvaFont state will be marked as "FallbackFont"
 	 */
 	void PostSerialize(const FArchive& Ar);
 
@@ -204,36 +204,36 @@ private:
 	static FAvaDefaultFontObjects DefaultFontObjects;
 
 	/**
-	 * Deprecated - used to reference the font used by FAvalancheFont.
-	 * Older assets still have the font stored in this field. Calling GetFont() will move the font to the currently used AvalancheFontObject property
+	 * Deprecated - used to reference the font used by FAvaFont.
+	 * Older assets still have the font stored in this field. Calling GetFont() will move the font to the currently used MotionDesignFontObject property
 	 */
-	UPROPERTY(Transient, meta=(DeprecatedProperty, DeprecationMessage="AvalancheFontObject is the property currently handling the font."))
+	UPROPERTY(Transient, meta=(DeprecatedProperty, DeprecationMessage="MotionDesignFontObject is the property currently handling the font."))
 	TObjectPtr<UFont> CurrentFont_DEPRECATED;
 
-	/** Stores the UAvalancheFontObject holding all the information about the font currently used by FAvalancheFont */
+	/** Stores the UAvaFontObject holding all the information about the font currently used by FAvaFont */
 	UPROPERTY()
-	TObjectPtr<UAvaFontObject> AvalancheFontObject;
+	TObjectPtr<UAvaFontObject> MotionDesignFontObject;
 
-	/** Updated when AvalancheFontObject is set. Mainly used when loading FAvaFont, to help looking for a font in case the current font is null */
+	/** Updated when MotionDesignFontObject is set. Mainly used when loading FAvaFont, to help looking for a font in case the current font is null */
 	UPROPERTY()
 	FString FontName;
 
-	/** The current font asset state, used internally to track the current state of FAvalancheFont with respect to its referenced font asset */
+	/** The current font asset state, used internally to track the current state of FAvaFont with respect to its referenced font asset */
 	EFontAssetState FontAssetState;
 
-	/** FAvalancheFonts can be marked as favorites by the user */
+	/** FAvaFonts can be marked as favorites by the user */
 	bool bIsFavorite;
 
 	/** When a font is missing and its name is available, it will be stored here. Old assets version cannot retrieve the name of the missing font */
 	FString MissingFontName;
 };
 
-/** We need to intercept FAvalancheFont serialization to detect if there's a missing font when loading */
+/** We need to intercept FAvaFont serialization to detect if there's a missing font when loading */
 template<> struct TStructOpsTypeTraits<FAvaFont> : public TStructOpsTypeTraitsBase2<FAvaFont>
 {
 	enum
 	{
-		/** void PostSerialize(const FArchive& Ar) in FAvalancheFont */
+		/** void PostSerialize(const FArchive& Ar) in FAvaFont */
 		WithPostSerialize = true 
 	};
 };
diff --git a/Engine/Plugins/Experimental/Avalanche/Source/AvalancheTextEditor/Private/Font/AvaFontDetailsCustomization.cpp b/Engine/Plugins/Experimental/Avalanche/Source/AvalancheTextEditor/Private/Font/AvaFontDetailsCustomization.cpp
--- a/Engine/Plugins/Experimental/Avalanche/Source/AvalancheTextEditor/Private/Font/AvaFontDetailsCustomization.cpp
+++ b/Engine/Plugins/Experimental/Avalanche/Source/AvalancheTextEditor/Private/Font/AvaFontDetailsCustomization.cpp
@@ -523,7 +523,7 @@ void FAvaFontDetailsCustomization::UpdateFontObjectProperty(UAvaFontObject* InFo
 						NewFontObj->InitFromFontObject(InFontObject);
 
 						FString ValueString;
-						FAvaFont::GenerateAvalancheFontFormattedString(NewFontObj, ValueString);
+						FAvaFont::GenerateFontFormattedString(NewFontObj, ValueString);
 
 						AvaFontPropertyHandle->SetPerObjectValue(Index, ValueString);
 					}
@@ -556,7 +556,7 @@ void FAvaFontDetailsCustomization::UpdateFontObjectProperty(UAvaFontObject* InFo
 				}
 
 				FString ValueString;
-				bool bStringSuccess = FAvaFont::GenerateAvalancheFontFormattedString(SourceFontObject, ValueString);
+				bool bStringSuccess = FAvaFont::GenerateFontFormattedString(SourceFontObject, ValueString);
 
 				if (bStringSuccess)
 				{
@@ -593,7 +593,7 @@ void FAvaFontDetailsCustomization::UpdateFontObjectProperty(UAvaFontObject* InFo
 				NewFontObj->InitFromFontObject(InFontObject);
 
 				FString ValueString;
-				FAvaFont::GenerateAvalancheFontFormattedString(NewFontObj, ValueString);
+				FAvaFont::GenerateFontFormattedString(NewFontObj, ValueString);
 
 				PerObjectValues.Add(ValueString);
 			}
diff --git a/Engine/Plugins/Experimental/Avalanche/Source/AvalancheTextEditor/Private/Font/AvaFontManagerSubsystem.h b/Engine/Plugins/Experimental/Avalanche/Source/AvalancheTextEditor/Private/Font/AvaFontManagerSubsystem.h
--- a/Engine/Plugins/Experimental/Avalanche/Source/AvalancheTextEditor/Private/Font/AvaFontManagerSubsystem.h
+++ b/Engine/Plugins/Experimental/Avalanche/Source/AvalancheTextEditor/Private/Font/AvaFontManagerSubsystem.h
@@ -126,8 +126,8 @@ public:
 	FOnProjectFontsChange& OnProjectFontDeleted() { return OnProjectFontDeletedDelegate; }
 	FOnSystemFontsUpdated& OnSystemFontsUpdated() { return OnSystemFontsUpdatedDelegate; }
 	/**
-	 * Import the font contained in the UAvalancheFontObject as .uasset, if yet to be imported
-	 * @param InAvaFontPropertyHandle a property handle to the FAvalancheFont referencing the UFont to be saved
+	 * Import the font contained in the UAvaFontObject as .uasset, if yet to be imported
+	 * @param InAvaFontPropertyHandle a property handle to the FAvaFont referencing the UFont to be saved
 	 * @return success
 	 */
 	bool ImportFont(const TSharedPtr<IPropertyHandle>& InAvaFontPropertyHandle);
