Mask modifier disable/enable support

diff --git a/Engine/Plugins/Experimental/Avalanche/Source/AvalancheMask/Private/Mask2D/AvaMask2DBaseModifier.cpp b/Engine/Plugins/Experimental/Avalanche/Source/AvalancheMask/Private/Mask2D/AvaMask2DBaseModifier.cpp
--- a/Engine/Plugins/Experimental/Avalanche/Source/AvalancheMask/Private/Mask2D/AvaMask2DBaseModifier.cpp
+++ b/Engine/Plugins/Experimental/Avalanche/Source/AvalancheMask/Private/Mask2D/AvaMask2DBaseModifier.cpp
@@ -187,30 +187,6 @@ void UAvaMask2DBaseModifier::OnModifierAdded(EActorModifierCoreEnableReason InRe
 	}
 }
 
-void UAvaMask2DBaseModifier::OnModifierRemoved(EActorModifierCoreDisableReason InReason)
-{
-	Super::OnModifierRemoved(InReason);
-
-	TRACE_BOOKMARK(TEXT("UAvaMask2DModifier::OnModifierRemoved"));
-
-	TArray<AActor*> ActorDataActors;
-	ActorDataActors.Reserve(ActorData.Num());	
-	Algo::Transform(ActorData, ActorDataActors, [](const TPair<TWeakObjectPtr<AActor>, FAvaMask2DActorData>& InActorDataPair)
-	{
-		return InActorDataPair.Key.Get();
-	});
-	
-	for (AActor* Actor : ActorDataActors)
-	{
-		if (!Actor)
-		{
-			continue;
-		}
-		
-		RemoveFromActor(Actor);
-	}
-}
-
 void UAvaMask2DBaseModifier::OnModifiedActorTransformed()
 {
 	// This is needed to override parent behavior (to disable it)
@@ -248,6 +224,23 @@ void UAvaMask2DBaseModifier::RestorePreState()
 			RestoreActorPreState(Actor, ActorDataPair.Value);
 		}
 	}
+	
+	TArray<AActor*> ActorDataActors;
+	ActorDataActors.Reserve(ActorData.Num());	
+	Algo::Transform(ActorData, ActorDataActors, [](const TPair<TWeakObjectPtr<AActor>, FAvaMask2DActorData>& InActorDataPair)
+	{
+		return InActorDataPair.Key.Get();
+	});
+	
+	for (AActor* Actor : ActorDataActors)
+	{
+		if (!Actor)
+		{
+			continue;
+		}
+		
+		RemoveFromActor(Actor);
+	}
 }
 
 void UAvaMask2DBaseModifier::Apply()
diff --git a/Engine/Plugins/Experimental/Avalanche/Source/AvalancheMask/Public/Mask2D/AvaMask2DBaseModifier.h b/Engine/Plugins/Experimental/Avalanche/Source/AvalancheMask/Public/Mask2D/AvaMask2DBaseModifier.h
--- a/Engine/Plugins/Experimental/Avalanche/Source/AvalancheMask/Public/Mask2D/AvaMask2DBaseModifier.h
+++ b/Engine/Plugins/Experimental/Avalanche/Source/AvalancheMask/Public/Mask2D/AvaMask2DBaseModifier.h
@@ -105,7 +105,6 @@ protected:
 	//~ Begin UActorModifierCoreBase
 	virtual void OnModifierCDOSetup(FActorModifierCoreMetadata& InMetadata) override;
 	virtual void OnModifierAdded(EActorModifierCoreEnableReason InReason) override;
-	virtual void OnModifierRemoved(EActorModifierCoreDisableReason InReason) override;
 	virtual void OnModifiedActorTransformed() override;
 	virtual void SavePreState() override;
 	virtual void RestorePreState() override;
