[Motion Design] Motion Design Text bug fixes: alignment widgets + import text string change for deprecated property

- fix wrong name for deprecated property in Import Text string generation
- fix issue with alignment customization applying alignment as soon as Text Actors are selected

diff --git a/Engine/Plugins/Experimental/Avalanche/Source/AvalancheText/Private/Font/AvaFont.cpp b/Engine/Plugins/Experimental/Avalanche/Source/AvalancheText/Private/Font/AvaFont.cpp
--- a/Engine/Plugins/Experimental/Avalanche/Source/AvalancheText/Private/Font/AvaFont.cpp
+++ b/Engine/Plugins/Experimental/Avalanche/Source/AvalancheText/Private/Font/AvaFont.cpp
@@ -64,7 +64,7 @@ FString FAvaFont::GenerateFontFormattedString(const FString& InFontName, const F
 	// e.g. (Property1=Value1,Property2=Value2,Property3=Value3,...)
 
 	FString FormattedString = TEXT("(");
-	FormattedString += TEXT("CurrentFont=None");
+	FormattedString += TEXT("CurrentFont_DEPRECATED=None");
 	FormattedString += TEXT(",MotionDesignFontObject=/Script/AvalancheText.AvaFontObject'") + InFontObjectPathName + TEXT("'");
 	FormattedString += TEXT(",FontName=\"") + InFontName  + TEXT("\"");
 	FormattedString += TEXT(")");
diff --git a/Engine/Plugins/Experimental/Avalanche/Source/AvalancheTextEditor/Private/Slate/SAvaTextAlignmentWidget.cpp b/Engine/Plugins/Experimental/Avalanche/Source/AvalancheTextEditor/Private/Slate/SAvaTextAlignmentWidget.cpp
--- a/Engine/Plugins/Experimental/Avalanche/Source/AvalancheTextEditor/Private/Slate/SAvaTextAlignmentWidget.cpp
+++ b/Engine/Plugins/Experimental/Avalanche/Source/AvalancheTextEditor/Private/Slate/SAvaTextAlignmentWidget.cpp
@@ -85,8 +85,6 @@ void SAvaTextAlignmentWidget::Construct(const FArguments& InArgs)
 		]
 	];
 
-	OnAlignmentPropertyChanged();
-
 	TextAlignmentPropertyHandle.Get()->SetOnPropertyValueChanged(FSimpleDelegate::CreateSP(this, &SAvaTextAlignmentWidget::OnAlignmentPropertyChanged));
 }
 
