Avalanche Sequence:
- Improved UI  to include a checkbox to denote whether that time entry is enabled or disabled, similar to how EditConditionInline works.
- This is primarily to solve the confusion where users have to set "-1" for End Time to mean "play until the end". With this solution, users can opt to set the checkbox to false to not have to set any specific time

diff --git a/Engine/Plugins/Experimental/Avalanche/Source/AvalancheEditor/Private/Transition/AvaTransitionExtension.cpp b/Engine/Plugins/Experimental/Avalanche/Source/AvalancheEditor/Private/Transition/AvaTransitionExtension.cpp
--- a/Engine/Plugins/Experimental/Avalanche/Source/AvalancheEditor/Private/Transition/AvaTransitionExtension.cpp
+++ b/Engine/Plugins/Experimental/Avalanche/Source/AvalancheEditor/Private/Transition/AvaTransitionExtension.cpp
@@ -177,8 +177,8 @@ void FAvaTransitionExtension::BuildDefaultTransitionTree(UAvaTransitionTreeEdito
 			SetStateCompletedResult(ChangeOut, EStateTreeTransitionType::Succeeded);
 
 			FAvaSequencePlayParams& PlaySettings = AddPlayTask(ChangeOut, ChangeTag);
-			PlaySettings.Start = FAvaSequenceTime(0);
-            PlaySettings.End   = MarkTime;
+			PlaySettings.Start.bHasTimeConstraint = false;
+            PlaySettings.End = MarkTime;
 		}
 
 		// Change in (A to End)
@@ -197,7 +197,7 @@ void FAvaTransitionExtension::BuildDefaultTransitionTree(UAvaTransitionTreeEdito
 
 			FAvaSequencePlayParams& PlaySettings = AddPlayTask(PlayChangeInState, ChangeTag);
 			PlaySettings.Start = MarkTime;
-			PlaySettings.End   = FAvaSequenceTime(-1);
+			PlaySettings.End.bHasTimeConstraint = false;
 		}
 	}
 
diff --git a/Engine/Plugins/Experimental/Avalanche/Source/AvalancheSequence/Private/AvaSequencePlayer.cpp b/Engine/Plugins/Experimental/Avalanche/Source/AvalancheSequence/Private/AvaSequencePlayer.cpp
--- a/Engine/Plugins/Experimental/Avalanche/Source/AvalancheSequence/Private/AvaSequencePlayer.cpp
+++ b/Engine/Plugins/Experimental/Avalanche/Source/AvalancheSequence/Private/AvaSequencePlayer.cpp
@@ -60,43 +60,15 @@ void UAvaSequencePlayer::SetPlaySettings(const FAvaSequencePlayParams& InPlaySet
 
 	bReversePlayback = InPlaySettings.PlayMode == EAvaSequencePlayMode::Reverse;
 
-	// Playback Range in Seconds
-	const double TotalDuration = MovieScene->GetTickResolution().AsSeconds(MovieScene->GetPlaybackRange().Size<FFrameNumber>());
-
-	double StartTimeSeconds = InPlaySettings.Start.ToSeconds(*PlaybackSequence, *MovieScene);
-	double EndTimeSeconds   = InPlaySettings.End.ToSeconds(*PlaybackSequence, *MovieScene);
-
-	StartTimeSeconds = FMath::Max(StartTimeSeconds, 0.0);
-
-	// Treat negative values as just meaning "play to end"
-	if (EndTimeSeconds < 0)
-	{
-		EndTimeSeconds = TotalDuration;
-	}
-
-	EndTimeSeconds = FMath::Max(EndTimeSeconds, 0.0);
-
-	if (bReversePlayback && !InPlaySettings.Start.IsAbsoluteTime())
-	{
-		StartTimeSeconds = TotalDuration - StartTimeSeconds;
-	}
-
-	if (bReversePlayback && !InPlaySettings.End.IsAbsoluteTime())
-	{
-		EndTimeSeconds = TotalDuration - EndTimeSeconds;
-	}
+	const double TotalDuration    = MovieScene->GetTickResolution().AsSeconds(MovieScene->GetPlaybackRange().Size<FFrameNumber>());
+	const double StartTimeSeconds = InPlaySettings.Start.ToSeconds(*PlaybackSequence, *MovieScene, 0.0);
+	const double EndTimeSeconds   = InPlaySettings.End.ToSeconds(*PlaybackSequence, *MovieScene, TotalDuration);
 
 	double DurationSeconds = FMath::Abs(EndTimeSeconds - StartTimeSeconds);
 
 	// Ensure that Duration Seconds is at least 1 Frame Long, as this gets Rounded to the nearest Frame and if it results in 0, the Sequence does not get evaluated
 	DurationSeconds = FMath::Max(DurationSeconds, PlayPosition.GetInputRate().AsSeconds(FFrameNumber(1)));
 
-	if (bReversePlayback)
-	{
-		// Remove Duration Seconds as this is is later added by the Sequence Player for Reverse Playbacks
-		StartTimeSeconds -= DurationSeconds;
-	}
-
 	SetTimeRange(StartTimeSeconds, DurationSeconds);
 
 	PlaybackSettings = FMovieSceneSequencePlaybackSettings();
diff --git a/Engine/Plugins/Experimental/Avalanche/Source/AvalancheSequence/Private/AvaSequenceShared.cpp b/Engine/Plugins/Experimental/Avalanche/Source/AvalancheSequence/Private/AvaSequenceShared.cpp
--- a/Engine/Plugins/Experimental/Avalanche/Source/AvalancheSequence/Private/AvaSequenceShared.cpp
+++ b/Engine/Plugins/Experimental/Avalanche/Source/AvalancheSequence/Private/AvaSequenceShared.cpp
@@ -2,9 +2,43 @@
 
 #include "AvaSequenceShared.h"
 #include "AvaSequence.h"
+#include "AvaSequenceVersion.h"
 
-double FAvaSequenceTime::ToSeconds(const UAvaSequence& InSequence, const UMovieScene& InMovieScene) const
+bool FAvaSequenceTime::Serialize(FArchive& Ar)
 {
+	Ar.UsingCustomVersion(FAvaSequenceVersion::GUID);
+
+	UScriptStruct* ScriptStruct = FAvaSequenceTime::StaticStruct();
+	ScriptStruct->SerializeTaggedProperties(Ar, reinterpret_cast<uint8*>(this), ScriptStruct, nullptr);
+
+	if (Ar.IsLoading() && Ar.CustomVer(FAvaSequenceVersion::GUID) < FAvaSequenceVersion::SequenceTimeConstraintOption)
+	{
+		switch (TimeType)
+		{
+		case EAvaSequenceTimeType::Frame:
+			bHasTimeConstraint = Frame >= 0;
+			break;
+
+		case EAvaSequenceTimeType::Seconds:
+			bHasTimeConstraint = Seconds >= 0.0;
+			break;
+
+		case EAvaSequenceTimeType::Mark:
+			bHasTimeConstraint = !MarkLabel.IsEmpty();
+			break;
+		}
+	}
+
+	return true;
+}
+
+double FAvaSequenceTime::ToSeconds(const UAvaSequence& InSequence, const UMovieScene& InMovieScene, double InDefaultTime) const
+{
+	if (!bHasTimeConstraint)
+	{
+		return InDefaultTime;
+	}
+
 	switch (TimeType)
 	{
 	case EAvaSequenceTimeType::Frame:
@@ -22,5 +56,6 @@ double FAvaSequenceTime::ToSeconds(const UAvaSequence& InSequence, const UMovieS
 			}
 		}
 	}
-	return -1.0;
+
+	return InDefaultTime;
 }
diff --git a/Engine/Plugins/Experimental/Avalanche/Source/AvalancheSequence/Private/AvaSequenceVersion.h b/Engine/Plugins/Experimental/Avalanche/Source/AvalancheSequence/Private/AvaSequenceVersion.h
--- a/Engine/Plugins/Experimental/Avalanche/Source/AvalancheSequence/Private/AvaSequenceVersion.h
+++ b/Engine/Plugins/Experimental/Avalanche/Source/AvalancheSequence/Private/AvaSequenceVersion.h
@@ -17,6 +17,9 @@ public:
 		/** Moved to become child class of ULevelSequence */
 		LevelSequence, 
 
+		/** Introduced Time Constraint Option to FAvaSequenceTime */
+		SequenceTimeConstraintOption,
+
 		VersionPlusOne,
 		LatestVersion = VersionPlusOne - 1
 	};
diff --git a/Engine/Plugins/Experimental/Avalanche/Source/AvalancheSequence/Public/AvaSequenceShared.h b/Engine/Plugins/Experimental/Avalanche/Source/AvalancheSequence/Public/AvaSequenceShared.h
--- a/Engine/Plugins/Experimental/Avalanche/Source/AvalancheSequence/Public/AvaSequenceShared.h
+++ b/Engine/Plugins/Experimental/Avalanche/Source/AvalancheSequence/Public/AvaSequenceShared.h
@@ -40,24 +40,34 @@ struct FAvaSequenceTime
 {
 	GENERATED_BODY()
 
+	enum ENoTimeConstraint {NoTimeConstraint};
+
 	FAvaSequenceTime() = default;
 
+	explicit FAvaSequenceTime(ENoTimeConstraint)
+		: bHasTimeConstraint(false)
+	{
+	}
+
 	explicit FAvaSequenceTime(FFrameTime FrameTime)
 		: TimeType(EAvaSequenceTimeType::Frame)
 		, Frame(FrameTime.GetFrame().Value)
 		, SubFrame(FrameTime.GetSubFrame())
+		, bHasTimeConstraint(true)
 	{
 	}
 
 	explicit FAvaSequenceTime(double InSeconds)
 		: TimeType(EAvaSequenceTimeType::Seconds)
 		, Seconds(InSeconds)
+		, bHasTimeConstraint(true)
 	{
 	}
 
 	explicit FAvaSequenceTime(const FString& InMarkLabel)
 		: TimeType(EAvaSequenceTimeType::Mark)
 		, MarkLabel(InMarkLabel)
+		, bHasTimeConstraint(true)
 	{
 	}
 
@@ -67,21 +77,43 @@ struct FAvaSequenceTime
 		return TimeType == EAvaSequenceTimeType::Mark;
 	}
 
-	AVALANCHESEQUENCE_API double ToSeconds(const UAvaSequence& InSequence, const UMovieScene& InMovieScene) const;
+	AVALANCHESEQUENCE_API bool Serialize(FArchive& Ar);
 
-	UPROPERTY(EditAnywhere, BlueprintReadWrite, Category = "Motion Design Sequence")
+	/**
+	 * Converts this Motion Design Sequence Time struct to time frame in seconds
+	 * @param InSequence the sequence to use to query for marks (only applicable if TimeType is set to Mark)
+	 * @param InMovieScene the movie scene used to get the Tick Resolution or Display Rate
+	 * @param InDefaultTime the default time to return if time is invalid or unset
+	 * @return the timeframe in seconds
+	 */
+	AVALANCHESEQUENCE_API double ToSeconds(const UAvaSequence& InSequence, const UMovieScene& InMovieScene, double InDefaultTime) const;
+
+	UPROPERTY(EditAnywhere, BlueprintReadWrite, Category = "Motion Design Sequence", meta=(EditCondition="bHasTimeConstraint"))
 	EAvaSequenceTimeType TimeType = EAvaSequenceTimeType::Frame;
 
-	UPROPERTY(EditAnywhere, BlueprintReadWrite, Category = "Motion Design Sequence", meta=(EditCondition="TimeType==EAvaSequenceTimeType::Frame"))
+	UPROPERTY(EditAnywhere, BlueprintReadWrite, Category = "Motion Design Sequence", meta=(EditCondition="bHasTimeConstraint && TimeType==EAvaSequenceTimeType::Frame"))
 	int32 Frame = 0;
 
+	UPROPERTY(EditAnywhere, BlueprintReadWrite, Category = "Motion Design Sequence", meta=(EditCondition="bHasTimeConstraint && TimeType==EAvaSequenceTimeType::Frame"))
 	double SubFrame = 0.0;
 
-	UPROPERTY(EditAnywhere, BlueprintReadWrite, Category = "Motion Design Sequence", meta=(EditCondition="TimeType==EAvaSequenceTimeType::Seconds", Unit=s))
+	UPROPERTY(EditAnywhere, BlueprintReadWrite, Category = "Motion Design Sequence", meta=(EditCondition="bHasTimeConstraint && TimeType==EAvaSequenceTimeType::Seconds", Unit=s))
 	double Seconds = 0.0;
 
-	UPROPERTY(EditAnywhere, BlueprintReadWrite, Category = "Motion Design Sequence", meta=(EditCondition="TimeType==EAvaSequenceTimeType::Mark"))
+	UPROPERTY(EditAnywhere, BlueprintReadWrite, Category = "Motion Design Sequence", meta=(EditCondition="bHasTimeConstraint && TimeType==EAvaSequenceTimeType::Mark"))
 	FString MarkLabel;
+
+	UPROPERTY(EditAnywhere, BlueprintReadWrite, Category = "Motion Design Sequence")
+	bool bHasTimeConstraint = false;
+};
+
+template<>
+struct TStructOpsTypeTraits<FAvaSequenceTime> : public TStructOpsTypeTraitsBase2<FAvaSequenceTime>
+{
+	enum
+	{
+		WithSerializer = true
+	};
 };
 
 USTRUCT(BlueprintType)
@@ -109,7 +141,7 @@ struct FAvaSequencePlayParams
 	FAvaSequenceTime Start = FAvaSequenceTime(0.0);
 
 	UPROPERTY(EditAnywhere, BlueprintReadWrite, Category = "Motion Design Sequence")
-	FAvaSequenceTime End = FAvaSequenceTime(-1.0);
+	FAvaSequenceTime End = FAvaSequenceTime(FAvaSequenceTime::NoTimeConstraint);
 
 	UPROPERTY(EditAnywhere, BlueprintReadWrite, Category = "Motion Design Sequence")
 	EAvaSequencePlayMode PlayMode = EAvaSequencePlayMode::Forward;
diff --git a/Engine/Plugins/Experimental/Avalanche/Source/AvalancheSequencer/Private/Customization/AvaSequenceTimeCustomization.cpp b/Engine/Plugins/Experimental/Avalanche/Source/AvalancheSequencer/Private/Customization/AvaSequenceTimeCustomization.cpp
--- a/Engine/Plugins/Experimental/Avalanche/Source/AvalancheSequencer/Private/Customization/AvaSequenceTimeCustomization.cpp
+++ b/Engine/Plugins/Experimental/Avalanche/Source/AvalancheSequencer/Private/Customization/AvaSequenceTimeCustomization.cpp
@@ -18,7 +18,7 @@ void FAvaSequenceTimeCustomization::CustomizeHeader(TSharedRef<IPropertyHandle>
 
 		TSharedRef<SWidget> PropertyValueWidget = ChildPropertyHandle->CreatePropertyValueWidget(/*bDisplayDefaultPropertyButtons*/false);
 
-		TAttribute<EVisibility>::FGetter VisibilityGetter = TAttribute<EVisibility>::FGetter::CreateSPLambda(ChildPropertyHandle.Get(), [ChildPropertyHandle]
+		TAttribute<EVisibility>::FGetter VisibilityGetter = TAttribute<EVisibility>::FGetter::CreateLambda([ChildPropertyHandle]
 			{
 				return ChildPropertyHandle->IsEditable()
 					? EVisibility::SelfHitTestInvisible
@@ -45,6 +45,8 @@ void FAvaSequenceTimeCustomization::CustomizeHeader(TSharedRef<IPropertyHandle>
 			InPropertyHandle->CreateDefaultPropertyButtonWidgets()
 		];
 
+	TSharedPtr<IPropertyHandle> HasTimeHandle = InPropertyHandle->GetChildHandle(GET_MEMBER_NAME_CHECKED(FAvaSequenceTime, bHasTimeConstraint));
+
 	InHeaderRow
 		.NameContent()
 		[
@@ -55,6 +57,12 @@ void FAvaSequenceTimeCustomization::CustomizeHeader(TSharedRef<IPropertyHandle>
 				SNew(SHorizontalBox)
 				+ SHorizontalBox::Slot()
 				.AutoWidth()
+				[
+					HasTimeHandle->CreatePropertyValueWidget(/*bDisplayDefaultPropertyButtons*/false)
+				]
+				+ SHorizontalBox::Slot()
+				.AutoWidth()
+				.Padding(2.f, 0.f, 0.f, 0.f)
 				[
 					InPropertyHandle->CreatePropertyNameWidget()
 				]
