Avalanche Tag:
- Improved UX on Tag Collection Picker, collapsing Asset Actions into 3 Vertical Dot Option and be able to do more things such as opening the TagCollection asset (editor) or navigating to it in the Content Browser.

diff --git a/Engine/Plugins/Experimental/Avalanche/Source/AvalancheTagEditor/AvalancheTagEditor.Build.cs b/Engine/Plugins/Experimental/Avalanche/Source/AvalancheTagEditor/AvalancheTagEditor.Build.cs
--- a/Engine/Plugins/Experimental/Avalanche/Source/AvalancheTagEditor/AvalancheTagEditor.Build.cs
+++ b/Engine/Plugins/Experimental/Avalanche/Source/AvalancheTagEditor/AvalancheTagEditor.Build.cs
@@ -20,6 +20,7 @@ public class AvalancheTagEditor : ModuleRules
                 "PropertyEditor",
                 "Slate",
                 "SlateCore",
+                "ToolMenus",
                 "UnrealEd",
             }
         );
diff --git a/Engine/Plugins/Experimental/Avalanche/Source/AvalancheTagEditor/Private/Customization/Menu/AvaTagCollectionMenuContext.h b/Engine/Plugins/Experimental/Avalanche/Source/AvalancheTagEditor/Private/Customization/Menu/AvaTagCollectionMenuContext.h
new file mode 100644
--- /dev/null
+++ b/Engine/Plugins/Experimental/Avalanche/Source/AvalancheTagEditor/Private/Customization/Menu/AvaTagCollectionMenuContext.h
@@ -0,0 +1,18 @@
+﻿// Copyright Epic Games, Inc. All Rights Reserved.
+
+#pragma once
+
+#include "UObject/Object.h"
+#include "Templates/SharedPointer.h"
+#include "AvaTagCollectionMenuContext.generated.h"
+
+class IPropertyHandle;
+
+UCLASS()
+class UAvaTagCollectionMenuContext : public UObject
+{
+	GENERATED_BODY()
+
+public:
+	TWeakPtr<IPropertyHandle> TagCollectionHandleWeak;
+};
diff --git a/Engine/Plugins/Experimental/Avalanche/Source/AvalancheTagEditor/Private/Customization/Menu/AvaTagCollectionPickerContextMenu.cpp b/Engine/Plugins/Experimental/Avalanche/Source/AvalancheTagEditor/Private/Customization/Menu/AvaTagCollectionPickerContextMenu.cpp
new file mode 100644
--- /dev/null
+++ b/Engine/Plugins/Experimental/Avalanche/Source/AvalancheTagEditor/Private/Customization/Menu/AvaTagCollectionPickerContextMenu.cpp
@@ -0,0 +1,102 @@
+﻿// Copyright Epic Games, Inc. All Rights Reserved.
+
+#include "AvaTagCollectionPickerContextMenu.h"
+#include "AvaTagCollection.h"
+#include "AvaTagCollectionMenuContext.h"
+#include "Editor.h"
+#include "PropertyHandle.h"
+#include "Styling/AppStyle.h"
+#include "Subsystems/AssetEditorSubsystem.h"
+#include "ToolMenu.h"
+#include "ToolMenuSection.h"
+#include "ToolMenus.h"
+#include "Widgets/SNullWidget.h"
+
+#define LOCTEXT_NAMESPACE "AvaTagCollectionPickerContextMenu"
+
+FAvaTagCollectionPickerContextMenu& FAvaTagCollectionPickerContextMenu::Get()
+{
+	static TSharedRef<FAvaTagCollectionPickerContextMenu> ContextMenu = MakeShared<FAvaTagCollectionPickerContextMenu>(FPrivateToken());
+	return ContextMenu.Get();
+}
+
+TSharedRef<SWidget> FAvaTagCollectionPickerContextMenu::GenerateContextMenuWidget(const TSharedPtr<IPropertyHandle>& InTagCollectionHandle)
+{
+	UToolMenus* ToolMenus = UToolMenus::Get();
+	if (!ToolMenus)
+	{
+		return SNullWidget::NullWidget;
+	}
+
+	constexpr const TCHAR* ContextMenuName = TEXT("AvaTagCollectionPickerContextMenu");
+
+	if (!ToolMenus->IsMenuRegistered(ContextMenuName))
+	{
+		UToolMenu* const ToolMenu = ToolMenus->RegisterMenu(ContextMenuName, NAME_None, EMultiBoxType::Menu);
+		check(ToolMenu);
+		ToolMenu->AddDynamicSection(TEXT("PopulateContextMenu"), FNewToolMenuDelegate::CreateSP(this, &FAvaTagCollectionPickerContextMenu::PopulateContextMenu));
+	}
+
+	UAvaTagCollectionMenuContext* MenuContext = NewObject<UAvaTagCollectionMenuContext>();
+	MenuContext->TagCollectionHandleWeak = InTagCollectionHandle;
+
+	FToolMenuContext ToolMenuContext(MenuContext);	
+	return ToolMenus->GenerateWidget(ContextMenuName, ToolMenuContext);
+}
+
+void FAvaTagCollectionPickerContextMenu::PopulateContextMenu(UToolMenu* InToolMenu)
+{
+	UAvaTagCollectionMenuContext* MenuContext = InToolMenu
+		? InToolMenu->FindContext<UAvaTagCollectionMenuContext>()
+		: nullptr;
+
+	if (!MenuContext || !MenuContext->TagCollectionHandleWeak.IsValid())
+	{
+		return;
+	}
+
+	FToolMenuSection& AssetActions = InToolMenu->FindOrAddSection(TEXT("AssetActions"), LOCTEXT("AssetSectionLabel", "Asset"));
+
+	AssetActions.AddMenuEntry(TEXT("EditAsset")
+		, LOCTEXT("EditAssetLabel", "Edit...")
+		, LOCTEXT("EditAssetTooltip", "Opens the Tag Collection asset editor")
+		, FSlateIcon(FAppStyle::GetAppStyleSetName(), "Icons.Edit")
+		, FExecuteAction::CreateSP(this, &FAvaTagCollectionPickerContextMenu::EditAsset, MenuContext->TagCollectionHandleWeak));
+
+	AssetActions.AddMenuEntry(TEXT("FindInContentBrowser")
+		, LOCTEXT("FindInContentBrowserLabel", "Find in Content Browser")
+		, LOCTEXT("FindInContentBrowserTooltip", "Navigates to the Tag Collection asset in the Content Browser")
+		, FSlateIcon(FAppStyle::GetAppStyleSetName(), "SystemWideCommands.FindInContentBrowser")
+		, FExecuteAction::CreateSP(this, &FAvaTagCollectionPickerContextMenu::FindInContentBrowser, MenuContext->TagCollectionHandleWeak));
+}
+
+bool FAvaTagCollectionPickerContextMenu::TryGetAsset(const TWeakPtr<IPropertyHandle>& InTagCollectionHandleWeak, FAssetData& OutAssetData) const
+{
+	if (TSharedPtr<IPropertyHandle> TagCollectionHandle = InTagCollectionHandleWeak.Pin())
+	{
+		return TagCollectionHandle->GetValue(OutAssetData) == FPropertyAccess::Success;
+	}
+	return false;
+}
+
+void FAvaTagCollectionPickerContextMenu::EditAsset(TWeakPtr<IPropertyHandle> InTagCollectionHandleWeak)
+{
+	UAssetEditorSubsystem* AssetEditorSubsystem = GEditor ? GEditor->GetEditorSubsystem<UAssetEditorSubsystem>() : nullptr;
+
+	FAssetData TagCollectionAsset;
+	if (AssetEditorSubsystem && TryGetAsset(InTagCollectionHandleWeak, TagCollectionAsset))
+	{
+		AssetEditorSubsystem->OpenEditorForAsset(TagCollectionAsset.GetSoftObjectPath());
+	}
+}
+
+void FAvaTagCollectionPickerContextMenu::FindInContentBrowser(TWeakPtr<IPropertyHandle> InTagCollectionHandleWeak)
+{
+	FAssetData TagCollectionAsset;
+	if (GEditor && TryGetAsset(InTagCollectionHandleWeak, TagCollectionAsset))
+	{
+		GEditor->SyncBrowserToObject(TagCollectionAsset);
+	}
+}
+
+#undef LOCTEXT_NAMESPACE
diff --git a/Engine/Plugins/Experimental/Avalanche/Source/AvalancheTagEditor/Private/Customization/Menu/AvaTagCollectionPickerContextMenu.h b/Engine/Plugins/Experimental/Avalanche/Source/AvalancheTagEditor/Private/Customization/Menu/AvaTagCollectionPickerContextMenu.h
new file mode 100644
--- /dev/null
+++ b/Engine/Plugins/Experimental/Avalanche/Source/AvalancheTagEditor/Private/Customization/Menu/AvaTagCollectionPickerContextMenu.h
@@ -0,0 +1,31 @@
+﻿// Copyright Epic Games, Inc. All Rights Reserved.
+
+#pragma once
+
+#include "Templates/SharedPointer.h"
+
+class IPropertyHandle;
+class SWidget;
+class UToolMenu;
+struct FAssetData;
+
+class FAvaTagCollectionPickerContextMenu : public TSharedFromThis<FAvaTagCollectionPickerContextMenu>
+{
+	struct FPrivateToken { explicit FPrivateToken() = default; };
+
+public:
+	explicit FAvaTagCollectionPickerContextMenu(FPrivateToken) {}
+
+	static FAvaTagCollectionPickerContextMenu& Get();
+
+	TSharedRef<SWidget> GenerateContextMenuWidget(const TSharedPtr<IPropertyHandle>& InTagCollectionHandle);
+
+private:
+	void PopulateContextMenu(UToolMenu* InToolMenu);
+
+	bool TryGetAsset(const TWeakPtr<IPropertyHandle>& InTagCollectionHandleWeak, FAssetData& OutAssetData) const;
+
+	void EditAsset(TWeakPtr<IPropertyHandle> InTagCollectionHandleWeak);
+
+	void FindInContentBrowser(TWeakPtr<IPropertyHandle> InTagCollectionHandleWeak);
+};
diff --git a/Engine/Plugins/Experimental/Avalanche/Source/AvalancheTagEditor/Private/Customization/SAvaTagCollectionPicker.cpp b/Engine/Plugins/Experimental/Avalanche/Source/AvalancheTagEditor/Private/Customization/SAvaTagCollectionPicker.cpp
--- a/Engine/Plugins/Experimental/Avalanche/Source/AvalancheTagEditor/Private/Customization/SAvaTagCollectionPicker.cpp
+++ b/Engine/Plugins/Experimental/Avalanche/Source/AvalancheTagEditor/Private/Customization/SAvaTagCollectionPicker.cpp
@@ -5,10 +5,8 @@
 #include "DetailLayoutBuilder.h"
 #include "Editor.h"
 #include "PropertyCustomizationHelpers.h"
-#include "Widgets/Images/SImage.h"
-#include "Widgets/Input/SButton.h"
 #include "Widgets/Input/SComboButton.h"
-#include "Widgets/SBoxPanel.h"
+#include "Widgets/SNullWidget.h"
 #include "Widgets/Text/STextBlock.h"
 
 #define LOCTEXT_NAMESPACE "SAvaTagCollectionPicker"
@@ -21,39 +19,17 @@ void SAvaTagCollectionPicker::Construct(const FArguments& InArgs, const TSharedR
 
 	ChildSlot
 	[
-		SNew(SHorizontalBox)
-		+ SHorizontalBox::Slot()
-		.FillWidth(1)
+		SAssignNew(TagCollectionPickerButton, SComboButton)
+		.OnGetMenuContent(this, &SAvaTagCollectionPicker::MakeTagCollectionPicker)
+		.OnMenuOpenChanged(this, &SAvaTagCollectionPicker::OnTagCollectionMenuOpenChanged)
+		.ContentPadding(FMargin(6, -1, 0, -1))
+		.ButtonStyle(FAppStyle::Get(), "SimpleButton")
+		.ButtonContent()
 		[
-			SAssignNew(TagCollectionPickerButton, SComboButton)
-			.OnGetMenuContent(this, &SAvaTagCollectionPicker::MakeTagCollectionPicker)
-			.OnMenuOpenChanged(this, &SAvaTagCollectionPicker::OnTagCollectionMenuOpenChanged)
-			.ContentPadding(FMargin(6, -1, 0, -1))
-			.ButtonStyle(FAppStyle::Get(), "SimpleButton")
-			.ButtonContent()
-			[
-				SNew(STextBlock)
-				.Text(this, &SAvaTagCollectionPicker::GetTagCollectionTitleText)
-				.ToolTipText(this, &SAvaTagCollectionPicker::GetTagCollectionTooltipText)
-				.Font(IDetailLayoutBuilder::GetDetailFont())
-			]
-		]
-		+ SHorizontalBox::Slot()
-		.AutoWidth()
-		.HAlign(HAlign_Center)
-		.VAlign(VAlign_Center)
-		[
-			SNew(SButton)
-			.OnClicked(this, &SAvaTagCollectionPicker::FindInContentBrowser)
-			.ButtonStyle(FAppStyle::Get(), "HoverHintOnly")
-			.ToolTipText(LOCTEXT("FindInContentBrowserTooltip", "Browse to this tag collection in the Content Browser"))
-			.VAlign(VAlign_Center)
-			.HAlign(HAlign_Center)
-			.DesiredSizeScale(0.75f)
-			[
-				SNew(SImage)
-				.Image(FSlateIcon(FAppStyle::GetAppStyleSetName(), "SystemWideCommands.FindInContentBrowser").GetIcon())
-			]
+			SNew(STextBlock)
+			.Text(this, &SAvaTagCollectionPicker::GetTagCollectionTitleText)
+			.ToolTipText(this, &SAvaTagCollectionPicker::GetTagCollectionTooltipText)
+			.Font(IDetailLayoutBuilder::GetDetailFont())
 		]
 	];
 }
@@ -68,25 +44,6 @@ void SAvaTagCollectionPicker::SetIsOpen(bool bInIsOpen)
 	TagCollectionPickerButton->SetIsOpen(bInIsOpen);
 }
 
-FReply SAvaTagCollectionPicker::FindInContentBrowser()
-{
-	if (!GEditor)
-	{
-		return FReply::Unhandled();
-	}
-
-	FAssetData AssetData;
-	TagCollectionPropertyHandle->GetValue(AssetData);
-
-	if (AssetData.IsValid())
-	{
-		GEditor->SyncBrowserToObject(AssetData);
-		return FReply::Handled();
-	}
-
-	return FReply::Unhandled();
-}
-
 FText SAvaTagCollectionPicker::GetTagCollectionTitleText() const
 {
 	FAssetData AssetData;
@@ -171,5 +128,4 @@ void SAvaTagCollectionPicker::CloseTagCollectionPicker()
 	TagCollectionPickerButton->SetIsOpen(false);
 }
 
-
 #undef LOCTEXT_NAMESPACE
diff --git a/Engine/Plugins/Experimental/Avalanche/Source/AvalancheTagEditor/Private/Customization/SAvaTagCollectionPicker.h b/Engine/Plugins/Experimental/Avalanche/Source/AvalancheTagEditor/Private/Customization/SAvaTagCollectionPicker.h
--- a/Engine/Plugins/Experimental/Avalanche/Source/AvalancheTagEditor/Private/Customization/SAvaTagCollectionPicker.h
+++ b/Engine/Plugins/Experimental/Avalanche/Source/AvalancheTagEditor/Private/Customization/SAvaTagCollectionPicker.h
@@ -6,7 +6,6 @@
 #include "Templates/SharedPointer.h"
 #include "Widgets/SCompoundWidget.h"
 
-class FReply;
 class FText;
 class IPropertyHandle;
 class SComboButton;
@@ -26,8 +25,6 @@ public:
 	void SetIsOpen(bool bInIsOpen);
 
 private:
-	FReply FindInContentBrowser();
-
 	FText GetTagCollectionTitleText() const;
 
 	FText GetTagCollectionTooltipText() const;
diff --git a/Engine/Plugins/Experimental/Avalanche/Source/AvalancheTagEditor/Private/Customization/SAvaTagPicker.cpp b/Engine/Plugins/Experimental/Avalanche/Source/AvalancheTagEditor/Private/Customization/SAvaTagPicker.cpp
--- a/Engine/Plugins/Experimental/Avalanche/Source/AvalancheTagEditor/Private/Customization/SAvaTagPicker.cpp
+++ b/Engine/Plugins/Experimental/Avalanche/Source/AvalancheTagEditor/Private/Customization/SAvaTagPicker.cpp
@@ -6,10 +6,15 @@
 #include "AvaTagEditorStyle.h"
 #include "AvaTagHandle.h"
 #include "DetailLayoutBuilder.h"
+#include "Menu/AvaTagCollectionPickerContextMenu.h"
 #include "SAvaTagCollectionPicker.h"
 #include "SAvaTagHandleEntry.h"
+#include "Styling/ToolBarStyle.h"
 #include "TagCustomizers/IAvaTagHandleCustomizer.h"
+#include "Widgets/Images/SImage.h"
+#include "Widgets/Input/SButton.h"
 #include "Widgets/Input/SComboButton.h"
+#include "Widgets/Input/SMenuAnchor.h"
 #include "Widgets/SBoxPanel.h"
 #include "Widgets/Text/STextBlock.h"
 #include "Widgets/Views/SListView.h"
@@ -23,10 +28,11 @@ void SAvaTagPicker::Construct(const FArguments& InArgs, const TSharedRef<IProper
 
 	TagCustomizer = InTagCustomizer;
 
-	TSharedPtr<IPropertyHandle> TagCollectionPropertyHandle = InTagCustomizer->GetTagCollectionHandle(InStructPropertyHandle);
-
+	TagCollectionPropertyHandle = InTagCustomizer->GetTagCollectionHandle(InStructPropertyHandle);
 	check(TagCollectionPropertyHandle.IsValid());
 
+	const FToolBarStyle& SlimToolbarStyle = FAppStyle::Get().GetWidgetStyle<FToolBarStyle>("SlimToolBar");
+
 	// Selection Mode set to single, yet preventing Items to be selected via OnIsSelectableOrNavigable, so that hover cue appears
 	TagListView = SNew(SListView<TSharedPtr<FAvaTagHandle>>)
 		.ListViewStyle(&FAppStyle::Get().GetWidgetStyle<FTableViewStyle>("SimpleListView"))
@@ -41,35 +47,60 @@ void SAvaTagPicker::Construct(const FArguments& InArgs, const TSharedRef<IProper
 
 	ChildSlot
 	[
-		SNew(SVerticalBox)
-		+ SVerticalBox::Slot()
-		.AutoHeight()
-		.Padding(0)
-		.HAlign(HAlign_Left)
-		.VAlign(VAlign_Center)
+		SNew(SHorizontalBox)
+		+ SHorizontalBox::Slot()
+		.FillWidth(1.f)
 		[
-			SAssignNew(TagCollectionPicker, SAvaTagCollectionPicker, TagCollectionPropertyHandle.ToSharedRef())
-			.OnTagCollectionChanged(this, &SAvaTagPicker::OnTagCollectionChanged)
+			SNew(SVerticalBox)
+			+ SVerticalBox::Slot()
+			.AutoHeight()
+			.Padding(0)
+			.HAlign(HAlign_Left)
+			.VAlign(VAlign_Center)
+			[
+				SAssignNew(TagCollectionPicker, SAvaTagCollectionPicker, TagCollectionPropertyHandle.ToSharedRef())
+				.OnTagCollectionChanged(this, &SAvaTagPicker::OnTagCollectionChanged)
+			]
+			+ SVerticalBox::Slot()
+			.AutoHeight()
+			.Padding(0)
+			.HAlign(HAlign_Fill)
+			.VAlign(VAlign_Center)
+			[
+				SAssignNew(TagComboButton, SComboButton)
+				.HasDownArrow(true)
+				.ContentPadding(FMargin(0, -1))
+				.OnMenuOpenChanged(this, &SAvaTagPicker::OnTagMenuOpenChanged)
+				.MenuContent()
+				[
+					TagListView.ToSharedRef()
+				]
+				.ButtonContent()
+				[
+					SNew(STextBlock)
+					.Text(this, &SAvaTagPicker::GetValueDisplayText)
+					.Font(IDetailLayoutBuilder::GetDetailFont())
+				]
+			]
 		]
-		+ SVerticalBox::Slot()
-		.AutoHeight()
-		.Padding(0)
-		.HAlign(HAlign_Fill)
+		+ SHorizontalBox::Slot()
+		.AutoWidth()
+		.HAlign(HAlign_Center)
 		.VAlign(VAlign_Center)
 		[
-			SAssignNew(TagComboButton, SComboButton)
-			.HasDownArrow(true)
-			.ContentPadding(FMargin(0, -1))
-			.OnMenuOpenChanged(this, &SAvaTagPicker::OnTagMenuOpenChanged)
-			.MenuContent()
+			SAssignNew(TagCollectionOptions, SMenuAnchor)
+			.Content()
 			[
-				TagListView.ToSharedRef()
-			]
-			.ButtonContent()
-			[
-				SNew(STextBlock)
-				.Text(this, &SAvaTagPicker::GetValueDisplayText)
-				.Font(IDetailLayoutBuilder::GetDetailFont())
+				SNew(SButton)
+				.OnClicked(this, &SAvaTagPicker::OpenContextMenu)
+				.ButtonStyle(FAppStyle::Get(), "SimpleButton")
+				.ToolTipText(LOCTEXT("TagCollectionOptionsToolTip", "Tag Collection Options"))
+				.ContentPadding(FMargin(4, 6))
+				[
+					SNew(SImage)
+					.Image(&SlimToolbarStyle.SettingsComboButton.DownArrowImage)
+					.ColorAndOpacity(FSlateColor::UseForeground())
+				]
 			]
 		]
 	];
@@ -85,6 +116,25 @@ void SAvaTagPicker::Tick(const FGeometry& InGeometry, const double InCurrentTime
 	}
 }
 
+FReply SAvaTagPicker::OnMouseButtonUp(const FGeometry& InGeometry, const FPointerEvent& InMouseEvent)
+{
+	if (InMouseEvent.GetEffectingButton() == EKeys::RightMouseButton)
+	{
+		return OpenContextMenu();
+	}
+	return FReply::Unhandled();
+}
+
+FReply SAvaTagPicker::OpenContextMenu()
+{
+	if (!TagCollectionOptions->IsOpen())
+	{
+		TagCollectionOptions->SetMenuContent(FAvaTagCollectionPickerContextMenu::Get().GenerateContextMenuWidget(TagCollectionPropertyHandle));
+		TagCollectionOptions->SetIsOpen(true);
+	}
+	return FReply::Handled();
+}
+
 const UAvaTagCollection* SAvaTagPicker::GetOrLoadTagCollection() const
 {
 	if (!StructPropertyHandle.IsValid())
diff --git a/Engine/Plugins/Experimental/Avalanche/Source/AvalancheTagEditor/Private/Customization/SAvaTagPicker.h b/Engine/Plugins/Experimental/Avalanche/Source/AvalancheTagEditor/Private/Customization/SAvaTagPicker.h
--- a/Engine/Plugins/Experimental/Avalanche/Source/AvalancheTagEditor/Private/Customization/SAvaTagPicker.h
+++ b/Engine/Plugins/Experimental/Avalanche/Source/AvalancheTagEditor/Private/Customization/SAvaTagPicker.h
@@ -12,6 +12,7 @@ class IPropertyHandle;
 class ITableRow;
 class SAvaTagCollectionPicker;
 class SComboButton;
+class SMenuAnchor;
 class STableViewBase;
 class UAvaTagCollection;
 struct FAvaTagHandle;
@@ -28,8 +29,11 @@ public:
 private:
 	//~ Begin SWidget
 	virtual void Tick(const FGeometry& InGeometry, const double InCurrentTime, const float InDeltaTime) override;
+	virtual FReply OnMouseButtonUp(const FGeometry& InGeometry, const FPointerEvent& InMouseEvent) override;
 	//~ End SWidget
 
+	FReply OpenContextMenu();
+
 	const UAvaTagCollection* GetOrLoadTagCollection() const;
 
 	void RefreshTagOptions();
@@ -49,8 +53,13 @@ private:
 	/** Handle to the struct property */
 	TSharedPtr<IPropertyHandle> StructPropertyHandle;
 
+	/** Handle to the Tag Collection property */
+	TSharedPtr<IPropertyHandle> TagCollectionPropertyHandle;
+
 	TSharedPtr<SAvaTagCollectionPicker> TagCollectionPicker;
 
+	TSharedPtr<SMenuAnchor> TagCollectionOptions;
+
 	TSharedPtr<SListView<TSharedPtr<FAvaTagHandle>>> TagListView;
 
 	TSharedPtr<IAvaTagHandleCustomizer> TagCustomizer;
