Motion Design Transition:
- Added Label for Layer Picker in Toolbar for Clarity
- Category so that the Transition Widget is at the bottom (EnterCond -> Tasks -> Transition)

diff --git a/Engine/Plugins/Experimental/Avalanche/Source/AvalancheTransitionEditor/Private/Customizations/AvaStateTreeStateCustomization.cpp b/Engine/Plugins/Experimental/Avalanche/Source/AvalancheTransitionEditor/Private/Customizations/AvaStateTreeStateCustomization.cpp
--- a/Engine/Plugins/Experimental/Avalanche/Source/AvalancheTransitionEditor/Private/Customizations/AvaStateTreeStateCustomization.cpp
+++ b/Engine/Plugins/Experimental/Avalanche/Source/AvalancheTransitionEditor/Private/Customizations/AvaStateTreeStateCustomization.cpp
@@ -31,9 +31,12 @@ void FAvaStateTreeStateCustomization::CustomizeDetails(IDetailLayoutBuilder& InD
 	InDetailBuilder.EditCategory(TEXT("Task")).InitiallyCollapsed(false);
 	InDetailBuilder.EditCategory(TEXT("Tasks")).InitiallyCollapsed(false);
 
-	IDetailCategoryBuilder& Category = InDetailBuilder.EditCategory(TEXT("General"), LOCTEXT("GeneralCategory", "General"), ECategoryPriority::Important);
-	Category.SetSortOrder(0);
-	Category.AddProperty(InDetailBuilder.GetProperty(GET_MEMBER_NAME_CHECKED(UStateTreeState, ColorRef)));
+	IDetailCategoryBuilder& GeneralCategory = InDetailBuilder.EditCategory(TEXT("General"), LOCTEXT("GeneralCategory", "General"), ECategoryPriority::Important);
+	GeneralCategory.SetSortOrder(0);
+	GeneralCategory.AddProperty(InDetailBuilder.GetProperty(GET_MEMBER_NAME_CHECKED(UStateTreeState, ColorRef)));
+
+	IDetailCategoryBuilder& TransitionCategory = InDetailBuilder.EditCategory(TEXT("Transition"), LOCTEXT("TransitionCategory", "Transition"), ECategoryPriority::Important);
+	TransitionCategory.SetSortOrder(10);
 
 	TSharedPtr<IPropertyHandleArray> TransitionArrayHandle;
 	if (TSharedPtr<IPropertyHandle> TransitionsHandle = InDetailBuilder.GetProperty(GET_MEMBER_NAME_CHECKED(UStateTreeState, Transitions)))
@@ -42,7 +45,7 @@ void FAvaStateTreeStateCustomization::CustomizeDetails(IDetailLayoutBuilder& InD
 	}
 
 	const FText TransitionType = LOCTEXT("TransitionType", "Transition Type");
-	Category.AddCustomRow(TransitionType)
+	TransitionCategory.AddCustomRow(TransitionType)
 		.NameContent()
 		[
 			SNew(STextBlock)
diff --git a/Engine/Plugins/Experimental/Avalanche/Source/AvalancheTransitionEditor/Private/Menu/AvaTransitionToolbar.cpp b/Engine/Plugins/Experimental/Avalanche/Source/AvalancheTransitionEditor/Private/Menu/AvaTransitionToolbar.cpp
--- a/Engine/Plugins/Experimental/Avalanche/Source/AvalancheTransitionEditor/Private/Menu/AvaTransitionToolbar.cpp
+++ b/Engine/Plugins/Experimental/Avalanche/Source/AvalancheTransitionEditor/Private/Menu/AvaTransitionToolbar.cpp
@@ -11,6 +11,8 @@
 #include "Toolkits/AssetEditorToolkit.h"
 #include "ViewModels/AvaTransitionEditorViewModel.h"
 #include "ViewModels/AvaTransitionViewModelSharedData.h"
+#include "Widgets/Layout/SBox.h"
+#include "Widgets/Text/STextBlock.h"
 
 #define LOCTEXT_NAMESPACE "AvaTransitionToolbar"
 
@@ -56,6 +58,17 @@ void FAvaTransitionToolbar::ExtendEditorToolbar(UToolMenu* InToolbarMenu)
 		LayerPicker->SetEnabled(!bReadOnly);
 
 		Section.AddSeparator(NAME_None);
+
+		TSharedRef<SWidget> TransitionLayerWidget = SNew(SBox)
+			.Padding(8.f, 0.f, 2.f, 0.f)
+			.VAlign(VAlign_Center)
+			[
+				SNew(STextBlock)
+				.Text(LOCTEXT("TransitionLayerText", "Layer"))
+				.Justification(ETextJustify::Right)
+			];
+
+		Section.AddEntry(FToolMenuEntry::InitWidget("TransitionLayerWidget", TransitionLayerWidget, FText::GetEmpty()));
 		Section.AddEntry(FToolMenuEntry::InitWidget(UE::AvaTransitionEditor::Private::TransitionLayerPickerName
 			, LayerPicker.ToSharedRef()
 			, LOCTEXT("TransitionLogicLabel", "Transition Logic Layer")));
