Motion Design Transition:
- Fix issue where GetDescription does not have access to instance data. A node context was made so that these type of functions can get access to the instance data

diff --git a/Engine/Plugins/Experimental/Avalanche/Source/Avalanche/Private/Transition/Conditions/AvaSceneContainsTagAttributeCondition.cpp b/Engine/Plugins/Experimental/Avalanche/Source/Avalanche/Private/Transition/Conditions/AvaSceneContainsTagAttributeCondition.cpp
--- a/Engine/Plugins/Experimental/Avalanche/Source/Avalanche/Private/Transition/Conditions/AvaSceneContainsTagAttributeCondition.cpp
+++ b/Engine/Plugins/Experimental/Avalanche/Source/Avalanche/Private/Transition/Conditions/AvaSceneContainsTagAttributeCondition.cpp
@@ -13,7 +13,7 @@
 
 #define LOCTEXT_NAMESPACE "AvaSceneContainsTagAttributeCondition"
 
-FText FAvaSceneContainsTagAttributeCondition::GenerateDescription() const
+FText FAvaSceneContainsTagAttributeCondition::GenerateDescription(const FAvaTransitionNodeContext& InContext) const
 {
 	FFormatNamedArguments Arguments;
 
diff --git a/Engine/Plugins/Experimental/Avalanche/Source/Avalanche/Private/Transition/Conditions/AvaTransitionRCControllerMatchCondition.cpp b/Engine/Plugins/Experimental/Avalanche/Source/Avalanche/Private/Transition/Conditions/AvaTransitionRCControllerMatchCondition.cpp
--- a/Engine/Plugins/Experimental/Avalanche/Source/Avalanche/Private/Transition/Conditions/AvaTransitionRCControllerMatchCondition.cpp
+++ b/Engine/Plugins/Experimental/Avalanche/Source/Avalanche/Private/Transition/Conditions/AvaTransitionRCControllerMatchCondition.cpp
@@ -18,7 +18,7 @@
 
 #define LOCTEXT_NAMESPACE "AvaTransitionRCControllerMatchCondition"
 
-FText FAvaTransitionRCControllerMatchCondition::GenerateDescription() const
+FText FAvaTransitionRCControllerMatchCondition::GenerateDescription(const FAvaTransitionNodeContext& InContext) const
 {
 	return FText::Format(LOCTEXT("ConditionDescription", "'{0}' is {1}")
 		, ControllerId.ToText()
diff --git a/Engine/Plugins/Experimental/Avalanche/Source/Avalanche/Private/Transition/Tasks/AvaSceneAddTagAttributeTask.cpp b/Engine/Plugins/Experimental/Avalanche/Source/Avalanche/Private/Transition/Tasks/AvaSceneAddTagAttributeTask.cpp
--- a/Engine/Plugins/Experimental/Avalanche/Source/Avalanche/Private/Transition/Tasks/AvaSceneAddTagAttributeTask.cpp
+++ b/Engine/Plugins/Experimental/Avalanche/Source/Avalanche/Private/Transition/Tasks/AvaSceneAddTagAttributeTask.cpp
@@ -6,7 +6,7 @@
 
 #define LOCTEXT_NAMESPACE "AvaSceneAddTagAttributeTask"
 
-FText FAvaSceneAddTagAttributeTask::GenerateDescription() const
+FText FAvaSceneAddTagAttributeTask::GenerateDescription(const FAvaTransitionNodeContext& InContext) const
 {
 	return FText::Format(LOCTEXT("TaskDescription", "Add '{0}' tag attribute to this scene"), FText::FromName(TagAttribute.ToName()));
 }
diff --git a/Engine/Plugins/Experimental/Avalanche/Source/Avalanche/Private/Transition/Tasks/AvaSceneRemoveTagAttributeTask.cpp b/Engine/Plugins/Experimental/Avalanche/Source/Avalanche/Private/Transition/Tasks/AvaSceneRemoveTagAttributeTask.cpp
--- a/Engine/Plugins/Experimental/Avalanche/Source/Avalanche/Private/Transition/Tasks/AvaSceneRemoveTagAttributeTask.cpp
+++ b/Engine/Plugins/Experimental/Avalanche/Source/Avalanche/Private/Transition/Tasks/AvaSceneRemoveTagAttributeTask.cpp
@@ -6,7 +6,7 @@
 
 #define LOCTEXT_NAMESPACE "AvaSceneRemoveTagAttributeTask"
 
-FText FAvaSceneRemoveTagAttributeTask::GenerateDescription() const
+FText FAvaSceneRemoveTagAttributeTask::GenerateDescription(const FAvaTransitionNodeContext& InContext) const
 {
 	return FText::Format(LOCTEXT("TaskDescription", "Remove '{0}' tag attribute from this scene"), FText::FromName(TagAttribute.ToName()));
 }
diff --git a/Engine/Plugins/Experimental/Avalanche/Source/Avalanche/Public/Transition/Conditions/AvaSceneContainsTagAttributeCondition.h b/Engine/Plugins/Experimental/Avalanche/Source/Avalanche/Public/Transition/Conditions/AvaSceneContainsTagAttributeCondition.h
--- a/Engine/Plugins/Experimental/Avalanche/Source/Avalanche/Public/Transition/Conditions/AvaSceneContainsTagAttributeCondition.h
+++ b/Engine/Plugins/Experimental/Avalanche/Source/Avalanche/Public/Transition/Conditions/AvaSceneContainsTagAttributeCondition.h
@@ -25,7 +25,7 @@ struct AVALANCHE_API FAvaSceneContainsTagAttributeCondition : public FAvaTransit
 	using FInstanceDataType = FAvaSceneContainsAttributeConditionInstanceData;
 
 	//~ Begin FAvaTransitionCondition
-	virtual FText GenerateDescription() const override;
+	virtual FText GenerateDescription(const FAvaTransitionNodeContext& InContext) const override;
 	//~ End FAvaTransitionCondition
 
 	//~ Begin FStateTreeNodeBase
diff --git a/Engine/Plugins/Experimental/Avalanche/Source/Avalanche/Public/Transition/Conditions/AvaTransitionRCControllerMatchCondition.h b/Engine/Plugins/Experimental/Avalanche/Source/Avalanche/Public/Transition/Conditions/AvaTransitionRCControllerMatchCondition.h
--- a/Engine/Plugins/Experimental/Avalanche/Source/Avalanche/Public/Transition/Conditions/AvaTransitionRCControllerMatchCondition.h
+++ b/Engine/Plugins/Experimental/Avalanche/Source/Avalanche/Public/Transition/Conditions/AvaTransitionRCControllerMatchCondition.h
@@ -25,7 +25,7 @@ struct AVALANCHE_API FAvaTransitionRCControllerMatchCondition : public FAvaTrans
 	using FInstanceDataType = FAvaTransitionRCControllerMatchConditionInstanceData;
 
 	//~ Begin FAvaTransitionCondition
-	virtual FText GenerateDescription() const override;
+	virtual FText GenerateDescription(const FAvaTransitionNodeContext& InContext) const override;
 	//~ End FAvaTransitionCondition
 
 	//~ Begin FStateTreeNodeBase
diff --git a/Engine/Plugins/Experimental/Avalanche/Source/Avalanche/Public/Transition/Tasks/AvaSceneAddTagAttributeTask.h b/Engine/Plugins/Experimental/Avalanche/Source/Avalanche/Public/Transition/Tasks/AvaSceneAddTagAttributeTask.h
--- a/Engine/Plugins/Experimental/Avalanche/Source/Avalanche/Public/Transition/Tasks/AvaSceneAddTagAttributeTask.h
+++ b/Engine/Plugins/Experimental/Avalanche/Source/Avalanche/Public/Transition/Tasks/AvaSceneAddTagAttributeTask.h
@@ -11,7 +11,7 @@ struct AVALANCHE_API FAvaSceneAddTagAttributeTask : public FAvaSceneTask
 	GENERATED_BODY()
 
 	//~ Begin FAvaTransitionTask
-	virtual FText GenerateDescription() const override;
+	virtual FText GenerateDescription(const FAvaTransitionNodeContext& InContext) const override;
 	//~ End FAvaTransitionTask
 
 	//~ Begin FStateTreeTaskBase
diff --git a/Engine/Plugins/Experimental/Avalanche/Source/Avalanche/Public/Transition/Tasks/AvaSceneRemoveTagAttributeTask.h b/Engine/Plugins/Experimental/Avalanche/Source/Avalanche/Public/Transition/Tasks/AvaSceneRemoveTagAttributeTask.h
--- a/Engine/Plugins/Experimental/Avalanche/Source/Avalanche/Public/Transition/Tasks/AvaSceneRemoveTagAttributeTask.h
+++ b/Engine/Plugins/Experimental/Avalanche/Source/Avalanche/Public/Transition/Tasks/AvaSceneRemoveTagAttributeTask.h
@@ -11,7 +11,7 @@ struct AVALANCHE_API FAvaSceneRemoveTagAttributeTask : public FAvaSceneTask
 	GENERATED_BODY()
 
 	//~ Begin FAvaTransitionTask
-	virtual FText GenerateDescription() const override;
+	virtual FText GenerateDescription(const FAvaTransitionNodeContext& InContext) const override;
 	//~ End FAvaTransitionTask
 
 	//~ Begin FStateTreeTaskBase
diff --git a/Engine/Plugins/Experimental/Avalanche/Source/AvalancheSequence/Private/Transition/AvaTransitionPlaySequenceTask.cpp b/Engine/Plugins/Experimental/Avalanche/Source/AvalancheSequence/Private/Transition/AvaTransitionPlaySequenceTask.cpp
--- a/Engine/Plugins/Experimental/Avalanche/Source/AvalancheSequence/Private/Transition/AvaTransitionPlaySequenceTask.cpp
+++ b/Engine/Plugins/Experimental/Avalanche/Source/AvalancheSequence/Private/Transition/AvaTransitionPlaySequenceTask.cpp
@@ -10,7 +10,7 @@
 
 #define LOCTEXT_NAMESPACE "AvaTransitionPlaySequenceTask"
 
-FText FAvaTransitionPlaySequenceTask::GenerateDescription() const
+FText FAvaTransitionPlaySequenceTask::GenerateDescription(const FAvaTransitionNodeContext& InContext) const
 {
 	FText DescriptionFormat = LOCTEXT("TaskDescription", "Play {SequenceQueryText} {AdditionalArgs}");
 
diff --git a/Engine/Plugins/Experimental/Avalanche/Source/AvalancheSequence/Public/Transition/AvaTransitionPlaySequenceTask.h b/Engine/Plugins/Experimental/Avalanche/Source/AvalancheSequence/Public/Transition/AvaTransitionPlaySequenceTask.h
--- a/Engine/Plugins/Experimental/Avalanche/Source/AvalancheSequence/Public/Transition/AvaTransitionPlaySequenceTask.h
+++ b/Engine/Plugins/Experimental/Avalanche/Source/AvalancheSequence/Public/Transition/AvaTransitionPlaySequenceTask.h
@@ -12,7 +12,7 @@ struct AVALANCHESEQUENCE_API FAvaTransitionPlaySequenceTask : public FAvaTransit
 	GENERATED_BODY()
 
 	//~ Begin FAvaTransitionTask
-	virtual FText GenerateDescription() const override;
+	virtual FText GenerateDescription(const FAvaTransitionNodeContext& InContext) const override;
 	//~ End FAvaTransitionTask
 
 	//~ Begin FAvaTransitionSequenceTaskBase
diff --git a/Engine/Plugins/Experimental/Avalanche/Source/AvalancheTransition/Private/Conditions/AvaTransitionLayerMatchCondition.cpp b/Engine/Plugins/Experimental/Avalanche/Source/AvalancheTransition/Private/Conditions/AvaTransitionLayerMatchCondition.cpp
--- a/Engine/Plugins/Experimental/Avalanche/Source/AvalancheTransition/Private/Conditions/AvaTransitionLayerMatchCondition.cpp
+++ b/Engine/Plugins/Experimental/Avalanche/Source/AvalancheTransition/Private/Conditions/AvaTransitionLayerMatchCondition.cpp
@@ -7,7 +7,7 @@
 
 #define LOCTEXT_NAMESPACE "AvaTransitionLayerMatchCondition"
 
-FText FAvaTransitionLayerMatchCondition::GenerateDescription() const
+FText FAvaTransitionLayerMatchCondition::GenerateDescription(const FAvaTransitionNodeContext& InContext) const
 {
 	return FText::Format(LOCTEXT("ConditionDescription", "scenes transitioning in {0}"), GetLayerQueryText());
 }
diff --git a/Engine/Plugins/Experimental/Avalanche/Source/AvalancheTransition/Private/Conditions/AvaTransitionSceneMatchCondition.cpp b/Engine/Plugins/Experimental/Avalanche/Source/AvalancheTransition/Private/Conditions/AvaTransitionSceneMatchCondition.cpp
--- a/Engine/Plugins/Experimental/Avalanche/Source/AvalancheTransition/Private/Conditions/AvaTransitionSceneMatchCondition.cpp
+++ b/Engine/Plugins/Experimental/Avalanche/Source/AvalancheTransition/Private/Conditions/AvaTransitionSceneMatchCondition.cpp
@@ -8,7 +8,7 @@
 
 #define LOCTEXT_NAMESPACE "AvaTransitionSceneMatchCondition"
 
-FText FAvaTransitionSceneMatchCondition::GenerateDescription() const
+FText FAvaTransitionSceneMatchCondition::GenerateDescription(const FAvaTransitionNodeContext& InContext) const
 {
 	return FText::Format(LOCTEXT("ConditionDescription", "{0} scene in {1}")
 		, UEnum::GetDisplayValueAsText(SceneComparisonType).ToLower()
diff --git a/Engine/Plugins/Experimental/Avalanche/Source/AvalancheTransition/Private/Conditions/AvaTransitionStateMatchCondition.cpp b/Engine/Plugins/Experimental/Avalanche/Source/AvalancheTransition/Private/Conditions/AvaTransitionStateMatchCondition.cpp
--- a/Engine/Plugins/Experimental/Avalanche/Source/AvalancheTransition/Private/Conditions/AvaTransitionStateMatchCondition.cpp
+++ b/Engine/Plugins/Experimental/Avalanche/Source/AvalancheTransition/Private/Conditions/AvaTransitionStateMatchCondition.cpp
@@ -5,7 +5,7 @@
 
 #define LOCTEXT_NAMESPACE "AvaTransitionSceneMatchCondition"
 
-FText FAvaTransitionStateMatchCondition::GenerateDescription() const
+FText FAvaTransitionStateMatchCondition::GenerateDescription(const FAvaTransitionNodeContext& InContext) const
 {
 	return FText::Format(LOCTEXT("ConditionDescription", "{0} scene in {1}")
 		, UEnum::GetDisplayValueAsText(TransitionState).ToLower()
diff --git a/Engine/Plugins/Experimental/Avalanche/Source/AvalancheTransition/Private/Conditions/AvaTransitionTypeMatchCondition.cpp b/Engine/Plugins/Experimental/Avalanche/Source/AvalancheTransition/Private/Conditions/AvaTransitionTypeMatchCondition.cpp
--- a/Engine/Plugins/Experimental/Avalanche/Source/AvalancheTransition/Private/Conditions/AvaTransitionTypeMatchCondition.cpp
+++ b/Engine/Plugins/Experimental/Avalanche/Source/AvalancheTransition/Private/Conditions/AvaTransitionTypeMatchCondition.cpp
@@ -6,7 +6,7 @@
 
 #define LOCTEXT_NAMESPACE "AvaTransitionTypeMatchCondition"
 
-FText FAvaTransitionTypeMatchCondition::GenerateDescription() const
+FText FAvaTransitionTypeMatchCondition::GenerateDescription(const FAvaTransitionNodeContext& InContext) const
 {
 	return FText::Format(LOCTEXT("ConditionDescription", "transitioning {0}")
 		, UEnum::GetDisplayValueAsText(TransitionType).ToLower());
diff --git a/Engine/Plugins/Experimental/Avalanche/Source/AvalancheTransition/Private/Tasks/AvaTransitionDelayTask.cpp b/Engine/Plugins/Experimental/Avalanche/Source/AvalancheTransition/Private/Tasks/AvaTransitionDelayTask.cpp
--- a/Engine/Plugins/Experimental/Avalanche/Source/AvalancheTransition/Private/Tasks/AvaTransitionDelayTask.cpp
+++ b/Engine/Plugins/Experimental/Avalanche/Source/AvalancheTransition/Private/Tasks/AvaTransitionDelayTask.cpp
@@ -5,7 +5,7 @@
 
 #define LOCTEXT_NAMESPACE "AvaTransitionDelayTask"
 
-FText FAvaTransitionDelayTask::GenerateDescription() const
+FText FAvaTransitionDelayTask::GenerateDescription(const FAvaTransitionNodeContext& InContext) const
 {
 	return FText::Format(LOCTEXT("TaskDescription", "Delay {0} seconds"), FText::AsNumber(Duration));
 }
diff --git a/Engine/Plugins/Experimental/Avalanche/Source/AvalancheTransition/Private/Tasks/AvaTransitionWaitForLayerTask.cpp b/Engine/Plugins/Experimental/Avalanche/Source/AvalancheTransition/Private/Tasks/AvaTransitionWaitForLayerTask.cpp
--- a/Engine/Plugins/Experimental/Avalanche/Source/AvalancheTransition/Private/Tasks/AvaTransitionWaitForLayerTask.cpp
+++ b/Engine/Plugins/Experimental/Avalanche/Source/AvalancheTransition/Private/Tasks/AvaTransitionWaitForLayerTask.cpp
@@ -6,7 +6,7 @@
 
 #define LOCTEXT_NAMESPACE "AvaTransitionWaitForLayerTask"
 
-FText FAvaTransitionWaitForLayerTask::GenerateDescription() const
+FText FAvaTransitionWaitForLayerTask::GenerateDescription(const FAvaTransitionNodeContext& InContext) const
 {
 	return FText::Format(LOCTEXT("TaskDescription", "Wait for others in {0} to finish")
 		, GetLayerQueryText());
diff --git a/Engine/Plugins/Experimental/Avalanche/Source/AvalancheTransition/Public/AvaTransitionNode.h b/Engine/Plugins/Experimental/Avalanche/Source/AvalancheTransition/Public/AvaTransitionNode.h
--- a/Engine/Plugins/Experimental/Avalanche/Source/AvalancheTransition/Public/AvaTransitionNode.h
+++ b/Engine/Plugins/Experimental/Avalanche/Source/AvalancheTransition/Public/AvaTransitionNode.h
@@ -4,6 +4,7 @@
 
 #include "StateTreeExecutionTypes.h"
 
+class FAvaTransitionNodeContext;
 class UAvaTransitionSubsystem;
 struct FAvaTransitionContext;
 struct FStateTreeLinker;
@@ -14,7 +15,7 @@ public:
 	AVALANCHETRANSITION_API bool LinkNode(FStateTreeLinker& InLinker);
 
 	/** Generates a more user-friendly dynamic description of the Node based on its settings */
-	virtual FText GenerateDescription() const
+	virtual FText GenerateDescription(const FAvaTransitionNodeContext& InContext) const
 	{
 		return FText::GetEmpty();
 	}
diff --git a/Engine/Plugins/Experimental/Avalanche/Source/AvalancheTransition/Public/AvaTransitionNodeContext.h b/Engine/Plugins/Experimental/Avalanche/Source/AvalancheTransition/Public/AvaTransitionNodeContext.h
new file mode 100644
--- /dev/null
+++ b/Engine/Plugins/Experimental/Avalanche/Source/AvalancheTransition/Public/AvaTransitionNodeContext.h
@@ -0,0 +1,25 @@
+// Copyright Epic Games, Inc. All Rights Reserved.
+
+#pragma once
+
+#include "StateTreeExecutionTypes.h"
+#include "StateTreeNodeBase.h"
+
+class FAvaTransitionNodeContext
+{
+public:
+	explicit FAvaTransitionNodeContext(const FStateTreeDataView& InInstanceDataView)
+		: InstanceDataView(InInstanceDataView)
+	{
+	}
+
+	template<typename InNodeType
+		UE_REQUIRES(std::is_base_of_v<FStateTreeNodeBase, InNodeType>)>
+	const typename InNodeType::FInstanceDataType* GetInstanceData(const InNodeType* InNode) const
+	{
+		return InstanceDataView.GetPtr<typename InNodeType::FInstanceDataType>();
+	}
+
+private:
+	const FStateTreeDataView InstanceDataView;
+};
diff --git a/Engine/Plugins/Experimental/Avalanche/Source/AvalancheTransition/Public/Conditions/AvaTransitionLayerMatchCondition.h b/Engine/Plugins/Experimental/Avalanche/Source/AvalancheTransition/Public/Conditions/AvaTransitionLayerMatchCondition.h
--- a/Engine/Plugins/Experimental/Avalanche/Source/AvalancheTransition/Public/Conditions/AvaTransitionLayerMatchCondition.h
+++ b/Engine/Plugins/Experimental/Avalanche/Source/AvalancheTransition/Public/Conditions/AvaTransitionLayerMatchCondition.h
@@ -19,7 +19,7 @@ struct AVALANCHETRANSITION_API FAvaTransitionLayerMatchCondition : public FAvaTr
 	using FInstanceDataType = FAvaTransitionLayerMatchConditionInstanceData;
 
 	//~ Begin FAvaTransitionCondition
-	virtual FText GenerateDescription() const override;
+	virtual FText GenerateDescription(const FAvaTransitionNodeContext& InContext) const override;
 	//~ End FAvaTransitionCondition
 
 	//~ Begin FStateTreeNodeBase
diff --git a/Engine/Plugins/Experimental/Avalanche/Source/AvalancheTransition/Public/Conditions/AvaTransitionSceneMatchCondition.h b/Engine/Plugins/Experimental/Avalanche/Source/AvalancheTransition/Public/Conditions/AvaTransitionSceneMatchCondition.h
--- a/Engine/Plugins/Experimental/Avalanche/Source/AvalancheTransition/Public/Conditions/AvaTransitionSceneMatchCondition.h
+++ b/Engine/Plugins/Experimental/Avalanche/Source/AvalancheTransition/Public/Conditions/AvaTransitionSceneMatchCondition.h
@@ -20,7 +20,7 @@ struct AVALANCHETRANSITION_API FAvaTransitionSceneMatchCondition : public FAvaTr
 	using FInstanceDataType = FAvaTransitionSceneMatchConditionInstanceData;
 
 	//~ Begin FAvaTransitionCondition
-	virtual FText GenerateDescription() const override;
+	virtual FText GenerateDescription(const FAvaTransitionNodeContext& InContext) const override;
 	//~ End FAvaTransitionCondition
 
 	//~ Begin FStateTreeNodeBase
diff --git a/Engine/Plugins/Experimental/Avalanche/Source/AvalancheTransition/Public/Conditions/AvaTransitionStateMatchCondition.h b/Engine/Plugins/Experimental/Avalanche/Source/AvalancheTransition/Public/Conditions/AvaTransitionStateMatchCondition.h
--- a/Engine/Plugins/Experimental/Avalanche/Source/AvalancheTransition/Public/Conditions/AvaTransitionStateMatchCondition.h
+++ b/Engine/Plugins/Experimental/Avalanche/Source/AvalancheTransition/Public/Conditions/AvaTransitionStateMatchCondition.h
@@ -20,7 +20,7 @@ struct AVALANCHETRANSITION_API FAvaTransitionStateMatchCondition : public FAvaTr
 	using FInstanceDataType = FAvaTransitionStateMatchConditionInstanceData;
 
 	//~ Begin FAvaTransitionCondition
-	virtual FText GenerateDescription() const override;
+	virtual FText GenerateDescription(const FAvaTransitionNodeContext& InContext) const override;
 	//~ End FAvaTransitionCondition
 
 	//~ Begin FStateTreeNodeBase
diff --git a/Engine/Plugins/Experimental/Avalanche/Source/AvalancheTransition/Public/Conditions/AvaTransitionTypeMatchCondition.h b/Engine/Plugins/Experimental/Avalanche/Source/AvalancheTransition/Public/Conditions/AvaTransitionTypeMatchCondition.h
--- a/Engine/Plugins/Experimental/Avalanche/Source/AvalancheTransition/Public/Conditions/AvaTransitionTypeMatchCondition.h
+++ b/Engine/Plugins/Experimental/Avalanche/Source/AvalancheTransition/Public/Conditions/AvaTransitionTypeMatchCondition.h
@@ -22,7 +22,7 @@ struct AVALANCHETRANSITION_API FAvaTransitionTypeMatchCondition : public FAvaTra
 	FAvaTransitionTypeMatchCondition() = default;
 
 	//~ Begin FAvaTransitionCondition
-	virtual FText GenerateDescription() const override;
+	virtual FText GenerateDescription(const FAvaTransitionNodeContext& InContext) const override;
 	//~ End FAvaTransitionCondition
 
 	//~ Begin FStateTreeNodeBase
diff --git a/Engine/Plugins/Experimental/Avalanche/Source/AvalancheTransition/Public/Tasks/AvaTransitionDelayTask.h b/Engine/Plugins/Experimental/Avalanche/Source/AvalancheTransition/Public/Tasks/AvaTransitionDelayTask.h
--- a/Engine/Plugins/Experimental/Avalanche/Source/AvalancheTransition/Public/Tasks/AvaTransitionDelayTask.h
+++ b/Engine/Plugins/Experimental/Avalanche/Source/AvalancheTransition/Public/Tasks/AvaTransitionDelayTask.h
@@ -22,7 +22,7 @@ struct AVALANCHETRANSITION_API FAvaTransitionDelayTask : public FAvaTransitionTa
 	using FInstanceDataType = FAvaTransitionDelayTaskInstanceData;
 
 	//~ Begin FAvaTransitionTask
-	virtual FText GenerateDescription() const override;
+	virtual FText GenerateDescription(const FAvaTransitionNodeContext& InContext) const override;
 	//~ End FAvaTransitionTask
 
 	//~ Begin FStateTreeNodeBase
diff --git a/Engine/Plugins/Experimental/Avalanche/Source/AvalancheTransition/Public/Tasks/AvaTransitionWaitForLayerTask.h b/Engine/Plugins/Experimental/Avalanche/Source/AvalancheTransition/Public/Tasks/AvaTransitionWaitForLayerTask.h
--- a/Engine/Plugins/Experimental/Avalanche/Source/AvalancheTransition/Public/Tasks/AvaTransitionWaitForLayerTask.h
+++ b/Engine/Plugins/Experimental/Avalanche/Source/AvalancheTransition/Public/Tasks/AvaTransitionWaitForLayerTask.h
@@ -19,7 +19,7 @@ struct AVALANCHETRANSITION_API FAvaTransitionWaitForLayerTask : public FAvaTrans
 	using FInstanceDataType = FAvaTransitionWaitForLayerTaskInstanceData;
 
 	//~ Begin FAvaTransitionTask
-	virtual FText GenerateDescription() const override;
+	virtual FText GenerateDescription(const FAvaTransitionNodeContext& InContext) const override;
 	//~ End FAvaTransitionTask
 
 	//~ Begin FStateTreeNodeBase
diff --git a/Engine/Plugins/Experimental/Avalanche/Source/AvalancheTransitionEditor/Private/ViewModels/Condition/AvaTransitionConditionContainerViewModel.cpp b/Engine/Plugins/Experimental/Avalanche/Source/AvalancheTransitionEditor/Private/ViewModels/Condition/AvaTransitionConditionContainerViewModel.cpp
--- a/Engine/Plugins/Experimental/Avalanche/Source/AvalancheTransitionEditor/Private/ViewModels/Condition/AvaTransitionConditionContainerViewModel.cpp
+++ b/Engine/Plugins/Experimental/Avalanche/Source/AvalancheTransitionEditor/Private/ViewModels/Condition/AvaTransitionConditionContainerViewModel.cpp
@@ -2,9 +2,10 @@
 
 #include "AvaTransitionConditionContainerViewModel.h"
 #include "AvaTransitionConditionViewModel.h"
+#include "AvaTransitionNodeContext.h"
+#include "Conditions/AvaTransitionCondition.h"
 #include "StateTreeEditorStyle.h"
 #include "StateTreeState.h"
-#include "Conditions/AvaTransitionCondition.h"
 #include "ViewModels/AvaTransitionViewModelUtils.h"
 #include "Widgets/Images/SImage.h"
 #include "Widgets/Layout/SBox.h"
@@ -28,7 +29,7 @@ namespace UE::AvaTransitionEditor::Private
 		FText ConditionDescription;
 		if (const FAvaTransitionCondition* Condition = InEditorNode.Node.GetPtr<FAvaTransitionCondition>())
 		{
-			ConditionDescription = Condition->GenerateDescription(); 	
+			ConditionDescription = Condition->GenerateDescription(FAvaTransitionNodeContext(InEditorNode.GetInstance()));
 		}
 
 		const UScriptStruct* Struct = InEditorNode.Node.GetScriptStruct();
diff --git a/Engine/Plugins/Experimental/Avalanche/Source/AvalancheTransitionEditor/Private/ViewModels/Task/AvaTransitionTaskViewModel.cpp b/Engine/Plugins/Experimental/Avalanche/Source/AvalancheTransitionEditor/Private/ViewModels/Task/AvaTransitionTaskViewModel.cpp
--- a/Engine/Plugins/Experimental/Avalanche/Source/AvalancheTransitionEditor/Private/ViewModels/Task/AvaTransitionTaskViewModel.cpp
+++ b/Engine/Plugins/Experimental/Avalanche/Source/AvalancheTransitionEditor/Private/ViewModels/Task/AvaTransitionTaskViewModel.cpp
@@ -1,6 +1,7 @@
 // Copyright Epic Games, Inc. All Rights Reserved.
 
 #include "AvaTransitionTaskViewModel.h"
+#include "AvaTransitionNodeContext.h"
 #include "AvaTransitionTreeEditorData.h"
 #include "StateTreeEditorData.h"
 #include "StateTreeEditorStyle.h"
@@ -31,8 +32,8 @@ void FAvaTransitionTaskViewModel::UpdateTaskDescription()
 {
 	TaskDescription = FText::GetEmpty();
 
-	const FInstancedStruct* Node = GetNode();
-	if (!Node)
+	const FStateTreeEditorNode* EditorNode = GetEditorNode();
+	if (!EditorNode)
 	{
 		return;
 	}
@@ -40,7 +41,7 @@ void FAvaTransitionTaskViewModel::UpdateTaskDescription()
 	// Advanced Mode: Show Task Descriptions from the Node Name (Raw) if name not none
 	if (GetSharedData()->GetEditorMode() == EAvaTransitionEditorMode::Advanced)
 	{
-		const FStateTreeNodeBase* Task = Node->GetPtr<FStateTreeNodeBase>();
+		const FStateTreeNodeBase* Task = EditorNode->Node.GetPtr<FStateTreeNodeBase>();
 		if (Task && Task->Name != NAME_None)
 		{
 			TaskDescription = FText::FromName(Task->Name);
@@ -48,12 +49,12 @@ void FAvaTransitionTaskViewModel::UpdateTaskDescription()
 		}
 	}
 
-	if (const FAvaTransitionTask* Task = Node->GetPtr<FAvaTransitionTask>())
+	if (const FAvaTransitionTask* Task = EditorNode->Node.GetPtr<FAvaTransitionTask>())
 	{
-		TaskDescription = Task->GenerateDescription();
+		TaskDescription = Task->GenerateDescription(FAvaTransitionNodeContext(EditorNode->GetInstance()));
 	}
 
-	const UScriptStruct* Struct = Node->GetScriptStruct();
+	const UScriptStruct* Struct = EditorNode->Node.GetScriptStruct();
 	if (Struct && TaskDescription.IsEmpty())
 	{
 		TaskDescription = Struct->GetDisplayNameText();
