Avalanche Transition:
- Added extra condition for "doesn't contain", as the "Invert Condition" is confusing to have
- Fixed extra space after "if" enter condition nodes

diff --git a/Engine/Plugins/Experimental/Avalanche/Config/DefaultAvalanche.ini b/Engine/Plugins/Experimental/Avalanche/Config/DefaultAvalanche.ini
--- a/Engine/Plugins/Experimental/Avalanche/Config/DefaultAvalanche.ini
+++ b/Engine/Plugins/Experimental/Avalanche/Config/DefaultAvalanche.ini
@@ -172,5 +172,7 @@
 +ClassRedirects=(OldName="/Script/AvalancheMaskEditor.AvalancheMaskEditorMode",NewName="/Script/AvalancheMaskEditor.AvaMaskEditorMode")
 +FunctionRedirects=(OldName="/Script/AvalancheTest.AvaTestBlueprintFunctionLibrary.ExportAvaBlueprintsToWorld",NewName="/Script/AvalancheTest.AvaTestBlueprintFunctionLibrary.ExportMotionDesignBlueprintsToWorld")
 
++StructRedirects=(OldName="/Script/Avalanche.AvaSceneContainsAttributeConditionInstanceData",NewName="/Script/Avalanche.AvaSceneContainsTagAttributeConditionInstanceData")
+
 [/Script/Avalanche.AvaAnimation]
 DefaultCompletionMode=KeepState
diff --git a/Engine/Plugins/Experimental/Avalanche/Source/Avalanche/Private/Transition/Conditions/AvaSceneContainsTagAttributeCondition.cpp b/Engine/Plugins/Experimental/Avalanche/Source/Avalanche/Private/Transition/Conditions/AvaSceneContainsTagAttributeCondition.cpp
--- a/Engine/Plugins/Experimental/Avalanche/Source/Avalanche/Private/Transition/Conditions/AvaSceneContainsTagAttributeCondition.cpp
+++ b/Engine/Plugins/Experimental/Avalanche/Source/Avalanche/Private/Transition/Conditions/AvaSceneContainsTagAttributeCondition.cpp
@@ -11,50 +11,49 @@
 #include "StateTreeExecutionContext.h"
 #include "StateTreeLinker.h"
 
-#define LOCTEXT_NAMESPACE "AvaSceneContainsTagAttributeCondition"
+#define LOCTEXT_NAMESPACE "AvaSceneContainsTagAttributeConditionBase"
 
-FText FAvaSceneContainsTagAttributeCondition::GenerateDescription(const FAvaTransitionNodeContext& InContext) const
+FText FAvaSceneContainsTagAttributeConditionBase::GenerateDescription(const FAvaTransitionNodeContext& InContext) const
 {
 	FFormatNamedArguments Arguments;
 
 	switch (SceneType)
 	{
 	case EAvaTransitionSceneType::This:
+		Arguments.Add(TEXT("IndefinitePronoun"), FText::GetEmpty());
+
 		Arguments.Add(TEXT("Scene"), LOCTEXT("ThisScene", "this scene"));
+
+		Arguments.Add(TEXT("Contains"), bInvertCondition ? LOCTEXT("ThisDoesntContain", "does not contain") : LOCTEXT("ThisContains", "contains"));
 		break;
 
 	case EAvaTransitionSceneType::Other:
+		Arguments.Add(TEXT("IndefinitePronoun"), bInvertCondition ? LOCTEXT("NoScene", "no ") : LOCTEXT("AnyScene", "a "));
+
 		Arguments.Add(TEXT("Scene"), FText::Format(LOCTEXT("OtherScene", "scene in {0}"), FAvaTransitionLayerUtils::GetLayerQueryText(LayerType, *SpecificLayers.ToString())));
-		break;
-	}
 
-	if (bInvert)
-	{
-		Arguments.Add(TEXT("Contains"), LOCTEXT("ContainsInvert", "does NOT contain"));
-	}
-	else
-	{
-		Arguments.Add(TEXT("Contains"), LOCTEXT("Contains", "contains"));
+		Arguments.Add(TEXT("Contains"), LOCTEXT("OtherSceneContains", "contains"));
+		break;
 	}
 
 	Arguments.Add(TEXT("TagAttribute"), FText::FromName(TagAttribute.ToName()));
 
-	return FText::Format(LOCTEXT("ConditionDescription", "{Scene} {Contains} tag attribute '{TagAttribute}'"), Arguments);
+	return FText::Format(LOCTEXT("ConditionDescription", "{IndefinitePronoun}{Scene} {Contains} tag attribute '{TagAttribute}'"), Arguments);
 }
 
-bool FAvaSceneContainsTagAttributeCondition::Link(FStateTreeLinker& InLinker)
+bool FAvaSceneContainsTagAttributeConditionBase::Link(FStateTreeLinker& InLinker)
 {
 	FAvaTransitionCondition::Link(InLinker);
 	InLinker.LinkExternalData(SceneSubsystemHandle);
 	return true;
 }
 
-bool FAvaSceneContainsTagAttributeCondition::TestCondition(FStateTreeExecutionContext& InContext) const
+bool FAvaSceneContainsTagAttributeConditionBase::TestCondition(FStateTreeExecutionContext& InContext) const
 {
-	return ContainsTagAttribute(InContext) ^ bInvert;
+	return ContainsTagAttribute(InContext) ^ bInvertCondition;
 }
 
-bool FAvaSceneContainsTagAttributeCondition::ContainsTagAttribute(FStateTreeExecutionContext& InContext) const
+bool FAvaSceneContainsTagAttributeConditionBase::ContainsTagAttribute(FStateTreeExecutionContext& InContext) const
 {
 	TArray<const FAvaTransitionScene*> TransitionScenes = GetTransitionScenes(InContext);
 	if (TransitionScenes.IsEmpty())
@@ -87,7 +86,7 @@ bool FAvaSceneContainsTagAttributeCondition::ContainsTagAttribute(FStateTreeExec
 	return false;
 }
 
-TArray<const FAvaTransitionScene*> FAvaSceneContainsTagAttributeCondition::GetTransitionScenes(FStateTreeExecutionContext& InContext) const
+TArray<const FAvaTransitionScene*> FAvaSceneContainsTagAttributeConditionBase::GetTransitionScenes(FStateTreeExecutionContext& InContext) const
 {
 	const FAvaTransitionContext& TransitionContext = InContext.GetExternalData(TransitionContextHandle);
 
diff --git a/Engine/Plugins/Experimental/Avalanche/Source/Avalanche/Public/Transition/Conditions/AvaSceneContainsTagAttributeCondition.h b/Engine/Plugins/Experimental/Avalanche/Source/Avalanche/Public/Transition/Conditions/AvaSceneContainsTagAttributeCondition.h
--- a/Engine/Plugins/Experimental/Avalanche/Source/Avalanche/Public/Transition/Conditions/AvaSceneContainsTagAttributeCondition.h
+++ b/Engine/Plugins/Experimental/Avalanche/Source/Avalanche/Public/Transition/Conditions/AvaSceneContainsTagAttributeCondition.h
@@ -12,17 +12,24 @@ class UAvaSceneSubsystem;
 struct FAvaTransitionScene;
 
 USTRUCT()
-struct FAvaSceneContainsAttributeConditionInstanceData
+struct FAvaSceneContainsTagAttributeConditionInstanceData
 {
 	GENERATED_BODY()
 };
 
-USTRUCT(DisplayName="Scene contains tag attribute", Category="Scene Attributes")
-struct AVALANCHE_API FAvaSceneContainsTagAttributeCondition : public FAvaTransitionCondition
+USTRUCT(meta=(Hidden))
+struct AVALANCHE_API FAvaSceneContainsTagAttributeConditionBase : public FAvaTransitionCondition
 {
 	GENERATED_BODY()
 
-	using FInstanceDataType = FAvaSceneContainsAttributeConditionInstanceData;
+	FAvaSceneContainsTagAttributeConditionBase() = default;
+
+	FAvaSceneContainsTagAttributeConditionBase(bool bInInvertCondition)
+		: bInvertCondition(bInInvertCondition)
+	{
+	}
+
+	using FInstanceDataType = FAvaSceneContainsTagAttributeConditionInstanceData;
 
 	//~ Begin FAvaTransitionCondition
 	virtual FText GenerateDescription(const FAvaTransitionNodeContext& InContext) const override;
@@ -57,9 +64,30 @@ struct AVALANCHE_API FAvaSceneContainsTagAttributeCondition : public FAvaTransit
 	UPROPERTY(EditAnywhere, Category="Parameter")
 	FAvaTagHandle TagAttribute;
 
-	/** Whether to invert the condition (i.e. if true, would become "does NOT contain attribute")*/
-	UPROPERTY(EditAnywhere, Category="Parameter")
-	bool bInvert = false;
+protected:
+	bool bInvertCondition = false;
 
 	TStateTreeExternalDataHandle<UAvaSceneSubsystem> SceneSubsystemHandle;
 };
+
+USTRUCT(DisplayName="A scene contains tag attribute", Category="Scene Attributes")
+struct FAvaSceneContainsTagAttributeCondition : public FAvaSceneContainsTagAttributeConditionBase
+{
+	GENERATED_BODY()
+
+	FAvaSceneContainsTagAttributeCondition()
+		: FAvaSceneContainsTagAttributeConditionBase(/*bInvertCondition*/false)
+	{
+	}
+};
+
+USTRUCT(DisplayName="No scene contains tag attribute", Category="Scene Attributes")
+struct FAvaNoSceneContainsTagAttributeCondition : public FAvaSceneContainsTagAttributeConditionBase
+{
+	GENERATED_BODY()
+
+	FAvaNoSceneContainsTagAttributeCondition()
+		: FAvaSceneContainsTagAttributeConditionBase(/*bInvertCondition*/true)
+	{
+	}
+};
diff --git a/Engine/Plugins/Experimental/Avalanche/Source/AvalancheTransition/Public/AvaTransitionEnums.h b/Engine/Plugins/Experimental/Avalanche/Source/AvalancheTransition/Public/AvaTransitionEnums.h
--- a/Engine/Plugins/Experimental/Avalanche/Source/AvalancheTransition/Public/AvaTransitionEnums.h
+++ b/Engine/Plugins/Experimental/Avalanche/Source/AvalancheTransition/Public/AvaTransitionEnums.h
@@ -55,7 +55,7 @@ enum class EAvaTransitionLayerCompareType : uint8
 {
 	None UMETA(Hidden),
 	Same,
-	Different,
+	Different UMETA(DisplayName="Other"),
 	MatchingTag,
 	Any,
 };
diff --git a/Engine/Plugins/Experimental/Avalanche/Source/AvalancheTransitionEditor/Private/ViewModels/Condition/AvaTransitionConditionContainerViewModel.cpp b/Engine/Plugins/Experimental/Avalanche/Source/AvalancheTransitionEditor/Private/ViewModels/Condition/AvaTransitionConditionContainerViewModel.cpp
--- a/Engine/Plugins/Experimental/Avalanche/Source/AvalancheTransitionEditor/Private/ViewModels/Condition/AvaTransitionConditionContainerViewModel.cpp
+++ b/Engine/Plugins/Experimental/Avalanche/Source/AvalancheTransitionEditor/Private/ViewModels/Condition/AvaTransitionConditionContainerViewModel.cpp
@@ -19,7 +19,7 @@ namespace UE::AvaTransitionEditor::Private
 		// First Conditions or Copy operands should not have any operand displayed
 		if (InConditionIndex > 0 && InEditorNode.ConditionOperand != EStateTreeConditionOperand::Copy)
 		{
-			return UEnum::GetDisplayValueAsText(InEditorNode.ConditionOperand).ToLower();	
+			return FText::Format(INVTEXT("{0} "), UEnum::GetDisplayValueAsText(InEditorNode.ConditionOperand).ToLower());
 		}
 		return FText::GetEmpty();
 	}
@@ -109,8 +109,8 @@ FText FAvaTransitionConditionContainerViewModel::UpdateStateDescription() const
 		using namespace UE::AvaTransitionEditor;
 
 		FTextFormat TextFormat = DeltaIndent > 0
-			? INVTEXT("{Operand} {Parenthesis}{Description}")
-			: INVTEXT("{Operand} {Description}{Parenthesis}");
+			? LOCTEXT("ConditionFormatOpening", "{Operand}{Parenthesis}{Description}")
+			: LOCTEXT("ConditionFormatClosing", "{Operand}{Description}{Parenthesis}");
 
 		FFormatNamedArguments TextArguments;
 		TextArguments.Add(TEXT("Operand"), Private::GetOperandText(ConditionIndex, *EditorNode));
@@ -120,12 +120,7 @@ FText FAvaTransitionConditionContainerViewModel::UpdateStateDescription() const
 		ConditionDescriptions.Add(FText::Format(TextFormat, TextArguments));
 	}
 
-	if (!ConditionDescriptions.IsEmpty())
-	{
-		ConditionDescriptions.Insert(LOCTEXT("If", "If"), 0);
-	}
-
-	return FText::Join(INVTEXT(" "), ConditionDescriptions);
+	return FText::Format(LOCTEXT("ConditionFormat", "If {0}"), FText::Join(INVTEXT(" "), ConditionDescriptions));
 }
 
 void FAvaTransitionConditionContainerViewModel::GatherChildren(FAvaTransitionViewModelChildren& OutChildren)
