[Motion Design] Ease curve Tool
 - Shortcut keys F8, Shift+F8, and Shift+Cmd+F8 for quick ease
 - Set quick preset with Shift+Click in preset dropdown list
 - Show selected preset item
 - Instance ease curve tool in AvaSequencer
 - Change around Tangent and Weight num boxes
 - Add start and end graph text
 - Auto zoom to fit config option

diff --git a/Engine/Plugins/Experimental/Avalanche/Source/AvalancheSequencer/Private/AvaSequencer.cpp b/Engine/Plugins/Experimental/Avalanche/Source/AvalancheSequencer/Private/AvaSequencer.cpp
--- a/Engine/Plugins/Experimental/Avalanche/Source/AvalancheSequencer/Private/AvaSequencer.cpp
+++ b/Engine/Plugins/Experimental/Avalanche/Source/AvalancheSequencer/Private/AvaSequencer.cpp
@@ -17,6 +17,8 @@
 #include "DetailsView/Section/AvaSequencePlaybackDetails.h"
 #include "DetailsView/Section/AvaSequenceSelectionDetails.h"
 #include "DetailsView/Section/AvaSequenceSettingsDetails.h"
+#include "EaseCurveTool/AvaEaseCurveTool.h"
+#include "EaseCurveTool/AvaEaseCurveToolCommands.h"
 #include "Editor/Sequencer/Private/Sequencer.h"
 #include "EngineUtils.h"
 #include "EntitySystem/MovieSceneEntitySystemRunner.h"
@@ -261,6 +263,21 @@ void FAvaSequencer::EnsureSequencer()
 	{
 		SequenceProvider->OnEditorSequencerCreated(Sequencer);
 	}
+
+	// Create ease curve tool and map commands
+	EaseCurveTool = MakeShared<FAvaEaseCurveTool>(SharedThis(this));
+
+	const TSharedRef<FAvaEaseCurveTool> EaseCurveToolRef = EaseCurveTool.ToSharedRef();
+	const FAvaEaseCurveToolCommands& EaseCurveToolCommands = FAvaEaseCurveToolCommands::Get();
+
+	CommandList->MapAction(EaseCurveToolCommands.QuickEaseIn
+		, FExecuteAction::CreateSP(EaseCurveToolRef, &FAvaEaseCurveTool::ApplyQuickEaseToSequencerKeySelections, FAvaEaseCurveTool::EOperation::In));
+
+	CommandList->MapAction(EaseCurveToolCommands.QuickEase
+		, FExecuteAction::CreateSP(EaseCurveToolRef, &FAvaEaseCurveTool::ApplyQuickEaseToSequencerKeySelections, FAvaEaseCurveTool::EOperation::InOut));
+
+	CommandList->MapAction(EaseCurveToolCommands.QuickEaseOut
+		, FExecuteAction::CreateSP(EaseCurveToolRef, &FAvaEaseCurveTool::ApplyQuickEaseToSequencerKeySelections, FAvaEaseCurveTool::EOperation::Out));
 }
 
 TSharedRef<ISequencer> FAvaSequencer::CreateSequencer()
@@ -1677,4 +1694,9 @@ void FAvaSequencer::OnUpdateCameraCut(UObject* InCameraObject, bool bInJumpCut)
 	GetProvider().OnUpdateCameraCut(InCameraObject, bInJumpCut);
 }
 
+TSharedRef<FAvaEaseCurveTool> FAvaSequencer::GetEaseCurveTool() const
+{
+	return EaseCurveTool.ToSharedRef(); 
+}
+
 #undef LOCTEXT_NAMESPACE
diff --git a/Engine/Plugins/Experimental/Avalanche/Source/AvalancheSequencer/Private/AvaSequencer.h b/Engine/Plugins/Experimental/Avalanche/Source/AvalancheSequencer/Private/AvaSequencer.h
--- a/Engine/Plugins/Experimental/Avalanche/Source/AvalancheSequencer/Private/AvaSequencer.h
+++ b/Engine/Plugins/Experimental/Avalanche/Source/AvalancheSequencer/Private/AvaSequencer.h
@@ -9,6 +9,7 @@
 #include "IAvaSequencerProvider.h"
 
 class AActor;
+class FAvaEaseCurveTool;
 class FAvaEditorSelection;
 class FAvaSequencerAction;
 class FAvaSequencerCleanView;
@@ -180,6 +181,7 @@ public:
 
 	void FixBindingHierarchy();
 
+	TSharedRef<FAvaEaseCurveTool> GetEaseCurveTool() const;
 	
 	TArrayView<TWeakObjectPtr<>> ResolveBoundObjects(const FGuid& InBindingId, class UMovieSceneSequence* Sequence) const;
 
@@ -264,4 +266,6 @@ private:
 
 	/** Selected sequence details sections to restore when a new sequence is selected. */
 	TSet<FName> SelectedSections;
+
+	TSharedPtr<FAvaEaseCurveTool> EaseCurveTool;
 };
diff --git a/Engine/Plugins/Experimental/Avalanche/Source/AvalancheSequencer/Private/Customization/AvaSequenceCustomization.cpp b/Engine/Plugins/Experimental/Avalanche/Source/AvalancheSequencer/Private/Customization/AvaSequenceCustomization.cpp
--- a/Engine/Plugins/Experimental/Avalanche/Source/AvalancheSequencer/Private/Customization/AvaSequenceCustomization.cpp
+++ b/Engine/Plugins/Experimental/Avalanche/Source/AvalancheSequencer/Private/Customization/AvaSequenceCustomization.cpp
@@ -6,6 +6,7 @@
 #include "AvaTypeSharedPointer.h"
 #include "Commands/AvaSequencerCommands.h"
 #include "DragDropOps/AvaOutlinerItemDragDropOp.h"
+#include "EaseCurveTool/AvaEaseCurveToolCommands.h"
 #include "Framework/MultiBox/MultiBoxBuilder.h"
 #include "Framework/MultiBox/MultiBoxExtender.h"
 #include "ISequencer.h"
@@ -13,6 +14,7 @@
 #include "MVVM/ViewModels/ObjectBindingModel.h"
 #include "ScopedTransaction.h"
 #include "SequencerCommands.h"
+#include "ToolMenus.h"
 
 #define LOCTEXT_NAMESPACE "AvaSequenceCustomization"
 
@@ -106,16 +108,36 @@ TSharedRef<SWidget> FAvaSequenceCustomization::MakePlaybackMenu() const
 		return SNullWidget::NullWidget;
 	}
 
-	const FAvaSequencerCommands& AvaSequencerCommands = FAvaSequencerCommands::Get();
+	UToolMenus* ToolMenus = UToolMenus::Get();
+	if (!ToolMenus)
+	{
+		return SNullWidget::NullWidget;
+	}
 
-	FMenuBuilder MenuBuilder(true, Sequencer->GetCommandBindings());
+	constexpr const TCHAR* MenuName = TEXT("AvaPlaybackMenu");
+
+	if (!ToolMenus->IsMenuRegistered(MenuName))
+	{
+		UToolMenu* const ToolMenu = ToolMenus->RegisterMenu(MenuName, NAME_None, EMultiBoxType::Menu);
 
-	MenuBuilder.AddMenuEntry(AvaSequencerCommands.FixBindingPaths);
-	MenuBuilder.AddMenuEntry(AvaSequencerCommands.FixInvalidBindings);
-	MenuBuilder.AddMenuEntry(AvaSequencerCommands.FixBindingHierarchy);
-	MenuBuilder.AddMenuEntry(AvaSequencerCommands.StaggerLayerBars);
+		const FAvaSequencerCommands& AvaSequencerCommands = FAvaSequencerCommands::Get();
+		const FAvaEaseCurveToolCommands& EaseCurveToolCommands = FAvaEaseCurveToolCommands::Get();
+
+		FToolMenuSection& SequencerSection = ToolMenu->FindOrAddSection(TEXT("AvaSequencerActions"), LOCTEXT("AvaSequencerActions", "Sequence"));
+		SequencerSection.AddMenuEntry(AvaSequencerCommands.StaggerLayerBars);
+
+		FToolMenuSection& EaseCurveToolSection = ToolMenu->FindOrAddSection(TEXT("AvaEaseCurveToolActions"), LOCTEXT("AvaEaseCurveToolActions", "Ease Curve Tool"));
+		EaseCurveToolSection.AddMenuEntry(EaseCurveToolCommands.QuickEase);
+		EaseCurveToolSection.AddMenuEntry(EaseCurveToolCommands.QuickEaseIn);
+		EaseCurveToolSection.AddMenuEntry(EaseCurveToolCommands.QuickEaseOut);
+
+		FToolMenuSection& BindingSection = ToolMenu->FindOrAddSection(TEXT("AvaBindingActions"), LOCTEXT("AvaBindingActions", "Binding"));
+		BindingSection.AddMenuEntry(AvaSequencerCommands.FixBindingPaths);
+		BindingSection.AddMenuEntry(AvaSequencerCommands.FixInvalidBindings);
+		BindingSection.AddMenuEntry(AvaSequencerCommands.FixBindingHierarchy);
+	}
 
-	return MenuBuilder.MakeWidget();
+	return ToolMenus->GenerateWidget(MenuName, FToolMenuContext(Sequencer->GetCommandBindings()));
 }
 
 bool FAvaSequenceCustomization::OnSequencerReceiveDragOver(const FGeometry& InGeometry
diff --git a/Engine/Plugins/Experimental/Avalanche/Source/AvalancheSequencer/Private/DetailsView/SAvaKeyFrameEdit.cpp b/Engine/Plugins/Experimental/Avalanche/Source/AvalancheSequencer/Private/DetailsView/SAvaKeyFrameEdit.cpp
--- a/Engine/Plugins/Experimental/Avalanche/Source/AvalancheSequencer/Private/DetailsView/SAvaKeyFrameEdit.cpp
+++ b/Engine/Plugins/Experimental/Avalanche/Source/AvalancheSequencer/Private/DetailsView/SAvaKeyFrameEdit.cpp
@@ -4,6 +4,7 @@
 #include "AvaSequencer.h"
 #include "EaseCurveTool/AvaEaseCurveTool.h"
 #include "EaseCurveTool/AvaEaseCurveToolSettings.h"
+#include "EaseCurveTool/Widgets/SAvaEaseCurveTool.h"
 #include "MVVM/Selection/Selection.h"
 #include "MVVM/ViewModels/SequencerEditorViewModel.h"
 #include "Widgets/Layout/SScrollBox.h"
@@ -58,7 +59,7 @@ namespace UE::AvaSequencer::Private
 	};
 }
 
-void SAvaKeyFrameEdit::Construct(const FArguments& InArgs, TSharedRef<FAvaSequencer> InSequencer)
+void SAvaKeyFrameEdit::Construct(const FArguments& InArgs, const TSharedRef<FAvaSequencer>& InSequencer)
 {
 	AvaSequencerWeak = InSequencer;
 
@@ -71,8 +72,6 @@ void SAvaKeyFrameEdit::Construct(const FArguments& InArgs, TSharedRef<FAvaSequen
 		SequencerSelectionWeak = SequencerViewModel->GetSelection();
 	}
 
-	CurveEaseTool = MakeShared<FAvaEaseCurveTool>(InSequencer);
-
 	if (!AvaSequencerWeak.IsValid())
 	{
 		return;
@@ -93,7 +92,7 @@ void SAvaKeyFrameEdit::Construct(const FArguments& InArgs, TSharedRef<FAvaSequen
 				Left/Right/Center will maintain 1:1 ratio. */
 				SNew(UE::AvaSequencer::Private::SDynamicHAlign)
 				.HAlign_UObject(GetDefault<UAvaEaseCurveToolSettings>(), &UAvaEaseCurveToolSettings::GetGraphHAlign)
-				.Content(CurveEaseTool->GenerateWidget())
+				.Content(InSequencer->GetEaseCurveTool()->GenerateWidget())
 			]
 			+ SScrollBox::Slot()
 			.AutoSize()
diff --git a/Engine/Plugins/Experimental/Avalanche/Source/AvalancheSequencer/Private/DetailsView/SAvaKeyFrameEdit.h b/Engine/Plugins/Experimental/Avalanche/Source/AvalancheSequencer/Private/DetailsView/SAvaKeyFrameEdit.h
--- a/Engine/Plugins/Experimental/Avalanche/Source/AvalancheSequencer/Private/DetailsView/SAvaKeyFrameEdit.h
+++ b/Engine/Plugins/Experimental/Avalanche/Source/AvalancheSequencer/Private/DetailsView/SAvaKeyFrameEdit.h
@@ -6,7 +6,6 @@
 #include "Widgets/DeclarativeSyntaxSupport.h"
 #include "Widgets/SCompoundWidget.h"
 
-class FAvaEaseCurveTool;
 class FAvaSequencer;
 
 namespace UE::Sequencer
@@ -21,13 +20,11 @@ public:
 		SLATE_ATTRIBUTE(FKeyEditData, KeyEditData)
 	SLATE_END_ARGS()
 
-	void Construct(const FArguments& InArgs, TSharedRef<FAvaSequencer> InSequencer);
+	void Construct(const FArguments& InArgs, const TSharedRef<FAvaSequencer>& InSequencer);
 
 protected:
 	TWeakPtr<FAvaSequencer> AvaSequencerWeak;
 	TWeakPtr<UE::Sequencer::FSequencerSelection> SequencerSelectionWeak;
 
-	TSharedPtr<FAvaEaseCurveTool> CurveEaseTool;
-
 	TAttribute<FKeyEditData> KeyEditData;
 };
diff --git a/Engine/Plugins/Experimental/Avalanche/Source/AvalancheSequencer/Private/EaseCurveTool/AvaEaseCurvePreset.h b/Engine/Plugins/Experimental/Avalanche/Source/AvalancheSequencer/Private/EaseCurveTool/AvaEaseCurvePreset.h
--- a/Engine/Plugins/Experimental/Avalanche/Source/AvalancheSequencer/Private/EaseCurveTool/AvaEaseCurvePreset.h
+++ b/Engine/Plugins/Experimental/Avalanche/Source/AvalancheSequencer/Private/EaseCurveTool/AvaEaseCurvePreset.h
@@ -31,12 +31,12 @@ struct FAvaEaseCurvePreset
 		return Tangents.CalculateCurveLength() < InRhs.Tangents.CalculateCurveLength();
 	}
 
-	UPROPERTY()
+	UPROPERTY(Config, EditAnywhere, Category = "EaseCurvePreset")
 	FString Name;
 
-	UPROPERTY()
+	UPROPERTY(Config, EditAnywhere, Category = "EaseCurvePreset")
 	FString Category;
 
-	UPROPERTY()
+	UPROPERTY(Config, EditAnywhere, Category = "EaseCurvePreset")
 	FAvaEaseCurveTangents Tangents;
 };
diff --git a/Engine/Plugins/Experimental/Avalanche/Source/AvalancheSequencer/Private/EaseCurveTool/AvaEaseCurveStyle.cpp b/Engine/Plugins/Experimental/Avalanche/Source/AvalancheSequencer/Private/EaseCurveTool/AvaEaseCurveStyle.cpp
--- a/Engine/Plugins/Experimental/Avalanche/Source/AvalancheSequencer/Private/EaseCurveTool/AvaEaseCurveStyle.cpp
+++ b/Engine/Plugins/Experimental/Avalanche/Source/AvalancheSequencer/Private/EaseCurveTool/AvaEaseCurveStyle.cpp
@@ -17,10 +17,10 @@ FAvaEaseCurveStyle::FAvaEaseCurveStyle()
 {
 	const TSharedPtr<IPlugin> Plugin = IPluginManager::Get().FindPlugin(TEXT("Avalanche"));
 	check(Plugin.IsValid());
-	if (Plugin.IsValid())
-	{
-		SetContentRoot(FPaths::Combine(Plugin->GetBaseDir(), TEXT("Resources")));
-	}
+	
+	SetContentRoot(FPaths::Combine(Plugin->GetBaseDir(), TEXT("Resources")));
+
+	Set("Preset.Selected", new FSlateRoundedBoxBrush(FStyleColors::Transparent, FVector4(4.f, 0.f, 0.f, 4.f), FStyleColors::Select.GetSpecifiedColor(), 1.f));
 
 	Set("EditMode.Background", new FSlateRoundedBoxBrush(FStyleColors::Transparent, 4.f, FLinearColor(0.1f, 0.1f, 0.1f, 1.f), 1.f));
 	Set("EditMode.Background.Highlight", new FSlateRoundedBoxBrush(FStyleColors::Transparent, 4.f, FLinearColor(0.6f, 0.6f, 0.6f, 1.f), 1.f));
@@ -56,6 +56,8 @@ FAvaEaseCurveStyle::FAvaEaseCurveStyle()
 		.SetPadding(ToolButtonPadding);
 	Set("ToolToggleButton", ToolToggleButtonStyle);
 
+	Set("Editor.LabelFont", FSlateFontInfo(FCoreStyle::GetDefaultFont(), 7, TEXT("Regular")));
+
 	FSlateStyleRegistry::RegisterSlateStyle(*this);
 }
 
diff --git a/Engine/Plugins/Experimental/Avalanche/Source/AvalancheSequencer/Private/EaseCurveTool/AvaEaseCurveSubsystem.cpp b/Engine/Plugins/Experimental/Avalanche/Source/AvalancheSequencer/Private/EaseCurveTool/AvaEaseCurveSubsystem.cpp
--- a/Engine/Plugins/Experimental/Avalanche/Source/AvalancheSequencer/Private/EaseCurveTool/AvaEaseCurveSubsystem.cpp
+++ b/Engine/Plugins/Experimental/Avalanche/Source/AvalancheSequencer/Private/EaseCurveTool/AvaEaseCurveSubsystem.cpp
@@ -25,7 +25,7 @@ namespace UE::EaseCurveTool::Json::Private
 
 		if (!FPaths::DirectoryExists(PresetsPath))
 		{
-			UE_LOG(LogAvaEaseCurveSubsystem, Warning, TEXT("CurveEaseToolJson::GetPresetFilePath() - Configured presets directory does not exist: %s"), *PresetsPath);
+			UE_LOG(LogAvaEaseCurveSubsystem, Warning, TEXT("EaseCurveToolJson::GetPresetFilePath() - Configured presets directory does not exist: %s"), *PresetsPath);
 			return false;
 		}
 
@@ -39,7 +39,7 @@ namespace UE::EaseCurveTool::Json::Private
 			const FString EaseCurveTool_NewPresetCategory = Settings->GetNewPresetCategory();
 			if (EaseCurveTool_NewPresetCategory.IsEmpty())
 			{
-				UE_LOG(LogAvaEaseCurveSubsystem, Warning, TEXT("CurveEaseToolJson::GetPresetFilePath() - No preset category and no configured default category"));
+				UE_LOG(LogAvaEaseCurveSubsystem, Warning, TEXT("EaseCurveToolJson::GetPresetFilePath() - No preset category and no configured default category"));
 				return false;
 			}
 
@@ -57,7 +57,7 @@ namespace UE::EaseCurveTool::Json::Private
 		FString FileJsonString;
 		if (!FFileHelper::LoadFileToString(FileJsonString, *InFilePath))
 		{
-			UE_LOG(LogAvaEaseCurveSubsystem, Warning, TEXT("CurveEaseToolJson::LoadCurvePresetsJson() - Unable to load Json file: %s"), *InFilePath);
+			UE_LOG(LogAvaEaseCurveSubsystem, Warning, TEXT("EaseCurveToolJson::LoadCurvePresetsJson() - Unable to load Json file: %s"), *InFilePath);
 			return false;
 		}
 
@@ -65,7 +65,7 @@ namespace UE::EaseCurveTool::Json::Private
 		OutRootObject = MakeShared<FJsonObject>();
 		if (!FJsonSerializer::Deserialize(JsonReader.ToSharedRef(), OutRootObject))
 		{
-			UE_LOG(LogAvaEaseCurveSubsystem, Warning, TEXT("CurveEaseToolJson::LoadCurvePresetsJson() - Unable to parse file [%s]. Json=[%s]"), *InFilePath, *FileJsonString);
+			UE_LOG(LogAvaEaseCurveSubsystem, Warning, TEXT("EaseCurveToolJson::LoadCurvePresetsJson() - Unable to parse file [%s]. Json=[%s]"), *InFilePath, *FileJsonString);
 			return false;
 		}
 
@@ -78,13 +78,13 @@ namespace UE::EaseCurveTool::Json::Private
 		const TSharedRef<TJsonWriter<>> Writer = TJsonWriterFactory<>::Create(&FileJsonString);
 		if (!FJsonSerializer::Serialize(InRootObject, Writer))
 		{
-			UE_LOG(LogAvaEaseCurveSubsystem, Warning, TEXT("CurveEaseToolJson::SaveCurvePresetsJson() - Unable to serialize [%s]. Json=[%s]"), *InFilePath, *FileJsonString);
+			UE_LOG(LogAvaEaseCurveSubsystem, Warning, TEXT("EaseCurveToolJson::SaveCurvePresetsJson() - Unable to serialize [%s]. Json=[%s]"), *InFilePath, *FileJsonString);
 			return false;
 		}
 
 		if (!FFileHelper::SaveStringToFile(FileJsonString, *InFilePath))
 		{
-			UE_LOG(LogAvaEaseCurveSubsystem, Warning, TEXT("CurveEaseToolJson::SaveCurvePresetsJson() - Unable to save Json file: %s"), *InFilePath);
+			UE_LOG(LogAvaEaseCurveSubsystem, Warning, TEXT("EaseCurveToolJson::SaveCurvePresetsJson() - Unable to save Json file: %s"), *InFilePath);
 			return false;
 		}
 
@@ -180,14 +180,14 @@ namespace UE::EaseCurveTool::Json::Private
 		FString PresetsPath;
 		if (!GetPresetFilePath(InCategory, false, PresetsPath))
 		{
-			UE_LOG(LogAvaEaseCurveSubsystem, Warning, TEXT("CurveEaseToolJson::DoesPresetExist() - Could not find file for preset category: %s"), *InCategory);
+			UE_LOG(LogAvaEaseCurveSubsystem, Warning, TEXT("EaseCurveToolJson::DoesPresetExist() - Could not find file for preset category: %s"), *InCategory);
 			return false;
 		}
 
 		TSharedPtr<FJsonObject> RootObject;
 		if (!LoadCurvePresetsJson(PresetsPath, RootObject) || !RootObject.IsValid())
 		{
-			UE_LOG(LogAvaEaseCurveSubsystem, Warning, TEXT("CurveEaseToolJson::DoesPresetExist() - Unable to load Json file: %s"), *PresetsPath);
+			UE_LOG(LogAvaEaseCurveSubsystem, Warning, TEXT("EaseCurveToolJson::DoesPresetExist() - Unable to load Json file: %s"), *PresetsPath);
 			return false;
 		}
 
@@ -200,7 +200,7 @@ namespace UE::EaseCurveTool::Json::Private
 		const bool bCategoryFileExists = GetPresetFilePath(InCategory, true, PresetsPath);
 		if (PresetsPath.IsEmpty())
 		{
-			UE_LOG(LogAvaEaseCurveSubsystem, Warning, TEXT("CurveEaseToolJson::AddPreset() - Could not find file for preset category: %s"), *InCategory);
+			UE_LOG(LogAvaEaseCurveSubsystem, Warning, TEXT("EaseCurveToolJson::AddPreset() - Could not find file for preset category: %s"), *InCategory);
 			return false;
 		}
 
@@ -210,7 +210,7 @@ namespace UE::EaseCurveTool::Json::Private
 		{
 			if (!UE::EaseCurveTool::Json::Private::LoadCurvePresetsJson(PresetsPath, RootObject) || !RootObject.IsValid())
 			{
-				UE_LOG(LogAvaEaseCurveSubsystem, Warning, TEXT("CurveEaseToolJson::AddPreset() - Unable to load Json file: %s"), *PresetsPath);
+				UE_LOG(LogAvaEaseCurveSubsystem, Warning, TEXT("EaseCurveToolJson::AddPreset() - Unable to load Json file: %s"), *PresetsPath);
 				return false;
 			}
 		}
@@ -223,7 +223,7 @@ namespace UE::EaseCurveTool::Json::Private
 
 		if (!SaveCurvePresetsJson(PresetsPath, RootObject.ToSharedRef()))
 		{
-			UE_LOG(LogAvaEaseCurveSubsystem, Warning, TEXT("CurveEaseToolJson::AddPreset() - Unable to save Json file: %s"), *PresetsPath);
+			UE_LOG(LogAvaEaseCurveSubsystem, Warning, TEXT("EaseCurveToolJson::AddPreset() - Unable to save Json file: %s"), *PresetsPath);
 			return false;
 		}
 
@@ -235,14 +235,14 @@ namespace UE::EaseCurveTool::Json::Private
 		FString PresetsPath;
 		if (!GetPresetFilePath(InCategory, false, PresetsPath))
 		{
-			UE_LOG(LogAvaEaseCurveSubsystem, Warning, TEXT("CurveEaseToolJson::RemovePreset() - Could not find file for preset category: %s"), *InCategory);
+			UE_LOG(LogAvaEaseCurveSubsystem, Warning, TEXT("EaseCurveToolJson::RemovePreset() - Could not find file for preset category: %s"), *InCategory);
 			return false;
 		}
 
 		TSharedPtr<FJsonObject> RootObject;
 		if (!LoadCurvePresetsJson(PresetsPath, RootObject) || !RootObject.IsValid())
 		{
-			UE_LOG(LogAvaEaseCurveSubsystem, Warning, TEXT("CurveEaseToolJson::RemovePreset() - Unable to load Json file: %s"), *PresetsPath);
+			UE_LOG(LogAvaEaseCurveSubsystem, Warning, TEXT("EaseCurveToolJson::RemovePreset() - Unable to load Json file: %s"), *PresetsPath);
 			return false;
 		}
 
@@ -250,7 +250,7 @@ namespace UE::EaseCurveTool::Json::Private
 
 		if (!SaveCurvePresetsJson(PresetsPath, RootObject.ToSharedRef()))
 		{
-			UE_LOG(LogAvaEaseCurveSubsystem, Warning, TEXT("CurveEaseToolJson::RemovePreset() - Unable to save Json file: %s"), *PresetsPath);
+			UE_LOG(LogAvaEaseCurveSubsystem, Warning, TEXT("EaseCurveToolJson::RemovePreset() - Unable to save Json file: %s"), *PresetsPath);
 			return false;
 		}
 
@@ -262,14 +262,14 @@ namespace UE::EaseCurveTool::Json::Private
 		FString PresetsPath;
 		if (!GetPresetFilePath(InCategory, false, PresetsPath))
 		{
-			UE_LOG(LogAvaEaseCurveSubsystem, Warning, TEXT("CurveEaseToolJson::RemovePreset() - Could not find file for preset category: %s"), *InCategory);
+			UE_LOG(LogAvaEaseCurveSubsystem, Warning, TEXT("EaseCurveToolJson::RemovePreset() - Could not find file for preset category: %s"), *InCategory);
 			return false;
 		}
 
 		TSharedPtr<FJsonObject> RootObject;
 		if (!LoadCurvePresetsJson(PresetsPath, RootObject) || !RootObject.IsValid())
 		{
-			UE_LOG(LogAvaEaseCurveSubsystem, Warning, TEXT("CurveEaseToolJson::RemovePreset() - Unable to load Json file: %s"), *PresetsPath);
+			UE_LOG(LogAvaEaseCurveSubsystem, Warning, TEXT("EaseCurveToolJson::RemovePreset() - Unable to load Json file: %s"), *PresetsPath);
 			return false;
 		}
 
@@ -282,7 +282,7 @@ namespace UE::EaseCurveTool::Json::Private
 
 		if (!SaveCurvePresetsJson(PresetsPath, RootObject.ToSharedRef()))
 		{
-			UE_LOG(LogAvaEaseCurveSubsystem, Warning, TEXT("CurveEaseToolJson::RemovePreset() - Unable to save Json file: %s"), *PresetsPath);
+			UE_LOG(LogAvaEaseCurveSubsystem, Warning, TEXT("EaseCurveToolJson::RemovePreset() - Unable to save Json file: %s"), *PresetsPath);
 			return false;
 		}
 
diff --git a/Engine/Plugins/Experimental/Avalanche/Source/AvalancheSequencer/Private/EaseCurveTool/AvaEaseCurveTangents.h b/Engine/Plugins/Experimental/Avalanche/Source/AvalancheSequencer/Private/EaseCurveTool/AvaEaseCurveTangents.h
--- a/Engine/Plugins/Experimental/Avalanche/Source/AvalancheSequencer/Private/EaseCurveTool/AvaEaseCurveTangents.h
+++ b/Engine/Plugins/Experimental/Avalanche/Source/AvalancheSequencer/Private/EaseCurveTool/AvaEaseCurveTangents.h
@@ -45,7 +45,7 @@ public:
 		: Start(InStartMovieSceneDoubleValue.Tangent.LeaveTangent), StartWeight(InStartMovieSceneDoubleValue.Tangent.LeaveTangentWeight)
 		, End(InEndMovieSceneDoubleValue.Tangent.ArriveTangent), EndWeight(InEndMovieSceneDoubleValue.Tangent.ArriveTangentWeight)
 	{}
-	/** Constructor from string consisting of cubic bezier points. Ex. "0.45, 34.0, 0.0, 3.45" */
+	/** Constructor from string consisting of cubic bezier points. Ex. "0.45, 0.34, 0.0, 1.00" */
 	explicit FAvaEaseCurveTangents(const FString& InTangentsString);
 
 	FORCEINLINE bool operator==(const FAvaEaseCurveTangents& InRhs) const
@@ -80,15 +80,15 @@ public:
 	/** Calculates the length of this curve. The higher the sample count, the higher the calculation precision. */
 	double CalculateCurveLength(const int32 InSampleCount = 10) const;
 
-	UPROPERTY()
+	UPROPERTY(Config, EditAnywhere, Category = "EaseCurveTangents")
 	double Start = 0.0;
 
-	UPROPERTY()
+	UPROPERTY(Config, EditAnywhere, Category = "EaseCurveTangents")
 	double StartWeight = 0.0;
 
-	UPROPERTY()
+	UPROPERTY(Config, EditAnywhere, Category = "EaseCurveTangents")
 	double End = 0.0;
 
-	UPROPERTY()
+	UPROPERTY(Config, EditAnywhere, Category = "EaseCurveTangents")
 	double EndWeight = 0.0;
 };
diff --git a/Engine/Plugins/Experimental/Avalanche/Source/AvalancheSequencer/Private/EaseCurveTool/AvaEaseCurveTool.cpp b/Engine/Plugins/Experimental/Avalanche/Source/AvalancheSequencer/Private/EaseCurveTool/AvaEaseCurveTool.cpp
--- a/Engine/Plugins/Experimental/Avalanche/Source/AvalancheSequencer/Private/EaseCurveTool/AvaEaseCurveTool.cpp
+++ b/Engine/Plugins/Experimental/Avalanche/Source/AvalancheSequencer/Private/EaseCurveTool/AvaEaseCurveTool.cpp
@@ -5,6 +5,7 @@
 #include "AssetToolsModule.h"
 #include "AvaSequencer.h"
 #include "CurveEditor.h"
+#include "EaseCurveTool/AvaEaseCurveToolCommands.h"
 #include "EaseCurveTool/AvaEaseCurveToolSettings.h"
 #include "EaseCurveTool/Widgets/SAvaEaseCurveEditor.h"
 #include "EaseCurveTool/Widgets/SAvaEaseCurveTool.h"
@@ -25,7 +26,7 @@
 
 using namespace UE::Sequencer;
 
-#define LOCTEXT_NAMESPACE "CurveEaseTool"
+#define LOCTEXT_NAMESPACE "AvaEaseCurveTool"
 
 FAvaEaseCurveTool::FAvaEaseCurveTool(const TSharedRef<FAvaSequencer>& InSequencer)
 	: AvaSequencerWeak(InSequencer)
@@ -37,40 +38,30 @@ FAvaEaseCurveTool::FAvaEaseCurveTool(const TSharedRef<FAvaSequencer>& InSequence
 	if (const TSharedPtr<FSequencerEditorViewModel> SequencerViewModel = Sequencer->GetViewModel())
 	{
 		SequencerSelectionWeak = SequencerViewModel->GetSelection();
-
-		if (const FCurveEditorExtension* CurveEditorExtension = SequencerViewModel->CastDynamic<FCurveEditorExtension>())
-		{
-			SequencerCurveEditorWeak = CurveEditorExtension->GetCurveEditor();
-			check(SequencerCurveEditorWeak.IsValid());
-		}
-
-		CacheSelectionData(SequencerViewModel->GetSelection().ToSharedRef());
 	}
 
 	UpdateEaseCurveFromSequencerKeySelections();
 }
 
-void FAvaEaseCurveTool::CacheSelectionData(const TSharedRef<FSequencerSelection>& InSequencerSelection)
+void FAvaEaseCurveTool::CacheSelectionData()
 {
-	if (!SequencerCurveEditorWeak.IsValid())
+	if (!SequencerSelectionWeak.IsValid())
 	{
 		return;
 	}
 
-	TSharedPtr<FCurveEditor> SequencerCurveEditor = SequencerCurveEditorWeak.Pin();
+	const TSharedPtr<FSequencerSelection> SequencerSelection = SequencerSelectionWeak.Pin();
 
-	CachedChannelKeyData.Reset();
-	bCachedAllChannelSingleKeySelections = true;
-	bCachedIsLastOnlySelectedKey = false;
+	FKeyDataCache NewKeyCache;
 
-	for (const FKeyHandle Key : InSequencerSelection->KeySelection)
+	for (const FKeyHandle Key : SequencerSelection->KeySelection)
 	{
 		if (Key == FKeyHandle::Invalid())
 		{
 			continue;
 		}
 
-		TViewModelPtr<FChannelModel> ChannelModel = InSequencerSelection->KeySelection.GetModelForKey(Key);
+		TViewModelPtr<FChannelModel> ChannelModel = SequencerSelection->KeySelection.GetModelForKey(Key);
 		if (!ChannelModel.IsValid())
 		{
 			continue;
@@ -88,31 +79,33 @@ void FAvaEaseCurveTool::CacheSelectionData(const TSharedRef<FSequencerSelection>
 			continue;
 		}
 
-		FCachedKeyData& Entry = CachedChannelKeyData.FindOrAdd(ChannelModel->GetChannelName());
+		FKeyDataCache::FChannelData& Entry = NewKeyCache.ChannelKeyData.FindOrAdd(ChannelModel->GetChannelName());
 		Entry.ChannelModel = ChannelModel;
 		Entry.DoubleChannel = static_cast<FMovieSceneDoubleChannel*>(ChannelHandle.Get());
 		Entry.Section = ChannelModel->GetSection();
 		Entry.KeyHandles.Add(Key);
 
-		TotalSelectedKeys++;
+		NewKeyCache.TotalSelectedKeys++;
 
 		const TArrayView<const FFrameNumber> ChannelTimes = Entry.DoubleChannel->GetTimes();
 		const int32 AllKeyCount = ChannelTimes.Num();
 		const int32 SelectedKeyCount = Entry.KeyHandles.Num();
 
-		if (SelectedKeyCount == 1 && TotalSelectedKeys == 1)
+		if (SelectedKeyCount == 1 && NewKeyCache.TotalSelectedKeys == 1)
 		{
 			const int32 KeyIndex = Entry.DoubleChannel->GetIndex(Entry.KeyHandles[0]);
 			if (KeyIndex == AllKeyCount - 1)
 			{
-				bCachedIsLastOnlySelectedKey = true;
+				NewKeyCache.bIsLastOnlySelectedKey = true;
 			}
 		}
 		else if (SelectedKeyCount > 1)
 		{
-			bCachedAllChannelSingleKeySelections = false;
+			NewKeyCache.bAllChannelSingleKeySelections = false;
 		}
 	}
+
+	KeyCache = NewKeyCache;
 }
 
 void FAvaEaseCurveTool::AddReferencedObjects(FReferenceCollector& Collector)
@@ -122,24 +115,26 @@ void FAvaEaseCurveTool::AddReferencedObjects(FReferenceCollector& Collector)
 
 FString FAvaEaseCurveTool::GetReferencerName() const
 {
-	return TEXT("AvaCurveEaseTool");
+	return TEXT("AvaEaseCurveTool");
 }
 
 TSharedRef<SWidget> FAvaEaseCurveTool::GenerateWidget()
 {
-	if (!CurveEaseToolWidget.IsValid())
+	CacheSelectionData();
+
+	if (!ToolWidget.IsValid())
 	{
-		CurveEaseToolWidget = SNew(SAvaEaseCurveTool, SharedThis(this))
-			.ToolMode(EMode::DoubleKeyEdit)
+		ToolWidget = SNew(SAvaEaseCurveTool, SharedThis(this))
 			.Visibility(this, &FAvaEaseCurveTool::GetVisibility)
 			.ToolOperation(this, &FAvaEaseCurveTool::GetToolOperation);
 	}
-	return CurveEaseToolWidget.ToSharedRef();
+
+	return ToolWidget.ToSharedRef();
 }
 
 EVisibility FAvaEaseCurveTool::GetVisibility() const
 {
-	return (TotalSelectedKeys > 0 && !bCachedIsLastOnlySelectedKey) ? EVisibility::Visible : EVisibility::Collapsed;
+	return (KeyCache.TotalSelectedKeys > 0 && !KeyCache.bIsLastOnlySelectedKey) ? EVisibility::Visible : EVisibility::Collapsed;
 }
 
 TObjectPtr<UAvaEaseCurve> FAvaEaseCurveTool::GetToolCurve() const
@@ -252,20 +247,48 @@ void FAvaEaseCurveTool::FlattenOrStraightenTangents(const bool bInStartTangent,
 	EaseCurve->BroadcastUpdate();
 }
 
-void FAvaEaseCurveTool::SetSequencerKeySelectionTangents(const FAvaEaseCurveTangents& InTangents)
+void FAvaEaseCurveTool::ApplyEaseCurveToSequencerKeySelections()
+{
+	SetSequencerKeySelectionTangents(GetEaseCurveTangents(), OperationMode);
+}
+
+void FAvaEaseCurveTool::ApplyQuickEaseToSequencerKeySelections(const EOperation InOperation)
 {
-	if (TotalSelectedKeys == 0)
+	const UAvaEaseCurveToolSettings* const Settings = GetDefault<UAvaEaseCurveToolSettings>();
+
+	FAvaEaseCurveTangents Tangents;
+	if (!FAvaEaseCurveTangents::FromString(Settings->GetQuickEaseTangents(), Tangents))
+	{
+		UE_LOG(LogTemp, Warning, TEXT("Ease curve tool failed to apply quick ease tangents: "
+			"Could not parse configured quick ease tangent string."));
+		return;
+	}
+
+	SetEaseCurveTangents(Tangents, true, true);
+
+	if (ToolWidget.IsValid())
+	{
+		ToolWidget->ZoomToFit();
+	}
+}
+
+void FAvaEaseCurveTool::SetSequencerKeySelectionTangents(const FAvaEaseCurveTangents& InTangents, const EOperation InOperation)
+{
+	CacheSelectionData();
+
+	if (KeyCache.TotalSelectedKeys == 0)
 	{
 		return;
 	}
 
 	const FFrameRate DisplayRate = GetDisplayRate();
 	const FFrameRate TickResolution = GetTickResolution();
+
 	const UAvaEaseCurveToolSettings* const EaseCurveToolSettings = GetDefault<UAvaEaseCurveToolSettings>();
 	check(EaseCurveToolSettings);
 	const bool bAutoFlipTangents = EaseCurveToolSettings->GetAutoFlipTangents();
 
-	for (const TPair<FName, FCachedKeyData>& ChannelEntry : CachedChannelKeyData)
+	for (const TPair<FName, FKeyDataCache::FChannelData>& ChannelEntry : KeyCache.ChannelKeyData)
 	{
 		for (const FKeyHandle& KeyHandleToEdit : ChannelEntry.Value.KeyHandles)
 		{
@@ -276,12 +299,12 @@ void FAvaEaseCurveTool::SetSequencerKeySelectionTangents(const FAvaEaseCurveTang
 			{
 				continue;
 			}
-			
+
 			// InTangents should always be a normalized range, allowing for weights in the range of 0 - 10
 			const TArrayView<FFrameNumber> ChannelTimes = ChannelData.GetTimes();
 			TArrayView<FMovieSceneDoubleValue> ChannelValues = ChannelData.GetValues();
 			const int32 KeyCount = ChannelValues.Num();
-			
+
 			// If there is a keyframe after this keyframe that we are editing, we check if the that keyframe value is less
 			// than or greater than this keyframe value. If less, flip the tangent (if option is set).
 			int32 NextKeyIndex = INDEX_NONE;
@@ -312,7 +335,7 @@ void FAvaEaseCurveTool::SetSequencerKeySelectionTangents(const FAvaEaseCurveTang
 			ChannelEntry.Value.Section->MarkAsChanged();
 
 			// Set this keys leave tangent
-			if (OperationMode == EOperation::Out || OperationMode == EOperation::InOut)
+			if (InOperation == EOperation::Out || InOperation == EOperation::InOut)
 			{
 				ChannelValues[KeyIndex].InterpMode = ERichCurveInterpMode::RCIM_Cubic;
 				ChannelValues[KeyIndex].Tangent.TangentWeightMode = ERichCurveTangentWeightMode::RCTWM_WeightedBoth;
@@ -320,9 +343,9 @@ void FAvaEaseCurveTool::SetSequencerKeySelectionTangents(const FAvaEaseCurveTang
 				ChannelValues[KeyIndex].Tangent.LeaveTangent = ScaledTangents.Start;
 				ChannelValues[KeyIndex].Tangent.LeaveTangentWeight = ScaledTangents.StartWeight;
 			}
-			
+
 			// Set the next keys arrive tangent
-			if (NextKeyIndex != INDEX_NONE && (OperationMode == EOperation::In || OperationMode == EOperation::InOut))
+			if (NextKeyIndex != INDEX_NONE && (InOperation == EOperation::In || InOperation == EOperation::InOut))
 			{
 				ChannelValues[NextKeyIndex].InterpMode = ERichCurveInterpMode::RCIM_Cubic;
 				ChannelValues[NextKeyIndex].Tangent.TangentWeightMode = ERichCurveTangentWeightMode::RCTWM_WeightedBoth;
@@ -336,7 +359,9 @@ void FAvaEaseCurveTool::SetSequencerKeySelectionTangents(const FAvaEaseCurveTang
 
 void FAvaEaseCurveTool::UpdateEaseCurveFromSequencerKeySelections()
 {
-	if (TotalSelectedKeys == 0)
+	CacheSelectionData();
+	
+	if (KeyCache.TotalSelectedKeys == 0)
 	{
 		return;
 	}
@@ -349,7 +374,7 @@ void FAvaEaseCurveTool::UpdateEaseCurveFromSequencerKeySelections()
 
 	TArray<FAvaEaseCurveTangents> KeySetTangents;
 
-	for (const TPair<FName, FCachedKeyData>& ChannelEntry : CachedChannelKeyData)
+	for (const TPair<FName, FKeyDataCache::FChannelData>& ChannelEntry : KeyCache.ChannelKeyData)
 	{
 		for (const FKeyHandle& KeyHandleToEdit : ChannelEntry.Value.KeyHandles)
 		{
@@ -409,11 +434,6 @@ void FAvaEaseCurveTool::UpdateEaseCurveFromSequencerKeySelections()
 	SetEaseCurveTangents(FAvaEaseCurveTangents::Average(KeySetTangents), true, false);
 }
 
-void FAvaEaseCurveTool::ApplyEaseCurveToSequencerKeySelections()
-{
-	SetSequencerKeySelectionTangents(GetEaseCurveTangents());
-}
-
 UCurveBase* FAvaEaseCurveTool::CreateCurveAsset() const
 {
 	const FAssetToolsModule& AssetToolsModule = FModuleManager::LoadModuleChecked<FAssetToolsModule>(TEXT("AssetTools"));
@@ -624,7 +644,7 @@ FFrameRate FAvaEaseCurveTool::GetDisplayRate() const
 	{
 		return AvaSequencer->GetSequencer()->GetFocusedDisplayRate();
 	}
-	
+
 	// Fallback to using config display rate if tool is being used outside sequencer
 	return GetDefault<UAvaSequencerSettings>()->GetDisplayRate();
 }
@@ -639,7 +659,7 @@ void FAvaEaseCurveTool::ShowNotificationMessage(const FText& InMessageText)
 void FAvaEaseCurveTool::SelectNextChannelKey()
 {
 	// This command is only available for single key selections
-	if (TotalSelectedKeys != 1 || !SequencerSelectionWeak.IsValid())
+	if (KeyCache.TotalSelectedKeys != 1 || !SequencerSelectionWeak.IsValid())
 	{
 		return;
 	}
@@ -651,7 +671,7 @@ void FAvaEaseCurveTool::SelectNextChannelKey()
 	FKeyHandle NextKeyHandle;
 	TViewModelPtr<FChannelModel> ChannelModel;
 
-	for (const TPair<FName, FCachedKeyData>& ChannelEntry : CachedChannelKeyData)
+	for (const TPair<FName, FKeyDataCache::FChannelData>& ChannelEntry : KeyCache.ChannelKeyData)
 	{
 		if (ChannelEntry.Value.KeyHandles.Num() == 1)
 		{
@@ -688,7 +708,7 @@ void FAvaEaseCurveTool::SelectNextChannelKey()
 void FAvaEaseCurveTool::SelectPreviousChannelKey()
 {
 	// This command is only available for single key selections
-	if (TotalSelectedKeys != 1 || !SequencerSelectionWeak.IsValid())
+	if (KeyCache.TotalSelectedKeys != 1 || !SequencerSelectionWeak.IsValid())
 	{
 		return;
 	}
@@ -700,7 +720,7 @@ void FAvaEaseCurveTool::SelectPreviousChannelKey()
 	FKeyHandle PreviousKeyHandle;
 	TViewModelPtr<FChannelModel> ChannelModel;
 
-	for (const TPair<FName, FCachedKeyData>& ChannelEntry : CachedChannelKeyData)
+	for (const TPair<FName, FKeyDataCache::FChannelData>& ChannelEntry : KeyCache.ChannelKeyData)
 	{
 		if (ChannelEntry.Value.KeyHandles.Num() == 1)
 		{
diff --git a/Engine/Plugins/Experimental/Avalanche/Source/AvalancheSequencer/Private/EaseCurveTool/AvaEaseCurveTool.h b/Engine/Plugins/Experimental/Avalanche/Source/AvalancheSequencer/Private/EaseCurveTool/AvaEaseCurveTool.h
--- a/Engine/Plugins/Experimental/Avalanche/Source/AvalancheSequencer/Private/EaseCurveTool/AvaEaseCurveTool.h
+++ b/Engine/Plugins/Experimental/Avalanche/Source/AvalancheSequencer/Private/EaseCurveTool/AvaEaseCurveTool.h
@@ -16,11 +16,16 @@
 
 class FAvaSequencer;
 class FCurveEditor;
+class IAvaSequencer;
+class ISequencer;
 class SAvaEaseCurveTool;
 class SCurveEditor;
 class SCurveEditorPanel;
 class SEditableTextBox;
+class UAvaEaseCurve;
 class UCurveFloat;
+class UMovieSceneSection;
+struct FAvaEaseCurveTangents;
 struct FMovieSceneDoubleChannel;
 
 namespace UE::Sequencer
@@ -35,18 +40,11 @@ class FAvaEaseCurveTool
 	, public FSelfRegisteringEditorUndoClient
 {
 public:
-	static void ShowNotificationMessage(const FText& InMessageText);
-
-	/** Returns true if the clipboard paste data contains tangent information. */
-	static bool TangentsFromClipboardPaste(FAvaEaseCurveTangents& OutTangents);
-
-	/**
-	 * Current default and only implemented is DoubleKeyEdit.
-	 */
+	/** Current default and only implemented is DualKeyEdit. */
 	enum class EMode : uint8
 	{
 		/** Edits the selected key's leave tangent and the next key's arrive tangent in the curve editor graph. */
-		DoubleKeyEdit,
+		DualKeyEdit,
 		/** Edits only the selected key.
 		 * The leave tangent in the curve editor graph will set the sequence key arrive tangent.
 		 * The arrive tangent in the curve editor graph will set the sequence key leave tangent. */
@@ -55,11 +53,16 @@ public:
 
 	enum class EOperation : uint8
 	{
+		InOut,
 		In,
-		Out,
-		InOut
+		Out
 	};
 
+	static void ShowNotificationMessage(const FText& InMessageText);
+
+	/** Returns true if the clipboard paste data contains tangent information. */
+	static bool TangentsFromClipboardPaste(FAvaEaseCurveTangents& OutTangents);
+
 	FAvaEaseCurveTool(const TSharedRef<FAvaSequencer>& InSequencer);
 
 	TSharedRef<SWidget> GenerateWidget();
@@ -80,14 +83,16 @@ public:
 
 	/** Creates a new external float curve from the internet curve editor curve. */
 	UCurveBase* CreateCurveAsset() const;
-	
-	void SetSequencerKeySelectionTangents(const FAvaEaseCurveTangents& InTangents);
 
-	/** Updates the ease curve graph view based on the active sequencer key selection. */
-	void UpdateEaseCurveFromSequencerKeySelections();
+	void SetSequencerKeySelectionTangents(const FAvaEaseCurveTangents& InTangents, const EOperation InOperation = EOperation::InOut);
 
 	void ApplyEaseCurveToSequencerKeySelections();
 
+	void ApplyQuickEaseToSequencerKeySelections(const EOperation InOperation = EOperation::InOut);
+
+	/** Updates the ease curve graph view based on the active sequencer key selection. */
+	void UpdateEaseCurveFromSequencerKeySelections();
+
 	EOperation GetToolOperation() const;
 	void SetToolOperation(const EOperation InNewOperation);
 	bool IsToolOperation(const EOperation InNewOperation) const;
@@ -127,34 +132,38 @@ public:
 protected:
 	TWeakPtr<FAvaSequencer> AvaSequencerWeak;
 	TWeakPtr<UE::Sequencer::FSequencerSelection> SequencerSelectionWeak;
-	TWeakPtr<FCurveEditor> SequencerCurveEditorWeak;
 
 	TObjectPtr<UAvaEaseCurve> EaseCurve;
 
-	EMode ToolMode = EMode::DoubleKeyEdit;
+	EMode ToolMode = EMode::DualKeyEdit;
 	EOperation OperationMode = EOperation::InOut;
 
-	TSharedPtr<SAvaEaseCurveTool> CurveEaseToolWidget;
-
-	FSimpleMulticastDelegate::FDelegate LastDelegate;
+	TSharedPtr<SAvaEaseCurveTool> ToolWidget;
 	
 	/** Cached data set when a new sequencer selection is made. */
-	struct FCachedKeyData
+	struct FKeyDataCache
 	{
-		TSharedPtr<UE::Sequencer::FChannelModel> ChannelModel;
-		UMovieSceneSection* Section = nullptr;
-		FMovieSceneDoubleChannel* DoubleChannel = nullptr;
-		TArray<FKeyHandle> KeyHandles;
-	};
-	TMap<FName, FCachedKeyData> CachedChannelKeyData;
-	int32 TotalSelectedKeys = 0;
+		struct FChannelData
+		{
+			TSharedPtr<UE::Sequencer::FChannelModel> ChannelModel;
+			UMovieSceneSection* Section = nullptr;
+			FMovieSceneDoubleChannel* DoubleChannel = nullptr;
+			TArray<FKeyHandle> KeyHandles;
+		};
+
+		TMap<FName, FChannelData> ChannelKeyData;
 
-	/** Indicates only one selected key on each of the selected channels. */
-	bool bCachedAllChannelSingleKeySelections = false;
+		int32 TotalSelectedKeys = 0;
 
-	/** Indicates only one key selected and it is the last key of the channel. */
-	bool bCachedIsLastOnlySelectedKey = false;
+		/** Indicates only one selected key on each of the selected channels. */
+		bool bAllChannelSingleKeySelections = true;
+
+		/** Indicates only one key selected and it is the last key of the channel. */
+		bool bIsLastOnlySelectedKey = false;
+	};
+	FKeyDataCache KeyCache;
+	bool bAreKeysCached = false;
 
 private:
-	void CacheSelectionData(const TSharedRef<UE::Sequencer::FSequencerSelection>& InSequencerSelection);
+	void CacheSelectionData();
 };
diff --git a/Engine/Plugins/Experimental/Avalanche/Source/AvalancheSequencer/Private/EaseCurveTool/AvaEaseCurveToolCommands.cpp b/Engine/Plugins/Experimental/Avalanche/Source/AvalancheSequencer/Private/EaseCurveTool/AvaEaseCurveToolCommands.cpp
--- a/Engine/Plugins/Experimental/Avalanche/Source/AvalancheSequencer/Private/EaseCurveTool/AvaEaseCurveToolCommands.cpp
+++ b/Engine/Plugins/Experimental/Avalanche/Source/AvalancheSequencer/Private/EaseCurveTool/AvaEaseCurveToolCommands.cpp
@@ -2,11 +2,13 @@
 
 #include "AvaEaseCurveToolCommands.h"
 
-#define LOCTEXT_NAMESPACE "AvaCurveEaseToolCommands"
+#define LOCTEXT_NAMESPACE "AvaEaseCurveToolCommands"
 
 void FAvaEaseCurveToolCommands::RegisterCommands()
 {
-	FUICommandInfo::MakeCommandInfo(AsShared()
+	const TSharedRef<FBindingContext> SharedThis = AsShared();
+
+	FUICommandInfo::MakeCommandInfo(SharedThis
 		, OpenToolSettings
 		, TEXT("OpenToolSettings")
 		, LOCTEXT("OpenToolSettings_Label", "Open Tool Settings...")
@@ -15,7 +17,7 @@ void FAvaEaseCurveToolCommands::RegisterCommands()
 		, EUserInterfaceActionType::Button
 		, FInputChord());
 
-	FUICommandInfo::MakeCommandInfo(AsShared()
+	FUICommandInfo::MakeCommandInfo(SharedThis
 		, Refresh
 		, TEXT("Refresh")
 		, LOCTEXT("Refresh_Label", "Refresh")
@@ -24,7 +26,7 @@ void FAvaEaseCurveToolCommands::RegisterCommands()
 		, EUserInterfaceActionType::Button
 		, FInputChord(EKeys::F5));
 
-	FUICommandInfo::MakeCommandInfo(AsShared()
+	FUICommandInfo::MakeCommandInfo(SharedThis
 		, Apply
 		, TEXT("Apply")
 		, LOCTEXT("Apply_Label", "Apply")
@@ -33,7 +35,7 @@ void FAvaEaseCurveToolCommands::RegisterCommands()
 		, EUserInterfaceActionType::Button
 		, FInputChord());
 
-	FUICommandInfo::MakeCommandInfo(AsShared()
+	FUICommandInfo::MakeCommandInfo(SharedThis
 		, ZoomToFit
 		, TEXT("ZoomToFit")
 		, LOCTEXT("ZoomToFit_Label", "Zoom to Fit")
@@ -42,7 +44,7 @@ void FAvaEaseCurveToolCommands::RegisterCommands()
 		, EUserInterfaceActionType::Button
 		, FInputChord(EKeys::F));
 
-	FUICommandInfo::MakeCommandInfo(AsShared()
+	FUICommandInfo::MakeCommandInfo(SharedThis
 		, SetOperationToEaseOut
 		, TEXT("SetOperationToEaseOut")
 		, LOCTEXT("SetOperationToEaseOut_Label", "Ease Out")
@@ -51,7 +53,7 @@ void FAvaEaseCurveToolCommands::RegisterCommands()
 		, EUserInterfaceActionType::RadioButton
 		, FInputChord(EKeys::O));
 
-	FUICommandInfo::MakeCommandInfo(AsShared()
+	FUICommandInfo::MakeCommandInfo(SharedThis
 		, SetOperationToEaseInOut
 		, TEXT("SetOperationToEaseInOut")
 		, LOCTEXT("SetOperationToEaseInOut_Label", "Ease In Out")
@@ -60,7 +62,7 @@ void FAvaEaseCurveToolCommands::RegisterCommands()
 		, EUserInterfaceActionType::RadioButton
 		, FInputChord(EKeys::U));
 
-	FUICommandInfo::MakeCommandInfo(AsShared()
+	FUICommandInfo::MakeCommandInfo(SharedThis
 		, SetOperationToEaseIn
 		, TEXT("SetOperationToEaseIn")
 		, LOCTEXT("SetOperationToEaseIn_Label", "Ease In")
@@ -69,7 +71,7 @@ void FAvaEaseCurveToolCommands::RegisterCommands()
 		, EUserInterfaceActionType::RadioButton
 		, FInputChord(EKeys::I));
 
-	FUICommandInfo::MakeCommandInfo(AsShared()
+	FUICommandInfo::MakeCommandInfo(SharedThis
 		, ToggleGridSnap
 		, TEXT("ToggleGridSnap")
 		, LOCTEXT("ToggleGridSnap_Label", "Snap To Grid")
@@ -78,7 +80,7 @@ void FAvaEaseCurveToolCommands::RegisterCommands()
 		, EUserInterfaceActionType::Check
 		, FInputChord(EKeys::G));
 
-	FUICommandInfo::MakeCommandInfo(AsShared()
+	FUICommandInfo::MakeCommandInfo(SharedThis
 		, ToggleAutoFlipTangents
 		, TEXT("ToggleAutoFlipTangents")
 		, LOCTEXT("ToggleAutoFlipTangents_Label", "Auto Flip Tangents")
@@ -87,7 +89,13 @@ void FAvaEaseCurveToolCommands::RegisterCommands()
 		, EUserInterfaceActionType::Check
 		, FInputChord());
 
-	FUICommandInfo::MakeCommandInfo(AsShared()
+	UI_COMMAND(ToggleAutoZoomToFit
+		, "Auto Zoom To Fit"
+		, "Auto zoom the graph editor to fit the tangent handles after they have been changed."
+		, EUserInterfaceActionType::Check
+		, FInputChord());
+
+	FUICommandInfo::MakeCommandInfo(SharedThis
 		, ResetTangents
 		, TEXT("ResetBothTangents")
 		, LOCTEXT("ResetBothTangents_Label", "Both Tangents")
@@ -96,7 +104,7 @@ void FAvaEaseCurveToolCommands::RegisterCommands()
 		, EUserInterfaceActionType::Button
 		, FInputChord(EKeys::R));
 
-	FUICommandInfo::MakeCommandInfo(AsShared()
+	FUICommandInfo::MakeCommandInfo(SharedThis
 		, ResetStartTangent
 		, TEXT("ResetStartTangent")
 		, LOCTEXT("ResetStartTangent_Label", "Start Tangent")
@@ -105,7 +113,7 @@ void FAvaEaseCurveToolCommands::RegisterCommands()
 		, EUserInterfaceActionType::Button
 		, FInputChord(EKeys::R, EModifierKey::Control));
 
-	FUICommandInfo::MakeCommandInfo(AsShared()
+	FUICommandInfo::MakeCommandInfo(SharedThis
 		, ResetEndTangent
 		, TEXT("ResetEndTangent")
 		, LOCTEXT("ResetEndTangent_Label", "End Tangent")
@@ -114,7 +122,7 @@ void FAvaEaseCurveToolCommands::RegisterCommands()
 		, EUserInterfaceActionType::Button
 		, FInputChord(EKeys::R, EModifierKey::Alt));
 
-	FUICommandInfo::MakeCommandInfo(AsShared()
+	FUICommandInfo::MakeCommandInfo(SharedThis
 		, FlattenTangents
 		, TEXT("FlattenBothTangents")
 		, LOCTEXT("FlattenBothTangents_Label", "Both Tangents")
@@ -123,7 +131,7 @@ void FAvaEaseCurveToolCommands::RegisterCommands()
 		, EUserInterfaceActionType::Button
 		, FInputChord(EKeys::T));
 
-	FUICommandInfo::MakeCommandInfo(AsShared()
+	FUICommandInfo::MakeCommandInfo(SharedThis
 		, FlattenStartTangent
 		, TEXT("FlattenStartTangent")
 		, LOCTEXT("FlattenStartTangent_Label", "Start Tangent")
@@ -132,7 +140,7 @@ void FAvaEaseCurveToolCommands::RegisterCommands()
 		, EUserInterfaceActionType::Button
 		, FInputChord(EKeys::T, EModifierKey::Control));
 
-	FUICommandInfo::MakeCommandInfo(AsShared()
+	FUICommandInfo::MakeCommandInfo(SharedThis
 		, FlattenEndTangent
 		, TEXT("FlattenEndTangent")
 		, LOCTEXT("FlattenEndTangent_Label", "End Tangent")
@@ -141,7 +149,7 @@ void FAvaEaseCurveToolCommands::RegisterCommands()
 		, EUserInterfaceActionType::Button
 		, FInputChord(EKeys::T, EModifierKey::Alt));
 
-	FUICommandInfo::MakeCommandInfo(AsShared()
+	FUICommandInfo::MakeCommandInfo(SharedThis
 		, StraightenTangents
 		, TEXT("StraightenBothTangents")
 		, LOCTEXT("StraightenBothTangents_Label", "Both Tangents")
@@ -150,7 +158,7 @@ void FAvaEaseCurveToolCommands::RegisterCommands()
 		, EUserInterfaceActionType::Button
 		, FInputChord(EKeys::S));
 
-	FUICommandInfo::MakeCommandInfo(AsShared()
+	FUICommandInfo::MakeCommandInfo(SharedThis
 		, StraightenStartTangent
 		, TEXT("StraightenStartTangent")
 		, LOCTEXT("StraightenStartTangent_Label", "Start Tangent")
@@ -159,7 +167,7 @@ void FAvaEaseCurveToolCommands::RegisterCommands()
 		, EUserInterfaceActionType::Button
 		, FInputChord(EKeys::S, EModifierKey::Control));
 
-	FUICommandInfo::MakeCommandInfo(AsShared()
+	FUICommandInfo::MakeCommandInfo(SharedThis
 		, StraightenEndTangent
 		, TEXT("StraightenEndTangent")
 		, LOCTEXT("StraightenEndTangent_Label", "End Tangent")
@@ -168,7 +176,7 @@ void FAvaEaseCurveToolCommands::RegisterCommands()
 		, EUserInterfaceActionType::Button
 		, FInputChord(EKeys::S, EModifierKey::Alt));
 
-	FUICommandInfo::MakeCommandInfo(AsShared()
+	FUICommandInfo::MakeCommandInfo(SharedThis
 		, CopyTangents
 		, TEXT("CopyTangents")
 		, LOCTEXT("CopyTangents_Label", "Copy Tangents")
@@ -177,7 +185,7 @@ void FAvaEaseCurveToolCommands::RegisterCommands()
 		, EUserInterfaceActionType::Button
 		, FInputChord(EKeys::C, EModifierKey::Control));
 
-	FUICommandInfo::MakeCommandInfo(AsShared()
+	FUICommandInfo::MakeCommandInfo(SharedThis
 		, PasteTangents
 		, TEXT("PasteTangents")
 		, LOCTEXT("PasteTangents_Label", "Paste Tangents")
@@ -186,7 +194,7 @@ void FAvaEaseCurveToolCommands::RegisterCommands()
 		, EUserInterfaceActionType::Button
 		, FInputChord(EKeys::V, EModifierKey::Control));
 
-	FUICommandInfo::MakeCommandInfo(AsShared()
+	FUICommandInfo::MakeCommandInfo(SharedThis
 		, CreateExternalCurveAsset
 		, TEXT("CreateExternalCurve")
 		, LOCTEXT("CreateExternalCurve_Label", "Create External Curve...")
@@ -203,7 +211,7 @@ void FAvaEaseCurveToolCommands::RegisterCommands()
 	UI_COMMAND(SetKeyInterpCubicBreak, "Break", "Cubic interpolation - User broken tangents", EUserInterfaceActionType::ToggleButton, FInputChord(EKeys::Three));
 	UI_COMMAND(SetKeyInterpToggleWeighted, "Weighted Tangents", "Toggle weighted tangents for cubic interpolation modes. Only supported on some curve types", EUserInterfaceActionType::ToggleButton, FInputChord(EModifierKey::Control, EKeys::W));
 
-	FUICommandInfo::MakeCommandInfo(AsShared()
+	FUICommandInfo::MakeCommandInfo(SharedThis
 		, SelectNextChannelKey
 		, TEXT("SelectNextChannelKey")
 		, LOCTEXT("SelectNextChannelKey_Label", "Select Next Channel Key")
@@ -212,7 +220,7 @@ void FAvaEaseCurveToolCommands::RegisterCommands()
 		, EUserInterfaceActionType::Button
 		, FInputChord(EKeys::Period, EModifierKey::Control));
 
-	FUICommandInfo::MakeCommandInfo(AsShared()
+	FUICommandInfo::MakeCommandInfo(SharedThis
 		, SelectPreviousChannelKey
 		, TEXT("SelectsPreviousChannelKey")
 		, LOCTEXT("SelectsPreviousChannelKey_Label", "Selects Previous Channel Key")
@@ -220,6 +228,33 @@ void FAvaEaseCurveToolCommands::RegisterCommands()
 		, FSlateIcon(FAppStyle::Get().GetStyleSetName(), TEXT("Wizard.BackIcon"))
 		, EUserInterfaceActionType::Button
 		, FInputChord(EKeys::Comma, EModifierKey::Control));
+
+	FUICommandInfo::MakeCommandInfo(SharedThis
+		, QuickEase
+		, TEXT("QuickEase")
+		, LOCTEXT("QuickEase_Label", "Quick Ease")
+		, LOCTEXT("QuickEase_ToolTip", "Apply the quick ease preset")
+		, FSlateIcon(FAppStyle::Get().GetStyleSetName(), TEXT("Icons.BulletPoint"))
+		, EUserInterfaceActionType::Button
+		, FInputChord(EKeys::F8));
+
+	FUICommandInfo::MakeCommandInfo(SharedThis
+		, QuickEaseIn
+		, TEXT("QuickEaseIn")
+		, LOCTEXT("QuickEaseIn_Label", "Quick Ease In")
+		, LOCTEXT("QuickEaseIn_ToolTip", "Apply the quick ease in preset")
+		, FSlateIcon(FAppStyle::Get().GetStyleSetName(), TEXT("Icons.ChevronRight"))
+		, EUserInterfaceActionType::Button
+		, FInputChord(EKeys::F8, EModifierKey::Shift));
+
+	FUICommandInfo::MakeCommandInfo(SharedThis
+		, QuickEaseOut
+		, TEXT("QuickEaseOut")
+		, LOCTEXT("QuickEaseOut_Label", "Quick Ease Out")
+		, LOCTEXT("QuickEaseOut_ToolTip", "Apply the quick ease out preset")
+		, FSlateIcon(FAppStyle::Get().GetStyleSetName(), TEXT("Icons.ChevronLeft"))
+		, EUserInterfaceActionType::Button
+		, FInputChord(EKeys::F8, EModifierKey::Command | EModifierKey::Shift));
 }
 
 #undef LOCTEXT_NAMESPACE
diff --git a/Engine/Plugins/Experimental/Avalanche/Source/AvalancheSequencer/Private/EaseCurveTool/AvaEaseCurveToolCommands.h b/Engine/Plugins/Experimental/Avalanche/Source/AvalancheSequencer/Private/EaseCurveTool/AvaEaseCurveToolCommands.h
--- a/Engine/Plugins/Experimental/Avalanche/Source/AvalancheSequencer/Private/EaseCurveTool/AvaEaseCurveToolCommands.h
+++ b/Engine/Plugins/Experimental/Avalanche/Source/AvalancheSequencer/Private/EaseCurveTool/AvaEaseCurveToolCommands.h
@@ -9,8 +9,8 @@ class FAvaEaseCurveToolCommands : public TCommands<FAvaEaseCurveToolCommands>
 {
 public:
 	FAvaEaseCurveToolCommands()
-		: TCommands<FAvaEaseCurveToolCommands>(TEXT("AvalancheCurveEaseToolCommands")
-			, NSLOCTEXT("AvaCurveEaseToolCommands", "AvaCurveEaseToolCommands", "Motion Design Curve Ease Tool Commands")
+		: TCommands<FAvaEaseCurveToolCommands>(TEXT("AvalancheEaseCurveToolCommands")
+			, NSLOCTEXT("AvaEaseCurveToolCommands", "AvaEaseCurveToolCommands", "Motion Design Curve Ease Tool Commands")
 			, NAME_None
 			, FAppStyle::GetAppStyleSetName())
 	{
@@ -33,6 +33,8 @@ public:
 
 	TSharedPtr<FUICommandInfo> ToggleAutoFlipTangents;
 
+	TSharedPtr<FUICommandInfo> ToggleAutoZoomToFit;
+
 	TSharedPtr<FUICommandInfo> ResetTangents;
 	TSharedPtr<FUICommandInfo> ResetStartTangent;
 	TSharedPtr<FUICommandInfo> ResetEndTangent;
@@ -60,4 +62,8 @@ public:
 
 	TSharedPtr<FUICommandInfo> SelectNextChannelKey;
 	TSharedPtr<FUICommandInfo> SelectPreviousChannelKey;
+
+	TSharedPtr<FUICommandInfo> QuickEase;
+	TSharedPtr<FUICommandInfo> QuickEaseIn;
+	TSharedPtr<FUICommandInfo> QuickEaseOut;
 };
diff --git a/Engine/Plugins/Experimental/Avalanche/Source/AvalancheSequencer/Private/EaseCurveTool/AvaEaseCurveToolSettings.h b/Engine/Plugins/Experimental/Avalanche/Source/AvalancheSequencer/Private/EaseCurveTool/AvaEaseCurveToolSettings.h
--- a/Engine/Plugins/Experimental/Avalanche/Source/AvalancheSequencer/Private/EaseCurveTool/AvaEaseCurveToolSettings.h
+++ b/Engine/Plugins/Experimental/Avalanche/Source/AvalancheSequencer/Private/EaseCurveTool/AvaEaseCurveToolSettings.h
@@ -3,6 +3,7 @@
 #pragma once
 
 #include "Containers/EnumAsByte.h"
+#include "EaseCurveTool/AvaEaseCurveTangents.h"
 #include "Engine/DeveloperSettings.h"
 #include "Types/SlateEnums.h"
 #include "UObject/Object.h"
@@ -17,24 +18,34 @@ public:
 	UAvaEaseCurveToolSettings();
 
 	int32 GetGraphSize() const { return GraphSize; }
-	void SetGraphSize(const int32 InNewSize) { GraphSize = InNewSize; }
+	void SetGraphSize(const int32 InSize) { GraphSize = InSize; }
 
 	EHorizontalAlignment GetGraphHAlign() const { return GraphHAlign.GetValue(); }
-	void SetGraphHAlign(const EHorizontalAlignment InNewHAlign) { GraphHAlign = InNewHAlign; }
+	void SetGraphHAlign(const EHorizontalAlignment InHAlign) { GraphHAlign = InHAlign; }
 
 	bool GetGridSnap() const { return bGridSnap; }
-	void SetGridSnap(const bool bInCurveEaseToolGridSnap) { bGridSnap = bInCurveEaseToolGridSnap; }
+	void SetGridSnap(const bool bInGridSnap) { bGridSnap = bInGridSnap; }
 	void ToggleGridSnap() { bGridSnap = !bGridSnap; }
 
 	int32 GetGridSize() const { return GridSize; }
-	void SetGridSize(const int32 InNewSize) { GridSize = InNewSize; }
+	void SetGridSize(const int32 InSize) { GridSize = InSize; }
+
+	bool GetAutoZoomToFit() const { return bAutoZoomToFit; }
+	void SetAutoZoomToFit(const bool bInAutoZoomToFit) { bAutoZoomToFit = bInAutoZoomToFit; }
+	void ToggleAutoZoomToFit() { bAutoZoomToFit = !bAutoZoomToFit; }
 
 	bool GetAutoFlipTangents() const { return bAutoFlipTangents; }
 	void SetAutoFlipTangents(const bool bInAutoFlipTangents) { bAutoFlipTangents = bInAutoFlipTangents; }
 	void ToggleAutoFlipTangents() { bAutoFlipTangents = !bAutoFlipTangents; }
 
 	FString GetPresetsPath() const { return PresetsPath; }
+	void SetPresetsPath(const FString& InPresetsPath) { PresetsPath = InPresetsPath; }
+
 	FString GetNewPresetCategory() const { return NewPresetCategory; }
+	void SetNewPresetCategory(const FString& InNewPresetCategory) { NewPresetCategory = InNewPresetCategory; }
+
+	FString GetQuickEaseTangents() const { return QuickEaseTangents; }
+	void SetQuickEaseTangents(const FString& InTangents) { QuickEaseTangents = InTangents; }
 
 private:
 	/** The height of the curve ease tool in the details panel. */
@@ -53,6 +64,10 @@ private:
 	UPROPERTY(Config, EditAnywhere, Category = "EaseCurveTool", meta = (DisplayName = "Grid Size", UIMin = 4, ClampMin = 4, UIMax = 24, ClampMax = 24, Delta = 1))
 	int32 GridSize = 8;
 
+	/** If true, will auto zoom the graph editor to fit the tangent handles after they have been changed. */
+	UPROPERTY(Config, EditAnywhere, Category = "EaseCurveTool", meta = (DisplayName = "Auto Zoom To Fit"))
+	bool bAutoZoomToFit = true;
+
 	/** If true, auto flips tangents when sequential key frame curve values are descending. */
 	UPROPERTY(Config, EditAnywhere, Category = "EaseCurveTool", meta = (DisplayName = "Auto Flip Tangents"))
 	bool bAutoFlipTangents = true;
@@ -64,4 +79,8 @@ private:
 	/** The name of the category to place newly created curve presets. */
 	UPROPERTY(Config, EditAnywhere, Category = "EaseCurveTool", meta = (DisplayName = "New Preset Category"))
 	FString NewPresetCategory;
+
+	/** The tangents to apply for quick ease. Should be in the format of four comma-separated cubic bezier points. Ex. "0.45, 0.34, 0.0, 1.00" */
+	UPROPERTY(Config, EditAnywhere, Category = "EaseCurveTool", meta = (DisplayName = "Quick Ease Tangents"))
+	FString QuickEaseTangents;
 };
diff --git a/Engine/Plugins/Experimental/Avalanche/Source/AvalancheSequencer/Private/EaseCurveTool/Widgets/SAvaEaseCurveEditor.cpp b/Engine/Plugins/Experimental/Avalanche/Source/AvalancheSequencer/Private/EaseCurveTool/Widgets/SAvaEaseCurveEditor.cpp
--- a/Engine/Plugins/Experimental/Avalanche/Source/AvalancheSequencer/Private/EaseCurveTool/Widgets/SAvaEaseCurveEditor.cpp
+++ b/Engine/Plugins/Experimental/Avalanche/Source/AvalancheSequencer/Private/EaseCurveTool/Widgets/SAvaEaseCurveEditor.cpp
@@ -4,8 +4,10 @@
 #include "EaseCurveTool/Widgets/SAvaEaseCurvePreview.h"
 #include "EaseCurveTool/Widgets/SAvaEaseCurvePreviewToolTip.h"
 #include "Curves/SimpleCurve.h"
+#include "EaseCurveTool/AvaEaseCurveStyle.h"
 #include "Editor.h"
 #include "Factories/Factory.h"
+#include "Fonts/FontMeasure.h"
 #include "Fonts/SlateFontInfo.h"
 #include "Framework/Application/MenuStack.h"
 #include "Framework/Application/SlateApplication.h"
@@ -22,7 +24,9 @@
 #include "Styling/StyleColors.h"
 #include "Textures/SlateIcon.h"
 #include "Widgets/Colors/SColorBlock.h"
+#include "Widgets/IToolTip.h"
 #include "Widgets/Layout/SBox.h"
+#include "Widgets/Layout/SWidgetSwitcher.h"
 #include "Widgets/SBoxPanel.h"
 #include "Widgets/SToolTip.h"
 #include "Widgets/Text/STextBlock.h"
@@ -31,16 +35,20 @@
 
 namespace UE::EaseCurveTool::Private
 {
-	const static FVector2d CONST_KeySize = FVector2d(7, 7);
-	const static FVector2d CONST_TangentSize = FVector2d(16, 16);
-	const static FVector2d CONST_CurveSize = FVector2d(12, 12);
+	static const FVector2D KeyHitSize = FVector2D(8.f);
+	static const FVector2D TangentHitSize = FVector2D(16.f);
+	static const FVector2D CurveHitSize = FVector2D(6.f);
 
-	constexpr static float CONST_FitMargin = 0.1f;
-	constexpr static float CONST_MinViewRange = 0.5f;
-	constexpr static float CONST_DefaultZoomRange = 1.f;
+	static const FVector2D HalfKeyHitSize = KeyHitSize * 0.5f;
+	static const FVector2D HalfTangentHitSize = TangentHitSize * 0.5f;
+	static const FVector2D HalfCurveHitSize = CurveHitSize * 0.5f;
 
-	constexpr static float CONST_NormalBoundsThickness = 1.f;
-	constexpr static float CONST_TangentHandleLineThickness = 2.f;
+	static constexpr float FitMargin = 0.1f;
+	static constexpr float MinViewRange = 0.5f;
+	static constexpr float HalfMinViewRange = MinViewRange * 0.5f;
+
+	static constexpr float NormalBoundsThickness = 1.f;
+	static constexpr float TangentHandleLineThickness = 2.f;
 }
 
 using namespace UE::EaseCurveTool::Private;
@@ -79,6 +87,11 @@ void SAvaEaseCurveEditor::Construct(const FArguments& InArgs, const TObjectPtr<U
 	CurveThickness = InArgs._CurveThickness;
 	CurveColor = InArgs._CurveColor;
 	Operation = InArgs._Operation;
+
+	StartText = InArgs._StartText;
+	StartTooltipText = InArgs._StartTooltipText;
+	EndText = InArgs._EndText;
+	EndTooltipText = InArgs._EndTooltipText;
 	
 	OnTangentsChanged = InArgs._OnTangentsChanged;
 	GetContextMenuContent = InArgs._GetContextMenuContent;
@@ -165,6 +178,35 @@ int32 SAvaEaseCurveEditor::OnPaint(const FPaintArgs& InArgs, const FGeometry& In
 
 	InLayerId = PaintGrid(ScaleInfo, InAllottedGeometry, OutDrawElements, ++InLayerId, InMyCullingRect, DrawEffects);
 
+	if (StartText.IsSet() || EndText.IsSet())
+	{
+		const FSlateFontInfo FontInfo = FAvaEaseCurveStyle::Get().GetFontStyle(TEXT("Editor.LabelFont"));
+		const FLinearColor TextColor = FLinearColor(0.1f, 0.1f, 0.1f, 1.f);
+		const FVector2D TextOffset = FVector2D(4.f);
+
+		if (StartText.IsSet())
+		{
+			const FText& Text = StartText.Get();
+			const TSharedRef<FSlateFontMeasure> FontMeasure = FSlateApplication::Get().GetRenderer()->GetFontMeasureService();
+			const FVector2D TextSize = FontMeasure->Measure(Text, FontInfo);
+
+			const FVector2D ActualOffset(ZeroInputX - (TextOffset.X + TextSize.X), ZeroOutputY - (TextOffset.Y + TextSize.Y));
+
+			FSlateDrawElement::MakeText(OutDrawElements, ++InLayerId
+				, InAllottedGeometry.ToPaintGeometry(InAllottedGeometry.Size, FSlateLayoutTransform(ActualOffset))
+				, Text, FontInfo, DrawEffects, TextColor);
+		}
+
+		if (EndText.IsSet())
+		{
+			const FVector2D ActualOffset(OneInputX + TextOffset.X, OneOutputY + TextOffset.Y);
+
+			FSlateDrawElement::MakeText(OutDrawElements, ++InLayerId
+				, InAllottedGeometry.ToPaintGeometry(InAllottedGeometry.Size, FSlateLayoutTransform(ActualOffset))
+				, EndText.Get(), FontInfo, DrawEffects, TextColor);
+		}
+	}
+
 	PaintNormalBounds(ScaleInfo, InAllottedGeometry, OutDrawElements, ++InLayerId, InMyCullingRect, DrawEffects, InWidgetStyle);
 
 	PaintCurve(ScaleInfo, InAllottedGeometry, OutDrawElements, ++InLayerId, InMyCullingRect, DrawEffects, InWidgetStyle);
@@ -226,7 +268,8 @@ int32 SAvaEaseCurveEditor::PaintGrid(const FTrackScaleInfo& InScaleInfo, const F
 				LinePoints.Add(FVector2D(0.f, ScreenY));
 				LinePoints.Add(FVector2D(LocalSize.X, ScreenY));
 
-				FSlateDrawElement::MakeLines(OutDrawElements, InLayerId, InAllottedGeometry.ToPaintGeometry(), LinePoints, InDrawEffects, ExtendedGridColor, false);
+				FSlateDrawElement::MakeLines(OutDrawElements, InLayerId
+					, InAllottedGeometry.ToPaintGeometry(), LinePoints, InDrawEffects, ExtendedGridColor, false);
 			}
 		}
 	}
@@ -246,7 +289,8 @@ int32 SAvaEaseCurveEditor::PaintGrid(const FTrackScaleInfo& InScaleInfo, const F
 			LinePoints.Add(FVector2D(ScreenX, CurveZeroY1));
 			LinePoints.Add(FVector2D(ScreenX, CurveZeroY2));
 
-			FSlateDrawElement::MakeLines(OutDrawElements, NormalizedLayerId, InAllottedGeometry.ToPaintGeometry(), LinePoints, InDrawEffects, GridColor, false);
+			FSlateDrawElement::MakeLines(OutDrawElements, NormalizedLayerId
+				, InAllottedGeometry.ToPaintGeometry(), LinePoints, InDrawEffects, GridColor, false);
 		}
 
 		// Horizontal grid lines
@@ -261,7 +305,8 @@ int32 SAvaEaseCurveEditor::PaintGrid(const FTrackScaleInfo& InScaleInfo, const F
 			LinePoints.Add(FVector2D(CurveZeroX1, ScreenY));
 			LinePoints.Add(FVector2D(CurveZeroX2, ScreenY));
 
-			FSlateDrawElement::MakeLines(OutDrawElements, NormalizedLayerId, InAllottedGeometry.ToPaintGeometry(), LinePoints, InDrawEffects, GridColor, false);
+			FSlateDrawElement::MakeLines(OutDrawElements, NormalizedLayerId
+				, InAllottedGeometry.ToPaintGeometry(), LinePoints, InDrawEffects, GridColor, false);
 		}
 	}
 
@@ -280,25 +325,29 @@ void SAvaEaseCurveEditor::PaintNormalBounds(const FTrackScaleInfo& InScaleInfo,
 	TArray<FVector2D> ZeroTimeLinePoints;
 	ZeroTimeLinePoints.Add(FVector2D(ZeroInputX, 0.f));
 	ZeroTimeLinePoints.Add(FVector2D(ZeroInputX, InAllottedGeometry.GetLocalSize().Y));
-	FSlateDrawElement::MakeLines(OutDrawElements, InLayerId, InAllottedGeometry.ToPaintGeometry(), ZeroTimeLinePoints, InDrawEffects, NormalBoundsColor, false, CONST_NormalBoundsThickness);
+	FSlateDrawElement::MakeLines(OutDrawElements, InLayerId, InAllottedGeometry.ToPaintGeometry()
+		, ZeroTimeLinePoints, InDrawEffects, NormalBoundsColor, false, NormalBoundsThickness);
 
 	// Time = 1 line
 	TArray<FVector2D> OneTimeLinePoints;
 	OneTimeLinePoints.Add(FVector2D(OneInputX, 0.f));
 	OneTimeLinePoints.Add(FVector2D(OneInputX, InAllottedGeometry.GetLocalSize().Y));
-	FSlateDrawElement::MakeLines(OutDrawElements, InLayerId, InAllottedGeometry.ToPaintGeometry(), OneTimeLinePoints, InDrawEffects, NormalBoundsColor, false, CONST_NormalBoundsThickness);
+	FSlateDrawElement::MakeLines(OutDrawElements, InLayerId, InAllottedGeometry.ToPaintGeometry()
+		, OneTimeLinePoints, InDrawEffects, NormalBoundsColor, false, NormalBoundsThickness);
 
 	// Value = 0 line
 	TArray<FVector2D> ZeroValueLinePoints;
 	ZeroValueLinePoints.Add(FVector2D(0.f, ZeroOutputY));
 	ZeroValueLinePoints.Add(FVector2D(InAllottedGeometry.GetLocalSize().X, ZeroOutputY));
-	FSlateDrawElement::MakeLines(OutDrawElements, InLayerId, InAllottedGeometry.ToPaintGeometry(), ZeroValueLinePoints, InDrawEffects, NormalBoundsColor, false, CONST_NormalBoundsThickness);
+	FSlateDrawElement::MakeLines(OutDrawElements, InLayerId, InAllottedGeometry.ToPaintGeometry()
+		, ZeroValueLinePoints, InDrawEffects, NormalBoundsColor, false, NormalBoundsThickness);
 
 	// Value = 1 line
 	TArray<FVector2D> OneValueLinePoints;
 	OneValueLinePoints.Add(FVector2D(0.f, OneOutputY));
 	OneValueLinePoints.Add(FVector2D(InAllottedGeometry.GetLocalSize().X, OneOutputY));
-	FSlateDrawElement::MakeLines(OutDrawElements, InLayerId, InAllottedGeometry.ToPaintGeometry(), OneValueLinePoints, InDrawEffects, NormalBoundsColor, false, CONST_NormalBoundsThickness);
+	FSlateDrawElement::MakeLines(OutDrawElements, InLayerId, InAllottedGeometry.ToPaintGeometry()
+		, OneValueLinePoints, InDrawEffects, NormalBoundsColor, false, NormalBoundsThickness);
 }
 
 void SAvaEaseCurveEditor::PaintCurve(const FTrackScaleInfo& InScaleInfo, const FGeometry& InAllottedGeometry, FSlateWindowElementList& OutDrawElements,
@@ -442,13 +491,13 @@ int32 SAvaEaseCurveEditor::PaintKeys(const FTrackScaleInfo& InScaleInfo, const F
 
 		const FVector2D KeyLocation(InScaleInfo.InputToLocalX(EaseCurve->FloatCurve.GetKeyTime(KeyHandle))
 			, InScaleInfo.OutputToLocalY(EaseCurve->FloatCurve.GetKeyValue(KeyHandle)));
-		const FVector2D KeyIconLocation = KeyLocation - (CONST_KeySize / 2);
+		const FVector2D KeyIconLocation = KeyLocation - HalfKeyHitSize;
 		const bool bIsSelected = (KeyHandle == SelectedTangent.KeyHandle);
-		const FSlateBrush* KeyBrush = /*bIsSelected ? FAppStyle::GetBrush("CurveEd.CurveKeySelected") : */ FAppStyle::GetBrush("CurveEd.CurveKey");
+		const FSlateBrush* KeyBrush = FAppStyle::GetBrush("CurveEd.CurveKey");
 		const int32 LayerToUse = bIsSelected ? SelectedLayerId : LayerId;
 
 		FSlateDrawElement::MakeBox(OutDrawElements, LayerToUse,
-			InAllottedGeometry.ToPaintGeometry(CONST_KeySize, FSlateLayoutTransform(KeyIconLocation)),
+			InAllottedGeometry.ToPaintGeometry(KeyHitSize, FSlateLayoutTransform(KeyIconLocation)),
 			KeyBrush, InDrawEffects,
 			KeyBrush->GetTint(InWidgetStyle) * InWidgetStyle.GetColorAndOpacityTint());
 
@@ -494,17 +543,18 @@ int32 SAvaEaseCurveEditor::PaintTangentHandle(const FGeometry& InAllottedGeometr
 		LinePoints.Add(FVector2D(InKeyLocation));
 		LinePoints.Add(FVector2D(InTangentLocation));
 
-		FSlateDrawElement::MakeLines(OutDrawElements, InLayerId++, InAllottedGeometry.ToPaintGeometry(), LinePoints, InDrawEffects, LineColor, true, CONST_TangentHandleLineThickness);
+		FSlateDrawElement::MakeLines(OutDrawElements, InLayerId++, InAllottedGeometry.ToPaintGeometry()
+			, LinePoints, InDrawEffects, LineColor, true, TangentHandleLineThickness);
 	}
 
 	// Draw tangent handle
 	{
 		const FSlateBrush* ImageBrush = FAppStyle::GetBrush(TEXT("Icons.BulletPoint"));
-		const FVector2D TangentIconLocation = InTangentLocation - (CONST_TangentSize / 2);
+		const FVector2D TangentIconLocation = InTangentLocation - HalfTangentHitSize;
 		const FLinearColor HandleColor = bInSelected ? FStyleColors::AccentBlue.GetSpecifiedColor() : FLinearColor::White;
 
 		FSlateDrawElement::MakeBox(OutDrawElements, InLayerId++
-			, InAllottedGeometry.ToPaintGeometry(CONST_TangentSize, FSlateLayoutTransform(TangentIconLocation))
+			, InAllottedGeometry.ToPaintGeometry(TangentHitSize, FSlateLayoutTransform(TangentIconLocation))
 			, ImageBrush, InDrawEffects, HandleColor * InWidgetStyle.GetColorAndOpacityTint());
 	}
 
@@ -593,7 +643,8 @@ FReply SAvaEaseCurveEditor::OnMouseButtonUp(const FGeometry& InMyGeometry, const
 	return FReply::Unhandled();
 }
 
-void ClampViewRangeToDataIfBound(float& NewViewMin, float& NewViewMax, const TAttribute< TOptional<float> >& DataMin, const TAttribute< TOptional<float> >& DataMax, const float ViewRange)
+void ClampViewRangeToDataIfBound(float& NewViewMin, float& NewViewMax
+	, const TAttribute<TOptional<float>>& DataMin, const TAttribute<TOptional<float>>& DataMax, const float ViewRange)
 {
 	const TOptional<float>& Min = DataMin.Get();
 	const TOptional<float>& Max = DataMax.Get();
@@ -632,58 +683,116 @@ FReply SAvaEaseCurveEditor::OnMouseMove(const FGeometry& InMyGeometry, const FPo
 	return FReply::Unhandled();
 }
 
+TSharedRef<IToolTip> SAvaEaseCurveEditor::CreateCurveToolTip()
+{
+	const FAvaEaseCurveTangents Tangents = EaseCurve->GetTangents();
+
+	SAvaEaseCurvePreview::FArguments PreviewArgs;
+	PreviewArgs.Tangents(Tangents)
+		.PreviewSize(256.f)
+		.CanExpandPreview(true)
+		.CustomToolTip(false)
+		.Animate(true)
+		.DisplayRate(DisplayRate.Get())
+		.DrawMotionTrails(true);
+
+	const TSharedRef<SWidget> AdditionalContent = SNew(SVerticalBox)
+		+ SVerticalBox::Slot()
+		.AutoHeight()
+		.Padding(15.f, 5.f, 0.f, 0.f)
+		[
+			SNew(STextBlock)
+			.Text(this, &SAvaEaseCurveEditor::GetCurveToolTipInputText)
+		]
+		+ SVerticalBox::Slot()
+		.AutoHeight()
+		.Padding(15.f, 5.f, 0.f, 15.f)
+		[
+			SNew(STextBlock)
+			.Text(this, &SAvaEaseCurveEditor::GetCurveToolTipOutputText)
+		];
+
+	return SNew(SToolTip)
+		.TextMargin(0.f)
+		.Visibility_Lambda([this]()
+			{
+				return (ToolTipIndex != INDEX_NONE) ? EVisibility::Visible : EVisibility::Hidden;
+			})
+		[
+			SNew(SWidgetSwitcher)
+			.WidgetIndex_Lambda([this]()
+				{
+					return (ToolTipIndex == INDEX_NONE) ? 0 : ToolTipIndex;
+				})
+			+ SWidgetSwitcher::Slot()
+			[
+				SAvaEaseCurvePreviewToolTip::CreateDefaultToolTip(PreviewArgs, AdditionalContent)
+			]
+			+ SWidgetSwitcher::Slot()
+			[
+				SNew(STextBlock)
+				.Margin(FMargin(5.f))
+				.Text(StartTooltipText)
+			]
+			+ SWidgetSwitcher::Slot()
+			[
+				SNew(STextBlock)
+				.Margin(FMargin(5.f))
+				.Text(EndTooltipText)
+			]
+		];
+}
+
 void SAvaEaseCurveEditor::UpdateCurveToolTip(const FGeometry& InMyGeometry, const FPointerEvent& InMouseEvent)
 {
-	if (DragState != EDragState::DragTangent && HitTestCurves(InMyGeometry, InMouseEvent))
-	{
-		const FTrackScaleInfo ScaleInfo(ViewMinInput, ViewMaxInput, ViewMinOutput, ViewMaxOutput, InMyGeometry.GetLocalSize());
-		const FVector2D HitPosition = InMyGeometry.AbsoluteToLocal(InMouseEvent.GetScreenSpacePosition());
-		const float Time = ScaleInfo.LocalXToInput(HitPosition.X);
-		const float Value = EaseCurve->FloatCurve.Eval(Time);
+	const FTrackScaleInfo ScaleInfo(ViewMinInput, ViewMaxInput, ViewMinOutput, ViewMaxOutput, InMyGeometry.GetLocalSize());
+	const FVector2D MousePosition = InMyGeometry.AbsoluteToLocal(InMouseEvent.GetScreenSpacePosition());
 
-		FNumberFormattingOptions FormattingOptions;
-		FormattingOptions.MaximumFractionalDigits = 2;
-		CurveToolTipOutputText = FText::Format(LOCTEXT("CurveToolTipValueFormat", "Value:\t{0}"), FText::AsNumber(Value, &FormattingOptions));
-		CurveToolTipInputText = FText::Format(LOCTEXT("CurveToolTipTimeFormat", "Time:\t{0}"), FText::AsNumber(Time, &FormattingOptions));
+	const float Time = ScaleInfo.LocalXToInput(MousePosition.X);
+	const float Value = EaseCurve->FloatCurve.Eval(Time);
+	const FVector2D TimeValueRange = FVector2D(Time, Value);
 
-		if (!CurveToolTip.IsValid())
-		{
-			const FAvaEaseCurveTangents Tangents = EaseCurve->GetTangents();
-
-			SAvaEaseCurvePreview::FArguments PreviewArgs;
-			PreviewArgs.Tangents(Tangents)
-				.PreviewSize(256.f)
-				.CanExpandPreview(true)
-				.CustomToolTip(false)
-				.Animate(true)
-				.DisplayRate(DisplayRate.Get())
-				.DrawMotionTrails(true);
-
-			const TSharedPtr<SWidget> AdditionalContent = SNew(SVerticalBox)
-				+ SVerticalBox::Slot()
-				.AutoHeight()
-				.Padding(15.f, 5.f, 0.f, 0.f)
-				[
-					SNew(STextBlock)
-					.Text(this, &SAvaEaseCurveEditor::GetCurveToolTipInputText)
-				]
-				+ SVerticalBox::Slot()
-				.AutoHeight()
-				.Padding(15.f, 5.f, 0.f, 15.f)
-				[
-					SNew(STextBlock)
-					.Text(this, &SAvaEaseCurveEditor::GetCurveToolTipOutputText)
-				];
-
-			CurveToolTip = SAvaEaseCurvePreviewToolTip::CreateDefaultToolTip(PreviewArgs, AdditionalContent);
-
-			SetToolTip(CurveToolTip);
-		}
+	FNumberFormattingOptions FormattingOptions;
+	FormattingOptions.MaximumFractionalDigits = 2;
+	CurveToolTipOutputText = FText::Format(LOCTEXT("CurveToolTipValueFormat", "Value:\t{0}"), FText::AsNumber(TimeValueRange.Y, &FormattingOptions));
+	CurveToolTipInputText = FText::Format(LOCTEXT("CurveToolTipTimeFormat", "Time:\t{0}"), FText::AsNumber(TimeValueRange.X, &FormattingOptions));
+
+	if (DragState != EDragState::None)
+	{
+		ToolTipIndex = INDEX_NONE;
+		CurveToolTip.Reset();
+		SetToolTip(CurveToolTip);
+		return;
+	}
+
+	if (StartTooltipText.IsSet() && HitTestKey(InMyGeometry, InMouseEvent, FVector2D(0.f, 0.f)))
+	{
+		ToolTipIndex = 1;
+	}
+	else if (EndTooltipText.IsSet() && HitTestKey(InMyGeometry, InMouseEvent, FVector2D(1.f, 1.f)))
+	{
+		ToolTipIndex = 2;
+	}
+	else if (HitTestCurves(InMyGeometry, InMouseEvent))
+	{
+		ToolTipIndex = 0;
 	}
 	else
+	{
+		ToolTipIndex = INDEX_NONE;
+	}
+
+	if (ToolTipIndex == INDEX_NONE)
 	{
 		CurveToolTip.Reset();
 		SetToolTip(CurveToolTip);
+		return;
+	}
+
+	if (!CurveToolTip.IsValid())
+	{
+		CurveToolTip = CreateCurveToolTip();
+		SetToolTip(CurveToolTip);
 	}
 }
 
@@ -703,17 +812,19 @@ void SAvaEaseCurveEditor::ZoomView(const FVector2D& InDelta)
 
 	const float OutputViewSize = ViewMaxOutput - ViewMinOutput;
 	const float OutputChange = OutputViewSize * ZoomDelta.Y;
+	const float HalfOutputChange = OutputChange * 0.5f;
 
-	const float NewMinOutput = ViewMinOutput - (OutputChange * 0.5f);
-	const float NewMaxOutput = ViewMaxOutput + (OutputChange * 0.5f);
+	const float NewMinOutput = ViewMinOutput - HalfOutputChange;
+	const float NewMaxOutput = ViewMaxOutput + HalfOutputChange;
 
 	SetOutputMinMax(NewMinOutput, NewMaxOutput);
 
 	const float InputViewSize = ViewMaxInput - ViewMinInput;
 	const float InputChange = InputViewSize * ZoomDelta.X;
+	const float HalfInputChange = InputChange * 0.5f;
 
-	const float NewMinInput = ViewMinInput - (InputChange * 0.5f);
-	const float NewMaxInput = ViewMaxInput + (InputChange * 0.5f);
+	const float NewMinInput = ViewMinInput - HalfInputChange;
+	const float NewMaxInput = ViewMaxInput + HalfInputChange;
 
 	SetInputMinMax(NewMinInput, NewMaxInput);
 }
@@ -762,7 +873,7 @@ void SAvaEaseCurveEditor::TryStartDrag(const FGeometry& InMyGeometry, const FPoi
 		if (bLeftMouseButton)
 		{
 			// Check if we should start dragging a tangent.
-			const FSelectedTangent HitTangent = HitTestCubicTangents(InMyGeometry, InMyGeometry.LocalToAbsolute(MouseDownLocation));
+			const FSelectedTangent HitTangent = HitTestTangentHandle(InMyGeometry, InMouseEvent);
 			if (HitTangent.IsValid())
 			{
 				SelectedTangent = HitTangent;
@@ -859,7 +970,7 @@ void SAvaEaseCurveEditor::ProcessClick(const FGeometry& InMyGeometry, const FPoi
 
 	if (EffectingButton == EKeys::LeftMouseButton)
 	{
-		const FSelectedTangent HitTangent = HitTestCubicTangents(InMyGeometry, InMouseEvent.GetScreenSpacePosition());
+		const FSelectedTangent HitTangent = HitTestTangentHandle(InMyGeometry, InMouseEvent);
 		if (HitTangent.IsValid())
 		{
 			SelectedTangent = HitTangent;
@@ -881,10 +992,11 @@ void SAvaEaseCurveEditor::ClearSelection()
 	SelectedTangent = FSelectedTangent();
 }
 
-void SAvaEaseCurveEditor::SetDefaultOutput(const float MinZoomRange)
+void SAvaEaseCurveEditor::SetDefaultOutput(const float InMinZoomRange)
 {
-	const float NewMinOutput = (ViewMinOutput - (0.5f * MinZoomRange));
-	const float NewMaxOutput = (ViewMaxOutput + (0.5f * MinZoomRange));
+	const float HalfMinZoomRange = InMinZoomRange * 0.5f;
+	const float NewMinOutput = ViewMinOutput - HalfMinZoomRange;
+	const float NewMaxOutput = ViewMaxOutput + HalfMinZoomRange;
 
 	SetOutputMinMax(NewMinOutput, NewMaxOutput);
 }
@@ -897,8 +1009,10 @@ void SAvaEaseCurveEditor::ZoomToFit()
 	const FVector2D ArriveTangentDir = CalcTangentDir(EndKey.ArriveTangent) * EndKey.ArriveTangentWeight;
 	const FVector2D LeaveTangentDir = CalcTangentDir(StartKey.LeaveTangent) * StartKey.LeaveTangentWeight;
 
-	const FVector2D StartKeyPosition(EaseCurve->FloatCurve.GetKeyTime(EaseCurve->GetStartKeyHandle()), EaseCurve->FloatCurve.GetKeyValue(EaseCurve->GetStartKeyHandle()));
-	const FVector2D EndKeyPosition(EaseCurve->FloatCurve.GetKeyTime(EaseCurve->GetEndKeyHandle()), EaseCurve->FloatCurve.GetKeyValue(EaseCurve->GetEndKeyHandle()));
+	const FVector2D StartKeyPosition(EaseCurve->FloatCurve.GetKeyTime(EaseCurve->GetStartKeyHandle())
+		, EaseCurve->FloatCurve.GetKeyValue(EaseCurve->GetStartKeyHandle()));
+	const FVector2D EndKeyPosition(EaseCurve->FloatCurve.GetKeyTime(EaseCurve->GetEndKeyHandle())
+		, EaseCurve->FloatCurve.GetKeyValue(EaseCurve->GetEndKeyHandle()));
 	const FVector2D ArrivePosition = -ArriveTangentDir + EndKeyPosition;
 	const FVector2D LeavePosition = LeaveTangentDir - StartKeyPosition;
 
@@ -922,16 +1036,16 @@ void SAvaEaseCurveEditor::ZoomToFit()
 
 		// Clamp the minimum size
 		float SizeX = InMaxX - InMinX;
-		if (SizeX < CONST_MinViewRange)
+		if (SizeX < MinViewRange)
 		{
-			InMinX -= 0.5f * CONST_MinViewRange;
-			InMaxX += 0.5f * CONST_MinViewRange;
+			InMinX -= HalfMinViewRange;
+			InMaxX += HalfMinViewRange;
 			SizeX = InMaxX - InMinX;
 		}
 
 		// Add margin
-		InMinX -= CONST_FitMargin * SizeX;
-		InMaxX += CONST_FitMargin * SizeX;
+		InMinX -= FitMargin * SizeX;
+		InMaxX += FitMargin * SizeX;
 
 		SetInputMinMax(InMinX, InMaxX);
 	}
@@ -963,17 +1077,18 @@ void SAvaEaseCurveEditor::ZoomToFit()
 
 		// Clamp the minimum size
 		float SizeY = InMaxY - InMinY;
-		if (SizeY < CONST_MinViewRange)
+		if (SizeY < MinViewRange)
 		{
-			SetDefaultOutput(CONST_MinViewRange);
+			SetDefaultOutput(MinViewRange);
+
 			InMinY = ViewMinOutput;
 			InMaxY = ViewMaxOutput;
 			SizeY = InMaxY - InMinY;
 		}
 
 		// Add margin
-		const float NewMinOutputY = (InMinY - CONST_FitMargin * SizeY);
-		const float NewMaxOutputY = (InMaxY + CONST_FitMargin * SizeY);
+		const float NewMinOutputY = (InMinY - FitMargin * SizeY);
+		const float NewMaxOutputY = (InMaxY + FitMargin * SizeY);
 
 		SetOutputMinMax(NewMinOutputY, NewMaxOutputY);
 	}
@@ -1004,11 +1119,13 @@ void SAvaEaseCurveEditor::SetOutputMinMax(const float InNewMin, const float InNe
 bool SAvaEaseCurveEditor::HitTestCurves(const FGeometry& InMyGeometry, const FPointerEvent& InMouseEvent) const
 {
 	const FTrackScaleInfo ScaleInfo(ViewMinInput, ViewMaxInput, ViewMinOutput, ViewMaxOutput, InMyGeometry.GetLocalSize());
-	const FVector2D HitPosition = InMyGeometry.AbsoluteToLocal(InMouseEvent.GetScreenSpacePosition());
-	const float Time = ScaleInfo.LocalXToInput(HitPosition.X);
-	const float KeyScreenY = ScaleInfo.OutputToLocalY(EaseCurve->FloatCurve.Eval(Time));
+	const FVector2D MousePosition = InMyGeometry.AbsoluteToLocal(InMouseEvent.GetScreenSpacePosition());
+
+	const float Time = ScaleInfo.LocalXToInput(MousePosition.X);
+	const float Value = ScaleInfo.OutputToLocalY(EaseCurve->FloatCurve.Eval(Time));
 
-	if (HitPosition.Y > (KeyScreenY - (0.5f * CONST_CurveSize.Y)) && HitPosition.Y < (KeyScreenY + (0.5f * CONST_CurveSize.Y)))
+	if (MousePosition.Y > (Value - HalfCurveHitSize.Y)
+		&& MousePosition.Y < (Value + HalfCurveHitSize.Y))
 	{
 		return true;
 	}
@@ -1016,21 +1133,21 @@ bool SAvaEaseCurveEditor::HitTestCurves(const FGeometry& InMyGeometry, const FPo
 	return false;
 }
 
-SAvaEaseCurveEditor::FSelectedTangent SAvaEaseCurveEditor::HitTestCubicTangents(const FGeometry& InMyGeometry, const FVector2D& InHitScreenPosition) const
+SAvaEaseCurveEditor::FSelectedTangent SAvaEaseCurveEditor::HitTestTangentHandle(const FGeometry& InMyGeometry, const FPointerEvent& InMouseEvent) const
 {
-	FSelectedTangent OutTangent;
-
 	const FTrackScaleInfo ScaleInfo(ViewMinInput, ViewMaxInput, ViewMinOutput, ViewMaxOutput, InMyGeometry.GetLocalSize());
-	const FVector2D HitPosition = InMyGeometry.AbsoluteToLocal(InHitScreenPosition);
+	const FVector2D MousePosition = InMyGeometry.AbsoluteToLocal(InMouseEvent.GetScreenSpacePosition());
 
-	auto IsInside = [&HitPosition](const FVector2D& InTangent) -> bool
+	auto IsInside = [&MousePosition](const FVector2D& InTangent) -> bool
 		{
-			return HitPosition.Y > (InTangent.Y - (0.5f * CONST_CurveSize.Y)) &&
-				HitPosition.Y < (InTangent.Y + (0.5f * CONST_CurveSize.Y)) &&
-				HitPosition.X > (InTangent.X - (0.5f * CONST_TangentSize.X)) &&
-				HitPosition.X < (InTangent.X + (0.5f * CONST_TangentSize.X));
+			return MousePosition.Y > (InTangent.Y - HalfTangentHitSize.Y)
+				&& MousePosition.Y < (InTangent.Y + HalfTangentHitSize.Y)
+				&& MousePosition.X > (InTangent.X - HalfTangentHitSize.X)
+				&& MousePosition.X < (InTangent.X + HalfTangentHitSize.X);
 		};
 
+	FSelectedTangent OutTangent;
+
 	for (auto It(EaseCurve->FloatCurve.GetKeyHandleIterator()); It; ++It)
 	{
 		const FKeyHandle KeyHandle = *It;
@@ -1060,6 +1177,20 @@ SAvaEaseCurveEditor::FSelectedTangent SAvaEaseCurveEditor::HitTestCubicTangents(
 	return OutTangent;
 }
 
+bool SAvaEaseCurveEditor::HitTestKey(const FGeometry& InMyGeometry, const FPointerEvent& InMouseEvent, const FVector2D& InInputPosition) const
+{
+	const FTrackScaleInfo ScaleInfo(ViewMinInput, ViewMaxInput, ViewMinOutput, ViewMaxOutput, InMyGeometry.GetLocalSize());
+	const FVector2D MousePosition = InMyGeometry.AbsoluteToLocal(InMouseEvent.GetScreenSpacePosition());
+
+	const float Time = ScaleInfo.InputToLocalX(InInputPosition.X);
+	const float Value = ScaleInfo.OutputToLocalY(InInputPosition.Y);
+
+	return MousePosition.Y > (Value - HalfKeyHitSize.Y)
+		&& MousePosition.Y < (Value + HalfKeyHitSize.Y)
+		&& MousePosition.X > (Time - HalfKeyHitSize.X)
+		&& MousePosition.X < (Time + HalfKeyHitSize.X);
+}
+
 void SAvaEaseCurveEditor::MoveSelectedTangent(const FTrackScaleInfo& InScaleInfo, const FVector2D& InScreenDelta)
 {
 	if (!SelectedTangent.IsValid())
diff --git a/Engine/Plugins/Experimental/Avalanche/Source/AvalancheSequencer/Private/EaseCurveTool/Widgets/SAvaEaseCurveEditor.h b/Engine/Plugins/Experimental/Avalanche/Source/AvalancheSequencer/Private/EaseCurveTool/Widgets/SAvaEaseCurveEditor.h
--- a/Engine/Plugins/Experimental/Avalanche/Source/AvalancheSequencer/Private/EaseCurveTool/Widgets/SAvaEaseCurveEditor.h
+++ b/Engine/Plugins/Experimental/Avalanche/Source/AvalancheSequencer/Private/EaseCurveTool/Widgets/SAvaEaseCurveEditor.h
@@ -21,6 +21,7 @@
 #include "Widgets/SWidget.h"
 
 class IDetailCategoryBuilder;
+class IToolTip;
 class FDetailWidgetRow;
 class SHorizontalBox;
 struct FCurveEditorScreenSpace;
@@ -67,6 +68,11 @@ public:
 		SLATE_ARGUMENT(FLinearColor, CurveColor)
 		SLATE_ATTRIBUTE(FAvaEaseCurveTool::EOperation, Operation)
 
+		SLATE_ATTRIBUTE(FText, StartText)
+		SLATE_ATTRIBUTE(FText, StartTooltipText)
+		SLATE_ATTRIBUTE(FText, EndText)
+		SLATE_ATTRIBUTE(FText, EndTooltipText)
+
 		SLATE_EVENT(FAvaOnEaseCurveChanged, OnTangentsChanged)
 		SLATE_EVENT(FOnGetContent, GetContextMenuContent)
 		SLATE_EVENT(FOnKeyDown, OnKeyDown)
@@ -204,18 +210,20 @@ protected:
 	FText GetCurveToolTipInputText() const;
 	FText GetCurveToolTipOutputText() const;
 
+	TSharedRef<IToolTip> CreateCurveToolTip();
 	void UpdateCurveToolTip(const FGeometry& MyGeometry, const FPointerEvent& MouseEvent);
 
 	void OnObjectPropertyChanged(UObject* Object, FPropertyChangedEvent& PropertyChangedEvent);
 
 	/** Set Default output values when range is too small */
-	virtual void SetDefaultOutput(const float MinZoomRange);
+	virtual void SetDefaultOutput(const float InMinZoomRange);
 
 	void SetInputMinMax(const float InNewMin, const float InNewMax);
 	void SetOutputMinMax(const float InNewMin, const float InNewMax);
 
-	FSelectedTangent HitTestCubicTangents(const FGeometry& InMyGeometry, const FVector2D& InHitScreenPosition) const;
 	bool HitTestCurves(const FGeometry& InMyGeometry, const FPointerEvent& InMouseEvent) const;
+	FSelectedTangent HitTestTangentHandle(const FGeometry& InMyGeometry, const FPointerEvent& InMouseEvent) const;
+	bool HitTestKey(const FGeometry& InMyGeometry, const FPointerEvent& InMouseEvent, const FVector2D& InInputPosition) const;
 
 	TObjectPtr<UAvaEaseCurve> EaseCurve;
 
@@ -239,6 +247,11 @@ protected:
 	FLinearColor CurveColor;
 	TAttribute<FAvaEaseCurveTool::EOperation> Operation;
 
+	TAttribute<FText> StartText;
+	TAttribute<FText> StartTooltipText;
+	TAttribute<FText> EndText;
+	TAttribute<FText> EndTooltipText;
+
 	FAvaOnEaseCurveChanged OnTangentsChanged;
 	FOnGetContent GetContextMenuContent;
 	FOnKeyDown OnKeyDownEvent;
@@ -262,7 +275,8 @@ protected:
 	/** Tangent values at the beginning of a drag operation. */
 	FAvaEaseCurveTangents PreDragTangents;
 
-	TSharedPtr<SToolTip> CurveToolTip;
+	TSharedPtr<IToolTip> CurveToolTip;
+	int8 ToolTipIndex = INDEX_NONE;
 
 	FText CurveToolTipInputText;
 	FText CurveToolTipOutputText;
diff --git a/Engine/Plugins/Experimental/Avalanche/Source/AvalancheSequencer/Private/EaseCurveTool/Widgets/SAvaEaseCurvePreset.cpp b/Engine/Plugins/Experimental/Avalanche/Source/AvalancheSequencer/Private/EaseCurveTool/Widgets/SAvaEaseCurvePreset.cpp
--- a/Engine/Plugins/Experimental/Avalanche/Source/AvalancheSequencer/Private/EaseCurveTool/Widgets/SAvaEaseCurvePreset.cpp
+++ b/Engine/Plugins/Experimental/Avalanche/Source/AvalancheSequencer/Private/EaseCurveTool/Widgets/SAvaEaseCurvePreset.cpp
@@ -4,6 +4,7 @@
 #include "EaseCurveTool/AvaEaseCurvePreset.h"
 #include "EaseCurveTool/AvaEaseCurveStyle.h"
 #include "EaseCurveTool/AvaEaseCurveSubsystem.h"
+#include "EaseCurveTool/AvaEaseCurveToolSettings.h"
 #include "EaseCurveTool/Widgets/SAvaEaseCurvePreview.h"
 #include "Editor.h"
 #include "DetailLayoutBuilder.h"
@@ -25,6 +26,7 @@ void SAvaEaseCurvePreset::Construct(const FArguments& InArgs)
 {
 	DisplayRate = InArgs._DisplayRate;
 	OnPresetChanged = InArgs._OnPresetChanged;
+	OnQuickPresetChanged = InArgs._OnQuickPresetChanged;
 	OnGetNewPresetTangents = InArgs._OnGetNewPresetTangents;
 
 	ChildSlot
@@ -284,6 +286,7 @@ void SAvaEaseCurvePreset::RegenerateGroupWrapBox()
 		TSharedRef<SAvaEaseCurvePresetGroup> NewGroupWidget = SNew(SAvaEaseCurvePresetGroup)
 			.CategoryName(Category)
 			.Presets(EaseCurveSubsystem.GetEaseCurvePresets(Category))
+			.SelectedPreset(SelectedItem)
 			.IsEditMode(bIsInEditMode)
 			.DisplayRate(DisplayRate.Get())
 			.OnCategoryDelete(this, &SAvaEaseCurvePreset::HandleCategoryDelete)
@@ -316,6 +319,25 @@ TSharedRef<SWidget> SAvaEaseCurvePreset::GenerateSelectedRowWidget() const
 	}
 
 	return SNew(SHorizontalBox)
+		+ SHorizontalBox::Slot()
+		.AutoWidth()
+		.HAlign(HAlign_Right)
+		.VAlign(VAlign_Center)
+		.Padding(5.f, 2.f, 4.f, 2.f)
+		[
+			SNew(SBorder)
+			.BorderBackgroundColor(FStyleColors::White25)
+			[
+				SNew(SAvaEaseCurvePreview)
+				.PreviewSize(12.f)
+				.CustomToolTip(true)
+				.DisplayRate(DisplayRate.Get())
+				.Tangents_Lambda([this]()
+					{
+						return SelectedItem.IsValid() ? SelectedItem->Tangents : FAvaEaseCurveTangents();
+					})
+			]
+		]
 		+ SHorizontalBox::Slot()
 		.HAlign(HAlign_Fill)
 		.VAlign(VAlign_Center)
@@ -346,25 +368,6 @@ TSharedRef<SWidget> SAvaEaseCurvePreset::GenerateSelectedRowWidget() const
 					})
 			]
 
-		]
-		+ SHorizontalBox::Slot()
-		.AutoWidth()
-		.HAlign(HAlign_Right)
-		.VAlign(VAlign_Center)
-		.Padding(5.f, 2.f, 4.f, 2.f)
-		[
-			SNew(SBorder)
-			.BorderBackgroundColor(FStyleColors::White25)
-			[
-				SNew(SAvaEaseCurvePreview)
-				.PreviewSize(12.f)
-				.CustomToolTip(true)
-				.DisplayRate(DisplayRate.Get())
-				.Tangents_Lambda([this]()
-					{
-						return SelectedItem.IsValid() ? SelectedItem->Tangents : FAvaEaseCurveTangents();
-					})
-			]
 		];
 }
 
@@ -596,20 +599,24 @@ bool SAvaEaseCurvePreset::HandleEndPresetMove(const TSharedPtr<FAvaEaseCurvePres
 	return true;
 }
 
-bool SAvaEaseCurvePreset::HandlePresetClick(const TSharedPtr<FAvaEaseCurvePreset>& InPreset)
+bool SAvaEaseCurvePreset::HandlePresetClick(const TSharedPtr<FAvaEaseCurvePreset>& InPreset, const FModifierKeysState& InModifierKeys)
 {
-	if (ComboBoxSlot)
+	if (InModifierKeys.IsShiftDown())
 	{
-		FSlateApplication::Get().SetAllUserFocus(ComboBoxSlot->GetWidget());
-
-		SetSelectedItem(*InPreset);
+		UAvaEaseCurveToolSettings* const EaseCurveToolSettings = GetMutableDefault<UAvaEaseCurveToolSettings>();
+		EaseCurveToolSettings->SetQuickEaseTangents(InPreset->Tangents.ToJson());
+		EaseCurveToolSettings->SaveConfig();
 
-		OnPresetChanged.ExecuteIfBound(InPreset);
+		OnQuickPresetChanged.ExecuteIfBound(InPreset);
 
 		return true;
 	}
 
-	return false;
+	SetSelectedItem(*InPreset);
+
+	OnPresetChanged.ExecuteIfBound(InPreset);
+
+	return true;
 }
 
 #undef LOCTEXT_NAMESPACE
diff --git a/Engine/Plugins/Experimental/Avalanche/Source/AvalancheSequencer/Private/EaseCurveTool/Widgets/SAvaEaseCurvePreset.h b/Engine/Plugins/Experimental/Avalanche/Source/AvalancheSequencer/Private/EaseCurveTool/Widgets/SAvaEaseCurvePreset.h
--- a/Engine/Plugins/Experimental/Avalanche/Source/AvalancheSequencer/Private/EaseCurveTool/Widgets/SAvaEaseCurvePreset.h
+++ b/Engine/Plugins/Experimental/Avalanche/Source/AvalancheSequencer/Private/EaseCurveTool/Widgets/SAvaEaseCurvePreset.h
@@ -25,6 +25,7 @@ public:
 	{}
 		SLATE_ATTRIBUTE(FFrameRate, DisplayRate)
 		SLATE_EVENT(FAvaOnPresetChanged, OnPresetChanged)
+		SLATE_EVENT(FAvaOnPresetChanged, OnQuickPresetChanged)
 		SLATE_EVENT(FAvaOnGetNewPresetTangents, OnGetNewPresetTangents)
 	SLATE_END_ARGS()
 
@@ -66,11 +67,12 @@ protected:
 	bool HandlePresetRename(const TSharedPtr<FAvaEaseCurvePreset>& InPreset, const FString& InNewName);
 	bool HandleBeginPresetMove(const TSharedPtr<FAvaEaseCurvePreset>& InPreset, const FString& InNewCategoryName);
 	bool HandleEndPresetMove(const TSharedPtr<FAvaEaseCurvePreset>& InPreset, const FString& InNewCategoryName);
-	bool HandlePresetClick(const TSharedPtr<FAvaEaseCurvePreset>& InPreset);
+	bool HandlePresetClick(const TSharedPtr<FAvaEaseCurvePreset>& InPreset, const FModifierKeysState& InModifierKeys);
 
 	TAttribute<FFrameRate> DisplayRate;
 
 	FAvaOnPresetChanged OnPresetChanged;
+	FAvaOnPresetChanged OnQuickPresetChanged;
 	FAvaOnGetNewPresetTangents OnGetNewPresetTangents;
 
 	TSharedPtr<SBox> GroupWidgetsParent;
diff --git a/Engine/Plugins/Experimental/Avalanche/Source/AvalancheSequencer/Private/EaseCurveTool/Widgets/SAvaEaseCurvePresetGroup.cpp b/Engine/Plugins/Experimental/Avalanche/Source/AvalancheSequencer/Private/EaseCurveTool/Widgets/SAvaEaseCurvePresetGroup.cpp
--- a/Engine/Plugins/Experimental/Avalanche/Source/AvalancheSequencer/Private/EaseCurveTool/Widgets/SAvaEaseCurvePresetGroup.cpp
+++ b/Engine/Plugins/Experimental/Avalanche/Source/AvalancheSequencer/Private/EaseCurveTool/Widgets/SAvaEaseCurvePresetGroup.cpp
@@ -23,6 +23,7 @@ void SAvaEaseCurvePresetGroup::Construct(const FArguments& InArgs)
 {
 	CategoryName = InArgs._CategoryName;
 	Presets = InArgs._Presets;
+	SelectedPreset = InArgs._SelectedPreset;
 	SearchText = InArgs._SearchText;
 	bIsEditMode = InArgs._IsEditMode;
 	DisplayRate = InArgs._DisplayRate;
@@ -137,12 +138,13 @@ TSharedRef<SWidget> SAvaEaseCurvePresetGroup::ConstructHeader()
 		];
 }
 
-TSharedRef<ITableRow> SAvaEaseCurvePresetGroup::GeneratePresetWidget(TSharedPtr<FAvaEaseCurvePreset> InPreset, const TSharedRef<STableViewBase>& InOwnerTable)
+TSharedRef<ITableRow> SAvaEaseCurvePresetGroup::GeneratePresetWidget(const TSharedPtr<FAvaEaseCurvePreset> InPreset, const TSharedRef<STableViewBase>& InOwnerTable)
 {
 	TSharedRef<SAvaEaseCurvePresetGroupItem> NewPresetWidget = SNew(SAvaEaseCurvePresetGroupItem, InOwnerTable)
 		.Preset(InPreset)
 		.IsEditMode(bIsEditMode)
 		.DisplayRate(DisplayRate)
+		.IsSelected(this, &SAvaEaseCurvePresetGroup::IsSelected, InPreset)
 		.OnClick(this, &SAvaEaseCurvePresetGroup::HandlePresetClick)
 		.OnDelete(this, &SAvaEaseCurvePresetGroup::HandlePresetDelete)
 		.OnRename(this, &SAvaEaseCurvePresetGroup::HandlePresetRename)
@@ -285,11 +287,11 @@ bool SAvaEaseCurvePresetGroup::HandlePresetEndMove(const TSharedPtr<FAvaEaseCurv
 	return false;
 }
 
-bool SAvaEaseCurvePresetGroup::HandlePresetClick(const TSharedPtr<FAvaEaseCurvePreset>& InPreset) const
+bool SAvaEaseCurvePresetGroup::HandlePresetClick(const TSharedPtr<FAvaEaseCurvePreset>& InPreset, const FModifierKeysState& InModifierKeys) const
 {
 	if (OnPresetClick.IsBound())
 	{
-		OnPresetClick.Execute(InPreset);
+		OnPresetClick.Execute(InPreset, InModifierKeys);
 
 		return true;
 	}
@@ -370,4 +372,14 @@ void SAvaEaseCurvePresetGroup::ResetDragBorder()
 	bIsOverDifferentCategory = false;
 }
 
+bool SAvaEaseCurvePresetGroup::IsSelected(const TSharedPtr<FAvaEaseCurvePreset> InPreset) const
+{
+	if (SelectedPreset.IsSet() && SelectedPreset.Get().IsValid())
+	{
+		return *InPreset == *SelectedPreset.Get();
+	}
+
+	return false;
+}
+
 #undef LOCTEXT_NAMESPACE
diff --git a/Engine/Plugins/Experimental/Avalanche/Source/AvalancheSequencer/Private/EaseCurveTool/Widgets/SAvaEaseCurvePresetGroup.h b/Engine/Plugins/Experimental/Avalanche/Source/AvalancheSequencer/Private/EaseCurveTool/Widgets/SAvaEaseCurvePresetGroup.h
--- a/Engine/Plugins/Experimental/Avalanche/Source/AvalancheSequencer/Private/EaseCurveTool/Widgets/SAvaEaseCurvePresetGroup.h
+++ b/Engine/Plugins/Experimental/Avalanche/Source/AvalancheSequencer/Private/EaseCurveTool/Widgets/SAvaEaseCurvePresetGroup.h
@@ -21,6 +21,7 @@ public:
 	{}
 		SLATE_ARGUMENT(FString, CategoryName)
 		SLATE_ARGUMENT(TArray<TSharedPtr<FAvaEaseCurvePreset>>, Presets)
+		SLATE_ATTRIBUTE(TSharedPtr<FAvaEaseCurvePreset>, SelectedPreset)
 		SLATE_ARGUMENT(FText, SearchText)
 		SLATE_ATTRIBUTE(bool, IsEditMode)
 		SLATE_ARGUMENT(FFrameRate, DisplayRate)
@@ -30,7 +31,7 @@ public:
 		SLATE_EVENT(FAvaEaseCurvePresetRenameDelegate, OnPresetRename)
 		SLATE_EVENT(FAvaEaseCurvePresetMoveDelegate, OnBeginPresetMove)
 		SLATE_EVENT(FAvaEaseCurvePresetMoveDelegate, OnEndPresetMove)
-		SLATE_EVENT(FAvaEaseCurvePresetDelegate, OnPresetClick)
+		SLATE_EVENT(FAvaEaseCurvePresetClickDelegate, OnPresetClick)
 	SLATE_END_ARGS()
 
 	void Construct(const FArguments& InArgs);
@@ -47,10 +48,12 @@ public:
 
 	FString GetCategoryName() const { return CategoryName; }
 
+	bool IsSelected(const TSharedPtr<FAvaEaseCurvePreset> InPreset) const;
+
 protected:
 	TSharedRef<SWidget> ConstructHeader();
 
-	TSharedRef<ITableRow> GeneratePresetWidget(TSharedPtr<FAvaEaseCurvePreset> InPreset, const TSharedRef<STableViewBase>& InOwnerTable);
+	TSharedRef<ITableRow> GeneratePresetWidget(const TSharedPtr<FAvaEaseCurvePreset> InPreset, const TSharedRef<STableViewBase>& InOwnerTable);
 	
 	EVisibility GetEditModeVisibility() const;
 
@@ -61,7 +64,7 @@ protected:
 	void HandleCategoryRenameCommitted(const FText& InNewText, ETextCommit::Type InCommitType);
 	FReply HandleCategoryDelete() const;
 
-	bool HandlePresetClick(const TSharedPtr<FAvaEaseCurvePreset>& InPreset) const;
+	bool HandlePresetClick(const TSharedPtr<FAvaEaseCurvePreset>& InPreset, const FModifierKeysState& InModifierKeys) const;
 	bool HandlePresetDelete(const TSharedPtr<FAvaEaseCurvePreset>& InPreset);
 	bool HandlePresetRename(const TSharedPtr<FAvaEaseCurvePreset>& InPreset, const FString& InNewName);
 	bool HandlePresetBeginMove(const TSharedPtr<FAvaEaseCurvePreset>& InPreset, const FString& InNewCategoryName) const;
@@ -79,13 +82,14 @@ protected:
 	FText SearchText;
 	TAttribute<bool> bIsEditMode;
 	FFrameRate DisplayRate;
+	TAttribute<TSharedPtr<FAvaEaseCurvePreset>> SelectedPreset;
 	FAvaEaseCurveCategoryDeleteDelegate OnCategoryDelete;
 	FAvaEaseCurveCategoryRenameDelegate OnCategoryRename;
 	FAvaEaseCurvePresetDelegate OnPresetDelete;
 	FAvaEaseCurvePresetRenameDelegate OnPresetRename;
 	FAvaEaseCurvePresetMoveDelegate OnBeginPresetMove;
 	FAvaEaseCurvePresetMoveDelegate OnEndPresetMove;
-	FAvaEaseCurvePresetDelegate OnPresetClick;
+	FAvaEaseCurvePresetClickDelegate OnPresetClick;
 
 	TArray<TSharedPtr<FAvaEaseCurvePreset>> VisiblePresets;
 
diff --git a/Engine/Plugins/Experimental/Avalanche/Source/AvalancheSequencer/Private/EaseCurveTool/Widgets/SAvaEaseCurvePresetGroupItem.cpp b/Engine/Plugins/Experimental/Avalanche/Source/AvalancheSequencer/Private/EaseCurveTool/Widgets/SAvaEaseCurvePresetGroupItem.cpp
--- a/Engine/Plugins/Experimental/Avalanche/Source/AvalancheSequencer/Private/EaseCurveTool/Widgets/SAvaEaseCurvePresetGroupItem.cpp
+++ b/Engine/Plugins/Experimental/Avalanche/Source/AvalancheSequencer/Private/EaseCurveTool/Widgets/SAvaEaseCurvePresetGroupItem.cpp
@@ -1,9 +1,11 @@
 // Copyright Epic Games, Inc. All Rights Reserved.
 
 #include "SAvaEaseCurvePresetGroupItem.h"
-#include "EaseCurveTool/AvaEaseCurveStyle.h"
 #include "EaseCurveTool/AvaEaseCurvePreset.h"
 #include "EaseCurveTool/AvaEaseCurvePresetDragDropOp.h"
+#include "EaseCurveTool/AvaEaseCurveStyle.h"
+#include "EaseCurveTool/AvaEaseCurveToolCommands.h"
+#include "EaseCurveTool/AvaEaseCurveToolSettings.h"
 #include "EaseCurveTool/Widgets/SAvaEaseCurvePreview.h"
 #include "Internationalization/Text.h"
 #include "Styling/AppStyle.h"
@@ -25,12 +27,15 @@ void SAvaEaseCurvePresetGroupItem::Construct(const FArguments& InArgs, const TSh
 {
 	Preset = InArgs._Preset;
 	bIsEditMode = InArgs._IsEditMode;
+	IsSelected = InArgs._IsSelected;
 	OnClick = InArgs._OnClick;
 	OnDelete = InArgs._OnDelete;
 	OnRename = InArgs._OnRename;
 	OnBeginMove = InArgs._OnBeginMove;
 	OnEndMove = InArgs._OnEndMove;
 
+	const FText ItemTooltipText = FText::Format(LOCTEXT("ItemTooltip", "{0}\n\nShift + Click to set as active quick preset"), FText::FromString(Preset->Name));
+
 	ChildSlot
 		[
 			SNew(SOverlay)
@@ -60,11 +65,30 @@ void SAvaEaseCurvePresetGroupItem::Construct(const FArguments& InArgs, const TSh
 							})
 						+ SWidgetSwitcher::Slot()
 						[
-							SNew(STextBlock)
-							.TextStyle(FAppStyle::Get(), TEXT("Menu.Label"))
-							.OverflowPolicy(ETextOverflowPolicy::Ellipsis)
-							.Text(FText::FromString(Preset->Name))
-							.ToolTipText(FText::FromString(Preset->Name))
+							SNew(SHorizontalBox)
+							+ SHorizontalBox::Slot()
+							.AutoWidth()
+							.VAlign(VAlign_Center)
+							[
+								SNew(SBox)
+								.Padding(0.f, 0.f, 0.f, 3.f)
+								.Visibility(this, &SAvaEaseCurvePresetGroupItem::GetQuickPresetIconVisibility)
+								[
+									SNew(SImage)
+									.DesiredSizeOverride(FVector2D(10.f))
+									.Image(FAppStyle::GetBrush(TEXT("Icons.Adjust")))
+									.ToolTipText(this, &SAvaEaseCurvePresetGroupItem::GetQuickPresetIconToolTip)
+								]
+							]
+							+ SHorizontalBox::Slot()
+							[
+								SNew(STextBlock)
+								.TextStyle(FAppStyle::Get(), TEXT("Menu.Label"))
+								.OverflowPolicy(ETextOverflowPolicy::Ellipsis)
+								.ColorAndOpacity(FStyleColors::White)
+								.Text(FText::FromString(Preset->Name))
+								.ToolTipText(ItemTooltipText)
+							]
 						]
 						+ SWidgetSwitcher::Slot()
 						[
@@ -106,6 +130,7 @@ void SAvaEaseCurvePresetGroupItem::Construct(const FArguments& InArgs, const TSh
 					SNew(SBorder)
 					.BorderBackgroundColor(FStyleColors::White25)
 					.Padding(2.f)
+					.OnMouseButtonDown(this, &SAvaEaseCurvePresetGroupItem::OnMouseButtonDown)
 					[
 						SNew(SAvaEaseCurvePreview)
 						.PreviewSize(20.f)
@@ -132,14 +157,6 @@ void SAvaEaseCurvePresetGroupItem::SetPreset(const TSharedPtr<FAvaEaseCurvePrese
 	Preset = InPreset;
 }
 
-void SAvaEaseCurvePresetGroupItem::HandlePresetClick() const
-{
-	if (OnClick.IsBound())
-	{
-		OnClick.Execute(Preset);
-	}
-}
-
 bool SAvaEaseCurvePresetGroupItem::IsEditMode() const
 {
 	return bIsEditMode.Get(false);
@@ -152,12 +169,24 @@ EVisibility SAvaEaseCurvePresetGroupItem::GetEditModeVisibility() const
 
 EVisibility SAvaEaseCurvePresetGroupItem::GetBorderVisibility() const
 {
-	return bIsDragging && IsEditMode() ? EVisibility::Visible : EVisibility::Collapsed;
+	return ((bIsDragging && IsEditMode()) || (IsSelected.IsSet() && IsSelected.Get()))
+		? EVisibility::Visible: EVisibility::Collapsed;
 }
 
 const FSlateBrush* SAvaEaseCurvePresetGroupItem::GetBackgroundImage() const
 {
-	return bIsDragging ? FAvaEaseCurveStyle::Get().GetBrush(TEXT("EditMode.Background.Over")) : nullptr;
+	FName BrushName = NAME_None;
+
+	if (bIsDragging)
+	{
+		BrushName = TEXT("EditMode.Background.Over");
+	}
+	else if (IsSelected.IsSet() && IsSelected.Get())
+	{
+		BrushName = TEXT("Preset.Selected");
+	}
+
+	return BrushName.IsNone() ? nullptr : FAvaEaseCurveStyle::Get().GetBrush(BrushName);
 }
 
 void SAvaEaseCurvePresetGroupItem::HandleRenameTextCommitted(const FText& InNewText, ETextCommit::Type InCommitType) const
@@ -194,7 +223,7 @@ FReply SAvaEaseCurvePresetGroupItem::OnMouseButtonDown(const FGeometry& InGeomet
 
 	if (OnClick.IsBound())
 	{
-		OnClick.Execute(Preset);
+		OnClick.Execute(Preset, InMouseEvent.GetModifierKeys());
 	}
 
 	return FReply::Handled();
@@ -247,4 +276,46 @@ void SAvaEaseCurvePresetGroupItem::TriggerEndMove()
 	bIsDragging = false;
 }
 
+EVisibility SAvaEaseCurvePresetGroupItem::GetQuickPresetIconVisibility() const
+{
+	const UAvaEaseCurveToolSettings* const Settings = GetDefault<UAvaEaseCurveToolSettings>();
+
+	FAvaEaseCurveTangents Tangents;
+	if (!FAvaEaseCurveTangents::FromString(Settings->GetQuickEaseTangents(), Tangents))
+	{
+		return EVisibility::Hidden;
+	}
+
+	return (Preset->Tangents == Tangents) ? EVisibility::Visible : EVisibility::Hidden;
+}
+
+FText SAvaEaseCurvePresetGroupItem::GetQuickPresetIconToolTip() const
+{
+	static const FText QuickEaseText = LOCTEXT("QuickEaseIconTooltip", "Active Quick Ease Preset");
+
+	const FAvaEaseCurveToolCommands& EaseCurveToolCommands = FAvaEaseCurveToolCommands::Get();
+
+	FText CommandText;
+
+	if (EaseCurveToolCommands.QuickEase->GetFirstValidChord()->IsValidChord())
+	{
+		CommandText = FText::Format(LOCTEXT("QuickEaseIconInputInOutTooltip", "{0} - Apply quick preset to Out (Leave) and In (Arrive) tangents\n")
+			, EaseCurveToolCommands.QuickEase->GetInputText());
+	}
+
+	if (EaseCurveToolCommands.QuickEaseIn->GetFirstValidChord()->IsValidChord())
+	{
+		CommandText = FText::Format(LOCTEXT("QuickEaseIconInputInTooltip", "{0} - Apply quick preset to In (Arrive) tangent only\n")
+			, EaseCurveToolCommands.QuickEaseIn->GetInputText());
+	}
+
+	if (EaseCurveToolCommands.QuickEaseOut->GetFirstValidChord()->IsValidChord())
+	{
+		CommandText = FText::Format(LOCTEXT("QuickEaseIconInputOutTooltip", "{0} - Apply quick preset to Out (Leave) tangent only\n")
+			, EaseCurveToolCommands.QuickEaseOut->GetInputText());
+	}
+
+	return CommandText.IsEmpty() ? QuickEaseText : FText::Format(LOCTEXT("", "{0}\n\n{1}"), QuickEaseText, CommandText);
+}
+
 #undef LOCTEXT_NAMESPACE
diff --git a/Engine/Plugins/Experimental/Avalanche/Source/AvalancheSequencer/Private/EaseCurveTool/Widgets/SAvaEaseCurvePresetGroupItem.h b/Engine/Plugins/Experimental/Avalanche/Source/AvalancheSequencer/Private/EaseCurveTool/Widgets/SAvaEaseCurvePresetGroupItem.h
--- a/Engine/Plugins/Experimental/Avalanche/Source/AvalancheSequencer/Private/EaseCurveTool/Widgets/SAvaEaseCurvePresetGroupItem.h
+++ b/Engine/Plugins/Experimental/Avalanche/Source/AvalancheSequencer/Private/EaseCurveTool/Widgets/SAvaEaseCurvePresetGroupItem.h
@@ -12,7 +12,8 @@ class SEditableTextBox;
 class STableViewBase;
 struct FAvaEaseCurvePreset;
 
-DECLARE_DELEGATE_RetVal_OneParam(bool, FAvaEaseCurvePresetDelegate, const TSharedPtr<FAvaEaseCurvePreset>&)
+DECLARE_DELEGATE_RetVal_OneParam(bool, FAvaEaseCurvePresetDelegate, const TSharedPtr<FAvaEaseCurvePreset>& /*InPreset*/)
+DECLARE_DELEGATE_RetVal_TwoParams(bool, FAvaEaseCurvePresetClickDelegate, const TSharedPtr<FAvaEaseCurvePreset>& /*InPreset*/, const FModifierKeysState& /*InModifierKeys*/)
 DECLARE_DELEGATE_RetVal_TwoParams(bool, FAvaEaseCurvePresetRenameDelegate, const TSharedPtr<FAvaEaseCurvePreset>& /*InPreset*/, const FString& /*InNewName*/)
 DECLARE_DELEGATE_RetVal_TwoParams(bool, FAvaEaseCurvePresetMoveDelegate, const TSharedPtr<FAvaEaseCurvePreset>& /*InPreset*/, const FString& /*InNewCategoryName*/)
 
@@ -25,7 +26,8 @@ public:
 		SLATE_ARGUMENT(TSharedPtr<FAvaEaseCurvePreset>, Preset)
 		SLATE_ATTRIBUTE(bool, IsEditMode)
 		SLATE_ARGUMENT(FFrameRate, DisplayRate)
-		SLATE_EVENT(FAvaEaseCurvePresetDelegate, OnClick)
+		SLATE_ATTRIBUTE(bool, IsSelected)
+		SLATE_EVENT(FAvaEaseCurvePresetClickDelegate, OnClick)
 		SLATE_EVENT(FAvaEaseCurvePresetDelegate, OnDelete)
 		SLATE_EVENT(FAvaEaseCurvePresetRenameDelegate, OnRename)
 		SLATE_EVENT(FAvaEaseCurvePresetMoveDelegate, OnBeginMove)
@@ -53,6 +55,9 @@ protected:
 	EVisibility GetBorderVisibility() const;
 	const FSlateBrush* GetBackgroundImage() const;
 
+	EVisibility GetQuickPresetIconVisibility() const;
+	FText GetQuickPresetIconToolTip() const;
+
 	//~ Begin SWidget
 	virtual FReply OnMouseButtonDown(const FGeometry& InGeometry, const FPointerEvent& InMouseEvent) override;
 	virtual FReply OnMouseButtonUp(const FGeometry& InGeometry, const FPointerEvent& InMouseEvent) override;
@@ -62,7 +67,8 @@ protected:
 
 	TSharedPtr<FAvaEaseCurvePreset> Preset;
 	TAttribute<bool> bIsEditMode;
-	FAvaEaseCurvePresetDelegate OnClick;
+	TAttribute<bool> IsSelected;
+	FAvaEaseCurvePresetClickDelegate OnClick;
 	FAvaEaseCurvePresetDelegate OnDelete;
 	FAvaEaseCurvePresetRenameDelegate OnRename;
 	FAvaEaseCurvePresetMoveDelegate OnBeginMove;
diff --git a/Engine/Plugins/Experimental/Avalanche/Source/AvalancheSequencer/Private/EaseCurveTool/Widgets/SAvaEaseCurveTool.cpp b/Engine/Plugins/Experimental/Avalanche/Source/AvalancheSequencer/Private/EaseCurveTool/Widgets/SAvaEaseCurveTool.cpp
--- a/Engine/Plugins/Experimental/Avalanche/Source/AvalancheSequencer/Private/EaseCurveTool/Widgets/SAvaEaseCurveTool.cpp
+++ b/Engine/Plugins/Experimental/Avalanche/Source/AvalancheSequencer/Private/EaseCurveTool/Widgets/SAvaEaseCurveTool.cpp
@@ -3,6 +3,7 @@
 #include "SAvaEaseCurveTool.h"
 #include "Curves/KeyHandle.h"
 #include "Curves/RichCurve.h"
+#include "DetailLayoutBuilder.h"
 #include "EaseCurveTool/AvaEaseCurvePreset.h"
 #include "EaseCurveTool/AvaEaseCurveTool.h"
 #include "EaseCurveTool/AvaEaseCurveToolCommands.h"
@@ -10,9 +11,9 @@
 #include "EaseCurveTool/Widgets/SAvaEaseCurveEditor.h"
 #include "EaseCurveTool/Widgets/SAvaEaseCurvePreset.h"
 #include "Editor.h"
-#include "DetailLayoutBuilder.h"
 #include "Framework/Commands/GenericCommands.h"
 #include "Framework/MultiBox/MultiBoxBuilder.h"
+#include "IAvaSequencer.h"
 #include "Math/UnrealMathUtility.h"
 #include "SCurveEditor.h"
 #include "Styling/AppStyle.h"
@@ -115,12 +116,12 @@ namespace UE::AvaSequencer
 	}
 }
 
-void SAvaEaseCurveTool::Construct(const FArguments& InArgs, const TSharedRef<FAvaEaseCurveTool>& InCurveEaseTool)
+void SAvaEaseCurveTool::Construct(const FArguments& InArgs, const TSharedRef<FAvaEaseCurveTool>& InEaseCurveTool)
 {
 	ToolMode = InArgs._ToolMode;
 	ToolOperation = InArgs._ToolOperation;
 
-	CurveEaseTool = InCurveEaseTool;
+	EaseCurveTool = InEaseCurveTool;
 
 	ChildSlot
 		[
@@ -132,9 +133,10 @@ void SAvaEaseCurveTool::Construct(const FArguments& InArgs, const TSharedRef<FAv
 			[
 				SAssignNew(CurvePresetWidget, SAvaEaseCurvePreset)
 				.OnPresetChanged(this, &SAvaEaseCurveTool::OnPresetChanged)
+				.OnQuickPresetChanged(this, &SAvaEaseCurveTool::OnQuickPresetChanged)
 				.OnGetNewPresetTangents_Lambda([this](FAvaEaseCurveTangents& OutTangents) -> bool
 					{
-						OutTangents = CurveEaseTool->GetEaseCurveTangents();
+						OutTangents = EaseCurveTool->GetEaseCurveTangents();
 						return true;
 					})
 			]
@@ -166,14 +168,14 @@ TSharedRef<SWidget> SAvaEaseCurveTool::ConstructCurveEditorPanel()
 {
 	CurrentGraphSize = GetDefault<UAvaEaseCurveToolSettings>()->GetGraphSize();
 
-	const TSharedRef<FAvaEaseCurveTool> EaseCurveToolRef = CurveEaseTool.ToSharedRef();
+	const TSharedRef<FAvaEaseCurveTool> EaseCurveToolRef = EaseCurveTool.ToSharedRef();
 
 	return SNew(SBorder)
 		[
 			SNew(SOverlay)
 			+ SOverlay::Slot()
 			[
-				SAssignNew(CurveEaseEditorWidget, SAvaEaseCurveEditor, CurveEaseTool->GetToolCurve())
+				SAssignNew(CurveEaseEditorWidget, SAvaEaseCurveEditor, EaseCurveTool->GetToolCurve())
 				.DisplayRate(EaseCurveToolRef, &FAvaEaseCurveTool::GetDisplayRate)
 				.Operation(EaseCurveToolRef, &FAvaEaseCurveTool::GetOperation)
 				.DesiredSize_Lambda([this]() -> FVector2D
@@ -184,6 +186,10 @@ TSharedRef<SWidget> SAvaEaseCurveTool::ConstructCurveEditorPanel()
 				.GridSnap_UObject(GetDefault<UAvaEaseCurveToolSettings>(), &UAvaEaseCurveToolSettings::GetGridSnap)
 				.GridSize_UObject(GetDefault<UAvaEaseCurveToolSettings>(), &UAvaEaseCurveToolSettings::GetGridSize)
 				.GetContextMenuContent(this, &SAvaEaseCurveTool::CreateContextMenuContent)
+				.StartText(this, &SAvaEaseCurveTool::GetStartText)
+				.StartTooltipText(this, &SAvaEaseCurveTool::GetStartTooltipText)
+				.EndText(this, &SAvaEaseCurveTool::GetEndText)
+				.EndTooltipText(this, &SAvaEaseCurveTool::GetEndTooltipText)
 				.OnKeyDown(this, &SAvaEaseCurveTool::OnKeyDown)
 				.OnDragStart(this, &SAvaEaseCurveTool::OnEditorDragStart)
 				.OnDragEnd(this, &SAvaEaseCurveTool::OnEditorDragEnd)
@@ -262,22 +268,23 @@ TSharedRef<SWidget> SAvaEaseCurveTool::ConstructInputBoxes()
 			.HAlign(HAlign_Fill)
 			.VAlign(VAlign_Center)
 			[
-				ConstructTangentNumBox(LOCTEXT("OutTangentLabel", "T")
-					, LOCTEXT("OutTangentToolTip", "Out Tangent")
-					, TAttribute<float>::CreateSP(this, &SAvaEaseCurveTool::GetStartTangent)
-					, SNumericEntryBox<float>::FOnValueChanged::CreateSP(this, &SAvaEaseCurveTool::OnStartTangentSpinBoxChanged)
-					, MinTangent, MaxTangent)
+				ConstructTangentNumBox(LOCTEXT("OutTangentWeightLabel", "W")
+					, LOCTEXT("OutTangentWeightToolTip", "Out Tangent Weight")
+					, TAttribute<float>::CreateSP(this, &SAvaEaseCurveTool::GetStartTangentWeight)
+					, SNumericEntryBox<float>::FOnValueChanged::CreateSP(this, &SAvaEaseCurveTool::OnStartTangentWeightSpinBoxChanged)
+					, MinWeight, MaxWeight)
 			]
 			+ SHorizontalBox::Slot()
 			.FillWidth(1.f)
 			.HAlign(HAlign_Fill)
 			.VAlign(VAlign_Center)
+			.Padding(2.f, 0.f, 0.f, 0.f)
 			[
-				ConstructTangentNumBox(LOCTEXT("OutTangentWeightLabel", "W")
-					, LOCTEXT("OutTangentWeightToolTip", "Out Tangent Weight")
-					, TAttribute<float>::CreateSP(this, &SAvaEaseCurveTool::GetStartTangentWeight)
-					, SNumericEntryBox<float>::FOnValueChanged::CreateSP(this, &SAvaEaseCurveTool::OnStartTangentWeightSpinBoxChanged)
-					, MinWeight, MaxWeight)
+				ConstructTangentNumBox(LOCTEXT("OutTangentLabel", "T")
+					, LOCTEXT("OutTangentToolTip", "Out Tangent")
+					, TAttribute<float>::CreateSP(this, &SAvaEaseCurveTool::GetStartTangent)
+					, SNumericEntryBox<float>::FOnValueChanged::CreateSP(this, &SAvaEaseCurveTool::OnStartTangentSpinBoxChanged)
+					, MinTangent, MaxTangent)
 			]
 		]
 		+ SWrapBox::Slot()
@@ -303,22 +310,23 @@ TSharedRef<SWidget> SAvaEaseCurveTool::ConstructInputBoxes()
 			.HAlign(HAlign_Fill)
 			.VAlign(VAlign_Center)
 			[
-				ConstructTangentNumBox(LOCTEXT("InTangentLabel", "T")
-					, LOCTEXT("InTangentToolTip", "In Tangent")
-					, TAttribute<float>::CreateSP(this, &SAvaEaseCurveTool::GetEndTangent)
-					, SNumericEntryBox<float>::FOnValueChanged::CreateSP(this, &SAvaEaseCurveTool::OnEndTangentSpinBoxChanged)
-					, MinTangent, MaxTangent)
+				ConstructTangentNumBox(LOCTEXT("InTangentWeightLabel", "W")
+					, LOCTEXT("InTangentWeightToolTip", "In Tangent Weight")
+					, TAttribute<float>::CreateSP(this, &SAvaEaseCurveTool::GetEndTangentWeight)
+					, SNumericEntryBox<float>::FOnValueChanged::CreateSP(this, &SAvaEaseCurveTool::OnEndTangentWeightSpinBoxChanged)
+					, MinWeight, MaxWeight)
 			]
 			+ SHorizontalBox::Slot()
 			.FillWidth(1.f)
 			.HAlign(HAlign_Fill)
 			.VAlign(VAlign_Center)
+			.Padding(2.f, 0.f, 0.f, 0.f)
 			[
-				ConstructTangentNumBox(LOCTEXT("InTangentWeightLabel", "W")
-					, LOCTEXT("InTangentWeightToolTip", "In Tangent Weight")
-					, TAttribute<float>::CreateSP(this, &SAvaEaseCurveTool::GetEndTangentWeight)
-					, SNumericEntryBox<float>::FOnValueChanged::CreateSP(this, &SAvaEaseCurveTool::OnEndTangentWeightSpinBoxChanged)
-					, MinWeight, MaxWeight)
+				ConstructTangentNumBox(LOCTEXT("InTangentLabel", "T")
+					, LOCTEXT("InTangentToolTip", "In Tangent")
+					, TAttribute<float>::CreateSP(this, &SAvaEaseCurveTool::GetEndTangent)
+					, SNumericEntryBox<float>::FOnValueChanged::CreateSP(this, &SAvaEaseCurveTool::OnEndTangentSpinBoxChanged)
+					, MinTangent, MaxTangent)
 			]
 		];
 }
@@ -360,11 +368,11 @@ TSharedRef<SWidget> SAvaEaseCurveTool::ConstructTangentNumBox(const FText& InLab
 				.Value(InValue)
 				.OnBeginSliderMovement_Lambda([this]()
 					{
-						CurveEaseTool->BeginTransaction(LOCTEXT("SliderDragStartLabel", "Ease Curve Slider Drag"));
+						EaseCurveTool->BeginTransaction(LOCTEXT("SliderDragStartLabel", "Ease Curve Slider Drag"));
 					})
 				.OnEndSliderMovement_Lambda([this](const float InNewValue)
 					{
-						CurveEaseTool->EndTransaction();
+						EaseCurveTool->EndTransaction();
 					})
 				.OnValueChanged_Lambda([this, InOnValueChanged](const float InNewValue)
 					{
@@ -385,12 +393,12 @@ void SAvaEaseCurveTool::HandleEditorTangentsChanged(const FAvaEaseCurveTangents&
 
 void SAvaEaseCurveTool::OnEditorDragStart() const
 {
-	CurveEaseTool->BeginTransaction(LOCTEXT("EditorDragStartLabel", "Ease Curve Graph Drag"));
+	EaseCurveTool->BeginTransaction(LOCTEXT("EditorDragStartLabel", "Ease Curve Graph Drag"));
 }
 
 void SAvaEaseCurveTool::OnEditorDragEnd() const
 {
-	CurveEaseTool->EndTransaction();
+	EaseCurveTool->EndTransaction();
 }
 
 void SAvaEaseCurveTool::SetTangents(const FAvaEaseCurveTangents& InTangents, const bool bInSetEaseCurve, const bool bInBroadcastUpdate, const bool bInSetSequencerTangents) const
@@ -404,61 +412,61 @@ void SAvaEaseCurveTool::SetTangents(const FAvaEaseCurveTangents& InTangents, con
 	}
 
 	// To change the graph UI tangents, we need to change the ease curve object tangents and the graph will reflect.
-	if (bInSetEaseCurve && CurveEaseTool.IsValid())
+	if (bInSetEaseCurve && EaseCurveTool.IsValid())
 	{
-		CurveEaseTool->SetEaseCurveTangents(InTangents, bInBroadcastUpdate, bInSetSequencerTangents);
+		EaseCurveTool->SetEaseCurveTangents(InTangents, bInBroadcastUpdate, bInSetSequencerTangents);
 	}
 
-	if (CurveEaseEditorWidget.IsValid())
+	if (GetDefault<UAvaEaseCurveToolSettings>()->GetAutoZoomToFit())
 	{
-		CurveEaseEditorWidget->ZoomToFit();
+		ZoomToFit();
 	}
 }
 
 float SAvaEaseCurveTool::GetStartTangent() const
 {
-	return CurveEaseTool->GetEaseCurveTangents().Start;
+	return EaseCurveTool->GetEaseCurveTangents().Start;
 }
 
 float SAvaEaseCurveTool::GetStartTangentWeight() const
 {
-	return CurveEaseTool->GetEaseCurveTangents().StartWeight;
+	return EaseCurveTool->GetEaseCurveTangents().StartWeight;
 }
 
 float SAvaEaseCurveTool::GetEndTangent() const
 {
-	return CurveEaseTool->GetEaseCurveTangents().End;
+	return EaseCurveTool->GetEaseCurveTangents().End;
 }
 
 float SAvaEaseCurveTool::GetEndTangentWeight() const
 {
-	return CurveEaseTool->GetEaseCurveTangents().EndWeight;
+	return EaseCurveTool->GetEaseCurveTangents().EndWeight;
 }
 
 void SAvaEaseCurveTool::OnStartTangentSpinBoxChanged(const float InNewValue) const
 {
-	FAvaEaseCurveTangents NewTangents = CurveEaseTool->GetEaseCurveTangents();
+	FAvaEaseCurveTangents NewTangents = EaseCurveTool->GetEaseCurveTangents();
 	NewTangents.Start = InNewValue;
 	SetTangents(NewTangents, true, true, true);
 }
 
 void SAvaEaseCurveTool::OnStartTangentWeightSpinBoxChanged(const float InNewValue) const
 {
-	FAvaEaseCurveTangents NewTangents = CurveEaseTool->GetEaseCurveTangents();
+	FAvaEaseCurveTangents NewTangents = EaseCurveTool->GetEaseCurveTangents();
 	NewTangents.StartWeight = InNewValue;
 	SetTangents(NewTangents, true, true, true);
 }
 
 void SAvaEaseCurveTool::OnEndTangentSpinBoxChanged(const float InNewValue) const
 {
-	FAvaEaseCurveTangents NewTangents = CurveEaseTool->GetEaseCurveTangents();
+	FAvaEaseCurveTangents NewTangents = EaseCurveTool->GetEaseCurveTangents();
 	NewTangents.End = InNewValue;
 	SetTangents(NewTangents, true, true, true);
 }
 
 void SAvaEaseCurveTool::OnEndTangentWeightSpinBoxChanged(const float InNewValue) const
 {
-	FAvaEaseCurveTangents NewTangents = CurveEaseTool->GetEaseCurveTangents();
+	FAvaEaseCurveTangents NewTangents = EaseCurveTool->GetEaseCurveTangents();
 	NewTangents.EndWeight = InNewValue;
 	SetTangents(NewTangents, true, true, true);
 }
@@ -466,128 +474,149 @@ void SAvaEaseCurveTool::OnEndTangentWeightSpinBoxChanged(const float InNewValue)
 void SAvaEaseCurveTool::OnPresetChanged(const TSharedPtr<FAvaEaseCurvePreset>& InPreset) const
 {
 	SetTangents(InPreset->Tangents, true, true, true);
+
+	FSlateApplication::Get().SetAllUserFocus(CurveEaseEditorWidget);
+}
+
+void SAvaEaseCurveTool::OnQuickPresetChanged(const TSharedPtr<FAvaEaseCurvePreset>& InPreset) const
+{
+	FSlateApplication::Get().SetAllUserFocus(CurveEaseEditorWidget);
 }
 
 void SAvaEaseCurveTool::BindCommands()
 {
-	const FAvaEaseCurveToolCommands& AvaCurveEaseToolCommands = FAvaEaseCurveToolCommands::Get();
+	const FAvaEaseCurveToolCommands& EaseCurveToolCommands = FAvaEaseCurveToolCommands::Get();
 
-	FAvaEaseCurveTool* const CurveEaseToolRawPtr = CurveEaseTool.Get();
+	const TSharedRef<FAvaEaseCurveTool> EaseCurveToolRef = EaseCurveTool.ToSharedRef();
 
 	CommandList = MakeShared<FUICommandList>();
 
-	CommandList->MapAction(FGenericCommands::Get().Undo, FExecuteAction::CreateSP(CurveEaseToolRawPtr, &FAvaEaseCurveTool::UndoAction));
+	CommandList->MapAction(FGenericCommands::Get().Undo, FExecuteAction::CreateSP(EaseCurveToolRef, &FAvaEaseCurveTool::UndoAction));
 	
-	CommandList->MapAction(FGenericCommands::Get().Redo, FExecuteAction::CreateSP(CurveEaseToolRawPtr, &FAvaEaseCurveTool::RedoAction));
+	CommandList->MapAction(FGenericCommands::Get().Redo, FExecuteAction::CreateSP(EaseCurveToolRef, &FAvaEaseCurveTool::RedoAction));
 
-	CommandList->MapAction(AvaCurveEaseToolCommands.OpenToolSettings, FExecuteAction::CreateSP(CurveEaseToolRawPtr, &FAvaEaseCurveTool::OpenToolSettings));
+	CommandList->MapAction(EaseCurveToolCommands.OpenToolSettings, FExecuteAction::CreateSP(EaseCurveToolRef, &FAvaEaseCurveTool::OpenToolSettings));
 
-	CommandList->MapAction(AvaCurveEaseToolCommands.Refresh, FExecuteAction::CreateSP(CurveEaseToolRawPtr, &FAvaEaseCurveTool::UpdateEaseCurveFromSequencerKeySelections));
+	CommandList->MapAction(EaseCurveToolCommands.Refresh, FExecuteAction::CreateSP(EaseCurveToolRef, &FAvaEaseCurveTool::UpdateEaseCurveFromSequencerKeySelections));
 	
-	CommandList->MapAction(AvaCurveEaseToolCommands.Apply, FExecuteAction::CreateSP(CurveEaseToolRawPtr, &FAvaEaseCurveTool::ApplyEaseCurveToSequencerKeySelections));
+	CommandList->MapAction(EaseCurveToolCommands.Apply, FExecuteAction::CreateSP(EaseCurveToolRef, &FAvaEaseCurveTool::ApplyEaseCurveToSequencerKeySelections));
 
-	CommandList->MapAction(AvaCurveEaseToolCommands.ZoomToFit, FExecuteAction::CreateSP(this, &SAvaEaseCurveTool::ZoomToFit));
+	CommandList->MapAction(EaseCurveToolCommands.ZoomToFit, FExecuteAction::CreateSP(this, &SAvaEaseCurveTool::ZoomToFit));
 
-	CommandList->MapAction(AvaCurveEaseToolCommands.ToggleGridSnap
+	CommandList->MapAction(EaseCurveToolCommands.ToggleGridSnap
 		, FExecuteAction::CreateUObject(GetMutableDefault<UAvaEaseCurveToolSettings>(), &UAvaEaseCurveToolSettings::ToggleGridSnap)
 		, FCanExecuteAction()
 		, FIsActionChecked::CreateUObject(GetDefault<UAvaEaseCurveToolSettings>(), &UAvaEaseCurveToolSettings::GetGridSnap));
 
-	CommandList->MapAction(AvaCurveEaseToolCommands.ToggleAutoFlipTangents
+	CommandList->MapAction(EaseCurveToolCommands.ToggleAutoFlipTangents
 		, FExecuteAction::CreateUObject(GetMutableDefault<UAvaEaseCurveToolSettings>(), &UAvaEaseCurveToolSettings::ToggleAutoFlipTangents)
 		, FCanExecuteAction()
 		, FIsActionChecked::CreateUObject(GetDefault<UAvaEaseCurveToolSettings>(), &UAvaEaseCurveToolSettings::GetAutoFlipTangents));
 
-	CommandList->MapAction(AvaCurveEaseToolCommands.SelectNextChannelKey, FExecuteAction::CreateSP(CurveEaseToolRawPtr, &FAvaEaseCurveTool::SelectNextChannelKey));
+	CommandList->MapAction(EaseCurveToolCommands.ToggleAutoZoomToFit
+		, FExecuteAction::CreateUObject(GetMutableDefault<UAvaEaseCurveToolSettings>(), &UAvaEaseCurveToolSettings::ToggleAutoZoomToFit)
+		, FCanExecuteAction()
+		, FIsActionChecked::CreateUObject(GetDefault<UAvaEaseCurveToolSettings>(), &UAvaEaseCurveToolSettings::GetAutoZoomToFit));
+
+	CommandList->MapAction(EaseCurveToolCommands.SelectNextChannelKey, FExecuteAction::CreateSP(EaseCurveToolRef, &FAvaEaseCurveTool::SelectNextChannelKey));
 	
-	CommandList->MapAction(AvaCurveEaseToolCommands.SelectPreviousChannelKey, FExecuteAction::CreateSP(CurveEaseToolRawPtr, &FAvaEaseCurveTool::SelectPreviousChannelKey));
+	CommandList->MapAction(EaseCurveToolCommands.SelectPreviousChannelKey, FExecuteAction::CreateSP(EaseCurveToolRef, &FAvaEaseCurveTool::SelectPreviousChannelKey));
 
-	CommandList->MapAction(AvaCurveEaseToolCommands.SetOperationToEaseOut
-		, FExecuteAction::CreateSP(CurveEaseToolRawPtr, &FAvaEaseCurveTool::SetToolOperation, FAvaEaseCurveTool::EOperation::Out)
+	CommandList->MapAction(EaseCurveToolCommands.SetOperationToEaseOut
+		, FExecuteAction::CreateSP(EaseCurveToolRef, &FAvaEaseCurveTool::SetToolOperation, FAvaEaseCurveTool::EOperation::Out)
 		, FCanExecuteAction()
-		, FIsActionChecked::CreateSP(CurveEaseToolRawPtr, &FAvaEaseCurveTool::IsToolOperation, FAvaEaseCurveTool::EOperation::Out));
+		, FIsActionChecked::CreateSP(EaseCurveToolRef, &FAvaEaseCurveTool::IsToolOperation, FAvaEaseCurveTool::EOperation::Out));
 
-	CommandList->MapAction(AvaCurveEaseToolCommands.SetOperationToEaseInOut
-		, FExecuteAction::CreateSP(CurveEaseToolRawPtr, &FAvaEaseCurveTool::SetToolOperation, FAvaEaseCurveTool::EOperation::InOut)
+	CommandList->MapAction(EaseCurveToolCommands.SetOperationToEaseInOut
+		, FExecuteAction::CreateSP(EaseCurveToolRef, &FAvaEaseCurveTool::SetToolOperation, FAvaEaseCurveTool::EOperation::InOut)
 		, FCanExecuteAction()
-		, FIsActionChecked::CreateSP(CurveEaseToolRawPtr, &FAvaEaseCurveTool::IsToolOperation, FAvaEaseCurveTool::EOperation::InOut));
+		, FIsActionChecked::CreateSP(EaseCurveToolRef, &FAvaEaseCurveTool::IsToolOperation, FAvaEaseCurveTool::EOperation::InOut));
 
-	CommandList->MapAction(AvaCurveEaseToolCommands.SetOperationToEaseIn
-		, FExecuteAction::CreateSP(CurveEaseToolRawPtr, &FAvaEaseCurveTool::SetToolOperation, FAvaEaseCurveTool::EOperation::In)
+	CommandList->MapAction(EaseCurveToolCommands.SetOperationToEaseIn
+		, FExecuteAction::CreateSP(EaseCurveToolRef, &FAvaEaseCurveTool::SetToolOperation, FAvaEaseCurveTool::EOperation::In)
 		, FCanExecuteAction()
-		, FIsActionChecked::CreateSP(CurveEaseToolRawPtr, &FAvaEaseCurveTool::IsToolOperation, FAvaEaseCurveTool::EOperation::In));
+		, FIsActionChecked::CreateSP(EaseCurveToolRef, &FAvaEaseCurveTool::IsToolOperation, FAvaEaseCurveTool::EOperation::In));
 
-	CommandList->MapAction(AvaCurveEaseToolCommands.ResetTangents
-		, FExecuteAction::CreateSP(CurveEaseToolRawPtr, &FAvaEaseCurveTool::ResetEaseCurveTangents, true, true));
+	CommandList->MapAction(EaseCurveToolCommands.ResetTangents
+		, FExecuteAction::CreateSP(EaseCurveToolRef, &FAvaEaseCurveTool::ResetEaseCurveTangents, true, true));
 
-	CommandList->MapAction(AvaCurveEaseToolCommands.ResetStartTangent
-		, FExecuteAction::CreateSP(CurveEaseToolRawPtr, &FAvaEaseCurveTool::ResetEaseCurveTangents, true, false));
+	CommandList->MapAction(EaseCurveToolCommands.ResetStartTangent
+		, FExecuteAction::CreateSP(EaseCurveToolRef, &FAvaEaseCurveTool::ResetEaseCurveTangents, true, false));
 
-	CommandList->MapAction(AvaCurveEaseToolCommands.ResetEndTangent
-		, FExecuteAction::CreateSP(CurveEaseToolRawPtr, &FAvaEaseCurveTool::ResetEaseCurveTangents, false, true));
+	CommandList->MapAction(EaseCurveToolCommands.ResetEndTangent
+		, FExecuteAction::CreateSP(EaseCurveToolRef, &FAvaEaseCurveTool::ResetEaseCurveTangents, false, true));
 
-	CommandList->MapAction(AvaCurveEaseToolCommands.FlattenTangents
-		, FExecuteAction::CreateSP(CurveEaseToolRawPtr, &FAvaEaseCurveTool::FlattenOrStraightenTangents, true, true, true));
+	CommandList->MapAction(EaseCurveToolCommands.FlattenTangents
+		, FExecuteAction::CreateSP(EaseCurveToolRef, &FAvaEaseCurveTool::FlattenOrStraightenTangents, true, true, true));
 
-	CommandList->MapAction(AvaCurveEaseToolCommands.FlattenStartTangent
-		, FExecuteAction::CreateSP(CurveEaseToolRawPtr, &FAvaEaseCurveTool::FlattenOrStraightenTangents, true, false, true));
+	CommandList->MapAction(EaseCurveToolCommands.FlattenStartTangent
+		, FExecuteAction::CreateSP(EaseCurveToolRef, &FAvaEaseCurveTool::FlattenOrStraightenTangents, true, false, true));
 
-	CommandList->MapAction(AvaCurveEaseToolCommands.FlattenEndTangent
-		, FExecuteAction::CreateSP(CurveEaseToolRawPtr, &FAvaEaseCurveTool::FlattenOrStraightenTangents, false, true, true));
+	CommandList->MapAction(EaseCurveToolCommands.FlattenEndTangent
+		, FExecuteAction::CreateSP(EaseCurveToolRef, &FAvaEaseCurveTool::FlattenOrStraightenTangents, false, true, true));
 
-	CommandList->MapAction(AvaCurveEaseToolCommands.StraightenTangents
-		, FExecuteAction::CreateSP(CurveEaseToolRawPtr, &FAvaEaseCurveTool::FlattenOrStraightenTangents, true, true, false));
+	CommandList->MapAction(EaseCurveToolCommands.StraightenTangents
+		, FExecuteAction::CreateSP(EaseCurveToolRef, &FAvaEaseCurveTool::FlattenOrStraightenTangents, true, true, false));
 
-	CommandList->MapAction(AvaCurveEaseToolCommands.StraightenStartTangent
-		, FExecuteAction::CreateSP(CurveEaseToolRawPtr, &FAvaEaseCurveTool::FlattenOrStraightenTangents, true, false, false));
+	CommandList->MapAction(EaseCurveToolCommands.StraightenStartTangent
+		, FExecuteAction::CreateSP(EaseCurveToolRef, &FAvaEaseCurveTool::FlattenOrStraightenTangents, true, false, false));
 
-	CommandList->MapAction(AvaCurveEaseToolCommands.StraightenEndTangent
-		, FExecuteAction::CreateSP(CurveEaseToolRawPtr, &FAvaEaseCurveTool::FlattenOrStraightenTangents, false, true, false));
+	CommandList->MapAction(EaseCurveToolCommands.StraightenEndTangent
+		, FExecuteAction::CreateSP(EaseCurveToolRef, &FAvaEaseCurveTool::FlattenOrStraightenTangents, false, true, false));
 
-	CommandList->MapAction(AvaCurveEaseToolCommands.CopyTangents
-		, FExecuteAction::CreateSP(CurveEaseToolRawPtr, &FAvaEaseCurveTool::CopyTangentsToClipboard)
-		, FCanExecuteAction::CreateSP(CurveEaseToolRawPtr, &FAvaEaseCurveTool::CanCopyTangentsToClipboard));
+	CommandList->MapAction(EaseCurveToolCommands.CopyTangents
+		, FExecuteAction::CreateSP(EaseCurveToolRef, &FAvaEaseCurveTool::CopyTangentsToClipboard)
+		, FCanExecuteAction::CreateSP(EaseCurveToolRef, &FAvaEaseCurveTool::CanCopyTangentsToClipboard));
 
-	CommandList->MapAction(AvaCurveEaseToolCommands.PasteTangents
-		, FExecuteAction::CreateSP(CurveEaseToolRawPtr, &FAvaEaseCurveTool::PasteTangentsFromClipboard)
-		, FCanExecuteAction::CreateSP(CurveEaseToolRawPtr, &FAvaEaseCurveTool::CanPasteTangentsFromClipboard));
+	CommandList->MapAction(EaseCurveToolCommands.PasteTangents
+		, FExecuteAction::CreateSP(EaseCurveToolRef, &FAvaEaseCurveTool::PasteTangentsFromClipboard)
+		, FCanExecuteAction::CreateSP(EaseCurveToolRef, &FAvaEaseCurveTool::CanPasteTangentsFromClipboard));
 
-	CommandList->MapAction(AvaCurveEaseToolCommands.CreateExternalCurveAsset
+	CommandList->MapAction(EaseCurveToolCommands.CreateExternalCurveAsset
 		, FExecuteAction::CreateLambda([this]()
 			{
-				CurveEaseTool->CreateCurveAsset();
+				EaseCurveTool->CreateCurveAsset();
 			})
 		, FCanExecuteAction());
 
-	CommandList->MapAction(AvaCurveEaseToolCommands.SetKeyInterpConstant,
-		FExecuteAction::CreateSP(CurveEaseToolRawPtr, &FAvaEaseCurveTool::SetKeyInterpMode, RCIM_Constant, RCTM_Auto),
+	CommandList->MapAction(EaseCurveToolCommands.SetKeyInterpConstant,
+		FExecuteAction::CreateSP(EaseCurveToolRef, &FAvaEaseCurveTool::SetKeyInterpMode, RCIM_Constant, RCTM_Auto),
 		FCanExecuteAction(),
-		FIsActionChecked::CreateSP(CurveEaseToolRawPtr, &FAvaEaseCurveTool::IsKeyInterpMode, RCIM_Constant, RCTM_Auto));
+		FIsActionChecked::CreateSP(EaseCurveToolRef, &FAvaEaseCurveTool::IsKeyInterpMode, RCIM_Constant, RCTM_Auto));
 
-	CommandList->MapAction(AvaCurveEaseToolCommands.SetKeyInterpLinear,
-		FExecuteAction::CreateSP(CurveEaseToolRawPtr, &FAvaEaseCurveTool::SetKeyInterpMode, RCIM_Linear, RCTM_Auto),
+	CommandList->MapAction(EaseCurveToolCommands.SetKeyInterpLinear,
+		FExecuteAction::CreateSP(EaseCurveToolRef, &FAvaEaseCurveTool::SetKeyInterpMode, RCIM_Linear, RCTM_Auto),
 		FCanExecuteAction(),
-		FIsActionChecked::CreateSP(CurveEaseToolRawPtr, &FAvaEaseCurveTool::IsKeyInterpMode, RCIM_Linear, RCTM_Auto));
+		FIsActionChecked::CreateSP(EaseCurveToolRef, &FAvaEaseCurveTool::IsKeyInterpMode, RCIM_Linear, RCTM_Auto));
 
-	CommandList->MapAction(AvaCurveEaseToolCommands.SetKeyInterpCubicAuto,
-		FExecuteAction::CreateSP(CurveEaseToolRawPtr, &FAvaEaseCurveTool::SetKeyInterpMode, RCIM_Cubic, RCTM_Auto),
+	CommandList->MapAction(EaseCurveToolCommands.SetKeyInterpCubicAuto,
+		FExecuteAction::CreateSP(EaseCurveToolRef, &FAvaEaseCurveTool::SetKeyInterpMode, RCIM_Cubic, RCTM_Auto),
 		FCanExecuteAction(),
-		FIsActionChecked::CreateSP(CurveEaseToolRawPtr, &FAvaEaseCurveTool::IsKeyInterpMode, RCIM_Cubic, RCTM_Auto));
+		FIsActionChecked::CreateSP(EaseCurveToolRef, &FAvaEaseCurveTool::IsKeyInterpMode, RCIM_Cubic, RCTM_Auto));
 
-	CommandList->MapAction(AvaCurveEaseToolCommands.SetKeyInterpCubicSmartAuto,
-		FExecuteAction::CreateSP(CurveEaseToolRawPtr, &FAvaEaseCurveTool::SetKeyInterpMode, RCIM_Cubic, RCTM_SmartAuto),
+	CommandList->MapAction(EaseCurveToolCommands.SetKeyInterpCubicSmartAuto,
+		FExecuteAction::CreateSP(EaseCurveToolRef, &FAvaEaseCurveTool::SetKeyInterpMode, RCIM_Cubic, RCTM_SmartAuto),
 		FCanExecuteAction(),
-		FIsActionChecked::CreateSP(CurveEaseToolRawPtr, &FAvaEaseCurveTool::IsKeyInterpMode, RCIM_Cubic, RCTM_SmartAuto));
+		FIsActionChecked::CreateSP(EaseCurveToolRef, &FAvaEaseCurveTool::IsKeyInterpMode, RCIM_Cubic, RCTM_SmartAuto));
 
-	CommandList->MapAction(AvaCurveEaseToolCommands.SetKeyInterpCubicUser,
-		FExecuteAction::CreateSP(CurveEaseToolRawPtr, &FAvaEaseCurveTool::SetKeyInterpMode, RCIM_Cubic, RCTM_User),
+	CommandList->MapAction(EaseCurveToolCommands.SetKeyInterpCubicUser,
+		FExecuteAction::CreateSP(EaseCurveToolRef, &FAvaEaseCurveTool::SetKeyInterpMode, RCIM_Cubic, RCTM_User),
 		FCanExecuteAction(),
-		FIsActionChecked::CreateSP(CurveEaseToolRawPtr, &FAvaEaseCurveTool::IsKeyInterpMode, RCIM_Cubic, RCTM_User));
+		FIsActionChecked::CreateSP(EaseCurveToolRef, &FAvaEaseCurveTool::IsKeyInterpMode, RCIM_Cubic, RCTM_User));
 
-	CommandList->MapAction(AvaCurveEaseToolCommands.SetKeyInterpCubicBreak,
-		FExecuteAction::CreateSP(CurveEaseToolRawPtr, &FAvaEaseCurveTool::SetKeyInterpMode, RCIM_Cubic, RCTM_Break),
+	CommandList->MapAction(EaseCurveToolCommands.SetKeyInterpCubicBreak,
+		FExecuteAction::CreateSP(EaseCurveToolRef, &FAvaEaseCurveTool::SetKeyInterpMode, RCIM_Cubic, RCTM_Break),
 		FCanExecuteAction(),
-		FIsActionChecked::CreateSP(CurveEaseToolRawPtr, &FAvaEaseCurveTool::IsKeyInterpMode, RCIM_Cubic, RCTM_Break));
+		FIsActionChecked::CreateSP(EaseCurveToolRef, &FAvaEaseCurveTool::IsKeyInterpMode, RCIM_Cubic, RCTM_Break));
+
+	CommandList->MapAction(EaseCurveToolCommands.QuickEase
+		, FExecuteAction::CreateSP(EaseCurveToolRef, &FAvaEaseCurveTool::ApplyQuickEaseToSequencerKeySelections, FAvaEaseCurveTool::EOperation::InOut));
+
+	CommandList->MapAction(EaseCurveToolCommands.QuickEaseIn
+		, FExecuteAction::CreateSP(EaseCurveToolRef, &FAvaEaseCurveTool::ApplyQuickEaseToSequencerKeySelections, FAvaEaseCurveTool::EOperation::In));
+
+	CommandList->MapAction(EaseCurveToolCommands.QuickEaseOut
+		, FExecuteAction::CreateSP(EaseCurveToolRef, &FAvaEaseCurveTool::ApplyQuickEaseToSequencerKeySelections, FAvaEaseCurveTool::EOperation::Out));
 }
 
 TSharedRef<SWidget> SAvaEaseCurveTool::CreateContextMenuContent()
@@ -606,7 +635,7 @@ TSharedRef<SWidget> SAvaEaseCurveTool::CreateContextMenuContent()
 
 		FToolMenuSection& Section = ToolMenu->FindOrAddSection(TEXT("EaseCurveTool"), LOCTEXT("EaseCurveToolActions", "Curve Ease Tool Actions"));
 
-		const FAvaEaseCurveToolCommands& AvaCurveEaseToolCommands = FAvaEaseCurveToolCommands::Get();
+		const FAvaEaseCurveToolCommands& EaseCurveToolCommands = FAvaEaseCurveToolCommands::Get();
 
 		Section.AddSubMenu(TEXT("Settings"),
 			LOCTEXT("SettingsSubMenuLabel", "Settings"),
@@ -617,89 +646,95 @@ TSharedRef<SWidget> SAvaEaseCurveTool::CreateContextMenuContent()
 
 		Section.AddSeparator(NAME_None);
 
-		Section.AddMenuEntry(AvaCurveEaseToolCommands.CreateExternalCurveAsset);
+		Section.AddMenuEntry(EaseCurveToolCommands.CreateExternalCurveAsset);
 
 		Section.AddSeparator(NAME_None);
 
-		Section.AddMenuEntry(AvaCurveEaseToolCommands.CopyTangents);
+		Section.AddMenuEntry(EaseCurveToolCommands.CopyTangents);
 
-		Section.AddMenuEntry(AvaCurveEaseToolCommands.PasteTangents);
+		Section.AddMenuEntry(EaseCurveToolCommands.PasteTangents);
 
 		Section.AddSeparator(NAME_None);
 
 		Section.AddSubMenu(TEXT("StraightenTangents"),
 			LOCTEXT("StraightenTangentsSubMenuLabel", "Straighten Tangents"),
 			LOCTEXT("StraightenTangentsSubMenuToolTip", ""),
-			FNewMenuDelegate::CreateLambda([&AvaCurveEaseToolCommands](FMenuBuilder& InMenuBuilder)
+			FNewToolMenuDelegate::CreateLambda([&EaseCurveToolCommands](UToolMenu* InToolMenu)
 				{
-					InMenuBuilder.AddMenuEntry(AvaCurveEaseToolCommands.StraightenTangents);
-					InMenuBuilder.AddMenuEntry(AvaCurveEaseToolCommands.StraightenStartTangent);
-					InMenuBuilder.AddMenuEntry(AvaCurveEaseToolCommands.StraightenEndTangent);
+					FToolMenuSection& NewSection = InToolMenu->FindOrAddSection(TEXT("StraightenTangents"));
+					NewSection.AddMenuEntry(EaseCurveToolCommands.StraightenTangents);
+					NewSection.AddSeparator(NAME_None);
+					NewSection.AddMenuEntry(EaseCurveToolCommands.StraightenStartTangent);
+					NewSection.AddMenuEntry(EaseCurveToolCommands.StraightenEndTangent);
 				}),
 			false, FSlateIcon(FAppStyle::Get().GetStyleSetName(), TEXT("GenericCurveEditor.StraightenTangents")));
 
 		Section.AddSubMenu(TEXT("FlattenTangents"),
 			LOCTEXT("FlattenTangentsSubMenuLabel", "Flatten Tangents"),
 			LOCTEXT("FlattenTangentsSubMenuToolTip", ""),
-			FNewMenuDelegate::CreateLambda([&AvaCurveEaseToolCommands](FMenuBuilder& InMenuBuilder)
+			FNewToolMenuDelegate::CreateLambda([&EaseCurveToolCommands](UToolMenu* InToolMenu)
 				{
-					InMenuBuilder.AddMenuEntry(AvaCurveEaseToolCommands.FlattenTangents);
-					InMenuBuilder.AddMenuEntry(AvaCurveEaseToolCommands.FlattenStartTangent);
-					InMenuBuilder.AddMenuEntry(AvaCurveEaseToolCommands.FlattenEndTangent);
+					FToolMenuSection& NewSection = InToolMenu->FindOrAddSection(TEXT("FlattenTangents"));
+					NewSection.AddMenuEntry(EaseCurveToolCommands.FlattenTangents);
+					NewSection.AddSeparator(NAME_None);
+					NewSection.AddMenuEntry(EaseCurveToolCommands.FlattenStartTangent);
+					NewSection.AddMenuEntry(EaseCurveToolCommands.FlattenEndTangent);
 				}), 
 			false, FSlateIcon(FAppStyle::Get().GetStyleSetName(), TEXT("GenericCurveEditor.FlattenTangents")));
 
 		Section.AddSubMenu(TEXT("ResetTangents"),
 			LOCTEXT("ResetTangentsSubMenuLabel", "Reset Tangents"),
 			LOCTEXT("ResetTangentsSubMenuToolTip", ""),
-			FNewMenuDelegate::CreateLambda([&AvaCurveEaseToolCommands](FMenuBuilder& InMenuBuilder)
+			FNewToolMenuDelegate::CreateLambda([&EaseCurveToolCommands](UToolMenu* InToolMenu)
 				{
-					InMenuBuilder.AddMenuEntry(AvaCurveEaseToolCommands.ResetTangents);
-					InMenuBuilder.AddMenuEntry(AvaCurveEaseToolCommands.ResetStartTangent);
-					InMenuBuilder.AddMenuEntry(AvaCurveEaseToolCommands.ResetEndTangent);
+					FToolMenuSection& NewSection = InToolMenu->FindOrAddSection(TEXT("ResetTangents"));
+					NewSection.AddMenuEntry(EaseCurveToolCommands.ResetTangents);
+					NewSection.AddSeparator(NAME_None);
+					NewSection.AddMenuEntry(EaseCurveToolCommands.ResetStartTangent);
+					NewSection.AddMenuEntry(EaseCurveToolCommands.ResetEndTangent);
 				}),
 			false, FSlateIcon(FAppStyle::Get().GetStyleSetName(), TEXT("PropertyWindow.DiffersFromDefault")));
 
 		Section.AddSeparator(NAME_None);
 
 		/** @TODO: Only show these in single edit mode(?)
-		Section.AddMenuEntry(AvaCurveEaseToolCommands.SetKeyInterpConstant);
+		Section.AddMenuEntry(EaseCurveToolCommands.SetKeyInterpConstant);
 
-		Section.AddMenuEntry(AvaCurveEaseToolCommands.SetKeyInterpLinear);
+		Section.AddMenuEntry(EaseCurveToolCommands.SetKeyInterpLinear);
 
-		Section.AddMenuEntry(AvaCurveEaseToolCommands.SetKeyInterpCubicAuto);
+		Section.AddMenuEntry(EaseCurveToolCommands.SetKeyInterpCubicAuto);
 
-		Section.AddMenuEntry(AvaCurveEaseToolCommands.SetKeyInterpCubicSmartAuto);
+		Section.AddMenuEntry(EaseCurveToolCommands.SetKeyInterpCubicSmartAuto);
 
-		Section.AddMenuEntry(AvaCurveEaseToolCommands.SetKeyInterpCubicUser);
+		Section.AddMenuEntry(EaseCurveToolCommands.SetKeyInterpCubicUser);
 
-		Section.AddMenuEntry(AvaCurveEaseToolCommands.SetKeyInterpCubicBreak);
+		Section.AddMenuEntry(EaseCurveToolCommands.SetKeyInterpCubicBreak);
 
 		Section.AddSeparator(NAME_None);*/
 
-		Section.AddMenuEntry(AvaCurveEaseToolCommands.SetOperationToEaseOut);
+		Section.AddMenuEntry(EaseCurveToolCommands.SetOperationToEaseOut);
 
-		Section.AddMenuEntry(AvaCurveEaseToolCommands.SetOperationToEaseInOut);
+		Section.AddMenuEntry(EaseCurveToolCommands.SetOperationToEaseInOut);
 
-		Section.AddMenuEntry(AvaCurveEaseToolCommands.SetOperationToEaseIn);
+		Section.AddMenuEntry(EaseCurveToolCommands.SetOperationToEaseIn);
 
 		Section.AddSeparator(NAME_None);
 
-		Section.AddMenuEntry(AvaCurveEaseToolCommands.SelectPreviousChannelKey);
+		Section.AddMenuEntry(EaseCurveToolCommands.SelectPreviousChannelKey);
 
-		Section.AddMenuEntry(AvaCurveEaseToolCommands.SelectNextChannelKey);
+		Section.AddMenuEntry(EaseCurveToolCommands.SelectNextChannelKey);
 
 		Section.AddSeparator(NAME_None);
 
-		Section.AddMenuEntry(AvaCurveEaseToolCommands.ToggleGridSnap);
+		Section.AddMenuEntry(EaseCurveToolCommands.ToggleGridSnap);
 
-		Section.AddMenuEntry(AvaCurveEaseToolCommands.ZoomToFit);
+		Section.AddMenuEntry(EaseCurveToolCommands.ZoomToFit);
 
 		Section.AddSeparator(NAME_None);
 
-		Section.AddMenuEntry(AvaCurveEaseToolCommands.Refresh);
+		Section.AddMenuEntry(EaseCurveToolCommands.Refresh);
 
-		Section.AddMenuEntry(AvaCurveEaseToolCommands.Apply);
+		Section.AddMenuEntry(EaseCurveToolCommands.Apply);
 	}
 
 	return ToolMenus->GenerateWidget(MenuName, FToolMenuContext(CommandList));
@@ -707,12 +742,17 @@ TSharedRef<SWidget> SAvaEaseCurveTool::CreateContextMenuContent()
 
 void SAvaEaseCurveTool::MakeContextMenuSettings(FMenuBuilder& InMenuBuilder)
 {
-	const FAvaEaseCurveToolCommands& AvaCurveEaseToolCommands = FAvaEaseCurveToolCommands::Get();
+	const FAvaEaseCurveToolCommands& EaseCurveToolCommands = FAvaEaseCurveToolCommands::Get();
+
+	InMenuBuilder.AddMenuEntry(EaseCurveToolCommands.OpenToolSettings);
 
-	InMenuBuilder.AddMenuEntry(AvaCurveEaseToolCommands.OpenToolSettings);
 	InMenuBuilder.AddSeparator();
-	InMenuBuilder.AddMenuEntry(AvaCurveEaseToolCommands.ToggleAutoFlipTangents);
+
+	InMenuBuilder.AddMenuEntry(EaseCurveToolCommands.ToggleAutoFlipTangents);
+
 	InMenuBuilder.AddSeparator();
+
+	// Graph Size
 	{
 		FProperty* const GraphSizeProperty = UAvaEaseCurveToolSettings::StaticClass()->FindPropertyByName(TEXT("GraphSize"));
 		check(GraphSizeProperty);
@@ -763,6 +803,8 @@ void SAvaEaseCurveTool::MakeContextMenuSettings(FMenuBuilder& InMenuBuilder)
 				});
 		InMenuBuilder.AddWidget(GraphSizeWidget, LOCTEXT("ToolSizeLabel", "Tool Size"));
 	}
+
+	// Grid Size
 	{
 		FProperty* const GridSizeProperty = UAvaEaseCurveToolSettings::StaticClass()->FindPropertyByName(TEXT("GridSize"));
 		check(GridSizeProperty);
@@ -813,6 +855,10 @@ void SAvaEaseCurveTool::MakeContextMenuSettings(FMenuBuilder& InMenuBuilder)
 				});
 		InMenuBuilder.AddWidget(GridSizeWidget, LOCTEXT("GridSizeLabel", "Grid Size"));
 	}
+
+	InMenuBuilder.AddSeparator();
+
+	InMenuBuilder.AddMenuEntry(EaseCurveToolCommands.ToggleAutoZoomToFit);
 }
 
 void SAvaEaseCurveTool::UndoAction()
@@ -857,4 +903,32 @@ FReply SAvaEaseCurveTool::OnKeyDown(const FGeometry& MyGeometry, const FKeyEvent
 	return FReply::Unhandled();
 }
 
+FText SAvaEaseCurveTool::GetStartText() const
+{
+	return (ToolMode.Get(FAvaEaseCurveTool::EMode::DualKeyEdit) == FAvaEaseCurveTool::EMode::DualKeyEdit)
+		? LOCTEXT("StartText", "Leave")
+		: LOCTEXT("ArriveText", "Arrive");
+}
+
+FText SAvaEaseCurveTool::GetStartTooltipText() const
+{
+	return (ToolMode.Get(FAvaEaseCurveTool::EMode::DualKeyEdit) == FAvaEaseCurveTool::EMode::DualKeyEdit)
+		? LOCTEXT("StartTooltipText", "Start: The selected key's leave tangent")
+		: LOCTEXT("ArriveTooltipText", "Arrive");
+}
+
+FText SAvaEaseCurveTool::GetEndText() const
+{
+	return (ToolMode.Get(FAvaEaseCurveTool::EMode::DualKeyEdit) == FAvaEaseCurveTool::EMode::DualKeyEdit)
+		? LOCTEXT("EndText", "Arrive")
+		: LOCTEXT("LeaveText", "Leave");
+}
+
+FText SAvaEaseCurveTool::GetEndTooltipText() const
+{
+	return (ToolMode.Get(FAvaEaseCurveTool::EMode::DualKeyEdit) == FAvaEaseCurveTool::EMode::DualKeyEdit)
+		? LOCTEXT("EndTooltipText", "End: The next key's arrive tangent")
+		: LOCTEXT("LeaveTooltipText", "Leave");
+}
+
 #undef LOCTEXT_NAMESPACE
diff --git a/Engine/Plugins/Experimental/Avalanche/Source/AvalancheSequencer/Private/EaseCurveTool/Widgets/SAvaEaseCurveTool.h b/Engine/Plugins/Experimental/Avalanche/Source/AvalancheSequencer/Private/EaseCurveTool/Widgets/SAvaEaseCurveTool.h
--- a/Engine/Plugins/Experimental/Avalanche/Source/AvalancheSequencer/Private/EaseCurveTool/Widgets/SAvaEaseCurveTool.h
+++ b/Engine/Plugins/Experimental/Avalanche/Source/AvalancheSequencer/Private/EaseCurveTool/Widgets/SAvaEaseCurveTool.h
@@ -23,7 +23,7 @@ class SAvaEaseCurveTool
 {
 public:
 	SLATE_BEGIN_ARGS(SAvaEaseCurveTool)
-		: _ToolMode(FAvaEaseCurveTool::EMode::DoubleKeyEdit)
+		: _ToolMode(FAvaEaseCurveTool::EMode::DualKeyEdit)
 		, _ToolOperation(FAvaEaseCurveTool::EOperation::InOut)
 	{}
 		SLATE_ATTRIBUTE(FAvaEaseCurveTool::EMode, ToolMode)
@@ -31,7 +31,7 @@ public:
 		SLATE_ARGUMENT(FAvaEaseCurveTangents, InitialTangents)
 	SLATE_END_ARGS()
 
-	void Construct(const FArguments& InArgs, const TSharedRef<FAvaEaseCurveTool>& InCurveEaseTool);
+	void Construct(const FArguments& InArgs, const TSharedRef<FAvaEaseCurveTool>& InEaseCurveTool);
 	
 	void SetTangents(const FAvaEaseCurveTangents& InTangents, const bool bInSetEaseCurve, const bool bInBroadcastUpdate, const bool bInSetSequencerTangents) const;
 
@@ -42,6 +42,8 @@ public:
 
 	FKeyHandle GetSelectedKeyHandle() const;
 
+	void ZoomToFit() const;
+
 protected:
 	TSharedRef<SWidget> ConstructCurveEditorPanel();
 	TSharedRef<SWidget> ConstructInputBoxes();
@@ -60,6 +62,7 @@ protected:
 	void OnEndTangentWeightSpinBoxChanged(const float InNewValue) const;
 
 	void OnPresetChanged(const TSharedPtr<FAvaEaseCurvePreset>& InPreset) const;
+	void OnQuickPresetChanged(const TSharedPtr<FAvaEaseCurvePreset>& InPreset) const;
 
 	void BindCommands();
 
@@ -69,18 +72,21 @@ protected:
 	void UndoAction();
 	void RedoAction();
 
-	void ZoomToFit() const;
-
 	void OnEditorDragStart() const;
 	void OnEditorDragEnd() const;
 
+	FText GetStartText() const;
+	FText GetStartTooltipText() const;
+	FText GetEndText() const;
+	FText GetEndTooltipText() const;
+
 	//~ Begin SWidget
 	virtual FReply OnKeyDown(const FGeometry& MyGeometry, const FKeyEvent& InKeyEvent) override;
 	//~ End SWidget
 
 	TSharedPtr<FUICommandList> CommandList;
 
-	TSharedPtr<FAvaEaseCurveTool> CurveEaseTool;
+	TSharedPtr<FAvaEaseCurveTool> EaseCurveTool;
 
 	TAttribute<FAvaEaseCurveTool::EMode> ToolMode;
 	TAttribute<FAvaEaseCurveTool::EOperation> ToolOperation;
