Avalanche Outliner:
- First pass unexporting unused functions, and moved public classes not used outside the module to private

diff --git a/Engine/Plugins/Experimental/Avalanche/Source/Avalanche/Private/AvaActorUtils.cpp b/Engine/Plugins/Experimental/Avalanche/Source/Avalanche/Private/AvaActorUtils.cpp
--- a/Engine/Plugins/Experimental/Avalanche/Source/Avalanche/Private/AvaActorUtils.cpp
+++ b/Engine/Plugins/Experimental/Avalanche/Source/Avalanche/Private/AvaActorUtils.cpp
@@ -14,10 +14,10 @@
 #include "Math/OrientedBox.h"
 
 #if WITH_EDITOR
-#include "AvaOutliner.h"
 #include "AvaOutlinerDefines.h"
 #include "AvaOutlinerSubsystem.h"
 #include "AvaOutlinerUtils.h"
+#include "IAvaOutliner.h"
 #include "Item/AvaOutlinerTreeRoot.h"
 #endif
 
@@ -160,7 +160,7 @@ AActor* FAvaActorUtils::ActorFromReferenceContainer(AActor* const InActor, const
 	UAvaOutlinerSubsystem* const OutlinerSubsystem = World->GetSubsystem<UAvaOutlinerSubsystem>();
 	if (IsValid(OutlinerSubsystem))
 	{
-		TSharedPtr<FAvaOutliner> AvaOutliner = OutlinerSubsystem->GetOutliner();
+		TSharedPtr<IAvaOutliner> AvaOutliner = OutlinerSubsystem->GetOutliner();
 		if (AvaOutliner.IsValid())
 		{
 			bIsOutlinerAttachedActors = true;
diff --git a/Engine/Plugins/Experimental/Avalanche/Source/Avalanche/Private/AvaOutlinerUtils.cpp b/Engine/Plugins/Experimental/Avalanche/Source/Avalanche/Private/AvaOutlinerUtils.cpp
--- a/Engine/Plugins/Experimental/Avalanche/Source/Avalanche/Private/AvaOutlinerUtils.cpp
+++ b/Engine/Plugins/Experimental/Avalanche/Source/Avalanche/Private/AvaOutlinerUtils.cpp
@@ -4,14 +4,14 @@
 
 #if WITH_EDITOR
 
-#include "AvaOutliner.h"
 #include "AvaOutlinerSubsystem.h"
 #include "Engine/World.h"
+#include "GameFramework/Actor.h"
+#include "IAvaOutliner.h"
 #include "Item/AvaOutlinerActor.h"
 #include "Item/AvaOutlinerTreeRoot.h"
-#include "GameFramework/Actor.h"
 
-TSharedPtr<FAvaOutliner> FAvaOutlinerUtils::EditorGetOutliner(const UWorld* const InWorld)
+TSharedPtr<IAvaOutliner> FAvaOutlinerUtils::EditorGetOutliner(const UWorld* const InWorld)
 {
 	check(IsValid(InWorld));
 
@@ -56,7 +56,7 @@ bool FAvaOutlinerUtils::EditorOutlinerItemsToActors(const TArray<FAvaOutlinerIte
 	return true;
 }
 
-TArray<AActor*> FAvaOutlinerUtils::EditorOutlinerChildActors(TSharedPtr<FAvaOutliner> InOutliner, AActor* const InParentActor)
+TArray<AActor*> FAvaOutlinerUtils::EditorOutlinerChildActors(TSharedPtr<IAvaOutliner> InOutliner, AActor* const InParentActor)
 {
 	if (!InOutliner.IsValid())
 	{
@@ -87,7 +87,7 @@ TArray<AActor*> FAvaOutlinerUtils::EditorOutlinerChildActors(TSharedPtr<FAvaOutl
 	return {};
 }
 
-bool FAvaOutlinerUtils::EditorActorIsolationInfo(TSharedPtr<FAvaOutliner> InOutliner, TArray<TWeakObjectPtr<const AActor>>& OutIsolatedActors)
+bool FAvaOutlinerUtils::EditorActorIsolationInfo(TSharedPtr<IAvaOutliner> InOutliner, TArray<TWeakObjectPtr<const AActor>>& OutIsolatedActors)
 {
 	// todo
 	OutIsolatedActors.Empty();
diff --git a/Engine/Plugins/Experimental/Avalanche/Source/Avalanche/Public/AvaOutlinerUtils.h b/Engine/Plugins/Experimental/Avalanche/Source/Avalanche/Public/AvaOutlinerUtils.h
--- a/Engine/Plugins/Experimental/Avalanche/Source/Avalanche/Public/AvaOutlinerUtils.h
+++ b/Engine/Plugins/Experimental/Avalanche/Source/Avalanche/Public/AvaOutlinerUtils.h
@@ -10,19 +10,19 @@
 #include "UObject/WeakObjectPtr.h"
 
 class AActor;
-class FAvaOutliner;
+class IAvaOutliner;
 
 /** Editor Only Outliner Access Functions */
 struct AVALANCHE_API FAvaOutlinerUtils
 {
 	/** Returns the Avalanche outliner from the outliner module at runtime in the editor. */
-	static TSharedPtr<FAvaOutliner> EditorGetOutliner(const UWorld* const InWorld);
+	static TSharedPtr<IAvaOutliner> EditorGetOutliner(const UWorld* const InWorld);
 
 	/** Returns an array of direct child actors attached to a parent actor or all root level actors if the parent actor is null. */
-	static TArray<AActor*> EditorOutlinerChildActors(TSharedPtr<FAvaOutliner> InOutliner, AActor* const InParentActor = nullptr);
+	static TArray<AActor*> EditorOutlinerChildActors(TSharedPtr<IAvaOutliner> InOutliner, AActor* const InParentActor = nullptr);
 
 	/** Returns true if the Avalanche BP editor is isolating actors and the array of actors currently being isolated. */
-	static bool EditorActorIsolationInfo(TSharedPtr<FAvaOutliner> InOutliner, TArray<TWeakObjectPtr<const AActor>>& OutIsolatedActors);
+	static bool EditorActorIsolationInfo(TSharedPtr<IAvaOutliner> InOutliner, TArray<TWeakObjectPtr<const AActor>>& OutIsolatedActors);
 
 	/**
 	 * Gets an array of actors from an array of outliner items while maintaining index ordering (for actor outliner items only).
diff --git a/Engine/Plugins/Experimental/Avalanche/Source/AvalancheEditor/Private/Outliner/AvaOutlinerExtension.cpp b/Engine/Plugins/Experimental/Avalanche/Source/AvalancheEditor/Private/Outliner/AvaOutlinerExtension.cpp
--- a/Engine/Plugins/Experimental/Avalanche/Source/AvalancheEditor/Private/Outliner/AvaOutlinerExtension.cpp
+++ b/Engine/Plugins/Experimental/Avalanche/Source/AvalancheEditor/Private/Outliner/AvaOutlinerExtension.cpp
@@ -5,7 +5,6 @@
 #include "AvaEditorCommands.h"
 #include "AvaEditorModule.h"
 #include "AvaEditorSettings.h"
-#include "AvaOutliner.h"
 #include "AvaOutlinerSubsystem.h"
 #include "AvaOutlinerTabSpawner.h"
 #include "AvaScene.h"
@@ -21,8 +20,9 @@
 #include "Filters/AvaOutlinerItemTypeFilter.h"
 #include "Framework/AvaNullActor.h"
 #include "Framework/Docking/LayoutExtender.h"
+#include "IAvaOutliner.h"
 #include "IAvaOutlinerModule.h"
-#include "ItemActions/AvaOutlinerAddItem.h"
+#include "Item/AvaOutlinerComponent.h"
 #include "ItemProxies/AvaOutlinerItemProxyRegistry.h"
 #include "LevelEditor.h"
 #include "Materials/Material.h"
@@ -33,7 +33,6 @@
 #include "Sequencer/AvaSequencerExtension.h"
 #include "Styling/SlateIconFinder.h"
 #include "Text3DActor.h"
-#include "Item/AvaOutlinerComponent.h"
 #include "ToolMenuContext/AvaOutlinerItemsContext.h"
 #include "Viewport/AvaViewportExtension.h"
 
@@ -52,7 +51,7 @@ void FAvaOutlinerExtension::StaticStartup()
 	OutlinerModule.GetItemProxyRegistry().RegisterItemProxyWithDefaultFactory<FAvaOutlinerMaterialDesignerProxy, 5>();
 
 	OutlinerModule.GetOnExtendItemProxiesForItem().AddLambda(
-		[](FAvaOutliner& InOutliner, const FAvaOutlinerItemPtr& InItemPtr, TArray<TSharedPtr<FAvaOutlinerItemProxy>>& OutItemProxies)
+		[](IAvaOutliner& InOutliner, const FAvaOutlinerItemPtr& InItemPtr, TArray<TSharedPtr<FAvaOutlinerItemProxy>>& OutItemProxies)
 		{
 			if (FAvaOutlinerComponent* ComponentItem = InItemPtr->CastTo<FAvaOutlinerComponent>())
 			{
@@ -301,7 +300,7 @@ void FAvaOutlinerExtension::ExtendOutlinerItemContextMenu(UToolMenu* InItemConte
 			, FToolMenuInsert(NAME_None, EToolMenuInsertType::First));
 	}
 
-	// Note: Since the Outliner Command List is linked to the Ava Command List (see FAvaOutliner::SetBaseCommandList),
+	// Note: Since the Outliner Command List is linked to the Ava Command List (see IAvaOutliner::SetBaseCommandList),
 	// we do NOT need to add the entry with a different Command List
 	const FAvaEditorCommands AvaEditorCommands = FAvaEditorCommands::Get();
 
diff --git a/Engine/Plugins/Experimental/Avalanche/Source/AvalancheEditor/Private/Outliner/AvaOutlinerExtension.h b/Engine/Plugins/Experimental/Avalanche/Source/AvalancheEditor/Private/Outliner/AvaOutlinerExtension.h
--- a/Engine/Plugins/Experimental/Avalanche/Source/AvalancheEditor/Private/Outliner/AvaOutlinerExtension.h
+++ b/Engine/Plugins/Experimental/Avalanche/Source/AvalancheEditor/Private/Outliner/AvaOutlinerExtension.h
@@ -7,7 +7,7 @@
 #include "IAvaOutlinerProvider.h"
 #include "Templates/SharedPointer.h"
 
-class FAvaOutliner;
+class IAvaOutliner;
 
 UE_AVA_TYPE_EXTERNAL(IAvaOutlinerProvider);
 
@@ -23,7 +23,7 @@ public:
 
 	FAvaOutlinerExtension();
 
-	TSharedPtr<FAvaOutliner> GetAvaOutliner() const { return AvaOutliner; }
+	TSharedPtr<IAvaOutliner> GetAvaOutliner() const { return AvaOutliner; }
 
 	//~ Begin IAvaEditorExtension
 	virtual void Activate() override;
@@ -64,7 +64,7 @@ public:
 private:
 	void GroupSelection();
 
-	TSharedPtr<FAvaOutliner> AvaOutliner;
+	TSharedPtr<IAvaOutliner> AvaOutliner;
 
 	TSharedRef<FUICommandList> OutlinerCommands;
 };
diff --git a/Engine/Plugins/Experimental/Avalanche/Source/AvalancheEditor/Private/Outliner/AvaOutlinerMaterialDesigner.cpp b/Engine/Plugins/Experimental/Avalanche/Source/AvalancheEditor/Private/Outliner/AvaOutlinerMaterialDesigner.cpp
--- a/Engine/Plugins/Experimental/Avalanche/Source/AvalancheEditor/Private/Outliner/AvaOutlinerMaterialDesigner.cpp
+++ b/Engine/Plugins/Experimental/Avalanche/Source/AvalancheEditor/Private/Outliner/AvaOutlinerMaterialDesigner.cpp
@@ -13,7 +13,7 @@
 
 #define LOCTEXT_NAMESPACE "AvaOutlinerMaterialDesigner"
 
-FAvaOutlinerMaterialDesigner::FAvaOutlinerMaterialDesigner(FAvaOutliner& InOutliner
+FAvaOutlinerMaterialDesigner::FAvaOutlinerMaterialDesigner(IAvaOutliner& InOutliner
 		, UMaterialInterface* InMaterial
 		, const FAvaOutlinerItemPtr& InReferencingItem
 		, int32 InMaterialIndex)
diff --git a/Engine/Plugins/Experimental/Avalanche/Source/AvalancheEditor/Private/Outliner/AvaOutlinerMaterialDesigner.h b/Engine/Plugins/Experimental/Avalanche/Source/AvalancheEditor/Private/Outliner/AvaOutlinerMaterialDesigner.h
--- a/Engine/Plugins/Experimental/Avalanche/Source/AvalancheEditor/Private/Outliner/AvaOutlinerMaterialDesigner.h
+++ b/Engine/Plugins/Experimental/Avalanche/Source/AvalancheEditor/Private/Outliner/AvaOutlinerMaterialDesigner.h
@@ -17,7 +17,7 @@ class FAvaOutlinerMaterialDesigner : public FAvaOutlinerMaterial
 public:
 	UE_AVA_INHERITS_WITH_SUPER(FAvaOutlinerMaterialDesigner, FAvaOutlinerMaterial)
 
-	FAvaOutlinerMaterialDesigner(FAvaOutliner& InOutliner
+	FAvaOutlinerMaterialDesigner(IAvaOutliner& InOutliner
 		, UMaterialInterface* InMaterial
 		, const FAvaOutlinerItemPtr& InReferencingItem
 		, int32 InMaterialIndex);
diff --git a/Engine/Plugins/Experimental/Avalanche/Source/AvalancheEditor/Private/Outliner/AvaOutlinerMaterialDesignerProxy.cpp b/Engine/Plugins/Experimental/Avalanche/Source/AvalancheEditor/Private/Outliner/AvaOutlinerMaterialDesignerProxy.cpp
--- a/Engine/Plugins/Experimental/Avalanche/Source/AvalancheEditor/Private/Outliner/AvaOutlinerMaterialDesignerProxy.cpp
+++ b/Engine/Plugins/Experimental/Avalanche/Source/AvalancheEditor/Private/Outliner/AvaOutlinerMaterialDesignerProxy.cpp
@@ -1,19 +1,18 @@
 // Copyright Epic Games, Inc. All Rights Reserved.
 
 #include "Outliner/AvaOutlinerMaterialDesignerProxy.h"
-#include "AvaOutliner.h"
 #include "Components/PrimitiveComponent.h"
+#include "IAvaOutliner.h"
 #include "Item/AvaOutlinerComponent.h"
+#include "Item/AvaOutlinerItemParameters.h"
 #include "Item/AvaOutlinerMaterial.h"
-#include "ItemActions/AvaOutlinerAddItem.h"
-#include "ItemActions/AvaOutlinerRemoveItem.h"
 #include "Material/DynamicMaterialInstance.h"
 #include "Materials/Material.h"
 #include "Outliner/AvaOutlinerMaterialDesigner.h"
 
 #define LOCTEXT_NAMESPACE "AvaOutlinerMaterialDesignerProxy"
 
-FAvaOutlinerMaterialDesignerProxy::FAvaOutlinerMaterialDesignerProxy(FAvaOutliner& InOutliner, const FAvaOutlinerItemPtr& InParentItem)
+FAvaOutlinerMaterialDesignerProxy::FAvaOutlinerMaterialDesignerProxy(IAvaOutliner& InOutliner, const FAvaOutlinerItemPtr& InParentItem)
 	: FAvaOutlinerMaterialProxy(InOutliner, InParentItem)
 {
 }
diff --git a/Engine/Plugins/Experimental/Avalanche/Source/AvalancheEditor/Private/Outliner/AvaOutlinerMaterialDesignerProxy.h b/Engine/Plugins/Experimental/Avalanche/Source/AvalancheEditor/Private/Outliner/AvaOutlinerMaterialDesignerProxy.h
--- a/Engine/Plugins/Experimental/Avalanche/Source/AvalancheEditor/Private/Outliner/AvaOutlinerMaterialDesignerProxy.h
+++ b/Engine/Plugins/Experimental/Avalanche/Source/AvalancheEditor/Private/Outliner/AvaOutlinerMaterialDesignerProxy.h
@@ -10,7 +10,7 @@ class FAvaOutlinerMaterialDesignerProxy : public FAvaOutlinerMaterialProxy
 public:
 	UE_AVA_INHERITS_WITH_SUPER(FAvaOutlinerMaterialDesignerProxy, FAvaOutlinerMaterialProxy);
 
-	FAvaOutlinerMaterialDesignerProxy(FAvaOutliner& InOutliner, const FAvaOutlinerItemPtr& InParentItem);
+	FAvaOutlinerMaterialDesignerProxy(IAvaOutliner& InOutliner, const FAvaOutlinerItemPtr& InParentItem);
 
 	virtual ~FAvaOutlinerMaterialDesignerProxy() override = default;
 
diff --git a/Engine/Plugins/Experimental/Avalanche/Source/AvalancheEditor/Private/Outliner/AvaOutlinerTabSpawner.cpp b/Engine/Plugins/Experimental/Avalanche/Source/AvalancheEditor/Private/Outliner/AvaOutlinerTabSpawner.cpp
--- a/Engine/Plugins/Experimental/Avalanche/Source/AvalancheEditor/Private/Outliner/AvaOutlinerTabSpawner.cpp
+++ b/Engine/Plugins/Experimental/Avalanche/Source/AvalancheEditor/Private/Outliner/AvaOutlinerTabSpawner.cpp
@@ -1,12 +1,12 @@
 ﻿// Copyright Epic Games, Inc. All Rights Reserved.
 
 #include "AvaOutlinerTabSpawner.h"
-#include "AvaOutliner.h"
 #include "AvaOutlinerExtension.h"
-#include "AvaOutlinerView.h"
 #include "Framework/Docking/TabManager.h"
 #include "Framework/Docking/WorkspaceItem.h"
 #include "IAvaEditor.h"
+#include "IAvaOutliner.h"
+#include "IAvaOutlinerView.h"
 #include "Styling/AppStyle.h"
 
 #define LOCTEXT_NAMESPACE "AvaOutlinerTabSpawner"
@@ -43,13 +43,13 @@ TSharedRef<SWidget> FAvaOutlinerTabSpawner::CreateTabBody()
 		return GetNullWidget();
 	}
 
-	const TSharedPtr<FAvaOutliner> Outliner = OutlinerExtension->GetAvaOutliner();
+	const TSharedPtr<IAvaOutliner> Outliner = OutlinerExtension->GetAvaOutliner();
 	if (!ensure(Outliner.IsValid()))
 	{
 		return GetNullWidget();
 	}
 
-	TSharedPtr<FAvaOutlinerView> OutlinerView = Outliner->GetOutlinerView(OutlinerViewId);
+	TSharedPtr<IAvaOutlinerView> OutlinerView = Outliner->GetOutlinerView(OutlinerViewId);
 	if (!OutlinerView.IsValid())
 	{
 		OutlinerView = Outliner->RegisterOutlinerView(OutlinerViewId);
diff --git a/Engine/Plugins/Experimental/Avalanche/Source/AvalancheEffectors/Private/Cloner/AvaClonerComponent.cpp b/Engine/Plugins/Experimental/Avalanche/Source/AvalancheEffectors/Private/Cloner/AvaClonerComponent.cpp
--- a/Engine/Plugins/Experimental/Avalanche/Source/AvalancheEffectors/Private/Cloner/AvaClonerComponent.cpp
+++ b/Engine/Plugins/Experimental/Avalanche/Source/AvalancheEffectors/Private/Cloner/AvaClonerComponent.cpp
@@ -26,10 +26,10 @@
 #include "UObject/Package.h"
 
 #if WITH_EDITOR
-#include "AvaOutliner.h"
 #include "AvaOutlinerDefines.h"
 #include "AvaOutlinerSubsystem.h"
 #include "Framework/Notifications/NotificationManager.h"
+#include "IAvaOutliner.h"
 #include "Item/AvaOutlinerActor.h"
 #include "Materials/Material.h"
 #include "Misc/MessageDialog.h"
@@ -595,7 +595,7 @@ void UAvaClonerComponent::GetOrderedRootActors(TArray<AActor*>& OutActors) const
 #if WITH_EDITOR
 	if (const UAvaOutlinerSubsystem* const OutlinerSubsystem = ClonerWorld->GetSubsystem<UAvaOutlinerSubsystem>())
 	{
-		const TSharedPtr<FAvaOutliner> AvaOutliner = OutlinerSubsystem->GetOutliner();
+		const TSharedPtr<IAvaOutliner> AvaOutliner = OutlinerSubsystem->GetOutliner();
 		if (AvaOutliner.IsValid())
 		{
 			const FAvaOutlinerItemPtr OutlinerClonerItem = AvaOutliner->FindItem(ClonerActor);
diff --git a/Engine/Plugins/Experimental/Avalanche/Source/AvalancheModifiers/Private/AvaModifiersActorUtils.cpp b/Engine/Plugins/Experimental/Avalanche/Source/AvalancheModifiers/Private/AvaModifiersActorUtils.cpp
--- a/Engine/Plugins/Experimental/Avalanche/Source/AvalancheModifiers/Private/AvaModifiersActorUtils.cpp
+++ b/Engine/Plugins/Experimental/Avalanche/Source/AvalancheModifiers/Private/AvaModifiersActorUtils.cpp
@@ -7,8 +7,8 @@
 #include "GameFramework/Actor.h"
 
 #if WITH_EDITOR
-#include "AvaOutliner.h"
 #include "AvaOutlinerUtils.h"
+#include "IAvaOutliner.h"
 #endif
 
 FBox FAvaModifiersActorUtils::GetActorsBounds(const TSet<TWeakObjectPtr<AActor>>& InActors, const FTransform& InReferenceTransform, bool bInSkipHidden)
@@ -180,7 +180,7 @@ bool FAvaModifiersActorUtils::IsActorNotIsolated(const AActor* InActor)
 	bool bIsIsolatingActors = false;
 	TArray<TWeakObjectPtr<const AActor>> IsolatedActors;
 
-	TSharedPtr<FAvaOutliner> AvaOutliner = FAvaOutlinerUtils::EditorGetOutliner(InActor->GetWorld());
+	TSharedPtr<IAvaOutliner> AvaOutliner = FAvaOutlinerUtils::EditorGetOutliner(InActor->GetWorld());
 	if (AvaOutliner.IsValid())
 	{
 		bIsIsolatingActors = FAvaOutlinerUtils::EditorActorIsolationInfo(AvaOutliner, IsolatedActors);
diff --git a/Engine/Plugins/Experimental/Avalanche/Source/AvalancheModifiers/Private/Extensions/AvaSceneTreeUpdateModifierExtension.cpp b/Engine/Plugins/Experimental/Avalanche/Source/AvalancheModifiers/Private/Extensions/AvaSceneTreeUpdateModifierExtension.cpp
--- a/Engine/Plugins/Experimental/Avalanche/Source/AvalancheModifiers/Private/Extensions/AvaSceneTreeUpdateModifierExtension.cpp
+++ b/Engine/Plugins/Experimental/Avalanche/Source/AvalancheModifiers/Private/Extensions/AvaSceneTreeUpdateModifierExtension.cpp
@@ -7,10 +7,10 @@
 #include "GameFramework/Actor.h"
 
 #if WITH_EDITOR
-#include "AvaOutliner.h"
 #include "AvaOutlinerSubsystem.h"
 #include "AvaOutlinerUtils.h"
 #include "Engine/World.h"
+#include "IAvaOutliner.h"
 #endif
 
 FAvaSceneTreeUpdateModifierExtension::FAvaSceneTreeUpdateModifierExtension(IAvaSceneTreeUpdateHandler* InExtensionHandler)
@@ -93,10 +93,11 @@ void FAvaSceneTreeUpdateModifierExtension::OnExtensionEnabled(EActorModifierCore
 		ActorHierarchyChanged.RemoveAll(this);
 		ActorHierarchyChanged.AddSP(this, &FAvaSceneTreeUpdateModifierExtension::OnActorHierarchyChanged);
 
-		if (const TSharedPtr<FAvaOutliner> Outliner = OutlinerSubsystem->GetOutliner())
+		if (const TSharedPtr<IAvaOutliner> Outliner = OutlinerSubsystem->GetOutliner())
 		{
-			Outliner->OnOutlinerLoaded.RemoveAll(this);
-			Outliner->OnOutlinerLoaded.AddSP(this, &FAvaSceneTreeUpdateModifierExtension::OnOutlinerLoaded);
+			IAvaOutliner::FOnOutlinerLoaded& OnOutlinerLoaded = Outliner->GetOnOutlinerLoaded();
+			OnOutlinerLoaded.RemoveAll(this);
+			OnOutlinerLoaded.AddSP(this, &FAvaSceneTreeUpdateModifierExtension::OnOutlinerLoaded);
 		}
 	}
 #endif
@@ -120,9 +121,9 @@ void FAvaSceneTreeUpdateModifierExtension::OnExtensionDisabled(EActorModifierCor
 	if (UAvaOutlinerSubsystem* const OutlinerSubsystem = World->GetSubsystem<UAvaOutlinerSubsystem>())
 	{
 		OutlinerSubsystem->OnActorHierarchyChanged().RemoveAll(this);
-		if (const TSharedPtr<FAvaOutliner> Outliner = OutlinerSubsystem->GetOutliner())
+		if (const TSharedPtr<IAvaOutliner> Outliner = OutlinerSubsystem->GetOutliner())
 		{
-			Outliner->OnOutlinerLoaded.RemoveAll(this);
+			Outliner->GetOnOutlinerLoaded().RemoveAll(this);
 		}
 	}
 #endif
@@ -335,7 +336,7 @@ TArray<TWeakObjectPtr<AActor>> FAvaSceneTreeUpdateModifierExtension::GetDirectCh
 	const UAvaOutlinerSubsystem* const OutlinerSubsystem = World->GetSubsystem<UAvaOutlinerSubsystem>();
 	if (IsValid(OutlinerSubsystem))
 	{
-		const TSharedPtr<FAvaOutliner> AvaOutliner = OutlinerSubsystem->GetOutliner();
+		const TSharedPtr<IAvaOutliner> AvaOutliner = OutlinerSubsystem->GetOutliner();
 		if (AvaOutliner.IsValid())
 		{
 			DirectChildren = FAvaOutlinerUtils::EditorOutlinerChildActors(AvaOutliner, InActor);
diff --git a/Engine/Plugins/Experimental/Avalanche/Source/AvalancheModifiers/Private/Modifiers/AvaTranslucentPriorityModifier.cpp b/Engine/Plugins/Experimental/Avalanche/Source/AvalancheModifiers/Private/Modifiers/AvaTranslucentPriorityModifier.cpp
--- a/Engine/Plugins/Experimental/Avalanche/Source/AvalancheModifiers/Private/Modifiers/AvaTranslucentPriorityModifier.cpp
+++ b/Engine/Plugins/Experimental/Avalanche/Source/AvalancheModifiers/Private/Modifiers/AvaTranslucentPriorityModifier.cpp
@@ -12,9 +12,9 @@
 #include "Shared/AvaTranslucentPriorityModifierShared.h"
 
 #if WITH_EDITOR
-#include "AvaOutliner.h"
 #include "AvaOutlinerSubsystem.h"
 #include "AvaOutlinerUtils.h"
+#include "IAvaOutliner.h"
 #endif
 
 #define LOCTEXT_NAMESPACE "AvaTranslucentPriorityModifier"
diff --git a/Engine/Plugins/Experimental/Avalanche/Source/AvalancheModifiers/Private/Shared/AvaTranslucentPriorityModifierShared.cpp b/Engine/Plugins/Experimental/Avalanche/Source/AvalancheModifiers/Private/Shared/AvaTranslucentPriorityModifierShared.cpp
--- a/Engine/Plugins/Experimental/Avalanche/Source/AvalancheModifiers/Private/Shared/AvaTranslucentPriorityModifierShared.cpp
+++ b/Engine/Plugins/Experimental/Avalanche/Source/AvalancheModifiers/Private/Shared/AvaTranslucentPriorityModifierShared.cpp
@@ -3,14 +3,14 @@
 #include "Shared/AvaTranslucentPriorityModifierShared.h"
 #include "AvaActorUtils.h"
 #include "AvaBlueprint.h"
-#include "AvaOutlinerUtils.h"
 #include "Camera/CameraActor.h"
 #include "Camera/CameraComponent.h"
 #include "Framework/AvaGameInstance.h"
 
 #if WITH_EDITOR
-#include "AvaOutliner.h"
-#include "AvaOutlinerSubsystem.h"
+#include "AvaOutlinerUtils.h"
+#include "IAvaOutliner.h"
+#include "Item/AvaOutlinerItemUtils.h"
 #endif
 
 void UAvaTranslucentPriorityModifierShared::SetComponentsState(UAvaTranslucentPriorityModifier* InModifierContext, const TSet<TWeakObjectPtr<UPrimitiveComponent>>& InComponents)
@@ -239,7 +239,7 @@ TArray<const FAvaTranslucentPriorityModifierComponentState*> UAvaTranslucentPrio
 		// Sort current modifiers by outliner hierarchy
 #if WITH_EDITOR
 		// Use the outliner first if available for sorting
-		TSharedPtr<FAvaOutliner> AvaOutliner = FAvaOutlinerUtils::EditorGetOutliner(World);
+		TSharedPtr<IAvaOutliner> AvaOutliner = FAvaOutlinerUtils::EditorGetOutliner(World);
 		if (AvaOutliner.IsValid())
 		{
 			SortedComponentStates.Sort([AvaOutliner](const FAvaTranslucentPriorityModifierComponentState& InComponentA, const FAvaTranslucentPriorityModifierComponentState& InComponentB)->bool
@@ -247,7 +247,7 @@ TArray<const FAvaTranslucentPriorityModifierComponentState*> UAvaTranslucentPrio
 				const FAvaOutlinerItemPtr OutlinerItemA = AvaOutliner->FindItem(InComponentA.GetOwningActor());
 				const FAvaOutlinerItemPtr OutlinerItemB = AvaOutliner->FindItem(InComponentB.GetOwningActor());
 
-				return FAvaOutliner::CompareOutlinerItemOrder(OutlinerItemB, OutlinerItemA);
+				return UE::AvaOutliner::CompareOutlinerItemOrder(OutlinerItemB, OutlinerItemA);
 			});
 		}
 		else
diff --git a/Engine/Plugins/Experimental/Avalanche/Source/AvalancheModifiersEditor/Private/AvaModifiersEditorModule.cpp b/Engine/Plugins/Experimental/Avalanche/Source/AvalancheModifiersEditor/Private/AvaModifiersEditorModule.cpp
--- a/Engine/Plugins/Experimental/Avalanche/Source/AvalancheModifiersEditor/Private/AvaModifiersEditorModule.cpp
+++ b/Engine/Plugins/Experimental/Avalanche/Source/AvalancheModifiersEditor/Private/AvaModifiersEditorModule.cpp
@@ -3,9 +3,8 @@
 #include "AvaModifiersEditorModule.h"
 #include "AvaDefs.h"
 #include "AvaModifiersEditorStyle.h"
-#include "AvaOutliner.h"
-#include "AvaOutlinerView.h"
 #include "DragDropOps/AvaOutlinerItemDragDropOp.h"
+#include "IAvaOutliner.h"
 #include "IAvaOutlinerModule.h"
 #include "Item/AvaOutlinerActor.h"
 #include "ItemProxies/AvaOutlinerItemProxyRegistry.h"
@@ -37,7 +36,7 @@ void FAvaModifiersEditorModule::RegisterOutlinerItems()
 	ItemProxyRegistry.RegisterItemProxyWithDefaultFactory<FAvaOutlinerModifierProxy, 60>();
 
 	OutlinerProxiesExtensionDelegateHandle = IAvaOutlinerModule::Get().GetOnExtendItemProxiesForItem().AddLambda(
-		[](FAvaOutliner& InOutliner, const FAvaOutlinerItemPtr& InItem, TArray<TSharedPtr<FAvaOutlinerItemProxy>>& OutItemProxies)
+		[](IAvaOutliner& InOutliner, const FAvaOutlinerItemPtr& InItem, TArray<TSharedPtr<FAvaOutlinerItemProxy>>& OutItemProxies)
 		{
 			if (InItem->IsA<FAvaOutlinerActor>())
 			{
diff --git a/Engine/Plugins/Experimental/Avalanche/Source/AvalancheModifiersEditor/Private/Outliner/AvaOutlinerModifier.cpp b/Engine/Plugins/Experimental/Avalanche/Source/AvalancheModifiersEditor/Private/Outliner/AvaOutlinerModifier.cpp
--- a/Engine/Plugins/Experimental/Avalanche/Source/AvalancheModifiersEditor/Private/Outliner/AvaOutlinerModifier.cpp
+++ b/Engine/Plugins/Experimental/Avalanche/Source/AvalancheModifiersEditor/Private/Outliner/AvaOutlinerModifier.cpp
@@ -9,7 +9,7 @@
 #include "Subsystems/ActorModifierCoreEditorSubsystem.h"
 #include "Subsystems/ActorModifierCoreSubsystem.h"
 
-FAvaOutlinerModifier::FAvaOutlinerModifier(FAvaOutliner& InOutliner, UActorModifierCoreBase* InModifier)
+FAvaOutlinerModifier::FAvaOutlinerModifier(IAvaOutliner& InOutliner, UActorModifierCoreBase* InModifier)
 	: FAvaOutlinerObject(InOutliner, InModifier)
 	, Modifier(InModifier)
 {
diff --git a/Engine/Plugins/Experimental/Avalanche/Source/AvalancheModifiersEditor/Private/Outliner/AvaOutlinerModifier.h b/Engine/Plugins/Experimental/Avalanche/Source/AvalancheModifiersEditor/Private/Outliner/AvaOutlinerModifier.h
--- a/Engine/Plugins/Experimental/Avalanche/Source/AvalancheModifiersEditor/Private/Outliner/AvaOutlinerModifier.h
+++ b/Engine/Plugins/Experimental/Avalanche/Source/AvalancheModifiersEditor/Private/Outliner/AvaOutlinerModifier.h
@@ -14,7 +14,7 @@ class AVALANCHEMODIFIERSEDITOR_API FAvaOutlinerModifier : public FAvaOutlinerObj
 public:
 	UE_AVA_INHERITS_WITH_SUPER(FAvaOutlinerModifier, FAvaOutlinerObject);
 	
-	FAvaOutlinerModifier(FAvaOutliner& InOutliner, UActorModifierCoreBase* InObject);
+	FAvaOutlinerModifier(IAvaOutliner& InOutliner, UActorModifierCoreBase* InObject);
 	
 	UActorModifierCoreBase* GetModifier() const { return Modifier.Get(); }
 	
diff --git a/Engine/Plugins/Experimental/Avalanche/Source/AvalancheModifiersEditor/Private/Outliner/AvaOutlinerModifierProxy.cpp b/Engine/Plugins/Experimental/Avalanche/Source/AvalancheModifiersEditor/Private/Outliner/AvaOutlinerModifierProxy.cpp
--- a/Engine/Plugins/Experimental/Avalanche/Source/AvalancheModifiersEditor/Private/Outliner/AvaOutlinerModifierProxy.cpp
+++ b/Engine/Plugins/Experimental/Avalanche/Source/AvalancheModifiersEditor/Private/Outliner/AvaOutlinerModifierProxy.cpp
@@ -1,7 +1,7 @@
 // Copyright Epic Games, Inc. All Rights Reserved.
 
 #include "Outliner/AvaOutlinerModifierProxy.h"
-#include "AvaOutliner.h"
+#include "IAvaOutliner.h"
 #include "Item/AvaOutlinerActor.h"
 #include "Modifiers/ActorModifierCoreStack.h"
 #include "Modifiers/AvaBaseModifier.h"
@@ -12,7 +12,7 @@
 
 #define LOCTEXT_NAMESPACE "AvaOutlinerModifierProxy"
 
-FAvaOutlinerModifierProxy::FAvaOutlinerModifierProxy(FAvaOutliner& InOutliner, const FAvaOutlinerItemPtr& InParentItem)
+FAvaOutlinerModifierProxy::FAvaOutlinerModifierProxy(IAvaOutliner& InOutliner, const FAvaOutlinerItemPtr& InParentItem)
 	: Super(InOutliner, InParentItem)
 {
 	ModifierIcon = FSlateIconFinder::FindIconForClass(UAvaBaseModifier::StaticClass());
diff --git a/Engine/Plugins/Experimental/Avalanche/Source/AvalancheModifiersEditor/Private/Outliner/AvaOutlinerModifierProxy.h b/Engine/Plugins/Experimental/Avalanche/Source/AvalancheModifiersEditor/Private/Outliner/AvaOutlinerModifierProxy.h
--- a/Engine/Plugins/Experimental/Avalanche/Source/AvalancheModifiersEditor/Private/Outliner/AvaOutlinerModifierProxy.h
+++ b/Engine/Plugins/Experimental/Avalanche/Source/AvalancheModifiersEditor/Private/Outliner/AvaOutlinerModifierProxy.h
@@ -15,7 +15,7 @@ class AVALANCHEMODIFIERSEDITOR_API FAvaOutlinerModifierProxy : public FAvaOutlin
 public:
 	UE_AVA_INHERITS_WITH_SUPER(FAvaOutlinerModifierProxy, FAvaOutlinerItemProxy);
 
-	FAvaOutlinerModifierProxy(FAvaOutliner& InOutliner, const FAvaOutlinerItemPtr& InParentItem);
+	FAvaOutlinerModifierProxy(IAvaOutliner& InOutliner, const FAvaOutlinerItemPtr& InParentItem);
 
 	/** Gets the Modifier Stack to use (e.g. for an Actor it would be the Root Modifier Stack) */
 	UActorModifierCoreStack* GetModifierStack() const;
diff --git a/Engine/Plugins/Experimental/Avalanche/Source/AvalancheOutliner/Private/AvaOutliner.cpp b/Engine/Plugins/Experimental/Avalanche/Source/AvalancheOutliner/Private/AvaOutliner.cpp
--- a/Engine/Plugins/Experimental/Avalanche/Source/AvalancheOutliner/Private/AvaOutliner.cpp
+++ b/Engine/Plugins/Experimental/Avalanche/Source/AvalancheOutliner/Private/AvaOutliner.cpp
@@ -23,6 +23,7 @@
 #include "IAvaOutlinerProvider.h"
 #include "Item/AvaOutlinerActor.h"
 #include "Item/AvaOutlinerComponent.h"
+#include "Item/AvaOutlinerItemUtils.h"
 #include "Item/AvaOutlinerMaterial.h"
 #include "Item/AvaOutlinerTreeRoot.h"
 #include "ItemActions/AvaOutlinerAddItem.h"
@@ -143,34 +144,11 @@
 		&& OutlinerProvider.CanOutlinerProcessActorSpawn(InActor);
 }
 
-void FAvaOutliner::SetBaseCommandList(const TSharedPtr<FUICommandList>& InBaseCommandList)
-{
-	BaseCommandListWeak = InBaseCommandList;
-}
-
 TSharedPtr<FUICommandList> FAvaOutliner::GetBaseCommandList() const
 {
 	return BaseCommandListWeak.Pin();
 }
 
-void FAvaOutliner::Serialize(FArchive& Ar)
-{
-	Ar.UsingCustomVersion(FAvaOutlinerVersion::GUID);
-	
-	FCustomVersionContainer CustomVersionContainer = Ar.GetCustomVersions();
-	CustomVersionContainer.Serialize(Ar);
-	Ar.SetCustomVersions(CustomVersionContainer);
-	
-	SaveState->Serialize(*this, Ar);
-	
-	if (Ar.IsLoading())
-	{
-		Refresh();
-		bOutlinerDirty = false;
-		OnOutlinerLoaded.Broadcast();
-	}
-}
-
 TArray<FName> FAvaOutliner::GetRegisteredItemProxyTypeNames() const
 {
 	TArray<FName> OutItemProxyTypeNames;
@@ -204,8 +182,7 @@
 	InItem->GetItemProxies(OutItemProxies);
 	
 	IAvaOutlinerModule::Get().GetOnExtendItemProxiesForItem().Broadcast(*this, InItem, OutItemProxies);	
-	OutlinerProvider.GetItemProxiesForItem(*this, InItem, OutItemProxies);
-	
+
 	// Clean up any invalid Item Proxy
 	OutItemProxies.RemoveAll([](const TSharedPtr<FAvaOutlinerItemProxy>& InItemProxy) { return !InItemProxy.IsValid(); });
 	
@@ -216,6 +193,22 @@
 	});
 }
 
+IAvaOutlinerItemProxyFactory* FAvaOutliner::GetItemProxyFactory(FName InItemProxyTypeName) const
+{
+	// First look for the Registry in Outliner
+	if (IAvaOutlinerItemProxyFactory* Factory = ItemProxyRegistry.GetItemProxyFactory(InItemProxyTypeName))
+	{
+		return Factory;
+	}
+	// Fallback to finding the Factory in the Module if the Outliner did not find it
+	return IAvaOutlinerModule::Get().GetItemProxyRegistry().GetItemProxyFactory(InItemProxyTypeName);
+}
+
+const TSharedRef<FAvaOutlinerSaveState>& FAvaOutliner::GetSaveState() const
+{
+	return SaveState;
+}
+
 bool FAvaOutliner::IsOutlinerLocked() const
 {
 	return OutlinerProvider.ShouldLockOutliner();
@@ -226,113 +219,30 @@
 	RequestRefresh();
 }
 
-void FAvaOutliner::PostUndo(bool bSuccess)
+void FAvaOutliner::SetBaseCommandList(const TSharedPtr<FUICommandList>& InBaseCommandList)
 {
-	if (bSuccess)
-	{
-		const int32 QueueIndex = GEditor->Trans->GetQueueLength() - GEditor->Trans->GetUndoCount();
-		const FTransaction* Transaction = GEditor->Trans->GetTransaction(QueueIndex);
-		HandleUndoRedoTransaction(Transaction, true);
-	}
+	BaseCommandListWeak = InBaseCommandList;
 }
 
-void FAvaOutliner::PostRedo(bool bSuccess)
+void FAvaOutliner::Serialize(FArchive& Ar)
 {
-	if (bSuccess)
+	Ar.UsingCustomVersion(FAvaOutlinerVersion::GUID);
+	
+	FCustomVersionContainer CustomVersionContainer = Ar.GetCustomVersions();
+	CustomVersionContainer.Serialize(Ar);
+	Ar.SetCustomVersions(CustomVersionContainer);
+	
+	SaveState->Serialize(*this, Ar);
+	
+	if (Ar.IsLoading())
 	{
-		const int32 QueueIndex = GEditor->Trans->GetQueueLength() - GEditor->Trans->GetUndoCount();
-		const FTransaction* Transaction = GEditor->Trans->GetTransaction(QueueIndex);
-		HandleUndoRedoTransaction(Transaction, false);
-	}
-}
-
-TStatId FAvaOutliner::GetStatId() const
-{
-	RETURN_QUICK_DECLARE_CYCLE_STAT(FAvaOutliner, STATGROUP_Tickables);
-}
-
-void FAvaOutliner::Tick(float InDeltaTime)
-{
-	if (NeedsRefresh())
-	{
 		Refresh();
-	}
-
-	if (bOutlinerDirty)
-	{
-		OutlinerProvider.OnOutlinerModified();
 		bOutlinerDirty = false;
+		OnOutlinerLoaded.Broadcast();
 	}
-
-	// Select Objects Pending Selection
-	if (ObjectsLastSelected.IsValid())
-	{
-		TArray<FAvaOutlinerItemPtr> ItemsToSelect;
-
-		if (ObjectsLastSelected->Num() > 0)
-		{
-			ItemsToSelect.Reserve(ObjectsLastSelected->Num());
-			for (TWeakObjectPtr<UObject> Object : *ObjectsLastSelected)
-			{
-				if (Object.IsValid())
-				{
-					if (FAvaOutlinerItemPtr Item = FindItem(FAvaOutlinerItemId(Object.Get())))
-					{
-						// If it's a Shared Object, select all its References as Shared Objects are not in the Outliner directly
-						// and the outliner does not know which Object Reference should be selected
-						if (const FAvaOutlinerSharedObject* const SharedObjectItem = Item->CastTo<FAvaOutlinerSharedObject>())
-						{
-							ItemsToSelect.Append(SharedObjectItem->GetObjectReferences());
-						}
-						else
-						{
-							ItemsToSelect.Add(Item);
-						}
-					}
-				}
-			}
-		}
-
-		// Only Scroll Into View, don't signal selection since we just come from the selection notify itself
-		SelectItems(ItemsToSelect, EAvaOutlinerItemSelectionFlags::ScrollIntoView);
-		ObjectsLastSelected.Reset();
-	}
-
-	ForEachOutlinerView([InDeltaTime](const TSharedPtr<FAvaOutlinerView>& InOutlinerView)
-	{
-		InOutlinerView->Tick(InDeltaTime);
-	});
 }
 
-void FAvaOutliner::DuplicateItems(TArray<FAvaOutlinerItemPtr> InItems
-	, FAvaOutlinerItemPtr InRelativeItem
-	, TOptional<EItemDropZone> InRelativeDropZone)
+TSharedPtr<IAvaOutlinerView> FAvaOutliner::RegisterOutlinerView(int32 InOutlinerViewId)
 {
-	SortItems(InItems);
-
-	TArray<AActor*> TemplateActors;
-	for (const FAvaOutlinerItemPtr& Item : InItems)
-	{
-		if (!Item.IsValid())
-		{
-			continue;
-		}
-		
-		if (const FAvaOutlinerActor* const ActorItem = Item->CastTo<FAvaOutlinerActor>())
-		{
-			TemplateActors.Add(ActorItem->GetActor());
-		}
-	}
-
-	if (TemplateActors.IsEmpty())
-	{
-		return;
-	}
-
-	OutlinerProvider.OutlinerDuplicateActors(TemplateActors);
-}
-
-TSharedPtr<FAvaOutlinerView> FAvaOutliner::RegisterOutlinerView(int32 InOutlinerViewId)
-{
 	const bool bShouldCreateWidget = OutlinerProvider.ShouldCreateWidget();
 	
@@ -348,18 +258,7 @@
 	return OutlinerView;
 }
 
-void FAvaOutliner::UnregisterOutlinerView(int32 InOutlinerViewId)
+TSharedPtr<IAvaOutlinerView> FAvaOutliner::GetOutlinerView(int32 InOutlinerViewId) const
 {
-	OutlinerViews.Remove(InOutlinerViewId);
-}
-
-void FAvaOutliner::UpdateRecentOutlinerViews(int32 InOutlinerViewId)
-{
-	RecentOutlinerViews.Remove(InOutlinerViewId);
-	RecentOutlinerViews.Add(InOutlinerViewId);
-}
-
-TSharedPtr<FAvaOutlinerView> FAvaOutliner::GetOutlinerView(int32 InOutlinerViewId) const
-{
 	if (const TSharedPtr<FAvaOutlinerView>* const FoundOutlinerView = OutlinerViews.Find(InOutlinerViewId))
 	{
@@ -368,36 +267,6 @@
 	return nullptr;
 }
 
-TSharedPtr<FAvaOutlinerView> FAvaOutliner::GetMostRecentOutlinerView() const
-{
-	for (int32 Index = RecentOutlinerViews.Num() - 1; Index >= 0; --Index)
-	{
-		TSharedPtr<FAvaOutlinerView> OutlinerView = GetOutlinerView(RecentOutlinerViews[Index]);
-		if (OutlinerView.IsValid())
-		{
-			return OutlinerView;
-		}
-	}
-	return nullptr;
-}
-
-void FAvaOutliner::ForEachOutlinerView(const TFunction<void(const TSharedPtr<FAvaOutlinerView>& InOutlinerView)>& InPredicate) const
-{
-	FAvaOutliner* const MutableThis = const_cast<FAvaOutliner*>(this);
-	for (TMap<int32, TSharedPtr<FAvaOutlinerView>>::TIterator Iter(MutableThis->OutlinerViews); Iter; ++Iter)
-	{
-		const TSharedPtr<FAvaOutlinerView>& OutlinerView = Iter.Value();
-		if (OutlinerView.IsValid())
-		{
-			InPredicate(OutlinerView);
-		}
-		else
-		{
-			Iter.RemoveCurrent();
-		}
-	}
-}
-
 void FAvaOutliner::RegisterItem(const FAvaOutlinerItemPtr& InItem)
 {
 	if (InItem.IsValid())
@@ -426,32 +295,6 @@
 	}
 }
 
-void FAvaOutliner::EnqueueItemActions(TArray<TSharedPtr<IAvaOutlinerAction>>&& InItemActions) noexcept
-{
-	PendingActions.Append(MoveTemp(InItemActions));
-}
-
-void FAvaOutliner::EnqueueItemActions(const TArray<TSharedPtr<IAvaOutlinerAction>>& InItemActions)
-{
-	PendingActions.Append(InItemActions);
-}
-
-bool FAvaOutliner::NeedsRefresh() const
-{
-	//Return false if we're already refreshing
-	if (bRefreshing)
-	{
-		return false;
-	}
-
-	if (bRefreshRequested || PendingActions.Num() > 0)
-	{
-		return true;
-	}
-
-	return false;
-}
-
 void FAvaOutliner::RequestRefresh()
 {
 	bRefreshRequested = true;
@@ -566,6 +409,426 @@
 	return nullptr;
 }
 
+void FAvaOutliner::SetIgnoreNotify(EAvaOutlinerIgnoreNotifyFlags InFlag, bool bIgnore)
+{
+	if (bIgnore)
+	{
+		EnumAddFlags(IgnoreNotifyFlags, InFlag);
+	}
+	else
+	{
+		EnumRemoveFlags(IgnoreNotifyFlags, InFlag);
+	}
+}
+
+void FAvaOutliner::OnActorsCopied(FString& InOutCopiedData, TConstArrayView<AActor*> InCopiedActors)
+{
+	FAvaOutlinerExporter Exporter(SharedThis(this));
+	Exporter.ExportText(InOutCopiedData, InCopiedActors);
+}
+
+void FAvaOutliner::OnActorsPasted(FStringView InPastedData, const TMap<FName, AActor*>& InPastedActors)
+{
+	FAvaOutlinerImporter Importer(SharedThis(this));
+	Importer.ImportText(InPastedData, InPastedActors);
+}
+
+void FAvaOutliner::OnActorsDuplicated(const TMap<AActor*, AActor*>& InDuplicateActorMap
+	, FAvaOutlinerItemPtr InRelativeItem
+	, TOptional<EItemDropZone> InRelativeDropZone)
+{
+	if (InDuplicateActorMap.IsEmpty())
+	{
+		return;
+	}
+
+	// Bail if we're currently Ignoring Duplication Notifies
+	if (EnumHasAnyFlags(IgnoreNotifyFlags, EAvaOutlinerIgnoreNotifyFlags::Duplication))
+	{
+		return;
+	}
+
+	// Map of the Duplicate Actor to the Template Actor's Item
+	TMap<AActor*, FAvaOutlinerItemPtr> DuplicateItemMap;
+	
+	// Try finding the Template Item with the Lowest Index (i.e. Highest in the Tree) to use as Placeholder Relative Item
+	if (!InRelativeItem.IsValid())
+	{
+		DuplicateItemMap.Reserve(InDuplicateActorMap.Num());
+		for (const TPair<AActor*, AActor*>& Pair : InDuplicateActorMap)
+		{
+			AActor* const TemplateActor = Pair.Value;
+			if (IsValid(TemplateActor))
+			{
+				if (FAvaOutlinerItemPtr FoundTemplateItem = FindItem(TemplateActor))
+				{
+					DuplicateItemMap.Emplace(Pair.Key, FoundTemplateItem);
+					if (!InRelativeItem.IsValid() || UE::AvaOutliner::CompareOutlinerItemOrder(FoundTemplateItem, InRelativeItem))
+					{
+						InRelativeItem = FoundTemplateItem;
+					}
+				}
+			}
+		}
+	}
+
+	const EItemDropZone RelativeDropZone = InRelativeDropZone.Get(EItemDropZone::AboveItem);
+
+	// Newly Duplicated Items don't have a Parent Item set yet, so it's not going to call RemoveFromParent, which handles Detachment
+	bool bShouldDetachActors = false;
+	if (InRelativeItem.IsValid())
+	{
+		const bool bRelativeItemIsRoot        = InRelativeItem == RootItem;
+		const bool bRelativeItemIsChildOfRoot = InRelativeItem->GetParent() == RootItem;
+		
+		bShouldDetachActors = bRelativeItemIsRoot || (bRelativeItemIsChildOfRoot && RelativeDropZone != EItemDropZone::OntoItem);
+	}
+
+	TArray<AActor*> DuplicateActors;
+	InDuplicateActorMap.GetKeys(DuplicateActors);
+
+	// Sort based on the ordering of the template actor items in outliner
+	DuplicateActors.Sort([&InDuplicateActorMap, this](const AActor& A, const AActor& B)
+	{
+		AActor* const * const TemplateActorA = InDuplicateActorMap.Find(&A);
+		AActor* const * const TemplateActorB = InDuplicateActorMap.Find(&B);
+
+		const FAvaOutlinerItemPtr TemplateItemA = FindItem(*TemplateActorA);
+		const FAvaOutlinerItemPtr TemplateItemB = FindItem(*TemplateActorB);
+
+		return UE::AvaOutliner::CompareOutlinerItemOrder(TemplateItemA, TemplateItemB);
+	});
+
+	/*
+	 * REVERSE the ordering of Duplicate Actors
+	 * The order needs to be reversed when items are not added above the target item (e.g. RelativeDropZone is either below or onto)
+	 * Some context:
+	 * ABOVE- items are added just above the target, so the first one will always be topmost. Order is preserved
+	 * BELOW- items are added just below the target, so the last one added will be the topmost. Order is in reverse.
+	 * ONTO - items are added at index 0, so last one added will be topmost. Order is in reverse
+	 */
+	if (RelativeDropZone != EItemDropZone::AboveItem)
+	{
+		Algo::Reverse(DuplicateActors);
+	}
+	
+	TArray<TSharedPtr<IAvaOutlinerAction>> ItemActions;
+	ItemActions.Reserve(DuplicateActors.Num());
+	
+	FAvaOutlinerAddItemParams Params;
+	Params.RelativeDropZone = RelativeDropZone;
+	Params.Flags            = EAvaOutlinerAddItemFlags::AddChildren;
+
+	// Finally, Enqueue the Add Item Action to these new Actors
+	for (AActor* const DuplicateActor : DuplicateActors)
+	{
+		Params.RelativeItem.Reset();
+		
+		//We don't need to worry about Child Items since the Parent Item will do all the Attachment Work
+		const AActor* const ParentActor = DuplicateActor->GetAttachParentActor();
+
+		const bool bIsParentDuplicated = InDuplicateActorMap.Contains(ParentActor);
+		if (!bIsParentDuplicated && bShouldDetachActors)
+		{
+			DuplicateActor->DetachFromActor(FDetachmentTransformRules::KeepWorldTransform);
+		}
+
+		FAvaOutlinerItemPtr Item = FindOrAdd<FAvaOutlinerActor>(DuplicateActor);
+		Params.Item = Item;
+		
+		if (FAvaOutlinerItemPtr* const FoundTemplateItem = DuplicateItemMap.Find(DuplicateActor))
+		{
+			FAvaOutlinerItemPtr TemplateItem = *FoundTemplateItem;
+
+			if (!bIsParentDuplicated)
+			{
+				Params.RelativeItem = TemplateItem;
+			}
+			
+			//Sanitize the Flags by Removing the Temp Ones
+			EAvaOutlinerItemFlags TemplateItemFlags = TemplateItem->GetFlags();
+			EnumRemoveFlags(TemplateItemFlags, EAvaOutlinerItemFlags::PendingRemoval | EAvaOutlinerItemFlags::IgnorePendingKill);
+			
+			Item->SetFlags(TemplateItemFlags);
+
+			//Copy the template flags of each outliner view into the new item
+			ForEachOutlinerView([TemplateItem, Item](const TSharedPtr<FAvaOutlinerView>& InOutlinerView)
+			{
+				InOutlinerView->SetViewItemFlags(Item, InOutlinerView->GetViewItemFlags(TemplateItem));
+			});			
+		}
+		
+		//If a Relative Item hasn't been set yet (i.e. no Template Item found for the Item)
+		if (!Params.RelativeItem.IsValid() && !bIsParentDuplicated)
+		{
+			Params.RelativeItem = InRelativeItem;
+		}
+
+		TSharedRef<FAvaOutlinerAddItem> AddItemAction = NewItemAction<FAvaOutlinerAddItem>(Params);
+
+		// If parent is duplicated, make sure we add this item to index 0 so that children are added in reverse order
+		// this is because ParentItem->AddChild(...) adds the Child at Index 0, so it's correcting the ordering issue
+		if (bIsParentDuplicated)
+		{
+			ItemActions.EmplaceAt(0, AddItemAction);
+		}
+		else
+		{
+			ItemActions.Add(AddItemAction);
+		}		
+		
+		// Append after first selected item is added
+		Params.SelectionFlags |= EAvaOutlinerItemSelectionFlags::AppendToCurrentSelection;
+	}
+	
+	EnqueueItemActions(MoveTemp(ItemActions));
+}
+
+void FAvaOutliner::GroupSelection(AActor* InGroupingActor, const TOptional<FAttachmentTransformRules>& InTransformRules)
+{
+	if (!InGroupingActor)
+	{
+		return;
+	}
+
+	FAvaOutlinerItemPtr GroupingItem = FindItem(InGroupingActor);
+	bool bIsGroupingItemNew = false;	
+	if (!GroupingItem.IsValid())
+	{
+		GroupingItem = FindOrAdd<FAvaOutlinerActor>(InGroupingActor);
+		bIsGroupingItemNew = true;
+	}
+	GroupingItem->SetFlags(EAvaOutlinerItemFlags::Expanded);
+
+	TArray<FAvaOutlinerItemPtr> SelectedItems = GetSelectedItems();
+	NormalizeItems(SelectedItems);
+
+	// Enqueue the Grouping Item if New
+	if (bIsGroupingItemNew)
+	{
+		FAvaOutlinerAddItemParams AddGroupingItemParams;
+		AddGroupingItemParams.Item  = GroupingItem;
+		AddGroupingItemParams.Flags = EAvaOutlinerAddItemFlags::AddChildren;
+
+		AddGroupingItemParams.AttachmentTransformRules = InTransformRules;
+
+		if (const FAvaOutlinerItemPtr LowestCommonAncestor = FindLowestCommonAncestor(SelectedItems))
+		{
+			//Find the first path that leads to a Selected Item
+			TArray<FAvaOutlinerItemPtr> Descendants = LowestCommonAncestor->FindPath(SelectedItems);
+			if (Descendants.Num() > 0)
+			{
+				AddGroupingItemParams.RelativeItem     = Descendants[0];
+				AddGroupingItemParams.RelativeDropZone = EItemDropZone::BelowItem;
+			}
+			else
+			{
+				AddGroupingItemParams.RelativeItem     = LowestCommonAncestor;
+				AddGroupingItemParams.RelativeDropZone = EItemDropZone::OntoItem;
+			}
+		}
+
+		EnqueueItemAction<FAvaOutlinerAddItem>(AddGroupingItemParams);
+	}	
+	
+	//If there are no Selected Items, stop here
+	if (SelectedItems.IsEmpty())
+	{
+		return;
+	}
+
+	// Sort Items in reverse order of hierarchy as Add Item will add them at index 0 by default
+	SortItems(SelectedItems, true);
+		
+	FAvaOutlinerAddItemParams AddChildItemParams;
+	AddChildItemParams.RelativeItem     = GroupingItem;
+	AddChildItemParams.RelativeDropZone = EItemDropZone::OntoItem;
+	AddChildItemParams.SelectionFlags   = EAvaOutlinerItemSelectionFlags::AppendToCurrentSelection | EAvaOutlinerItemSelectionFlags::ScrollIntoView;
+
+	AddChildItemParams.AttachmentTransformRules = InTransformRules;
+
+	for (const FAvaOutlinerItemPtr& Item : SelectedItems)
+	{
+		if (Item.IsValid())
+		{
+			AddChildItemParams.Item = Item;
+			EnqueueItemAction<FAvaOutlinerAddItem>(AddChildItemParams);
+		}
+	}
+}
+
+void FAvaOutliner::PostUndo(bool bSuccess)
+{
+	if (bSuccess)
+	{
+		const int32 QueueIndex = GEditor->Trans->GetQueueLength() - GEditor->Trans->GetUndoCount();
+		const FTransaction* Transaction = GEditor->Trans->GetTransaction(QueueIndex);
+		HandleUndoRedoTransaction(Transaction, true);
+	}
+}
+
+void FAvaOutliner::PostRedo(bool bSuccess)
+{
+	if (bSuccess)
+	{
+		const int32 QueueIndex = GEditor->Trans->GetQueueLength() - GEditor->Trans->GetUndoCount();
+		const FTransaction* Transaction = GEditor->Trans->GetTransaction(QueueIndex);
+		HandleUndoRedoTransaction(Transaction, false);
+	}
+}
+
+TStatId FAvaOutliner::GetStatId() const
+{
+	RETURN_QUICK_DECLARE_CYCLE_STAT(FAvaOutliner, STATGROUP_Tickables);
+}
+
+void FAvaOutliner::Tick(float InDeltaTime)
+{
+	if (NeedsRefresh())
+	{
+		Refresh();
+	}
+
+	if (bOutlinerDirty)
+	{
+		bOutlinerDirty = false;
+	}
+
+	// Select Objects Pending Selection
+	if (ObjectsLastSelected.IsValid())
+	{
+		TArray<FAvaOutlinerItemPtr> ItemsToSelect;
+
+		if (ObjectsLastSelected->Num() > 0)
+		{
+			ItemsToSelect.Reserve(ObjectsLastSelected->Num());
+			for (TWeakObjectPtr<UObject> Object : *ObjectsLastSelected)
+			{
+				if (Object.IsValid())
+				{
+					if (FAvaOutlinerItemPtr Item = FindItem(FAvaOutlinerItemId(Object.Get())))
+					{
+						// If it's a Shared Object, select all its References as Shared Objects are not in the Outliner directly
+						// and the outliner does not know which Object Reference should be selected
+						if (const FAvaOutlinerSharedObject* const SharedObjectItem = Item->CastTo<FAvaOutlinerSharedObject>())
+						{
+							ItemsToSelect.Append(SharedObjectItem->GetObjectReferences());
+						}
+						else
+						{
+							ItemsToSelect.Add(Item);
+						}
+					}
+				}
+			}
+		}
+
+		// Only Scroll Into View, don't signal selection since we just come from the selection notify itself
+		SelectItems(ItemsToSelect, EAvaOutlinerItemSelectionFlags::ScrollIntoView);
+		ObjectsLastSelected.Reset();
+	}
+
+	ForEachOutlinerView([InDeltaTime](const TSharedPtr<FAvaOutlinerView>& InOutlinerView)
+	{
+		InOutlinerView->Tick(InDeltaTime);
+	});
+}
+
+void FAvaOutliner::DuplicateItems(TArray<FAvaOutlinerItemPtr> InItems
+	, FAvaOutlinerItemPtr InRelativeItem
+	, TOptional<EItemDropZone> InRelativeDropZone)
+{
+	SortItems(InItems);
+
+	TArray<AActor*> TemplateActors;
+	for (const FAvaOutlinerItemPtr& Item : InItems)
+	{
+		if (!Item.IsValid())
+		{
+			continue;
+		}
+		
+		if (const FAvaOutlinerActor* const ActorItem = Item->CastTo<FAvaOutlinerActor>())
+		{
+			TemplateActors.Add(ActorItem->GetActor());
+		}
+	}
+
+	if (TemplateActors.IsEmpty())
+	{
+		return;
+	}
+
+	OutlinerProvider.OutlinerDuplicateActors(TemplateActors);
+}
+
+void FAvaOutliner::UnregisterOutlinerView(int32 InOutlinerViewId)
+{
+	OutlinerViews.Remove(InOutlinerViewId);
+}
+
+void FAvaOutliner::UpdateRecentOutlinerViews(int32 InOutlinerViewId)
+{
+	RecentOutlinerViews.Remove(InOutlinerViewId);
+	RecentOutlinerViews.Add(InOutlinerViewId);
+}
+
+TSharedPtr<FAvaOutlinerView> FAvaOutliner::GetMostRecentOutlinerView() const
+{
+	for (int32 Index = RecentOutlinerViews.Num() - 1; Index >= 0; --Index)
+	{
+		TSharedPtr<IAvaOutlinerView> OutlinerView = GetOutlinerView(RecentOutlinerViews[Index]);
+		if (OutlinerView.IsValid())
+		{
+			return StaticCastSharedPtr<FAvaOutlinerView>(OutlinerView);
+		}
+	}
+	return nullptr;
+}
+
+void FAvaOutliner::ForEachOutlinerView(const TFunction<void(const TSharedPtr<FAvaOutlinerView>& InOutlinerView)>& InPredicate) const
+{
+	FAvaOutliner* const MutableThis = const_cast<FAvaOutliner*>(this);
+	for (TMap<int32, TSharedPtr<FAvaOutlinerView>>::TIterator Iter(MutableThis->OutlinerViews); Iter; ++Iter)
+	{
+		const TSharedPtr<FAvaOutlinerView>& OutlinerView = Iter.Value();
+		if (OutlinerView.IsValid())
+		{
+			InPredicate(OutlinerView);
+		}
+		else
+		{
+			Iter.RemoveCurrent();
+		}
+	}
+}
+
+void FAvaOutliner::EnqueueItemActions(TArray<TSharedPtr<IAvaOutlinerAction>>&& InItemActions) noexcept
+{
+	PendingActions.Append(MoveTemp(InItemActions));
+}
+
+int32 FAvaOutliner::GetPendingItemActionCount() const
+{
+	return PendingActions.Num();
+}
+
+bool FAvaOutliner::NeedsRefresh() const
+{
+	//Return false if we're already refreshing
+	if (bRefreshing)
+	{
+		return false;
+	}
+
+	if (bRefreshRequested || PendingActions.Num() > 0)
+	{
+		return true;
+	}
+
+	return false;
+}
+
 TOptional<FAvaOutlinerColorPair> FAvaOutliner::FindItemColor(const FAvaOutlinerItemPtr& InItem, bool bRecurseParent) const
 {
 	if (InItem.IsValid())
@@ -753,44 +1016,4 @@
 	return LowestCommonAncestor;
 }
 
-bool FAvaOutliner::CompareOutlinerItemOrder(const FAvaOutlinerItemPtr& A, const FAvaOutlinerItemPtr& B)
-{
-	if (!A.IsValid())
-	{
-		return false;
-	}
-	if (!B.IsValid())
-	{
-		return false;
-	}
-	if (const FAvaOutlinerItemPtr LowestCommonAncestor = FindLowestCommonAncestor({A, B}))
-	{
-		const TArray<FAvaOutlinerItemPtr> PathToA = LowestCommonAncestor->FindPath({A});
-		const TArray<FAvaOutlinerItemPtr> PathToB = LowestCommonAncestor->FindPath({B});
-
-		int32 Index = 0;
-		
-		int32 PathAIndex = -1;
-		int32 PathBIndex = -1;
-
-		while (PathAIndex == PathBIndex)
-		{
-			if (!PathToA.IsValidIndex(Index))
-			{
-				return true;
-			}
-			if (!PathToB.IsValidIndex(Index))
-			{
-				return false;
-			}
-
-			PathAIndex = LowestCommonAncestor->GetChildIndex(PathToA[Index]);
-			PathBIndex = LowestCommonAncestor->GetChildIndex(PathToB[Index]);
-			++Index;
-		}
-		return PathAIndex < PathBIndex;
-	}
-	return false;
-}
-
 FAvaSceneItem FAvaOutliner::MakeSceneItemFromOutlinerItem(const FAvaOutlinerItemPtr& InItem)
@@ -820,7 +1043,7 @@
 {
 	OutOutlinerItems.Sort([bInReverseOrder](const FAvaOutlinerItemPtr& ItemA, const FAvaOutlinerItemPtr& ItemB)
 	{
-		return FAvaOutliner::CompareOutlinerItemOrder(ItemA, ItemB) != bInReverseOrder;
+		return UE::AvaOutliner::CompareOutlinerItemOrder(ItemA, ItemB) != bInReverseOrder;
 	});
 }
 
@@ -850,79 +1073,6 @@
 	});
 }
 
-void FAvaOutliner::GroupSelection(AActor* InGroupingActor, const TOptional<FAttachmentTransformRules>& InTransformRules)
-{
-	if (!InGroupingActor)
-	{
-		return;
-	}
-
-	FAvaOutlinerItemPtr GroupingItem = FindItem(InGroupingActor);
-	bool bIsGroupingItemNew = false;	
-	if (!GroupingItem.IsValid())
-	{
-		GroupingItem = FindOrAdd<FAvaOutlinerActor>(InGroupingActor);
-		bIsGroupingItemNew = true;
-	}
-	GroupingItem->SetFlags(EAvaOutlinerItemFlags::Expanded);
-
-	TArray<FAvaOutlinerItemPtr> SelectedItems = GetSelectedItems();
-	NormalizeItems(SelectedItems);
-
-	// Enqueue the Grouping Item if New
-	if (bIsGroupingItemNew)
-	{
-		FAvaOutlinerAddItemParams AddGroupingItemParams;
-		AddGroupingItemParams.Item  = GroupingItem;
-		AddGroupingItemParams.Flags = EAvaOutlinerAddItemFlags::AddChildren;
-
-		AddGroupingItemParams.AttachmentTransformRules = InTransformRules;
-
-		if (const FAvaOutlinerItemPtr LowestCommonAncestor = FindLowestCommonAncestor(SelectedItems))
-		{
-			//Find the first path that leads to a Selected Item
-			TArray<FAvaOutlinerItemPtr> Descendants = LowestCommonAncestor->FindPath(SelectedItems);
-			if (Descendants.Num() > 0)
-			{
-				AddGroupingItemParams.RelativeItem     = Descendants[0];
-				AddGroupingItemParams.RelativeDropZone = EItemDropZone::BelowItem;
-			}
-			else
-			{
-				AddGroupingItemParams.RelativeItem     = LowestCommonAncestor;
-				AddGroupingItemParams.RelativeDropZone = EItemDropZone::OntoItem;
-			}
-		}
-
-		EnqueueItemAction<FAvaOutlinerAddItem>(AddGroupingItemParams);
-	}	
-	
-	//If there are no Selected Items, stop here
-	if (SelectedItems.IsEmpty())
-	{
-		return;
-	}
-
-	// Sort Items in reverse order of hierarchy as Add Item will add them at index 0 by default
-	SortItems(SelectedItems, true);
-		
-	FAvaOutlinerAddItemParams AddChildItemParams;
-	AddChildItemParams.RelativeItem     = GroupingItem;
-	AddChildItemParams.RelativeDropZone = EItemDropZone::OntoItem;
-	AddChildItemParams.SelectionFlags   = EAvaOutlinerItemSelectionFlags::AppendToCurrentSelection | EAvaOutlinerItemSelectionFlags::ScrollIntoView;
-
-	AddChildItemParams.AttachmentTransformRules = InTransformRules;
-
-	for (const FAvaOutlinerItemPtr& Item : SelectedItems)
-	{
-		if (Item.IsValid())
-		{
-			AddChildItemParams.Item = Item;
-			EnqueueItemAction<FAvaOutlinerAddItem>(AddChildItemParams);
-		}
-	}
-}
-
 FEditorModeTools* FAvaOutliner::GetModeTools() const
 {
 	return OutlinerProvider.GetOutlinerModeTools();
@@ -979,6 +1129,16 @@
 	return OutlinerProvider.GetOutlinerWorld();
 }
 
+const FAvaOutlinerItemProxyRegistry& FAvaOutliner::GetItemProxyRegistry() const
+{
+	return ItemProxyRegistry;
+}
+
+FAvaOutlinerItemProxyRegistry& FAvaOutliner::GetItemProxyRegistry()
+{
+	return ItemProxyRegistry;
+}
+
 TArray<TWeakObjectPtr<AActor>> FAvaOutliner::GetActorSceneOutlinerChildren(AActor* InParentActor) const
 {
 	if (const TArray<TWeakObjectPtr<AActor>>* const FoundChildren = SceneOutlinerParentMap.Find(InParentActor))
@@ -1007,169 +1167,6 @@
 	}
 }
 
-void FAvaOutliner::OnActorsCopied(FString& InOutCopiedData, TConstArrayView<AActor*> InCopiedActors)
-{
-	FAvaOutlinerExporter Exporter(SharedThis(this));
-	Exporter.ExportText(InOutCopiedData, InCopiedActors);
-}
-
-void FAvaOutliner::OnActorsPasted(FStringView InPastedData, const TMap<FName, AActor*>& InPastedActors)
-{
-	FAvaOutlinerImporter Importer(SharedThis(this));
-	Importer.ImportText(InPastedData, InPastedActors);
-}
-
-void FAvaOutliner::OnActorsDuplicated(const TMap<AActor*, AActor*>& InDuplicateActorMap
-	, FAvaOutlinerItemPtr InRelativeItem
-	, TOptional<EItemDropZone> InRelativeDropZone)
-{
-	if (InDuplicateActorMap.IsEmpty())
-	{
-		return;
-	}
-
-	// Bail if we're currently Ignoring Duplication Notifies
-	if (EnumHasAnyFlags(IgnoreNotifyFlags, EAvaOutlinerIgnoreNotifyFlags::Duplication))
-	{
-		return;
-	}
-
-	// Map of the Duplicate Actor to the Template Actor's Item
-	TMap<AActor*, FAvaOutlinerItemPtr> DuplicateItemMap;
-	
-	// Try finding the Template Item with the Lowest Index (i.e. Highest in the Tree) to use as Placeholder Relative Item
-	if (!InRelativeItem.IsValid())
-	{
-		DuplicateItemMap.Reserve(InDuplicateActorMap.Num());
-		for (const TPair<AActor*, AActor*>& Pair : InDuplicateActorMap)
-		{
-			AActor* const TemplateActor = Pair.Value;
-			if (IsValid(TemplateActor))
-			{
-				if (FAvaOutlinerItemPtr FoundTemplateItem = FindItem(TemplateActor))
-				{
-					DuplicateItemMap.Emplace(Pair.Key, FoundTemplateItem);
-					if (!InRelativeItem.IsValid() || CompareOutlinerItemOrder(FoundTemplateItem, InRelativeItem))
-					{
-						InRelativeItem = FoundTemplateItem;
-					}
-				}
-			}
-		}
-	}
-
-	const EItemDropZone RelativeDropZone = InRelativeDropZone.Get(EItemDropZone::AboveItem);
-
-	// Newly Duplicated Items don't have a Parent Item set yet, so it's not going to call RemoveFromParent, which handles Detachment
-	bool bShouldDetachActors = false;
-	if (InRelativeItem.IsValid())
-	{
-		const bool bRelativeItemIsRoot        = InRelativeItem == RootItem;
-		const bool bRelativeItemIsChildOfRoot = InRelativeItem->GetParent() == RootItem;
-		
-		bShouldDetachActors = bRelativeItemIsRoot || (bRelativeItemIsChildOfRoot && RelativeDropZone != EItemDropZone::OntoItem);
-	}
-
-	TArray<AActor*> DuplicateActors;
-	InDuplicateActorMap.GetKeys(DuplicateActors);
-
-	// Sort based on the ordering of the template actor items in outliner
-	DuplicateActors.Sort([&InDuplicateActorMap, this](const AActor& A, const AActor& B)
-	{
-		AActor* const * const TemplateActorA = InDuplicateActorMap.Find(&A);
-		AActor* const * const TemplateActorB = InDuplicateActorMap.Find(&B);
-
-		const FAvaOutlinerItemPtr TemplateItemA = FindItem(*TemplateActorA);
-		const FAvaOutlinerItemPtr TemplateItemB = FindItem(*TemplateActorB);
-
-		return CompareOutlinerItemOrder(TemplateItemA, TemplateItemB);
-	});
-
-	/*
-	 * REVERSE the ordering of Duplicate Actors
-	 * The order needs to be reversed when items are not added above the target item (e.g. RelativeDropZone is either below or onto)
-	 * Some context:
-	 * ABOVE- items are added just above the target, so the first one will always be topmost. Order is preserved
-	 * BELOW- items are added just below the target, so the last one added will be the topmost. Order is in reverse.
-	 * ONTO - items are added at index 0, so last one added will be topmost. Order is in reverse
-	 */
-	if (RelativeDropZone != EItemDropZone::AboveItem)
-	{
-		Algo::Reverse(DuplicateActors);
-	}
-	
-	TArray<TSharedPtr<IAvaOutlinerAction>> ItemActions;
-	ItemActions.Reserve(DuplicateActors.Num());
-	
-	FAvaOutlinerAddItemParams Params;
-	Params.RelativeDropZone = RelativeDropZone;
-	Params.Flags            = EAvaOutlinerAddItemFlags::AddChildren;
-
-	// Finally, Enqueue the Add Item Action to these new Actors
-	for (AActor* const DuplicateActor : DuplicateActors)
-	{
-		Params.RelativeItem.Reset();
-		
-		//We don't need to worry about Child Items since the Parent Item will do all the Attachment Work
-		const AActor* const ParentActor = DuplicateActor->GetAttachParentActor();
-
-		const bool bIsParentDuplicated = InDuplicateActorMap.Contains(ParentActor);
-		if (!bIsParentDuplicated && bShouldDetachActors)
-		{
-			DuplicateActor->DetachFromActor(FDetachmentTransformRules::KeepWorldTransform);
-		}
-
-		FAvaOutlinerItemPtr Item = FindOrAdd<FAvaOutlinerActor>(DuplicateActor);
-		Params.Item = Item;
-		
-		if (FAvaOutlinerItemPtr* const FoundTemplateItem = DuplicateItemMap.Find(DuplicateActor))
-		{
-			FAvaOutlinerItemPtr TemplateItem = *FoundTemplateItem;
-
-			if (!bIsParentDuplicated)
-			{
-				Params.RelativeItem = TemplateItem;
-			}
-			
-			//Sanitize the Flags by Removing the Temp Ones
-			EAvaOutlinerItemFlags TemplateItemFlags = TemplateItem->GetFlags();
-			EnumRemoveFlags(TemplateItemFlags, EAvaOutlinerItemFlags::PendingRemoval | EAvaOutlinerItemFlags::IgnorePendingKill);
-			
-			Item->SetFlags(TemplateItemFlags);
-
-			//Copy the template flags of each outliner view into the new item
-			ForEachOutlinerView([TemplateItem, Item](const TSharedPtr<FAvaOutlinerView>& InOutlinerView)
-			{
-				InOutlinerView->SetViewItemFlags(Item, InOutlinerView->GetViewItemFlags(TemplateItem));
-			});			
-		}
-		
-		//If a Relative Item hasn't been set yet (i.e. no Template Item found for the Item)
-		if (!Params.RelativeItem.IsValid() && !bIsParentDuplicated)
-		{
-			Params.RelativeItem = InRelativeItem;
-		}
-
-		TSharedRef<FAvaOutlinerAddItem> AddItemAction = NewItemAction<FAvaOutlinerAddItem>(Params);
-
-		// If parent is duplicated, make sure we add this item to index 0 so that children are added in reverse order
-		// this is because ParentItem->AddChild(...) adds the Child at Index 0, so it's correcting the ordering issue
-		if (bIsParentDuplicated)
-		{
-			ItemActions.EmplaceAt(0, AddItemAction);
-		}
-		else
-		{
-			ItemActions.Add(AddItemAction);
-		}		
-		
-		// Append after first selected item is added
-		Params.SelectionFlags |= EAvaOutlinerItemSelectionFlags::AppendToCurrentSelection;
-	}
-	
-	EnqueueItemActions(MoveTemp(ItemActions));
-}
-
 void FAvaOutliner::OnActorDestroyed(AActor* InActor)
 {
 	const FAvaOutlinerItemPtr ActorItem = FindItem(FAvaOutlinerItemId(InActor));
@@ -1223,11 +1220,6 @@
 	});
 }
 
-FTransform FAvaOutliner::GetActorDefaultSpawnTransform() const
-{
-	return OutlinerProvider.GetOutlinerDefaultActorSpawnTransform();
-}
-
 void FAvaOutliner::SetOutlinerModified()
 {
 	if (!bOutlinerDirty)
diff --git a/Engine/Plugins/Experimental/Avalanche/Source/AvalancheOutliner/Private/AvaOutliner.h b/Engine/Plugins/Experimental/Avalanche/Source/AvalancheOutliner/Private/AvaOutliner.h
new file mode 100644
--- /dev/null
+++ b/Engine/Plugins/Experimental/Avalanche/Source/AvalancheOutliner/Private/AvaOutliner.h
@@ -0,0 +1,308 @@
+// Copyright Epic Games, Inc. All Rights Reserved.
+
+#pragma once
+
+#include "AvaOutlinerDefines.h"
+#include "EditorUndoClient.h"
+#include "Engine/EngineTypes.h"
+#include "IAvaOutliner.h"
+#include "IAvaOutlinerModule.h"
+#include "Item/AvaOutlinerItemId.h"
+#include "Item/AvaOutlinerItemProxy.h"
+#include "Item/IAvaOutlinerItem.h"
+#include "ItemProxies/AvaOutlinerItemProxyRegistry.h"
+#include "ItemProxies/IAvaOutlinerItemProxyFactory.h"
+#include "TickableEditorObject.h"
+
+class FAvaEditorSelection;
+class FAvaOutlinerTreeRoot;
+class FAvaOutlinerView;
+class FEditorModeTools;
+class FTransaction;
+class FUICommandList;
+class IAvaOutlinerAction;
+class IAvaOutlinerProvider;
+class IAvaOutlinerView;
+class UAvaOutlinerSubsystem;
+enum class EItemDropZone;
+struct FAttachmentTransformRules;
+struct FAvaOutlinerSaveState;
+struct FAvaSceneItem;
+
+class FAvaOutliner
+	: public IAvaOutliner
+	, public FTickableEditorObject
+	, public FEditorUndoClient
+{
+public:
+	FAvaOutliner(IAvaOutlinerProvider& InOutlinerProvider);
+
+	virtual ~FAvaOutliner() override;
+
+	/**
+	 * Gets the Outliner Subsystem of the World this Outliner is responsible for
+	 * NOTE: The OutlinerSubsystem's Outliner Instance should be the same as this one in the default implementation 
+	 * but can differ if there are multiple FAvaOutliner instances being used in custom implementations
+	 */
+	UAvaOutlinerSubsystem* GetOutlinerSubsystem() const;
+
+	/**
+	 * Determines whether the given actor can be presented in the Outliner, at all.
+	 * This is a permanent check unlike filters that are temporary.
+	 */
+	bool IsActorAllowedInOutliner(const AActor* InActor) const;
+
+	/**
+	 * Determines whether the given scene component can be presented in the Outliner, at all.
+	 * This is a permanent check unlike filters that are temporary.
+	 */
+	bool IsComponentAllowedInOutliner(const USceneComponent* InComponent) const;
+
+	bool CanProcessActorSpawn(AActor* InActor) const;
+
+	TSharedPtr<FUICommandList> GetBaseCommandList() const;
+
+	FOnOutlinerLoaded OnOutlinerLoaded;
+
+	/** Gathers the Type Names of all the Item Proxies that are registered both in the outliner proxy registry and the module's */
+	TArray<FName> GetRegisteredItemProxyTypeNames() const;
+
+	/** Gathers all previously existing and new Item Proxies for a given Item */
+	void GetItemProxiesForItem(const FAvaOutlinerItemPtr& InItem, TArray<TSharedPtr<FAvaOutlinerItemProxy>>& OutItemProxies);
+
+	/** Tries to find the Item Proxy Factory for the given Item Proxy Type Name */
+	IAvaOutlinerItemProxyFactory* GetItemProxyFactory(FName InItemProxyTypeName) const;
+
+	const TSharedRef<FAvaOutlinerSaveState>& GetSaveState() const;
+
+	/** Returns whether the Outliner is in Read-only mode */
+	bool IsOutlinerLocked() const;
+
+	void HandleUndoRedoTransaction(const FTransaction* Transaction, bool bIsUndo);
+
+	//~ Begin IAvaOutliner
+	virtual FOnOutlinerLoaded& GetOnOutlinerLoaded() override { return OnOutlinerLoaded; }
+	virtual IAvaOutlinerProvider& GetProvider() const override { return OutlinerProvider; }
+	virtual void SetBaseCommandList(const TSharedPtr<FUICommandList>& InBaseCommandList) override;
+	virtual void Serialize(FArchive& Ar) override;
+	virtual TSharedPtr<IAvaOutlinerView> RegisterOutlinerView(int32 InOutlinerViewId) override;
+	virtual TSharedPtr<IAvaOutlinerView> GetOutlinerView(int32 InOutlinerViewId) const override;
+	virtual void RegisterItem(const FAvaOutlinerItemPtr& InItem) override;
+	virtual void UnregisterItem(const FAvaOutlinerItemId& InItemId) override;
+	virtual void RequestRefresh() override;
+	virtual void Refresh() override;
+	virtual TSharedRef<FAvaOutlinerTreeRoot> GetTreeRoot() const override { return RootItem; }
+	virtual FAvaOutlinerItemPtr FindItem(const FAvaOutlinerItemId& InItemId) const override;
+	virtual void SetIgnoreNotify(EAvaOutlinerIgnoreNotifyFlags InFlag, bool bIgnore) override;
+	virtual void OnActorsCopied(FString& InOutCopiedData, TConstArrayView<AActor*> InCopiedActors) override;
+	virtual void OnActorsPasted(FStringView InPastedData, const TMap<FName, AActor*>& InPastedActors) override;
+	virtual void OnActorsDuplicated(const TMap<AActor*, AActor*>& InDuplicateActorMap, FAvaOutlinerItemPtr InRelativeItem = nullptr, TOptional<EItemDropZone> InRelativeDropZone = TOptional<EItemDropZone>()) override;
+	virtual void GroupSelection(AActor* InGroupingActor, const TOptional<FAttachmentTransformRules>& InTransformRules = TOptional<FAttachmentTransformRules>()) override;
+	virtual void OnObjectSelectionChanged(const FAvaEditorSelection& InEditorSelection) override;
+	virtual UWorld* GetWorld() const override;
+	virtual const FAvaOutlinerItemProxyRegistry& GetItemProxyRegistry() const override;
+	//~ End IAvaOutliner
+
+	FAvaOutlinerItemProxyRegistry& GetItemProxyRegistry();
+
+	//~ Begin FEditorUndoClient
+	virtual void PostUndo(bool bSuccess) override;
+	virtual void PostRedo(bool bSuccess) override;
+	//~ End FEditorUndoClient
+
+	//~ Begin FTickableObjectBase
+	virtual TStatId GetStatId() const override;
+	virtual void Tick(float InDeltaTime) override;
+	//~ End FTickableObjectBase
+
+	/**
+	 * Creates a new set of items in the outliner based on the given template items
+	 * @param InItems the template items to use for duplication
+	 * @param InRelativeItem the existing outliner item that items should use as positional reference
+	 * @param InRelativeDropZone where to put the duplicate items with relation to InRelativeItem (above, below, onto)
+	 */
+	void DuplicateItems(TArray<FAvaOutlinerItemPtr> InItems
+		, FAvaOutlinerItemPtr InRelativeItem
+		, TOptional<EItemDropZone> InRelativeDropZone);
+	
+	/** Unregisters the Outliner View bound to the given id */
+	void UnregisterOutlinerView(int32 InOutlinerViewId);
+
+	/** Sets the given Outliner View Id as the most recent Outliner View */
+	void UpdateRecentOutlinerViews(int32 InOutlinerViewId);
+
+	/** Gets the outliner view that was most recently used (i.e. called FAvaOutliner::UpdateRecentOutlinerViews) */
+	TSharedPtr<FAvaOutlinerView> GetMostRecentOutlinerView() const;
+
+	/** Executes the given predicate for each Outliner View registered */
+	void ForEachOutlinerView(const TFunction<void(const TSharedPtr<FAvaOutlinerView>& InOutlinerView)>& InPredicate) const;
+
+	/**
+	 * Instantiates a new item action without adding it to the Pending Actions Queue.
+	 * This should only be used directly when planning to enqueue multiple actions.
+	 * @see FAvaOutliner::EnqueueItemActions
+	 */
+	template<typename InItemActionType, typename = typename TEnableIf<TIsDerivedFrom<InItemActionType, IAvaOutlinerAction>::IsDerived>::Type, typename ...InArgTypes>
+	TSharedRef<InItemActionType> NewItemAction(InArgTypes&&... InArgs)
+	{
+		return MakeShared<InItemActionType>(Forward<InArgTypes>(InArgs)...);
+	}
+
+	/**
+	 * Instantiates a single new item action and immediately adds it to the Pending Actions Queue.
+	 * Ideal for when dealing with a single action.
+	 * For multiple actions use FAvaOutliner::EnqueueItemActions.
+	 */
+	template<typename InItemActionType, typename = typename TEnableIf<TIsDerivedFrom<InItemActionType, IAvaOutlinerAction>::IsDerived>::Type, typename ...InArgTypes>
+	void EnqueueItemAction(InArgTypes&&... InArgs)
+	{
+		EnqueueItemActions({ NewItemAction<InItemActionType>(Forward<InArgTypes>(InArgs)...) });
+	}
+
+	/** Adds the given actions to the Pending Action Queue */
+	void EnqueueItemActions(TArray<TSharedPtr<IAvaOutlinerAction>>&& InItemActions) noexcept;
+
+	/** Returns the number of actions that been added to the queue so far before triggering a refresh */
+	int32 GetPendingItemActionCount() const;
+
+	/** Returns whether the Outliner is currently in need of a Refresh */
+	bool NeedsRefresh() const;
+
+	/**
+	 * Gets the color pair (color name, linear color) related to the Item
+	 * @param InItem the item to query
+	 * @param bRecurseParent whether to get the color of the parent (recursively) if the given item does not have a color by itself
+	 * @returns the matching color pair or unset if item is invalid or no color could be found.
+	 */
+	TOptional<FAvaOutlinerColorPair> FindItemColor(const FAvaOutlinerItemPtr& InItem, bool bRecurseParent = true) const;
+
+	/** Pairs the Item with the given color name, overriding the inherited color if different */
+	void SetItemColor(const FAvaOutlinerItemPtr& InItem, const FName& InColorName);
+
+	/** Removes the Color pairing of the given Item (can still have an inherited color though) */
+	void RemoveItemColor(const FAvaOutlinerItemPtr& InItem);
+
+	/** Returns the color map from the Outliner Settings */
+	const TMap<FName, FLinearColor>& GetColorMap() const;
+
+	/**
+	 * Replaces the Item's Id in the Item Map. This can be due to an object item changing it's object
+	 * (e.g. a bp component getting destroyed and recreated, the item should be the same but the underlying component will not be)
+	 */
+	void NotifyItemIdChanged(const FAvaOutlinerItemId& OldId, const FAvaOutlinerItemPtr& InItem);
+
+	/** Returns the currently selected items in the most recent outliner view (since this list can vary between outliner views) */
+	TArray<FAvaOutlinerItemPtr> GetSelectedItems() const;
+
+	/** Returns the number of currently selected items in the most recent outliner view */
+	int32 GetSelectedItemCount() const;
+
+	/**
+	 * Selects the given Items on all Outliner Views
+	 * @param InItems the items to select
+	 * @param InFlags how the items should be selected (appended, notify of selections, etc)
+	 */
+	void SelectItems(const TArray<FAvaOutlinerItemPtr>& InItems, EAvaOutlinerItemSelectionFlags InFlags = EAvaOutlinerItemSelectionFlags::SignalSelectionChange) const;
+
+	/**
+	 * Clears the Item Selection from all Outliner Views
+	 * @param bSignalSelectionChange whether to notify the change in selection
+	 */
+	void ClearItemSelection(bool bSignalSelectionChange) const;
+
+	/** Gets the closest item to all the given items while also being their common ancestor */
+	static FAvaOutlinerItemPtr FindLowestCommonAncestor(const TArray<FAvaOutlinerItemPtr>& Items);
+
+	/** Converts the given Outliner Item to a Scene Item that can be serialized in the Scene Tree */
+	static FAvaSceneItem MakeSceneItemFromOutlinerItem(const FAvaOutlinerItemPtr& InItem);
+
+	/**
+	 * Helper function to sort the given array of items based on their ordering in the Outliner
+	 * @see FAvaOutliner::CompareOutlinerItemOrder
+	 */
+	static void SortItems(TArray<FAvaOutlinerItemPtr>& OutOutlinerItems, bool bInReverseOrder = false);
+
+	/** Normalizes the given Items by removing selected items that have their parent item also in the selection */
+	static void NormalizeItems(TArray<FAvaOutlinerItemPtr>& InOutItems);
+
+	/** Gets the Editor Mode Tools used to handle selections */
+	FEditorModeTools* GetModeTools() const;
+
+	/** Have the given Selected Items sync to the USelection Instances of Mode Tools */
+	void SyncModeToolsSelection(const TArray<FAvaOutlinerItemPtr>& InSelectedItems) const;
+
+	/** Gets all the Actors that have as their AActor::GetSceneOutlinerParent the given InParentActor */
+	TArray<TWeakObjectPtr<AActor>> GetActorSceneOutlinerChildren(AActor* InParentActor) const;
+
+	/** Tries to add the new Actor to the Outliner if not added already and if the Outliner allows the given actor to be added */
+	void OnActorSpawned(AActor* InActor);
+
+	/** Called when an Actor has been destroyed. This enqueues the Removal the Actor Item to the Pending Action Queue */
+	void OnActorDestroyed(AActor* InActor);
+
+	/** Called when an Actor's attachment has changed. This triggers a refresh */
+	void OnActorAttachmentChanged(AActor* InActor, const AActor* InParent, bool bAttach);
+
+	/** Called the engine replaces an object. A common example is when a BP Component is destroyed, and replaced */
+	void OnObjectsReplaced(const TMap<UObject*, UObject*>& InReplacementMap);
+
+	/** Marks the Outliner dirty. This triggers IAvaOutlinerProvider::OnOutlinerModified on next tick */
+	void SetOutlinerModified();
+
+private:
+	void AddItem(const FAvaOutlinerItemPtr& InItem);
+	void RemoveItem(const FAvaOutlinerItemId& InItemId);
+	void ForEachItem(TFunctionRef<void(const FAvaOutlinerItemPtr&)> InFunc);
+
+	/** Interface providing the Outliner with logic such as Actor Duplication, Mode Tools, and extensibility options */
+	IAvaOutlinerProvider& OutlinerProvider;
+
+	/** the base command list gotten from the outliner provider to append our command list to */
+	TWeakPtr<FUICommandList> BaseCommandListWeak;
+
+	/** the map of the registered items */
+	TMap<FAvaOutlinerItemId, FAvaOutlinerItemPtr> ItemMap;
+
+	TMap<FAvaOutlinerItemId, FAvaOutlinerItemPtr> ItemsPendingAdd;
+
+	TSet<FAvaOutlinerItemId> ItemsPendingRemove;
+
+	/** Map of an Actor to its Child Actors defined as such via AActor::GetSceneOutlinerParent */
+	TMap<TWeakObjectPtr<AActor>, TArray<TWeakObjectPtr<AActor>>> SceneOutlinerParentMap;
+
+	/** The root of all the items in the outliner */
+	TSharedRef<FAvaOutlinerTreeRoot> RootItem;
+
+	/** Outliner's Item Proxy Factory Registry Instance. This takes precedence over the Module's Factory Registry */
+	FAvaOutlinerItemProxyRegistry ItemProxyRegistry;
+	
+	/** the current pending actions before refresh is called */
+	TArray<TSharedPtr<IAvaOutlinerAction>> PendingActions;
+
+	/** the list of objects pending selection processing. filled in when getting notifies from the USelection instances */
+	TSharedPtr<TArray<TWeakObjectPtr<UObject>>> ObjectsLastSelected;
+
+	/** the save state object to help serialize the outliner state */
+	TSharedRef<FAvaOutlinerSaveState> SaveState;
+
+	/** the map of registered outliner views */
+	TMap<int32, TSharedPtr<FAvaOutlinerView>> OutlinerViews;
+
+	/** List of Outliner View Ids in order from least recent to most recent (i.e. Index 0 is least recent) */
+	TArray<int32> RecentOutlinerViews;
+
+	/** the current events to ignore and not handle automatically */
+	EAvaOutlinerIgnoreNotifyFlags IgnoreNotifyFlags = EAvaOutlinerIgnoreNotifyFlags::None;
+
+	/** Flag indicating whether the Outliner has been changed this tick and should call IAvaOutlinerProvider::OnOutlinerModified next tick */
+	bool bOutlinerDirty = false;
+	
+	/** Flag indicating Refreshing is taking place */
+	bool bRefreshing = false;
+
+	/** Flag indicating that a refresh must take place next tick */
+	bool bRefreshRequested = false;
+
+	/** Flag indicating that the Item Map is iterating */
+	bool bIteratingItemMap = false;
+};
diff --git a/Engine/Plugins/Experimental/Avalanche/Source/AvalancheOutliner/Private/AvaOutlinerModule.cpp b/Engine/Plugins/Experimental/Avalanche/Source/AvalancheOutliner/Private/AvaOutlinerModule.cpp
--- a/Engine/Plugins/Experimental/Avalanche/Source/AvalancheOutliner/Private/AvaOutlinerModule.cpp
+++ b/Engine/Plugins/Experimental/Avalanche/Source/AvalancheOutliner/Private/AvaOutlinerModule.cpp
@@ -36,16 +36,6 @@ void FAvaOutlinerModule::ShutdownModule()
 	FAvaOutlinerCommands::Unregister();
 }
 
-void FAvaOutlinerModule::UnregisterFilterExpressionFactory(const FName& InFilterIdentifier)
-{
-	FilterExpressionFactories.Remove(InFilterIdentifier);
-}
-
-void FAvaOutlinerModule::UnregisterFilterSuggestionFactory(const FName& InFilterIdentifier)
-{
-	FilterSuggestionsFactories.Remove(InFilterIdentifier);
-}
-
 bool FAvaOutlinerModule::CanFilterSupportComparisonOperation(const FName& InFilterKey, const ETextFilterComparisonOperation InOperation) const
 {
 	if (const TSharedPtr<IAvaFilterExpressionFactory>* FilterExpressionFactory = FilterExpressionFactories.Find(InFilterKey))
diff --git a/Engine/Plugins/Experimental/Avalanche/Source/AvalancheOutliner/Private/AvaOutlinerModule.h b/Engine/Plugins/Experimental/Avalanche/Source/AvalancheOutliner/Private/AvaOutlinerModule.h
--- a/Engine/Plugins/Experimental/Avalanche/Source/AvalancheOutliner/Private/AvaOutlinerModule.h
+++ b/Engine/Plugins/Experimental/Avalanche/Source/AvalancheOutliner/Private/AvaOutlinerModule.h
@@ -2,6 +2,8 @@
 
 #pragma once
 
+#include "Filters/IAvaFilterExpressionFactory.h"
+#include "Filters/IAvaFilterSuggestionFactory.h"
 #include "IAvaOutlinerModule.h"
 #include "ItemProxies/AvaOutlinerItemProxyRegistry.h"
 
@@ -10,57 +12,98 @@ class IAvaFilterExpressionFactory;
 class FAvaOutlinerModule : public IAvaOutlinerModule
 {
 public:
+	static FAvaOutlinerModule& Get()
+	{
+		return FModuleManager::LoadModuleChecked<FAvaOutlinerModule>(TEXT("AvalancheOutliner"));
+	};
+
 	//~ Begin IModuleInterface
 	virtual void StartupModule() override;
 	virtual void ShutdownModule() override;
 	//~ End IModuleInterface
 
 	//~ Begin IAvaOutlinerModule
-	virtual FOnInitOutlinerToolMenu& GetOnInitOutlinerToolBar() override { return OnInitOutlinerToolBar; }
-	virtual FOnInitOutlinerToolMenu& GetOnInitOutlinerItemContextMenu() override { return OnInitOutlinerItemContextMenu; }
-	virtual FOnExtendOutlinerToolMenu& GetOnExtendOutlinerToolBar() override  { return OnExtendOutlinerToolBar; }
-	virtual FOnExtendOutlinerToolMenu& GetOnExtendOutlinerItemContextMenu() override   { return OnExtendOutlinerItemContextMenu; }
-	virtual FOnExtendOutlinerColumns& GetOnExtendOutlinerColumns() override { return OnExtendOutlinerColumns; }
+	virtual FOnExtendOutlinerToolMenu& GetOnExtendOutlinerItemContextMenu() override { return OnExtendOutlinerItemContextMenu; }
 	virtual FAvaOutlinerItemProxyRegistry& GetItemProxyRegistry() override { return ItemProxyRegistry; }
 	virtual FOnExtendItemProxiesForItem& GetOnExtendItemProxiesForItem() override { return OnGetItemProxiesForItem; }
-	virtual void UnregisterFilterExpressionFactory(const FName& InFilterIdentifier) override;
-	virtual void UnregisterFilterSuggestionFactory(const FName& InFilterIdentifier) override;
-	virtual bool CanFilterSupportComparisonOperation(const FName& InFilterKey, const ETextFilterComparisonOperation InOperation) const override;
-	virtual bool FilterExpression(const FName& InFilterKey, const IAvaOutlinerItem& InItem, const FAvaTextFilterArgs& InArgs) const override;
-	virtual TArray<TSharedPtr<IAvaFilterSuggestionFactory>> GetSuggestions(const EAvaFilterSuggestionType& InSuggestionType) const override;
-	virtual FSlateIcon FindOverrideIcon(TSharedPtr<const FAvaOutlinerItem> InItem) const override;
 	//~ End IAvaOutlinerModule
 
-protected:
-	/** Delegate to extend the Tool Menu Context of the Outliner ToolBar */
-	FOnInitOutlinerToolMenu OnInitOutlinerToolBar;
-
-	/** Delegate to extend the Tool Menu Context of the Outliner Item Context Menu */
-	FOnInitOutlinerToolMenu OnInitOutlinerItemContextMenu;
+	/**
+ 	 * Create and Register the given filter expression factory
+ 	 * @tparam InFilterExpressionFactoryType The type of filter expression factory to create and register
+ 	 * @tparam InArgsType Constructor args types if needed
+ 	 * @param InArgs Additional args passed to the constructor of the filter expression factory if needed
+ 	 */
+	template<typename InFilterExpressionFactoryType, typename... InArgsType
+			UE_REQUIRES(TIsDerivedFrom<InFilterExpressionFactoryType, IAvaFilterExpressionFactory>::Value)>
+	void RegisterFilterExpressionFactory(InArgsType&&... InArgs)
+	{
+		RegisterFilterExpressionFactory_Internal(IAvaFilterExpressionFactory::MakeInstance<InFilterExpressionFactoryType>(Forward<InArgsType>(InArgs)...));
+	}
 
-	/** Delegate to extend the Outliner ToolBar */
-	FOnExtendOutlinerToolMenu OnExtendOutlinerToolBar;
+	/**
+  	 * Create and Register the given filter suggestion factory
+  	 * @tparam InSuggestionFactoryType The type of filter suggestion factory to create and register
+  	 * @tparam InArgsType Constructor args types if needed
+  	 * @param InArgs Additional args passed to the constructor of the filter suggestion factory if needed
+  	 */
+	template<typename InSuggestionFactoryType, typename... InArgsType
+			UE_REQUIRES(TIsDerivedFrom<InSuggestionFactoryType, IAvaFilterSuggestionFactory>::Value)>
+	void RegisterFilterSuggestionFactory(InArgsType&&... InArgs)
+	{
+		RegisterFilterSuggestionFactory_Internal(IAvaFilterSuggestionFactory::MakeInstance<InSuggestionFactoryType>(Forward<InArgsType>(InArgs)...));
+	}
 
-	/** Delegate to extend the Item Context Menu */
-	FOnExtendOutlinerToolMenu OnExtendOutlinerItemContextMenu;
+	/**
+	 * Check if current filter expression factory support the comparison operation
+	 * @param InFilterKey Filter key to get the filter expression factory needed
+	 * @param InOperation Operation to check if supported
+	 * @return True if operation is supported, False otherwise
+	 */
+	bool CanFilterSupportComparisonOperation(const FName& InFilterKey, const ETextFilterComparisonOperation InOperation) const;
 
-	/** Delegates to add in additional Outliner Columns */
-	FOnExtendOutlinerColumns OnExtendOutlinerColumns;
+	/**
+	 * Evaluate the expression and return the result 
+	 * @param InFilterKey Filter Key to get the Factory
+	 * @param InItem Item that is currently checked
+	 * @param InArgs Args to evaluate the expression see FAvaTextFilterArgs for more information
+	 * @return True if the expression evaluated to True, False otherwise
+	 */
+	bool FilterExpression(const FName& InFilterKey, const IAvaOutlinerItem& InItem, const FAvaTextFilterArgs& InArgs) const;
 
-	/** Delegate to add in Item Proxies for a given Item / Item Type */
-	FOnExtendItemProxiesForItem OnGetItemProxiesForItem;
+	/**
+	 * Get all suggestions with the given type (Generic/ItemBased/All)
+	 * @param InSuggestionType Type of suggestion to get, see EFilterSuggestionType for more information
+	 * @return An Array containing all suggestions of the given type
+	 */
+	TArray<TSharedPtr<IAvaFilterSuggestionFactory>> GetSuggestions(const EAvaFilterSuggestionType& InSuggestionType) const;
 
-	/** Module Instance of the Item Registry */
-	FAvaOutlinerItemProxyRegistry ItemProxyRegistry;
+	/**
+	 * Get the custom icon for the item if a customization for it exist
+	 * @param InItem The item to search the icon for
+	 * @return The custom icon for the item if a customization is found, FSlateIcon() otherwise
+	 */
+	FSlateIcon FindOverrideIcon(TSharedPtr<const FAvaOutlinerItem> InItem) const;
 
+protected:
 	//~ Begin IAvaOutlinerModule
 	virtual void UnregisterOverriddenIcon_Internal(const FAvaTypeId& InItemTypeId, const FName& InSpecializationIdentifier) override;
 	virtual void RegisterOverridenIcon_Internal(const FAvaTypeId& InItemTypeId, const TSharedRef<IAvaOutlinerIconCustomization>& InAvaOutlinerIconCustomization) override;
-	virtual void RegisterFilterExpressionFactory_Internal(const TSharedRef<IAvaFilterExpressionFactory>& InFilterExpressionFactory) override;
-	virtual void RegisterFilterSuggestionFactory_Internal(const TSharedRef<IAvaFilterSuggestionFactory>& InFilterSuggestionFactory) override;
 	//~ End IAvaOutlinerModule
 
 private:
+	/**
+	 * Adds the given factory to the Map containing all of them if not already added.
+	 * @param InFilterExpressionFactory Filter expression factory to register.
+	 */
+	void RegisterFilterExpressionFactory_Internal(const TSharedRef<IAvaFilterExpressionFactory>& InFilterExpressionFactory);
+
+	/**
+	 * Adds the given factory to the Map containing all of them if not already added.
+	 * @param InFilterSuggestionFactory Filter suggestion factory to register.
+	 */
+	void RegisterFilterSuggestionFactory_Internal(const TSharedRef<IAvaFilterSuggestionFactory>& InFilterSuggestionFactory) ;
+
 	/**
 	 * Get the customization for the given item if any
 	 * @param InItem Item to search the customization for
@@ -78,7 +121,15 @@ private:
 	  */
 	void RegisterFilterSuggestionFactories();
 
-private:
+	/** Delegate to extend the Item Context Menu */
+	FOnExtendOutlinerToolMenu OnExtendOutlinerItemContextMenu;
+
+	/** Delegate to add in Item Proxies for a given Item / Item Type */
+	FOnExtendItemProxiesForItem OnGetItemProxiesForItem;
+
+	/** Module Instance of the Item Registry */
+	FAvaOutlinerItemProxyRegistry ItemProxyRegistry;
+
 	/** Holds all the FilterExpressionFactory */
 	TMap<FName, TSharedPtr<IAvaFilterExpressionFactory>> FilterExpressionFactories;
 
diff --git a/Engine/Plugins/Experimental/Avalanche/Source/AvalancheOutliner/Private/AvaOutlinerSettings.h b/Engine/Plugins/Experimental/Avalanche/Source/AvalancheOutliner/Private/AvaOutlinerSettings.h
new file mode 100644
--- /dev/null
+++ b/Engine/Plugins/Experimental/Avalanche/Source/AvalancheOutliner/Private/AvaOutlinerSettings.h
@@ -0,0 +1,79 @@
+// Copyright Epic Games, Inc. All Rights Reserved.
+
+#pragma once
+
+#include "Engine/DeveloperSettings.h"
+#include "Filters/AvaOutlinerItemTypeFilter.h"
+#include "AvaOutlinerSettings.generated.h"
+
+UCLASS(config=EditorPerProjectUserSettings, meta=(DisplayName="Outliner"))
+class UAvaOutlinerSettings : public UDeveloperSettings
+{
+	GENERATED_BODY()
+
+public:
+	UAvaOutlinerSettings();
+
+	static UAvaOutlinerSettings* Get();
+
+	const TMap<FName, FLinearColor>& GetColorMap() const { return ItemColorMap; }
+
+	static FName GetCustomItemTypeFiltersName();
+
+	const TMap<FName, FAvaOutlinerItemTypeFilterData>& GetCustomItemTypeFilters() const { return CustomItemTypeFilters; }
+
+	bool AddCustomItemTypeFilter(FName InKey, FAvaOutlinerItemTypeFilterData& InFilter);
+
+	bool ShouldUseMutedHierarchy() const { return bUseMutedHierarchy; }
+
+	bool ShouldAutoExpandToSelection() const { return bAutoExpandToSelection; }
+
+	bool ShouldAlwaysShowVisibilityState() const { return bAlwaysShowVisibilityState; }
+
+	bool ShouldAlwaysShowLockState() const { return bAlwaysShowLockState; }
+
+	EAvaOutlinerItemViewMode GetItemDefaultViewMode() const { return static_cast<EAvaOutlinerItemViewMode>(ItemDefaultViewMode); }
+
+	EAvaOutlinerItemViewMode GetItemProxyViewMode() const { return static_cast<EAvaOutlinerItemViewMode>(ItemProxyViewMode); }
+
+private:
+	virtual void PostInitProperties() override;
+
+	virtual void PostEditChangeProperty(FPropertyChangedEvent& PropertyChangedEvent) override;
+
+	/**
+	 * Hides the Alpha Channel for the Value Color Property in ItemColorMap
+	 * this has to be done manually as putting in meta=(HideAlphaChannel) in a map doesn't apply to the Value Property
+	 */
+	void HideItemColorMapAlphaChannel() const;
+
+	UPROPERTY(Config, EditAnywhere, Category = "Color")
+	TMap<FName, FLinearColor> ItemColorMap;
+
+	UPROPERTY(Config, EditAnywhere, Category = "Filters")
+	TMap<FName, FAvaOutlinerItemTypeFilterData> CustomItemTypeFilters;
+
+	/** Whether to show the visibility state always, rather than only showing when the item is hidden or hovered */
+	UPROPERTY(Config, EditAnywhere, Category = "Editor UX")
+	bool bAlwaysShowVisibilityState = false;
+
+	/** Whether to show the lock state always, rather than only showing when the item is locked or hovered */
+	UPROPERTY(Config, EditAnywhere, Category = "Editor UX")
+	bool bAlwaysShowLockState = false;
+
+	/** Whether to show the parent of the shown items, even if the parents are filtered out */
+	UPROPERTY(Config, EditAnywhere, Category = "Initial Values")
+	bool bUseMutedHierarchy = true;
+
+	/** Whether to auto expand the hierarchy to show the item when selected */
+	UPROPERTY(Config, EditAnywhere, Category = "Initial Values")
+	bool bAutoExpandToSelection = true;
+
+	/** The View Mode a Non-Actor / Non-Component Item supports by default */
+	UPROPERTY(Config, EditAnywhere, Category = "Initial Values", meta=(Bitmask, BitmaskEnum="/Script/AvalancheOutliner.EAvaOutlinerItemViewMode"))
+	int32 ItemDefaultViewMode;
+
+	/** The View Mode a Proxy Item supports by default */
+	UPROPERTY(Config, EditAnywhere, Category = "Initial Values", meta=(Bitmask, BitmaskEnum="/Script/AvalancheOutliner.EAvaOutlinerItemViewMode"))
+	int32 ItemProxyViewMode;
+};
diff --git a/Engine/Plugins/Experimental/Avalanche/Source/AvalancheOutliner/Private/AvaOutlinerSubsystem.cpp b/Engine/Plugins/Experimental/Avalanche/Source/AvalancheOutliner/Private/AvaOutlinerSubsystem.cpp
--- a/Engine/Plugins/Experimental/Avalanche/Source/AvalancheOutliner/Private/AvaOutlinerSubsystem.cpp
+++ b/Engine/Plugins/Experimental/Avalanche/Source/AvalancheOutliner/Private/AvaOutlinerSubsystem.cpp
@@ -8,10 +8,10 @@
 
 DEFINE_LOG_CATEGORY_STATIC(LogAvaOutlinerSubsystem, Log, All);
 
-TSharedRef<FAvaOutliner> UAvaOutlinerSubsystem::GetOrCreateOutliner(IAvaOutlinerProvider& InProvider, bool bInForceCreate)
+TSharedRef<IAvaOutliner> UAvaOutlinerSubsystem::GetOrCreateOutliner(IAvaOutlinerProvider& InProvider, bool bInForceCreate)
 {
-	TSharedPtr<FAvaOutliner> Outliner = OutlinerWeak.Pin();
-	if (!Outliner.IsValid() || &Outliner->GetProvider() != &InProvider || bInForceCreate)
+	TSharedPtr<IAvaOutliner> Outliner = OutlinerWeak.Pin();
+	if (!Outliner.IsValid() || &StaticCastSharedPtr<FAvaOutliner>(Outliner)->GetProvider() != &InProvider || bInForceCreate)
 	{
 		Outliner = MakeShared<FAvaOutliner>(InProvider);
 	}
diff --git a/Engine/Plugins/Experimental/Avalanche/Source/AvalancheOutliner/Private/AvaOutlinerView.cpp b/Engine/Plugins/Experimental/Avalanche/Source/AvalancheOutliner/Private/AvaOutlinerView.cpp
--- a/Engine/Plugins/Experimental/Avalanche/Source/AvalancheOutliner/Private/AvaOutlinerView.cpp
+++ b/Engine/Plugins/Experimental/Avalanche/Source/AvalancheOutliner/Private/AvaOutlinerView.cpp
@@ -24,13 +24,13 @@
 #include "Item/AvaOutlinerComponent.h"
 #include "Item/AvaOutlinerComponentProxy.h"
 #include "Item/AvaOutlinerTreeRoot.h"
+#include "Menu/AvaOutlinerViewToolbarContext.h"
 #include "Slate/SAvaOutliner.h"
 #include "Slate/SAvaOutlinerTreeView.h"
 #include "Stats/AvaOutlinerStats.h"
 #include "Styling/SlateTypes.h"
 #include "ToolMenu.h"
 #include "ToolMenuContext/AvaOutlinerItemsContext.h"
-#include "ToolMenuContext/AvaOutlinerToolBarContext.h"
 #include "ToolMenus.h"
 #include "Widgets/SNullWidget.h"
 
@@ -56,7 +56,7 @@ FAvaOutlinerView::~FAvaOutlinerView()
 	}
 }
 
-FName FAvaOutlinerView::GetOutlinerToolBarName()
+FName FAvaOutlinerView::GetOutlinerToolbarName()
 {
 	static const FName ToolBarName(TEXT("AvalancheOutliner.MainToolBar"));
 	return ToolBarName;
@@ -85,7 +85,7 @@ void FAvaOutlinerView::Init(const TSharedRef<FAvaOutliner>& InOutliner, bool bCr
 	if (bCreateOutlinerWidget)
 	{
 		OutlinerWidget = SNew(SAvaOutliner, SharedThis(this));
-		CreateToolBar();
+		CreateToolbar();
 	}
 
 	UpdateRecentOutlinerViews();
@@ -104,13 +104,6 @@ void FAvaOutlinerView::CreateColumns()
 	ColumnExtender.AddColumn<FAvaOutlinerItemsColumn>();
 	ColumnExtender.AddColumn<FAvaOutlinerTagColumn>();
 
-	// Option to Extend Columns via Provider and Module
-	if (const TSharedPtr<FAvaOutliner> Outliner = OutlinerWeak.Pin())
-	{
-		Outliner->GetProvider().ExtendOutlinerColumns(ColumnExtender);
-	}
-	IAvaOutlinerModule::Get().GetOnExtendOutlinerColumns().Broadcast(ColumnExtender);
-
 	const TArray<TSharedPtr<IAvaOutlinerColumn>>& FoundColumns = ColumnExtender.GetColumns();
 	
 	Columns.Empty(FoundColumns.Num());
@@ -364,33 +357,35 @@ TSharedRef<SWidget> FAvaOutlinerView::GetOutlinerWidget() const
 	return SNullWidget::NullWidget;
 }
 
-void FAvaOutlinerView::CreateToolBar()
+TSharedPtr<IAvaOutliner> FAvaOutlinerView::GetOwnerOutliner() const
+{
+	return GetOutliner();
+}
+
+void FAvaOutlinerView::CreateToolbar()
 {
 	UToolMenus* const ToolMenus = UToolMenus::Get();
 
-	const FName ToolBarName = FAvaOutlinerView::GetOutlinerToolBarName();
-	
-	if (!ToolMenus->IsMenuRegistered(ToolBarName))
+	const FName ToolbarName = FAvaOutlinerView::GetOutlinerToolbarName();
+
+	if (!ToolMenus->IsMenuRegistered(ToolbarName))
 	{
-		UToolMenu* const ToolBar = ToolMenus->RegisterMenu(ToolBarName, NAME_None, EMultiBoxType::SlimHorizontalToolBar);
+		UToolMenu* const ToolBar = ToolMenus->RegisterMenu(ToolbarName, NAME_None, EMultiBoxType::SlimHorizontalToolBar);
 		ToolBar->StyleName               = "StatusBarToolBar";
 		ToolBar->bToolBarForceSmallIcons = true;
 		ToolBar->bToolBarIsFocusable     = true;
 		ToolBar->AddDynamicSection("PopulateToolBar", FNewToolMenuDelegate::CreateStatic(&FAvaOutlinerView::PopulateToolBar));
 	}
-	
+
 	TSharedPtr<FExtender> Extender;
-	
-	UAvaOutlinerToolBarContext* const ContextObject = NewObject<UAvaOutlinerToolBarContext>();
+
+	UAvaOutlinerViewToolbarContext* const ContextObject = NewObject<UAvaOutlinerViewToolbarContext>();
 	ContextObject->OutlinerViewWeak = SharedThis(this);
-	
+
 	FToolMenuContext Context(GetBaseCommandList(), Extender, ContextObject);
-	
-	IAvaOutlinerModule::Get().GetOnInitOutlinerToolBar().Broadcast(Context);
-	
 	if (OutlinerWidget.IsValid())
 	{
-		OutlinerWidget->SetToolBarWidget(ToolMenus->GenerateWidget(ToolBarName, Context));
+		OutlinerWidget->SetToolBarWidget(ToolMenus->GenerateWidget(ToolbarName, Context));
 	}
 }
 
@@ -409,13 +404,10 @@ TSharedPtr<SWidget> FAvaOutlinerView::CreateItemContextMenu()
 	TSharedPtr<FExtender> Extender;
 
 	UAvaOutlinerItemsContext* const ContextObject = NewObject<UAvaOutlinerItemsContext>();
-	ContextObject->OutlinerViewWeak = SharedThis(this);
+	ContextObject->OutlinerWeak = GetOutliner();
 	ContextObject->ItemListWeak.Append(SelectedItems);
 
 	FToolMenuContext Context(GetViewCommandList(), Extender, ContextObject);
-
-	IAvaOutlinerModule::Get().GetOnInitOutlinerItemContextMenu().Broadcast(Context);
-
 	return ToolMenus->GenerateWidget(ItemContextMenuName, Context);
 }
 
@@ -1645,7 +1637,7 @@ void FAvaOutlinerView::PopulateItemContextMenu(UToolMenu* InToolMenu)
 	{
 		return;
 	}
-	
+
 	UAvaOutlinerItemsContext* const ItemsContext = InToolMenu->FindContext<UAvaOutlinerItemsContext>();
 	if (!ItemsContext)
 	{
@@ -1664,7 +1656,7 @@ void FAvaOutlinerView::PopulateItemContextMenu(UToolMenu* InToolMenu)
 		GenericActionsSection.AddMenuEntry(GenericCommands.Delete);
 		GenericActionsSection.AddMenuEntry(GenericCommands.Rename);
 	}
-	
+
 	// Outliner Commands
 	{
 		FToolMenuSection& OutlinerActionsSection = InToolMenu->AddSection(TEXT("OutlinerActions"), LOCTEXT("OutlinerActionsHeader", "Outliner Actions"));
@@ -1675,19 +1667,15 @@ void FAvaOutlinerView::PopulateItemContextMenu(UToolMenu* InToolMenu)
 		OutlinerActionsSection.AddMenuEntry(OutlinerCommands.SelectAllChildren);
 		OutlinerActionsSection.AddMenuEntry(OutlinerCommands.SelectImmediateChildren);
 	}
-	
-	const TSharedPtr<FAvaOutlinerView> OutlinerView = ItemsContext->GetOutlinerView();
-	
-	if (!OutlinerView.IsValid())
+
+	const TSharedPtr<IAvaOutliner> Outliner = ItemsContext->GetOutliner();
+	if (!Outliner.IsValid())
 	{
 		return;
 	}
-	
+
 	// Give Option to Quickly Extend without having to implement UToolMenus::Get()->ExtendMenu
-	if (TSharedPtr<FAvaOutliner> Outliner = OutlinerView->GetOutliner())
-	{
-		Outliner->GetProvider().ExtendOutlinerItemContextMenu(InToolMenu);
-	}
+	Outliner->GetProvider().ExtendOutlinerItemContextMenu(InToolMenu);
 	IAvaOutlinerModule::Get().GetOnExtendOutlinerItemContextMenu().Broadcast(InToolMenu);
 }
 
@@ -1697,30 +1685,30 @@ void FAvaOutlinerView::PopulateToolBar(UToolMenu* InToolMenu)
 	{
 		return;
 	}
-	
-	UAvaOutlinerToolBarContext* const ToolBarContext = InToolMenu->FindContext<UAvaOutlinerToolBarContext>();
+
+	UAvaOutlinerViewToolbarContext* const ToolBarContext = InToolMenu->FindContext<UAvaOutlinerViewToolbarContext>();
 	if (!ToolBarContext)
 	{
 		return;
 	}
-	
+
 	if (!ToolBarContext->GetOutlinerView().IsValid())
 	{
 		return;
 	}
-	
+
 	TSharedRef<FAvaOutlinerView> OutlinerView = ToolBarContext->GetOutlinerView().ToSharedRef();
-	
+
 	TSharedRef<SWidget> Widget = OutlinerView->GetOutlinerWidget();
 	if (Widget == SNullWidget::NullWidget)
 	{
 		return;
 	}
-	
+
 	TSharedRef<SAvaOutliner> OutlinerWidget = StaticCastSharedRef<SAvaOutliner>(Widget);
-	
+
 	FToolMenuSection& MainSection = InToolMenu->AddSection(TEXT("Main"));
-		
+
 	// View Options
 	const FToolMenuEntry ViewOptionsEntry = FToolMenuEntry::InitComboButton("ViewOptions", FUIAction()
 			, FOnGetContent::CreateSP(OutlinerView, &FAvaOutlinerView::CreateOutlinerViewOptionsMenu)
@@ -1728,7 +1716,7 @@ void FAvaOutlinerView::PopulateToolBar(UToolMenu* InToolMenu)
 			, LOCTEXT("ViewOptionsToolTip", "View Options")
 			, FSlateIcon(FAppStyle::Get().GetStyleSetName(), "Icons.Visibility"));
 	MainSection.AddEntry(ViewOptionsEntry);
-	
+
 	// Settings
 	const FToolMenuEntry SettingsEntry = FToolMenuEntry::InitComboButton("Settings", FUIAction()
 			, FOnGetContent::CreateSP(OutlinerView, &FAvaOutlinerView::CreateOutlinerSettingsMenu)
@@ -1736,13 +1724,12 @@ void FAvaOutlinerView::PopulateToolBar(UToolMenu* InToolMenu)
 			, LOCTEXT("SettingsToolTip", "Settings")
 			, FSlateIcon(FAppStyle::Get().GetStyleSetName(), "Icons.Settings"));
 	MainSection.AddEntry(SettingsEntry);
-	
+
 	// Give Option to Quickly Extend without having to implement UToolMenus::Get()->ExtendMenu
 	if (TSharedPtr<FAvaOutliner> Outliner = OutlinerView->GetOutliner())
 	{
 		Outliner->GetProvider().ExtendOutlinerToolBar(InToolMenu);
 	}
-	IAvaOutlinerModule::Get().GetOnExtendOutlinerToolBar().Broadcast(InToolMenu);
 }
 
 void FAvaOutlinerView::RefreshOutliner(bool bInImmediateRefresh)
diff --git a/Engine/Plugins/Experimental/Avalanche/Source/AvalancheOutliner/Private/AvaOutlinerView.h b/Engine/Plugins/Experimental/Avalanche/Source/AvalancheOutliner/Private/AvaOutlinerView.h
new file mode 100644
--- /dev/null
+++ b/Engine/Plugins/Experimental/Avalanche/Source/AvalancheOutliner/Private/AvaOutlinerView.h
@@ -0,0 +1,497 @@
+// Copyright Epic Games, Inc. All Rights Reserved.
+
+#pragma once
+
+#include "AvaOutlinerDefines.h"
+#include "IAvaOutlinerView.h"
+#include "Item/AvaOutlinerItemId.h"
+
+class FAvaOutliner;
+class FAvaOutlinerTextFilter;
+class FDragDropEvent;
+class FReply;
+class FUICommandList;
+class IAvaOutlinerColumn;
+class IAvaOutlinerItemFilter;
+class SAvaOutliner;
+class SHorizontalBox;
+class SWidget;
+class UToolMenu;
+enum class ECheckBoxState : uint8;
+enum class EItemDropZone;
+struct FAvaOutlinerSaveState;
+struct FAvaOutlinerStats;
+struct FGeometry;
+struct FPointerEvent;
+
+/**
+  * A view instance of the Outliner, that handles viewing a subset of the outliner items based on
+  * item filters, search text, hierarchy type, etc
+  */
+class FAvaOutlinerView : public IAvaOutlinerView
+{
+	friend FAvaOutlinerSaveState;
+	
+	struct FPrivateToken { explicit FPrivateToken() = default; };
+
+	/** Initializes the Outliner Instance. Only executed once */
+	void Init(const TSharedRef<FAvaOutliner>& InOutliner, bool bCreateOutlinerWidget);
+
+	/** Creates the Base Columns and any other columns gotten via the Module or Provider */
+	void CreateColumns();
+	
+public:
+	explicit FAvaOutlinerView(FPrivateToken);
+
+	virtual ~FAvaOutlinerView() override;
+
+	/** Gets the Tool Menu name for the ToolBar in each Outliner View Widget */
+	static FName GetOutlinerToolbarName();
+
+	/** Gets the Tool Menu name for the Item Context Menu */
+	static FName GetOutlinerItemContextMenuName();
+	
+	/**
+	 * Creates an Outliner View Instance and register it to the Outliner
+	 * @param InOutlinerViewId the Id assigned to the Instance created
+	 * @param InOutliner the Outliner this Instance View belongs to
+	 * @param bCreateOutlinerWidget whether to create an SAvaOutliner to this view. Can be false if used for testing
+	 */
+	static TSharedRef<FAvaOutlinerView> CreateInstance(int32 InOutlinerViewId
+		, const TSharedRef<FAvaOutliner>& InOutliner
+		, bool bCreateOutlinerWidget);
+
+	void PostLoad();
+	
+	/** Called when the UAvaOutlinerSettings has a property change */
+	void OnOutlinerSettingsChanged(UObject* InObject, FPropertyChangedEvent& InPropertyChangedEvent);
+
+	/**
+	 * Retrieves the latest list of Custom Filters specified in UAvaOutlinerSettings
+	 * and recreates the filter instances for these
+	 */
+	void UpdateCustomFilters();
+	
+	DECLARE_MULTICAST_DELEGATE(FOnCustomFiltersChanged)
+	FOnCustomFiltersChanged& GetOnCustomFiltersChanged() { return OnCustomFiltersChanged; }
+	
+	void Tick(float InDeltaTime);
+
+	/**
+	 * Saves the Outliner View Data into the Save State object in FAvaOutliner
+	 * @see FAvaOutlinerSaveState::SaveOutlinerViewState
+	 */
+	void SaveState();
+
+	/**
+	 * Loads the data found in the FAvaOutliner's SaveState into the respective properties in this View
+	 * @see FAvaOutlinerSaveState::LoadOutlinerViewState
+	 */
+	void LoadState();
+	
+	void BindCommands(const TSharedPtr<FUICommandList>& InBaseCommandList);
+
+	TSharedPtr<FUICommandList> GetBaseCommandList() const;
+
+	TSharedPtr<FUICommandList> GetViewCommandList() const { return ViewCommandList; }
+
+	/** Notifies FAvaOutliner so that this view instance becomes the most recent interacted outliner view */
+	void UpdateRecentOutlinerViews();
+
+	/** Returns whether this view instance is the most recently interacted between the instances in FAvaOutliner */
+	bool IsMostRecentOutlinerView() const;
+	
+	int32 GetOutlinerViewId() const { return OutlinerViewId; }
+
+	TSharedPtr<FAvaOutliner> GetOutliner() const { return OutlinerWeak.Pin(); }
+
+	TSharedRef<FAvaOutlinerStats> GetOutlinerStats() const { return OutlinerStats; }
+
+	//~ Begin IAvaOutlinerView
+	virtual TSharedRef<SWidget> GetOutlinerWidget() const override;
+	virtual TSharedPtr<IAvaOutliner> GetOwnerOutliner() const override;
+	//~ End IAvaOutlinerView
+
+	void CreateToolbar();
+	
+	TSharedPtr<SWidget> CreateItemContextMenu();
+	
+	TSharedRef<SWidget> CreateOutlinerSettingsMenu();
+
+	TSharedRef<SWidget> CreateOutlinerViewOptionsMenu();
+	
+	/** Get the Columns currently allocated in this Outliner Instance */
+	const TMap<FName, TSharedPtr<IAvaOutlinerColumn>>& GetColumns() const { return Columns; }
+
+	/** Whether the given column should be shown by default. This is only used at the start when Creating the Columns */
+	bool ShouldShowColumnByDefault(const TSharedPtr<IAvaOutlinerColumn>& InColumn) const;
+
+	/** Updates the Cached Column Visibility Map to the latest current state of this View's Widget Columns */
+	void UpdateColumnVisibilityMap();
+
+	/** Marks the Outliner View to be refreshed on Next Tick */
+	void RequestRefresh();
+	
+	DECLARE_MULTICAST_DELEGATE(FOnOutlinerViewRefreshed);
+	FOnOutlinerViewRefreshed& GetOnOutlinerViewRefreshed() { return OnOutlinerViewRefreshed; }
+	
+	/** Refreshes the Items visible in this View, and refreshes the Outliner Widget if it was created */
+	void Refresh();
+
+	void SetKeyboardFocus();
+	
+	/** Refreshes the Items that will be at the Top Level of the Tree (e.g. Actors with no parent item) */
+	void UpdateRootVisibleItems();
+
+	/** Updates the Item Expansions in the Tree Widget based on whether the Items have the Expanded flag or not */
+	void UpdateItemExpansions();
+
+	/** Called when Object Replacement has taken place. Used to invalidate the widget for painting */
+	void NotifyObjectsReplaced();
+
+	/** Called when the Text Filter has changed and thus needing to refresh to visualize the new filtered items */
+	void OnFilterChanged();
+
+	/** The Root Item of the Tree. Convenience function to get the RootItem from FAvaOutliner */
+	FAvaOutlinerItemPtr GetRootItem() const;
+
+	/** Gets the Top Level Items that should be visible in the Tree. This can differ between views depending on Filters and other factors */
+	const TArray<FAvaOutlinerItemPtr>& GetRootVisibleItems() const { return RootVisibleItems; }
+
+	void SetViewItemFlags(const FAvaOutlinerItemPtr& InItem, EAvaOutlinerItemFlags InFlags);
+	
+	EAvaOutlinerItemFlags GetViewItemFlags(const FAvaOutlinerItemPtr& InItem) const;
+
+	void GetChildrenOfItem(FAvaOutlinerItemPtr InItem, TArray<FAvaOutlinerItemPtr>& OutChildren) const;
+
+	/**
+	 * Gets the Children of a given Item. Can recurse if the immediate child is hidden (the children of these hidden items should still be given the opportunity to show up)
+	 * @param InItem the Item to get the children of
+	 * @param OutChildren the visible children in the give view mode.
+	 * @param InViewMode the view mode(s) that the children should support to be added to OutChildren
+	 * @param InRecursionDisallowedItems the items where recursion should not be performed
+	 */
+	void GetChildrenOfItem(FAvaOutlinerItemPtr InItem
+		, TArray<FAvaOutlinerItemPtr>& OutChildren
+		, EAvaOutlinerItemViewMode InViewMode
+		, const TSet<FAvaOutlinerItemPtr>& InRecursionDisallowedItems) const;
+
+	FLinearColor GetItemBrushColor(FAvaOutlinerItemPtr InItem) const;
+	
+	const TArray<TSharedPtr<IAvaOutlinerItemFilter>>& GetItemFilters() const { return ItemFilters; }
+	
+	const TArray<TSharedPtr<IAvaOutlinerItemFilter>>& GetCustomItemFilters() const { return CustomItemFilters; }
+
+	TSharedRef<FAvaOutlinerTextFilter> GetTextFilter() const { return TextFilter; }
+	
+	/** Gets the Currently Selected Items in the Tree View */
+	TArray<FAvaOutlinerItemPtr> GetViewSelectedItems() const;
+
+	/** Gets the Currently Selected Item Count in the Tree View */
+	int32 GetViewSelectedItemCount() const;
+
+	/** Calculates the amount of Items that are visible in this Outliner View */
+	int32 CalculateVisibleItemCount() const;
+	
+	/** Selects the Item in this Outliner View */
+	void SelectItems(TArray<FAvaOutlinerItemPtr> InItems, EAvaOutlinerItemSelectionFlags InFlags = EAvaOutlinerItemSelectionFlags::SignalSelectionChange);
+
+	/** Clears the currently Selected Items in the Outliner View */
+	void ClearItemSelection(bool bSignalSelectionChange = true);
+
+	/** Whether Sync Selection is currently taking place */
+	bool IsSyncingItemSelection() const { return bSyncingItemSelection; }
+	
+private:
+	/** Sets the Selected Items to be the given Items in this List. */
+	void SetItemSelectionImpl(TArray<FAvaOutlinerItemPtr>&& InItems
+		, bool bSignalSelectionChange);
+
+public:
+	/** Called when Item selection has changed. This can be called by the OutlinerWidget or by the View itself if there's no OutlinerWidget */
+	void NotifyItemSelectionChanged(const TArray<FAvaOutlinerItemPtr>& InSelectedItems
+		, const FAvaOutlinerItemPtr& InItem
+		, bool bUpdateModeTools);
+
+	/** Whether the given Item is explicitly marked as Read-only in the Outliner */
+	bool IsItemReadOnly(const FAvaOutlinerItemPtr& InItem) const;
+	
+	/** Whether the given Item can be selected in this Outliner View / Widget */
+	bool CanSelectItem(const FAvaOutlinerItemPtr& InItem) const;
+
+	/** Whether the current item is selected in this Outliner View / Widget */
+	bool IsItemSelected(const FAvaOutlinerItemPtr& InItem) const;
+
+	/** Change the expansion state of the parents (recursively) of the given item */
+	void SetParentItemExpansions(const FAvaOutlinerItemPtr& InItem, bool bIsExpanded);
+
+	/** Expands / Collapses the given Item */
+	void SetItemExpansion(const FAvaOutlinerItemPtr& InItem, bool bIsExpanded, bool bInUseFilter = true);
+
+	/** Expands / Collapses the given Item and its children recursively */
+	void SetItemExpansionRecursive(FAvaOutlinerItemPtr InItem, bool bIsExpanded);
+
+	/** Called when the given item's expansion state (expanded/collapsed) has changed */
+	void OnItemExpansionChanged(FAvaOutlinerItemPtr InItem, bool bIsExpanded);
+
+	/** Return whether the given item should be visible in this Outliner Instance */
+	bool ShouldShowItem(const FAvaOutlinerItemPtr& InItem, bool bInUseFilters, EAvaOutlinerItemViewMode InViewMode) const;
+
+	/** Returns the Index of the Child from the Parent's Children Visible List (Visible List can vary between Outliner Instances) */
+	int32 GetVisibleChildIndex(const FAvaOutlinerItemPtr& InParentItem, const FAvaOutlinerItemPtr& InChildItem) const;
+
+	/** Returns the Child from the Parent's Children Visible List at the given Index (Visible List can vary between Outliner Instances) */
+	FAvaOutlinerItemPtr GetVisibleChildAt(const FAvaOutlinerItemPtr& InParentItem, int32 InChildIndex) const;
+
+	/** Returns whether the Outliner is Locked. See FAvaOutliner::IsOutlinerLocked */
+	bool IsOutlinerLocked() const;
+
+	/**
+	 * Activate the given Item Filter
+	 * @param InFilter the filter to activate
+	 * @param bInRefreshOutliner whether to refresh this view if the filter was added. false recommended if doing operations in bulk
+	 */
+	void EnableItemFilter(const TSharedPtr<IAvaOutlinerItemFilter>& InFilter, const bool bInRefreshOutliner = true);
+
+	/**
+ 	* Activate the Item Filter with the given ID
+ 	* @param InFilterId the filter id associated with the Filter to activate
+ 	* @param bInRefreshOutliner whether to refresh this view if the filter was added. false recommended if doing operations in bulk
+ 	*/
+	void EnableItemFilterById(const FName& InFilterId, const bool bInRefreshOutliner = true);
+
+	/**
+	 * Deactivates the given Item Filter
+	 * @param InFilter the filter to remove from the active filter list
+	 * @param bInRefreshOutliner whether to refresh this view if the filter was removed. false recommended if doing operations in bulk
+	 */
+	void DisableItemFilter(const TSharedPtr<IAvaOutlinerItemFilter>& InFilter, const bool bInRefreshOutliner = true);
+
+	/** Whether the given Filter is in the Active Filter List */
+	bool IsItemFilterEnabled(const TSharedPtr<IAvaOutlinerItemFilter>& InFilter) const;
+
+	/** Sets the Filter Text */
+	void SetSearchText(const FText& InText);
+
+	/** Gets the Current Filter Text */
+	FText GetSearchText() const;
+
+	EAvaOutlinerItemViewMode GetItemDefaultViewMode() const { return ItemDefaultViewMode; }
+	EAvaOutlinerItemViewMode GetItemProxyViewMode() const { return ItemProxyViewMode; }
+
+	void ToggleViewModeSupport(EAvaOutlinerItemViewMode& InOutViewMode, EAvaOutlinerItemViewMode InFlags);
+	void ToggleItemDefaultViewModeSupport(EAvaOutlinerItemViewMode InFlags);
+	void ToggleItemProxyViewModeSupport(EAvaOutlinerItemViewMode InFlags);
+
+	ECheckBoxState GetViewModeCheckState(EAvaOutlinerItemViewMode InViewMode, EAvaOutlinerItemViewMode InFlags) const;
+	ECheckBoxState GetItemDefaultViewModeCheckState(EAvaOutlinerItemViewMode InFlags) const;
+	ECheckBoxState GetItemProxyViewModeCheckState(EAvaOutlinerItemViewMode InFlags) const;
+	
+	/**
+	 * Action to turn Muted Hierarchy on or off
+	 * @see bUseMutedHierarchy
+	 */
+	void ToggleMutedHierarchy();
+	bool CanToggleMutedHierarchy() const { return true; }
+	bool IsMutedHierarchyActive() const { return bUseMutedHierarchy; }
+
+	/**
+	 * Action to turn Auto Expand to Selection on or off
+	 * @see bAutoExpandToSelection
+	 */
+	void ToggleAutoExpandToSelection();
+	bool CanToggleAutoExpandToSelection() const { return true; }
+	bool ShouldAutoExpandToSelection() const { return bAutoExpandToSelection; }
+
+	void ToggleShowItemFilters();
+	bool CanToggleShowItemFilters() const { return true; }
+	bool ShouldShowItemFilters() const { return bShowItemFilters; }
+
+	/** Sets whether the given item type should be hidden or not */
+	void SetItemTypeHidden(FName InItemTypeName, bool bHidden);
+	
+	/** Toggles the given Item Types to Hide/Show */
+	void ToggleHideItemTypes(FName InItemTypeName);
+	ECheckBoxState GetToggleHideItemTypesState(FName InItemTypeName) const;
+	
+	/** Hides the given Outliner Item Type from showing in this Outliner View (e.g. Hide Components is used) */
+	template<typename InItemType, typename = typename TEnableIf<TIsDerivedFrom<InItemType, IAvaOutlinerItem>::IsDerived>::Type>
+	void HideItemType()
+	{
+		SetItemTypeHidden(TAvaType<InItemType>::GetTypeId().ToName(), true);
+	}
+	
+	/** Shows the given Outliner Item Type if it was hidden in this Outliner View (e.g. Showing Components again) */
+	template<typename InItemType, typename = typename TEnableIf<TIsDerivedFrom<InItemType, IAvaOutlinerItem>::IsDerived>::Type>
+	void ExposeItemType()
+	{
+		SetItemTypeHidden(TAvaType<InItemType>::GetTypeId().ToName(), false);
+	}
+	
+	/** Whether the given Outliner Item Type is currently hidden in this Outliner View */
+	bool IsItemTypeHidden(FName InItemTypeName) const;
+	bool IsItemTypeHidden(const FAvaOutlinerItemPtr& InItem) const;
+	
+	template<typename InItemType, typename = typename TEnableIf<TIsDerivedFrom<InItemType, IAvaOutlinerItem>::IsDerived>::Type>
+	bool IsItemTypeHidden() const
+	{
+		return IsItemTypeHidden(InItemType::GetStaticTypeName());
+	}
+
+	/** Called when a Drag enters the Outliner Widgets for the given Item (if null, treat as root) */
+	void OnDragEnter(const FDragDropEvent& InDragDropEvent, FAvaOutlinerItemPtr InTargetItem);
+
+	/** Called when a Drag leaves the Outliner Widgets for the given Item (if null, treat as root) */
+	void OnDragLeave(const FDragDropEvent& InDragDropEvent, FAvaOutlinerItemPtr InTargetItem);
+	
+	/** Called a Drag is attempted for the selected items in view */
+	FReply OnDragDetected(const FGeometry& InGeometry, const FPointerEvent& InMouseEvent, FAvaOutlinerItemPtr InTargetItem);
+
+	/** Processes the Drag and Drop for the given Item if valid, else it will process it for the Root Item */
+	FReply OnDrop(const FDragDropEvent& InDragDropEvent, EItemDropZone InDropZone, FAvaOutlinerItemPtr InTargetItem);
+
+	/** Determines whether the Drag and Drop can be processed by the given item */
+	TOptional<EItemDropZone> OnCanDrop(const FDragDropEvent& InDragDropEvent
+		, EItemDropZone InDropZone
+		, FAvaOutlinerItemPtr InTargetItem) const;
+
+	/** Called when there's a Drag being started/ended to the Outliner Widget as a whole (i.e. root) and not an individual item */
+	void SetDragIntoTreeRoot(bool bInIsDraggingIntoTreeRoot);
+	
+	void RenameSelected();
+	void ResetRenaming();
+	void OnItemRenameAction(EAvaOutlinerRenameAction InRenameAction, const TSharedPtr<FAvaOutlinerView>& InOutlinerView);
+	bool CanRenameSelected() const;
+
+	void DuplicateSelected();
+	bool CanDuplicateSelected() const;
+
+	void SelectChildren(bool bIsRecursive);
+	bool CanSelectChildren() const;
+
+	void SelectParent();
+	bool CanSelectParent() const;
+
+	void SelectFirstChild();
+	bool CanSelectFirstChild() const;
+
+	void SelectSibling(int32 InDeltaIndex);
+	bool CanSelectSibling() const;
+
+	void ExpandAll();
+	bool CanExpandAll() const;
+
+	void CollapseAll();
+	bool CanCollapseAll() const;
+
+	void ScrollNextIntoView();
+	void ScrollPrevIntoView();
+	bool CanScrollNextIntoView() const;
+	void ScrollDeltaIndexIntoView(int32 DeltaIndex);
+	void ScrollItemIntoView(const FAvaOutlinerItemPtr& InItem);
+
+	/** Sorts the given Items in the user-defined order and selects them, and if the Widget is valid focusing to the item in the Widget */
+	void SortAndSelectItems(TArray<FAvaOutlinerItemPtr> InItemsToSelect);
+	
+private:
+	/** Static Function to Populate the Outliner Item Context Menu. The Outliner View and Context Items are gotten via UAvaOutlinerItemsContext */
+	static void PopulateItemContextMenu(UToolMenu* InToolMenu);
+	
+	/** Static Function to Populate the Outliner Main Toolbar. The Outliner View is gotten via UAvaOutlinerToolbarContext */
+	static void PopulateToolBar(UToolMenu* InToolMenu);
+	
+	/** Triggers a Refresh on the FAvaOutliner */
+	void RefreshOutliner(bool bInImmediateRefresh);
+	
+	/** Local Identifier of this Instance */
+	int32 OutlinerViewId = -1;
+
+	/** Weak pointer to the outliner this is a view of */
+	TWeakPtr<FAvaOutliner> OutlinerWeak;
+
+	/** Widget showing the View. Can be null if instanced for testing (i.e. bCreateOutlinerWidget is false in Init) */
+	TSharedPtr<SAvaOutliner> OutlinerWidget;
+
+	/** Command List mapped to the View to handle things like Selected Items */
+	TSharedPtr<FUICommandList> ViewCommandList;
+	
+	/** Root Items from Outliner visible to this Instance. i.e. a Subset of all the Items in FAvaOutliner */
+	TArray<FAvaOutlinerItemPtr> RootVisibleItems;
+
+	/** A list of the Selected Items in this Outliner View */
+	TArray<FAvaOutlinerItemPtr> SelectedItems;
+	
+	/** Set of Items that are Currently Read Only in this Instance. Mutable as this changes on GetChildrenOfItem const func */
+	mutable TSet<FAvaOutlinerItemPtr> ReadOnlyItems;
+
+	/** Items specific to this Outliner Instance, rather than being shared across Outliners (e.g. Expansion flags) */
+	TMap<FAvaOutlinerItemId, EAvaOutlinerItemFlags> ViewItemFlags;
+
+	TSharedRef<FAvaOutlinerTextFilter> TextFilter;
+
+	/** Native-constructed Item Filters in the Outliner */
+	TArray<TSharedPtr<IAvaOutlinerItemFilter>> ItemFilters;
+
+	/** Custom Item Filters constructed from UAvaOutlinerSettings */
+	TArray<TSharedPtr<IAvaOutlinerItemFilter>> CustomItemFilters;
+
+	/** Delegate called when the Custom Item Filters have been modified, as these can be instanced after */
+	FOnCustomFiltersChanged OnCustomFiltersChanged;
+
+	/** Delegate called at the end of the Outliner View Refresh */
+	FOnOutlinerViewRefreshed OnOutlinerViewRefreshed;
+	
+	/** Active Item Filters. Subset of ItemFilters */
+	TSet<TSharedPtr<IAvaOutlinerItemFilter>> ActiveItemFilters;
+
+	/** Map of the Column Ids and the Colum with that Id. This is filled in FAvaOutlinerView::CreateColumns */
+	TMap<FName, TSharedPtr<IAvaOutlinerColumn>> Columns;
+
+	/** Map of the Column Ids to their Override (i.e. saved) Visibility */
+	TMap<FName, bool> ColumnVisibility;
+	
+	/** A set of the Outliner Item Type Names that should be hidden in the Outliner */
+	TSet<FName> HiddenItemTypes;
+
+	/** The Index of the Item that will be scrolled into View */
+	int32 NextSelectedItemIntoView = -1;
+
+	/** Selected Items Sorted from Top To Bottom */
+	TArray<FAvaOutlinerItemPtr> SortedSelectedItems;
+
+	/** Stats for the Outliner Widget */
+	TSharedRef<FAvaOutlinerStats> OutlinerStats;
+
+	/** The list of items that need renaming and are waiting for their turn to be renamed */
+	TArray<FAvaOutlinerItemPtr> ItemsRemainingRename;
+	
+	/** The current item in the process of renaming. Null if no renaming is taking place */
+	FAvaOutlinerItemPtr CurrentItemRenaming;
+	
+	/** This Outliner View's Setting for the Default Item's Default View Mode */
+	EAvaOutlinerItemViewMode ItemDefaultViewMode = EAvaOutlinerItemViewMode::HorizontalItemList;
+	
+	/** This Outliner View's Setting for the Item Proxy's Default View Mode */
+	EAvaOutlinerItemViewMode ItemProxyViewMode = EAvaOutlinerItemViewMode::None;
+	
+	/** Whether the expand the parents of an item when it's selected */
+	bool bAutoExpandToSelection = false;
+	
+	/** Whether to also show the Parent Items in read-only mode when these are filtered out */
+	bool bUseMutedHierarchy = false;
+	
+	/** Flag to call FAvaOutlinerView::Refresh next tick */
+	bool bRefreshRequested = false;
+	
+	/** Whether Item renaming is taking place in this View */
+	bool bRenamingItems = false;
+	
+	/** Whether Renaming Items should be processed next tick */
+	bool bRequestedRename = false;
+
+	/** Whether Items are currently being synced to what the Widget or Interface sent. Used as a re-enter guard */
+	bool bSyncingItemSelection = false;
+	
+	/** Flag used for the Outliner Widget to determine if the Item Filter Bar should open and show the Item Filters */
+	bool bShowItemFilters = false;
+};
diff --git a/Engine/Plugins/Experimental/Avalanche/Source/AvalancheOutliner/Private/Columns/AvaOutlinerColumnExtender.h b/Engine/Plugins/Experimental/Avalanche/Source/AvalancheOutliner/Private/Columns/AvaOutlinerColumnExtender.h
new file mode 100644
--- /dev/null
+++ b/Engine/Plugins/Experimental/Avalanche/Source/AvalancheOutliner/Private/Columns/AvaOutlinerColumnExtender.h
@@ -0,0 +1,50 @@
+// Copyright Epic Games, Inc. All Rights Reserved.
+
+#pragma once
+
+#include "AvaOutlinerDefines.h"
+#include "Templates/UnrealTypeTraits.h"
+
+class IAvaOutlinerColumn;
+
+/**
+ * Extension Class to add Outliner Columns
+ *
+ * Example #1:
+ * This would create a Tag Column at the end of the current column list (order matters!)
+ *    ColumnExtender.AddColumn<FAvaOutlinerTagColumn>();
+ *
+ * Example #2:
+ * This would create a Tag Column before the Label Column (if it doesn't exist, it's the same behavior as the above example)
+ *    ColumnExtender.AddColumn<FAvaOutlinerTagColumn, EAvaOutlinerExtensionPosition::Before, FAvaOutlinerLabelColumn>();
+ */
+class FAvaOutlinerColumnExtender
+{
+	typedef TSharedPtr<IAvaOutlinerColumn> FAvaOutlinerColumnPtr;
+
+public:
+	template<typename InColumnType, EAvaOutlinerExtensionPosition InExtensionPosition, typename InRefColumnType
+		UE_REQUIRES(TIsDerivedFrom<InColumnType, IAvaOutlinerColumn>::Value && TIsDerivedFrom<InRefColumnType, IAvaOutlinerColumn>::Value)>
+	void AddColumn()
+	{
+		AddColumn(MakeShared<InColumnType>(), InExtensionPosition, InRefColumnType::GetStaticTypeName());
+	}
+
+	template<typename InColumnType
+		UE_REQUIRES(TIsDerivedFrom<InColumnType, IAvaOutlinerColumn>::Value)>
+	void AddColumn()
+	{
+		AddColumn(MakeShared<InColumnType>(), EAvaOutlinerExtensionPosition::Before, NAME_None);
+	}
+
+	const TArray<FAvaOutlinerColumnPtr>& GetColumns() const { return Columns; }
+
+private:
+	void AddColumn(FAvaOutlinerColumnPtr InColumn
+		, EAvaOutlinerExtensionPosition InExtensionPosition
+		, FName InReferenceColumnId);
+
+	int32 FindColumnIndex(FName InColumnId) const;
+
+	TArray<FAvaOutlinerColumnPtr> Columns;
+};
diff --git a/Engine/Plugins/Experimental/Avalanche/Source/AvalancheOutliner/Private/Columns/IAvaOutlinerColumn.h b/Engine/Plugins/Experimental/Avalanche/Source/AvalancheOutliner/Private/Columns/IAvaOutlinerColumn.h
new file mode 100644
--- /dev/null
+++ b/Engine/Plugins/Experimental/Avalanche/Source/AvalancheOutliner/Private/Columns/IAvaOutlinerColumn.h
@@ -0,0 +1,34 @@
+// Copyright Epic Games, Inc. All Rights Reserved.
+
+#pragma once
+
+#include "AvaOutlinerDefines.h"
+#include "AvaType.h"
+#include "Widgets/Views/SHeaderRow.h"
+
+class FAvaOutlinerView;
+class SAvaOutlinerTreeRow;
+
+class IAvaOutlinerColumn : public IAvaTypeCastable, public TSharedFromThis<IAvaOutlinerColumn>
+{
+public:	
+	UE_AVA_INHERITS(IAvaOutlinerColumn, IAvaTypeCastable);
+
+	FName GetColumnId() const { return GetTypeId().ToName(); }
+
+	virtual FText GetColumnDisplayNameText() const = 0;
+
+	virtual SHeaderRow::FColumn::FArguments ConstructHeaderRowColumn() = 0;
+
+	/*
+	 * Determines whether the Column should be Showing by Default while still be able to toggle it on/off.
+	 * Used when calling SHeaderRow::SetShowGeneratedColumn (requires ShouldGenerateWidget to not be set).
+	 */
+	virtual bool ShouldShowColumnByDefault() const { return false; }
+	
+	virtual TSharedRef<SWidget> ConstructRowWidget(FAvaOutlinerItemPtr InItem
+		, const TSharedRef<FAvaOutlinerView>& InOutlinerView
+		, const TSharedRef<SAvaOutlinerTreeRow>& InRow) = 0;
+
+	virtual void Tick(float InDeltaTime) {}
+};
diff --git a/Engine/Plugins/Experimental/Avalanche/Source/AvalancheOutliner/Private/Columns/Slate/SAvaOutlinerItemChip.cpp b/Engine/Plugins/Experimental/Avalanche/Source/AvalancheOutliner/Private/Columns/Slate/SAvaOutlinerItemChip.cpp
--- a/Engine/Plugins/Experimental/Avalanche/Source/AvalancheOutliner/Private/Columns/Slate/SAvaOutlinerItemChip.cpp
+++ b/Engine/Plugins/Experimental/Avalanche/Source/AvalancheOutliner/Private/Columns/Slate/SAvaOutlinerItemChip.cpp
@@ -210,7 +210,7 @@ bool SAvaOutlinerItemChip::IsSelected() const
 
 const FSlateBrush* SAvaOutlinerItemChip::GetItemBackgroundBrush() const
 {
-	using namespace UE::AvalancheOutliner::Private;
+	using namespace UE::AvaOutliner::Private;
 
 	const EStyleType StyleType = IsHovered()
 		? EStyleType::Hovered
diff --git a/Engine/Plugins/Experimental/Avalanche/Source/AvalancheOutliner/Private/Columns/Slate/SAvaOutlinerLock.cpp b/Engine/Plugins/Experimental/Avalanche/Source/AvalancheOutliner/Private/Columns/Slate/SAvaOutlinerLock.cpp
--- a/Engine/Plugins/Experimental/Avalanche/Source/AvalancheOutliner/Private/Columns/Slate/SAvaOutlinerLock.cpp
+++ b/Engine/Plugins/Experimental/Avalanche/Source/AvalancheOutliner/Private/Columns/Slate/SAvaOutlinerLock.cpp
@@ -8,7 +8,7 @@
 
 #define LOCTEXT_NAMESPACE "SAvaOutlinerLock"
 
-namespace UE::AvalancheOutliner::Private
+namespace UE::AvaOutliner::Private
 {
 	class FLockDragDropOp : public FDragDropOperation, public TSharedFromThis<FLockDragDropOp>
 	{
@@ -102,12 +102,12 @@ FSlateColor SAvaOutlinerLock::GetForegroundColor() const
 
 void SAvaOutlinerLock::SetItemLocked(bool bInLocked)
 {
-	UE::AvalancheOutliner::Private::SetItemLocked(ItemWeak.Pin(), bInLocked);
+	UE::AvaOutliner::Private::SetItemLocked(ItemWeak.Pin(), bInLocked);
 }
 
 bool SAvaOutlinerLock::IsItemLocked() const
 {
-	return UE::AvalancheOutliner::Private::IsItemLocked(ItemWeak.Pin());
+	return UE::AvaOutliner::Private::IsItemLocked(ItemWeak.Pin());
 }
 
 const FSlateBrush* SAvaOutlinerLock::GetBrush() const
@@ -119,14 +119,14 @@ FReply SAvaOutlinerLock::OnDragDetected(const FGeometry& MyGeometry, const FPoin
 {
 	if (MouseEvent.IsMouseButtonDown(EKeys::LeftMouseButton))
 	{
-		return FReply::Handled().BeginDragDrop(UE::AvalancheOutliner::Private::FLockDragDropOp::New(IsItemLocked(), UndoTransaction));
+		return FReply::Handled().BeginDragDrop(UE::AvaOutliner::Private::FLockDragDropOp::New(IsItemLocked(), UndoTransaction));
 	}
 	return FReply::Unhandled();
 }
 
 void SAvaOutlinerLock::OnDragEnter(const FGeometry& MyGeometry, const FDragDropEvent& DragDropEvent)
 {
-	using namespace UE::AvalancheOutliner::Private;
+	using namespace UE::AvaOutliner::Private;
 	if (const TSharedPtr<FLockDragDropOp> DragDropOp = DragDropEvent.GetOperationAs<FLockDragDropOp>())
 	{
 		SetItemLocked(DragDropOp->bShouldLock);
@@ -153,9 +153,9 @@ FReply SAvaOutlinerLock::HandleClick()
 	{
 		for (TSharedPtr<IAvaOutlinerItem>& SelectedItem : OutlinerView->GetViewSelectedItems())
 		{
-			if (UE::AvalancheOutliner::Private::IsItemLocked(SelectedItem) != bLockItem)
+			if (UE::AvaOutliner::Private::IsItemLocked(SelectedItem) != bLockItem)
 			{
-				UE::AvalancheOutliner::Private::SetItemLocked(SelectedItem, bLockItem);
+				UE::AvaOutliner::Private::SetItemLocked(SelectedItem, bLockItem);
 			}
 		}
 	}
diff --git a/Engine/Plugins/Experimental/Avalanche/Source/AvalancheOutliner/Private/Columns/Slate/SAvaOutlinerVisibility.cpp b/Engine/Plugins/Experimental/Avalanche/Source/AvalancheOutliner/Private/Columns/Slate/SAvaOutlinerVisibility.cpp
--- a/Engine/Plugins/Experimental/Avalanche/Source/AvalancheOutliner/Private/Columns/Slate/SAvaOutlinerVisibility.cpp
+++ b/Engine/Plugins/Experimental/Avalanche/Source/AvalancheOutliner/Private/Columns/Slate/SAvaOutlinerVisibility.cpp
@@ -10,7 +10,7 @@
 
 #define LOCTEXT_NAMESPACE "SAvaOutlinerVisibility"
 
-namespace UE::AvalancheOutliner::Private
+namespace UE::AvaOutliner::Private
 {
 	class FVisibilityDragDropOp : public FDragDropOperation, public TSharedFromThis<FVisibilityDragDropOp>
 	{
@@ -73,7 +73,7 @@ void SAvaOutlinerVisibility::Construct(const FArguments& InArgs
 
 FReply SAvaOutlinerVisibility::OnDragDetected(const FGeometry& MyGeometry, const FPointerEvent& MouseEvent)
 {
-	using namespace UE::AvalancheOutliner::Private;
+	using namespace UE::AvaOutliner::Private;
 	if (MouseEvent.IsMouseButtonDown(EKeys::LeftMouseButton))
 	{
 		return FReply::Handled().BeginDragDrop(FVisibilityDragDropOp::New(!IsItemVisible(), UndoTransaction));
@@ -84,7 +84,7 @@ FReply SAvaOutlinerVisibility::OnDragDetected(const FGeometry& MyGeometry, const
 /** If a visibility drag drop operation has entered this widget, set its item to the new visibility state */
 void SAvaOutlinerVisibility::OnDragEnter(const FGeometry& MyGeometry, const FDragDropEvent& DragDropEvent)
 {
-	using namespace UE::AvalancheOutliner::Private;
+	using namespace UE::AvaOutliner::Private;
 	const TSharedPtr<FVisibilityDragDropOp> DragDropOp = DragDropEvent.GetOperationAs<FVisibilityDragDropOp>();
 	if (DragDropOp.IsValid())
 	{
@@ -118,7 +118,7 @@ FReply SAvaOutlinerVisibility::HandleClick()
 	{
 		for (TSharedPtr<IAvaOutlinerItem>& SelectedItem : OutlinerView->GetViewSelectedItems())
 		{
-			if (UE::AvalancheOutliner::Private::IsItemVisible(SelectedItem, Column) != bVisible)
+			if (UE::AvaOutliner::Private::IsItemVisible(SelectedItem, Column) != bVisible)
 			{
 				OnSetItemVisibility(SelectedItem, bVisible);
 			}
@@ -221,7 +221,7 @@ FSlateColor SAvaOutlinerVisibility::GetForegroundColor() const
 /** Check if our wrapped tree item is visible */
 bool SAvaOutlinerVisibility::IsItemVisible() const
 {
-	return UE::AvalancheOutliner::Private::IsItemVisible(ItemWeak.Pin(), ColumnWeak.Pin());
+	return UE::AvaOutliner::Private::IsItemVisible(ItemWeak.Pin(), ColumnWeak.Pin());
 }
 
 /** Set the item this widget is responsible for to be hidden or shown */
diff --git a/Engine/Plugins/Experimental/Avalanche/Source/AvalancheOutliner/Private/Data/AvaOutlinerSaveState.cpp b/Engine/Plugins/Experimental/Avalanche/Source/AvalancheOutliner/Private/Data/AvaOutlinerSaveState.cpp
--- a/Engine/Plugins/Experimental/Avalanche/Source/AvalancheOutliner/Private/Data/AvaOutlinerSaveState.cpp
+++ b/Engine/Plugins/Experimental/Avalanche/Source/AvalancheOutliner/Private/Data/AvaOutlinerSaveState.cpp
@@ -288,7 +288,7 @@ void FAvaOutlinerSaveState::LoadSceneTree(const FAvaOutlinerItemPtr& InParentIte
 	
 	TArray<FAvaOutlinerItemPtr> Sortable;
 	TArray<FAvaOutlinerItemPtr> Unsortable;
-	UE::AvalancheOutliner::SplitItems(Children, Sortable, Unsortable);
+	UE::AvaOutliner::SplitItems(Children, Sortable, Unsortable);
 
 	// If Scene Tree is valid, Item Sorting should be empty as this function only takes a valid Scene Tree if
 	// loaded version supports Scene Trees (i.e. when Item Sorting stops being loaded in)
diff --git a/Engine/Plugins/Experimental/Avalanche/Source/AvalancheOutliner/Private/DragDropOps/AvaOutlinerItemDragDropOp.cpp b/Engine/Plugins/Experimental/Avalanche/Source/AvalancheOutliner/Private/DragDropOps/AvaOutlinerItemDragDropOp.cpp
--- a/Engine/Plugins/Experimental/Avalanche/Source/AvalancheOutliner/Private/DragDropOps/AvaOutlinerItemDragDropOp.cpp
+++ b/Engine/Plugins/Experimental/Avalanche/Source/AvalancheOutliner/Private/DragDropOps/AvaOutlinerItemDragDropOp.cpp
@@ -54,15 +54,6 @@ TSharedRef<FAvaOutlinerItemDragDropOp> FAvaOutlinerItemDragDropOp::New(const TAr
 	return DragDropOp;
 }
 
-TSharedPtr<FAvaOutliner> FAvaOutlinerItemDragDropOp::GetOutliner() const
-{
-	if (TSharedPtr<FAvaOutlinerView> OutlinerView = GetOutlinerView())
-	{
-		return OutlinerView->GetOutliner();
-	}
-	return nullptr;
-}
-
 void FAvaOutlinerItemDragDropOp::GetDragDropOpActors(TSharedPtr<FDragDropOperation> InDragDropOp, TArray<TWeakObjectPtr<AActor>>& OutActors)
 {
 	if (InDragDropOp.IsValid() && InDragDropOp->IsOfType<FAvaOutlinerItemDragDropOp>())
@@ -120,8 +111,10 @@ FReply FAvaOutlinerItemDragDropOp::Drop(EItemDropZone InDropZone, FAvaOutlinerIt
 
 TOptional<EItemDropZone> FAvaOutlinerItemDragDropOp::CanDrop(EItemDropZone InDropZone, FAvaOutlinerItemPtr InTargetItem) const
 {
+	TSharedPtr<IAvaOutlinerView> OutlinerView = GetOutlinerView();
+
 	// Only support Drag/Drop from Same Outliner
-	if (InTargetItem->GetOwnerOutliner() != GetOutliner())
+	if (!OutlinerView.IsValid() || InTargetItem->GetOwnerOutliner() != OutlinerView->GetOwnerOutliner())
 	{
 		return TOptional<EItemDropZone>();
 	}
@@ -170,11 +163,6 @@ void FAvaOutlinerItemDragDropOp::Init(const TArray<FAvaOutlinerItemPtr>& InItems
 	AddDropHandler<FAvaOutlinerActorDropHandler>();
 
 	FAvaOutlinerItemDragDropOp::OnItemDragDropOpInitialized().Broadcast(*this);
-
-	if (TSharedPtr<FAvaOutliner> Outliner = GetOutliner())
-	{
-		Outliner->GetProvider().OnItemDragDropInitialized(*this);
-	}
 }
 
 #undef LOCTEXT_NAMESPACE
diff --git a/Engine/Plugins/Experimental/Avalanche/Source/AvalancheOutliner/Private/DragDropOps/Handlers/AvaOutlinerActorDropHandler.cpp b/Engine/Plugins/Experimental/Avalanche/Source/AvalancheOutliner/Private/DragDropOps/Handlers/AvaOutlinerActorDropHandler.cpp
--- a/Engine/Plugins/Experimental/Avalanche/Source/AvalancheOutliner/Private/DragDropOps/Handlers/AvaOutlinerActorDropHandler.cpp
+++ b/Engine/Plugins/Experimental/Avalanche/Source/AvalancheOutliner/Private/DragDropOps/Handlers/AvaOutlinerActorDropHandler.cpp
@@ -42,7 +42,7 @@ bool FAvaOutlinerActorDropHandler::Drop(EItemDropZone InDropZone, FAvaOutlinerIt
 		break;
 
 	case EAvaOutlinerDragDropActionType::Copy:
-		InTargetItem->GetOwnerOutliner()->DuplicateItems(Items, InTargetItem, InDropZone);
+		StaticCastSharedRef<FAvaOutliner>(InTargetItem->GetOwnerOutliner())->DuplicateItems(Items, InTargetItem, InDropZone);
 		break;
 
 	default:
@@ -78,7 +78,7 @@ void FAvaOutlinerActorDropHandler::MoveItems(EItemDropZone InDropZone, FAvaOutli
 		Algo::Reverse(Items);
 	}
 
-	TSharedRef<FAvaOutliner> Outliner = InTargetItem->GetOwnerOutliner();
+	TSharedRef<FAvaOutliner> Outliner = StaticCastSharedRef<FAvaOutliner>(InTargetItem->GetOwnerOutliner());
 
 	for (const FAvaOutlinerItemPtr& Item : Items)
 	{
diff --git a/Engine/Plugins/Experimental/Avalanche/Source/AvalancheOutliner/Private/Filters/AvaOutlinerFilterExpressionContext.cpp b/Engine/Plugins/Experimental/Avalanche/Source/AvalancheOutliner/Private/Filters/AvaOutlinerFilterExpressionContext.cpp
--- a/Engine/Plugins/Experimental/Avalanche/Source/AvalancheOutliner/Private/Filters/AvaOutlinerFilterExpressionContext.cpp
+++ b/Engine/Plugins/Experimental/Avalanche/Source/AvalancheOutliner/Private/Filters/AvaOutlinerFilterExpressionContext.cpp
@@ -2,6 +2,8 @@
 
 #include "AvaOutlinerFilterExpressionContext.h"
 #include "AvaOutliner.h"
+#include "AvaOutlinerModule.h"
+#include "IAvaFilterExpressionFactory.h"
 #include "Item/AvaOutlinerObject.h"
 #include "Item/IAvaOutlinerItem.h"
 
@@ -46,9 +48,9 @@ bool FAvaOutlinerFilterExpressionContext::TestComplexExpression(const FName& InK
 	, const ETextFilterComparisonOperation InComparisonOperation
 	, const ETextFilterTextComparisonMode InTextComparisonMode) const
 {
-	const IAvaOutlinerModule& AvaOutlinerModule = IAvaOutlinerModule::Get();
+	const FAvaOutlinerModule& OutlinerModule = FAvaOutlinerModule::Get();
 
-	if (AvaOutlinerModule.CanFilterSupportComparisonOperation(InKey, InComparisonOperation))
+	if (OutlinerModule.CanFilterSupportComparisonOperation(InKey, InComparisonOperation))
 	{
 		FAvaTextFilterArgs FilterArgs;
 		FilterArgs.ItemClass = ItemClass;
@@ -57,7 +59,7 @@ bool FAvaOutlinerFilterExpressionContext::TestComplexExpression(const FName& InK
 		FilterArgs.ComparisonOperation = InComparisonOperation;
 		FilterArgs.ComparisonMode = InTextComparisonMode;
 
-		return AvaOutlinerModule.FilterExpression(InKey, *Item, FilterArgs);
+		return OutlinerModule.FilterExpression(InKey, *Item, FilterArgs);
 	}
 
 	return false;
diff --git a/Engine/Plugins/Experimental/Avalanche/Source/AvalancheOutliner/Private/Filters/IAvaFilterExpressionFactory.h b/Engine/Plugins/Experimental/Avalanche/Source/AvalancheOutliner/Private/Filters/IAvaFilterExpressionFactory.h
new file mode 100644
--- /dev/null
+++ b/Engine/Plugins/Experimental/Avalanche/Source/AvalancheOutliner/Private/Filters/IAvaFilterExpressionFactory.h
@@ -0,0 +1,66 @@
+// Copyright Epic Games, Inc. All Rights Reserved.
+
+#pragma once
+
+#include "UObject/NameTypes.h"
+#include "Templates/SharedPointer.h"
+
+class IAvaOutlinerItem;
+
+/**
+ * Holds the arguments required for the expression evaluation.
+ */
+struct FAvaTextFilterArgs
+{
+	/** item Class */
+	FName ItemClass;
+
+	/** Item Display Name */
+	FName ItemDisplayName;
+
+	/** Second Value */
+	FName ValueToCheck;
+
+	/** Comparison Operation requested (Equal/Greater/etc...) still needed for certain Filter */
+	ETextFilterComparisonOperation ComparisonOperation;
+
+	/** Comparison Mode requested see ETextFilterTextComparisonMode for more information */
+	ETextFilterTextComparisonMode ComparisonMode;
+};
+
+class IAvaFilterExpressionFactory : public TSharedFromThis<IAvaFilterExpressionFactory>
+{
+public:
+	virtual ~IAvaFilterExpressionFactory() {}
+
+	/** Get the filter identifier for this factory */
+	virtual FName GetFilterIdentifier() const = 0;
+
+	/**
+ 	 * Create and return an Instance of the AvaFilterExpressionFactory Requested
+ 	 * @tparam InFilterExpressionFactoryType The type of the factory to instantiate
+ 	 * @param InArgs Additional Args for constructor of the AvaFilterExpressionFactory class if needed
+ 	 * @return The new filter factory
+ 	 */
+	template<typename InFilterExpressionFactoryType, typename... InArgsType
+			, typename = typename TEnableIf<TIsDerivedFrom<InFilterExpressionFactoryType, IAvaFilterExpressionFactory>::IsDerived>::Type>
+	static TSharedRef<InFilterExpressionFactoryType> MakeInstance(InArgsType&&... InArgs)
+	{
+		return MakeShared<InFilterExpressionFactoryType>(Forward<InArgsType>(InArgs)...);
+	}
+
+	/**
+	 * Evaluate the expression with the Value given
+	 * @param InItem Current Item being checked
+	 * @param InArgs The argument containing the data to evaluate the expression see FAvaTextFilterArgs for more information
+	 * @return True if the expression evaluate to True, False otherwise
+	 */
+	virtual bool FilterExpression(const IAvaOutlinerItem& InItem, const FAvaTextFilterArgs& InArgs) const = 0;
+
+	/**
+	 * Whether the factory support a given Comparison Operation
+	 * @param InComparisonOperation Comparison Operation to check
+	 * @return True if Factory support passed Comparison Operation, False otherwise
+	 */
+	virtual bool SupportsComparisonOperation(const ETextFilterComparisonOperation& InComparisonOperation) const = 0;
+};
diff --git a/Engine/Plugins/Experimental/Avalanche/Source/AvalancheOutliner/Private/Filters/IAvaFilterSuggestionFactory.h b/Engine/Plugins/Experimental/Avalanche/Source/AvalancheOutliner/Private/Filters/IAvaFilterSuggestionFactory.h
new file mode 100644
--- /dev/null
+++ b/Engine/Plugins/Experimental/Avalanche/Source/AvalancheOutliner/Private/Filters/IAvaFilterSuggestionFactory.h
@@ -0,0 +1,101 @@
+// Copyright Epic Games, Inc. All Rights Reserved.
+
+#pragma once
+
+#include "AvaOutlinerDefines.h"
+#include "AvaType.h"
+
+struct FAssetSearchBoxSuggestion;
+
+enum class EAvaFilterSuggestionType : uint8
+{
+	None      = 0x00, /** None */
+	Generic   = 0x01, /** Executed without the need to cycle through outliner item (for example Colors in the Settings) */
+	ItemBased = 0x10, /** Executed for each item in the outliner, the item is passed down to get suggestion from it */
+	All       = 0x11  /** Executed for both ways */
+};
+
+ENUM_CLASS_FLAGS(EAvaFilterSuggestionType)
+
+/**
+ * Holds the arguments required for the suggestion factory to get the suggestion.
+ */
+struct FAvaFilterSuggestionPayload : public IAvaTypeCastable, public TSharedFromThis<FAvaFilterSuggestionPayload>
+{
+public:
+	UE_AVA_INHERITS(FAvaFilterSuggestionPayload, IAvaTypeCastable)
+
+	FAvaFilterSuggestionPayload(TArray<FAssetSearchBoxSuggestion>& OutPossibleSuggestions
+		, const FString& InFilterValue)
+		: OutPossibleSuggestions(OutPossibleSuggestions)
+		, FilterValue(InFilterValue)
+	{
+	}
+
+public:
+	/** Holds all the suggestion that will be shown */
+	TArray<FAssetSearchBoxSuggestion>& OutPossibleSuggestions;
+
+	/** The current string value written in the searchbox */
+	FString FilterValue;
+};
+
+/**
+ * Specialization for the ItemBased suggestion factory that needs more arguments
+ */
+struct FAvaFilterSuggestionItemPayload : public FAvaFilterSuggestionPayload
+{
+public:
+	UE_AVA_INHERITS(FAvaFilterSuggestionItemPayload, FAvaFilterSuggestionPayload)
+
+	FAvaFilterSuggestionItemPayload(TArray<FAssetSearchBoxSuggestion>& OutPossibleSuggestions
+		, const FString& InFilterValue
+		, const FAvaOutlinerItemPtr& InItem
+		, TSet<FString>& OutFilterCache)
+		: FAvaFilterSuggestionPayload(OutPossibleSuggestions, InFilterValue)
+		, Item(InItem)
+		, OutFilterCache(OutFilterCache)
+	{
+	}
+
+public:
+	/** Current item to check */
+	FAvaOutlinerItemPtr Item;
+
+	/** Used to speed up the check on whether a suggestion is already added */
+	TSet<FString>& OutFilterCache;
+};
+
+class IAvaFilterSuggestionFactory : public TSharedFromThis<IAvaFilterSuggestionFactory>
+{
+public:
+	virtual ~IAvaFilterSuggestionFactory() {}
+
+	/**
+	  * Create and return an Instance of the AvaFilterSuggestionFactory Requested
+	  * @tparam InSuggestionFactoryType The type of the factory to instantiate
+	  * @param InArgs Additional Args for constructor of the AvaFilterSuggestionFactory class if needed
+	  * @return The new suggestion factory
+	  */
+	template<typename InSuggestionFactoryType, typename... InArgsType
+		, typename = typename TEnableIf<TIsDerivedFrom<InSuggestionFactoryType, IAvaFilterSuggestionFactory>::IsDerived>::Type>
+	static TSharedRef<InSuggestionFactoryType> MakeInstance(InArgsType&&... InArgs)
+	{
+		return MakeShared<InSuggestionFactoryType>(Forward<InArgsType>(InArgs)...);
+	}
+
+	/**
+	 * Get the type of the suggestion
+	 * @return The type of the suggestion see EAvaFilterSuggestionType for more information
+	 */
+	virtual EAvaFilterSuggestionType GetSuggestionType() const = 0;
+
+	/** Get the suggestion identifier for this factory */
+	virtual FName GetSuggestionIdentifier() const = 0;
+
+	/**
+	 * Add suggestions entry to be shown
+	 * @param InPayload Arguments to be passed to the factory to get the suggestion to show
+	 */
+	virtual void AddSuggestion(const TSharedRef<FAvaFilterSuggestionPayload> InPayload) = 0;
+};
diff --git a/Engine/Plugins/Experimental/Avalanche/Source/AvalancheOutliner/Private/Item/AvaOutlinerActor.cpp b/Engine/Plugins/Experimental/Avalanche/Source/AvalancheOutliner/Private/Item/AvaOutlinerActor.cpp
--- a/Engine/Plugins/Experimental/Avalanche/Source/AvalancheOutliner/Private/Item/AvaOutlinerActor.cpp
+++ b/Engine/Plugins/Experimental/Avalanche/Source/AvalancheOutliner/Private/Item/AvaOutlinerActor.cpp
@@ -17,7 +17,7 @@
 
 #define LOCTEXT_NAMESPACE "AvaOutlinerActor"
 
-FAvaOutlinerActor::FAvaOutlinerActor(FAvaOutliner& InOutliner, AActor* InActor)
+FAvaOutlinerActor::FAvaOutlinerActor(IAvaOutliner& InOutliner, AActor* InActor)
 	: Super(InOutliner, InActor)
 	, Actor(InActor)
 {
@@ -33,7 +33,8 @@ void FAvaOutlinerActor::FindChildren(TArray<FAvaOutlinerItemPtr>& OutChildren, b
 		return;
 	}
 
-	const TArray<TWeakObjectPtr<AActor>> OutlinerChildren = Outliner.GetActorSceneOutlinerChildren(UnderlyingActor);
+	FAvaOutliner& OutlinerPrivate = static_cast<FAvaOutliner&>(Outliner);
+	const TArray<TWeakObjectPtr<AActor>> OutlinerChildren = OutlinerPrivate.GetActorSceneOutlinerChildren(UnderlyingActor);
 
 	// Worst case and most likely case: all Outliner Children are valid.
 	// Note: if recursive, re-allocations will still need to be done past this reserve count, as it's unknown at this point how many items are going to be added
@@ -107,8 +108,9 @@ bool FAvaOutlinerActor::AddChild(const FAvaOutlinerAddItemParams& InAddItemParam
 				{
 					ChangeType = EAvaOutlinerHierarchyChangeType::Rearranged;
 				}
-				
-				if (UAvaOutlinerSubsystem* const OutlinerSubsystem = Outliner.GetOutlinerSubsystem())
+
+				FAvaOutliner& OutlinerPrivate = static_cast<FAvaOutliner&>(Outliner);
+				if (UAvaOutlinerSubsystem* const OutlinerSubsystem = OutlinerPrivate.GetOutlinerSubsystem())
 				{
 					OutlinerSubsystem->BroadcastActorHierarchyChanged(ChildActor, ParentActor, ChangeType);
 				}
@@ -141,8 +143,9 @@ bool FAvaOutlinerActor::RemoveChild(const FAvaOutlinerRemoveItemParams& InRemove
 					{
 						OldParentActor->Modify(false);
 						ChildRoot->DetachFromComponent(InRemoveItemParams.DetachmentTransformRules.Get(FDetachmentTransformRules::KeepWorldTransform));
-						
-						if (UAvaOutlinerSubsystem* const OutlinerSubsystem = Outliner.GetOutlinerSubsystem())
+
+						FAvaOutliner& OutlinerPrivate = static_cast<FAvaOutliner&>(Outliner);
+						if (UAvaOutlinerSubsystem* const OutlinerSubsystem = OutlinerPrivate.GetOutlinerSubsystem())
 						{
 							OutlinerSubsystem->BroadcastActorHierarchyChanged(ChildActor, ParentActor
 								, EAvaOutlinerHierarchyChangeType::Detached);
@@ -158,7 +161,8 @@ bool FAvaOutlinerActor::RemoveChild(const FAvaOutlinerRemoveItemParams& InRemove
 
 bool FAvaOutlinerActor::IsAllowedInOutliner() const
 {
-	return Outliner.IsActorAllowedInOutliner(GetActor());
+	FAvaOutliner& OutlinerPrivate = static_cast<FAvaOutliner&>(Outliner);
+	return OutlinerPrivate.IsActorAllowedInOutliner(GetActor());
 }
 
 EAvaOutlinerItemViewMode FAvaOutlinerActor::GetSupportedViewModes(const FAvaOutlinerView& InOutlinerView) const
diff --git a/Engine/Plugins/Experimental/Avalanche/Source/AvalancheOutliner/Private/Item/AvaOutlinerComponent.cpp b/Engine/Plugins/Experimental/Avalanche/Source/AvalancheOutliner/Private/Item/AvaOutlinerComponent.cpp
--- a/Engine/Plugins/Experimental/Avalanche/Source/AvalancheOutliner/Private/Item/AvaOutlinerComponent.cpp
+++ b/Engine/Plugins/Experimental/Avalanche/Source/AvalancheOutliner/Private/Item/AvaOutlinerComponent.cpp
@@ -14,7 +14,7 @@
 
 #define LOCTEXT_NAMESPACE "AvaOutlinerComponent"
 
-FAvaOutlinerComponent::FAvaOutlinerComponent(FAvaOutliner& InOutliner, USceneComponent* InComponent)
+FAvaOutlinerComponent::FAvaOutlinerComponent(IAvaOutliner& InOutliner, USceneComponent* InComponent)
 	: Super(InOutliner, InComponent)
 	, Component(InComponent)
 {
@@ -122,13 +122,15 @@ EAvaOutlinerItemViewMode FAvaOutlinerComponent::GetSupportedViewModes(const FAva
 
 bool FAvaOutlinerComponent::IsAllowedInOutliner() const
 {
+	FAvaOutliner& OutlinerPrivate = static_cast<FAvaOutliner&>(Outliner);
+
 	//Make sure Owner is Allowed (this also returns false if Owner is invalid)
 	const USceneComponent* const UnderlyingComponent = GetComponent();
-	if (!UnderlyingComponent || !Outliner.IsActorAllowedInOutliner(UnderlyingComponent->GetOwner()))
+	if (!UnderlyingComponent || !OutlinerPrivate.IsActorAllowedInOutliner(UnderlyingComponent->GetOwner()))
 	{
 		return false;
 	}
-	return Outliner.IsComponentAllowedInOutliner(UnderlyingComponent);
+	return OutlinerPrivate.IsComponentAllowedInOutliner(UnderlyingComponent);
 }
 
 TSharedRef<SWidget> FAvaOutlinerComponent::GenerateLabelWidget(const TSharedRef<SAvaOutlinerTreeRow>& InRow)
diff --git a/Engine/Plugins/Experimental/Avalanche/Source/AvalancheOutliner/Private/Item/AvaOutlinerComponentProxy.cpp b/Engine/Plugins/Experimental/Avalanche/Source/AvalancheOutliner/Private/Item/AvaOutlinerComponentProxy.cpp
--- a/Engine/Plugins/Experimental/Avalanche/Source/AvalancheOutliner/Private/Item/AvaOutlinerComponentProxy.cpp
+++ b/Engine/Plugins/Experimental/Avalanche/Source/AvalancheOutliner/Private/Item/AvaOutlinerComponentProxy.cpp
@@ -10,7 +10,7 @@
 
 #define LOCTEXT_NAMESPACE "AvaOutlinerComponentProxy"
 
-FAvaOutlinerComponentProxy::FAvaOutlinerComponentProxy(FAvaOutliner& InOutliner, const FAvaOutlinerItemPtr& InParentItem)
+FAvaOutlinerComponentProxy::FAvaOutlinerComponentProxy(IAvaOutliner& InOutliner, const FAvaOutlinerItemPtr& InParentItem)
 	: Super(InOutliner, InParentItem)
 {
 }
diff --git a/Engine/Plugins/Experimental/Avalanche/Source/AvalancheOutliner/Private/Item/AvaOutlinerItem.cpp b/Engine/Plugins/Experimental/Avalanche/Source/AvalancheOutliner/Private/Item/AvaOutlinerItem.cpp
--- a/Engine/Plugins/Experimental/Avalanche/Source/AvalancheOutliner/Private/Item/AvaOutlinerItem.cpp
+++ b/Engine/Plugins/Experimental/Avalanche/Source/AvalancheOutliner/Private/Item/AvaOutlinerItem.cpp
@@ -1,10 +1,10 @@
 // Copyright Epic Games, Inc. All Rights Reserved.
 
 #include "Item/AvaOutlinerItem.h"
-
 #include "ActorFactories/ActorFactory.h"
 #include "AssetSelection.h"
 #include "AvaOutliner.h"
+#include "AvaOutlinerModule.h"
 #include "AvaOutlinerView.h"
 #include "AvaSceneTree.h"
 #include "Columns/Slate/SAvaOutlinerLabelItem.h"
@@ -12,8 +12,8 @@
 #include "DragAndDrop/AssetDragDropOp.h"
 #include "DragDropOps/AvaOutlinerItemDragDropOp.h"
 #include "Editor.h"
-#include "IAvaOutlinerProvider.h"
 #include "Framework/Application/SlateApplication.h"
+#include "IAvaOutlinerProvider.h"
 #include "Item/AvaOutlinerActor.h"
 #include "Item/AvaOutlinerItemUtils.h"
 #include "ItemActions/AvaOutlinerAddItem.h"
@@ -23,7 +23,7 @@
 
 #define LOCTEXT_NAMESPACE "AvaOutlinerItem"
 
-namespace UE::AvalancheOutliner::Private
+namespace UE::AvaOutliner::Private
 {
 	EDetachmentRule GetDetachmentRuleFromAttachmentRule(EAttachmentRule AttachmentRule)
 	{
@@ -53,7 +53,7 @@ namespace UE::AvalancheOutliner::Private
 	};
 }
 
-FAvaOutlinerItem::FAvaOutlinerItem(FAvaOutliner& InOutliner)
+FAvaOutlinerItem::FAvaOutlinerItem(IAvaOutliner& InOutliner)
 	: Outliner(InOutliner)
 {
 }
@@ -62,7 +62,7 @@ FAvaOutlinerItem::~FAvaOutlinerItem()
 {
 }
 
-TSharedRef<FAvaOutliner> FAvaOutlinerItem::GetOwnerOutliner() const
+TSharedRef<IAvaOutliner> FAvaOutlinerItem::GetOwnerOutliner() const
 {
 	return Outliner.AsShared();
 }
@@ -79,7 +79,7 @@ void FAvaOutlinerItem::RefreshChildren()
 
 	TArray<FAvaOutlinerItemPtr> Sortable;
 	TArray<FAvaOutlinerItemPtr> Unsortable;
-	UE::AvalancheOutliner::SplitItems(FoundChildren, Sortable, Unsortable);
+	UE::AvaOutliner::SplitItems(FoundChildren, Sortable, Unsortable);
 
 	// Start with all Sortable/Unsortable Items, and remove every item seen by iterating Children
 	TSet<FAvaOutlinerItemPtr> NewSortableChildren(Sortable);
@@ -142,10 +142,12 @@ void FAvaOutlinerItem::RefreshChildren()
 
 void FAvaOutlinerItem::FindChildren(TArray<FAvaOutlinerItemPtr>& OutChildren, bool bRecursive)
 {
+	FAvaOutliner& OutlinerPrivate = static_cast<FAvaOutliner&>(Outliner);
+
 	TArray<TSharedPtr<FAvaOutlinerItemProxy>> ItemProxies;
-	Outliner.GetItemProxiesForItem(SharedThis(this), ItemProxies);
+	OutlinerPrivate.GetItemProxiesForItem(SharedThis(this), ItemProxies);
 	OutChildren.Reserve(OutChildren.Num() + ItemProxies.Num());
-	
+
 	for (const TSharedPtr<FAvaOutlinerItemProxy>& ItemProxy : ItemProxies)
 	{
 		OutChildren.Add(ItemProxy);
@@ -203,7 +205,7 @@ FAvaOutlinerItemId FAvaOutlinerItem::GetItemId() const
 
 const FSlateBrush* FAvaOutlinerItem::GetIconBrush() const
 {
-	const FSlateIcon Icon = IAvaOutlinerModule::Get().FindOverrideIcon(SharedThis(this));
+	const FSlateIcon Icon = FAvaOutlinerModule::Get().FindOverrideIcon(SharedThis(this));
 	if (Icon.IsSet())
 	{
 		return Icon.GetIcon();
@@ -250,7 +252,8 @@ bool FAvaOutlinerItem::HasAllFlags(EAvaOutlinerItemFlags Flags) const
 
 TOptional<FAvaOutlinerColorPair> FAvaOutlinerItem::GetColor(bool bRecurse) const
 {
-	return Outliner.FindItemColor(SharedThis(const_cast<FAvaOutlinerItem*>(this)), bRecurse);
+	FAvaOutliner& OutlinerPrivate = static_cast<FAvaOutliner&>(Outliner);
+	return OutlinerPrivate.FindItemColor(SharedThis(const_cast<FAvaOutlinerItem*>(this)), bRecurse);
 }
 
 TOptional<EItemDropZone> FAvaOutlinerItem::CanAcceptDrop(const FDragDropEvent& InDragDropEvent, EItemDropZone InDropZone)
@@ -288,7 +291,9 @@ void FAvaOutlinerItem::RecalculateItemId()
 {
 	const FAvaOutlinerItemId OldItemId = ItemId;
 	ItemId = CalculateItemId();
-	Outliner.NotifyItemIdChanged(OldItemId, SharedThis(this));
+
+	FAvaOutliner& OutlinerPrivate = static_cast<FAvaOutliner&>(Outliner);
+	OutlinerPrivate.NotifyItemIdChanged(OldItemId, SharedThis(this));
 }
 
 void FAvaOutlinerItem::AddChildChecked(const FAvaOutlinerAddItemParams& InAddItemParams)
@@ -304,7 +309,7 @@ void FAvaOutlinerItem::AddChildChecked(const FAvaOutlinerAddItemParams& InAddIte
 		else
 		{
 			FAvaOutlinerRemoveItemParams RemoveParams(InAddItemParams.Item);
-			RemoveParams.DetachmentTransformRules = UE::AvalancheOutliner::Private::GetDetachmentRulesFromAttachmentRules(InAddItemParams.AttachmentTransformRules);
+			RemoveParams.DetachmentTransformRules = UE::AvaOutliner::Private::GetDetachmentRulesFromAttachmentRules(InAddItemParams.AttachmentTransformRules);
 
 			OldParent->RemoveChild(RemoveParams);
 		}
@@ -426,7 +431,7 @@ FReply FAvaOutlinerItem::CreateItemsFromAssetDrop(const TSharedPtr<FAssetDragDro
 
 				if (AActor* const SpawnedActor = ActorFactory->CreateActor(Asset.GetAsset()
 					, Level
-					, Outliner.GetActorDefaultSpawnTransform()
+					, Outliner.GetProvider().GetOutlinerDefaultActorSpawnTransform()
 					, SpawnParams))
 				{
 					++ItemsAdded;
@@ -443,7 +448,9 @@ FReply FAvaOutlinerItem::CreateItemsFromAssetDrop(const TSharedPtr<FAssetDragDro
 						// When there is more than 1 item being added, start appending to Selection Instead
 						AddChildParams.SelectionFlags |= EAvaOutlinerItemSelectionFlags::AppendToCurrentSelection;
 					}
-					Outliner.EnqueueItemAction<FAvaOutlinerAddItem>(AddChildParams);
+
+					FAvaOutliner& OutlinerPrivate = static_cast<FAvaOutliner&>(Outliner);
+					OutlinerPrivate.EnqueueItemAction<FAvaOutlinerAddItem>(AddChildParams);
 				}
 
 				Outliner.SetIgnoreNotify(EAvaOutlinerIgnoreNotifyFlags::Spawn, false);
diff --git a/Engine/Plugins/Experimental/Avalanche/Source/AvalancheOutliner/Private/Item/AvaOutlinerItemProxy.cpp b/Engine/Plugins/Experimental/Avalanche/Source/AvalancheOutliner/Private/Item/AvaOutlinerItemProxy.cpp
--- a/Engine/Plugins/Experimental/Avalanche/Source/AvalancheOutliner/Private/Item/AvaOutlinerItemProxy.cpp
+++ b/Engine/Plugins/Experimental/Avalanche/Source/AvalancheOutliner/Private/Item/AvaOutlinerItemProxy.cpp
@@ -5,7 +5,7 @@
 #include "AvaOutlinerView.h"
 #include "Input/Reply.h"
 
-FAvaOutlinerItemProxy::FAvaOutlinerItemProxy(FAvaOutliner& InOutliner, const FAvaOutlinerItemPtr& InParentItem)
+FAvaOutlinerItemProxy::FAvaOutlinerItemProxy(IAvaOutliner& InOutliner, const FAvaOutlinerItemPtr& InParentItem)
 	: FAvaOutlinerItem(InOutliner)
 {
 	ParentWeak = InParentItem;
diff --git a/Engine/Plugins/Experimental/Avalanche/Source/AvalancheOutliner/Private/Item/AvaOutlinerItemUtils.cpp b/Engine/Plugins/Experimental/Avalanche/Source/AvalancheOutliner/Private/Item/AvaOutlinerItemUtils.cpp
--- a/Engine/Plugins/Experimental/Avalanche/Source/AvalancheOutliner/Private/Item/AvaOutlinerItemUtils.cpp
+++ b/Engine/Plugins/Experimental/Avalanche/Source/AvalancheOutliner/Private/Item/AvaOutlinerItemUtils.cpp
@@ -1,10 +1,51 @@
 // Copyright Epic Games, Inc. All Rights Reserved.
 
 #include "Item/AvaOutlinerItemUtils.h"
+#include "AvaOutliner.h"
 #include "Containers/Array.h"
 #include "Item/IAvaOutlinerItem.h"
 
-void UE::AvalancheOutliner::SplitItems(const TArray<FAvaOutlinerItemPtr>& InItems
+bool UE::AvaOutliner::CompareOutlinerItemOrder(const FAvaOutlinerItemPtr& A, const FAvaOutlinerItemPtr& B)
+{
+	if (!A.IsValid())
+	{
+		return false;
+	}
+	if (!B.IsValid())
+	{
+		return false;
+	}
+	if (const FAvaOutlinerItemPtr LowestCommonAncestor = FAvaOutliner::FindLowestCommonAncestor({A, B}))
+	{
+		const TArray<FAvaOutlinerItemPtr> PathToA = LowestCommonAncestor->FindPath({A});
+		const TArray<FAvaOutlinerItemPtr> PathToB = LowestCommonAncestor->FindPath({B});
+
+		int32 Index = 0;
+		
+		int32 PathAIndex = -1;
+		int32 PathBIndex = -1;
+
+		while (PathAIndex == PathBIndex)
+		{
+			if (!PathToA.IsValidIndex(Index))
+			{
+				return true;
+			}
+			if (!PathToB.IsValidIndex(Index))
+			{
+				return false;
+			}
+
+			PathAIndex = LowestCommonAncestor->GetChildIndex(PathToA[Index]);
+			PathBIndex = LowestCommonAncestor->GetChildIndex(PathToB[Index]);
+			++Index;
+		}
+		return PathAIndex < PathBIndex;
+	}
+	return false;
+}
+
+void UE::AvaOutliner::SplitItems(const TArray<FAvaOutlinerItemPtr>& InItems
 	, TArray<FAvaOutlinerItemPtr>& OutSortable
 	, TArray<FAvaOutlinerItemPtr>& OutUnsortable)
 {
diff --git a/Engine/Plugins/Experimental/Avalanche/Source/AvalancheOutliner/Private/Item/AvaOutlinerMaterial.cpp b/Engine/Plugins/Experimental/Avalanche/Source/AvalancheOutliner/Private/Item/AvaOutlinerMaterial.cpp
--- a/Engine/Plugins/Experimental/Avalanche/Source/AvalancheOutliner/Private/Item/AvaOutlinerMaterial.cpp
+++ b/Engine/Plugins/Experimental/Avalanche/Source/AvalancheOutliner/Private/Item/AvaOutlinerMaterial.cpp
@@ -12,7 +12,7 @@
 
 #define LOCTEXT_NAMESPACE "AvaOutlinerMaterial"
 
-FAvaOutlinerMaterial::FAvaOutlinerMaterial(FAvaOutliner& InOutliner
+FAvaOutlinerMaterial::FAvaOutlinerMaterial(IAvaOutliner& InOutliner
 		, UMaterialInterface* InMaterial
 		, const FAvaOutlinerItemPtr& InReferencingItem
 		, int32 InMaterialIndex)
diff --git a/Engine/Plugins/Experimental/Avalanche/Source/AvalancheOutliner/Private/Item/AvaOutlinerMaterialProxy.cpp b/Engine/Plugins/Experimental/Avalanche/Source/AvalancheOutliner/Private/Item/AvaOutlinerMaterialProxy.cpp
--- a/Engine/Plugins/Experimental/Avalanche/Source/AvalancheOutliner/Private/Item/AvaOutlinerMaterialProxy.cpp
+++ b/Engine/Plugins/Experimental/Avalanche/Source/AvalancheOutliner/Private/Item/AvaOutlinerMaterialProxy.cpp
@@ -10,7 +10,7 @@
 
 #define LOCTEXT_NAMESPACE "AvaOutlinerMaterialProxy"
 
-FAvaOutlinerMaterialProxy::FAvaOutlinerMaterialProxy(FAvaOutliner& InOutliner, const FAvaOutlinerItemPtr& InParentItem)
+FAvaOutlinerMaterialProxy::FAvaOutlinerMaterialProxy(IAvaOutliner& InOutliner, const FAvaOutlinerItemPtr& InParentItem)
 	: FAvaOutlinerItemProxy(InOutliner, InParentItem)
 {
 }
diff --git a/Engine/Plugins/Experimental/Avalanche/Source/AvalancheOutliner/Private/Item/AvaOutlinerObject.cpp b/Engine/Plugins/Experimental/Avalanche/Source/AvalancheOutliner/Private/Item/AvaOutlinerObject.cpp
--- a/Engine/Plugins/Experimental/Avalanche/Source/AvalancheOutliner/Private/Item/AvaOutlinerObject.cpp
+++ b/Engine/Plugins/Experimental/Avalanche/Source/AvalancheOutliner/Private/Item/AvaOutlinerObject.cpp
@@ -8,7 +8,7 @@
 
 #define LOCTEXT_NAMESPACE "AvaOutlinerObject"
 
-FAvaOutlinerObject::FAvaOutlinerObject(FAvaOutliner& InOutliner, UObject* InObject)
+FAvaOutlinerObject::FAvaOutlinerObject(IAvaOutliner& InOutliner, UObject* InObject)
 	: FAvaOutlinerItem(InOutliner)
 	, Object(InObject)
 {
diff --git a/Engine/Plugins/Experimental/Avalanche/Source/AvalancheOutliner/Private/Item/AvaOutlinerObjectReference.cpp b/Engine/Plugins/Experimental/Avalanche/Source/AvalancheOutliner/Private/Item/AvaOutlinerObjectReference.cpp
--- a/Engine/Plugins/Experimental/Avalanche/Source/AvalancheOutliner/Private/Item/AvaOutlinerObjectReference.cpp
+++ b/Engine/Plugins/Experimental/Avalanche/Source/AvalancheOutliner/Private/Item/AvaOutlinerObjectReference.cpp
@@ -8,7 +8,7 @@
 
 #define LOCTEXT_NAMESPACE "AvaOutlinerObjectReference"
 
-FAvaOutlinerSharedObject::FAvaOutlinerSharedObject(FAvaOutliner& InOutliner, UObject* InObject)
+FAvaOutlinerSharedObject::FAvaOutlinerSharedObject(IAvaOutliner& InOutliner, UObject* InObject)
 	: Super(InOutliner, InObject)
 {
 }
@@ -42,7 +42,7 @@ TArray<FAvaOutlinerItemPtr> FAvaOutlinerSharedObject::GetObjectReferences() cons
 	return OutObjectReferences;
 }
 
-FAvaOutlinerObjectReference::FAvaOutlinerObjectReference(FAvaOutliner& InOutliner
+FAvaOutlinerObjectReference::FAvaOutlinerObjectReference(IAvaOutliner& InOutliner
 		, UObject* InObject
 		, const FAvaOutlinerItemPtr& InReferencingItem
 		, const FString& InReferenceId)
diff --git a/Engine/Plugins/Experimental/Avalanche/Source/AvalancheOutliner/Private/Item/AvaOutlinerTreeRoot.cpp b/Engine/Plugins/Experimental/Avalanche/Source/AvalancheOutliner/Private/Item/AvaOutlinerTreeRoot.cpp
--- a/Engine/Plugins/Experimental/Avalanche/Source/AvalancheOutliner/Private/Item/AvaOutlinerTreeRoot.cpp
+++ b/Engine/Plugins/Experimental/Avalanche/Source/AvalancheOutliner/Private/Item/AvaOutlinerTreeRoot.cpp
@@ -25,8 +25,10 @@ void FAvaOutlinerTreeRoot::FindChildren(TArray<FAvaOutlinerItemPtr>& OutChildren
 		return;
 	}
 
+	FAvaOutliner& OutlinerPrivate = static_cast<FAvaOutliner&>(Outliner);
+
 	// All the Actors that don't have an Outliner Parent will be at the Root
-	const TArray<TWeakObjectPtr<AActor>> RootChildren = Outliner.GetActorSceneOutlinerChildren(nullptr);
+	const TArray<TWeakObjectPtr<AActor>> RootChildren = OutlinerPrivate.GetActorSceneOutlinerChildren(nullptr);
 
 	// Worst case and most likely case: all Outliner Children are valid.
 	// Note: if recursive, re-allocations will still need to be done past this reserve count, as it's unknown at this point how many items are going to be added
@@ -60,12 +62,14 @@ bool FAvaOutlinerTreeRoot::AddChild(const FAvaOutlinerAddItemParams& InAddItemPa
 
 	// Is it a new actor that we just spawned
 	const bool bSpawning = InAddItemParams.Item.IsValid() && !Children.Contains(InAddItemParams.Item);
-	
+
 	const bool bResult = FAvaOutlinerItem::AddChild(InAddItemParams);
-	
+
 	if ((bSpawning || bRearranging) && InAddItemParams.Item->IsA<FAvaOutlinerActor>())
 	{
-		UAvaOutlinerSubsystem* const OutlinerSubsystem = Outliner.GetOutlinerSubsystem();
+		FAvaOutliner& OutlinerPrivate = static_cast<FAvaOutliner&>(Outliner);
+
+		UAvaOutlinerSubsystem* const OutlinerSubsystem = OutlinerPrivate.GetOutlinerSubsystem();
 		
 		AActor* const Actor = StaticCastSharedPtr<FAvaOutlinerActor>(InAddItemParams.Item)->GetActor();
 		
diff --git a/Engine/Plugins/Experimental/Avalanche/Source/AvalancheOutliner/Private/ItemActions/AvaOutlinerAddItem.h b/Engine/Plugins/Experimental/Avalanche/Source/AvalancheOutliner/Private/ItemActions/AvaOutlinerAddItem.h
new file mode 100644
--- /dev/null
+++ b/Engine/Plugins/Experimental/Avalanche/Source/AvalancheOutliner/Private/ItemActions/AvaOutlinerAddItem.h
@@ -0,0 +1,27 @@
+// Copyright Epic Games, Inc. All Rights Reserved.
+
+#pragma once
+
+#include "AvaOutlinerItemAction.h"
+#include "Item/AvaOutlinerItemParameters.h"
+
+/**
+ * Item Action responsible of adding an item to the Tree under a given optional Parent.
+ * If Parent is null, it is added as a Top Level Item.
+ */
+class FAvaOutlinerAddItem : public IAvaOutlinerAction
+{
+public:
+	UE_AVA_INHERITS(FAvaOutlinerAddItem, IAvaOutlinerAction);
+
+	FAvaOutlinerAddItem(const FAvaOutlinerAddItemParams& InAddItemParams);
+
+	//~ Begin IAvaOutlinerAction
+	virtual bool ShouldTransact() const override;
+	virtual void Execute(FAvaOutliner& InOutliner) override;
+	virtual void OnObjectsReplaced(const TMap<UObject*, UObject*>& InReplacementMap, bool bRecursive) override;
+	//~ End IAvaOutlinerAction
+
+protected:
+	FAvaOutlinerAddItemParams AddParams;
+};
diff --git a/Engine/Plugins/Experimental/Avalanche/Source/AvalancheOutliner/Private/ItemActions/AvaOutlinerItemAction.h b/Engine/Plugins/Experimental/Avalanche/Source/AvalancheOutliner/Private/ItemActions/AvaOutlinerItemAction.h
new file mode 100644
--- /dev/null
+++ b/Engine/Plugins/Experimental/Avalanche/Source/AvalancheOutliner/Private/ItemActions/AvaOutlinerItemAction.h
@@ -0,0 +1,27 @@
+// Copyright Epic Games, Inc. All Rights Reserved.
+
+#pragma once
+
+#include "AvaOutlinerDefines.h"
+#include "AvaType.h"
+
+class FAvaOutliner;
+class UObject;
+
+/**
+ * Interface class for an Action in the Outliner (e.g. Add/Delete/Move Tree Item)
+ */ 
+class IAvaOutlinerAction : public IAvaTypeCastable
+{
+public:
+	UE_AVA_INHERITS(IAvaOutlinerAction, IAvaTypeCastable);
+
+	/** Determines whether the given action modifies its objects and should transact */
+	virtual bool ShouldTransact() const { return false; }
+
+	/** The Action to execute on the given Outliner */
+	virtual void Execute(FAvaOutliner& InOutliner) = 0;
+
+	/** Replace any Objects that might be held in this Action that has been killed and replaced by a new object (e.g. BP Components) */
+	virtual void OnObjectsReplaced(const TMap<UObject*, UObject*>& InReplacementMap, bool bRecursive) = 0;
+};
diff --git a/Engine/Plugins/Experimental/Avalanche/Source/AvalancheOutliner/Private/ItemActions/AvaOutlinerRemoveItem.h b/Engine/Plugins/Experimental/Avalanche/Source/AvalancheOutliner/Private/ItemActions/AvaOutlinerRemoveItem.h
new file mode 100644
--- /dev/null
+++ b/Engine/Plugins/Experimental/Avalanche/Source/AvalancheOutliner/Private/ItemActions/AvaOutlinerRemoveItem.h
@@ -0,0 +1,25 @@
+// Copyright Epic Games, Inc. All Rights Reserved.
+
+#pragma once
+
+#include "AvaOutlinerItemAction.h"
+#include "Item/AvaOutlinerItemParameters.h"
+
+/**
+ * Item Action responsible of removing/unregistering Items from the Tree
+ */
+class FAvaOutlinerRemoveItem : public IAvaOutlinerAction
+{
+public:
+	UE_AVA_INHERITS(FAvaOutlinerRemoveItem, IAvaOutlinerAction);
+
+	FAvaOutlinerRemoveItem(const FAvaOutlinerRemoveItemParams& InRemoveItemParams);
+
+	//~ Begin IAvaOutlinerAction
+	virtual void Execute(FAvaOutliner& InOutliner) override;
+	virtual void OnObjectsReplaced(const TMap<UObject*, UObject*>& InReplacementMap, bool bRecursive) override;
+	//~ End IAvaOutlinerAction
+
+protected:
+	FAvaOutlinerRemoveItemParams RemoveParams;
+};
diff --git a/Engine/Plugins/Experimental/Avalanche/Source/AvalancheOutliner/Private/Menu/AvaOutlinerViewToolbarContext.h b/Engine/Plugins/Experimental/Avalanche/Source/AvalancheOutliner/Private/Menu/AvaOutlinerViewToolbarContext.h
new file mode 100644
--- /dev/null
+++ b/Engine/Plugins/Experimental/Avalanche/Source/AvalancheOutliner/Private/Menu/AvaOutlinerViewToolbarContext.h
@@ -0,0 +1,25 @@
+// Copyright Epic Games, Inc. All Rights Reserved.
+
+#pragma once
+
+#include "Templates/SharedPointer.h"
+#include "UObject/Object.h"
+#include "AvaOutlinerViewToolbarContext.generated.h"
+
+class FAvaOutlinerView;
+
+UCLASS()
+class UAvaOutlinerViewToolbarContext : public UObject
+{
+	GENERATED_BODY()
+
+	friend FAvaOutlinerView;
+
+public:
+	UAvaOutlinerViewToolbarContext() = default;
+
+	TSharedPtr<FAvaOutlinerView> GetOutlinerView() const { return OutlinerViewWeak.Pin(); }
+
+private:
+	TWeakPtr<FAvaOutlinerView> OutlinerViewWeak;
+};
diff --git a/Engine/Plugins/Experimental/Avalanche/Source/AvalancheOutliner/Private/Selection/AvaOutlinerScopedSelection.cpp b/Engine/Plugins/Experimental/Avalanche/Source/AvalancheOutliner/Private/Selection/AvaOutlinerScopedSelection.cpp
--- a/Engine/Plugins/Experimental/Avalanche/Source/AvalancheOutliner/Private/Selection/AvaOutlinerScopedSelection.cpp
+++ b/Engine/Plugins/Experimental/Avalanche/Source/AvalancheOutliner/Private/Selection/AvaOutlinerScopedSelection.cpp
@@ -6,7 +6,7 @@
 #include "GameFramework/Actor.h"
 #include "Selection.h"
 
-namespace UE::AvalancheOutliner::Private
+namespace UE::AvaOutliner::Private
 {
 	USelection* GetSelection(const FEditorModeTools& InEditorModeTools, const UObject* InObject)
 	{
@@ -122,7 +122,7 @@ bool FAvaOutlinerScopedSelection::IsSelected(const UObject* InObject) const
 		return true;
 	}
 
-	USelection* const Selection = UE::AvalancheOutliner::Private::GetSelection(EditorModeTools, InObject);
+	USelection* const Selection = UE::AvaOutliner::Private::GetSelection(EditorModeTools, InObject);
 	if (ensure(Selection))
 	{
 		return Selection->IsSelected(InObject);
@@ -132,7 +132,7 @@ bool FAvaOutlinerScopedSelection::IsSelected(const UObject* InObject) const
 
 void FAvaOutlinerScopedSelection::SyncSelections()
 {
-	UE::AvalancheOutliner::Private::SyncSelection(EditorModeTools.GetSelectedActors(), SelectedActors);
-	UE::AvalancheOutliner::Private::SyncSelection(EditorModeTools.GetSelectedComponents(), SelectedComponents);
-	UE::AvalancheOutliner::Private::SyncSelection(EditorModeTools.GetSelectedObjects(), SelectedObjects);
+	UE::AvaOutliner::Private::SyncSelection(EditorModeTools.GetSelectedActors(), SelectedActors);
+	UE::AvaOutliner::Private::SyncSelection(EditorModeTools.GetSelectedComponents(), SelectedComponents);
+	UE::AvaOutliner::Private::SyncSelection(EditorModeTools.GetSelectedObjects(), SelectedObjects);
 }
diff --git a/Engine/Plugins/Experimental/Avalanche/Source/AvalancheOutliner/Private/Slate/SAvaOutliner.cpp b/Engine/Plugins/Experimental/Avalanche/Source/AvalancheOutliner/Private/Slate/SAvaOutliner.cpp
--- a/Engine/Plugins/Experimental/Avalanche/Source/AvalancheOutliner/Private/Slate/SAvaOutliner.cpp
+++ b/Engine/Plugins/Experimental/Avalanche/Source/AvalancheOutliner/Private/Slate/SAvaOutliner.cpp
@@ -2,6 +2,7 @@
 
 #include "Slate/SAvaOutliner.h"
 #include "AvaOutliner.h"
+#include "AvaOutlinerModule.h"
 #include "AvaOutlinerSettings.h"
 #include "AvaOutlinerView.h"
 #include "Columns/IAvaOutlinerColumn.h"
@@ -448,10 +449,10 @@ void SAvaOutliner::OnAssetSearchSuggestionFilter(const FText& InSearchText, TArr
 	FString FilterValue;
 	UE::AvaOutliner::Private::ExtractAssetSearchFilterTerms(InSearchText, &FilterKey, &FilterValue, nullptr);
 
-	const IAvaOutlinerModule& AvaOutlinerModule = IAvaOutlinerModule::Get();
+	const FAvaOutlinerModule& OutlinerModule = FAvaOutlinerModule::Get();
 	const TSharedRef<FAvaFilterSuggestionPayload> GenericPayload(MakeShared<FAvaFilterSuggestionPayload>(OutPossibleSuggestions,FilterValue));
 
-	for (const TSharedPtr<IAvaFilterSuggestionFactory>& GenericSuggestion : AvaOutlinerModule.GetSuggestions(EAvaFilterSuggestionType::Generic))
+	for (const TSharedPtr<IAvaFilterSuggestionFactory>& GenericSuggestion : OutlinerModule.GetSuggestions(EAvaFilterSuggestionType::Generic))
 	{
 		GenericSuggestion->AddSuggestion(GenericPayload);
 	}
@@ -466,7 +467,7 @@ void SAvaOutliner::OnAssetSearchSuggestionFilter(const FText& InSearchText, TArr
 
 			for (const FAvaOutlinerItemPtr& RootChildFromRoot : OutlinerPtr->GetTreeRoot()->GetChildren())
 			{
-				for (const TSharedPtr<IAvaFilterSuggestionFactory>& ItemSuggestion : AvaOutlinerModule.GetSuggestions(EAvaFilterSuggestionType::ItemBased))
+				for (const TSharedPtr<IAvaFilterSuggestionFactory>& ItemSuggestion : OutlinerModule.GetSuggestions(EAvaFilterSuggestionType::ItemBased))
 				{
 					ItemBasedPayload->Item = RootChildFromRoot;
 					ItemSuggestion->AddSuggestion(ItemBasedPayload);
@@ -477,7 +478,7 @@ void SAvaOutliner::OnAssetSearchSuggestionFilter(const FText& InSearchText, TArr
 
 				for (const FAvaOutlinerItemPtr& Child : Children)
 				{
-					for (const TSharedPtr<IAvaFilterSuggestionFactory>& ItemSuggestion : AvaOutlinerModule.GetSuggestions(EAvaFilterSuggestionType::ItemBased))
+					for (const TSharedPtr<IAvaFilterSuggestionFactory>& ItemSuggestion : OutlinerModule.GetSuggestions(EAvaFilterSuggestionType::ItemBased))
 					{
 						ItemBasedPayload->Item = Child;
 						ItemSuggestion->AddSuggestion(ItemBasedPayload);
diff --git a/Engine/Plugins/Experimental/Avalanche/Source/AvalancheOutliner/Private/Slate/SAvaOutlinerItemFilters.cpp b/Engine/Plugins/Experimental/Avalanche/Source/AvalancheOutliner/Private/Slate/SAvaOutlinerItemFilters.cpp
--- a/Engine/Plugins/Experimental/Avalanche/Source/AvalancheOutliner/Private/Slate/SAvaOutlinerItemFilters.cpp
+++ b/Engine/Plugins/Experimental/Avalanche/Source/AvalancheOutliner/Private/Slate/SAvaOutlinerItemFilters.cpp
@@ -20,7 +20,7 @@
 
 #define LOCTEXT_NAMESPACE "SAvaOutlinerItemFilters"
 
-namespace UE::AvalancheOutliner::Private
+namespace UE::AvaOutliner::Private
 {
 	const FButtonStyle& GetFilterItemMenuButtonStyle()
 	{
@@ -46,7 +46,7 @@ namespace UE::AvalancheOutliner::Private
 
 void SAvaOutlinerItemFilters::Construct(const FArguments& InArgs, const TSharedRef<FAvaOutlinerView>& InOutlinerView)
 {
-	using namespace UE::AvalancheOutliner::Private;
+	using namespace UE::AvaOutliner::Private;
 
 	OutlinerViewWeak = InOutlinerView;
 
@@ -191,7 +191,7 @@ void SAvaOutlinerItemFilters::AddItemFilterSlot(const TSharedPtr<IAvaOutlinerIte
 	}
 
 	TSharedRef<SWidget> Slot = SNew(SCheckBox)
-		.Style(&UE::AvalancheOutliner::Private::GetItemFilterCheckboxStyle())
+		.Style(&UE::AvaOutliner::Private::GetItemFilterCheckboxStyle())
 		.ToolTipText(InFilter->GetTooltipText())
 		.OnCheckStateChanged(this, &SAvaOutlinerItemFilters::OnCheckBoxStateChanged, InFilter)
 		.IsChecked(this, &SAvaOutlinerItemFilters::IsChecked, InFilter)
diff --git a/Engine/Plugins/Experimental/Avalanche/Source/AvalancheOutliner/Private/Slate/Styling/AvaOutlinerStyleUtils.cpp b/Engine/Plugins/Experimental/Avalanche/Source/AvalancheOutliner/Private/Slate/Styling/AvaOutlinerStyleUtils.cpp
--- a/Engine/Plugins/Experimental/Avalanche/Source/AvalancheOutliner/Private/Slate/Styling/AvaOutlinerStyleUtils.cpp
+++ b/Engine/Plugins/Experimental/Avalanche/Source/AvalancheOutliner/Private/Slate/Styling/AvaOutlinerStyleUtils.cpp
@@ -7,7 +7,7 @@
 #include "Styling/SlateBrush.h"
 #include "Styling/ToolBarStyle.h"
 
-namespace UE::AvalancheOutliner::Private
+namespace UE::AvaOutliner::Private
 {
 	const FToolBarStyle& FStyleUtils::GetSlimToolBarStyle()
 	{
diff --git a/Engine/Plugins/Experimental/Avalanche/Source/AvalancheOutliner/Private/Slate/Styling/AvaOutlinerStyleUtils.h b/Engine/Plugins/Experimental/Avalanche/Source/AvalancheOutliner/Private/Slate/Styling/AvaOutlinerStyleUtils.h
--- a/Engine/Plugins/Experimental/Avalanche/Source/AvalancheOutliner/Private/Slate/Styling/AvaOutlinerStyleUtils.h
+++ b/Engine/Plugins/Experimental/Avalanche/Source/AvalancheOutliner/Private/Slate/Styling/AvaOutlinerStyleUtils.h
@@ -9,7 +9,7 @@ struct FSlateColor;
 struct FSlateColorBrush;
 struct FToolBarStyle;
 
-namespace UE::AvalancheOutliner::Private
+namespace UE::AvaOutliner::Private
 {
 	enum class EStyleType
 	{
diff --git a/Engine/Plugins/Experimental/Avalanche/Source/AvalancheOutliner/Private/Tests/AvaOutliner.spec.cpp b/Engine/Plugins/Experimental/Avalanche/Source/AvalancheOutliner/Private/Tests/AvaOutliner.spec.cpp
--- a/Engine/Plugins/Experimental/Avalanche/Source/AvalancheOutliner/Private/Tests/AvaOutliner.spec.cpp
+++ b/Engine/Plugins/Experimental/Avalanche/Source/AvalancheOutliner/Private/Tests/AvaOutliner.spec.cpp
@@ -56,7 +56,7 @@ BEGIN_DEFINE_SPEC(FAvaOutlinerSpec, "Avalanche.Outliner"
 	TSharedPtr<FAvaOutlinerView> OutlinerView;
 
 	// Custom test provider to fill some of the info/actions requested by the outliner (e.g. USelection Instances, or Actor Duplication)
-	TSharedPtr<UE::AvalancheOutliner::Private::FAvaOutlinerProviderTest> OutlinerProvider;
+	TSharedPtr<UE::AvaOutliner::Private::FAvaOutlinerProviderTest> OutlinerProvider;
 
 	// A list of outliner items to test things with
 	TArray<FAvaOutlinerItemPtr> TestItems;
@@ -67,7 +67,7 @@ void FAvaOutlinerSpec::Define()
 {
 	BeforeEach([this]()
 	{
-		OutlinerProvider = MakeShared<UE::AvalancheOutliner::Private::FAvaOutlinerProviderTest>();
+		OutlinerProvider = MakeShared<UE::AvaOutliner::Private::FAvaOutlinerProviderTest>();
 
 		constexpr int32 OutlinerViewId = 0;
 
@@ -75,7 +75,7 @@ void FAvaOutlinerSpec::Define()
 		Outliner->RegisterOutlinerView(OutlinerViewId);
 		Outliner->Refresh();
 
-		OutlinerView = Outliner->GetOutlinerView(0);
+		OutlinerView = Outliner->GetMostRecentOutlinerView();
 
 		const TArray<UObject*> TestObjects
 		{
diff --git a/Engine/Plugins/Experimental/Avalanche/Source/AvalancheOutliner/Private/Tests/AvaOutlinerTestUtils.cpp b/Engine/Plugins/Experimental/Avalanche/Source/AvalancheOutliner/Private/Tests/AvaOutlinerTestUtils.cpp
--- a/Engine/Plugins/Experimental/Avalanche/Source/AvalancheOutliner/Private/Tests/AvaOutlinerTestUtils.cpp
+++ b/Engine/Plugins/Experimental/Avalanche/Source/AvalancheOutliner/Private/Tests/AvaOutlinerTestUtils.cpp
@@ -1,7 +1,6 @@
 // Copyright Epic Games, Inc. All Rights Reserved.
 
 #include "AvaOutlinerTestUtils.h"
-
 #include "AvaOutliner.h"
 #include "AvaOutlinerSubsystem.h"
 #include "Components/LightComponent.h"
@@ -15,7 +14,7 @@
 #include "PackageTools.h"
 #include "Selection.h"
 
-void UE::AvalancheOutliner::Private::FAvaOutlinerEditorModeTools::Init(UWorld* InWorld)
+void UE::AvaOutliner::Private::FAvaOutlinerEditorModeTools::Init(UWorld* InWorld)
 {
 	World = InWorld;
 
@@ -38,7 +37,7 @@ void UE::AvalancheOutliner::Private::FAvaOutlinerEditorModeTools::Init(UWorld* I
 	ObjectSelection->AddToRoot();
 }
 
-void UE::AvalancheOutliner::Private::FAvaOutlinerEditorModeTools::Cleanup()
+void UE::AvaOutliner::Private::FAvaOutlinerEditorModeTools::Cleanup()
 {
 	if (!UObjectInitialized())
 	{
@@ -77,12 +76,12 @@ void UE::AvalancheOutliner::Private::FAvaOutlinerEditorModeTools::Cleanup()
 	World = nullptr;
 }
 
-UE::AvalancheOutliner::Private::FAvaOutlinerEditorModeTools::~FAvaOutlinerEditorModeTools()
+UE::AvaOutliner::Private::FAvaOutlinerEditorModeTools::~FAvaOutlinerEditorModeTools()
 {
 	Cleanup();
 }
 
-UE::AvalancheOutliner::Private::FAvaOutlinerProviderTest::FAvaOutlinerProviderTest()
+UE::AvaOutliner::Private::FAvaOutlinerProviderTest::FAvaOutlinerProviderTest()
 	: Outliner(MakeShared<FAvaOutliner>(*this))
 	, ModeTools(MakeShared<FAvaOutlinerEditorModeTools>())
 {
@@ -101,12 +100,12 @@ UE::AvalancheOutliner::Private::FAvaOutlinerProviderTest::FAvaOutlinerProviderTe
 	FillWorld();
 }
 
-UE::AvalancheOutliner::Private::FAvaOutlinerProviderTest::~FAvaOutlinerProviderTest()
+UE::AvaOutliner::Private::FAvaOutlinerProviderTest::~FAvaOutlinerProviderTest()
 {
 	CleanupWorld();
 }
 
-void UE::AvalancheOutliner::Private::FAvaOutlinerProviderTest::FillWorld()
+void UE::AvaOutliner::Private::FAvaOutlinerProviderTest::FillWorld()
 {
 	if (!World)
 	{
@@ -208,7 +207,7 @@ void UE::AvalancheOutliner::Private::FAvaOutlinerProviderTest::FillWorld()
 	}
 }
 
-void UE::AvalancheOutliner::Private::FAvaOutlinerProviderTest::CleanupWorld()
+void UE::AvaOutliner::Private::FAvaOutlinerProviderTest::CleanupWorld()
 {
 	ModeTools->Cleanup();
 
@@ -233,7 +232,7 @@ void UE::AvalancheOutliner::Private::FAvaOutlinerProviderTest::CleanupWorld()
 	}
 }
 
-TArray<FAvaOutlinerItemPtr> UE::AvalancheOutliner::Private::FAvaOutlinerProviderTest::GetOutlinerItems(const TArray<UObject*>& InObjects) const
+TArray<FAvaOutlinerItemPtr> UE::AvaOutliner::Private::FAvaOutlinerProviderTest::GetOutlinerItems(const TArray<UObject*>& InObjects) const
 {
 	TArray<FAvaOutlinerItemPtr> OutlinerItems;
 	OutlinerItems.Reserve(InObjects.Num());
@@ -249,7 +248,7 @@ TArray<FAvaOutlinerItemPtr> UE::AvalancheOutliner::Private::FAvaOutlinerProvider
 	return OutlinerItems;
 }
 
-void UE::AvalancheOutliner::Private::FAvaOutlinerProviderTest::TestSpawnActor()
+void UE::AvaOutliner::Private::FAvaOutlinerProviderTest::TestSpawnActor()
 {
 	TestSpawnedActor = World->SpawnActorDeferred<AStaticMeshActor>(AStaticMeshActor::StaticClass()
 			, FTransform::Identity
@@ -259,12 +258,12 @@ void UE::AvalancheOutliner::Private::FAvaOutlinerProviderTest::TestSpawnActor()
 	TestSpawnedActor->FinishSpawning(FTransform::Identity, true);
 }
 
-bool UE::AvalancheOutliner::Private::FAvaOutlinerProviderTest::CanOutlinerProcessActorSpawn(AActor* InActor) const
+bool UE::AvaOutliner::Private::FAvaOutlinerProviderTest::CanOutlinerProcessActorSpawn(AActor* InActor) const
 {
 	return !InActor->bIsEditorPreviewActor;
 }
 
-void UE::AvalancheOutliner::Private::FAvaOutlinerProviderTest::OutlinerDuplicateActors(const TArray<AActor*>& InTemplateActors)
+void UE::AvaOutliner::Private::FAvaOutlinerProviderTest::OutlinerDuplicateActors(const TArray<AActor*>& InTemplateActors)
 {
 	constexpr EAvaOutlinerIgnoreNotifyFlags NotifiesToIgnore = EAvaOutlinerIgnoreNotifyFlags::Spawn
 		| EAvaOutlinerIgnoreNotifyFlags::Duplication;
@@ -389,21 +388,21 @@ void UE::AvalancheOutliner::Private::FAvaOutlinerProviderTest::OutlinerDuplicate
 	Outliner->OnActorsDuplicated(DuplicateActorMap);
 }
 
-TOptional<EItemDropZone> UE::AvalancheOutliner::Private::FAvaOutlinerProviderTest::OnOutlinerItemCanAcceptDrop(const FDragDropEvent& DragDropEvent
+TOptional<EItemDropZone> UE::AvaOutliner::Private::FAvaOutlinerProviderTest::OnOutlinerItemCanAcceptDrop(const FDragDropEvent& DragDropEvent
 	, EItemDropZone DropZone
 	, FAvaOutlinerItemPtr TargetItem) const
 {
 	return TOptional<EItemDropZone>();
 }
 
-FReply UE::AvalancheOutliner::Private::FAvaOutlinerProviderTest::OnOutlinerItemAcceptDrop(const FDragDropEvent& DragDropEvent
+FReply UE::AvaOutliner::Private::FAvaOutlinerProviderTest::OnOutlinerItemAcceptDrop(const FDragDropEvent& DragDropEvent
 	, EItemDropZone DropZone
 	, FAvaOutlinerItemPtr TargetItem)
 {
 	return FReply::Unhandled();
 }
 
-const FAttachmentTransformRules& UE::AvalancheOutliner::Private::FAvaOutlinerProviderTest::GetTransformRule(bool bIsPrimaryTransformRule) const
+const FAttachmentTransformRules& UE::AvaOutliner::Private::FAvaOutlinerProviderTest::GetTransformRule(bool bIsPrimaryTransformRule) const
 {
 	return bIsPrimaryTransformRule ? FAttachmentTransformRules::KeepWorldTransform: FAttachmentTransformRules::KeepRelativeTransform;
 }
diff --git a/Engine/Plugins/Experimental/Avalanche/Source/AvalancheOutliner/Private/Tests/AvaOutlinerTestUtils.h b/Engine/Plugins/Experimental/Avalanche/Source/AvalancheOutliner/Private/Tests/AvaOutlinerTestUtils.h
--- a/Engine/Plugins/Experimental/Avalanche/Source/AvalancheOutliner/Private/Tests/AvaOutlinerTestUtils.h
+++ b/Engine/Plugins/Experimental/Avalanche/Source/AvalancheOutliner/Private/Tests/AvaOutlinerTestUtils.h
@@ -10,8 +10,9 @@ class APostProcessVolume;
 class ASkyLight;
 class AStaticMeshActor;
 class AVolume;
+class FAvaOutliner;
 
-namespace UE::AvalancheOutliner::Private
+namespace UE::AvaOutliner::Private
 {
 	/** Mode Tools used to test the Outliner */
 	class FAvaOutlinerEditorModeTools : public FEditorModeTools
diff --git a/Engine/Plugins/Experimental/Avalanche/Source/AvalancheOutliner/Public/AvaOutliner.h b/Engine/Plugins/Experimental/Avalanche/Source/AvalancheOutliner/Public/AvaOutliner.h
deleted file mode 100644
diff --git a/Engine/Plugins/Experimental/Avalanche/Source/AvalancheOutliner/Public/AvaOutlinerSettings.h b/Engine/Plugins/Experimental/Avalanche/Source/AvalancheOutliner/Public/AvaOutlinerSettings.h
deleted file mode 100644
diff --git a/Engine/Plugins/Experimental/Avalanche/Source/AvalancheOutliner/Public/AvaOutlinerSubsystem.h b/Engine/Plugins/Experimental/Avalanche/Source/AvalancheOutliner/Public/AvaOutlinerSubsystem.h
--- a/Engine/Plugins/Experimental/Avalanche/Source/AvalancheOutliner/Public/AvaOutlinerSubsystem.h
+++ b/Engine/Plugins/Experimental/Avalanche/Source/AvalancheOutliner/Public/AvaOutlinerSubsystem.h
@@ -2,11 +2,12 @@
 
 #pragma once
 
+#include "AvaOutlinerDefines.h"
 #include "Subsystems/WorldSubsystem.h"
 #include "Templates/SharedPointer.h"
 #include "AvaOutlinerSubsystem.generated.h"
 
-class FAvaOutliner;
+class IAvaOutliner;
 class IAvaOutlinerProvider;
 
 UENUM()
@@ -21,8 +22,8 @@ enum class EAvaOutlinerHierarchyChangeType : uint8
 };
 
 /** Subsystem in charge of instancing and keeping reference of the World's Outliner */
-UCLASS()
-class AVALANCHEOUTLINER_API UAvaOutlinerSubsystem : public UWorldSubsystem
+UCLASS(MinimalAPI)
+class UAvaOutlinerSubsystem : public UWorldSubsystem
 {
 	GENERATED_BODY()
 
@@ -30,16 +31,22 @@ public:
 	/**
 	 * Instantiates the World's Outliner
 	 * @param InProvider the outliner provider
-	 * @param bInForceCreate whether to create a new FAvaOutliner even if one already exists
-	 * @returns an existing or new valid FAvaOutliner instance
+	 * @param bInForceCreate whether to create a new IAvaOutliner even if one already exists
+	 * @returns an existing or new valid IAvaOutliner instance
 	 * @note the subsystem holds a non-owning reference of the outliner, the provider is what ideally should be holding the owning reference
 	 */
-	TSharedRef<FAvaOutliner> GetOrCreateOutliner(IAvaOutlinerProvider& InProvider, bool bInForceCreate = false);
+	AVALANCHEOUTLINER_API TSharedRef<IAvaOutliner> GetOrCreateOutliner(IAvaOutlinerProvider& InProvider, bool bInForceCreate = false);
 
-	TSharedPtr<FAvaOutliner> GetOutliner() const { return OutlinerWeak.Pin(); }
+	TSharedPtr<IAvaOutliner> GetOutliner() const
+	{
+		return OutlinerWeak.Pin();
+	}
 
 	DECLARE_EVENT_ThreeParams(UAvaOutlinerSubsystem, FActorHierarchyChanged, AActor* /*InActor*/, const AActor* /*InParentActor*/, EAvaOutlinerHierarchyChangeType);
-	FActorHierarchyChanged& OnActorHierarchyChanged() { return ActorHierarchyChangedEvent; }
+	FActorHierarchyChanged& OnActorHierarchyChanged()
+	{
+		return ActorHierarchyChangedEvent;
+	}
 
 	/**
 	 * Called when the Actor has changed in the Outliner Hierarchy
@@ -49,12 +56,13 @@ public:
 	 */
 	void BroadcastActorHierarchyChanged(AActor* InActor, const AActor* InParentActor, EAvaOutlinerHierarchyChangeType InChangeType) const;
 
+protected:
 	//~ Begin UWorldSubsystem
 	virtual bool DoesSupportWorldType(const EWorldType::Type InWorldType) const override;
 	//~ End UWorldSubsystem
 
 private:
-	TWeakPtr<FAvaOutliner> OutlinerWeak;
+	TWeakPtr<IAvaOutliner> OutlinerWeak;
 
 	FActorHierarchyChanged ActorHierarchyChangedEvent;
 };
diff --git a/Engine/Plugins/Experimental/Avalanche/Source/AvalancheOutliner/Public/AvaOutlinerView.h b/Engine/Plugins/Experimental/Avalanche/Source/AvalancheOutliner/Public/AvaOutlinerView.h
deleted file mode 100644
diff --git a/Engine/Plugins/Experimental/Avalanche/Source/AvalancheOutliner/Public/Columns/AvaOutlinerColumnExtender.h b/Engine/Plugins/Experimental/Avalanche/Source/AvalancheOutliner/Public/Columns/AvaOutlinerColumnExtender.h
deleted file mode 100644
diff --git a/Engine/Plugins/Experimental/Avalanche/Source/AvalancheOutliner/Public/Columns/IAvaOutlinerColumn.h b/Engine/Plugins/Experimental/Avalanche/Source/AvalancheOutliner/Public/Columns/IAvaOutlinerColumn.h
deleted file mode 100644
diff --git a/Engine/Plugins/Experimental/Avalanche/Source/AvalancheOutliner/Public/DragDropOps/AvaOutlinerItemDragDropOp.h b/Engine/Plugins/Experimental/Avalanche/Source/AvalancheOutliner/Public/DragDropOps/AvaOutlinerItemDragDropOp.h
--- a/Engine/Plugins/Experimental/Avalanche/Source/AvalancheOutliner/Public/DragDropOps/AvaOutlinerItemDragDropOp.h
+++ b/Engine/Plugins/Experimental/Avalanche/Source/AvalancheOutliner/Public/DragDropOps/AvaOutlinerItemDragDropOp.h
@@ -6,7 +6,7 @@
 #include "DragAndDrop/DecoratedDragDropOp.h"
 #include "Handlers/AvaOutlinerItemDropHandler.h"
 
-class FAvaOutliner;
+class IAvaOutlinerView;
 class FAvaOutlinerItemDropHandler;
 
 enum class EItemDropZone;
@@ -21,13 +21,11 @@ public:
 		, const TSharedPtr<FAvaOutlinerView>& InOutlinerView
 		, EAvaOutlinerDragDropActionType InActionType);
 
-	TSharedPtr<FAvaOutlinerView> GetOutlinerView() const
+	TSharedPtr<IAvaOutlinerView> GetOutlinerView() const
 	{
 		return OutlinerViewWeak.Pin();
 	}
 
-	TSharedPtr<FAvaOutliner> GetOutliner() const;
-
 	TConstArrayView<FAvaOutlinerItemPtr> GetItems() const
 	{
 		return Items;
@@ -68,7 +66,7 @@ protected:
 
 	TArray<TSharedRef<FAvaOutlinerItemDropHandler>> DropHandlers;
 
-	TWeakPtr<FAvaOutlinerView> OutlinerViewWeak;
+	TWeakPtr<IAvaOutlinerView> OutlinerViewWeak;
 
 	EAvaOutlinerDragDropActionType ActionType = EAvaOutlinerDragDropActionType::Move;
 };
diff --git a/Engine/Plugins/Experimental/Avalanche/Source/AvalancheOutliner/Public/Factories/Filters/AvaFilterExpressionFactoryUtils.h b/Engine/Plugins/Experimental/Avalanche/Source/AvalancheOutliner/Public/Factories/Filters/AvaFilterExpressionFactoryUtils.h
deleted file mode 100644
diff --git a/Engine/Plugins/Experimental/Avalanche/Source/AvalancheOutliner/Public/Filters/IAvaFilterExpressionFactory.h b/Engine/Plugins/Experimental/Avalanche/Source/AvalancheOutliner/Public/Filters/IAvaFilterExpressionFactory.h
deleted file mode 100644
diff --git a/Engine/Plugins/Experimental/Avalanche/Source/AvalancheOutliner/Public/Filters/IAvaFilterSuggestionFactory.h b/Engine/Plugins/Experimental/Avalanche/Source/AvalancheOutliner/Public/Filters/IAvaFilterSuggestionFactory.h
deleted file mode 100644
diff --git a/Engine/Plugins/Experimental/Avalanche/Source/AvalancheOutliner/Public/IAvaOutliner.h b/Engine/Plugins/Experimental/Avalanche/Source/AvalancheOutliner/Public/IAvaOutliner.h
new file mode 100644
--- /dev/null
+++ b/Engine/Plugins/Experimental/Avalanche/Source/AvalancheOutliner/Public/IAvaOutliner.h
@@ -0,0 +1,197 @@
+// Copyright Epic Games, Inc. All Rights Reserved.
+
+#pragma once
+
+#include "AvaOutlinerDefines.h"
+#include "Containers/ArrayView.h"
+#include "Containers/ContainersFwd.h"
+#include "Containers/StringFwd.h"
+#include "Delegates/Delegate.h"
+#include "IAvaOutlinerModule.h"
+#include "Item/AvaOutlinerItemId.h"
+#include "Item/AvaOutlinerItemProxy.h"
+#include "Item/IAvaOutlinerItem.h"
+#include "ItemProxies/AvaOutlinerItemProxyRegistry.h"
+#include "ItemProxies/IAvaOutlinerItemProxyFactory.h"
+#include "Templates/SharedPointer.h"
+
+class FAvaOutlinerTreeRoot;
+class AActor;
+class FArchive;
+class FAvaEditorSelection;
+class FUICommandList;
+class IAvaOutlinerProvider;
+class IAvaOutlinerView;
+class UWorld;
+enum class EItemDropZone;
+struct FAttachmentTransformRules;
+struct FAvaOutlinerItemId;
+template<typename OptionalType> struct TOptional;
+
+/** 
+ * The Outliner Object that is commonly instanced once per World
+ * (unless for advanced use where there are different outliner instances with different item ordering and behaviors).
+ * This is the object that dictates core outliner behavior like how items are sorted, which items are allowed, etc.
+ * Views are the objects that take this core behavior and show a part of it (e.g. through filters).
+ */
+class IAvaOutliner : public TSharedFromThis<IAvaOutliner>
+{
+public:
+	virtual ~IAvaOutliner() = default;
+
+	DECLARE_MULTICAST_DELEGATE(FOnOutlinerLoaded);
+	virtual FOnOutlinerLoaded& GetOnOutlinerLoaded() = 0;
+
+	virtual IAvaOutlinerProvider& GetProvider() const = 0;
+
+	/** Sets the Command List that the Outliner Views will use to append their Command Lists to */
+	virtual void SetBaseCommandList(const TSharedPtr<FUICommandList>& InBaseCommandList) = 0;
+
+	/** Serializes the Outliner Save State to the given archive */
+	virtual void Serialize(FArchive& Ar) = 0;
+
+	/** Register a new Outliner View to the Outliner to the given id, replacing the old view that was bound to the given id */
+	virtual TSharedPtr<IAvaOutlinerView> RegisterOutlinerView(int32 InOutlinerViewId) = 0;
+
+	/** Gets the Outliner View bound to the given id */
+	virtual TSharedPtr<IAvaOutlinerView> GetOutlinerView(int32 InOutlinerViewId) const = 0;
+
+	/** Instantiates a new Item and automatically registers it to the Outliner */
+	template<typename InItemType, typename = typename TEnableIf<TIsDerivedFrom<InItemType, IAvaOutlinerItem>::IsDerived>::Type, typename ...InArgTypes>
+	TSharedRef<InItemType> FindOrAdd(InArgTypes&&... InArgs)
+	{
+		TSharedRef<InItemType> Item = MakeShared<InItemType>(*this, Forward<InArgTypes>(InArgs)...);
+
+		// If an existing item already exists and has a valid state, use that and forget about the newly created
+		FAvaOutlinerItemPtr ExistingItem = FindItem(Item->GetItemId());
+		if (ExistingItem.IsValid() && ExistingItem->IsItemValid() && ExistingItem->IsA<InItemType>())
+		{
+			return StaticCastSharedPtr<InItemType>(ExistingItem).ToSharedRef();
+		}
+
+		if (Item->IsAllowedInOutliner())
+		{
+			RegisterItem(Item);
+		}
+
+		return Item;
+	}
+
+	/** Tries to find the Item Proxy Factory for the given Item Proxy Type Name */
+	template<typename InItemProxyType, typename = typename TEnableIf<TIsDerivedFrom<InItemProxyType, FAvaOutlinerItemProxy>::IsDerived>::Type>
+	IAvaOutlinerItemProxyFactory* GetItemProxyFactory() const
+	{
+		// First look for the Registry in Outliner
+		if (IAvaOutlinerItemProxyFactory* Factory = GetItemProxyRegistry().GetItemProxyFactory<InItemProxyType>())
+		{
+			return Factory;
+		}
+		// Fallback to finding the Factory in the Module if the Outliner did not find it
+		return IAvaOutlinerModule::Get().GetItemProxyRegistry().GetItemProxyFactory<InItemProxyType>();
+	}
+
+	/**
+	 * Tries to get the Item Proxy Factory for the given Item Proxy type, first trying the Outliner Registry then the Module's
+	 * then returns an existing item proxy created via the factory, or creates one if there's no existing item proxy
+	 * @returns the Item Proxy created by the Factory. Can be null if no factory was found or if the factory intentionally returns null
+	 */
+	template<typename InItemProxyType, typename = typename TEnableIf<TIsDerivedFrom<InItemProxyType, FAvaOutlinerItemProxy>::IsDerived>::Type>
+	TSharedPtr<FAvaOutlinerItemProxy> GetOrCreateItemProxy(const FAvaOutlinerItemPtr& InParentItem)
+	{
+		if (!InParentItem.IsValid() || !InParentItem->IsAllowedInOutliner())
+		{
+			return nullptr;
+		}
+		
+		IAvaOutlinerItemProxyFactory* const Factory = GetItemProxyFactory<InItemProxyType>();
+		if (!Factory)
+		{
+			return nullptr;
+		}
+		
+		TSharedPtr<FAvaOutlinerItemProxy> OutItemProxy;
+		if (FAvaOutlinerItemPtr ExistingItemProxy = FindItem(FAvaOutlinerItemId(InParentItem, *Factory)))
+		{
+			check(ExistingItemProxy->IsA<FAvaOutlinerItemProxy>());
+			ExistingItemProxy->SetParent(InParentItem);
+			OutItemProxy = StaticCastSharedPtr<FAvaOutlinerItemProxy>(ExistingItemProxy);
+		}
+		else
+		{
+			OutItemProxy = Factory->CreateItemProxy(*this, InParentItem);
+		}
+		RegisterItem(OutItemProxy);
+		return OutItemProxy;
+	}
+
+	/** Registers the given Item, replacing the old one. */
+	virtual void RegisterItem(const FAvaOutlinerItemPtr& InItem) = 0;
+
+	/** Unregisters the Item having the given ItemId */
+	virtual void UnregisterItem(const FAvaOutlinerItemId& InItemId) = 0;
+
+	/** Ensures that the next time Refresh is called in tick, Refresh will be called */
+	virtual void RequestRefresh() = 0;
+
+	/**
+	 * Flushes the Pending Actions from the Queue while also updating the state of the Outliner.
+	 * Calling it directly is forcing it to happen.
+	 * If a refresh is needed it will be called on the next tick automatically.
+	 */
+	virtual void Refresh() = 0;
+
+	/** Gets the Tree Root Item of the Outliner */
+	virtual TSharedRef<FAvaOutlinerTreeRoot> GetTreeRoot() const = 0;
+
+	/**
+	 * Finds the Registered Item that has the given Id
+	 * @returns the Item with the given Id, or null if the item does not exist or was not registered to the Outliner
+	 */
+	virtual FAvaOutlinerItemPtr FindItem(const FAvaOutlinerItemId& InItemId) const = 0;
+
+	/**
+ 	 * Adds or Removes the Ignore Notify Flags to prevent certain actions from automatically happening when they're triggered
+ 	 * @param InFlag the ignore flag to add or remove
+ 	 * @param bIgnore whether to add (true) or remove (false) the flag
+ 	 */
+	virtual void SetIgnoreNotify(EAvaOutlinerIgnoreNotifyFlags InFlag, bool bIgnore) = 0;
+
+	/**
+ 	 * Should be called when Actors have been copied and give Ava Outliner opportunity to add to the Buffer to copy the Outliner data for those Actors
+ 	 * @param InOutCopiedData the data to process / append to for copy
+ 	 * @param InCopiedActors the actors to copy
+ 	 */
+	virtual void OnActorsCopied(FString& InOutCopiedData, TConstArrayView<AActor*> InCopiedActors) = 0;
+
+	/**
+	 * Should be called when Actors have been pasted to parse the data that was filled in by IAvaOutliner::OnActorsCopied
+	 * @param InPastedData pasted string data
+	 * @param InPastedActors map of the original actor name to its created actor on paste
+	 */
+	virtual void OnActorsPasted(FStringView InPastedData, const TMap<FName, AActor*>& InPastedActors) = 0;
+
+	/**
+	 * Handles when Actors have been Duplicated
+	 * @param InDuplicateActorMap the map of the Duplicate Actors to their Templates
+	 * @param InRelativeItem he item to use as positional reference to where to place the duplicate items in the outliner
+	 * @param InRelativeDropZone where to put the duplicate items relative to the InRelativeItem (above, below, onto)
+	 */
+	virtual void OnActorsDuplicated(const TMap<AActor*, AActor*>& InDuplicateActorMap
+		, FAvaOutlinerItemPtr InRelativeItem = nullptr
+		, TOptional<EItemDropZone> InRelativeDropZone = TOptional<EItemDropZone>()) = 0;
+
+	/**
+ 	 * Gets all the Selected Items and puts/attaches them under the given Grouping Actor.
+ 	 * Requires that the Grouping Actor is valid and spawned in the World.
+ 	 */
+	virtual void GroupSelection(AActor* InGroupingActor, const TOptional<FAttachmentTransformRules>& InTransformRules) = 0;
+
+	/** Called when the objects have been selected and notified through USelection Instances in Mode Tools */
+	virtual void OnObjectSelectionChanged(const FAvaEditorSelection& InEditorSelection) = 0;
+
+	/** Gets the World the Outliner is working with */
+	virtual UWorld* GetWorld() const = 0;
+
+private:
+	virtual const FAvaOutlinerItemProxyRegistry& GetItemProxyRegistry() const = 0;
+};
diff --git a/Engine/Plugins/Experimental/Avalanche/Source/AvalancheOutliner/Public/IAvaOutlinerModule.h b/Engine/Plugins/Experimental/Avalanche/Source/AvalancheOutliner/Public/IAvaOutlinerModule.h
--- a/Engine/Plugins/Experimental/Avalanche/Source/AvalancheOutliner/Public/IAvaOutlinerModule.h
+++ b/Engine/Plugins/Experimental/Avalanche/Source/AvalancheOutliner/Public/IAvaOutlinerModule.h
@@ -5,15 +5,12 @@
 #include "AvaOutlinerDefines.h"
 #include "Containers/ContainersFwd.h"
 #include "Delegates/Delegate.h"
-#include "Filters/IAvaFilterExpressionFactory.h"
-#include "Filters/IAvaFilterSuggestionFactory.h"
 #include "Icon/AvaOutlinerIconCustomization.h"
 #include "Item/AvaOutlinerItem.h"
 #include "Modules/ModuleInterface.h"
 #include "Modules/ModuleManager.h"
 #include "Templates/SharedPointer.h"
 
-class FAvaOutliner;
 class FAvaOutlinerItem;
 class FAvaOutlinerItemProxy;
 class FAvaOutlinerItemProxyRegistry;
@@ -36,61 +33,19 @@ public:
 		return FModuleManager::Get().IsModuleLoaded(TEXT("AvalancheOutliner"));
 	};
 
-	DECLARE_MULTICAST_DELEGATE_OneParam(FOnInitOutlinerToolMenu, FToolMenuContext&);
-
-	/** Called when Initializing Outliner ToolBar to extend its Tool Menu Context */
-	virtual FOnInitOutlinerToolMenu& GetOnInitOutlinerToolBar() = 0;
-
-	/** Called when Initializing Outliner Item Context Menu to extend its Tool Menu Context */
-	virtual FOnInitOutlinerToolMenu& GetOnInitOutlinerItemContextMenu() = 0;
-
 	DECLARE_MULTICAST_DELEGATE_OneParam(FOnExtendOutlinerToolMenu, UToolMenu*);
 
-	/** Option to extend the Outliner ToolBar */
-	virtual FOnExtendOutlinerToolMenu& GetOnExtendOutlinerToolBar() = 0;
-
 	/** Option to extend the Outliner Item Context Menu */
 	virtual FOnExtendOutlinerToolMenu& GetOnExtendOutlinerItemContextMenu() = 0;
 
-	DECLARE_MULTICAST_DELEGATE_OneParam(FOnExtendOutlinerColumns, class FAvaOutlinerColumnExtender&);
-
-	/** Extends Outliner Columns regardless of the Outliner Provider used */
-	virtual FOnExtendOutlinerColumns& GetOnExtendOutlinerColumns() = 0;
-
 	/** Get the Module's Item Proxy Registry. Used as back up in case the Outliner's Instance does not have the Item type registered */
 	virtual FAvaOutlinerItemProxyRegistry& GetItemProxyRegistry() = 0;
 
-	DECLARE_MULTICAST_DELEGATE_ThreeParams(FOnExtendItemProxiesForItem, FAvaOutliner&, const FAvaOutlinerItemPtr&, TArray<TSharedPtr<FAvaOutlinerItemProxy>>&);
+	DECLARE_MULTICAST_DELEGATE_ThreeParams(FOnExtendItemProxiesForItem, IAvaOutliner&, const FAvaOutlinerItemPtr&, TArray<TSharedPtr<FAvaOutlinerItemProxy>>&);
 
 	/** Gets the Delegate to call for adding in the Item Proxies to put under the given Item */
 	virtual FOnExtendItemProxiesForItem& GetOnExtendItemProxiesForItem() = 0;
 
-	/**
- 	 * Create and Register the given filter expression factory
- 	 * @tparam InFilterExpressionFactoryType The type of filter expression factory to create and register
- 	 * @tparam InArgsType Constructor args types if needed
- 	 * @param InArgs Additional args passed to the constructor of the filter expression factory if needed
- 	 */
-	template<typename InFilterExpressionFactoryType, typename... InArgsType
-			UE_REQUIRES(TIsDerivedFrom<InFilterExpressionFactoryType, IAvaFilterExpressionFactory>::Value)>
-	void RegisterFilterExpressionFactory(InArgsType&&... InArgs)
-	{
-		RegisterFilterExpressionFactory_Internal(IAvaFilterExpressionFactory::MakeInstance<InFilterExpressionFactoryType>(Forward<InArgsType>(InArgs)...));
-	}
-
-	/**
-  	 * Create and Register the given filter suggestion factory
-  	 * @tparam InSuggestionFactoryType The type of filter suggestion factory to create and register
-  	 * @tparam InArgsType Constructor args types if needed
-  	 * @param InArgs Additional args passed to the constructor of the filter suggestion factory if needed
-  	 */
-	template<typename InSuggestionFactoryType, typename... InArgsType
-			UE_REQUIRES(TIsDerivedFrom<InSuggestionFactoryType, IAvaFilterSuggestionFactory>::Value)>
-	void RegisterFilterSuggestionFactory(InArgsType&&... InArgs)
-	{
-		RegisterFilterSuggestionFactory_Internal(IAvaFilterSuggestionFactory::MakeInstance<InSuggestionFactoryType>(Forward<InArgsType>(InArgs)...));
-	}
-
 	/**
 	 * Register an icon customization for the given OutlinerItemClass
 	 * @tparam InOutlinerItemClass Outliner class to add a customization for
@@ -120,74 +75,7 @@ public:
 		UnregisterOverriddenIcon_Internal(TAvaType<InOutlinerItemClass>::GetTypeId(), InAdditionalKeyParameter);
 	}
 
-	/**
-	 * Get the custom icon for the item if a customization for it exist
-	 * @param InItem The item to search the icon for
-	 * @return The custom icon for the item if a customization is found, FSlateIcon() otherwise
-	 */
-	virtual FSlateIcon FindOverrideIcon(TSharedPtr<const FAvaOutlinerItem> InItem) const = 0;
-
-	/**
- 	 * Unregister a previously registered filter expression factory. 
- 	 * @param InFilterIdentifier 
- 	 */
-	virtual void UnregisterFilterSuggestionFactory(const FName& InFilterIdentifier) = 0;
-
-	/**
-	 * Unregister a previously registered filter expression factory. 
-	 * @param InFilterIdentifier 
-	 */
-	virtual void UnregisterFilterExpressionFactory(const FName& InFilterIdentifier) = 0;
-
-	/**
-	 *	Check if current filter expression factory support the comparison operation
-	 * @param InFilterKey Filter key to get the filter expression factory needed
-	 * @param InOperation Operation to check if supported
-	 * @return True if operation is supported, False otherwise
-	 */
-	virtual bool CanFilterSupportComparisonOperation(const FName& InFilterKey, const ETextFilterComparisonOperation InOperation) const = 0;
-
-	/**
-	 *	Evaluate the expression and return the result 
-	 * @param InFilterKey Filter Key to get the Factory
-	 * @param InItem Item that is currently checked
-	 * @param InArgs Args to evaluate the expression see FAvaTextFilterArgs for more information
-	 * @return True if the expression evaluated to True, False otherwise
-	 */
-	virtual bool FilterExpression(const FName& InFilterKey, const IAvaOutlinerItem& InItem, const FAvaTextFilterArgs& InArgs) const = 0;
-
-	/**
-	 * Get all suggestions with the given type (Generic/ItemBased/All)
-	 * @param InSuggestionType Type of suggestion to get, see EFilterSuggestionType for more information
-	 * @return An Array containing all suggestions of the given type
-	 */
-	virtual TArray<TSharedPtr<IAvaFilterSuggestionFactory>> GetSuggestions(const EAvaFilterSuggestionType& InSuggestionType) const = 0;
-
 protected:
-
-	/**
-	 * Adds the given Customization to the Map containing all of them if not already added.
-	 * @param InItemTypeId Outliner Item Type Id
-	 * @param InAvaOutlinerIconCustomization Customization to add
-	 */
 	virtual void RegisterOverridenIcon_Internal(const FAvaTypeId& InItemTypeId, const TSharedRef<IAvaOutlinerIconCustomization>& InAvaOutlinerIconCustomization) = 0;
-
-	/**
-	 * Unregister the customization with the given key
-	 * @param InItemTypeId Outliner Item Type Id
-	 * @param InSpecializationIdentifier Specialization identifier (ex. AActor class FName for ActorIconCustomization)
-	 */
 	virtual void UnregisterOverriddenIcon_Internal(const FAvaTypeId& InItemTypeId, const FName& InSpecializationIdentifier) = 0;
-
-	/**
-	 * Adds the given factory to the Map containing all of them if not already added.
-	 * @param InFilterExpressionFactory Filter expression factory to register.
-	 */
-	virtual void RegisterFilterExpressionFactory_Internal(const TSharedRef<IAvaFilterExpressionFactory>& InFilterExpressionFactory) = 0;
-
-	/**
-	 * Adds the given factory to the Map containing all of them if not already added.
-	 * @param InFilterSuggestionFactory Filter suggestion factory to register.
-	 */
-	virtual void RegisterFilterSuggestionFactory_Internal(const TSharedRef<IAvaFilterSuggestionFactory>& InFilterSuggestionFactory) = 0;
 };
diff --git a/Engine/Plugins/Experimental/Avalanche/Source/AvalancheOutliner/Public/IAvaOutlinerProvider.h b/Engine/Plugins/Experimental/Avalanche/Source/AvalancheOutliner/Public/IAvaOutlinerProvider.h
--- a/Engine/Plugins/Experimental/Avalanche/Source/AvalancheOutliner/Public/IAvaOutlinerProvider.h
+++ b/Engine/Plugins/Experimental/Avalanche/Source/AvalancheOutliner/Public/IAvaOutlinerProvider.h
@@ -4,7 +4,6 @@
 
 #include "AvaOutlinerDefines.h"
 
-class FAvaOutliner;
 class FAvaOutlinerItemDragDropOp;
 class FAvaOutlinerItemProxy;
 class FDragDropEvent;
@@ -30,16 +29,10 @@ public:
 
 	/** Whether an Outliner Widget should be created */
 	virtual bool ShouldCreateWidget() const { return true; }
-	
+
 	/** Determines whether an Actor can be added and be visible in the Outliner */
 	virtual bool CanOutlinerProcessActorSpawn(AActor* InActor) const = 0;
-	
-	/** Called when the Outliner Save Data has been Changed */
-	virtual void OnOutlinerModified() {}
 
-	/** Gathers the Item Proxies for the given Item */
-	virtual void GetItemProxiesForItem(FAvaOutliner& InOutliner, const FAvaOutlinerItemPtr& InItem, TArray<TSharedPtr<FAvaOutlinerItemProxy>>& OutItemProxies) {}
-	
 	/** Determines whether the Outliner should be read-only */
 	virtual bool ShouldLockOutliner() const = 0;
 
@@ -51,34 +44,28 @@ public:
 	 * Implementation must return a map of the Duplicate Actor to the Template Actor
 	 */
 	virtual void OutlinerDuplicateActors(const TArray<AActor*>& InTemplateActors) = 0;
-	
+
 	/** Returns the Mode Tools used (e.g. for Selections) */
 	virtual FEditorModeTools* GetOutlinerModeTools() const = 0;
 
 	/** Returns the Scene Tree to use for the Outliner Tree Ordering */
 	virtual FAvaSceneTree* GetSceneTree() const = 0;
-	
+
 	/** Returns the World to use for the Outliner */
 	virtual UWorld* GetOutlinerWorld() const = 0;
 
 	/** Gets the Default Transform to use when Spawning an Actor from the Outliner (used when there's no info related to Transform) */
 	virtual FTransform GetOutlinerDefaultActorSpawnTransform() const = 0;
 
-	/** Option to extend the Outliner Columns via FAvaOutlinerColumnExtender */
-	virtual void ExtendOutlinerColumns(class FAvaOutlinerColumnExtender& InColumnExtender) {}
-	
 	/** Option to quickly extend the Outliner ToolBar without having to implement UToolMenus::Get()->ExtendMenu */
 	virtual void ExtendOutlinerToolBar(UToolMenu* InToolBarMenu) {}
-	
+
 	/** Option to quickly extend the Outliner Item Context Menu without having to implement UToolMenus::Get()->ExtendMenu */
 	virtual void ExtendOutlinerItemContextMenu(UToolMenu* InItemContextMenu) {}
 	
 	/** Extension to add Item Filters to the Outliner */
 	virtual void ExtendOutlinerItemFilters(TArray<TSharedPtr<class IAvaOutlinerItemFilter>>& InItemFilters) {}
 
-	/** Called when an Item Drag Drop Op has been Instantiated and Initialized. Useful for Adding Custom Drop Handlers*/
-	virtual void OnItemDragDropInitialized(FAvaOutlinerItemDragDropOp& InCreatedDragDropOp) {}
-
 	/** Determines whether an external Drag Drop event (i.e. not an Outliner one) can be accepted by the Outliner for a given Target Item */
 	virtual TOptional<EItemDropZone> OnOutlinerItemCanAcceptDrop(const FDragDropEvent& DragDropEvent
 		, EItemDropZone DropZone
diff --git a/Engine/Plugins/Experimental/Avalanche/Source/AvalancheOutliner/Public/IAvaOutlinerView.h b/Engine/Plugins/Experimental/Avalanche/Source/AvalancheOutliner/Public/IAvaOutlinerView.h
new file mode 100644
--- /dev/null
+++ b/Engine/Plugins/Experimental/Avalanche/Source/AvalancheOutliner/Public/IAvaOutlinerView.h
@@ -0,0 +1,19 @@
+// Copyright Epic Games, Inc. All Rights Reserved.
+
+#pragma once
+
+#include "Templates/SharedPointer.h"
+
+class IAvaOutliner;
+class SWidget;
+
+class IAvaOutlinerView : public TSharedFromThis<IAvaOutlinerView>
+{
+public:
+	virtual ~IAvaOutlinerView() = default;
+
+	/** Returns the Outliner Widget. Can be null widget */
+	virtual TSharedRef<SWidget> GetOutlinerWidget() const = 0;
+
+	virtual TSharedPtr<IAvaOutliner> GetOwnerOutliner() const  = 0;
+};
diff --git a/Engine/Plugins/Experimental/Avalanche/Source/AvalancheOutliner/Public/Icon/AvaOutlinerObjectIconCustomization.h b/Engine/Plugins/Experimental/Avalanche/Source/AvalancheOutliner/Public/Icon/AvaOutlinerObjectIconCustomization.h
--- a/Engine/Plugins/Experimental/Avalanche/Source/AvalancheOutliner/Public/Icon/AvaOutlinerObjectIconCustomization.h
+++ b/Engine/Plugins/Experimental/Avalanche/Source/AvalancheOutliner/Public/Icon/AvaOutlinerObjectIconCustomization.h
@@ -16,15 +16,16 @@ class AVALANCHEOUTLINER_API FAvaOutlinerObjectIconCustomization : public IAvaOut
 public:
 	FAvaOutlinerObjectIconCustomization(const UClass* InSupportedClass);
 
+	void SetOverriddenIcon(const FOnGetOverriddenObjectIcon& InOverriddenIcon);
+
+protected:
 	//~ Begin IAvaOutlinerIconCustomization
 	virtual FName GetOutlinerItemIdentifier() const override { return SupportedClassName; }
 	virtual bool HasOverrideIcon(TSharedPtr<const FAvaOutlinerItem> InOutlinerItem) const override;
 	virtual FSlateIcon GetOverrideIcon(TSharedPtr<const FAvaOutlinerItem> InOutlinerItem) const override;
 	//~ End IAvaOutlinerIconCustomization
 
-	void SetOverriddenIcon(const FOnGetOverriddenObjectIcon& InOverriddenIcon);
-
-protected:
+private:
 	FName SupportedClassName;
 
 	FOnGetOverriddenObjectIcon OnGetOverriddenIcon;
diff --git a/Engine/Plugins/Experimental/Avalanche/Source/AvalancheOutliner/Public/Item/AvaOutlinerActor.h b/Engine/Plugins/Experimental/Avalanche/Source/AvalancheOutliner/Public/Item/AvaOutlinerActor.h
--- a/Engine/Plugins/Experimental/Avalanche/Source/AvalancheOutliner/Public/Item/AvaOutlinerActor.h
+++ b/Engine/Plugins/Experimental/Avalanche/Source/AvalancheOutliner/Public/Item/AvaOutlinerActor.h
@@ -14,7 +14,7 @@ class AVALANCHEOUTLINER_API FAvaOutlinerActor : public FAvaOutlinerObject
 public:
 	UE_AVA_INHERITS_WITH_SUPER(FAvaOutlinerActor, FAvaOutlinerObject);
 
-	FAvaOutlinerActor(FAvaOutliner& InOutliner, AActor* InActor);
+	FAvaOutlinerActor(IAvaOutliner& InOutliner, AActor* InActor);
 
 	//~ Begin IAvaOutlinerItem
 	virtual void FindChildren(TArray<FAvaOutlinerItemPtr>& OutChildren, bool bRecursive) override;
diff --git a/Engine/Plugins/Experimental/Avalanche/Source/AvalancheOutliner/Public/Item/AvaOutlinerComponent.h b/Engine/Plugins/Experimental/Avalanche/Source/AvalancheOutliner/Public/Item/AvaOutlinerComponent.h
--- a/Engine/Plugins/Experimental/Avalanche/Source/AvalancheOutliner/Public/Item/AvaOutlinerComponent.h
+++ b/Engine/Plugins/Experimental/Avalanche/Source/AvalancheOutliner/Public/Item/AvaOutlinerComponent.h
@@ -14,7 +14,7 @@ class AVALANCHEOUTLINER_API FAvaOutlinerComponent : public FAvaOutlinerObject
 public:
 	UE_AVA_INHERITS_WITH_SUPER(FAvaOutlinerComponent, FAvaOutlinerObject);
 
-	FAvaOutlinerComponent(FAvaOutliner& InOutliner, USceneComponent* InComponent);
+	FAvaOutlinerComponent(IAvaOutliner& InOutliner, USceneComponent* InComponent);
 
 	//~ Begin IAvaOutlinerItem
 	virtual void FindChildren(TArray<FAvaOutlinerItemPtr>& OutChildren, bool bRecursive) override;
diff --git a/Engine/Plugins/Experimental/Avalanche/Source/AvalancheOutliner/Public/Item/AvaOutlinerComponentProxy.h b/Engine/Plugins/Experimental/Avalanche/Source/AvalancheOutliner/Public/Item/AvaOutlinerComponentProxy.h
--- a/Engine/Plugins/Experimental/Avalanche/Source/AvalancheOutliner/Public/Item/AvaOutlinerComponentProxy.h
+++ b/Engine/Plugins/Experimental/Avalanche/Source/AvalancheOutliner/Public/Item/AvaOutlinerComponentProxy.h
@@ -9,7 +9,7 @@ class AVALANCHEOUTLINER_API FAvaOutlinerComponentProxy : public FAvaOutlinerItem
 public:
 	UE_AVA_INHERITS_WITH_SUPER(FAvaOutlinerComponentProxy, FAvaOutlinerItemProxy)
 
-	FAvaOutlinerComponentProxy(FAvaOutliner& InOutliner, const FAvaOutlinerItemPtr& InParentItem);
+	FAvaOutlinerComponentProxy(IAvaOutliner& InOutliner, const FAvaOutlinerItemPtr& InParentItem);
 	
 	//~ Begin IAvaOutlinerItem
 	virtual FText GetDisplayName() const override;
diff --git a/Engine/Plugins/Experimental/Avalanche/Source/AvalancheOutliner/Public/Item/AvaOutlinerItem.h b/Engine/Plugins/Experimental/Avalanche/Source/AvalancheOutliner/Public/Item/AvaOutlinerItem.h
--- a/Engine/Plugins/Experimental/Avalanche/Source/AvalancheOutliner/Public/Item/AvaOutlinerItem.h
+++ b/Engine/Plugins/Experimental/Avalanche/Source/AvalancheOutliner/Public/Item/AvaOutlinerItem.h
@@ -12,17 +12,17 @@ class ULevel;
 /** Base Implementation of IAvaOutlinerItem */
 class AVALANCHEOUTLINER_API FAvaOutlinerItem : public IAvaOutlinerItem
 {
-	friend FAvaOutliner;
+	friend class FAvaOutliner;
 
 public:
 	UE_AVA_INHERITS_WITH_SUPER(FAvaOutlinerItem, IAvaOutlinerItem);
 
-	FAvaOutlinerItem(FAvaOutliner& InOutliner);
+	FAvaOutlinerItem(IAvaOutliner& InOutliner);
 
 	virtual ~FAvaOutlinerItem() override;
 
 	//~ Begin IAvaOutlinerItem
-	virtual TSharedRef<FAvaOutliner> GetOwnerOutliner() const override;
+	virtual TSharedRef<IAvaOutliner> GetOwnerOutliner() const override;
 	virtual bool IsItemValid() const override;
 	virtual void RefreshChildren() override;
 	virtual const TArray<FAvaOutlinerItemPtr>& GetChildren() const override { return Children; }
@@ -79,14 +79,11 @@ protected:
 	/** Careful handling of multiple children being detected and added to this item children array */
 	void HandleNewSortableChildren(TArray<FAvaOutlinerItemPtr> InSortableChildren);
 
-	/** The implementation when an Outliner item has dropped on this item */
-	FReply HandleItemDragDrop(const FAvaOutlinerItemDragDropOp& ItemDragDrop, EItemDropZone DropZone);
-
 	/** Handling when an Asset from Content Browser has been dragged on to this item */
 	FReply CreateItemsFromAssetDrop(const TSharedPtr<FAssetDragDropOp>& AssetDragDropOp, EItemDropZone DropZone, ULevel* Level);
 
 	/** Reference to the Owning Outliner */
-	FAvaOutliner& Outliner;
+	IAvaOutliner& Outliner;
 
 	/** Weak pointer to the Parent Item. Can be null, but if valid, the Parent should have this item in the Children Array */
 	TWeakPtr<IAvaOutlinerItem> ParentWeak;
diff --git a/Engine/Plugins/Experimental/Avalanche/Source/AvalancheOutliner/Public/Item/AvaOutlinerItemParameters.h b/Engine/Plugins/Experimental/Avalanche/Source/AvalancheOutliner/Public/Item/AvaOutlinerItemParameters.h
new file mode 100644
--- /dev/null
+++ b/Engine/Plugins/Experimental/Avalanche/Source/AvalancheOutliner/Public/Item/AvaOutlinerItemParameters.h
@@ -0,0 +1,67 @@
+// Copyright Epic Games, Inc. All Rights Reserved.
+
+#pragma once
+
+#include "AvaOutlinerDefines.h"
+#include "Engine/EngineTypes.h"
+#include "Misc/Optional.h"
+
+enum class EItemDropZone;
+
+enum class EAvaOutlinerAddItemFlags : uint8
+{
+	None        = 0,
+	/** Also add the children of the given item even if they were not made into their own Add Item Action */
+	AddChildren = 1 << 0,
+	/** Select this Item on Add */
+	Select      = 1 << 1,
+	/** Make a Transaction for this Action */
+	Transact    = 1 << 2,
+};
+ENUM_CLASS_FLAGS(EAvaOutlinerAddItemFlags);
+
+struct FAvaOutlinerAddItemParams
+{
+	FAvaOutlinerAddItemParams(const FAvaOutlinerItemPtr& InItem = nullptr
+			, EAvaOutlinerAddItemFlags InFlags            = EAvaOutlinerAddItemFlags::None
+			, const FAvaOutlinerItemPtr& InRelativeItem   = nullptr
+			, TOptional<EItemDropZone> InRelativeDropZone = TOptional<EItemDropZone>())
+		: Item(InItem)
+		, RelativeItem(InRelativeItem)
+		, RelativeDropZone(InRelativeDropZone)
+		, Flags(InFlags)
+		, SelectionFlags(EAvaOutlinerItemSelectionFlags::None)
+	{
+	}
+
+	/** The Item to Add */
+	FAvaOutlinerItemPtr Item;
+
+	/** The Item to use as base in where to place the Item */
+	FAvaOutlinerItemPtr RelativeItem;
+
+	/** The Placement Order from the Relative Item (Onto/Inside, Above, Below) */
+	TOptional<EItemDropZone> RelativeDropZone;
+
+	/** Some Extra Flags for what to do when Adding or After Adding the Items */
+	EAvaOutlinerAddItemFlags Flags;
+
+	/** Flags to Indicate how we should Select the Item. This only applies if bool bSelectItem is true */
+	EAvaOutlinerItemSelectionFlags SelectionFlags;
+
+	/** Optional Transform override Rule when Attaching Items */
+	TOptional<FAttachmentTransformRules> AttachmentTransformRules;
+};
+
+struct FAvaOutlinerRemoveItemParams
+{
+	FAvaOutlinerRemoveItemParams(const FAvaOutlinerItemPtr& InItem = nullptr)
+		: Item(InItem)
+	{
+	}
+
+	FAvaOutlinerItemPtr Item;
+
+	/** Optional Transform override Rule when Detaching Items */
+	TOptional<FDetachmentTransformRules> DetachmentTransformRules;
+};
diff --git a/Engine/Plugins/Experimental/Avalanche/Source/AvalancheOutliner/Public/Item/AvaOutlinerItemProxy.h b/Engine/Plugins/Experimental/Avalanche/Source/AvalancheOutliner/Public/Item/AvaOutlinerItemProxy.h
--- a/Engine/Plugins/Experimental/Avalanche/Source/AvalancheOutliner/Public/Item/AvaOutlinerItemProxy.h
+++ b/Engine/Plugins/Experimental/Avalanche/Source/AvalancheOutliner/Public/Item/AvaOutlinerItemProxy.h
@@ -16,7 +16,7 @@ class AVALANCHEOUTLINER_API FAvaOutlinerItemProxy : public FAvaOutlinerItem
 public:
 	UE_AVA_INHERITS_WITH_SUPER(FAvaOutlinerItemProxy, FAvaOutlinerItem);
 
-	FAvaOutlinerItemProxy(FAvaOutliner& InOutliner, const FAvaOutlinerItemPtr& InParentItem);
+	FAvaOutlinerItemProxy(IAvaOutliner& InOutliner, const FAvaOutlinerItemPtr& InParentItem);
 
 	//~ Begin IAvaOutlinerItem
 	virtual bool IsItemValid() const override;
diff --git a/Engine/Plugins/Experimental/Avalanche/Source/AvalancheOutliner/Public/Item/AvaOutlinerItemUtils.h b/Engine/Plugins/Experimental/Avalanche/Source/AvalancheOutliner/Public/Item/AvaOutlinerItemUtils.h
--- a/Engine/Plugins/Experimental/Avalanche/Source/AvalancheOutliner/Public/Item/AvaOutlinerItemUtils.h
+++ b/Engine/Plugins/Experimental/Avalanche/Source/AvalancheOutliner/Public/Item/AvaOutlinerItemUtils.h
@@ -5,8 +5,14 @@
 #include "AvaOutlinerDefines.h"
 #include "Containers/ContainersFwd.h"
 
-namespace UE::AvalancheOutliner
+namespace UE::AvaOutliner
 {
+	/**
+	 * Compares the absolute order of the items in the Outliner and returns true if A comes before B in the outliner.
+	 * Useful to use when sorting Items.
+	 */
+	AVALANCHEOUTLINER_API bool CompareOutlinerItemOrder(const FAvaOutlinerItemPtr& A, const FAvaOutlinerItemPtr& B);
+
 	/** Returns two subset arrays of Items: one is containing only Sortable Items and the other Non Sortable Items */
 	AVALANCHEOUTLINER_API void SplitItems(const TArray<FAvaOutlinerItemPtr>& InItems
 		, TArray<FAvaOutlinerItemPtr>& OutSortable
diff --git a/Engine/Plugins/Experimental/Avalanche/Source/AvalancheOutliner/Public/Item/AvaOutlinerMaterial.h b/Engine/Plugins/Experimental/Avalanche/Source/AvalancheOutliner/Public/Item/AvaOutlinerMaterial.h
--- a/Engine/Plugins/Experimental/Avalanche/Source/AvalancheOutliner/Public/Item/AvaOutlinerMaterial.h
+++ b/Engine/Plugins/Experimental/Avalanche/Source/AvalancheOutliner/Public/Item/AvaOutlinerMaterial.h
@@ -16,7 +16,7 @@ class AVALANCHEOUTLINER_API FAvaOutlinerMaterial : public FAvaOutlinerObjectRefe
 public:
 	UE_AVA_INHERITS_WITH_SUPER(FAvaOutlinerMaterial, FAvaOutlinerObjectReference)
 
-	FAvaOutlinerMaterial(FAvaOutliner& InOutliner
+	FAvaOutlinerMaterial(IAvaOutliner& InOutliner
 		, UMaterialInterface* InMaterial
 		, const FAvaOutlinerItemPtr& InReferencingItem
 		, int32 InMaterialIndex);
diff --git a/Engine/Plugins/Experimental/Avalanche/Source/AvalancheOutliner/Public/Item/AvaOutlinerMaterialProxy.h b/Engine/Plugins/Experimental/Avalanche/Source/AvalancheOutliner/Public/Item/AvaOutlinerMaterialProxy.h
--- a/Engine/Plugins/Experimental/Avalanche/Source/AvalancheOutliner/Public/Item/AvaOutlinerMaterialProxy.h
+++ b/Engine/Plugins/Experimental/Avalanche/Source/AvalancheOutliner/Public/Item/AvaOutlinerMaterialProxy.h
@@ -10,7 +10,7 @@ class AVALANCHEOUTLINER_API FAvaOutlinerMaterialProxy : public FAvaOutlinerItemP
 public:
 	UE_AVA_INHERITS_WITH_SUPER(FAvaOutlinerMaterialProxy, FAvaOutlinerItemProxy);
 
-	FAvaOutlinerMaterialProxy(FAvaOutliner& InOutliner, const FAvaOutlinerItemPtr& InParentItem);
+	FAvaOutlinerMaterialProxy(IAvaOutliner& InOutliner, const FAvaOutlinerItemPtr& InParentItem);
 
 	virtual ~FAvaOutlinerMaterialProxy() override;
 	
diff --git a/Engine/Plugins/Experimental/Avalanche/Source/AvalancheOutliner/Public/Item/AvaOutlinerObject.h b/Engine/Plugins/Experimental/Avalanche/Source/AvalancheOutliner/Public/Item/AvaOutlinerObject.h
--- a/Engine/Plugins/Experimental/Avalanche/Source/AvalancheOutliner/Public/Item/AvaOutlinerObject.h
+++ b/Engine/Plugins/Experimental/Avalanche/Source/AvalancheOutliner/Public/Item/AvaOutlinerObject.h
@@ -11,7 +11,7 @@ class AVALANCHEOUTLINER_API FAvaOutlinerObject : public FAvaOutlinerItem
 public:
 	UE_AVA_INHERITS_WITH_SUPER(FAvaOutlinerObject, FAvaOutlinerItem);
 
-	FAvaOutlinerObject(FAvaOutliner& InOutliner, UObject* InObject);
+	FAvaOutlinerObject(IAvaOutliner& InOutliner, UObject* InObject);
 
 	//~ Begin IAvaOutlinerItem
 	virtual bool IsItemValid() const override;
diff --git a/Engine/Plugins/Experimental/Avalanche/Source/AvalancheOutliner/Public/Item/AvaOutlinerObjectReference.h b/Engine/Plugins/Experimental/Avalanche/Source/AvalancheOutliner/Public/Item/AvaOutlinerObjectReference.h
--- a/Engine/Plugins/Experimental/Avalanche/Source/AvalancheOutliner/Public/Item/AvaOutlinerObjectReference.h
+++ b/Engine/Plugins/Experimental/Avalanche/Source/AvalancheOutliner/Public/Item/AvaOutlinerObjectReference.h
@@ -21,7 +21,7 @@ class AVALANCHEOUTLINER_API FAvaOutlinerSharedObject : public FAvaOutlinerObject
 public:
 	UE_AVA_INHERITS_WITH_SUPER(FAvaOutlinerSharedObject, FAvaOutlinerObject);
 
-	FAvaOutlinerSharedObject(FAvaOutliner& InOutliner, UObject* InObject);
+	FAvaOutlinerSharedObject(IAvaOutliner& InOutliner, UObject* InObject);
 
 	//~ Begin IAvaOutlinerItem
 	virtual bool IsAllowedInOutliner() const override { return false; }
@@ -47,7 +47,7 @@ class AVALANCHEOUTLINER_API FAvaOutlinerObjectReference : public FAvaOutlinerObj
 public:
 	UE_AVA_INHERITS_WITH_SUPER(FAvaOutlinerObjectReference, FAvaOutlinerObject);
 
-	FAvaOutlinerObjectReference(FAvaOutliner& InOutliner
+	FAvaOutlinerObjectReference(IAvaOutliner& InOutliner
 		, UObject* InObject
 		, const FAvaOutlinerItemPtr& InReferencingItem
 		, const FString& InReferenceId);
diff --git a/Engine/Plugins/Experimental/Avalanche/Source/AvalancheOutliner/Public/Item/AvaOutlinerTreeRoot.h b/Engine/Plugins/Experimental/Avalanche/Source/AvalancheOutliner/Public/Item/AvaOutlinerTreeRoot.h
--- a/Engine/Plugins/Experimental/Avalanche/Source/AvalancheOutliner/Public/Item/AvaOutlinerTreeRoot.h
+++ b/Engine/Plugins/Experimental/Avalanche/Source/AvalancheOutliner/Public/Item/AvaOutlinerTreeRoot.h
@@ -13,7 +13,7 @@ class AVALANCHEOUTLINER_API FAvaOutlinerTreeRoot final : public FAvaOutlinerItem
 public:
 	UE_AVA_INHERITS_WITH_SUPER(FAvaOutlinerTreeRoot, FAvaOutlinerItem)
 
-	FAvaOutlinerTreeRoot(FAvaOutliner& InOutliner)
+	FAvaOutlinerTreeRoot(IAvaOutliner& InOutliner)
 		: FAvaOutlinerItem(InOutliner)
 	{
 	}
diff --git a/Engine/Plugins/Experimental/Avalanche/Source/AvalancheOutliner/Public/Item/IAvaOutlinerItem.h b/Engine/Plugins/Experimental/Avalanche/Source/AvalancheOutliner/Public/Item/IAvaOutlinerItem.h
--- a/Engine/Plugins/Experimental/Avalanche/Source/AvalancheOutliner/Public/Item/IAvaOutlinerItem.h
+++ b/Engine/Plugins/Experimental/Avalanche/Source/AvalancheOutliner/Public/Item/IAvaOutlinerItem.h
@@ -6,13 +6,13 @@
 #include "AvaType.h"
 #include "Delegates/Delegate.h"
 
-class FAvaOutliner;
 class FAvaOutlinerItemProxy;
 class FAvaOutlinerScopedSelection;
 class FAvaOutlinerView;
 class FDragDropEvent;
 class FEditorModeTools;
 class FReply;
+class IAvaOutliner;
 class SAvaOutlinerTreeRow;
 class SWidget;
 class UObject;
@@ -46,7 +46,7 @@ public:
 	virtual bool IsSelectable() const { return true; }
 	
 	/** Gets the Outliner that owns this Item */
-	virtual TSharedRef<FAvaOutliner> GetOwnerOutliner() const = 0;
+	virtual TSharedRef<IAvaOutliner> GetOwnerOutliner() const = 0;
 
 	/** Determines whether the Item properties are in a valid state */
 	virtual bool IsItemValid() const = 0;
diff --git a/Engine/Plugins/Experimental/Avalanche/Source/AvalancheOutliner/Public/ItemActions/AvaOutlinerAddItem.h b/Engine/Plugins/Experimental/Avalanche/Source/AvalancheOutliner/Public/ItemActions/AvaOutlinerAddItem.h
deleted file mode 100644
diff --git a/Engine/Plugins/Experimental/Avalanche/Source/AvalancheOutliner/Public/ItemActions/AvaOutlinerItemAction.h b/Engine/Plugins/Experimental/Avalanche/Source/AvalancheOutliner/Public/ItemActions/AvaOutlinerItemAction.h
deleted file mode 100644
diff --git a/Engine/Plugins/Experimental/Avalanche/Source/AvalancheOutliner/Public/ItemActions/AvaOutlinerRemoveItem.h b/Engine/Plugins/Experimental/Avalanche/Source/AvalancheOutliner/Public/ItemActions/AvaOutlinerRemoveItem.h
deleted file mode 100644
diff --git a/Engine/Plugins/Experimental/Avalanche/Source/AvalancheOutliner/Public/ItemProxies/Factories/AvaOutlinerItemProxyDefaultFactory.h b/Engine/Plugins/Experimental/Avalanche/Source/AvalancheOutliner/Public/ItemProxies/Factories/AvaOutlinerItemProxyDefaultFactory.h
--- a/Engine/Plugins/Experimental/Avalanche/Source/AvalancheOutliner/Public/ItemProxies/Factories/AvaOutlinerItemProxyDefaultFactory.h
+++ b/Engine/Plugins/Experimental/Avalanche/Source/AvalancheOutliner/Public/ItemProxies/Factories/AvaOutlinerItemProxyDefaultFactory.h
@@ -34,7 +34,7 @@ template<typename InItemProxyType, uint32 InItemProxyPriority = 0
 class TAvaOutlinerItemProxyDefaultFactory : public TAvaOutlinerItemProxyDefaultFactoryBase<InItemProxyType, InItemProxyPriority>
 {
 public:
-	virtual TSharedPtr<FAvaOutlinerItemProxy> CreateItemProxy(FAvaOutliner& InOutliner, const FAvaOutlinerItemPtr& InParentItem) override
+	virtual TSharedPtr<FAvaOutlinerItemProxy> CreateItemProxy(IAvaOutliner& InOutliner, const FAvaOutlinerItemPtr& InParentItem) override
 	{
 		return this->DefaultCreateItemProxy(InOutliner, InParentItem);
 	}
diff --git a/Engine/Plugins/Experimental/Avalanche/Source/AvalancheOutliner/Public/ItemProxies/IAvaOutlinerItemProxyFactory.h b/Engine/Plugins/Experimental/Avalanche/Source/AvalancheOutliner/Public/ItemProxies/IAvaOutlinerItemProxyFactory.h
--- a/Engine/Plugins/Experimental/Avalanche/Source/AvalancheOutliner/Public/ItemProxies/IAvaOutlinerItemProxyFactory.h
+++ b/Engine/Plugins/Experimental/Avalanche/Source/AvalancheOutliner/Public/ItemProxies/IAvaOutlinerItemProxyFactory.h
@@ -5,9 +5,9 @@
 #include "AvaOutlinerDefines.h"
 #include "Templates/SharedPointer.h"
 
-class FAvaOutliner;
 class FAvaOutlinerItemProxy;
 class FName;
+class IAvaOutliner;
 
 /**
  * Item Proxy Factories are the classes that instance or get the existing Outliner Item Proxies for a given Item
@@ -22,5 +22,5 @@ public:
 	virtual FName GetItemProxyTypeName() const = 0;
 
 	/** Returns a newly created instance of the Relevant Item Proxy if successful */
-	virtual TSharedPtr<FAvaOutlinerItemProxy> CreateItemProxy(FAvaOutliner& InOutliner, const FAvaOutlinerItemPtr& InParentItem) = 0;
+	virtual TSharedPtr<FAvaOutlinerItemProxy> CreateItemProxy(IAvaOutliner& InOutliner, const FAvaOutlinerItemPtr& InParentItem) = 0;
 };
diff --git a/Engine/Plugins/Experimental/Avalanche/Source/AvalancheOutliner/Public/ToolMenuContext/AvaOutlinerItemsContext.h b/Engine/Plugins/Experimental/Avalanche/Source/AvalancheOutliner/Public/ToolMenuContext/AvaOutlinerItemsContext.h
--- a/Engine/Plugins/Experimental/Avalanche/Source/AvalancheOutliner/Public/ToolMenuContext/AvaOutlinerItemsContext.h
+++ b/Engine/Plugins/Experimental/Avalanche/Source/AvalancheOutliner/Public/ToolMenuContext/AvaOutlinerItemsContext.h
@@ -9,24 +9,28 @@
 #include "UObject/Object.h"
 #include "AvaOutlinerItemsContext.generated.h"
 
-class FAvaOutlinerView;
+class IAvaOutliner;
 
 UCLASS(MinimalAPI)
 class UAvaOutlinerItemsContext : public UObject
 {
 	GENERATED_BODY()
-	
-	friend FAvaOutlinerView;
-	
+
+	friend class FAvaOutlinerView;
+
 public:
-	UAvaOutlinerItemsContext() = default;
-	
-	TSharedPtr<FAvaOutlinerView> GetOutlinerView() const { return OutlinerViewWeak.Pin(); }
-	
-	TConstArrayView<FAvaOutlinerItemWeakPtr> GetItems() const { return ItemListWeak; }
-	
+	TSharedPtr<IAvaOutliner> GetOutliner() const
+	{
+		return OutlinerWeak.Pin();
+	}
+
+	TConstArrayView<FAvaOutlinerItemWeakPtr> GetItems() const
+	{
+		return ItemListWeak;
+	}
+
 private:
-	TWeakPtr<FAvaOutlinerView> OutlinerViewWeak;
-	
+	TWeakPtr<IAvaOutliner> OutlinerWeak;
+
 	TArray<FAvaOutlinerItemWeakPtr> ItemListWeak;
 };
diff --git a/Engine/Plugins/Experimental/Avalanche/Source/AvalancheOutliner/Public/ToolMenuContext/AvaOutlinerToolBarContext.h b/Engine/Plugins/Experimental/Avalanche/Source/AvalancheOutliner/Public/ToolMenuContext/AvaOutlinerToolBarContext.h
deleted file mode 100644
diff --git a/Engine/Plugins/Experimental/Avalanche/Source/AvalanchePropertyAnimatorEditor/Private/AvaPropertyAnimatorEditorModule.cpp b/Engine/Plugins/Experimental/Avalanche/Source/AvalanchePropertyAnimatorEditor/Private/AvaPropertyAnimatorEditorModule.cpp
--- a/Engine/Plugins/Experimental/Avalanche/Source/AvalanchePropertyAnimatorEditor/Private/AvaPropertyAnimatorEditorModule.cpp
+++ b/Engine/Plugins/Experimental/Avalanche/Source/AvalanchePropertyAnimatorEditor/Private/AvaPropertyAnimatorEditorModule.cpp
@@ -2,7 +2,7 @@
 
 #include "AvaPropertyAnimatorEditorModule.h"
 
-#include "AvaOutliner.h"
+#include "IAvaOutliner.h"
 #include "IAvaOutlinerModule.h"
 #include "Item/AvaOutlinerActor.h"
 #include "ItemProxies/AvaOutlinerItemProxyRegistry.h"
@@ -28,16 +28,16 @@ void FAvaPropertyAnimatorEditorModule::RegisterOutlinerItems()
 	ItemProxyRegistry.RegisterItemProxyWithDefaultFactory<FAvaPropertyAnimatorEditorOutlinerProxy, 50>();
 
 	OutlinerProxiesExtensionDelegateHandle = IAvaOutlinerModule::Get().GetOnExtendItemProxiesForItem().AddLambda(
-	[](FAvaOutliner& InOutliner, const FAvaOutlinerItemPtr& InItem, TArray<TSharedPtr<FAvaOutlinerItemProxy>>& OutItemProxies)
-	{
-		if (InItem->IsA<FAvaOutlinerActor>())
+		[](IAvaOutliner& InOutliner, const FAvaOutlinerItemPtr& InItem, TArray<TSharedPtr<FAvaOutlinerItemProxy>>& OutItemProxies)
 		{
-			if (const TSharedPtr<FAvaOutlinerItemProxy> ControllerProxy = InOutliner.GetOrCreateItemProxy<FAvaPropertyAnimatorEditorOutlinerProxy>(InItem))
+			if (InItem->IsA<FAvaOutlinerActor>())
 			{
-				OutItemProxies.Add(ControllerProxy);
+				if (const TSharedPtr<FAvaOutlinerItemProxy> ControllerProxy = InOutliner.GetOrCreateItemProxy<FAvaPropertyAnimatorEditorOutlinerProxy>(InItem))
+				{
+					OutItemProxies.Add(ControllerProxy);
+				}
 			}
-		}
-	});
+		});
 
 	OutlinerContextDelegateHandle = IAvaOutlinerModule::Get().GetOnExtendOutlinerItemContextMenu()
 		.AddStatic(&FAvaPropertyAnimatorEditorOutlinerContextMenu::OnExtendOutlinerContextMenu);
diff --git a/Engine/Plugins/Experimental/Avalanche/Source/AvalanchePropertyAnimatorEditor/Private/Outliner/AvaPropertyAnimatorEditorOutliner.cpp b/Engine/Plugins/Experimental/Avalanche/Source/AvalanchePropertyAnimatorEditor/Private/Outliner/AvaPropertyAnimatorEditorOutliner.cpp
--- a/Engine/Plugins/Experimental/Avalanche/Source/AvalanchePropertyAnimatorEditor/Private/Outliner/AvaPropertyAnimatorEditorOutliner.cpp
+++ b/Engine/Plugins/Experimental/Avalanche/Source/AvalanchePropertyAnimatorEditor/Private/Outliner/AvaPropertyAnimatorEditorOutliner.cpp
@@ -7,7 +7,7 @@
 #include "Selection/AvaOutlinerScopedSelection.h"
 #include "Styling/SlateIconFinder.h"
 
-FAvaPropertyAnimatorEditorOutliner::FAvaPropertyAnimatorEditorOutliner(FAvaOutliner& InOutliner, UPropertyAnimatorCoreBase* InAnimator)
+FAvaPropertyAnimatorEditorOutliner::FAvaPropertyAnimatorEditorOutliner(IAvaOutliner& InOutliner, UPropertyAnimatorCoreBase* InAnimator)
 	: FAvaOutlinerObject(InOutliner, InAnimator)
 	, PropertyAnimator(InAnimator)
 {
diff --git a/Engine/Plugins/Experimental/Avalanche/Source/AvalanchePropertyAnimatorEditor/Private/Outliner/AvaPropertyAnimatorEditorOutliner.h b/Engine/Plugins/Experimental/Avalanche/Source/AvalanchePropertyAnimatorEditor/Private/Outliner/AvaPropertyAnimatorEditorOutliner.h
--- a/Engine/Plugins/Experimental/Avalanche/Source/AvalanchePropertyAnimatorEditor/Private/Outliner/AvaPropertyAnimatorEditorOutliner.h
+++ b/Engine/Plugins/Experimental/Avalanche/Source/AvalanchePropertyAnimatorEditor/Private/Outliner/AvaPropertyAnimatorEditorOutliner.h
@@ -14,7 +14,7 @@ class FAvaPropertyAnimatorEditorOutliner : public FAvaOutlinerObject
 public:
 	UE_AVA_INHERITS_WITH_SUPER(FAvaPropertyAnimatorEditorOutliner, FAvaOutlinerObject);
 
-	FAvaPropertyAnimatorEditorOutliner(FAvaOutliner& InOutliner, UPropertyAnimatorCoreBase* InObject);
+	FAvaPropertyAnimatorEditorOutliner(IAvaOutliner& InOutliner, UPropertyAnimatorCoreBase* InObject);
 
 	UPropertyAnimatorCoreBase* GetPropertyAnimator() const
 	{
diff --git a/Engine/Plugins/Experimental/Avalanche/Source/AvalanchePropertyAnimatorEditor/Private/Outliner/AvaPropertyAnimatorEditorOutlinerProxy.cpp b/Engine/Plugins/Experimental/Avalanche/Source/AvalanchePropertyAnimatorEditor/Private/Outliner/AvaPropertyAnimatorEditorOutlinerProxy.cpp
--- a/Engine/Plugins/Experimental/Avalanche/Source/AvalanchePropertyAnimatorEditor/Private/Outliner/AvaPropertyAnimatorEditorOutlinerProxy.cpp
+++ b/Engine/Plugins/Experimental/Avalanche/Source/AvalanchePropertyAnimatorEditor/Private/Outliner/AvaPropertyAnimatorEditorOutlinerProxy.cpp
@@ -3,8 +3,8 @@
 #include "Outliner/AvaPropertyAnimatorEditorOutlinerProxy.h"
 
 #include "Animators/PropertyAnimatorCoreBase.h"
-#include "AvaOutliner.h"
 #include "Components/PropertyAnimatorCoreComponent.h"
+#include "IAvaOutliner.h"
 #include "Item/AvaOutlinerActor.h"
 #include "Outliner/AvaPropertyAnimatorEditorOutliner.h"
 #include "Selection/AvaOutlinerScopedSelection.h"
@@ -12,7 +12,7 @@
 
 #define LOCTEXT_NAMESPACE "AvaPropertyAnimatorEditorOutlinerProxy"
 
-FAvaPropertyAnimatorEditorOutlinerProxy::FAvaPropertyAnimatorEditorOutlinerProxy(FAvaOutliner& InOutliner, const FAvaOutlinerItemPtr& InParentItem)
+FAvaPropertyAnimatorEditorOutlinerProxy::FAvaPropertyAnimatorEditorOutlinerProxy(IAvaOutliner& InOutliner, const FAvaOutlinerItemPtr& InParentItem)
 	: Super(InOutliner, InParentItem)
 {
 	ItemIcon = FSlateIconFinder::FindIconForClass(UPropertyAnimatorCoreComponent::StaticClass());
diff --git a/Engine/Plugins/Experimental/Avalanche/Source/AvalanchePropertyAnimatorEditor/Private/Outliner/AvaPropertyAnimatorEditorOutlinerProxy.h b/Engine/Plugins/Experimental/Avalanche/Source/AvalanchePropertyAnimatorEditor/Private/Outliner/AvaPropertyAnimatorEditorOutlinerProxy.h
--- a/Engine/Plugins/Experimental/Avalanche/Source/AvalanchePropertyAnimatorEditor/Private/Outliner/AvaPropertyAnimatorEditorOutlinerProxy.h
+++ b/Engine/Plugins/Experimental/Avalanche/Source/AvalanchePropertyAnimatorEditor/Private/Outliner/AvaPropertyAnimatorEditorOutlinerProxy.h
@@ -14,7 +14,7 @@ class FAvaPropertyAnimatorEditorOutlinerProxy : public FAvaOutlinerItemProxy
 public:
 	UE_AVA_INHERITS_WITH_SUPER(FAvaPropertyAnimatorEditorOutlinerProxy, FAvaOutlinerItemProxy);
 
-	FAvaPropertyAnimatorEditorOutlinerProxy(FAvaOutliner& InOutliner, const FAvaOutlinerItemPtr& InParentItem);
+	FAvaPropertyAnimatorEditorOutlinerProxy(IAvaOutliner& InOutliner, const FAvaOutlinerItemPtr& InParentItem);
 
 	UPropertyAnimatorCoreComponent* GetPropertyAnimatorComponent() const;
 
diff --git a/Engine/Plugins/Experimental/Avalanche/Source/AvalancheRemoteControlEditor/Private/AvalancheRemoteControlEditorModule.cpp b/Engine/Plugins/Experimental/Avalanche/Source/AvalancheRemoteControlEditor/Private/AvalancheRemoteControlEditorModule.cpp
--- a/Engine/Plugins/Experimental/Avalanche/Source/AvalancheRemoteControlEditor/Private/AvalancheRemoteControlEditorModule.cpp
+++ b/Engine/Plugins/Experimental/Avalanche/Source/AvalancheRemoteControlEditor/Private/AvalancheRemoteControlEditorModule.cpp
@@ -1,9 +1,9 @@
 // Copyright Epic Games, Inc. All Rights Reserved.
 
 #include "AvalancheRemoteControlEditorModule.h"
-#include "AvaOutliner.h"
 #include "AvaRCControllerId.h"
 #include "Customizations/AvaRCControllerIdCustomization.h"
+#include "IAvaOutliner.h"
 #include "IAvaOutlinerModule.h"
 #include "Item/AvaOutlinerActor.h"
 #include "ItemProxies/AvaOutlinerItemProxyRegistry.h"
@@ -22,7 +22,7 @@ void FAvalancheRemoteControlEditorModule::StartupModule()
 	ItemProxyRegistry.RegisterItemProxyWithDefaultFactory<FAvaOutlinerRCTrackerComponentProxy, 40>();
 	
 	OutlinerProxiesExtensionDelegateHandle = IAvaOutlinerModule::Get().GetOnExtendItemProxiesForItem().AddLambda(
-		[this](FAvaOutliner& InOutliner, const FAvaOutlinerItemPtr& InItem, TArray<TSharedPtr<FAvaOutlinerItemProxy>>& OutItemProxies)
+		[this](IAvaOutliner& InOutliner, const FAvaOutlinerItemPtr& InItem, TArray<TSharedPtr<FAvaOutlinerItemProxy>>& OutItemProxies)
 		{
 			if (InItem->IsA<FAvaOutlinerActor>())
 			{
diff --git a/Engine/Plugins/Experimental/Avalanche/Source/AvalancheRemoteControlEditor/Private/Outliner/AvaOutlinerRCTrackerComponent.cpp b/Engine/Plugins/Experimental/Avalanche/Source/AvalancheRemoteControlEditor/Private/Outliner/AvaOutlinerRCTrackerComponent.cpp
--- a/Engine/Plugins/Experimental/Avalanche/Source/AvalancheRemoteControlEditor/Private/Outliner/AvaOutlinerRCTrackerComponent.cpp
+++ b/Engine/Plugins/Experimental/Avalanche/Source/AvalancheRemoteControlEditor/Private/Outliner/AvaOutlinerRCTrackerComponent.cpp
@@ -1,7 +1,7 @@
 // Copyright Epic Games, Inc. All Rights Reserved.
 
 #include "Outliner/AvaOutlinerRCTrackerComponent.h"
-#include "AvaOutliner.h"
+#include "IAvaOutliner.h"
 #include "RemoteControlTrackerComponent.h"
 #include "Selection/AvaOutlinerScopedSelection.h"
 #include "Subsystems/RemoteControlComponentsEditorUtils.h"
@@ -10,7 +10,7 @@
 
 FAvaOutlinerRCTrackerComponent::FOnAvaOutlinerRCTrackerComponentSelected FAvaOutlinerRCTrackerComponent::OnRCTrackerSelectedDelegate;
 
-FAvaOutlinerRCTrackerComponent::FAvaOutlinerRCTrackerComponent(FAvaOutliner& InOutliner, URemoteControlTrackerComponent* InComponent)
+FAvaOutlinerRCTrackerComponent::FAvaOutlinerRCTrackerComponent(IAvaOutliner& InOutliner, URemoteControlTrackerComponent* InComponent)
 	: FAvaOutlinerObject(InOutliner, InComponent)
 	, TrackerComponentWeak(InComponent)
 {
diff --git a/Engine/Plugins/Experimental/Avalanche/Source/AvalancheRemoteControlEditor/Private/Outliner/AvaOutlinerRCTrackerComponentProxy.cpp b/Engine/Plugins/Experimental/Avalanche/Source/AvalancheRemoteControlEditor/Private/Outliner/AvaOutlinerRCTrackerComponentProxy.cpp
--- a/Engine/Plugins/Experimental/Avalanche/Source/AvalancheRemoteControlEditor/Private/Outliner/AvaOutlinerRCTrackerComponentProxy.cpp
+++ b/Engine/Plugins/Experimental/Avalanche/Source/AvalancheRemoteControlEditor/Private/Outliner/AvaOutlinerRCTrackerComponentProxy.cpp
@@ -1,8 +1,8 @@
 // Copyright Epic Games, Inc. All Rights Reserved.
 
 #include "Outliner/AvaOutlinerRCTrackerComponentProxy.h"
-#include "AvaOutliner.h"
 #include "GameFramework/Actor.h"
+#include "IAvaOutliner.h"
 #include "Item/AvaOutlinerActor.h"
 #include "Outliner/AvaOutlinerRCTrackerComponent.h"
 #include "RemoteControlTrackerComponent.h"
@@ -11,7 +11,7 @@
 
 #define LOCTEXT_NAMESPACE "AvaOutlinerRemoteControlComponentProxy"
 
-FAvaOutlinerRCTrackerComponentProxy::FAvaOutlinerRCTrackerComponentProxy(FAvaOutliner& InOutliner, const FAvaOutlinerItemPtr& InParentItem)
+FAvaOutlinerRCTrackerComponentProxy::FAvaOutlinerRCTrackerComponentProxy(IAvaOutliner& InOutliner, const FAvaOutlinerItemPtr& InParentItem)
 	: Super(InOutliner, InParentItem)
 {
 	TrackerIcon = FRemoteControlComponentsEditorUtils::GetIcon("ClassIcon.RemoteControlTracker");
diff --git a/Engine/Plugins/Experimental/Avalanche/Source/AvalancheRemoteControlEditor/Private/Outliner/AvaOutlinerRCTrackerComponentProxy.h b/Engine/Plugins/Experimental/Avalanche/Source/AvalancheRemoteControlEditor/Private/Outliner/AvaOutlinerRCTrackerComponentProxy.h
--- a/Engine/Plugins/Experimental/Avalanche/Source/AvalancheRemoteControlEditor/Private/Outliner/AvaOutlinerRCTrackerComponentProxy.h
+++ b/Engine/Plugins/Experimental/Avalanche/Source/AvalancheRemoteControlEditor/Private/Outliner/AvaOutlinerRCTrackerComponentProxy.h
@@ -12,7 +12,7 @@ class FAvaOutlinerRCTrackerComponentProxy : public FAvaOutlinerItemProxy
 public:
 	UE_AVA_INHERITS_WITH_SUPER(FAvaOutlinerRCTrackerComponentProxy, FAvaOutlinerItemProxy);
 
-	FAvaOutlinerRCTrackerComponentProxy(FAvaOutliner& InOutliner, const FAvaOutlinerItemPtr& InParentItem);
+	FAvaOutlinerRCTrackerComponentProxy(IAvaOutliner& InOutliner, const FAvaOutlinerItemPtr& InParentItem);
 
 	URemoteControlTrackerComponent* GetTrackerComponent() const;
 	
diff --git a/Engine/Plugins/Experimental/Avalanche/Source/AvalancheRemoteControlEditor/Public/Outliner/AvaOutlinerRCTrackerComponent.h b/Engine/Plugins/Experimental/Avalanche/Source/AvalancheRemoteControlEditor/Public/Outliner/AvaOutlinerRCTrackerComponent.h
--- a/Engine/Plugins/Experimental/Avalanche/Source/AvalancheRemoteControlEditor/Public/Outliner/AvaOutlinerRCTrackerComponent.h
+++ b/Engine/Plugins/Experimental/Avalanche/Source/AvalancheRemoteControlEditor/Public/Outliner/AvaOutlinerRCTrackerComponent.h
@@ -19,7 +19,7 @@ public:
 	DECLARE_MULTICAST_DELEGATE_OneParam(FOnAvaOutlinerRCTrackerComponentSelected, const URemoteControlTrackerComponent*)
 	static FOnAvaOutlinerRCTrackerComponentSelected& OnRemoteControlTrackerSelected() { return OnRCTrackerSelectedDelegate; }
 	
-	FAvaOutlinerRCTrackerComponent(FAvaOutliner& InOutliner, URemoteControlTrackerComponent* InComponent);
+	FAvaOutlinerRCTrackerComponent(IAvaOutliner& InOutliner, URemoteControlTrackerComponent* InComponent);
 
 	//~ Begin IAvaOutlinerItem
 	virtual FSlateIcon GetIcon() const override;
diff --git a/Engine/Plugins/Experimental/Avalanche/Source/AvalancheSequencer/Private/AvaSequencerModule.cpp b/Engine/Plugins/Experimental/Avalanche/Source/AvalancheSequencer/Private/AvaSequencerModule.cpp
--- a/Engine/Plugins/Experimental/Avalanche/Source/AvalancheSequencer/Private/AvaSequencerModule.cpp
+++ b/Engine/Plugins/Experimental/Avalanche/Source/AvalancheSequencer/Private/AvaSequencerModule.cpp
@@ -1,20 +1,20 @@
 // Copyright Epic Games, Inc. All Rights Reserved.
 
 #include "AvaSequencerModule.h"
-#include "AvaOutliner.h"
 #include "AvaSequence.h"
 #include "AvaSequenceEditor.h"
 #include "AvaSequencer.h"
 #include "AvaSequencerUtils.h"
 #include "Commands/AvaSequencerCommands.h"
-#include "EaseCurveTool/AvaEaseCurveStyle.h"
-#include "EaseCurveTool/AvaEaseCurveToolCommands.h"
 #include "Customization/AvaDisplayRateCustomization.h"
 #include "Customization/AvaMarkSettingCustomization.h"
 #include "Customization/AvaSequenceCustomization.h"
 #include "Customization/AvaSequenceTimeCustomization.h"
 #include "Director/AvaSequenceDirectorBlueprint.h"
 #include "Director/AvaSequenceDirectorCompiler.h"
+#include "EaseCurveTool/AvaEaseCurveStyle.h"
+#include "EaseCurveTool/AvaEaseCurveToolCommands.h"
+#include "IAvaOutliner.h"
 #include "IAvaOutlinerModule.h"
 #include "ISequencerModule.h"
 #include "Item/AvaOutlinerActor.h"
@@ -146,7 +146,7 @@ void FAvaSequencerModule::RegisterOutlinerItems()
 	ItemProxyRegistry.RegisterItemProxyWithDefaultFactory<FAvaOutlinerSequenceProxy, 30>();
 
 	OutlinerProxiesExtensionDelegateHandle = IAvaOutlinerModule::Get().GetOnExtendItemProxiesForItem().AddLambda(
-		[](FAvaOutliner& InOutliner, const FAvaOutlinerItemPtr& InItem, TArray<TSharedPtr<FAvaOutlinerItemProxy>>& OutItemProxies)
+		[](IAvaOutliner& InOutliner, const FAvaOutlinerItemPtr& InItem, TArray<TSharedPtr<FAvaOutlinerItemProxy>>& OutItemProxies)
 		{
 			if (InItem->IsA<FAvaOutlinerActor>())
 			{
diff --git a/Engine/Plugins/Experimental/Avalanche/Source/AvalancheSequencer/Private/Outliner/AvaOutlinerSequence.cpp b/Engine/Plugins/Experimental/Avalanche/Source/AvalancheSequencer/Private/Outliner/AvaOutlinerSequence.cpp
--- a/Engine/Plugins/Experimental/Avalanche/Source/AvalancheSequencer/Private/Outliner/AvaOutlinerSequence.cpp
+++ b/Engine/Plugins/Experimental/Avalanche/Source/AvalancheSequencer/Private/Outliner/AvaOutlinerSequence.cpp
@@ -8,7 +8,7 @@
 #include "AvaSequence.h"
 #include "Styling/AppStyle.h"
 
-FAvaOutlinerSequence::FAvaOutlinerSequence(FAvaOutliner& InOutliner, UAvaSequence* InSequence, const FAvaOutlinerItemPtr& InReferencingItem)
+FAvaOutlinerSequence::FAvaOutlinerSequence(IAvaOutliner& InOutliner, UAvaSequence* InSequence, const FAvaOutlinerItemPtr& InReferencingItem)
 	: Super(InOutliner, InSequence, InReferencingItem, TEXT("Sequence"))
 	, Sequence(InSequence)
 {
diff --git a/Engine/Plugins/Experimental/Avalanche/Source/AvalancheSequencer/Private/Outliner/AvaOutlinerSequence.h b/Engine/Plugins/Experimental/Avalanche/Source/AvalancheSequencer/Private/Outliner/AvaOutlinerSequence.h
--- a/Engine/Plugins/Experimental/Avalanche/Source/AvalancheSequencer/Private/Outliner/AvaOutlinerSequence.h
+++ b/Engine/Plugins/Experimental/Avalanche/Source/AvalancheSequencer/Private/Outliner/AvaOutlinerSequence.h
@@ -17,7 +17,7 @@ class FAvaOutlinerSequence : public FAvaOutlinerObjectReference
 public:
 	UE_AVA_INHERITS_WITH_SUPER(FAvaOutlinerSequence, FAvaOutlinerObjectReference);
 
-	FAvaOutlinerSequence(FAvaOutliner& InOutliner, UAvaSequence* InSequence, const FAvaOutlinerItemPtr& InReferencingItem);
+	FAvaOutlinerSequence(IAvaOutliner& InOutliner, UAvaSequence* InSequence, const FAvaOutlinerItemPtr& InReferencingItem);
 
 	UAvaSequence* GetSequence() const { return Sequence.Get(IsIgnoringPendingKill()); }
 	
diff --git a/Engine/Plugins/Experimental/Avalanche/Source/AvalancheSequencer/Private/Outliner/AvaOutlinerSequenceProxy.cpp b/Engine/Plugins/Experimental/Avalanche/Source/AvalancheSequencer/Private/Outliner/AvaOutlinerSequenceProxy.cpp
--- a/Engine/Plugins/Experimental/Avalanche/Source/AvalancheSequencer/Private/Outliner/AvaOutlinerSequenceProxy.cpp
+++ b/Engine/Plugins/Experimental/Avalanche/Source/AvalancheSequencer/Private/Outliner/AvaOutlinerSequenceProxy.cpp
@@ -1,12 +1,12 @@
 // Copyright Epic Games, Inc. All Rights Reserved.
 
 #include "AvaOutlinerSequenceProxy.h"
-#include "AvaOutliner.h"
 #include "AvaSequence.h"
 #include "AvaSequencer.h"
 #include "AvaSequencerSubsystem.h"
 #include "AvaTypeSharedPointer.h"
 #include "Engine/World.h"
+#include "IAvaOutliner.h"
 #include "ISequencer.h"
 #include "Item/AvaOutlinerActor.h"
 #include "Outliner/AvaOutlinerSequence.h"
@@ -14,7 +14,7 @@
 
 #define LOCTEXT_NAMESPACE "AvaOutlinerSequenceProxy"
 
-FAvaOutlinerSequenceProxy::FAvaOutlinerSequenceProxy(FAvaOutliner& InOutliner, const FAvaOutlinerItemPtr& InParentItem)
+FAvaOutlinerSequenceProxy::FAvaOutlinerSequenceProxy(IAvaOutliner& InOutliner, const FAvaOutlinerItemPtr& InParentItem)
 	: Super(InOutliner, InParentItem)
 {
 	SequenceIcon = FSlateIcon(FAppStyle::GetAppStyleSetName(), "UMGEditor.AnimTabIcon");
diff --git a/Engine/Plugins/Experimental/Avalanche/Source/AvalancheSequencer/Private/Outliner/AvaOutlinerSequenceProxy.h b/Engine/Plugins/Experimental/Avalanche/Source/AvalancheSequencer/Private/Outliner/AvaOutlinerSequenceProxy.h
--- a/Engine/Plugins/Experimental/Avalanche/Source/AvalancheSequencer/Private/Outliner/AvaOutlinerSequenceProxy.h
+++ b/Engine/Plugins/Experimental/Avalanche/Source/AvalancheSequencer/Private/Outliner/AvaOutlinerSequenceProxy.h
@@ -18,7 +18,7 @@ class FAvaOutlinerSequenceProxy : public FAvaOutlinerItemProxy
 public:
 	UE_AVA_INHERITS_WITH_SUPER(FAvaOutlinerSequenceProxy, FAvaOutlinerItemProxy);
 
-	FAvaOutlinerSequenceProxy(FAvaOutliner& InOutliner, const FAvaOutlinerItemPtr& InParentItem);
+	FAvaOutlinerSequenceProxy(IAvaOutliner& InOutliner, const FAvaOutlinerItemPtr& InParentItem);
 
 	virtual ~FAvaOutlinerSequenceProxy() override;
 
