Avalanche Sequence:
- fixed issue with InitializeSequence playing entire range instead of just a single frame

diff --git a/Engine/Plugins/Experimental/Avalanche/Source/AvalancheSequence/Private/AvaSequencePlayer.cpp b/Engine/Plugins/Experimental/Avalanche/Source/AvalancheSequence/Private/AvaSequencePlayer.cpp
--- a/Engine/Plugins/Experimental/Avalanche/Source/AvalancheSequence/Private/AvaSequencePlayer.cpp
+++ b/Engine/Plugins/Experimental/Avalanche/Source/AvalancheSequence/Private/AvaSequencePlayer.cpp
@@ -61,14 +61,23 @@ void UAvaSequencePlayer::SetPlaySettings(const FAvaSequencePlayParams& InPlaySet
 	bReversePlayback = InPlaySettings.PlayMode == EAvaSequencePlayMode::Reverse;
 
 	const double TotalDuration    = MovieScene->GetTickResolution().AsSeconds(MovieScene->GetPlaybackRange().Size<FFrameNumber>());
-	const double StartTimeSeconds = InPlaySettings.Start.ToSeconds(*PlaybackSequence, *MovieScene, 0.0);
-	const double EndTimeSeconds   = InPlaySettings.End.ToSeconds(*PlaybackSequence, *MovieScene, TotalDuration);
+	double StartTimeSeconds = InPlaySettings.Start.ToSeconds(*PlaybackSequence, *MovieScene, 0.0);
+	double EndTimeSeconds   = InPlaySettings.End.ToSeconds(*PlaybackSequence, *MovieScene, TotalDuration);
+
+	// Ensure that StartTime <= EndTime
+	StartTimeSeconds = FMath::Min(StartTimeSeconds, EndTimeSeconds);
+	EndTimeSeconds   = FMath::Max(StartTimeSeconds, EndTimeSeconds);
 
 	double DurationSeconds = FMath::Abs(EndTimeSeconds - StartTimeSeconds);
 
 	// Ensure that Duration Seconds is at least 1 Frame Long, as this gets Rounded to the nearest Frame and if it results in 0, the Sequence does not get evaluated
 	DurationSeconds = FMath::Max(DurationSeconds, PlayPosition.GetInputRate().AsSeconds(FFrameNumber(1)));
 
+	if (bReversePlayback)
+	{
+		StartTimeSeconds = EndTimeSeconds - DurationSeconds;
+	}
+
 	SetTimeRange(StartTimeSeconds, DurationSeconds);
 
 	PlaybackSettings = FMovieSceneSequencePlaybackSettings();
diff --git a/Engine/Plugins/Experimental/Avalanche/Source/AvalancheSequence/Private/Transition/AvaTransitionInitializeSequence.cpp b/Engine/Plugins/Experimental/Avalanche/Source/AvalancheSequence/Private/Transition/AvaTransitionInitializeSequence.cpp
--- a/Engine/Plugins/Experimental/Avalanche/Source/AvalancheSequence/Private/Transition/AvaTransitionInitializeSequence.cpp
+++ b/Engine/Plugins/Experimental/Avalanche/Source/AvalancheSequence/Private/Transition/AvaTransitionInitializeSequence.cpp
@@ -2,6 +2,7 @@
 
 #include "Transition/AvaTransitionInitializeSequence.h"
 #include "AvaSequencePlaybackObject.h"
+#include "Math/NumericLimits.h"
 
 TArray<UAvaSequencePlayer*> FAvaTransitionInitializeSequence::ExecuteSequenceTask(FStateTreeExecutionContext& InContext) const
 {
@@ -12,7 +13,10 @@ TArray<UAvaSequencePlayer*> FAvaTransitionInitializeSequence::ExecuteSequenceTas
 	}
 
 	FAvaSequencePlayParams PlaySettings;
-	PlaySettings.Start = PlaySettings.End = InitializeTime;
+
+	// Set start to largest double, so that it gets clamped down to be at the End Time (i.e. time that the Sequence should evaluate)
+	PlaySettings.Start = FAvaSequenceTime(TNumericLimits<double>::Max());
+	PlaySettings.End = InitializeTime;
 	PlaySettings.PlayMode = PlayMode;
 
 	switch (QueryType)
