Avalanche Outliner:
- Trivial fix changed private -> protected

diff --git a/Engine/Plugins/Experimental/Avalanche/Source/AvalancheOutliner/Public/IAvaOutliner.h b/Engine/Plugins/Experimental/Avalanche/Source/AvalancheOutliner/Public/IAvaOutliner.h
--- a/Engine/Plugins/Experimental/Avalanche/Source/AvalancheOutliner/Public/IAvaOutliner.h
+++ b/Engine/Plugins/Experimental/Avalanche/Source/AvalancheOutliner/Public/IAvaOutliner.h
@@ -192,6 +192,6 @@ public:
 	/** Gets the World the Outliner is working with */
 	virtual UWorld* GetWorld() const = 0;
 
-private:
+protected:
 	virtual const FAvaOutlinerItemProxyRegistry& GetItemProxyRegistry() const = 0;
 };
