Avalanche:
- remove coreminimal includes

diff --git a/Engine/Plugins/Experimental/Avalanche/Source/AvalancheComponentVisualizers/Public/AvaVisBase.h b/Engine/Plugins/Experimental/Avalanche/Source/AvalancheComponentVisualizers/Public/AvaVisBase.h
--- a/Engine/Plugins/Experimental/Avalanche/Source/AvalancheComponentVisualizers/Public/AvaVisBase.h
+++ b/Engine/Plugins/Experimental/Avalanche/Source/AvalancheComponentVisualizers/Public/AvaVisBase.h
@@ -2,7 +2,6 @@
 
 #pragma once
 
-#include "CoreMinimal.h"
 #include "ComponentVisualizer.h"
 #include "UnrealWidgetFwd.h"
 
diff --git a/Engine/Plugins/Experimental/Avalanche/Source/AvalancheEditor/Private/AvaEditorCommands.h b/Engine/Plugins/Experimental/Avalanche/Source/AvalancheEditor/Private/AvaEditorCommands.h
--- a/Engine/Plugins/Experimental/Avalanche/Source/AvalancheEditor/Private/AvaEditorCommands.h
+++ b/Engine/Plugins/Experimental/Avalanche/Source/AvalancheEditor/Private/AvaEditorCommands.h
@@ -2,7 +2,6 @@
 
 #pragma once
 
-#include "CoreMinimal.h"
 #include "Framework/Commands/Commands.h"
 
 class FAvaEditorCommands : public TCommands<FAvaEditorCommands>
diff --git a/Engine/Plugins/Experimental/Avalanche/Source/AvalancheEditor/Private/AvaEditorSettings.h b/Engine/Plugins/Experimental/Avalanche/Source/AvalancheEditor/Private/AvaEditorSettings.h
--- a/Engine/Plugins/Experimental/Avalanche/Source/AvalancheEditor/Private/AvaEditorSettings.h
+++ b/Engine/Plugins/Experimental/Avalanche/Source/AvalancheEditor/Private/AvaEditorSettings.h
@@ -3,8 +3,12 @@
 #pragma once
 
 #include "AvaDefs.h"
-#include "CoreMinimal.h"
+#include "Containers/Array.h"
+#include "Containers/Map.h"
+#include "Containers/UnrealString.h"
+#include "Delegates/Delegate.h"
 #include "Engine/DeveloperSettings.h"
+#include "UObject/NameTypes.h"
 #include "Viewport/AvaViewportQualitySettings.h"
 #include "AvaEditorSettings.generated.h"
 
diff --git a/Engine/Plugins/Experimental/Avalanche/Source/AvalancheEditor/Private/DetailView/Customizations/AvaCategoryHiderCustomization.h b/Engine/Plugins/Experimental/Avalanche/Source/AvalancheEditor/Private/DetailView/Customizations/AvaCategoryHiderCustomization.h
--- a/Engine/Plugins/Experimental/Avalanche/Source/AvalancheEditor/Private/DetailView/Customizations/AvaCategoryHiderCustomization.h
+++ b/Engine/Plugins/Experimental/Avalanche/Source/AvalancheEditor/Private/DetailView/Customizations/AvaCategoryHiderCustomization.h
@@ -2,7 +2,6 @@
 
 #pragma once
 
-#include "CoreMinimal.h"
 #include "IDetailCustomization.h"
 
 class FAvaCategoryHiderCustomization : public IDetailCustomization
diff --git a/Engine/Plugins/Experimental/Avalanche/Source/AvalancheEditor/Private/DetailView/Customizations/AvaVectorPropertyTypeCustomization.h b/Engine/Plugins/Experimental/Avalanche/Source/AvalancheEditor/Private/DetailView/Customizations/AvaVectorPropertyTypeCustomization.h
--- a/Engine/Plugins/Experimental/Avalanche/Source/AvalancheEditor/Private/DetailView/Customizations/AvaVectorPropertyTypeCustomization.h
+++ b/Engine/Plugins/Experimental/Avalanche/Source/AvalancheEditor/Private/DetailView/Customizations/AvaVectorPropertyTypeCustomization.h
@@ -2,7 +2,6 @@
 
 #pragma once
 
-#include "CoreMinimal.h"
 #include "IPropertyTypeCustomization.h"
 #include "PropertyEditorModule.h"
 #include "PropertyHandle.h"
diff --git a/Engine/Plugins/Experimental/Avalanche/Source/AvalancheLevelViewport/Private/Widgets/AvaLevelViewportComponentTransformDetails.h b/Engine/Plugins/Experimental/Avalanche/Source/AvalancheLevelViewport/Private/Widgets/AvaLevelViewportComponentTransformDetails.h
--- a/Engine/Plugins/Experimental/Avalanche/Source/AvalancheLevelViewport/Private/Widgets/AvaLevelViewportComponentTransformDetails.h
+++ b/Engine/Plugins/Experimental/Avalanche/Source/AvalancheLevelViewport/Private/Widgets/AvaLevelViewportComponentTransformDetails.h
@@ -2,7 +2,6 @@
 
 #pragma once
 
-#include "CoreMinimal.h"
 #include "AssetSelection.h"
 #include "Framework/Commands/UIAction.h"
 #include "IDetailCustomNodeBuilder.h"
@@ -110,7 +109,7 @@ private:
 	 * @param Axis				Bitfield of which axis to set, can be multiple
 	 * @param NewValue			The new vector values, it only uses the ones with specified axis
 	 * @param bMirror			If true, set the value to it's inverse instead of using NewValue
-	 * @param bCommittted		True if the value was committed, false is the value comes from the slider
+	 * @param bCommitted		True if the value was committed, false is the value comes from the slider
 	 */
 	void OnSetTransform(EAvaLevelViewportTransformFieldType TransformField, EAxisList::Type Axis, FVector NewValue, bool bMirror, bool bCommitted);
 
@@ -124,7 +123,7 @@ private:
 	 * @param NewValue			The new translation value
 	 * @param CommitInfo		Whether or not this was committed from pressing enter or losing focus
 	 * @param Axis				Bitfield of which axis to set, can be multiple
-	 * @param bCommittted		true if the value was committed, false is the value comes from the slider
+	 * @param bCommitted		true if the value was committed, false is the value comes from the slider
 	 */
 	void OnSetTransformAxis(FVector::FReal NewValue, ETextCommit::Type CommitInfo, EAvaLevelViewportTransformFieldType TransformField, EAxisList::Type Axis, bool bCommitted);
 
@@ -135,11 +134,10 @@ private:
 	 * @param NewValue			The new translation value
 	 * @param CommitInfo		Whether or not this was committed from pressing enter or losing focus
 	 * @param Axis				Bitfield of which axis to set, can be multiple
-	 * @param bCommittted		true if the value was committed, false is the value comes from the slider
+	 * @param bCommitted		true if the value was committed, false is the value comes from the slider
 	 */
 	void OnSetTransformAxisFloat(float NewValue, ETextCommit::Type CommitInfo, EAvaLevelViewportTransformFieldType TransformField, EAxisList::Type Axis, bool bCommitted);
 
-
 	/**
 	 * Helper to begin a new transaction for a slider interaction.
 	 * @param ActorTransaction		The name to give the transaction when changing an actor transform.
@@ -210,7 +208,8 @@ private:
 	/**
 	 * Sets relative transform on the specified field
 	 *
-	 * @param The field that should be set to relative
+	 * @param TransformField the field that should be set to relative
+	 * @param bAbsoluteEnabled true to absolute, false for relative 
 	 */
 	void OnSetAbsoluteTransform(EAvaLevelViewportTransformFieldType TransformField, bool bAbsoluteEnabled);
 
diff --git a/Engine/Plugins/Experimental/Avalanche/Source/AvalancheLevelViewport/Private/Widgets/SAvaEyeDropperButton.h b/Engine/Plugins/Experimental/Avalanche/Source/AvalancheLevelViewport/Private/Widgets/SAvaEyeDropperButton.h
--- a/Engine/Plugins/Experimental/Avalanche/Source/AvalancheLevelViewport/Private/Widgets/SAvaEyeDropperButton.h
+++ b/Engine/Plugins/Experimental/Avalanche/Source/AvalancheLevelViewport/Private/Widgets/SAvaEyeDropperButton.h
@@ -2,7 +2,6 @@
 
 #pragma once
 
-#include "CoreMinimal.h"
 #include "Misc/Attribute.h"
 #include "Input/CursorReply.h"
 #include "Input/Reply.h"
@@ -92,8 +91,8 @@ private:
 	TOptional<FVector2f> LastCursorPosition;
 
 	// Dropper states
-	bool bWasClicked;
-	bool bWasClickActivated;
-	bool bWasLeft;
-	bool bWasReEntered;
+	bool bWasClicked = false;
+	bool bWasClickActivated = false;
+	bool bWasLeft = false;
+	bool bWasReEntered = false;
 };
diff --git a/Engine/Plugins/Experimental/Avalanche/Source/AvalancheMaskEditor/Private/AvaMaskEditorStyle.h b/Engine/Plugins/Experimental/Avalanche/Source/AvalancheMaskEditor/Private/AvaMaskEditorStyle.h
--- a/Engine/Plugins/Experimental/Avalanche/Source/AvalancheMaskEditor/Private/AvaMaskEditorStyle.h
+++ b/Engine/Plugins/Experimental/Avalanche/Source/AvalancheMaskEditor/Private/AvaMaskEditorStyle.h
@@ -2,8 +2,8 @@
 
 #pragma once
 
-#include "CoreMinimal.h"
 #include "Styling/SlateStyle.h"
+#include "Templates/SharedPointer.h"
 
 class FSlateStyleSet;
 
diff --git a/Engine/Plugins/Experimental/Avalanche/Source/AvalancheShapesEditor/Private/Visualizers/AvaShapeConeDynMeshVis.h b/Engine/Plugins/Experimental/Avalanche/Source/AvalancheShapesEditor/Private/Visualizers/AvaShapeConeDynMeshVis.h
--- a/Engine/Plugins/Experimental/Avalanche/Source/AvalancheShapesEditor/Private/Visualizers/AvaShapeConeDynMeshVis.h
+++ b/Engine/Plugins/Experimental/Avalanche/Source/AvalancheShapesEditor/Private/Visualizers/AvaShapeConeDynMeshVis.h
@@ -2,7 +2,6 @@
 
 #pragma once
 
-#include "CoreMinimal.h"
 #include "Visualizers/AvaShape3DDynMeshVis.h"
 
 class UAvaShapeConeDynamicMesh;
diff --git a/Engine/Plugins/Experimental/Avalanche/Source/AvalancheShapesEditor/Private/Visualizers/AvaShapeCubeDynMeshVis.h b/Engine/Plugins/Experimental/Avalanche/Source/AvalancheShapesEditor/Private/Visualizers/AvaShapeCubeDynMeshVis.h
--- a/Engine/Plugins/Experimental/Avalanche/Source/AvalancheShapesEditor/Private/Visualizers/AvaShapeCubeDynMeshVis.h
+++ b/Engine/Plugins/Experimental/Avalanche/Source/AvalancheShapesEditor/Private/Visualizers/AvaShapeCubeDynMeshVis.h
@@ -2,7 +2,6 @@
 
 #pragma once
 
-#include "CoreMinimal.h"
 #include "Visualizers/AvaShape3DDynMeshVis.h"
 
 class UAvaShapeCubeDynamicMesh;
diff --git a/Engine/Plugins/Experimental/Avalanche/Source/AvalancheShapesEditor/Private/Visualizers/AvaShapeEllipseDynMeshVis.h b/Engine/Plugins/Experimental/Avalanche/Source/AvalancheShapesEditor/Private/Visualizers/AvaShapeEllipseDynMeshVis.h
--- a/Engine/Plugins/Experimental/Avalanche/Source/AvalancheShapesEditor/Private/Visualizers/AvaShapeEllipseDynMeshVis.h
+++ b/Engine/Plugins/Experimental/Avalanche/Source/AvalancheShapesEditor/Private/Visualizers/AvaShapeEllipseDynMeshVis.h
@@ -2,7 +2,6 @@
 
 #pragma once
 
-#include "CoreMinimal.h"
 #include "Visualizers/AvaShape2DDynMeshVis.h"
 
 class UAvaShapeEllipseDynamicMesh;
diff --git a/Engine/Plugins/Experimental/Avalanche/Source/AvalancheShapesEditor/Private/Visualizers/AvaShapeIrregularPolygonDynMeshVis.h b/Engine/Plugins/Experimental/Avalanche/Source/AvalancheShapesEditor/Private/Visualizers/AvaShapeIrregularPolygonDynMeshVis.h
--- a/Engine/Plugins/Experimental/Avalanche/Source/AvalancheShapesEditor/Private/Visualizers/AvaShapeIrregularPolygonDynMeshVis.h
+++ b/Engine/Plugins/Experimental/Avalanche/Source/AvalancheShapesEditor/Private/Visualizers/AvaShapeIrregularPolygonDynMeshVis.h
@@ -2,7 +2,6 @@
 
 #pragma once
 
-#include "CoreMinimal.h"
 #include "Visualizers/AvaShape2DDynMeshVis.h"
 
 struct HAvaShapeIrregularPolygonPointHitProxy : HAvaHitProxy
diff --git a/Engine/Plugins/Experimental/Avalanche/Source/AvalancheShapesEditor/Private/Visualizers/AvaShapeLineDynMeshVis.h b/Engine/Plugins/Experimental/Avalanche/Source/AvalancheShapesEditor/Private/Visualizers/AvaShapeLineDynMeshVis.h
--- a/Engine/Plugins/Experimental/Avalanche/Source/AvalancheShapesEditor/Private/Visualizers/AvaShapeLineDynMeshVis.h
+++ b/Engine/Plugins/Experimental/Avalanche/Source/AvalancheShapesEditor/Private/Visualizers/AvaShapeLineDynMeshVis.h
@@ -2,7 +2,6 @@
 
 #pragma once
 
-#include "CoreMinimal.h"
 #include "AvaShapeRoundedPolygonDynMeshVis.h"
 
 struct HAvaShapeLineEndHitProxy : HAvaHitProxy
diff --git a/Engine/Plugins/Experimental/Avalanche/Source/AvalancheShapesEditor/Private/Visualizers/AvaShapeNGonDynMeshVis.h b/Engine/Plugins/Experimental/Avalanche/Source/AvalancheShapesEditor/Private/Visualizers/AvaShapeNGonDynMeshVis.h
--- a/Engine/Plugins/Experimental/Avalanche/Source/AvalancheShapesEditor/Private/Visualizers/AvaShapeNGonDynMeshVis.h
+++ b/Engine/Plugins/Experimental/Avalanche/Source/AvalancheShapesEditor/Private/Visualizers/AvaShapeNGonDynMeshVis.h
@@ -2,7 +2,6 @@
 
 #pragma once
 
-#include "CoreMinimal.h"
 #include "AvaShapeRoundedPolygonDynMeshVis.h"
 
 class UAvaShapeNGonDynamicMesh;
diff --git a/Engine/Plugins/Experimental/Avalanche/Source/AvalancheShapesEditor/Private/Visualizers/AvaShapeRectangleDynMeshVis.h b/Engine/Plugins/Experimental/Avalanche/Source/AvalancheShapesEditor/Private/Visualizers/AvaShapeRectangleDynMeshVis.h
--- a/Engine/Plugins/Experimental/Avalanche/Source/AvalancheShapesEditor/Private/Visualizers/AvaShapeRectangleDynMeshVis.h
+++ b/Engine/Plugins/Experimental/Avalanche/Source/AvalancheShapesEditor/Private/Visualizers/AvaShapeRectangleDynMeshVis.h
@@ -2,7 +2,6 @@
 
 #pragma once
 
-#include "CoreMinimal.h"
 #include "AvaDefs.h"
 #include "Visualizers/AvaShape2DDynMeshVis.h"
 
diff --git a/Engine/Plugins/Experimental/Avalanche/Source/AvalancheShapesEditor/Private/Visualizers/AvaShapeRingDynMeshVis.h b/Engine/Plugins/Experimental/Avalanche/Source/AvalancheShapesEditor/Private/Visualizers/AvaShapeRingDynMeshVis.h
--- a/Engine/Plugins/Experimental/Avalanche/Source/AvalancheShapesEditor/Private/Visualizers/AvaShapeRingDynMeshVis.h
+++ b/Engine/Plugins/Experimental/Avalanche/Source/AvalancheShapesEditor/Private/Visualizers/AvaShapeRingDynMeshVis.h
@@ -2,7 +2,6 @@
 
 #pragma once
 
-#include "CoreMinimal.h"
 #include "Visualizers/AvaShape2DDynMeshVis.h"
 
 class UAvaShapeRingDynamicMesh;
diff --git a/Engine/Plugins/Experimental/Avalanche/Source/AvalancheShapesEditor/Private/Visualizers/AvaShapeRoundedPolygonDynMeshVis.h b/Engine/Plugins/Experimental/Avalanche/Source/AvalancheShapesEditor/Private/Visualizers/AvaShapeRoundedPolygonDynMeshVis.h
--- a/Engine/Plugins/Experimental/Avalanche/Source/AvalancheShapesEditor/Private/Visualizers/AvaShapeRoundedPolygonDynMeshVis.h
+++ b/Engine/Plugins/Experimental/Avalanche/Source/AvalancheShapesEditor/Private/Visualizers/AvaShapeRoundedPolygonDynMeshVis.h
@@ -2,7 +2,6 @@
 
 #pragma once
 
-#include "CoreMinimal.h"
 #include "Visualizers/AvaShape2DDynMeshVis.h"
 
 class UAvaShapeRoundedPolygonDynamicMesh;
diff --git a/Engine/Plugins/Experimental/Avalanche/Source/AvalancheShapesEditor/Private/Visualizers/AvaShapeSphereDynMeshVis.h b/Engine/Plugins/Experimental/Avalanche/Source/AvalancheShapesEditor/Private/Visualizers/AvaShapeSphereDynMeshVis.h
--- a/Engine/Plugins/Experimental/Avalanche/Source/AvalancheShapesEditor/Private/Visualizers/AvaShapeSphereDynMeshVis.h
+++ b/Engine/Plugins/Experimental/Avalanche/Source/AvalancheShapesEditor/Private/Visualizers/AvaShapeSphereDynMeshVis.h
@@ -3,7 +3,6 @@
 
 #pragma once
 
-#include "CoreMinimal.h"
 #include "Visualizers/AvaShape3DDynMeshVis.h"
 
 class UAvaShapeSphereDynamicMesh;
diff --git a/Engine/Plugins/Experimental/Avalanche/Source/AvalancheShapesEditor/Private/Visualizers/AvaShapeStarDynMeshVis.h b/Engine/Plugins/Experimental/Avalanche/Source/AvalancheShapesEditor/Private/Visualizers/AvaShapeStarDynMeshVis.h
--- a/Engine/Plugins/Experimental/Avalanche/Source/AvalancheShapesEditor/Private/Visualizers/AvaShapeStarDynMeshVis.h
+++ b/Engine/Plugins/Experimental/Avalanche/Source/AvalancheShapesEditor/Private/Visualizers/AvaShapeStarDynMeshVis.h
@@ -2,7 +2,6 @@
 
 #pragma once
 
-#include "CoreMinimal.h"
 #include "AvaShapeRoundedPolygonDynMeshVis.h"
 
 class UAvaShapeStarDynamicMesh;
diff --git a/Engine/Plugins/Experimental/Avalanche/Source/AvalancheShapesEditor/Private/Visualizers/AvaShapeTorusDynMeshVis.h b/Engine/Plugins/Experimental/Avalanche/Source/AvalancheShapesEditor/Private/Visualizers/AvaShapeTorusDynMeshVis.h
--- a/Engine/Plugins/Experimental/Avalanche/Source/AvalancheShapesEditor/Private/Visualizers/AvaShapeTorusDynMeshVis.h
+++ b/Engine/Plugins/Experimental/Avalanche/Source/AvalancheShapesEditor/Private/Visualizers/AvaShapeTorusDynMeshVis.h
@@ -2,7 +2,6 @@
 
 #pragma once
 
-#include "CoreMinimal.h"
 #include "Visualizers/AvaShape3DDynMeshVis.h"
 
 class UAvaShapeTorusDynamicMesh;
diff --git a/Engine/Plugins/Experimental/Avalanche/Source/AvalancheShapesEditor/Public/Visualizers/AvaShape2DDynMeshVis.h b/Engine/Plugins/Experimental/Avalanche/Source/AvalancheShapesEditor/Public/Visualizers/AvaShape2DDynMeshVis.h
--- a/Engine/Plugins/Experimental/Avalanche/Source/AvalancheShapesEditor/Public/Visualizers/AvaShape2DDynMeshVis.h
+++ b/Engine/Plugins/Experimental/Avalanche/Source/AvalancheShapesEditor/Public/Visualizers/AvaShape2DDynMeshVis.h
@@ -2,7 +2,6 @@
 
 #pragma once
 
-#include "CoreMinimal.h"
 #include "AvaShapeDynMeshVis.h"
 #include "DynamicMeshes/AvaShape2DDynMeshBase.h"
 
diff --git a/Engine/Plugins/Experimental/Avalanche/Source/AvalancheShapesEditor/Public/Visualizers/AvaShape3DDynMeshVis.h b/Engine/Plugins/Experimental/Avalanche/Source/AvalancheShapesEditor/Public/Visualizers/AvaShape3DDynMeshVis.h
--- a/Engine/Plugins/Experimental/Avalanche/Source/AvalancheShapesEditor/Public/Visualizers/AvaShape3DDynMeshVis.h
+++ b/Engine/Plugins/Experimental/Avalanche/Source/AvalancheShapesEditor/Public/Visualizers/AvaShape3DDynMeshVis.h
@@ -2,7 +2,6 @@
 
 #pragma once
 
-#include "CoreMinimal.h"
 #include "AvaShapeDynMeshVis.h"
 
 class UAvaShape3DDynMeshBase;
diff --git a/Engine/Plugins/Experimental/Avalanche/Source/AvalancheShapesEditor/Public/Visualizers/AvaShapeDynMeshVis.h b/Engine/Plugins/Experimental/Avalanche/Source/AvalancheShapesEditor/Public/Visualizers/AvaShapeDynMeshVis.h
--- a/Engine/Plugins/Experimental/Avalanche/Source/AvalancheShapesEditor/Public/Visualizers/AvaShapeDynMeshVis.h
+++ b/Engine/Plugins/Experimental/Avalanche/Source/AvalancheShapesEditor/Public/Visualizers/AvaShapeDynMeshVis.h
@@ -2,7 +2,6 @@
 
 #pragma once
 
-#include "CoreMinimal.h"
 #include "AvaDefs.h"
 #include "AvaVisBase.h"
 
diff --git a/Engine/Plugins/Experimental/Avalanche/Source/AvalancheTextEditor/Private/DetailsView/AvaLinearGradientSettingsCustomization.h b/Engine/Plugins/Experimental/Avalanche/Source/AvalancheTextEditor/Private/DetailsView/AvaLinearGradientSettingsCustomization.h
--- a/Engine/Plugins/Experimental/Avalanche/Source/AvalancheTextEditor/Private/DetailsView/AvaLinearGradientSettingsCustomization.h
+++ b/Engine/Plugins/Experimental/Avalanche/Source/AvalancheTextEditor/Private/DetailsView/AvaLinearGradientSettingsCustomization.h
@@ -2,7 +2,6 @@
 
 #pragma once
 
-#include "CoreMinimal.h"
 #include "IPropertyTypeCustomization.h"
 
 class SBox;
diff --git a/Engine/Plugins/Experimental/Avalanche/Source/AvalancheTextEditor/Private/Visualizer/AvaTextVisualizer.h b/Engine/Plugins/Experimental/Avalanche/Source/AvalancheTextEditor/Private/Visualizer/AvaTextVisualizer.h
--- a/Engine/Plugins/Experimental/Avalanche/Source/AvalancheTextEditor/Private/Visualizer/AvaTextVisualizer.h
+++ b/Engine/Plugins/Experimental/Avalanche/Source/AvalancheTextEditor/Private/Visualizer/AvaTextVisualizer.h
@@ -2,9 +2,9 @@
 
 #pragma once
 
-#include "CoreMinimal.h"
 #include "AvaVisBase.h"
 
+class FProperty;
 class UAvaText3DComponent;
 class UText3DComponent;
 
