[Motion Design]
 - Adjust quick preset icon placement
 - Fix LOCTEXT empty identifier
 - Fix copy of tangents

diff --git a/Engine/Plugins/Experimental/Avalanche/Source/AvalancheSequencer/Private/EaseCurveTool/AvaEaseCurveTool.cpp b/Engine/Plugins/Experimental/Avalanche/Source/AvalancheSequencer/Private/EaseCurveTool/AvaEaseCurveTool.cpp
--- a/Engine/Plugins/Experimental/Avalanche/Source/AvalancheSequencer/Private/EaseCurveTool/AvaEaseCurveTool.cpp
+++ b/Engine/Plugins/Experimental/Avalanche/Source/AvalancheSequencer/Private/EaseCurveTool/AvaEaseCurveTool.cpp
@@ -516,7 +516,7 @@ bool FAvaEaseCurveTool::CanCopyTangentsToClipboard() const
 
 void FAvaEaseCurveTool::CopyTangentsToClipboard()
 {
-	FPlatformApplicationMisc::ClipboardCopy(*EaseCurve->GetTangents().ToDisplayString());
+	FPlatformApplicationMisc::ClipboardCopy(*EaseCurve->GetTangents().ToJson());
 
 	ShowNotificationMessage(LOCTEXT("EaseCurveToolTangentsCopied", "Ease Curve Tool Tangents Copied!"));
 }
diff --git a/Engine/Plugins/Experimental/Avalanche/Source/AvalancheSequencer/Private/EaseCurveTool/Widgets/SAvaEaseCurvePreset.cpp b/Engine/Plugins/Experimental/Avalanche/Source/AvalancheSequencer/Private/EaseCurveTool/Widgets/SAvaEaseCurvePreset.cpp
--- a/Engine/Plugins/Experimental/Avalanche/Source/AvalancheSequencer/Private/EaseCurveTool/Widgets/SAvaEaseCurvePreset.cpp
+++ b/Engine/Plugins/Experimental/Avalanche/Source/AvalancheSequencer/Private/EaseCurveTool/Widgets/SAvaEaseCurvePreset.cpp
@@ -323,7 +323,7 @@ TSharedRef<SWidget> SAvaEaseCurvePreset::GenerateSelectedRowWidget() const
 		.AutoWidth()
 		.HAlign(HAlign_Right)
 		.VAlign(VAlign_Center)
-		.Padding(5.f, 2.f, 4.f, 2.f)
+		.Padding(0.f, 2.f, 5.f, 2.f)
 		[
 			SNew(SBorder)
 			.BorderBackgroundColor(FStyleColors::White25)
@@ -341,6 +341,7 @@ TSharedRef<SWidget> SAvaEaseCurvePreset::GenerateSelectedRowWidget() const
 		+ SHorizontalBox::Slot()
 		.HAlign(HAlign_Fill)
 		.VAlign(VAlign_Center)
+		.Padding(0.f, 0.f, 5.f, 0.f)
 		[
 			SNew(SHorizontalBox)
 			+ SHorizontalBox::Slot()
diff --git a/Engine/Plugins/Experimental/Avalanche/Source/AvalancheSequencer/Private/EaseCurveTool/Widgets/SAvaEaseCurvePresetGroupItem.cpp b/Engine/Plugins/Experimental/Avalanche/Source/AvalancheSequencer/Private/EaseCurveTool/Widgets/SAvaEaseCurvePresetGroupItem.cpp
--- a/Engine/Plugins/Experimental/Avalanche/Source/AvalancheSequencer/Private/EaseCurveTool/Widgets/SAvaEaseCurvePresetGroupItem.cpp
+++ b/Engine/Plugins/Experimental/Avalanche/Source/AvalancheSequencer/Private/EaseCurveTool/Widgets/SAvaEaseCurvePresetGroupItem.cpp
@@ -71,7 +71,7 @@ void SAvaEaseCurvePresetGroupItem::Construct(const FArguments& InArgs, const TSh
 							.VAlign(VAlign_Center)
 							[
 								SNew(SBox)
-								.Padding(0.f, 0.f, 0.f, 3.f)
+								.Padding(0.f, 0.f, 3.f, 0.f)
 								.Visibility(this, &SAvaEaseCurvePresetGroupItem::GetQuickPresetIconVisibility)
 								[
 									SNew(SImage)
@@ -286,7 +286,7 @@ EVisibility SAvaEaseCurvePresetGroupItem::GetQuickPresetIconVisibility() const
 		return EVisibility::Hidden;
 	}
 
-	return (Preset->Tangents == Tangents) ? EVisibility::Visible : EVisibility::Hidden;
+	return (Preset->Tangents == Tangents) ? EVisibility::Visible : EVisibility::Collapsed;
 }
 
 FText SAvaEaseCurvePresetGroupItem::GetQuickPresetIconToolTip() const
@@ -315,7 +315,7 @@ FText SAvaEaseCurvePresetGroupItem::GetQuickPresetIconToolTip() const
 			, EaseCurveToolCommands.QuickEaseOut->GetInputText());
 	}
 
-	return CommandText.IsEmpty() ? QuickEaseText : FText::Format(LOCTEXT("", "{0}\n\n{1}"), QuickEaseText, CommandText);
+	return CommandText.IsEmpty() ? QuickEaseText : FText::Format(LOCTEXT("QuickEasePresetIconTooltip", "{0}\n\n{1}"), QuickEaseText, CommandText);
 }
 
 #undef LOCTEXT_NAMESPACE
