Motion Design
- Media module refactor to remove Avalanche names, organise into correct submodules and homogenise prefixes.

diff --git a/Engine/Plugins/Experimental/Avalanche/Config/DefaultAvalanche.ini b/Engine/Plugins/Experimental/Avalanche/Config/DefaultAvalanche.ini
--- a/Engine/Plugins/Experimental/Avalanche/Config/DefaultAvalanche.ini
+++ b/Engine/Plugins/Experimental/Avalanche/Config/DefaultAvalanche.ini
@@ -14,7 +14,7 @@
 +StructRedirects=(OldName="/Script/Avalanche.AvaAnimBindingReferences",NewName="/Script/AvalancheSequence.AvaAnimBindingReferences")
 +StructRedirects=(OldName="/Script/Avalanche.AvaAnimPlayParams",NewName="/Script/AvalancheSequence.AvaSequencePlayParams")
 
-+PropertyRedirects=(OldName="AvalancheRemoteControlValues.Values", NewName="AvalancheRemoteControlValues.EntityValues")
++PropertyRedirects=(OldName="AvalancheRemoteControlValues.Values", NewName="AvaPlayableRemoveControls.EntityValues")
 +PropertyRedirects=(OldName="AvaMediaOutputInfo.HostName", NewName="AvaMediaOutputInfo.ServerName")
 +PropertyRedirects=(OldName="/Script/Avalanche.AvalanchePlaylist.Pages", NewName="Pages_DEPRECATED")
 +PropertyRedirects=(OldName="/Script/AvalancheSceneTree.AvaSceneItem.ObjectPath", NewName="Id")
@@ -161,7 +161,7 @@
 +StructRedirects=(OldName="/Script/Avalanche.AvalancheActorData",NewName="/Script/Avalanche.AvaActorData")
 
 +StructRedirects=(OldName="/Script/Avalanche.AvalancheInstanceSettings",NewName="/Script/Avalanche.AvaInstanceSettings")
-+PropertyRedirects=(OldName="/Script/AvalancheMedia.AvalancheMediaSettings.AvalancheInstanceSettings",NewName="/Script/AvalancheMedia.AvalancheMediaSettings.AvaInstanceSettings")
++PropertyRedirects=(OldName="/Script/AvalancheMedia.AvalancheMediaSettings.AvalancheInstanceSettings",NewName="/Script/AvalancheMedia.AvaMediaSettings.AvaInstanceSettings")
 
 +FunctionRedirects=(OldName="/Script/AvalancheText.AvaText3DComponent.SetAvalancheFont",NewName="/Script/AvalancheText.AvaText3DComponent.SetMotionDesignFont")
 +PropertyRedirects=(OldName="/Script/AvalancheText.AvaText3DComponent.AvalancheFont",NewName="/Script/AvalancheText.AvaText3DComponent.MotionDesignFont")
@@ -174,5 +174,140 @@
 
 +StructRedirects=(OldName="/Script/Avalanche.AvaSceneContainsAttributeConditionInstanceData",NewName="/Script/Avalanche.AvaSceneContainsTagAttributeConditionInstanceData")
 
++StructRedirects=(OldName="/Script/AvalancheMedia.AnimPlaySettings",NewName="/Script/AvalancheMedia.AvaPlaybackAnimPlaySettings")
++StructRedirects=(OldName="/Script/AvalancheMedia.AvalancheAnimations",NewName="/Script/AvalancheMedia.AvaPlaybackAnimations")
++ClassRedirects=(OldName="/Script/AvalancheMedia.AvalanchePlayable",NewName="/Script/AvalancheMedia.AvaPlayable")
++ClassRedirects=(OldName="/Script/AvalancheMedia.AvalancheBlueprintPlayable",NewName="/Script/AvalancheMedia.AvaPlayableBlueprint")
++ClassRedirects=(OldName="/Script/AvalancheMedia.AvalancheBroadcast",NewName="/Script/AvalancheMedia.AvaBroadcast")
++FunctionRedirects=(OldName="/Script/AvalancheMedia.AvaBroadcast.GetAvalancheBroadcast",NewName="/Script/AvalancheMedia.AvaBroadcast.GetBroadcast")
++ClassRedirects=(OldName="/Script/AvalancheMedia.AvalancheBroadcastComponent",NewName="/Script/AvalancheMedia.AvaBroadcastComponent")
++StructRedirects=(OldName="/Script/AvalancheMedia.AvaBroadcastProfile",NewName="/Script/AvalancheMedia.AvaBroadcastProfile")
++EnumRedirects=(OldName="/Script/AvalancheMedia.EAvaChannelState",NewName="/Script/AvalancheMedia.EAvaBroadcastChannelState")
++EnumRedirects=(OldName="/Script/AvalancheMedia.EAvaBroadcastChannelType",NewName="/Script/AvalancheMedia.EAvaBroadcastChannelType")
++EnumRedirects=(OldName="/Script/AvalancheMedia.EAvalanchePlayableStatus",NewName="/Script/AvalancheMedia.EAvaPlayableStatus")
++EnumRedirects=(OldName="/Script/AvalancheMedia.EAvalanchePlayableSequenceEventType",NewName="/Script/AvalancheMedia.EAvaPlayableSequenceEventType")
++EnumRedirects=(OldName="/Script/AvalancheMedia.EAvaPageListType",NewName="/Script/AvalancheMedia.EAvaRundownPageListType")
++StructRedirects=(OldName="/Script/AvalancheMedia.AvaPageListReference",NewName="/Script/AvalancheMedia.AvaRundownPageListReference")
++StructRedirects=(OldName="/Script/AvalancheMedia.AvaDeviceProviderData",NewName="/Script/AvalancheMedia.AvaBroadcastDeviceProviderData")
++ClassRedirects=(OldName="/Script/AvalancheMedia.AvaDisplayMediaCapture",NewName="/Script/AvalancheMedia.AvaBroadcastDisplayMediaCapture")
++ClassRedirects=(OldName="/Script/AvalancheMedia.AvaDisplayMediaOutput",NewName="/Script/AvalancheMedia.AvaBroadcastDisplayMediaOutput")
++ClassRedirects=(OldName="/Script/AvalancheMedia.AvalancheLevelStreamingPlayable",NewName="/Script/AvalancheMedia.AvaPlayableLevelStreaming")
++ClassRedirects=(OldName="/Script/AvalancheMedia.AvalanchePlayableTransition",NewName="/Script/AvalancheMedia.AvaPlayableTransition")
++ClassRedirects=(OldName="/Script/AvalancheMedia.AvaLocalPlayableTransition",NewName="/Script/AvalancheMedia.AvaPlayableLocalTransition")
++StructRedirects=(OldName="/Script/AvalancheMedia.AvaOutputChannel",NewName="/Script/AvalancheMedia.AvaBroadcastOutputChannel")
++EnumRedirects=(OldName="/Script/AvalancheMedia.EAvalanchePageStatus",NewName="/Script/AvalancheMedia.EAvaRundownPageStatus")
++StructRedirects=(OldName="/Script/AvalancheMedia.AvalanchePageStatus",NewName="/Script/AvalancheMedia.AvaRundownChannelPageStatus")
++StructRedirects=(OldName="/Script/AvalancheMedia.AvalanchePage",NewName="/Script/AvalancheMedia.AvaRundownPage")
++ClassRedirects=(OldName="/Script/AvalancheMedia.UAvaRundownPlaybackInstancePlayer",NewName="/Script/AvalancheMedia.AvaRundownPlaybackInstancePlayer")
++ClassRedirects=(OldName="/Script/AvalancheMedia.AvalanchePagePlayer",NewName="/Script/AvalancheMedia.AvaRundownPagePlayer")
++ClassRedirects=(OldName="/Script/AvalancheMedia.AvalanchePageTransition",NewName="/Script/AvalancheMedia.AvaRundownPageTransition")
++EnumRedirects=(OldName="/Script/AvalancheMedia.EAvaPlaybackStopOptions",NewName="/Script/AvalancheMedia.EAvaPlaybackStopOptions")
++EnumRedirects=(OldName="/Script/AvalancheMedia.EAvaPlaybackUnloadOptions",NewName="/Script/AvalancheMedia.EAvaPlaybackUnloadOptions")
++StructRedirects=(OldName="/Script/AvalancheMedia.AvalanchePlaybackPlayableGroup",NewName="/Script/AvalancheMedia.AvaPlaybackPlayableGroup")
++ClassRedirects=(OldName="/Script/AvalancheMedia.AvalanchePlayback",NewName="/Script/AvalancheMedia.AvaPlaybackGraph")
++PropertyRedirects=(OldName="/Script/AvalancheMedia.AvaPlaybackNodeBlueprintPlayer.BlueprintAsset",NewName="/Script/AvalancheMedia.AvaPlaybackNodeBlueprintPlayer.Asset")
++EnumRedirects=(OldName="/Script/AvalancheMedia.EAvaRemoteControlChanges",NewName="/Script/AvalancheMedia.EAvaPlayableRemoteControlChanges")
++StructRedirects=(OldName="/Script/AvalancheMedia.AvalancheRemoteControlValue",NewName="/Script/AvalancheMedia.AvaPlayableRemoteControlValue")
++StructRedirects=(OldName="/Script/AvalancheMedia.AvalancheRemoteControlValues",NewName="/Script/AvalancheMedia.AvaPlayableRemoteControlValues")
++StructRedirects=(OldName="/Script/AvalancheMedia.AvalancheRemoteControlValueAsBytes_Legacy",NewName="/Script/AvalancheMedia.AvaPlayableRemoteControlValueAsBytes_Legacy")
++ClassRedirects=(OldName="/Script/AvalancheMedia.AvaRemotePlayableTransition",NewName="/Script/AvalancheMedia.AvaPlayableRemoteTransition")
++ClassRedirects=(OldName="/Script/AvalancheMedia.AvalancheRemoteProxyPlayable",NewName="/Script/AvalancheMedia.AvaPlayableRemoteProxy")
++ClassRedirects=(OldName="/Script/AvalancheMedia.AvaRenderTargetMediaCapture",NewName="/Script/AvalancheMedia.AvaBroadcastRenderTargetMediaCapture")
++ClassRedirects=(OldName="/Script/AvalancheMedia.AvaRenderTargetMediaOutput",NewName="/Script/AvalancheMedia.AvaBroadcastRenderTargetMediaOutput")
++StructRedirects=(OldName="/Script/AvalancheMedia.AvalanchePageCollection",NewName="/Script/AvalancheMedia.AvaRundownPageCollection")
++StructRedirects=(OldName="/Script/AvalancheMedia.AvalancheSubList",NewName="/Script/AvalancheMedia.AvaRundownSubList")
++EnumRedirects=(OldName="/Script/AvalancheMedia.EAvaPlaylistPageStopOptions",NewName="/Script/AvalancheMedia.EAvaRundownPageStopOptions")
++ClassRedirects=(OldName="/Script/AvalancheMedia.AvalanchePlaylist",NewName="/Script/AvalancheMedia.AvaRundown")
++ClassRedirects=(OldName="/Script/AvalancheMedia.AvalanchePlaylistComponent",NewName="/Script/AvalancheMedia.AvaRundownComponent")
++StructRedirects=(OldName="/Script/AvalancheMedia.AvaRundownMacroCommand",NewName="/Script/AvalancheMedia.AvaRundownMacroCommand")
++StructRedirects=(OldName="/Script/AvalancheMedia.AvaRundownMacroKeyBinding",NewName="/Script/AvalancheMedia.AvaRundownMacroKeyBinding")
++ClassRedirects=(OldName="/Script/AvalancheMedia.AvaRundownMacroCollection",NewName="/Script/AvalancheMedia.AvaRundownMacroCollection")
++StructRedirects=(OldName="/Script/AvalancheMedia.AvaPlaybackServerLoggingEntry",NewName="/Script/AvalancheMedia.AvaPlaybackServerLoggingEntry")
++StructRedirects=(OldName="/Script/AvalancheMedia.LocalPlaybackServerSettings",NewName="/Script/AvalancheMedia.AvaMediaLocalPlaybackServerSettings")
++ClassRedirects=(OldName="/Script/AvalancheMedia.AvalancheMediaSettings",NewName="/Script/AvalancheMedia.AvaMediaSettings")
++PropertyRedirects=(OldName="/Script/AvalancheMedia.AvaMediaSettings.ManagedAvalancheInstanceCacheMaximumSize",NewName="/Script/AvalancheMedia.AvaMediaSettings.ManagedInstanceCacheMaximumSize")
++EnumRedirects=(OldName="/Script/AvalancheMedia.EAvaSyncResponseResult",NewName="/Script/AvalancheMedia.EAvaMediaSyncResponseResult")
++EnumRedirects=(OldName="/Script/AvalancheMedia.EAvaSyncEngineType",NewName="/Script/AvalancheMedia.EAvaMediaSyncEngineType")
++StructRedirects=(OldName="/Script/AvalancheMedia.AvaSyncConnectionInfo",NewName="/Script/AvalancheMedia.AvaMediaSyncConnectionInfo")
++StructRedirects=(OldName="/Script/AvalancheMedia.AvaSyncResponse",NewName="/Script/AvalancheMedia.AvaMediaSyncResponse")
++StructRedirects=(OldName="/Script/AvalancheMedia.AvaSyncCompareResponse",NewName="/Script/AvalancheMedia.AvaMediaSyncCompareResponse")
++StructRedirects=(OldName="/Script/AvalancheMedia.AvaTransitionPlayableScene",NewName="/Script/AvalancheMedia.AvaPlayableTransitionScene")
++ClassRedirects=(OldName="/Script/AvalancheMediaEditor.AvalancheMediaEditorSettings",NewName="/Script/AvalancheMediaEditor.AvaMediaEditorSettings")
++ClassRedirects=(OldName="/Script/AvalancheMediaEditor.AvaPlaybackGraph",NewName="/Script/AvalancheMediaEditor.AvaPlaybackEditorGraph")
++ClassRedirects=(OldName="/Script/AvalancheMediaEditor.AvaPlaybackGraphSchema",NewName="/Script/AvalancheMediaEditor.AvaPlaybackEditorGraphSchema")
++EnumRedirects=(OldName="/Script/AvalancheMediaEditor.EAvaRundownPageSet",NewName="/Script/AvalancheMediaEditor.EAvaRundownPageSet")
++ClassRedirects=(OldName="/Script/AvalancheMediaEditor.AvaRundownEditorSettings",NewName="/Script/AvalancheMediaEditor.AvaRundownEditorSettings")
++StructRedirects=(OldName="/Script/AvalancheMediaEditor.AvaPlaylistMsgBase",NewName="/Script/AvalancheMediaEditor.AvaRundownMsgBase")
++StructRedirects=(OldName="/Script/AvalancheMediaEditor.AvaPlaylistServerMsg",NewName="/Script/AvalancheMediaEditor.AvaRundownServerMsg")
++StructRedirects=(OldName="/Script/AvalancheMediaEditor.AvaPlaylistPing",NewName="/Script/AvalancheMediaEditor.AvaRundownPing")
++StructRedirects=(OldName="/Script/AvalancheMediaEditor.AvaPlaylistPong",NewName="/Script/AvalancheMediaEditor.AvaRundownPong")
++StructRedirects=(OldName="/Script/AvalancheMediaEditor.AvaPlaylistGetPlaylists",NewName="/Script/AvalancheMediaEditor.AvaRundownGetRundowns")
++StructRedirects=(OldName="/Script/AvalancheMediaEditor.AvaPlaylistPlaylists",NewName="/Script/AvalancheMediaEditor.AvaRundownRundowns")
++StructRedirects=(OldName="/Script/AvalancheMediaEditor.AvaPlaylistLoadPlaylist",NewName="/Script/AvalancheMediaEditor.AvaRundownLoadRundown")
++StructRedirects=(OldName="/Script/AvalancheMediaEditor.AvaPlaylistGetPages",NewName="/Script/AvalancheMediaEditor.AvaRundownGetPages")
++StructRedirects=(OldName="/Script/AvalancheMediaEditor.AvaPlaylistCreatePage",NewName="/Script/AvalancheMediaEditor.AvaRundownCreatePage")
++StructRedirects=(OldName="/Script/AvalancheMediaEditor.AvaPlaylistDeletePage",NewName="/Script/AvalancheMediaEditor.AvaRundownDeletePage")
++StructRedirects=(OldName="/Script/AvalancheMediaEditor.AvaPlaylistCreateTemplate",NewName="/Script/AvalancheMediaEditor.AvaRundownCreateTemplate")
++StructRedirects=(OldName="/Script/AvalancheMediaEditor.AvaPlaylistDeleteTemplate",NewName="/Script/AvalancheMediaEditor.AvaRundownDeleteTemplate")
++StructRedirects=(OldName="/Script/AvalancheMediaEditor.AvaPlaylistChangeTemplateBP",NewName="/Script/AvalancheMediaEditor.AvaRundownChangeTemplateBP")
++StructRedirects=(OldName="/Script/AvalancheMediaEditor.AvaPlaylistPage",NewName="/Script/AvalancheMediaEditor.AvaRundownPageInfo")
++StructRedirects=(OldName="/Script/AvalancheMediaEditor.AvaPlaylistPages",NewName="/Script/AvalancheMediaEditor.AvaRundownPages")
++StructRedirects=(OldName="/Script/AvalancheMediaEditor.AvaPlaylistGetPageDetails",NewName="/Script/AvalancheMediaEditor.AvaRundownGetPageDetails")
++StructRedirects=(OldName="/Script/AvalancheMediaEditor.AvaPlaylistPageDetails",NewName="/Script/AvalancheMediaEditor.AvaRundownPageDetails")
++StructRedirects=(OldName="/Script/AvalancheMediaEditor.AvaPlaylistPagesStatuses",NewName="/Script/AvalancheMediaEditor.AvaRundownPagesStatuses")
++StructRedirects=(OldName="/Script/AvalancheMediaEditor.AvaPlaylistPageListChanged",NewName="/Script/AvalancheMediaEditor.AvaRundownPageListChanged")
++StructRedirects=(OldName="/Script/AvalancheMediaEditor.AvaPlaylistPageBlueprintChanged",NewName="/Script/AvalancheMediaEditor.AvaRundownPageBlueprintChanged")
++StructRedirects=(OldName="/Script/AvalancheMediaEditor.AvaPlaylistPageChannelChanged",NewName="/Script/AvalancheMediaEditor.AvaRundownPageChannelChanged")
++StructRedirects=(OldName="/Script/AvalancheMediaEditor.AvaPlaylistPageAnimSettingsChanged",NewName="/Script/AvalancheMediaEditor.AvaRundownPageAnimSettingsChanged")
++StructRedirects=(OldName="/Script/AvalancheMediaEditor.AvaPlaylistPageChangeChannel",NewName="/Script/AvalancheMediaEditor.AvaRundownPageChangeChannel")
++StructRedirects=(OldName="/Script/AvalancheMediaEditor.AvaPlaylistUpdatePageFromRCP",NewName="/Script/AvalancheMediaEditor.AvaRundownUpdatePageFromRCP")
++EnumRedirects=(OldName="/Script/AvalancheMediaEditor.EAvaPlaylistPageActions",NewName="/Script/AvalancheMediaEditor.EAvaRundownPageActions")
++StructRedirects=(OldName="/Script/AvalancheMediaEditor.AvaPlaylistPageAction",NewName="/Script/AvalancheMediaEditor.AvaRundownPageAction")
++StructRedirects=(OldName="/Script/AvalancheMediaEditor.AvaPlaylistPagePreviewAction",NewName="/Script/AvalancheMediaEditor.AvaRundownPagePreviewAction")
++StructRedirects=(OldName="/Script/AvalancheMediaEditor.AvaPlaylistPageActions",NewName="/Script/AvalancheMediaEditor.AvaRundownPageActions")
++StructRedirects=(OldName="/Script/AvalancheMediaEditor.AvaPlaylistPagePreviewActions",NewName="/Script/AvalancheMediaEditor.AvaRundownPagePreviewActions")
++StructRedirects=(OldName="/Script/AvalancheMediaEditor.AvaPlaylistPageEvent",NewName="/Script/AvalancheMediaEditor.AvaRundownPageEvent")
++StructRedirects=(OldName="/Script/AvalancheMediaEditor.AvaPlaylistOutputDeviceItem",NewName="/Script/AvalancheMediaEditor.AvaRundownOutputDeviceItem")
++StructRedirects=(OldName="/Script/AvalancheMediaEditor.AvaPlaylistOutputClassItem",NewName="/Script/AvalancheMediaEditor.AvaRundownOutputClassItem")
++StructRedirects=(OldName="/Script/AvalancheMediaEditor.AvaPlaylistDevicesList",NewName="/Script/AvalancheMediaEditor.AvaRundownDevicesList")
++StructRedirects=(OldName="/Script/AvalancheMediaEditor.AvaPlaylistGetChannel",NewName="/Script/AvalancheMediaEditor.AvaRundownGetChannel")
++StructRedirects=(OldName="/Script/AvalancheMediaEditor.AvaPlaylistGetChannels",NewName="/Script/AvalancheMediaEditor.AvaRundownGetChannels")
++StructRedirects=(OldName="/Script/AvalancheMediaEditor.AvaPlaylistChannel",NewName="/Script/AvalancheMediaEditor.AvaRundownChannel")
++StructRedirects=(OldName="/Script/AvalancheMediaEditor.AvaBroadcastChannelListChanged",NewName="/Script/AvalancheMediaEditor.AvaRundownChannelListChanged")
++StructRedirects=(OldName="/Script/AvalancheMediaEditor.AvaPlaylistChannelResponse",NewName="/Script/AvalancheMediaEditor.AvaRundownChannelResponse")
++StructRedirects=(OldName="/Script/AvalancheMediaEditor.AvaPlaylistChannels",NewName="/Script/AvalancheMediaEditor.AvaRundownChannels")
++StructRedirects=(OldName="/Script/AvalancheMediaEditor.AvaPlaylistAssetsChanged",NewName="/Script/AvalancheMediaEditor.AvaRundownAssetsChanged")
++EnumRedirects=(OldName="/Script/AvalancheMediaEditor.EAvaPlaylistChannelActions",NewName="/Script/AvalancheMediaEditor.EAvaRundownChannelActions")
++StructRedirects=(OldName="/Script/AvalancheMediaEditor.AvaPlaylistChannelAction",NewName="/Script/AvalancheMediaEditor.AvaRundownChannelAction")
++StructRedirects=(OldName="/Script/AvalancheMediaEditor.AvaPlaylistGetDevices",NewName="/Script/AvalancheMediaEditor.AvaRundownGetDevices")
++StructRedirects=(OldName="/Script/AvalancheMediaEditor.AvaPlaylistAddChannelDevice",NewName="/Script/AvalancheMediaEditor.AvaRundownAddChannelDevice")
++StructRedirects=(OldName="/Script/AvalancheMediaEditor.AvaPlaylistEditChannelDevice",NewName="/Script/AvalancheMediaEditor.AvaRundownEditChannelDevice")
++StructRedirects=(OldName="/Script/AvalancheMediaEditor.AvaPlaylistRemoveChannelDevice",NewName="/Script/AvalancheMediaEditor.AvaRundownRemoveChannelDevice")
++StructRedirects=(OldName="/Script/AvalancheMediaEditor.AvaPlaylistOutputDeviceItemResponse",NewName="/Script/AvalancheMediaEditor.AvaRundownOutputDeviceItemResponse")
++StructRedirects=(OldName="/Script/AvalancheMediaEditor.AvaPlaylistGetChannelImage",NewName="/Script/AvalancheMediaEditor.AvaRundownGetChannelImage")
++StructRedirects=(OldName="/Script/AvalancheMediaEditor.AvaPlaylistChannelImage",NewName="/Script/AvalancheMediaEditor.AvaRundownChannelImage")
++ClassRedirects=(OldName="/Script/AvalancheMediaEditor.AvaPlaylistPageContext",NewName="/Script/AvalancheMediaEditor.AvaRundownPageContext")
++EnumRedirects=(OldName="/Script/AvalancheMediaEditor.EAvaRundownEditorMacroCommand",NewName="/Script/AvalancheMediaEditor.EAvaRundownEditorMacroCommand")
++PropertyRedirects=(OldName="/Script/AvalancheMedia.AvaRundownComponent.Playlist",NewName="/Script/AvalancheMedia.AvaRundownComponent.Rundown")
++EnumRedirects=(OldName="/Script/AvalancheMedia.EAvaPlayType",NewName="/Script/AvalancheMedia.EAvaRundownPagePlayType")
++PropertyRedirects=(OldName="/Script/AvalancheMediaEditor.AvaRundownRundowns.Playlists",NewName="/Script/AvalancheMediaEditor.AvaRundownRundowns.Rundowns")
++EnumRedirects=(OldName="/Script/AvalancheMediaEditor.EAvaPlaylistPageEvents",NewName="/Script/AvalancheMediaEditor.EAvaRundownPageEvents")
++PropertyRedirects=(OldName="/Script/AvalancheMediaEditor.AvaRundownLoadRundown.Playlist",NewName="/Script/AvalancheMediaEditor.AvaRundownLoadRundown.Rundown")
++PropertyRedirects=(OldName="/Script/AvalancheMediaEditor.AvaRundownGetPages.Playlist",NewName="/Script/AvalancheMediaEditor.AvaRundownGetPages.Rundown")
++PropertyRedirects=(OldName="/Script/AvalancheMediaEditor.AvaRundownCreatePage.Playlist",NewName="/Script/AvalancheMediaEditor.AvaRundownCreatePage.Rundown")
++PropertyRedirects=(OldName="/Script/AvalancheMediaEditor.AvaRundownDeletePage.Playlist",NewName="/Script/AvalancheMediaEditor.AvaRundownDeletePage.Rundown")
++PropertyRedirects=(OldName="/Script/AvalancheMediaEditor.AvaRundownCreateTemplate.Playlist",NewName="/Script/AvalancheMediaEditor.AvaRundownCreateTemplate.Rundown")
++PropertyRedirects=(OldName="/Script/AvalancheMediaEditor.AvaRundownDeleteTemplate.Playlist",NewName="/Script/AvalancheMediaEditor.AvaRundownDeleteTemplate.Rundown")
++PropertyRedirects=(OldName="/Script/AvalancheMediaEditor.AvaRundownChangeTemplateBP.Playlist",NewName="/Script/AvalancheMediaEditor.AvaRundownChangeTemplateBP.Rundown")
++PropertyRedirects=(OldName="/Script/AvalancheMediaEditor.AvaRundownGetPageDetails.Playlist",NewName="/Script/AvalancheMediaEditor.AvaRundownGetPageDetails.Rundown")
++PropertyRedirects=(OldName="/Script/AvalancheMediaEditor.AvaRundownPageDetails.Playlist",NewName="/Script/AvalancheMediaEditor.AvaRundownPageDetails.Rundown")
++PropertyRedirects=(OldName="/Script/AvalancheMediaEditor.AvaRundownPagesStatuses.Playlist",NewName="/Script/AvalancheMediaEditor.AvaRundownPagesStatuses.Rundown")
++PropertyRedirects=(OldName="/Script/AvalancheMediaEditor.AvaRundownPageListChanged.Playlist",NewName="/Script/AvalancheMediaEditor.AvaRundownPageListChanged.Rundown")
++PropertyRedirects=(OldName="/Script/AvalancheMediaEditor.AvaRundownPageBlueprintChanged.Playlist",NewName="/Script/AvalancheMediaEditor.AvaRundownPageBlueprintChanged.Rundown")
++PropertyRedirects=(OldName="/Script/AvalancheMediaEditor.AvaRundownPageChannelChanged.Playlist",NewName="/Script/AvalancheMediaEditor.AvaRundownPageChannelChanged.Rundown")
++PropertyRedirects=(OldName="/Script/AvalancheMediaEditor.AvaRundownPageAnimSettingsChanged.Playlist",NewName="/Script/AvalancheMediaEditor.AvaRundownPageAnimSettingsChanged.Rundown")
++PropertyRedirects=(OldName="/Script/AvalancheMediaEditor.AvaRundownPageChangeChannel.Playlist",NewName="/Script/AvalancheMediaEditor.AvaRundownPageChangeChannel.Rundown")
++EnumRedirects=(OldName="/Script/AvalancheMedia.EAvaRemoteControlResult",NewName="/Script/AvalancheMedia.EAvaPlayableRemoteControlResult")
+
 [/Script/Avalanche.AvaAnimation]
 DefaultCompletionMode=KeepState
diff --git a/Engine/Plugins/Experimental/Avalanche/Source/AvalancheEditor/Private/Transition/AvaTransitionExtension.cpp b/Engine/Plugins/Experimental/Avalanche/Source/AvalancheEditor/Private/Transition/AvaTransitionExtension.cpp
--- a/Engine/Plugins/Experimental/Avalanche/Source/AvalancheEditor/Private/Transition/AvaTransitionExtension.cpp
+++ b/Engine/Plugins/Experimental/Avalanche/Source/AvalancheEditor/Private/Transition/AvaTransitionExtension.cpp
@@ -63,11 +63,11 @@ namespace UE::AvaEditor::Private
 
 FDelegateHandle FAvaTransitionExtension::PropertyFilterHandle;
 
-FAvaTransitionPlaylistExtension FAvaTransitionExtension::PlaylistExtension;
+FAvaTransitionRundownExtension FAvaTransitionExtension::RundownExtension;
 
 void FAvaTransitionExtension::StaticStartup()
 {
-	PlaylistExtension.Startup();
+	RundownExtension.Startup();
 
 	IRemoteControlUIModule& RemoteControlUIModule = IRemoteControlUIModule::Get();
 
@@ -89,7 +89,7 @@ void FAvaTransitionExtension::StaticStartup()
 
 void FAvaTransitionExtension::StaticShutdown()
 {
-	PlaylistExtension.Shutdown();
+	RundownExtension.Shutdown();
 
 	if (IRemoteControlUIModule* RemoteControlUIModule = FModuleManager::Get().GetModulePtr<IRemoteControlUIModule>(TEXT("RemoteControlUI")))
 	{
diff --git a/Engine/Plugins/Experimental/Avalanche/Source/AvalancheEditor/Private/Transition/AvaTransitionExtension.h b/Engine/Plugins/Experimental/Avalanche/Source/AvalancheEditor/Private/Transition/AvaTransitionExtension.h
--- a/Engine/Plugins/Experimental/Avalanche/Source/AvalancheEditor/Private/Transition/AvaTransitionExtension.h
+++ b/Engine/Plugins/Experimental/Avalanche/Source/AvalancheEditor/Private/Transition/AvaTransitionExtension.h
@@ -2,7 +2,7 @@
 
 #pragma once
 
-#include "AvaTransitionPlaylistExtension.h"
+#include "AvaTransitionRundownExtension.h"
 #include "Delegates/Delegate.h"
 #include "IAvaEditorExtension.h"
 
@@ -37,5 +37,5 @@ private:
 
 	static FDelegateHandle PropertyFilterHandle;
 
-	static FAvaTransitionPlaylistExtension PlaylistExtension;
+	static FAvaTransitionRundownExtension RundownExtension;
 };
diff --git a/Engine/Plugins/Experimental/Avalanche/Source/AvalancheMRQ/Private/AvaMRQRundownPageSetting.cpp b/Engine/Plugins/Experimental/Avalanche/Source/AvalancheMRQ/Private/AvaMRQRundownPageSetting.cpp
--- a/Engine/Plugins/Experimental/Avalanche/Source/AvalancheMRQ/Private/AvaMRQRundownPageSetting.cpp
+++ b/Engine/Plugins/Experimental/Avalanche/Source/AvalancheMRQ/Private/AvaMRQRundownPageSetting.cpp
@@ -12,7 +12,7 @@
 #include "Misc/CoreMiscDefines.h"
 #include "Misc/PackageName.h"
 #include "MoviePipeline.h"
-#include "Playlist/AvalanchePlaylist.h"
+#include "Rundown/AvaRundown.h"
 
 DEFINE_LOG_CATEGORY_STATIC(LogAvaMRQSequenceData, Log, All);
 
@@ -58,7 +58,7 @@ void UAvaMRQRundownPageSetting::SetupForPipelineImpl(UMoviePipeline* InPipeline)
 		return;
 	}
 
-	UAvalanchePlaylist* Rundown = RundownPage.Rundown.LoadSynchronous();
+	UAvaRundown* Rundown = RundownPage.Rundown.LoadSynchronous();
 	if (!Rundown)
 	{
 		UE_LOG(LogAvaMRQSequenceData, Error
@@ -67,7 +67,7 @@ void UAvaMRQRundownPageSetting::SetupForPipelineImpl(UMoviePipeline* InPipeline)
 		return;
 	}
 
-	const FAvalanchePage& Page = Rundown->GetPage(RundownPage.PageId);
+	const FAvaRundownPage& Page = Rundown->GetPage(RundownPage.PageId);
 	if (!Page.IsValidPage())
 	{
 		UE_LOG(LogAvaMRQSequenceData, Error
@@ -81,13 +81,13 @@ void UAvaMRQRundownPageSetting::SetupForPipelineImpl(UMoviePipeline* InPipeline)
 	int32 PieInstanceId = UE::AvaMRQ::Private::FindPieInstanceId(World);
 	if (PieInstanceId != INDEX_NONE)
 	{
-		FSoftObjectPath SourcePath  = Page.GetAvalancheAssetPath(Rundown);
+		FSoftObjectPath SourcePath  = Page.GetAssetPath(Rundown);
 		FSoftObjectPath CurrentPath = UE::AvaMRQ::Private::RemovePiePrefix(World, PieInstanceId);
 
 		if (CurrentPath != SourcePath)
 		{
 			UE_LOG(LogAvaMRQSequenceData, Error
-				, TEXT("Asset path '%s' in Page '%d' of Avalanche Rundown '%s' did not match the provided PIE sanitized world path '%s'")
+				, TEXT("Asset path '%s' in Page '%d' of Motion Design Rundown '%s' did not match the provided PIE sanitized world path '%s'")
 				, *SourcePath.ToString()
 				, Page.GetPageId()
 				, *Rundown->GetFullName()
diff --git a/Engine/Plugins/Experimental/Avalanche/Source/AvalancheMRQ/Public/AvaMRQRundownPageSetting.h b/Engine/Plugins/Experimental/Avalanche/Source/AvalancheMRQ/Public/AvaMRQRundownPageSetting.h
--- a/Engine/Plugins/Experimental/Avalanche/Source/AvalancheMRQ/Public/AvaMRQRundownPageSetting.h
+++ b/Engine/Plugins/Experimental/Avalanche/Source/AvalancheMRQ/Public/AvaMRQRundownPageSetting.h
@@ -6,7 +6,7 @@
 #include "UObject/SoftObjectPtr.h"
 #include "AvaMRQRundownPageSetting.generated.h"
 
-class UAvalanchePlaylist;
+class UAvaRundown;
 
 USTRUCT(BlueprintType, DisplayName = "Motion Design MRQ Rundown Page")
 struct FAvaMRQRundownPage
@@ -14,7 +14,7 @@ struct FAvaMRQRundownPage
 	GENERATED_BODY()
 
 	UPROPERTY(EditAnywhere, BlueprintReadWrite, Category = "Motion Design Sequence")
-	TSoftObjectPtr<UAvalanchePlaylist> Rundown;
+	TSoftObjectPtr<UAvaRundown> Rundown;
 
 	UPROPERTY(EditAnywhere, BlueprintReadWrite, Category = "Motion Design Sequence")
 	int32 PageId = 0;
diff --git a/Engine/Plugins/Experimental/Avalanche/Source/AvalancheMRQEditor/Private/AvaMRQEditorModule.cpp b/Engine/Plugins/Experimental/Avalanche/Source/AvalancheMRQEditor/Private/AvaMRQEditorModule.cpp
--- a/Engine/Plugins/Experimental/Avalanche/Source/AvalancheMRQEditor/Private/AvaMRQEditorModule.cpp
+++ b/Engine/Plugins/Experimental/Avalanche/Source/AvalancheMRQEditor/Private/AvaMRQEditorModule.cpp
@@ -2,27 +2,27 @@
 
 #include "AvaMRQEditorModule.h"
 #include "AvaMRQEditorCommands.h"
-#include "AvaMRQEditorPlaylistUtils.h"
+#include "AvaMRQEditorRundownUtils.h"
 #include "Framework/MultiBox/MultiBoxExtender.h"
 #include "IAvaMediaEditorModule.h"
 #include "Modules/ModuleManager.h"
-#include "Playlist/AvaPlaylistEditor.h"
-#include "Playlist/AvalanchePlaylist.h"
+#include "Rundown/AvaRundown.h"
+#include "Rundown/AvaRundownEditor.h"
 #include "Toolkits/AssetEditorToolkit.h"
 
 void FAvaMRQEditorModule::StartupModule()
 {
 	FAvaMRQEditorCommands::Register();
 
-	TSharedPtr<FExtensibilityManager> PlaylistToolbarExtensibility = IAvaMediaEditorModule::Get().GetPlaylistToolBarExtensibilityManager();
-	if (PlaylistToolbarExtensibility.IsValid())
+	TSharedPtr<FExtensibilityManager> RundownToolbarExtensibility = IAvaMediaEditorModule::Get().GetRundownToolBarExtensibilityManager();
+	if (RundownToolbarExtensibility.IsValid())
 	{
-		PlaylistToolbarExtensibilityWeak = PlaylistToolbarExtensibility;
+		RundownToolbarExtensibilityWeak = RundownToolbarExtensibility;
 
-		TArray<FAssetEditorExtender>& ExtenderDelegates = PlaylistToolbarExtensibility->GetExtenderDelegates();
-		ExtenderDelegates.Add(FAssetEditorExtender::CreateStatic(&FAvaMRQEditorModule::ExtendPlaylistToolbar));
+		TArray<FAssetEditorExtender>& ExtenderDelegates = RundownToolbarExtensibility->GetExtenderDelegates();
+		ExtenderDelegates.Add(FAssetEditorExtender::CreateStatic(&FAvaMRQEditorModule::ExtendRundownToolbar));
 
-		PlaylistToolbarExtenderHandle = ExtenderDelegates.Last().GetHandle();
+		RundownToolbarExtenderHandle = ExtenderDelegates.Last().GetHandle();
 	}
 }
 
@@ -30,21 +30,21 @@ void FAvaMRQEditorModule::ShutdownModule()
 {
 	FAvaMRQEditorCommands::Unregister();
 
-	if (TSharedPtr<FExtensibilityManager> PlaylistToolbarExtensibility = PlaylistToolbarExtensibilityWeak.Pin())
+	if (TSharedPtr<FExtensibilityManager> RundownToolbarExtensibility = RundownToolbarExtensibilityWeak.Pin())
 	{
-		if (PlaylistToolbarExtenderHandle.IsValid())
+		if (RundownToolbarExtenderHandle.IsValid())
 		{
-			PlaylistToolbarExtensibility->GetExtenderDelegates().RemoveAll([this](const FAssetEditorExtender& InExtender)
+			RundownToolbarExtensibility->GetExtenderDelegates().RemoveAll([this](const FAssetEditorExtender& InExtender)
 			{
-				return InExtender.GetHandle() == PlaylistToolbarExtenderHandle;
+				return InExtender.GetHandle() == RundownToolbarExtenderHandle;
 			});
-			PlaylistToolbarExtenderHandle.Reset();
+			RundownToolbarExtenderHandle.Reset();
 		}
-		PlaylistToolbarExtensibilityWeak.Reset();
+		RundownToolbarExtensibilityWeak.Reset();
 	}
 }
 
-TSharedRef<FExtender> FAvaMRQEditorModule::ExtendPlaylistToolbar(const TSharedRef<FUICommandList> InCommandList, const TArray<UObject*> InObjects)
+TSharedRef<FExtender> FAvaMRQEditorModule::ExtendRundownToolbar(const TSharedRef<FUICommandList> InCommandList, const TArray<UObject*> InObjects)
 {
 	TSharedRef<FExtender> Extender = MakeShared<FExtender>();
 
@@ -54,35 +54,35 @@ TSharedRef<FExtender> FAvaMRQEditorModule::ExtendPlaylistToolbar(const TSharedRe
 		return Extender;
 	}
 
-	TSharedRef<FAvaMRQPlaylistContext> Context = MakeShared<FAvaMRQPlaylistContext>();
-	Context->PlaylistEditors.Reserve(InObjects.Num());
+	TSharedRef<FAvaMRQRundownContext> Context = MakeShared<FAvaMRQRundownContext>();
+	Context->RundownEditors.Reserve(InObjects.Num());
 
-	const FName PlaylistEditorName = TEXT("AvaPlaylistEditor");
+	const FName RundownEditorName = TEXT("AvaRundownEditor");
 
-	// Gather Playlist Editors
+	// Gather Rundown Editors
 	for (UObject* Object : InObjects)
 	{
-		if (UAvalanchePlaylist* Playlist = Cast<UAvalanchePlaylist>(Object))
+		if (UAvaRundown* Rundown = Cast<UAvaRundown>(Object))
 		{
-			IAssetEditorInstance* AssetEditor = AssetEditorSubsystem->FindEditorForAsset(Playlist, /*bFocusIfOpen*/false);
-			if (!AssetEditor || AssetEditor->GetEditorName() != PlaylistEditorName)
+			IAssetEditorInstance* AssetEditor = AssetEditorSubsystem->FindEditorForAsset(Rundown, /*bFocusIfOpen*/false);
+			if (!AssetEditor || AssetEditor->GetEditorName() != RundownEditorName)
 			{
 				continue;
 			}
 
-			FAvaPlaylistEditor* PlaylistEditor = static_cast<FAvaPlaylistEditor*>(AssetEditor);
-			Context->PlaylistEditors.Add(StaticCastWeakPtr<FAvaPlaylistEditor>(PlaylistEditor->AsWeak()));
+			FAvaRundownEditor* RundownEditor = static_cast<FAvaRundownEditor*>(AssetEditor);
+			Context->RundownEditors.Add(StaticCastWeakPtr<FAvaRundownEditor>(RundownEditor->AsWeak()));
 		}
 	}
 
-	if (Context->PlaylistEditors.IsEmpty())
+	if (Context->RundownEditors.IsEmpty())
 	{
 		return Extender;
 	}
 
 	Extender->AddToolBarExtension("Pages"
 		, EExtensionHook::After
-		, CreatePlaylistActions(Context)
+		, CreateRundownActions(Context)
 		, FToolBarExtensionDelegate::CreateLambda([](FToolBarBuilder& ToolbarBuilder)
 		{
 			const FAvaMRQEditorCommands& MRQEditorCommands = FAvaMRQEditorCommands::Get();
@@ -93,25 +93,25 @@ TSharedRef<FExtender> FAvaMRQEditorModule::ExtendPlaylistToolbar(const TSharedRe
 	return Extender;
 }
 
-TSharedRef<FUICommandList> FAvaMRQEditorModule::CreatePlaylistActions(TSharedRef<FAvaMRQPlaylistContext> InContext)
+TSharedRef<FUICommandList> FAvaMRQEditorModule::CreateRundownActions(TSharedRef<FAvaMRQRundownContext> InContext)
 {
 	const FAvaMRQEditorCommands& MRQEditorCommands = FAvaMRQEditorCommands::Get();
 
 	TSharedRef<FUICommandList> CommandList = MakeShared<FUICommandList>();
 
-	static auto HasAnySelectedPage = [](const TWeakPtr<const FAvaPlaylistEditor>& InPlaylistEditorWeak)
+	static auto HasAnySelectedPage = [](const TWeakPtr<const FAvaRundownEditor>& InRundownEditorWeak)
 	{
-		TSharedPtr<const FAvaPlaylistEditor> PlaylistEditor = InPlaylistEditorWeak.Pin();
-		return PlaylistEditor.IsValid() && PlaylistEditor->GetSelectedPagesOnActiveSubListWidget().Num() > 0;
+		TSharedPtr<const FAvaRundownEditor> RundownEditor = InRundownEditorWeak.Pin();
+		return RundownEditor.IsValid() && RundownEditor->GetSelectedPagesOnActiveSubListWidget().Num() > 0;
 	};
 
 	FCanExecuteAction CanExecuteSelectedPageAction = FCanExecuteAction::CreateLambda([InContext]
 	{
-		return InContext->PlaylistEditors.ContainsByPredicate(HasAnySelectedPage);
+		return InContext->RundownEditors.ContainsByPredicate(HasAnySelectedPage);
 	});
 
 	CommandList->MapAction(MRQEditorCommands.RenderSelectedPages
-		, FExecuteAction::CreateLambda([InContext]{ FAvaMRQEditorPlaylistUtils::RenderSelectedPages(InContext->PlaylistEditors); })
+		, FExecuteAction::CreateLambda([InContext]{ FAvaMRQEditorRundownUtils::RenderSelectedPages(InContext->RundownEditors); })
 		, CanExecuteSelectedPageAction);
 
 	return CommandList;
diff --git a/Engine/Plugins/Experimental/Avalanche/Source/AvalancheMRQEditor/Private/AvaMRQEditorModule.h b/Engine/Plugins/Experimental/Avalanche/Source/AvalancheMRQEditor/Private/AvaMRQEditorModule.h
--- a/Engine/Plugins/Experimental/Avalanche/Source/AvalancheMRQEditor/Private/AvaMRQEditorModule.h
+++ b/Engine/Plugins/Experimental/Avalanche/Source/AvalancheMRQEditor/Private/AvaMRQEditorModule.h
@@ -7,16 +7,16 @@
 #include "Modules/ModuleInterface.h"
 #include "Templates/SharedPointer.h"
 
-class FAvaPlaylistEditor;
+class FAvaRundownEditor;
 class FExtender;
 class FExtensibilityManager;
 class FUICommandList;
-class UAvalanchePlaylist;
+class UAvaRundown;
 class UObject;
 
-struct FAvaMRQPlaylistContext
+struct FAvaMRQRundownContext
 {
-    TArray<TWeakPtr<const FAvaPlaylistEditor>, TInlineAllocator<1>> PlaylistEditors;
+    TArray<TWeakPtr<const FAvaRundownEditor>, TInlineAllocator<1>> RundownEditors;
 };
 
 class FAvaMRQEditorModule : public IModuleInterface
@@ -26,11 +26,11 @@ class FAvaMRQEditorModule : public IModuleInterface
     virtual void ShutdownModule() override;
     //~ End IModuleInterface
 
-    static TSharedRef<FExtender> ExtendPlaylistToolbar(const TSharedRef<FUICommandList> InCommandList, const TArray<UObject*> InObjects);
+    static TSharedRef<FExtender> ExtendRundownToolbar(const TSharedRef<FUICommandList> InCommandList, const TArray<UObject*> InObjects);
 
-    static TSharedRef<FUICommandList> CreatePlaylistActions(TSharedRef<FAvaMRQPlaylistContext> InContext);
+    static TSharedRef<FUICommandList> CreateRundownActions(TSharedRef<FAvaMRQRundownContext> InContext);
 
-    TWeakPtr<FExtensibilityManager> PlaylistToolbarExtensibilityWeak;
+    TWeakPtr<FExtensibilityManager> RundownToolbarExtensibilityWeak;
 
-    FDelegateHandle PlaylistToolbarExtenderHandle;
+    FDelegateHandle RundownToolbarExtenderHandle;
 };
diff --git a/Engine/Plugins/Experimental/Avalanche/Source/AvalancheMRQEditor/Public/AvaMRQEditorRundownUtils.h b/Engine/Plugins/Experimental/Avalanche/Source/AvalancheMRQEditor/Public/AvaMRQEditorRundownUtils.h
new file mode 100644
--- /dev/null
+++ b/Engine/Plugins/Experimental/Avalanche/Source/AvalancheMRQEditor/Public/AvaMRQEditorRundownUtils.h
@@ -0,0 +1,20 @@
+// Copyright Epic Games, Inc. All Rights Reserved.
+
+#pragma once
+
+#include "Containers/ArrayView.h"
+#include "Templates/SharedPointerFwd.h"
+
+class FAvaRundownEditor;
+class UAvaSequence;
+class UAvaRundown;
+struct FAvaRundownPage;
+
+struct FAvaMRQEditorRundownUtils
+{
+	AVALANCHEMRQEDITOR_API static void RenderSelectedPages(TConstArrayView<TWeakPtr<const FAvaRundownEditor>> InRundownEditors);
+
+	AVALANCHEMRQEDITOR_API static void RenderPages(const UAvaRundown& InRundown, TConstArrayView<int32> InPageIds);
+
+	AVALANCHEMRQEDITOR_API static void RenderPage(const UAvaRundown& InRundown, const FAvaRundownPage& InPage);
+};
diff --git a/Engine/Plugins/Experimental/Avalanche/Source/AvalancheMedia/Private/AvaMediaModule.cpp b/Engine/Plugins/Experimental/Avalanche/Source/AvalancheMedia/Private/AvaMediaModule.cpp
--- a/Engine/Plugins/Experimental/Avalanche/Source/AvalancheMedia/Private/AvaMediaModule.cpp
+++ b/Engine/Plugins/Experimental/Avalanche/Source/AvalancheMedia/Private/AvaMediaModule.cpp
@@ -1,15 +1,16 @@
 // Copyright Epic Games, Inc. All Rights Reserved.
 
 #include "AvaMediaModule.h"
+
 #include "Application/ThrottleManager.h"
-#include "AvalancheMediaSettings.h"
+#include "AvaMediaSettings.h"
 #include "AvaMediaStyle.h"
+#include "Broadcast/OutputDevices/AvaBroadcastRenderTargetMediaUtils.h"
 #include "IMediaIOCoreModule.h"
 #include "Interfaces/IPluginManager.h"
 #include "Misc/CommandLine.h"
 #include "Misc/CoreDelegates.h"
-#include "OutputDevices/AvaRenderTargetMediaUtils.h"
-#include "Playback/AvaMediaPlaybackClientDelegates.h"
+#include "Playback/AvaPlaybackClientDelegates.h"
 #include "ShaderCore.h"
 
 #if WITH_EDITOR
@@ -23,8 +24,8 @@ namespace UE::AvaMediaModule::Private
 	// Command line parsing helper.
 	bool IsPlaybackServerManuallyStarted(FString& OutPlaybackServerName)
 	{
-		return FParse::Value(FCommandLine::Get(),TEXT("AvaMediaPlaybackServerStart="), OutPlaybackServerName) ||
-		 FParse::Param(FCommandLine::Get(), TEXT("AvaMediaPlaybackServerStart"));
+		return FParse::Value(FCommandLine::Get(),TEXT("MotionDesignPlaybackServerStart="), OutPlaybackServerName) ||
+		 FParse::Param(FCommandLine::Get(), TEXT("MotionDesignPlaybackServerStart"));
 	}
 }
 
@@ -41,84 +42,84 @@ void FAvaMediaModule::StartupModule()
 	check(Plugin.IsValid());
 
 	const FString PluginShaderDir = FPaths::Combine(Plugin->GetBaseDir(), TEXT("Shaders"));
-	AddShaderSourceDirectoryMapping(UE::AvaRenderTargetMediaUtils::VirtualShaderMountPoint, PluginShaderDir);
+	AddShaderSourceDirectoryMapping(UE::AvaBroadcastRenderTargetMediaUtils::VirtualShaderMountPoint, PluginShaderDir);
 
 	IMediaIOCoreModule::Get().RegisterDeviceProvider(&AvaDisplayDeviceProvider);
 
 	ConsoleCmds.Add(IConsoleManager::Get().RegisterConsoleCommand(
-				TEXT("AvaMediaServer.Start"),
+				TEXT("MotionDesignPlaybackServer.Start"),
 				TEXT("Starts the playback server."),
 				FConsoleCommandWithArgsDelegate::CreateRaw(this, &FAvaMediaModule::StartPlaybackServerCommand),
 				ECVF_Default
 				));
 	ConsoleCmds.Add(IConsoleManager::Get().RegisterConsoleCommand(
-				TEXT("AvaMediaServer.Stop"),
+				TEXT("MotionDesignPlaybackServer.Stop"),
 				TEXT("Stops the playback server."),
 				FConsoleCommandWithArgsDelegate::CreateRaw(this, &FAvaMediaModule::StopPlaybackServerCommand),
 				ECVF_Default
 				));
 	ConsoleCmds.Add(IConsoleManager::Get().RegisterConsoleCommand(
-				TEXT("AvaMediaClient.Start"),
+				TEXT("MotionDesignPlaybackClient.Start"),
 				TEXT("Starts the playback client."),
 				FConsoleCommandWithArgsDelegate::CreateRaw(this, &FAvaMediaModule::StartPlaybackClientCommand),
 				ECVF_Default
 				));
 	ConsoleCmds.Add(IConsoleManager::Get().RegisterConsoleCommand(
-				TEXT("AvaMediaClient.Stop"),
+				TEXT("MotionDesignPlaybackClient.Stop"),
 				TEXT("Stops the playback client."),
 				FConsoleCommandWithArgsDelegate::CreateRaw(this, &FAvaMediaModule::StopPlaybackClientCommand),
 				ECVF_Default
 				));
 	ConsoleCmds.Add(IConsoleManager::Get().RegisterConsoleCommand(
-				TEXT("AvaMediaLocalServer.Launch"),
+				TEXT("MotionDesignPlaybackLocalServer.Launch"),
 				TEXT("Launches the local playback server."),
 				FConsoleCommandWithArgsDelegate::CreateRaw(this, &FAvaMediaModule::LaunchLocalPlaybackServerCommand),
 				ECVF_Default
 				));
 	ConsoleCmds.Add(IConsoleManager::Get().RegisterConsoleCommand(
-				TEXT("AvaMediaLocalServer.Stop"),
+				TEXT("MotionDesignPlaybackLocalServer.Stop"),
 				TEXT("Stops the local playback server."),
 				FConsoleCommandWithArgsDelegate::CreateRaw(this, &FAvaMediaModule::StopLocalPlaybackServerCommand),
 				ECVF_Default
 				));
 	ConsoleCmds.Add(IConsoleManager::Get().RegisterConsoleCommand(
-				TEXT("AvaMediaHttpServer.Start"),
+				TEXT("MotionDesignPlaybackHttpServer.Start"),
 				TEXT("Starts the http playback server."),
 				FConsoleCommandWithArgsDelegate::CreateRaw(this, &FAvaMediaModule::StartHttpPlaybackServerCommand),
 				ECVF_Default
 				));
 	ConsoleCmds.Add(IConsoleManager::Get().RegisterConsoleCommand(
-				TEXT("AvaMediaHttpServer.Stop"),
+				TEXT("MotionDesignPlaybackHttpServer.Stop"),
 				TEXT("Stops the http playback server."),
 				FConsoleCommandWithArgsDelegate::CreateRaw(this, &FAvaMediaModule::StopHttpPlaybackServerCommand),
 				ECVF_Default
 				));
 	ConsoleCmds.Add(IConsoleManager::Get().RegisterConsoleCommand(
-				TEXT("AvaMediaDevices.Save"),
+				TEXT("MotionDesignPlaybackDevices.Save"),
 				TEXT("Save Device Providers data."),
 				FConsoleCommandWithArgsDelegate::CreateRaw(this, &FAvaMediaModule::SaveDeviceProvidersCommand),
 				ECVF_Default
 				));
 	ConsoleCmds.Add(IConsoleManager::Get().RegisterConsoleCommand(
-				TEXT("AvaMediaDevices.Load"),
+				TEXT("MotionDesignPlaybackDevices.Load"),
 				TEXT("Load device providers."),
 				FConsoleCommandWithArgsDelegate::CreateRaw(this, &FAvaMediaModule::LoadDeviceProvidersCommand),
 				ECVF_Default
 				));
 	ConsoleCmds.Add(IConsoleManager::Get().RegisterConsoleCommand(
-				TEXT("AvaMediaDevices.Unload"),
+				TEXT("MotionDesignPlaybackDevices.Unload"),
 				TEXT("Load device providers."),
 				FConsoleCommandWithArgsDelegate::CreateRaw(this, &FAvaMediaModule::UnloadDeviceProvidersCommand),
 				ECVF_Default
 				));
 	ConsoleCmds.Add(IConsoleManager::Get().RegisterConsoleCommand(
-				TEXT("AvaMediaDevices.List"),
+				TEXT("MotionDesignPlaybackDevices.List"),
 				TEXT("List device providers."),
 				FConsoleCommandWithArgsDelegate::CreateRaw(this, &FAvaMediaModule::ListDeviceProvidersCommand),
 				ECVF_Default
 				));
 	ConsoleCmds.Add(IConsoleManager::Get().RegisterConsoleCommand(
-				TEXT("AvaStat"),
+				TEXT("MotionDesignPlayback.Stat"),
 				TEXT("Enable engine performance statistics. Same as 'stat' command but will affect Motion Design Playback outputs and propagate to connected servers."),
 				FConsoleCommandWithArgsDelegate::CreateRaw(this, &FAvaMediaModule::HandleStatCommand),
 				ECVF_Default
@@ -128,7 +129,7 @@ void FAvaMediaModule::StartupModule()
 	FCoreDelegates::OnEnginePreExit.AddRaw(this, &FAvaMediaModule::EnginePreExit);
 	
 	FString DummyServerName;
-	if (!AvaMediaPlaybackServer.IsValid() && IsPlaybackServerManuallyStarted(DummyServerName))
+	if (!AvaPlaybackServer.IsValid() && IsPlaybackServerManuallyStarted(DummyServerName))
 	{
 		// Prevent throttling when server is started.
 		// This has to be done before any SLevelViewport are ticked since the cvar value is cached on first tick.
@@ -144,7 +145,7 @@ void FAvaMediaModule::StartupModule()
 
 	// StormSyncAvaBridge has some issues with servers in game mode. This option
 	// allows us to disable some of it until all the bugs are fixed.
-	if (FParse::Param(FCommandLine::Get(), TEXT("DisableAvaMediaSync")))
+	if (FParse::Param(FCommandLine::Get(), TEXT("DisableMotionDesignSync")))
 	{
 		AvaMediaSync->SetFeatureEnabled(false);
 	}
@@ -193,12 +194,12 @@ void FAvaMediaModule::StopMediaPlaybackServer()
 	StopPlaybackServerCommand({});
 }
 
-class FAvaMediaPlaybackClientDummy : public IAvaMediaPlaybackClient
+class FAvaPlaybackClientDummy : public IAvaPlaybackClient
 {
 	static const FString EmptyString;
 	static const TArray<FString> EmptyStringArray;
 public:
-	//IAvaMediaPlaybackClient
+	//IAvaPlaybackClient
 	virtual int32 GetNumConnectedServers() const override { return 0; }
 	virtual TArray<FString> GetServerNames() const override { return TArray<FString>();}
 	virtual FMessageAddress GetServerAddress(const FString& InServerName) const override
@@ -215,53 +216,53 @@ public:
 	virtual void RemoveUserData(const FString& InKey) override {}
 	virtual void BroadcastStatCommand(const FString& InCommand, bool bInBroadcastLocalState) override {}
 	virtual void RequestPlaybackAssetStatus(const FSoftObjectPath& InAssetPath, const FString& InChannelOrServerName, bool bInForceRefresh) override {}
-	virtual void RequestPlayback(const FGuid& InInstanceId, const FSoftObjectPath& InAssetPath, const FString& InChannelName, EAvaMediaPlaybackAction InAction, const FString& InArguments) override {}
-	virtual void RequestAnimPlayback(const FGuid& InInstanceId, const FSoftObjectPath& InAssetPath, const FString& InChannelName, const FAnimPlaySettings& InAnimSettings) override {};
-	virtual void RequestAnimAction(const FGuid& InInstanceId, const FSoftObjectPath& InAssetPath, const FString& InChannelName, const FString& InAnimationName, EAvaMediaAnimAction InAction) override {};
-	virtual void RequestRemoteControlUpdate(const FGuid& InInstanceId, const FSoftObjectPath& InAssetPath, const FString& InChannelName, const FAvalancheRemoteControlValues& InRemoteControlValues) override {};
-	virtual void RequestPlayableTransitionStart(const FGuid& InTransitionId, TArray<FGuid>&& InEnterInstanceIds, TArray<FGuid>&& InPlayingInstanceIds, TArray<FGuid>&& InExitInstanceIds, TArray<FAvalancheRemoteControlValues>&& InEnterValues, const FName& InChannelName, EAvalanchePlayableTransitionFlags InTransitionFlags) override {}
+	virtual void RequestPlayback(const FGuid& InInstanceId, const FSoftObjectPath& InAssetPath, const FString& InChannelName, EAvaPlaybackAction InAction, const FString& InArguments) override {}
+	virtual void RequestAnimPlayback(const FGuid& InInstanceId, const FSoftObjectPath& InAssetPath, const FString& InChannelName, const FAvaPlaybackAnimPlaySettings& InAnimSettings) override {};
+	virtual void RequestAnimAction(const FGuid& InInstanceId, const FSoftObjectPath& InAssetPath, const FString& InChannelName, const FString& InAnimationName, EAvaPlaybackAnimAction InAction) override {};
+	virtual void RequestRemoteControlUpdate(const FGuid& InInstanceId, const FSoftObjectPath& InAssetPath, const FString& InChannelName, const FAvaPlayableRemoteControlValues& InRemoteControlValues) override {};
+	virtual void RequestPlayableTransitionStart(const FGuid& InTransitionId, TArray<FGuid>&& InEnterInstanceIds, TArray<FGuid>&& InPlayingInstanceIds, TArray<FGuid>&& InExitInstanceIds, TArray<FAvaPlayableRemoteControlValues>&& InEnterValues, const FName& InChannelName, EAvaPlayableTransitionFlags InTransitionFlags) override {}
 	virtual void RequestPlayableTransitionStop(const FGuid& InTransitionId, const FName& InChannelName) override {}
-	virtual void RequestBroadcast(const FString& InProfile, const FName& InChannel, const TArray<UMediaOutput*>& InRemoteMediaOutputs, EAvaMediaBroadcastAction InAction) override {}
+	virtual void RequestBroadcast(const FString& InProfile, const FName& InChannel, const TArray<UMediaOutput*>& InRemoteMediaOutputs, EAvaBroadcastAction InAction) override {}
 	virtual bool IsMediaOutputRemoteFallback(const UMediaOutput* InMediaOutput) override { return false;}
-	virtual EAvaMediaIssueSeverity GetMediaOutputIssueSeverity(const FString& InServerName, const FString& InChannelName, const FGuid& InOutputGuid) const override { return EAvaMediaIssueSeverity::None;}
+	virtual EAvaBroadcastIssueSeverity GetMediaOutputIssueSeverity(const FString& InServerName, const FString& InChannelName, const FGuid& InOutputGuid) const override { return EAvaBroadcastIssueSeverity::None;}
 	virtual const TArray<FString>& GetMediaOutputIssueMessages(const FString& InServerName, const FString& InChannelName, const FGuid& InOutputGuid) const override { return EmptyStringArray; }
-	virtual EAvaMediaOutputState GetMediaOutputState(const FString& InServerName, const FString& InChannelName, const FGuid& InOutputGuid) const override { return EAvaMediaOutputState::Error; }
+	virtual EAvaBroadcastOutputState GetMediaOutputState(const FString& InServerName, const FString& InChannelName, const FGuid& InOutputGuid) const override { return EAvaBroadcastOutputState::Error; }
 	virtual bool HasAnyServerOnlineForChannel(const FName& InChannelName) const override { return false;}
-   	virtual TOptional<EAvaMediaPlaybackStatus> GetRemotePlaybackStatus(const FGuid& InInstanceId, const FSoftObjectPath& InAssetPath, const FString& InChannelName, const FString& InServerName) const override
+   	virtual TOptional<EAvaPlaybackStatus> GetRemotePlaybackStatus(const FGuid& InInstanceId, const FSoftObjectPath& InAssetPath, const FString& InChannelName, const FString& InServerName) const override
 	{
-		return EAvaMediaPlaybackStatus::Unknown;
+		return EAvaPlaybackStatus::Unknown;
 	}
 	virtual const FString* GetRemotePlaybackUserData(const FGuid& InInstanceId, const FSoftObjectPath& InAssetPath, const FString& InChannelName, const FString& InServerName) const override
 	{
 		return nullptr;
 	}
-	virtual TOptional<EAvaMediaPlaybackAssetStatus> GetRemotePlaybackAssetStatus(const FSoftObjectPath& InAssetPath, const FString& InServerName) const override
+	virtual TOptional<EAvaPlaybackAssetStatus> GetRemotePlaybackAssetStatus(const FSoftObjectPath& InAssetPath, const FString& InServerName) const override
 	{
-		return EAvaMediaPlaybackAssetStatus::Unknown;
+		return EAvaPlaybackAssetStatus::Unknown;
 	}
-	//~IAvaMediaPlaybackClient
+	//~IAvaPlaybackClient
 };
-const FString FAvaMediaPlaybackClientDummy::EmptyString;
-const TArray<FString> FAvaMediaPlaybackClientDummy::EmptyStringArray;
+const FString FAvaPlaybackClientDummy::EmptyString;
+const TArray<FString> FAvaPlaybackClientDummy::EmptyStringArray;
 
-IAvaMediaPlaybackClient& FAvaMediaModule::GetMediaPlaybackClient()
+IAvaPlaybackClient& FAvaMediaModule::GetMediaPlaybackClient()
 {
-	if (AvaMediaPlaybackClient.IsValid())
+	if (AvaPlaybackClient.IsValid())
 	{
-		return *AvaMediaPlaybackClient;
+		return *AvaPlaybackClient;
 	}
 	else
 	{
-		static FAvaMediaPlaybackClientDummy DummyClient;
+		static FAvaPlaybackClientDummy DummyClient;
 		return DummyClient;
 	}
 }
 
 const IMediaIOCoreDeviceProvider* FAvaMediaModule::GetDeviceProvider(FName InProviderName, const FMediaIOOutputConfiguration* InMediaIOOutputConfiguration) const
 {
-	if (AvaMediaPlaybackClient.IsValid() && InMediaIOOutputConfiguration)
+	if (AvaPlaybackClient.IsValid() && InMediaIOOutputConfiguration)
 	{
-		const FAvaDeviceProviderWrapper* Wrapper = DeviceProviderProxyManager.GetDeviceProviderWrapper(InProviderName);
+		const FAvaBroadcastDeviceProviderWrapper* Wrapper = DeviceProviderProxyManager.GetDeviceProviderWrapper(InProviderName);
 		if (Wrapper)
 		{
 			return Wrapper->GetProviderForDeviceName(InMediaIOOutputConfiguration->MediaConfiguration.MediaConnection.Device.DeviceName);
@@ -272,7 +273,7 @@ const IMediaIOCoreDeviceProvider* FAvaMediaModule::GetDeviceProvider(FName InPro
 
 TArray<const IMediaIOCoreDeviceProvider*> FAvaMediaModule::GetDeviceProvidersForServer(const FString& InServerName) const
 {
-	if (AvaMediaPlaybackClient.IsValid())
+	if (AvaPlaybackClient.IsValid())
 	{
 		return DeviceProviderProxyManager.GetDeviceProvidersForServer(InServerName);
 	}
@@ -294,10 +295,10 @@ FString FAvaMediaModule::GetServerNameForDevice(const FName& InDeviceProviderNam
 
 	// Fallback (legacy method)
 	// Assumes the device name starts with the server name.
-	if (AvaMediaPlaybackClient.IsValid())
+	if (AvaPlaybackClient.IsValid())
 	{
 		const FString DeviceName = InDeviceName.ToString();
-		TArray<FString> ServerNames = AvaMediaPlaybackClient->GetServerNames();
+		TArray<FString> ServerNames = AvaPlaybackClient->GetServerNames();
 		for (const FString& ServerName : ServerNames)
 		{
 			if (DeviceName.StartsWith(ServerName))
@@ -332,7 +333,7 @@ void FAvaMediaModule::LaunchGameModeLocalPlaybackServer()
 
 	if (!LocalPlaybackServerProcess.IsValid())
 	{
-		LocalPlaybackServerProcess = FAvaMediaPlaybackServerProcess::FindOrCreate(*AvaMediaPlaybackClient);
+		LocalPlaybackServerProcess = FAvaPlaybackServerProcess::FindOrCreate(*AvaPlaybackClient);
 	}
 
 	if (LocalPlaybackServerProcess.IsValid())
@@ -348,9 +349,9 @@ void FAvaMediaModule::StopGameModeLocalPlaybackServer()
 {
 	// We may not have a local handle to the server process if it was started by
 	// another client instance.
-	if (!LocalPlaybackServerProcess.IsValid() && AvaMediaPlaybackClient.IsValid())
+	if (!LocalPlaybackServerProcess.IsValid() && AvaPlaybackClient.IsValid())
 	{
-		LocalPlaybackServerProcess = FAvaMediaPlaybackServerProcess::Find(*AvaMediaPlaybackClient);
+		LocalPlaybackServerProcess = FAvaPlaybackServerProcess::Find(*AvaPlaybackClient);
 	}
 	
 	if (LocalPlaybackServerProcess.IsValid() && LocalPlaybackServerProcess->IsLaunched())
@@ -367,7 +368,7 @@ bool FAvaMediaModule::IsGameModeLocalPlaybackServerLaunched() const
 	return LocalPlaybackServerProcess.IsValid() && LocalPlaybackServerProcess->IsLaunched();
 }
 
-const IAvalancheBroadcastSettings& FAvaMediaModule::GetBroadcastSettings() const
+const IAvaBroadcastSettings& FAvaMediaModule::GetBroadcastSettings() const
 {
 	return BroadcastSettingsBridge;
 }
@@ -375,24 +376,24 @@ const IAvalancheBroadcastSettings& FAvaMediaModule::GetBroadcastSettings() const
 const FAvaInstanceSettings& FAvaMediaModule::GetAvaInstanceSettings() const
 {
 	// if the server is enabled, fetch the settings from the currently connected client.
-	if (AvaMediaPlaybackServer.IsValid())
+	if (AvaPlaybackServer.IsValid())
 	{
-		if (const FAvaInstanceSettings* SettingsFromClient = AvaMediaPlaybackServer->GetAvaInstanceSettings())
+		if (const FAvaInstanceSettings* SettingsFromClient = AvaPlaybackServer->GetAvaInstanceSettings())
 		{
 			return *SettingsFromClient;
 		}
 	}
 	// Return the local settings.
-	return UAvalancheMediaSettings::Get().AvaInstanceSettings;
+	return UAvaMediaSettings::Get().AvaInstanceSettings;
 }
 
-FAvaMediaPlaybackManager& FAvaMediaModule::GetLocalPlaybackManager() const
+FAvaPlaybackManager& FAvaMediaModule::GetLocalPlaybackManager() const
 {
 	check(LocalPlaybackManager.IsValid());
 	return *LocalPlaybackManager;
 }
 
-FAvalancheManagedInstanceCache& FAvaMediaModule::GetManagedInstanceCache() const
+FAvaRundownManagedInstanceCache& FAvaMediaModule::GetManagedInstanceCache() const
 {
 	check(ManagedInstanceCache.IsValid());
 	return *ManagedInstanceCache;
@@ -419,19 +420,19 @@ void FAvaMediaModule::PostEngineInit()
 	using namespace UE::AvaMediaModule::Private;
 
 	// This needs to happen late in the loading process, otherwise it fails.
-	const UAvalancheMediaSettings& Settings = UAvalancheMediaSettings::Get();
+	const UAvaMediaSettings& Settings = UAvaMediaSettings::Get();
 
-	ManagedInstanceCache = MakeShared<FAvalancheManagedInstanceCache>();
+	ManagedInstanceCache = MakeShared<FAvaRundownManagedInstanceCache>();
 	
 	// Allow for specification of the server name in the command line.
 	// Command line has priority over project settings. 
 	FString PlaybackServerName = Settings.PlaybackServerName;
 	const bool bIsServerManuallyStarted = IsPlaybackServerManuallyStarted(PlaybackServerName);
-	const bool bIsClientManuallyStarted = FParse::Param(FCommandLine::Get(), TEXT("AvaMediaPlaybackClientStart"));
+	const bool bIsClientManuallyStarted = FParse::Param(FCommandLine::Get(), TEXT("MotionDesignPlaybackClientStart"));
 
 	// Adding a command to suppress the client from auto-starting. This is used when spawning
 	// extra server process from the same project, while preventing extra clients.
-	const bool bIsClientAutoStartSuppressed = FParse::Param(FCommandLine::Get(), TEXT("AvaMediaPlaybackClientSuppress"));
+	const bool bIsClientAutoStartSuppressed = FParse::Param(FCommandLine::Get(), TEXT("MotionDesignPlaybackClientSuppress"));
 	
 	bool bShouldStartClient = bIsClientManuallyStarted || (Settings.bAutoStartPlaybackClient && !bIsClientAutoStartSuppressed && !IsRunningCommandlet());
 	bool bShouldStartServer = bIsServerManuallyStarted || (Settings.bAutoStartPlaybackServer && !IsRunningCommandlet());
@@ -470,12 +471,12 @@ void FAvaMediaModule::PostEngineInit()
 		StartPlaybackServerCommand({PlaybackServerName});
 	}
 	
-	LocalPlaybackManager = MakeShared<FAvaMediaPlaybackManager>();
+	LocalPlaybackManager = MakeShared<FAvaPlaybackManager>();
 
 #if WITH_EDITOR
 	// Capture Raw Ptr to avoid keeping ref count on capture
 	// Only Local Playback Manager should handle tear down for PIE End
-	FAvaMediaPlaybackManager* LocalPlaybackManagerRaw = LocalPlaybackManager.Get();
+	FAvaPlaybackManager* LocalPlaybackManagerRaw = LocalPlaybackManager.Get();
 	FEditorDelegates::PrePIEEnded.AddSPLambda(LocalPlaybackManagerRaw, [LocalPlaybackManagerRaw](const bool)
 	{
 		LocalPlaybackManagerRaw->OnParentWorldBeginTearDown();
@@ -483,7 +484,7 @@ void FAvaMediaModule::PostEngineInit()
 #endif
 
 	// Playback server required by Http server
-	if (AvaMediaPlaybackServer.IsValid() && Settings.bAutoStartWebServer)
+	if (AvaPlaybackServer.IsValid() && Settings.bAutoStartWebServer)
 	{
 		StartHttpPlaybackServerCommand({});
 	}
@@ -511,72 +512,72 @@ void FAvaMediaModule::StopAllServices()
 void FAvaMediaModule::StartPlaybackServerCommand(const TArray<FString>& InArgs)
 {
 	// Starting a playback server in the same process as playback client is forbidden.
-	if (AvaMediaPlaybackClient)
+	if (AvaPlaybackClient)
 	{
 		UE_LOG(LogAvaMedia, Error, TEXT("A Playback Server can't be started in the same process as a Playback Client."));
 		return;
 	}
 	
-	if (!AvaMediaPlaybackServer)
+	if (!AvaPlaybackServer)
 	{
-		AvaMediaPlaybackServer = MakeShared<FAvaMediaPlaybackServer>();
-		AvaMediaPlaybackServer->Init(InArgs.Num() > 0 ? InArgs[0] : TEXT(""));
-		OnAvaMediaPlaybackServerStarted.Broadcast();
+		AvaPlaybackServer = MakeShared<FAvaPlaybackServer>();
+		AvaPlaybackServer->Init(InArgs.Num() > 0 ? InArgs[0] : TEXT(""));
+		OnAvaPlaybackServerStarted.Broadcast();
 		UE_LOG(LogAvaMedia, Log, TEXT("Playback Server Started"));
 	}
 }
 
 void FAvaMediaModule::StopPlaybackServerCommand(const TArray<FString>& InArgs)
 {
-	if (AvaMediaPlaybackServer)
+	if (AvaPlaybackServer)
 	{
-		AvaMediaPlaybackServer->StartShuttingDown();
-		AvaMediaPlaybackServer->StopBroadcast();
-		AvaMediaPlaybackServer->StopPlaybacks();
-		OnAvaMediaPlaybackServerStopped.Broadcast();
+		AvaPlaybackServer->StartShuttingDown();
+		AvaPlaybackServer->StopBroadcast();
+		AvaPlaybackServer->StopPlaybacks();
+		OnAvaPlaybackServerStopped.Broadcast();
 	}
-	AvaMediaPlaybackServer.Reset();
+	AvaPlaybackServer.Reset();
 }
 
 void FAvaMediaModule::StartPlaybackClientCommand(const TArray<FString>& InArgs)
 {
 	// Starting a playback server in the same process as playback client is forbidden.
-	if (AvaMediaPlaybackServer)
+	if (AvaPlaybackServer)
 	{
 		UE_LOG(LogAvaMedia, Error, TEXT("A Playback Client can't be started in the same process as a Playback Server."));
 		return;
 	}
 
-	if (!AvaMediaPlaybackClient)
+	if (!AvaPlaybackClient)
 	{
-		using namespace UE::AvaMediaPlaybackClient::Delegates;
+		using namespace UE::AvaPlaybackClient::Delegates;
 		if (GetOnConnectionEvent().IsBoundToObject(this))
 		{
-			GetOnConnectionEvent().AddRaw(this, &FAvaMediaModule::OnAvaMediaPlaybackClientConnectionEvent);
+			GetOnConnectionEvent().AddRaw(this, &FAvaMediaModule::OnAvaPlaybackClientConnectionEvent);
 		}
 
-		AvaMediaPlaybackClient = MakeShared<FAvaMediaPlaybackClient>(this);
-		AvaMediaPlaybackClient->Init();
-		OnAvaMediaPlaybackClientStarted.Broadcast();
+		AvaPlaybackClient = MakeShared<FAvaPlaybackClient>(this);
+		AvaPlaybackClient->Init();
+		OnAvaPlaybackClientStarted.Broadcast();
 		UE_LOG(LogAvaMedia, Log, TEXT("Playback Client Started"));
 	}
 }
 
 void FAvaMediaModule::StopPlaybackClientCommand(const TArray<FString>& InArgs)
 {
-	if (AvaMediaPlaybackClient)
+	if (AvaPlaybackClient)
 	{
-		OnAvaMediaPlaybackClientStopped.Broadcast();
+		OnAvaPlaybackClientStopped.Broadcast();
 	}
-	AvaMediaPlaybackClient.Reset();
-	UE::AvaMediaPlaybackClient::Delegates::GetOnConnectionEvent().RemoveAll(this);
+	AvaPlaybackClient.Reset();
+	UE::AvaPlaybackClient::Delegates::GetOnConnectionEvent().RemoveAll(this);
 }
 
 void FAvaMediaModule::StartHttpPlaybackServerCommand(const TArray<FString>& InArgs)
 {
-	if (!AvaMediaHttpPlaybackServer)
+	if (!AvaPlaybackHttpPlaybackServer)
 	{
-		AvaMediaHttpPlaybackServer = MakeShared<FAvaMediaHttpServer>();
+		AvaPlaybackHttpPlaybackServer = MakeShared<FAvaPlaybackHttpServer>();
 		
 		if (InArgs.Num() > 0)
 		{
@@ -584,23 +585,23 @@ void FAvaMediaModule::StartHttpPlaybackServerCommand(const TArray<FString>& InAr
 		}
 	}
 	
-	if (!AvaMediaHttpPlaybackServer->IsRunning())
+	if (!AvaPlaybackHttpPlaybackServer->IsRunning())
 	{
-		const int32 Port = GetDefault<UAvalancheMediaSettings>()->HttpServerPort;
-		AvaMediaHttpPlaybackServer->Start(Port);
+		const int32 Port = GetDefault<UAvaMediaSettings>()->HttpServerPort;
+		AvaPlaybackHttpPlaybackServer->Start(Port);
 		UE_LOG(LogAvaMedia, Log, TEXT("Http Playback Server Started"));
 	}
 }
 
 void FAvaMediaModule::StopHttpPlaybackServerCommand(const TArray<FString>& InArgs)
 {
-	AvaMediaHttpPlaybackServer.Reset();
+	AvaPlaybackHttpPlaybackServer.Reset();
 }
 
 
 void FAvaMediaModule::SaveDeviceProvidersCommand(const TArray<FString>& InArgs)
 {
-	FAvaDeviceProviderDataList OutProviders;
+	FAvaBroadcastDeviceProviderDataList OutProviders;
 	OutProviders.Populate(InArgs.Num() > 0 ? InArgs[0] : FPlatformProcess::ComputerName());
 	OutProviders.SaveToJson();	// Saves in the project's config folder.
 	OutProviders.SaveToXml();	// Saves in the project's config folder.
@@ -631,23 +632,23 @@ void FAvaMediaModule::HandleStatCommand(const TArray<FString>& InArgs)
 
 	const bool bLocalCommandSucceeded = LocalPlaybackManager->HandleStatCommand(InArgs);
 	
-	if (AvaMediaPlaybackClient.IsValid())
+	if (AvaPlaybackClient.IsValid())
 	{
-		AvaMediaPlaybackClient->BroadcastStatCommand(InArgs[0], bLocalCommandSucceeded);
+		AvaPlaybackClient->BroadcastStatCommand(InArgs[0], bLocalCommandSucceeded);
 	}
 }
 
-void FAvaMediaModule::OnAvaMediaPlaybackClientConnectionEvent(IAvaMediaPlaybackClient& InPlaybackClient,
-		const UE::AvaMediaPlaybackClient::Delegates::FConnectionEventArgs& InArgs)
+void FAvaMediaModule::OnAvaPlaybackClientConnectionEvent(IAvaPlaybackClient& InPlaybackClient,
+		const UE::AvaPlaybackClient::Delegates::FConnectionEventArgs& InArgs)
 {
-	using namespace UE::AvaMediaPlaybackClient::Delegates;
+	using namespace UE::AvaPlaybackClient::Delegates;
 	// When a playback server connection event occurs, update the status of the playback server process.
 	switch (InArgs.Event)
 	{
 	case EConnectionEvent::ServerConnected:
 		if (!IsGameModeLocalPlaybackServerLaunched() && IsMediaPlaybackClientStarted())
 		{
-			LocalPlaybackServerProcess = FAvaMediaPlaybackServerProcess::Find(*AvaMediaPlaybackClient);
+			LocalPlaybackServerProcess = FAvaPlaybackServerProcess::Find(*AvaPlaybackClient);
 		}
 		break;
 	case EConnectionEvent::ServerDisconnected:
@@ -662,27 +663,27 @@ void FAvaMediaModule::OnAvaMediaPlaybackClientConnectionEvent(IAvaMediaPlaybackC
 
 const FLinearColor& FAvaMediaModule::FLocalBroadcastSettings::GetChannelClearColor() const
 {
-	return UAvalancheMediaSettings::Get().ChannelClearColor;
+	return UAvaMediaSettings::Get().ChannelClearColor;
 }
 
 EPixelFormat FAvaMediaModule::FLocalBroadcastSettings::GetDefaultPixelFormat() const
 {
-	return UAvalancheMediaSettings::Get().ChannelDefaultPixelFormat;
+	return UAvaMediaSettings::Get().ChannelDefaultPixelFormat;
 }
 
 const FIntPoint& FAvaMediaModule::FLocalBroadcastSettings::GetDefaultResolution() const
 {
-	return UAvalancheMediaSettings::Get().ChannelDefaultResolution;
+	return UAvaMediaSettings::Get().ChannelDefaultResolution;
 }
 
 bool FAvaMediaModule::FLocalBroadcastSettings::IsDrawPlaceholderWidget() const
 {
-	return UAvalancheMediaSettings::Get().bDrawPlaceholderWidget;
+	return UAvaMediaSettings::Get().bDrawPlaceholderWidget;
 }
 
 const FSoftObjectPath& FAvaMediaModule::FLocalBroadcastSettings::GetPlaceholderWidgetClass() const
 {
-	return UAvalancheMediaSettings::Get().PlaceholderWidgetClass.ToSoftObjectPath();
+	return UAvaMediaSettings::Get().PlaceholderWidgetClass.ToSoftObjectPath();
 }
 
 const FLinearColor& FAvaMediaModule::FBroadcastSettingsBridge::GetChannelClearColor() const
@@ -710,12 +711,12 @@ const FSoftObjectPath& FAvaMediaModule::FBroadcastSettingsBridge::GetPlaceholder
 	return GetSettings().GetPlaceholderWidgetClass();
 }
 
-const IAvalancheBroadcastSettings& FAvaMediaModule::FBroadcastSettingsBridge::GetSettings() const
+const IAvaBroadcastSettings& FAvaMediaModule::FBroadcastSettingsBridge::GetSettings() const
 {
 	// if server is enabled, fetch the setting from the currently connected client.
-	if (ParentModule->AvaMediaPlaybackServer.IsValid())
+	if (ParentModule->AvaPlaybackServer.IsValid())
 	{
-		if (const IAvalancheBroadcastSettings* SettingsFromClient = ParentModule->AvaMediaPlaybackServer->GetBroadcastSettings())
+		if (const IAvaBroadcastSettings* SettingsFromClient = ParentModule->AvaPlaybackServer->GetBroadcastSettings())
 		{
 			return *SettingsFromClient;
 		}
@@ -723,7 +724,7 @@ const IAvalancheBroadcastSettings& FAvaMediaModule::FBroadcastSettingsBridge::Ge
 	return ParentModule->LocalBroadcastSettings;
 }
 
-IAvaDeviceProviderProxyManager& FAvaMediaModule::GetDeviceProviderProxyManager()
+IAvaBroadcastDeviceProviderProxyManager& FAvaMediaModule::GetDeviceProviderProxyManager()
 {
 	return DeviceProviderProxyManager;
 }
diff --git a/Engine/Plugins/Experimental/Avalanche/Source/AvalancheMedia/Private/AvaMediaModule.h b/Engine/Plugins/Experimental/Avalanche/Source/AvalancheMedia/Private/AvaMediaModule.h
--- a/Engine/Plugins/Experimental/Avalanche/Source/AvalancheMedia/Private/AvaMediaModule.h
+++ b/Engine/Plugins/Experimental/Avalanche/Source/AvalancheMedia/Private/AvaMediaModule.h
@@ -2,19 +2,18 @@
 
 #pragma once
 
-#include "Http/AvaMediaHttpServer.h"
+#include "Broadcast/IAvaBroadcastSettings.h"
+#include "Broadcast/OutputDevices/AvaBroadcastDeviceProviderProxy.h"
+#include "Broadcast/OutputDevices/AvaBroadcastDisplayDeviceProvider.h"
 #include "IAvaMediaModule.h"
-#include "IAvalancheBroadcastSettings.h"
 #include "ModularFeature/AvaMediaSync.h"
-#include "ModuleDescriptor.h"
-#include "OutputDevices/AvaDeviceProviderProxy.h"
-#include "OutputDevices/AvaDisplayDeviceProvider.h"
-#include "Playback/AvaMediaPlaybackClient.h"
-#include "Playback/AvaMediaPlaybackClientDelegates.h"
-#include "Playback/AvaMediaPlaybackManager.h"
-#include "Playback/AvaMediaPlaybackServer.h"
-#include "Playback/AvaMediaPlaybackServerProcess.h"
-#include "Playlist/AvalancheManagedInstanceCache.h"
+#include "Playback/AvaPlaybackClient.h"
+#include "Playback/AvaPlaybackClientDelegates.h"
+#include "Playback/AvaPlaybackManager.h"
+#include "Playback/AvaPlaybackServer.h"
+#include "Playback/AvaPlaybackServerProcess.h"
+#include "Playback/Http/AvaPlaybackHttpServer.h"
+#include "Rundown/AvaRundownManagedInstanceCache.h"
 
 class FAvaMediaModule : public IAvaMediaModule
 {
@@ -24,15 +23,15 @@ public:
 	//~ Begin IAvaMediaModule
 	virtual void StartupModule() override;
 	virtual void ShutdownModule() override;
-	virtual bool IsMediaPlaybackClientStarted() const override { return AvaMediaPlaybackClient.IsValid();}
+	virtual bool IsMediaPlaybackClientStarted() const override { return AvaPlaybackClient.IsValid();}
 	virtual void StartMediaPlaybackClient() override;
 	virtual void StopMediaPlaybackClient() override;
-	virtual bool IsMediaPlaybackServerStarted() const override { return AvaMediaPlaybackServer.IsValid();}
+	virtual bool IsMediaPlaybackServerStarted() const override { return AvaPlaybackServer.IsValid();}
 	virtual void StartMediaPlaybackServer(const FString& InPlaybackServerName) override;
 	virtual void StopMediaPlaybackServer() override;
 	
-	virtual IAvaMediaPlaybackClient& GetMediaPlaybackClient() override;
-	virtual TSharedPtr<FAvaMediaPlaybackServer> GetMediaPlaybackServer() const override { return AvaMediaPlaybackServer; }
+	virtual IAvaPlaybackClient& GetMediaPlaybackClient() override;
+	virtual TSharedPtr<FAvaPlaybackServer> GetMediaPlaybackServer() const override { return AvaPlaybackServer; }
 	virtual const IMediaIOCoreDeviceProvider* GetDeviceProvider(FName InProviderName, const FMediaIOOutputConfiguration* InMediaIOOutputConfiguration) const override;
 	virtual TArray<const IMediaIOCoreDeviceProvider*> GetDeviceProvidersForServer(const FString& InServerName) const override;
 	virtual FString GetServerNameForDevice(const FName& InDeviceProviderName, const FName& InDeviceName) const override;
@@ -40,22 +39,22 @@ public:
 	virtual void LaunchGameModeLocalPlaybackServer() override;
 	virtual void StopGameModeLocalPlaybackServer() override;
 	virtual bool IsGameModeLocalPlaybackServerLaunched() const override;
-	virtual const IAvalancheBroadcastSettings& GetBroadcastSettings() const override;
+	virtual const IAvaBroadcastSettings& GetBroadcastSettings() const override;
 	virtual const FAvaInstanceSettings& GetAvaInstanceSettings() const override;
-	virtual FAvaMediaPlaybackManager& GetLocalPlaybackManager() const override;
-	virtual FAvalancheManagedInstanceCache& GetManagedInstanceCache() const override;
+	virtual FAvaPlaybackManager& GetLocalPlaybackManager() const override;
+	virtual FAvaRundownManagedInstanceCache& GetManagedInstanceCache() const override;
 	virtual bool IsAvaMediaSyncProviderFeatureAvailable() const override;
 	virtual IAvaMediaSyncProvider* GetAvaMediaSyncProvider() const override;
 	virtual void NotifyMapChangedEvent(UWorld* InWorld, EAvaMediaMapChangeType InEventType) override;
 	virtual FOnAvaMediaSyncProviderChanged& GetOnAvaMediaSyncProviderChanged() override { return OnAvaMediaSyncProviderChanged; }
 	virtual FOnAvaMediaSyncPackageModified& GetOnAvaMediaSyncPackageModified() override { return OnAvaMediaSyncPackageModified; }
 	virtual FOnMapChangedEvent& GetOnMapChangedEvent() override { return OnMapChangedEvent; }
-	virtual FOnAvaMediaPlaybackClientStarted& GetOnAvaMediaPlaybackClientStarted() override { return OnAvaMediaPlaybackClientStarted; }
-	virtual FOnAvaMediaPlaybackClientStopped& GetAvaMediaPlaybackClientStopped() override { return OnAvaMediaPlaybackClientStopped; }
-	virtual FOnAvaMediaPlaybackServerStarted& GetOnAvaMediaPlaybackServerStarted() override { return OnAvaMediaPlaybackServerStarted; }
-	virtual FOnAvaMediaPlaybackServerStopped& GetAvaMediaPlaybackServerStopped() override { return OnAvaMediaPlaybackServerStopped; }
+	virtual FOnAvaPlaybackClientStarted& GetOnAvaPlaybackClientStarted() override { return OnAvaPlaybackClientStarted; }
+	virtual FOnAvaPlaybackClientStopped& GetAvaPlaybackClientStopped() override { return OnAvaPlaybackClientStopped; }
+	virtual FOnAvaPlaybackServerStarted& GetOnAvaPlaybackServerStarted() override { return OnAvaPlaybackServerStarted; }
+	virtual FOnAvaPlaybackServerStopped& GetAvaPlaybackServerStopped() override { return OnAvaPlaybackServerStopped; }
 	virtual FGetEditorViewportClient& GetEditorViewportClientDelegate() override { return GetEditorViewportClient; }
-	virtual IAvaDeviceProviderProxyManager& GetDeviceProviderProxyManager() override;
+	virtual IAvaBroadcastDeviceProviderProxyManager& GetDeviceProviderProxyManager() override;
 	//~ End IAvaMediaModule
 	
 private:
@@ -81,31 +80,31 @@ private:
 	void HandleStatCommand(const TArray<FString>& InArgs);
 
 	// Event Handlers
-	void OnAvaMediaPlaybackClientConnectionEvent(IAvaMediaPlaybackClient& InPlaybackClient,
-		const UE::AvaMediaPlaybackClient::Delegates::FConnectionEventArgs& InArgs);
+	void OnAvaPlaybackClientConnectionEvent(IAvaPlaybackClient& InPlaybackClient,
+		const UE::AvaPlaybackClient::Delegates::FConnectionEventArgs& InArgs);
 
 private:
-	FAvaDisplayDeviceProvider AvaDisplayDeviceProvider;
-	FAvaDeviceProviderProxyManager DeviceProviderProxyManager;
+	FAvaBroadcastDisplayDeviceProvider AvaDisplayDeviceProvider;
+	FAvaBroadcastDeviceProviderProxyManager DeviceProviderProxyManager;
 	TUniquePtr<FAvaMediaSync> AvaMediaSync;
 
 	TArray<IConsoleObject*> ConsoleCmds;
 	
 	/**
-	 *	Wraps the local default UAvalancheMediaSettings.
+	 *	Wraps the local default UAvaMediaSettings.
 	 */
-	class FLocalBroadcastSettings : public IAvalancheBroadcastSettings
+	class FLocalBroadcastSettings : public IAvaBroadcastSettings
 	{
 	public:
 		virtual ~FLocalBroadcastSettings() override = default;
 
-		//~ Begin IAvalancheBroadcastSettings
+		//~ Begin IAvaBroadcastSettings
 		virtual const FLinearColor& GetChannelClearColor() const override;
 		virtual EPixelFormat GetDefaultPixelFormat() const override;
 		virtual const FIntPoint& GetDefaultResolution() const override;
 		virtual bool IsDrawPlaceholderWidget() const override;
 		virtual const FSoftObjectPath& GetPlaceholderWidgetClass() const override;
-		//~ End IAvalancheBroadcastSettings
+		//~ End IAvaBroadcastSettings
 	};
 	FLocalBroadcastSettings LocalBroadcastSettings;
 
@@ -114,41 +113,41 @@ private:
 	 * even if a client disconnects, it will either switch to another client
 	 * or use the local settings.
 	 */
-	class FBroadcastSettingsBridge final : public IAvalancheBroadcastSettings
+	class FBroadcastSettingsBridge final : public IAvaBroadcastSettings
 	{
 	public:
 		explicit FBroadcastSettingsBridge(FAvaMediaModule* Module) : ParentModule(Module) {}
 		virtual ~FBroadcastSettingsBridge() override = default;
 
-		//~ Begin IAvalancheBroadcastSettings
+		//~ Begin IAvaBroadcastSettings
 		virtual const FLinearColor& GetChannelClearColor() const override;
 		virtual EPixelFormat GetDefaultPixelFormat() const override;
 		virtual const FIntPoint& GetDefaultResolution() const override;
 		virtual bool IsDrawPlaceholderWidget() const override;
 		virtual const FSoftObjectPath& GetPlaceholderWidgetClass() const override;
-		//~ End IAvalancheBroadcastSettings
+		//~ End IAvaBroadcastSettings
 
 	private:
-		const IAvalancheBroadcastSettings& GetSettings() const;
+		const IAvaBroadcastSettings& GetSettings() const;
 		
 		FAvaMediaModule* ParentModule = nullptr;
 	};
 	FBroadcastSettingsBridge BroadcastSettingsBridge;
 	
-	TSharedPtr<FAvaMediaPlaybackServer> AvaMediaPlaybackServer;	
-	TSharedPtr<FAvaMediaPlaybackClient> AvaMediaPlaybackClient;
-	TSharedPtr<FAvaMediaPlaybackServerProcess> LocalPlaybackServerProcess;
-	TSharedPtr<FAvaMediaPlaybackManager> LocalPlaybackManager;
-	TSharedPtr<FAvalancheManagedInstanceCache> ManagedInstanceCache;
+	TSharedPtr<FAvaPlaybackServer> AvaPlaybackServer;	
+	TSharedPtr<FAvaPlaybackClient> AvaPlaybackClient;
+	TSharedPtr<FAvaPlaybackServerProcess> LocalPlaybackServerProcess;
+	TSharedPtr<FAvaPlaybackManager> LocalPlaybackManager;
+	TSharedPtr<FAvaRundownManagedInstanceCache> ManagedInstanceCache;
 
-	TSharedPtr<FAvaMediaHttpServer> AvaMediaHttpPlaybackServer;
+	TSharedPtr<FAvaPlaybackHttpServer> AvaPlaybackHttpPlaybackServer;
 
 	FOnAvaMediaSyncProviderChanged OnAvaMediaSyncProviderChanged;
 	FOnAvaMediaSyncPackageModified OnAvaMediaSyncPackageModified;
 	FOnMapChangedEvent OnMapChangedEvent;
-	FOnAvaMediaPlaybackClientStarted OnAvaMediaPlaybackClientStarted;
-	FOnAvaMediaPlaybackClientStopped OnAvaMediaPlaybackClientStopped;
-	FOnAvaMediaPlaybackServerStarted OnAvaMediaPlaybackServerStarted;
-	FOnAvaMediaPlaybackServerStopped OnAvaMediaPlaybackServerStopped;
+	FOnAvaPlaybackClientStarted OnAvaPlaybackClientStarted;
+	FOnAvaPlaybackClientStopped OnAvaPlaybackClientStopped;
+	FOnAvaPlaybackServerStarted OnAvaPlaybackServerStarted;
+	FOnAvaPlaybackServerStopped OnAvaPlaybackServerStopped;
 	FGetEditorViewportClient GetEditorViewportClient;
 };
diff --git a/Engine/Plugins/Experimental/Avalanche/Source/AvalancheMedia/Private/AvaMediaStyle.cpp b/Engine/Plugins/Experimental/Avalanche/Source/AvalancheMedia/Private/AvaMediaStyle.cpp
--- a/Engine/Plugins/Experimental/Avalanche/Source/AvalancheMedia/Private/AvaMediaStyle.cpp
+++ b/Engine/Plugins/Experimental/Avalanche/Source/AvalancheMedia/Private/AvaMediaStyle.cpp
@@ -1,6 +1,7 @@
 // Copyright Epic Games, Inc. All Rights Reserved.
 
 #include "AvaMediaStyle.h"
+
 #include "Brushes/SlateImageBrush.h"
 #include "Interfaces/IPluginManager.h"
 #include "Styling/SlateStyle.h"
@@ -26,7 +27,7 @@ void FAvaMediaStyle::Shutdown()
 
 FName FAvaMediaStyle::GetStyleSetName()
 {
-	static FName StyleSetName(TEXT("AvalancheMedia"));
+	static FName StyleSetName(TEXT("AvaMedia"));
 	return StyleSetName;
 }
 
@@ -50,17 +51,17 @@ const FVector2D Icon200x200(200.f, 200.f);
 
 TSharedRef<FSlateStyleSet> FAvaMediaStyle::Create()
 {
-	TSharedRef<FSlateStyleSet> Style = MakeShared<FSlateStyleSet>("AvalancheMedia");
+	TSharedRef<FSlateStyleSet> Style = MakeShared<FSlateStyleSet>("AvaMedia");
 	
-	TSharedPtr<IPlugin> Plugin = IPluginManager::Get().FindPlugin(TEXT("Avalanche"));
+	TSharedPtr<IPlugin> Plugin = IPluginManager::Get().FindPlugin(UE_PLUGIN_NAME);
 	check(Plugin.IsValid());
 	if (Plugin.IsValid())
 	{
 		Style->SetContentRoot(FPaths::Combine(Plugin->GetBaseDir(), TEXT("Resources")));
 	}
 	
-	Style->Set("AvalancheMedia.UnrealIcon", new IMAGE_BRUSH("Icons/MediaIcons/ue_logo", Icon200x200));
-	Style->Set("ClassIcon.AvalancheBroadcast", new IMAGE_BRUSH("Icons/MediaIcons/MediaOutput", Icon16x16));
+	Style->Set("AvaMedia.UnrealIcon", new IMAGE_BRUSH("Icons/MediaIcons/ue_logo", Icon200x200));
+	Style->Set("ClassIcon.AvaBroadcast", new IMAGE_BRUSH("Icons/MediaIcons/MediaOutput", Icon16x16));
 
 	return Style;
 }
diff --git a/Engine/Plugins/Experimental/Avalanche/Source/AvalancheMedia/Private/Playable/AvaPlayable.cpp b/Engine/Plugins/Experimental/Avalanche/Source/AvalancheMedia/Private/Playable/AvaPlayable.cpp
new file mode 100644
--- /dev/null
+++ b/Engine/Plugins/Experimental/Avalanche/Source/AvalancheMedia/Private/Playable/AvaPlayable.cpp
@@ -0,0 +1,386 @@
+// Copyright Epic Games, Inc. All Rights Reserved.
+
+#include "Playable/AvaPlayable.h"
+
+#include "AvaSequence.h"
+#include "AvaSequencePlayer.h"
+#include "Broadcast/AvaBroadcast.h"
+#include "Engine/Engine.h"
+#include "Framework/AvaSoftAssetPtr.h"
+#include "IAvaSceneInterface.h"
+#include "IAvaSequenceProvider.h"
+#include "Playable/AvaPlayableGroup.h"
+#include "Playable/AvaPlayableRemoteControlValues.h"
+#include "Playable/Playables/AvaPlayableBlueprint.h"
+#include "Playable/Playables/AvaPlayableLevelStreaming.h"
+#include "Playable/Playables/AvaPlayableRemoteProxy.h"
+#include "Playback/AvaPlaybackUtils.h"
+#include "RemoteControlPreset.h"
+
+DEFINE_LOG_CATEGORY(LogAvaPlayable);
+
+#define LOCTEXT_NAMESPACE "AvaPlayable"
+
+namespace UE::AvaPlayable::Private
+{
+	bool ShouldCreateLocalPlayable(const FName& InChannelName, const UAvaBroadcast& InBroadcast)
+	{
+		const FAvaBroadcastOutputChannel& Channel = InBroadcast.GetCurrentProfile().GetChannel(InChannelName);
+
+		// If there is no broadcast channel defined, like for preview (by default), then this is a local playable.
+		if (!Channel.IsValidChannel())
+		{
+			return true;
+		}
+
+		// For non-preview, the commands will be executed locally if the channel has at least one local outputs or no outputs.
+		// The "no outputs" condition is considered valid. Empty channels run locally.
+		if (Channel.HasAnyLocalMediaOutputs() || Channel.GetMediaOutputs().IsEmpty())
+		{
+			return true;
+		}
+		
+		return false;
+	}
+
+	bool HasRemoteOutputs(const FName& InChannelName, const UAvaBroadcast& InBroadcast)
+	{
+		const FAvaBroadcastOutputChannel& Channel = InBroadcast.GetCurrentProfile().GetChannel(InChannelName);
+		return Channel.IsValidChannel() && Channel.HasAnyRemoteMediaOutputs();
+	}
+
+	FString GetPrettyPlayableInfo(const UAvaPlayable* InPlayable)
+	{
+		if (InPlayable)
+		{
+			return FString::Printf(TEXT("Id:%s, Asset:%s, Status:%s"),
+				*InPlayable->GetInstanceId().ToString(),
+				*InPlayable->GetSourceAssetPath().ToString(),
+				*StaticEnum<EAvaPlayableStatus>()->GetNameByValue(static_cast<int32>(InPlayable->GetPlayableStatus())).ToString());
+		}
+		return TEXT("(nullptr)");
+	}
+
+	FString GetBriefFrameInfo()
+	{
+		return UE::AvaPlayback::Utils::GetBriefFrameInfo();
+	}
+	
+	FString GetPrettySequenceCommandInfo(EAvaPlaybackAnimAction InAnimAction, const FAvaPlaybackAnimPlaySettings& InAnimPlaySettings)
+	{
+		return FString::Printf(TEXT("Action:%s, Name:%s"),
+			*StaticEnum<EAvaPlaybackAnimAction>()->GetNameByValue(static_cast<int32>(InAnimAction)).ToString(),
+			*InAnimPlaySettings.AnimationName.ToString());
+	}
+}
+
+UAvaPlayable::FOnSequenceEvent UAvaPlayable::OnSequenceEventDelegate;
+UAvaPlayable::FOnTransitionEvent UAvaPlayable::OnTransitionEventDelegate;
+
+UAvaPlayable* UAvaPlayable::Create(UObject* InOuter, const FPlayableCreationInfo& InPlayableInfo)
+{
+	using namespace UE::AvaPlayable::Private;
+
+	const UAvaBroadcast& Broadcast = UAvaBroadcast::Get();
+	
+	UAvaPlayable* NewPlayable;
+	
+	// Forked channels considerations:
+	// - The case of forked remote channels is/will be handled internally to the RemoteProxy playable.
+	// - The case of forked local and remote channels will lead to a local playable and a remote proxy playable.
+	//   It would require wrapping the playables in a composite (or facade?) proxy. TODO
+	
+	if (ShouldCreateLocalPlayable(InPlayableInfo.ChannelName, Broadcast))
+	{
+		// For the moment, remote outputs will be ignored.
+		if (HasRemoteOutputs(InPlayableInfo.ChannelName, Broadcast))
+		{
+			UE_LOG(LogAvaPlayable, Error, TEXT("Forked Channels with both local and remote outputs are not supported in this version. Only local instance will be created."));
+		}
+		
+		NewPlayable = CreateLocalPlayable(InOuter, InPlayableInfo);
+	}
+	else
+	{
+		// Purely remote channel.
+		NewPlayable = CreateRemoteProxyPlayable(InOuter, InPlayableInfo);
+	}
+
+	// Finish the setup.
+	if (NewPlayable && !NewPlayable->InitPlayable(InPlayableInfo))
+	{
+		// final setup may fail, in this case the playable is discarded.
+		return nullptr;
+	}
+
+	return NewPlayable;
+}
+
+const FSoftObjectPath& UAvaPlayable::GetSourceAssetPath() const
+{
+	static const FSoftObjectPath Empty;
+	return Empty;
+}
+
+EAvaPlayableCommandResult UAvaPlayable::ExecuteAnimationCommand(EAvaPlaybackAnimAction InAnimAction, const FAvaPlaybackAnimPlaySettings& InAnimPlaySettings)
+{
+	using namespace UE::AvaPlayable::Private;
+
+	const EAvaPlayableStatus PlayableStatus = GetPlayableStatus();
+
+	if (PlayableStatus == EAvaPlayableStatus::Unknown
+		|| PlayableStatus == EAvaPlayableStatus::Error
+		|| PlayableStatus == EAvaPlayableStatus::Unloaded)
+	{
+		UE_LOG(LogAvaPlayable, Verbose,
+			TEXT("%s Playable {%s} -> Discarding Sequence Command: {%s}."),
+			*GetBriefFrameInfo(), *GetPrettyPlayableInfo(this), *GetPrettySequenceCommandInfo(InAnimAction, InAnimPlaySettings));
+	
+		// Discard the command
+		return EAvaPlayableCommandResult::ErrorDiscard;
+	}
+
+	// Asset status must be visible to run the animation commands.
+	// If not visible, the components are not yet added to the world and animations won't execute.
+	if (PlayableStatus != EAvaPlayableStatus::Visible)
+	{
+		UE_LOG(LogAvaPlayable, Verbose,
+			TEXT("%s Playable {%s} -> ReQueueing Sequence Command: {%s}."),
+			*GetBriefFrameInfo(), *GetPrettyPlayableInfo(this), *GetPrettySequenceCommandInfo(InAnimAction, InAnimPlaySettings));
+		
+		// Keep the command in the queue for next tick.
+		return EAvaPlayableCommandResult::KeepPending;
+	}
+
+	const IAvaSceneInterface* Scene = GetSceneInterface();
+	if (!Scene)
+	{
+		return EAvaPlayableCommandResult::ErrorDiscard;
+	}
+	
+	IAvaSequencePlaybackObject* PlaybackObject = Scene->GetPlaybackObject();
+	
+	if (!PlaybackObject)
+	{
+		return EAvaPlayableCommandResult::ErrorDiscard;
+	}
+
+	const IAvaSequenceProvider* SequenceProvider = Scene->GetSequenceProvider();
+
+	if (!SequenceProvider)
+	{
+		return EAvaPlayableCommandResult::ErrorDiscard;
+	}
+
+	UE_LOG(LogAvaPlayable, Verbose,
+		TEXT("%s Playable {%s} -> Executing Sequence Command: {%s}."),
+		*GetBriefFrameInfo(), *GetPrettyPlayableInfo(this), *GetPrettySequenceCommandInfo(InAnimAction, InAnimPlaySettings));
+	
+	for (const TObjectPtr<UAvaSequence>& Sequence : SequenceProvider->GetSequences())
+	{
+		// Remark: if the command doesn't specify the sequence name, we run the command on all the sequences.
+		if (Sequence && (Sequence->GetFName() == InAnimPlaySettings.AnimationName || InAnimPlaySettings.AnimationName.IsNone()))
+		{
+			if (InAnimAction == EAvaPlaybackAnimAction::Play || InAnimAction == EAvaPlaybackAnimAction::PreviewFrame)
+			{
+				UAvaSequencePlayer* const SequencePlayer = PlaybackObject->PlaySequence(Sequence, InAnimPlaySettings.AsPlayParams());
+				if (InAnimAction == EAvaPlaybackAnimAction::PreviewFrame && SequencePlayer)
+				{
+					SequencePlayer->PreviewFrame();
+				}
+			}
+			else if (InAnimAction == EAvaPlaybackAnimAction::Continue)
+			{
+				PlaybackObject->ContinueSequence(Sequence);
+			}
+			else if (InAnimAction == EAvaPlaybackAnimAction::Stop)
+			{
+				PlaybackObject->StopSequence(Sequence);
+			}
+		}
+	}
+	
+	return EAvaPlayableCommandResult::Executed;
+}
+
+EAvaPlayableCommandResult UAvaPlayable::UpdateRemoteControlCommand(const TSharedRef<FAvaPlayableRemoteControlValues>& InRemoteControlValues)
+{
+	const EAvaPlayableStatus PlayableStatus = GetPlayableStatus();
+
+	if (PlayableStatus == EAvaPlayableStatus::Unknown
+		|| PlayableStatus == EAvaPlayableStatus::Error
+		|| PlayableStatus == EAvaPlayableStatus::Unloaded)
+	{
+		using namespace UE::AvaPlayable::Private;
+		UE_LOG(LogAvaPlayable, Verbose, TEXT("%s Playable {%s} -> Discarding RC Update."), *GetBriefFrameInfo(), *GetPrettyPlayableInfo(this));
+
+		// Discard the command
+		return EAvaPlayableCommandResult::ErrorDiscard;
+	}
+
+	// Asset status must be visible to run the command.
+	// If not visible, the components are not yet added to the world.
+	if (PlayableStatus != EAvaPlayableStatus::Visible)
+	{
+		using namespace UE::AvaPlayable::Private;
+		UE_LOG(LogAvaPlayable, Verbose, TEXT("%s Playable {%s} -> ReQueueing RC Update."), *GetBriefFrameInfo(), *GetPrettyPlayableInfo(this));
+
+		// Keep the command in the queue for next tick.
+		return EAvaPlayableCommandResult::KeepPending;
+	}
+
+	const IAvaSceneInterface* Scene = GetSceneInterface();
+	if (!Scene)
+	{
+		return EAvaPlayableCommandResult::ErrorDiscard;
+	}
+	
+	URemoteControlPreset* RemoteControlPreset = Scene->GetRemoteControlPreset();
+	
+	if (!IsValid(RemoteControlPreset))
+	{
+		UE_LOG(LogAvaPlayable, Error,
+			TEXT("Remote Control command for asset \"%s\": Remote Control Preset is null."),
+			*GetSourceAssetPath().ToString());
+		return EAvaPlayableCommandResult::ErrorDiscard;
+	}
+
+	LatestRemoteControlValues = InRemoteControlValues;
+
+	using namespace UE::AvaPlayable::Private;
+	UE_LOG(LogAvaPlayable, Verbose, TEXT("%s Playable {%s} -> Executing RC Update."), *GetBriefFrameInfo(), *GetPrettyPlayableInfo(this));
+
+	// WYSIWYG (Solution): For the runtime/playback RCP, we don't apply the controllers.
+	// We assume the controller actions are already executed in the rundown's managed RCP
+	// during page edition and the resulting entity values are already captured.
+	InRemoteControlValues->ApplyEntityValuesToRemoteControlPreset(RemoteControlPreset);
+
+	return EAvaPlayableCommandResult::Executed;
+}
+
+void UAvaPlayable::BeginPlay(const FAvaInstancePlaySettings& InWorldPlaySettings)
+{
+	if (!PlayableGroup)
+	{
+		return;
+	}
+
+	const bool bGroupHasBegunPlay = PlayableGroup->ConditionalBeginPlay(InWorldPlaySettings);
+	
+	if (!bIsPlaying || bGroupHasBegunPlay)
+	{
+		bIsPlaying = true;
+		
+		// Playable events need to transit through playback events to reach the rundown for proper impl layer separation.
+		UAvaSequencePlayer::OnSequenceStarted().AddUObject(this, &UAvaPlayable::HandleOnSequenceStarted);
+		UAvaSequencePlayer::OnSequenceFinished().AddUObject(this, &UAvaPlayable::HandleOnSequenceFinished);
+
+		OnPlay();
+	}
+}
+
+void UAvaPlayable::EndPlay(EAvaPlayableEndPlayOptions InOptions)
+{
+	if (!bIsPlaying)
+	{
+		return;
+	}
+
+	bIsPlaying = false;
+	UAvaSequencePlayer::OnSequenceStarted().RemoveAll(this);
+	UAvaSequencePlayer::OnSequenceFinished().RemoveAll(this);
+	OnEndPlay();
+
+	if (PlayableGroup)
+	{
+		PlayableGroup->UpdateCameraSetup();
+
+		if (EnumHasAnyFlags(InOptions, EAvaPlayableEndPlayOptions::ConditionalEndPlayWorld) && !PlayableGroup->HasPlayingPlayables())
+		{
+			const bool bForceImmediate = EnumHasAnyFlags(InOptions, EAvaPlayableEndPlayOptions::ForceImmediate);
+			PlayableGroup->RequestEndPlayWorld(bForceImmediate);
+		}
+	}
+}
+
+bool UAvaPlayable::HasSequence(const UAvaSequence* InSequence) const
+{
+	const IAvaSceneInterface* SceneInterface = GetSceneInterface();
+	if (!SceneInterface)
+	{
+		return false;
+	}
+	
+	const IAvaSequenceProvider* SequenceProvider = SceneInterface->GetSequenceProvider();
+	if (!SequenceProvider)
+	{
+		return false;
+	}
+	
+	for (const TObjectPtr<UAvaSequence>& Sequence : SequenceProvider->GetSequences())
+	{
+		if (Sequence == InSequence)
+		{
+			return true;
+		}
+	}
+	return false;
+}
+
+bool UAvaPlayable::InitPlayable(const FPlayableCreationInfo& InPlayableInfo)
+{
+	if (PlayableGroup)
+	{
+		// Register this playable in the instance group.
+		// This is necessary to determine what is playing in what group.
+		PlayableGroup->RegisterPlayable(this);
+		return true;
+	}
+	
+	// Currently, playables must have a playable group otherwise they are unplayable.
+	UE_LOG(LogAvaPlayable, Error, TEXT("Failed to create or acquire a playable group for \"%s\". Playable will be discarded."), *InPlayableInfo.SourceAsset.ToSoftObjectPath().ToString());
+	return false;
+}
+
+void UAvaPlayable::HandleOnSequenceStarted(UAvaSequencePlayer* InSequencePlayer, UAvaSequence* InSequence)
+{
+	if (HasSequence(InSequence))
+	{
+		using namespace UE::AvaPlayable::Private;
+		UE_LOG(LogAvaPlayable, Verbose, TEXT("%s Playable {%s}: Sequence \"%s\" started."),
+			*GetBriefFrameInfo(), *GetPrettyPlayableInfo(this), *InSequence->GetFName().ToString());
+		OnSequenceEventDelegate.Broadcast(this, InSequence->GetFName(), EAvaPlayableSequenceEventType::Started);
+	}
+}
+
+void UAvaPlayable::HandleOnSequenceFinished(UAvaSequencePlayer* InSequencePlayer, UAvaSequence* InSequence)
+{
+	if (HasSequence(InSequence))
+	{
+		using namespace UE::AvaPlayable::Private;
+		UE_LOG(LogAvaPlayable, Verbose, TEXT("%s Playable {%s}: Sequence \"%s\" finished."),
+			*GetBriefFrameInfo(), *GetPrettyPlayableInfo(this), *InSequence->GetFName().ToString());
+		OnSequenceEventDelegate.Broadcast(this, InSequence->GetFName(), EAvaPlayableSequenceEventType::Finished);
+	}
+}
+
+UAvaPlayable* UAvaPlayable::CreateLocalPlayable(UObject* InOuter, const FPlayableCreationInfo& InPlayableInfo)
+{
+	switch (InPlayableInfo.SourceAsset.GetAssetType())
+	{
+	case EAvalancheAssetType::Blueprint:
+		return NewObject<UAvaPlayableBlueprint>(InOuter ? InOuter : GEngine);
+	case EAvalancheAssetType::World:
+		return NewObject<UAvaPlayableLevelStreaming>(InOuter ? InOuter : GEngine);
+	default:
+		UE_LOG(LogAvaPlayable, Error, TEXT("Asset \"%s\" is an unsupported type."), *InPlayableInfo.SourceAsset.ToSoftObjectPath().ToString());
+		return nullptr;
+	}
+}
+
+UAvaPlayable* UAvaPlayable::CreateRemoteProxyPlayable(UObject* InOuter, const FPlayableCreationInfo& InPlayableInfo)
+{
+	return NewObject<UAvaPlayableRemoteProxy>(InOuter ? InOuter : GEngine);
+}
+
+#undef LOCTEXT_NAMESPACE
