Motion Design

diff --git a/Engine/Plugins/Experimental/Avalanche/Source/AvalancheMedia/Private/Playable/AvaPlayableGroup.cpp b/Engine/Plugins/Experimental/Avalanche/Source/AvalancheMedia/Private/Playable/AvaPlayableGroup.cpp
new file mode 100644
--- /dev/null
+++ b/Engine/Plugins/Experimental/Avalanche/Source/AvalancheMedia/Private/Playable/AvaPlayableGroup.cpp
@@ -0,0 +1,622 @@
+// Copyright Epic Games, Inc. All Rights Reserved.
+
+#include "Playable/AvaPlayableGroup.h"
+
+#include "Engine/ViewportStatsSubsystem.h"
+#include "Framework/AvaGameInstance.h"
+#include "Playable/AvaPlayable.h"
+#include "Playable/AvaPlayableGroupManager.h"
+#include "Playable/Transition/AvaPlayableTransition.h"
+#include "Playback/AvaPlaybackUtils.h"
+#include "UObject/Package.h"
+
+#define LOCTEXT_NAMESPACE "AvaPlayableGroup"
+
+namespace UE::AvaMedia::PlayableGroup::Private
+{
+	template<typename InElementType>
+	void CleanStaleKeys(TSet<TObjectKey<InElementType>>& InOutSetToClean)
+	{
+		for (typename TSet<TObjectKey<InElementType>>::TIterator It(InOutSetToClean); It; ++It)
+		{
+			if (!It->ResolveObjectPtr())
+			{
+				It.RemoveCurrent();
+			}
+		}
+	}
+}
+
+UPackage* UAvaPlayableGroup::MakeGameInstancePackage(const FSoftObjectPath& InSourceAssetPath, const FName& InChannelName)
+{
+	// Remote Control Preset will be registered with the package name.
+	// We want a package name unique to the game instance and that will be
+	// human readable since it will show up in the Web Remote Control page.
+	
+	FString InstancePackageName = TEXT("/Temp/");	// Keep it short for web page.
+
+	// Using channel name, since we should have one instance per channel.
+	InstancePackageName += InChannelName.ToString();
+
+	// In order to keep things short, we remove "/Game" since we added the channel name instead.
+	FString InstanceSubPath;
+	if (!InSourceAssetPath.GetLongPackageName().Split(TEXT("/Game"), nullptr, &InstanceSubPath))
+	{
+		InstanceSubPath = InSourceAssetPath.GetLongPackageName();
+	}
+
+	// This may happen if the original asset path is not specified.
+	if (InstanceSubPath.IsEmpty())
+	{
+		// Add something to get a valid path at least.
+		InstanceSubPath = TEXT("InvalidAssetName");
+	}
+	
+	InstancePackageName += InstanceSubPath;
+
+	return MakeInstancePackage(InstancePackageName);
+}
+
+UPackage* UAvaPlayableGroup::MakeSharedInstancePackage(const FName& InChannelName)
+{
+	// Remote Control Preset will be registered with the package name.
+	// We want a package name unique to the game instance and that will be
+	// human readable since it will show up in the Web Remote Control page.
+	
+	FString SharedPackageName = TEXT("/Temp/");	// Keep it short for web page.
+
+	// Using channel name, since we should have one instance per channel.
+	SharedPackageName += InChannelName.ToString();
+
+	// Shared for all levels.
+	SharedPackageName += TEXT("/SharedLevels");
+	
+	return MakeInstancePackage(SharedPackageName);
+}
+
+UPackage* UAvaPlayableGroup::MakeInstancePackage(const FString& InInstancePackageName)
+{
+	UPackage* InstancePackage = CreatePackage(*InInstancePackageName);
+	if (InstancePackage)
+	{
+		InstancePackage->SetFlags(RF_Transient);
+	}
+	else
+	{
+		// Note: The outer will fallback to GEngine in that case.
+		UE_LOG(LogAvaPlayable, Error, TEXT("Unable to create package \"%s\" for Motion Design Game Instance."), *InInstancePackageName);
+	}
+	return InstancePackage;
+}
+
+UAvaPlayableGroup* UAvaPlayableGroup::MakePlayableGroup(UObject* InOuter, const FPlayableGroupCreationInfo& InPlayableGroupInfo)
+{
+	UObject* Outer = InOuter ? InOuter : GetTransientPackage();
+	
+	UAvaPlayableGroup* GameInstanceGroup;
+	if (InPlayableGroupInfo.bIsRemoteProxy)
+	{
+		// Remote Proxy group doesn't have a game instance.
+		GameInstanceGroup = NewObject<UAvaPlayableRemoteProxyGroup>(Outer);
+		GameInstanceGroup->ParentPlayableGroupManagerWeak = InPlayableGroupInfo.PlayableGroupManager;
+	}
+	else
+	{
+		GameInstanceGroup = NewObject<UAvaPlayableGroup>(Outer);
+		GameInstanceGroup->ParentPlayableGroupManagerWeak = InPlayableGroupInfo.PlayableGroupManager;
+		
+		if (InPlayableGroupInfo.bIsSharedGroup)
+		{
+			GameInstanceGroup->GameInstancePackage = MakeSharedInstancePackage(InPlayableGroupInfo.ChannelName);
+		}
+		else
+		{
+			// We can create the package even if the name is null, it will have a generic name. But that will be considered an error.
+			if (!InPlayableGroupInfo.SourceAssetPath.IsNull())
+			{
+				UE_LOG(LogAvaPlayable, Error, TEXT("Creating game instance package for asset with unspecified name."));
+			}
+			GameInstanceGroup->GameInstancePackage = MakeGameInstancePackage(InPlayableGroupInfo.SourceAssetPath, InPlayableGroupInfo.ChannelName);
+		}
+		
+		GameInstanceGroup->GameInstance = UAvaGameInstance::Create(GameInstanceGroup->GameInstancePackage);
+	}
+	return GameInstanceGroup;
+}
+
+void UAvaPlayableGroup::RegisterPlayable(UAvaPlayable* InPlayable)
+{
+	// Prevent accumulation of stale keys.
+	UE::AvaMedia::PlayableGroup::Private::CleanStaleKeys(Playables);
+
+	if (InPlayable)
+	{
+		InPlayable->OnPlayableStatusChanged().AddUObject(this, &UAvaPlayableGroup::OnPlayableStatusChanged);
+		Playables.Add(InPlayable);
+	}
+}
+
+/** Unregister a playable when it is about to be deleted. */
+void UAvaPlayableGroup::UnregisterPlayable(UAvaPlayable* InPlayable)
+{
+	if (InPlayable)
+	{
+		InPlayable->OnPlayableStatusChanged().RemoveAll(this);
+		Playables.Remove(InPlayable);
+	}
+}
+
+bool UAvaPlayableGroup::HasPlayables() const
+{
+	for (const TObjectKey<UAvaPlayable>& PlayableKey : Playables)
+	{
+		if (PlayableKey.ResolveObjectPtr())
+		{
+			return true;
+		}
+	}
+	return false;
+}
+
+bool UAvaPlayableGroup::HasPlayingPlayables() const
+{
+	for (const TObjectKey<UAvaPlayable>& PlayableKey : Playables)
+	{
+		const UAvaPlayable* Playable = PlayableKey.ResolveObjectPtr();
+		if (Playable && Playable->IsPlaying())
+		{
+			return true;
+		}
+	}
+	return false;
+}
+
+void UAvaPlayableGroup::FindPlayablesBySourceAssetPath(const FSoftObjectPath& InSourceAssetPath, TArray<UAvaPlayable*>& OutFoundPlayables) const
+{
+	for (const TObjectKey<UAvaPlayable>& PlayableKey : Playables)
+	{
+		UAvaPlayable* Playable = PlayableKey.ResolveObjectPtr();
+		if (Playable && Playable->GetSourceAssetPath() == InSourceAssetPath)
+		{
+			OutFoundPlayables.Add(Playable);
+		}
+	}
+}
+
+void UAvaPlayableGroup::RegisterPlayableTransition(UAvaPlayableTransition* InPlayableTransition)
+{
+	if (InPlayableTransition)
+	{
+		// Protect PlayableTransitions iterator.
+		if (bIsTickingTransitions)
+		{
+			PlayableTransitionsToRemove.Remove(InPlayableTransition);
+			PlayableTransitionsToAdd.Add(InPlayableTransition);
+			return;
+		}
+
+		// Prevent accumulation of stale keys.
+		UE::AvaMedia::PlayableGroup::Private::CleanStaleKeys(PlayableTransitions);
+
+		PlayableTransitions.Add(InPlayableTransition);
+		
+		if (UAvaPlayableGroupManager* PlayableGroupManager = GetPlayableGroupManager())
+		{
+			PlayableGroupManager->RegisterForTransitionTicking(this);
+		}
+	}
+}
+	
+void UAvaPlayableGroup::UnregisterPlayableTransition(UAvaPlayableTransition* InPlayableTransition)
+{
+	if (InPlayableTransition)
+	{
+		// Protect PlayableTransitions iterator.
+		if (bIsTickingTransitions)
+		{
+			PlayableTransitionsToAdd.Remove(InPlayableTransition);
+			PlayableTransitionsToRemove.Add(InPlayableTransition);
+			return;
+		}
+		
+		PlayableTransitions.Remove(InPlayableTransition);
+
+		UE::AvaMedia::PlayableGroup::Private::CleanStaleKeys(PlayableTransitions);
+		
+		if (PlayableTransitions.IsEmpty())
+		{
+			if (UAvaPlayableGroupManager* PlayableGroupManager = GetPlayableGroupManager())
+			{
+				PlayableGroupManager->UnregisterFromTransitionTicking(this);
+			}
+		}
+	}
+}
+
+void UAvaPlayableGroup::TickTransitions(double InDeltaSeconds)
+{
+	// Ticking the transitions will lead to some transitions being removed.
+	// So we need to protect the iterator with a scope and do the operations when iteration is done.
+	{
+		TGuardValue TickGuard(bIsTickingTransitions, true);
+	
+		for (TSet<TObjectKey<UAvaPlayableTransition>>::TIterator TickIt(PlayableTransitions); TickIt; ++TickIt)
+		{
+			if (UAvaPlayableTransition* TransitionToTick = TickIt->ResolveObjectPtr())
+			{
+				TransitionToTick->Tick(InDeltaSeconds);
+			}
+			else
+			{
+				// Prevent accumulation of stale keys.
+				// Remark: if there is a stale key, it means the transition was not stopped properly.
+				TickIt.RemoveCurrent();
+			}
+		}
+	}
+	
+	for (const TObjectKey<UAvaPlayableTransition>& ToRemove : PlayableTransitionsToRemove)
+	{
+		UnregisterPlayableTransition(ToRemove.ResolveObjectPtr());
+	}
+	PlayableTransitionsToRemove.Reset();
+	
+	for (const TObjectKey<UAvaPlayableTransition>& ToAdd : PlayableTransitionsToAdd)
+	{
+		RegisterPlayableTransition(ToAdd.ResolveObjectPtr());
+	}
+	PlayableTransitionsToAdd.Reset();
+}
+
+bool UAvaPlayableGroup::HasTransitions() const
+{
+	return !PlayableTransitions.IsEmpty();
+}
+
+bool UAvaPlayableGroup::ConditionalCreateWorld()
+{
+	if (!GameInstance)
+	{
+		return false;
+	}
+
+	bool bWorldWasCreated = false;
+
+	if (!GameInstance->IsWorldCreated())
+	{
+		bWorldWasCreated = GameInstance->CreateWorld();
+	}
+	
+	// Make sure we register our delegates to this world.
+	if (GameInstance->GetPlayWorld())
+	{
+		ConditionalRegisterWorldDelegates(GameInstance->GetPlayWorld());
+	}
+	
+	return bWorldWasCreated;
+}
+
+bool UAvaPlayableGroup::ConditionalBeginPlay(const FAvaInstancePlaySettings& InWorldPlaySettings)
+{
+	bool bHasBegunPlay = false;
+	if (!GameInstance)
+	{
+		return bHasBegunPlay;
+	}
+	
+	// Make sure we don't have pending unload or stop requests left over in the game instance.
+	GameInstance->CancelWorldRequests();
+
+	if (!GameInstance->IsWorldPlaying())
+	{
+		bHasBegunPlay = GameInstance->BeginPlayWorld(InWorldPlaySettings);
+	}
+	else
+	{
+		GameInstance->UpdateRenderTarget(InWorldPlaySettings.RenderTarget);
+		GameInstance->UpdateSceneViewportSize(InWorldPlaySettings.ViewportSize);
+	}
+	return bHasBegunPlay;
+}
+
+void UAvaPlayableGroup::RequestEndPlayWorld(bool bInForceImmediate)
+{
+	if (GameInstance)
+	{
+		GameInstance->RequestEndPlayWorld(bInForceImmediate);
+	}
+}
+
+void UAvaPlayableGroup::SetLastAppliedCameraPlayable(UAvaPlayable* InPlayable)
+{
+	LastAppliedCameraPlayableWeak = InPlayable;
+}
+
+bool UAvaPlayableGroup::UpdateCameraSetup()
+{
+	// With rigs in sub-playables, it may still be valid and won't need updating.
+	if (const UAvaPlayable* LastAppliedCameraPlayable = LastAppliedCameraPlayableWeak.Get())
+	{
+		if (LastAppliedCameraPlayable->GetPlayableStatus() == EAvaPlayableStatus::Visible
+			&& LastAppliedCameraPlayable->GetShouldBeVisible())
+		{
+			return true;	// Camera setup is still valid.
+		}
+	}
+
+	// This is called when a playable from the group is stopped. We need to
+	// select a playable in the group that is still playing and will use it's camera.
+	for (const TObjectKey<UAvaPlayable>& PlayableKey : Playables )
+	{
+		UAvaPlayable* Playable = PlayableKey.ResolveObjectPtr();
+		
+		if (Playable && Playable->IsPlaying())
+		{
+			if (Playable->ApplyCamera())
+			{
+				return true;
+			}
+		}
+	}
+	return false;
+}
+
+bool UAvaPlayableGroup::ConditionalRequestUnloadWorld(bool bForceImmediate)
+{
+	if (!GameInstance)
+	{
+		return false;
+	}
+	
+	if (!HasPlayables())
+	{
+		UnregisterWorldDelegates(GameInstance->GetPlayWorld());
+		GameInstance->RequestUnloadWorld(bForceImmediate);
+		return true;
+	}
+	return false;
+}
+
+void UAvaPlayableGroup::QueueCameraCut()
+{
+	if (GameInstance)
+	{
+		if (UAvaGameViewportClient* GameViewportClient = GameInstance->GetAvaGameViewportClient())
+		{
+			GameViewportClient->SetCameraCutThisFrame();
+		}
+	}
+}
+
+void UAvaPlayableGroup::NotifyLevelStreaming(UAvaPlayable* InPlayable)
+{
+	// If the world is not playing, we need to make sure the level streaming still updates.
+	if (!IsWorldPlaying())
+	{
+		if (UAvaPlayableGroupManager* PlayableGroupManager = GetPlayableGroupManager())
+		{
+			PlayableGroupManager->RegisterForLevelStreamingUpdate(this);
+		}
+	}
+}
+
+void UAvaPlayableGroup::FVisibilityRequest::Execute(const UAvaPlayableGroup* InPlayableGroup) const
+{
+	UAvaPlayable* Playable = PlayableWeak.Get();
+	
+	if (!Playable)
+	{
+		using namespace UE::AvaPlayback::Utils;
+		UE_LOG(LogAvaPlayable, Error,
+			TEXT("%s Failed to Set Visibility to \"%s\" because the playable has become stale. Playable Group: \"%s\"."),
+			*GetBriefFrameInfo(), bShouldBeVisible ? TEXT("true") : TEXT("false"), *InPlayableGroup->GetFullName());
+	}
+
+	Playable->SetShouldBeVisible(bShouldBeVisible);
+}
+
+void UAvaPlayableGroup::RegisterVisibilityConstraint(const TWeakInterfacePtr<IAvaPlayableVisibilityConstraint>& InVisibilityConstraint)
+{
+	if (!VisibilityConstraints.Contains(InVisibilityConstraint))
+	{
+		VisibilityConstraints.RemoveAll([](const TWeakInterfacePtr<IAvaPlayableVisibilityConstraint>& InConstraint) { return InConstraint.IsStale();});
+		VisibilityConstraints.Add(InVisibilityConstraint);
+	}
+}
+
+void UAvaPlayableGroup::UnregisterVisibilityConstraint(const IAvaPlayableVisibilityConstraint* InVisibilityConstraint)
+{
+	VisibilityConstraints.RemoveAll([InVisibilityConstraint](const TWeakInterfacePtr<IAvaPlayableVisibilityConstraint>& InConstraint)
+	{
+		// also remove stale pointers.
+		return InConstraint.IsStale() || InConstraint.Get() == InVisibilityConstraint;
+	});
+}
+
+void UAvaPlayableGroup::RequestSetVisibility(UAvaPlayable* InPlayable, bool bInShouldBeVisible)
+{
+	FVisibilityRequest Request = { InPlayable, bInShouldBeVisible };
+
+	if (IsVisibilityConstrained(InPlayable))
+	{
+		VisibilityRequests.Add(MoveTemp(Request));
+	}
+	else
+	{
+		Request.Execute(this);
+	}
+}
+
+bool UAvaPlayableGroup::IsVisibilityConstrained(const UAvaPlayable* InPlayable) const
+{
+	for (const TWeakInterfacePtr<IAvaPlayableVisibilityConstraint>& ConstraintWeak : VisibilityConstraints)
+	{
+		if (const IAvaPlayableVisibilityConstraint* Constraint = ConstraintWeak.Get())
+		{
+			if (Constraint->IsVisibilityConstrained(InPlayable))
+			{
+				return true;
+			}
+		}
+	}
+	return false;
+}
+
+void UAvaPlayableGroup::OnPlayableStatusChanged(UAvaPlayable* InPlayable)
+{
+	// Evaluate the playable visibility requests
+	for (TArray<FVisibilityRequest>::TIterator RequestIt(VisibilityRequests); RequestIt; ++RequestIt)
+	{
+		const UAvaPlayable* Playable = RequestIt->PlayableWeak.Get();
+		if (Playable && IsVisibilityConstrained(Playable))
+		{
+			continue;
+		}
+		
+		RequestIt->Execute(this);
+		RequestIt.RemoveCurrent();
+	}
+}
+
+void UAvaPlayableGroup::ConditionalRegisterWorldDelegates(UWorld* InWorld)
+{
+	if (!DisplayDelegateIndices.IsEmpty() && LastWorldBoundToDisplayDelegates.Get() == InWorld)
+	{
+		return;
+	}
+
+	if (!DisplayDelegateIndices.IsEmpty() && LastWorldBoundToDisplayDelegates.IsValid())
+	{
+		UnregisterWorldDelegates(LastWorldBoundToDisplayDelegates.Get());
+	}
+
+	if (UViewportStatsSubsystem* ViewportSubsystem = InWorld->GetSubsystem<UViewportStatsSubsystem>())
+	{
+		DisplayDelegateIndices.Add(ViewportSubsystem->AddDisplayDelegate([this](FText& OutText, FLinearColor& OutColor)
+		{
+			return DisplayLoadedAssets(OutText, OutColor);
+		}));
+		DisplayDelegateIndices.Add(ViewportSubsystem->AddDisplayDelegate([this](FText& OutText, FLinearColor& OutColor)
+		{
+			return DisplayPlayingAssets(OutText, OutColor);
+		}));
+		DisplayDelegateIndices.Add(ViewportSubsystem->AddDisplayDelegate([this](FText& OutText, FLinearColor& OutColor)
+		{
+			return DisplayTransitions(OutText, OutColor);
+		}));
+
+		LastWorldBoundToDisplayDelegates = InWorld;
+	}
+}
+
+void UAvaPlayableGroup::UnregisterWorldDelegates(UWorld* InWorld)
+{
+	if (!DisplayDelegateIndices.IsEmpty())
+	{
+		// Check that we are indeed registered to the world
+		check(LastWorldBoundToDisplayDelegates.Get() == InWorld);
+	
+		if (UViewportStatsSubsystem* ViewportSubsystem = InWorld->GetSubsystem<UViewportStatsSubsystem>())
+		{
+			// Remark: removing them in the reverse order they where added. It is using RemoveAtSwap()
+			// so removing the first one would change the index of the next one. Removing a delegate in the
+			// middle of the array will invalidate the indices above.
+			for (int32 Index = DisplayDelegateIndices.Num()-1; Index >= 0; --Index)
+			{
+				ViewportSubsystem->RemoveDisplayDelegate(DisplayDelegateIndices[Index]);
+			}
+		}
+
+		DisplayDelegateIndices.Reset();
+		LastWorldBoundToDisplayDelegates.Reset();
+	}
+
+	check(LastWorldBoundToDisplayDelegates.Get() == nullptr);
+}
+
+bool UAvaPlayableGroup::DisplayLoadedAssets(FText& OutText, FLinearColor& OutColor)
+{
+	FString AssetList;
+	for (const TObjectKey<UAvaPlayable>& PlayableKey : Playables)
+	{
+		if (const UAvaPlayable* Playable = PlayableKey.ResolveObjectPtr())
+		{
+			AssetList += AssetList.IsEmpty() ? TEXT("") : TEXT(", ");
+			AssetList += Playable->GetSourceAssetPath().GetAssetName();
+		}
+	}
+
+	if (!AssetList.IsEmpty())
+	{
+		OutText = FText::Format(LOCTEXT("DisplayLoadedGraphics", "Loaded Graphic(s): {0}"), FText::FromString(AssetList));
+		OutColor = FLinearColor::Red;
+		return true;
+	}
+	return false;
+}
+
+bool UAvaPlayableGroup::DisplayPlayingAssets(FText& OutText, FLinearColor& OutColor)
+{
+	FString AssetList;
+	for (const TObjectKey<UAvaPlayable>& PlayableKey : Playables)
+	{
+		const UAvaPlayable* Playable = PlayableKey.ResolveObjectPtr();
+		if (Playable && Playable->IsPlaying())
+		{
+			AssetList += AssetList.IsEmpty() ? TEXT("") : TEXT(", ");
+			AssetList += Playable->GetSourceAssetPath().GetAssetName();
+
+			if (!Playable->GetUserData().IsEmpty())
+			{
+				AssetList += FString::Printf(TEXT(" (%s)"), *Playable->GetUserData());
+			}
+		}
+	}
+
+	if (!AssetList.IsEmpty())
+	{
+		OutText = FText::Format(LOCTEXT("DisplayPlayingGraphics", "Playing Graphic(s): {0}"), FText::FromString(AssetList));
+		OutColor = FLinearColor::Green;
+		return true;
+	}
+	return false;
+}
+
+bool UAvaPlayableGroup::DisplayTransitions(FText& OutText, FLinearColor& OutColor)
+{
+	FString TransitionList;
+	for (const TObjectKey<UAvaPlayableTransition>& TransitionKey : PlayableTransitions)
+	{
+		const UAvaPlayableTransition* Transition = TransitionKey.ResolveObjectPtr();
+		if (Transition && Transition->IsRunning())
+		{
+			TransitionList += TransitionList.IsEmpty() ? TEXT("") : TEXT(", ");
+			TransitionList += Transition->GetPrettyInfo();
+		}
+	}
+
+	if (!TransitionList.IsEmpty())
+	{
+		OutText = FText::Format(LOCTEXT("DisplayTransitions", "Transition(s): {0}"), FText::FromString(TransitionList));
+		OutColor = FLinearColor::Green;
+		return true;
+	}
+	return false;
+}
+
+bool UAvaPlayableRemoteProxyGroup::ConditionalBeginPlay(const FAvaInstancePlaySettings& InWorldPlaySettings)
+{
+	if (!bIsPlaying)
+	{
+		bIsPlaying = true;
+		return true;
+	}
+	return false;
+}
+
+void UAvaPlayableRemoteProxyGroup::RequestEndPlayWorld(bool bInForceImmediate)
+{
+	bIsPlaying = false;
+}
+
+#undef LOCTEXT_NAMESPACE
\ No newline at end of file
diff --git a/Engine/Plugins/Experimental/Avalanche/Source/AvalancheMedia/Private/Playable/Playables/AvaPlayableBlueprint.cpp b/Engine/Plugins/Experimental/Avalanche/Source/AvalancheMedia/Private/Playable/Playables/AvaPlayableBlueprint.cpp
new file mode 100644
--- /dev/null
+++ b/Engine/Plugins/Experimental/Avalanche/Source/AvalancheMedia/Private/Playable/Playables/AvaPlayableBlueprint.cpp
@@ -0,0 +1,183 @@
+// Copyright Epic Games, Inc. All Rights Reserved.
+
+#include "Playable/Playables/AvaPlayableBlueprint.h"
+
+#include "AvaBlueprint.h"
+#include "AvaRemoteControlRebind.h"
+#include "AvaSequencePlaybackObject.h"
+#include "Framework/AvaSoftAssetPtr.h"
+#include "Playable/AvaPlayableGroup.h"
+#include "Playable/AvaPlayableGroupManager.h"
+#include "Viewport/AvaCameraManager.h"
+
+#define LOCTEXT_NAMESPACE "AvaPlayableBlueprint"
+
+bool UAvaPlayableBlueprint::LoadAsset(const FAvaSoftAssetPtr& InSourceAsset, bool bInInitiallyVisible)
+{
+	if (!PlayableGroup)
+	{
+		return false;
+	}
+
+	// Ensure world is created. Does nothing if already created.
+	PlayableGroup->ConditionalCreateWorld(); 
+
+	check(InSourceAsset.GetAssetType() == EAvalancheAssetType::Blueprint);
+	return LoadBlueprintInternal(TSoftObjectPtr<UAvalancheBlueprint>(InSourceAsset.ToSoftObjectPath()));
+}
+
+bool UAvaPlayableBlueprint::UnloadAsset()
+{
+	if (IsValid(ManagedBlueprint.Get()))
+	{
+		ManagedBlueprint->UnregisterRemoteControlPreset();
+	}
+	SourceBlueprint.Reset();
+	ManagedBlueprint = nullptr;
+	return true;
+}
+
+const FSoftObjectPath& UAvaPlayableBlueprint::GetSourceAssetPath() const
+{
+	return SourceBlueprint.ToSoftObjectPath();
+}
+
+EAvaPlayableStatus UAvaPlayableBlueprint::GetPlayableStatus() const
+{
+	// The blueprint assets don't stream. We can't even make them visible/invisible.
+	return ManagedBlueprint ? EAvaPlayableStatus::Visible : EAvaPlayableStatus::Unloaded;
+}
+
+IAvaSceneInterface* UAvaPlayableBlueprint::GetSceneInterface() const
+{
+	if (ManagedBlueprint)
+	{
+		return static_cast<IAvaSceneInterface*>(ManagedBlueprint);
+	}
+	return nullptr;
+}
+
+EAvaPlayableCommandResult UAvaPlayableBlueprint::ExecuteAnimationCommand(EAvaPlaybackAnimAction InAnimAction, const FAvaPlaybackAnimPlaySettings& InAnimPlaySettings)
+{
+	if (InAnimAction ==  EAvaPlaybackAnimAction::CameraCut)
+	{
+		if (GetPlayableGroup())
+		{
+			GetPlayableGroup()->QueueCameraCut();
+		}
+		return EAvaPlayableCommandResult::Executed;
+	}
+	
+	return Super::ExecuteAnimationCommand(InAnimAction, InAnimPlaySettings);
+}
+
+bool UAvaPlayableBlueprint::InitPlayable(const FPlayableCreationInfo& InPlayableInfo)
+{
+	// Motion Design blueprint requires a separate playable group for each asset.
+	
+	UAvaPlayableGroup::FPlayableGroupCreationInfo PlayableGroupCreationInfo;
+	PlayableGroupCreationInfo.PlayableGroupManager = InPlayableInfo.PlayableGroupManager;
+	PlayableGroupCreationInfo.SourceAssetPath = InPlayableInfo.SourceAsset.ToSoftObjectPath();
+	PlayableGroupCreationInfo.ChannelName = InPlayableInfo.ChannelName;
+	PlayableGroupCreationInfo.bIsRemoteProxy = false;
+	PlayableGroupCreationInfo.bIsSharedGroup = false;
+
+	PlayableGroup = InPlayableInfo.PlayableGroup ?
+		InPlayableInfo.PlayableGroup : UAvaPlayableGroup::MakePlayableGroup(InPlayableInfo.PlayableGroupManager, PlayableGroupCreationInfo);
+	
+	return Super::InitPlayable(InPlayableInfo);
+}
+
+void UAvaPlayableBlueprint::OnPlay()
+{
+	if (!GetPlayableGroup())
+	{
+		return;
+	}
+
+	const UAvaGameInstance* AvaGameInstance = Cast<UAvaGameInstance>(GetPlayableGroup()->GetGameInstance());
+	if (!AvaGameInstance)
+	{
+		return;
+	}
+	
+	UAvaGameViewportClient* ViewportClient = AvaGameInstance->GetAvaGameViewportClient();
+
+	// Verify if this can be changed on the fly.
+	ManagedBlueprint->GetViewportQualitySettings().Apply(ViewportClient->EngineShowFlags);
+	
+	// Old code using ava camera manager. To retire.
+	constexpr bool bIsCanvasController = false;
+	ViewportClient->GetCameraManager()->Init(ManagedBlueprint->GetPlaybackObject(), bIsCanvasController);
+#if WITH_EDITOR
+	ViewportClient->GetCameraManager()->SetDefaultViewTarget(GetPlayWorld(), ManagedBlueprint->GetStartupCameraName());
+#endif
+}
+
+void UAvaPlayableBlueprint::OnEndPlay()
+{
+	if (ManagedBlueprint)
+	{
+		if (IAvaSequencePlaybackObject* PlaybackObject = ManagedBlueprint->GetPlaybackObject())
+		{
+			PlaybackObject->CleanupPlayers();
+		}
+	}
+}
+
+bool UAvaPlayableBlueprint::LoadBlueprintInternal(const TSoftObjectPtr<UAvalancheBlueprint>& InSourceBlueprint)
+{
+	if (!GetPlayableGroup())
+	{
+		return false;
+	}
+
+	UAvaGameInstance* AvaGameInstance = Cast<UAvaGameInstance>(GetPlayableGroup()->GetGameInstance());
+	if (!AvaGameInstance)
+	{
+		return false;
+	}
+	
+	UAvalancheBlueprint* Source;
+	{
+		TRACE_CPUPROFILER_EVENT_SCOPE(UAvaPlayableBlueprint::LoadAsset::LoadSourceAvaBp);
+		Source = InSourceBlueprint.LoadSynchronous();
+		check(Source);
+	}
+
+#if WITH_EDITORONLY_DATA
+	FGuardValue_Bitfield(Source->bDuplicatingReadOnly, true);
+#endif
+
+	{
+		TRACE_CPUPROFILER_EVENT_SCOPE(UAvaPlayableBlueprint::LoadAsset::DupAvaBp);
+		const FName AvaInstanceName = MakeUniqueObjectName(AvaGameInstance, UAvalancheBlueprint::StaticClass(), Source->GetFName());
+		ManagedBlueprint = Cast<UAvalancheBlueprint>(StaticDuplicateObject(Source, AvaGameInstance, AvaInstanceName, RF_Transient));
+	}
+
+	if (!ManagedBlueprint)
+	{
+		return false;
+	}
+	
+	{
+		TRACE_CPUPROFILER_EVENT_SCOPE(UAvaPlayableBlueprint::LoadAsset::LoadAvaBpWorld);
+		ManagedBlueprint->SetAvalancheWorld(AvaGameInstance->GetPlayWorld());
+		ManagedBlueprint->LoadAvalancheWorld();
+	}
+
+	ManagedBlueprint->UpdatePlaceholderActor();
+	
+	{
+		TRACE_CPUPROFILER_EVENT_SCOPE(UAvaPlayableBlueprint::LoadAsset::RebindRCP);
+		ManagedBlueprint->RegisterRemoteControlPreset();
+		FAvaRemoteControlRebind::RebindUnboundEntities(ManagedBlueprint->GetRemoteControlPreset());
+	}
+
+	AvaGameInstance->MarkSynchronousAssetLoadingThisFrame();
+	
+	SourceBlueprint = InSourceBlueprint;
+	return true;
+}
+
+#undef LOCTEXT_NAMESPACE
diff --git a/Engine/Plugins/Experimental/Avalanche/Source/AvalancheMedia/Private/Playable/Playables/AvaPlayableBlueprint.h b/Engine/Plugins/Experimental/Avalanche/Source/AvalancheMedia/Private/Playable/Playables/AvaPlayableBlueprint.h
new file mode 100644
--- /dev/null
+++ b/Engine/Plugins/Experimental/Avalanche/Source/AvalancheMedia/Private/Playable/Playables/AvaPlayableBlueprint.h
@@ -0,0 +1,46 @@
+// Copyright Epic Games, Inc. All Rights Reserved.
+
+#pragma once
+
+#include "AvaMediaDefines.h"
+#include "Playable/AvaPlayable.h"
+#include "AvaPlayableBlueprint.generated.h"
+
+UCLASS(NotBlueprintable, BlueprintType, ClassGroup = "Motion Design Playable", 
+	meta = (DisplayName = "Motion Design Blueprint Playable"))
+class UAvaPlayableBlueprint : public UAvaPlayable
+{
+	GENERATED_BODY()
+
+public:
+	//~ Begin UAvaPlayable
+	virtual bool LoadAsset(const FAvaSoftAssetPtr& InSourceAsset, bool bInInitiallyVisible) override;
+	virtual bool UnloadAsset() override;
+	virtual const FSoftObjectPath& GetSourceAssetPath() const override;
+	virtual EAvaPlayableStatus GetPlayableStatus() const override;
+	virtual IAvaSceneInterface* GetSceneInterface() const override;
+	virtual EAvaPlayableCommandResult ExecuteAnimationCommand(EAvaPlaybackAnimAction InAnimAction, const FAvaPlaybackAnimPlaySettings& InAnimPlaySettings) override;
+
+protected:
+	virtual bool InitPlayable(const FPlayableCreationInfo& InPlayableInfo) override;
+	virtual void OnPlay() override;
+	virtual void OnEndPlay() override;
+	//~ End UAvaPlayable
+
+public:
+	UAvalancheBlueprint* GetManagedBlueprint() const { return ManagedBlueprint; }
+	
+protected:
+	bool LoadBlueprintInternal(const TSoftObjectPtr<UAvalancheBlueprint>& InSourceBlueprint);
+
+protected:
+	UPROPERTY(Transient)
+	TSoftObjectPtr<UAvalancheBlueprint> SourceBlueprint;
+
+	/*
+	 * The Managed Blueprint, a Duplication of the Source Blueprint when Begin Play was called.
+	 * This shouldn't be GCd while in play, as it contains Animation and other data used during play.
+	 */
+	UPROPERTY(Transient)
+	TObjectPtr<UAvalancheBlueprint> ManagedBlueprint;
+};
\ No newline at end of file
diff --git a/Engine/Plugins/Experimental/Avalanche/Source/AvalancheMedia/Private/Playable/Playables/AvaPlayableRemoteProxy.cpp b/Engine/Plugins/Experimental/Avalanche/Source/AvalancheMedia/Private/Playable/Playables/AvaPlayableRemoteProxy.cpp
new file mode 100644
--- /dev/null
+++ b/Engine/Plugins/Experimental/Avalanche/Source/AvalancheMedia/Private/Playable/Playables/AvaPlayableRemoteProxy.cpp
@@ -0,0 +1,299 @@
+// Copyright Epic Games, Inc. All Rights Reserved.
+
+#include "Playable/Playables/AvaPlayableRemoteProxy.h"
+
+#include "Broadcast/AvaBroadcast.h"
+#include "Framework/AvaSoftAssetPtr.h"
+#include "IAvaMediaModule.h"
+#include "Playable/AvaPlayableGroupManager.h"
+#include "Playback/AvaPlaybackClientDelegates.h"
+#include "Playback/IAvaPlaybackClient.h"
+
+#define LOCTEXT_NAMESPACE "AvaPlayableRemoteProxy"
+
+bool UAvaPlayableRemoteProxy::LoadAsset(const FAvaSoftAssetPtr& InSourceAsset, bool bInInitiallyVisible)
+{
+	if (!PlayableGroup)
+	{
+		return false;
+	}
+	
+	SourceAssetPath = InSourceAsset.ToSoftObjectPath();
+
+	IAvaPlaybackClient& PlaybackClient = IAvaMediaModule::Get().GetMediaPlaybackClient();
+	
+	if (PlaybackClient.HasAnyServerOnlineForChannel(PlayingChannelFName))
+	{
+		const FString ChannelName = PlayingChannelName;
+		const FSoftObjectPath& AssetPath = InSourceAsset.ToSoftObjectPath();
+		const TOptional<EAvaPlaybackStatus> RemoteStatusOpt = PlaybackClient.GetRemotePlaybackStatus(InstanceId, AssetPath, ChannelName);
+		const EAvaPlaybackStatus RemoteStatus = RemoteStatusOpt.IsSet() ? RemoteStatusOpt.GetValue() : EAvaPlaybackStatus::Unknown;
+
+		const bool bCanLoad = RemoteStatus == EAvaPlaybackStatus::Available
+			|| RemoteStatus == EAvaPlaybackStatus::Unknown;
+		const bool bIsLoaded = RemoteStatus == EAvaPlaybackStatus::Loading
+			|| RemoteStatus == EAvaPlaybackStatus::Loaded
+			|| RemoteStatus == EAvaPlaybackStatus::Starting
+			|| RemoteStatus == EAvaPlaybackStatus::Started;
+
+		if (bCanLoad && !bIsLoaded)
+		{
+			PlaybackClient.RequestPlayback(InstanceId, AssetPath, ChannelName, EAvaPlaybackAction::Load);
+			PlaybackClient.RequestPlayback(InstanceId, AssetPath, ChannelName, EAvaPlaybackAction::SetUserData, UserData);
+		}
+	}
+	return true;
+}
+
+bool UAvaPlayableRemoteProxy::UnloadAsset()
+{
+	IAvaPlaybackClient& PlaybackClient = IAvaMediaModule::Get().GetMediaPlaybackClient();
+	if (PlaybackClient.HasAnyServerOnlineForChannel(PlayingChannelFName))
+	{
+		PlaybackClient.RequestPlayback(InstanceId, SourceAssetPath, PlayingChannelName, EAvaPlaybackAction::Unload);
+	}
+	return true;
+}
+
+namespace UE::AvaMedIaRemoteProxyPlayable::Private
+{
+	EAvaPlayableStatus GetPlayableStatus(EAvaPlaybackStatus InPlaybackStatus)
+	{
+		switch (InPlaybackStatus)
+		{
+		case EAvaPlaybackStatus::Unknown:
+			return EAvaPlayableStatus::Unknown;
+		case EAvaPlaybackStatus::Missing:
+		case EAvaPlaybackStatus::Syncing:
+		case EAvaPlaybackStatus::Available:
+			return EAvaPlayableStatus::Unloaded;
+		case EAvaPlaybackStatus::Loading:
+			return EAvaPlayableStatus::Loading;
+		case EAvaPlaybackStatus::Loaded:
+			return EAvaPlayableStatus::Loaded;
+		case EAvaPlaybackStatus::Starting:
+			return EAvaPlayableStatus::Loaded;
+		case EAvaPlaybackStatus::Started:
+			return EAvaPlayableStatus::Visible;
+		case EAvaPlaybackStatus::Stopping:
+			return EAvaPlayableStatus::Loaded;
+		case EAvaPlaybackStatus::Unloading:
+			return EAvaPlayableStatus::Unloaded;
+		case EAvaPlaybackStatus::Error:
+		default:
+			return EAvaPlayableStatus::Error;
+		}
+	}
+}
+
+EAvaPlayableStatus UAvaPlayableRemoteProxy::GetPlayableStatus() const
+{
+	IAvaPlaybackClient& PlaybackClient = IAvaMediaModule::Get().GetMediaPlaybackClient();
+	const TArray<FString> OnlineServers = GetOnlineServerForChannel(PlayingChannelFName);
+	
+	for (const FString& Server : OnlineServers)
+	{
+		TOptional<EAvaPlaybackStatus> PlaybackStatus = PlaybackClient.GetRemotePlaybackStatus(InstanceId, SourceAssetPath, PlayingChannelName, Server);
+
+		if (!PlaybackStatus.IsSet())
+		{
+			PlaybackClient.RequestPlayback(InstanceId, SourceAssetPath, PlayingChannelName, EAvaPlaybackAction::Status);
+			PlaybackStatus = EAvaPlaybackStatus::Unknown;
+		}
+
+		// TODO: reconcile forked channels.
+		return UE::AvaMedIaRemoteProxyPlayable::Private::GetPlayableStatus(PlaybackStatus.GetValue());
+	}
+	
+	return EAvaPlayableStatus::Unknown;
+}
+
+IAvaSceneInterface* UAvaPlayableRemoteProxy::GetSceneInterface() const
+{
+	return nullptr;
+}
+
+EAvaPlayableCommandResult UAvaPlayableRemoteProxy::ExecuteAnimationCommand(EAvaPlaybackAnimAction InAnimAction, const FAvaPlaybackAnimPlaySettings& InAnimPlaySettings)
+{
+	IAvaPlaybackClient& PlaybackClient = IAvaMediaModule::Get().GetMediaPlaybackClient();
+	
+	// If an animation event was locally scheduled on a remote playable,
+	// we need to propagate the event.
+	if (PlaybackClient.HasAnyServerOnlineForChannel(PlayingChannelFName))
+	{
+		switch (InAnimAction)
+		{
+		case EAvaPlaybackAnimAction::None:
+			break;
+		case EAvaPlaybackAnimAction::Play:
+		case EAvaPlaybackAnimAction::PreviewFrame:
+			PlaybackClient.RequestAnimPlayback(InstanceId, SourceAssetPath, PlayingChannelName, InAnimPlaySettings);
+			break;
+
+		case EAvaPlaybackAnimAction::Continue:
+		case EAvaPlaybackAnimAction::Stop:
+		case EAvaPlaybackAnimAction::CameraCut:
+			PlaybackClient.RequestAnimAction(InstanceId, SourceAssetPath, PlayingChannelName, InAnimPlaySettings.AnimationName.ToString(), InAnimAction);
+			break;
+
+		default:
+			UE_LOG(LogAvaPlayable, Warning,
+				TEXT("Animation command action \"%s\" for asset \"%s\" on channel \"%s\" is not implemented."),
+				*StaticEnum<EAvaPlaybackAnimAction>()->GetValueAsString(InAnimAction),
+				*SourceAssetPath.ToString(), *PlayingChannelName);
+			break;
+		}
+	}
+	return EAvaPlayableCommandResult::Executed;
+}
+
+EAvaPlayableCommandResult UAvaPlayableRemoteProxy::UpdateRemoteControlCommand(const TSharedRef<FAvaPlayableRemoteControlValues>& InRemoteControlValues)
+{
+	IAvaPlaybackClient& PlaybackClient = IAvaMediaModule::Get().GetMediaPlaybackClient();
+	if (PlaybackClient.HasAnyServerOnlineForChannel(PlayingChannelFName))
+	{
+		PlaybackClient.RequestRemoteControlUpdate(InstanceId, SourceAssetPath, PlayingChannelName, *InRemoteControlValues);
+	}
+	return EAvaPlayableCommandResult::Executed;
+}
+
+bool UAvaPlayableRemoteProxy::ApplyCamera()
+{
+	return false;
+}
+
+void UAvaPlayableRemoteProxy::SetUserData(const FString& InUserData)
+{
+	if (UserData != InUserData)
+	{
+		IAvaPlaybackClient& PlaybackClient = IAvaMediaModule::Get().GetMediaPlaybackClient();
+		if (PlaybackClient.HasAnyServerOnlineForChannel(PlayingChannelFName))
+		{
+			PlaybackClient.RequestPlayback(InstanceId, GetSourceAssetPath(), PlayingChannelName, EAvaPlaybackAction::SetUserData, InUserData);
+		}
+	}
+	
+	Super::SetUserData(InUserData);
+}
+
+bool UAvaPlayableRemoteProxy::InitPlayable(const FPlayableCreationInfo& InPlayableInfo)
+{
+	// We keep track of the channel this playable is part of.
+	PlayingChannelFName = InPlayableInfo.ChannelName;
+	PlayingChannelName = InPlayableInfo.ChannelName.ToString();
+	
+	constexpr bool bIsRemoteProxy = true;
+
+	// Remote playables have proxy playable groups imitating the same setup as local ones.
+	switch (InPlayableInfo.SourceAsset.GetAssetType())
+	{
+	case EAvalancheAssetType::Blueprint:
+		{
+			UAvaPlayableGroup::FPlayableGroupCreationInfo PlayableGroupCreationInfo;
+			PlayableGroupCreationInfo.PlayableGroupManager = InPlayableInfo.PlayableGroupManager;
+			PlayableGroupCreationInfo.SourceAssetPath = InPlayableInfo.SourceAsset.ToSoftObjectPath();
+			PlayableGroupCreationInfo.ChannelName = InPlayableInfo.ChannelName;
+			PlayableGroupCreationInfo.bIsRemoteProxy = bIsRemoteProxy;
+			PlayableGroupCreationInfo.bIsSharedGroup = false;
+			
+			PlayableGroup = UAvaPlayableGroup::MakePlayableGroup(InPlayableInfo.PlayableGroupManager, PlayableGroupCreationInfo);
+		}
+		break;
+	case EAvalancheAssetType::World:
+		PlayableGroup = InPlayableInfo.PlayableGroupManager->GetOrCreateSharedLevelGroup(InPlayableInfo.ChannelName, bIsRemoteProxy);
+		break;
+	default:
+		UE_LOG(LogAvaPlayable, Error, TEXT("Asset \"%s\" is an unsupported type."), *InPlayableInfo.SourceAsset.ToSoftObjectPath().ToString());
+		break;
+	}
+	
+	const bool bInitSucceeded = Super::InitPlayable(InPlayableInfo);
+	if (bInitSucceeded)
+	{
+		RegisterClientEventHandlers();
+	}
+	return bInitSucceeded;
+}
+
+void UAvaPlayableRemoteProxy::OnPlay()
+{
+	IAvaMediaModule& AvaMediaModule = IAvaMediaModule::Get();
+	if (!AvaMediaModule.IsMediaPlaybackClientStarted())
+	{
+		return;
+	}
+	
+	IAvaPlaybackClient& Client = AvaMediaModule.GetMediaPlaybackClient();
+	
+	if (Client.HasAnyServerOnlineForChannel(PlayingChannelFName))
+	{
+		const FString ChannelName = PlayingChannelName;
+		const TOptional<EAvaPlaybackStatus> RemoteStatusOpt = Client.GetRemotePlaybackStatus(InstanceId, SourceAssetPath, ChannelName);
+		const EAvaPlaybackStatus RemoteStatus = RemoteStatusOpt.IsSet() ? RemoteStatusOpt.GetValue() : EAvaPlaybackStatus::Unknown;
+		// TODO: rework this logic.
+		if (RemoteStatus == EAvaPlaybackStatus::Available
+			|| RemoteStatus == EAvaPlaybackStatus::Loading || RemoteStatus == EAvaPlaybackStatus::Loaded
+			|| RemoteStatus == EAvaPlaybackStatus::Unknown || RemoteStatus == EAvaPlaybackStatus::Stopping
+			|| RemoteStatus == EAvaPlaybackStatus::Unloading)
+		{
+			Client.RequestPlayback(InstanceId, SourceAssetPath, ChannelName, EAvaPlaybackAction::Start);
+		}
+	}
+}
+
+void UAvaPlayableRemoteProxy::OnEndPlay()
+{
+	IAvaPlaybackClient& PlaybackClient = IAvaMediaModule::Get().GetMediaPlaybackClient();
+	if (PlaybackClient.HasAnyServerOnlineForChannel(PlayingChannelFName))
+	{
+		PlaybackClient.RequestPlayback(InstanceId, SourceAssetPath, PlayingChannelName, EAvaPlaybackAction::Stop);
+	}
+}
+
+void UAvaPlayableRemoteProxy::BeginDestroy()
+{
+	UnregisterClientEventHandlers();
+	Super::BeginDestroy();
+}
+
+void UAvaPlayableRemoteProxy::RegisterClientEventHandlers()
+{
+	using namespace UE::AvaPlaybackClient::Delegates;
+	GetOnPlaybackSequenceEvent().RemoveAll(this);
+	GetOnPlaybackSequenceEvent().AddUObject(this, &UAvaPlayableRemoteProxy::HandleAvaPlaybackSequenceEvent);
+}
+
+void UAvaPlayableRemoteProxy::UnregisterClientEventHandlers() const
+{
+	using namespace UE::AvaPlaybackClient::Delegates;
+	GetOnPlaybackSequenceEvent().RemoveAll(this);
+}
+
+TArray<FString> UAvaPlayableRemoteProxy::GetOnlineServerForChannel(const FName& InChannelName)
+{
+	const FAvaBroadcastOutputChannel& Channel = UAvaBroadcast::Get().GetCurrentProfile().GetChannel(InChannelName);
+	const TArray<UMediaOutput*>& Outputs = Channel.GetMediaOutputs();
+	TArray<FString> OnlineServers;
+	OnlineServers.Reserve(Outputs.Num());
+	
+	for (const UMediaOutput* Output : Outputs)
+	{
+		if (Channel.IsMediaOutputRemote(Output) && Channel.GetMediaOutputState(Output) != EAvaBroadcastOutputState::Offline)
+		{
+			OnlineServers.AddUnique(Channel.GetMediaOutputServerName(Output));
+		}
+	}
+	return OnlineServers;
+}
+
+void UAvaPlayableRemoteProxy::HandleAvaPlaybackSequenceEvent(IAvaPlaybackClient& InPlaybackClient,
+	const UE::AvaPlaybackClient::Delegates::FPlaybackSequenceEventArgs& InEventArgs)
+{
+	if (InEventArgs.InstanceId == InstanceId && InEventArgs.ChannelName == PlayingChannelName)
+	{
+		const FName SequenceName(InEventArgs.SequenceName);
+		OnSequenceEventDelegate.Broadcast(this, SequenceName, InEventArgs.EventType);
+	}
+}
+
+#undef LOCTEXT_NAMESPACE
diff --git a/Engine/Plugins/Experimental/Avalanche/Source/AvalancheMedia/Private/Playback/AvaPlaybackServerProcess.cpp b/Engine/Plugins/Experimental/Avalanche/Source/AvalancheMedia/Private/Playback/AvaPlaybackServerProcess.cpp
new file mode 100644
--- /dev/null
+++ b/Engine/Plugins/Experimental/Avalanche/Source/AvalancheMedia/Private/Playback/AvaPlaybackServerProcess.cpp
@@ -0,0 +1,183 @@
+// Copyright Epic Games, Inc. All Rights Reserved.
+
+#include "AvaPlaybackServerProcess.h"
+
+#include "AvaMediaSettings.h"
+#include "IAvaMediaModule.h"
+#include "Misc/App.h"
+#include "Misc/Paths.h"
+#include "Playback/AvaPlaybackClient.h"
+
+namespace UE::AvaPlaybackServerProcess::Private
+{
+	FString GetLocalServerName()
+	{
+		return UAvaMediaSettings::Get().LocalPlaybackServerSettings.ServerName;
+	}
+
+	const TCHAR* GetPlaybackServerLogReplicationVerbosityString()
+	{
+		const UAvaMediaSettings& Settings = UAvaMediaSettings::Get(); 
+		return ToString(UAvaMediaSettings::ToLogVerbosity(Settings.PlaybackServerLogReplicationVerbosity));
+	}
+}
+
+FAvaPlaybackServerProcess::FAvaPlaybackServerProcess(FProcHandle&& InProcessHandle)
+	: ProcessHandle(MoveTemp(InProcessHandle))
+{}
+
+bool FAvaPlaybackServerProcess::IsLaunched()
+{
+	return ProcessHandle.IsValid() && FPlatformProcess::IsProcRunning(ProcessHandle);
+}
+
+bool FAvaPlaybackServerProcess::Launch()
+{
+	if (IsLaunched())
+	{
+		UE_LOG(LogAvaMedia, Warning, TEXT("Local playback server is already running."));
+		return true;
+	}
+
+	using namespace UE::AvaPlaybackServerProcess::Private;
+
+	FString GameNameOrProjectFile;
+	if (FPaths::IsProjectFilePathSet())
+	{
+		GameNameOrProjectFile = FString::Printf(TEXT("\"%s\""), *FPaths::GetProjectFilePath());
+	}
+	else
+	{
+		GameNameOrProjectFile = FApp::GetProjectName();
+	}
+
+	FString CommandLine;
+	CommandLine += GameNameOrProjectFile;
+
+	// Todo: It will likely need a more complete set of parameters. Reference: UDisplayClusterLaunchEditorProjectSettings.
+	const FAvaMediaLocalPlaybackServerSettings& Settings = UAvaMediaSettings::Get().LocalPlaybackServerSettings;
+
+	// Standard arguments
+	const FIntPoint Resolution = Settings.Resolution;
+	CommandLine += FString::Printf(TEXT(" -game -windowed -ResX=%d -ResY=%d -messaging"), Resolution.X, Resolution.Y);
+	if (Settings.bEnableLogConsole)
+	{
+		CommandLine += TEXT(" -log");
+	}
+	// Motion Design arguments
+	CommandLine += FString::Printf(TEXT(" -MotionDesignPlaybackServerStart=%s"), *GetLocalServerName());
+	CommandLine += TEXT(" -MotionDesignPlaybackClientSuppress");
+	CommandLine += FString::Printf(TEXT(" -MotionDesignPlaybackServerLogReplication=%s"), GetPlaybackServerLogReplicationVerbosityString());
+	// Storm Sync arguments
+	CommandLine += TEXT(" -NoStormSyncServerAutoStart");
+
+	if (!Settings.ExtraCommandLineArguments.IsEmpty())
+	{
+		CommandLine += TEXT(" ");
+		CommandLine += Settings.ExtraCommandLineArguments;
+	}
+
+	FString LogCommands;
+	
+	for (const FAvaPlaybackServerLoggingEntry& LoggingEntry : Settings.Logging)
+	{
+		if (!LoggingEntry.Category.IsNone())
+		{
+			LogCommands += FString::Printf(TEXT("%s%s %s"),
+					LogCommands.IsEmpty() ? TEXT("") : TEXT(", "), *LoggingEntry.Category.ToString(),
+					ToString(UAvaMediaSettings::ToLogVerbosity(LoggingEntry.VerbosityLevel)));
+		}
+	}
+
+	if (!LogCommands.IsEmpty())
+	{
+		CommandLine += FString::Printf(TEXT(" -LogCmds=\"%s\""), *LogCommands);		
+	}
+
+	const FString ExecutablePath = FPlatformProcess::ExecutablePath();
+
+	uint32 ProcessID = 0;
+	constexpr bool bLaunchDetached = true;
+	constexpr bool bLaunchMinimized = false;
+	constexpr bool bLaunchWindowHidden = false;
+	constexpr uint32 PriorityModifier = 0;
+
+	UE_LOG(LogAvaMedia, Log, TEXT("Launching a playback server in game mode in a new process with the following command line:"));
+	UE_LOG(LogAvaMedia, Log, TEXT("%s %s"), *ExecutablePath, *CommandLine);
+
+	ProcessHandle = FPlatformProcess::CreateProc(
+		*ExecutablePath, *CommandLine, bLaunchDetached,
+		bLaunchMinimized, bLaunchWindowHidden, &ProcessID,
+		PriorityModifier, nullptr, nullptr, nullptr);
+	
+	return ProcessHandle.IsValid();
+}
+
+void FAvaPlaybackServerProcess::Stop()
+{
+	FPlatformProcess::TerminateProc(ProcessHandle);
+}
+
+TSharedPtr<FAvaPlaybackServerProcess> FAvaPlaybackServerProcess::Find(const FAvaPlaybackClient& InPlaybackClient)
+{
+	using namespace UE::AvaPlaybackServerProcess::Private;
+
+	// We use the playback client to find connected server processes that are considered local.
+	// Local means that it is in the same project content path and the same machine. We figure out
+	// if it is the same machine if the process Id can be opened.
+
+	const FString ProjectContentPath = InPlaybackClient.GetProjectContentPath();
+	const uint32 ClientProcessId = FPlatformProcess::GetCurrentProcessId();
+	const FString LocalServerName = GetLocalServerName();
+	
+	// First try with the configured local server name.
+	if (InPlaybackClient.GetServerProjectContentPath(LocalServerName) == ProjectContentPath)
+	{
+		const uint32 ServerProcessId = InPlaybackClient.GetServerProcessId(LocalServerName);
+		if (ServerProcessId != 0 && ClientProcessId != ServerProcessId)
+		{
+			FProcHandle ProcessHandle = FPlatformProcess::OpenProcess(ServerProcessId);
+			if (ProcessHandle.IsValid())
+			{
+				return MakeShared<FAvaPlaybackServerProcess>(MoveTemp(ProcessHandle));
+			}
+
+			// Warn about the server name collision issue.
+			UE_LOG(LogAvaMedia, Warning, TEXT("Found Connected Server \"%s\" but it is not a local process."), *LocalServerName);
+			UE_LOG(LogAvaMedia, Warning, TEXT("This indicates a name collision between the servers. The local server should be renamed."));
+		}
+	}
+
+	// Fallback
+	// Try to find any connected server that is "local", i.e. on the same machine and
+	// in the same project content folder, but on a different process id.
+	
+	TArray<FString> ServerNames = InPlaybackClient.GetServerNames();
+	for (const FString& ServerName : ServerNames)
+	{
+		if (InPlaybackClient.GetServerProjectContentPath(ServerName) == ProjectContentPath)
+		{
+			const uint32 ServerProcessId = InPlaybackClient.GetServerProcessId(ServerName);
+			if (ClientProcessId != ServerProcessId)
+			{
+				FProcHandle ProcessHandle = FPlatformProcess::OpenProcess(ServerProcessId);
+				if (ProcessHandle.IsValid())
+				{
+					return MakeShared<FAvaPlaybackServerProcess>(MoveTemp(ProcessHandle));
+				}
+			}
+		}
+	}
+	
+	return nullptr;	
+}
+
+TSharedPtr<FAvaPlaybackServerProcess> FAvaPlaybackServerProcess::FindOrCreate(const FAvaPlaybackClient& InPlaybackClient)
+{
+	TSharedPtr<FAvaPlaybackServerProcess> ServerProcess = Find(InPlaybackClient);
+	if (!ServerProcess.IsValid())
+	{
+		ServerProcess = MakeShared<FAvaPlaybackServerProcess>();
+	}
+	return ServerProcess;
+}
\ No newline at end of file
diff --git a/Engine/Plugins/Experimental/Avalanche/Source/AvalancheMedia/Private/Playback/Nodes/AvaPlaybackNode.cpp b/Engine/Plugins/Experimental/Avalanche/Source/AvalancheMedia/Private/Playback/Nodes/AvaPlaybackNode.cpp
--- a/Engine/Plugins/Experimental/Avalanche/Source/AvalancheMedia/Private/Playback/Nodes/AvaPlaybackNode.cpp
+++ b/Engine/Plugins/Experimental/Avalanche/Source/AvalancheMedia/Private/Playback/Nodes/AvaPlaybackNode.cpp
@@ -4,14 +4,14 @@
 
 #include "EngineUtils.h"
 #include "Misc/App.h"
-#include "Playback/AvalanchePlayback.h"
-#include "Playback/IAvaPlaybackEditor.h"
+#include "Playback/AvaPlaybackGraph.h"
+#include "Playback/IAvaPlaybackGraphEditor.h"
 
 #if WITH_EDITOR
 #include "EdGraph/EdGraph.h"
 #endif
 
-#define LOCTEXT_NAMESPACE "AvalanchePlayback"
+#define LOCTEXT_NAMESPACE "AvaPlaybackNode"
 
 const FText UAvaPlaybackNode::NodeCategory::Default          = LOCTEXT("Category_Default", "Default");
 const FText UAvaPlaybackNode::NodeCategory::EventTrigger     = LOCTEXT("Category_EventTrigger", "Event Triggers");
@@ -20,7 +20,7 @@ const FText UAvaPlaybackNode::NodeCategory::EventFlowControl = LOCTEXT("Category
 
 void UAvaPlaybackNode::PostAllocateNode()
 {
-	if (UAvalanchePlayback* const Playback = GetPlayback())
+	if (UAvaPlaybackGraph* const Playback = GetPlayback())
 	{
 		PlaybackStateChangedHandle = Playback->OnPlaybackStateChanged.AddUObject(this, &UAvaPlaybackNode::NotifyPlaybackStateChanged);
 	}
@@ -35,9 +35,9 @@ void UAvaPlaybackNode::BeginDestroy()
 	UObject::BeginDestroy();
 }
 
-UAvalanchePlayback* UAvaPlaybackNode::GetPlayback() const
+UAvaPlaybackGraph* UAvaPlaybackNode::GetPlayback() const
 {
-	return Cast<UAvalanchePlayback>(GetOuter());
+	return Cast<UAvaPlaybackGraph>(GetOuter());
 }
 
 FText UAvaPlaybackNode::GetNodeDisplayNameText() const
@@ -214,7 +214,7 @@ void UAvaPlaybackNode::ReconstructNode()
 
 void UAvaPlaybackNode::RefreshNode(bool bDryRunGraph)
 {
-	if (UAvalanchePlayback* const Playback = GetPlayback())
+	if (UAvaPlaybackGraph* const Playback = GetPlayback())
 	{
 #if WITH_EDITOR
 		Playback->RefreshPlaybackNode(this);
@@ -249,7 +249,7 @@ void UAvaPlaybackNode::InsertChildNode(int32 Index)
 	{
 		ChildNodes.InsertZeroed(Index);
 #if WITH_EDITOR
-		UAvalanchePlayback* const Playback = GetPlayback();
+		UAvaPlaybackGraph* const Playback = GetPlayback();
 		if (Playback && GraphNode && Playback->GetGraphEditor().IsValid())
 		{
 			Playback->GetGraphEditor()->CreateInputPin(GraphNode);
diff --git a/Engine/Plugins/Experimental/Avalanche/Source/AvalancheMedia/Private/Playback/Nodes/AvaPlaybackNodeBlueprintPlayer.cpp b/Engine/Plugins/Experimental/Avalanche/Source/AvalancheMedia/Private/Playback/Nodes/AvaPlaybackNodeBlueprintPlayer.cpp
--- a/Engine/Plugins/Experimental/Avalanche/Source/AvalancheMedia/Private/Playback/Nodes/AvaPlaybackNodeBlueprintPlayer.cpp
+++ b/Engine/Plugins/Experimental/Avalanche/Source/AvalancheMedia/Private/Playback/Nodes/AvaPlaybackNodeBlueprintPlayer.cpp
@@ -1,12 +1,11 @@
 // Copyright Epic Games, Inc. All Rights Reserved.
 
 #include "Playback/Nodes/AvaPlaybackNodeBlueprintPlayer.h"
+
 #include "AvaBlueprint.h"
-#include "Framework/AvaGameInstance.h"
-#include "Http/AvaMediaHttpServer.h"
-#include "Playback/AvalanchePlayback.h"
+#include "Playback/AvaPlaybackGraph.h"
 
-#define LOCTEXT_NAMESPACE "AvalanchePlayback"
+#define LOCTEXT_NAMESPACE "AvaPlaybackNodeBlueprintPlayer"
 
 UAvaPlaybackNodeBlueprintPlayer::UAvaPlaybackNodeBlueprintPlayer()
 {
@@ -26,14 +25,14 @@ void UAvaPlaybackNodeBlueprintPlayer::PostLoad()
 	UpdateDisplayNameText();
 }
 
-void UAvaPlaybackNodeBlueprintPlayer::SetAvalancheAsset(const TSoftObjectPtr<UAvalancheBlueprint>& InAsset)
+void UAvaPlaybackNodeBlueprintPlayer::SetAsset(const TSoftObjectPtr<UAvalancheBlueprint>& InAsset)
 {
-	BlueprintAsset = InAsset;
+	Asset = InAsset;
 }
 
 void UAvaPlaybackNodeBlueprintPlayer::UpdateDisplayNameText()
 {
-	const FString AssetName = BlueprintAsset.GetAssetName();
+	const FString AssetName = Asset.GetAssetName();
 	
 	if (AssetName.IsEmpty())
 	{
@@ -46,12 +45,12 @@ void UAvaPlaybackNodeBlueprintPlayer::UpdateDisplayNameText()
 	}
 }
 
-FAvaSoftAssetPtr UAvaPlaybackNodeBlueprintPlayer::GetAvalancheAssetPtr() const
+FAvaSoftAssetPtr UAvaPlaybackNodeBlueprintPlayer::GetAssetPtr() const
 {
-	FAvaSoftAssetPtr OutAvalancheAsset;
-	OutAvalancheAsset.AssetClassPath = FSoftClassPath(UAvalancheBlueprint::StaticClass());
-	OutAvalancheAsset.AssetPtr = BlueprintAsset;
-	return OutAvalancheAsset;
+	FAvaSoftAssetPtr OutAsset;
+	OutAsset.AssetClassPath = FSoftClassPath(UAvalancheBlueprint::StaticClass());
+	OutAsset.AssetPtr = Asset;
+	return OutAsset;
 }
 
 #undef LOCTEXT_NAMESPACE
diff --git a/Engine/Plugins/Experimental/Avalanche/Source/AvalancheMedia/Private/Playback/Nodes/AvaPlaybackNodeCombiner.cpp b/Engine/Plugins/Experimental/Avalanche/Source/AvalancheMedia/Private/Playback/Nodes/AvaPlaybackNodeCombiner.cpp
--- a/Engine/Plugins/Experimental/Avalanche/Source/AvalancheMedia/Private/Playback/Nodes/AvaPlaybackNodeCombiner.cpp
+++ b/Engine/Plugins/Experimental/Avalanche/Source/AvalancheMedia/Private/Playback/Nodes/AvaPlaybackNodeCombiner.cpp
@@ -2,7 +2,7 @@
 
 #include "Playback/Nodes/AvaPlaybackNodeCombiner.h"
 
-#define LOCTEXT_NAMESPACE "AvalanchePlayback"
+#define LOCTEXT_NAMESPACE "AvaPlaybackNodeCombiner"
 
 FText UAvaPlaybackNodeCombiner::GetNodeDisplayNameText() const
 {
diff --git a/Engine/Plugins/Experimental/Avalanche/Source/AvalancheMedia/Private/Playback/Nodes/AvaPlaybackNodeLevelPlayer.cpp b/Engine/Plugins/Experimental/Avalanche/Source/AvalancheMedia/Private/Playback/Nodes/AvaPlaybackNodeLevelPlayer.cpp
--- a/Engine/Plugins/Experimental/Avalanche/Source/AvalancheMedia/Private/Playback/Nodes/AvaPlaybackNodeLevelPlayer.cpp
+++ b/Engine/Plugins/Experimental/Avalanche/Source/AvalancheMedia/Private/Playback/Nodes/AvaPlaybackNodeLevelPlayer.cpp
@@ -3,9 +3,9 @@
 #include "Playback/Nodes/AvaPlaybackNodeLevelPlayer.h"
 
 #include "Engine/World.h"
-#include "Playback/AvalanchePlayback.h"
+#include "Playback/AvaPlaybackGraph.h"
 
-#define LOCTEXT_NAMESPACE "AvalanchePlayback"
+#define LOCTEXT_NAMESPACE "AvaPlaybackNodeLevelPlayer"
 
 UAvaPlaybackNodeLevelPlayer::UAvaPlaybackNodeLevelPlayer()
 {
@@ -25,7 +25,7 @@ void UAvaPlaybackNodeLevelPlayer::PostLoad()
 	UpdateDisplayNameText();
 }
 
-void UAvaPlaybackNodeLevelPlayer::SetAvalancheAsset(const TSoftObjectPtr<UWorld>& InAsset)
+void UAvaPlaybackNodeLevelPlayer::SetAsset(const TSoftObjectPtr<UWorld>& InAsset)
 {
 	LevelAsset = InAsset;
 }
@@ -45,12 +45,12 @@ void UAvaPlaybackNodeLevelPlayer::UpdateDisplayNameText()
 	}
 }
 
-FAvaSoftAssetPtr UAvaPlaybackNodeLevelPlayer::GetAvalancheAssetPtr() const
+FAvaSoftAssetPtr UAvaPlaybackNodeLevelPlayer::GetAssetPtr() const
 {
-	FAvaSoftAssetPtr OutAvalancheAsset;
-	OutAvalancheAsset.AssetClassPath = FSoftClassPath(UWorld::StaticClass());
-	OutAvalancheAsset.AssetPtr = LevelAsset;
-	return OutAvalancheAsset;
+	FAvaSoftAssetPtr OutAsset;
+	OutAsset.AssetClassPath = FSoftClassPath(UWorld::StaticClass());
+	OutAsset.AssetPtr = LevelAsset;
+	return OutAsset;
 }
 
 #undef LOCTEXT_NAMESPACE
diff --git a/Engine/Plugins/Experimental/Avalanche/Source/AvalancheMedia/Private/Playback/Nodes/AvaPlaybackNodePlayer.cpp b/Engine/Plugins/Experimental/Avalanche/Source/AvalancheMedia/Private/Playback/Nodes/AvaPlaybackNodePlayer.cpp
--- a/Engine/Plugins/Experimental/Avalanche/Source/AvalancheMedia/Private/Playback/Nodes/AvaPlaybackNodePlayer.cpp
+++ b/Engine/Plugins/Experimental/Avalanche/Source/AvalancheMedia/Private/Playback/Nodes/AvaPlaybackNodePlayer.cpp
@@ -2,19 +2,19 @@
 
 #include "Playback/Nodes/AvaPlaybackNodePlayer.h"
 
-#include "AvalancheBroadcast.h"
-#include "Playback/AvalanchePlayable.h"
-#include "Playback/AvalanchePlayback.h"
+#include "Broadcast/AvaBroadcast.h"
+#include "Playable/AvaPlayable.h"
+#include "Playback/AvaPlaybackGraph.h"
 #include "Playback/Nodes/Events/AvaPlaybackNodeEvent.h"
 
-#define LOCTEXT_NAMESPACE "AvalanchePlayback"
+#define LOCTEXT_NAMESPACE "AvaPlaybackNodePlayer"
 
 UAvaPlaybackNodePlayer::UAvaPlaybackNodePlayer() = default;
 
 void UAvaPlaybackNodePlayer::PostAllocateNode()
 {
 	Super::PostAllocateNode();
-	if (UAvalanchePlayback* const Playback = GetPlayback())
+	if (UAvaPlaybackGraph* const Playback = GetPlayback())
 	{
 		Playback->AddPlayerNode(this);
 	}
@@ -22,7 +22,7 @@ void UAvaPlaybackNodePlayer::PostAllocateNode()
 
 void UAvaPlaybackNodePlayer::Tick(float DeltaTime, FAvaPlaybackChannelParameters& ChannelParameters)
 {
-	ChannelParameters.AvalancheAssets.AddUnique(GetAvalancheAssetPtr());
+	ChannelParameters.Assets.AddUnique(GetAssetPtr());
 	ChannelIndices.Add(ChannelParameters.ChannelIndex);
 	//No Children Ticks. The Player Node is the Dead End for Ticking.
 	//Events are handled separately
@@ -37,7 +37,7 @@ void UAvaPlaybackNodePlayer::TickEventFeed(float DeltaTime)
 		{
 			FAvaPlaybackEventParameters EventParameters;
 			EventParameters.ChannelIndices = ChannelIndices;
-			EventParameters.AvalancheAsset = GetAvalancheAssetPtr();
+			EventParameters.Asset = GetAssetPtr();
 			
 			EventNode->TickEvent(DeltaTime, EventParameters);
 			
@@ -87,17 +87,17 @@ UTextureRenderTarget2D* UAvaPlaybackNodePlayer::GetPreviewRenderTarget() const
 	// Returns the first valid channel index we have.
 	for (const int32 ChannelIndex : LastTickChannelIndices)
 	{
-		const FName ChannelName = UAvalancheBroadcast::Get().GetChannelName(ChannelIndex);
+		const FName ChannelName = UAvaBroadcast::Get().GetChannelName(ChannelIndex);
 		if (ChannelName.IsNone())
 		{
 			continue;
 		}
 
-		if (const UAvalanchePlayback* const Playback = GetPlayback())
+		if (const UAvaPlaybackGraph* const Playback = GetPlayback())
 		{
-			if (const UAvalanchePlayable* const Playable = Playback->FindPlayable(GetAvalancheAssetPath(), ChannelName))
+			if (const UAvaPlayable* const Playable = Playback->FindPlayable(GetAssetPath(), ChannelName))
 			{
-				if (const UAvaMediaPlayableGroup* PlayableGroup = Playable->GetPlayableGroup())
+				if (const UAvaPlayableGroup* PlayableGroup = Playable->GetPlayableGroup())
 				{
 					if (UTextureRenderTarget2D* RenderTarget = PlayableGroup->GetRenderTarget())
 					{
@@ -108,7 +108,7 @@ UTextureRenderTarget2D* UAvaPlaybackNodePlayer::GetPreviewRenderTarget() const
 		}
 
 		// Fallback to channel RT which is likely to be (actually should be) the same as the game instance.
-		const FAvaOutputChannel& Channel = UAvalancheBroadcast::Get().GetCurrentProfile().GetChannel(ChannelName);
+		const FAvaBroadcastOutputChannel& Channel = UAvaBroadcast::Get().GetCurrentProfile().GetChannel(ChannelName);
 		if (Channel.IsValidChannel())
 		{
 			return Channel.GetCurrentRenderTarget(true);
@@ -118,7 +118,7 @@ UTextureRenderTarget2D* UAvaPlaybackNodePlayer::GetPreviewRenderTarget() const
 }
 #endif
 
-const FSoftObjectPath& UAvaPlaybackNodePlayer::GetAvalancheAssetPath() const
+const FSoftObjectPath& UAvaPlaybackNodePlayer::GetAssetPath() const
 {
 	static FSoftObjectPath EmptyPath;
 	return EmptyPath;
diff --git a/Engine/Plugins/Experimental/Avalanche/Source/AvalancheMedia/Private/Playback/Nodes/AvaPlaybackNodeRoot.cpp b/Engine/Plugins/Experimental/Avalanche/Source/AvalancheMedia/Private/Playback/Nodes/AvaPlaybackNodeRoot.cpp
--- a/Engine/Plugins/Experimental/Avalanche/Source/AvalancheMedia/Private/Playback/Nodes/AvaPlaybackNodeRoot.cpp
+++ b/Engine/Plugins/Experimental/Avalanche/Source/AvalancheMedia/Private/Playback/Nodes/AvaPlaybackNodeRoot.cpp
@@ -1,24 +1,25 @@
 // Copyright Epic Games, Inc. All Rights Reserved.
 
 #include "Playback/Nodes/AvaPlaybackNodeRoot.h"
-#include "AvalancheBroadcast.h"
-#include "Playback/AvalanchePlayback.h"
 
-#define LOCTEXT_NAMESPACE "AvalanchePlayback"
+#include "Broadcast/AvaBroadcast.h"
+#include "Playback/AvaPlaybackGraph.h"
+
+#define LOCTEXT_NAMESPACE "AvaPlaybackNodeRoot"
 
 void UAvaPlaybackNodeRoot::PostAllocateNode()
 {
 	Super::PostAllocateNode();
 	
-	UAvalancheBroadcast& Broadcast = UAvalancheBroadcast::Get();
+	UAvaBroadcast& Broadcast = UAvaBroadcast::Get();
 	
 	BroadcastChangedHandle = Broadcast.AddChangeListener(FOnAvaBroadcastChanged::FDelegate::CreateUObject(this
 		, &UAvaPlaybackNodeRoot::OnBroadcastChanged));
 		
-	ChannelChangedHandle = FAvaOutputChannel::GetOnChannelChanged().AddUObject(this
+	ChannelChangedHandle = FAvaBroadcastOutputChannel::GetOnChannelChanged().AddUObject(this
 		, &UAvaPlaybackNodeRoot::OnChannelChanged);
 
-	if (UAvalanchePlayback* const Playback = GetPlayback())
+	if (UAvaPlaybackGraph* const Playback = GetPlayback())
 	{
 		Playback->SetRootNode(this);
 	}
@@ -28,11 +29,11 @@ void UAvaPlaybackNodeRoot::BeginDestroy()
 {
 	if (BroadcastChangedHandle.IsValid())
 	{
-		UAvalancheBroadcast::Get().RemoveChangeListener(BroadcastChangedHandle);
+		UAvaBroadcast::Get().RemoveChangeListener(BroadcastChangedHandle);
 	}
 	if (ChannelChangedHandle.IsValid())
 	{
-		FAvaOutputChannel::GetOnChannelChanged().Remove(ChannelChangedHandle);
+		FAvaBroadcastOutputChannel::GetOnChannelChanged().Remove(ChannelChangedHandle);
 	}
 	Super::BeginDestroy();
 }
@@ -70,9 +71,9 @@ void UAvaPlaybackNodeRoot::OnBroadcastChanged(EAvaBroadcastChange InChange)
 	ReconstructNode();
 }
 
-void UAvaPlaybackNodeRoot::OnChannelChanged(const FAvaOutputChannel& InChannel, EAvaChannelChange InChange)
+void UAvaPlaybackNodeRoot::OnChannelChanged(const FAvaBroadcastOutputChannel& InChannel, EAvaBroadcastChannelChange InChange)
 {
-	if (EnumHasAnyFlags(InChange, EAvaChannelChange::State))
+	if (EnumHasAnyFlags(InChange, EAvaBroadcastChannelChange::State))
 	{
 		//TODO: Instead of Reconstructing, have the Pins just react to each Individual Channel State Change
 		ReconstructNode();
@@ -81,25 +82,25 @@ void UAvaPlaybackNodeRoot::OnChannelChanged(const FAvaOutputChannel& InChannel,
 
 int32 UAvaPlaybackNodeRoot::GetMinChildNodes() const
 {
-	return UAvalancheBroadcast::Get().GetChannelNameCount();
+	return UAvaBroadcast::Get().GetChannelNameCount();
 }
 
 int32 UAvaPlaybackNodeRoot::GetMaxChildNodes() const
 {
-	return UAvalancheBroadcast::Get().GetChannelNameCount();
+	return UAvaBroadcast::Get().GetChannelNameCount();
 }
 
 #if WITH_EDITOR
 FName UAvaPlaybackNodeRoot::GetInputPinName(int32 InputPinIndex) const
 {
-	return UAvalancheBroadcast::Get().GetChannelName(InputPinIndex);
+	return UAvaBroadcast::Get().GetChannelName(InputPinIndex);
 }
 #endif
 
 FName UAvaPlaybackNodeRoot::GetChannelName(int32 InChannelNameIndex) const
 {
-	const UAvalancheBroadcast& AvalancheBroadcast = UAvalancheBroadcast::Get();
-	return AvalancheBroadcast.GetChannelName(InChannelNameIndex);
+	const UAvaBroadcast& AvaBroadcast = UAvaBroadcast::Get();
+	return AvaBroadcast.GetChannelName(InChannelNameIndex);
 }
 
 #undef LOCTEXT_NAMESPACE
diff --git a/Engine/Plugins/Experimental/Avalanche/Source/AvalancheMedia/Private/Playback/Nodes/AvaPlaybackNodeSwitcher.cpp b/Engine/Plugins/Experimental/Avalanche/Source/AvalancheMedia/Private/Playback/Nodes/AvaPlaybackNodeSwitcher.cpp
--- a/Engine/Plugins/Experimental/Avalanche/Source/AvalancheMedia/Private/Playback/Nodes/AvaPlaybackNodeSwitcher.cpp
+++ b/Engine/Plugins/Experimental/Avalanche/Source/AvalancheMedia/Private/Playback/Nodes/AvaPlaybackNodeSwitcher.cpp
@@ -2,7 +2,7 @@
 
 #include "Playback/Nodes/AvaPlaybackNodeSwitcher.h"
 
-#define LOCTEXT_NAMESPACE "AvalanchePlayback"
+#define LOCTEXT_NAMESPACE "AvaPlaybackNodeSwitcher"
 
 FText UAvaPlaybackNodeSwitcher::GetNodeDisplayNameText() const
 {
diff --git a/Engine/Plugins/Experimental/Avalanche/Source/AvalancheMedia/Private/Playback/Nodes/Events/Actions/AvaPlaybackNode_PlayAnim.cpp b/Engine/Plugins/Experimental/Avalanche/Source/AvalancheMedia/Private/Playback/Nodes/Events/Actions/AvaPlaybackNode_PlayAnim.cpp
--- a/Engine/Plugins/Experimental/Avalanche/Source/AvalancheMedia/Private/Playback/Nodes/Events/Actions/AvaPlaybackNode_PlayAnim.cpp
+++ b/Engine/Plugins/Experimental/Avalanche/Source/AvalancheMedia/Private/Playback/Nodes/Events/Actions/AvaPlaybackNode_PlayAnim.cpp
@@ -1,19 +1,20 @@
 // Copyright Epic Games, Inc. All Rights Reserved.
 
 #include "AvaPlaybackNode_PlayAnim.h"
+
 #include "Async/Async.h"
 #include "AvaBlueprint.h"
 #include "AvaScene.h"
-#include "AvalancheBroadcast.h"
+#include "AvaSequence.h"
+#include "Broadcast/AvaBroadcast.h"
 #include "Engine/Level.h"
-#include "Playback/AvalanchePlayback.h"
+#include "Playback/AvaPlaybackGraph.h"
 #include "Playback/Nodes/AvaPlaybackNodeBlueprintPlayer.h"
 #include "Playback/Nodes/AvaPlaybackNodeLevelPlayer.h"
-#include "AvaSequence.h"
 
-#define LOCTEXT_NAMESPACE "AvalanchePlayback"
+#define LOCTEXT_NAMESPACE "AvaPlaybackNode_PlayAnim"
 
-namespace Avalanche
+namespace UE::AvaMedia::NodePlayAnim::Private
 {
 	FName GetSequenceName(const TObjectPtr<UAvaSequence>& InSequence)
 	{
@@ -31,12 +32,12 @@ FText UAvaPlaybackNode_PlayAnim::GetNodeDisplayNameText() const
 	else
 	{
 		FString AnimationTree;
-		for (TMap<FSoftObjectPath, FAvalancheAnimations>::TConstIterator Iter(AnimationMap); Iter; ++Iter)
+		for (TMap<FSoftObjectPath, FAvaPlaybackAnimations>::TConstIterator Iter(AnimationMap); Iter; ++Iter)
 		{
 			bool bAction = false;
-			for (const FAnimPlaySettings& PlaySettings : Iter->Value.AvailableAnimations)
+			for (const FAvaPlaybackAnimPlaySettings& PlaySettings : Iter->Value.AvailableAnimations)
 			{
-				if (PlaySettings.Action != EAvaMediaAnimAction::None)
+				if (PlaySettings.Action != EAvaPlaybackAnimAction::None)
 				{
 					bAction = true;
 					break;
@@ -60,25 +61,25 @@ FText UAvaPlaybackNode_PlayAnim::GetNodeTooltipText() const
 
 void UAvaPlaybackNode_PlayAnim::OnEventTriggered(const FAvaPlaybackEventParameters& InEventParameters)
 {
-	UAvalanchePlayback* const Playback = GetPlayback();
+	UAvaPlaybackGraph* const Playback = GetPlayback();
 	if (!Playback)
 	{
 		return;
 	}
 
-	const FSoftObjectPath& AvalancheAssetPath = InEventParameters.AvalancheAsset.ToSoftObjectPath();
+	const FSoftObjectPath& AssetPath = InEventParameters.Asset.ToSoftObjectPath();
 	TArray<FName> ChannelNames = Playback->GetChannelNamesForIndices(InEventParameters.ChannelIndices);
 	
 	// Gather the Animations to Play
-	if (FAvalancheAnimations* const FoundAnimations = AnimationMap.Find(AvalancheAssetPath))
+	if (FAvaPlaybackAnimations* const FoundAnimations = AnimationMap.Find(AssetPath))
 	{
-		for (const FAnimPlaySettings& PlaySettings : FoundAnimations->AvailableAnimations)
+		for (const FAvaPlaybackAnimPlaySettings& PlaySettings : FoundAnimations->AvailableAnimations)
 		{
-			if (PlaySettings.Action != EAvaMediaAnimAction::None)
+			if (PlaySettings.Action != EAvaPlaybackAnimAction::None)
 			{
 				for (const FName& ChannelName : ChannelNames)
 				{
-					Playback->PushAnimationCommand(AvalancheAssetPath, ChannelName.ToString(), PlaySettings.Action, PlaySettings);
+					Playback->PushAnimationCommand(AssetPath, ChannelName.ToString(), PlaySettings.Action, PlaySettings);
 				}
 			}
 		}
@@ -96,37 +97,37 @@ void UAvaPlaybackNode_PlayAnim::DryRun(const TArray<UAvaPlaybackNode*>& InAncest
 	{
 		if (UAvaPlaybackNodeBlueprintPlayer* const PlayerNode = Cast<UAvaPlaybackNodeBlueprintPlayer>(PlaybackNode))
 		{
-			TSoftObjectPtr<UAvalancheBlueprint> AvalancheAsset = PlayerNode->GetAvalancheAsset();
-			if (UAvalancheBlueprint* const AvalancheBlueprint = AvalancheAsset.LoadSynchronous())
+			TSoftObjectPtr<UAvalancheBlueprint> Asset = PlayerNode->GetAsset();
+			if (UAvalancheBlueprint* const AvalancheBlueprint = Asset.LoadSynchronous())
 			{
-				SeenAssetsInDryRun.Add(PlayerNode->GetAvalancheAssetPath());
-				FAvalancheAnimations& Animations = AnimationMap.FindOrAdd(PlayerNode->GetAvalancheAssetPath());
+				SeenAssetsInDryRun.Add(PlayerNode->GetAssetPath());
+				FAvaPlaybackAnimations& Animations = AnimationMap.FindOrAdd(PlayerNode->GetAssetPath());
 				for (const TObjectPtr<UAvaSequence>& Animation : AvalancheBlueprint->GetSequences())
 				{
 					if (Animation)
 					{
-						Animations.AvailableAnimations.FindOrAdd(Avalanche::GetSequenceName(Animation));
+						Animations.AvailableAnimations.FindOrAdd(UE::AvaMedia::NodePlayAnim::Private::GetSequenceName(Animation));
 					}
 				}
 			}
 		}
 		if (UAvaPlaybackNodeLevelPlayer* const PlayerNode = Cast<UAvaPlaybackNodeLevelPlayer>(PlaybackNode))
 		{
-			TSoftObjectPtr<UWorld> AvalancheAsset = PlayerNode->GetAvalancheAsset();
-			if (UWorld* const AvalancheWorld = AvalancheAsset.LoadSynchronous())
+			TSoftObjectPtr<UWorld> Asset = PlayerNode->GetAsset();
+			if (UWorld* const World = Asset.LoadSynchronous())
 			{
-				SeenAssetsInDryRun.Add(PlayerNode->GetAvalancheAssetPath());
-				FAvalancheAnimations& Animations = AnimationMap.FindOrAdd(PlayerNode->GetAvalancheAssetPath());
+				SeenAssetsInDryRun.Add(PlayerNode->GetAssetPath());
+				FAvaPlaybackAnimations& Animations = AnimationMap.FindOrAdd(PlayerNode->GetAssetPath());
 
 				AAvaScene* Scene = nullptr;
-				AvalancheWorld->PersistentLevel->Actors.FindItemByClass(&Scene);
+				World->PersistentLevel->Actors.FindItemByClass(&Scene);
 				if (IsValid(Scene))
 				{
 					for (const TObjectPtr<UAvaSequence>& Animation : Scene->GetSequences())
 					{
 						if (Animation)
 						{
-							Animations.AvailableAnimations.FindOrAdd(Avalanche::GetSequenceName(Animation));
+							Animations.AvailableAnimations.FindOrAdd(UE::AvaMedia::NodePlayAnim::Private::GetSequenceName(Animation));
 						}
 					}
 				}
@@ -140,7 +141,7 @@ void UAvaPlaybackNode_PlayAnim::PostDryRun()
 	bool bRefreshNode = false;
 	
 	//Remove all the Assets that were not Seen.
-	for (TMap<FSoftObjectPath, FAvalancheAnimations>::TIterator Iter(AnimationMap); Iter; ++Iter)
+	for (TMap<FSoftObjectPath, FAvaPlaybackAnimations>::TIterator Iter(AnimationMap); Iter; ++Iter)
 	{
 		if (!SeenAssetsInDryRun.Contains(Iter->Key))
 		{
diff --git a/Engine/Plugins/Experimental/Avalanche/Source/AvalancheMedia/Private/Playback/Nodes/Events/Actions/AvaPlaybackNode_PlayAnim.h b/Engine/Plugins/Experimental/Avalanche/Source/AvalancheMedia/Private/Playback/Nodes/Events/Actions/AvaPlaybackNode_PlayAnim.h
--- a/Engine/Plugins/Experimental/Avalanche/Source/AvalancheMedia/Private/Playback/Nodes/Events/Actions/AvaPlaybackNode_PlayAnim.h
+++ b/Engine/Plugins/Experimental/Avalanche/Source/AvalancheMedia/Private/Playback/Nodes/Events/Actions/AvaPlaybackNode_PlayAnim.h
@@ -3,7 +3,8 @@
 #pragma once
 
 #include "CoreMinimal.h"
-#include "Playback/Nodes/Events/Actions/AvalancheAnimations.h"
+#include "Playback/Nodes/Events/Actions/AvaPlaybackAnimations.h"
+#include "Playback/Nodes/Events/AvaPlaybackNodeAction.h"
 #include "AvaPlaybackNode_PlayAnim.generated.h"
 
 UCLASS()
@@ -23,7 +24,7 @@ public:
 protected:
 	
 	UPROPERTY(VisibleAnywhere, Category = "Motion Design")
-	TMap<FSoftObjectPath, FAvalancheAnimations> AnimationMap;
+	TMap<FSoftObjectPath, FAvaPlaybackAnimations> AnimationMap;
 
 	//All the Blueprints connected to this Node from the Dry Run (only populated while Dry Running!)
 	TSet<FSoftObjectPath> SeenAssetsInDryRun;
diff --git a/Engine/Plugins/Experimental/Avalanche/Source/AvalancheMedia/Private/Playback/Nodes/Events/Actions/AvaPlaybackNode_PreloadPlayer.cpp b/Engine/Plugins/Experimental/Avalanche/Source/AvalancheMedia/Private/Playback/Nodes/Events/Actions/AvaPlaybackNode_PreloadPlayer.cpp
--- a/Engine/Plugins/Experimental/Avalanche/Source/AvalancheMedia/Private/Playback/Nodes/Events/Actions/AvaPlaybackNode_PreloadPlayer.cpp
+++ b/Engine/Plugins/Experimental/Avalanche/Source/AvalancheMedia/Private/Playback/Nodes/Events/Actions/AvaPlaybackNode_PreloadPlayer.cpp
@@ -1,10 +1,11 @@
 // Copyright Epic Games, Inc. All Rights Reserved.
 
 #include "AvaPlaybackNode_PreloadPlayer.h"
-#include "AvalancheBroadcast.h"
-#include "Playback/AvalanchePlayback.h"
 
-#define LOCTEXT_NAMESPACE "AvalanchePlayback"
+#include "Broadcast/AvaBroadcast.h"
+#include "Playback/AvaPlaybackGraph.h"
+
+#define LOCTEXT_NAMESPACE "AvaPlaybackNode_PreloadPlayer"
 
 FText UAvaPlaybackNode_PreloadPlayer::GetNodeDisplayNameText() const
 {
@@ -18,15 +19,15 @@ FText UAvaPlaybackNode_PreloadPlayer::GetNodeTooltipText() const
 
 void UAvaPlaybackNode_PreloadPlayer::OnEventTriggered(const FAvaPlaybackEventParameters& InEventParameters)
 {
-	UAvalanchePlayback* const Playback = GetPlayback();
+	UAvaPlaybackGraph* const Playback = GetPlayback();
 	
-	//Make sure we have a valid Avalanche Asset
+	//Make sure we have a valid Motion Design Asset
 	if (Playback && InEventParameters.IsAssetValid())
 	{
 		TArray<FName> ChannelNames = Playback->GetChannelNamesForIndices(InEventParameters.ChannelIndices);
 		for (const FName& ChannelName : ChannelNames)
 		{
-			Playback->LoadAsset(InEventParameters.AvalancheAsset, ChannelName);
+			Playback->LoadAsset(InEventParameters.Asset, ChannelName);
 		}
 	}
 }
diff --git a/Engine/Plugins/Experimental/Avalanche/Source/AvalancheMedia/Private/Playback/Nodes/Events/FlowControl/AvaPlaybackNode_Hub.cpp b/Engine/Plugins/Experimental/Avalanche/Source/AvalancheMedia/Private/Playback/Nodes/Events/FlowControl/AvaPlaybackNode_Hub.cpp
--- a/Engine/Plugins/Experimental/Avalanche/Source/AvalancheMedia/Private/Playback/Nodes/Events/FlowControl/AvaPlaybackNode_Hub.cpp
+++ b/Engine/Plugins/Experimental/Avalanche/Source/AvalancheMedia/Private/Playback/Nodes/Events/FlowControl/AvaPlaybackNode_Hub.cpp
@@ -1,9 +1,8 @@
 // Copyright Epic Games, Inc. All Rights Reserved.
 
-
 #include "AvaPlaybackNode_Hub.h"
 
-#define LOCTEXT_NAMESPACE "AvalanchePlayback"
+#define LOCTEXT_NAMESPACE "AvaPlaybackNode_Hub"
 
 FText UAvaPlaybackNode_Hub::GetNodeDisplayNameText() const
 {
@@ -22,7 +21,7 @@ void UAvaPlaybackNode_Hub::TickEvent(float DeltaTime, FAvaPlaybackEventParameter
 		if (ChildNodes.IsValidIndex(Index) && ChildNodes[Index])
 		{
 			FAvaPlaybackEventParameters EventParameters;
-			EventParameters.AvalancheAsset = OutEventParameters.AvalancheAsset;
+			EventParameters.Asset = OutEventParameters.Asset;
 			
 			Cast<UAvaPlaybackNodeEvent>(ChildNodes[Index])->TickEvent(DeltaTime, EventParameters);
 			
diff --git a/Engine/Plugins/Experimental/Avalanche/Source/AvalancheMedia/Private/Playback/Nodes/Events/Triggers/AvaPlaybackNode_BeginPlay.cpp b/Engine/Plugins/Experimental/Avalanche/Source/AvalancheMedia/Private/Playback/Nodes/Events/Triggers/AvaPlaybackNode_BeginPlay.cpp
--- a/Engine/Plugins/Experimental/Avalanche/Source/AvalancheMedia/Private/Playback/Nodes/Events/Triggers/AvaPlaybackNode_BeginPlay.cpp
+++ b/Engine/Plugins/Experimental/Avalanche/Source/AvalancheMedia/Private/Playback/Nodes/Events/Triggers/AvaPlaybackNode_BeginPlay.cpp
@@ -2,7 +2,7 @@
 
 #include "AvaPlaybackNode_BeginPlay.h"
 
-#define LOCTEXT_NAMESPACE "AvalanchePlayback"
+#define LOCTEXT_NAMESPACE "AvaPlaybackNode_BeginPlay"
 
 FText UAvaPlaybackNode_BeginPlay::GetNodeDisplayNameText() const
 {
diff --git a/Engine/Plugins/Experimental/Avalanche/Source/AvalancheMedia/Private/Playback/Nodes/Events/Triggers/AvaPlaybackNode_KeyInput.cpp b/Engine/Plugins/Experimental/Avalanche/Source/AvalancheMedia/Private/Playback/Nodes/Events/Triggers/AvaPlaybackNode_KeyInput.cpp
--- a/Engine/Plugins/Experimental/Avalanche/Source/AvalancheMedia/Private/Playback/Nodes/Events/Triggers/AvaPlaybackNode_KeyInput.cpp
+++ b/Engine/Plugins/Experimental/Avalanche/Source/AvalancheMedia/Private/Playback/Nodes/Events/Triggers/AvaPlaybackNode_KeyInput.cpp
@@ -1,11 +1,12 @@
 // Copyright Epic Games, Inc. All Rights Reserved.
 
 #include "AvaPlaybackNode_KeyInput.h"
+
 #include "Framework/Application/SlateApplication.h"
 #include "Input/Events.h"
-#include "Playback/AvalanchePlayback.h"
+#include "Playback/AvaPlaybackGraph.h"
 
-#define LOCTEXT_NAMESPACE "AvalanchePlayback"
+#define LOCTEXT_NAMESPACE "AvaPlaybackNode_KeyInput"
 
 bool UAvaPlaybackNode_KeyInput::FEventInputProcessor::HandleKeyDownEvent(FSlateApplication& SlateApp, const FKeyEvent& InKeyEvent)
 {
@@ -70,7 +71,7 @@ FText UAvaPlaybackNode_KeyInput::GetNodeTooltipText() const
 bool UAvaPlaybackNode_KeyInput::HandleKeyDownEvent(FSlateApplication& SlateApp, const FKeyEvent& InKeyEvent)
 {
 	//Dont process input while not Playing
-	UAvalanchePlayback* const Playback = GetPlayback();
+	UAvaPlaybackGraph* const Playback = GetPlayback();
 	if (!Playback || !Playback->IsPlaying())
 	{
 		return false;
diff --git a/Engine/Plugins/Experimental/Avalanche/Source/AvalancheMedia/Private/Rundown/AvaRundownManagedInstanceBlueprint.cpp b/Engine/Plugins/Experimental/Avalanche/Source/AvalancheMedia/Private/Rundown/AvaRundownManagedInstanceBlueprint.cpp
new file mode 100644
--- /dev/null
+++ b/Engine/Plugins/Experimental/Avalanche/Source/AvalancheMedia/Private/Rundown/AvaRundownManagedInstanceBlueprint.cpp
@@ -0,0 +1,142 @@
+// Copyright Epic Games, Inc. All Rights Reserved.
+
+#include "Rundown/AvaRundownManagedInstanceBlueprint.h"
+
+#include "AssetRegistry/AssetData.h"
+#include "AvaBlueprint.h"
+#include "AvaRemoteControlRebind.h"
+#include "AvaRemoteControlUtils.h"
+#include "Engine/Engine.h"
+#include "Engine/World.h"
+#include "IAvaMediaModule.h"
+#include "RemoteControlPreset.h"
+#include "Rundown/AvaRundownManagedInstanceUtils.h"
+#include "UObject/Package.h"
+
+namespace UE::AvaRundownManagedInstanceBlueprint::Private
+{
+	UWorld* CreateManagedInstanceWorld(UPackage* InOuter, EWorldType::Type InWorldType)
+	{
+		const FName WorldName = MakeUniqueObjectName(InOuter, UWorld::StaticClass(), TEXT("MotionDesignManagedInstanceWorld"));
+
+		// Note: this will initialize the world.
+		UWorld* World = UWorld::CreateWorld(InWorldType, false, WorldName, InOuter, false);
+
+		if (InOuter == nullptr)
+		{
+			UPackage* const Package = World->GetPackage();
+			Package->SetFlags(RF_Transient | RF_Public);
+		}
+
+		FWorldContext& WorldContext = GEngine->CreateNewWorldContext(InWorldType);
+		WorldContext.SetCurrentWorld(World);
+		return World;
+	}
+	
+	UAvalancheBlueprint* LoadBlueprint(const FSoftObjectPath& InAssetPath)
+	{
+		return Cast<UAvalancheBlueprint>(InAssetPath.TryLoad());
+	}
+}
+
+FAvaRundownManagedInstanceBlueprint::FAvaRundownManagedInstanceBlueprint(FAvaRundownManagedInstanceCache* InParentCache, const FSoftObjectPath& InAssetPath)
+	: FAvaRundownManagedInstance(InParentCache, InAssetPath)
+{
+	using namespace UE::AvaRundownManagedInstanceBlueprint::Private;
+	SourceBlueprint = LoadBlueprint(InAssetPath);
+	if (!SourceBlueprint)
+	{
+		UE_LOG(LogAvaMedia, Error, TEXT("Unable to load Source Motion Design Blueprint: %s"), *InAssetPath.ToString());
+		return;
+	}
+	
+	RegisterSourceRemoteControlPresetDelegates(SourceBlueprint->GetRemoteControlPreset());
+
+#if WITH_EDITORONLY_DATA
+	FGuardValue_Bitfield(SourceBlueprint->bDuplicatingReadOnly, true);
+#endif
+
+	ManagedBlueprintPackage = FAvaRundownManagedInstanceUtils::MakeManagedInstancePackage(InAssetPath);
+	if (!ManagedBlueprintPackage)
+	{
+		UE_LOG(LogAvaMedia, Error, TEXT("Unable to create a Managed Motion Design Blueprint Package for %s"), *InAssetPath.ToString());
+		return;
+	}
+
+	// Remark: using PIE duplication mode to avoid entity and controller Ids from being renewed.
+	ManagedBlueprint = Cast<UAvalancheBlueprint>(StaticDuplicateObject(SourceBlueprint.Get(), ManagedBlueprintPackage.Get(), NAME_None, RF_NoFlags, nullptr, EDuplicateMode::PIE));
+	if (!ManagedBlueprint)
+	{
+		UE_LOG(LogAvaMedia, Error, TEXT("Unable to duplicate Source Motion Design Blueprint: %s"), *InAssetPath.ToString());
+		return;
+	}
+
+	ManagedBlueprint->SetFlags(RF_Public | RF_Transient);
+
+	ManagedWorld = CreateManagedInstanceWorld(ManagedBlueprintPackage.Get(), EWorldType::Type::Game);
+	ManagedBlueprint->SetAvalancheWorld(ManagedWorld.Get());
+	ManagedBlueprint->LoadAvalancheWorld();
+	FAvaRemoteControlRebind::RebindUnboundEntities(ManagedBlueprint->GetRemoteControlPreset());
+
+	// Backup the remote control values from the source.
+	constexpr bool bIsDefault = true;	// Flag the values as "default".
+	DefaultRemoteControlValues.CopyFrom(ManagedBlueprint->GetRemoteControlPreset(), bIsDefault);
+
+	FAvaRundownManagedInstanceUtils::PreventWorldFromBeingSeenAsLeakingByLevelEditor(ManagedWorld.Get());
+}
+
+FAvaRundownManagedInstanceBlueprint::~FAvaRundownManagedInstanceBlueprint()
+{
+	if (ManagedBlueprintPackage)
+	{
+		ManagedBlueprintPackage->ClearDirtyFlag();
+	}
+
+	// Simplified version of UAvaGameInstance::EndPlay()
+	if (ManagedWorld)
+	{
+		GEngine->DestroyWorldContext(ManagedWorld.Get());
+		ManagedWorld->DestroyWorld(true);
+	}
+
+	ManagedWorld = nullptr;
+	if (ManagedBlueprint && ManagedBlueprint->GetRemoteControlPreset())
+	{
+		FAvaRemoteControlUtils::UnregisterRemoteControlPreset(ManagedBlueprint->GetRemoteControlPreset());
+	}
+	ManagedBlueprint = nullptr;
+
+	if (SourceBlueprint)
+	{
+		UnregisterSourceRemoteControlPresetDelegates(SourceBlueprint->GetRemoteControlPreset());
+	}
+	SourceBlueprint = nullptr;
+}
+
+void FAvaRundownManagedInstanceBlueprint::AddReferencedObjects( FReferenceCollector& Collector )
+{
+	Collector.AddReferencedObject( SourceBlueprint );
+	Collector.AddReferencedObject( ManagedWorld );
+	Collector.AddReferencedObject( ManagedBlueprint );
+	Collector.AddReferencedObject( ManagedBlueprintPackage );
+}
+
+FString FAvaRundownManagedInstanceBlueprint::GetReferencerName() const
+{
+	return TEXT("FAvaRundownManagedInstanceBlueprint");
+}
+
+IAvaSceneInterface* FAvaRundownManagedInstanceBlueprint::GetSceneInterface() const
+{
+	return ManagedBlueprint ? static_cast<IAvaSceneInterface*>(ManagedBlueprint.Get()) : nullptr;
+}
+
+bool FAvaRundownManagedInstanceBlueprint::IsValid() const
+{
+	return ::IsValid(ManagedBlueprint) && ::IsValid(ManagedWorld);
+}
+
+URemoteControlPreset* FAvaRundownManagedInstanceBlueprint::GetRemoteControlPreset() const
+{
+	return IsValid() ? ManagedBlueprint->GetRemoteControlPreset() : nullptr;
+}
\ No newline at end of file
diff --git a/Engine/Plugins/Experimental/Avalanche/Source/AvalancheMedia/Private/Rundown/AvaRundownManagedInstanceBlueprint.h b/Engine/Plugins/Experimental/Avalanche/Source/AvalancheMedia/Private/Rundown/AvaRundownManagedInstanceBlueprint.h
new file mode 100644
--- /dev/null
+++ b/Engine/Plugins/Experimental/Avalanche/Source/AvalancheMedia/Private/Rundown/AvaRundownManagedInstanceBlueprint.h
@@ -0,0 +1,39 @@
+// Copyright Epic Games, Inc. All Rights Reserved.
+
+#pragma once
+
+#include "Rundown/AvaRundownManagedInstance.h"
+#include "UObject/StrongObjectPtr.h"
+
+class UAvalancheBlueprint;
+
+/**
+ * Container for a fully instanced Motion Design Blueprint.
+ * The World is not fully initialized, but contains all loaded actors.
+ * This is required for RemoteControl bindings to be resolved.
+ */
+class FAvaRundownManagedInstanceBlueprint : public FAvaRundownManagedInstance
+{
+public:
+	UE_AVA_INHERITS(FAvaRundownManagedInstanceBlueprint, FAvaRundownManagedInstance);
+	
+	FAvaRundownManagedInstanceBlueprint(FAvaRundownManagedInstanceCache* InParentCache, const FSoftObjectPath& InAssetPath);
+	virtual ~FAvaRundownManagedInstanceBlueprint() override;
+
+	//~ Begin FGCObject
+	virtual void AddReferencedObjects( FReferenceCollector& Collector ) override;
+	virtual FString GetReferencerName() const override;
+	virtual IAvaSceneInterface* GetSceneInterface() const override;
+	//~ End FGCObject
+
+	//~ Begin FAvaRundownManagedInstance
+	virtual bool IsValid() const override;
+	virtual URemoteControlPreset* GetRemoteControlPreset() const override;
+	//~ End FAvaRundownManagedInstance
+
+private:
+	TObjectPtr<UAvalancheBlueprint> SourceBlueprint;
+	TObjectPtr<UWorld> ManagedWorld;
+	TObjectPtr<UAvalancheBlueprint> ManagedBlueprint;
+	TObjectPtr<UPackage> ManagedBlueprintPackage;
+};
diff --git a/Engine/Plugins/Experimental/Avalanche/Source/AvalancheMedia/Private/Rundown/AvaRundownPageAssetUtils.cpp b/Engine/Plugins/Experimental/Avalanche/Source/AvalancheMedia/Private/Rundown/AvaRundownPageAssetUtils.cpp
new file mode 100644
--- /dev/null
+++ b/Engine/Plugins/Experimental/Avalanche/Source/AvalancheMedia/Private/Rundown/AvaRundownPageAssetUtils.cpp
@@ -0,0 +1,95 @@
+// Copyright Epic Games, Inc. All Rights Reserved.
+
+#include "Rundown/AvaRundownPageAssetUtils.h"
+
+#include "AvaScene.h"
+#include "AvaTagHandle.h"
+#include "AvaTransitionSubsystem.h"
+#include "AvaTransitionTree.h"
+#include "Behavior/IAvaTransitionBehavior.h"
+#include "Engine/Level.h"
+#include "Engine/World.h"
+#include "IAvaSceneInterface.h"
+
+namespace UE::AvaRundownPageAssetUtils::Private
+{
+	template<typename InInterfaceType>
+	const InInterfaceType* GetInterfaceFromAsset(const UObject* InLoadedAsset)
+	{
+		if (const InInterfaceType* const Interface = Cast<InInterfaceType>(InLoadedAsset))
+		{
+			return Interface;
+		}
+		if (const UWorld* const World = Cast<UWorld>(InLoadedAsset))
+		{
+			if (const AAvaScene* AvaScene = AAvaScene::GetScene(World->PersistentLevel, false))
+			{
+				return static_cast<const InInterfaceType*>(AvaScene);
+			}
+		}
+		return nullptr;
+	}
+}
+
+const IAvaSceneInterface* FAvaRundownPageAssetUtils::GetSceneInterface(const UObject* InLoadedAsset)
+{
+	using namespace UE::AvaRundownPageAssetUtils::Private;
+	return GetInterfaceFromAsset<IAvaSceneInterface>(InLoadedAsset);
+}
+
+const UAvaTransitionTree* FAvaRundownPageAssetUtils::FindTransitionTree(const IAvaSceneInterface* InSceneInterface)
+{
+	if (!InSceneInterface)
+	{
+		return nullptr;
+	}
+	
+	ULevel* SceneLevel = InSceneInterface->GetSceneLevel();
+	if (!SceneLevel)
+	{
+		return nullptr;
+	}
+
+	// Method 1 - Using the UAvaTransitionSubsystem
+	// This method may not work if the world is not initialized (inactive, as in the managed instance).
+	if (const UWorld* SceneWorld = SceneLevel->GetWorld())
+	{
+		if (const UAvaTransitionSubsystem* TransitionSubsystem = SceneWorld->GetSubsystem<UAvaTransitionSubsystem>())
+		{
+			if (const IAvaTransitionBehavior* TransitionBehavior = TransitionSubsystem->GetTransitionBehavior(SceneLevel))
+			{
+				if (const UAvaTransitionTree* TransitionTree = TransitionBehavior->GetTransitionTree())
+				{
+					return TransitionTree;
+				}
+			}
+			// If we had a subsystem and it didn't find the behavior then there was no state tree either and we are done.
+			return nullptr;
+		}
+	}
+	
+	// Method 2 - Fallback using direct lookup of the private actor class.
+	// This is needed to lookup in managed instances and source asset.
+
+	static FString AvaTransitionBehaviorActor(TEXT("AvaTransitionBehaviorActor"));
+	for (const AActor* Actor : SceneLevel->Actors)
+	{
+		const UClass* ActorClass = Actor->GetClass();
+		if (ActorClass && ActorClass->GetName() == AvaTransitionBehaviorActor)
+		{
+			const FObjectPropertyBase* const Property = FindFProperty<FObjectPropertyBase>(ActorClass, TEXT("TransitionTree"));
+			if (Property
+				&& Property->PropertyClass
+				&& Property->PropertyClass->IsChildOf(UAvaTransitionTree::StaticClass()))
+			{
+				return Cast<UAvaTransitionTree>(Property->GetObjectPropertyValue_InContainer(Actor));
+			}
+		}
+	}
+	return nullptr;
+}
+
+FAvaTagHandle FAvaRundownPageAssetUtils::GetTransitionLayerTag(const UAvaTransitionTree* InTransitionTree)
+{
+	return InTransitionTree ? InTransitionTree->GetTransitionLayer() : FAvaTagHandle();
+}
\ No newline at end of file
diff --git a/Engine/Plugins/Experimental/Avalanche/Source/AvalancheMedia/Public/AvaMediaDefines.h b/Engine/Plugins/Experimental/Avalanche/Source/AvalancheMedia/Public/AvaMediaDefines.h
--- a/Engine/Plugins/Experimental/Avalanche/Source/AvalancheMedia/Public/AvaMediaDefines.h
+++ b/Engine/Plugins/Experimental/Avalanche/Source/AvalancheMedia/Public/AvaMediaDefines.h
@@ -2,12 +2,12 @@
 
 #pragma once
 
-#include "Misc/EnumClassFlags.h"
 #include "Framework/AvaSoftAssetPtr.h"
+#include "Misc/EnumClassFlags.h"
 #include "AvaMediaDefines.generated.h"
 
 UENUM()
-enum class EAvaChannelState : uint8
+enum class EAvaBroadcastChannelState : uint8
 {
 	Offline,
 	Idle,
@@ -42,7 +42,7 @@ enum class EAvaBroadcastChannelType : uint8
  *	Status of the media output.
  */
 UENUM()
-enum class EAvaMediaOutputState : uint8
+enum class EAvaBroadcastOutputState : uint8
 {
 	/** Invalid/Uninitialized state. */
 	Invalid,
@@ -59,7 +59,7 @@ enum class EAvaMediaOutputState : uint8
 };
 
 UENUM()
-enum class EAvaMediaIssueSeverity : uint8
+enum class EAvaBroadcastIssueSeverity : uint8
 {
 	None,
 	Warnings,
@@ -82,7 +82,7 @@ enum class EAvaBroadcastChange : uint8
 ENUM_CLASS_FLAGS(EAvaBroadcastChange);
 
 //An enum indicating what changed in Channel
-enum class EAvaChannelChange : uint8
+enum class EAvaBroadcastChannelChange : uint8
 {
 	None         = 0,
 	State        = 1 << 0,
@@ -90,7 +90,7 @@ enum class EAvaChannelChange : uint8
 	MediaOutputs = 1 << 2,
 	Settings     = 1 << 3,
 };
-ENUM_CLASS_FLAGS(EAvaChannelChange);
+ENUM_CLASS_FLAGS(EAvaBroadcastChannelChange);
 
 /**
  * The status of a playable object.
@@ -100,7 +100,7 @@ ENUM_CLASS_FLAGS(EAvaChannelChange);
  * referred to as the "runtime" asset instance.
  */
 UENUM()
-enum class EAvalanchePlayableStatus : uint8
+enum class EAvaPlayableStatus : uint8
 {
 	Unknown,
 	Error,
@@ -111,7 +111,7 @@ enum class EAvalanchePlayableStatus : uint8
 };
 
 UENUM()
-enum class EAvalanchePlayableSequenceEventType : uint8
+enum class EAvaPlayableSequenceEventType : uint8
 {
 	None,
 	Started,
@@ -121,15 +121,15 @@ enum class EAvalanchePlayableSequenceEventType : uint8
 /**
  * Transition flags carry additional context to help resolve the behaviors.
  */
-enum class EAvalanchePlayableTransitionFlags : uint8
+enum class EAvaPlayableTransitionFlags : uint8
 {
 	None = 0,
 	/** This transition contains only exit playables. */
 	ExitOnly = 1 << 0,
 };
-ENUM_CLASS_FLAGS(EAvalanchePlayableTransitionFlags);
+ENUM_CLASS_FLAGS(EAvaPlayableTransitionFlags);
 
-enum class EAvalanchePlayableTransitionEventFlags : uint8
+enum class EAvaPlayableTransitionEventFlags : uint8
 {
 	None = 0,
 	/** The transition is starting. */
@@ -141,7 +141,7 @@ enum class EAvalanchePlayableTransitionEventFlags : uint8
 	/** The transition is finished and can be cleaned up. */
 	Finished = 1 << 3,
 };
-ENUM_CLASS_FLAGS(EAvalanchePlayableTransitionEventFlags);
+ENUM_CLASS_FLAGS(EAvaPlayableTransitionEventFlags);
 
 /*
  * Parameters used to tell a Channel its Playback Settings
@@ -155,11 +155,11 @@ struct FAvaPlaybackChannelParameters
 	 * If a playback node is reached during tick traversal,
 	 * it is going to set this value.
 	 */
-	TArray<FAvaSoftAssetPtr> AvalancheAssets;
+	TArray<FAvaSoftAssetPtr> Assets;
 
 	bool HasAssets() const
 	{
-		return !AvalancheAssets.IsEmpty();
+		return !Assets.IsEmpty();
 	}
 };
 
@@ -171,16 +171,16 @@ struct FAvaPlaybackEventParameters
 	/** List of channel indices that this event originates from. */
 	TArray<int32> ChannelIndices;
 
-	FAvaSoftAssetPtr AvalancheAsset;
+	FAvaSoftAssetPtr Asset;
 
 	bool IsAssetValid() const
 	{
-		return !AvalancheAsset.IsNull();
+		return !Asset.IsNull();
 	}
 
 	const FSoftObjectPath& GetAssetPath() const
 	{
-		return AvalancheAsset.ToSoftObjectPath();
+		return Asset.ToSoftObjectPath();
 	}
 	
 	void RequestTriggerEventAction()
@@ -199,7 +199,7 @@ protected:
 };
 
 
-enum class EAvaPageListChange : uint8
+enum class EAvaRundownPageListChange : uint8
 {
 	None = 0,
 
@@ -211,9 +211,9 @@ enum class EAvaPageListChange : uint8
 
 	All              = 0xFF,
 };
-ENUM_CLASS_FLAGS(EAvaPageListChange);
+ENUM_CLASS_FLAGS(EAvaRundownPageListChange);
 
-enum class EAvaPageChanges : uint8
+enum class EAvaRundownPageChanges : uint8
 {
 	None = 0,
 
@@ -225,7 +225,7 @@ enum class EAvaPageChanges : uint8
 
 	All              = 0xFF,
 };
-ENUM_CLASS_FLAGS(EAvaPageChanges);
+ENUM_CLASS_FLAGS(EAvaRundownPageChanges);
 
 /**
  *	The status of the playback asset on disk.
@@ -235,7 +235,7 @@ ENUM_CLASS_FLAGS(EAvaPageChanges);
  *	the source Motion Design blueprint, rather than a playback object.
  */
 UENUM()
-enum class EAvaMediaPlaybackAssetStatus
+enum class EAvaPlaybackAssetStatus
 {
 	/** Invalid status. */
 	Unknown,
@@ -245,11 +245,11 @@ enum class EAvaMediaPlaybackAssetStatus
 	 * Missing asset dependencies
 	 * Note: An Motion Design asset can run even with missing dependencies.
 	 */
-	 MissingDependencies,
-	 /** Asset is out of date from compare with remote. */
-	 NeedsSync,
-	 /** Asset is fully available and up to date. */
-	 Available
+	MissingDependencies,
+	/** Asset is out of date from compare with remote. */
+	NeedsSync,
+	/** Asset is fully available and up to date. */
+	Available
 };
 
 /**
@@ -260,7 +260,7 @@ enum class EAvaMediaPlaybackAssetStatus
  * referred to as the "managed" or "runtime" asset.
  */
 UENUM()
-enum class EAvaMediaPlaybackStatus
+enum class EAvaPlaybackStatus
 {
 	/** Invalid status. */
 	Unknown,
@@ -286,33 +286,33 @@ enum class EAvaMediaPlaybackStatus
 	Error
 };
 
-UENUM(BlueprintType)
-enum class EAvaPageListType : uint8
+UENUM(BlueprintType, DisplayName = "Motion Design Rundown Page List Type")
+enum class EAvaRundownPageListType : uint8
 {
 	Template,
 	Instance,
 	View
 };
 
-USTRUCT(BlueprintType)
-struct FAvaPageListReference
+USTRUCT(BlueprintType, DisplayName = "Motion Design Rundown Page List Reference")
+struct FAvaRundownPageListReference
 {
 	GENERATED_BODY()
 
 	UPROPERTY()
-	EAvaPageListType Type = EAvaPageListType::Instance;
+	EAvaRundownPageListType Type = EAvaRundownPageListType::Instance;
 
 	UPROPERTY()
 	int32 SubListIndex = INDEX_NONE;
 
-	bool operator==(const FAvaPageListReference& InOther) const
+	bool operator==(const FAvaRundownPageListReference& InOther) const
 	{
 		if (Type != InOther.Type)
 		{
 			return false;
 		}
 
-		if (Type == EAvaPageListType::View && SubListIndex != InOther.SubListIndex)
+		if (Type == EAvaRundownPageListType::View && SubListIndex != InOther.SubListIndex)
 		{
 			return false;
 		}
@@ -320,5 +320,5 @@ struct FAvaPageListReference
 		return true;
 	}
 
-	bool operator!=(const FAvaPageListReference& InOther) const { return !(*this == InOther); }
+	bool operator!=(const FAvaRundownPageListReference& InOther) const { return !(*this == InOther); }
 };
diff --git a/Engine/Plugins/Experimental/Avalanche/Source/AvalancheMedia/Public/IAvaMediaModule.h b/Engine/Plugins/Experimental/Avalanche/Source/AvalancheMedia/Public/IAvaMediaModule.h
--- a/Engine/Plugins/Experimental/Avalanche/Source/AvalancheMedia/Public/IAvaMediaModule.h
+++ b/Engine/Plugins/Experimental/Avalanche/Source/AvalancheMedia/Public/IAvaMediaModule.h
@@ -8,14 +8,14 @@
 
 AVALANCHEMEDIA_API DECLARE_LOG_CATEGORY_EXTERN(LogAvaMedia, Log, All);
 
-class FAvaMediaPlaybackManager;
-class FAvaMediaPlaybackServer;
-class FAvalancheManagedInstanceCache;
+class FAvaPlaybackManager;
+class FAvaPlaybackServer;
+class FAvaRundownManagedInstanceCache;
 class FCommonViewportClient;
-class IAvaDeviceProviderProxyManager;
-class IAvaMediaPlaybackClient;
+class IAvaBroadcastDeviceProviderProxyManager;
+class IAvaBroadcastSettings;
 class IAvaMediaSyncProvider;
-class IAvalancheBroadcastSettings;
+class IAvaPlaybackClient;
 class IMediaIOCoreDeviceProvider;
 class UWorld;
 struct FAvaInstanceSettings;
@@ -77,8 +77,8 @@ public:
 	 */
 	virtual void StopMediaPlaybackServer() = 0;
 
-	virtual IAvaMediaPlaybackClient& GetMediaPlaybackClient() = 0;
-	virtual TSharedPtr<FAvaMediaPlaybackServer> GetMediaPlaybackServer() const = 0;
+	virtual IAvaPlaybackClient& GetMediaPlaybackClient() = 0;
+	virtual TSharedPtr<FAvaPlaybackServer> GetMediaPlaybackServer() const = 0;
 	virtual const IMediaIOCoreDeviceProvider* GetDeviceProvider(FName InProviderName, const FMediaIOOutputConfiguration* InMediaIOOutputConfiguration) const = 0;
 	virtual TArray<const IMediaIOCoreDeviceProvider*> GetDeviceProvidersForServer(const FString& InServerName) const = 0;
 	virtual FString GetServerNameForDevice(const FName& InDeviceProviderName, const FName& InDeviceName) const = 0;
@@ -102,10 +102,10 @@ public:
 	/**
 	 * Access the global broadcast settings.
 	 */
-	virtual const IAvalancheBroadcastSettings& GetBroadcastSettings() const = 0;
+	virtual const IAvaBroadcastSettings& GetBroadcastSettings() const = 0;
 
 	/**
-	 * Access the global Avalanche instance settings.
+	 * Access the global Motion Design instance settings.
 	 * 
 	 * Remark: lifetime of the returned reference is not guaranteed beyond the current call context.
 	 * If the settings are replicated from a client, it could get deleted if the client disconnects.
@@ -115,14 +115,14 @@ public:
 	virtual const FAvaInstanceSettings& GetAvaInstanceSettings() const = 0;
 	
 	/**
-	 *	This is the backend for playing avalanche blueprints locally.
+	 *	This is the backend for playing Motion Design blueprints locally.
 	 */
-	virtual FAvaMediaPlaybackManager& GetLocalPlaybackManager() const = 0;
+	virtual FAvaPlaybackManager& GetLocalPlaybackManager() const = 0;
 
 	/**
-	 *	Access the "managed" Avalanche Asset Instance cache.
+	 *	Access the "managed" Motion Design Asset Instance cache.
 	 */
-	virtual FAvalancheManagedInstanceCache& GetManagedInstanceCache() const = 0;
+	virtual FAvaRundownManagedInstanceCache& GetManagedInstanceCache() const = 0;
 
 	/**
 	 * Returns true if the AvaMediaSyncProvider modular feature is available.
@@ -154,17 +154,17 @@ public:
 	DECLARE_MULTICAST_DELEGATE_TwoParams(FOnMapChangedEvent, UWorld* /*InWorld*/, EAvaMediaMapChangeType /*InEventType*/);
 	virtual FOnMapChangedEvent& GetOnMapChangedEvent() = 0;
 	
-	DECLARE_MULTICAST_DELEGATE(FOnAvaMediaPlaybackClientStarted);
-	virtual FOnAvaMediaPlaybackClientStarted& GetOnAvaMediaPlaybackClientStarted() = 0;
+	DECLARE_MULTICAST_DELEGATE(FOnAvaPlaybackClientStarted);
+	virtual FOnAvaPlaybackClientStarted& GetOnAvaPlaybackClientStarted() = 0;
 
-	DECLARE_MULTICAST_DELEGATE(FOnAvaMediaPlaybackClientStopped);
-	virtual FOnAvaMediaPlaybackClientStopped& GetAvaMediaPlaybackClientStopped() = 0;
+	DECLARE_MULTICAST_DELEGATE(FOnAvaPlaybackClientStopped);
+	virtual FOnAvaPlaybackClientStopped& GetAvaPlaybackClientStopped() = 0;
 
-	DECLARE_MULTICAST_DELEGATE(FOnAvaMediaPlaybackServerStarted);
-	virtual FOnAvaMediaPlaybackServerStarted& GetOnAvaMediaPlaybackServerStarted() = 0;
+	DECLARE_MULTICAST_DELEGATE(FOnAvaPlaybackServerStarted);
+	virtual FOnAvaPlaybackServerStarted& GetOnAvaPlaybackServerStarted() = 0;
 
-	DECLARE_MULTICAST_DELEGATE(FOnAvaMediaPlaybackServerStopped);
-	virtual FOnAvaMediaPlaybackServerStopped& GetAvaMediaPlaybackServerStopped() = 0;
+	DECLARE_MULTICAST_DELEGATE(FOnAvaPlaybackServerStopped);
+	virtual FOnAvaPlaybackServerStopped& GetAvaPlaybackServerStopped() = 0;
 	
 	/** Use to query the current editor viewport from the corresponding editor module. */
 	DECLARE_DELEGATE_OneParam(FGetEditorViewportClient, FCommonViewportClient** );
@@ -173,5 +173,5 @@ public:
 	/**
 	 * Access the device provider proxy manager.
 	 */
-	virtual IAvaDeviceProviderProxyManager& GetDeviceProviderProxyManager() = 0;
+	virtual IAvaBroadcastDeviceProviderProxyManager& GetDeviceProviderProxyManager() = 0;
 };
