Motion Design

diff --git a/Engine/Plugins/Experimental/Avalanche/Source/AvalancheMedia/Public/ModularFeature/IAvaMediaSyncProvider.h b/Engine/Plugins/Experimental/Avalanche/Source/AvalancheMedia/Public/ModularFeature/IAvaMediaSyncProvider.h
--- a/Engine/Plugins/Experimental/Avalanche/Source/AvalancheMedia/Public/ModularFeature/IAvaMediaSyncProvider.h
+++ b/Engine/Plugins/Experimental/Avalanche/Source/AvalancheMedia/Public/ModularFeature/IAvaMediaSyncProvider.h
@@ -2,17 +2,17 @@
 
 #pragma once
 
-#include "AvaSyncProviderFeatureTypes.h"
+#include "AvaMediaSyncProviderFeatureTypes.h"
 #include "Features/IModularFeature.h"
 
 /** Generic delegate response handler */
-DECLARE_DELEGATE_OneParam(FOnAvaSyncResponse, const TSharedPtr<FAvaSyncResponse>&);
+DECLARE_DELEGATE_OneParam(FOnAvaMediaSyncResponse, const TSharedPtr<FAvaMediaSyncResponse>&);
 
-/** Delegate with a comparison response. Payload is a FAvaSyncCompareResponse shared ptr */
-DECLARE_DELEGATE_OneParam(FOnAvaSyncCompareResponse, const TSharedPtr<FAvaSyncCompareResponse>&);
+/** Delegate with a comparison response. Payload is a FAvaMediaSyncCompareResponse shared ptr */
+DECLARE_DELEGATE_OneParam(FOnAvaMediaSyncCompareResponse, const TSharedPtr<FAvaMediaSyncCompareResponse>&);
 
 /** Delegate called when a package is modified by a sync operation, either added or modified. */
-DECLARE_MULTICAST_DELEGATE_OneParam(FOnAvaSyncPackageModified, const FName& /*PackageName*/);
+DECLARE_MULTICAST_DELEGATE_OneParam(FOnAvaMediaSyncPackageModified, const FName& /*PackageName*/);
 
 /**
  * Base Sync feature interface
@@ -28,7 +28,7 @@ public:
 
 	static FName GetModularFeatureName()
 	{
-		static const FName FeatureName = TEXT("Avalanche_Feature_Sync");
+		static const FName FeatureName = TEXT("MotionDesign_Feature_Sync");
 		return FeatureName;
 	}
 
@@ -53,9 +53,9 @@ public:
 	 *
 	 * @param InRemoteName The playback client or server name to push to
 	 * @param InPackageNames The list of "top level" package names (without inner dependencies / referencers) to send on remote
-	 * @param DoneDelegate A completion delegate called when receiving a response with a `FOnAvaSyncResponse` payload
+	 * @param DoneDelegate A completion delegate called when receiving a response with a `FOnAvaMediaSyncResponse` payload
 	 */
-	virtual void PushToRemote(const FString& InRemoteName, const TArray<FName>& InPackageNames, const FOnAvaSyncResponse& DoneDelegate = FOnAvaSyncResponse()) = 0;
+	virtual void PushToRemote(const FString& InRemoteName, const TArray<FName>& InPackageNames, const FOnAvaMediaSyncResponse& DoneDelegate = FOnAvaMediaSyncResponse()) = 0;
 
 	/**
 	 * Sends a "pull" request to a specific remote (either a playback server or client name)
@@ -64,9 +64,9 @@ public:
 	 *
 	 * @param InRemoteName The playback client or server name to push to
 	 * @param InPackageNames The list of "top level" package names (without inner dependencies / referencers) to send on remote
-	 * @param DoneDelegate A completion delegate called when receiving a response with a `FOnAvaSyncResponse` payload
+	 * @param DoneDelegate A completion delegate called when receiving a response with a `FOnAvaMediaSyncResponse` payload
 	 */
-	virtual void PullFromRemote(const FString& InRemoteName, const TArray<FName>& InPackageNames, const FOnAvaSyncResponse& DoneDelegate = FOnAvaSyncResponse()) = 0;
+	virtual void PullFromRemote(const FString& InRemoteName, const TArray<FName>& InPackageNames, const FOnAvaMediaSyncResponse& DoneDelegate = FOnAvaMediaSyncResponse()) = 0;
 
 	/**
 	 * Issue a "compare" request to a given remote, and calls back the completion delegate with status payload.
@@ -79,10 +79,10 @@ public:
 	 *
 	 * @param InRemoteName The playback client or server name to compare with
 	 * @param InPackageNames The list of "top level" package names (without inner dependencies / referencers) to base the comparison on
-	 * @param DoneDelegate A completion delegate called when receiving a response with a `FAvaSyncCompareResponse` payload
+	 * @param DoneDelegate A completion delegate called when receiving a response with a `FAvaMediaSyncCompareResponse` payload
 	 */
-	virtual void CompareWithRemote(const FString& InRemoteName, const TArray<FName>& InPackageNames, const FOnAvaSyncCompareResponse& DoneDelegate = FOnAvaSyncCompareResponse()) = 0;
+	virtual void CompareWithRemote(const FString& InRemoteName, const TArray<FName>& InPackageNames, const FOnAvaMediaSyncCompareResponse& DoneDelegate = FOnAvaMediaSyncCompareResponse()) = 0;
 
 	/** Delegate called when a package is modified by a sync write operation, either being added or modified. */
-	virtual FOnAvaSyncPackageModified& GetOnAvaSyncPackageModified() = 0;
+	virtual FOnAvaMediaSyncPackageModified& GetOnAvaSyncPackageModified() = 0;
 };
diff --git a/Engine/Plugins/Experimental/Avalanche/Source/AvalancheMedia/Public/Playable/AvaPlayableGroup.h b/Engine/Plugins/Experimental/Avalanche/Source/AvalancheMedia/Public/Playable/AvaPlayableGroup.h
new file mode 100644
--- /dev/null
+++ b/Engine/Plugins/Experimental/Avalanche/Source/AvalancheMedia/Public/Playable/AvaPlayableGroup.h
@@ -0,0 +1,222 @@
+// Copyright Epic Games, Inc. All Rights Reserved.
+
+#pragma once
+
+#include "Framework/AvaGameInstance.h"
+#include "Playable/IAvaPlayableVisibilityConstraint.h"
+#include "UObject/Object.h"
+#include "UObject/ObjectKey.h"
+#include "UObject/ObjectPtr.h"
+#include "UObject/SoftObjectPtr.h"
+#include "AvaPlayableGroup.generated.h"
+
+class UAvaPlayable;
+class UAvaPlayableGroupManager;
+class UAvaPlayableTransition;
+class UTextureRenderTarget2D;
+
+/**
+ * This defines a game instance playable group.
+ * It tracks and manage the game instance and playables state.
+ *
+ * The design goal of this class is to remove all the playback management
+ * from UAvaGameInstance and move it to the playable framework. This
+ * should allow us to hook the playback framework to any game instance, including PIE
+ * so it can work with any work flow (editor, PIE, game, nDisplay, etc).
+ *
+ * Also, ideally, the playable class itself should be "game instance" agnostic
+ * and do all it's bidding on it's container through this class.
+ */
+UCLASS()
+class AVALANCHEMEDIA_API UAvaPlayableGroup : public UObject
+{
+	GENERATED_BODY()
+
+public:
+	/**
+	 * PlayableGroup creation information contains the necessary information to
+	 * create an instance of a playable group.
+	 */
+	struct FPlayableGroupCreationInfo
+	{
+		/** Container for shared playable groups. */
+		UAvaPlayableGroupManager* PlayableGroupManager = nullptr;
+		/** Source asset path. */
+		FSoftObjectPath SourceAssetPath;
+		/** Channel name this playable will be instanced in. */
+		FName ChannelName;
+		/** Indicate if the group is for remote proxy playables. */
+		bool bIsRemoteProxy = false;
+		/** Indicate if the group is shared. */
+		bool bIsSharedGroup = false;
+	};
+	
+	static UAvaPlayableGroup* MakePlayableGroup(UObject* InOuter, const FPlayableGroupCreationInfo& InPlayableGroupInfo);
+
+	/** Register a playable to this group when it is created. */
+	void RegisterPlayable(UAvaPlayable* InPlayable);
+
+	/** Unregister a playable when it is about to be deleted. */
+	void UnregisterPlayable(UAvaPlayable* InPlayable);
+
+	/** Returns true if there are any valid registered playables. */
+	bool HasPlayables() const;
+
+	/** Returns true if there are any valid registered playables that are currently playing. */
+	bool HasPlayingPlayables() const;
+
+	/** Finds all the playables that are instances of the given source asset. */
+	void FindPlayablesBySourceAssetPath(const FSoftObjectPath& InSourceAssetPath, TArray<UAvaPlayable*>& OutFoundPlayables) const;
+
+	void RegisterPlayableTransition(UAvaPlayableTransition* InPlayableTransition);
+	
+	void UnregisterPlayableTransition(UAvaPlayableTransition* InPlayableTransition);
+
+	/** Tick transitions that have been registered. Returns the number of transitions that where ticked. */
+	void TickTransitions(double InDeltaSeconds);
+
+	bool HasTransitions() const;
+
+	/**
+	 * Creates the game instance's world if it wasn't already.
+	 * @return true if the world was created. false if nothing was done.
+	 */
+	bool ConditionalCreateWorld();
+
+	/**
+	 * Begin playing the game instance's world if it wasn't already.
+	 * @return true if the BeginPlay was done (i.e. on the state transition only), false otherwise.
+	 */
+	virtual bool ConditionalBeginPlay(const FAvaInstancePlaySettings& InWorldPlaySettings);
+
+	virtual void RequestEndPlayWorld(bool bInForceImmediate);
+
+	/** Keep track of the last playable that applied it's camera in the viewport/controller. */
+	void SetLastAppliedCameraPlayable(UAvaPlayable* InPlayable);
+	
+	/** Search for the first remaining playing playable and use it's camera. */
+	bool UpdateCameraSetup();
+
+	bool IsWorldPlaying() const { return GameInstance ? GameInstance->IsWorldPlaying() : false;}
+
+	bool IsRenderTargetReady() const { return GameInstance ? GameInstance->IsRenderTargetReady() : false;}
+
+	/**
+	 * Current logic for the render target, use the game instance's if there, fallback to internal one if not.
+	 */
+	UTextureRenderTarget2D* GetRenderTarget() const { return GameInstance ? GameInstance->GetRenderTarget() : RenderTarget.Get(); }
+
+	/** Return a vanilla game instance, we want to eventually support any game instance. */
+	UGameInstance* GetGameInstance() const { return GameInstance;}
+	
+	UWorld* GetPlayWorld() const { return GameInstance ? GameInstance->GetPlayWorld() : nullptr; }
+	
+	/**
+	 * Unloads the game instance's world if no more playables are loaded.
+	 * @return true if the world was unloaded. false if nothing was done.
+	 */ 
+	bool ConditionalRequestUnloadWorld(bool bForceImmediate);
+
+	/**
+	 * @brief Queue a camera cut for the next rendered frame.
+	 */
+	void QueueCameraCut();
+
+	/**
+	 * @brief Notify the playable group that a playable is loading an asset.
+	 * @param InPlayable Source of the event.
+	 */
+	void NotifyLevelStreaming(UAvaPlayable* InPlayable);
+
+	UAvaPlayableGroupManager* GetPlayableGroupManager() const { return ParentPlayableGroupManagerWeak.Get(); }
+
+	void RegisterVisibilityConstraint(const TWeakInterfacePtr<IAvaPlayableVisibilityConstraint>& InVisibilityConstraint);
+	void UnregisterVisibilityConstraint(const IAvaPlayableVisibilityConstraint* InVisibilityConstraint);
+
+	void RequestSetVisibility(UAvaPlayable* InPlayable, bool bInShouldBeVisible);
+
+protected:
+	bool IsVisibilityConstrained(const UAvaPlayable* InPlayable) const;
+
+	void OnPlayableStatusChanged(UAvaPlayable* InPlayable);
+	
+	void ConditionalRegisterWorldDelegates(UWorld* InWorld);
+	void UnregisterWorldDelegates(UWorld* InWorld);
+
+	bool DisplayLoadedAssets(FText& OutText, FLinearColor& OutColor);
+	bool DisplayPlayingAssets(FText& OutText, FLinearColor& OutColor);
+	bool DisplayTransitions(FText& OutText, FLinearColor& OutColor);
+
+	static UPackage* MakeGameInstancePackage(const FSoftObjectPath& InSourceAssetPath, const FName& InChannelName);
+	static UPackage* MakeSharedInstancePackage(const FName& InChannelName);
+	static UPackage* MakeInstancePackage(const FString& InInstancePackageName);
+
+public:	
+	UPROPERTY(Transient)
+	TObjectPtr<UTextureRenderTarget2D> RenderTarget;	// Optional
+	
+	UPROPERTY(Transient)
+	TObjectPtr<UAvaGameInstance> GameInstance;
+
+	UPROPERTY(Transient)
+	TObjectPtr<UPackage> GameInstancePackage;
+
+protected:
+	/** PlayableGroup Manager handling this playable group. */
+	TWeakObjectPtr<UAvaPlayableGroupManager> ParentPlayableGroupManagerWeak;
+	
+	/** List of playables for this group. */
+	TSet<TObjectKey<UAvaPlayable>> Playables;
+
+	/** Last playable that applied a camera. */
+	TWeakObjectPtr<UAvaPlayable> LastAppliedCameraPlayableWeak;
+
+	/** Set of registered playable transitions for this group. Remark: used for ticking. */
+	TSet<TObjectKey<UAvaPlayableTransition>> PlayableTransitions;
+
+	/** If transitions are added or removed while ticking, we need to protect transition iterator. */
+	bool bIsTickingTransitions = false;
+
+	/** Set of transitions to remove accumulated during transition ticking. */
+	TSet<TObjectKey<UAvaPlayableTransition>> PlayableTransitionsToRemove;
+
+	/** Set of transitions to add accumulated during transition ticking. */
+	TSet<TObjectKey<UAvaPlayableTransition>> PlayableTransitionsToAdd;
+
+	/**
+	 * Since the UViewportStatsSubsystem delegation mechanism does not allow us to verify
+	 * if it is bound to this world or another one. We need this auxiliary binding
+	 * tracking to compensate.
+	 */
+	TWeakObjectPtr<UWorld> LastWorldBoundToDisplayDelegates;
+
+	/** UViewportStatsSubsystem display delegate indices. */
+	TArray<int32> DisplayDelegateIndices;
+
+	struct FVisibilityRequest
+	{
+		TWeakObjectPtr<UAvaPlayable> PlayableWeak;
+		bool bShouldBeVisible;
+
+		void Execute(const UAvaPlayableGroup* InPlayableGroup) const;
+	};
+	
+	TArray<FVisibilityRequest> VisibilityRequests;
+
+	TArray<TWeakInterfacePtr<IAvaPlayableVisibilityConstraint>> VisibilityConstraints;
+};
+
+/**
+ *	Remote Proxy Playable Group doesn't have a game instance.
+ */
+UCLASS()
+class AVALANCHEMEDIA_API UAvaPlayableRemoteProxyGroup : public UAvaPlayableGroup
+{
+	GENERATED_BODY()
+	
+public:
+	virtual bool ConditionalBeginPlay(const FAvaInstancePlaySettings& InWorldPlaySettings) override;
+	virtual void RequestEndPlayWorld(bool bInForceImmediate) override;
+	
+	bool bIsPlaying = false;
+};
\ No newline at end of file
diff --git a/Engine/Plugins/Experimental/Avalanche/Source/AvalancheMedia/Public/Playback/Nodes/AvaPlaybackNode.h b/Engine/Plugins/Experimental/Avalanche/Source/AvalancheMedia/Public/Playback/Nodes/AvaPlaybackNode.h
--- a/Engine/Plugins/Experimental/Avalanche/Source/AvalancheMedia/Public/Playback/Nodes/AvaPlaybackNode.h
+++ b/Engine/Plugins/Experimental/Avalanche/Source/AvalancheMedia/Public/Playback/Nodes/AvaPlaybackNode.h
@@ -2,17 +2,16 @@
 
 #pragma once
 
-#include "CoreMinimal.h"
 #include "AvaMediaDefines.h"
+#include "CoreMinimal.h"
 #include "UObject/Object.h"
 #include "AvaPlaybackNode.generated.h"
 
+class UAvaPlaybackGraph;
 class UEdGraphNode;
-class UAvalanchePlayback;
-
 
 /**
- * Base Class for all the Nodes found in the Avalanche Playback Graph
+ * Base Class for all the Nodes found in the Motion Design Playback Graph
  */
 UCLASS(Abstract)
 class AVALANCHEMEDIA_API UAvaPlaybackNode : public UObject
@@ -37,7 +36,7 @@ public:
 	
 	virtual void BeginDestroy() override;
 	
-	UAvalanchePlayback* GetPlayback() const;
+	UAvaPlaybackGraph* GetPlayback() const;
 
 	/**
 	 * Gets the Display Name of the Node
@@ -56,7 +55,7 @@ public:
 	virtual FText GetNodeTooltipText() const;
 
 	/**
-	 * Called when the UAvalanchePlayback changes its state from Playing <-> Stopped
+	 * Called when the UAvaPlaybackGraph changes its state from Playing <-> Stopped
 	 * @param  bPlaying The new state of Playback
 	 */
 	virtual void NotifyPlaybackStateChanged(bool bPlaying) {}
diff --git a/Engine/Plugins/Experimental/Avalanche/Source/AvalancheMedia/Public/Playback/Nodes/AvaPlaybackNodeBlueprintPlayer.h b/Engine/Plugins/Experimental/Avalanche/Source/AvalancheMedia/Public/Playback/Nodes/AvaPlaybackNodeBlueprintPlayer.h
--- a/Engine/Plugins/Experimental/Avalanche/Source/AvalancheMedia/Public/Playback/Nodes/AvaPlaybackNodeBlueprintPlayer.h
+++ b/Engine/Plugins/Experimental/Avalanche/Source/AvalancheMedia/Public/Playback/Nodes/AvaPlaybackNodeBlueprintPlayer.h
@@ -2,8 +2,8 @@
 
 #pragma once
 
-#include "CoreMinimal.h"
 #include "AvaPlaybackNodePlayer.h"
+#include "CoreMinimal.h"
 #include "AvaPlaybackNodeBlueprintPlayer.generated.h"
 
 class UAvalancheBlueprint;
@@ -19,17 +19,17 @@ public:
 	virtual void RefreshNode(bool bDryRunGraph) override;
 	virtual void PostLoad() override;
 	
-	TSoftObjectPtr<UAvalancheBlueprint> GetAvalancheAsset() const { return BlueprintAsset; }
-	void SetAvalancheAsset(const TSoftObjectPtr<UAvalancheBlueprint>& InAsset);
+	TSoftObjectPtr<UAvalancheBlueprint> GetAsset() const { return Asset; }
+	void SetAsset(const TSoftObjectPtr<UAvalancheBlueprint>& InAsset);
 	
 	void UpdateDisplayNameText();
 
-	virtual const FSoftObjectPath& GetAvalancheAssetPath() const override { return BlueprintAsset.ToSoftObjectPath();}
+	virtual const FSoftObjectPath& GetAssetPath() const override { return Asset.ToSoftObjectPath();}
 	
-	virtual FAvaSoftAssetPtr GetAvalancheAssetPtr() const override;
+	virtual FAvaSoftAssetPtr GetAssetPtr() const override;
 	
 protected:
 	
 	UPROPERTY(EditAnywhere, Category = "Motion Design")
-	TSoftObjectPtr<UAvalancheBlueprint> BlueprintAsset;
+	TSoftObjectPtr<UAvalancheBlueprint> Asset;
 };
diff --git a/Engine/Plugins/Experimental/Avalanche/Source/AvalancheMedia/Public/Playback/Nodes/AvaPlaybackNodeLevelPlayer.h b/Engine/Plugins/Experimental/Avalanche/Source/AvalancheMedia/Public/Playback/Nodes/AvaPlaybackNodeLevelPlayer.h
--- a/Engine/Plugins/Experimental/Avalanche/Source/AvalancheMedia/Public/Playback/Nodes/AvaPlaybackNodeLevelPlayer.h
+++ b/Engine/Plugins/Experimental/Avalanche/Source/AvalancheMedia/Public/Playback/Nodes/AvaPlaybackNodeLevelPlayer.h
@@ -2,8 +2,8 @@
 
 #pragma once
 
-#include "CoreMinimal.h"
 #include "AvaPlaybackNodePlayer.h"
+#include "CoreMinimal.h"
 #include "AvaPlaybackNodeLevelPlayer.generated.h"
 
 class UWorld;
@@ -20,14 +20,14 @@ public:
 	virtual void RefreshNode(bool bDryRunGraph) override;
 	virtual void PostLoad() override;
 	
-	TSoftObjectPtr<UWorld> GetAvalancheAsset() const { return LevelAsset; }
-	void SetAvalancheAsset(const TSoftObjectPtr<UWorld>& InAsset);
+	TSoftObjectPtr<UWorld> GetAsset() const { return LevelAsset; }
+	void SetAsset(const TSoftObjectPtr<UWorld>& InAsset);
 	
 	void UpdateDisplayNameText();
 
-	virtual const FSoftObjectPath& GetAvalancheAssetPath() const override { return LevelAsset.ToSoftObjectPath();}
+	virtual const FSoftObjectPath& GetAssetPath() const override { return LevelAsset.ToSoftObjectPath();}
 
-	virtual FAvaSoftAssetPtr GetAvalancheAssetPtr() const override;
+	virtual FAvaSoftAssetPtr GetAssetPtr() const override;
 
 protected:
 	
diff --git a/Engine/Plugins/Experimental/Avalanche/Source/AvalancheMedia/Public/Playback/Nodes/AvaPlaybackNodePlayer.h b/Engine/Plugins/Experimental/Avalanche/Source/AvalancheMedia/Public/Playback/Nodes/AvaPlaybackNodePlayer.h
--- a/Engine/Plugins/Experimental/Avalanche/Source/AvalancheMedia/Public/Playback/Nodes/AvaPlaybackNodePlayer.h
+++ b/Engine/Plugins/Experimental/Avalanche/Source/AvalancheMedia/Public/Playback/Nodes/AvaPlaybackNodePlayer.h
@@ -2,8 +2,8 @@
 
 #pragma once
 
-#include "CoreMinimal.h"
 #include "AvaPlaybackNode.h"
+#include "CoreMinimal.h"
 #include "AvaPlaybackNodePlayer.generated.h"
 
 class UTextureRenderTarget2D;
@@ -38,9 +38,9 @@ public:
 	virtual bool EditorDryRunGraphOnNodeRefresh(FPropertyChangedEvent& PropertyChangedEvent) const override { return true; }
 #endif
 
-	virtual const FSoftObjectPath& GetAvalancheAssetPath() const;
+	virtual const FSoftObjectPath& GetAssetPath() const;
 
-	virtual FAvaSoftAssetPtr GetAvalancheAssetPtr() const { return FAvaSoftAssetPtr();}
+	virtual FAvaSoftAssetPtr GetAssetPtr() const { return FAvaSoftAssetPtr(); }
 
 	/** Get the channel indices from the last time this node was traversed. */
 	const TArray<int32>& GetLastTickChannelIndices() const { return LastTickChannelIndices; }
diff --git a/Engine/Plugins/Experimental/Avalanche/Source/AvalancheMedia/Public/Playback/Nodes/AvaPlaybackNodeRoot.h b/Engine/Plugins/Experimental/Avalanche/Source/AvalancheMedia/Public/Playback/Nodes/AvaPlaybackNodeRoot.h
--- a/Engine/Plugins/Experimental/Avalanche/Source/AvalancheMedia/Public/Playback/Nodes/AvaPlaybackNodeRoot.h
+++ b/Engine/Plugins/Experimental/Avalanche/Source/AvalancheMedia/Public/Playback/Nodes/AvaPlaybackNodeRoot.h
@@ -1,12 +1,13 @@
 // Copyright Epic Games, Inc. All Rights Reserved.
 
 #pragma once
+
 #include "AvaMediaDefines.h"
 #include "Playback/Nodes/AvaPlaybackNode.h"
 #include "AvaPlaybackNodeRoot.generated.h"
 
-struct FAvaOutputChannel;
 class UTextureRenderTarget2D;
+struct FAvaBroadcastOutputChannel;
 
 UCLASS()
 class AVALANCHEMEDIA_API UAvaPlaybackNodeRoot : public UAvaPlaybackNode
@@ -24,7 +25,7 @@ public:
 	virtual FText GetNodeTooltipText() const override;	
 	
 	void OnBroadcastChanged(EAvaBroadcastChange InChange);
-	void OnChannelChanged(const FAvaOutputChannel& InChannel, EAvaChannelChange InChange);
+	void OnChannelChanged(const FAvaBroadcastOutputChannel& InChannel, EAvaBroadcastChannelChange InChange);
 	
 	virtual int32 GetMinChildNodes() const override;
 	virtual int32 GetMaxChildNodes() const override;
diff --git a/Engine/Plugins/Experimental/Avalanche/Source/AvalancheMediaEditor/Private/AvaMediaEditorModule.cpp b/Engine/Plugins/Experimental/Avalanche/Source/AvalancheMediaEditor/Private/AvaMediaEditorModule.cpp
--- a/Engine/Plugins/Experimental/Avalanche/Source/AvalancheMediaEditor/Private/AvaMediaEditorModule.cpp
+++ b/Engine/Plugins/Experimental/Avalanche/Source/AvalancheMediaEditor/Private/AvaMediaEditorModule.cpp
@@ -1,38 +1,38 @@
 // Copyright Epic Games, Inc. All Rights Reserved.
 
 #include "AvaMediaEditorModule.h"
+
 #include "AvaMediaEditorStyle.h"
 #include "Broadcast/AvaBroadcastEditor.h"
+#include "Broadcast/OutputDevices/AvaBroadcastMediaIOOutputConfigurationCustomization.h"
 #include "Editor.h"
 #include "Engine/GameViewportClient.h"
 #include "Framework/Application/SlateApplication.h"
 #include "IAvaMediaModule.h"
 #include "LevelEditor.h"
-#include "LevelEditor/AvaLevelEditorUtils.h"
 #include "LevelEditorViewport.h"
-#include "OutputDevices/AvaMediaIOOutputConfigurationCustomization.h"
 #include "Playback/AvaPlaybackCommands.h"
 #include "Playback/Graph/AvaPlaybackConnectionDrawingPolicy.h"
-#include "Playlist/AvaPlaylistCommands.h"
-#include "Playlist/AvaRundownEditorSettings.h"
-#include "Playlist/AvaRundownMacroCollection.h"
-#include "Playlist/Customization/AvaRundownMacroCommandCustomization.h"
-#include "Playlist/Customization/AvaRundownMacroKeyBindingCustomization.h"
-#include "Playlist/Factories/Filters/AvaPlaylistFactoriesUtils.h"
-#include "Playlist/Factories/Filters/AvaPlaylistFilterChannelExpressionFactory.h"
-#include "Playlist/Factories/Filters/AvaPlaylistFilterIdExpressionFactory.h"
-#include "Playlist/Factories/Filters/AvaPlaylistFilterNameExpressionFactory.h"
-#include "Playlist/Factories/Filters/AvaPlaylistFilterPathExpressionFactory.h"
-#include "Playlist/Factories/Filters/AvaPlaylistFilterStatusExpressionFactory.h"
-#include "Playlist/Factories/Filters/AvaPlaylistFilterTransitionLayerExpressionFactory.h"
-#include "Playlist/Factories/Filters/IAvaPlaylistFilterExpressionFactory.h"
-#include "Playlist/Factories/Filters/IAvaPlaylistFilterSuggestionFactory.h"
-#include "Playlist/Factories/Suggestions/AvaPlaylistFilterChannelSuggestionFactory.h"
-#include "Playlist/Factories/Suggestions/AvaPlaylistFilterIdSuggestionFactory.h"
-#include "Playlist/Factories/Suggestions/AvaPlaylistFilterNameSuggestionFactory.h"
-#include "Playlist/Factories/Suggestions/AvaPlaylistFilterPathSuggestionFactory.h"
-#include "Playlist/Factories/Suggestions/AvaPlaylistFilterStatusSuggestionFactory.h"
-#include "Playlist/Factories/Suggestions/AvaPlaylistFilterTransitionLayerSuggestionFactory.h"
+#include "Rundown/AvaRundownCommands.h"
+#include "Rundown/AvaRundownEditorSettings.h"
+#include "Rundown/AvaRundownMacroCollection.h"
+#include "Rundown/Customization/AvaRundownMacroCommandCustomization.h"
+#include "Rundown/Customization/AvaRundownMacroKeyBindingCustomization.h"
+#include "Rundown/Factories/Filters/AvaRundownFactoriesUtils.h"
+#include "Rundown/Factories/Filters/AvaRundownFilterChannelExpressionFactory.h"
+#include "Rundown/Factories/Filters/AvaRundownFilterIdExpressionFactory.h"
+#include "Rundown/Factories/Filters/AvaRundownFilterNameExpressionFactory.h"
+#include "Rundown/Factories/Filters/AvaRundownFilterPathExpressionFactory.h"
+#include "Rundown/Factories/Filters/AvaRundownFilterStatusExpressionFactory.h"
+#include "Rundown/Factories/Filters/AvaRundownFilterTransitionLayerExpressionFactory.h"
+#include "Rundown/Factories/Filters/IAvaRundownFilterExpressionFactory.h"
+#include "Rundown/Factories/Filters/IAvaRundownFilterSuggestionFactory.h"
+#include "Rundown/Factories/Suggestions/AvaRundownFilterChannelSuggestionFactory.h"
+#include "Rundown/Factories/Suggestions/AvaRundownFilterIdSuggestionFactory.h"
+#include "Rundown/Factories/Suggestions/AvaRundownFilterNameSuggestionFactory.h"
+#include "Rundown/Factories/Suggestions/AvaRundownFilterPathSuggestionFactory.h"
+#include "Rundown/Factories/Suggestions/AvaRundownFilterStatusSuggestionFactory.h"
+#include "Rundown/Factories/Suggestions/AvaRundownFilterTransitionLayerSuggestionFactory.h"
 #include "ToolMenus.h"
 
 #define LOCTEXT_NAMESPACE "AvaMediaEditorModule"
@@ -44,14 +44,14 @@ namespace UE::AvaMediaEditorModule::Private
 	namespace BroadcastEditorEntry
 	{
 		static const FName MenuName(TEXT("LevelEditor.StatusBar.ToolBar"));
-		static const FName SectionName(TEXT("AvalancheMedia"));
+		static const FName SectionName(TEXT("MotionDesign"));
 	}
 	
 	// Command line parsing helper.
-	bool IsPlaylistServerManuallyStarted(FString& OutPlaylistServerName)
+	bool IsRundownServerManuallyStarted(FString& OutRundownServerName)
 	{
-		return FParse::Value(FCommandLine::Get(), TEXT("AvaMediaPlaylistServerStart="), OutPlaylistServerName) ||
-			FParse::Param(FCommandLine::Get(), TEXT("AvaMediaPlaylistServerStart"));
+		return FParse::Value(FCommandLine::Get(), TEXT("MotionDesignRundownServerStart="), OutRundownServerName) ||
+			FParse::Param(FCommandLine::Get(), TEXT("MotionDesignRundownServerStart"));
 	}
 
 	static void GetEditorViewportClient(FCommonViewportClient** OutViewportClient)
@@ -77,7 +77,7 @@ void FAvaMediaEditorModule::StartupModule()
 	InitExtensibilityManagers();
 
 	FAvaPlaybackCommands::Register();
-	FAvaPlaylistCommands::Register();
+	FAvaRundownCommands::Register();
 
 	FAvaMediaEditorStyle::Initialize();
 
@@ -86,7 +86,7 @@ void FAvaMediaEditorModule::StartupModule()
 		AddEditorToolbarButtons();
 	}
 
-	// Register the Avalanche Playback Graph connection policy with the graph editor
+	// Register the Motion Design Playback Graph connection policy with the graph editor
 	PlaybackConnectionFactory = MakeShared<FAvaPlaybackConnectionDrawingPolicyFactory>();
 	FEdGraphUtilities::RegisterVisualPinConnectionFactory(PlaybackConnectionFactory);
 
@@ -102,22 +102,22 @@ void FAvaMediaEditorModule::StartupModule()
 	LevelEditor.OnMapChanged().AddRaw(this, &FAvaMediaEditorModule::HandleMapChanged);
 
 	ConsoleCmds.Add(IConsoleManager::Get().RegisterConsoleCommand(
-		TEXT("AvaRundownServer.Start"),
+		TEXT("MotionDesignRundownServer.Start"),
 		TEXT("Starts the rundown server."),
-		FConsoleCommandWithArgsDelegate::CreateRaw(this, &FAvaMediaEditorModule::StartPlaylistServerCommand),
+		FConsoleCommandWithArgsDelegate::CreateRaw(this, &FAvaMediaEditorModule::StartRundownServerCommand),
 		ECVF_Default
 	));
 	ConsoleCmds.Add(IConsoleManager::Get().RegisterConsoleCommand(
-		TEXT("AvaRundownServer.Stop"),
+		TEXT("MotionDesignRundownServer.Stop"),
 		TEXT("Stops the rundown server."),
-		FConsoleCommandWithArgsDelegate::CreateRaw(this, &FAvaMediaEditorModule::StopPlaylistServerCommand),
+		FConsoleCommandWithArgsDelegate::CreateRaw(this, &FAvaMediaEditorModule::StopRundownServerCommand),
 		ECVF_Default
 	));
 
 	IAvaMediaModule::Get().GetEditorViewportClientDelegate().BindStatic(&GetEditorViewportClient);
 
 	FString DummyServerName;
-	if (IsPlaylistServerManuallyStarted(DummyServerName))
+	if (IsRundownServerManuallyStarted(DummyServerName))
 	{
 		// Prevent throttling when the server is started.
 		// This has to be done before any SLevelViewport are ticked since the cvar value is cached on first tick.
@@ -128,8 +128,8 @@ void FAvaMediaEditorModule::StartupModule()
 			UE_LOG(LogAvaMediaEditor, Log, TEXT("Setting Slate.bAllowThrottling to false."));
 		}
 	}
-	RegisterPlaylistFilterExpressionFactories();
-	RegisterPlaylistFilterSuggestionFactories();
+	RegisterRundownFilterExpressionFactories();
+	RegisterRundownFilterSuggestionFactories();
 }
 
 void FAvaMediaEditorModule::ShutdownModule()
@@ -158,7 +158,7 @@ void FAvaMediaEditorModule::ShutdownModule()
 
 	FAvaMediaEditorStyle::Shutdown();
 	FAvaPlaybackCommands::Unregister();
-	FAvaPlaylistCommands::Unregister();
+	FAvaRundownCommands::Unregister();
 
 	if (PlaybackConnectionFactory.IsValid())
 	{
@@ -183,14 +183,14 @@ TSharedPtr<FExtensibilityManager> FAvaMediaEditorModule::GetPlaybackToolBarExten
 	return PlaybackToolBarExtensibility;
 }
 
-TSharedPtr<FExtensibilityManager> FAvaMediaEditorModule::GetPlaylistToolBarExtensibilityManager()
+TSharedPtr<FExtensibilityManager> FAvaMediaEditorModule::GetRundownToolBarExtensibilityManager()
 {
-	return PlaylistToolBarExtensibility;
+	return RundownToolBarExtensibility;
 }
 
-TSharedPtr<FExtensibilityManager> FAvaMediaEditorModule::GetPlaylistMenuExtensibilityManager()
+TSharedPtr<FExtensibilityManager> FAvaMediaEditorModule::GetRundownMenuExtensibilityManager()
 {
-	return PlaylistMenuExtensibility;
+	return RundownMenuExtensibility;
 }
 
 FSlateIcon FAvaMediaEditorModule::GetToolbarBroadcastButtonIcon() const
@@ -199,38 +199,38 @@ FSlateIcon FAvaMediaEditorModule::GetToolbarBroadcastButtonIcon() const
 
 	if (MediaModule.IsMediaPlaybackClientStarted())
 	{
-		return FSlateIcon(FAvaMediaEditorStyle::GetStyleSetName(), "AvalancheMediaEditor.BroadcastClient", "AvalancheMediaEditor.BroadcastClient.Small");
+		return FSlateIcon(FAvaMediaEditorStyle::GetStyleSetName(), "AvaMediaEditor.BroadcastClient", "AvaMediaEditor.BroadcastClient.Small");
 	}
 	else if (MediaModule.IsMediaPlaybackServerStarted())
 	{
-		return FSlateIcon(FAvaMediaEditorStyle::GetStyleSetName(), "AvalancheMediaEditor.BroadcastServer", "AvalancheMediaEditor.BroadcastServer.Small");
+		return FSlateIcon(FAvaMediaEditorStyle::GetStyleSetName(), "AvaMediaEditor.BroadcastServer", "AvaMediaEditor.BroadcastServer.Small");
 	}
-	return FSlateIcon(FAvaMediaEditorStyle::GetStyleSetName(), "AvalancheMediaEditor.BroadcastIcon");
+	return FSlateIcon(FAvaMediaEditorStyle::GetStyleSetName(), "AvaMediaEditor.BroadcastIcon");
 }
 
-bool FAvaMediaEditorModule::CanFilterSupportComparisonOperation(const FName& InFilterKey, ETextFilterComparisonOperation InOperation, EAvaPlaylistSearchListType InPlaylistSearchListType) const
+bool FAvaMediaEditorModule::CanFilterSupportComparisonOperation(const FName& InFilterKey, ETextFilterComparisonOperation InOperation, EAvaRundownSearchListType InRundownSearchListType) const
 {
-	if (const TSharedPtr<IAvaPlaylistFilterExpressionFactory>* FilterExpressionFactory = FilterExpressionFactories.Find(InFilterKey))
+	if (const TSharedPtr<IAvaRundownFilterExpressionFactory>* FilterExpressionFactory = FilterExpressionFactories.Find(InFilterKey))
 	{
-		return FilterExpressionFactory->Get()->SupportsComparisonOperation(InOperation, InPlaylistSearchListType);
+		return FilterExpressionFactory->Get()->SupportsComparisonOperation(InOperation, InRundownSearchListType);
 	}
 	return false;
 }
 
-bool FAvaMediaEditorModule::FilterExpression(const FName& InFilterKey, const FAvalanchePage& InItem, const FAvaPlaylistTextFilterArgs& InArgs) const
+bool FAvaMediaEditorModule::FilterExpression(const FName& InFilterKey, const FAvaRundownPage& InItem, const FAvaRundownTextFilterArgs& InArgs) const
 {
-	if (const TSharedPtr<IAvaPlaylistFilterExpressionFactory>* FilterExpressionFactory = FilterExpressionFactories.Find(InFilterKey))
+	if (const TSharedPtr<IAvaRundownFilterExpressionFactory>* FilterExpressionFactory = FilterExpressionFactories.Find(InFilterKey))
 	{
 		return FilterExpressionFactory->Get()->FilterExpression(InItem, InArgs);
 	}
 	return false;
 }
 
-TArray<TSharedPtr<IAvaPlaylistFilterSuggestionFactory>> FAvaMediaEditorModule::GetSimpleSuggestions(EAvaPlaylistSearchListType InSuggestionType) const
+TArray<TSharedPtr<IAvaRundownFilterSuggestionFactory>> FAvaMediaEditorModule::GetSimpleSuggestions(EAvaRundownSearchListType InSuggestionType) const
 {
-	TArray<TSharedPtr<IAvaPlaylistFilterSuggestionFactory>> OutArray;
+	TArray<TSharedPtr<IAvaRundownFilterSuggestionFactory>> OutArray;
 
-	for (const TPair<FName, TSharedPtr<IAvaPlaylistFilterSuggestionFactory>>& Suggestion : FilterSuggestionFactories)
+	for (const TPair<FName, TSharedPtr<IAvaRundownFilterSuggestionFactory>>& Suggestion : FilterSuggestionFactories)
 	{
 		if (Suggestion.Value->SupportSuggestionType(InSuggestionType) && Suggestion.Value->IsSimpleSuggestion())
 		{
@@ -241,11 +241,11 @@ TArray<TSharedPtr<IAvaPlaylistFilterSuggestionFactory>> FAvaMediaEditorModule::G
 	return OutArray;
 }
 
-TArray<TSharedPtr<IAvaPlaylistFilterSuggestionFactory>> FAvaMediaEditorModule::GetComplexSuggestions(EAvaPlaylistSearchListType InSuggestionType) const
+TArray<TSharedPtr<IAvaRundownFilterSuggestionFactory>> FAvaMediaEditorModule::GetComplexSuggestions(EAvaRundownSearchListType InSuggestionType) const
 {
-	TArray<TSharedPtr<IAvaPlaylistFilterSuggestionFactory>> OutArray;
+	TArray<TSharedPtr<IAvaRundownFilterSuggestionFactory>> OutArray;
 
-	for (const TPair<FName, TSharedPtr<IAvaPlaylistFilterSuggestionFactory>>& Suggestion : FilterSuggestionFactories)
+	for (const TPair<FName, TSharedPtr<IAvaRundownFilterSuggestionFactory>>& Suggestion : FilterSuggestionFactories)
 	{
 		if (Suggestion.Value->SupportSuggestionType(InSuggestionType) && !Suggestion.Value->IsSimpleSuggestion())
 		{
@@ -294,22 +294,22 @@ void FAvaMediaEditorModule::InitExtensibilityManagers()
 {
 	BroadcastToolBarExtensibility = MakeShared<FExtensibilityManager>();
 	PlaybackToolBarExtensibility  = MakeShared<FExtensibilityManager>();
-	PlaylistToolBarExtensibility  = MakeShared<FExtensibilityManager>();
-	PlaylistMenuExtensibility     = MakeShared<FExtensibilityManager>();
+	RundownToolBarExtensibility  = MakeShared<FExtensibilityManager>();
+	RundownMenuExtensibility     = MakeShared<FExtensibilityManager>();
 }
 
 void FAvaMediaEditorModule::ResetExtensibilityManagers()
 {
 	BroadcastToolBarExtensibility = nullptr;
 	PlaybackToolBarExtensibility  = nullptr;
-	PlaylistToolBarExtensibility  = nullptr;
-	PlaylistMenuExtensibility     = nullptr;
+	RundownToolBarExtensibility  = nullptr;
+	RundownMenuExtensibility     = nullptr;
 }
 
 void FAvaMediaEditorModule::RegisterCustomizations() const
 {
 	FPropertyEditorModule& PropertyModule = FModuleManager::LoadModuleChecked<FPropertyEditorModule>("PropertyEditor");
-	PropertyModule.RegisterCustomPropertyTypeLayout(FMediaIOOutputConfiguration::StaticStruct()->GetFName(), FOnGetPropertyTypeCustomizationInstance::CreateStatic(&FAvaMediaIOOutputConfigurationCustomization::MakeInstance));
+	PropertyModule.RegisterCustomPropertyTypeLayout(FMediaIOOutputConfiguration::StaticStruct()->GetFName(), FOnGetPropertyTypeCustomizationInstance::CreateStatic(&FAvaBroadcastMediaIOOutputConfigurationCustomization::MakeInstance));
 	PropertyModule.RegisterCustomPropertyTypeLayout(FAvaRundownMacroCommand::StaticStruct()->GetFName(), FOnGetPropertyTypeCustomizationInstance::CreateStatic(&FAvaRundownMacroCommandCustomization::MakeInstance));
 	PropertyModule.RegisterCustomPropertyTypeLayout(FAvaRundownMacroKeyBinding::StaticStruct()->GetFName(), FOnGetPropertyTypeCustomizationInstance::CreateStatic(&FAvaRundownMacroKeyBindingCustomization::MakeInstance));
 }
@@ -322,34 +322,34 @@ void FAvaMediaEditorModule::UnregisterCustomizations() const
 	PropertyModule.UnregisterCustomPropertyTypeLayout(FAvaRundownMacroKeyBinding::StaticStruct()->GetFName());
 }
 
-void FAvaMediaEditorModule::StartPlaylistServerCommand(const TArray<FString>& Args)
+void FAvaMediaEditorModule::StartRundownServerCommand(const TArray<FString>& Args)
 {
-	if (PlaylistServer)
+	if (RundownServer)
 	{
 		UE_LOG(LogAvaMediaEditor, Log, TEXT("Rundown Server is already started."));
 		return;
 	}
 	
-	PlaylistServer = MakeShared<FAvaPlaylistServer>();
+	RundownServer = MakeShared<FAvaRundownServer>();
 	
-	// Remark: Only the module's playlist server register console commands to avoid
+	// Remark: Only the module's rundown server register console commands to avoid
 	// conflicts with temporary servers (for testing).
-	PlaylistServer->RegisterConsoleCommands();
+	RundownServer->RegisterConsoleCommands();
 
-	PlaylistServer->Init(Args.Num() > 0 ? Args[0] : TEXT(""));
-	OnPlaylistServerStarted.Broadcast();
+	RundownServer->Init(Args.Num() > 0 ? Args[0] : TEXT(""));
+	OnRundownServerStarted.Broadcast();
 
 	UE_LOG(LogAvaMediaEditor, Log, TEXT("Rundown Server Started."));
 }
 
-void FAvaMediaEditorModule::StopPlaylistServerCommand(const TArray<FString>& Args)
+void FAvaMediaEditorModule::StopRundownServerCommand(const TArray<FString>& Args)
 {
-	if (PlaylistServer)
+	if (RundownServer)
 	{
 		UE_LOG(LogAvaMediaEditor, Log, TEXT("Stopping Rundown Server..."));
-		OnPlaylistServerStopped.Broadcast();
+		OnRundownServerStopped.Broadcast();
 	}
-	PlaylistServer.Reset();
+	RundownServer.Reset();
 }
 
 void FAvaMediaEditorModule::PostEngineInit()
@@ -362,13 +362,13 @@ void FAvaMediaEditorModule::PostEngineInit()
 	// Allow for specification of the server name in the command line.
 	// Command line has priority over project settings.
 	FString ServerName;
-	if (IsPlaylistServerManuallyStarted(ServerName))
+	if (IsRundownServerManuallyStarted(ServerName))
 	{
-		StartPlaylistServerCommand({ServerName});
+		StartRundownServerCommand({ServerName});
 	}
 	else if (Settings && Settings->bAutoStartRundownServer)
 	{
-		StartPlaylistServerCommand({Settings->RundownServerName});
+		StartRundownServerCommand({Settings->RundownServerName});
 	}
 }
 
@@ -379,7 +379,7 @@ void FAvaMediaEditorModule::EnginePreExit()
 
 void FAvaMediaEditorModule::StopAllServices()
 {
-	StopPlaylistServerCommand({});
+	StopRundownServerCommand({});
 }
 
 void FAvaMediaEditorModule::HandleMapChanged(UWorld* InWorld, EMapChangeType InMapChangeType)
@@ -413,14 +413,14 @@ void FAvaMediaEditorModule::HandleMapChanged(UWorld* InWorld, EMapChangeType InM
 }
 
 template <
-	typename InPlaylistFilterExpressionFactoryType,
+	typename InRundownFilterExpressionFactoryType,
 	typename ... InArgsType
-	UE_REQUIRES(TIsDerivedFrom<InPlaylistFilterExpressionFactoryType, IAvaPlaylistFilterExpressionFactory>::Value)
+	UE_REQUIRES(TIsDerivedFrom<InRundownFilterExpressionFactoryType, IAvaRundownFilterExpressionFactory>::Value)
 >
-void FAvaMediaEditorModule::RegisterPlaylistFilterExpressionFactory(InArgsType&&... InArgs)
+void FAvaMediaEditorModule::RegisterRundownFilterExpressionFactory(InArgsType&&... InArgs)
 {
-	const TSharedRef<IAvaPlaylistFilterExpressionFactory> FilterExpressionFactory =
-		IAvaPlaylistFilterExpressionFactory::MakeInstance<InPlaylistFilterExpressionFactoryType>(Forward<InArgsType>(InArgs)...);
+	const TSharedRef<IAvaRundownFilterExpressionFactory> FilterExpressionFactory =
+		IAvaRundownFilterExpressionFactory::MakeInstance<InRundownFilterExpressionFactoryType>(Forward<InArgsType>(InArgs)...);
 
 	if (!FilterExpressionFactories.Contains(FilterExpressionFactory->GetFilterIdentifier()))
 	{
@@ -429,14 +429,14 @@ void FAvaMediaEditorModule::RegisterPlaylistFilterExpressionFactory(InArgsType&&
 }
 
 template <
-	typename InPlaylistSuggestionFactoryType,
+	typename InRundownSuggestionFactoryType,
 	typename ... InArgsType
-	UE_REQUIRES(TIsDerivedFrom<InPlaylistSuggestionFactoryType, IAvaPlaylistFilterSuggestionFactory>::Value)
+	UE_REQUIRES(TIsDerivedFrom<InRundownSuggestionFactoryType, IAvaRundownFilterSuggestionFactory>::Value)
 >
-void FAvaMediaEditorModule::RegisterPlaylistFilterSuggestionFactory(InArgsType&&... InArgs)
+void FAvaMediaEditorModule::RegisterRundownFilterSuggestionFactory(InArgsType&&... InArgs)
 {
-	const TSharedRef<IAvaPlaylistFilterSuggestionFactory> FilterSuggestionFactory =
-		IAvaPlaylistFilterSuggestionFactory::MakeInstance<InPlaylistSuggestionFactoryType>(Forward<InArgsType>(InArgs)...);
+	const TSharedRef<IAvaRundownFilterSuggestionFactory> FilterSuggestionFactory =
+		IAvaRundownFilterSuggestionFactory::MakeInstance<InRundownSuggestionFactoryType>(Forward<InArgsType>(InArgs)...);
 
 	if (!FilterSuggestionFactories.Contains(FilterSuggestionFactory->GetSuggestionIdentifier()))
 	{
@@ -444,24 +444,24 @@ void FAvaMediaEditorModule::RegisterPlaylistFilterSuggestionFactory(InArgsType&&
 	}
 }
 
-void FAvaMediaEditorModule::RegisterPlaylistFilterExpressionFactories()
+void FAvaMediaEditorModule::RegisterRundownFilterExpressionFactories()
 {
-	RegisterPlaylistFilterExpressionFactory<FAvaPlaylistFilterNameExpressionFactory>();
-	RegisterPlaylistFilterExpressionFactory<FAvaPlaylistFilterIdExpressionFactory>();
-	RegisterPlaylistFilterExpressionFactory<FAvaPlaylistFilterPathExpressionFactory>();
-	RegisterPlaylistFilterExpressionFactory<FAvaPlaylistFilterChannelExpressionFactory>();
-	RegisterPlaylistFilterExpressionFactory<FAvaPlaylistFilterStatusExpressionFactory>();
-	RegisterPlaylistFilterExpressionFactory<FAvaPlaylistFilterTransitionLayerExpressionFactory>();
+	RegisterRundownFilterExpressionFactory<FAvaRundownFilterNameExpressionFactory>();
+	RegisterRundownFilterExpressionFactory<FAvaRundownFilterIdExpressionFactory>();
+	RegisterRundownFilterExpressionFactory<FAvaRundownFilterPathExpressionFactory>();
+	RegisterRundownFilterExpressionFactory<FAvaRundownFilterChannelExpressionFactory>();
+	RegisterRundownFilterExpressionFactory<FAvaRundownFilterStatusExpressionFactory>();
+	RegisterRundownFilterExpressionFactory<FAvaRundownFilterTransitionLayerExpressionFactory>();
 }
 
-void FAvaMediaEditorModule::RegisterPlaylistFilterSuggestionFactories()
+void FAvaMediaEditorModule::RegisterRundownFilterSuggestionFactories()
 {
-	RegisterPlaylistFilterSuggestionFactory<FAvaPlaylistFilterNameSuggestionFactory>();
-	RegisterPlaylistFilterSuggestionFactory<FAvaPlaylistFilterIdSuggestionFactory>();
-	RegisterPlaylistFilterSuggestionFactory<FAvaPlaylistFilterPathSuggestionFactory>();
-	RegisterPlaylistFilterSuggestionFactory<FAvaPlaylistFilterChannelSuggestionFactory>();
-	RegisterPlaylistFilterSuggestionFactory<FAvaPlaylistFilterStatusSuggestionFactory>();
-	RegisterPlaylistFilterSuggestionFactory<FAvaPlaylistFilterTransitionLayerSuggestionFactory>();
+	RegisterRundownFilterSuggestionFactory<FAvaRundownFilterNameSuggestionFactory>();
+	RegisterRundownFilterSuggestionFactory<FAvaRundownFilterIdSuggestionFactory>();
+	RegisterRundownFilterSuggestionFactory<FAvaRundownFilterPathSuggestionFactory>();
+	RegisterRundownFilterSuggestionFactory<FAvaRundownFilterChannelSuggestionFactory>();
+	RegisterRundownFilterSuggestionFactory<FAvaRundownFilterStatusSuggestionFactory>();
+	RegisterRundownFilterSuggestionFactory<FAvaRundownFilterTransitionLayerSuggestionFactory>();
 }
 
 #undef LOCTEXT_NAMESPACE
diff --git a/Engine/Plugins/Experimental/Avalanche/Source/AvalancheMediaEditor/Private/AvaMediaEditorModule.h b/Engine/Plugins/Experimental/Avalanche/Source/AvalancheMediaEditor/Private/AvaMediaEditorModule.h
--- a/Engine/Plugins/Experimental/Avalanche/Source/AvalancheMediaEditor/Private/AvaMediaEditorModule.h
+++ b/Engine/Plugins/Experimental/Avalanche/Source/AvalancheMediaEditor/Private/AvaMediaEditorModule.h
@@ -3,13 +3,13 @@
 #pragma once
 
 #include "IAvaMediaEditorModule.h"
-#include "Playlist/AvaPlaylistServer.h"
+#include "Rundown/AvaRundownServer.h"
 #include "Templates/UnrealTypeTraits.h"
 
-class IAvaPlaylistFilterExpressionFactory;
-class IAvaPlaylistFilterSuggestionFactory;
+class IAvaRundownFilterExpressionFactory;
+class IAvaRundownFilterSuggestionFactory;
+enum class EAvaRundownSearchListType : uint8;
 struct FGraphPanelPinConnectionFactory;
-enum class EAvaPlaylistSearchListType : uint8;
 
 class FAvaMediaEditorModule : public IAvaMediaEditorModule
 {
@@ -24,15 +24,15 @@ public:
 	//IAvaMediaEditorModule
 	virtual TSharedPtr<FExtensibilityManager> GetBroadcastToolBarExtensibilityManager() override;
 	virtual TSharedPtr<FExtensibilityManager> GetPlaybackToolBarExtensibilityManager() override;
-	virtual TSharedPtr<FExtensibilityManager> GetPlaylistToolBarExtensibilityManager() override;
-	virtual TSharedPtr<FExtensibilityManager> GetPlaylistMenuExtensibilityManager() override;
-	virtual FOnPlaylistServerStarted& GetOnPlaylistServerStarted() override { return OnPlaylistServerStarted; }
-	virtual FOnPlaylistServerStopped& GetOnPlaylistServerStopped() override { return OnPlaylistServerStopped; }
-	virtual TSharedPtr<FAvaPlaylistServer> GetPlaylistServer() const override { return PlaylistServer; }
-	virtual bool CanFilterSupportComparisonOperation(const FName& InFilterKey, ETextFilterComparisonOperation InOperation, EAvaPlaylistSearchListType InPlaylistSearchListType) const override;
-	virtual bool FilterExpression(const FName& InFilterKey, const FAvalanchePage& InItem, const FAvaPlaylistTextFilterArgs& InArgs) const override;
-	virtual TArray<TSharedPtr<IAvaPlaylistFilterSuggestionFactory>> GetSimpleSuggestions(EAvaPlaylistSearchListType InSuggestionType) const override;
-	virtual TArray<TSharedPtr<IAvaPlaylistFilterSuggestionFactory>> GetComplexSuggestions(EAvaPlaylistSearchListType InSuggestionType) const override;
+	virtual TSharedPtr<FExtensibilityManager> GetRundownToolBarExtensibilityManager() override;
+	virtual TSharedPtr<FExtensibilityManager> GetRundownMenuExtensibilityManager() override;
+	virtual FOnRundownServerStarted& GetOnRundownServerStarted() override { return OnRundownServerStarted; }
+	virtual FOnRundownServerStopped& GetOnRundownServerStopped() override { return OnRundownServerStopped; }
+	virtual TSharedPtr<FAvaRundownServer> GetRundownServer() const override { return RundownServer; }
+	virtual bool CanFilterSupportComparisonOperation(const FName& InFilterKey, ETextFilterComparisonOperation InOperation, EAvaRundownSearchListType InRundownSearchListType) const override;
+	virtual bool FilterExpression(const FName& InFilterKey, const FAvaRundownPage& InItem, const FAvaRundownTextFilterArgs& InArgs) const override;
+	virtual TArray<TSharedPtr<IAvaRundownFilterSuggestionFactory>> GetSimpleSuggestions(EAvaRundownSearchListType InSuggestionType) const override;
+	virtual TArray<TSharedPtr<IAvaRundownFilterSuggestionFactory>> GetComplexSuggestions(EAvaRundownSearchListType InSuggestionType) const override;
 	//~IAvaMediaEditorModule
 
 	void AddEditorToolbarButtons();
@@ -51,8 +51,8 @@ protected:
 	/** Unregister details view customizations. */
 	void UnregisterCustomizations() const;
 
-	void StartPlaylistServerCommand(const TArray<FString>& Args);
-	void StopPlaylistServerCommand(const TArray<FString>& Args);
+	void StartRundownServerCommand(const TArray<FString>& Args);
+	void StopRundownServerCommand(const TArray<FString>& Args);
 
 private:
 	void PostEngineInit();
@@ -61,41 +61,41 @@ private:
 	void HandleMapChanged(UWorld* InWorld, EMapChangeType InMapChangeType);
 	
 	template <
-		typename InPlaylistFilterExpressionFactoryType,
+		typename InRundownFilterExpressionFactoryType,
 		typename... InArgsType
-		UE_REQUIRES(TIsDerivedFrom<InPlaylistFilterExpressionFactoryType, IAvaPlaylistFilterExpressionFactory>::Value)
+		UE_REQUIRES(TIsDerivedFrom<InRundownFilterExpressionFactoryType, IAvaRundownFilterExpressionFactory>::Value)
 	>
-	void RegisterPlaylistFilterExpressionFactory(InArgsType&&... InArgs);
+	void RegisterRundownFilterExpressionFactory(InArgsType&&... InArgs);
 
 	template <
-		typename InPlaylistSuggestionFactoryType,
+		typename InRundownSuggestionFactoryType,
 		typename... InArgsType
-		UE_REQUIRES(TIsDerivedFrom<InPlaylistSuggestionFactoryType, IAvaPlaylistFilterSuggestionFactory>::Value)
+		UE_REQUIRES(TIsDerivedFrom<InRundownSuggestionFactoryType, IAvaRundownFilterSuggestionFactory>::Value)
 	>
-	void RegisterPlaylistFilterSuggestionFactory(InArgsType&&... InArgs);
+	void RegisterRundownFilterSuggestionFactory(InArgsType&&... InArgs);
 
-	void RegisterPlaylistFilterExpressionFactories();
+	void RegisterRundownFilterExpressionFactories();
 
-	void RegisterPlaylistFilterSuggestionFactories();
+	void RegisterRundownFilterSuggestionFactories();
 
 private:
 	TSharedPtr<FExtensibilityManager> BroadcastToolBarExtensibility;
 	TSharedPtr<FExtensibilityManager> PlaybackToolBarExtensibility;
-	TSharedPtr<FExtensibilityManager> PlaylistToolBarExtensibility;
-	TSharedPtr<FExtensibilityManager> PlaylistMenuExtensibility;
+	TSharedPtr<FExtensibilityManager> RundownToolBarExtensibility;
+	TSharedPtr<FExtensibilityManager> RundownMenuExtensibility;
 
 	TSharedPtr<FGraphPanelPinConnectionFactory> PlaybackConnectionFactory;
 
-	TSharedPtr<FAvaPlaylistServer> PlaylistServer;
+	TSharedPtr<FAvaRundownServer> RundownServer;
 
-	FOnPlaylistServerStarted OnPlaylistServerStarted;
-	FOnPlaylistServerStopped OnPlaylistServerStopped;
+	FOnRundownServerStarted OnRundownServerStarted;
+	FOnRundownServerStopped OnRundownServerStopped;
 
 	TArray<IConsoleObject*> ConsoleCmds;
 
-	/** Holds all the PlaylistFilterExpressionFactory */
-	TMap<FName, TSharedPtr<IAvaPlaylistFilterExpressionFactory>> FilterExpressionFactories;
+	/** Holds all the RundownFilterExpressionFactory */
+	TMap<FName, TSharedPtr<IAvaRundownFilterExpressionFactory>> FilterExpressionFactories;
 
-	/** Holds all the PlaylistFilterSuggestionFactory */
-	TMap<FName, TSharedPtr<IAvaPlaylistFilterSuggestionFactory>> FilterSuggestionFactories;
+	/** Holds all the RundownFilterSuggestionFactory */
+	TMap<FName, TSharedPtr<IAvaRundownFilterSuggestionFactory>> FilterSuggestionFactories;
 };
diff --git a/Engine/Plugins/Experimental/Avalanche/Source/AvalancheMediaEditor/Private/AvaMediaEditorStyle.cpp b/Engine/Plugins/Experimental/Avalanche/Source/AvalancheMediaEditor/Private/AvaMediaEditorStyle.cpp
--- a/Engine/Plugins/Experimental/Avalanche/Source/AvalancheMediaEditor/Private/AvaMediaEditorStyle.cpp
+++ b/Engine/Plugins/Experimental/Avalanche/Source/AvalancheMediaEditor/Private/AvaMediaEditorStyle.cpp
@@ -1,13 +1,13 @@
 // Copyright Epic Games, Inc. All Rights Reserved.
 
 #include "AvaMediaEditorStyle.h"
+
 #include "Brushes/SlateImageBrush.h"
 #include "Interfaces/IPluginManager.h"
 #include "Styling/AppStyle.h"
 #include "Styling/SlateStyle.h"
-#include "Styling/SlateTypes.h"
 #include "Styling/SlateStyleRegistry.h"
-#include "Styling/StarshipCoreStyle.h"
+#include "Styling/SlateTypes.h"
 #include "Styling/ToolBarStyle.h"
 
 TSharedPtr<FSlateStyleSet> FAvaMediaEditorStyle::StyleInstance = nullptr;
@@ -30,7 +30,7 @@ void FAvaMediaEditorStyle::Shutdown()
 
 FName FAvaMediaEditorStyle::GetStyleSetName()
 {
-	static FName StyleSetName(TEXT("AvalancheMediaEditor"));
+	static FName StyleSetName(TEXT("AvaMediaEditor"));
 	return StyleSetName;
 }
 
@@ -58,9 +58,9 @@ const FVector2D Icon64x64(64.0f, 64.0f);
 
 TSharedRef<FSlateStyleSet> FAvaMediaEditorStyle::Create()
 {
-	TSharedRef<FSlateStyleSet> Style = MakeShared<FSlateStyleSet>("AvalancheMediaEditor");
+	TSharedRef<FSlateStyleSet> Style = MakeShared<FSlateStyleSet>("AvaMediaEditor");
 
-	TSharedPtr<IPlugin> Plugin = IPluginManager::Get().FindPlugin(TEXT("Avalanche"));
+	TSharedPtr<IPlugin> Plugin = IPluginManager::Get().FindPlugin(UE_PLUGIN_NAME);
 	check(Plugin.IsValid());
 	if (Plugin.IsValid())
 	{
@@ -68,102 +68,102 @@ TSharedRef<FSlateStyleSet> FAvaMediaEditorStyle::Create()
 	}
 
 	// Asset Class Icons
-	Style->Set("ClassIcon.AvalanchePlayback", new IMAGE_BRUSH("Icons/MediaIcons/AvaPlaybackIcon_16x"     , Icon16x16));
-	Style->Set("ClassThumbnail.AvalanchePlayback", new IMAGE_BRUSH("Icons/MediaIcons/AvaPlaybackIcon_64x", Icon64x64));
-	Style->Set("ClassIcon.AvalanchePlaylist", new IMAGE_BRUSH("Icons/MediaIcons/AvaRundownIcon_16x"     , Icon16x16));
-	Style->Set("ClassThumbnail.AvalanchePlaylist", new IMAGE_BRUSH("Icons/MediaIcons/AvaRundownIcon_64x", Icon64x64));
+	Style->Set("ClassIcon.AvaPlaybackGraph", new IMAGE_BRUSH("Icons/MediaIcons/AvaPlaybackIcon_16x"     , Icon16x16));
+	Style->Set("ClassThumbnail.AvaPlaybackGraph", new IMAGE_BRUSH("Icons/MediaIcons/AvaPlaybackIcon_64x", Icon64x64));
+	Style->Set("ClassIcon.AvaRundown", new IMAGE_BRUSH("Icons/MediaIcons/AvaRundownIcon_16x"     , Icon16x16));
+	Style->Set("ClassThumbnail.AvaRundown", new IMAGE_BRUSH("Icons/MediaIcons/AvaRundownIcon_64x", Icon64x64));
 	Style->Set("ClassIcon.AvaRundownMacroCollection", new IMAGE_BRUSH("Icons/MediaIcons/AvaRundownMacroCollectionIcon_16x"     , Icon16x16));
 	Style->Set("ClassThumbnail.AvaRundownMacroCollection", new IMAGE_BRUSH("Icons/MediaIcons/AvaRundownMacroCollectionIcon_64x", Icon64x64));
 
-	//Playlist Commands
-	Style->Set("AvaPlaylistCommands.AddPage"         , new IMAGE_BRUSH_SVG("Icons/MediaIcons/AddPage"      , Icon16x16));
-	Style->Set("AvaPlaylistCommands.AddTemplate"     , new IMAGE_BRUSH_SVG("Icons/MediaIcons/AddPage"      , Icon16x16));
-	Style->Set("AvaPlaylistCommands.CreatePageInstanceFromTemplate"  , new IMAGE_BRUSH_SVG("Icons/MediaIcons/AddPage"      , Icon16x16));
-	Style->Set("AvaPlaylistCommands.CreateComboTemplate", new IMAGE_BRUSH_SVG("Icons/MediaIcons/AddPage"   , Icon16x16));
-	Style->Set("AvaPlaylistCommands.RemovePage"      , new IMAGE_BRUSH_SVG("Icons/MediaIcons/RemovePage"   , Icon16x16));
-	Style->Set("AvaPlaylistCommands.RenumberPage"    , new IMAGE_BRUSH_SVG("Icons/MediaIcons/RenumberPage" , Icon16x16));
-	Style->Set("AvaPlaylistCommands.ReimportPage"    , new IMAGE_BRUSH_SVG("Icons/MediaIcons/ReimportPage" , Icon16x16));
-	Style->Set("AvaPlaylistCommands.EditPageSource"  , new IMAGE_BRUSH_SVG("Icons/MediaIcons/OpenAsset"	  , Icon16x16));
-	Style->Set("AvaPlaylistCommands.Play"            , new IMAGE_BRUSH_SVG("Icons/MediaIcons/Play"         , Icon16x16));
-	Style->Set("AvaPlaylistCommands.UpdateValues"    , new IMAGE_BRUSH_SVG("Icons/MediaIcons/UpdateValues" , Icon16x16));
-	Style->Set("AvaPlaylistCommands.Stop"            , new IMAGE_BRUSH_SVG("Icons/MediaIcons/Stop"         , Icon16x16));
-	Style->Set("AvaPlaylistCommands.ForceStop"       , new IMAGE_BRUSH_SVG("Icons/MediaIcons/Stop"         , Icon16x16));
-	Style->Set("AvaPlaylistCommands.Continue"        , new IMAGE_BRUSH_SVG("Icons/MediaIcons/Continue"     , Icon16x16));
-	Style->Set("AvaPlaylistCommands.PlayNext"        , new IMAGE_BRUSH_SVG("Icons/MediaIcons/PlayNext"     , Icon16x16));
-	Style->Set("AvaPlaylistCommands.PreviewFrame"    , new IMAGE_BRUSH_SVG("Icons/MediaIcons/PreviewFrame" , Icon16x16));
-	Style->Set("AvaPlaylistCommands.PreviewPlay"     , new IMAGE_BRUSH_SVG("Icons/MediaIcons/Play"         , Icon16x16));
-	Style->Set("AvaPlaylistCommands.PreviewStop"     , new IMAGE_BRUSH_SVG("Icons/MediaIcons/Stop"         , Icon16x16));
-	Style->Set("AvaPlaylistCommands.PreviewForceStop", new IMAGE_BRUSH_SVG("Icons/MediaIcons/Stop"         , Icon16x16));
-	Style->Set("AvaPlaylistCommands.PreviewContinue" , new IMAGE_BRUSH_SVG("Icons/MediaIcons/Continue"     , Icon16x16));
-	Style->Set("AvaPlaylistCommands.PreviewPlayNext" , new IMAGE_BRUSH_SVG("Icons/MediaIcons/PlayNext"     , Icon16x16));
-	Style->Set("AvaPlaylistCommands.TakeToProgram"   , new IMAGE_BRUSH_SVG("Icons/MediaIcons/Play"		   , Icon16x16));
+	//Rundown Commands
+	Style->Set("AvaRundownCommands.AddPage"         , new IMAGE_BRUSH_SVG("Icons/MediaIcons/AddPage"      , Icon16x16));
+	Style->Set("AvaRundownCommands.AddTemplate"     , new IMAGE_BRUSH_SVG("Icons/MediaIcons/AddPage"      , Icon16x16));
+	Style->Set("AvaRundownCommands.CreatePageInstanceFromTemplate"  , new IMAGE_BRUSH_SVG("Icons/MediaIcons/AddPage"      , Icon16x16));
+	Style->Set("AvaRundownCommands.CreateComboTemplate", new IMAGE_BRUSH_SVG("Icons/MediaIcons/AddPage"   , Icon16x16));
+	Style->Set("AvaRundownCommands.RemovePage"      , new IMAGE_BRUSH_SVG("Icons/MediaIcons/RemovePage"   , Icon16x16));
+	Style->Set("AvaRundownCommands.RenumberPage"    , new IMAGE_BRUSH_SVG("Icons/MediaIcons/RenumberPage" , Icon16x16));
+	Style->Set("AvaRundownCommands.ReimportPage"    , new IMAGE_BRUSH_SVG("Icons/MediaIcons/ReimportPage" , Icon16x16));
+	Style->Set("AvaRundownCommands.EditPageSource"  , new IMAGE_BRUSH_SVG("Icons/MediaIcons/OpenAsset"	  , Icon16x16));
+	Style->Set("AvaRundownCommands.Play"            , new IMAGE_BRUSH_SVG("Icons/MediaIcons/Play"         , Icon16x16));
+	Style->Set("AvaRundownCommands.UpdateValues"    , new IMAGE_BRUSH_SVG("Icons/MediaIcons/UpdateValues" , Icon16x16));
+	Style->Set("AvaRundownCommands.Stop"            , new IMAGE_BRUSH_SVG("Icons/MediaIcons/Stop"         , Icon16x16));
+	Style->Set("AvaRundownCommands.ForceStop"       , new IMAGE_BRUSH_SVG("Icons/MediaIcons/Stop"         , Icon16x16));
+	Style->Set("AvaRundownCommands.Continue"        , new IMAGE_BRUSH_SVG("Icons/MediaIcons/Continue"     , Icon16x16));
+	Style->Set("AvaRundownCommands.PlayNext"        , new IMAGE_BRUSH_SVG("Icons/MediaIcons/PlayNext"     , Icon16x16));
+	Style->Set("AvaRundownCommands.PreviewFrame"    , new IMAGE_BRUSH_SVG("Icons/MediaIcons/PreviewFrame" , Icon16x16));
+	Style->Set("AvaRundownCommands.PreviewPlay"     , new IMAGE_BRUSH_SVG("Icons/MediaIcons/Play"         , Icon16x16));
+	Style->Set("AvaRundownCommands.PreviewStop"     , new IMAGE_BRUSH_SVG("Icons/MediaIcons/Stop"         , Icon16x16));
+	Style->Set("AvaRundownCommands.PreviewForceStop", new IMAGE_BRUSH_SVG("Icons/MediaIcons/Stop"         , Icon16x16));
+	Style->Set("AvaRundownCommands.PreviewContinue" , new IMAGE_BRUSH_SVG("Icons/MediaIcons/Continue"     , Icon16x16));
+	Style->Set("AvaRundownCommands.PreviewPlayNext" , new IMAGE_BRUSH_SVG("Icons/MediaIcons/PlayNext"     , Icon16x16));
+	Style->Set("AvaRundownCommands.TakeToProgram"   , new IMAGE_BRUSH_SVG("Icons/MediaIcons/Play"		   , Icon16x16));
 
 	//Broadcast
-	Style->Set("AvalancheMediaEditor.OutputIcon"           , new IMAGE_BRUSH("Icons/MediaIcons/MediaOutput"     , Icon20x20));
-	Style->Set("AvalancheMediaEditor.BroadcastIcon"        , new IMAGE_BRUSH("Icons/MediaIcons/MediaOutput"     , Icon20x20));
-	Style->Set("AvalancheMediaEditor.BroadcastClient"      , new IMAGE_BRUSH("Icons/MediaIcons/BroadcastClient" , Icon64x64));
-	Style->Set("AvalancheMediaEditor.BroadcastServer"      , new IMAGE_BRUSH("Icons/MediaIcons/BroadcastServer" , Icon64x64));
-	Style->Set("AvalancheMediaEditor.BroadcastClient.Small", new IMAGE_BRUSH("Icons/MediaIcons/BroadcastClient" , Icon20x20));
-	Style->Set("AvalancheMediaEditor.BroadcastServer.Small", new IMAGE_BRUSH("Icons/MediaIcons/BroadcastServer" , Icon20x20));
-	Style->Set("AvalancheMediaEditor.BroadcastOffline"     , new IMAGE_BRUSH_SVG("Icons/MediaIcons/BroadcastOffline"	, Icon16x16));
-	Style->Set("AvalancheMediaEditor.BroadcastError"       , new IMAGE_BRUSH_SVG("Icons/MediaIcons/BroadcastError"		, Icon16x16));
-	Style->Set("AvalancheMediaEditor.BroadcastWarning"     , new IMAGE_BRUSH_SVG("Icons/MediaIcons/BroadcastWarning"	, Icon16x16));
-	Style->Set("AvalancheMediaEditor.BroadcastLive"        , new IMAGE_BRUSH_SVG("Icons/MediaIcons/BroadcastLive"		, Icon16x16)); 
-	Style->Set("AvalancheMediaEditor.BroadcastIdle"        , new IMAGE_BRUSH_SVG("Icons/MediaIcons/BroadcastIdle"		, Icon16x16));
+	Style->Set("AvaMediaEditor.OutputIcon"           , new IMAGE_BRUSH("Icons/MediaIcons/MediaOutput"     , Icon20x20));
+	Style->Set("AvaMediaEditor.BroadcastIcon"        , new IMAGE_BRUSH("Icons/MediaIcons/MediaOutput"     , Icon20x20));
+	Style->Set("AvaMediaEditor.BroadcastClient"      , new IMAGE_BRUSH("Icons/MediaIcons/BroadcastClient" , Icon64x64));
+	Style->Set("AvaMediaEditor.BroadcastServer"      , new IMAGE_BRUSH("Icons/MediaIcons/BroadcastServer" , Icon64x64));
+	Style->Set("AvaMediaEditor.BroadcastClient.Small", new IMAGE_BRUSH("Icons/MediaIcons/BroadcastClient" , Icon20x20));
+	Style->Set("AvaMediaEditor.BroadcastServer.Small", new IMAGE_BRUSH("Icons/MediaIcons/BroadcastServer" , Icon20x20));
+	Style->Set("AvaMediaEditor.BroadcastOffline"     , new IMAGE_BRUSH_SVG("Icons/MediaIcons/BroadcastOffline"	, Icon16x16));
+	Style->Set("AvaMediaEditor.BroadcastError"       , new IMAGE_BRUSH_SVG("Icons/MediaIcons/BroadcastError"		, Icon16x16));
+	Style->Set("AvaMediaEditor.BroadcastWarning"     , new IMAGE_BRUSH_SVG("Icons/MediaIcons/BroadcastWarning"	, Icon16x16));
+	Style->Set("AvaMediaEditor.BroadcastLive"        , new IMAGE_BRUSH_SVG("Icons/MediaIcons/BroadcastLive"		, Icon16x16)); 
+	Style->Set("AvaMediaEditor.BroadcastIdle"        , new IMAGE_BRUSH_SVG("Icons/MediaIcons/BroadcastIdle"		, Icon16x16));
 
 	// Broadcast Channel Types
-	Style->Set("AvalancheMediaEditor.ChannelTypeProgram"	, new IMAGE_BRUSH_SVG("Icons/MediaIcons/ChannelTypeProgram"	, Icon16x16)); 
-	Style->Set("AvalancheMediaEditor.ChannelTypePreview"	, new IMAGE_BRUSH_SVG("Icons/MediaIcons/ChannelTypePreview"	, Icon16x16));
+	Style->Set("AvaMediaEditor.ChannelTypeProgram"	, new IMAGE_BRUSH_SVG("Icons/MediaIcons/ChannelTypeProgram"	, Icon16x16)); 
+	Style->Set("AvaMediaEditor.ChannelTypePreview"	, new IMAGE_BRUSH_SVG("Icons/MediaIcons/ChannelTypePreview"	, Icon16x16));
 	
 	// Media Status
 	/* For when we have SVGs
-	Style->Set("AvalancheMediaEditor.MediaAssetStatus", new IMAGE_BRUSH_SVG("Icons/MediaIcons/MediaLoaded", Icon16x16));
-	Style->Set("AvalancheMediaEditor.MediaSyncStatus",  new IMAGE_BRUSH_SVG("Icons/MediaIcons/MediaLoaded", Icon16x16));
-	Style->Set("AvalancheMediaEditor.MediaPreviewing",  new IMAGE_BRUSH_SVG("Icons/MediaIcons/MediaLoaded", Icon16x16));
-	Style->Set("AvalancheMediaEditor.MediaPlaying",     new IMAGE_BRUSH_SVG("Icons/MediaIcons/MediaLoaded", Icon16x16));
+	Style->Set("AvaMediaEditor.MediaAssetStatus", new IMAGE_BRUSH_SVG("Icons/MediaIcons/MediaLoaded", Icon16x16));
+	Style->Set("AvaMediaEditor.MediaSyncStatus",  new IMAGE_BRUSH_SVG("Icons/MediaIcons/MediaLoaded", Icon16x16));
+	Style->Set("AvaMediaEditor.MediaPreviewing",  new IMAGE_BRUSH_SVG("Icons/MediaIcons/MediaLoaded", Icon16x16));
+	Style->Set("AvaMediaEditor.MediaPlaying",     new IMAGE_BRUSH_SVG("Icons/MediaIcons/MediaLoaded", Icon16x16));
 	*/
-	Style->Set("AvalancheMediaEditor.MediaAssetStatus", new IMAGE_BRUSH("Icons/MediaIcons/MediaStatus",     Icon16x16));
-	Style->Set("AvalancheMediaEditor.MediaSyncStatus",  new IMAGE_BRUSH("Icons/MediaIcons/MediaSyncStatus", Icon16x16));
-	Style->Set("AvalancheMediaEditor.MediaPreviewing",  new IMAGE_BRUSH("Icons/MediaIcons/MediaPreviewing", Icon16x16));
-	Style->Set("AvalancheMediaEditor.MediaPlaying",     new IMAGE_BRUSH("Icons/MediaIcons/MediaPlaying",    Icon16x16));
+	Style->Set("AvaMediaEditor.MediaAssetStatus", new IMAGE_BRUSH("Icons/MediaIcons/MediaStatus",     Icon16x16));
+	Style->Set("AvaMediaEditor.MediaSyncStatus",  new IMAGE_BRUSH("Icons/MediaIcons/MediaSyncStatus", Icon16x16));
+	Style->Set("AvaMediaEditor.MediaPreviewing",  new IMAGE_BRUSH("Icons/MediaIcons/MediaPreviewing", Icon16x16));
+	Style->Set("AvaMediaEditor.MediaPlaying",     new IMAGE_BRUSH("Icons/MediaIcons/MediaPlaying",    Icon16x16));
 
 	// Media Output Status
-	Style->Set("AvalancheMediaEditor.MediaOutputOffline"   , new IMAGE_BRUSH_SVG("Icons/MediaIcons/MediaOutputOffline"	, Icon16x16));
-	Style->Set("AvalancheMediaEditor.MediaOutputIdle"      , new IMAGE_BRUSH_SVG("Icons/MediaIcons/MediaOutputIdle"		, Icon16x16));
-	Style->Set("AvalancheMediaEditor.MediaOutputPreparing" , new IMAGE_BRUSH_SVG("Icons/MediaIcons/MediaOutputPreparing", Icon16x16));
-	Style->Set("AvalancheMediaEditor.MediaOutputLive"      , new IMAGE_BRUSH_SVG("Icons/MediaIcons/MediaOutputLive"		, Icon16x16));
-	Style->Set("AvalancheMediaEditor.MediaOutputLiveWarn"  , new IMAGE_BRUSH_SVG("Icons/MediaIcons/MediaOutputLiveWarn"	, Icon16x16));
-	Style->Set("AvalancheMediaEditor.MediaOutputError"     , new IMAGE_BRUSH_SVG("Icons/MediaIcons/MediaOutputError"	, Icon16x16));
+	Style->Set("AvaMediaEditor.MediaOutputOffline"   , new IMAGE_BRUSH_SVG("Icons/MediaIcons/MediaOutputOffline"	, Icon16x16));
+	Style->Set("AvaMediaEditor.MediaOutputIdle"      , new IMAGE_BRUSH_SVG("Icons/MediaIcons/MediaOutputIdle"		, Icon16x16));
+	Style->Set("AvaMediaEditor.MediaOutputPreparing" , new IMAGE_BRUSH_SVG("Icons/MediaIcons/MediaOutputPreparing", Icon16x16));
+	Style->Set("AvaMediaEditor.MediaOutputLive"      , new IMAGE_BRUSH_SVG("Icons/MediaIcons/MediaOutputLive"		, Icon16x16));
+	Style->Set("AvaMediaEditor.MediaOutputLiveWarn"  , new IMAGE_BRUSH_SVG("Icons/MediaIcons/MediaOutputLiveWarn"	, Icon16x16));
+	Style->Set("AvaMediaEditor.MediaOutputError"     , new IMAGE_BRUSH_SVG("Icons/MediaIcons/MediaOutputError"	, Icon16x16));
 
-	// Avalanche Preview
-	Style->Set("AvalancheMediaEditor.Checkerboard" , new IMAGE_BRUSH("Images/AvaPreviewCheckerboard", Icon16x16, FLinearColor::White, ESlateBrushTileType::Both));
+	// Motion Design Preview
+	Style->Set("AvaMediaEditor.Checkerboard" , new IMAGE_BRUSH("Images/AvaPreviewCheckerboard", Icon16x16, FLinearColor::White, ESlateBrushTileType::Both));
 
 	// Button styles
 	const FButtonStyle& AppStyle_SimpleButton = FAppStyle::Get().GetWidgetStyle<FButtonStyle>("SimpleButton");
 
-	Style->Set("AvalancheMediaEditor.BorderlessButton", FButtonStyle(AppStyle_SimpleButton));
+	Style->Set("AvaMediaEditor.BorderlessButton", FButtonStyle(AppStyle_SimpleButton));
 
 	const FButtonStyle& MenuButtonStyleGreen = FAppStyle::Get().GetWidgetStyle<FButtonStyle>("FlatButton.Success");
-	Style->Set("AvalancheMediaEditor.ButtonGreen", MenuButtonStyleGreen);
+	Style->Set("AvaMediaEditor.ButtonGreen", MenuButtonStyleGreen);
 	const FButtonStyle& MenuButtonStyleRed = FAppStyle::Get().GetWidgetStyle<FButtonStyle>("FlatButton.Danger");
-	Style->Set("AvalancheMediaEditor.ButtonRed", MenuButtonStyleRed);
+	Style->Set("AvaMediaEditor.ButtonRed", MenuButtonStyleRed);
 
 	const FToolBarStyle& SlimToolbarStyle = FAppStyle::Get().GetWidgetStyle<FToolBarStyle>("SlimToolBar");
-	Style->Set("AvalancheMediaEditor.ToolBar", SlimToolbarStyle);
+	Style->Set("AvaMediaEditor.ToolBar", SlimToolbarStyle);
 	const FButtonStyle& MenuButtonStyle = SlimToolbarStyle.ButtonStyle;
 
 	FToolBarStyle SlimToolbarStyleOverrideRedButton = FAppStyle::Get().GetWidgetStyle<FToolBarStyle>("SlimToolBar");
-	SlimToolbarStyleOverrideRedButton.SetButtonStyle(Style->GetWidgetStyle<FButtonStyle>("AvalancheMediaEditor.ButtonRed"));
+	SlimToolbarStyleOverrideRedButton.SetButtonStyle(Style->GetWidgetStyle<FButtonStyle>("AvaMediaEditor.ButtonRed"));
 	SlimToolbarStyleOverrideRedButton.ButtonStyle.SetDisabled(MenuButtonStyle.Disabled);
-	Style->Set("AvalancheMediaEditor.ToolBarRedButtonOverride", SlimToolbarStyleOverrideRedButton);
+	Style->Set("AvaMediaEditor.ToolBarRedButtonOverride", SlimToolbarStyleOverrideRedButton);
 
 	FToolBarStyle SlimToolbarStyleOverrideGreenButton = FAppStyle::Get().GetWidgetStyle<FToolBarStyle>("SlimToolBar");
-	SlimToolbarStyleOverrideGreenButton.SetButtonStyle(Style->GetWidgetStyle<FButtonStyle>("AvalancheMediaEditor.ButtonGreen"));
+	SlimToolbarStyleOverrideGreenButton.SetButtonStyle(Style->GetWidgetStyle<FButtonStyle>("AvaMediaEditor.ButtonGreen"));
 	SlimToolbarStyleOverrideGreenButton.ButtonStyle.SetDisabled(MenuButtonStyle.Disabled);
-	Style->Set("AvalancheMediaEditor.ToolBarGreenButtonOverride", SlimToolbarStyleOverrideGreenButton);
+	Style->Set("AvaMediaEditor.ToolBarGreenButtonOverride", SlimToolbarStyleOverrideGreenButton);
 
 	const FToolBarStyle& CalloutToolbarStyleOverride = FAppStyle::Get().GetWidgetStyle<FToolBarStyle>("CalloutToolbar");
-	Style->Set("AvalancheMediaEditor.CalloutToolbar", CalloutToolbarStyleOverride);
+	Style->Set("AvaMediaEditor.CalloutToolbar", CalloutToolbarStyleOverride);
 
 	Style->Set("TableView.ActivePageBrush", new FSlateColorBrush(FLinearColor(0.05f, 0.1f, 0.05f, 0.5f)));
 	Style->Set("TableView.Hovered.ActivePageBrush", new FSlateColorBrush(FLinearColor(0.05f, 0.2f, 0.05f, 0.5f)));
@@ -173,9 +173,9 @@ TSharedRef<FSlateStyleSet> FAvaMediaEditorStyle::Create()
 	Style->Set("TableView.Selected.DisabledPageBrush", new FSlateColorBrush(FLinearColor(0.3f, 0.05f, 0.05f, 0.5f)));
 
 	// Asset Colors - in the shades of yellow (hue = 36, saturation = 221)
-	Style->Set("AvalancheMediaEditor.AssetColors.Rundown", FLinearColor::MakeFromHSV8(36, 221, 170));
-	Style->Set("AvalancheMediaEditor.AssetColors.RundownMacroCollection", FLinearColor::MakeFromHSV8(36, 221, 124));
-	Style->Set("AvalancheMediaEditor.AssetColors.Playback", FLinearColor::MakeFromHSV8(36, 221, 105));
+	Style->Set("AvaMediaEditor.AssetColors.Rundown", FLinearColor::MakeFromHSV8(36, 221, 170));
+	Style->Set("AvaMediaEditor.AssetColors.RundownMacroCollection", FLinearColor::MakeFromHSV8(36, 221, 124));
+	Style->Set("AvaMediaEditor.AssetColors.Playback", FLinearColor::MakeFromHSV8(36, 221, 105));
 
 	return Style;
 }
diff --git a/Engine/Plugins/Experimental/Avalanche/Source/AvalancheMediaEditor/Private/AvaMediaEditorUtils.cpp b/Engine/Plugins/Experimental/Avalanche/Source/AvalancheMediaEditor/Private/AvaMediaEditorUtils.cpp
--- a/Engine/Plugins/Experimental/Avalanche/Source/AvalancheMediaEditor/Private/AvaMediaEditorUtils.cpp
+++ b/Engine/Plugins/Experimental/Avalanche/Source/AvalancheMediaEditor/Private/AvaMediaEditorUtils.cpp
@@ -2,34 +2,34 @@
 
 #include "AvaMediaEditorUtils.h"
 
-const FSlateBrush* FAvaMediaEditorUtils::GetChannelStatusBrush(EAvaChannelState InChannelState, EAvaMediaIssueSeverity InChannelIssueSeverity)
+const FSlateBrush* FAvaMediaEditorUtils::GetChannelStatusBrush(EAvaBroadcastChannelState InChannelState, EAvaBroadcastIssueSeverity InChannelIssueSeverity)
 {
-	if (InChannelState == EAvaChannelState::Live)
+	if (InChannelState == EAvaBroadcastChannelState::Live)
 	{
 		switch (InChannelIssueSeverity)
 		{
-		case EAvaMediaIssueSeverity::None:
-			return FAvaMediaEditorStyle::Get().GetBrush("AvalancheMediaEditor.BroadcastLive");
-		case EAvaMediaIssueSeverity::Warnings:
-			return FAvaMediaEditorStyle::Get().GetBrush("AvalancheMediaEditor.BroadcastWarning");
-		case EAvaMediaIssueSeverity::Errors:
-			return FAvaMediaEditorStyle::Get().GetBrush("AvalancheMediaEditor.BroadcastError");
+		case EAvaBroadcastIssueSeverity::None:
+			return FAvaMediaEditorStyle::Get().GetBrush("AvaMediaEditor.BroadcastLive");
+		case EAvaBroadcastIssueSeverity::Warnings:
+			return FAvaMediaEditorStyle::Get().GetBrush("AvaMediaEditor.BroadcastWarning");
+		case EAvaBroadcastIssueSeverity::Errors:
+			return FAvaMediaEditorStyle::Get().GetBrush("AvaMediaEditor.BroadcastError");
 		}
 	}
-	else if (InChannelState == EAvaChannelState::Idle)
+	else if (InChannelState == EAvaBroadcastChannelState::Idle)
 	{
-		return FAvaMediaEditorStyle::Get().GetBrush("AvalancheMediaEditor.BroadcastIdle");
+		return FAvaMediaEditorStyle::Get().GetBrush("AvaMediaEditor.BroadcastIdle");
 	}
 	else
 	{
-		return FAvaMediaEditorStyle::Get().GetBrush("AvalancheMediaEditor.BroadcastOffline");
+		return FAvaMediaEditorStyle::Get().GetBrush("AvaMediaEditor.BroadcastOffline");
 	}
 	
-	return FAvaMediaEditorStyle::Get().GetBrush("AvalancheMediaEditor.BroadcastIdle");
+	return FAvaMediaEditorStyle::Get().GetBrush("AvaMediaEditor.BroadcastIdle");
 }
 
-FText FAvaMediaEditorUtils::GetChannelStatusText(EAvaChannelState InChannelState, EAvaMediaIssueSeverity InChannelIssueSeverity)
+FText FAvaMediaEditorUtils::GetChannelStatusText(EAvaBroadcastChannelState InChannelState, EAvaBroadcastIssueSeverity InChannelIssueSeverity)
 {
 	//TODO: Possibly also add Issue Severity (if not none)
-	return StaticEnum<EAvaChannelState>()->GetDisplayNameTextByIndex(static_cast<int32>(InChannelState));
+	return StaticEnum<EAvaBroadcastChannelState>()->GetDisplayNameTextByIndex(static_cast<int32>(InChannelState));
 }
diff --git a/Engine/Plugins/Experimental/Avalanche/Source/AvalancheMediaEditor/Private/AvaMediaEditorUtils.h b/Engine/Plugins/Experimental/Avalanche/Source/AvalancheMediaEditor/Private/AvaMediaEditorUtils.h
--- a/Engine/Plugins/Experimental/Avalanche/Source/AvalancheMediaEditor/Private/AvaMediaEditorUtils.h
+++ b/Engine/Plugins/Experimental/Avalanche/Source/AvalancheMediaEditor/Private/AvaMediaEditorUtils.h
@@ -8,6 +8,6 @@
 class FAvaMediaEditorUtils
 {
 public:
-	static const FSlateBrush* GetChannelStatusBrush(EAvaChannelState InChannelState, EAvaMediaIssueSeverity InChannelIssueSeverity);
-	static FText GetChannelStatusText(EAvaChannelState InChannelState, EAvaMediaIssueSeverity InChannelIssueSeverity);
+	static const FSlateBrush* GetChannelStatusBrush(EAvaBroadcastChannelState InChannelState, EAvaBroadcastIssueSeverity InChannelIssueSeverity);
+	static FText GetChannelStatusText(EAvaBroadcastChannelState InChannelState, EAvaBroadcastIssueSeverity InChannelIssueSeverity);
 };
diff --git a/Engine/Plugins/Experimental/Avalanche/Source/AvalancheMediaEditor/Private/Broadcast/AppModes/AvaBroadcastDefaultMode.cpp b/Engine/Plugins/Experimental/Avalanche/Source/AvalancheMediaEditor/Private/Broadcast/AppModes/AvaBroadcastDefaultMode.cpp
--- a/Engine/Plugins/Experimental/Avalanche/Source/AvalancheMediaEditor/Private/Broadcast/AppModes/AvaBroadcastDefaultMode.cpp
+++ b/Engine/Plugins/Experimental/Avalanche/Source/AvalancheMediaEditor/Private/Broadcast/AppModes/AvaBroadcastDefaultMode.cpp
@@ -1,6 +1,7 @@
 // Copyright Epic Games, Inc. All Rights Reserved.
 
 #include "AvaBroadcastDefaultMode.h"
+
 #include "Broadcast/AvaBroadcastEditor.h"
 #include "Broadcast/TabFactories/AvaBroadcastChannelsTabFactory.h"
 #include "Broadcast/TabFactories/AvaBroadcastDetailsTabFactory.h"
@@ -15,7 +16,7 @@ FAvaBroadcastDefaultMode::FAvaBroadcastDefaultMode(const TSharedPtr<FAvaBroadcas
 	WorkspaceMenuCategory = FWorkspaceItem::NewGroup(LOCTEXT("WorkspaceMenu_AvaBroadcast", "Motion Design Broadcast"));
 	
 	check(InBroadcastEditor.IsValid());
-	TabLayout = FTabManager::NewLayout("AvalancheBroadcastEditor_Default_Layout_V1")
+	TabLayout = FTabManager::NewLayout("MotionDesignBroadcastEditor_Default_Layout_V1")
 		->AddArea
 		(
 			FTabManager::NewPrimaryArea() ->SetOrientation(Orient_Vertical)
diff --git a/Engine/Plugins/Experimental/Avalanche/Source/AvalancheMediaEditor/Private/Broadcast/AvaBroadcastEditor.cpp b/Engine/Plugins/Experimental/Avalanche/Source/AvalancheMediaEditor/Private/Broadcast/AvaBroadcastEditor.cpp
--- a/Engine/Plugins/Experimental/Avalanche/Source/AvalancheMediaEditor/Private/Broadcast/AvaBroadcastEditor.cpp
+++ b/Engine/Plugins/Experimental/Avalanche/Source/AvalancheMediaEditor/Private/Broadcast/AvaBroadcastEditor.cpp
@@ -3,19 +3,19 @@
 #include "Broadcast/AvaBroadcastEditor.h"
 
 #include "AppModes/AvaBroadcastDefaultMode.h"
-#include "AvalancheBroadcast.h"
-#include "AvalancheMediaSettings.h"
+#include "AvaMediaSettings.h"
+#include "Broadcast/AvaBroadcast.h"
 #include "ChannelGrid/Slate/SAvaBroadcastProfileEntry.h"
 #include "Framework/Application/SlateApplication.h"
 #include "Framework/MultiBox/MultiBoxBuilder.h"
 #include "IAvaMediaEditorModule.h"
 #include "IAvaMediaModule.h"
 #include "Misc/MessageDialog.h"
-#include "Playback/IAvaMediaPlaybackClient.h"
+#include "Playback/IAvaPlaybackClient.h"
 #include "ScopedTransaction.h"
 #include "WorkflowOrientedApp/ApplicationMode.h"
 
-#define LOCTEXT_NAMESPACE "FAvaBroadcastEditor"
+#define LOCTEXT_NAMESPACE "AvaBroadcastEditor"
 
 TSharedPtr<FAvaBroadcastEditor> FAvaBroadcastEditor::BroadcastEditor;
 
@@ -23,13 +23,13 @@ FAvaBroadcastEditor::~FAvaBroadcastEditor()
 {
 	if (UObjectInitialized())
 	{
-		UAvalancheMediaSettings& AvaMediaSettings = UAvalancheMediaSettings::GetMutable();
+		UAvaMediaSettings& AvaMediaSettings = UAvaMediaSettings::GetMutable();
 		AvaMediaSettings.OnSettingChanged().RemoveAll(this);
 	}
 	
-	if (AvalancheBroadcast.IsValid())
+	if (BroadcastWeak.IsValid())
 	{
-		AvalancheBroadcast->RemoveChangeListener(this);
+		BroadcastWeak->RemoveChangeListener(this);
 	}
 }
 
@@ -43,11 +43,11 @@ void FAvaBroadcastEditor::OpenBroadcastEditor()
 	{
 		BroadcastEditor.Reset();
 		BroadcastEditor = MakeShared<FAvaBroadcastEditor>(FPrivateToken());
-		BroadcastEditor->InitBroadcastEditor(&UAvalancheBroadcast::Get());
+		BroadcastEditor->InitBroadcastEditor(&UAvaBroadcast::Get());
 	}
 }
 
-void FAvaBroadcastEditor::SelectOutputTile(const TSharedPtr<FAvaOutputTileItem>& InOutputTile)
+void FAvaBroadcastEditor::SelectOutputTile(const TSharedPtr<FAvaBroadcastOutputTileItem>& InOutputTile)
 {
 	if (SelectedOutputTileWeak != InOutputTile)
 	{
@@ -56,7 +56,7 @@ void FAvaBroadcastEditor::SelectOutputTile(const TSharedPtr<FAvaOutputTileItem>&
 	}
 }
 
-TSharedPtr<FAvaOutputTileItem> FAvaBroadcastEditor::GetSelectedOutputTile() const
+TSharedPtr<FAvaBroadcastOutputTileItem> FAvaBroadcastEditor::GetSelectedOutputTile() const
 {
 	return SelectedOutputTileWeak.Pin();
 }
@@ -70,9 +70,9 @@ void FAvaBroadcastEditor::OnClose()
 	}
 }
 
-void FAvaBroadcastEditor::InitBroadcastEditor(UAvalancheBroadcast* InBroadcast)
+void FAvaBroadcastEditor::InitBroadcastEditor(UAvaBroadcast* InBroadcast)
 {
-	AvalancheBroadcast = InBroadcast;
+	BroadcastWeak = InBroadcast;
 	
 	if (InBroadcast)
 	{
@@ -80,12 +80,12 @@ void FAvaBroadcastEditor::InitBroadcastEditor(UAvalancheBroadcast* InBroadcast)
 			, &FAvaBroadcastEditor::OnBroadcastChanged));
 	}
 
-	UAvalancheMediaSettings& AvaMediaSettings = UAvalancheMediaSettings::GetMutable();
+	UAvaMediaSettings& AvaMediaSettings = UAvaMediaSettings::GetMutable();
 	AvaMediaSettings.OnSettingChanged().AddRaw(this, &FAvaBroadcastEditor::OnAvaMediaSettingsChanged);
 
 	CreateDefaultCommands();
 	
-	const FName BroadcastEditorAppName(TEXT("AvalancheBroadcastEditorApp"));
+	const FName BroadcastEditorAppName(TEXT("MotionDesignBroadcastEditorApp"));
 	constexpr bool bCreateDefaultStandaloneMenu = true;
 	constexpr bool bCreateDefaultToolbar = true;
 
@@ -104,7 +104,7 @@ void FAvaBroadcastEditor::OnBroadcastChanged(EAvaBroadcastChange ChangedEvent)
 {
 	if (EnumHasAnyFlags(ChangedEvent, EAvaBroadcastChange::CurrentProfile))
 	{
-		TSharedPtr<FAvaOutputTileItem> NullItem;
+		TSharedPtr<FAvaBroadcastOutputTileItem> NullItem;
 		SelectedOutputTileWeak = NullItem;
 		OnOutputTileSelectionChanged.Broadcast(NullItem);
 	}
@@ -112,18 +112,18 @@ void FAvaBroadcastEditor::OnBroadcastChanged(EAvaBroadcastChange ChangedEvent)
 
 void FAvaBroadcastEditor::OnAvaMediaSettingsChanged(UObject*, FPropertyChangedEvent&)
 {
-	if (AvalancheBroadcast.IsValid())
+	if (BroadcastWeak.IsValid())
 	{
-		AvalancheBroadcast->GetCurrentProfile().UpdateChannels(true);
+		BroadcastWeak->GetCurrentProfile().UpdateChannels(true);
 	}
 }
 
 void FAvaBroadcastEditor::SaveAsset_Execute()
 {
 	FWorkflowCentricApplication::SaveAsset_Execute();
-	if (AvalancheBroadcast.IsValid())
+	if (BroadcastWeak.IsValid())
 	{
-		AvalancheBroadcast->SaveBroadcast();
+		BroadcastWeak->SaveBroadcast();
 	}
 }
 
@@ -170,9 +170,9 @@ const FSlateBrush* FAvaBroadcastEditor::GetDefaultTabIcon() const
 	return IAvaMediaEditorModule::Get().GetToolbarBroadcastButtonIcon().GetIcon();
 }
 
-UAvalancheBroadcast* FAvaBroadcastEditor::GetBroadcastObject() const
+UAvaBroadcast* FAvaBroadcastEditor::GetBroadcastObject() const
 {
-	return AvalancheBroadcast.Get();
+	return BroadcastWeak.Get();
 }
 
 void FAvaBroadcastEditor::ExtendToolBar(TSharedPtr<FExtender> Extender)
@@ -185,7 +185,7 @@ void FAvaBroadcastEditor::ExtendToolBar(TSharedPtr<FExtender> Extender)
 
 void FAvaBroadcastEditor::FillPlayToolBar(FToolBarBuilder& ToolBarBuilder)
 {
-	TWeakObjectPtr<UAvalancheBroadcast> Broadcast = GetBroadcastObject();
+	TWeakObjectPtr<UAvaBroadcast> Broadcast = GetBroadcastObject();
 
 	//TODO: Change Lambdas to their Own Command Action
 	ToolBarBuilder.BeginSection(TEXT("Player"));
@@ -341,10 +341,10 @@ void FAvaBroadcastEditor::FillPlayToolBar(FToolBarBuilder& ToolBarBuilder)
 				{
 					// Stop and unload all playbacks.
 					IAvaMediaModule::Get().GetMediaPlaybackClient().RequestPlayback(
-						FGuid(), FSoftObjectPath(), FString(), EAvaMediaPlaybackAction::Unload);
+						FGuid(), FSoftObjectPath(), FString(), EAvaPlaybackAction::Unload);
 					// Stop all broadcast channels.
 					IAvaMediaModule::Get().GetMediaPlaybackClient().RequestBroadcast(
-						TEXT(""), FName(), TArray<UMediaOutput*>(), EAvaMediaBroadcastAction::Stop);
+						TEXT(""), FName(), TArray<UMediaOutput*>(), EAvaBroadcastAction::Stop);
 				})
 				, FCanExecuteAction::CreateLambda([]
 				{
@@ -368,7 +368,7 @@ void FAvaBroadcastEditor::FillPlayToolBar(FToolBarBuilder& ToolBarBuilder)
 
 FText FAvaBroadcastEditor::GetCurrentProfileName() const
 {
-	return FText::FromName(UAvalancheBroadcast::Get().GetCurrentProfileName());
+	return FText::FromName(UAvaBroadcast::Get().GetCurrentProfileName());
 }
 
 void FAvaBroadcastEditor::MakeProfilesToolbar(FToolBarBuilder& ToolBarBuilder)
@@ -377,7 +377,7 @@ void FAvaBroadcastEditor::MakeProfilesToolbar(FToolBarBuilder& ToolBarBuilder)
 	{
 		static const FCanExecuteAction DefaultCanExecuteAction(FCanExecuteAction::CreateLambda([]()
 		{
-			return !UAvalancheBroadcast::Get().IsBroadcastingAnyChannel();
+			return !UAvaBroadcast::Get().IsBroadcastingAnyChannel();
 		}));
 		
 		ToolBarBuilder.AddComboButton(FUIAction(FExecuteAction(), DefaultCanExecuteAction)
@@ -391,7 +391,7 @@ void FAvaBroadcastEditor::MakeProfilesToolbar(FToolBarBuilder& ToolBarBuilder)
 		ToolBarBuilder.AddToolBarButton(FUIAction(FExecuteAction::CreateLambda([]()
 				{
 					FScopedTransaction Transaction(LOCTEXT("DuplicateProfile", "Duplicate Profile"));
-					UAvalancheBroadcast& Broadcast = UAvalancheBroadcast::Get();
+					UAvaBroadcast& Broadcast = UAvaBroadcast::Get();
 					Broadcast.Modify();
 					const bool bResult = Broadcast.DuplicateCurrentProfile();
 					if (!bResult)
@@ -416,7 +416,7 @@ TSharedRef<SWidget> FAvaBroadcastEditor::MakeProfileComboButton()
 
 	MenuBuilder.BeginSection(TEXT("Profiles"), LOCTEXT("Profiles_Title", "Profiles"));
 	{
-		TArray<FName> ProfileNames = UAvalancheBroadcast::Get().GetProfileNames();
+		TArray<FName> ProfileNames = UAvaBroadcast::Get().GetProfileNames();
 		for (const FName& ProfileName : ProfileNames)
 		{
 			MenuBuilder.AddWidget(SNew(SAvaBroadcastProfileEntry, ProfileName)
@@ -443,7 +443,7 @@ TSharedRef<SWidget> FAvaBroadcastEditor::MakeProfileComboButton()
 
 void FAvaBroadcastEditor::CreateNewProfile()
 {
-	if (UAvalancheBroadcast* const Broadcast = AvalancheBroadcast.Get())
+	if (UAvaBroadcast* const Broadcast = BroadcastWeak.Get())
 	{
 		FScopedTransaction Transaction(LOCTEXT("CreateProfile", "Create Profile"));
 		Broadcast->Modify();
@@ -457,7 +457,7 @@ FReply FAvaBroadcastEditor::OnProfileSelected(FName InProfileName)
 
 	FScopedTransaction Transaction(LOCTEXT("SetCurrentProfile", "Set Current Profile"));
 	
-	UAvalancheBroadcast& Broadcast = UAvalancheBroadcast::Get();
+	UAvaBroadcast& Broadcast = UAvaBroadcast::Get();
 	Broadcast.Modify();
 
 	if (Broadcast.SetCurrentProfile(InProfileName))
@@ -493,7 +493,7 @@ void FAvaBroadcastEditor::CreateDefaultCommands()
 
 FText FAvaBroadcastEditor::GetStopPlaybackClientTooltip()
 {
-	if (UAvalancheMediaSettings::Get().bAutoStartPlaybackClient)
+	if (UAvaMediaSettings::Get().bAutoStartPlaybackClient)
 	{
 		return LOCTEXT("PlaybackClientStopAuto_ToolTip", "Stop Playback Client (auto started)");
 	}
diff --git a/Engine/Plugins/Experimental/Avalanche/Source/AvalancheMediaEditor/Private/Broadcast/ChannelGrid/Slate/SAvaBroadcastChannels.cpp b/Engine/Plugins/Experimental/Avalanche/Source/AvalancheMediaEditor/Private/Broadcast/ChannelGrid/Slate/SAvaBroadcastChannels.cpp
--- a/Engine/Plugins/Experimental/Avalanche/Source/AvalancheMediaEditor/Private/Broadcast/ChannelGrid/Slate/SAvaBroadcastChannels.cpp
+++ b/Engine/Plugins/Experimental/Avalanche/Source/AvalancheMediaEditor/Private/Broadcast/ChannelGrid/Slate/SAvaBroadcastChannels.cpp
@@ -1,12 +1,12 @@
 // Copyright Epic Games, Inc. All Rights Reserved.
 
 #include "SAvaBroadcastChannels.h"
-#include "AvalancheBroadcast.h"
-#include "AvalancheMediaEditorSettings.h"
+
+#include "AvaMediaEditorSettings.h"
+#include "Broadcast/AvaBroadcast.h"
 #include "Broadcast/AvaBroadcastEditor.h"
 #include "Framework/MultiBox/MultiBoxBuilder.h"
-#include "SAvaBroadcastProfileEntry.h"
-#include "SAvaChannel.h"
+#include "SAvaBroadcastChannel.h"
 #include "ScopedTransaction.h"
 #include "Widgets/Layout/SGridPanel.h"
 
@@ -23,7 +23,7 @@ void SAvaBroadcastChannels::FAvaChannelMaximizer::Reset()
 	bMaximizing = false;
 }
 
-void SAvaBroadcastChannels::FAvaChannelMaximizer::ToggleMaximize(const TSharedRef<SAvaChannel>& InChannelWidget)
+void SAvaBroadcastChannels::FAvaChannelMaximizer::ToggleMaximize(const TSharedRef<SAvaBroadcastChannel>& InChannelWidget)
 {
 	ChannelWidgetWeak = InChannelWidget;
 	bMaximizing = !bMaximizing;
@@ -74,7 +74,7 @@ void SAvaBroadcastChannels::Construct(const FArguments& InArgs, const TSharedPtr
 	BroadcastEditorWeak = InBroadcastEditor;
 	check(InBroadcastEditor.IsValid());
 	
-	UAvalancheBroadcast* const Broadcast = InBroadcastEditor->GetBroadcastObject();
+	UAvaBroadcast* const Broadcast = InBroadcastEditor->GetBroadcastObject();
 	BroadcastWeak = Broadcast;
 	check(IsValid(Broadcast));
 	
@@ -116,9 +116,9 @@ SAvaBroadcastChannels::~SAvaBroadcastChannels()
 
 bool SAvaBroadcastChannels::CanAddChannel()
 {
-	const int32 ChannelCount = UAvalancheBroadcast::Get().GetCurrentProfile().GetChannels().Num();
+	const int32 ChannelCount = UAvaBroadcast::Get().GetCurrentProfile().GetChannels().Num();
 					
-	const UAvalancheMediaEditorSettings& MediaSettings = UAvalancheMediaEditorSettings::Get();
+	const UAvaMediaEditorSettings& MediaSettings = UAvaMediaEditorSettings::Get();
 	
 	return !MediaSettings.bBroadcastEnforceMaxChannelCount
 		|| ChannelCount < MediaSettings.BroadcastMaxChannelCount;
@@ -126,7 +126,7 @@ bool SAvaBroadcastChannels::CanAddChannel()
 
 void SAvaBroadcastChannels::AddChannel()
 {
-	if (UAvalancheBroadcast* const Broadcast = BroadcastWeak.Get())
+	if (UAvaBroadcast* const Broadcast = BroadcastWeak.Get())
 	{
 		FScopedTransaction Transaction(LOCTEXT("AddChannel", "Add Channel"));
 		Broadcast->Modify();
@@ -171,7 +171,7 @@ bool SAvaBroadcastChannels::CanMaximizeChannel() const
 	return ChannelCount == 1 || (ChannelCount > 1 && !ChannelMaximizer.bMaximizing);
 }
 
-void SAvaBroadcastChannels::ToggleMaximizeChannel(const TSharedRef<SAvaChannel>& InWidget)
+void SAvaBroadcastChannels::ToggleMaximizeChannel(const TSharedRef<SAvaBroadcastChannel>& InWidget)
 {
 	if (Channels.Num() > 1)
 	{
@@ -197,7 +197,7 @@ void SAvaBroadcastChannels::RefreshChannelGrid()
 		const FAvaBroadcastProfile& Profile = BroadcastWeak->GetCurrentProfile();
 		
 		//Remove Items that are not in the New Channel List
-		for (TMap<FName, TSharedPtr<SAvaChannel>>::TIterator Iter = Channels.CreateIterator(); Iter; ++Iter)
+		for (TMap<FName, TSharedPtr<SAvaBroadcastChannel>>::TIterator Iter = Channels.CreateIterator(); Iter; ++Iter)
 		{
 			//Remove Invalid Channels
 			if (!Iter->Value.IsValid())
@@ -209,7 +209,7 @@ void SAvaBroadcastChannels::RefreshChannelGrid()
 			//TODO: Currently Remove all Slots and Re-add them in Order
 			ChannelGrid->RemoveSlot(Iter->Value.ToSharedRef());
 
-			const FAvaOutputChannel& Channel = Profile.GetChannel(Iter->Key);
+			const FAvaBroadcastOutputChannel& Channel = Profile.GetChannel(Iter->Key);
 			
 			//Refresh the Media Tiles as they might've been deleted in an Undo operation
 			Iter->Value->OnChannelMediaOutputsChanged(Channel);
@@ -230,7 +230,7 @@ void SAvaBroadcastChannels::RefreshChannelGrid()
 
 		int32 ItemCount = 0;
 
-		const TArray<FAvaOutputChannel*>& BroadcastChannels = Profile.GetChannels();
+		const TArray<FAvaBroadcastOutputChannel*>& BroadcastChannels = Profile.GetChannels();
 		
 		const int32 ItemsPerRow = FMath::RoundHalfToZero(FMath::Sqrt(static_cast<float>(BroadcastChannels.Num())));
 
@@ -239,12 +239,12 @@ void SAvaBroadcastChannels::RefreshChannelGrid()
 		ChannelGrid->ClearFill();
 		
 		//Re-add existing Widgets or add new ones
-		for (const FAvaOutputChannel* InChannel : BroadcastChannels)
+		for (const FAvaBroadcastOutputChannel* InChannel : BroadcastChannels)
 		{
 			const FName ChannelName = InChannel->GetChannelName();
 			if (!Channels.Contains(ChannelName))
 			{
-				TSharedRef<SAvaChannel> NewChannel = SNew(SAvaChannel, BroadcastEditor)
+				TSharedRef<SAvaBroadcastChannel> NewChannel = SNew(SAvaBroadcastChannel, BroadcastEditor)
 					.ChannelName(ChannelName)
 					.CanMaximize(this, &SAvaBroadcastChannels::CanMaximizeChannel)
 					.OnMaximizeClicked(this, &SAvaBroadcastChannels::ToggleMaximizeChannel);
@@ -255,7 +255,7 @@ void SAvaBroadcastChannels::RefreshChannelGrid()
 			const int32 Row = ItemCount / ItemsPerRow;
 			const int32 Column = ItemCount % ItemsPerRow;
 			
-			TSharedRef<SAvaChannel> Channel = Channels.Find(ChannelName)->ToSharedRef();
+			TSharedRef<SAvaBroadcastChannel> Channel = Channels.Find(ChannelName)->ToSharedRef();
 			Channel->SetPosition(Column, Row);
 
 			ChannelGrid->AddSlot(Column, Row)
diff --git a/Engine/Plugins/Experimental/Avalanche/Source/AvalancheMediaEditor/Private/Broadcast/ChannelGrid/Slate/SAvaBroadcastChannels.h b/Engine/Plugins/Experimental/Avalanche/Source/AvalancheMediaEditor/Private/Broadcast/ChannelGrid/Slate/SAvaBroadcastChannels.h
--- a/Engine/Plugins/Experimental/Avalanche/Source/AvalancheMediaEditor/Private/Broadcast/ChannelGrid/Slate/SAvaBroadcastChannels.h
+++ b/Engine/Plugins/Experimental/Avalanche/Source/AvalancheMediaEditor/Private/Broadcast/ChannelGrid/Slate/SAvaBroadcastChannels.h
@@ -2,16 +2,15 @@
 
 #pragma once
 
-#include "CoreMinimal.h"
-#include "AvaMediaDefines.h"
 #include "Animation/CurveSequence.h"
+#include "AvaMediaDefines.h"
+#include "CoreMinimal.h"
 #include "Widgets/SCompoundWidget.h"
-#include "Widgets/Views/STileView.h"
 
 class FAvaBroadcastEditor;
-class SAvaChannel;
+class SAvaBroadcastChannel;
 class SGridPanel;
-class UAvalancheBroadcast;
+class UAvaBroadcast;
 
 class SAvaBroadcastChannels : public SCompoundWidget
 {
@@ -20,12 +19,12 @@ class SAvaBroadcastChannels : public SCompoundWidget
 		FAvaChannelMaximizer();
 
 		void Reset();
-		void ToggleMaximize(const TSharedRef<SAvaChannel>& InChannelWidget);
+		void ToggleMaximize(const TSharedRef<SAvaBroadcastChannel>& InChannelWidget);
 		
 		float GetRowFill(int32 InRowIndex) const;
 		float GetColumnFill(int32 InColumnIndex) const;
 
-		TWeakPtr<SAvaChannel> ChannelWidgetWeak;
+		TWeakPtr<SAvaBroadcastChannel> ChannelWidgetWeak;
 		FCurveSequence MaximizeSequence;
 		bool bMaximizing  = false;
 	};
@@ -49,7 +48,7 @@ public:
 	TSharedRef<SWidget> MakeChannelsToolbar();
 	
 	bool CanMaximizeChannel() const;
-	void ToggleMaximizeChannel(const TSharedRef<SAvaChannel>& InWidget);
+	void ToggleMaximizeChannel(const TSharedRef<SAvaBroadcastChannel>& InWidget);
 	
 	void OnBroadcastChanged(EAvaBroadcastChange ChangedEvent);
 	void RefreshChannelGrid();
@@ -58,13 +57,13 @@ public:
 	
 protected:
 	
-	TWeakObjectPtr<UAvalancheBroadcast> BroadcastWeak;
+	TWeakObjectPtr<UAvaBroadcast> BroadcastWeak;
 	
 	TWeakPtr<FAvaBroadcastEditor> BroadcastEditorWeak;
 	
 	TSharedPtr<SGridPanel> ChannelGrid;
 	
-	TMap<FName, TSharedPtr<SAvaChannel>> Channels;
+	TMap<FName, TSharedPtr<SAvaBroadcastChannel>> Channels;
 
 	FAvaChannelMaximizer ChannelMaximizer;
 };
diff --git a/Engine/Plugins/Experimental/Avalanche/Source/AvalancheMediaEditor/Private/Broadcast/ChannelGrid/Slate/SAvaBroadcastProfileEntry.cpp b/Engine/Plugins/Experimental/Avalanche/Source/AvalancheMediaEditor/Private/Broadcast/ChannelGrid/Slate/SAvaBroadcastProfileEntry.cpp
--- a/Engine/Plugins/Experimental/Avalanche/Source/AvalancheMediaEditor/Private/Broadcast/ChannelGrid/Slate/SAvaBroadcastProfileEntry.cpp
+++ b/Engine/Plugins/Experimental/Avalanche/Source/AvalancheMediaEditor/Private/Broadcast/ChannelGrid/Slate/SAvaBroadcastProfileEntry.cpp
@@ -2,15 +2,15 @@
 
 #include "SAvaBroadcastProfileEntry.h"
 
-#include "AvalancheBroadcast.h"
+#include "Broadcast/AvaBroadcast.h"
 #include "Framework/Application/SlateApplication.h"
 #include "ScopedTransaction.h"
-#include "Widgets/SBoxPanel.h"
 #include "Widgets/Images/SImage.h"
+#include "Widgets/Input/SButton.h"
 #include "Widgets/Layout/SBox.h"
 #include "Widgets/Layout/SScaleBox.h"
 #include "Widgets/Layout/SSpacer.h"
-#include "Widgets/Input/SButton.h"
+#include "Widgets/SBoxPanel.h"
 #include "Widgets/Text/SInlineEditableTextBlock.h"
 
 #define LOCTEXT_NAMESPACE "SAvaBroadcastProfileEntry"
@@ -108,7 +108,7 @@ void SAvaBroadcastProfileEntry::OnProfileTextCommitted(const FText& InProfileTex
 	
 	FScopedTransaction Transaction(LOCTEXT("RenameProfile", "Rename Profile"));
 	
-	UAvalancheBroadcast& Broadcast = UAvalancheBroadcast::Get();
+	UAvaBroadcast& Broadcast = UAvaBroadcast::Get();
 	Broadcast.Modify();
 	
 	if (Broadcast.RenameProfile(ProfileName, DesiredProfileName))
@@ -131,7 +131,7 @@ bool SAvaBroadcastProfileEntry::OnVerifyProfileTextChanged(const FText& InProfil
 	
 	const FName DesiredProfileName(*InProfileText.ToString());
 	
-	if (!UAvalancheBroadcast::Get().CanRenameProfile(ProfileName, DesiredProfileName))
+	if (!UAvaBroadcast::Get().CanRenameProfile(ProfileName, DesiredProfileName))
 	{
 		OutErrorMessage = LOCTEXT("BroadcastCannotRename", "Broadcast cannot rename the Profile to the given name");
 		return false;
@@ -164,7 +164,7 @@ FReply SAvaBroadcastProfileEntry::DeleteProfile()
 {
 	FScopedTransaction Transaction(LOCTEXT("RemoveProfile", "Remove Profile"));
 	
-	UAvalancheBroadcast& Broadcast = UAvalancheBroadcast::Get();
+	UAvaBroadcast& Broadcast = UAvaBroadcast::Get();
 	Broadcast.Modify();
 	
 	const bool bResult = Broadcast.RemoveProfile(ProfileName);
diff --git a/Engine/Plugins/Experimental/Avalanche/Source/AvalancheMediaEditor/Private/Broadcast/ChannelGrid/Slate/SAvaBroadcastProfileEntry.h b/Engine/Plugins/Experimental/Avalanche/Source/AvalancheMediaEditor/Private/Broadcast/ChannelGrid/Slate/SAvaBroadcastProfileEntry.h
--- a/Engine/Plugins/Experimental/Avalanche/Source/AvalancheMediaEditor/Private/Broadcast/ChannelGrid/Slate/SAvaBroadcastProfileEntry.h
+++ b/Engine/Plugins/Experimental/Avalanche/Source/AvalancheMediaEditor/Private/Broadcast/ChannelGrid/Slate/SAvaBroadcastProfileEntry.h
@@ -8,14 +8,14 @@
 
 class SInlineEditableTextBlock;
 
-DECLARE_DELEGATE_RetVal_OneParam(FReply, FOnProfileEntrySelected, FName ProfileName);
+DECLARE_DELEGATE_RetVal_OneParam(FReply, FOnAvaBroadcastProfileEntrySelected, FName ProfileName);
 
 class SAvaBroadcastProfileEntry : public SCompoundWidget
 {
 public:
 	
 	SLATE_BEGIN_ARGS(SAvaBroadcastProfileEntry){}
-		SLATE_EVENT(FOnProfileEntrySelected, OnProfileEntrySelected)
+		SLATE_EVENT(FOnAvaBroadcastProfileEntrySelected, OnProfileEntrySelected)
 	SLATE_END_ARGS()
 
 	/** Constructs this widget with InArgs */
@@ -44,7 +44,7 @@ protected:
 	
 	FText ProfileNameText;
 
-	FOnProfileEntrySelected OnProfileEntrySelected;
+	FOnAvaBroadcastProfileEntrySelected OnProfileEntrySelected;
 	
 	TSharedPtr<SInlineEditableTextBlock> ProfileTextBlock;
 	
diff --git a/Engine/Plugins/Experimental/Avalanche/Source/AvalancheMediaEditor/Private/Broadcast/DetailsView/SAvaBroadcastDetailsView.cpp b/Engine/Plugins/Experimental/Avalanche/Source/AvalancheMediaEditor/Private/Broadcast/DetailsView/SAvaBroadcastDetailsView.cpp
--- a/Engine/Plugins/Experimental/Avalanche/Source/AvalancheMediaEditor/Private/Broadcast/DetailsView/SAvaBroadcastDetailsView.cpp
+++ b/Engine/Plugins/Experimental/Avalanche/Source/AvalancheMediaEditor/Private/Broadcast/DetailsView/SAvaBroadcastDetailsView.cpp
@@ -3,8 +3,9 @@
 #include "SAvaBroadcastDetailsView.h"
 
 #include "Broadcast/AvaBroadcastEditor.h"
-#include "Broadcast/ChannelGrid/AvaOutputTileItem.h"
-#include "Channel/AvaOutputChannel.h"
+#include "Broadcast/Channel/AvaBroadcastOutputChannel.h"
+#include "Broadcast/ChannelGrid/AvaBroadcastOutputTileItem.h"
+#include "MediaOutput.h"
 #include "Modules/ModuleManager.h"
 #include "PropertyEditorModule.h"
 #include "Widgets/Text/STextBlock.h"
@@ -61,16 +62,16 @@ bool SAvaBroadcastDetailsView::IsMediaOutputEditingEnabled() const
 {
 	if (TSharedPtr<FAvaBroadcastEditor> BroadcastEditor = BroadcastEditorWeak.Pin())
 	{
-		if (TSharedPtr<FAvaOutputTileItem> MediaOutputItem = BroadcastEditor->GetSelectedOutputTile())
+		if (TSharedPtr<FAvaBroadcastOutputTileItem> MediaOutputItem = BroadcastEditor->GetSelectedOutputTile())
 		{
-			const FAvaOutputChannel& Channel = MediaOutputItem->GetChannel();
+			const FAvaBroadcastOutputChannel& Channel = MediaOutputItem->GetChannel();
 			return Channel.IsValidChannel() && Channel.IsEditingEnabled();
 		}
 	}
 	return false;
 }
 
-void SAvaBroadcastDetailsView::OnMediaOutputSelectionChanged(const TSharedPtr<FAvaOutputTileItem>& InSelectedItem)
+void SAvaBroadcastDetailsView::OnMediaOutputSelectionChanged(const TSharedPtr<FAvaBroadcastOutputTileItem>& InSelectedItem)
 {
 	check(DetailsView.IsValid());
 
diff --git a/Engine/Plugins/Experimental/Avalanche/Source/AvalancheMediaEditor/Private/Broadcast/DetailsView/SAvaBroadcastDetailsView.h b/Engine/Plugins/Experimental/Avalanche/Source/AvalancheMediaEditor/Private/Broadcast/DetailsView/SAvaBroadcastDetailsView.h
--- a/Engine/Plugins/Experimental/Avalanche/Source/AvalancheMediaEditor/Private/Broadcast/DetailsView/SAvaBroadcastDetailsView.h
+++ b/Engine/Plugins/Experimental/Avalanche/Source/AvalancheMediaEditor/Private/Broadcast/DetailsView/SAvaBroadcastDetailsView.h
@@ -6,10 +6,10 @@
 #include "Widgets/SCompoundWidget.h"
 
 class FAvaBroadcastEditor;
-class FAvaOutputTileItem;
+class FAvaBroadcastOutputTileItem;
 class IDetailsView;
 
-class  SAvaBroadcastDetailsView : public SCompoundWidget
+class SAvaBroadcastDetailsView : public SCompoundWidget
 {
 public:
 	
@@ -20,7 +20,7 @@ public:
 	virtual ~SAvaBroadcastDetailsView() override;
 
 	bool IsMediaOutputEditingEnabled() const;
-	void OnMediaOutputSelectionChanged(const TSharedPtr<FAvaOutputTileItem>& InSelectedItem);
+	void OnMediaOutputSelectionChanged(const TSharedPtr<FAvaBroadcastOutputTileItem>& InSelectedItem);
 
 	EVisibility GetEmptySelectionTextVisibility() const;
 	
diff --git a/Engine/Plugins/Experimental/Avalanche/Source/AvalancheMediaEditor/Private/Broadcast/OutputDevices/DragDropOps/AvaBroadcastOutputTreeItemDragDropOp.h b/Engine/Plugins/Experimental/Avalanche/Source/AvalancheMediaEditor/Private/Broadcast/OutputDevices/DragDropOps/AvaBroadcastOutputTreeItemDragDropOp.h
new file mode 100644
--- /dev/null
+++ b/Engine/Plugins/Experimental/Avalanche/Source/AvalancheMediaEditor/Private/Broadcast/OutputDevices/DragDropOps/AvaBroadcastOutputTreeItemDragDropOp.h
@@ -0,0 +1,27 @@
+// Copyright Epic Games, Inc. All Rights Reserved.
+
+#pragma once
+
+#include "DragAndDrop/DecoratedDragDropOp.h"
+
+class FAvaBroadcastOutputTreeItem;
+
+class FAvaBroadcastOutputTreeItemDragDropOp : public FDecoratedDragDropOp
+{
+public:
+	DRAG_DROP_OPERATOR_TYPE(FAvaOutputClassDragDropOp, FDecoratedDragDropOp)
+
+	static TSharedRef<FAvaBroadcastOutputTreeItemDragDropOp> New(const TSharedPtr<FAvaBroadcastOutputTreeItem>& InOutputClassItem);
+
+	bool IsValidToDropInChannel(FName InTargetChannelName) const;
+
+	TSharedPtr<FAvaBroadcastOutputTreeItem> GetOutputTreeItem() const { return OutputTreeItem; }
+
+	FReply OnChannelDrop(FName InTargetChannelName);
+
+protected:
+	void Init(const TSharedPtr<FAvaBroadcastOutputTreeItem>& InOutputClassItem);
+
+	/** Keep Reference Count while Drag Dropping */
+	TSharedPtr<FAvaBroadcastOutputTreeItem> OutputTreeItem;
+};
