Motion Design

diff --git a/Engine/Plugins/Experimental/Avalanche/Source/AvalancheMediaEditor/Private/Broadcast/TabFactories/AvaBroadcastChannelsTabFactory.cpp b/Engine/Plugins/Experimental/Avalanche/Source/AvalancheMediaEditor/Private/Broadcast/TabFactories/AvaBroadcastChannelsTabFactory.cpp
--- a/Engine/Plugins/Experimental/Avalanche/Source/AvalancheMediaEditor/Private/Broadcast/TabFactories/AvaBroadcastChannelsTabFactory.cpp
+++ b/Engine/Plugins/Experimental/Avalanche/Source/AvalancheMediaEditor/Private/Broadcast/TabFactories/AvaBroadcastChannelsTabFactory.cpp
@@ -1,12 +1,13 @@
 // Copyright Epic Games, Inc. All Rights Reserved.
 
 #include "AvaBroadcastChannelsTabFactory.h"
+
 #include "AvaMediaEditorStyle.h"
 #include "Broadcast/AvaBroadcastEditor.h"
 #include "Broadcast/ChannelGrid/Slate/SAvaBroadcastChannels.h"
 #include "IAvaMediaEditorModule.h"
 
-const FName FAvaBroadcastChannelsTabFactory::TabID(TEXT("AvalancheBroadcastChannels"));
+const FName FAvaBroadcastChannelsTabFactory::TabID(TEXT("MotionDesignBroadcastChannels"));
 
 #define LOCTEXT_NAMESPACE "AvaBroadcastChannelsTabFactory"
 
@@ -14,7 +15,7 @@ FAvaBroadcastChannelsTabFactory::FAvaBroadcastChannelsTabFactory(const TSharedPt
 	: FAvaBroadcastTabFactory(TabID, InBroadcastEditor)
 {
 	TabLabel = LOCTEXT("BroadcastChannels_TabLabel", "Channels");
-	TabIcon = FSlateIcon(FAvaMediaEditorStyle::GetStyleSetName(), "AvalancheMediaEditor.BroadcastIcon");
+	TabIcon = FSlateIcon(FAvaMediaEditorStyle::GetStyleSetName(), "AvaMediaEditor.BroadcastIcon");
 
 	bIsSingleton = true;
 
diff --git a/Engine/Plugins/Experimental/Avalanche/Source/AvalancheMediaEditor/Private/Broadcast/TabFactories/AvaBroadcastDetailsTabFactory.cpp b/Engine/Plugins/Experimental/Avalanche/Source/AvalancheMediaEditor/Private/Broadcast/TabFactories/AvaBroadcastDetailsTabFactory.cpp
--- a/Engine/Plugins/Experimental/Avalanche/Source/AvalancheMediaEditor/Private/Broadcast/TabFactories/AvaBroadcastDetailsTabFactory.cpp
+++ b/Engine/Plugins/Experimental/Avalanche/Source/AvalancheMediaEditor/Private/Broadcast/TabFactories/AvaBroadcastDetailsTabFactory.cpp
@@ -1,11 +1,11 @@
 // Copyright Epic Games, Inc. All Rights Reserved.
 
 #include "AvaBroadcastDetailsTabFactory.h"
-#include "AvalancheBroadcast.h"
+
 #include "Broadcast/AvaBroadcastEditor.h"
 #include "Broadcast/DetailsView/SAvaBroadcastDetailsView.h"
 
-const FName FAvaBroadcastDetailsTabFactory::TabID(TEXT("AvalancheBroadcastDetails"));
+const FName FAvaBroadcastDetailsTabFactory::TabID(TEXT("MotionDesignBroadcastDetails"));
 
 #define LOCTEXT_NAMESPACE "AvaBroadcastDetailsTabFactory"
 
diff --git a/Engine/Plugins/Experimental/Avalanche/Source/AvalancheMediaEditor/Private/Broadcast/TabFactories/AvaBroadcastOutputsTabFactory.cpp b/Engine/Plugins/Experimental/Avalanche/Source/AvalancheMediaEditor/Private/Broadcast/TabFactories/AvaBroadcastOutputsTabFactory.cpp
--- a/Engine/Plugins/Experimental/Avalanche/Source/AvalancheMediaEditor/Private/Broadcast/TabFactories/AvaBroadcastOutputsTabFactory.cpp
+++ b/Engine/Plugins/Experimental/Avalanche/Source/AvalancheMediaEditor/Private/Broadcast/TabFactories/AvaBroadcastOutputsTabFactory.cpp
@@ -1,10 +1,11 @@
 // Copyright Epic Games, Inc. All Rights Reserved.
 
 #include "AvaBroadcastOutputsTabFactory.h"
+
 #include "Broadcast/AvaBroadcastEditor.h"
-#include "Broadcast/OutputDevices/Slate/SAvaOutputDevices.h"
+#include "Broadcast/OutputDevices/Slate/SAvaBroadcastOutputDevices.h"
 
-const FName FAvaBroadcastOutputsTabFactory::TabID(TEXT("AvalancheBroadcastOutputList"));
+const FName FAvaBroadcastOutputsTabFactory::TabID(TEXT("MotionDesignBroadcastOutputList"));
 
 #define LOCTEXT_NAMESPACE "AvaBroadcastOutputsTabFactory"
 
@@ -24,7 +25,7 @@ TSharedRef<SWidget> FAvaBroadcastOutputsTabFactory::CreateTabBody(const FWorkflo
 {
 	TSharedPtr<FAvaBroadcastEditor> BroadcastEditor = BroadcastEditorWeak.Pin();
 	check(BroadcastEditor.IsValid());
-	return SNew(SAvaOutputDevices, BroadcastEditor);
+	return SNew(SAvaBroadcastOutputDevices, BroadcastEditor);
 }
 
 #undef LOCTEXT_NAMESPACE
diff --git a/Engine/Plugins/Experimental/Avalanche/Source/AvalancheMediaEditor/Private/Playback/AppModes/AvaPlaybackAppMode.cpp b/Engine/Plugins/Experimental/Avalanche/Source/AvalancheMediaEditor/Private/Playback/AppModes/AvaPlaybackAppMode.cpp
--- a/Engine/Plugins/Experimental/Avalanche/Source/AvalancheMediaEditor/Private/Playback/AppModes/AvaPlaybackAppMode.cpp
+++ b/Engine/Plugins/Experimental/Avalanche/Source/AvalancheMediaEditor/Private/Playback/AppModes/AvaPlaybackAppMode.cpp
@@ -1,13 +1,14 @@
 // Copyright Epic Games, Inc. All Rights Reserved.
 
 #include "AvaPlaybackAppMode.h"
-#include "Playback/AvaPlaybackEditor.h"
+
+#include "Playback/AvaPlaybackGraphEditor.h"
 
 const FName FAvaPlaybackAppMode::DefaultMode("DefaultName");
 
 #define LOCTEXT_NAMESPACE "AvaPlaybackAppMode"
 
-FAvaPlaybackAppMode::FAvaPlaybackAppMode(const TSharedPtr<FAvaPlaybackEditor>& InPlaybackEditor, const FName& InModeName)
+FAvaPlaybackAppMode::FAvaPlaybackAppMode(const TSharedPtr<FAvaPlaybackGraphEditor>& InPlaybackEditor, const FName& InModeName)
 	: FApplicationMode(InModeName, FAvaPlaybackAppMode::GetLocalizedMode)
 	, PlaybackEditorWeak(InPlaybackEditor)
 {
@@ -15,7 +16,7 @@ FAvaPlaybackAppMode::FAvaPlaybackAppMode(const TSharedPtr<FAvaPlaybackEditor>& I
 
 void FAvaPlaybackAppMode::RegisterTabFactories(TSharedPtr<FTabManager> InTabManager)
 {
-	TSharedPtr<FAvaPlaybackEditor> PlaybackEditor = PlaybackEditorWeak.Pin();
+	TSharedPtr<FAvaPlaybackGraphEditor> PlaybackEditor = PlaybackEditorWeak.Pin();
 	PlaybackEditor->PushTabFactories(TabFactories);
 	FApplicationMode::RegisterTabFactories(InTabManager);
 }
diff --git a/Engine/Plugins/Experimental/Avalanche/Source/AvalancheMediaEditor/Private/Playback/AppModes/AvaPlaybackAppMode.h b/Engine/Plugins/Experimental/Avalanche/Source/AvalancheMediaEditor/Private/Playback/AppModes/AvaPlaybackAppMode.h
--- a/Engine/Plugins/Experimental/Avalanche/Source/AvalancheMediaEditor/Private/Playback/AppModes/AvaPlaybackAppMode.h
+++ b/Engine/Plugins/Experimental/Avalanche/Source/AvalancheMediaEditor/Private/Playback/AppModes/AvaPlaybackAppMode.h
@@ -5,7 +5,7 @@
 #include "WorkflowOrientedApp/ApplicationMode.h"
 #include "WorkflowOrientedApp/WorkflowTabManager.h"
 
-class FAvaPlaybackEditor;
+class FAvaPlaybackGraphEditor;
 
 class FAvaPlaybackAppMode : public FApplicationMode
 {
@@ -15,14 +15,14 @@ public:
 	static const FName DefaultMode;
 	//Add more here
 	
-	FAvaPlaybackAppMode(const TSharedPtr<FAvaPlaybackEditor>& InPlaybackEditor, const FName& InModeName);
+	FAvaPlaybackAppMode(const TSharedPtr<FAvaPlaybackGraphEditor>& InPlaybackEditor, const FName& InModeName);
 	virtual void RegisterTabFactories(TSharedPtr<FTabManager> InTabManager) override;
 
 protected:
 	
 	static FText GetLocalizedMode(const FName InMode);
 
-	TWeakPtr<FAvaPlaybackEditor> PlaybackEditorWeak;
+	TWeakPtr<FAvaPlaybackGraphEditor> PlaybackEditorWeak;
 	
 	// Set of spawnable tabs in the mode
 	FWorkflowAllowedTabSet TabFactories;
diff --git a/Engine/Plugins/Experimental/Avalanche/Source/AvalancheMediaEditor/Private/Playback/AppModes/AvaPlaybackDefaultMode.cpp b/Engine/Plugins/Experimental/Avalanche/Source/AvalancheMediaEditor/Private/Playback/AppModes/AvaPlaybackDefaultMode.cpp
--- a/Engine/Plugins/Experimental/Avalanche/Source/AvalancheMediaEditor/Private/Playback/AppModes/AvaPlaybackDefaultMode.cpp
+++ b/Engine/Plugins/Experimental/Avalanche/Source/AvalancheMediaEditor/Private/Playback/AppModes/AvaPlaybackDefaultMode.cpp
@@ -1,20 +1,21 @@
 // Copyright Epic Games, Inc. All Rights Reserved.
 
 #include "AvaPlaybackDefaultMode.h"
+
 #include "IAvaMediaEditorModule.h"
-#include "Playback/AvaPlaybackEditor.h"
+#include "Playback/AvaPlaybackGraphEditor.h"
 #include "Playback/TabFactories/AvaPlaybackDetailsTabFactory.h"
-#include "Playback/TabFactories/AvaPlaybackGraphTabFactory.h"
+#include "Playback/TabFactories/AvaPlaybackEditorGraphTabFactory.h"
 
 #define LOCTEXT_NAMESPACE "AvaPlaybackDefaultMode"
 
-FAvaPlaybackDefaultMode::FAvaPlaybackDefaultMode(const TSharedPtr<FAvaPlaybackEditor>& InPlaybackEditor)
+FAvaPlaybackDefaultMode::FAvaPlaybackDefaultMode(const TSharedPtr<FAvaPlaybackGraphEditor>& InPlaybackEditor)
 	: FAvaPlaybackAppMode(InPlaybackEditor, FAvaPlaybackAppMode::DefaultMode)
 {
 	WorkspaceMenuCategory = FWorkspaceItem::NewGroup(LOCTEXT("WorkspaceMenu_AvaPlayback", "Motion Design Playback"));
 	
 	check(InPlaybackEditor.IsValid());
-	TabLayout = FTabManager::NewLayout("AvalanchePlaybackEditor_Default_Layout_V1")
+	TabLayout = FTabManager::NewLayout("MotionDesignPlaybackEditor_Default_Layout_V1")
 		->AddArea
 		(
 			FTabManager::NewPrimaryArea() ->SetOrientation(Orient_Vertical)
@@ -26,8 +27,8 @@ FAvaPlaybackDefaultMode::FAvaPlaybackDefaultMode(const TSharedPtr<FAvaPlaybackEd
 				(
 					FTabManager::NewStack()
 					->SetSizeCoefficient(0.75f)
-					->AddTab(FAvaPlaybackGraphTabFactory::TabID, ETabState::OpenedTab)
-					->SetForegroundTab(FAvaPlaybackGraphTabFactory::TabID)
+					->AddTab(FAvaPlaybackEditorGraphTabFactory::TabID, ETabState::OpenedTab)
+					->SetForegroundTab(FAvaPlaybackEditorGraphTabFactory::TabID)
 					->SetHideTabWell(true)
 				)
 				->Split
@@ -40,7 +41,7 @@ FAvaPlaybackDefaultMode::FAvaPlaybackDefaultMode(const TSharedPtr<FAvaPlaybackEd
 		);
 
 	// Add Tab Spawners
-	TabFactories.RegisterFactory(MakeShared<FAvaPlaybackGraphTabFactory>(InPlaybackEditor));
+	TabFactories.RegisterFactory(MakeShared<FAvaPlaybackEditorGraphTabFactory>(InPlaybackEditor));
 	TabFactories.RegisterFactory(MakeShared<FAvaPlaybackDetailsTabFactory>(InPlaybackEditor));
 
 	//Make sure we start with our existing list of extenders instead of creating a new one
diff --git a/Engine/Plugins/Experimental/Avalanche/Source/AvalancheMediaEditor/Private/Playback/AppModes/AvaPlaybackDefaultMode.h b/Engine/Plugins/Experimental/Avalanche/Source/AvalancheMediaEditor/Private/Playback/AppModes/AvaPlaybackDefaultMode.h
--- a/Engine/Plugins/Experimental/Avalanche/Source/AvalancheMediaEditor/Private/Playback/AppModes/AvaPlaybackDefaultMode.h
+++ b/Engine/Plugins/Experimental/Avalanche/Source/AvalancheMediaEditor/Private/Playback/AppModes/AvaPlaybackDefaultMode.h
@@ -7,5 +7,5 @@
 class FAvaPlaybackDefaultMode : public FAvaPlaybackAppMode
 {
 public:
-	FAvaPlaybackDefaultMode(const TSharedPtr<FAvaPlaybackEditor>& InPlaybackEditor);
+	FAvaPlaybackDefaultMode(const TSharedPtr<FAvaPlaybackGraphEditor>& InPlaybackEditor);
 };
diff --git a/Engine/Plugins/Experimental/Avalanche/Source/AvalancheMediaEditor/Private/Playback/AvaPlaybackGraphFactory.cpp b/Engine/Plugins/Experimental/Avalanche/Source/AvalancheMediaEditor/Private/Playback/AvaPlaybackGraphFactory.cpp
new file mode 100644
--- /dev/null
+++ b/Engine/Plugins/Experimental/Avalanche/Source/AvalancheMediaEditor/Private/Playback/AvaPlaybackGraphFactory.cpp
@@ -0,0 +1,71 @@
+// Copyright Epic Games, Inc. All Rights Reserved.
+
+#include "Playback/AvaPlaybackGraphFactory.h"
+
+#include "AssetToolsModule.h"
+#include "IAssetTools.h"
+#include "Playback/AvaPlaybackGraph.h"
+
+#define LOCTEXT_NAMESPACE "AvaPlaybackGraphFactory"
+
+UAvaPlaybackGraphFactory::UAvaPlaybackGraphFactory()
+{
+	// Provide the factory with information about how to handle our asset
+	bCreateNew = true;
+	bEditAfterNew = true;
+	SupportedClass = UAvaPlaybackGraph::StaticClass();
+}
+
+UAvaPlaybackGraphFactory::~UAvaPlaybackGraphFactory()
+{
+}
+
+uint32 UAvaPlaybackGraphFactory::GetMenuCategories() const
+{
+	IAssetTools& AssetTools = FAssetToolsModule::GetModule().Get();
+	return AssetTools.FindAdvancedAssetCategory("MotionDesignCategory");
+}
+
+bool UAvaPlaybackGraphFactory::ConfigureProperties()
+{
+	return Super::ConfigureProperties();
+}
+
+bool UAvaPlaybackGraphFactory::ShouldShowInNewMenu() const
+{
+	return Super::ShouldShowInNewMenu();
+}
+
+UObject* UAvaPlaybackGraphFactory::FactoryCreateNew(UClass* Class, UObject* InParent, FName Name, EObjectFlags Flags
+  , UObject* Context, FFeedbackContext* Warn, FName CallingContext)
+{
+	return Super::FactoryCreateNew(Class, InParent, Name, Flags, Context, Warn, CallingContext);
+}
+
+UObject* UAvaPlaybackGraphFactory::FactoryCreateNew(UClass* Class, UObject* InParent, FName Name, EObjectFlags Flags
+  , UObject* Context, FFeedbackContext* Warn)
+{
+	UAvaPlaybackGraph* Playback = nullptr;
+	if (ensure(SupportedClass == Class))
+	{
+		Playback = NewObject<UAvaPlaybackGraph>(InParent, Name, Flags);
+	}
+	return Playback;
+}
+
+bool UAvaPlaybackGraphFactory::DoesSupportClass(UClass* Class)
+{
+	return Class == UAvaPlaybackGraph::StaticClass();
+}
+
+UClass* UAvaPlaybackGraphFactory::ResolveSupportedClass()
+{
+	return Super::ResolveSupportedClass();
+}
+
+FString UAvaPlaybackGraphFactory::GetDefaultNewAssetName() const
+{
+	return TEXT("NewPlaybackGraph");
+}
+
+#undef LOCTEXT_NAMESPACE
diff --git a/Engine/Plugins/Experimental/Avalanche/Source/AvalancheMediaEditor/Private/Playback/DetailsView/SAvaPlaybackDetailsView.cpp b/Engine/Plugins/Experimental/Avalanche/Source/AvalancheMediaEditor/Private/Playback/DetailsView/SAvaPlaybackDetailsView.cpp
--- a/Engine/Plugins/Experimental/Avalanche/Source/AvalancheMediaEditor/Private/Playback/DetailsView/SAvaPlaybackDetailsView.cpp
+++ b/Engine/Plugins/Experimental/Avalanche/Source/AvalancheMediaEditor/Private/Playback/DetailsView/SAvaPlaybackDetailsView.cpp
@@ -1,11 +1,12 @@
 // Copyright Epic Games, Inc. All Rights Reserved.
 
 #include "SAvaPlaybackDetailsView.h"
+
 #include "Modules/ModuleManager.h"
-#include "Playback/AvaPlaybackEditor.h"
+#include "Playback/AvaPlaybackGraphEditor.h"
 #include "PropertyEditorModule.h"
 
-void SAvaPlaybackDetailsView::Construct(const FArguments& InArgs, const TSharedPtr<FAvaPlaybackEditor>& InPlaybackEditor)
+void SAvaPlaybackDetailsView::Construct(const FArguments& InArgs, const TSharedPtr<FAvaPlaybackGraphEditor>& InPlaybackEditor)
 {
 	PlaybackEditorWeak = InPlaybackEditor;
 	
@@ -29,7 +30,7 @@ void SAvaPlaybackDetailsView::Construct(const FArguments& InArgs, const TSharedP
 
 SAvaPlaybackDetailsView::~SAvaPlaybackDetailsView()
 {
-	if (TSharedPtr<FAvaPlaybackEditor> PlaybackEditor = PlaybackEditorWeak.Pin())
+	if (TSharedPtr<FAvaPlaybackGraphEditor> PlaybackEditor = PlaybackEditorWeak.Pin())
 	{
 		PlaybackEditor->OnPlaybackSelectionChanged.RemoveAll(this);
 	}
diff --git a/Engine/Plugins/Experimental/Avalanche/Source/AvalancheMediaEditor/Private/Playback/DetailsView/SAvaPlaybackDetailsView.h b/Engine/Plugins/Experimental/Avalanche/Source/AvalancheMediaEditor/Private/Playback/DetailsView/SAvaPlaybackDetailsView.h
--- a/Engine/Plugins/Experimental/Avalanche/Source/AvalancheMediaEditor/Private/Playback/DetailsView/SAvaPlaybackDetailsView.h
+++ b/Engine/Plugins/Experimental/Avalanche/Source/AvalancheMediaEditor/Private/Playback/DetailsView/SAvaPlaybackDetailsView.h
@@ -5,7 +5,7 @@
 #include "CoreMinimal.h"
 #include "Widgets/SCompoundWidget.h"
 
-class FAvaPlaybackEditor;
+class FAvaPlaybackGraphEditor;
 class IDetailsView;
 
 class SAvaPlaybackDetailsView : public SCompoundWidget
@@ -15,14 +15,14 @@ public:
 	SLATE_BEGIN_ARGS(SAvaPlaybackDetailsView){}
 	SLATE_END_ARGS()
 	
-	void Construct(const FArguments& InArgs, const TSharedPtr<FAvaPlaybackEditor>& InPlaybackEditor);
+	void Construct(const FArguments& InArgs, const TSharedPtr<FAvaPlaybackGraphEditor>& InPlaybackEditor);
 	virtual ~SAvaPlaybackDetailsView() override;
 	
 	void OnPlaybackNodeSelectionChanged(const TArray<UObject*>& InSelectedObjects);
 
 protected:
 
-	TWeakPtr<FAvaPlaybackEditor> PlaybackEditorWeak;
+	TWeakPtr<FAvaPlaybackGraphEditor> PlaybackEditorWeak;
 	
 	TSharedPtr<IDetailsView> DetailsView;
 };
diff --git a/Engine/Plugins/Experimental/Avalanche/Source/AvalancheMediaEditor/Private/Playback/Graph/AvaPlaybackConnectionDrawingPolicy.cpp b/Engine/Plugins/Experimental/Avalanche/Source/AvalancheMediaEditor/Private/Playback/Graph/AvaPlaybackConnectionDrawingPolicy.cpp
--- a/Engine/Plugins/Experimental/Avalanche/Source/AvalancheMediaEditor/Private/Playback/Graph/AvaPlaybackConnectionDrawingPolicy.cpp
+++ b/Engine/Plugins/Experimental/Avalanche/Source/AvalancheMediaEditor/Private/Playback/Graph/AvaPlaybackConnectionDrawingPolicy.cpp
@@ -1,10 +1,11 @@
 // Copyright Epic Games, Inc. All Rights Reserved.
 
 #include "AvaPlaybackConnectionDrawingPolicy.h"
-#include "Playback/AvalanchePlayback.h"
-#include "Playback/Graph/AvaPlaybackGraph.h"
-#include "Playback/Graph/AvaPlaybackGraphSchema.h"
-#include "Playback/Graph/Nodes/AvaPlaybackGraphNode_Root.h"
+
+#include "Playback/AvaPlaybackGraph.h"
+#include "Playback/Graph/AvaPlaybackEditorGraph.h"
+#include "Playback/Graph/AvaPlaybackEditorGraphSchema.h"
+#include "Playback/Graph/Nodes/AvaPlaybackEditorGraphNode_Root.h"
 
 FConnectionDrawingPolicy* FAvaPlaybackConnectionDrawingPolicyFactory::CreateConnectionPolicy(const UEdGraphSchema* Schema
 	, int32 InBackLayerID
@@ -14,7 +15,7 @@ FConnectionDrawingPolicy* FAvaPlaybackConnectionDrawingPolicyFactory::CreateConn
 	, FSlateWindowElementList& InDrawElements
 	, UEdGraph* InGraphObj) const
 {
-	if (Schema->IsA(UAvaPlaybackGraphSchema::StaticClass()))
+	if (Schema->IsA(UAvaPlaybackEditorGraphSchema::StaticClass()))
 	{
 		return new FAvaPlaybackConnectionDrawingPolicy(InBackLayerID
 			, InFrontLayerID
@@ -61,7 +62,7 @@ void FAvaPlaybackConnectionDrawingPolicy::DetermineWiringStyle(UEdGraphPin* Outp
 	const UEdGraphSchema* const Schema = GraphObject->GetSchema();
 
 	Params.WireColor = Schema->GetPinTypeColor(OutputPin->PinType);
-	UAvaPlaybackGraph* const Graph = Cast<UAvaPlaybackGraph>(GraphObject);
+	UAvaPlaybackEditorGraph* const Graph = Cast<UAvaPlaybackEditorGraph>(GraphObject);
 	
 	if (!InputPin || !Graph)
 	{
@@ -70,11 +71,11 @@ void FAvaPlaybackConnectionDrawingPolicy::DetermineWiringStyle(UEdGraphPin* Outp
 	
 	bool bExecuted = false;
 
-	const UAvalanchePlayback* const Playback = Graph->GetPlayback();
+	const UAvaPlaybackGraph* const Playback = Graph->GetPlaybackGraph();
 	const bool bIsPlaying = Playback ? Playback->IsPlaying() : false;
 	
-	UAvaPlaybackGraphNode* const InputNode  = Cast<UAvaPlaybackGraphNode>(InputPin->GetOwningNode());
-	UAvaPlaybackGraphNode* const OutputNode = Cast<UAvaPlaybackGraphNode>(OutputPin->GetOwningNode());
+	UAvaPlaybackEditorGraphNode* const InputNode  = Cast<UAvaPlaybackEditorGraphNode>(InputPin->GetOwningNode());
+	UAvaPlaybackEditorGraphNode* const OutputNode = Cast<UAvaPlaybackEditorGraphNode>(OutputPin->GetOwningNode());
 
 	if (bIsPlaying && InputNode && OutputNode)
 	{
diff --git a/Engine/Plugins/Experimental/Avalanche/Source/AvalancheMediaEditor/Private/Playback/Graph/Nodes/AvaPlaybackEditorGraphNode_BlueprintPlayer.cpp b/Engine/Plugins/Experimental/Avalanche/Source/AvalancheMediaEditor/Private/Playback/Graph/Nodes/AvaPlaybackEditorGraphNode_BlueprintPlayer.cpp
new file mode 100644
--- /dev/null
+++ b/Engine/Plugins/Experimental/Avalanche/Source/AvalancheMediaEditor/Private/Playback/Graph/Nodes/AvaPlaybackEditorGraphNode_BlueprintPlayer.cpp
@@ -0,0 +1,28 @@
+// Copyright Epic Games, Inc. All Rights Reserved.
+
+#include "Playback/Graph/Nodes/AvaPlaybackEditorGraphNode_BlueprintPlayer.h"
+
+#include "AvaMediaEditorSettings.h"
+#include "Playback/Graph/AvaPlaybackEditorGraphSchema.h"
+#include "Playback/Nodes/AvaPlaybackNodeBlueprintPlayer.h"
+#include "Slate/SAvaPlaybackEditorGraphNode_Player.h"
+
+TSubclassOf<UAvaPlaybackNode> UAvaPlaybackEditorGraphNode_BlueprintPlayer::GetPlaybackNodeClass() const
+{
+	return UAvaPlaybackNodeBlueprintPlayer::StaticClass();
+}
+
+FName UAvaPlaybackEditorGraphNode_BlueprintPlayer::GetInputPinCategory(int32 InputPinIndex) const
+{
+	return UAvaPlaybackEditorGraphSchema::PC_Event;
+}
+
+FLinearColor UAvaPlaybackEditorGraphNode_BlueprintPlayer::GetNodeTitleColor() const
+{
+	return UAvaMediaEditorSettings::Get().PlaybackPlayerNodeColor;
+}
+
+TSharedPtr<SGraphNode> UAvaPlaybackEditorGraphNode_BlueprintPlayer::CreateVisualWidget()
+{
+	return SNew(SAvaPlaybackEditorGraphNode_Player, this);
+}
diff --git a/Engine/Plugins/Experimental/Avalanche/Source/AvalancheMediaEditor/Private/Playback/Graph/SchemaActions/AvaPlaybackAction_NewComment.cpp b/Engine/Plugins/Experimental/Avalanche/Source/AvalancheMediaEditor/Private/Playback/Graph/SchemaActions/AvaPlaybackAction_NewComment.cpp
--- a/Engine/Plugins/Experimental/Avalanche/Source/AvalancheMediaEditor/Private/Playback/Graph/SchemaActions/AvaPlaybackAction_NewComment.cpp
+++ b/Engine/Plugins/Experimental/Avalanche/Source/AvalancheMediaEditor/Private/Playback/Graph/SchemaActions/AvaPlaybackAction_NewComment.cpp
@@ -1,10 +1,11 @@
 // Copyright Epic Games, Inc. All Rights Reserved.
 
 #include "Playback/Graph/SchemaActions/AvaPlaybackAction_NewComment.h"
+
 #include "EdGraphNode_Comment.h"
-#include "Playback/AvalanchePlayback.h"
-#include "Playback/Graph/AvaPlaybackGraph.h"
-#include "Playback/IAvaPlaybackEditor.h"
+#include "Playback/AvaPlaybackGraph.h"
+#include "Playback/Graph/AvaPlaybackEditorGraph.h"
+#include "Playback/IAvaPlaybackGraphEditor.h"
 
 UEdGraphNode* FAvaPlaybackAction_NewComment::PerformAction(UEdGraph* ParentGraph, UEdGraphPin* FromPin
 	, const FVector2D Location, bool bSelectNewNode)
@@ -14,7 +15,7 @@ UEdGraphNode* FAvaPlaybackAction_NewComment::PerformAction(UEdGraph* ParentGraph
 
 	FVector2D SpawnLocation = Location;
 	
-	UAvalanchePlayback* const Playback = CastChecked<UAvaPlaybackGraph>(ParentGraph)->GetPlayback();
+	UAvaPlaybackGraph* const Playback = CastChecked<UAvaPlaybackEditorGraph>(ParentGraph)->GetPlaybackGraph();
 	check(Playback);
 	
 	TSharedPtr<IAvaPlaybackGraphEditor> GraphEditor = Playback->GetGraphEditor();
diff --git a/Engine/Plugins/Experimental/Avalanche/Source/AvalancheMediaEditor/Private/Playback/Graph/SchemaActions/AvaPlaybackAction_NewNode.cpp b/Engine/Plugins/Experimental/Avalanche/Source/AvalancheMediaEditor/Private/Playback/Graph/SchemaActions/AvaPlaybackAction_NewNode.cpp
--- a/Engine/Plugins/Experimental/Avalanche/Source/AvalancheMediaEditor/Private/Playback/Graph/SchemaActions/AvaPlaybackAction_NewNode.cpp
+++ b/Engine/Plugins/Experimental/Avalanche/Source/AvalancheMediaEditor/Private/Playback/Graph/SchemaActions/AvaPlaybackAction_NewNode.cpp
@@ -1,8 +1,9 @@
 // Copyright Epic Games, Inc. All Rights Reserved.
 
 #include "Playback/Graph/SchemaActions/AvaPlaybackAction_NewNode.h"
-#include "Playback/AvalanchePlayback.h"
-#include "Playback/Graph/AvaPlaybackGraph.h"
+
+#include "Playback/AvaPlaybackGraph.h"
+#include "Playback/Graph/AvaPlaybackEditorGraph.h"
 #include "Playback/Nodes/AvaPlaybackNode.h"
 #include "ScopedTransaction.h"
 
@@ -15,7 +16,7 @@ UEdGraphNode* FAvaPlaybackAction_NewNode::PerformAction(UEdGraph* ParentGraph
 {
 	check(IsValid(PlaybackNodeClass));
 
-	UAvalanchePlayback* Playback = CastChecked<UAvaPlaybackGraph>(ParentGraph)->GetPlayback();
+	UAvaPlaybackGraph* Playback = CastChecked<UAvaPlaybackEditorGraph>(ParentGraph)->GetPlaybackGraph();
 	check(Playback);
 	
 	const FScopedTransaction Transaction(LOCTEXT("NewPlaybackNode", "Motion Design Playback: New Node"));
diff --git a/Engine/Plugins/Experimental/Avalanche/Source/AvalancheMediaEditor/Private/Playback/Graph/SchemaActions/AvaPlaybackAction_PasteNode.cpp b/Engine/Plugins/Experimental/Avalanche/Source/AvalancheMediaEditor/Private/Playback/Graph/SchemaActions/AvaPlaybackAction_PasteNode.cpp
--- a/Engine/Plugins/Experimental/Avalanche/Source/AvalancheMediaEditor/Private/Playback/Graph/SchemaActions/AvaPlaybackAction_PasteNode.cpp
+++ b/Engine/Plugins/Experimental/Avalanche/Source/AvalancheMediaEditor/Private/Playback/Graph/SchemaActions/AvaPlaybackAction_PasteNode.cpp
@@ -1,14 +1,15 @@
 // Copyright Epic Games, Inc. All Rights Reserved.
 
 #include "Playback/Graph/SchemaActions/AvaPlaybackAction_PasteNode.h"
-#include "Playback/AvalanchePlayback.h"
-#include "Playback/Graph/AvaPlaybackGraph.h"
-#include "Playback/IAvaPlaybackEditor.h"
+
+#include "Playback/AvaPlaybackGraph.h"
+#include "Playback/Graph/AvaPlaybackEditorGraph.h"
+#include "Playback/IAvaPlaybackGraphEditor.h"
 
 UEdGraphNode* FAvaPlaybackAction_PasteNode::PerformAction(UEdGraph* ParentGraph, UEdGraphPin* FromPin
 	, const FVector2D Location, bool bSelectNewNode)
 {
-	UAvalanchePlayback* const Playback = CastChecked<UAvaPlaybackGraph>(ParentGraph)->GetPlayback();
+	UAvaPlaybackGraph* const Playback = CastChecked<UAvaPlaybackEditorGraph>(ParentGraph)->GetPlaybackGraph();
 	check(Playback);
 	
 	if (TSharedPtr<IAvaPlaybackGraphEditor> GraphEditor = Playback->GetGraphEditor())
diff --git a/Engine/Plugins/Experimental/Avalanche/Source/AvalancheMediaEditor/Private/Playback/TabFactories/AvaPlaybackDetailsTabFactory.cpp b/Engine/Plugins/Experimental/Avalanche/Source/AvalancheMediaEditor/Private/Playback/TabFactories/AvaPlaybackDetailsTabFactory.cpp
--- a/Engine/Plugins/Experimental/Avalanche/Source/AvalancheMediaEditor/Private/Playback/TabFactories/AvaPlaybackDetailsTabFactory.cpp
+++ b/Engine/Plugins/Experimental/Avalanche/Source/AvalancheMediaEditor/Private/Playback/TabFactories/AvaPlaybackDetailsTabFactory.cpp
@@ -1,15 +1,15 @@
 // Copyright Epic Games, Inc. All Rights Reserved.
 
 #include "AvaPlaybackDetailsTabFactory.h"
-#include "Playback/AvalanchePlayback.h"
-#include "Playback/AvaPlaybackEditor.h"
+
+#include "Playback/AvaPlaybackGraphEditor.h"
 #include "Playback/DetailsView/SAvaPlaybackDetailsView.h"
 
-const FName FAvaPlaybackDetailsTabFactory::TabID(TEXT("AvalanchePlaybackDetails"));
+const FName FAvaPlaybackDetailsTabFactory::TabID(TEXT("MotionDesignPlaybackDetails"));
 
 #define LOCTEXT_NAMESPACE "AvaPlaybackDetailsTabFactory"
 
-FAvaPlaybackDetailsTabFactory::FAvaPlaybackDetailsTabFactory(const TSharedPtr<FAvaPlaybackEditor>& InPlaybackEditor)
+FAvaPlaybackDetailsTabFactory::FAvaPlaybackDetailsTabFactory(const TSharedPtr<FAvaPlaybackGraphEditor>& InPlaybackEditor)
 	: FAvaPlaybackTabFactory(TabID, InPlaybackEditor)
 {
 	TabLabel = LOCTEXT("PlaybackDetails_TabLabel", "Details");
@@ -23,7 +23,7 @@ FAvaPlaybackDetailsTabFactory::FAvaPlaybackDetailsTabFactory(const TSharedPtr<FA
 
 TSharedRef<SWidget> FAvaPlaybackDetailsTabFactory::CreateTabBody(const FWorkflowTabSpawnInfo& Info) const
 {
-	TSharedPtr<FAvaPlaybackEditor> PlaybackEditor = PlaybackEditorWeak.Pin();
+	TSharedPtr<FAvaPlaybackGraphEditor> PlaybackEditor = PlaybackEditorWeak.Pin();
 	check(PlaybackEditor.IsValid());
 	return SNew(SAvaPlaybackDetailsView, PlaybackEditor);
 }
diff --git a/Engine/Plugins/Experimental/Avalanche/Source/AvalancheMediaEditor/Private/Playback/TabFactories/AvaPlaybackDetailsTabFactory.h b/Engine/Plugins/Experimental/Avalanche/Source/AvalancheMediaEditor/Private/Playback/TabFactories/AvaPlaybackDetailsTabFactory.h
--- a/Engine/Plugins/Experimental/Avalanche/Source/AvalancheMediaEditor/Private/Playback/TabFactories/AvaPlaybackDetailsTabFactory.h
+++ b/Engine/Plugins/Experimental/Avalanche/Source/AvalancheMediaEditor/Private/Playback/TabFactories/AvaPlaybackDetailsTabFactory.h
@@ -9,7 +9,7 @@ class FAvaPlaybackDetailsTabFactory : public FAvaPlaybackTabFactory
 public:
 	static const FName TabID;
 	
-	FAvaPlaybackDetailsTabFactory(const TSharedPtr<FAvaPlaybackEditor>& InPlaybackEditor);
+	FAvaPlaybackDetailsTabFactory(const TSharedPtr<FAvaPlaybackGraphEditor>& InPlaybackEditor);
 	
 	virtual TSharedRef<SWidget> CreateTabBody(const FWorkflowTabSpawnInfo& Info) const override;
 };
diff --git a/Engine/Plugins/Experimental/Avalanche/Source/AvalancheMediaEditor/Private/Playback/TabFactories/AvaPlaybackTabFactory.cpp b/Engine/Plugins/Experimental/Avalanche/Source/AvalancheMediaEditor/Private/Playback/TabFactories/AvaPlaybackTabFactory.cpp
--- a/Engine/Plugins/Experimental/Avalanche/Source/AvalancheMediaEditor/Private/Playback/TabFactories/AvaPlaybackTabFactory.cpp
+++ b/Engine/Plugins/Experimental/Avalanche/Source/AvalancheMediaEditor/Private/Playback/TabFactories/AvaPlaybackTabFactory.cpp
@@ -1,9 +1,10 @@
 // Copyright Epic Games, Inc. All Rights Reserved.
 
 #include "AvaPlaybackTabFactory.h"
-#include "Playback/AvaPlaybackEditor.h"
 
-FAvaPlaybackTabFactory::FAvaPlaybackTabFactory(const FName& InTabID, const TSharedPtr<FAvaPlaybackEditor>& InPlaybackEditor)
+#include "Playback/AvaPlaybackGraphEditor.h"
+
+FAvaPlaybackTabFactory::FAvaPlaybackTabFactory(const FName& InTabID, const TSharedPtr<FAvaPlaybackGraphEditor>& InPlaybackEditor)
 	: FWorkflowTabFactory(InTabID, InPlaybackEditor)
 	, PlaybackEditorWeak(InPlaybackEditor)
 {
diff --git a/Engine/Plugins/Experimental/Avalanche/Source/AvalancheMediaEditor/Private/Playback/TabFactories/AvaPlaybackTabFactory.h b/Engine/Plugins/Experimental/Avalanche/Source/AvalancheMediaEditor/Private/Playback/TabFactories/AvaPlaybackTabFactory.h
--- a/Engine/Plugins/Experimental/Avalanche/Source/AvalancheMediaEditor/Private/Playback/TabFactories/AvaPlaybackTabFactory.h
+++ b/Engine/Plugins/Experimental/Avalanche/Source/AvalancheMediaEditor/Private/Playback/TabFactories/AvaPlaybackTabFactory.h
@@ -4,16 +4,16 @@
 
 #include "WorkflowOrientedApp/WorkflowTabFactory.h"
 
-class FAvaPlaybackEditor;
+class FAvaPlaybackGraphEditor;
 
 //Base class for all Tab Factories in Ava Playback Editor
 class FAvaPlaybackTabFactory : public FWorkflowTabFactory
 {
 public:
 	
-	FAvaPlaybackTabFactory(const FName& InTabID, const TSharedPtr<FAvaPlaybackEditor>& InPlaybackEditor);
+	FAvaPlaybackTabFactory(const FName& InTabID, const TSharedPtr<FAvaPlaybackGraphEditor>& InPlaybackEditor);
 
 protected:
 
-	TWeakPtr<FAvaPlaybackEditor> PlaybackEditorWeak;
+	TWeakPtr<FAvaPlaybackGraphEditor> PlaybackEditorWeak;
 };
diff --git a/Engine/Plugins/Experimental/Avalanche/Source/AvalancheMediaEditor/Private/Rundown/AvaRundownDefines.cpp b/Engine/Plugins/Experimental/Avalanche/Source/AvalancheMediaEditor/Private/Rundown/AvaRundownDefines.cpp
new file mode 100644
--- /dev/null
+++ b/Engine/Plugins/Experimental/Avalanche/Source/AvalancheMediaEditor/Private/Rundown/AvaRundownDefines.cpp
@@ -0,0 +1,7 @@
+// Copyright Epic Games, Inc. All Rights Reserved.
+
+#include "Rundown/AvaRundownDefines.h"
+
+const FNumberFormattingOptions UE::AvaRundown::FEditorMetrics::PageIdFormattingOptions = FNumberFormattingOptions()
+	.SetMinimumIntegralDigits(4)
+	.SetUseGrouping(false);
diff --git a/Engine/Plugins/Experimental/Avalanche/Source/AvalancheMediaEditor/Private/Rundown/AvaRundownFactory.cpp b/Engine/Plugins/Experimental/Avalanche/Source/AvalancheMediaEditor/Private/Rundown/AvaRundownFactory.cpp
new file mode 100644
--- /dev/null
+++ b/Engine/Plugins/Experimental/Avalanche/Source/AvalancheMediaEditor/Private/Rundown/AvaRundownFactory.cpp
@@ -0,0 +1,68 @@
+// Copyright Epic Games, Inc. All Rights Reserved.
+
+#include "Rundown/AvaRundownFactory.h"
+
+#include "AssetToolsModule.h"
+#include "IAssetTools.h"
+#include "Rundown/AvaRundown.h"
+
+UAvaRundownFactory::UAvaRundownFactory()
+{
+	// Provide the factory with information about how to handle our asset
+	bCreateNew = true;
+	bEditAfterNew = true;
+	SupportedClass = UAvaRundown::StaticClass();
+}
+
+UAvaRundownFactory::~UAvaRundownFactory()
+{
+}
+
+uint32 UAvaRundownFactory::GetMenuCategories() const
+{
+	IAssetTools& AssetTools = FAssetToolsModule::GetModule().Get();
+	return AssetTools.FindAdvancedAssetCategory("MotionDesignCategory");
+}
+
+bool UAvaRundownFactory::ConfigureProperties()
+{
+	return Super::ConfigureProperties();
+}
+
+bool UAvaRundownFactory::ShouldShowInNewMenu() const
+{
+	return Super::ShouldShowInNewMenu();
+}
+
+UObject* UAvaRundownFactory::FactoryCreateNew(UClass* Class, UObject* InParent, FName Name, EObjectFlags Flags
+	, UObject* Context, FFeedbackContext* Warn, FName CallingContext)
+{
+	return Super::FactoryCreateNew(Class, InParent, Name, Flags, Context, Warn, CallingContext);
+}
+
+UObject* UAvaRundownFactory::FactoryCreateNew(UClass* Class, UObject* InParent, FName Name, EObjectFlags Flags
+	, UObject* Context, FFeedbackContext* Warn)
+{
+	UAvaRundown* Rundown = nullptr;
+	if (ensure(SupportedClass == Class))
+	{
+		Rundown = NewObject<UAvaRundown>(InParent, Name, Flags);
+	}
+	return Rundown;
+}
+
+bool UAvaRundownFactory::DoesSupportClass(UClass* Class)
+{
+	return Class == UAvaRundown::StaticClass();
+}
+
+UClass* UAvaRundownFactory::ResolveSupportedClass()
+{
+	return Super::ResolveSupportedClass();
+}
+
+FString UAvaRundownFactory::GetDefaultNewAssetName() const
+{
+	return TEXT("NewRundown");
+}
+
diff --git a/Engine/Plugins/Experimental/Avalanche/Source/AvalancheMediaEditor/Private/Rundown/DetailsView/RemoteControl/Controllers/SAvaRundownRCControllerPanel.h b/Engine/Plugins/Experimental/Avalanche/Source/AvalancheMediaEditor/Private/Rundown/DetailsView/RemoteControl/Controllers/SAvaRundownRCControllerPanel.h
new file mode 100644
--- /dev/null
+++ b/Engine/Plugins/Experimental/Avalanche/Source/AvalancheMediaEditor/Private/Rundown/DetailsView/RemoteControl/Controllers/SAvaRundownRCControllerPanel.h
@@ -0,0 +1,107 @@
+// Copyright Epic Games, Inc. All Rights Reserved.
+
+#pragma once
+
+#include "Misc/NotifyHook.h"
+#include "Templates/SharedPointer.h"
+#include "Widgets/SCompoundWidget.h"
+
+class FAvaRundownEditor;
+class FAvaRundownManagedInstance;
+class FAvaRundownRCControllerItem;
+class IPropertyRowGenerator;
+class ITableRow;
+class SAvaRundownRCControllerPanel;
+class SHeaderRow;
+class STableViewBase;
+class UAvaRundown;
+class URCVirtualPropertyBase;
+class URemoteControlPreset;
+struct FAvaPlayableRemoteControlValue;
+struct FAvaRundownPage;
+struct FPropertyChangedEvent;
+template <typename ItemType> class SListView;
+
+using FAvaRundownRCControllerItemPtr = TSharedPtr<FAvaRundownRCControllerItem>;
+
+DECLARE_MULTICAST_DELEGATE_TwoParams(FAvaRundownRCControllerHeaderRowExtensionDelegate, TSharedRef<SAvaRundownRCControllerPanel> Panel, 
+	TSharedRef<SHeaderRow>& HeaderRow)
+DECLARE_DELEGATE_ThreeParams(FAvaRundownRCControllerTableRowExtensionDelegate, TSharedRef<SAvaRundownRCControllerPanel> Panel, 
+	TSharedRef<const FAvaRundownRCControllerItem> ItemPtr, TSharedPtr<SWidget>& CurrentWidget)
+
+class SAvaRundownRCControllerPanel : public SCompoundWidget
+{
+public:
+	static const FName ControllerColumnName;
+	static const FName ValueColumnName;
+
+	static FAvaRundownRCControllerHeaderRowExtensionDelegate& GetHeaderRowExtensionDelegate() { return HeaderRowExtensionDelegate; }
+	static TArray<FAvaRundownRCControllerTableRowExtensionDelegate>& GetTableRowExtensionDelegates(FName InExtensionName);
+
+	SLATE_BEGIN_ARGS(SAvaRundownRCControllerPanel) {}
+	SLATE_END_ARGS()
+	
+	void Construct(const FArguments& InArgs, const TSharedPtr<FAvaRundownEditor>& InRundownEditor);
+	
+	bool HasRemoteControlPreset(const URemoteControlPreset* InPreset) const;
+
+	void OnPageSelectionChanged(const TArray<int32>& InSelectedPageIds);
+	
+	void Refresh(const TArray<int32>& InSelectedPageIds);
+	
+	TSharedRef<ITableRow> OnGenerateControllerRow(FAvaRundownRCControllerItemPtr InItem, const TSharedRef<STableViewBase>& InOwnerTable);
+
+private:
+	static FAvaRundownRCControllerHeaderRowExtensionDelegate HeaderRowExtensionDelegate;
+	static TMap<FName, TArray<FAvaRundownRCControllerTableRowExtensionDelegate>> TableRowExtensionDelegates;
+
+	void UpdatePropertyRowGenerators(int32 InNumGenerators);
+	
+	void RefreshForManagedInstance(int32 InInstanceIndex, const FAvaRundownManagedInstance& InManagedInstance, const FAvaRundownPage& InPage);
+	
+	/** Update the current page's remote control values from the defaults then refresh the widget. */
+	void UpdateDefaultValuesAndRefresh();
+
+	void OnRemoteControlControllerAdded(URemoteControlPreset* InPreset, const FName NewControllerName, const FGuid& InControllerId) { UpdateDefaultValuesAndRefresh(); }
+	void OnRemoteControlControllerRemoved(URemoteControlPreset* InPreset, const FGuid& InControllerId) { UpdateDefaultValuesAndRefresh(); }
+	void OnRemoteControlControllerRenamed(URemoteControlPreset* InPreset, const FName InOldLabel, const FName InNewLabel) { UpdateDefaultValuesAndRefresh(); }
+	void OnRemoteControlControllerModified(URemoteControlPreset* InPreset, const TSet<FGuid>& InModifiedControllerIds);
+	void BindRemoteControlDelegates(URemoteControlPreset* InPreset);
+
+	const FAvaPlayableRemoteControlValue* GetSelectedPageControllerValue(const URCVirtualPropertyBase* InController) const;
+	bool SetSelectedPageControllerValue(const URCVirtualPropertyBase* InController, const FAvaPlayableRemoteControlValue& InValue) const;
+	void UpdatePageSummary(bool bInIsPresetChanged);
+
+	UAvaRundown* GetRundown() const;
+	const FAvaRundownPage& GetActivePage(const UAvaRundown* InRundown) const;
+	const FAvaRundownPage& GetActivePage() const { return GetActivePage(GetRundown()); }
+	FAvaRundownPage& GetActivePageMutable(UAvaRundown* InRundown) const;
+	FAvaRundownPage& GetActivePageMutable() const { return GetActivePageMutable(GetRundown()); }
+
+	TWeakPtr<FAvaRundownEditor> RundownEditorWeak;
+
+	class FPropertyRowGeneratorWrapper : public FNotifyHook
+	{
+	public:
+		TSharedPtr<IPropertyRowGenerator> PropertyRowGenerator;
+		TWeakObjectPtr<URemoteControlPreset> PresetWeak;
+		SAvaRundownRCControllerPanel* ParentPanel = nullptr;
+
+		FPropertyRowGeneratorWrapper(SAvaRundownRCControllerPanel* InParentPanel);
+		virtual ~FPropertyRowGeneratorWrapper();
+
+		//~ Begin FNotifyHook
+		virtual void NotifyPostChange(const FPropertyChangedEvent& InPropertyChangedEvent, FProperty* InPropertyThatChanged) override;
+		//~ End FNotifyHook
+	};
+	
+	TArray<TUniquePtr<FPropertyRowGeneratorWrapper>> PropertyRowGenerators;
+
+	TArray<TSharedPtr<FAvaRundownManagedInstance>> ManagedInstances;
+	
+	TSharedPtr<SListView<FAvaRundownRCControllerItemPtr>> ControllerContainer;
+	
+	TArray<FAvaRundownRCControllerItemPtr> ControllerItems;
+
+	int32 ActivePageId = -1;
+};
diff --git a/Engine/Plugins/Experimental/Avalanche/Source/AvalancheMediaEditor/Private/Rundown/DetailsView/RemoteControl/Properties/AvaRundownPageRemoteControlWidgetUtils.cpp b/Engine/Plugins/Experimental/Avalanche/Source/AvalancheMediaEditor/Private/Rundown/DetailsView/RemoteControl/Properties/AvaRundownPageRemoteControlWidgetUtils.cpp
new file mode 100644
--- /dev/null
+++ b/Engine/Plugins/Experimental/Avalanche/Source/AvalancheMediaEditor/Private/Rundown/DetailsView/RemoteControl/Properties/AvaRundownPageRemoteControlWidgetUtils.cpp
@@ -0,0 +1,115 @@
+// Copyright Epic Games, Inc. All Rights Reserved.
+
+#include "AvaRundownPageRemoteControlWidgetUtils.h"
+
+#include "IDetailTreeNode.h"
+#include "PropertyHandle.h"
+#include "SlateOptMacros.h"
+#include "Widgets/SBoxPanel.h"
+
+#define LOCTEXT_NAMESPACE "AvaRundownPageRemoteControlWidgetUtils"
+
+BEGIN_SLATE_FUNCTION_BUILD_OPTIMIZATION
+
+TSharedRef<SWidget> FAvaRundownPageRemoteControlWidgetUtils::CreateNodeValueWidget(const TSharedPtr<IDetailTreeNode>& InNode)
+{
+	FNodeWidgets NodeWidgets = InNode->CreateNodeWidgets();
+
+	TSharedRef<SHorizontalBox> FieldWidget = SNew(SHorizontalBox);
+
+	if (NodeWidgets.ValueWidget)
+	{
+		FieldWidget->AddSlot()
+			.Padding(FMargin(3.0f, 2.0f))
+			.HAlign(HAlign_Right)
+			.FillWidth(1.0f)
+			[
+				NodeWidgets.ValueWidget.ToSharedRef()
+			];
+	}
+	else if (NodeWidgets.WholeRowWidget)
+	{
+		FieldWidget->AddSlot()
+			.Padding(FMargin(3.0f, 2.0f))
+			.FillWidth(1.0f)
+			[
+				NodeWidgets.WholeRowWidget.ToSharedRef()
+			];
+	}
+
+	return FieldWidget;
+}
+
+bool FAvaRundownPageRemoteControlWidgetUtils::FindPropertyHandleRecursive(const TSharedPtr<IPropertyHandle>& InPropertyHandle, const FString& InPropertyNameOrPath, EFindNodeMethod InFindMethod)
+{
+	if (InPropertyHandle && InPropertyHandle->IsValidHandle())
+	{
+		uint32 ChildrenCount = 0;
+		InPropertyHandle->GetNumChildren(ChildrenCount);
+		for (uint32 Index = 0; Index < ChildrenCount; ++Index)
+		{
+			TSharedPtr<IPropertyHandle> ChildHandle = InPropertyHandle->GetChildHandle(Index);
+			if (FindPropertyHandleRecursive(ChildHandle, InPropertyNameOrPath, InFindMethod))
+			{
+				return true;
+			}
+		}
+
+		if (InPropertyHandle->GetProperty())
+		{
+			if (InFindMethod == EFindNodeMethod::Path)
+			{
+				if (InPropertyHandle->GeneratePathToProperty() == InPropertyNameOrPath)
+				{
+					return true;
+				}
+			}
+			else if (InPropertyHandle->GetProperty()->GetName() == InPropertyNameOrPath)
+			{
+				return true;
+			}
+		}
+	}
+
+	return false;
+}
+
+TSharedPtr<IDetailTreeNode> FAvaRundownPageRemoteControlWidgetUtils::FindTreeNodeRecursive(const TSharedRef<IDetailTreeNode>& InRootNode, const FString& InPropertyNameOrPath, EFindNodeMethod InFindMethod)
+{
+	TArray<TSharedRef<IDetailTreeNode>> Children;
+	InRootNode->GetChildren(Children);
+	for (TSharedRef<IDetailTreeNode>& Child : Children)
+	{
+		TSharedPtr<IDetailTreeNode> FoundNode = FindTreeNodeRecursive(Child, InPropertyNameOrPath, InFindMethod);
+		if (FoundNode.IsValid())
+		{
+			return FoundNode;
+		}
+	}
+
+	TSharedPtr<IPropertyHandle> Handle = InRootNode->CreatePropertyHandle();
+	if (FindPropertyHandleRecursive(Handle, InPropertyNameOrPath, InFindMethod))
+	{
+		return InRootNode;
+	}
+
+	return nullptr;
+}
+
+TSharedPtr<IDetailTreeNode> FAvaRundownPageRemoteControlWidgetUtils::FindNode(const TArray<TSharedRef<IDetailTreeNode>>& InRootNodes, const FString& InQualifiedPropertyName, EFindNodeMethod InFindMethod)
+{
+	for (const TSharedRef<IDetailTreeNode>& CategoryNode : InRootNodes)
+	{
+		TSharedPtr<IDetailTreeNode> FoundNode = FindTreeNodeRecursive(CategoryNode, InQualifiedPropertyName, InFindMethod);
+		if (FoundNode.IsValid())
+		{
+			return FoundNode;
+		}
+	}
+
+	return nullptr;
+}
+
+END_SLATE_FUNCTION_BUILD_OPTIMIZATION
+
+#undef LOCTEXT_NAMESPACE
diff --git a/Engine/Plugins/Experimental/Avalanche/Source/AvalancheMediaEditor/Private/Rundown/DetailsView/RemoteControl/Properties/AvaRundownPageRemoteControlWidgetUtils.h b/Engine/Plugins/Experimental/Avalanche/Source/AvalancheMediaEditor/Private/Rundown/DetailsView/RemoteControl/Properties/AvaRundownPageRemoteControlWidgetUtils.h
new file mode 100644
--- /dev/null
+++ b/Engine/Plugins/Experimental/Avalanche/Source/AvalancheMediaEditor/Private/Rundown/DetailsView/RemoteControl/Properties/AvaRundownPageRemoteControlWidgetUtils.h
@@ -0,0 +1,28 @@
+// Copyright Epic Games, Inc. All Rights Reserved.
+
+#pragma once
+
+#include "CoreMinimal.h"
+
+class IDetailTreeNode;
+class IPropertyHandle;
+class SWidget;
+
+/**
+ * A static class containing utility functions for widgets.
+ */
+class FAvaRundownPageRemoteControlWidgetUtils final
+{
+public:
+	enum class EFindNodeMethod : uint8
+	{
+		Name, // Find a node by providing its property name
+		Path // Find a node by providing a property path
+	};
+
+public:
+	static TSharedRef<SWidget> CreateNodeValueWidget(const TSharedPtr<IDetailTreeNode>& InNode);
+	static bool FindPropertyHandleRecursive(const TSharedPtr<IPropertyHandle>& InPropertyHandle, const FString& InPropertyNameOrPath, EFindNodeMethod InFindMethod);
+	static TSharedPtr<IDetailTreeNode> FindTreeNodeRecursive(const TSharedRef<IDetailTreeNode>& InRootNode, const FString& InPropertyNameOrPath, EFindNodeMethod InFindMethod);
+	static TSharedPtr<IDetailTreeNode> FindNode(const TArray<TSharedRef<IDetailTreeNode>>& InRootNodes, const FString& InQualifiedPropertyName, EFindNodeMethod InFindMethod);
+};
diff --git a/Engine/Plugins/Experimental/Avalanche/Source/AvalancheMediaEditor/Private/Rundown/DetailsView/RemoteControl/Properties/AvaRundownRCPropertyItem.cpp b/Engine/Plugins/Experimental/Avalanche/Source/AvalancheMediaEditor/Private/Rundown/DetailsView/RemoteControl/Properties/AvaRundownRCPropertyItem.cpp
new file mode 100644
--- /dev/null
+++ b/Engine/Plugins/Experimental/Avalanche/Source/AvalancheMediaEditor/Private/Rundown/DetailsView/RemoteControl/Properties/AvaRundownRCPropertyItem.cpp
@@ -0,0 +1,17 @@
+// Copyright Epic Games, Inc. All Rights Reserved.
+
+#include "AvaRundownRCPropertyItem.h"
+
+#include "RemoteControlEntity.h"
+#include "SAvaRundownRCPropertyItemRow.h"
+
+FAvaRundownRCPropertyItem::FAvaRundownRCPropertyItem(TSharedRef<FRemoteControlEntity> InEntity, bool bInControlled)
+{
+	Entity = InEntity;
+	bEntityControlled = bInControlled;
+}
+
+TSharedRef<ITableRow> FAvaRundownRCPropertyItem::CreateWidget(TSharedRef<SAvaRundownPageRemoteControlProps> InPropertyPanel, const TSharedRef<STableViewBase>& InOwnerTable) const
+{
+	return SNew(SAvaRundownRCPropertyItemRow, InPropertyPanel, InOwnerTable, SharedThis(this));
+}
diff --git a/Engine/Plugins/Experimental/Avalanche/Source/AvalancheMediaEditor/Private/Rundown/DetailsView/RemoteControl/Properties/AvaRundownRCPropertyItem.h b/Engine/Plugins/Experimental/Avalanche/Source/AvalancheMediaEditor/Private/Rundown/DetailsView/RemoteControl/Properties/AvaRundownRCPropertyItem.h
new file mode 100644
--- /dev/null
+++ b/Engine/Plugins/Experimental/Avalanche/Source/AvalancheMediaEditor/Private/Rundown/DetailsView/RemoteControl/Properties/AvaRundownRCPropertyItem.h
@@ -0,0 +1,29 @@
+// Copyright Epic Games, Inc. All Rights Reserved.
+
+#pragma once
+
+#include "Internationalization/Text.h"
+#include "Templates/SharedPointer.h"
+
+class IDetailTreeNode;
+class ITableRow;
+class SAvaRundownPageRemoteControlProps;
+class STableViewBase;
+class URemoteControlPreset;
+struct FAvaPageRemoteControlGenerateWidgetArgs;
+struct FRemoteControlEntity;
+
+class FAvaRundownRCPropertyItem : public TSharedFromThis<FAvaRundownRCPropertyItem>
+{
+public:
+	FAvaRundownRCPropertyItem(TSharedRef<FRemoteControlEntity> InEntity, bool bInControlled);
+
+	TSharedRef<ITableRow> CreateWidget(TSharedRef<SAvaRundownPageRemoteControlProps> InPropertyPanel, const TSharedRef<STableViewBase>& InOwnerTable) const;
+
+	TSharedPtr<FRemoteControlEntity> GetEntity() const { return Entity.Pin(); }
+	bool IsEntityControlled() const { return bEntityControlled; }
+
+private:
+	TWeakPtr<FRemoteControlEntity> Entity = nullptr;
+	bool bEntityControlled = false;
+};
diff --git a/Engine/Plugins/Experimental/Avalanche/Source/AvalancheMediaEditor/Private/Rundown/DetailsView/RemoteControl/Properties/SAvaRundownPageRemoteControlProps.cpp b/Engine/Plugins/Experimental/Avalanche/Source/AvalancheMediaEditor/Private/Rundown/DetailsView/RemoteControl/Properties/SAvaRundownPageRemoteControlProps.cpp
new file mode 100644
--- /dev/null
+++ b/Engine/Plugins/Experimental/Avalanche/Source/AvalancheMediaEditor/Private/Rundown/DetailsView/RemoteControl/Properties/SAvaRundownPageRemoteControlProps.cpp
@@ -0,0 +1,469 @@
+// Copyright Epic Games, Inc. All Rights Reserved.
+
+#include "SAvaRundownPageRemoteControlProps.h"
+
+#include "AvaRundownRCPropertyItem.h"
+#include "Playable/AvaPlayableRemoteControl.h"
+#include "Rundown/AvaRundownEditor.h"
+#include "Rundown/AvaRundownManagedInstanceCache.h"
+#include "Rundown/AvaRundownPage.h"
+#include "SAvaRundownRCPropertyItemRow.h"
+#include "SlateOptMacros.h"
+#include "Widgets/Views/ITableRow.h"
+#include "Widgets/Views/SListView.h"
+#include "Widgets/Views/STableViewBase.h"
+
+#define LOCTEXT_NAMESPACE "SAvaRundownPageRemoteControlProps"
+
+const FName SAvaRundownPageRemoteControlProps::PropertyColumnName = "PropertyColumn";
+const FName SAvaRundownPageRemoteControlProps::ValueColumnName = "ValueColumn";
+
+FAvaRundownRCPropertyHeaderRowExtensionDelegate SAvaRundownPageRemoteControlProps::HeaderRowExtensionDelegate;
+TMap<FName, TArray<FAvaRundownRCPropertyTableRowExtensionDelegate>> SAvaRundownPageRemoteControlProps::TableRowExtensionDelegates;
+
+BEGIN_SLATE_FUNCTION_BUILD_OPTIMIZATION
+
+TArray<FAvaRundownRCPropertyTableRowExtensionDelegate>& SAvaRundownPageRemoteControlProps::GetTableRowExtensionDelegates(FName InExtensionName)
+{
+	return TableRowExtensionDelegates.FindOrAdd(InExtensionName);
+}
+
+void SAvaRundownPageRemoteControlProps::Construct(const FArguments& InArgs, TSharedPtr<FAvaRundownEditor> InRundownEditor)
+{
+	RundownEditorWeak = InRundownEditor;
+	ActivePageId = FAvaRundownPage::InvalidPageId;
+
+	TSharedRef<SHeaderRow> HeaderRow =
+		SNew(SHeaderRow)
+		+ SHeaderRow::Column(PropertyColumnName)
+		.DefaultLabel(LOCTEXT("Property", "Property"))
+		.FixedWidth(150.f)
+		+ SHeaderRow::Column(ValueColumnName)
+		.DefaultLabel(LOCTEXT("Value", "Value"))
+		.FillWidth(1.f);
+
+	HeaderRowExtensionDelegate.Broadcast(SharedThis(this), HeaderRow);
+
+	ChildSlot
+	[
+		SNew(SBorder)
+		.Padding(8.0f)
+		.BorderImage(FAppStyle::GetBrush("ToolPanel.GroupBorder"))
+		[
+			SAssignNew(PropertyContainer, SListView<FAvaRundownRCPropertyItemPtr>)
+			.ListItemsSource(&PropertyItems)
+			.SelectionMode(ESelectionMode::None)
+			.OnGenerateRow(this, &SAvaRundownPageRemoteControlProps::OnGenerateControllerRow)
+			.HeaderRow(HeaderRow)
+		]
+	];
+
+	Refresh({});
+}
+
+END_SLATE_FUNCTION_BUILD_OPTIMIZATION
+
+SAvaRundownPageRemoteControlProps::~SAvaRundownPageRemoteControlProps()
+{
+}
+
+void SAvaRundownPageRemoteControlProps::UpdateDefaultValuesAndRefresh(const TArray<int32>& InSelectedPageIds)
+{
+	// Remark:
+	// This is used in the "reimport page" work flow. We want to be able to
+	// reimport multiple pages at the same time.
+	
+	if (const TSharedPtr<FAvaRundownEditor> RundownEditor = RundownEditorWeak.Pin())
+	{
+		UAvaRundown* Rundown = RundownEditor->GetRundown();
+
+		if (IsValid(Rundown))
+		{
+			EAvaPlayableRemoteControlChanges Changes = EAvaPlayableRemoteControlChanges::None;
+			for (const int32 PageId : InSelectedPageIds)
+			{
+				FAvaRundownPage& Page = RundownEditor->GetRundown()->GetPage(PageId);
+
+				if (Page.IsValidPage())
+				{
+					ManagedInstances = FAvaRundownEditor::GetManagedInstancesForPage(Rundown, Page);
+					
+					if (!ManagedInstances.IsEmpty())
+					{
+						FAvaPlayableRemoteControlValues MergedDefaultRCValues;
+						FAvaRundownEditor::MergeDefaultRemoteControlValues(ManagedInstances, MergedDefaultRCValues);
+
+						// Using the rundown API for event propagation.
+						constexpr bool bUpdateDefaults = true;
+						Changes |= Rundown->UpdateRemoteControlValues(ActivePageId, MergedDefaultRCValues, bUpdateDefaults);
+					}
+				}
+			}
+			
+			if (Changes != EAvaPlayableRemoteControlChanges::None)
+			{
+				RundownEditor->MarkAsModified();
+			}
+			
+			Refresh(InSelectedPageIds);
+		}
+	}
+}
+
+TSharedRef<ITableRow> SAvaRundownPageRemoteControlProps::OnGenerateControllerRow(FAvaRundownRCPropertyItemPtr InItem, const TSharedRef<STableViewBase>& InOwnerTable)
+{
+	return InItem->CreateWidget(SharedThis(this), InOwnerTable);
+}
+
+
+void SAvaRundownPageRemoteControlProps::RefreshTable(const TSet<FGuid>& InEntityIds)
+{
+	for (const FAvaRundownRCPropertyItemPtr& PropertyItem : PropertyItems)
+	{
+		const TSharedPtr<FRemoteControlEntity> Entity = PropertyItem->GetEntity();
+		if (Entity && (InEntityIds.IsEmpty() || InEntityIds.Contains(Entity->GetId())))
+		{
+			TSharedPtr<ITableRow> TableRow = PropertyContainer->WidgetFromItem(PropertyItem);
+
+			if (TableRow.IsValid())
+			{
+				const TSharedRef<SAvaRundownRCPropertyItemRow> ItemRow = StaticCastSharedRef<SAvaRundownRCPropertyItemRow>(TableRow.ToSharedRef());
+				ItemRow->UpdateValue();
+			}
+		}
+	}
+}
+
+void SAvaRundownPageRemoteControlProps::Refresh(const TArray<int32>& InSelectedPageIds)
+{
+	if (!PropertyContainer.IsValid())
+	{
+		return;
+	}
+
+	ActivePageId = InSelectedPageIds.IsEmpty() ? FAvaRundownPage::InvalidPageId : InSelectedPageIds[0];
+	
+	if (const TSharedPtr<FAvaRundownEditor> RundownEditor = RundownEditorWeak.Pin())
+	{
+		const UAvaRundown* Rundown = RundownEditor->GetRundown();
+
+		if (IsValid(Rundown))
+		{
+			ManagedInstances.Reset();
+
+			if (FAvaRundownPage* ActivePage = GetActivePage())
+			{
+				ManagedInstances = FAvaRundownEditor::GetManagedInstancesForPage(Rundown, *ActivePage);
+
+				for (const TSharedPtr<FAvaRundownManagedInstance>& ManagedInstance : ManagedInstances)
+				{
+					URemoteControlPreset* Preset = ManagedInstance ? ManagedInstance->GetRemoteControlPreset() : nullptr;
+					if (Preset)
+					{
+						BindRemoteControlDelegates(Preset);	
+					}
+				}
+
+				if (!ManagedInstances.IsEmpty())
+				{
+					FAvaPlayableRemoteControlValues MergedDefaultRCValues;
+					FAvaRundownEditor::MergeDefaultRemoteControlValues(ManagedInstances, MergedDefaultRCValues);
+					
+					// Prune any extra stale values. This happens if templates are changed.
+					if (ActivePage->PruneRemoteControlValues(MergedDefaultRCValues) != EAvaPlayableRemoteControlChanges::None)
+					{
+						UE_LOG(LogAvaPlayableRemoteControl, Log, TEXT("Page %d had stale values that where pruned."), ActivePage->GetPageId());
+						RundownEditor->MarkAsModified();
+					}
+				}
+			}
+		}
+
+		struct FAvaPropertyDetails
+		{
+			TSharedRef<FRemoteControlEntity> Entity;
+			bool bEntityControlled;
+		};
+
+		TArray<FAvaPropertyDetails> NewItems;
+
+		for (const TSharedPtr<FAvaRundownManagedInstance>& ManagedInstance : ManagedInstances)
+		{
+			URemoteControlPreset* RemoteControlPreset = ManagedInstance ? ManagedInstance->GetRemoteControlPreset() : nullptr;
+			if (!RemoteControlPreset)
+			{
+				continue;
+			}
+			
+			const TArray<TWeakPtr<FRemoteControlEntity>> ExposedEntities = RemoteControlPreset->GetExposedEntities<FRemoteControlEntity>();
+			NewItems.Reserve(ExposedEntities.Num());
+			
+			for (const TWeakPtr<FRemoteControlEntity>& EntityWeakPtr : ExposedEntities)
+			{
+				if (const TSharedPtr<FRemoteControlEntity> Entity = EntityWeakPtr.Pin())
+				{
+					const FAvaPlayableRemoteControlValue* EntityValue = GetSelectedPageEntityValue(Entity);
+
+					if (!EntityValue)
+					{
+						// If the page doesn't already have a value, we get it from the template's default values.
+						const FAvaPlayableRemoteControlValue* const DefaultEntityValue = ManagedInstance->GetDefaultRemoteControlValues().GetEntityValue(Entity->GetId()); 
+						
+						if (!DefaultEntityValue)
+						{
+							UE_LOG(LogAvaPlayableRemoteControl, Error, TEXT("Entity \"%s\" (id:%s) doesn't have a template default value."),
+								*Entity->GetLabel().ToString(), *Entity->GetId().ToString());
+							// TODO: UX improvement: instead of skipping, could add empty element, with error mark (and error message in tooltip).
+							continue;
+						}
+
+						// Ensure the default values have the default flag.
+						ensure(DefaultEntityValue->bIsDefault == true);
+						
+						// WYSIWYG (Solution):
+						// Capture the default value (flagged as default) in the current page to ensure all values will be applied to runtime RCP.
+						if (!SetSelectedPageEntityValue(Entity, *DefaultEntityValue))
+						{
+							UE_LOG(LogAvaPlayableRemoteControl, Error, TEXT("Entity \"%s\" (id:%s): failed to set value in currently selected page."),
+								*Entity->GetLabel().ToString(), *Entity->GetId().ToString());
+						}
+
+						EntityValue = DefaultEntityValue;
+					}
+					
+					// Update Exposed entity value with value from page.
+					using namespace UE::AvaPlayableRemoteControl;
+					if (const EAvaPlayableRemoteControlResult Result = SetValueOfEntity(Entity, EntityValue->Value); Failed(Result))
+					{
+						UE_LOG(LogAvaPlayableRemoteControl, Error, TEXT("Entity \"%s\" (id:%s): failed to set entity value: %s."),
+							*Entity->GetLabel().ToString(), *Entity->GetId().ToString(), *EnumToString(Result));
+						// TODO: UX improvement: instead of skipping, could add empty element, with error mark (and error message in tooltip).
+						continue;
+					}
+
+					{
+						const bool bEntityControlled = ManagedInstance->GetDefaultRemoteControlValues().EntitiesControlledByController.Contains(Entity->GetId());
+						NewItems.Add({Entity.ToSharedRef(), bEntityControlled});
+					}
+				}
+			}
+		}
+
+		PropertyContainer->RebuildList();
+
+		bool bRecreateList = NewItems.Num() != PropertyItems.Num();
+
+		if (!bRecreateList)
+		{
+			for (int32 PropertyIdx = 0; PropertyIdx < NewItems.Num(); ++PropertyIdx)
+			{
+				const FAvaRundownRCPropertyItemPtr& PropertyItem = PropertyItems[PropertyIdx];
+				const FAvaPropertyDetails& NewItem = NewItems[PropertyIdx];
+
+				if (!PropertyItem.IsValid())
+				{
+					bRecreateList = true;
+					break;
+				}
+
+				TSharedPtr<FRemoteControlEntity> PropertyEntity = PropertyItem->GetEntity();
+
+				if (!PropertyEntity.IsValid())
+				{
+					bRecreateList = true;
+					break;
+				}
+
+				if (PropertyEntity.Get() != NewItem.Entity.ToSharedPtr().Get())
+				{
+					bRecreateList = true;
+					break;
+				}
+			}
+		}
+
+		if (!bRecreateList)
+		{
+			RefreshTable();
+			return;
+		}
+
+		PropertyItems.Empty();
+
+		for (const FAvaPropertyDetails& NewItem : NewItems)
+		{
+			PropertyItems.Add(MakeShared<FAvaRundownRCPropertyItem>(NewItem.Entity, NewItem.bEntityControlled));
+		}
+	}
+}
+
+// Remarks:
+// This is called by URemoteControlPreset::OnEndFrame() as a result of an entity being modified.
+// However, it doesn't seem to be called (or not always) if the entity is modified by a controller action. 
+void SAvaRundownPageRemoteControlProps::OnRemoteControlExposedPropertiesModified(URemoteControlPreset* InPreset, const TSet<FGuid>& InModifiedProperties)
+{
+	if (!IsValid(InPreset) ||!HasRemoteControlPreset(InPreset))
+	{
+		return;
+	}
+
+	if (const TSharedPtr<FAvaRundownEditor> RundownEditor = RundownEditorWeak.Pin())
+	{
+		bool bModified = false;
+
+		for (const FGuid& Id : InModifiedProperties)
+		{
+			if (const TSharedPtr<FRemoteControlEntity> Entity = InPreset->GetExposedEntity<FRemoteControlEntity>(Id).Pin())
+			{
+				FAvaPlayableRemoteControlValue EntityValue;
+
+				{
+					using namespace UE::AvaPlayableRemoteControl;
+					if (const EAvaPlayableRemoteControlResult Result = GetValueOfEntity(Entity, EntityValue.Value); Failed(Result))
+					{
+						UE_LOG(LogAvaPlayableRemoteControl, Error, TEXT("Unable to get value of entity \"%s\": %s"),
+							*Entity->GetLabel().ToString(), *EnumToString(Result));
+						continue;
+					}
+				}
+				
+				const FAvaPlayableRemoteControlValue* StoredEntityValue = GetSelectedPageEntityValue(Entity);
+
+				if (StoredEntityValue && StoredEntityValue->IsSameValueAs(EntityValue))
+				{
+					continue;	// Skip if value is identical.
+				}
+
+				if (!SetSelectedPageEntityValue(Entity, EntityValue))
+				{
+					UE_LOG(LogAvaPlayableRemoteControl, Error, TEXT("Unable to set page entity value for: \"%s\""), *Entity->GetLabel().ToString());
+					continue;
+				}
+
+				bModified = true;
+			}
+		}
+
+		if (bModified)
+		{
+			RundownEditor->MarkAsModified();
+		}
+	}
+}
+
+void SAvaRundownPageRemoteControlProps::OnRemoteControlControllerModified(URemoteControlPreset* InPreset, const TSet<FGuid>& InModifiedControllerIds)
+{
+	if (!IsValid(InPreset) || !HasRemoteControlPreset(InPreset))
+	{
+		return;
+	}
+
+	TSet<FGuid> EntityIds;
+	for (const FGuid& ControllerId : InModifiedControllerIds)
+	{
+		UE::AvaPlayableRemoteControl::GetEntitiesControlledByController(InPreset, InPreset->GetController(ControllerId), EntityIds);
+	}
+
+	// If a controller changed, we need to propagate the refresh of the field's widgets.
+	// Optimization: only refresh the widgets that are related to the modified controllers.
+	RefreshTable(EntityIds);
+
+	// It seems OnPropertyChangedDelegate (OnExposedPropertiesModified()) is not called when properties are
+	// changed by controllers. Ensure the values are saved by calling our handler directly.
+	OnRemoteControlExposedPropertiesModified(InPreset, EntityIds);
+}
+
+void SAvaRundownPageRemoteControlProps::BindRemoteControlDelegates(URemoteControlPreset* InPreset)
+{
+	if (IsValid(InPreset))
+	{
+		if (!InPreset->OnEntityExposed().IsBoundToObject(this))
+		{
+			InPreset->OnEntityExposed().AddSP(this, &SAvaRundownPageRemoteControlProps::OnRemoteControlEntitiesExposed);
+		}
+
+		if (!InPreset->OnEntityUnexposed().IsBoundToObject(this))
+		{
+			InPreset->OnEntityUnexposed().AddSP(this, &SAvaRundownPageRemoteControlProps::OnRemoteControlEntitiesUnexposed);
+		}
+
+		if (!InPreset->OnEntitiesUpdated().IsBoundToObject(this))
+		{
+			InPreset->OnEntitiesUpdated().AddSP(this, &SAvaRundownPageRemoteControlProps::OnRemoteControlEntitiesUpdated);
+		}
+
+		if (!InPreset->OnExposedPropertiesModified().IsBoundToObject(this))
+		{
+			InPreset->OnExposedPropertiesModified().AddSP(this, &SAvaRundownPageRemoteControlProps::OnRemoteControlExposedPropertiesModified);
+		}
+
+		if (!InPreset->OnControllerModified().IsBoundToObject(this))
+		{
+			InPreset->OnControllerModified().AddSP(this, &SAvaRundownPageRemoteControlProps::OnRemoteControlControllerModified);
+		}
+	}
+}
+
+bool SAvaRundownPageRemoteControlProps::HasRemoteControlPreset(const URemoteControlPreset* InPreset) const
+{
+	for (const TSharedPtr<FAvaRundownManagedInstance>& ManagedInstance : ManagedInstances)
+	{
+		if (ManagedInstance && ManagedInstance->GetRemoteControlPreset() == InPreset)
+		{
+			return true;
+		}
+	}
+	return false;
+}
+
+FAvaRundownPage* SAvaRundownPageRemoteControlProps::GetActivePage() const
+{
+	if (ActivePageId == FAvaRundownPage::InvalidPageId)
+	{
+		return nullptr;
+	}
+
+	if (const TSharedPtr<FAvaRundownEditor> RundownEditor = RundownEditorWeak.Pin())
+	{
+		if (RundownEditor->IsRundownValid())
+		{
+			FAvaRundownPage& Page = RundownEditor->GetRundown()->GetPage(ActivePageId);
+			return Page.IsValidPage() ? &Page : nullptr;
+		}
+	}
+
+	return nullptr;
+}
+
+const FAvaPlayableRemoteControlValue* SAvaRundownPageRemoteControlProps::GetSelectedPageEntityValue(const TSharedPtr<FRemoteControlEntity>& InRemoteControlEntity) const
+{
+	if (InRemoteControlEntity.IsValid())
+	{
+		if (const FAvaRundownPage* Page = GetActivePage())
+		{
+			return Page->GetRemoteControlEntityValue(InRemoteControlEntity->GetId()); 
+		}
+	}
+
+	return nullptr;
+}
+
+bool SAvaRundownPageRemoteControlProps::SetSelectedPageEntityValue(const TSharedPtr<FRemoteControlEntity>& InRemoteControlEntity, const FAvaPlayableRemoteControlValue& InValue) const
+{
+	if (InRemoteControlEntity.IsValid())
+	{
+		if (const TSharedPtr<FAvaRundownEditor> RundownEditor = RundownEditorWeak.Pin())
+		{
+			if (RundownEditor->IsRundownValid())
+			{
+				// Using the rundown API for event propagation.
+				UAvaRundown* Rundown = RundownEditor->GetRundown(); 
+				return Rundown->SetRemoteControlEntityValue(ActivePageId, InRemoteControlEntity->GetId(), InValue);
+			}
+		}
+	}
+
+	return false;
+}
+
+#undef LOCTEXT_NAMESPACE
diff --git a/Engine/Plugins/Experimental/Avalanche/Source/AvalancheMediaEditor/Private/Rundown/DetailsView/RemoteControl/Properties/SAvaRundownRCPropertyItemRow.cpp b/Engine/Plugins/Experimental/Avalanche/Source/AvalancheMediaEditor/Private/Rundown/DetailsView/RemoteControl/Properties/SAvaRundownRCPropertyItemRow.cpp
new file mode 100644
--- /dev/null
+++ b/Engine/Plugins/Experimental/Avalanche/Source/AvalancheMediaEditor/Private/Rundown/DetailsView/RemoteControl/Properties/SAvaRundownRCPropertyItemRow.cpp
@@ -0,0 +1,152 @@
+// Copyright Epic Games, Inc. All Rights Reserved.
+
+#include "SAvaRundownRCPropertyItemRow.h"
+
+#include "AvaRundownPageRemoteControlWidgetUtils.h"
+#include "AvaRundownRCPropertyItem.h"
+#include "IDetailTreeNode.h"
+#include "IPropertyRowGenerator.h"
+#include "Modules/ModuleManager.h"
+#include "PropertyEditorModule.h"
+#include "Widgets/Layout/SScissorRectBox.h"
+#include "Widgets/SNullWidget.h"
+#include "Widgets/Text/STextBlock.h"
+
+#define LOCTEXT_NAMESPACE "SAvaRundownRCPropertyItemRow"
+
+void SAvaRundownRCPropertyItemRow::Construct(const FArguments& InArgs, TSharedRef<SAvaRundownPageRemoteControlProps> InPropertyPanel,
+	const TSharedRef<STableViewBase>& InOwnerTableView, const TSharedPtr<const FAvaRundownRCPropertyItem>& InRowItem)
+{
+	ItemPtrWeak = InRowItem;
+	PropertyPanelWeak = InPropertyPanel;
+	Generator = nullptr;
+	ValueContainer = nullptr;
+	ValueWidget = nullptr;
+
+	SMultiColumnTableRow<FAvaRundownRCPropertyItemPtr>::Construct(FSuperRowType::FArguments(), InOwnerTableView);
+}
+
+TSharedRef<SWidget> SAvaRundownRCPropertyItemRow::GenerateWidgetForColumn(const FName& InColumnName)
+{
+	TSharedPtr<const FAvaRundownRCPropertyItem> ItemPtr = ItemPtrWeak.Pin();
+
+	if (ItemPtr.IsValid())
+	{
+		if (InColumnName == SAvaRundownPageRemoteControlProps::PropertyColumnName)
+		{
+			return SNew(SScissorRectBox)
+				[
+					SNew(SBox)
+					.VAlign(EVerticalAlignment::VAlign_Center)
+					.Padding(3.f, 2.f, 3.f, 2.f)
+					[
+						SNew(STextBlock)
+						.Text(GetFieldLabel())
+					]
+				];
+		}
+		else if (InColumnName == SAvaRundownPageRemoteControlProps::ValueColumnName)
+		{
+			return SAssignNew(ValueContainer, SBox)
+				[
+					CreateValue()
+				];
+		}
+		else
+		{
+			TSharedPtr<SAvaRundownPageRemoteControlProps> PropertyPanel = PropertyPanelWeak.Pin();
+
+			if (PropertyPanel.IsValid())
+			{
+				TSharedPtr<SWidget> Cell = nullptr;
+				const TArray<FAvaRundownRCPropertyTableRowExtensionDelegate>& TableRowExtensionDelegates = PropertyPanel->GetTableRowExtensionDelegates(InColumnName);
+
+				for (const FAvaRundownRCPropertyTableRowExtensionDelegate& TableRowExtensionDelegate : TableRowExtensionDelegates)
+				{
+					TableRowExtensionDelegate.ExecuteIfBound(PropertyPanel.ToSharedRef(), ItemPtr.ToSharedRef(), Cell);
+				}
+
+				if (Cell.IsValid())
+				{
+					return Cell.ToSharedRef();
+				}
+			}
+		}
+	}
+
+	return SNullWidget::NullWidget;
+}
+
+void SAvaRundownRCPropertyItemRow::UpdateValue()
+{
+	if (ValueContainer.IsValid())
+	{
+		ValueContainer->SetContent(CreateValue());
+	}
+}
+
+FText SAvaRundownRCPropertyItemRow::GetFieldLabel() const
+{
+	if (TSharedPtr<const FAvaRundownRCPropertyItem> ItemPtr = ItemPtrWeak.Pin())
+	{
+		if (TSharedPtr<FRemoteControlEntity> EntityPtr = ItemPtr->GetEntity())
+		{
+			TSharedPtr<FRemoteControlField> FieldPtr = StaticCastSharedPtr<FRemoteControlField>(EntityPtr);
+			return FText::FromName(FieldPtr->FieldName);
+		}
+	}
+
+	return FText::GetEmpty();
+}
+
+TSharedRef<SWidget> SAvaRundownRCPropertyItemRow::CreateValue()
+{
+	if (TSharedPtr<const FAvaRundownRCPropertyItem> ItemPtr = ItemPtrWeak.Pin())
+	{
+		if (TSharedPtr<FRemoteControlEntity> EntityPtr = ItemPtr->GetEntity())
+		{
+			TSharedPtr<FRemoteControlField> FieldPtr = StaticCastSharedPtr<FRemoteControlField>(EntityPtr);
+
+			// For the moment, just use the first object.
+			TArray<UObject*> Objects = FieldPtr->GetBoundObjects();
+
+			if ((FieldPtr->FieldType == EExposedFieldType::Property) && (Objects.Num() > 0))
+			{
+				FPropertyRowGeneratorArgs Args;
+				Generator = FModuleManager::GetModuleChecked<FPropertyEditorModule>("PropertyEditor").CreatePropertyRowGenerator(Args);
+				Generator->SetObjects({Objects[0]});
+
+				if (TSharedPtr<IDetailTreeNode> Node = FAvaRundownPageRemoteControlWidgetUtils::FindNode(Generator->GetRootTreeNodes(), FieldPtr->FieldPathInfo.ToPathPropertyString(), FAvaRundownPageRemoteControlWidgetUtils::EFindNodeMethod::Path))
+				{
+					const FNodeWidgets NodeWidgets = Node->CreateNodeWidgets();
+					ValueWidget = NodeWidgets.WholeRowWidget.IsValid() ? NodeWidgets.WholeRowWidget : NodeWidgets.ValueWidget;
+
+					if (ItemPtr->IsEntityControlled())
+					{
+						ValueWidget = SNew(SHorizontalBox)
+							+ SHorizontalBox::Slot()
+							.VAlign(EVerticalAlignment::VAlign_Center)
+							[
+								ValueWidget.ToSharedRef()
+							]
+						+ SHorizontalBox::Slot()
+							.VAlign(EVerticalAlignment::VAlign_Center)
+							.Padding(3.f, 0.f, 0.f, 0.f)
+							[
+								SNew(STextBlock)
+								.Text(LOCTEXT("Controlled", "(Controlled)"))
+							];
+
+						ValueWidget->SetEnabled(false);
+					}
+
+					return ValueWidget.ToSharedRef();
+				}
+			}
+		}
+	}
+
+	return SNullWidget::NullWidget;
+}
+
+#undef LOCTEXT_NAMESPACE
diff --git a/Engine/Plugins/Experimental/Avalanche/Source/AvalancheMediaEditor/Private/Rundown/DetailsView/RemoteControl/Properties/SAvaRundownRCPropertyItemRow.h b/Engine/Plugins/Experimental/Avalanche/Source/AvalancheMediaEditor/Private/Rundown/DetailsView/RemoteControl/Properties/SAvaRundownRCPropertyItemRow.h
new file mode 100644
--- /dev/null
+++ b/Engine/Plugins/Experimental/Avalanche/Source/AvalancheMediaEditor/Private/Rundown/DetailsView/RemoteControl/Properties/SAvaRundownRCPropertyItemRow.h
@@ -0,0 +1,35 @@
+// Copyright Epic Games, Inc. All Rights Reserved.
+
+#pragma once
+
+#include "SAvaRundownPageRemoteControlProps.h"
+#include "Widgets/Views/STableRow.h"
+
+class SBox;
+class SWidget;
+
+class SAvaRundownRCPropertyItemRow : public SMultiColumnTableRow<FAvaRundownRCPropertyItemPtr>
+{
+public:
+	SLATE_BEGIN_ARGS(SAvaRundownRCPropertyItemRow) { }
+	SLATE_END_ARGS()
+
+	void Construct(const FArguments& InArgs, TSharedRef<SAvaRundownPageRemoteControlProps> InPropertyPanel,
+		const TSharedRef<STableViewBase>& InOwnerTableView, const TSharedPtr<const FAvaRundownRCPropertyItem>& InRowItem);
+
+	virtual TSharedRef<SWidget> GenerateWidgetForColumn(const FName& InColumnName) override;
+
+	void UpdateValue();
+
+protected:
+	TWeakPtr<const FAvaRundownRCPropertyItem> ItemPtrWeak;
+	TWeakPtr<SAvaRundownPageRemoteControlProps> PropertyPanelWeak;
+	TSharedPtr<IPropertyRowGenerator> Generator;
+	TSharedPtr<SBox> ValueContainer;
+	TSharedPtr<SWidget> ValueWidget;
+
+	/** Get this field's label. */
+	FText GetFieldLabel() const;
+
+	TSharedRef<SWidget> CreateValue();
+};
diff --git a/Engine/Plugins/Experimental/Avalanche/Source/AvalancheMediaEditor/Private/Rundown/DetailsView/SAvaRundownPageDetails.cpp b/Engine/Plugins/Experimental/Avalanche/Source/AvalancheMediaEditor/Private/Rundown/DetailsView/SAvaRundownPageDetails.cpp
new file mode 100644
--- /dev/null
+++ b/Engine/Plugins/Experimental/Avalanche/Source/AvalancheMediaEditor/Private/Rundown/DetailsView/SAvaRundownPageDetails.cpp
@@ -0,0 +1,440 @@
+// Copyright Epic Games, Inc. All Rights Reserved.
+
+#include "SAvaRundownPageDetails.h"
+
+#include "Async/Async.h"
+#include "IAvaMediaModule.h"
+#include "RemoteControl/Controllers/SAvaRundownRCControllerPanel.h"
+#include "Rundown/AvaRundown.h"
+#include "Rundown/AvaRundownEditor.h"
+#include "Rundown/AvaRundownManagedInstanceCache.h"
+#include "Rundown/AvaRundownPage.h"
+#include "Rundown/DetailsView/RemoteControl/Properties/SAvaRundownPageRemoteControlProps.h"
+#include "Rundown/Pages/Slate/SAvaRundownInstancedPageList.h"
+#include "Rundown/Pages/Slate/SAvaRundownPageList.h"
+#include "Widgets/Images/SImage.h"
+#include "Widgets/Input/SButton.h"
+#include "Widgets/Input/SEditableTextBox.h"
+#include "Widgets/Layout/SExpandableArea.h"
+#include "Widgets/Layout/SScrollBox.h"
+#include "Widgets/Text/STextBlock.h"
+
+#define LOCTEXT_NAMESPACE "SAvaRundownPageDetails"
+
+void SAvaRundownPageDetails::Construct(const FArguments& InArgs, const TSharedPtr<FAvaRundownEditor>& InRundownEditor)
+{
+	RundownEditorWeak = InRundownEditor;
+	ActivePageId = FAvaRundownPage::InvalidPageId;
+
+	InRundownEditor->GetOnPageEvent().AddSP(this, &SAvaRundownPageDetails::OnPageEvent);
+	IAvaMediaModule::Get().GetManagedInstanceCache().OnEntryInvalidated.AddSP(this, &SAvaRundownPageDetails::OnManagedInstanceCacheEntryInvalidated);
+
+	TSharedRef<SHorizontalBox> AnimationHeader = SNew(SHorizontalBox);
+	{
+		static const TArray<FText> HeaderTitles
+    	{
+    		LOCTEXT("HeaderTitleAnimationName", "Animation Name"),
+    		LOCTEXT("HeaderTitleLoopsToPlay", "Loops to Play"),
+    		LOCTEXT("HeaderTitlePlaybackSpeed", "Playback Speed"),
+    		LOCTEXT("HeaderTitlePlayMode", "Play Mode")
+    	};
+
+    	for (const FText& Title : HeaderTitles)
+    	{
+    		AnimationHeader->AddSlot()
+    			.FillWidth(1.f)
+    			[
+    				SNew(STextBlock)
+    				.Text(Title)
+    				.Justification(ETextJustify::Center)
+    			];
+    	}
+	}
+
+	ChildSlot
+	[
+		SNew(SScrollBox)
+		+ SScrollBox::Slot()
+		[
+			SNew(SVerticalBox)
+			+ SVerticalBox::Slot()
+			.AutoHeight()
+			.Padding(FMargin(10.f, 10.f, 10.f, 0.f))
+			[
+				SNew(SHorizontalBox)
+				+ SHorizontalBox::Slot()
+				.VAlign(EVerticalAlignment::VAlign_Center)
+				.MaxWidth(75.f)
+				[
+					SNew(STextBlock)
+					.Text(LOCTEXT("PageId", "Page Id"))
+					.MinDesiredWidth(75.f)
+				]
+				+ SHorizontalBox::Slot()
+				.VAlign(EVerticalAlignment::VAlign_Center)
+				.Padding(5.f, 0.f, 0.f, 0.f)
+				.MaxWidth(70.f)
+				[
+					SNew(SEditableTextBox)
+					.HintText(LOCTEXT("PageIdHint", "Page Id"))
+					.OnTextCommitted(this, &SAvaRundownPageDetails::OnPageIdCommitted)
+					.Text(this, &SAvaRundownPageDetails::GetPageId)
+					.IsEnabled(this, &SAvaRundownPageDetails::HasSelectedPage)
+				]
+				+ SHorizontalBox::Slot()
+				.VAlign(EVerticalAlignment::VAlign_Center)
+				.Padding(5.f, 0.f, 0.f, 0.f)
+				.AutoWidth()
+				[
+					SNew(SButton)
+					.ButtonStyle(FAppStyle::Get(), "SimpleButton")
+					.ToolTipText(LOCTEXT("DuplicatePageTooltip", "DuplicatePage"))
+					.OnClicked(this, &SAvaRundownPageDetails::DuplicateSelectedPage)
+					.IsEnabled(this, &SAvaRundownPageDetails::HasSelectedPage)
+					[
+						SNew(SImage)
+						.Image(FAppStyle::Get().GetBrush("GenericCommands.Duplicate"))
+					]
+				]
+			]
+			+ SVerticalBox::Slot()
+			.AutoHeight()
+			.Padding(FMargin(10.f, 3.f, 10.f, 0.f))
+			[
+				SNew(SHorizontalBox)
+				+ SHorizontalBox::Slot()
+				.VAlign(EVerticalAlignment::VAlign_Center)
+				.MaxWidth(75.f)
+				[
+					SNew(STextBlock)
+					.Text(LOCTEXT("PageName", "Page Name"))
+					.MinDesiredWidth(75.f)
+				]
+				+ SHorizontalBox::Slot()
+				.VAlign(EVerticalAlignment::VAlign_Center)
+				.Padding(5.f, 0.f, 0.f, 0.f)
+				[
+					SNew(SEditableTextBox)
+					.HintText(LOCTEXT("PageNameHint", "Page Name"))
+					.OnTextChanged(this, &SAvaRundownPageDetails::OnPageNameChanged)
+					.Text(this, &SAvaRundownPageDetails::GetPageDescription)
+					.IsEnabled(this, &SAvaRundownPageDetails::HasSelectedPage)
+				]
+			]
+			// Controllers
+			+ SVerticalBox::Slot()
+			.AutoHeight()
+			.Padding(FMargin(0.f, 10.f, 0.f, 0.f))
+			[
+				SAssignNew(RCControllerPanel, SAvaRundownRCControllerPanel, InRundownEditor)
+			]
+			// Exposed Properties
+			+ SVerticalBox::Slot()
+			.Padding(FMargin(0.f, 10.f, 0.f, 0.f))
+			.AutoHeight()
+			[
+				SNew(SHorizontalBox)
+				+ SHorizontalBox::Slot()
+				.VAlign(VAlign_Center)
+				.HAlign(HAlign_Left)
+				.Padding(5.f, 0.f, 0.f, 0.f)
+				.AutoWidth()
+				[
+					SNew(SButton)
+					.ContentPadding(0)
+					.ButtonStyle(FAppStyle::Get(), "NoBorder")
+					.OnClicked(this, &SAvaRundownPageDetails::ToggleExposedPropertiesVisibility)
+					.ToolTipText(LOCTEXT("VisibilityButtonToolTip", "Toggle Exposed Properties Visibility"))
+					.Content()
+					[
+						SNew(SImage)
+						.Image(this, &SAvaRundownPageDetails::GetExposedPropertiesVisibilityBrush)
+					]
+				]
+				+ SHorizontalBox::Slot()
+				.VAlign(VAlign_Center)
+				.HAlign(HAlign_Left)
+				.Padding(5.f, 0.f, 0.f, 0.f)
+				.AutoWidth()
+				[
+					SNew(STextBlock)
+					.Text(LOCTEXT("Properties", "Properties"))
+				]
+			]
+			+ SVerticalBox::Slot()
+			.AutoHeight()
+			[
+				SAssignNew(RemoteControlProps, SAvaRundownPageRemoteControlProps, InRundownEditor)
+				.Visibility(EVisibility::Collapsed)
+			]
+		]
+	];
+
+	OnPageSelectionChanged({});
+}
+
+SAvaRundownPageDetails::~SAvaRundownPageDetails()
+{
+	if (const TSharedPtr<FAvaRundownEditor> RundownEditor = RundownEditorWeak.Pin())
+	{
+		RundownEditor->GetOnPageEvent().RemoveAll(this);
+	}
+	if (IAvaMediaModule::IsModuleLoaded())
+	{
+		IAvaMediaModule::Get().GetManagedInstanceCache().OnEntryInvalidated.RemoveAll(this);
+	}
+}
+
+void SAvaRundownPageDetails::OnPageEvent(const TArray<int32>& InSelectedPageIds, UE::AvaRundown::EPageEvent InPageEvent)
+{
+	if (InPageEvent == UE::AvaRundown::EPageEvent::SelectionChanged)
+	{
+		OnPageSelectionChanged(InSelectedPageIds);
+		RemoteControlProps->Refresh(InSelectedPageIds);
+		RCControllerPanel->Refresh(InSelectedPageIds);
+	}
+	else if (InPageEvent == UE::AvaRundown::EPageEvent::ReimportRequest)
+	{
+		// Note: all widgets are refreshed in this event handler to ensure it is done
+		// after the cache is invalidated.
+		OnPageSelectionChanged(InSelectedPageIds);
+		// Only call SAvaRundownPageRemoteControlProps::UpdateDefaultValuesAndRefresh to have it update the values of all selected pages.
+		RemoteControlProps->UpdateDefaultValuesAndRefresh(InSelectedPageIds);
+		// Pages values already updated above, only need to refresh UI.
+		RCControllerPanel->Refresh(InSelectedPageIds);
+	}
+}
+
+void SAvaRundownPageDetails::OnPageSelectionChanged(const TArray<int32>& InSelectedPageIds)
+{
+	ActivePageId = InSelectedPageIds.IsEmpty() ? FAvaRundownPage::InvalidPageId : InSelectedPageIds[0];
+}
+
+void SAvaRundownPageDetails::OnManagedInstanceCacheEntryInvalidated(const FSoftObjectPath& InAssetPath)
+{
+	if (!bRefreshSelectedPageQueued)
+	{
+		if (TSharedPtr<FAvaRundownEditor> RundownEditor = RundownEditorWeak.Pin())
+		{
+			UAvaRundown* Rundown = RundownEditor->GetRundown();
+
+			if (IsValid(Rundown))
+			{
+				const FAvaRundownPage& SelectedPage = GetSelectedPage();
+
+				if (SelectedPage.IsValidPage())
+				{
+					if (SelectedPage.GetAssetPath(Rundown) == InAssetPath)
+					{
+						bRefreshSelectedPageQueued = true;
+						// Queue a refresh on next tick.
+						// We don't want to refresh immediately to avoid issues with
+						// cascading events within the managed blueprint cache.
+						TWeakPtr<SWidget> ThisWeak(AsShared());
+						AsyncTask(ENamedThreads::GameThread, [ThisWeak]()
+							{
+								if (const TSharedPtr<SWidget> ThisWidget = ThisWeak.Pin())
+								{
+									SAvaRundownPageDetails* AvaPageDetails = static_cast<SAvaRundownPageDetails*>(ThisWidget.Get());
+									AvaPageDetails->RefreshSelectedPage();
+									AvaPageDetails->bRefreshSelectedPageQueued = false;
+								}
+							});
+					}
+				}
+			}
+		}
+	}
+}
+
+FReply SAvaRundownPageDetails::ToggleExposedPropertiesVisibility()
+{
+	if (RemoteControlProps->GetVisibility() == EVisibility::Collapsed)
+	{
+		RemoteControlProps->SetVisibility(EVisibility::SelfHitTestInvisible);
+	}
+	else
+	{
+		RemoteControlProps->SetVisibility(EVisibility::Collapsed);
+	}
+
+	return FReply::Handled();
+}
+
+const FSlateBrush* SAvaRundownPageDetails::GetExposedPropertiesVisibilityBrush() const
+{
+	if (RemoteControlProps->GetVisibility() == EVisibility::Collapsed)
+	{
+		return FAppStyle::GetBrush("Level.NotVisibleHighlightIcon16x");
+	}
+	else
+	{
+		return FAppStyle::GetBrush("Level.VisibleHighlightIcon16x");
+	}
+}
+
+const FAvaRundownPage& SAvaRundownPageDetails::GetSelectedPage() const
+{
+	const TSharedPtr<FAvaRundownEditor> RundownEditor = RundownEditorWeak.Pin();
+
+	if (HasSelectedPage() && RundownEditor && RundownEditor->IsRundownValid())
+	{
+		return RundownEditor->GetRundown()->GetPage(ActivePageId);
+	}
+
+	return FAvaRundownPage::NullPage;
+}
+
+FAvaRundownPage& SAvaRundownPageDetails::GetMutableSelectedPage() const
+{
+	const TSharedPtr<FAvaRundownEditor> RundownEditor = RundownEditorWeak.Pin();
+
+	if (HasSelectedPage() && RundownEditor && RundownEditor->IsRundownValid())
+	{
+		return RundownEditor->GetRundown()->GetPage(ActivePageId);
+	}
+
+	return FAvaRundownPage::NullPage;
+}
+
+void SAvaRundownPageDetails::RefreshSelectedPage()
+{
+	const FAvaRundownPage& SelectedPage = GetSelectedPage();
+
+	if (SelectedPage.IsValidPage())
+	{
+		OnPageSelectionChanged({SelectedPage.GetPageId()});
+		RemoteControlProps->UpdateDefaultValuesAndRefresh({SelectedPage.GetPageId()});
+		RCControllerPanel->Refresh({SelectedPage.GetPageId()});
+	}
+}
+
+bool SAvaRundownPageDetails::HasSelectedPage() const
+{
+	if (ActivePageId == FAvaRundownPage::InvalidPageId)
+	{
+		return false;
+	}
+
+	return RundownEditorWeak.IsValid();
+}
+
+FText SAvaRundownPageDetails::GetPageId() const
+{
+	const FAvaRundownPage& SelectedPage = GetSelectedPage();
+
+	if (SelectedPage.IsValidPage())
+	{
+		return FText::AsNumber(SelectedPage.GetPageId(), &UE::AvaRundown::FEditorMetrics::PageIdFormattingOptions);
+	}
+
+	return FText::GetEmpty();
+}
+
+void SAvaRundownPageDetails::OnPageIdCommitted(const FText& InNewText, ETextCommit::Type InCommitType)
+{
+	switch (InCommitType)
+	{
+		case ETextCommit::OnEnter:
+		case ETextCommit::OnUserMovedFocus:
+			break;
+
+		case ETextCommit::Default:
+		case ETextCommit::OnCleared:
+		default:
+			return;
+	}
+
+	if (!InNewText.IsNumeric())
+	{
+		return;
+	}
+
+	FAvaRundownPage& SelectedPage = GetMutableSelectedPage();
+
+	if (!SelectedPage.IsValidPage()) // Not FAvaRundownPage::NullPage
+	{
+		return;
+	}
+
+	int32 NewId = FCString::Atoi(*InNewText.ToString());
+
+	if (NewId == SelectedPage.GetPageId())
+	{
+		return;
+	}
+
+	const TSharedPtr<FAvaRundownEditor> RundownEditor = RundownEditorWeak.Pin();
+
+	if (!RundownEditor.IsValid() || !IsValid(RundownEditor->GetRundown()))
+	{
+		return;
+	}
+
+	const bool bSuccessful = RundownEditor->GetRundown()->RenumberPageId(SelectedPage.GetPageId(), NewId);
+
+	if (bSuccessful)
+	{
+		TSharedPtr<SAvaRundownInstancedPageList> PageList = RundownEditor->GetActiveListWidget();
+		
+		if (PageList.IsValid())
+		{
+			PageList->SelectPage(NewId);
+		}
+	}
+}
+
+FText SAvaRundownPageDetails::GetPageDescription() const
+{
+	const FAvaRundownPage& SelectedPage = GetSelectedPage();
+
+	if (SelectedPage.IsValidPage())
+	{
+		return SelectedPage.GetPageDescription();
+	}
+
+	return FText::GetEmpty();
+}
+
+void SAvaRundownPageDetails::OnPageNameChanged(const FText& InNewText)
+{
+	FAvaRundownPage& SelectedPage = GetMutableSelectedPage();
+
+	if (SelectedPage.IsValidPage()) // Not FAvaRundownPage::NullPage
+	{
+		SelectedPage.SetPageFriendlyName(InNewText);
+	}
+}
+
+FReply SAvaRundownPageDetails::DuplicateSelectedPage()
+{
+	FAvaRundownPage& SelectedPage = GetMutableSelectedPage();
+
+	if (!SelectedPage.IsValidPage()) // Not FAvaRundownPage::NullPage
+	{
+		return FReply::Unhandled();
+	}
+
+	const TSharedPtr<FAvaRundownEditor> RundownEditor = RundownEditorWeak.Pin();
+
+	if (!RundownEditor.IsValid())
+	{
+		return FReply::Unhandled();
+	}
+
+	TSharedPtr<SAvaRundownInstancedPageList> PageList = RundownEditor->GetActiveListWidget();
+
+	if (!PageList.IsValid())
+	{
+		return FReply::Unhandled();
+	}
+
+	TArray<int32> SelectedPages = TArray<int32>(PageList->GetSelectedPageIds());
+	PageList->SelectPage(SelectedPage.GetPageId());
+	PageList->DuplicateSelectedPages();
+	PageList->SelectPages(SelectedPages);
+
+	return FReply::Handled();
+}
+
+#undef LOCTEXT_NAMESPACE
diff --git a/Engine/Plugins/Experimental/Avalanche/Source/AvalancheMediaEditor/Private/Rundown/Factories/Suggestions/AvaRundownFilterIdSuggestionFactory.cpp b/Engine/Plugins/Experimental/Avalanche/Source/AvalancheMediaEditor/Private/Rundown/Factories/Suggestions/AvaRundownFilterIdSuggestionFactory.cpp
new file mode 100644
--- /dev/null
+++ b/Engine/Plugins/Experimental/Avalanche/Source/AvalancheMediaEditor/Private/Rundown/Factories/Suggestions/AvaRundownFilterIdSuggestionFactory.cpp
@@ -0,0 +1,29 @@
+// Copyright Epic Games, Inc. All Rights Reserved.
+
+#include "AvaRundownFilterIdSuggestionFactory.h"
+
+#include "Rundown/Factories/Filters/AvaRundownFactoriesUtils.h"
+#include "SAssetSearchBox.h"
+
+#define LOCTEXT_NAMESPACE "AvaRundownFilterIdSuggestionFactory"
+
+const FName FAvaRundownFilterIdSuggestionFactory::KeyName = FName(TEXT("ID"));
+
+void FAvaRundownFilterIdSuggestionFactory::AddSuggestion(const TSharedRef<FAvaRundownFilterSuggestionPayload>& InPayload)
+{
+	const FText IdCategoryLabel = LOCTEXT("IdCategoryLabel", "Ava-Rundown-Id");
+	FString IdNameSuggestion = TEXT("Id");
+	const bool bIsFilterValueValid = InPayload->FilterValue.IsEmpty() || IdNameSuggestion.Contains(InPayload->FilterValue);
+	if (bIsFilterValueValid && !InPayload->FilterCache.Contains(IdNameSuggestion))
+	{
+		InPayload->FilterCache.Add(IdNameSuggestion);
+		InPayload->PossibleSuggestions.Add(FAssetSearchBoxSuggestion{MoveTemp(IdNameSuggestion), LOCTEXT("IdSuggestion","Id"), IdCategoryLabel});
+	}
+}
+
+bool FAvaRundownFilterIdSuggestionFactory::SupportSuggestionType(EAvaRundownSearchListType InSuggestionType) const
+{
+	return InSuggestionType != EAvaRundownSearchListType::None;
+}
+
+#undef LOCTEXT_NAMESPACE
diff --git a/Engine/Plugins/Experimental/Avalanche/Source/AvalancheMediaEditor/Private/Rundown/Factories/Suggestions/AvaRundownFilterPathSuggestionFactory.cpp b/Engine/Plugins/Experimental/Avalanche/Source/AvalancheMediaEditor/Private/Rundown/Factories/Suggestions/AvaRundownFilterPathSuggestionFactory.cpp
new file mode 100644
--- /dev/null
+++ b/Engine/Plugins/Experimental/Avalanche/Source/AvalancheMediaEditor/Private/Rundown/Factories/Suggestions/AvaRundownFilterPathSuggestionFactory.cpp
@@ -0,0 +1,37 @@
+// Copyright Epic Games, Inc. All Rights Reserved.
+
+#include "AvaRundownFilterPathSuggestionFactory.h"
+
+#include "Rundown/AvaRundown.h"
+#include "Rundown/Factories/Filters/AvaRundownFactoriesUtils.h"
+#include "SAssetSearchBox.h"
+
+#define LOCTEXT_NAMESPACE "AvaRundownFilterPathSuggestionFactory"
+
+const FName FAvaRundownFilterPathSuggestionFactory::KeyName = FName(TEXT("ASSET"));
+
+void FAvaRundownFilterPathSuggestionFactory::AddSuggestion(const TSharedRef<FAvaRundownFilterSuggestionPayload>& InPayload)
+{
+	const FAvaRundownPage& PageItem = UAvaRundown::GetPageSafe(InPayload->Rundown, InPayload->ItemPageId);
+	if (PageItem.IsValidPage())
+	{
+		const FText PathCategoryLabel = LOCTEXT("PathCategoryLabel", "Ava-Rundown-Asset");
+		const FString AssetName = PageItem.GetAssetPath(InPayload->Rundown).GetAssetName();
+
+		FString BlueprintNameSuggestion = FString::Printf(TEXT("Asset=%s"), *AssetName);
+		const bool bIsFilterValueValid = InPayload->FilterValue.IsEmpty() || BlueprintNameSuggestion.Contains(InPayload->FilterValue);
+
+		if (bIsFilterValueValid && !InPayload->FilterCache.Contains(BlueprintNameSuggestion))
+		{
+			InPayload->FilterCache.Add(BlueprintNameSuggestion);
+			InPayload->PossibleSuggestions.Add(FAssetSearchBoxSuggestion{MoveTemp(BlueprintNameSuggestion), FText::FromString(AssetName), PathCategoryLabel});
+		}
+	}
+}
+
+bool FAvaRundownFilterPathSuggestionFactory::SupportSuggestionType(EAvaRundownSearchListType InSuggestionType) const
+{
+	return InSuggestionType != EAvaRundownSearchListType::None;
+}
+
+#undef LOCTEXT_NAMESPACE
diff --git a/Engine/Plugins/Experimental/Avalanche/Source/AvalancheMediaEditor/Private/Rundown/MediaOutputEditorUtils/AvaRundownOutputEditorUtils.cpp b/Engine/Plugins/Experimental/Avalanche/Source/AvalancheMediaEditor/Private/Rundown/MediaOutputEditorUtils/AvaRundownOutputEditorUtils.cpp
new file mode 100644
--- /dev/null
+++ b/Engine/Plugins/Experimental/Avalanche/Source/AvalancheMediaEditor/Private/Rundown/MediaOutputEditorUtils/AvaRundownOutputEditorUtils.cpp
@@ -0,0 +1,261 @@
+// Copyright Epic Games, Inc. All Rights Reserved.
+
+#include "AvaRundownOutputEditorUtils.h"
+
+#include "Dom/JsonObject.h"
+#include "Serialization/JsonReader.h"
+#include "Serialization/JsonSerializer.h"
+#include "Serialization/JsonWriter.h"
+
+DEFINE_LOG_CATEGORY(LogAvaRundownEditorOutputUtils);
+
+FAvaRundownOutputEditorUtils::FAvaRundownOutputEditorUtils()
+{
+}
+
+FAvaRundownOutputEditorUtils::~FAvaRundownOutputEditorUtils()
+{
+}
+
+FString FAvaRundownOutputEditorUtils::SerializeMediaOutput(const UMediaOutput* InMediaOutput)
+{
+	const TSharedPtr<FJsonObject> OutputObject = MakeShareable(new FJsonObject);
+
+	OutputObject->SetStringField("Class", InMediaOutput->GetClass()->GetPathName());
+	OutputObject->SetStringField("Name", InMediaOutput->GetName());
+
+	TArray<TSharedPtr<FJsonValue>> PropertyArray;
+	for (TFieldIterator<FProperty> Iterator(InMediaOutput->GetClass()); Iterator; ++Iterator)
+	{
+		PropertyArray.Add(MakeShareable(new FJsonValueObject(ParsePropertyInfo(Iterator, InMediaOutput))));
+	}
+
+	const TSharedPtr<FJsonObject> PropertiesObject = MakeShareable(new FJsonObject);
+	OutputObject->SetArrayField("Properties", PropertyArray);
+
+	FString OutputString;
+	const TSharedRef<TJsonWriter<>> JsonWriter = TJsonWriterFactory<>::Create(&OutputString);
+	FJsonSerializer::Serialize(OutputObject.ToSharedRef(), JsonWriter);
+	
+	return OutputString;
+}
+
+void FAvaRundownOutputEditorUtils::EditMediaOutput(UMediaOutput* InMediaOutput, const FString& InDeviceData)
+{
+	const TSharedRef<TJsonReader<TCHAR>> JsonReader = TJsonReaderFactory<TCHAR>::Create(InDeviceData);
+	if (TSharedPtr<FJsonObject> NewDeviceData; FJsonSerializer::Deserialize(JsonReader, NewDeviceData))
+	{
+		for (TArray<TSharedPtr<FJsonValue>> PropertyValues = NewDeviceData->GetArrayField(TEXT("Properties"));
+		     TSharedPtr<FJsonValue> PropertyValue : PropertyValues)
+		{
+			const TSharedPtr<FJsonObject> PropertyObject = PropertyValue.Get()->AsObject();
+
+			for (TFieldIterator<FProperty> PropertyIterator(InMediaOutput->GetClass());
+			     PropertyIterator; ++PropertyIterator)
+			{
+				if (PropertyObject->GetStringField(TEXT("Name")) == PropertyIterator->GetName())
+				{
+					SetProperty(InMediaOutput, PropertyObject, PropertyIterator);
+				}
+			}
+		}
+	}
+}
+
+TSharedPtr<FJsonObject> FAvaRundownOutputEditorUtils::ParsePropertyInfo(TFieldIterator<FProperty> InProperty,
+                                                                       const void* InOwnerObject)
+{
+	const FFieldClass* FieldClass = InProperty->GetClass();
+	if (FieldClass == FIntProperty::StaticClass() || FieldClass == FNameProperty::StaticClass() ||
+		FieldClass == FBoolProperty::StaticClass() || FieldClass == FStrProperty::StaticClass())
+	{
+		return ParseElementaryPropertyInfo(InProperty, InOwnerObject);
+	}
+	if (FieldClass == FEnumProperty::StaticClass())
+	{
+		return ParseEnumPropertyInfo(InProperty, InOwnerObject);
+	}
+	if (FieldClass == FStructProperty::StaticClass())
+	{
+		return ParseStructPropertyInfo(InProperty, InOwnerObject);
+	}
+	return MakeShareable(new FJsonObject);
+}
+
+TSharedPtr<FJsonObject> FAvaRundownOutputEditorUtils::ParseElementaryPropertyInfo(
+	TFieldIterator<FProperty> InProperty, const void* InOwnerObject)
+{
+	const TSharedPtr<FJsonObject> PropertyObject = MakeShareable(new FJsonObject);
+	if (InProperty->GetClass() == FIntProperty::StaticClass())
+	{
+		const FIntProperty* IntProperty = CastField<FIntProperty>(InProperty->GetOwnerProperty());
+		const int32 Value = IntProperty->GetPropertyValue_InContainer(InOwnerObject);
+		PropertyObject->SetStringField("Name", InProperty->GetName());
+		PropertyObject->SetStringField("DisplayName", InProperty->GetDisplayNameText().ToString());
+		PropertyObject->SetStringField("Type", "number");
+		PropertyObject->SetNumberField("Value", Value);
+	}
+	else if (InProperty->GetClass() == FBoolProperty::StaticClass())
+	{
+		const FBoolProperty* BoolProperty = CastField<FBoolProperty>(InProperty->GetOwnerProperty());
+		const bool Value = BoolProperty->GetPropertyValue_InContainer(InOwnerObject);
+		PropertyObject->SetStringField("Name", InProperty->GetName());
+		PropertyObject->SetStringField("DisplayName", InProperty->GetDisplayNameText().ToString());
+		PropertyObject->SetStringField("Type", "boolean");
+		PropertyObject->SetBoolField("Value", Value);
+	}
+	else if (InProperty->GetClass() == FStrProperty::StaticClass())
+	{
+		const FStrProperty* TextProperty = CastField<FStrProperty>(InProperty->GetOwnerProperty());
+		const FString Value = TextProperty->GetPropertyValue_InContainer(InOwnerObject);
+		PropertyObject->SetStringField("Name", InProperty->GetName());
+		PropertyObject->SetStringField("DisplayName", InProperty->GetDisplayNameText().ToString());
+		PropertyObject->SetStringField("Type", "string");
+		PropertyObject->SetStringField("Value", Value);
+	}
+	else if (InProperty->GetClass() == FNameProperty::StaticClass())
+	{
+		const FNameProperty* NameProperty = CastField<FNameProperty>(InProperty->GetOwnerProperty());
+		const FName Value = NameProperty->GetPropertyValue_InContainer(InOwnerObject);
+		PropertyObject->SetStringField("Name", InProperty->GetName());
+		PropertyObject->SetStringField("DisplayName", InProperty->GetDisplayNameText().ToString());
+		PropertyObject->SetStringField("Type", "name");
+		PropertyObject->SetStringField("Value", Value.ToString());
+	}
+	else
+	{
+		UE_LOG(LogAvaRundownEditorOutputUtils, Warning, TEXT("Property type: %s is not supported"), *InProperty->GetClass()->GetName());
+	}
+
+	return PropertyObject;
+}
+
+TSharedPtr<FJsonObject> FAvaRundownOutputEditorUtils::ParseEnumPropertyInfo(
+	TFieldIterator<FProperty> InProperty, const void* InOwnerObject)
+{
+	const TSharedPtr<FJsonObject> PropertyObject = MakeShareable(new FJsonObject);
+	const FEnumProperty* EnumProperty = CastField<FEnumProperty>(InProperty->GetOwnerProperty());
+
+	PropertyObject->SetStringField("Name", InProperty->GetName());
+	PropertyObject->SetStringField("DisplayName", InProperty->GetDisplayNameText().ToString());
+	PropertyObject->SetStringField("Type", "enum");
+	PropertyObject->SetStringField("Class", EnumProperty->GetEnum()->GetName());
+
+	const int32 MaxIndex = EnumProperty->GetEnum()->NumEnums();
+	TArray<TSharedPtr<FJsonValue>> EnumArray;
+	for (int32 Index = 0; Index < MaxIndex; ++Index)
+	{
+		const TSharedPtr<FJsonObject> EnumValueObject = MakeShareable(new FJsonObject);
+		EnumValueObject->SetStringField(FString::FromInt(Index),
+		                                *EnumProperty->GetEnum()->GetAuthoredNameStringByIndex(Index));
+		EnumArray.Add(MakeShareable(new FJsonValueObject(EnumValueObject)));
+		PropertyObject->SetArrayField("EnumValues", EnumArray);
+	}
+	const uint8* Value = EnumProperty->ContainerPtrToValuePtr<uint8>(InOwnerObject);
+
+	PropertyObject->SetStringField("Value", *EnumProperty->GetEnum()->GetAuthoredNameStringByValue(*Value));
+	return PropertyObject;
+}
+
+TSharedPtr<FJsonObject> FAvaRundownOutputEditorUtils::ParseStructPropertyInfo(
+	TFieldIterator<FProperty> InProperty, const void* InOwnerObject)
+{
+	const TSharedPtr<FJsonObject> PropertyObject = MakeShareable(new FJsonObject);
+	const FStructProperty* StructProperty = CastField<FStructProperty>(InProperty->GetOwnerProperty());
+
+	PropertyObject->SetStringField("Name", InProperty->GetName());
+	PropertyObject->SetStringField("DisplayName", InProperty->GetDisplayNameText().ToString());
+	PropertyObject->SetStringField("Type", "struct");
+
+	const UScriptStruct* Struct = StructProperty->Struct;
+	PropertyObject->SetStringField("Class", Struct->GetName());
+
+	const void* StructPtr = StructProperty->ContainerPtrToValuePtr<void>(InOwnerObject);
+
+	TArray<TSharedPtr<FJsonValue>> PropertyArray;
+	for (TFieldIterator<FProperty> StructIterator(Struct); StructIterator; ++StructIterator)
+	{
+		PropertyArray.Add(MakeShareable(new FJsonValueObject(ParsePropertyInfo(StructIterator, StructPtr))));
+	}
+	PropertyObject->SetArrayField("StructProperties", PropertyArray);
+
+	return PropertyObject;
+}
+
+void FAvaRundownOutputEditorUtils::SetProperty(void* InOwnerObject, const TSharedPtr<FJsonObject>& InPropertyObject,
+                                              const TFieldIterator<FProperty>& InProperty)
+{
+	if (InPropertyObject->GetStringField(TEXT("Type")) == "number" || InPropertyObject->GetStringField(TEXT("Type")) == "string"
+		|| InPropertyObject->GetStringField(TEXT("Type")) == "boolean" || InPropertyObject->GetStringField(TEXT("Type")) == "name")
+	{
+		SetElementaryProperty(InOwnerObject, InPropertyObject, InProperty);
+	}
+	else if (InPropertyObject->GetStringField(TEXT("Type")) == "enum")
+	{
+		SetEnumProperty(InOwnerObject, InPropertyObject, InProperty);
+	}
+	else if (InPropertyObject->GetStringField(TEXT("Type")) == "struct")
+	{
+		SetStructProperties(InOwnerObject, InPropertyObject, InProperty);
+	}
+}
+
+void FAvaRundownOutputEditorUtils::SetElementaryProperty(void* InOwnerObject,
+	const TSharedPtr<FJsonObject>& InPropertyObject, TFieldIterator<FProperty> InProperty)
+{
+	if (InPropertyObject->GetStringField(TEXT("Type")) == "number")
+	{
+		const int32 NumberValue = InPropertyObject->GetNumberField(TEXT("Value"));
+		InProperty->GetOwnerProperty()->SetValue_InContainer(InOwnerObject, &NumberValue);
+	}
+	else if (InPropertyObject->GetStringField(TEXT("Type")) == "string")
+	{
+		const FString StringValue = InPropertyObject->GetStringField(TEXT("Value"));
+		InProperty->GetOwnerProperty()->SetValue_InContainer(InOwnerObject, &StringValue);
+	}
+	else if (InPropertyObject->GetStringField(TEXT("Type")) == "boolean")
+	{
+		const bool BoolValue = InPropertyObject->GetBoolField(TEXT("Value"));
+		InProperty->GetOwnerProperty()->SetValue_InContainer(InOwnerObject, &BoolValue);
+	}
+	else if (InPropertyObject->GetStringField(TEXT("Type")) == "name")
+	{
+		const FName NameValue = FName(InPropertyObject->GetStringField(TEXT("Value")));
+		InProperty->GetOwnerProperty()->SetValue_InContainer(InOwnerObject, &NameValue);
+	}
+}
+
+void FAvaRundownOutputEditorUtils::SetEnumProperty(void* InOwnerObject,
+	const TSharedPtr<FJsonObject>& InPropertyObject, TFieldIterator<FProperty> InProperty)
+{
+	const FString ValueName = InPropertyObject->GetStringField(TEXT("Value"));
+
+	const FEnumProperty* EnumProperty = CastField<FEnumProperty>(
+		InProperty->GetOwnerProperty());
+	const int32 EnumValue = EnumProperty->GetEnum()->GetValueByName(FName(ValueName));
+	EnumProperty->GetOwnerProperty()->SetValue_InContainer(InOwnerObject, &EnumValue);
+}
+
+void FAvaRundownOutputEditorUtils::SetStructProperties(void* InOwnerObject,
+	const TSharedPtr<FJsonObject>& InPropertyObject, TFieldIterator<FProperty> InProperty)
+{
+	const FStructProperty* StructProperty = CastField<FStructProperty>(
+		InProperty->GetOwnerProperty());
+
+	const UScriptStruct* Struct = StructProperty->Struct;
+	void* StructPtr = StructProperty->ContainerPtrToValuePtr<void>(InOwnerObject);
+
+	for (TArray<TSharedPtr<FJsonValue>> StructPropertyValues = InPropertyObject->GetArrayField(TEXT("StructProperties"));
+			 TSharedPtr<FJsonValue> StructPropertyValue : StructPropertyValues)
+	{
+		const TSharedPtr<FJsonObject> StructPropertyObject = StructPropertyValue.Get()->AsObject();
+
+		for (TFieldIterator<FProperty> StructIterator(Struct); StructIterator; ++StructIterator)
+		{
+			if (StructPropertyObject->GetStringField(TEXT("Name")) == StructIterator->GetName())
+			{
+				SetProperty(StructPtr, StructPropertyObject, StructIterator);
+			}
+		}
+	}
+}
