Motion Design SVG:
- Misc fixes on exported funcs

diff --git a/Engine/Plugins/Experimental/Avalanche/Source/AvalancheSVGEditor/AvalancheSVGEditor.Build.cs b/Engine/Plugins/Experimental/Avalanche/Source/AvalancheSVGEditor/AvalancheSVGEditor.Build.cs
--- a/Engine/Plugins/Experimental/Avalanche/Source/AvalancheSVGEditor/AvalancheSVGEditor.Build.cs
+++ b/Engine/Plugins/Experimental/Avalanche/Source/AvalancheSVGEditor/AvalancheSVGEditor.Build.cs
@@ -8,19 +8,13 @@ public class AvalancheSVGEditor : ModuleRules
     {
         PCHUsage = ModuleRules.PCHUsageMode.UseExplicitOrSharedPCHs;
 
-        PublicDependencyModuleNames.AddRange(
-            new string[]
-            {
-                "Core",
-            }
-        );
-
         PrivateDependencyModuleNames.AddRange(
             new string[]
             {
 	            "ActorModifierCore",
 	            "AvalancheInteractiveTools",
 	            "AvalancheModifiers",
+                "Core",
 	            "CoreUObject",
                 "Engine",
                 "InteractiveToolsFramework",
diff --git a/Engine/Plugins/Experimental/Avalanche/Source/AvalancheSVGEditor/Private/AvaSVGEditorModule.cpp b/Engine/Plugins/Experimental/Avalanche/Source/AvalancheSVGEditor/Private/AvaSVGEditorModule.cpp
new file mode 100644
--- /dev/null
+++ b/Engine/Plugins/Experimental/Avalanche/Source/AvalancheSVGEditor/Private/AvaSVGEditorModule.cpp
@@ -0,0 +1,133 @@
+// Copyright Epic Games, Inc. All Rights Reserved.
+
+#include "AvaSVGEditorModule.h"
+#include "AvaInteractiveToolsDelegates.h"
+#include "IAvalancheInteractiveToolsModule.h"
+#include "Modifiers/ActorModifierCoreStack.h"
+#include "Modifiers/AvaBevelModifier.h"
+#include "Modifiers/AvaExtrudeModifier.h"
+#include "ProceduralMeshes/SVGDynamicMeshComponent.h"
+#include "ProceduralMeshes/SVGStrokeComponent.h"
+#include "SVGEngineSubsystem.h"
+#include "SVGShapesParentActor.h"
+#include "Subsystems/ActorModifierCoreSubsystem.h"
+#include "Tool/AvaSVGActorTool.h"
+
+#define LOCTEXT_NAMESPACE "AvalancheSVGEditorModule"
+
+void FAvaSVGEditorModule::StartupModule()
+{
+	FAvaInteractiveToolsDelegates::GetRegisterToolsDelegate().AddRaw(this, &FAvaSVGEditorModule::RegisterTools);
+
+	USVGEngineSubsystem::OnSVGActorSplit().BindRaw(this, &FAvaSVGEditorModule::OnSVGActorSplit);
+}
+
+void FAvaSVGEditorModule::ShutdownModule()
+{
+	FAvaInteractiveToolsDelegates::GetRegisterToolsDelegate().RemoveAll(this);
+
+	USVGEngineSubsystem::OnSVGActorSplit().Unbind();
+}
+
+void FAvaSVGEditorModule::RegisterTools(IAvalancheInteractiveToolsModule* InModule)
+{
+	InModule->RegisterTool(
+		IAvalancheInteractiveToolsModule::Get().CategoryNameActor,
+		GetDefault<UAvaSVGActorTool>()->GetToolParameters()
+	);
+}
+
+void FAvaSVGEditorModule::OnSVGActorSplit(ASVGShapesParentActor* InSVGShapesParent)
+{
+	if (!InSVGShapesParent)
+	{
+		return;
+	}
+
+	const UActorModifierCoreSubsystem* ModifierCoreSubsystem = UActorModifierCoreSubsystem::Get();
+
+	if (!ModifierCoreSubsystem)
+	{
+		return;
+	}
+
+	TMap<TObjectPtr<AActor>, TObjectPtr<USVGDynamicMeshComponent>> ShapesMap;
+	InSVGShapesParent->GetShapes(ShapesMap);
+	for (const TPair<TObjectPtr<AActor>, TObjectPtr<USVGDynamicMeshComponent>>& ShapePair : ShapesMap)
+	{
+		if (USVGDynamicMeshComponent* Shape = ShapePair.Value)
+		{
+			const float ExtrudeValue = Shape->GetExtrudeDepth();
+			ESVGExtrudeType ExtrudeType = Shape->ExtrudeType;
+			Shape->FlattenShape();
+
+			if (AActor* ShapeActor = ShapePair.Key)
+			{
+				UActorModifierCoreStack* ModifierStack = ModifierCoreSubsystem->GetActorModifierStack(ShapeActor);
+				if (!ModifierStack)
+				{
+					ModifierStack = ModifierCoreSubsystem->AddActorModifierStack(ShapeActor);
+				}
+
+				FActorModifierCoreStackInsertOp ExtrudeModifierInsertOp;
+				ExtrudeModifierInsertOp.NewModifierName = GetDefault<UAvaExtrudeModifier>()->GetModifierName();
+
+				if (!ModifierCoreSubsystem->GetAllowedModifiers(ShapeActor).Contains(ExtrudeModifierInsertOp.NewModifierName))
+				{
+					continue;
+				}
+
+				UActorModifierCoreBase* Modifier = nullptr;
+
+				Modifier = ModifierStack->InsertModifier(ExtrudeModifierInsertOp);
+				if (UAvaExtrudeModifier* ExtrudeModifier = Cast<UAvaExtrudeModifier>(Modifier))
+				{
+					EAvaExtrudeMode ExtrudeMode = EAvaExtrudeMode::Symmetrical;
+
+					switch (ExtrudeType)
+					{
+						case ESVGExtrudeType::FrontFaceOnly:
+							ExtrudeMode = EAvaExtrudeMode::Front;
+							break;
+
+						case ESVGExtrudeType::None:
+						case ESVGExtrudeType::FrontBackMirror:
+							ExtrudeMode = EAvaExtrudeMode::Symmetrical;
+							break;
+
+						default: ;
+					}
+
+					ExtrudeModifier->SetDepth(ExtrudeValue);
+					ExtrudeModifier->SetExtrudeMode(ExtrudeMode);
+
+					if (Shape->IsA<USVGStrokeComponent>())
+					{
+						ExtrudeModifier->SetCloseBack(false);
+					}
+				}
+
+				if (Shape->Bevel > 0.0f)
+				{
+					FActorModifierCoreStackInsertOp BevelModifierInsertOp;
+					BevelModifierInsertOp.NewModifierName = GetDefault<UAvaBevelModifier>()->GetModifierName();
+
+					if (!ModifierCoreSubsystem->GetAllowedModifiers(ShapeActor).Contains(BevelModifierInsertOp.NewModifierName))
+					{
+						continue;
+					}
+
+					Modifier = ModifierStack->InsertModifier(BevelModifierInsertOp);
+					if (UAvaBevelModifier* BevelModifier = Cast<UAvaBevelModifier>(Modifier))
+					{
+						BevelModifier->SetInset(Shape->Bevel);
+					}
+				}
+			}
+		}
+	}
+}
+
+#undef LOCTEXT_NAMESPACE
+
+IMPLEMENT_MODULE(FAvaSVGEditorModule, AvalancheSVGEditorModule)
diff --git a/Engine/Plugins/Experimental/Avalanche/Source/AvalancheSVGEditor/Private/AvaSVGEditorModule.h b/Engine/Plugins/Experimental/Avalanche/Source/AvalancheSVGEditor/Private/AvaSVGEditorModule.h
new file mode 100644
--- /dev/null
+++ b/Engine/Plugins/Experimental/Avalanche/Source/AvalancheSVGEditor/Private/AvaSVGEditorModule.h
@@ -0,0 +1,23 @@
+// Copyright Epic Games, Inc. All Rights Reserved.
+
+#pragma once
+
+#include "Containers/ContainersFwd.h"
+#include "Modules/ModuleManager.h"
+
+class ASVGShapesParentActor;
+class IAvalancheInteractiveToolsModule;
+class UEdMode;
+struct FAvaInteractiveToolsToolParameters;
+
+class FAvaSVGEditorModule : public IModuleInterface
+{
+public:
+    virtual void StartupModule() override;
+    virtual void ShutdownModule() override;
+
+private:
+	void RegisterTools(IAvalancheInteractiveToolsModule* InModule);
+
+	void OnSVGActorSplit(ASVGShapesParentActor* InSVGShapesParent);
+};
diff --git a/Engine/Plugins/Experimental/Avalanche/Source/AvalancheSVGEditor/Private/AvalancheSVGEditor.cpp b/Engine/Plugins/Experimental/Avalanche/Source/AvalancheSVGEditor/Private/AvalancheSVGEditor.cpp
deleted file mode 100644
diff --git a/Engine/Plugins/Experimental/Avalanche/Source/AvalancheSVGEditor/Public/AvalancheSVGEditor.h b/Engine/Plugins/Experimental/Avalanche/Source/AvalancheSVGEditor/Public/AvalancheSVGEditor.h
deleted file mode 100644
diff --git a/Engine/Plugins/Experimental/SVGImporter/Source/SVGImporterEditor/Private/SVGImporterEditorModule.cpp b/Engine/Plugins/Experimental/SVGImporter/Source/SVGImporterEditor/Private/SVGImporterEditorModule.cpp
--- a/Engine/Plugins/Experimental/SVGImporter/Source/SVGImporterEditor/Private/SVGImporterEditorModule.cpp
+++ b/Engine/Plugins/Experimental/SVGImporter/Source/SVGImporterEditor/Private/SVGImporterEditorModule.cpp
@@ -22,6 +22,7 @@
 #include "SVGData.h"
 #include "SVGDynamicMeshVisualizer.h"
 #include "SVGImporter.h"
+#include "SVGImporterEditorCommands.h"
 #include "SVGImporterEditorUtils.h"
 #include "SVGParsingUtils.h"
 #include "SVGThumbnailRenderer.h"
@@ -53,6 +54,8 @@ FName FSVGImporterEditorModule::SVGImporterCategoryName = TEXT("SVGImporter");;
 
 void FSVGImporterEditorModule::StartupModule()
 {
+	FSVGImporterEditorCommands::Register();
+
 	IAssetTools& AssetTools = FAssetToolsModule::GetModule().Get();
 	AssetTools.RegisterAdvancedAssetCategory(GetSVGImporterMenuCategoryName(), LOCTEXT("SVGImporterCategoryName", "SVGImporter"));
 
@@ -83,6 +86,8 @@ void FSVGImporterEditorModule::StartupModule()
 
 void FSVGImporterEditorModule::ShutdownModule()
 {
+	FSVGImporterEditorCommands::Unregister();
+
 	if (FSVGImporterModule::Get().OnDefaultSVGDataRequested.IsBoundToObject(this))
 	{
 		FSVGImporterModule::Get().OnDefaultSVGDataRequested.Unbind();
diff --git a/Engine/Plugins/Experimental/SVGImporter/Source/SVGImporterEditor/Public/ISVGImporterEditorModule.h b/Engine/Plugins/Experimental/SVGImporter/Source/SVGImporterEditor/Public/ISVGImporterEditorModule.h
--- a/Engine/Plugins/Experimental/SVGImporter/Source/SVGImporterEditor/Public/ISVGImporterEditorModule.h
+++ b/Engine/Plugins/Experimental/SVGImporter/Source/SVGImporterEditor/Public/ISVGImporterEditorModule.h
@@ -9,7 +9,7 @@ class AActor;
 class ULevel;
 class UToolMenu;
 
-class SVGIMPORTEREDITOR_API ISVGImporterEditorModule : public IModuleInterface
+class ISVGImporterEditorModule : public IModuleInterface
 {
 	static constexpr const TCHAR* ModuleName = TEXT("SVGImporterEditor");
 
diff --git a/Engine/Plugins/Experimental/SVGImporter/Source/SVGImporterEditor/SVGImporterEditor.Build.cs b/Engine/Plugins/Experimental/SVGImporter/Source/SVGImporterEditor/SVGImporterEditor.Build.cs
--- a/Engine/Plugins/Experimental/SVGImporter/Source/SVGImporterEditor/SVGImporterEditor.Build.cs
+++ b/Engine/Plugins/Experimental/SVGImporter/Source/SVGImporterEditor/SVGImporterEditor.Build.cs
@@ -15,7 +15,6 @@ public class SVGImporterEditor : ModuleRules
 			new string[]
 			{
 				"Core",
-				"Engine",
 			}
 		);
 
