[MotionDesign] Rundown: Fix preview when no preview channel is specified.

When previewing a mix of combo pages and single page, it may happen that a page is only partially visible, i.e. some of the sub-instance have been kicked out by TL of the single page. In that case, it would happen that the preview render target gets lost because it only looks at the first instance of a combo page (which got kicked out).

The fix is to search in all the instance players, and even in all the playing pages until the playable group is found and the render target can be retrieved.

#rnx

diff --git a/Engine/Plugins/Experimental/Avalanche/Source/AvalancheMedia/Private/Rundown/AvaRundown.cpp b/Engine/Plugins/Experimental/Avalanche/Source/AvalancheMedia/Private/Rundown/AvaRundown.cpp
--- a/Engine/Plugins/Experimental/Avalanche/Source/AvalancheMedia/Private/Rundown/AvaRundown.cpp
+++ b/Engine/Plugins/Experimental/Avalanche/Source/AvalancheMedia/Private/Rundown/AvaRundown.cpp
@@ -1377,13 +1377,27 @@ int32 UAvaRundown::RemovePagesFromSubList(int32 InSubListIndex, const TArray<int
 
 namespace UE::AvaMedia::Rundown::Private
 {
-	const UAvaRundownPagePlayer* FindPagePlayerForPreviewChannel(const FName& InPreviewChannel, const TArray<TObjectPtr<UAvaRundownPagePlayer>>& InPagePlayers)
+	const UAvaPlayableGroup* FindPlayableGroup(const UAvaRundownPagePlayer* InPagePlayer)
 	{
-		const TObjectPtr<UAvaRundownPagePlayer>* PreviewingPagePlayerPtr = InPagePlayers.FindByPredicate([InPreviewChannel](const UAvaRundownPagePlayer* InPagePlayer)
+		if (!InPagePlayer)
 		{
-			return InPagePlayer->bIsPreview && InPagePlayer->ChannelName == InPreviewChannel;
-		});
-		return PreviewingPagePlayerPtr ? *PreviewingPagePlayerPtr : nullptr;
+			return nullptr;
+		}
+
+		for (const TObjectPtr<UAvaRundownPlaybackInstancePlayer>& InstancePlayer : InPagePlayer->InstancePlayers)
+		{
+			if (InstancePlayer->IsPlaying())
+			{
+				if (const UAvaPlayable* const Playable = InstancePlayer->Playback->GetFirstPlayable())
+				{
+					if (const UAvaPlayableGroup* const PlayableGroup = Playable->GetPlayableGroup())
+					{
+						return PlayableGroup;
+					}
+				}
+			}
+		}
+		return nullptr;
 	}
 }
 
@@ -1401,25 +1415,19 @@ UTextureRenderTarget2D* UAvaRundown::GetPreviewRenderTarget(const FName& InPrevi
 
 	// If there is no channel, we can get the render target from the playable group of a previewing page's playable
 	// in the given channel. When playable group composition is implemented, this may have to change. 
-	
-	using namespace UE::AvaMedia::Rundown::Private;
-	const UAvaRundownPagePlayer* PreviewingPagePlayer = FindPagePlayerForPreviewChannel(InPreviewChannel, PagePlayers);
 
-	if (!PreviewingPagePlayer)
-	{
-		return nullptr;
-	}
-	
-	if (const UAvaPlaybackGraph* Playback = PreviewingPagePlayer->GetPlayback())
+	using namespace UE::AvaMedia::Rundown::Private;
+	for (const TObjectPtr<UAvaRundownPagePlayer>& PagePlayer : PagePlayers)
 	{
-		if (const UAvaPlayable* const Playable = Playback->GetFirstPlayable())
+		if (PagePlayer->bIsPreview && PagePlayer->ChannelName == InPreviewChannel)
 		{
-			if (const UAvaPlayableGroup* const PlayableGroup = Playable->GetPlayableGroup())
+			if (const UAvaPlayableGroup* PlayableGroup = FindPlayableGroup(PagePlayer))
 			{
-				return PlayableGroup->IsRenderTargetReady() ? PlayableGroup->GetRenderTarget() : nullptr;
+				return PlayableGroup->IsRenderTargetReady() ? PlayableGroup->GetRenderTarget() : nullptr;	
 			}
 		}
 	}
+	
 	return nullptr;
 }
 
