Avalanche:
- Removed Avalanche Blueprint and code depending on it

diff --git a/Engine/Plugins/Experimental/Avalanche/Avalanche.uplugin b/Engine/Plugins/Experimental/Avalanche/Avalanche.uplugin
--- a/Engine/Plugins/Experimental/Avalanche/Avalanche.uplugin
+++ b/Engine/Plugins/Experimental/Avalanche/Avalanche.uplugin
@@ -164,11 +164,6 @@
 			"Type": "Runtime",
 			"LoadingPhase": "Default"
 		},
-		{
-			"Name": "AvalancheTest",
-			"Type": "UncookedOnly",
-			"LoadingPhase": "Default"
-		},
 		{
 			"Name": "AvalancheTransition",
 			"Type": "Runtime",
diff --git a/Engine/Plugins/Experimental/Avalanche/Config/DefaultAvalanche.ini b/Engine/Plugins/Experimental/Avalanche/Config/DefaultAvalanche.ini
--- a/Engine/Plugins/Experimental/Avalanche/Config/DefaultAvalanche.ini
+++ b/Engine/Plugins/Experimental/Avalanche/Config/DefaultAvalanche.ini
@@ -149,7 +149,6 @@
 +PropertyRedirects=(OldName="/Script/AvalancheEffectors.AvaClonerActor.StepScaleDelta",NewName="/Script/AvalancheEffectors.AvaClonerActor.DeltaStepScale")
 +PropertyRedirects=(OldName="/Script/AvalancheEffectors.AvaEffectorActor.Strength",NewName="/Script/AvalancheEffectors.AvaEffectorActor.LocationStrength")
 
-+ClassRedirects=(OldName="/Script/Avalanche.AvalancheActor",NewName="/Script/Avalanche.AvaActor")
 +StructRedirects=(OldName="/Script/Avalanche.AvalancheFileVersionInfo",NewName="/Script/Avalanche.AvaFileVersionInfo")
 +StructRedirects=(OldName="/Script/Avalanche.AvalancheEngineVersionInfo",NewName="/Script/Avalanche.AvaEngineVersionInfo")
 +StructRedirects=(OldName="/Script/Avalanche.AvalancheCustomVersionInfo",NewName="/Script/Avalanche.AvaCustomVersionInfo")
@@ -170,14 +169,12 @@
 +ClassRedirects=(OldName="/Script/AvalancheEditor.AvalancheEditorSettings",NewName="/Script/AvalancheEditor.AvaEditorSettings")
 +ClassRedirects=(OldName="/Script/AvalancheMask.AvalancheMaskSettings",NewName="/Script/AvalancheMask.AvaMaskSettings")
 +ClassRedirects=(OldName="/Script/AvalancheMaskEditor.AvalancheMaskEditorMode",NewName="/Script/AvalancheMaskEditor.AvaMaskEditorMode")
-+FunctionRedirects=(OldName="/Script/AvalancheTest.AvaTestBlueprintFunctionLibrary.ExportAvaBlueprintsToWorld",NewName="/Script/AvalancheTest.AvaTestBlueprintFunctionLibrary.ExportMotionDesignBlueprintsToWorld")
 
 +StructRedirects=(OldName="/Script/Avalanche.AvaSceneContainsAttributeConditionInstanceData",NewName="/Script/Avalanche.AvaSceneContainsTagAttributeConditionInstanceData")
 
 +StructRedirects=(OldName="/Script/AvalancheMedia.AnimPlaySettings",NewName="/Script/AvalancheMedia.AvaPlaybackAnimPlaySettings")
 +StructRedirects=(OldName="/Script/AvalancheMedia.AvalancheAnimations",NewName="/Script/AvalancheMedia.AvaPlaybackAnimations")
 +ClassRedirects=(OldName="/Script/AvalancheMedia.AvalanchePlayable",NewName="/Script/AvalancheMedia.AvaPlayable")
-+ClassRedirects=(OldName="/Script/AvalancheMedia.AvalancheBlueprintPlayable",NewName="/Script/AvalancheMedia.AvaPlayableBlueprint")
 +ClassRedirects=(OldName="/Script/AvalancheMedia.AvalancheBroadcast",NewName="/Script/AvalancheMedia.AvaBroadcast")
 +FunctionRedirects=(OldName="/Script/AvalancheMedia.AvaBroadcast.GetAvalancheBroadcast",NewName="/Script/AvalancheMedia.AvaBroadcast.GetBroadcast")
 +ClassRedirects=(OldName="/Script/AvalancheMedia.AvalancheBroadcastComponent",NewName="/Script/AvalancheMedia.AvaBroadcastComponent")
diff --git a/Engine/Plugins/Experimental/Avalanche/Source/Avalanche/Private/Archive/AvaArchive.cpp b/Engine/Plugins/Experimental/Avalanche/Source/Avalanche/Private/Archive/AvaArchive.cpp
deleted file mode 100644
diff --git a/Engine/Plugins/Experimental/Avalanche/Source/Avalanche/Private/Archive/AvaArchive.h b/Engine/Plugins/Experimental/Avalanche/Source/Avalanche/Private/Archive/AvaArchive.h
deleted file mode 100644
diff --git a/Engine/Plugins/Experimental/Avalanche/Source/Avalanche/Private/Archive/AvaReader.cpp b/Engine/Plugins/Experimental/Avalanche/Source/Avalanche/Private/Archive/AvaReader.cpp
deleted file mode 100644
diff --git a/Engine/Plugins/Experimental/Avalanche/Source/Avalanche/Private/Archive/AvaReader.h b/Engine/Plugins/Experimental/Avalanche/Source/Avalanche/Private/Archive/AvaReader.h
deleted file mode 100644
diff --git a/Engine/Plugins/Experimental/Avalanche/Source/Avalanche/Private/Archive/AvaWriter.cpp b/Engine/Plugins/Experimental/Avalanche/Source/Avalanche/Private/Archive/AvaWriter.cpp
deleted file mode 100644
diff --git a/Engine/Plugins/Experimental/Avalanche/Source/Avalanche/Private/Archive/AvaWriter.h b/Engine/Plugins/Experimental/Avalanche/Source/Avalanche/Private/Archive/AvaWriter.h
deleted file mode 100644
diff --git a/Engine/Plugins/Experimental/Avalanche/Source/Avalanche/Private/AvaActor.cpp b/Engine/Plugins/Experimental/Avalanche/Source/Avalanche/Private/AvaActor.cpp
deleted file mode 100644
diff --git a/Engine/Plugins/Experimental/Avalanche/Source/Avalanche/Private/AvaActorUtils.cpp b/Engine/Plugins/Experimental/Avalanche/Source/Avalanche/Private/AvaActorUtils.cpp
--- a/Engine/Plugins/Experimental/Avalanche/Source/Avalanche/Private/AvaActorUtils.cpp
+++ b/Engine/Plugins/Experimental/Avalanche/Source/Avalanche/Private/AvaActorUtils.cpp
@@ -1,15 +1,17 @@
 // Copyright Epic Games, Inc. All Rights Reserved.
 
 #include "AvaActorUtils.h"
-#include "AvaBlueprint.h"
 #include "AvaDefs.h"
+#include "AvaSceneItem.h"
 #include "AvaSceneSubsystem.h"
+#include "AvaSceneTree.h"
 #include "Components/PrimitiveComponent.h"
 #include "Components/SceneComponent.h"
 #include "Engine/World.h"
 #include "EngineUtils.h"
 #include "Framework/AvaGameInstance.h"
 #include "GameFramework/Actor.h"
+#include "IAvaSceneInterface.h"
 #include "Math/MathFwd.h"
 #include "Math/OrientedBox.h"
 
diff --git a/Engine/Plugins/Experimental/Avalanche/Source/Avalanche/Private/AvaBlueprint.cpp b/Engine/Plugins/Experimental/Avalanche/Source/Avalanche/Private/AvaBlueprint.cpp
deleted file mode 100644
diff --git a/Engine/Plugins/Experimental/Avalanche/Source/Avalanche/Private/AvaBlueprintGeneratedClass.cpp b/Engine/Plugins/Experimental/Avalanche/Source/Avalanche/Private/AvaBlueprintGeneratedClass.cpp
deleted file mode 100644
diff --git a/Engine/Plugins/Experimental/Avalanche/Source/Avalanche/Private/AvaBlueprintVersion.cpp b/Engine/Plugins/Experimental/Avalanche/Source/Avalanche/Private/AvaBlueprintVersion.cpp
deleted file mode 100644
diff --git a/Engine/Plugins/Experimental/Avalanche/Source/Avalanche/Private/AvaBlueprintVersion.h b/Engine/Plugins/Experimental/Avalanche/Source/Avalanche/Private/AvaBlueprintVersion.h
deleted file mode 100644
diff --git a/Engine/Plugins/Experimental/Avalanche/Source/Avalanche/Private/AvaBlueprint_Serialize.cpp b/Engine/Plugins/Experimental/Avalanche/Source/Avalanche/Private/AvaBlueprint_Serialize.cpp
deleted file mode 100644
diff --git a/Engine/Plugins/Experimental/Avalanche/Source/Avalanche/Private/AvaBlueprint_Serialize.h b/Engine/Plugins/Experimental/Avalanche/Source/Avalanche/Private/AvaBlueprint_Serialize.h
deleted file mode 100644
diff --git a/Engine/Plugins/Experimental/Avalanche/Source/Avalanche/Private/AvaScene.cpp b/Engine/Plugins/Experimental/Avalanche/Source/Avalanche/Private/AvaScene.cpp
--- a/Engine/Plugins/Experimental/Avalanche/Source/Avalanche/Private/AvaScene.cpp
+++ b/Engine/Plugins/Experimental/Avalanche/Source/Avalanche/Private/AvaScene.cpp
@@ -16,7 +16,6 @@
 #include "UObject/UObjectThreadContext.h"
 
 #if WITH_EDITOR
-#include "AvaBlueprint.h"
 #include "AvaField.h"
 #include "Misc/ScopedSlowTask.h"
 #include "RemoteControlBinding.h"
@@ -57,139 +56,6 @@ AAvaScene* AAvaScene::GetScene(ULevel* InLevel, bool bInCreateSceneIfNotFound)
 	return NewScene;
 }
 
-#if WITH_EDITOR
-AAvaScene* AAvaScene::CreateFromBlueprint(ULevel* InNewLevel, UAvalancheBlueprint* InBlueprint)
-{
-	if (!IsValid(InNewLevel) || !IsValid(InNewLevel->OwningWorld) || !IsValid(InBlueprint))
-	{
-		return nullptr;
-	}
-
-	FScopedSlowTask SlowTask(100, LOCTEXT("ExportingBlueprintToWorld", "Exporting Blueprint to World"));
-
-	// Load Blueprint World Actors into the new Level 
-	SlowTask.EnterProgressFrame(80, LOCTEXT("CopyingActors", "Copying Blueprint Actors"));
-	{
-		auto InitActor = [](AActor& InActor, const FAvaActorData& InActorData)
-		{
-			InActor.bHiddenEd = InActorData.bEditorVisibility;
-		};
-
-		InBlueprint->LoadAvalancheWorld(InNewLevel->OwningWorld, UAvalancheBlueprint::FInitActorFunctionRef(InitActor));
-	}
-
-	// Create Scene Actor in Level
-	AAvaScene* Scene;
-	SlowTask.EnterProgressFrame(5, LOCTEXT("CreatingScene", "Creating Level Scene"));
-	{
-		Scene = GetScene(InNewLevel, /*bInCreateSceneIfNotFound=*/true);
-		check(Scene);
-
-		Scene->SceneTree         = InBlueprint->GetSceneTree();
-		Scene->OutlinerData      = InBlueprint->GetEditorData();
-		Scene->StartupCameraName = InBlueprint->GetStartupCameraName();
-	}
-
-	// Copy Sequences
-	SlowTask.EnterProgressFrame(10, LOCTEXT("CopyingSequences", "Copying Blueprint Sequences"));
-	{
-		// Copy to remove all invalid sequences
-		TArray<UAvaSequence*> SourceSequences = InBlueprint->GetSequences();
-		SourceSequences.RemoveAll([](UAvaSequence* InSequence) { return !IsValid(InSequence); });
-
-		// Expectation is for this new scene to have no sequences
-		ensureAlways(Scene->Animations.IsEmpty());
-		Scene->Animations.Empty(SourceSequences.Num());
-		Scene->DefaultSequenceIndex = SourceSequences.Find(InBlueprint->GetDefaultSequence());
-
-		const FTopLevelAssetPath OldContext = UAvalancheBlueprint::GetWorldContextPath(InBlueprint);
-		const FTopLevelAssetPath NewContext(InNewLevel->OwningWorld);
-
-		for (UAvaSequence* const SourceSequence : SourceSequences)
-		{
-			UAvaSequence* const NewSequence = DuplicateObject<UAvaSequence>(SourceSequence, Scene);
-			ensure(IsValid(NewSequence));
-
-			NewSequence->MigrateLegacyBindings(/*PlaybackContext*/InNewLevel->OwningWorld);
-
-			NewSequence->UpdateBindings(&OldContext, NewContext);
-
-			// Instead of calling AAvaScene::AddSequence which has some check overhead and schedules a Tree Rebuild, handle here manually
-			Scene->Animations.Add(NewSequence);
-		}
-
-		if (!Scene->Animations.IsValidIndex(Scene->DefaultSequenceIndex))
-		{
-			Scene->DefaultSequenceIndex = 0;
-		}
-
-		check(SourceSequences.Num() == Scene->Animations.Num());
-
-		for (int32 Index = 0; Index < SourceSequences.Num(); ++Index)
-		{
-			UAvaSequence* const SourceSequence = SourceSequences[Index];
-			UAvaSequence* const TargetSequence = Scene->Animations[Index];
-
-			// On Duplicate, the Child Anim property is marked as Duplicate Transient so should've not duplicated over
-			ensure(TargetSequence->GetChildren().IsEmpty());
-
-			for (const TWeakObjectPtr<UAvaSequence>& SourceChildren : SourceSequence->GetChildren())
-			{
-				const int32 ChildIndex = SourceSequences.Find(SourceChildren.Get());
-				if (ChildIndex != INDEX_NONE)
-				{
-					TargetSequence->AddChild(Scene->Animations[ChildIndex]);	
-				}
-			}
-		}
-
-		Scene->RebuildSequenceTree();
-	}
-
-	// Copy Remote Control Preset
-	SlowTask.EnterProgressFrame(5, LOCTEXT("CopyingRemoteControlPreset", "Copying Blueprint Remote Control Preset"));
-	{
-		URemoteControlPreset* const SourcePreset = InBlueprint->GetRemoteControlPreset();
-		if (IsValid(SourcePreset))
-		{
-			const EObjectFlags Flags   = IsValid(Scene->RemoteControlPreset) ? Scene->RemoteControlPreset->GetFlags() : SourcePreset->GetFlags();
-			Scene->RemoteControlPreset = DuplicateObject<URemoteControlPreset>(SourcePreset, Scene);
-			Scene->RemoteControlPreset->SetFlags(Flags);
-
-			FMapProperty* LevelMapProperty  = UE::AvaCore::GetProperty<URemoteControlLevelDependantBinding, FMapProperty>(TEXT("SubLevelSelectionMapByPath"));
-			FMapProperty* ObjectMapProperty = UE::AvaCore::GetProperty<URemoteControlLevelDependantBinding, FMapProperty>(TEXT("BoundObjectMapByPath"));
-			FProperty* LastLevelProperty    = UE::AvaCore::GetProperty<URemoteControlLevelDependantBinding>(TEXT("LevelWithLastSuccessfulResolve"));
-			FProperty* LastObjectProperty   = UE::AvaCore::GetProperty<URemoteControlBinding>(TEXT("LastBoundObjectPath"));
-
-			FSoftObjectPath NullObjectPath;
-			FSoftObjectPtr NullObjectPtr;
-
-			auto ResetMap = [](UObject* InObject, FMapProperty* InMapProperty)
-			{
-				FScriptMapHelper MapHelper(InMapProperty, InMapProperty->ContainerPtrToValuePtr<void*>(InObject));
-				MapHelper.EmptyValues();
-			};
-
-			for (TArray<TObjectPtr<URemoteControlBinding>>::TIterator BindingIter =  Scene->RemoteControlPreset->Bindings.CreateIterator(); BindingIter; ++BindingIter)
-			{
-				if (URemoteControlLevelDependantBinding* LevelBinding = Cast<URemoteControlLevelDependantBinding>(*BindingIter))
-				{
-					ResetMap(LevelBinding, LevelMapProperty);
-					ResetMap(LevelBinding, ObjectMapProperty);
-
-					LastLevelProperty->SetValue_InContainer(LevelBinding, &NullObjectPtr);
-					LastObjectProperty->SetValue_InContainer(LevelBinding, &NullObjectPath);
-				}
-			}
-
-			Scene->RemoteControlPreset->RebindUnboundEntities();
-		}
-	}
-
-	return Scene;
-}
-#endif
-
 AAvaScene::AAvaScene()
 {
 	SceneSettings = CreateDefaultSubobject<UAvaSceneSettings>(TEXT("SceneSettings"));
diff --git a/Engine/Plugins/Experimental/Avalanche/Source/Avalanche/Private/Data/AvalancheVersionInfo.cpp b/Engine/Plugins/Experimental/Avalanche/Source/Avalanche/Private/Data/AvalancheVersionInfo.cpp
deleted file mode 100644
diff --git a/Engine/Plugins/Experimental/Avalanche/Source/Avalanche/Private/Data/AvalancheWorldData.cpp b/Engine/Plugins/Experimental/Avalanche/Source/Avalanche/Private/Data/AvalancheWorldData.cpp
deleted file mode 100644
diff --git a/Engine/Plugins/Experimental/Avalanche/Source/Avalanche/Private/Framework/AvaGameInstance.cpp b/Engine/Plugins/Experimental/Avalanche/Source/Avalanche/Private/Framework/AvaGameInstance.cpp
--- a/Engine/Plugins/Experimental/Avalanche/Source/Avalanche/Private/Framework/AvaGameInstance.cpp
+++ b/Engine/Plugins/Experimental/Avalanche/Source/Avalanche/Private/Framework/AvaGameInstance.cpp
@@ -2,7 +2,6 @@
 
 #include "Framework/AvaGameInstance.h"
 #include "AudioDevice.h"
-#include "AvaBlueprint.h"
 #include "AvaRemoteControlRebind.h"
 #include "Components/ReflectionCaptureComponent.h"
 #include "Components/SkyLightComponent.h"
@@ -14,6 +13,7 @@
 #include "Misc/TimeGuard.h"
 #include "Slate/SceneViewport.h"
 #include "Viewport/AvaCameraManager.h"
+#include "Viewport/AvaViewportQualitySettings.h"
 
 namespace UE::AvaGameInstance::Private
 {
@@ -38,13 +38,7 @@ namespace UE::AvaGameInstance::Private
 			MarkedIgnoreDeltaSecondsFrame = 0;
 		}
 	}
-	
-	static void MarkSynchronousAssetLoadingThisFrame()
-	{
-		UpdateFrameMarkers();
-		MarkedSynchronousAssetLoadingFrame = GFrameCounter;
-	}
-	
+
 	static bool ShouldIgnoreDeltaSecondsForCurrentFrame()
 	{
 		UpdateFrameMarkers();
@@ -55,14 +49,6 @@ namespace UE::AvaGameInstance::Private
 UAvaGameInstance::FOnAvaGameInstanceEvent UAvaGameInstance::OnEndPlay;
 UAvaGameInstance::FOnAvaGameInstanceEvent UAvaGameInstance::OnRenderTargetReady;
 
-UAvaGameInstance* UAvaGameInstance::Create(UObject* InOuter, const TSoftObjectPtr<UAvalancheBlueprint>& InBlueprintTemplate)
-{
-	UAvaGameInstance* GameInstance = NewObject<UAvaGameInstance>(InOuter ? InOuter : GEngine);
-	GameInstance->CreateWorld();
-	GameInstance->LoadMotionDesignBlueprint(InBlueprintTemplate);
-	return GameInstance;
-}
-
 UAvaGameInstance* UAvaGameInstance::Create(UObject* InOuter)
 {
 	UAvaGameInstance* GameInstance = NewObject<UAvaGameInstance>(InOuter ? InOuter : GEngine);
@@ -126,55 +112,6 @@ bool UAvaGameInstance::CreateWorld()
 	return true;
 }
 
-bool UAvaGameInstance::LoadMotionDesignBlueprint(const TSoftObjectPtr<UAvalancheBlueprint>& InBlueprint)
-{
-	if (ManagedBlueprint && SourceBlueprint == InBlueprint)
-	{
-		return false;
-	}
-
-	if (!bWorldCreated)
-	{
-		CreateWorld();
-	}
-
-	UAvalancheBlueprint* Source;
-	{
-		TRACE_CPUPROFILER_EVENT_SCOPE(UAvaGameInstance::LoadAsset::LoadSourceAvaBp);
-		Source = InBlueprint.LoadSynchronous();
-		check(Source);
-	}
-
-#if WITH_EDITORONLY_DATA
-	FGuardValue_Bitfield(Source->bDuplicatingReadOnly, true);
-#endif
-
-	{
-		TRACE_CPUPROFILER_EVENT_SCOPE(UAvaGameInstance::LoadAsset::DupAvaBp);
-		const FName AvaInstanceName = MakeUniqueObjectName(this, UAvalancheBlueprint::StaticClass(), Source->GetFName());
-		ManagedBlueprint = Cast<UAvalancheBlueprint>(StaticDuplicateObject(Source, this, AvaInstanceName, RF_Transient));
-	}
-	
-	{
-		TRACE_CPUPROFILER_EVENT_SCOPE(UAvaGameInstance::LoadAsset::LoadAvaBpWorld);
-		ManagedBlueprint->SetAvalancheWorld(PlayWorld.Get());
-		ManagedBlueprint->LoadAvalancheWorld();
-	}
-	
-	ManagedBlueprint->UpdatePlaceholderActor();
-
-	{
-		TRACE_CPUPROFILER_EVENT_SCOPE(UAvaGameInstance::LoadAsset::RebindRCP);
-		ManagedBlueprint->RegisterRemoteControlPreset();
-		FAvaRemoteControlRebind::RebindUnboundEntities(ManagedBlueprint->GetRemoteControlPreset());
-	}
-	
-	MarkSynchronousAssetLoadingThisFrame();
-
-	SourceBlueprint = InBlueprint;
-	return true;
-}
-
 bool UAvaGameInstance::BeginPlayWorld(const FAvaInstancePlaySettings& InWorldPlaySettings)
 {
 	// Make sure we don't have pending unload or stop requests left over in the game instance. 
@@ -234,42 +171,6 @@ void UAvaGameInstance::RequestEndPlayWorld(bool bForceImmediate)
 	}
 }
 
-void UAvaGameInstance::BeginPlayMotionDesignBlueprint(const TSoftObjectPtr<UAvalancheBlueprint>& InBlueprint, const FAvaInstancePlaySettings& InWorldPlaySettings)
-{
-	// World States: Created -> Playing
-	// Asset States: Unloaded -> Loading -> Loaded -> Playing (Visible)
-	const FSoftObjectPath& SourceAssetPath = InBlueprint.ToSoftObjectPath();
-
-	if (!bWorldPlaying)
-	{
-		BeginPlayWorld(InWorldPlaySettings);
-	}
-	else
-	{
-		UpdateRenderTarget(InWorldPlaySettings.RenderTarget);
-		UpdateSceneViewportSize(InWorldPlaySettings.ViewportSize);
-	}
-
-	// Check if the asset is loaded.
-	if (SourceBlueprint != InBlueprint)
-	{
-		LoadMotionDesignBlueprint(InBlueprint);
-	}	
-	
-	if (ManagedBlueprint && SourceBlueprint == InBlueprint)
-	{
-		// Verify if this can be changed on the fly.
-		ManagedBlueprint->GetViewportQualitySettings().Apply(ViewportClient->EngineShowFlags);
-	
-		// Old code using ava camera manager. To retire.
-		constexpr bool bIsCanvasController = false;
-		ViewportClient->GetCameraManager()->Init(ManagedBlueprint->GetPlaybackObject(), bIsCanvasController);
-#if WITH_EDITOR
-		ViewportClient->GetCameraManager()->SetDefaultViewTarget(PlayWorld, ManagedBlueprint->GetStartupCameraName());
-#endif
-	}
-}
-
 void UAvaGameInstance::RequestUnloadWorld(bool bForceImmediate)
 {
 	if (bWorldCreated)
@@ -319,11 +220,6 @@ void UAvaGameInstance::UpdateSceneViewportSize(const FIntPoint& InViewportSize)
 	}
 }
 
-void UAvaGameInstance::MarkSynchronousAssetLoadingThisFrame()
-{
-	UE::AvaGameInstance::Private::MarkSynchronousAssetLoadingThisFrame();
-}
-
 void UAvaGameInstance::Tick(float DeltaSeconds)
 {
 	if (bIsTicking)
@@ -418,13 +314,6 @@ void UAvaGameInstance::UnloadWorld()
 	bWorldCreated = false;
 	bRequestUnloadWorld = false;
 
-	if (IsValid(ManagedBlueprint.Get()))
-	{
-		ManagedBlueprint->UnregisterRemoteControlPreset();
-	}
-	SourceBlueprint.Reset();
-	ManagedBlueprint = nullptr;
-	
 	if (PlayWorld)
 	{
 		if (FAudioDeviceHandle AudioDevice = PlayWorld->GetAudioDevice())
diff --git a/Engine/Plugins/Experimental/Avalanche/Source/Avalanche/Private/Framework/AvaSoftAssetPtr.cpp b/Engine/Plugins/Experimental/Avalanche/Source/Avalanche/Private/Framework/AvaSoftAssetPtr.cpp
--- a/Engine/Plugins/Experimental/Avalanche/Source/Avalanche/Private/Framework/AvaSoftAssetPtr.cpp
+++ b/Engine/Plugins/Experimental/Avalanche/Source/Avalanche/Private/Framework/AvaSoftAssetPtr.cpp
@@ -1,19 +1,12 @@
 // Copyright Epic Games, Inc. All Rights Reserved.
 
 #include "Framework/AvaSoftAssetPtr.h"
-#include "AvaBlueprint.h"
+#include "Engine/World.h"
 
 EAvalancheAssetType FAvaSoftAssetPath::GetAssetTypeFromClass(const FSoftClassPath& InAssetClassPath, bool bInLoadIfUnknown)
 {
-	static const FSoftClassPath AvalancheClassPath(UAvalancheBlueprint::StaticClass());
 	static const FSoftClassPath WorldClassPath(UWorld::StaticClass());
 
-	// Try comparing the asset paths directly.
-	// That is probably the fastest way, but can fail for derived types.
-	if (InAssetClassPath == AvalancheClassPath)
-	{
-		return EAvalancheAssetType::Blueprint;
-	}
 	if (InAssetClassPath == WorldClassPath)
 	{
 		return EAvalancheAssetType::World;
@@ -23,10 +16,6 @@ EAvalancheAssetType FAvaSoftAssetPath::GetAssetTypeFromClass(const FSoftClassPat
 	if (InAssetClassPath.IsValid())
 	{
 		const UClass* AssetClass = InAssetClassPath.ResolveClass();
-		if (AssetClass->IsChildOf(UAvalancheBlueprint::StaticClass()))
-		{
-			return EAvalancheAssetType::Blueprint;
-		}
 		if (AssetClass->IsChildOf(UWorld::StaticClass()))
 		{
 			return EAvalancheAssetType::World;
@@ -42,7 +31,7 @@ EAvalancheAssetType FAvaSoftAssetPtr::GetAssetType(bool bInLoadIfUnknown) const
 	{
 		return AssetTypeFromClass;
 	}
-	
+
 	if (bInLoadIfUnknown)
 	{
 		// Todo: problem with loading the asset sync here.
@@ -54,17 +43,12 @@ EAvalancheAssetType FAvaSoftAssetPtr::GetAssetType(bool bInLoadIfUnknown) const
 			const FSoftObjectPtr SourceAvalancheAsset(AssetPtr.ToSoftObjectPath());
 			LoadedSourceAsset = SourceAvalancheAsset.LoadSynchronous();
 		}
-	
-		if (Cast<UAvalancheBlueprint>(LoadedSourceAsset))
-		{
-			return EAvalancheAssetType::Blueprint;
-		}
-	
+
 		if (Cast<UWorld>(LoadedSourceAsset))
 		{
 			return EAvalancheAssetType::World;
 		}
 	}
-	
+
 	return EAvalancheAssetType::Unknown;
 }
\ No newline at end of file
diff --git a/Engine/Plugins/Experimental/Avalanche/Source/Avalanche/Private/Viewport/AvaCameraManager.cpp b/Engine/Plugins/Experimental/Avalanche/Source/Avalanche/Private/Viewport/AvaCameraManager.cpp
--- a/Engine/Plugins/Experimental/Avalanche/Source/Avalanche/Private/Viewport/AvaCameraManager.cpp
+++ b/Engine/Plugins/Experimental/Avalanche/Source/Avalanche/Private/Viewport/AvaCameraManager.cpp
@@ -1,7 +1,6 @@
 // Copyright Epic Games, Inc. All Rights Reserved.
 
 #include "Viewport/AvaCameraManager.h"
-#include "AvaActor.h"
 #include "AvaSequencePlaybackObject.h"
 #include "Camera/CameraComponent.h"
 #include "Camera/CameraTypes.h"
diff --git a/Engine/Plugins/Experimental/Avalanche/Source/Avalanche/Public/AvaActor.h b/Engine/Plugins/Experimental/Avalanche/Source/Avalanche/Public/AvaActor.h
deleted file mode 100644
diff --git a/Engine/Plugins/Experimental/Avalanche/Source/Avalanche/Public/AvaBlueprint.h b/Engine/Plugins/Experimental/Avalanche/Source/Avalanche/Public/AvaBlueprint.h
deleted file mode 100644
diff --git a/Engine/Plugins/Experimental/Avalanche/Source/Avalanche/Public/AvaBlueprintGeneratedClass.h b/Engine/Plugins/Experimental/Avalanche/Source/Avalanche/Public/AvaBlueprintGeneratedClass.h
deleted file mode 100644
diff --git a/Engine/Plugins/Experimental/Avalanche/Source/Avalanche/Public/AvaScene.h b/Engine/Plugins/Experimental/Avalanche/Source/Avalanche/Public/AvaScene.h
--- a/Engine/Plugins/Experimental/Avalanche/Source/Avalanche/Public/AvaScene.h
+++ b/Engine/Plugins/Experimental/Avalanche/Source/Avalanche/Public/AvaScene.h
@@ -25,7 +25,6 @@ class URemoteControlPreset;
 
 #if WITH_EDITOR
 class ISequencer;
-class UAvalancheBlueprint;
 #endif
 
 UCLASS(MinimalAPI, NotPlaceable, Hidden, NotBlueprintable, NotBlueprintType, DisplayName = "Motion Design Scene")
@@ -36,10 +35,6 @@ class AAvaScene : public AActor, public IAvaSequenceProvider, public IAvaSceneIn
 public:
 	AVALANCHE_API static AAvaScene* GetScene(ULevel* InLevel, bool bInCreateSceneIfNotFound);
 
-#if WITH_EDITOR
-	AVALANCHE_API static AAvaScene* CreateFromBlueprint(ULevel* InNewLevel, UAvalancheBlueprint* InBlueprint);
-#endif
-
 	AAvaScene();
 
 	IAvaSequencePlaybackObject* GetScenePlayback() const;
diff --git a/Engine/Plugins/Experimental/Avalanche/Source/Avalanche/Public/Data/AvaActorData.h b/Engine/Plugins/Experimental/Avalanche/Source/Avalanche/Public/Data/AvaActorData.h
deleted file mode 100644
diff --git a/Engine/Plugins/Experimental/Avalanche/Source/Avalanche/Public/Data/AvaComponentData.h b/Engine/Plugins/Experimental/Avalanche/Source/Avalanche/Public/Data/AvaComponentData.h
deleted file mode 100644
diff --git a/Engine/Plugins/Experimental/Avalanche/Source/Avalanche/Public/Data/AvaDataDefines.h b/Engine/Plugins/Experimental/Avalanche/Source/Avalanche/Public/Data/AvaDataDefines.h
deleted file mode 100644
diff --git a/Engine/Plugins/Experimental/Avalanche/Source/Avalanche/Public/Data/AvaObjectData.h b/Engine/Plugins/Experimental/Avalanche/Source/Avalanche/Public/Data/AvaObjectData.h
deleted file mode 100644
diff --git a/Engine/Plugins/Experimental/Avalanche/Source/Avalanche/Public/Data/AvaSubObjectData.h b/Engine/Plugins/Experimental/Avalanche/Source/Avalanche/Public/Data/AvaSubObjectData.h
deleted file mode 100644
diff --git a/Engine/Plugins/Experimental/Avalanche/Source/Avalanche/Public/Data/AvaVersionInfo.h b/Engine/Plugins/Experimental/Avalanche/Source/Avalanche/Public/Data/AvaVersionInfo.h
deleted file mode 100644
diff --git a/Engine/Plugins/Experimental/Avalanche/Source/Avalanche/Public/Data/AvaWorldData.h b/Engine/Plugins/Experimental/Avalanche/Source/Avalanche/Public/Data/AvaWorldData.h
deleted file mode 100644
diff --git a/Engine/Plugins/Experimental/Avalanche/Source/Avalanche/Public/Framework/AvaGameInstance.h b/Engine/Plugins/Experimental/Avalanche/Source/Avalanche/Public/Framework/AvaGameInstance.h
--- a/Engine/Plugins/Experimental/Avalanche/Source/Avalanche/Public/Framework/AvaGameInstance.h
+++ b/Engine/Plugins/Experimental/Avalanche/Source/Avalanche/Public/Framework/AvaGameInstance.h
@@ -9,7 +9,6 @@
 
 class FRenderCommandFence;
 class FSceneViewport;
-class UAvalancheBlueprint;
 class UTextureRenderTarget2D;
 struct FAvaViewportQualitySettings;
 struct FAvaInstanceSettings;
@@ -31,17 +30,10 @@ class AVALANCHE_API UAvaGameInstance : public UGameInstance
 public:
 	/**
 	 * Creates a new Motion Design Game instance from the given Motion Design Template Asset.
-	 * Supported asset types: UMotion DesignBlueprint, UWorld (Level).
-	 *
-	 * Remark: For Motion Design Blueprints, the instanced embedded RemoteControlPreset will be registered
-	 * under the name of the outer Package Name. It is thus required that each instance be contained in
-	 * a unique package to avoid stomping the registered RCPs.
+	 * Supported asset types: UWorld (Level).
 	 */
 	static UAvaGameInstance* Create(UObject* InOuter);
 
-	// Legacy - will create and load a Motion Design blueprint.
-	static UAvaGameInstance* Create(UObject* InOuter, const TSoftObjectPtr<UAvalancheBlueprint>& InBlueprintTemplate);
-	
 	/** Create the game world */
 	bool CreateWorld();
 
@@ -49,12 +41,6 @@ public:
 
 	UWorld* GetPlayWorld() const { return PlayWorld.Get(); }
 
-	// Legacy
-	bool LoadMotionDesignBlueprint(const TSoftObjectPtr<UAvalancheBlueprint>& InBlueprint);
-
-	// Legacy
-	void BeginPlayMotionDesignBlueprint(const TSoftObjectPtr<UAvalancheBlueprint>& InBlueprint, const FAvaInstancePlaySettings& InWorldPlaySettings);
-	
 	bool BeginPlayWorld(const FAvaInstancePlaySettings& InWorldPlaySettings);
 
 	/**
@@ -105,11 +91,6 @@ public:
 	 * @remark Only valid within a BeginPlay()/EndPlay() scope. Null otherwise.
 	 */
 	UAvaGameViewportClient* GetAvaGameViewportClient() const { return ViewportClient; }
-	
-	/*
-	 * @return The Motion Design Blueprint used in play. Can be null if called outside play
-	 */
-	UAvalancheBlueprint* GetMotionDesignBlueprint() const { return ManagedBlueprint; }
 
 	DECLARE_MULTICAST_DELEGATE_TwoParams(FOnAvaGameInstanceEvent, UAvaGameInstance* /*InGameInstance*/, FName /*InChannelName*/);
 
@@ -119,15 +100,6 @@ public:
 	/** Event called when the render target becomes ready. */
 	static FOnAvaGameInstanceEvent& GetOnRenderTargetReady() { return OnRenderTargetReady; }
 
-	/*
-	 * Mark the current frame as a frame in which an asset was loaded synchronously.
-	 * When an asset has been loaded synchronously (i.e. Motion Design Blueprint), it will cause a
-	 * large delta seconds to occur on the next tick. By marking that particular frame it can be
-	 * used on the next tick to clamp the delta seconds to avoid animations skipping by the
-	 * amount of the sync load time.
-	 */
-	void MarkSynchronousAssetLoadingThisFrame();
-
 protected:
 	void Tick(float DeltaSeconds);
 
@@ -135,25 +107,11 @@ protected:
 	void EndPlayWorld();
 	void OnEndFrameTick();
 	void OnEnginePreExit();
-	
+
 	//~ Begin UObject
 	virtual void BeginDestroy() override;
 	//~ End UObject
 
-	/*
-	 * A Soft Ptr to the Source Motion Design (i.e. the actual asset).
-	 * It's a Soft Ptr because it can become stale if it's only a Weak Object Ptr
-	 * and we want to reload it if that happens
-	 */
-	TSoftObjectPtr<UAvalancheBlueprint> SourceBlueprint;
-
-	/*
-	 * The Managed Motion Design Blueprint. A Duplication of the Source Blueprint when Begin Play was called.
-	 * This shouldn't be GCd while in play, as it contains Animation and other data used during play.
-	 */
-	UPROPERTY(Transient)
-	TObjectPtr<UAvalancheBlueprint> ManagedBlueprint;
-
 	UPROPERTY(Transient)
 	TObjectPtr<UWorld> PlayWorld;
 
diff --git a/Engine/Plugins/Experimental/Avalanche/Source/Avalanche/Public/Framework/AvaInstanceSettings.h b/Engine/Plugins/Experimental/Avalanche/Source/Avalanche/Public/Framework/AvaInstanceSettings.h
--- a/Engine/Plugins/Experimental/Avalanche/Source/Avalanche/Public/Framework/AvaInstanceSettings.h
+++ b/Engine/Plugins/Experimental/Avalanche/Source/Avalanche/Public/Framework/AvaInstanceSettings.h
@@ -5,7 +5,7 @@
 #include "AvaInstanceSettings.generated.h"
 
 /**
- * Settings applied when instancing a Motion Design blueprint for playback.
+ * Settings applied when instancing a Motion Design Asset for playback.
  */
 USTRUCT()
 struct FAvaInstanceSettings
diff --git a/Engine/Plugins/Experimental/Avalanche/Source/Avalanche/Public/Framework/AvaSoftAssetPtr.h b/Engine/Plugins/Experimental/Avalanche/Source/Avalanche/Public/Framework/AvaSoftAssetPtr.h
--- a/Engine/Plugins/Experimental/Avalanche/Source/Avalanche/Public/Framework/AvaSoftAssetPtr.h
+++ b/Engine/Plugins/Experimental/Avalanche/Source/Avalanche/Public/Framework/AvaSoftAssetPtr.h
@@ -7,13 +7,12 @@
 enum class EAvalancheAssetType
 {
 	Unknown,
-	Blueprint,
 	World
 };
 
 /**
  * Avalanche Extension to FSoftObjectPath.
- * Adding the class information along to resolve between blueprint and level without loading anything.
+ * Adding the class information along to resolve about a Motion Design Asset without loading anything.
  */
 struct FAvaSoftAssetPath
 {
@@ -29,7 +28,7 @@ struct FAvaSoftAssetPath
 
 /**
  * Avalanche Extension to FSoftObjectPtr.
- * Adding the class information along to resolve between blueprint and level without loading anything.
+ * Adding the class information along to resolve a Motion Design Asset without loading anything.
  */
 
 struct FAvaSoftAssetPtr
diff --git a/Engine/Plugins/Experimental/Avalanche/Source/Avalanche/Public/Viewport/AvaViewportQualitySettings.h b/Engine/Plugins/Experimental/Avalanche/Source/Avalanche/Public/Viewport/AvaViewportQualitySettings.h
--- a/Engine/Plugins/Experimental/Avalanche/Source/Avalanche/Public/Viewport/AvaViewportQualitySettings.h
+++ b/Engine/Plugins/Experimental/Avalanche/Source/Avalanche/Public/Viewport/AvaViewportQualitySettings.h
@@ -36,7 +36,7 @@ struct FAvaViewportQualitySettingsFeature
 /** 
  * Avalanche Viewport Quality Settings
  * 
- * Advanced render and quality viewport settings to control performance saved with each Avalanche blueprint and set by the designer.
+ * Advanced render and quality viewport settings to control performance for a given Viewport.
  * Human-readable and blueprintable structure that holds flags for the FShowEngineFlags structure.
  * Can convert FShowEngineFlags to FAvaViewportQualitySettings and apply FAvaViewportQualitySettings to a FShowEngineFlags structure.
  */
diff --git a/Engine/Plugins/Experimental/Avalanche/Source/AvalancheEditor/Private/AssetTypeActions/AssetTypeActions_AvaBlueprint.cpp b/Engine/Plugins/Experimental/Avalanche/Source/AvalancheEditor/Private/AssetTypeActions/AssetTypeActions_AvaBlueprint.cpp
deleted file mode 100644
diff --git a/Engine/Plugins/Experimental/Avalanche/Source/AvalancheEditor/Private/AssetTypeActions/AssetTypeActions_AvaBlueprint.h b/Engine/Plugins/Experimental/Avalanche/Source/AvalancheEditor/Private/AssetTypeActions/AssetTypeActions_AvaBlueprint.h
deleted file mode 100644
diff --git a/Engine/Plugins/Experimental/Avalanche/Source/AvalancheEditor/Private/AvaEditorFunctionLibrary.cpp b/Engine/Plugins/Experimental/Avalanche/Source/AvalancheEditor/Private/AvaEditorFunctionLibrary.cpp
deleted file mode 100644
diff --git a/Engine/Plugins/Experimental/Avalanche/Source/AvalancheEditor/Private/AvaEditorModule.cpp b/Engine/Plugins/Experimental/Avalanche/Source/AvalancheEditor/Private/AvaEditorModule.cpp
--- a/Engine/Plugins/Experimental/Avalanche/Source/AvalancheEditor/Private/AvaEditorModule.cpp
+++ b/Engine/Plugins/Experimental/Avalanche/Source/AvalancheEditor/Private/AvaEditorModule.cpp
@@ -2,7 +2,6 @@
 
 #include "AvaEditorModule.h"
 #include "AssetRegistry/AssetRegistryModule.h"
-#include "AssetTypeActions/AssetTypeActions_AvaBlueprint.h"
 #include "AvaEditorActorUtils.h"
 #include "AvaEditorCommands.h"
 #include "AvaEditorIntegration.h"
@@ -28,6 +27,7 @@
 #include "Widgets/AvaViewportColorPickerActorClassRegistry.h"
 
 // Details View
+#include "AssetToolsModule.h"
 #include "DetailsPanel/DMMaterialInterfaceTypeCustomizer.h"
 #include "DetailView/Customizations/AvaAnchorAlignmentPropertyTypeCustomization.h"
 #include "DetailView/Customizations/AvaCategoryHiderCustomization.h"
@@ -124,7 +124,6 @@ void FAvaEditorModule::ShutdownModule()
 	FCoreDelegates::OnPostEngineInit.RemoveAll(this);
 	FCoreDelegates::OnPreExit.RemoveAll(this);
 
-	UnregisterAssetTools();
 	if (UObjectInitialized() && !IsEngineExitRequested())
 	{
 		UnregisterCustomLayouts();
@@ -185,24 +184,6 @@ void FAvaEditorModule::RegisterAssetTools()
 	AssetTools.RegisterAdvancedAssetCategory(AvalancheCategoryName
 		, LOCTEXT("MotionDesignCategoryName", "Motion Design"));
 
-	TSharedRef<IAssetTypeActions> AvaBlueprintActions = MakeShared<FAssetTypeActions_AvaBlueprint>();
-	AssetTools.RegisterAssetTypeActions(AvaBlueprintActions);
-	AssetTypeActions.Add(AvaBlueprintActions);
-}
-
-void FAvaEditorModule::UnregisterAssetTools()
-{
-	// Unregister all the asset types that we registered
-	if (FAssetToolsModule::IsModuleLoaded())
-	{
-		IAssetTools& AssetTools = FAssetToolsModule::GetModule().Get();
-		for (const TSharedPtr<IAssetTypeActions>& AssetTypeAction : AssetTypeActions)
-		{
-			AssetTools.UnregisterAssetTypeActions(AssetTypeAction.ToSharedRef());
-		}
-	}
-
-	AssetTypeActions.Empty();
 }
 
 void FAvaEditorModule::RegisterPropertyEditorCategories()
diff --git a/Engine/Plugins/Experimental/Avalanche/Source/AvalancheEditor/Private/AvaEditorModule.h b/Engine/Plugins/Experimental/Avalanche/Source/AvalancheEditor/Private/AvaEditorModule.h
--- a/Engine/Plugins/Experimental/Avalanche/Source/AvalancheEditor/Private/AvaEditorModule.h
+++ b/Engine/Plugins/Experimental/Avalanche/Source/AvalancheEditor/Private/AvaEditorModule.h
@@ -30,19 +30,14 @@ private:
 	/** Avalanche for the Level Editor*/
 	TSharedPtr<IAvaEditor> AvaLevelEditor;
 
-	/** All created asset type actions.  Cached here so that we can unregister it during shutdown. */
-	TArray<TSharedPtr<IAssetTypeActions>> AssetTypeActions;
-
 	void CreateAvaLevelEditor();
 
 	void PostEngineInit();
 	void PreExit();
 
 	void RegisterAssetTools();
-	void UnregisterAssetTools();
-
 	void RegisterPropertyEditorCategories();
-	
+
 	void RegisterCustomLayouts();
 	void UnregisterCustomLayouts();
 
diff --git a/Engine/Plugins/Experimental/Avalanche/Source/AvalancheEditor/Private/AvaEditorStyle.cpp b/Engine/Plugins/Experimental/Avalanche/Source/AvalancheEditor/Private/AvaEditorStyle.cpp
--- a/Engine/Plugins/Experimental/Avalanche/Source/AvalancheEditor/Private/AvaEditorStyle.cpp
+++ b/Engine/Plugins/Experimental/Avalanche/Source/AvalancheEditor/Private/AvaEditorStyle.cpp
@@ -69,10 +69,6 @@ void FAvaEditorStyle::Init()
 	SetCoreContentRoot(FPaths::EngineContentDir() / TEXT("Editor/Slate"));
 
 	// Custom Class Icons
-	Set("ClassThumbnail.AvaActor", new IMAGE_BRUSH("Icons/Icon64", Icon64));
-	Set("ClassThumbnail.AvalancheBlueprint", new IMAGE_BRUSH("Icons/Icon64", Icon64));
-	Set("ClassIcon.AvaActor", new IMAGE_BRUSH("Icons/Icon16", Icon16));
-	Set("ClassIcon.AvalancheBlueprint", new IMAGE_BRUSH("Icons/Icon16", Icon16));
 	Set("ClassIcon.AvaNullActor", new CORE_IMAGE_BRUSH(TEXT("Icons/SequencerIcons/icon_Sequencer_Move_24x"), Icon16));
 	Set("ClassIcon.AvaToolboxStarDynamicMesh", new IMAGE_BRUSH_SVG("Icons/ToolboxIcons/Favorite", Icon16));
 	Set("ClassIcon.AvaToolboxLineDynamicMesh", new IMAGE_BRUSH_SVG("Icons/ToolboxIcons/minus", Icon16));
diff --git a/Engine/Plugins/Experimental/Avalanche/Source/AvalancheEditor/Public/AvaEditorFunctionLibrary.h b/Engine/Plugins/Experimental/Avalanche/Source/AvalancheEditor/Public/AvaEditorFunctionLibrary.h
deleted file mode 100644
diff --git a/Engine/Plugins/Experimental/Avalanche/Source/AvalancheEffectors/Private/Cloner/AvaClonerComponent.cpp b/Engine/Plugins/Experimental/Avalanche/Source/AvalancheEffectors/Private/Cloner/AvaClonerComponent.cpp
--- a/Engine/Plugins/Experimental/Avalanche/Source/AvalancheEffectors/Private/Cloner/AvaClonerComponent.cpp
+++ b/Engine/Plugins/Experimental/Avalanche/Source/AvalancheEffectors/Private/Cloner/AvaClonerComponent.cpp
@@ -3,7 +3,6 @@
 #include "Cloner/AvaClonerComponent.h"
 #include "Async/Async.h"
 #include "AvaActorUtils.h"
-#include "AvaBlueprint.h"
 #include "AvaSceneTree.h"
 #include "Cloner/AvaClonerActor.h"
 #include "Cloner/Layouts/AvaClonerLayoutBase.h"
@@ -14,9 +13,11 @@
 #include "DynamicMesh/MeshTransforms.h"
 #include "DynamicMesh/Operations/MergeCoincidentMeshEdges.h"
 #include "DynamicMeshEditor.h"
+#include "Engine/World.h"
 #include "GeometryScript/MeshAssetFunctions.h"
 #include "GeometryScript/MeshBasicEditFunctions.h"
 #include "GeometryScript/SceneUtilityFunctions.h"
+#include "IAvaSceneInterface.h"
 #include "Materials/Material.h"
 #include "NiagaraMeshRendererProperties.h"
 #include "NiagaraSystem.h"
diff --git a/Engine/Plugins/Experimental/Avalanche/Source/AvalancheLevelViewport/Private/Widgets/AvaLevelViewportComponentTransformDetails.cpp b/Engine/Plugins/Experimental/Avalanche/Source/AvalancheLevelViewport/Private/Widgets/AvaLevelViewportComponentTransformDetails.cpp
--- a/Engine/Plugins/Experimental/Avalanche/Source/AvalancheLevelViewport/Private/Widgets/AvaLevelViewportComponentTransformDetails.cpp
+++ b/Engine/Plugins/Experimental/Avalanche/Source/AvalancheLevelViewport/Private/Widgets/AvaLevelViewportComponentTransformDetails.cpp
@@ -687,27 +687,7 @@ void FAvaLevelViewportComponentTransformDetails::UpdatePropertyHandlesObjects(TA
 
 bool FAvaLevelViewportComponentTransformDetails::GetIsEnabled() const
 {
-	if (SelectedActorInfo.NumSelected == 0)
-	{
-		return false;
-	}
-
-	/*
-	 * @TODO Locked actors
-	 *
-	bool bHasLockedActors = false;
-	
-	if (BlueprintEditor.IsValid())
-	{
-		if (TSharedPtr<FAvaEditorModeTools> ModeTools = BlueprintEditor.Pin()->GetAvaModeTools())
-		{
-			bHasLockedActors = ModeTools->HasLockedActorSelected();
-		}
-	}
-
-	return !bHasLockedActors;
-	*/
-	return true;
+	return SelectedActorInfo.NumSelected > 0;
 }
 
 const FSlateBrush* FAvaLevelViewportComponentTransformDetails::GetPreserveScaleRatioImage() const
diff --git a/Engine/Plugins/Experimental/Avalanche/Source/AvalancheLevelViewport/Private/Widgets/AvaLevelViewportComponentTransformDetails.h b/Engine/Plugins/Experimental/Avalanche/Source/AvalancheLevelViewport/Private/Widgets/AvaLevelViewportComponentTransformDetails.h
--- a/Engine/Plugins/Experimental/Avalanche/Source/AvalancheLevelViewport/Private/Widgets/AvaLevelViewportComponentTransformDetails.h
+++ b/Engine/Plugins/Experimental/Avalanche/Source/AvalancheLevelViewport/Private/Widgets/AvaLevelViewportComponentTransformDetails.h
@@ -392,14 +392,14 @@ private:
 	/** Holds the property handles' outer objects. Used to update the handles' objects when the actor construction script runs. */
 	TArray<TWeakObjectPtr<UObject>> CachedHandlesObjects;
 
-	// Link to the blueprint editor for which this is displaying details
+	// Link to the viewport client for which this is displaying details
 	TWeakPtr<FAvaLevelViewportClient> AvaViewportClientWeak;
 
 	// If true, use the multi-select as an ad-hoc parent
 	static bool bUseSelectionAsParent;
 	bool bLastUseSelectionAsParent;
 
-	// If true, we can use the multi-select as a parent (valid blueprint, etc)
+	// If true, we can use the multi-select as a parent
 	bool bCanUseSelectionAsParent;
 
 	TSharedPtr<SBox> TransformBox;
diff --git a/Engine/Plugins/Experimental/Avalanche/Source/AvalancheLevelViewport/Private/Widgets/SAvaLevelViewportGuide.h b/Engine/Plugins/Experimental/Avalanche/Source/AvalancheLevelViewport/Private/Widgets/SAvaLevelViewportGuide.h
--- a/Engine/Plugins/Experimental/Avalanche/Source/AvalancheLevelViewport/Private/Widgets/SAvaLevelViewportGuide.h
+++ b/Engine/Plugins/Experimental/Avalanche/Source/AvalancheLevelViewport/Private/Widgets/SAvaLevelViewportGuide.h
@@ -2,7 +2,6 @@
 
 #pragma once
 
-#include "AvaBlueprint.h"
 #include "AvaViewportGuideInfo.h"
 #include "Input/DragAndDrop.h"
 #include "Templates/SharedPointer.h"
diff --git a/Engine/Plugins/Experimental/Avalanche/Source/AvalancheMedia/Private/Playable/AvaPlayable.cpp b/Engine/Plugins/Experimental/Avalanche/Source/AvalancheMedia/Private/Playable/AvaPlayable.cpp
--- a/Engine/Plugins/Experimental/Avalanche/Source/AvalancheMedia/Private/Playable/AvaPlayable.cpp
+++ b/Engine/Plugins/Experimental/Avalanche/Source/AvalancheMedia/Private/Playable/AvaPlayable.cpp
@@ -11,7 +11,6 @@
 #include "IAvaSequenceProvider.h"
 #include "Playable/AvaPlayableGroup.h"
 #include "Playable/AvaPlayableRemoteControlValues.h"
-#include "Playable/Playables/AvaPlayableBlueprint.h"
 #include "Playable/Playables/AvaPlayableLevelStreaming.h"
 #include "Playable/Playables/AvaPlayableRemoteProxy.h"
 #include "Playback/AvaPlaybackUtils.h"
@@ -368,8 +367,6 @@ UAvaPlayable* UAvaPlayable::CreateLocalPlayable(UObject* InOuter, const FPlayabl
 {
 	switch (InPlayableInfo.SourceAsset.GetAssetType())
 	{
-	case EAvalancheAssetType::Blueprint:
-		return NewObject<UAvaPlayableBlueprint>(InOuter ? InOuter : GEngine);
 	case EAvalancheAssetType::World:
 		return NewObject<UAvaPlayableLevelStreaming>(InOuter ? InOuter : GEngine);
 	default:
diff --git a/Engine/Plugins/Experimental/Avalanche/Source/AvalancheMedia/Private/Playable/Playables/AvaPlayableBlueprint.cpp b/Engine/Plugins/Experimental/Avalanche/Source/AvalancheMedia/Private/Playable/Playables/AvaPlayableBlueprint.cpp
deleted file mode 100644
diff --git a/Engine/Plugins/Experimental/Avalanche/Source/AvalancheMedia/Private/Playable/Playables/AvaPlayableBlueprint.h b/Engine/Plugins/Experimental/Avalanche/Source/AvalancheMedia/Private/Playable/Playables/AvaPlayableBlueprint.h
deleted file mode 100644
diff --git a/Engine/Plugins/Experimental/Avalanche/Source/AvalancheMedia/Private/Playable/Playables/AvaPlayableRemoteProxy.cpp b/Engine/Plugins/Experimental/Avalanche/Source/AvalancheMedia/Private/Playable/Playables/AvaPlayableRemoteProxy.cpp
--- a/Engine/Plugins/Experimental/Avalanche/Source/AvalancheMedia/Private/Playable/Playables/AvaPlayableRemoteProxy.cpp
+++ b/Engine/Plugins/Experimental/Avalanche/Source/AvalancheMedia/Private/Playable/Playables/AvaPlayableRemoteProxy.cpp
@@ -187,18 +187,6 @@ bool UAvaPlayableRemoteProxy::InitPlayable(const FPlayableCreationInfo& InPlayab
 	// Remote playables have proxy playable groups imitating the same setup as local ones.
 	switch (InPlayableInfo.SourceAsset.GetAssetType())
 	{
-	case EAvalancheAssetType::Blueprint:
-		{
-			UAvaPlayableGroup::FPlayableGroupCreationInfo PlayableGroupCreationInfo;
-			PlayableGroupCreationInfo.PlayableGroupManager = InPlayableInfo.PlayableGroupManager;
-			PlayableGroupCreationInfo.SourceAssetPath = InPlayableInfo.SourceAsset.ToSoftObjectPath();
-			PlayableGroupCreationInfo.ChannelName = InPlayableInfo.ChannelName;
-			PlayableGroupCreationInfo.bIsRemoteProxy = bIsRemoteProxy;
-			PlayableGroupCreationInfo.bIsSharedGroup = false;
-			
-			PlayableGroup = UAvaPlayableGroup::MakePlayableGroup(InPlayableInfo.PlayableGroupManager, PlayableGroupCreationInfo);
-		}
-		break;
 	case EAvalancheAssetType::World:
 		PlayableGroup = InPlayableInfo.PlayableGroupManager->GetOrCreateSharedLevelGroup(InPlayableInfo.ChannelName, bIsRemoteProxy);
 		break;
diff --git a/Engine/Plugins/Experimental/Avalanche/Source/AvalancheMedia/Private/Playback/AvaPlaybackManager.cpp b/Engine/Plugins/Experimental/Avalanche/Source/AvalancheMedia/Private/Playback/AvaPlaybackManager.cpp
--- a/Engine/Plugins/Experimental/Avalanche/Source/AvalancheMedia/Private/Playback/AvaPlaybackManager.cpp
+++ b/Engine/Plugins/Experimental/Avalanche/Source/AvalancheMedia/Private/Playback/AvaPlaybackManager.cpp
@@ -4,7 +4,6 @@
 
 #include "AssetRegistry/AssetData.h"
 #include "AssetRegistry/IAssetRegistry.h"
-#include "AvaBlueprint.h"
 #include "Broadcast/AvaBroadcast.h"
 #include "Engine/Engine.h"
 #include "Framework/AvaGameInstance.h"
@@ -14,7 +13,6 @@
 #include "Playable/AvaPlayableGroup.h"
 #include "Playable/AvaPlayableGroupManager.h"
 #include "Playback/AvaPlaybackUtils.h"
-#include "Playback/Nodes/AvaPlaybackNodeBlueprintPlayer.h"
 #include "Playback/Nodes/AvaPlaybackNodeLevelPlayer.h"
 #include "Playback/Transition/AvaPlaybackTransition.h"
 #include "UnrealClient.h"
@@ -368,7 +366,7 @@ UAvaPlaybackGraph*  FAvaPlaybackManager::LoadPlaybackObject(const FSoftObjectPat
 
 		// Todo: Investigate LoadPackageAsync.
 		// For now, we tolerate a sync load here because there will be hitch from converting the
-		// Motion Design blueprint to a world.
+		// Motion Design Asset to a world.
 		TempPackage = LoadPackage(nullptr, *PackageName, LOAD_None );
 	}
 	
@@ -381,14 +379,6 @@ UAvaPlaybackGraph*  FAvaPlaybackManager::LoadPlaybackObject(const FSoftObjectPat
 			{
 				return AvaPlayback;
 			}
-			
-			// When the asset is an Motion Design Blueprint, it is wrapped in a playback object.
-			if (const UAvalancheBlueprint* const AvaBlueprint = Cast<UAvalancheBlueprint>(FoundObject))
-			{
-				UAvaPlaybackGraph* const AvaPlayback = BuildPlaybackFromBlueprint(AvaBlueprint, InChannelName);
-				check(AvaPlayback);
-				return AvaPlayback;
-			}
 
 			if (const UWorld* const World = Cast<UWorld>(FoundObject))
 			{
@@ -417,25 +407,12 @@ UAvaPlaybackGraph*  FAvaPlaybackManager::LoadPlaybackObject(const FSoftObjectPat
 	return nullptr;
 }
 
-UAvaPlaybackGraph* FAvaPlaybackManager::BuildPlaybackFromBlueprint(const UAvalancheBlueprint* InBlueprint, const FString& InChannelName) const
-{
-	using namespace UE::AvaPlaybackManager::Private;
-	FAvaPlaybackGraphBuilder GraphBuilder(GetPlayableGroupManager());
-
-	// Construct Player Node and assign the Blueprint that we passed in
-	UAvaPlaybackNodeBlueprintPlayer* const PlayerNode = GraphBuilder.ConstructPlaybackNode<UAvaPlaybackNodeBlueprintPlayer>();
-	PlayerNode->SetAsset(InBlueprint);
-	
-	GraphBuilder.ConnectToRoot(InChannelName, PlayerNode);
-	return GraphBuilder.FinishBuilding();
-}
-
 UAvaPlaybackGraph* FAvaPlaybackManager::BuildPlaybackFromWorld(const TSoftObjectPtr<UWorld>& InWorld, const FString& InChannelName) const
 {
 	using namespace UE::AvaPlaybackManager::Private;
 	FAvaPlaybackGraphBuilder GraphBuilder(GetPlayableGroupManager());
 
-	// Construct Player Node and assign the Blueprint that we passed in
+	// Construct Player Node and assign the World that we passed in
 	UAvaPlaybackNodeLevelPlayer* const PlayerNode = GraphBuilder.ConstructPlaybackNode<UAvaPlaybackNodeLevelPlayer>();
 	PlayerNode->SetAsset(InWorld);
 	
diff --git a/Engine/Plugins/Experimental/Avalanche/Source/AvalancheMedia/Private/Playback/Nodes/AvaPlaybackNodeBlueprintPlayer.cpp b/Engine/Plugins/Experimental/Avalanche/Source/AvalancheMedia/Private/Playback/Nodes/AvaPlaybackNodeBlueprintPlayer.cpp
deleted file mode 100644
diff --git a/Engine/Plugins/Experimental/Avalanche/Source/AvalancheMedia/Private/Playback/Nodes/Events/Actions/AvaPlaybackNode_PlayAnim.cpp b/Engine/Plugins/Experimental/Avalanche/Source/AvalancheMedia/Private/Playback/Nodes/Events/Actions/AvaPlaybackNode_PlayAnim.cpp
--- a/Engine/Plugins/Experimental/Avalanche/Source/AvalancheMedia/Private/Playback/Nodes/Events/Actions/AvaPlaybackNode_PlayAnim.cpp
+++ b/Engine/Plugins/Experimental/Avalanche/Source/AvalancheMedia/Private/Playback/Nodes/Events/Actions/AvaPlaybackNode_PlayAnim.cpp
@@ -3,13 +3,11 @@
 #include "AvaPlaybackNode_PlayAnim.h"
 
 #include "Async/Async.h"
-#include "AvaBlueprint.h"
 #include "AvaScene.h"
 #include "AvaSequence.h"
 #include "Broadcast/AvaBroadcast.h"
 #include "Engine/Level.h"
 #include "Playback/AvaPlaybackGraph.h"
-#include "Playback/Nodes/AvaPlaybackNodeBlueprintPlayer.h"
 #include "Playback/Nodes/AvaPlaybackNodeLevelPlayer.h"
 
 #define LOCTEXT_NAMESPACE "AvaPlaybackNode_PlayAnim"
@@ -95,22 +93,6 @@ void UAvaPlaybackNode_PlayAnim::DryRun(const TArray<UAvaPlaybackNode*>& InAncest
 {
 	for (UAvaPlaybackNode* const PlaybackNode : InAncestors)
 	{
-		if (UAvaPlaybackNodeBlueprintPlayer* const PlayerNode = Cast<UAvaPlaybackNodeBlueprintPlayer>(PlaybackNode))
-		{
-			TSoftObjectPtr<UAvalancheBlueprint> Asset = PlayerNode->GetAsset();
-			if (UAvalancheBlueprint* const AvalancheBlueprint = Asset.LoadSynchronous())
-			{
-				SeenAssetsInDryRun.Add(PlayerNode->GetAssetPath());
-				FAvaPlaybackAnimations& Animations = AnimationMap.FindOrAdd(PlayerNode->GetAssetPath());
-				for (const TObjectPtr<UAvaSequence>& Animation : AvalancheBlueprint->GetSequences())
-				{
-					if (Animation)
-					{
-						Animations.AvailableAnimations.FindOrAdd(UE::AvaMedia::NodePlayAnim::Private::GetSequenceName(Animation));
-					}
-				}
-			}
-		}
 		if (UAvaPlaybackNodeLevelPlayer* const PlayerNode = Cast<UAvaPlaybackNodeLevelPlayer>(PlaybackNode))
 		{
 			TSoftObjectPtr<UWorld> Asset = PlayerNode->GetAsset();
diff --git a/Engine/Plugins/Experimental/Avalanche/Source/AvalancheMedia/Private/Playback/Nodes/Events/Actions/AvaPlaybackNode_PlayAnim.h b/Engine/Plugins/Experimental/Avalanche/Source/AvalancheMedia/Private/Playback/Nodes/Events/Actions/AvaPlaybackNode_PlayAnim.h
--- a/Engine/Plugins/Experimental/Avalanche/Source/AvalancheMedia/Private/Playback/Nodes/Events/Actions/AvaPlaybackNode_PlayAnim.h
+++ b/Engine/Plugins/Experimental/Avalanche/Source/AvalancheMedia/Private/Playback/Nodes/Events/Actions/AvaPlaybackNode_PlayAnim.h
@@ -26,6 +26,6 @@ protected:
 	UPROPERTY(VisibleAnywhere, Category = "Motion Design")
 	TMap<FSoftObjectPath, FAvaPlaybackAnimations> AnimationMap;
 
-	//All the Blueprints connected to this Node from the Dry Run (only populated while Dry Running!)
+	//All the assets connected to this Node from the Dry Run (only populated while Dry Running!)
 	TSet<FSoftObjectPath> SeenAssetsInDryRun;
 };
diff --git a/Engine/Plugins/Experimental/Avalanche/Source/AvalancheMedia/Private/Rundown/AvaRundownManagedInstance.cpp b/Engine/Plugins/Experimental/Avalanche/Source/AvalancheMedia/Private/Rundown/AvaRundownManagedInstance.cpp
--- a/Engine/Plugins/Experimental/Avalanche/Source/AvalancheMedia/Private/Rundown/AvaRundownManagedInstance.cpp
+++ b/Engine/Plugins/Experimental/Avalanche/Source/AvalancheMedia/Private/Rundown/AvaRundownManagedInstance.cpp
@@ -17,7 +17,7 @@ void FAvaRundownManagedInstance::InvalidateFromCache()
 {
 	if (ParentCache)
 	{
-		// The managed blueprint is not deleted immediately to avoid issues
+		// The managed asset is not deleted immediately to avoid issues
 		// with destroying delegates while they are being called.
 		ParentCache->InvalidateNoDelete(SourceAssetPath);
 		ParentCache = nullptr;
diff --git a/Engine/Plugins/Experimental/Avalanche/Source/AvalancheMedia/Private/Rundown/AvaRundownManagedInstanceBlueprint.cpp b/Engine/Plugins/Experimental/Avalanche/Source/AvalancheMedia/Private/Rundown/AvaRundownManagedInstanceBlueprint.cpp
deleted file mode 100644
diff --git a/Engine/Plugins/Experimental/Avalanche/Source/AvalancheMedia/Private/Rundown/AvaRundownManagedInstanceBlueprint.h b/Engine/Plugins/Experimental/Avalanche/Source/AvalancheMedia/Private/Rundown/AvaRundownManagedInstanceBlueprint.h
deleted file mode 100644
diff --git a/Engine/Plugins/Experimental/Avalanche/Source/AvalancheMedia/Private/Rundown/AvaRundownManagedInstanceCache.cpp b/Engine/Plugins/Experimental/Avalanche/Source/AvalancheMedia/Private/Rundown/AvaRundownManagedInstanceCache.cpp
--- a/Engine/Plugins/Experimental/Avalanche/Source/AvalancheMedia/Private/Rundown/AvaRundownManagedInstanceCache.cpp
+++ b/Engine/Plugins/Experimental/Avalanche/Source/AvalancheMedia/Private/Rundown/AvaRundownManagedInstanceCache.cpp
@@ -9,7 +9,6 @@
 #include "Engine/World.h"
 #include "IAvaMediaModule.h"
 #include "Playback/AvaPlaybackUtils.h"
-#include "Rundown/AvaRundownManagedInstanceBlueprint.h"
 #include "Rundown/AvaRundownManagedInstanceLevel.h"
 #include "UObject/ObjectSaveContext.h"
 #include "UObject/Package.h"
@@ -69,10 +68,6 @@ TSharedPtr<FAvaRundownManagedInstance> FAvaRundownManagedInstanceCache::GetOrLoa
 	{
 		NewEntry = MakeShared<FAvaRundownManagedInstanceLevel>(this, InAssetPath);
 	}
-	else
-	{
-		NewEntry = MakeShared<FAvaRundownManagedInstanceBlueprint>(this, InAssetPath);
-	}
 
 	if (NewEntry.IsValid() && NewEntry->IsValid())
 	{
@@ -123,9 +118,9 @@ void FAvaRundownManagedInstanceCache::Flush(const FSoftObjectPath& InAssetPath)
 
 void FAvaRundownManagedInstanceCache::Flush()
 {
-	RemoveEntries([](const FSoftObjectPath& InAssetPath, const TSharedPtr<FAvaRundownManagedInstance>& InManagedBlueprint)
+	RemoveEntries([](const FSoftObjectPath& InAssetPath, const TSharedPtr<FAvaRundownManagedInstance>& InManagedInstance)
 	{
-		return (InManagedBlueprint.GetSharedReferenceCount() <= 1) ? true : false;
+		return (InManagedInstance.GetSharedReferenceCount() <= 1) ? true : false;
 	}, false);
 }
 
@@ -201,7 +196,7 @@ void FAvaRundownManagedInstanceCache::OnPackageSaved(const FString& InPackageFil
 void FAvaRundownManagedInstanceCache::OnAvaSyncPackageModified(IAvaMediaSyncProvider* InAvaMediaSyncProvider, const FName& InPackageName)
 {
 	UE_LOG(LogAvaMedia, Verbose,
-		TEXT("A sync operation has touched the package \"%s\" on disk. Managed Motion Design Blueprint Cache notified."),
+		TEXT("A sync operation has touched the package \"%s\" on disk. Managed Motion Design Instance Cache notified."),
 		*InPackageName.ToString());
 
 	OnPackageModified(InPackageName);
@@ -221,7 +216,7 @@ void FAvaRundownManagedInstanceCache::OnPackageModified(const FName& InPackageNa
 			if (InAssetPath.GetLongPackageFName() == InPackageName)
 			{
 				UE_LOG(LogAvaMedia, Log,
-					TEXT("Managed Motion Design Blueprint Cache: Package \"%s\" being touched caused asset \"%s\" to be invalidated."),
+					TEXT("Managed Motion Design Instance Cache: Package \"%s\" being touched caused asset \"%s\" to be invalidated."),
 					*InPackageName.ToString(), *InAssetPath.ToString());
 				return true;
 			}
diff --git a/Engine/Plugins/Experimental/Avalanche/Source/AvalancheMedia/Private/Rundown/AvaRundownManagedInstanceUtils.cpp b/Engine/Plugins/Experimental/Avalanche/Source/AvalancheMedia/Private/Rundown/AvaRundownManagedInstanceUtils.cpp
--- a/Engine/Plugins/Experimental/Avalanche/Source/AvalancheMedia/Private/Rundown/AvaRundownManagedInstanceUtils.cpp
+++ b/Engine/Plugins/Experimental/Avalanche/Source/AvalancheMedia/Private/Rundown/AvaRundownManagedInstanceUtils.cpp
@@ -23,7 +23,7 @@ UPackage* FAvaRundownManagedInstanceUtils::MakeManagedInstancePackage(const FSof
 	}
 	InstancePackageName += InstanceSubPath;
 
-	UPackage* ManagedMotionDesignBlueprintPackage = nullptr;
+	UPackage* ManagedInstancePackage = nullptr;
 
 	// Make sure the package name is unique. We don't want stomping of existing instances.
 	if (FindObject<UPackage>( nullptr, *InstancePackageName))
@@ -37,10 +37,10 @@ UPackage* FAvaRundownManagedInstanceUtils::MakeManagedInstancePackage(const FSof
 		InstancePackageName = UniqueInstancePackageName;
 	}
 	
-	ManagedMotionDesignBlueprintPackage = CreatePackage(*InstancePackageName);
-	if (ManagedMotionDesignBlueprintPackage)
+	ManagedInstancePackage = CreatePackage(*InstancePackageName);
+	if (ManagedInstancePackage)
 	{
-		ManagedMotionDesignBlueprintPackage->SetFlags(RF_Transient | RF_Public);
+		ManagedInstancePackage->SetFlags(RF_Transient | RF_Public);
 	}
 	else
 	{
@@ -49,7 +49,7 @@ UPackage* FAvaRundownManagedInstanceUtils::MakeManagedInstancePackage(const FSof
 			TEXT("Unable to create package \"%s\" for Motion Design Managed Instance."),
 			*InstancePackageName);
 	}
-	return ManagedMotionDesignBlueprintPackage;
+	return ManagedInstancePackage;
 }
 
 void FAvaRundownManagedInstanceUtils::PreventWorldFromBeingSeenAsLeakingByLevelEditor(UWorld* InWorld)
diff --git a/Engine/Plugins/Experimental/Avalanche/Source/AvalancheMedia/Private/Rundown/AvaRundownPage.cpp b/Engine/Plugins/Experimental/Avalanche/Source/AvalancheMedia/Private/Rundown/AvaRundownPage.cpp
--- a/Engine/Plugins/Experimental/Avalanche/Source/AvalancheMedia/Private/Rundown/AvaRundownPage.cpp
+++ b/Engine/Plugins/Experimental/Avalanche/Source/AvalancheMedia/Private/Rundown/AvaRundownPage.cpp
@@ -542,11 +542,6 @@ void FAvaRundownPage::SetRemoteControlControllerValue(const FGuid& InId, const F
 
 void FAvaRundownPage::PostLoad()
 {
-	if (!AvalancheBlueprint_DEPRECATED.IsNull() && !AssetPath.IsValid())
-	{
-		AssetPath = AvalancheBlueprint_DEPRECATED.ToSoftObjectPath();
-	}
-	AvalancheBlueprint_DEPRECATED.Reset();
 }
 
 int32 FAvaRundownPage::GetNumTemplates(const UAvaRundown* InRundown) const
diff --git a/Engine/Plugins/Experimental/Avalanche/Source/AvalancheMedia/Public/AvaMediaDefines.h b/Engine/Plugins/Experimental/Avalanche/Source/AvalancheMedia/Public/AvaMediaDefines.h
--- a/Engine/Plugins/Experimental/Avalanche/Source/AvalancheMedia/Public/AvaMediaDefines.h
+++ b/Engine/Plugins/Experimental/Avalanche/Source/AvalancheMedia/Public/AvaMediaDefines.h
@@ -232,7 +232,7 @@ ENUM_CLASS_FLAGS(EAvaRundownPageChanges);
  *
  *	This is also sometimes referred to as the "source" asset, by contrast to the
  *	runtime (or managed) transient asset for runtime playback. It will also refer to
- *	the source Motion Design blueprint, rather than a playback object.
+ *	the source Motion Design asset, rather than a playback object.
  */
 UENUM()
 enum class EAvaPlaybackAssetStatus
diff --git a/Engine/Plugins/Experimental/Avalanche/Source/AvalancheMedia/Public/AvaMediaSettings.h b/Engine/Plugins/Experimental/Avalanche/Source/AvalancheMedia/Public/AvaMediaSettings.h
--- a/Engine/Plugins/Experimental/Avalanche/Source/AvalancheMedia/Public/AvaMediaSettings.h
+++ b/Engine/Plugins/Experimental/Avalanche/Source/AvalancheMedia/Public/AvaMediaSettings.h
@@ -183,7 +183,7 @@ public:
 	FAvaInstanceSettings AvaInstanceSettings;
 
 	/**
-	 * Maximum cached Managed Motion Design Blueprints used for rundown editor's page details.
+	 * Maximum cached Managed Motion Design assets used for rundown editor's page details.
 	 * A value of 0 indicate the cache will grow without limit.
 	 */
 	UPROPERTY(Config, EditAnywhere, Category = "Managed Motion Design Instance Cache", meta = (DisplayName = "Maximum Cache Size"))
diff --git a/Engine/Plugins/Experimental/Avalanche/Source/AvalancheMedia/Public/IAvaMediaModule.h b/Engine/Plugins/Experimental/Avalanche/Source/AvalancheMedia/Public/IAvaMediaModule.h
--- a/Engine/Plugins/Experimental/Avalanche/Source/AvalancheMedia/Public/IAvaMediaModule.h
+++ b/Engine/Plugins/Experimental/Avalanche/Source/AvalancheMedia/Public/IAvaMediaModule.h
@@ -115,7 +115,7 @@ public:
 	virtual const FAvaInstanceSettings& GetAvaInstanceSettings() const = 0;
 	
 	/**
-	 *	This is the backend for playing Motion Design blueprints locally.
+	 *	This is the backend for playing Motion Design assets locally.
 	 */
 	virtual FAvaPlaybackManager& GetLocalPlaybackManager() const = 0;
 
diff --git a/Engine/Plugins/Experimental/Avalanche/Source/AvalancheMedia/Public/Playable/AvaPlayable.h b/Engine/Plugins/Experimental/Avalanche/Source/AvalancheMedia/Public/Playable/AvaPlayable.h
--- a/Engine/Plugins/Experimental/Avalanche/Source/AvalancheMedia/Public/Playable/AvaPlayable.h
+++ b/Engine/Plugins/Experimental/Avalanche/Source/AvalancheMedia/Public/Playable/AvaPlayable.h
@@ -55,8 +55,7 @@ enum class EAvaPlayableCommandResult : uint8
  * Design goal:
  *
  * The design goal is to abstract the implementation of a playable.
- * So far we have 2 implementations:
- * - Motion Design Blueprint that requires to be in it's own game instance/world.
+ * So far we have 1 implementation:
  * - Level Streaming that can be streamed with other levels in the same game instance.
  *
  * To support multiple playable in the same channel/output, there are 2 ways:
diff --git a/Engine/Plugins/Experimental/Avalanche/Source/AvalancheMedia/Public/Playback/AvaPlaybackGraph.h b/Engine/Plugins/Experimental/Avalanche/Source/AvalancheMedia/Public/Playback/AvaPlaybackGraph.h
--- a/Engine/Plugins/Experimental/Avalanche/Source/AvalancheMedia/Public/Playback/AvaPlaybackGraph.h
+++ b/Engine/Plugins/Experimental/Avalanche/Source/AvalancheMedia/Public/Playback/AvaPlaybackGraph.h
@@ -190,8 +190,8 @@ public:
 	TArray<UAvaGameInstance*> GetActiveGameInstances() const;
 	
 	/**
-	 *	Determines if the graph contains a node with the given source Motion Design blueprint.
-	 *	This works even if the Motion Design blueprint running.
+	 *	Determines if the graph contains a node with the given source Motion Design Asset.
+	 *	This works even if the Motion Design asset running.
 	 */
 	bool HasPlayerNodeForSourceAsset(const FSoftObjectPath& InSourceAssetPath) const;
 
diff --git a/Engine/Plugins/Experimental/Avalanche/Source/AvalancheMedia/Public/Playback/AvaPlaybackManager.h b/Engine/Plugins/Experimental/Avalanche/Source/AvalancheMedia/Public/Playback/AvaPlaybackManager.h
--- a/Engine/Plugins/Experimental/Avalanche/Source/AvalancheMedia/Public/Playback/AvaPlaybackManager.h
+++ b/Engine/Plugins/Experimental/Avalanche/Source/AvalancheMedia/Public/Playback/AvaPlaybackManager.h
@@ -9,7 +9,6 @@
 class FAvaPlaybackManager;
 class FAvaPlaybackSourceAssetEntry;
 class IAvaMediaSyncProvider;
-class UAvalancheBlueprint;
 class UAvaPlayableGroupManager;
 class UAvaPlaybackTransition;
 struct FAvaPlaybackAnimPlaySettings;	// private
@@ -216,13 +215,12 @@ public:
 	void InvalidatePlaybackAssetEntry(const FSoftObjectPath& InAssetPath);
 	
 	UAvaPlaybackGraph* LoadPlaybackObject(const FSoftObjectPath& InAssetPath, const FString& InChannelName) const;
-	UAvaPlaybackGraph* BuildPlaybackFromBlueprint(const UAvalancheBlueprint* InBlueprint, const FString& InChannelName) const;
 	UAvaPlaybackGraph* BuildPlaybackFromWorld(const TSoftObjectPtr<UWorld>& InWorld, const FString& InChannelName) const;
 	
 	/**
 	 * Stops all currently playing playback objects.
 	 * @param bInUnload if true, will also unload the objects.
-	 * @return Returns the list of all source blueprint assets that where stopped.
+	 * @return Returns the list of all source assets that were stopped.
 	 **/
 	TArray<FSoftObjectPath> StopAllPlaybacks(bool bInUnload);
 
diff --git a/Engine/Plugins/Experimental/Avalanche/Source/AvalancheMedia/Public/Playback/AvaPlaybackMessages.h b/Engine/Plugins/Experimental/Avalanche/Source/AvalancheMedia/Public/Playback/AvaPlaybackMessages.h
--- a/Engine/Plugins/Experimental/Avalanche/Source/AvalancheMedia/Public/Playback/AvaPlaybackMessages.h
+++ b/Engine/Plugins/Experimental/Avalanche/Source/AvalancheMedia/Public/Playback/AvaPlaybackMessages.h
@@ -308,7 +308,7 @@ struct AVALANCHEMEDIA_API FAvaPlaybackCommand
 	UPROPERTY()
 	FSoftObjectPath AssetPath;
 
-	/** Channel to Play, only considered when the Asset in question is an Motion Design Playable (Blueprint or Level). */
+	/** Channel to Play, only considered when the Asset in question is an Motion Design Playable. */
 	UPROPERTY()
 	FString ChannelName;
 	
diff --git a/Engine/Plugins/Experimental/Avalanche/Source/AvalancheMedia/Public/Playback/AvaPlaybackServer.h b/Engine/Plugins/Experimental/Avalanche/Source/AvalancheMedia/Public/Playback/AvaPlaybackServer.h
--- a/Engine/Plugins/Experimental/Avalanche/Source/AvalancheMedia/Public/Playback/AvaPlaybackServer.h
+++ b/Engine/Plugins/Experimental/Avalanche/Source/AvalancheMedia/Public/Playback/AvaPlaybackServer.h
@@ -22,8 +22,8 @@ DECLARE_LOG_CATEGORY_EXTERN(LogAvaPlaybackServer, Log, All);
  *
  * The playback server implements the commands for broadcast (channels and outputs) and playback.
  * Playback assets are either "playables" or "playback graphs", however the playback server will
- * create playback graphs for everything. If the asset to play is a "playable" (either a level or
- * a ava blueprint), it will create a transient playback graph for it.
+ * create playback graphs for everything. If the asset to play is a "playable" (e.g. a level),
+ * it will create a transient playback graph for it.
  *
  * A lot of the playback commands are geared toward running a playback graph with a single playable node,
  * mostly because this system is only used with rundowns or playback graphs run on the client side.
diff --git a/Engine/Plugins/Experimental/Avalanche/Source/AvalancheMedia/Public/Playback/Nodes/AvaPlaybackNodeBlueprintPlayer.h b/Engine/Plugins/Experimental/Avalanche/Source/AvalancheMedia/Public/Playback/Nodes/AvaPlaybackNodeBlueprintPlayer.h
deleted file mode 100644
diff --git a/Engine/Plugins/Experimental/Avalanche/Source/AvalancheMedia/Public/Rundown/AvaRundownManagedInstanceCache.h b/Engine/Plugins/Experimental/Avalanche/Source/AvalancheMedia/Public/Rundown/AvaRundownManagedInstanceCache.h
--- a/Engine/Plugins/Experimental/Avalanche/Source/AvalancheMedia/Public/Rundown/AvaRundownManagedInstanceCache.h
+++ b/Engine/Plugins/Experimental/Avalanche/Source/AvalancheMedia/Public/Rundown/AvaRundownManagedInstanceCache.h
@@ -8,7 +8,7 @@ class IAvaMediaSyncProvider;
 class UPackage;
 
 /**
- * Cache for the managed Motion Design Blueprint instances.
+ * Cache for the managed Motion Design instances.
  *
  * This is LRU cache. 
  */
diff --git a/Engine/Plugins/Experimental/Avalanche/Source/AvalancheMedia/Public/Rundown/AvaRundownPage.h b/Engine/Plugins/Experimental/Avalanche/Source/AvalancheMedia/Public/Rundown/AvaRundownPage.h
--- a/Engine/Plugins/Experimental/Avalanche/Source/AvalancheMedia/Public/Rundown/AvaRundownPage.h
+++ b/Engine/Plugins/Experimental/Avalanche/Source/AvalancheMedia/Public/Rundown/AvaRundownPage.h
@@ -7,7 +7,6 @@
 #include "Playable/AvaPlayableRemoteControlValues.h"
 #include "AvaRundownPage.generated.h"
 
-class UAvalancheBlueprint;
 class UAvaRundown;
 
 UENUM()
@@ -245,11 +244,8 @@ protected:
 	UPROPERTY(VisibleAnywhere, BlueprintReadOnly, Category = "Motion Design")
 	FString PageName;
 
-	UPROPERTY(meta = (DeprecatedProperty, DeprecationMessage = "Use AssetPath."))
-	TSoftObjectPtr<UAvalancheBlueprint> AvalancheBlueprint_DEPRECATED;
-
 	/** Template property: path for this template. */
-	UPROPERTY(VisibleAnywhere, BlueprintReadOnly, Category = "Motion Design", meta=(AllowedClasses="/Script/Engine.World, /Script/Avalanche.AvalancheBlueprint"))
+	UPROPERTY(VisibleAnywhere, BlueprintReadOnly, Category = "Motion Design", meta=(AllowedClasses="/Script/Engine.World"))
 	FSoftObjectPath AssetPath;
 
 	/** Template property: List the Ids of all instances. */
diff --git a/Engine/Plugins/Experimental/Avalanche/Source/AvalancheMediaEditor/Private/Playback/Graph/Nodes/AvaPlaybackEditorGraphNode_BlueprintPlayer.cpp b/Engine/Plugins/Experimental/Avalanche/Source/AvalancheMediaEditor/Private/Playback/Graph/Nodes/AvaPlaybackEditorGraphNode_BlueprintPlayer.cpp
deleted file mode 100644
diff --git a/Engine/Plugins/Experimental/Avalanche/Source/AvalancheMediaEditor/Private/Rundown/DetailsView/SAvaRundownPageDetails.cpp b/Engine/Plugins/Experimental/Avalanche/Source/AvalancheMediaEditor/Private/Rundown/DetailsView/SAvaRundownPageDetails.cpp
--- a/Engine/Plugins/Experimental/Avalanche/Source/AvalancheMediaEditor/Private/Rundown/DetailsView/SAvaRundownPageDetails.cpp
+++ b/Engine/Plugins/Experimental/Avalanche/Source/AvalancheMediaEditor/Private/Rundown/DetailsView/SAvaRundownPageDetails.cpp
@@ -229,7 +229,7 @@ void SAvaRundownPageDetails::OnManagedInstanceCacheEntryInvalidated(const FSoftO
 						bRefreshSelectedPageQueued = true;
 						// Queue a refresh on next tick.
 						// We don't want to refresh immediately to avoid issues with
-						// cascading events within the managed blueprint cache.
+						// cascading events within the managed instance cache.
 						TWeakPtr<SWidget> ThisWeak(AsShared());
 						AsyncTask(ENamedThreads::GameThread, [ThisWeak]()
 							{
diff --git a/Engine/Plugins/Experimental/Avalanche/Source/AvalancheMediaEditor/Private/Rundown/Factories/Suggestions/AvaRundownFilterPathSuggestionFactory.cpp b/Engine/Plugins/Experimental/Avalanche/Source/AvalancheMediaEditor/Private/Rundown/Factories/Suggestions/AvaRundownFilterPathSuggestionFactory.cpp
--- a/Engine/Plugins/Experimental/Avalanche/Source/AvalancheMediaEditor/Private/Rundown/Factories/Suggestions/AvaRundownFilterPathSuggestionFactory.cpp
+++ b/Engine/Plugins/Experimental/Avalanche/Source/AvalancheMediaEditor/Private/Rundown/Factories/Suggestions/AvaRundownFilterPathSuggestionFactory.cpp
@@ -18,13 +18,13 @@ void FAvaRundownFilterPathSuggestionFactory::AddSuggestion(const TSharedRef<FAva
 		const FText PathCategoryLabel = LOCTEXT("PathCategoryLabel", "Ava-Rundown-Asset");
 		const FString AssetName = PageItem.GetAssetPath(InPayload->Rundown).GetAssetName();
 
-		FString BlueprintNameSuggestion = FString::Printf(TEXT("Asset=%s"), *AssetName);
-		const bool bIsFilterValueValid = InPayload->FilterValue.IsEmpty() || BlueprintNameSuggestion.Contains(InPayload->FilterValue);
+		FString AssetNameSuggestion = FString::Printf(TEXT("Asset=%s"), *AssetName);
+		const bool bIsFilterValueValid = InPayload->FilterValue.IsEmpty() || AssetNameSuggestion.Contains(InPayload->FilterValue);
 
-		if (bIsFilterValueValid && !InPayload->FilterCache.Contains(BlueprintNameSuggestion))
+		if (bIsFilterValueValid && !InPayload->FilterCache.Contains(AssetNameSuggestion))
 		{
-			InPayload->FilterCache.Add(BlueprintNameSuggestion);
-			InPayload->PossibleSuggestions.Add(FAssetSearchBoxSuggestion{MoveTemp(BlueprintNameSuggestion), FText::FromString(AssetName), PathCategoryLabel});
+			InPayload->FilterCache.Add(AssetNameSuggestion);
+			InPayload->PossibleSuggestions.Add(FAssetSearchBoxSuggestion{MoveTemp(AssetNameSuggestion), FText::FromString(AssetName), PathCategoryLabel});
 		}
 	}
 }
diff --git a/Engine/Plugins/Experimental/Avalanche/Source/AvalancheMediaEditor/Private/Rundown/Pages/Columns/AvaRundownPageAssetSelectorColumn.cpp b/Engine/Plugins/Experimental/Avalanche/Source/AvalancheMediaEditor/Private/Rundown/Pages/Columns/AvaRundownPageAssetSelectorColumn.cpp
--- a/Engine/Plugins/Experimental/Avalanche/Source/AvalancheMediaEditor/Private/Rundown/Pages/Columns/AvaRundownPageAssetSelectorColumn.cpp
+++ b/Engine/Plugins/Experimental/Avalanche/Source/AvalancheMediaEditor/Private/Rundown/Pages/Columns/AvaRundownPageAssetSelectorColumn.cpp
@@ -3,7 +3,7 @@
 #include "AvaRundownPageAssetSelectorColumn.h"
 
 #include "AssetRegistry/IAssetRegistry.h"
-#include "AvaBlueprint.h"
+#include "Engine/World.h"
 #include "Framework/Application/SlateApplication.h"
 #include "PropertyCustomizationHelpers.h"
 #include "Rundown/Pages/PageViews/IAvaRundownPageView.h"
@@ -42,7 +42,6 @@ namespace UE::AvaPageAssetSelectorColumn::Private
 
 		constexpr bool bAllowClear = true;
 		TArray<const UClass*> AllowedClasses;
-		AllowedClasses.Add(UAvalancheBlueprint::StaticClass());
 		AllowedClasses.Add(UWorld::StaticClass());
 
 		const FAssetData AssetData = IAssetRegistry::Get()->GetAssetByObjectPath(InPageView->GetObjectPath(Rundown));
diff --git a/Engine/Plugins/Experimental/Avalanche/Source/AvalancheMediaEditor/Private/Rundown/Pages/PageViews/AvaRundownPageViewImpl.cpp b/Engine/Plugins/Experimental/Avalanche/Source/AvalancheMediaEditor/Private/Rundown/Pages/PageViews/AvaRundownPageViewImpl.cpp
--- a/Engine/Plugins/Experimental/Avalanche/Source/AvalancheMediaEditor/Private/Rundown/Pages/PageViews/AvaRundownPageViewImpl.cpp
+++ b/Engine/Plugins/Experimental/Avalanche/Source/AvalancheMediaEditor/Private/Rundown/Pages/PageViews/AvaRundownPageViewImpl.cpp
@@ -135,7 +135,7 @@ void FAvaRundownPageViewImpl::OnObjectChanged(const FAssetData& InAssetData)
 		SetPageSelection(EAvaRundownPageViewSelectionChangeType::ReplaceSelection);
 	}
 
-	PerformWorkOnPages(LOCTEXT("SetAvaBlueprint", "Set Motion Design Blueprint"),
+	PerformWorkOnPages(LOCTEXT("UpdateAsset", "Update Motion Design Asset"),
 		[this, &InAssetData](FAvaRundownPage& InPage)->bool
 		{
 			if (!InPage.UpdateAsset(InAssetData.GetSoftObjectPath())) 
diff --git a/Engine/Plugins/Experimental/Avalanche/Source/AvalancheMediaEditor/Public/Playback/Graph/Nodes/AvaPlaybackEditorGraphNode_BlueprintPlayer.h b/Engine/Plugins/Experimental/Avalanche/Source/AvalancheMediaEditor/Public/Playback/Graph/Nodes/AvaPlaybackEditorGraphNode_BlueprintPlayer.h
deleted file mode 100644
diff --git a/Engine/Plugins/Experimental/Avalanche/Source/AvalancheMediaEditor/Public/Rundown/AvaRundownServer.h b/Engine/Plugins/Experimental/Avalanche/Source/AvalancheMediaEditor/Public/Rundown/AvaRundownServer.h
--- a/Engine/Plugins/Experimental/Avalanche/Source/AvalancheMediaEditor/Public/Rundown/AvaRundownServer.h
+++ b/Engine/Plugins/Experimental/Avalanche/Source/AvalancheMediaEditor/Public/Rundown/AvaRundownServer.h
@@ -197,7 +197,7 @@ private:
 	/** Cached data for page editing commands (GetPages and GetPageDetails). */
 	struct FRundownEditCommandData : public FRundownCache
 	{
-		/** PageId of the current managed ava blueprint. */
+		/** PageId of the current managed ava asset. */
 		int32 ManagedPageId = FAvaRundownPage::InvalidPageId;
 		TSharedPtr<FAvaRundownManagedInstance> ManagedInstance;
 
diff --git a/Engine/Plugins/Experimental/Avalanche/Source/AvalancheModifiers/Private/Extensions/AvaSceneTreeUpdateModifierExtension.cpp b/Engine/Plugins/Experimental/Avalanche/Source/AvalancheModifiers/Private/Extensions/AvaSceneTreeUpdateModifierExtension.cpp
--- a/Engine/Plugins/Experimental/Avalanche/Source/AvalancheModifiers/Private/Extensions/AvaSceneTreeUpdateModifierExtension.cpp
+++ b/Engine/Plugins/Experimental/Avalanche/Source/AvalancheModifiers/Private/Extensions/AvaSceneTreeUpdateModifierExtension.cpp
@@ -2,9 +2,10 @@
 
 #include "Extensions/AvaSceneTreeUpdateModifierExtension.h"
 #include "AvaActorUtils.h"
-#include "AvaBlueprint.h"
 #include "AvaSceneTree.h"
+#include "Containers/Ticker.h"
 #include "GameFramework/Actor.h"
+#include "IAvaSceneInterface.h"
 
 #if WITH_EDITOR
 #include "AvaOutlinerSubsystem.h"
diff --git a/Engine/Plugins/Experimental/Avalanche/Source/AvalancheModifiers/Private/Shared/AvaTranslucentPriorityModifierShared.cpp b/Engine/Plugins/Experimental/Avalanche/Source/AvalancheModifiers/Private/Shared/AvaTranslucentPriorityModifierShared.cpp
--- a/Engine/Plugins/Experimental/Avalanche/Source/AvalancheModifiers/Private/Shared/AvaTranslucentPriorityModifierShared.cpp
+++ b/Engine/Plugins/Experimental/Avalanche/Source/AvalancheModifiers/Private/Shared/AvaTranslucentPriorityModifierShared.cpp
@@ -2,10 +2,12 @@
 
 #include "Shared/AvaTranslucentPriorityModifierShared.h"
 #include "AvaActorUtils.h"
-#include "AvaBlueprint.h"
+#include "AvaSceneTree.h"
+#include "AvaSceneTreeNode.h"
 #include "Camera/CameraActor.h"
 #include "Camera/CameraComponent.h"
 #include "Framework/AvaGameInstance.h"
+#include "IAvaSceneInterface.h"
 
 #if WITH_EDITOR
 #include "AvaOutlinerUtils.h"
diff --git a/Engine/Plugins/Experimental/Avalanche/Source/AvalancheRemoteControl/Private/AvaRemoteControlRebind.cpp b/Engine/Plugins/Experimental/Avalanche/Source/AvalancheRemoteControl/Private/AvaRemoteControlRebind.cpp
--- a/Engine/Plugins/Experimental/Avalanche/Source/AvalancheRemoteControl/Private/AvaRemoteControlRebind.cpp
+++ b/Engine/Plugins/Experimental/Avalanche/Source/AvalancheRemoteControl/Private/AvaRemoteControlRebind.cpp
@@ -53,7 +53,7 @@ namespace UE::AvaRemoteControl::Private
 }
 
 /**
- *	Since the runtime Avalanche Blueprint is not in the world where the
+ *	Since the runtime Avalanche Scene Object is not in the world where the
  *	remote control bindings were done, it needs to be rebound.
  *	
  *	The original FRemoteControlPresetRebindingManager::Rebind_NewAlgo
diff --git a/Engine/Plugins/Experimental/Avalanche/Source/AvalancheRemoteControl/Private/AvaRemoteControlUtils.cpp b/Engine/Plugins/Experimental/Avalanche/Source/AvalancheRemoteControl/Private/AvaRemoteControlUtils.cpp
--- a/Engine/Plugins/Experimental/Avalanche/Source/AvalancheRemoteControl/Private/AvaRemoteControlUtils.cpp
+++ b/Engine/Plugins/Experimental/Avalanche/Source/AvalancheRemoteControl/Private/AvaRemoteControlUtils.cpp
@@ -67,7 +67,7 @@ bool FAvaRemoteControlUtils::RegisterRemoteControlPreset(URemoteControlPreset* I
 		RemoteControlComponentsSubsystem->RegisterPreset(InRemoteControlPreset);
 	}
 
-	// Avalanche's RCPs, either from blueprint or level asset, are considered "embedded" because
+	// Avalanche's RCPs, from an ava asset (e.g. level) are considered "embedded" because
 	// they are not an asset.
 	constexpr bool bReplaceExisting = true;
 	return RemoteControlModule.RegisterEmbeddedPreset(InRemoteControlPreset, bReplaceExisting);
diff --git a/Engine/Plugins/Experimental/Avalanche/Source/AvalancheTest/AvalancheTest.Build.cs b/Engine/Plugins/Experimental/Avalanche/Source/AvalancheTest/AvalancheTest.Build.cs
deleted file mode 100644
diff --git a/Engine/Plugins/Experimental/Avalanche/Source/AvalancheTest/Private/AvaScreenshotTestActor.cpp b/Engine/Plugins/Experimental/Avalanche/Source/AvalancheTest/Private/AvaScreenshotTestActor.cpp
deleted file mode 100644
diff --git a/Engine/Plugins/Experimental/Avalanche/Source/AvalancheTest/Private/AvaTestBlueprintFunctionLibrary.cpp b/Engine/Plugins/Experimental/Avalanche/Source/AvalancheTest/Private/AvaTestBlueprintFunctionLibrary.cpp
deleted file mode 100644
diff --git a/Engine/Plugins/Experimental/Avalanche/Source/AvalancheTest/Private/AvaTestModule.cpp b/Engine/Plugins/Experimental/Avalanche/Source/AvalancheTest/Private/AvaTestModule.cpp
deleted file mode 100644
diff --git a/Engine/Plugins/Experimental/Avalanche/Source/AvalancheTest/Private/AvaTestModule.h b/Engine/Plugins/Experimental/Avalanche/Source/AvalancheTest/Private/AvaTestModule.h
deleted file mode 100644
diff --git a/Engine/Plugins/Experimental/Avalanche/Source/AvalancheTest/Public/AvaScreenshotTestActor.h b/Engine/Plugins/Experimental/Avalanche/Source/AvalancheTest/Public/AvaScreenshotTestActor.h
deleted file mode 100644
diff --git a/Engine/Plugins/Experimental/Avalanche/Source/AvalancheTest/Public/AvaTestBlueprintFunctionLibrary.h b/Engine/Plugins/Experimental/Avalanche/Source/AvalancheTest/Public/AvaTestBlueprintFunctionLibrary.h
deleted file mode 100644
