Motion Design Sequence:
- removed deprecated struct that was used back when the ava blueprint was being used.
- after UE-178582, the only reason this anim binding ref was used was to keep load support for older ava blueprints.
- now that ava blueprint is removed, this is no longer needed

diff --git a/Engine/Plugins/Experimental/Avalanche/Source/AvalancheSequence/Private/AvaSequence.cpp b/Engine/Plugins/Experimental/Avalanche/Source/AvalancheSequence/Private/AvaSequence.cpp
--- a/Engine/Plugins/Experimental/Avalanche/Source/AvalancheSequence/Private/AvaSequence.cpp
+++ b/Engine/Plugins/Experimental/Avalanche/Source/AvalancheSequence/Private/AvaSequence.cpp
@@ -511,30 +511,6 @@ UObject* UAvaSequence::GetParentObject(UObject* InObject) const
 void UAvaSequence::Serialize(FArchive& Ar)
 {
 	Ar.UsingCustomVersion(FAvaSequenceVersion::GUID);
-	const int32 Version = Ar.CustomVer(FAvaSequenceVersion::GUID);
-
-	// Temporarily add a Redirector for BindingReference->ObjectBindingReferences for older versions
-	// as the BindingReference here would refer to the previous name for UAvaSequence::ObjectBindingReferences
-	// rather than ULevelSequence::BindingReferences
-	if (Version < FAvaSequenceVersion::LevelSequence)
-	{
-		TArray<FCoreRedirect> Redirects;
-
-		constexpr const TCHAR* SourceString = TEXT("AvaSequence");
-
-		Redirects.Emplace(ECoreRedirectFlags::Type_Property
-			, TEXT("/Script/AvalancheSequence.AvaSequence.BindingReferences")
-			, TEXT("/Script/AvalancheSequence.AvaSequence.ObjectBindingReferences"));
-
-		FCoreRedirects::AddRedirectList(Redirects, SourceString);
-
-		Super::Serialize(Ar);
-
-		FCoreRedirects::RemoveRedirectList(Redirects, SourceString);
-
-		return;
-	}
-
 	Super::Serialize(Ar);
 }
 
@@ -621,33 +597,6 @@ int32 UAvaSequence::UpdateBindings(const FTopLevelAssetPath* InOldContext, const
 	return BindingsUpdatedCount;
 }
 
-void UAvaSequence::MigrateLegacyBindings(UObject* InPlaybackContext)
-{
-	if (!ensureAlways(IsValid(InPlaybackContext)))
-	{
-		return;	
-	}
-
-	const FAvaAnimBindingReferences LegacyBindings = MoveTemp(ObjectBindingReferences);
-
-	for (const TPair<FGuid, FAvaAnimBindingReferenceArray>& Binding : LegacyBindings.BindingIdToReferences)
-	{
-		for (const FAvaAnimBindingReference& BindingReference : Binding.Value.References)
-		{
-			if (UObject* const ResolvedObject = FindObject<UObject>(InPlaybackContext, *BindingReference.ObjectPath, false))
-			{
-				UObject* const ParentObject = GetParentObject(ResolvedObject);
-
-				UObject* Context = bParentContextsAreSignificant && IsValid(ParentObject)
-					? ParentObject
-					: InPlaybackContext;
-
-				BindPossessableObject(Binding.Key, *ResolvedObject, Context);
-			}
-		}
-	}
-}
-
 TArrayView<TWeakObjectPtr<>> UAvaSequence::FindObjectsFromGuid(const FGuid& InGuid)
 {
 	if (!InGuid.IsValid())
diff --git a/Engine/Plugins/Experimental/Avalanche/Source/AvalancheSequence/Public/AvaAnimBindingReference.h b/Engine/Plugins/Experimental/Avalanche/Source/AvalancheSequence/Public/AvaAnimBindingReference.h
deleted file mode 100644
diff --git a/Engine/Plugins/Experimental/Avalanche/Source/AvalancheSequence/Public/AvaSequence.h b/Engine/Plugins/Experimental/Avalanche/Source/AvalancheSequence/Public/AvaSequence.h
--- a/Engine/Plugins/Experimental/Avalanche/Source/AvalancheSequence/Public/AvaSequence.h
+++ b/Engine/Plugins/Experimental/Avalanche/Source/AvalancheSequence/Public/AvaSequence.h
@@ -2,7 +2,6 @@
 
 #pragma once
 
-#include "AvaAnimBindingReference.h"
 #include "AvaSequenceShared.h"
 #include "AvaTagHandle.h"
 #include "Containers/Array.h"
@@ -132,9 +131,6 @@ public:
 
 	int32 UpdateBindings(const FTopLevelAssetPath* InOldContext, const FTopLevelAssetPath& InNewContext);
 
-	/** Move the Ava Bindings to use the Level Sequence Bindings */
-	void MigrateLegacyBindings(UObject* InPlaybackContext);
-
 	/** Gets the Objects Bound for this Sequence. Potentially slow */
 	TArrayView<TWeakObjectPtr<>> FindObjectsFromGuid(const FGuid& InGuid);
 
@@ -169,10 +165,6 @@ protected:
 	UPROPERTY(EditAnywhere, EditFixedSize, Category = "Sequence Settings", meta=(NoResetToDefault, AllowPrivateAccess="true"))
 	TSet<FAvaMark> Marks;
 
-	/** Legacy Bindings. Should be read-only now. Not Deprecating so that duplication of older versions still copies this property */
-	UPROPERTY()
-	FAvaAnimBindingReferences ObjectBindingReferences;
-
 	/** Delegate called after the node has finished cleaning up itself its children */
 	FSimpleMulticastDelegate OnTreeNodeUpdated;
 };
