Motion Design MRQ:
- moved editor util functions to private as they're only used within the module
- other minor export fixes

diff --git a/Engine/Plugins/Experimental/Avalanche/Source/AvalancheMRQ/Public/AvaMRQRundownPageSetting.h b/Engine/Plugins/Experimental/Avalanche/Source/AvalancheMRQ/Public/AvaMRQRundownPageSetting.h
--- a/Engine/Plugins/Experimental/Avalanche/Source/AvalancheMRQ/Public/AvaMRQRundownPageSetting.h
+++ b/Engine/Plugins/Experimental/Avalanche/Source/AvalancheMRQ/Public/AvaMRQRundownPageSetting.h
@@ -20,8 +20,8 @@ struct FAvaMRQRundownPage
 	int32 PageId = 0;
 };
 
-UCLASS(DisplayName = "Rundown Page")
-class AVALANCHEMRQ_API UAvaMRQRundownPageSetting : public UMoviePipelineSetting
+UCLASS(MinimalAPI, DisplayName = "Rundown Page")
+class UAvaMRQRundownPageSetting : public UMoviePipelineSetting
 {
 	GENERATED_BODY()
 
@@ -31,7 +31,7 @@ public:
 
 protected:
 	//~ Begin UMoviePipelineSetting
-	virtual void SetupForPipelineImpl(UMoviePipeline* InPipeline) override;
+	AVALANCHEMRQ_API virtual void SetupForPipelineImpl(UMoviePipeline* InPipeline) override;
 	virtual bool IsValidOnShots() const override { return false; }
 	virtual bool IsValidOnPrimary() const override { return true; }
 	//~ End UMoviePipelineSetting
diff --git a/Engine/Plugins/Experimental/Avalanche/Source/AvalancheMRQEditor/AvalancheMRQEditor.Build.cs b/Engine/Plugins/Experimental/Avalanche/Source/AvalancheMRQEditor/AvalancheMRQEditor.Build.cs
--- a/Engine/Plugins/Experimental/Avalanche/Source/AvalancheMRQEditor/AvalancheMRQEditor.Build.cs
+++ b/Engine/Plugins/Experimental/Avalanche/Source/AvalancheMRQEditor/AvalancheMRQEditor.Build.cs
@@ -8,23 +8,17 @@ public class AvalancheMRQEditor : ModuleRules
     {
         PCHUsage = ModuleRules.PCHUsageMode.UseExplicitOrSharedPCHs;
 
-        PublicDependencyModuleNames.AddRange(
-            new string[]
-            {
-                "Core",
-                "CoreUObject",
-                "AvalancheMedia",
-                "AvalancheMediaEditor",
-            }
-        );
-
         PrivateDependencyModuleNames.AddRange(
             new string[]
             {
                 "Avalanche",
                 "AvalancheMRQ",
+                "AvalancheMedia",
+                "AvalancheMediaEditor",
                 "AvalancheSequence",
                 "AvalancheSequencer",
+                "Core",
+                "CoreUObject",
                 "Engine",
                 "MovieRenderPipelineCore",
                 "MovieRenderPipelineEditor",
diff --git a/Engine/Plugins/Experimental/Avalanche/Source/AvalancheMRQEditor/Private/AvaMRQEditorRundownUtils.h b/Engine/Plugins/Experimental/Avalanche/Source/AvalancheMRQEditor/Private/AvaMRQEditorRundownUtils.h
new file mode 100644
--- /dev/null
+++ b/Engine/Plugins/Experimental/Avalanche/Source/AvalancheMRQEditor/Private/AvaMRQEditorRundownUtils.h
@@ -0,0 +1,13 @@
+// Copyright Epic Games, Inc. All Rights Reserved.
+
+#pragma once
+
+#include "Containers/ArrayView.h"
+#include "Templates/SharedPointerFwd.h"
+
+class FAvaRundownEditor;
+
+struct FAvaMRQEditorRundownUtils
+{
+	static void RenderSelectedPages(TConstArrayView<TWeakPtr<const FAvaRundownEditor>> InRundownEditors);
+};
diff --git a/Engine/Plugins/Experimental/Avalanche/Source/AvalancheMRQEditor/Private/AvaMRQEditorRundowntUtils.cpp b/Engine/Plugins/Experimental/Avalanche/Source/AvalancheMRQEditor/Private/AvaMRQEditorRundowntUtils.cpp
--- a/Engine/Plugins/Experimental/Avalanche/Source/AvalancheMRQEditor/Private/AvaMRQEditorRundowntUtils.cpp
+++ b/Engine/Plugins/Experimental/Avalanche/Source/AvalancheMRQEditor/Private/AvaMRQEditorRundowntUtils.cpp
@@ -257,32 +257,3 @@ void FAvaMRQEditorRundownUtils::RenderSelectedPages(TConstArrayView<TWeakPtr<con
 		Private::RenderPages(ScopedRender, *Rundown, PageIds);
 	}
 }
-
-void FAvaMRQEditorRundownUtils::RenderPages(const UAvaRundown& InRundown, TConstArrayView<int32> InPageIds)
-{
-	if (InPageIds.IsEmpty())
-	{
-		return;
-	}
-
-	using namespace UE::AvaMRQEditor;
-	Private::FAvaMRQScopedRender ScopedRender;
-	if (!ScopedRender.IsValid())
-	{
-		return;
-	}
-
-	Private::RenderPages(ScopedRender, InRundown, InPageIds);
-}
-
-void FAvaMRQEditorRundownUtils::RenderPage(const UAvaRundown& InRundown, const FAvaRundownPage& InPage)
-{
-	using namespace UE::AvaMRQEditor;
-	Private::FAvaMRQScopedRender ScopedRender;
-	if (!ScopedRender.IsValid())
-	{
-		return;
-	}
-
-	Private::RenderPage(ScopedRender, InRundown, InPage);
-}
diff --git a/Engine/Plugins/Experimental/Avalanche/Source/AvalancheMRQEditor/Public/AvaMRQEditorRundownUtils.h b/Engine/Plugins/Experimental/Avalanche/Source/AvalancheMRQEditor/Public/AvaMRQEditorRundownUtils.h
deleted file mode 100644
