MotionDesign
- Moving cloner/effector system into its own plugin called ClonerEffector (CE) displayed "Cloners and Effectors"
- Exporting only user public functions instead of entire classes
- Moving uassets through editor instead of P4 to avoid breaking assets
- Allowing to specify a custom actor resolver to gather ordered children that MotionDesign plugin can use with outliner
- Moving icons, details customization for editor too

#changelist validated
#virtualized

diff --git a/Engine/Plugins/Experimental/Avalanche/Avalanche.uplugin b/Engine/Plugins/Experimental/Avalanche/Avalanche.uplugin
--- a/Engine/Plugins/Experimental/Avalanche/Avalanche.uplugin
+++ b/Engine/Plugins/Experimental/Avalanche/Avalanche.uplugin
@@ -296,6 +296,10 @@
 		{
 			"Name": "GeometryMask",
 			"Enabled": true
+		},
+		{
+			"Name": "ClonerEffector",
+			"Enabled": true
 		}
 	]
 }
diff --git a/Engine/Plugins/Experimental/Avalanche/Config/DefaultAvalanche.ini b/Engine/Plugins/Experimental/Avalanche/Config/DefaultAvalanche.ini
--- a/Engine/Plugins/Experimental/Avalanche/Config/DefaultAvalanche.ini
+++ b/Engine/Plugins/Experimental/Avalanche/Config/DefaultAvalanche.ini
@@ -25,13 +25,6 @@
 +ClassRedirects=(OldName="/Script/AvalancheModifiersEditor.AvaContentBackgroundModifierExtension",NewName="/Script/AvalancheModifiersEditor.AvaAutoSizeModifierExtension")
 +ClassRedirects=(OldName="/Script/Avalanche.AvaViewportCameraActor",NewName="/Script/Engine.CameraActor")
 
-+ClassRedirects=(OldName="/Script/Avalanche.AvaClonerActor",NewName="/Script/AvalancheEffectors.AvaClonerActor")
-+ClassRedirects=(OldName="/Script/Avalanche.AvaClonerComponent",NewName="/Script/AvalancheEffectors.AvaClonerComponent")
-+ClassRedirects=(OldName="/Script/Avalanche.AvaEffectorActor",NewName="/Script/AvalancheEffectors.AvaEffectorActor")
-+ClassRedirects=(OldName="/Script/Avalanche.AvaEffectorComponent",NewName="/Script/AvalancheEffectors.AvaEffectorComponent")
-+ClassRedirects=(OldName="/Script/Avalanche.AvaSphereComponent",NewName="/Script/AvalancheEffectors.AvaSphereComponent")
-+ClassRedirects=(OldName="/Script/AvalancheEffectors.AvaSphereComponent",NewName="/Script/Engine.SphereComponent")
-+ClassRedirects=(OldName="/Script/AvalancheEffectors.AvaClonerSphereLayout",NewName="/Script/AvalancheEffectors.AvaClonerSphereUniformLayout")
 +ClassRedirects=(OldName="/Script/Avalanche.MovieSceneAvaToolboxRectCornerSection",NewName="/Script/AvalancheShapes.MovieSceneAvaShapeRectCornerSection")
 +ClassRedirects=(OldName="/Script/Avalanche.MovieSceneAvaToolboxRectCornerTrack",NewName="/Script/AvalancheShapes.MovieSceneAvaShapeRectCornerTrack")
 +ClassRedirects=(OldName="/Script/Avalanche.AvaToolboxShapeActor",NewName="/Script/AvalancheShapes.AvaShapeActor")
@@ -67,21 +60,6 @@
 +ClassRedirects=(OldName="/Script/AvalancheModifiers.AvaCloneModifier",NewName="/Script/AvalancheModifiers.AvaPatternModifier")
 
 +StructRedirects=(OldName="/Script/Avalanche.AvaGuideInfo",NewName="/Script/Avalanche.AvaViewportGuideInfo")
-+StructRedirects=(OldName="/Script/Avalanche.AvaClonerAttachmentItem",NewName="/Script/AvalancheEffectors.AvaClonerAttachmentItem")
-+StructRedirects=(OldName="/Script/Avalanche.AvaClonerAttachmentTree",NewName="/Script/AvalancheEffectors.AvaClonerAttachmentTree")
-+StructRedirects=(OldName="/Script/Avalanche.AvaClonerEffectorDataInterfaces",NewName="/Script/AvalancheEffectors.AvaClonerEffectorDataInterfaces")
-+StructRedirects=(OldName="/Script/Avalanche.AvaClonerSystem",NewName="/Script/AvalancheEffectors.AvaClonerSystem")
-+StructRedirects=(OldName="/Script/Avalanche.AvaClonerSampleMeshOptions",NewName="/Script/AvalancheEffectors.AvaClonerSampleMeshOptions")
-+StructRedirects=(OldName="/Script/Avalanche.AvaClonerSampleSplineOptions",NewName="/Script/AvalancheEffectors.AvaClonerSampleSplineOptions")
-+StructRedirects=(OldName="/Script/Avalanche.AvaClonerGridConstraintSphere",NewName="/Script/AvalancheEffectors.AvaClonerGridConstraintSphere")
-+StructRedirects=(OldName="/Script/Avalanche.AvaClonerGridConstraintCylinder",NewName="/Script/AvalancheEffectors.AvaClonerGridConstraintCylinder")
-+StructRedirects=(OldName="/Script/Avalanche.AvaClonerGridConstraintTexture",NewName="/Script/AvalancheEffectors.AvaClonerGridConstraintTexture")
-+StructRedirects=(OldName="/Script/Avalanche.AvaClonerGridLayoutOptions",NewName="/Script/AvalancheEffectors.AvaClonerGridLayoutOptions")
-+StructRedirects=(OldName="/Script/Avalanche.AvaClonerLineLayoutOptions",NewName="/Script/AvalancheEffectors.AvaClonerLineLayoutOptions")
-+StructRedirects=(OldName="/Script/Avalanche.AvaClonerCircleLayoutOptions",NewName="/Script/AvalancheEffectors.AvaClonerCircleLayoutOptions")
-+StructRedirects=(OldName="/Script/Avalanche.AvaClonerCylinderLayoutOptions",NewName="/Script/AvalancheEffectors.AvaClonerCylinderLayoutOptions")
-+StructRedirects=(OldName="/Script/Avalanche.AvaClonerSphereLayoutOptions",NewName="/Script/AvalancheEffectors.AvaClonerSphereLayoutOptions")
-+StructRedirects=(OldName="/Script/Avalanche.AvaClonerHoneycombLayoutOptions",NewName="/Script/AvalancheEffectors.AvaClonerHoneycombLayoutOptions")
 +StructRedirects=(OldName="/Script/Avalanche.MovieSceneAvaToolboxRectCornerSectionTemplate",NewName="/Script/AvalancheShapes.MovieSceneAvaShapeRectCornerSectionTemplate")
 +StructRedirects=(OldName="/Script/Avalanche.AvaToolboxMaskMaterial",NewName="/Script/AvalancheShapes.AvaShapeMaskMaterial") -- ?!!
 +StructRedirects=(OldName="/Script/Avalanche.AvaToolboxMesh",NewName="/Script/AvalancheShapes.AvaShapeMesh")
@@ -108,19 +86,6 @@
 +StructRedirects=(OldName="/Script/AvalancheModifiers.AvaCloneModifierLineLayoutOptions",NewName="/Script/AvalancheModifiers.AvaPatternModifierLineLayoutOptions")
 
 +EnumRedirects=(OldName="/Script/Avalanche.EAvaToolboxGuideState",NewName="/Script/Avalanche.EAvaViewportGuideState")
-+EnumRedirects=(OldName="/Script/Avalanche.EAvaClonerLayout",NewName="/Script/AvalancheEffectors.EAvaClonerLayout")
-+EnumRedirects=(OldName="/Script/Avalanche.EAvaClonerAxis",NewName="/Script/AvalancheEffectors.EAvaClonerAxis")
-+EnumRedirects=(OldName="/Script/Avalanche.EAvaClonerPlane",NewName="/Script/AvalancheEffectors.EAvaClonerPlane")
-+EnumRedirects=(OldName="/Script/Avalanche.EAvaClonerMeshRenderMode",NewName="/Script/AvalancheEffectors.EAvaClonerMeshRenderMode")
-+EnumRedirects=(OldName="/Script/Avalanche.EAvaClonerGridConstraint",NewName="/Script/AvalancheEffectors.EAvaClonerGridConstraint")
-+EnumRedirects=(OldName="/Script/Avalanche.EAvaClonerEasing",NewName="/Script/AvalancheEffectors.EAvaClonerEasing")
-+EnumRedirects=(OldName="/Script/Avalanche.EAvaClonerMeshAsset",NewName="/Script/AvalancheEffectors.EAvaClonerMeshAsset")
-+EnumRedirects=(OldName="/Script/Avalanche.EAvaClonerMeshSampleData",NewName="/Script/AvalancheEffectors.EAvaClonerMeshSampleData")
-+EnumRedirects=(OldName="/Script/Avalanche.EAvaClonerEffectorType",NewName="/Script/AvalancheEffectors.EAvaClonerEffectorType")
-+EnumRedirects=(OldName="/Script/Avalanche.EAvaClonerEffectorMode",NewName="/Script/AvalancheEffectors.EAvaClonerEffectorMode")
-+EnumRedirects=(OldName="/Script/Avalanche.EAvaClonerAttachmentStatus",NewName="/Script/AvalancheEffectors.EAvaClonerAttachmentStatus")
-+EnumRedirects=(OldName="/Script/Avalanche.EAvaClonerTextureSampleChannel",NewName="/Script/AvalancheEffectors.EAvaClonerTextureSampleChannel")
-+EnumRedirects=(OldName="/Script/Avalanche.EAvaClonerCompareMode",NewName="/Script/AvalancheEffectors.EAvaClonerCompareMode")
 +EnumRedirects=(OldName="/Script/Avalanche.EAvaToolboxUVMode",NewName="/Script/AvalancheShapes.EAvaShapeUVMode")
 +EnumRedirects=(OldName="/Script/Avalanche.EAvaToolboxCornerType",NewName="/Script/AvalancheShapes.EAvaShapeCornerType")
 +EnumRedirects=(OldName="/Script/Avalanche.ParametricMaterialStyle",NewName="/Script/AvalancheShapes.EAvaShapeParametricMaterialStyle",ValueChanges=(("SOLID","Solid"),("LINEAR_GRADIENT","LinearGradient")))
@@ -138,16 +103,6 @@
 
 +PackageRedirects=(OldName="/Script/AvalancheText3D",NewName="/Script/AvalancheText")
 +ClassRedirects=(OldName="/Script/AvalancheModifiers.AvaMaskModifier",NewName="/Script/AvalancheModifiers.AvaBooleanModifier")
-+PropertyRedirects=(OldName="/Script/AvalancheEffectors.AvaClonerActor.ForceOrientationRate",NewName="/Script/AvalancheEffectors.AvaClonerActor.OrientationForceRate")
-+PropertyRedirects=(OldName="/Script/AvalancheEffectors.AvaClonerActor.ForceOrientationMin",NewName="/Script/AvalancheEffectors.AvaClonerActor.OrientationForceMin")
-+PropertyRedirects=(OldName="/Script/AvalancheEffectors.AvaClonerActor.ForceOrientationMax",NewName="/Script/AvalancheEffectors.AvaClonerActor.OrientationForceMax")
-+PropertyRedirects=(OldName="/Script/AvalancheEffectors.AvaClonerActor.ForceVortexAmount",NewName="/Script/AvalancheEffectors.AvaClonerActor.VortexForceAmount")
-+PropertyRedirects=(OldName="/Script/AvalancheEffectors.AvaClonerActor.ForceVortexAxis",NewName="/Script/AvalancheEffectors.AvaClonerActor.VortexForceAxis")
-+PropertyRedirects=(OldName="/Script/AvalancheEffectors.AvaClonerActor.ForceCurlNoiseStrength",NewName="/Script/AvalancheEffectors.AvaClonerActor.CurlNoiseForceStrength")
-+PropertyRedirects=(OldName="/Script/AvalancheEffectors.AvaClonerActor.ForceCurlNoiseFrequency",NewName="/Script/AvalancheEffectors.AvaClonerActor.CurlNoiseForceFrequency")
-+PropertyRedirects=(OldName="/Script/AvalancheEffectors.AvaClonerActor.StepRotationDelta",NewName="/Script/AvalancheEffectors.AvaClonerActor.DeltaStepRotation")
-+PropertyRedirects=(OldName="/Script/AvalancheEffectors.AvaClonerActor.StepScaleDelta",NewName="/Script/AvalancheEffectors.AvaClonerActor.DeltaStepScale")
-+PropertyRedirects=(OldName="/Script/AvalancheEffectors.AvaEffectorActor.Strength",NewName="/Script/AvalancheEffectors.AvaEffectorActor.LocationStrength")
 
 +StructRedirects=(OldName="/Script/Avalanche.AvalancheFileVersionInfo",NewName="/Script/Avalanche.AvaFileVersionInfo")
 +StructRedirects=(OldName="/Script/Avalanche.AvalancheEngineVersionInfo",NewName="/Script/Avalanche.AvaEngineVersionInfo")
diff --git a/Engine/Plugins/Experimental/Avalanche/Source/AvalancheEditor/Private/AvaEditorStyle.cpp b/Engine/Plugins/Experimental/Avalanche/Source/AvalancheEditor/Private/AvaEditorStyle.cpp
--- a/Engine/Plugins/Experimental/Avalanche/Source/AvalancheEditor/Private/AvaEditorStyle.cpp
+++ b/Engine/Plugins/Experimental/Avalanche/Source/AvalancheEditor/Private/AvaEditorStyle.cpp
@@ -1,7 +1,6 @@
 // Copyright Epic Games, Inc. All Rights Reserved.
 
 #include "AvaEditorStyle.h"
-#include "AvaClonerEffectorShared.h"
 #include "Brushes/SlateImageBrush.h"
 #include "Engine/Texture2D.h"
 #include "Framework/Application/SlateApplication.h"
@@ -210,20 +209,6 @@ void FAvaEditorStyle::Init()
 	Set("AvalancheEditor.RectLightTool",        new CORE_IMAGE_BRUSH_SVG("Starship/AssetIcons/RectLight_16",        Icon16));
 	Set("AvalancheEditor.SpotLightTool",        new CORE_IMAGE_BRUSH_SVG("Starship/AssetIcons/SpotLight_16",        Icon16));
 	Set("AvalancheEditor.SkyLightTool",         new CORE_IMAGE_BRUSH_SVG("Starship/AssetIcons/SkyLight_16",         Icon16));
-
-	// Easing
-	if (const UEnum* EasingEnum = StaticEnum<EAvaClonerEasing>())
-	{
-		for (int32 Idx = 0; Idx < EasingEnum->GetMaxEnumValue(); Idx++)
-		{
-			FText ValueText;
-			EasingEnum->GetDisplayValueAsText(static_cast<EAvaClonerEasing>(Idx), ValueText);
-			const FString EasingString = ValueText.ToString().Replace(TEXT(" "), TEXT(""));
-			const FName StyleName(TEXT("AvalancheIcons.Easing.") + EasingString);
-			
-			Set(StyleName, new IMAGE_BRUSH_SVG("Icons/ClonerIcons/" + EasingString, Icon32));
-		}
-	}
 	
 	// Colors
 	Set("AvalancheEditor.PalettesTab.ExpanderHeader", new FSlateColorBrush(FStyleColors::Header));
diff --git a/Engine/Plugins/Experimental/Avalanche/Source/AvalancheEffectors/AvalancheEffectors.Build.cs b/Engine/Plugins/Experimental/Avalanche/Source/AvalancheEffectors/AvalancheEffectors.Build.cs
--- a/Engine/Plugins/Experimental/Avalanche/Source/AvalancheEffectors/AvalancheEffectors.Build.cs
+++ b/Engine/Plugins/Experimental/Avalanche/Source/AvalancheEffectors/AvalancheEffectors.Build.cs
@@ -11,26 +11,19 @@ public class AvalancheEffectors : ModuleRules
 		PublicDependencyModuleNames.AddRange(
 			new string[]
 			{
-				"AvalancheCore",
 				"Core",
 				"CoreUObject",
-				"Engine",
-				"Niagara",
-				"NiagaraCore"
+				"Engine"
 			}
 		);
 
 		PrivateDependencyModuleNames.AddRange(
 			new string[]
 			{
+				"AvalancheCore",
 				"Avalanche",
 				"AvalancheSceneTree",
-				"GeometryCore",
-				"GeometryFramework",
-				"GeometryScriptingCore",
-				"DynamicMesh",
-				"ProceduralMeshComponent",
-				"StaticMeshDescription"
+				"ClonerEffector"
 			}
 		);
 
@@ -39,9 +32,7 @@ public class AvalancheEffectors : ModuleRules
 			PrivateDependencyModuleNames.AddRange(new string[]
 			{
 				"AvalancheOutliner",
-				"UnrealEd",
-				"Slate",
-				"SlateCore"
+				"UnrealEd"
 			});
 		}
 	}
diff --git a/Engine/Plugins/Experimental/Avalanche/Source/AvalancheEffectors/Private/AvalancheEffectorsModule.cpp b/Engine/Plugins/Experimental/Avalanche/Source/AvalancheEffectors/Private/AvalancheEffectorsModule.cpp
--- a/Engine/Plugins/Experimental/Avalanche/Source/AvalancheEffectors/Private/AvalancheEffectorsModule.cpp
+++ b/Engine/Plugins/Experimental/Avalanche/Source/AvalancheEffectors/Private/AvalancheEffectorsModule.cpp
@@ -1,5 +1,115 @@
 // Copyright Epic Games, Inc. All Rights Reserved.
 
-#include "Modules/ModuleManager.h"
+#include "AvalancheEffectorsModule.h"
 
-IMPLEMENT_MODULE(FDefaultModuleImpl, AvalancheEffectors)
+#include "AvaActorUtils.h"
+#include "AvaSceneTree.h"
+#include "IAvaSceneInterface.h"
+#include "Engine/World.h"
+#include "GameFramework/Actor.h"
+#include "Subsystems/CEClonerSubsystem.h"
+
+#if WITH_EDITOR
+#include "IAvaOutliner.h"
+#include "AvaOutlinerDefines.h"
+#include "AvaOutlinerSubsystem.h"
+#include "Item/AvaOutlinerActor.h"
+#endif
+
+void FAvalancheEffectorsModule::StartupModule()
+{
+	UCEClonerSubsystem::OnSubsystemInitialized().AddRaw(this, &FAvalancheEffectorsModule::RegisterCustomActorResolver);
+}
+
+void FAvalancheEffectorsModule::ShutdownModule()
+{
+	UCEClonerSubsystem::OnSubsystemInitialized().RemoveAll(this);
+
+	UnregisterCustomActorResolver();
+}
+
+void FAvalancheEffectorsModule::RegisterCustomActorResolver()
+{
+	if (UCEClonerSubsystem* Subsystem = UCEClonerSubsystem::Get())
+	{
+		Subsystem->RegisterCustomActorResolver(UCEClonerSubsystem::FOnGetOrderedActors::CreateRaw(this, &FAvalancheEffectorsModule::GetOrderedChildrenActors));
+	}
+}
+
+void FAvalancheEffectorsModule::UnregisterCustomActorResolver()
+{
+	if (UObjectInitialized())
+	{
+		if (UCEClonerSubsystem* Subsystem = UCEClonerSubsystem::Get())
+		{
+			Subsystem->UnregisterCustomActorResolver();
+		}
+	}
+}
+
+TArray<AActor*> FAvalancheEffectorsModule::GetOrderedChildrenActors(const AActor* InParentActor)
+{
+	TArray<AActor*> ChildrenActors;
+
+	if (!InParentActor)
+	{
+		return ChildrenActors;
+	}
+
+	UWorld* ClonerWorld = InParentActor->GetTypedOuter<UWorld>();
+	if (!ClonerWorld)
+	{
+		return ChildrenActors;
+	}
+
+#if WITH_EDITOR
+	if (const UAvaOutlinerSubsystem* const OutlinerSubsystem = ClonerWorld->GetSubsystem<UAvaOutlinerSubsystem>())
+	{
+		if (const TSharedPtr<IAvaOutliner> AvaOutliner = OutlinerSubsystem->GetOutliner())
+		{
+			const FAvaOutlinerItemPtr OutlinerClonerItem = AvaOutliner->FindItem(InParentActor);
+			if (OutlinerClonerItem.IsValid())
+			{
+				for (FAvaOutlinerItemPtr OutlinerChild : OutlinerClonerItem->GetChildren())
+				{
+					if (!OutlinerChild.IsValid() || !OutlinerChild->IsA<FAvaOutlinerActor>())
+					{
+						continue;
+					}
+
+					if (const FAvaOutlinerActor* OutlinerActor = OutlinerChild->CastTo<FAvaOutlinerActor>())
+					{
+						ChildrenActors.Add(OutlinerActor->GetActor());
+					}
+				}
+			}
+
+			return ChildrenActors;
+		}
+	}
+#endif
+
+	if (const IAvaSceneInterface* SceneInterface = FAvaActorUtils::GetSceneInterfaceFromActor(InParentActor))
+	{
+		TArray<AActor*> AttachedActors;
+		InParentActor->GetAttachedActors(AttachedActors, true, false);
+		ChildrenActors.Reserve(AttachedActors.Num());
+		const FAvaSceneTree& SceneTree = SceneInterface->GetSceneTree();
+
+		for (AActor* ChildActor : AttachedActors)
+		{
+			if (const FAvaSceneTreeNode* ChildNode = SceneTree.FindTreeNode(FAvaSceneItem(ChildActor, ClonerWorld)))
+			{
+				ChildrenActors.Insert(ChildActor, ChildNode->GetLocalIndex());
+			}
+		}
+	}
+	else
+	{
+		InParentActor->GetAttachedActors(ChildrenActors, true, false);
+	}
+
+	return ChildrenActors;
+}
+
+IMPLEMENT_MODULE(FAvalancheEffectorsModule, AvalancheEffectorsModule)
diff --git a/Engine/Plugins/Experimental/Avalanche/Source/AvalancheEffectorsEditor/AvalancheEffectorsEditor.Build.cs b/Engine/Plugins/Experimental/Avalanche/Source/AvalancheEffectorsEditor/AvalancheEffectorsEditor.Build.cs
--- a/Engine/Plugins/Experimental/Avalanche/Source/AvalancheEffectorsEditor/AvalancheEffectorsEditor.Build.cs
+++ b/Engine/Plugins/Experimental/Avalanche/Source/AvalancheEffectorsEditor/AvalancheEffectorsEditor.Build.cs
@@ -28,6 +28,7 @@ public class AvalancheEffectorsEditor : ModuleRules
 				"AvalancheInteractiveTools",
 				"AvalancheShapes",
 				"AvalancheShapesEditor",
+				"ClonerEffector",
 				"ComponentVisualizers",
 				"InputCore",
 				"InteractiveToolsFramework",
diff --git a/Engine/Plugins/Experimental/Avalanche/Source/AvalancheEffectorsEditor/Private/AvaEffectorsEditorStyle.cpp b/Engine/Plugins/Experimental/Avalanche/Source/AvalancheEffectorsEditor/Private/AvaEffectorsEditorStyle.cpp
--- a/Engine/Plugins/Experimental/Avalanche/Source/AvalancheEffectorsEditor/Private/AvaEffectorsEditorStyle.cpp
+++ b/Engine/Plugins/Experimental/Avalanche/Source/AvalancheEffectorsEditor/Private/AvaEffectorsEditorStyle.cpp
@@ -1,7 +1,7 @@
 // Copyright Epic Games, Inc. All Rights Reserved.
 
 #include "AvaEffectorsEditorStyle.h"
-#include "AvaClonerEffectorShared.h"
+
 #include "Interfaces/IPluginManager.h"
 #include "Misc/Paths.h"
 #include "Styling/SlateStyleMacros.h"
@@ -10,15 +10,15 @@
 #include "Styling/SlateStyleRegistry.h"
 
 FAvaEffectorsEditorStyle::FAvaEffectorsEditorStyle()
-	: FSlateStyleSet(TEXT("AvalancheEffectorsEditor"))
+	: FSlateStyleSet(UE_MODULE_NAME)
 {
 	const FVector2f Icon20x20(20.0f, 20.0f);
 	const FVector2f Icon32x32(32.0f, 32.0f);
-	
-	const TSharedPtr<IPlugin> Plugin = IPluginManager::Get().FindPlugin(TEXT("Avalanche"));
-	
+
+	const TSharedPtr<IPlugin> Plugin = IPluginManager::Get().FindPlugin(UE_PLUGIN_NAME);
+
 	check(Plugin.IsValid());
-	
+
 	ContentRootDir = FPaths::Combine(Plugin->GetBaseDir(), TEXT("Resources"));
 
 	Set("AvalancheEffectorsEditor.Tool_Actor_Effector", new IMAGE_BRUSH_SVG("Icons/ToolboxIcons/effector", Icon20x20));
@@ -26,20 +26,6 @@ FAvaEffectorsEditorStyle::FAvaEffectorsEditorStyle()
 	Set("Tool_Actor_Effector", new IMAGE_BRUSH_SVG("Icons/ToolboxIcons/effector", Icon20x20));
 	Set("Tool_Actor_Cloner", new IMAGE_BRUSH_SVG("Icons/ToolboxIcons/cloner", Icon20x20));
 
-	// Easing
-	if (const UEnum* EasingEnum = StaticEnum<EAvaClonerEasing>())
-	{
-		for (int32 Idx = 0; Idx < EasingEnum->GetMaxEnumValue(); Idx++)
-		{
-			FText ValueText;
-			EasingEnum->GetDisplayValueAsText(static_cast<EAvaClonerEasing>(Idx), ValueText);
-			const FString EasingString = ValueText.ToString().Replace(TEXT(" "), TEXT(""));
-			const FName EasingName(TEXT("AvalancheIcons.Easing.") + EasingString);
-			
-			Set(EasingName, new IMAGE_BRUSH_SVG("Icons/ClonerIcons/" + EasingString, Icon32x32));
-		}
-	}
-	
 	FSlateStyleRegistry::RegisterSlateStyle(*this);
 }
 
diff --git a/Engine/Plugins/Experimental/Avalanche/Source/AvalancheEffectorsEditor/Private/AvalancheEffectorsEditorModule.cpp b/Engine/Plugins/Experimental/Avalanche/Source/AvalancheEffectorsEditor/Private/AvalancheEffectorsEditorModule.cpp
--- a/Engine/Plugins/Experimental/Avalanche/Source/AvalancheEffectorsEditor/Private/AvalancheEffectorsEditorModule.cpp
+++ b/Engine/Plugins/Experimental/Avalanche/Source/AvalancheEffectorsEditor/Private/AvalancheEffectorsEditorModule.cpp
@@ -5,34 +5,22 @@
 #include "AvaEffectorsEditorStyle.h"
 #include "AvaInteractiveToolsDelegates.h"
 #include "AvalancheShapesEditorModule.h"
-#include "Cloner/AvaClonerActor.h"
+#include "Cloner/CEClonerComponent.h"
 #include "Cloner/AvaClonerActorTool.h"
 #include "Cloner/AvaClonerActorVis.h"
-#include "Cloner/AvaClonerComponent.h"
-#include "Cloner/AvaClonerDetailCustomization.h"
 #include "ComponentVisualizers.h"
 #include "Effector/AvaEffectorActorTool.h"
 #include "Effector/AvaEffectorActorVis.h"
-#include "Effector/AvaEffectorComponent.h"
-#include "Effector/AvaEffectorDetailCustomization.h"
+#include "Effector/CEEffectorComponent.h"
 #include "Framework/Application/SlateApplication.h"
 #include "IAvalancheComponentVisualizersModule.h"
 #include "Modules/ModuleManager.h"
 #include "PropertyEditorModule.h"
 
-#define LOCTEXT_NAMESPACE "AvalancheEffectorsEditorModule"
-
-FAvalancheEffectorsEditorModule& FAvalancheEffectorsEditorModule::Get()
-{
-	static const FName ModuleName = TEXT("AvalancheEffectorsEditor");
-	return FModuleManager::LoadModuleChecked<FAvalancheEffectorsEditorModule>(ModuleName);
-}
-
 void FAvalancheEffectorsEditorModule::StartupModule()
 {
 	FCoreDelegates::OnPostEngineInit.AddRaw(this, &FAvalancheEffectorsEditorModule::PostEngineInit);
 
-	RegisterCustomLayouts();
 	FAvaEffectorsEditorStyle::Get();
 	FAvaEffectorsEditorCommands::Register();
 
@@ -43,13 +31,8 @@ void FAvalancheEffectorsEditorModule::ShutdownModule()
 {
 	FCoreDelegates::OnPostEngineInit.RemoveAll(this);
 
-	if (UObjectInitialized() && !IsEngineExitRequested())
-	{
-		UnregisterCustomLayouts();
-	}
-
 	FAvaEffectorsEditorCommands::Unregister();
-	
+
 	FAvaInteractiveToolsDelegates::GetRegisterToolsDelegate().RemoveAll(this);
 }
 
@@ -76,36 +59,8 @@ void FAvalancheEffectorsEditorModule::RegisterTools(IAvalancheInteractiveToolsMo
 
 void FAvalancheEffectorsEditorModule::RegisterComponentVisualizers()
 {
-	IAvalancheComponentVisualizersModule::RegisterComponentVisualizer<UAvaEffectorComponent, FAvaEffectorActorVisualizer>(&Visualizers);
-	IAvalancheComponentVisualizersModule::RegisterComponentVisualizer<UAvaClonerComponent, FAvaClonerActorVisualizer>(&Visualizers);
-}
-
-namespace UE::AvalancheEffectorsEditor::Private
-{
-	static const FName PropertyEditorModuleName("PropertyEditor");
+	IAvalancheComponentVisualizersModule::RegisterComponentVisualizer<UCEEffectorComponent, FAvaEffectorActorVisualizer>(&Visualizers);
+	IAvalancheComponentVisualizersModule::RegisterComponentVisualizer<UCEClonerComponent, FAvaClonerActorVisualizer>(&Visualizers);
 }
 
-void FAvalancheEffectorsEditorModule::RegisterCustomLayouts()
-{
-	FPropertyEditorModule& PropertyModule = FModuleManager::GetModuleChecked<FPropertyEditorModule>(UE::AvalancheEffectorsEditor::Private::PropertyEditorModuleName);
-
-	// Cloner/effector
-	PropertyModule.RegisterCustomClassLayout(AAvaEffectorActor::StaticClass()->GetFName(), FOnGetDetailCustomizationInstance::CreateStatic(&FAvaEffectorDetailCustomization::MakeInstance));
-	PropertyModule.RegisterCustomClassLayout(AAvaClonerActor::StaticClass()->GetFName(), FOnGetDetailCustomizationInstance::CreateStatic(&FAvaClonerDetailCustomization::MakeInstance));
-}
-
-void FAvalancheEffectorsEditorModule::UnregisterCustomLayouts()
-{
-	if (FModuleManager::Get().IsModuleLoaded(UE::AvalancheEffectorsEditor::Private::PropertyEditorModuleName))
-	{
-		FPropertyEditorModule& PropertyModule = FModuleManager::GetModuleChecked<FPropertyEditorModule>(UE::AvalancheEffectorsEditor::Private::PropertyEditorModuleName);
-
-		// Cloner/effector
-		PropertyModule.UnregisterCustomClassLayout(AAvaEffectorActor::StaticClass()->GetFName());
-		PropertyModule.UnregisterCustomClassLayout(AAvaClonerActor::StaticClass()->GetFName());
-	}
-}
-
-#undef LOCTEXT_NAMESPACE
-
 IMPLEMENT_MODULE(FAvalancheEffectorsEditorModule, AvalancheEffectorsEditorModule)
diff --git a/Engine/Plugins/Experimental/Avalanche/Source/AvalancheEffectorsEditor/Private/Cloner/AvaClonerActorTool.cpp b/Engine/Plugins/Experimental/Avalanche/Source/AvalancheEffectorsEditor/Private/Cloner/AvaClonerActorTool.cpp
--- a/Engine/Plugins/Experimental/Avalanche/Source/AvalancheEffectorsEditor/Private/Cloner/AvaClonerActorTool.cpp
+++ b/Engine/Plugins/Experimental/Avalanche/Source/AvalancheEffectorsEditor/Private/Cloner/AvaClonerActorTool.cpp
@@ -4,11 +4,11 @@
 #include "AvaEffectorsEditorCommands.h"
 #include "AvaInteractiveToolsSettings.h"
 #include "Builders/AvaInteractiveToolsToolBuilder.h"
-#include "Cloner/AvaClonerActor.h"
+#include "Cloner/CEClonerActor.h"
 
 UAvaClonerActorTool::UAvaClonerActorTool()
 {
-	ActorClass = AAvaClonerActor::StaticClass();
+	ActorClass = ACEClonerActor::StaticClass();
 }
 
 bool UAvaClonerActorTool::UseIdentityRotation() const
diff --git a/Engine/Plugins/Experimental/Avalanche/Source/AvalancheEffectorsEditor/Private/Cloner/AvaClonerActorVis.cpp b/Engine/Plugins/Experimental/Avalanche/Source/AvalancheEffectorsEditor/Private/Cloner/AvaClonerActorVis.cpp
--- a/Engine/Plugins/Experimental/Avalanche/Source/AvalancheEffectorsEditor/Private/Cloner/AvaClonerActorVis.cpp
+++ b/Engine/Plugins/Experimental/Avalanche/Source/AvalancheEffectorsEditor/Private/Cloner/AvaClonerActorVis.cpp
@@ -1,17 +1,17 @@
-// Copyright Epic Games, Inc. All Rights Reserved. 
+// Copyright Epic Games, Inc. All Rights Reserved.
 
 #include "Cloner/AvaClonerActorVis.h"
 
 #include "AvaField.h"
 #include "AvalancheShapesEditorModule.h"
-#include "Cloner/AvaClonerActor.h"
-#include "Cloner/AvaClonerComponent.h"
-#include "Cloner/Layouts/AvaClonerCircleLayout.h"
-#include "Cloner/Layouts/AvaClonerCylinderLayout.h"
-#include "Cloner/Layouts/AvaClonerGridLayout.h"
-#include "Cloner/Layouts/AvaClonerHoneycombLayout.h"
-#include "Cloner/Layouts/AvaClonerLineLayout.h"
-#include "Cloner/Layouts/AvaClonerSphereUniformLayout.h"
+#include "Cloner/CEClonerActor.h"
+#include "Cloner/CEClonerComponent.h"
+#include "Cloner/Layouts/CEClonerCircleLayout.h"
+#include "Cloner/Layouts/CEClonerCylinderLayout.h"
+#include "Cloner/Layouts/CEClonerGridLayout.h"
+#include "Cloner/Layouts/CEClonerHoneycombLayout.h"
+#include "Cloner/Layouts/CEClonerLineLayout.h"
+#include "Cloner/Layouts/CEClonerSphereUniformLayout.h"
 #include "EditorViewportClient.h"
 #include "Engine/Texture2D.h"
 #include "Framework/Application/SlateApplication.h"
@@ -31,56 +31,56 @@ FAvaClonerActorVisualizer::FAvaClonerActorVisualizer()
 void FAvaClonerActorVisualizer::StoreInitialValues()
 {
 	Super::StoreInitialValues();
-	
+
 	if (GetEditedComponent() == nullptr)
 	{
 		return;
 	}
 
-	const AAvaClonerActor* ClonerActor = ClonerActorWeak.Get();
+	const ACEClonerActor* ClonerActor = ClonerActorWeak.Get();
 
 	if (!ClonerActor)
 	{
 		return;
 	}
 
-	if (const UAvaClonerGridLayout* GridLayout = ClonerActor->GetActiveLayout<UAvaClonerGridLayout>())
+	if (const UCEClonerGridLayout* GridLayout = ClonerActor->GetActiveLayout<UCEClonerGridLayout>())
 	{
 		InitialSpacing = FVector(GridLayout->GetSpacingX(), GridLayout->GetSpacingY(), GridLayout->GetSpacingZ());
 	}
-	else if (const UAvaClonerLineLayout* LineLayout = ClonerActor->GetActiveLayout<UAvaClonerLineLayout>())
+	else if (const UCEClonerLineLayout* LineLayout = ClonerActor->GetActiveLayout<UCEClonerLineLayout>())
 	{
 		InitialSpacing = FVector(LineLayout->GetSpacing());
 	}
-	else if (const UAvaClonerHoneycombLayout* HoneycombLayout = ClonerActor->GetActiveLayout<UAvaClonerHoneycombLayout>())
+	else if (const UCEClonerHoneycombLayout* HoneycombLayout = ClonerActor->GetActiveLayout<UCEClonerHoneycombLayout>())
 	{
-		const EAvaClonerPlane Plane = HoneycombLayout->GetPlane();
-		
+		const ECEClonerPlane Plane = HoneycombLayout->GetPlane();
+
 		FVector Spacing;
-		if (Plane == EAvaClonerPlane::XY)
+		if (Plane == ECEClonerPlane::XY)
 		{
 			Spacing = FVector(HoneycombLayout->GetWidthSpacing(), HoneycombLayout->GetHeightSpacing(), 0);
 		}
-		else if (Plane == EAvaClonerPlane::YZ)
+		else if (Plane == ECEClonerPlane::YZ)
 		{
 			Spacing = FVector(0, HoneycombLayout->GetWidthSpacing(), HoneycombLayout->GetHeightSpacing());
 		}
-		else if (Plane == EAvaClonerPlane::XZ)
+		else if (Plane == ECEClonerPlane::XZ)
 		{
 			Spacing = FVector(HoneycombLayout->GetWidthSpacing(), 0, HoneycombLayout->GetHeightSpacing());
 		}
-		
+
 		InitialSpacing = Spacing;
 	}
-	else if (const UAvaClonerCircleLayout* CircleLayout = ClonerActor->GetActiveLayout<UAvaClonerCircleLayout>())
+	else if (const UCEClonerCircleLayout* CircleLayout = ClonerActor->GetActiveLayout<UCEClonerCircleLayout>())
 	{
 		InitialSpacing = FVector(CircleLayout->GetRadius());
 	}
-	else if (const UAvaClonerCylinderLayout* CylinderLayout = ClonerActor->GetActiveLayout<UAvaClonerCylinderLayout>())
+	else if (const UCEClonerCylinderLayout* CylinderLayout = ClonerActor->GetActiveLayout<UCEClonerCylinderLayout>())
 	{
 		InitialSpacing = FVector(0, CylinderLayout->GetRadius(), CylinderLayout->GetHeight());
 	}
-	else if (const UAvaClonerSphereUniformLayout* SphereLayout = ClonerActor->GetActiveLayout<UAvaClonerSphereUniformLayout>())
+	else if (const UCEClonerSphereUniformLayout* SphereLayout = ClonerActor->GetActiveLayout<UCEClonerSphereUniformLayout>())
 	{
 		InitialSpacing = FVector(0, SphereLayout->GetRadius(), 0);
 	}
@@ -88,9 +88,9 @@ void FAvaClonerActorVisualizer::StoreInitialValues()
 
 FBox FAvaClonerActorVisualizer::GetComponentBounds(const UActorComponent* InComponent) const
 {
-	if (const UAvaClonerComponent* ClonerComponent = Cast<UAvaClonerComponent>(InComponent))
+	if (const UCEClonerComponent* ClonerComponent = Cast<UCEClonerComponent>(InComponent))
 	{
-		if (const AAvaClonerActor* ClonerActor = Cast<AAvaClonerActor>(ClonerComponent->GetOwner()))
+		if (const ACEClonerActor* ClonerActor = Cast<ACEClonerActor>(ClonerComponent->GetOwner()))
 		{
 			FVector Origin;
 			FVector Extent;
@@ -109,7 +109,7 @@ bool FAvaClonerActorVisualizer::HandleInputDeltaInternal(FEditorViewportClient*
 		return false;
 	}
 
-	const AAvaClonerActor* ClonerActor = ClonerActorWeak.Get();
+	const ACEClonerActor* ClonerActor = ClonerActorWeak.Get();
 
 	if (!ClonerActor)
 	{
@@ -118,35 +118,35 @@ bool FAvaClonerActorVisualizer::HandleInputDeltaInternal(FEditorViewportClient*
 	else if (GetViewportWidgetMode(InViewportClient) == UE::Widget::WM_Translate)
 	{
 		const EAxisList::Type AxisList = GetViewportWidgetAxisList(InViewportClient);
-			
-		if (UAvaClonerGridLayout* GridLayout = ClonerActor->GetActiveLayout<UAvaClonerGridLayout>())
+
+		if (UCEClonerGridLayout* GridLayout = ClonerActor->GetActiveLayout<UCEClonerGridLayout>())
 		{
 			if (AxisList & EAxisList::X)
 			{
 				const float SpacingX = InitialSpacing.X + InAccumulatedTranslation.X / FMath::Max(1, GridLayout->GetCountX() / 2.f);
 				GridLayout->SetSpacingX(SpacingX);
-				OnPropertyModified(GridLayout, GET_MEMBER_NAME_CHECKED(UAvaClonerGridLayout, SpacingX));
+				OnPropertyModified(GridLayout, GET_MEMBER_NAME_CHECKED(UCEClonerGridLayout, SpacingX));
 			}
 			else if (AxisList & EAxisList::Y)
 			{
 				const float SpacingY = InitialSpacing.Y + InAccumulatedTranslation.Y / FMath::Max(1, GridLayout->GetCountY() / 2.f);
 				GridLayout->SetSpacingY(SpacingY);
-				OnPropertyModified(GridLayout, GET_MEMBER_NAME_CHECKED(UAvaClonerGridLayout, SpacingY));
+				OnPropertyModified(GridLayout, GET_MEMBER_NAME_CHECKED(UCEClonerGridLayout, SpacingY));
 			}
 			else if (AxisList & EAxisList::Z)
 			{
 				const float SpacingZ = InitialSpacing.Z + InAccumulatedTranslation.Z / FMath::Max(1, GridLayout->GetCountZ() / 2.f);
 				GridLayout->SetSpacingZ(SpacingZ);
-				OnPropertyModified(GridLayout, GET_MEMBER_NAME_CHECKED(UAvaClonerGridLayout, SpacingZ));
+				OnPropertyModified(GridLayout, GET_MEMBER_NAME_CHECKED(UCEClonerGridLayout, SpacingZ));
 			}
-			
+
 			return true;
 		}
-		else if (UAvaClonerLineLayout* LineLayout = ClonerActor->GetActiveLayout<UAvaClonerLineLayout>())
+		else if (UCEClonerLineLayout* LineLayout = ClonerActor->GetActiveLayout<UCEClonerLineLayout>())
 		{
 			const int32 Count = LineLayout->GetCount();
 			float Spacing = LineLayout->GetSpacing();
-            				
+
 			if (AxisList & EAxisList::X)
 			{
 				Spacing = InitialSpacing.X + InAccumulatedTranslation.X / FMath::Max(1, Count);
@@ -159,48 +159,48 @@ bool FAvaClonerActorVisualizer::HandleInputDeltaInternal(FEditorViewportClient*
 			{
 				Spacing = InitialSpacing.Z + InAccumulatedTranslation.Z / FMath::Max(1, Count);
 			}
-			
+
 			LineLayout->SetSpacing(Spacing);
-			OnPropertyModified(LineLayout, GET_MEMBER_NAME_CHECKED(UAvaClonerLineLayout, Spacing));
+			OnPropertyModified(LineLayout, GET_MEMBER_NAME_CHECKED(UCEClonerLineLayout, Spacing));
 
 			return true;
 		}
-		else if (UAvaClonerHoneycombLayout* HoneycombLayout = ClonerActor->GetActiveLayout<UAvaClonerHoneycombLayout>())
+		else if (UCEClonerHoneycombLayout* HoneycombLayout = ClonerActor->GetActiveLayout<UCEClonerHoneycombLayout>())
 		{
-			const EAvaClonerPlane Plane = HoneycombLayout->GetPlane();
+			const ECEClonerPlane Plane = HoneycombLayout->GetPlane();
 			const int32 WidthCount = HoneycombLayout->GetWidthCount();
 			const int32 HeightCount = HoneycombLayout->GetHeightCount();
 			float WidthSpacing = HoneycombLayout->GetWidthSpacing();
 			float HeightSpacing = HoneycombLayout->GetHeightSpacing();
-			
-			if (Plane == EAvaClonerPlane::XY)
+
+			if (Plane == ECEClonerPlane::XY)
 			{
 				WidthSpacing = InitialSpacing.X + InAccumulatedTranslation.X / FMath::Max(1, WidthCount / 2.f);
 				HeightSpacing = InitialSpacing.Y + InAccumulatedTranslation.Y / FMath::Max(1, HeightCount / 2.f);
 			}
-			else if (Plane == EAvaClonerPlane::YZ)
+			else if (Plane == ECEClonerPlane::YZ)
 			{
 				WidthSpacing = InitialSpacing.Y + InAccumulatedTranslation.Y / FMath::Max(1, WidthCount / 2.f);
 				HeightSpacing = InitialSpacing.Z + InAccumulatedTranslation.Z / FMath::Max(1, HeightCount / 2.f);
 			}
-			else if (Plane == EAvaClonerPlane::XZ)
+			else if (Plane == ECEClonerPlane::XZ)
 			{
 				WidthSpacing = InitialSpacing.X + InAccumulatedTranslation.X / FMath::Max(1, WidthCount / 2.f);
 				HeightSpacing = InitialSpacing.Z + InAccumulatedTranslation.Z / FMath::Max(1, HeightCount / 2.f);
 			}
-			
+
 			HoneycombLayout->SetWidthSpacing(WidthSpacing);
 			HoneycombLayout->SetHeightSpacing(HeightSpacing);
 
-			OnPropertyModified(HoneycombLayout, GET_MEMBER_NAME_CHECKED(UAvaClonerHoneycombLayout, WidthSpacing));
-			OnPropertyModified(HoneycombLayout, GET_MEMBER_NAME_CHECKED(UAvaClonerHoneycombLayout, HeightSpacing));
+			OnPropertyModified(HoneycombLayout, GET_MEMBER_NAME_CHECKED(UCEClonerHoneycombLayout, WidthSpacing));
+			OnPropertyModified(HoneycombLayout, GET_MEMBER_NAME_CHECKED(UCEClonerHoneycombLayout, HeightSpacing));
 
 			return true;
 		}
-		else if (UAvaClonerCircleLayout* CircleLayout = ClonerActor->GetActiveLayout<UAvaClonerCircleLayout>())
+		else if (UCEClonerCircleLayout* CircleLayout = ClonerActor->GetActiveLayout<UCEClonerCircleLayout>())
 		{
 			float Radius = CircleLayout->GetRadius();
-			
+
 			if (AxisList & EAxisList::X)
 			{
 				Radius = InitialSpacing.X + InAccumulatedTranslation.X;
@@ -213,19 +213,19 @@ bool FAvaClonerActorVisualizer::HandleInputDeltaInternal(FEditorViewportClient*
 			{
 				Radius = InitialSpacing.Z + InAccumulatedTranslation.Z;
 			}
-			
+
 			CircleLayout->SetRadius(Radius);
-			OnPropertyModified(CircleLayout, GET_MEMBER_NAME_CHECKED(UAvaClonerCircleLayout, Radius));
+			OnPropertyModified(CircleLayout, GET_MEMBER_NAME_CHECKED(UCEClonerCircleLayout, Radius));
 
 			return true;
 		}
-		else if (UAvaClonerCylinderLayout* CylinderLayout = ClonerActor->GetActiveLayout<UAvaClonerCylinderLayout>())
+		else if (UCEClonerCylinderLayout* CylinderLayout = ClonerActor->GetActiveLayout<UCEClonerCylinderLayout>())
 		{
-			const EAvaClonerPlane Plane = CylinderLayout->GetPlane();
+			const ECEClonerPlane Plane = CylinderLayout->GetPlane();
 			float Radius = CylinderLayout->GetRadius();
 			float Height = CylinderLayout->GetHeight();
-			
-			if (Plane == EAvaClonerPlane::XY)
+
+			if (Plane == ECEClonerPlane::XY)
 			{
 				if (AxisList & EAxisList::Y)
 				{
@@ -236,7 +236,7 @@ bool FAvaClonerActorVisualizer::HandleInputDeltaInternal(FEditorViewportClient*
 					Height = InitialSpacing.Z + InAccumulatedTranslation.Z;
 				}
 			}
-			else if (Plane == EAvaClonerPlane::YZ)
+			else if (Plane == ECEClonerPlane::YZ)
 			{
 				if (AxisList & EAxisList::Y)
 				{
@@ -247,7 +247,7 @@ bool FAvaClonerActorVisualizer::HandleInputDeltaInternal(FEditorViewportClient*
 					Height = InitialSpacing.Z + InAccumulatedTranslation.X;
 				}
 			}
-			else if (Plane == EAvaClonerPlane::XZ)
+			else if (Plane == ECEClonerPlane::XZ)
 			{
 				if (AxisList & EAxisList::Z)
 				{
@@ -262,27 +262,27 @@ bool FAvaClonerActorVisualizer::HandleInputDeltaInternal(FEditorViewportClient*
 			CylinderLayout->SetRadius(Radius);
 			CylinderLayout->SetHeight(Height);
 
-			OnPropertyModified(CylinderLayout, GET_MEMBER_NAME_CHECKED(UAvaClonerCylinderLayout, Radius));
-			OnPropertyModified(CylinderLayout, GET_MEMBER_NAME_CHECKED(UAvaClonerCylinderLayout, Height));
+			OnPropertyModified(CylinderLayout, GET_MEMBER_NAME_CHECKED(UCEClonerCylinderLayout, Radius));
+			OnPropertyModified(CylinderLayout, GET_MEMBER_NAME_CHECKED(UCEClonerCylinderLayout, Height));
 
 			return true;
 		}
-		else if (UAvaClonerSphereUniformLayout* SphereLayout = ClonerActor->GetActiveLayout<UAvaClonerSphereUniformLayout>())
+		else if (UCEClonerSphereUniformLayout* SphereLayout = ClonerActor->GetActiveLayout<UCEClonerSphereUniformLayout>())
 		{
 			float Radius = SphereLayout->GetRadius();
-			
+
 			if (AxisList & EAxisList::Y)
 			{
 				Radius = InitialSpacing.Y + InAccumulatedTranslation.Y;
 			}
-			
+
 			SphereLayout->SetRadius(Radius);
-			OnPropertyModified(SphereLayout, GET_MEMBER_NAME_CHECKED(UAvaClonerSphereUniformLayout, Radius));
+			OnPropertyModified(SphereLayout, GET_MEMBER_NAME_CHECKED(UCEClonerSphereUniformLayout, Radius));
 
 			return true;
 		}
 	}
-	
+
 	return Super::HandleInputDeltaInternal(InViewportClient, InViewport, InAccumulatedTranslation, InAccumulatedRotation, InAccumulatedScale);
 }
 
@@ -290,41 +290,41 @@ void FAvaClonerActorVisualizer::DrawVisualizationEditing(const UActorComponent*
 {
 	Super::DrawVisualizationEditing(InComponent, InView, InPDI, InOutIconIndex);
 
-	const UAvaClonerComponent* ClonerComponent = Cast<UAvaClonerComponent>(InComponent);
+	const UCEClonerComponent* ClonerComponent = Cast<UCEClonerComponent>(InComponent);
 
 	if (!ClonerComponent)
 	{
 		return;
 	}
 
-	const AAvaClonerActor* ClonerActor = Cast<AAvaClonerActor>(ClonerComponent->GetOwner());
+	const ACEClonerActor* ClonerActor = Cast<ACEClonerActor>(ClonerComponent->GetOwner());
 
 	if (!ClonerActor || !ClonerActor->GetEnabled() || ClonerActor->GetMeshCount() == 0)
 	{
 		return;
 	}
 
-	if (const UAvaClonerGridLayout* GridLayout = ClonerActor->GetActiveLayout<UAvaClonerGridLayout>())
+	if (const UCEClonerGridLayout* GridLayout = ClonerActor->GetActiveLayout<UCEClonerGridLayout>())
 	{
 		if (GridLayout->GetCountX() > 0)
 		{
-			DrawSpacingButton(ClonerActor, InView, InPDI, InOutIconIndex, EAvaClonerAxis::X, FAvaVisualizerBase::Inactive);
+			DrawSpacingButton(ClonerActor, InView, InPDI, InOutIconIndex, ECEClonerAxis::X, FAvaVisualizerBase::Inactive);
 			InOutIconIndex++;
 		}
-		
+
 		if (GridLayout->GetCountY() > 0)
 		{
-			DrawSpacingButton(ClonerActor, InView, InPDI, InOutIconIndex, EAvaClonerAxis::Y, FAvaVisualizerBase::Inactive);
+			DrawSpacingButton(ClonerActor, InView, InPDI, InOutIconIndex, ECEClonerAxis::Y, FAvaVisualizerBase::Inactive);
 			InOutIconIndex++;
 		}
-		
+
 		if (GridLayout->GetCountZ() > 0)
 		{
-			DrawSpacingButton(ClonerActor, InView, InPDI, InOutIconIndex, EAvaClonerAxis::Z, FAvaVisualizerBase::Inactive);
+			DrawSpacingButton(ClonerActor, InView, InPDI, InOutIconIndex, ECEClonerAxis::Z, FAvaVisualizerBase::Inactive);
 			InOutIconIndex++;
 		}
 	}
-	else if (const UAvaClonerLineLayout* LineLayout = ClonerActor->GetActiveLayout<UAvaClonerLineLayout>())
+	else if (const UCEClonerLineLayout* LineLayout = ClonerActor->GetActiveLayout<UCEClonerLineLayout>())
 	{
 		if (LineLayout->GetCount() > 0)
 		{
@@ -332,81 +332,81 @@ void FAvaClonerActorVisualizer::DrawVisualizationEditing(const UActorComponent*
 			InOutIconIndex++;
 		}
 	}
-	else if (const UAvaClonerHoneycombLayout* HoneycombLayout = ClonerActor->GetActiveLayout<UAvaClonerHoneycombLayout>())
+	else if (const UCEClonerHoneycombLayout* HoneycombLayout = ClonerActor->GetActiveLayout<UCEClonerHoneycombLayout>())
 	{
-		const EAvaClonerPlane Plane = HoneycombLayout->GetPlane();
-		
-		if (Plane == EAvaClonerPlane::XY)
+		const ECEClonerPlane Plane = HoneycombLayout->GetPlane();
+
+		if (Plane == ECEClonerPlane::XY)
 		{
-			DrawSpacingButton(ClonerActor, InView, InPDI, InOutIconIndex, EAvaClonerAxis::X, FAvaVisualizerBase::Inactive);
+			DrawSpacingButton(ClonerActor, InView, InPDI, InOutIconIndex, ECEClonerAxis::X, FAvaVisualizerBase::Inactive);
 			InOutIconIndex++;
-			DrawSpacingButton(ClonerActor, InView, InPDI, InOutIconIndex, EAvaClonerAxis::Y, FAvaVisualizerBase::Inactive);
+			DrawSpacingButton(ClonerActor, InView, InPDI, InOutIconIndex, ECEClonerAxis::Y, FAvaVisualizerBase::Inactive);
 			InOutIconIndex++;
 		}
-		else if (Plane == EAvaClonerPlane::YZ)
+		else if (Plane == ECEClonerPlane::YZ)
 		{
-			DrawSpacingButton(ClonerActor, InView, InPDI, InOutIconIndex, EAvaClonerAxis::Y, FAvaVisualizerBase::Inactive);
+			DrawSpacingButton(ClonerActor, InView, InPDI, InOutIconIndex, ECEClonerAxis::Y, FAvaVisualizerBase::Inactive);
 			InOutIconIndex++;
-			DrawSpacingButton(ClonerActor, InView, InPDI, InOutIconIndex, EAvaClonerAxis::Z, FAvaVisualizerBase::Inactive);
+			DrawSpacingButton(ClonerActor, InView, InPDI, InOutIconIndex, ECEClonerAxis::Z, FAvaVisualizerBase::Inactive);
 			InOutIconIndex++;
 		}
-		else if (Plane == EAvaClonerPlane::XZ)
+		else if (Plane == ECEClonerPlane::XZ)
 		{
-			DrawSpacingButton(ClonerActor, InView, InPDI, InOutIconIndex, EAvaClonerAxis::X, FAvaVisualizerBase::Inactive);
+			DrawSpacingButton(ClonerActor, InView, InPDI, InOutIconIndex, ECEClonerAxis::X, FAvaVisualizerBase::Inactive);
 			InOutIconIndex++;
-			DrawSpacingButton(ClonerActor, InView, InPDI, InOutIconIndex, EAvaClonerAxis::Z, FAvaVisualizerBase::Inactive);
+			DrawSpacingButton(ClonerActor, InView, InPDI, InOutIconIndex, ECEClonerAxis::Z, FAvaVisualizerBase::Inactive);
 			InOutIconIndex++;
 		}
 	}
-	else if (const UAvaClonerCircleLayout* CircleLayout = ClonerActor->GetActiveLayout<UAvaClonerCircleLayout>())
+	else if (const UCEClonerCircleLayout* CircleLayout = ClonerActor->GetActiveLayout<UCEClonerCircleLayout>())
 	{
-		const EAvaClonerPlane Plane = CircleLayout->GetPlane();
+		const ECEClonerPlane Plane = CircleLayout->GetPlane();
 
-		if (Plane == EAvaClonerPlane::XY)
+		if (Plane == ECEClonerPlane::XY)
 		{
-			DrawSpacingButton(ClonerActor, InView, InPDI, InOutIconIndex, EAvaClonerAxis::Y, FAvaVisualizerBase::Inactive);
+			DrawSpacingButton(ClonerActor, InView, InPDI, InOutIconIndex, ECEClonerAxis::Y, FAvaVisualizerBase::Inactive);
 			InOutIconIndex++;
 		}
-		else if (Plane == EAvaClonerPlane::YZ)
+		else if (Plane == ECEClonerPlane::YZ)
 		{
-			DrawSpacingButton(ClonerActor, InView, InPDI, InOutIconIndex, EAvaClonerAxis::Y, FAvaVisualizerBase::Inactive);
+			DrawSpacingButton(ClonerActor, InView, InPDI, InOutIconIndex, ECEClonerAxis::Y, FAvaVisualizerBase::Inactive);
 			InOutIconIndex++;
 		}
-		else if (Plane == EAvaClonerPlane::XZ)
+		else if (Plane == ECEClonerPlane::XZ)
 		{
-			DrawSpacingButton(ClonerActor, InView, InPDI, InOutIconIndex, EAvaClonerAxis::Z, FAvaVisualizerBase::Inactive);
+			DrawSpacingButton(ClonerActor, InView, InPDI, InOutIconIndex, ECEClonerAxis::Z, FAvaVisualizerBase::Inactive);
 			InOutIconIndex++;
 		}
 	}
-	else if (const UAvaClonerCylinderLayout* CylinderLayout = ClonerActor->GetActiveLayout<UAvaClonerCylinderLayout>())
+	else if (const UCEClonerCylinderLayout* CylinderLayout = ClonerActor->GetActiveLayout<UCEClonerCylinderLayout>())
 	{
-		const EAvaClonerPlane Plane = CylinderLayout->GetPlane();
-		
-		if (Plane == EAvaClonerPlane::XY)
+		const ECEClonerPlane Plane = CylinderLayout->GetPlane();
+
+		if (Plane == ECEClonerPlane::XY)
 		{
-			DrawSpacingButton(ClonerActor, InView, InPDI, InOutIconIndex, EAvaClonerAxis::Y, FAvaVisualizerBase::Inactive);
+			DrawSpacingButton(ClonerActor, InView, InPDI, InOutIconIndex, ECEClonerAxis::Y, FAvaVisualizerBase::Inactive);
 			InOutIconIndex++;
-			DrawSpacingButton(ClonerActor, InView, InPDI, InOutIconIndex, EAvaClonerAxis::Z, FAvaVisualizerBase::Inactive);
+			DrawSpacingButton(ClonerActor, InView, InPDI, InOutIconIndex, ECEClonerAxis::Z, FAvaVisualizerBase::Inactive);
 			InOutIconIndex++;
 		}
-		else if (Plane == EAvaClonerPlane::YZ)
+		else if (Plane == ECEClonerPlane::YZ)
 		{
-			DrawSpacingButton(ClonerActor, InView, InPDI, InOutIconIndex, EAvaClonerAxis::Y, FAvaVisualizerBase::Inactive);
+			DrawSpacingButton(ClonerActor, InView, InPDI, InOutIconIndex, ECEClonerAxis::Y, FAvaVisualizerBase::Inactive);
 			InOutIconIndex++;
-			DrawSpacingButton(ClonerActor, InView, InPDI, InOutIconIndex, EAvaClonerAxis::X, FAvaVisualizerBase::Inactive);
+			DrawSpacingButton(ClonerActor, InView, InPDI, InOutIconIndex, ECEClonerAxis::X, FAvaVisualizerBase::Inactive);
 			InOutIconIndex++;
 		}
-		else if (Plane == EAvaClonerPlane::XZ)
+		else if (Plane == ECEClonerPlane::XZ)
 		{
-			DrawSpacingButton(ClonerActor, InView, InPDI, InOutIconIndex, EAvaClonerAxis::Z, FAvaVisualizerBase::Inactive);
+			DrawSpacingButton(ClonerActor, InView, InPDI, InOutIconIndex, ECEClonerAxis::Z, FAvaVisualizerBase::Inactive);
 			InOutIconIndex++;
-			DrawSpacingButton(ClonerActor, InView, InPDI, InOutIconIndex, EAvaClonerAxis::Y, FAvaVisualizerBase::Inactive);
+			DrawSpacingButton(ClonerActor, InView, InPDI, InOutIconIndex, ECEClonerAxis::Y, FAvaVisualizerBase::Inactive);
 			InOutIconIndex++;
 		}
 	}
-	else if (ClonerActor->IsActiveLayout<UAvaClonerSphereUniformLayout>())
+	else if (ClonerActor->IsActiveLayout<UCEClonerSphereUniformLayout>())
 	{
-		DrawSpacingButton(ClonerActor, InView, InPDI, InOutIconIndex, EAvaClonerAxis::Y, FAvaVisualizerBase::Inactive);
+		DrawSpacingButton(ClonerActor, InView, InPDI, InOutIconIndex, ECEClonerAxis::Y, FAvaVisualizerBase::Inactive);
 		InOutIconIndex++;
 	}
 }
@@ -415,41 +415,41 @@ void FAvaClonerActorVisualizer::DrawVisualizationNotEditing(const UActorComponen
 {
 	Super::DrawVisualizationNotEditing(InComponent, InView, InPDI, InOutIconIndex);
 
-	const UAvaClonerComponent* ClonerComponent = Cast<UAvaClonerComponent>(InComponent);
+	const UCEClonerComponent* ClonerComponent = Cast<UCEClonerComponent>(InComponent);
 
 	if (!ClonerComponent)
 	{
 		return;
 	}
 
-	const AAvaClonerActor* ClonerActor = Cast<AAvaClonerActor>(ClonerComponent->GetOwner());
+	const ACEClonerActor* ClonerActor = Cast<ACEClonerActor>(ClonerComponent->GetOwner());
 
 	if (!ClonerActor || !ClonerActor->GetEnabled() || ClonerActor->GetMeshCount() == 0)
 	{
 		return;
 	}
 
-	if (const UAvaClonerGridLayout* GridLayout = ClonerActor->GetActiveLayout<UAvaClonerGridLayout>())
+	if (const UCEClonerGridLayout* GridLayout = ClonerActor->GetActiveLayout<UCEClonerGridLayout>())
 	{
 		if (GridLayout->GetCountX() > 0)
 		{
-			DrawSpacingButton(ClonerActor, InView, InPDI, InOutIconIndex, EAvaClonerAxis::X, FAvaVisualizerBase::Inactive);
+			DrawSpacingButton(ClonerActor, InView, InPDI, InOutIconIndex, ECEClonerAxis::X, FAvaVisualizerBase::Inactive);
 			InOutIconIndex++;
 		}
-		
+
 		if (GridLayout->GetCountY() > 0)
 		{
-			DrawSpacingButton(ClonerActor, InView, InPDI, InOutIconIndex, EAvaClonerAxis::Y, FAvaVisualizerBase::Inactive);
+			DrawSpacingButton(ClonerActor, InView, InPDI, InOutIconIndex, ECEClonerAxis::Y, FAvaVisualizerBase::Inactive);
 			InOutIconIndex++;
 		}
-		
+
 		if (GridLayout->GetCountZ() > 0)
 		{
-			DrawSpacingButton(ClonerActor, InView, InPDI, InOutIconIndex, EAvaClonerAxis::Z, FAvaVisualizerBase::Inactive);
+			DrawSpacingButton(ClonerActor, InView, InPDI, InOutIconIndex, ECEClonerAxis::Z, FAvaVisualizerBase::Inactive);
 			InOutIconIndex++;
 		}
 	}
-	else if (const UAvaClonerLineLayout* LineLayout = ClonerActor->GetActiveLayout<UAvaClonerLineLayout>())
+	else if (const UCEClonerLineLayout* LineLayout = ClonerActor->GetActiveLayout<UCEClonerLineLayout>())
 	{
 		if (LineLayout->GetCount() > 0)
 		{
@@ -457,176 +457,176 @@ void FAvaClonerActorVisualizer::DrawVisualizationNotEditing(const UActorComponen
 			InOutIconIndex++;
 		}
 	}
-	else if (const UAvaClonerHoneycombLayout* HoneycombLayout = ClonerActor->GetActiveLayout<UAvaClonerHoneycombLayout>())
+	else if (const UCEClonerHoneycombLayout* HoneycombLayout = ClonerActor->GetActiveLayout<UCEClonerHoneycombLayout>())
 	{
-		const EAvaClonerPlane Plane = HoneycombLayout->GetPlane();
-		
-		if (Plane == EAvaClonerPlane::XY)
+		const ECEClonerPlane Plane = HoneycombLayout->GetPlane();
+
+		if (Plane == ECEClonerPlane::XY)
 		{
-			DrawSpacingButton(ClonerActor, InView, InPDI, InOutIconIndex, EAvaClonerAxis::X, FAvaVisualizerBase::Inactive);
+			DrawSpacingButton(ClonerActor, InView, InPDI, InOutIconIndex, ECEClonerAxis::X, FAvaVisualizerBase::Inactive);
 			InOutIconIndex++;
-			DrawSpacingButton(ClonerActor, InView, InPDI, InOutIconIndex, EAvaClonerAxis::Y, FAvaVisualizerBase::Inactive);
+			DrawSpacingButton(ClonerActor, InView, InPDI, InOutIconIndex, ECEClonerAxis::Y, FAvaVisualizerBase::Inactive);
 			InOutIconIndex++;
 		}
-		else if (Plane == EAvaClonerPlane::YZ)
+		else if (Plane == ECEClonerPlane::YZ)
 		{
-			DrawSpacingButton(ClonerActor, InView, InPDI, InOutIconIndex, EAvaClonerAxis::Y, FAvaVisualizerBase::Inactive);
+			DrawSpacingButton(ClonerActor, InView, InPDI, InOutIconIndex, ECEClonerAxis::Y, FAvaVisualizerBase::Inactive);
 			InOutIconIndex++;
-			DrawSpacingButton(ClonerActor, InView, InPDI, InOutIconIndex, EAvaClonerAxis::Z, FAvaVisualizerBase::Inactive);
+			DrawSpacingButton(ClonerActor, InView, InPDI, InOutIconIndex, ECEClonerAxis::Z, FAvaVisualizerBase::Inactive);
 			InOutIconIndex++;
 		}
-		else if (Plane == EAvaClonerPlane::XZ)
+		else if (Plane == ECEClonerPlane::XZ)
 		{
-			DrawSpacingButton(ClonerActor, InView, InPDI, InOutIconIndex, EAvaClonerAxis::X, FAvaVisualizerBase::Inactive);
+			DrawSpacingButton(ClonerActor, InView, InPDI, InOutIconIndex, ECEClonerAxis::X, FAvaVisualizerBase::Inactive);
 			InOutIconIndex++;
-			DrawSpacingButton(ClonerActor, InView, InPDI, InOutIconIndex, EAvaClonerAxis::Z, FAvaVisualizerBase::Inactive);
+			DrawSpacingButton(ClonerActor, InView, InPDI, InOutIconIndex, ECEClonerAxis::Z, FAvaVisualizerBase::Inactive);
 			InOutIconIndex++;
 		}
 	}
-	else if (const UAvaClonerCircleLayout* CircleLayout = ClonerActor->GetActiveLayout<UAvaClonerCircleLayout>())
+	else if (const UCEClonerCircleLayout* CircleLayout = ClonerActor->GetActiveLayout<UCEClonerCircleLayout>())
 	{
-		const EAvaClonerPlane Plane = CircleLayout->GetPlane();
-		
-		if (Plane == EAvaClonerPlane::XY)
+		const ECEClonerPlane Plane = CircleLayout->GetPlane();
+
+		if (Plane == ECEClonerPlane::XY)
 		{
-			DrawSpacingButton(ClonerActor, InView, InPDI, InOutIconIndex, EAvaClonerAxis::Y, FAvaVisualizerBase::Inactive);
+			DrawSpacingButton(ClonerActor, InView, InPDI, InOutIconIndex, ECEClonerAxis::Y, FAvaVisualizerBase::Inactive);
 			InOutIconIndex++;
 		}
-		else if (Plane == EAvaClonerPlane::YZ)
+		else if (Plane == ECEClonerPlane::YZ)
 		{
-			DrawSpacingButton(ClonerActor, InView, InPDI, InOutIconIndex, EAvaClonerAxis::Y, FAvaVisualizerBase::Inactive);
+			DrawSpacingButton(ClonerActor, InView, InPDI, InOutIconIndex, ECEClonerAxis::Y, FAvaVisualizerBase::Inactive);
 			InOutIconIndex++;
 		}
-		else if (Plane == EAvaClonerPlane::XZ)
+		else if (Plane == ECEClonerPlane::XZ)
 		{
-			DrawSpacingButton(ClonerActor, InView, InPDI, InOutIconIndex, EAvaClonerAxis::Z, FAvaVisualizerBase::Inactive);
+			DrawSpacingButton(ClonerActor, InView, InPDI, InOutIconIndex, ECEClonerAxis::Z, FAvaVisualizerBase::Inactive);
 			InOutIconIndex++;
 		}
 		else
 		{
-			DrawSpacingButton(ClonerActor, InView, InPDI, InOutIconIndex, EAvaClonerAxis::Custom, FAvaVisualizerBase::Inactive);
+			DrawSpacingButton(ClonerActor, InView, InPDI, InOutIconIndex, ECEClonerAxis::Custom, FAvaVisualizerBase::Inactive);
 			InOutIconIndex++;
 		}
 	}
-	else if (const UAvaClonerCylinderLayout* CylinderLayout = ClonerActor->GetActiveLayout<UAvaClonerCylinderLayout>())
+	else if (const UCEClonerCylinderLayout* CylinderLayout = ClonerActor->GetActiveLayout<UCEClonerCylinderLayout>())
 	{
-		const EAvaClonerPlane Plane = CylinderLayout->GetPlane();
-		
-		if (Plane == EAvaClonerPlane::XY)
+		const ECEClonerPlane Plane = CylinderLayout->GetPlane();
+
+		if (Plane == ECEClonerPlane::XY)
 		{
-			DrawSpacingButton(ClonerActor, InView, InPDI, InOutIconIndex, EAvaClonerAxis::Y, FAvaVisualizerBase::Inactive);
+			DrawSpacingButton(ClonerActor, InView, InPDI, InOutIconIndex, ECEClonerAxis::Y, FAvaVisualizerBase::Inactive);
 			InOutIconIndex++;
-			DrawSpacingButton(ClonerActor, InView, InPDI, InOutIconIndex, EAvaClonerAxis::Z, FAvaVisualizerBase::Inactive);
+			DrawSpacingButton(ClonerActor, InView, InPDI, InOutIconIndex, ECEClonerAxis::Z, FAvaVisualizerBase::Inactive);
 			InOutIconIndex++;
 		}
-		else if (Plane == EAvaClonerPlane::YZ)
+		else if (Plane == ECEClonerPlane::YZ)
 		{
-			DrawSpacingButton(ClonerActor, InView, InPDI, InOutIconIndex, EAvaClonerAxis::Y, FAvaVisualizerBase::Inactive);
+			DrawSpacingButton(ClonerActor, InView, InPDI, InOutIconIndex, ECEClonerAxis::Y, FAvaVisualizerBase::Inactive);
 			InOutIconIndex++;
-			DrawSpacingButton(ClonerActor, InView, InPDI, InOutIconIndex, EAvaClonerAxis::X, FAvaVisualizerBase::Inactive);
+			DrawSpacingButton(ClonerActor, InView, InPDI, InOutIconIndex, ECEClonerAxis::X, FAvaVisualizerBase::Inactive);
 			InOutIconIndex++;
 		}
-		else if (Plane == EAvaClonerPlane::XZ)
+		else if (Plane == ECEClonerPlane::XZ)
 		{
-			DrawSpacingButton(ClonerActor, InView, InPDI, InOutIconIndex, EAvaClonerAxis::Z, FAvaVisualizerBase::Inactive);
+			DrawSpacingButton(ClonerActor, InView, InPDI, InOutIconIndex, ECEClonerAxis::Z, FAvaVisualizerBase::Inactive);
 			InOutIconIndex++;
-			DrawSpacingButton(ClonerActor, InView, InPDI, InOutIconIndex, EAvaClonerAxis::Y, FAvaVisualizerBase::Inactive);
+			DrawSpacingButton(ClonerActor, InView, InPDI, InOutIconIndex, ECEClonerAxis::Y, FAvaVisualizerBase::Inactive);
 			InOutIconIndex++;
 		}
 	}
-	else if (ClonerActor->IsActiveLayout<UAvaClonerSphereUniformLayout>())
+	else if (ClonerActor->IsActiveLayout<UCEClonerSphereUniformLayout>())
 	{
-		DrawSpacingButton(ClonerActor, InView, InPDI, InOutIconIndex, EAvaClonerAxis::Y, FAvaVisualizerBase::Inactive);
+		DrawSpacingButton(ClonerActor, InView, InPDI, InOutIconIndex, ECEClonerAxis::Y, FAvaVisualizerBase::Inactive);
 		InOutIconIndex++;
 	}
 }
 
-FVector FAvaClonerActorVisualizer::GetHandleSpacingLocation(const AAvaClonerActor* InClonerActor, EAvaClonerAxis InAxis) const
+FVector FAvaClonerActorVisualizer::GetHandleSpacingLocation(const ACEClonerActor* InClonerActor, ECEClonerAxis InAxis) const
 {
 	const FRotator ClonerRotation = InClonerActor->GetActorRotation();
 	const FVector ClonerScale = InClonerActor->GetActorScale();
 	FVector OutLocation = FVector::ZeroVector;
-	FVector Axis = InAxis == EAvaClonerAxis::X ? FVector::XAxisVector : (InAxis == EAvaClonerAxis::Z ? FVector::ZAxisVector : FVector::YAxisVector);
-	
-	if (const UAvaClonerGridLayout* GridLayout = InClonerActor->GetActiveLayout<UAvaClonerGridLayout>())
+	FVector Axis = InAxis == ECEClonerAxis::X ? FVector::XAxisVector : (InAxis == ECEClonerAxis::Z ? FVector::ZAxisVector : FVector::YAxisVector);
+
+	if (const UCEClonerGridLayout* GridLayout = InClonerActor->GetActiveLayout<UCEClonerGridLayout>())
 	{
 		const FVector Spacing(GridLayout->GetSpacingX(), GridLayout->GetSpacingY(), GridLayout->GetSpacingZ());
 		const FVector Count(GridLayout->GetCountX(), GridLayout->GetCountY(), GridLayout->GetCountZ());
 
 		OutLocation = ClonerRotation.RotateVector(Axis * Spacing * Count / 2);
 	}
-	else if (const UAvaClonerLineLayout* LineLayout = InClonerActor->GetActiveLayout<UAvaClonerLineLayout>())
+	else if (const UCEClonerLineLayout* LineLayout = InClonerActor->GetActiveLayout<UCEClonerLineLayout>())
 	{
-		if (InAxis == EAvaClonerAxis::Custom)
+		if (InAxis == ECEClonerAxis::Custom)
 		{
 			Axis = LineLayout->GetDirection().GetSafeNormal();
 		}
-		
+
 		OutLocation = ClonerRotation.RotateVector(Axis) * LineLayout->GetSpacing() * FVector(LineLayout->GetCount());
 	}
-	else if (const UAvaClonerHoneycombLayout* HoneycombLayout = InClonerActor->GetActiveLayout<UAvaClonerHoneycombLayout>())
+	else if (const UCEClonerHoneycombLayout* HoneycombLayout = InClonerActor->GetActiveLayout<UCEClonerHoneycombLayout>())
 	{
-		const EAvaClonerPlane Plane = HoneycombLayout->GetPlane();
+		const ECEClonerPlane Plane = HoneycombLayout->GetPlane();
 		const float WidthSpacing = HoneycombLayout->GetWidthSpacing();
 		const float HeightSpacing = HoneycombLayout->GetHeightSpacing();
 		const int32 WidthCount = HoneycombLayout->GetWidthCount();
 		const int32 HeightCount = HoneycombLayout->GetHeightCount();
-		
+
 		FVector Spacing;
 		FVector Count;
-		if (Plane == EAvaClonerPlane::XY)
+		if (Plane == ECEClonerPlane::XY)
 		{
 			Spacing = FVector(WidthSpacing, HeightSpacing, 0);
 			Count = FVector(WidthCount, HeightCount, 0);
 		}
-		else if (Plane == EAvaClonerPlane::YZ)
+		else if (Plane == ECEClonerPlane::YZ)
 		{
 			Spacing = FVector(0, WidthSpacing, HeightSpacing);
 			Count = FVector(0, WidthCount, HeightCount);
 		}
-		else if (Plane == EAvaClonerPlane::XZ)
+		else if (Plane == ECEClonerPlane::XZ)
 		{
 			Spacing = FVector(WidthSpacing, 0, HeightSpacing);
 			Count = FVector(WidthCount, 0, HeightCount);
 		}
-		
+
 		OutLocation = ClonerRotation.RotateVector(Axis * Spacing * Count/2);
 	}
-	else if (const UAvaClonerCircleLayout* CircleLayout = InClonerActor->GetActiveLayout<UAvaClonerCircleLayout>())
+	else if (const UCEClonerCircleLayout* CircleLayout = InClonerActor->GetActiveLayout<UCEClonerCircleLayout>())
 	{
 		FRotator Rotation = CircleLayout->GetRotation();
 		const FVector Scale = CircleLayout->GetScale();
-		const EAvaClonerPlane Plane = CircleLayout->GetPlane();
+		const ECEClonerPlane Plane = CircleLayout->GetPlane();
 		const float Radius = CircleLayout->GetRadius();
-		
-		if (Plane == EAvaClonerPlane::XY)
+
+		if (Plane == ECEClonerPlane::XY)
 		{
 			Rotation = FRotator::ZeroRotator;
 		}
-		else if (Plane == EAvaClonerPlane::YZ)
+		else if (Plane == ECEClonerPlane::YZ)
 		{
 			Rotation = FRotator(90, 0, 0);
 		}
-		else if (Plane == EAvaClonerPlane::XZ)
+		else if (Plane == ECEClonerPlane::XZ)
 		{
 			Rotation = FRotator(0, 90, 0);
 		}
-		
+
 		OutLocation = ClonerRotation.RotateVector(Rotation.RotateVector(Axis * Scale)) * Radius;
 	}
-	else if (const UAvaClonerCylinderLayout* CylinderLayout = InClonerActor->GetActiveLayout<UAvaClonerCylinderLayout>())
+	else if (const UCEClonerCylinderLayout* CylinderLayout = InClonerActor->GetActiveLayout<UCEClonerCylinderLayout>())
 	{
 		FRotator Rotation = CylinderLayout->GetRotation();
 		const FVector Scale = CylinderLayout->GetScale();
-		const EAvaClonerPlane Plane = CylinderLayout->GetPlane();
+		const ECEClonerPlane Plane = CylinderLayout->GetPlane();
 		const float Radius = CylinderLayout->GetRadius();
 		const float Height = CylinderLayout->GetHeight();
 		float Dim = 0.f;
-		
-		if (Plane == EAvaClonerPlane::XY)
+
+		if (Plane == ECEClonerPlane::XY)
 		{
 			Rotation = FRotator::ZeroRotator;
-			if (InAxis == EAvaClonerAxis::Y)
+			if (InAxis == ECEClonerAxis::Y)
 			{
 				Dim = Radius;
 			}
@@ -635,9 +635,9 @@ FVector FAvaClonerActorVisualizer::GetHandleSpacingLocation(const AAvaClonerActo
 				Dim = Height / 2;
 			}
 		}
-		else if (Plane == EAvaClonerPlane::YZ)
+		else if (Plane == ECEClonerPlane::YZ)
 		{
-			if (InAxis == EAvaClonerAxis::Y)
+			if (InAxis == ECEClonerAxis::Y)
 			{
 				Rotation = FRotator(90, 0, 0);
 				Dim = Radius;
@@ -648,9 +648,9 @@ FVector FAvaClonerActorVisualizer::GetHandleSpacingLocation(const AAvaClonerActo
 				Dim = Height / 2;
 			}
 		}
-		else if (Plane == EAvaClonerPlane::XZ)
+		else if (Plane == ECEClonerPlane::XZ)
 		{
-			if (InAxis == EAvaClonerAxis::Z)
+			if (InAxis == ECEClonerAxis::Z)
 			{
 				Rotation = FRotator(0, 90, 0);
 				Dim = Radius;
@@ -661,21 +661,21 @@ FVector FAvaClonerActorVisualizer::GetHandleSpacingLocation(const AAvaClonerActo
 				Dim = Height / 2;
 			}
 		}
-		
+
 		OutLocation = ClonerRotation.RotateVector(Rotation.RotateVector(Axis * Scale)) * Dim;
 	}
-	else if (const UAvaClonerSphereUniformLayout* SphereLayout = InClonerActor->GetActiveLayout<UAvaClonerSphereUniformLayout>())
+	else if (const UCEClonerSphereUniformLayout* SphereLayout = InClonerActor->GetActiveLayout<UCEClonerSphereUniformLayout>())
 	{
 		const float Radius = SphereLayout->GetRadius();
 		const FVector Scale = SphereLayout->GetScale();
-		
+
 		OutLocation = ClonerRotation.RotateVector(Axis) * Scale * Radius;
 	}
-	
+
 	return InClonerActor->GetActorLocation() + OutLocation * ClonerScale;
 }
 
-void FAvaClonerActorVisualizer::DrawSpacingButton(const AAvaClonerActor* InClonerActor, const FSceneView* InView, FPrimitiveDrawInterface* InPDI, int32 InIconIndex, EAvaClonerAxis InAxis, FLinearColor InColor) const
+void FAvaClonerActorVisualizer::DrawSpacingButton(const ACEClonerActor* InClonerActor, const FSceneView* InView, FPrimitiveDrawInterface* InPDI, int32 InIconIndex, ECEClonerAxis InAxis, FLinearColor InColor) const
 {
 	UTexture2D* SpacingSprite = IAvalancheComponentVisualizersModule::Get().GetSettings()->GetVisualizerSprite(FAvalancheShapesEditorModule::BevelSprite);
 
@@ -701,7 +701,7 @@ void FAvaClonerActorVisualizer::OnPropertyModified(UObject* InPropertyObject, FN
 	{
 		return;
 	}
-	
+
 	InPropertyObject->Modify();
 	bHasBeenModified = true;
 	FProperty* PropertyModified = UE::AvaCore::GetProperty<UObject>(InPropertyName);
@@ -715,40 +715,40 @@ UActorComponent* FAvaClonerActorVisualizer::GetEditedComponent() const
 
 TMap<UObject*, TArray<FProperty*>> FAvaClonerActorVisualizer::GatherEditableProperties(UObject* InObject) const
 {
-	if (UAvaClonerComponent* ClonerComponent = Cast<UAvaClonerComponent>(InObject))
+	if (UCEClonerComponent* ClonerComponent = Cast<UCEClonerComponent>(InObject))
 	{
-		if (UAvaClonerLayoutBase* Layout = ClonerComponent->GetClonerActiveLayout())
+		if (UCEClonerLayoutBase* Layout = ClonerComponent->GetClonerActiveLayout())
 		{
 			TArray<FName> PropertyNames;
 
-			if (UAvaClonerGridLayout* GridLayout = Cast<UAvaClonerGridLayout>(Layout))
+			if (UCEClonerGridLayout* GridLayout = Cast<UCEClonerGridLayout>(Layout))
 			{
 				PropertyNames = {
-					GET_MEMBER_NAME_CHECKED(UAvaClonerGridLayout, SpacingX),
-					GET_MEMBER_NAME_CHECKED(UAvaClonerGridLayout, SpacingY),
-					GET_MEMBER_NAME_CHECKED(UAvaClonerGridLayout, SpacingZ)
+					GET_MEMBER_NAME_CHECKED(UCEClonerGridLayout, SpacingX),
+					GET_MEMBER_NAME_CHECKED(UCEClonerGridLayout, SpacingY),
+					GET_MEMBER_NAME_CHECKED(UCEClonerGridLayout, SpacingZ)
 				};
 			}
-			else if (UAvaClonerLineLayout* LineLayout = Cast<UAvaClonerLineLayout>(Layout))
+			else if (UCEClonerLineLayout* LineLayout = Cast<UCEClonerLineLayout>(Layout))
 			{
-				PropertyNames = {GET_MEMBER_NAME_CHECKED(UAvaClonerLineLayout, Spacing)};
+				PropertyNames = {GET_MEMBER_NAME_CHECKED(UCEClonerLineLayout, Spacing)};
 			}
-			else if (UAvaClonerHoneycombLayout* HoneycombLayout = Cast<UAvaClonerHoneycombLayout>(Layout))
+			else if (UCEClonerHoneycombLayout* HoneycombLayout = Cast<UCEClonerHoneycombLayout>(Layout))
 			{
 				PropertyNames = {
-					GET_MEMBER_NAME_CHECKED(UAvaClonerHoneycombLayout, WidthSpacing),
-					GET_MEMBER_NAME_CHECKED(UAvaClonerHoneycombLayout, HeightSpacing)
+					GET_MEMBER_NAME_CHECKED(UCEClonerHoneycombLayout, WidthSpacing),
+					GET_MEMBER_NAME_CHECKED(UCEClonerHoneycombLayout, HeightSpacing)
 				};
 			}
-			else if (UAvaClonerCircleLayout* CircleLayout = Cast<UAvaClonerCircleLayout>(Layout))
+			else if (UCEClonerCircleLayout* CircleLayout = Cast<UCEClonerCircleLayout>(Layout))
 			{
-				PropertyNames = {GET_MEMBER_NAME_CHECKED(UAvaClonerCircleLayout, Radius)};
+				PropertyNames = {GET_MEMBER_NAME_CHECKED(UCEClonerCircleLayout, Radius)};
 			}
-			else if (UAvaClonerCylinderLayout* CylinderLayout = Cast<UAvaClonerCylinderLayout>(Layout))
+			else if (UCEClonerCylinderLayout* CylinderLayout = Cast<UCEClonerCylinderLayout>(Layout))
 			{
 				PropertyNames = {
-					GET_MEMBER_NAME_CHECKED(UAvaClonerCylinderLayout, Radius),
-					GET_MEMBER_NAME_CHECKED(UAvaClonerCylinderLayout, Height)
+					GET_MEMBER_NAME_CHECKED(UCEClonerCylinderLayout, Radius),
+					GET_MEMBER_NAME_CHECKED(UCEClonerCylinderLayout, Height)
 				};
 			}
 
@@ -790,11 +790,11 @@ bool FAvaClonerActorVisualizer::VisProxyHandleClick(FEditorViewportClient* InVie
 	if (InVisProxy->IsA(HAvaClonerActorSpacingHitProxy::StaticGetType()))
 	{
 		EndEditing();
-		ClonerActorWeak = Cast<AAvaClonerActor>(DynMesh->GetOwner());
+		ClonerActorWeak = Cast<ACEClonerActor>(DynMesh->GetOwner());
 		bEditingSpacing = true;
 		EditingAxis = static_cast<HAvaClonerActorSpacingHitProxy*>(InVisProxy)->Axis;
 		StartEditing(InViewportClient, DynMesh);
-		
+
 		return true;
 	}
 
@@ -803,7 +803,7 @@ bool FAvaClonerActorVisualizer::VisProxyHandleClick(FEditorViewportClient* InVie
 
 bool FAvaClonerActorVisualizer::GetWidgetLocation(const FEditorViewportClient* InViewportClient, FVector& OutLocation) const
 {
-	if (const AAvaClonerActor* ClonerActor = ClonerActorWeak.Get())
+	if (const ACEClonerActor* ClonerActor = ClonerActorWeak.Get())
 	{
 		if (bEditingSpacing)
 		{
@@ -811,7 +811,7 @@ bool FAvaClonerActorVisualizer::GetWidgetLocation(const FEditorViewportClient* I
 			return true;
 		}
 	}
-	
+
 	return Super::GetWidgetLocation(InViewportClient, OutLocation);
 }
 
@@ -823,7 +823,7 @@ bool FAvaClonerActorVisualizer::GetWidgetMode(const FEditorViewportClient* InVie
 		OutMode = UE::Widget::EWidgetMode::WM_Translate;
 		return true;
 	}
-	
+
 	return Super::GetWidgetMode(InViewportClient, OutMode);
 }
 
@@ -832,25 +832,25 @@ bool FAvaClonerActorVisualizer::GetWidgetAxisList(const FEditorViewportClient* I
 {
 	if (bEditingSpacing)
 	{
-		if (EditingAxis == EAvaClonerAxis::X)
+		if (EditingAxis == ECEClonerAxis::X)
 		{
 			OutAxisList = EAxisList::Type::X;
 		}
-		else if (EditingAxis == EAvaClonerAxis::Y)
+		else if (EditingAxis == ECEClonerAxis::Y)
 		{
 			OutAxisList = EAxisList::Type::Y;
 		}
-		else if (EditingAxis == EAvaClonerAxis::Z)
+		else if (EditingAxis == ECEClonerAxis::Z)
 		{
 			OutAxisList = EAxisList::Type::Z;
 		}
-		else if (EditingAxis == EAvaClonerAxis::Custom)
+		else if (EditingAxis == ECEClonerAxis::Custom)
 		{
 			OutAxisList = EAxisList::Type::XYZ;
 		}
 		return true;
 	}
-	
+
 	return Super::GetWidgetAxisList(InViewportClient, InWidgetMode, OutAxisList);
 }
 
@@ -866,71 +866,71 @@ bool FAvaClonerActorVisualizer::ResetValue(FEditorViewportClient* InViewportClie
 	{
 		return Super::ResetValue(InViewportClient, InHitProxy);
 	}
-	
+
 	const HAvaClonerActorSpacingHitProxy* ComponentHitProxy = static_cast<HAvaClonerActorSpacingHitProxy*>(InHitProxy);
-	
-	if (!ComponentHitProxy->Component.IsValid() || !ComponentHitProxy->Component->IsA<UAvaClonerComponent>())
+
+	if (!ComponentHitProxy->Component.IsValid() || !ComponentHitProxy->Component->IsA<UCEClonerComponent>())
 	{
 		return Super::ResetValue(InViewportClient, InHitProxy);
 	}
 
-	if (const AAvaClonerActor* ClonerActor = Cast<AAvaClonerActor>(ComponentHitProxy->Component->GetOwner()))
+	if (const ACEClonerActor* ClonerActor = Cast<ACEClonerActor>(ComponentHitProxy->Component->GetOwner()))
 	{
-		if (UAvaClonerGridLayout* GridLayout = ClonerActor->GetActiveLayout<UAvaClonerGridLayout>())
+		if (UCEClonerGridLayout* GridLayout = ClonerActor->GetActiveLayout<UCEClonerGridLayout>())
 		{
 			FScopedTransaction Transaction(LOCTEXT("VisualizerResetValue", "Visualizer Reset Value"));
 			GridLayout->SetFlags(RF_Transactional);
-			
+
 			GridLayout->SetSpacingX(100.f);
 			GridLayout->SetSpacingY(100.f);
 			GridLayout->SetSpacingZ(100.f);
-			
-			OnPropertyModified(GridLayout, GET_MEMBER_NAME_CHECKED(UAvaClonerGridLayout, SpacingX), EPropertyChangeType::ValueSet);
-			OnPropertyModified(GridLayout, GET_MEMBER_NAME_CHECKED(UAvaClonerGridLayout, SpacingY), EPropertyChangeType::ValueSet);
-			OnPropertyModified(GridLayout, GET_MEMBER_NAME_CHECKED(UAvaClonerGridLayout, SpacingZ), EPropertyChangeType::ValueSet);
+
+			OnPropertyModified(GridLayout, GET_MEMBER_NAME_CHECKED(UCEClonerGridLayout, SpacingX), EPropertyChangeType::ValueSet);
+			OnPropertyModified(GridLayout, GET_MEMBER_NAME_CHECKED(UCEClonerGridLayout, SpacingY), EPropertyChangeType::ValueSet);
+			OnPropertyModified(GridLayout, GET_MEMBER_NAME_CHECKED(UCEClonerGridLayout, SpacingZ), EPropertyChangeType::ValueSet);
 		}
-		else if (UAvaClonerLineLayout* LineLayout = ClonerActor->GetActiveLayout<UAvaClonerLineLayout>())
+		else if (UCEClonerLineLayout* LineLayout = ClonerActor->GetActiveLayout<UCEClonerLineLayout>())
 		{
 			FScopedTransaction Transaction(LOCTEXT("VisualizerResetValue", "Visualizer Reset Value"));
 			LineLayout->SetFlags(RF_Transactional);
 
 			LineLayout->SetSpacing(500.f);
-			
-			OnPropertyModified(LineLayout, GET_MEMBER_NAME_CHECKED(UAvaClonerLineLayout, Spacing), EPropertyChangeType::ValueSet);
+
+			OnPropertyModified(LineLayout, GET_MEMBER_NAME_CHECKED(UCEClonerLineLayout, Spacing), EPropertyChangeType::ValueSet);
 		}
-		else if (UAvaClonerHoneycombLayout* HoneycombLayout = ClonerActor->GetActiveLayout<UAvaClonerHoneycombLayout>())
+		else if (UCEClonerHoneycombLayout* HoneycombLayout = ClonerActor->GetActiveLayout<UCEClonerHoneycombLayout>())
 		{
 			FScopedTransaction Transaction(LOCTEXT( "VisualizerResetValue", "Visualizer Reset Value"));
 			HoneycombLayout->SetFlags(RF_Transactional);
-			
+
 			HoneycombLayout->SetWidthSpacing(100.f);
 			HoneycombLayout->SetHeightSpacing(100.f);
-			
-			OnPropertyModified(HoneycombLayout, GET_MEMBER_NAME_CHECKED(UAvaClonerHoneycombLayout, WidthSpacing), EPropertyChangeType::ValueSet);
-			OnPropertyModified(HoneycombLayout, GET_MEMBER_NAME_CHECKED(UAvaClonerHoneycombLayout, HeightSpacing), EPropertyChangeType::ValueSet);
+
+			OnPropertyModified(HoneycombLayout, GET_MEMBER_NAME_CHECKED(UCEClonerHoneycombLayout, WidthSpacing), EPropertyChangeType::ValueSet);
+			OnPropertyModified(HoneycombLayout, GET_MEMBER_NAME_CHECKED(UCEClonerHoneycombLayout, HeightSpacing), EPropertyChangeType::ValueSet);
 		}
-		else if (UAvaClonerCircleLayout* CircleLayout = ClonerActor->GetActiveLayout<UAvaClonerCircleLayout>())
+		else if (UCEClonerCircleLayout* CircleLayout = ClonerActor->GetActiveLayout<UCEClonerCircleLayout>())
 		{
 			FScopedTransaction Transaction(LOCTEXT("VisualizerResetValue", "Visualizer Reset Value"));
 			CircleLayout->SetFlags(RF_Transactional);
-			
+
 			CircleLayout->SetRadius(500.f);
-			
-			OnPropertyModified(CircleLayout, GET_MEMBER_NAME_CHECKED(UAvaClonerCircleLayout, Radius), EPropertyChangeType::ValueSet);
+
+			OnPropertyModified(CircleLayout, GET_MEMBER_NAME_CHECKED(UCEClonerCircleLayout, Radius), EPropertyChangeType::ValueSet);
 		}
-		else if (UAvaClonerCylinderLayout* CylinderLayout = ClonerActor->GetActiveLayout<UAvaClonerCylinderLayout>())
+		else if (UCEClonerCylinderLayout* CylinderLayout = ClonerActor->GetActiveLayout<UCEClonerCylinderLayout>())
 		{
 			FScopedTransaction Transaction(LOCTEXT("VisualizerResetValue", "Visualizer Reset Value"));
 			CylinderLayout->SetFlags(RF_Transactional);
-			
+
 			CylinderLayout->SetRadius(500.f);
 			CylinderLayout->SetHeight(1000.f);
 
-			OnPropertyModified(CylinderLayout, GET_MEMBER_NAME_CHECKED(UAvaClonerCylinderLayout, Radius), EPropertyChangeType::ValueSet);
-			OnPropertyModified(CylinderLayout, GET_MEMBER_NAME_CHECKED(UAvaClonerCylinderLayout, Height), EPropertyChangeType::ValueSet);
+			OnPropertyModified(CylinderLayout, GET_MEMBER_NAME_CHECKED(UCEClonerCylinderLayout, Radius), EPropertyChangeType::ValueSet);
+			OnPropertyModified(CylinderLayout, GET_MEMBER_NAME_CHECKED(UCEClonerCylinderLayout, Height), EPropertyChangeType::ValueSet);
 		}
 	}
-	
+
 	return true;
 }
 
@@ -940,7 +940,7 @@ bool FAvaClonerActorVisualizer::IsEditing() const
 	{
 		return true;
 	}
-	
+
 	return Super::IsEditing();
 }
 
diff --git a/Engine/Plugins/Experimental/Avalanche/Source/AvalancheEffectorsEditor/Private/Cloner/AvaClonerActorVis.h b/Engine/Plugins/Experimental/Avalanche/Source/AvalancheEffectorsEditor/Private/Cloner/AvaClonerActorVis.h
--- a/Engine/Plugins/Experimental/Avalanche/Source/AvalancheEffectorsEditor/Private/Cloner/AvaClonerActorVis.h
+++ b/Engine/Plugins/Experimental/Avalanche/Source/AvalancheEffectorsEditor/Private/Cloner/AvaClonerActorVis.h
@@ -2,23 +2,23 @@
 
 #pragma once
 
-#include "AvaClonerEffectorShared.h"
+#include "CEClonerEffectorShared.h"
 #include "AvaVisBase.h"
 #include "UObject/WeakObjectPtrTemplates.h"
 
-class AAvaClonerActor;
-class UAvaClonerComponent;
+class ACEClonerActor;
+class UCEClonerComponent;
 
 struct HAvaClonerActorSpacingHitProxy : HAvaHitProxy
 {
 	DECLARE_HIT_PROXY();
 
-	HAvaClonerActorSpacingHitProxy(const UActorComponent* InComponent, EAvaClonerAxis InAxis)
+	HAvaClonerActorSpacingHitProxy(const UActorComponent* InComponent, ECEClonerAxis InAxis)
 		: HAvaHitProxy(InComponent)
 		, Axis(InAxis)
 	{}
 
-	EAvaClonerAxis Axis = EAvaClonerAxis::Custom;
+	ECEClonerAxis Axis = ECEClonerAxis::Custom;
 };
 
 /** Custom visualization for cloner actor to handle spacing in various layouts */
@@ -26,10 +26,10 @@ class FAvaClonerActorVisualizer : public FAvaVisualizerBase
 {
 public:
 	typedef FAvaVisualizerBase Super;
-	typedef UAvaClonerComponent MeshType;
+	typedef UCEClonerComponent MeshType;
 
 	FAvaClonerActorVisualizer();
-	
+
 	//~ Begin FAvaVisualizerBase
 	virtual UActorComponent* GetEditedComponent() const override;
 	virtual TMap<UObject*, TArray<FProperty*>> GatherEditableProperties(UObject* InObject) const override;
@@ -47,20 +47,20 @@ public:
 	virtual void DrawVisualizationEditing(const UActorComponent* InComponent, const FSceneView* InView, FPrimitiveDrawInterface* InPDI, int32& InOutIconIndex) override;
 	virtual void DrawVisualizationNotEditing(const UActorComponent* InComponent, const FSceneView* InView, FPrimitiveDrawInterface* InPDI, int32& InOutIconIndex) override;
 	//~ End FAvaVisualizerBase
-	
-	AAvaClonerActor* GetClonerActor() const
+
+	ACEClonerActor* GetClonerActor() const
 	{
 		return ClonerActorWeak.Get();
 	}
 
 protected:
-	FVector GetHandleSpacingLocation(const AAvaClonerActor* InClonerActor, EAvaClonerAxis InAxis) const;
-	void DrawSpacingButton(const AAvaClonerActor* InClonerActor, const FSceneView* InView, FPrimitiveDrawInterface* InPDI, int32 InIconIndex, EAvaClonerAxis InAxis, FLinearColor InColor) const;
+	FVector GetHandleSpacingLocation(const ACEClonerActor* InClonerActor, ECEClonerAxis InAxis) const;
+	void DrawSpacingButton(const ACEClonerActor* InClonerActor, const FSceneView* InView, FPrimitiveDrawInterface* InPDI, int32 InIconIndex, ECEClonerAxis InAxis, FLinearColor InColor) const;
 
 	void OnPropertyModified(UObject* InPropertyObject, FName InPropertyName, EPropertyChangeType::Type InType = EPropertyChangeType::Interactive);
-	
-	TWeakObjectPtr<AAvaClonerActor> ClonerActorWeak = nullptr;
+
+	TWeakObjectPtr<ACEClonerActor> ClonerActorWeak = nullptr;
 	FVector InitialSpacing = FVector::ZeroVector;
 	bool bEditingSpacing = false;
-	EAvaClonerAxis EditingAxis = EAvaClonerAxis::X;
+	ECEClonerAxis EditingAxis = ECEClonerAxis::X;
 };
diff --git a/Engine/Plugins/Experimental/Avalanche/Source/AvalancheEffectorsEditor/Private/Effector/AvaEffectorActorTool.cpp b/Engine/Plugins/Experimental/Avalanche/Source/AvalancheEffectorsEditor/Private/Effector/AvaEffectorActorTool.cpp
--- a/Engine/Plugins/Experimental/Avalanche/Source/AvalancheEffectorsEditor/Private/Effector/AvaEffectorActorTool.cpp
+++ b/Engine/Plugins/Experimental/Avalanche/Source/AvalancheEffectorsEditor/Private/Effector/AvaEffectorActorTool.cpp
@@ -4,11 +4,11 @@
 #include "AvaEffectorsEditorCommands.h"
 #include "AvaInteractiveToolsSettings.h"
 #include "Builders/AvaInteractiveToolsToolBuilder.h"
-#include "Effector/AvaEffectorActor.h"
+#include "Effector/CEEffectorActor.h"
 
 UAvaEffectorActorTool::UAvaEffectorActorTool()
 {
-	ActorClass = AAvaEffectorActor::StaticClass();
+	ActorClass = ACEEffectorActor::StaticClass();
 }
 
 bool UAvaEffectorActorTool::UseIdentityRotation() const
diff --git a/Engine/Plugins/Experimental/Avalanche/Source/AvalancheEffectorsEditor/Private/Effector/AvaEffectorActorVis.cpp b/Engine/Plugins/Experimental/Avalanche/Source/AvalancheEffectorsEditor/Private/Effector/AvaEffectorActorVis.cpp
--- a/Engine/Plugins/Experimental/Avalanche/Source/AvalancheEffectorsEditor/Private/Effector/AvaEffectorActorVis.cpp
+++ b/Engine/Plugins/Experimental/Avalanche/Source/AvalancheEffectorsEditor/Private/Effector/AvaEffectorActorVis.cpp
@@ -6,7 +6,7 @@
 #include "AvalancheShapesEditorModule.h"
 #include "AvaVisBase.h"
 #include "EditorViewportClient.h"
-#include "Effector/AvaEffectorComponent.h"
+#include "Effector/CEEffectorComponent.h"
 #include "Engine/Texture2D.h"
 #include "Framework/Application/SlateApplication.h"
 #include "IAvalancheComponentVisualizersModule.h"
@@ -21,11 +21,11 @@ FAvaEffectorActorVisualizer::FAvaEffectorActorVisualizer()
 	: FAvaVisualizerBase()
 {
 	using namespace UE::AvaCore;
-	InnerRadiusProperty  = GetProperty<AAvaEffectorActor>(GET_MEMBER_NAME_CHECKED(AAvaEffectorActor, InnerRadius));
-	OuterRadiusProperty  = GetProperty<AAvaEffectorActor>(GET_MEMBER_NAME_CHECKED(AAvaEffectorActor, OuterRadius));
-	InnerExtentProperty  = GetProperty<AAvaEffectorActor>(GET_MEMBER_NAME_CHECKED(AAvaEffectorActor, InnerExtent));
-	OuterExtentProperty  = GetProperty<AAvaEffectorActor>(GET_MEMBER_NAME_CHECKED(AAvaEffectorActor, OuterExtent));
-	PlaneSpacingProperty = GetProperty<AAvaEffectorActor>(GET_MEMBER_NAME_CHECKED(AAvaEffectorActor, PlaneSpacing));
+	InnerRadiusProperty  = GetProperty<ACEEffectorActor>(GET_MEMBER_NAME_CHECKED(ACEEffectorActor, InnerRadius));
+	OuterRadiusProperty  = GetProperty<ACEEffectorActor>(GET_MEMBER_NAME_CHECKED(ACEEffectorActor, OuterRadius));
+	InnerExtentProperty  = GetProperty<ACEEffectorActor>(GET_MEMBER_NAME_CHECKED(ACEEffectorActor, InnerExtent));
+	OuterExtentProperty  = GetProperty<ACEEffectorActor>(GET_MEMBER_NAME_CHECKED(ACEEffectorActor, OuterExtent));
+	PlaneSpacingProperty = GetProperty<ACEEffectorActor>(GET_MEMBER_NAME_CHECKED(ACEEffectorActor, PlaneSpacing));
 }
 
 void FAvaEffectorActorVisualizer::StoreInitialValues()
@@ -46,19 +46,19 @@ void FAvaEffectorActorVisualizer::StoreInitialValues()
 
 FBox FAvaEffectorActorVisualizer::GetComponentBounds(const UActorComponent* InComponent) const
 {
-	if (const UAvaEffectorComponent* EffectorComponent = Cast<UAvaEffectorComponent>(InComponent))
+	if (const UCEEffectorComponent* EffectorComponent = Cast<UCEEffectorComponent>(InComponent))
 	{
-		if (const AAvaEffectorActor* EffectorActor = Cast<AAvaEffectorActor>(EffectorComponent->GetOwner()))
+		if (const ACEEffectorActor* EffectorActor = Cast<ACEEffectorActor>(EffectorComponent->GetOwner()))
 		{
-			if (EffectorActor->GetType() == EAvaClonerEffectorType::Box)
+			if (EffectorActor->GetType() == ECEClonerEffectorType::Box)
 			{
 				return FBox(-EffectorActor->GetOuterExtent(), EffectorActor->GetOuterExtent());
 			}
-			if (EffectorActor->GetType() == EAvaClonerEffectorType::Sphere)
+			if (EffectorActor->GetType() == ECEClonerEffectorType::Sphere)
 			{
 				return FBox(-FVector(EffectorActor->GetOuterRadius() / 2), FVector(EffectorActor->GetOuterRadius() / 2));
 			}
-			if (EffectorActor->GetType() == EAvaClonerEffectorType::Plane)
+			if (EffectorActor->GetType() == ECEClonerEffectorType::Plane)
 			{
 				return FBox(-FVector(EffectorActor->GetPlaneSpacing() / 2), FVector(EffectorActor->GetPlaneSpacing() / 2));
 			}
@@ -75,13 +75,13 @@ bool FAvaEffectorActorVisualizer::HandleInputDeltaInternal(FEditorViewportClient
 		return false;
 	}
 
-	if (AAvaEffectorActor* EffectorActor = EffectorActorWeak.Get())
+	if (ACEEffectorActor* EffectorActor = EffectorActorWeak.Get())
 	{
 		if (GetViewportWidgetMode(InViewportClient) == UE::Widget::WM_Translate)
 		{
 			if (GetViewportWidgetAxisList(InViewportClient) & EAxisList::XYZ)
 			{
-				if (EffectorActor->GetType() == EAvaClonerEffectorType::Box)
+				if (EffectorActor->GetType() == ECEClonerEffectorType::Box)
 				{
 					if (bEditingInnerZone)
 					{
@@ -103,7 +103,7 @@ bool FAvaEffectorActorVisualizer::HandleInputDeltaInternal(FEditorViewportClient
 			}
 			if (GetViewportWidgetAxisList(InViewportClient) & EAxisList::Y)
 			{
-				if (EffectorActor->GetType() == EAvaClonerEffectorType::Plane)
+				if (EffectorActor->GetType() == ECEClonerEffectorType::Plane)
 				{
 					if (bEditingOuterZone || bEditingInnerZone)
 					{
@@ -115,7 +115,7 @@ bool FAvaEffectorActorVisualizer::HandleInputDeltaInternal(FEditorViewportClient
 					
 					return true;
 				}
-				if (EffectorActor->GetType() == EAvaClonerEffectorType::Sphere)
+				if (EffectorActor->GetType() == ECEClonerEffectorType::Sphere)
 				{
 					if (bEditingInnerZone)
 					{
@@ -149,21 +149,21 @@ void FAvaEffectorActorVisualizer::DrawVisualizationEditing(const UActorComponent
 {
 	Super::DrawVisualizationEditing(InComponent, InView, InPDI, InOutIconIndex);
 
-	const UAvaEffectorComponent* EffectorComponent = Cast<UAvaEffectorComponent>(InComponent);
+	const UCEEffectorComponent* EffectorComponent = Cast<UCEEffectorComponent>(InComponent);
 
 	if (!EffectorComponent)
 	{
 		return;
 	}
 
-	const AAvaEffectorActor* EffectorActor = Cast<AAvaEffectorActor>(EffectorComponent->GetOwner());
+	const ACEEffectorActor* EffectorActor = Cast<ACEEffectorActor>(EffectorComponent->GetOwner());
 
 	if (!EffectorActor)
 	{
 		return;
 	}
 
-	if (EffectorActor->GetType() != EAvaClonerEffectorType::Plane)
+	if (EffectorActor->GetType() != ECEClonerEffectorType::Plane)
 	{
 		if (!bEditingInnerZone)
 		{
@@ -183,21 +183,21 @@ void FAvaEffectorActorVisualizer::DrawVisualizationNotEditing(const UActorCompon
 {
 	Super::DrawVisualizationNotEditing(InComponent, InView, InPDI, InOutIconIndex);
 
-	const UAvaEffectorComponent* EffectorComponent = Cast<UAvaEffectorComponent>(InComponent);
+	const UCEEffectorComponent* EffectorComponent = Cast<UCEEffectorComponent>(InComponent);
 
 	if (!EffectorComponent)
 	{
 		return;
 	}
 
-	const AAvaEffectorActor* EffectorActor = Cast<AAvaEffectorActor>(EffectorComponent->GetOwner());
+	const ACEEffectorActor* EffectorActor = Cast<ACEEffectorActor>(EffectorComponent->GetOwner());
 
 	if (!EffectorActor)
 	{
 		return;
 	}
 
-	if (EffectorActor->GetType() != EAvaClonerEffectorType::Plane)
+	if (EffectorActor->GetType() != ECEClonerEffectorType::Plane)
 	{
 		DrawZoneButton(EffectorActor, InView, InPDI, InOutIconIndex, true, FAvaVisualizerBase::Inactive);
 		InOutIconIndex++;
@@ -207,13 +207,13 @@ void FAvaEffectorActorVisualizer::DrawVisualizationNotEditing(const UActorCompon
 	InOutIconIndex++;
 }
 
-FVector FAvaEffectorActorVisualizer::GetHandleZoneLocation(const AAvaEffectorActor* InEffectorActor, bool bInInnerSize) const
+FVector FAvaEffectorActorVisualizer::GetHandleZoneLocation(const ACEEffectorActor* InEffectorActor, bool bInInnerSize) const
 {
 	const FVector EffectorScale = InEffectorActor->GetActorScale();
 	const FRotator EffectorRotation = InEffectorActor->GetActorRotation();
 	FVector OutLocation = InEffectorActor->GetActorLocation();
 
-	if (InEffectorActor->GetType() == EAvaClonerEffectorType::Box)
+	if (InEffectorActor->GetType() == ECEClonerEffectorType::Box)
 	{
 		if (bInInnerSize)
 		{
@@ -224,12 +224,12 @@ FVector FAvaEffectorActorVisualizer::GetHandleZoneLocation(const AAvaEffectorAct
 			OutLocation += EffectorRotation.RotateVector(InEffectorActor->GetOuterExtent()) * EffectorScale;
 		}
 	}
-	else if (InEffectorActor->GetType() == EAvaClonerEffectorType::Plane)
+	else if (InEffectorActor->GetType() == ECEClonerEffectorType::Plane)
 	{
 		const float ComponentScale = (EffectorRotation.RotateVector(-FVector::YAxisVector) * EffectorScale).Length();
 		OutLocation += EffectorRotation.RotateVector(FVector::YAxisVector) * (InEffectorActor->GetPlaneSpacing() / 2) * ComponentScale;
 	}
-	else if (InEffectorActor->GetType() == EAvaClonerEffectorType::Sphere)
+	else if (InEffectorActor->GetType() == ECEClonerEffectorType::Sphere)
 	{
 		const float MinComponentScale = FMath::Min<float>(FMath::Min<float>(EffectorScale.X, EffectorScale.Y), EffectorScale.Z);
 		if (bInInnerSize)
@@ -245,7 +245,7 @@ FVector FAvaEffectorActorVisualizer::GetHandleZoneLocation(const AAvaEffectorAct
 	return OutLocation;
 }
 
-void FAvaEffectorActorVisualizer::DrawZoneButton(const AAvaEffectorActor* InEffectorActor, const FSceneView* InView, FPrimitiveDrawInterface* InPDI, int32 InIconIndex, bool bInInnerZone, FLinearColor InColor) const
+void FAvaEffectorActorVisualizer::DrawZoneButton(const ACEEffectorActor* InEffectorActor, const FSceneView* InView, FPrimitiveDrawInterface* InPDI, int32 InIconIndex, bool bInInnerZone, FLinearColor InColor) const
 {
 	UTexture2D* ZoneSprite = IAvalancheComponentVisualizersModule::Get().GetSettings()->GetVisualizerSprite(FAvalancheShapesEditorModule::BevelSprite);
 
@@ -272,19 +272,19 @@ UActorComponent* FAvaEffectorActorVisualizer::GetEditedComponent() const
 
 TMap<UObject*, TArray<FProperty*>> FAvaEffectorActorVisualizer::GatherEditableProperties(UObject* InObject) const
 {
-	if (UAvaEffectorComponent* EffectorComponent = Cast<UAvaEffectorComponent>(InObject))
+	if (UCEEffectorComponent* EffectorComponent = Cast<UCEEffectorComponent>(InObject))
 	{
-		if (AAvaEffectorActor* EffectorActor = EffectorComponent->GetOuterAAvaEffectorActor())
+		if (ACEEffectorActor* EffectorActor = EffectorComponent->GetOuterACEEffectorActor())
 		{
 			switch (EffectorActor->GetType())
 			{
-				case EAvaClonerEffectorType::Plane:
+				case ECEClonerEffectorType::Plane:
 					return {{EffectorActor, {PlaneSpacingProperty}}};
 
-				case EAvaClonerEffectorType::Box:
+				case ECEClonerEffectorType::Box:
 					return {{EffectorActor, {InnerExtentProperty, OuterExtentProperty}}};
 
-				case EAvaClonerEffectorType::Sphere:
+				case ECEClonerEffectorType::Sphere:
 					return {{EffectorActor, {InnerRadiusProperty, OuterRadiusProperty}}};
 			}
 		}
@@ -303,7 +303,7 @@ bool FAvaEffectorActorVisualizer::VisProxyHandleClick(FEditorViewportClient* InV
 
 	UActorComponent* Component = const_cast<UActorComponent*>(InVisProxy->Component.Get());
 
-	if (!Component || !Component->GetOwner()->IsA<AAvaEffectorActor>())
+	if (!Component || !Component->GetOwner()->IsA<ACEEffectorActor>())
 	{
 		return Super::VisProxyHandleClick(InViewportClient, InVisProxy, InClick);
 	}
@@ -311,7 +311,7 @@ bool FAvaEffectorActorVisualizer::VisProxyHandleClick(FEditorViewportClient* InV
 	if (InVisProxy->IsA(HAvaEffectorActorZoneHitProxy::StaticGetType()))
 	{
 		EndEditing();
-		EffectorActorWeak = Cast<AAvaEffectorActor>(InVisProxy->Component->GetOwner());
+		EffectorActorWeak = Cast<ACEEffectorActor>(InVisProxy->Component->GetOwner());
 		bEditingInnerZone = static_cast<HAvaEffectorActorZoneHitProxy*>(InVisProxy)->bInnerZone;
 		bEditingOuterZone = !bEditingInnerZone;
 		StartEditing(InViewportClient, Component);
@@ -350,7 +350,7 @@ bool FAvaEffectorActorVisualizer::GetWidgetAxisList(const FEditorViewportClient*
 {
 	if (bEditingOuterZone || bEditingInnerZone)
 	{
-		if (EffectorActorWeak->GetType() != EAvaClonerEffectorType::Box)
+		if (EffectorActorWeak->GetType() != ECEClonerEffectorType::Box)
 		{
 			OutAxisList = EAxisList::Type::Y;
 		}
@@ -370,7 +370,7 @@ bool FAvaEffectorActorVisualizer::GetWidgetAxisListDragOverride(const FEditorVie
 {
 	if (bEditingOuterZone || bEditingInnerZone)
 	{
-		if (EffectorActorWeak->GetType() != EAvaClonerEffectorType::Box)
+		if (EffectorActorWeak->GetType() != ECEClonerEffectorType::Box)
 		{
 			OutAxisList = EAxisList::Type::Y;
 			return true;
@@ -389,14 +389,14 @@ bool FAvaEffectorActorVisualizer::ResetValue(FEditorViewportClient* InViewportCl
 	
 	const HAvaEffectorActorZoneHitProxy* ComponentHitProxy = static_cast<HAvaEffectorActorZoneHitProxy*>(InHitProxy);
 	
-	if (!ComponentHitProxy->Component.IsValid() || !ComponentHitProxy->Component->IsA<UAvaEffectorComponent>())
+	if (!ComponentHitProxy->Component.IsValid() || !ComponentHitProxy->Component->IsA<UCEEffectorComponent>())
 	{
 		return Super::ResetValue(InViewportClient, InHitProxy);
 	}
 
-	if (AAvaEffectorActor* EffectorActor = Cast<AAvaEffectorActor>(ComponentHitProxy->Component->GetOwner()))
+	if (ACEEffectorActor* EffectorActor = Cast<ACEEffectorActor>(ComponentHitProxy->Component->GetOwner()))
 	{
-		if (EffectorActor->GetType() == EAvaClonerEffectorType::Box)
+		if (EffectorActor->GetType() == ECEClonerEffectorType::Box)
 		{
 			if (ComponentHitProxy->bInnerZone)
 			{
@@ -415,7 +415,7 @@ bool FAvaEffectorActorVisualizer::ResetValue(FEditorViewportClient* InViewportCl
 				NotifyPropertyModified(EffectorActor, OuterExtentProperty, EPropertyChangeType::ValueSet);
 			}
 		}
-		else if (EffectorActor->GetType() == EAvaClonerEffectorType::Plane)
+		else if (EffectorActor->GetType() == ECEClonerEffectorType::Plane)
 		{
 			FScopedTransaction Transaction(LOCTEXT("VisualizerResetValue", "Visualizer Reset Value"));
 			EffectorActor->SetFlags(RF_Transactional);
@@ -423,7 +423,7 @@ bool FAvaEffectorActorVisualizer::ResetValue(FEditorViewportClient* InViewportCl
 			EffectorActor->Modify();
 			NotifyPropertyModified(EffectorActor, PlaneSpacingProperty, EPropertyChangeType::ValueSet);
 		}
-		else if (EffectorActor->GetType() == EAvaClonerEffectorType::Sphere)
+		else if (EffectorActor->GetType() == ECEClonerEffectorType::Sphere)
 		{
 			if (ComponentHitProxy->bInnerZone)
 			{
diff --git a/Engine/Plugins/Experimental/Avalanche/Source/AvalancheEffectorsEditor/Private/Effector/AvaEffectorActorVis.h b/Engine/Plugins/Experimental/Avalanche/Source/AvalancheEffectorsEditor/Private/Effector/AvaEffectorActorVis.h
--- a/Engine/Plugins/Experimental/Avalanche/Source/AvalancheEffectorsEditor/Private/Effector/AvaEffectorActorVis.h
+++ b/Engine/Plugins/Experimental/Avalanche/Source/AvalancheEffectorsEditor/Private/Effector/AvaEffectorActorVis.h
@@ -2,10 +2,10 @@
 
 #pragma once
 
-#include "Effector/AvaEffectorActor.h"
+#include "Effector/CEEffectorActor.h"
 #include "AvaVisBase.h"
 
-class UAvaEffectorComponent;
+class UCEEffectorComponent;
 
 struct HAvaEffectorActorZoneHitProxy : HAvaHitProxy
 {
@@ -26,7 +26,7 @@ public:
 	typedef FAvaVisualizerBase Super;
 
 	FAvaEffectorActorVisualizer();
-	
+
 	//~ Begin FAvaVisualizerBase
 	virtual UActorComponent* GetEditedComponent() const override;
 	virtual TMap<UObject*, TArray<FProperty*>> GatherEditableProperties(UObject* InObject) const override;
@@ -44,20 +44,20 @@ public:
 	virtual void DrawVisualizationEditing(const UActorComponent* InComponent, const FSceneView* InView, FPrimitiveDrawInterface* InPDI, int32& InOutIconIndex) override;
 	virtual void DrawVisualizationNotEditing(const UActorComponent* InComponent, const FSceneView* InView, FPrimitiveDrawInterface* InPDI, int32& InOutIconIndex) override;
 	//~ End FAvaVisualizerBase
-	
-	AAvaEffectorActor* GetEffectorActor() const { return EffectorActorWeak.Get(); };
+
+	ACEEffectorActor* GetEffectorActor() const { return EffectorActorWeak.Get(); };
 protected:
 
-	FVector GetHandleZoneLocation(const AAvaEffectorActor* InEffectorActor, bool bInInnerSize) const;
-	void DrawZoneButton(const AAvaEffectorActor* InEffectorActor, const FSceneView* InView, FPrimitiveDrawInterface* InPDI, int32 InIconIndex, bool bInInnerZone, FLinearColor InColor) const;
+	FVector GetHandleZoneLocation(const ACEEffectorActor* InEffectorActor, bool bInInnerSize) const;
+	void DrawZoneButton(const ACEEffectorActor* InEffectorActor, const FSceneView* InView, FPrimitiveDrawInterface* InPDI, int32 InIconIndex, bool bInInnerZone, FLinearColor InColor) const;
 
 	FProperty* InnerRadiusProperty;
 	FProperty* OuterRadiusProperty;
 	FProperty* InnerExtentProperty;
 	FProperty* OuterExtentProperty;
 	FProperty* PlaneSpacingProperty;
-	
-	TWeakObjectPtr<AAvaEffectorActor> EffectorActorWeak = nullptr;
+
+	TWeakObjectPtr<ACEEffectorActor> EffectorActorWeak = nullptr;
 	float InitialInnerRadius = 0.f;
 	float InitialOuterRadius = 0.f;
 	FVector InitialInnerExtent = FVector(0.f);
diff --git a/Engine/Plugins/Experimental/Avalanche/Source/AvalancheEffectorsEditor/Public/AvalancheEffectorsEditorModule.h b/Engine/Plugins/Experimental/Avalanche/Source/AvalancheEffectorsEditor/Public/AvalancheEffectorsEditorModule.h
--- a/Engine/Plugins/Experimental/Avalanche/Source/AvalancheEffectorsEditor/Public/AvalancheEffectorsEditorModule.h
+++ b/Engine/Plugins/Experimental/Avalanche/Source/AvalancheEffectorsEditor/Public/AvalancheEffectorsEditorModule.h
@@ -10,8 +10,6 @@ class FComponentVisualizer;
 class FAvalancheEffectorsEditorModule : public IModuleInterface
 {
 public:
-	static FAvalancheEffectorsEditorModule& Get();
-
 	//~ Begin IModuleInterface interface
 	virtual void StartupModule() override;
 	virtual void ShutdownModule() override;
@@ -24,8 +22,5 @@ private:
 
 	void RegisterComponentVisualizers();
 
-	void RegisterCustomLayouts();
-	void UnregisterCustomLayouts();
-
 	TArray<TSharedPtr<FComponentVisualizer>> Visualizers;
 };
