Motion Design Effectors:
- Moved module header to private

#changelist validated
#virtualized

diff --git a/Engine/Plugins/Experimental/Avalanche/Source/AvalancheEffectors/AvalancheEffectors.Build.cs b/Engine/Plugins/Experimental/Avalanche/Source/AvalancheEffectors/AvalancheEffectors.Build.cs
--- a/Engine/Plugins/Experimental/Avalanche/Source/AvalancheEffectors/AvalancheEffectors.Build.cs
+++ b/Engine/Plugins/Experimental/Avalanche/Source/AvalancheEffectors/AvalancheEffectors.Build.cs
@@ -8,22 +8,16 @@ public class AvalancheEffectors : ModuleRules
 	{
 		PCHUsage = ModuleRules.PCHUsageMode.UseExplicitOrSharedPCHs;
 
-		PublicDependencyModuleNames.AddRange(
-			new string[]
-			{
-				"Core",
-				"CoreUObject",
-				"Engine"
-			}
-		);
-
 		PrivateDependencyModuleNames.AddRange(
 			new string[]
 			{
-				"AvalancheCore",
 				"Avalanche",
+				"AvalancheCore",
 				"AvalancheSceneTree",
-				"ClonerEffector"
+				"ClonerEffector",
+				"Core",
+				"CoreUObject",
+				"Engine",
 			}
 		);
 
diff --git a/Engine/Plugins/Experimental/Avalanche/Source/AvalancheEffectorsEditor/AvalancheEffectorsEditor.Build.cs b/Engine/Plugins/Experimental/Avalanche/Source/AvalancheEffectorsEditor/AvalancheEffectorsEditor.Build.cs
--- a/Engine/Plugins/Experimental/Avalanche/Source/AvalancheEffectorsEditor/AvalancheEffectorsEditor.Build.cs
+++ b/Engine/Plugins/Experimental/Avalanche/Source/AvalancheEffectorsEditor/AvalancheEffectorsEditor.Build.cs
@@ -8,17 +8,6 @@ public class AvalancheEffectorsEditor : ModuleRules
 	{
 		PCHUsage = ModuleRules.PCHUsageMode.UseExplicitOrSharedPCHs;
 
-		PublicDependencyModuleNames.AddRange(
-			new string[]
-			{
-				"Core",
-				"CoreUObject",
-				"Engine",
-				"Slate",
-				"SlateCore"
-			}
-		);
-
 		PrivateDependencyModuleNames.AddRange(
 			new string[]
 			{
@@ -30,10 +19,15 @@ public class AvalancheEffectorsEditor : ModuleRules
 				"AvalancheShapesEditor",
 				"ClonerEffector",
 				"ComponentVisualizers",
+				"Core",
+				"CoreUObject",
+				"Engine",
 				"InputCore",
 				"InteractiveToolsFramework",
 				"Projects",
-				"UnrealEd"
+				"Slate",
+				"SlateCore",
+				"UnrealEd",
 			}
 		);
 	}
diff --git a/Engine/Plugins/Experimental/Avalanche/Source/AvalancheEffectorsEditor/Private/AvalancheEffectorsEditorModule.h b/Engine/Plugins/Experimental/Avalanche/Source/AvalancheEffectorsEditor/Private/AvalancheEffectorsEditorModule.h
new file mode 100644
--- /dev/null
+++ b/Engine/Plugins/Experimental/Avalanche/Source/AvalancheEffectorsEditor/Private/AvalancheEffectorsEditorModule.h
@@ -0,0 +1,28 @@
+// Copyright Epic Games, Inc. All Rights Reserved.
+
+#pragma once
+
+#include "Containers/Array.h"
+#include "Modules/ModuleInterface.h"
+#include "Templates/SharedPointerFwd.h"
+
+class FComponentVisualizer;
+class IAvalancheInteractiveToolsModule;
+
+class FAvalancheEffectorsEditorModule : public IModuleInterface
+{
+public:
+	//~ Begin IModuleInterface interface
+	virtual void StartupModule() override;
+	virtual void ShutdownModule() override;
+	//~ End IModuleInterface interface
+
+private:
+	void PostEngineInit();
+
+	void RegisterTools(IAvalancheInteractiveToolsModule* InModule);
+
+	void RegisterComponentVisualizers();
+
+	TArray<TSharedPtr<FComponentVisualizer>> Visualizers;
+};
diff --git a/Engine/Plugins/Experimental/Avalanche/Source/AvalancheEffectorsEditor/Public/AvalancheEffectorsEditorModule.h b/Engine/Plugins/Experimental/Avalanche/Source/AvalancheEffectorsEditor/Public/AvalancheEffectorsEditorModule.h
deleted file mode 100644
