[MotionDesign] Playback Service Messages: Removing comments about the idea of supporting a playback graph directly on the server side.

The way the playback graph would be implemented on the playack server side was either that it would directly run the playback graph asset, or it would create a transient playback graph that wraps the playable.
The requirements for the Rundown have evolved so that the second approach was heavily favored and the first approach was never used. It can therefore be considered not supported. And thus why indications to the contrary in the comments are removed for now.

#rnx

diff --git a/Engine/Plugins/Experimental/Avalanche/Source/AvalancheMedia/Public/Playback/AvaPlaybackMessages.h b/Engine/Plugins/Experimental/Avalanche/Source/AvalancheMedia/Public/Playback/AvaPlaybackMessages.h
--- a/Engine/Plugins/Experimental/Avalanche/Source/AvalancheMedia/Public/Playback/AvaPlaybackMessages.h
+++ b/Engine/Plugins/Experimental/Avalanche/Source/AvalancheMedia/Public/Playback/AvaPlaybackMessages.h
@@ -308,7 +308,7 @@ struct AVALANCHEMEDIA_API FAvaPlaybackCommand
 	UPROPERTY()
 	FSoftObjectPath AssetPath;
 
-	/** Channel to Play, only considered when the Asset in question is an Motion Design Playable. */
+	/** Broadcast channel the action will apply to. */
 	UPROPERTY()
 	FString ChannelName;
 	
@@ -411,7 +411,6 @@ struct FAvaPlaybackAnimPlaybackRequest : public FAvaPlaybackClientMessageBase
 	UPROPERTY()
 	FSoftObjectPath AssetPath;
 	
-	/** Channel to Play, only considered when the Asset in question is an Motion Design Blueprint */
 	UPROPERTY()
 	FString ChannelName;
 	
@@ -456,8 +455,7 @@ struct FAvaPlaybackRemoteControlUpdateRequest : public FAvaPlaybackClientMessage
 	
 	UPROPERTY()
 	FSoftObjectPath AssetPath;
-	
-	/** Channel to Play, only considered when the Asset in question is an Motion Design Blueprint. */
+
 	UPROPERTY()
 	FString ChannelName;
 	
