MotionDesign
- MotionDesign cloner effector plugin updated to fix module names changes due to refactoring

#changelist validated
#virtualized

diff --git a/Engine/Plugins/Experimental/Avalanche/Source/AvalancheEffectors/Private/AvalancheEffectorsModule.cpp b/Engine/Plugins/Experimental/Avalanche/Source/AvalancheEffectors/Private/AvalancheEffectorsModule.cpp
--- a/Engine/Plugins/Experimental/Avalanche/Source/AvalancheEffectors/Private/AvalancheEffectorsModule.cpp
+++ b/Engine/Plugins/Experimental/Avalanche/Source/AvalancheEffectors/Private/AvalancheEffectorsModule.cpp
@@ -93,14 +93,19 @@ TArray<AActor*> FAvalancheEffectorsModule::GetOrderedChildrenActors(const AActor
 	{
 		TArray<AActor*> AttachedActors;
 		InParentActor->GetAttachedActors(AttachedActors, true, false);
-		ChildrenActors.Reserve(AttachedActors.Num());
+		ChildrenActors.SetNumUninitialized(AttachedActors.Num(), EAllowShrinking::No);
 		const FAvaSceneTree& SceneTree = SceneInterface->GetSceneTree();
 
 		for (AActor* ChildActor : AttachedActors)
 		{
 			if (const FAvaSceneTreeNode* ChildNode = SceneTree.FindTreeNode(FAvaSceneItem(ChildActor, ClonerWorld)))
 			{
-				ChildrenActors.Insert(ChildActor, ChildNode->GetLocalIndex());
+				const int32 ChildIndex = ChildNode->GetLocalIndex();
+
+				if (ChildrenActors.IsValidIndex(ChildIndex))
+				{
+					ChildrenActors[ChildIndex] = ChildActor;
+				}
 			}
 		}
 	}
@@ -112,4 +117,4 @@ TArray<AActor*> FAvalancheEffectorsModule::GetOrderedChildrenActors(const AActor
 	return ChildrenActors;
 }
 
-IMPLEMENT_MODULE(FAvalancheEffectorsModule, AvalancheEffectorsModule)
+IMPLEMENT_MODULE(FAvalancheEffectorsModule, AvalancheEffectors)
diff --git a/Engine/Plugins/Experimental/Avalanche/Source/AvalancheEffectorsEditor/Private/AvalancheEffectorsEditorModule.cpp b/Engine/Plugins/Experimental/Avalanche/Source/AvalancheEffectorsEditor/Private/AvalancheEffectorsEditorModule.cpp
--- a/Engine/Plugins/Experimental/Avalanche/Source/AvalancheEffectorsEditor/Private/AvalancheEffectorsEditorModule.cpp
+++ b/Engine/Plugins/Experimental/Avalanche/Source/AvalancheEffectorsEditor/Private/AvalancheEffectorsEditorModule.cpp
@@ -63,4 +63,4 @@ void FAvalancheEffectorsEditorModule::RegisterComponentVisualizers()
 	IAvalancheComponentVisualizersModule::RegisterComponentVisualizer<UCEClonerComponent, FAvaClonerActorVisualizer>(&Visualizers);
 }
 
-IMPLEMENT_MODULE(FAvalancheEffectorsEditorModule, AvalancheEffectorsEditorModule)
+IMPLEMENT_MODULE(FAvalancheEffectorsEditorModule, AvalancheEffectorsEditor)
