Motion Design:
- Fixed motion design category name

#changelist validated
#virtualized

diff --git a/Engine/Plugins/Experimental/Avalanche/Source/AvalancheEditor/Private/AvaEditorModule.cpp b/Engine/Plugins/Experimental/Avalanche/Source/AvalancheEditor/Private/AvaEditorModule.cpp
--- a/Engine/Plugins/Experimental/Avalanche/Source/AvalancheEditor/Private/AvaEditorModule.cpp
+++ b/Engine/Plugins/Experimental/Avalanche/Source/AvalancheEditor/Private/AvaEditorModule.cpp
@@ -79,8 +79,6 @@ struct FAvaViewportColorPickerLightAdapter : public FAvaViewportColorPickerActor
 	}
 };
 
-const FName AvalancheCategoryName(TEXT("AvalancheCategory"));
-
 namespace UE::AvalancheEditor::Private
 {
 	static FString LevelTemplatesPath = FString::Printf(TEXT("/%hs/%s"), UE_PLUGIN_NAME, TEXT("LevelTemplates"));
@@ -181,7 +179,7 @@ void FAvaEditorModule::RegisterAssetTools()
 {
 	IAssetTools& AssetTools = FAssetToolsModule::GetModule().Get();
 
-	AssetTools.RegisterAdvancedAssetCategory(AvalancheCategoryName
+	AssetTools.RegisterAdvancedAssetCategory(TEXT("MotionDesignCategory")
 		, LOCTEXT("MotionDesignCategoryName", "Motion Design"));
 
 }
diff --git a/Engine/Plugins/Experimental/Avalanche/Source/AvalancheTagEditor/Private/AvaTagCollectionFactory.cpp b/Engine/Plugins/Experimental/Avalanche/Source/AvalancheTagEditor/Private/AvaTagCollectionFactory.cpp
--- a/Engine/Plugins/Experimental/Avalanche/Source/AvalancheTagEditor/Private/AvaTagCollectionFactory.cpp
+++ b/Engine/Plugins/Experimental/Avalanche/Source/AvalancheTagEditor/Private/AvaTagCollectionFactory.cpp
@@ -15,7 +15,7 @@ UAvaTagCollectionFactory::UAvaTagCollectionFactory()
 uint32 UAvaTagCollectionFactory::GetMenuCategories() const
 {
 	IAssetTools& AssetTools = FAssetToolsModule::GetModule().Get();
-	return AssetTools.FindAdvancedAssetCategory("AvalancheCategory");
+	return AssetTools.FindAdvancedAssetCategory(TEXT("MotionDesignCategory"));
 }
 
 UObject* UAvaTagCollectionFactory::FactoryCreateNew(UClass* InClass, UObject* InParent, FName InName, EObjectFlags InFlags, UObject* InContext, FFeedbackContext* InWarn)
diff --git a/Engine/Plugins/Experimental/Avalanche/Source/AvalancheTransitionEditor/Private/AvaTransitionTreeFactory.cpp b/Engine/Plugins/Experimental/Avalanche/Source/AvalancheTransitionEditor/Private/AvaTransitionTreeFactory.cpp
--- a/Engine/Plugins/Experimental/Avalanche/Source/AvalancheTransitionEditor/Private/AvaTransitionTreeFactory.cpp
+++ b/Engine/Plugins/Experimental/Avalanche/Source/AvalancheTransitionEditor/Private/AvaTransitionTreeFactory.cpp
@@ -17,7 +17,7 @@ UAvaTransitionTreeFactory::UAvaTransitionTreeFactory()
 uint32 UAvaTransitionTreeFactory::GetMenuCategories() const
 {
 	IAssetTools& AssetTools = FAssetToolsModule::GetModule().Get();
-	return AssetTools.FindAdvancedAssetCategory("AvalancheCategory");
+	return AssetTools.FindAdvancedAssetCategory(TEXT("MotionDesignCategory"));
 }
 
 UObject* UAvaTransitionTreeFactory::FactoryCreateNew(UClass* InClass, UObject* InParent, FName InName, EObjectFlags InFlags, UObject* InContext, FFeedbackContext* InWarn)
