MotionDesign
- Fix include missing for module AvaEffectors

#changelist validated
#virtualized

diff --git a/Engine/Plugins/Experimental/Avalanche/Source/AvalancheEffectors/Private/AvalancheEffectorsModule.cpp b/Engine/Plugins/Experimental/Avalanche/Source/AvalancheEffectors/Private/AvalancheEffectorsModule.cpp
--- a/Engine/Plugins/Experimental/Avalanche/Source/AvalancheEffectors/Private/AvalancheEffectorsModule.cpp
+++ b/Engine/Plugins/Experimental/Avalanche/Source/AvalancheEffectors/Private/AvalancheEffectorsModule.cpp
@@ -7,6 +7,7 @@
 #include "IAvaSceneInterface.h"
 #include "Engine/World.h"
 #include "GameFramework/Actor.h"
+#include "Modules/ModuleManager.h"
 #include "Subsystems/CEClonerSubsystem.h"
 
 #if WITH_EDITOR
