Motion Design RC:
- Moved RC Tracker Component Outliner item to private
- Moved GetCurrentPreset for Tracker Components to public

#changelist validated
#virtualized

diff --git a/Engine/Plugins/Experimental/Avalanche/Source/AvalancheEditor/Private/RemoteControl/AvaRCExtension.cpp b/Engine/Plugins/Experimental/Avalanche/Source/AvalancheEditor/Private/RemoteControl/AvaRCExtension.cpp
--- a/Engine/Plugins/Experimental/Avalanche/Source/AvalancheEditor/Private/RemoteControl/AvaRCExtension.cpp
+++ b/Engine/Plugins/Experimental/Avalanche/Source/AvalancheEditor/Private/RemoteControl/AvaRCExtension.cpp
@@ -5,7 +5,6 @@
 #include "Framework/Docking/LayoutExtender.h"
 #include "IAvaSceneInterface.h"
 #include "LevelEditor.h"
-#include "Outliner/AvaOutlinerRCTrackerComponent.h"
 #include "RemoteControlTrackerComponent.h"
 #include "Subsystems/AssetEditorSubsystem.h"
 #include "Subsystems/RemoteControlComponentsSubsystem.h"
@@ -21,13 +20,11 @@ URemoteControlPreset* FAvaRCExtension::GetRemoteControlPreset() const
 
 void FAvaRCExtension::Activate()
 {
-	FAvaOutlinerRCTrackerComponent::OnRemoteControlTrackerSelected().AddSP(this, &FAvaRCExtension::OnRemoteControlTrackerOutlinerSelection);
 	OpenRemoteControlTab();
 }
 
 void FAvaRCExtension::Deactivate()
 {
-	FAvaOutlinerRCTrackerComponent::OnRemoteControlTrackerSelected().RemoveAll(this);
 	CloseRemoteControlTab();
 }
 
@@ -86,9 +83,4 @@ void FAvaRCExtension::CloseRemoteControlTab() const
 	}
 }
 
-void FAvaRCExtension::OnRemoteControlTrackerOutlinerSelection(const URemoteControlTrackerComponent* InTrackerComponent) const
-{
-	OpenRemoteControlTab();
-}
-
 #undef LOCTEXT_NAMESPACE
diff --git a/Engine/Plugins/Experimental/Avalanche/Source/AvalancheEditor/Private/RemoteControl/AvaRCExtension.h b/Engine/Plugins/Experimental/Avalanche/Source/AvalancheEditor/Private/RemoteControl/AvaRCExtension.h
--- a/Engine/Plugins/Experimental/Avalanche/Source/AvalancheEditor/Private/RemoteControl/AvaRCExtension.h
+++ b/Engine/Plugins/Experimental/Avalanche/Source/AvalancheEditor/Private/RemoteControl/AvaRCExtension.h
@@ -27,6 +27,4 @@ protected:
 	void OpenRemoteControlTab() const;
 
 	void CloseRemoteControlTab() const;
-
-	void OnRemoteControlTrackerOutlinerSelection(const URemoteControlTrackerComponent* InTrackerComponent) const;
 };
diff --git a/Engine/Plugins/Experimental/Avalanche/Source/AvalancheRemoteControlEditor/AvalancheRemoteControlEditor.Build.cs b/Engine/Plugins/Experimental/Avalanche/Source/AvalancheRemoteControlEditor/AvalancheRemoteControlEditor.Build.cs
--- a/Engine/Plugins/Experimental/Avalanche/Source/AvalancheRemoteControlEditor/AvalancheRemoteControlEditor.Build.cs
+++ b/Engine/Plugins/Experimental/Avalanche/Source/AvalancheRemoteControlEditor/AvalancheRemoteControlEditor.Build.cs
@@ -8,27 +8,22 @@ public class AvalancheRemoteControlEditor : ModuleRules
     {
         PCHUsage = ModuleRules.PCHUsageMode.UseExplicitOrSharedPCHs;
 
-        PublicDependencyModuleNames.AddRange(
-            new string[]
-            {
-	            "AvalancheEditorCore",
-	            "AvalancheOutliner",
-                "Core",
-            }
-        );
-
         PrivateDependencyModuleNames.AddRange(
             new string[]
             {
+                "AvalancheEditorCore",
+                "AvalancheOutliner",
                 "AvalancheRemoteControl",
+                "Core",
                 "CoreUObject",
                 "Engine",
-                "RemoteControlUI",
                 "RemoteControlComponents",
                 "RemoteControlComponentsEditor",
+                "RemoteControlUI",
                 "Slate",
                 "SlateCore",
                 "ToolMenus",
+                "UnrealEd",
             }
         );
 
diff --git a/Engine/Plugins/Experimental/Avalanche/Source/AvalancheRemoteControlEditor/Private/Outliner/AvaOutlinerRCTrackerComponent.cpp b/Engine/Plugins/Experimental/Avalanche/Source/AvalancheRemoteControlEditor/Private/Outliner/AvaOutlinerRCTrackerComponent.cpp
--- a/Engine/Plugins/Experimental/Avalanche/Source/AvalancheRemoteControlEditor/Private/Outliner/AvaOutlinerRCTrackerComponent.cpp
+++ b/Engine/Plugins/Experimental/Avalanche/Source/AvalancheRemoteControlEditor/Private/Outliner/AvaOutlinerRCTrackerComponent.cpp
@@ -1,14 +1,32 @@
 // Copyright Epic Games, Inc. All Rights Reserved.
 
 #include "Outliner/AvaOutlinerRCTrackerComponent.h"
+#include "Editor.h"
 #include "IAvaOutliner.h"
+#include "RemoteControlPreset.h"
 #include "RemoteControlTrackerComponent.h"
 #include "Selection/AvaOutlinerScopedSelection.h"
+#include "Subsystems/AssetEditorSubsystem.h"
 #include "Subsystems/RemoteControlComponentsEditorUtils.h"
 
 #define LOCTEXT_NAMESPACE "AvaOutlinerRemoteControlComponent"
 
-FAvaOutlinerRCTrackerComponent::FOnAvaOutlinerRCTrackerComponentSelected FAvaOutlinerRCTrackerComponent::OnRCTrackerSelectedDelegate;
+namespace UE::AvaRCEditor::Private
+{
+	void OpenRCPreset(const URemoteControlTrackerComponent& InTrackerComponent)
+	{
+		UAssetEditorSubsystem* AssetEditorSubsystem = GEditor->GetEditorSubsystem<UAssetEditorSubsystem>();
+		if (!AssetEditorSubsystem)
+		{
+			return;
+		}
+
+		if (URemoteControlPreset* RCPreset = InTrackerComponent.GetCurrentPreset())
+		{
+			AssetEditorSubsystem->OpenEditorForAsset(RCPreset, EToolkitMode::WorldCentric);
+		}
+	}
+}
 
 FAvaOutlinerRCTrackerComponent::FAvaOutlinerRCTrackerComponent(IAvaOutliner& InOutliner, URemoteControlTrackerComponent* InComponent)
 	: FAvaOutlinerObject(InOutliner, InComponent)
@@ -45,7 +63,7 @@ void FAvaOutlinerRCTrackerComponent::Select(FAvaOutlinerScopedSelection& InSelec
 	{
 		if (InSelection.IsSelected(TrackerComponent))
 		{
-			OnRemoteControlTrackerSelected().Broadcast(TrackerComponent);
+			UE::AvaRCEditor::Private::OpenRCPreset(*TrackerComponent);
 		}
 	}
 }
diff --git a/Engine/Plugins/Experimental/Avalanche/Source/AvalancheRemoteControlEditor/Private/Outliner/AvaOutlinerRCTrackerComponent.h b/Engine/Plugins/Experimental/Avalanche/Source/AvalancheRemoteControlEditor/Private/Outliner/AvaOutlinerRCTrackerComponent.h
new file mode 100644
--- /dev/null
+++ b/Engine/Plugins/Experimental/Avalanche/Source/AvalancheRemoteControlEditor/Private/Outliner/AvaOutlinerRCTrackerComponent.h
@@ -0,0 +1,37 @@
+// Copyright Epic Games, Inc. All Rights Reserved.
+
+#pragma once
+
+#include "Item/AvaOutlinerObject.h"
+#include "Textures/SlateIcon.h"
+
+class URemoteControlTrackerComponent;
+
+/**
+ * Inheriting from FAvaOutlinerObject and not from FAvaOutlinerComponent, which only takes into account Scene Components.
+ * This class needs instead to handle URemoteControlTrackerComponent, which is of type UActorComponent, and not Scene Component.
+ */
+class FAvaOutlinerRCTrackerComponent : public FAvaOutlinerObject
+{
+public:
+	UE_AVA_INHERITS_WITH_SUPER(FAvaOutlinerRCTrackerComponent, FAvaOutlinerObject);
+
+	FAvaOutlinerRCTrackerComponent(IAvaOutliner& InOutliner, URemoteControlTrackerComponent* InComponent);
+
+	//~ Begin IAvaOutlinerItem
+	virtual FSlateIcon GetIcon() const override;
+	virtual bool ShowVisibility(EAvaOutlinerVisibilityType InVisibilityType) const override;
+	virtual bool GetVisibility(EAvaOutlinerVisibilityType InVisibilityType) const override;
+	virtual void OnVisibilityChanged(EAvaOutlinerVisibilityType InVisibilityType, bool bInNewVisibility) override;
+	virtual void Select(FAvaOutlinerScopedSelection& InSelection) const override;
+	//~ End IAvaOutlinerItem
+
+protected:
+	//~ Begin FAvaOutlinerObjectItem
+	virtual void SetObject_Impl(UObject* InObject) override;
+	//~ End FAvaOutlinerObjectItem
+
+	TWeakObjectPtr<URemoteControlTrackerComponent> TrackerComponentWeak;
+
+	FSlateIcon TrackerIcon;
+};
diff --git a/Engine/Plugins/Experimental/Avalanche/Source/AvalancheRemoteControlEditor/Public/Outliner/AvaOutlinerRCTrackerComponent.h b/Engine/Plugins/Experimental/Avalanche/Source/AvalancheRemoteControlEditor/Public/Outliner/AvaOutlinerRCTrackerComponent.h
deleted file mode 100644
