[MotionDesign] Fix playback server crash (regression).

#rnx

diff --git a/Engine/Plugins/Experimental/Avalanche/Source/AvalancheMedia/Private/Playback/Http/AvaPlaybackHttpServer.cpp b/Engine/Plugins/Experimental/Avalanche/Source/AvalancheMedia/Private/Playback/Http/AvaPlaybackHttpServer.cpp
--- a/Engine/Plugins/Experimental/Avalanche/Source/AvalancheMedia/Private/Playback/Http/AvaPlaybackHttpServer.cpp
+++ b/Engine/Plugins/Experimental/Avalanche/Source/AvalancheMedia/Private/Playback/Http/AvaPlaybackHttpServer.cpp
@@ -51,7 +51,7 @@ void FAvaPlaybackHttpServer::Start(int32 InPortToUse)
 
 void FAvaPlaybackHttpServer::RegisterRoutes()
 {
-	const TSharedPtr<FAvaPlaybackServer> MediaPlaybackServer = FModuleManager::GetModulePtr<FAvaMediaModule>(UE_PLUGIN_NAME)->GetMediaPlaybackServer();
+	const TSharedPtr<FAvaPlaybackServer> MediaPlaybackServer = FModuleManager::GetModulePtr<FAvaMediaModule>(UE_MODULE_NAME)->GetMediaPlaybackServer();
 	check(MediaPlaybackServer.IsValid());
 	
 	// Map as per FAvaPlaybackServer::Init(const FString& AssignedServerName)
