[Motion Design]
 - Hook up operations for In and Out separatedly

#changelist validated
#virtualized

diff --git a/Engine/Plugins/Experimental/Avalanche/Source/AvalancheSequencer/Private/EaseCurveTool/AvaEaseCurveTool.cpp b/Engine/Plugins/Experimental/Avalanche/Source/AvalancheSequencer/Private/EaseCurveTool/AvaEaseCurveTool.cpp
--- a/Engine/Plugins/Experimental/Avalanche/Source/AvalancheSequencer/Private/EaseCurveTool/AvaEaseCurveTool.cpp
+++ b/Engine/Plugins/Experimental/Avalanche/Source/AvalancheSequencer/Private/EaseCurveTool/AvaEaseCurveTool.cpp
@@ -152,9 +152,24 @@ FAvaEaseCurveTangents FAvaEaseCurveTool::GetEaseCurveTangents() const
 	return EaseCurve->GetTangents();
 }
 
-void FAvaEaseCurveTool::SetEaseCurveTangents(const FAvaEaseCurveTangents& InTangents, const bool bInBroadcastUpdate, const bool bInSetSequencerTangents)
+void FAvaEaseCurveTool::SetEaseCurveTangents(const FAvaEaseCurveTangents& InTangents, const EOperation InOperation
+	, const bool bInBroadcastUpdate, const bool bInSetSequencerTangents)
 {
-	EaseCurve->SetTangents(InTangents);
+	const FScopedTransaction Transaction(LOCTEXT("SetEaseCurveTangents", "Set Ease Curve Tangents"));
+	EaseCurve->Modify();
+
+	switch (InOperation)
+	{
+	case EOperation::InOut:
+		EaseCurve->SetTangents(InTangents);
+		break;
+	case EOperation::In:
+		EaseCurve->SetEndTangent(InTangents.End, InTangents.EndWeight);
+		break;
+	case EOperation::Out:
+		EaseCurve->SetStartTangent(InTangents.Start, InTangents.StartWeight);
+		break;
+	}
 
 	if (bInBroadcastUpdate)
 	{
@@ -163,83 +178,74 @@ void FAvaEaseCurveTool::SetEaseCurveTangents(const FAvaEaseCurveTangents& InTang
 
 	if (bInSetSequencerTangents)
 	{
-		SetSequencerKeySelectionTangents(InTangents);
+		SetSequencerKeySelectionTangents(InTangents, InOperation);
 	}
 }
 
-void FAvaEaseCurveTool::ResetEaseCurveTangents(const bool bInStartTangent, const bool bInEndTangent)
+void FAvaEaseCurveTool::ResetEaseCurveTangents(const EOperation InOperation)
 {
 	FText TransactionText;
-	if (bInStartTangent && !bInEndTangent)
-	{
-		TransactionText = LOCTEXT("ResetStartTangents", "Reset Start Tangents");
-	}
-	else if (!bInStartTangent && bInEndTangent)
-	{
-		TransactionText = LOCTEXT("ResetEndTangents", "Reset End Tangents");
-	}
-	else
+
+	switch (InOperation)
 	{
+	case EOperation::InOut:
 		TransactionText = LOCTEXT("ResetTangents", "Reset Tangents");
+		break;
+	case EOperation::In:
+		TransactionText = LOCTEXT("ResetEndTangents", "Reset End Tangents");
+		break;
+	case EOperation::Out:
+		TransactionText = LOCTEXT("ResetStartTangents", "Reset Start Tangents");
+		break;
 	}
+
 	const FScopedTransaction Transaction(TransactionText);
 	EaseCurve->ModifyOwner();
 
-	FAvaEaseCurveTangents NewTangents = GetEaseCurveTangents();
-	if (bInStartTangent)
-	{
-		NewTangents.Start = 0.f;
-		NewTangents.StartWeight = 0.f;
-	}
-	if (bInEndTangent)
-	{
-		NewTangents.End = 0.f;
-		NewTangents.EndWeight = 0.f;
-	}
-	SetEaseCurveTangents(NewTangents, true, true);
+	SetEaseCurveTangents(FAvaEaseCurveTangents(0.0, 0.0, 0.0, 0.0), InOperation, true, true);
 }
 
-void FAvaEaseCurveTool::FlattenOrStraightenTangents(const bool bInStartTangent, const bool bInEndTangent, const bool bInFlattenTangents) const
+void FAvaEaseCurveTool::FlattenOrStraightenTangents(const EOperation InOperation, const bool bInFlattenTangents) const
 {
 	FText TransactionText;
 	if (bInFlattenTangents)
 	{
-		if (bInStartTangent && !bInEndTangent)
-		{
-			TransactionText = LOCTEXT("FlattenStartTangents", "Flatten Start Tangents");
-		}
-		else if (!bInStartTangent && bInEndTangent)
-		{
-			TransactionText = LOCTEXT("FlattenEndTangents", "Flatten End Tangents");
-		}
-		else
+		switch (InOperation)
 		{
+		case EOperation::InOut:
 			TransactionText = LOCTEXT("FlattenTangents", "Flatten Tangents");
+			break;
+		case EOperation::In:
+			TransactionText = LOCTEXT("FlattenEndTangents", "Flatten End Tangents");
+			break;
+		case EOperation::Out:
+			TransactionText = LOCTEXT("FlattenStartTangents", "Flatten Start Tangents");
+			break;
 		}
 	}
 	else
 	{
-		if (bInStartTangent && !bInEndTangent)
-		{
-			TransactionText = LOCTEXT("StraightenStartTangents", "Straighten Start Tangents");
-		}
-		else if (!bInStartTangent && bInEndTangent)
-		{
-			TransactionText = LOCTEXT("StraightenEndTangents", "Straighten End Tangents");
-		}
-		else
+		switch (InOperation)
 		{
+		case EOperation::InOut:
 			TransactionText = LOCTEXT("StraightenTangents", "Straighten Tangents");
+			break;
+		case EOperation::In:
+			TransactionText = LOCTEXT("StraightenEndTangents", "Straighten End Tangents");
+			break;
+		case EOperation::Out:
+			TransactionText = LOCTEXT("StraightenStartTangents", "Straighten Start Tangents");
+			break;
 		}
 	}
 	const FScopedTransaction Transaction(TransactionText);
 	EaseCurve->ModifyOwner();
 
-	if (bInStartTangent)
+	if (InOperation == EOperation::Out || InOperation == EOperation::InOut)
 	{
 		EaseCurve->FlattenOrStraightenTangents(EaseCurve->GetStartKeyHandle(), bInFlattenTangents);
 	}
-	if (bInEndTangent)
+	if (InOperation == EOperation::In || InOperation == EOperation::InOut)
 	{
 		EaseCurve->FlattenOrStraightenTangents(EaseCurve->GetEndKeyHandle(), bInFlattenTangents);
 	}
@@ -264,7 +270,7 @@ void FAvaEaseCurveTool::ApplyQuickEaseToSequencerKeySelections(const EOperation
 		return;
 	}
 
-	SetEaseCurveTangents(Tangents, true, true);
+	SetEaseCurveTangents(Tangents, InOperation, true, true);
 
 	if (ToolWidget.IsValid())
 	{
@@ -431,7 +437,7 @@ void FAvaEaseCurveTool::UpdateEaseCurveFromSequencerKeySelections()
 		}
 	}
 
-	SetEaseCurveTangents(FAvaEaseCurveTangents::Average(KeySetTangents), true, false);
+	SetEaseCurveTangents(FAvaEaseCurveTangents::Average(KeySetTangents), EOperation::InOut, true, false);
 }
 
 UCurveBase* FAvaEaseCurveTool::CreateCurveAsset() const
diff --git a/Engine/Plugins/Experimental/Avalanche/Source/AvalancheSequencer/Private/EaseCurveTool/AvaEaseCurveTool.h b/Engine/Plugins/Experimental/Avalanche/Source/AvalancheSequencer/Private/EaseCurveTool/AvaEaseCurveTool.h
--- a/Engine/Plugins/Experimental/Avalanche/Source/AvalancheSequencer/Private/EaseCurveTool/AvaEaseCurveTool.h
+++ b/Engine/Plugins/Experimental/Avalanche/Source/AvalancheSequencer/Private/EaseCurveTool/AvaEaseCurveTool.h
@@ -73,10 +73,13 @@ public:
 	FRichCurve* GetToolRichCurve() const;
 
 	FAvaEaseCurveTangents GetEaseCurveTangents() const;
-	void SetEaseCurveTangents(const FAvaEaseCurveTangents& InTangents, const bool bBroadcastUpdate, const bool bInSetSequencerTangents);
-	void ResetEaseCurveTangents(const bool bInStartTangent, const bool bInEndTangent);
 
-	void FlattenOrStraightenTangents(const bool bInStartTangent, const bool bInEndTangent, const bool bInFlattenTangents) const;
+	void SetEaseCurveTangents(const FAvaEaseCurveTangents& InTangents, const EOperation InOperation,
+		const bool bBroadcastUpdate, const bool bInSetSequencerTangents);
+
+	void ResetEaseCurveTangents(const EOperation InOperation);
+
+	void FlattenOrStraightenTangents(const EOperation InOperation, const bool bInFlattenTangents) const;
 
 	FORCEINLINE EOperation GetOperation() const { return OperationMode; }
 	FORCEINLINE void SetOperation(const EOperation InOperation) { OperationMode = InOperation; }
diff --git a/Engine/Plugins/Experimental/Avalanche/Source/AvalancheSequencer/Private/EaseCurveTool/Widgets/SAvaEaseCurveEditor.cpp b/Engine/Plugins/Experimental/Avalanche/Source/AvalancheSequencer/Private/EaseCurveTool/Widgets/SAvaEaseCurveEditor.cpp
--- a/Engine/Plugins/Experimental/Avalanche/Source/AvalancheSequencer/Private/EaseCurveTool/Widgets/SAvaEaseCurveEditor.cpp
+++ b/Engine/Plugins/Experimental/Avalanche/Source/AvalancheSequencer/Private/EaseCurveTool/Widgets/SAvaEaseCurveEditor.cpp
@@ -396,7 +396,7 @@ void SAvaEaseCurveEditor::CreateLinesForSegment(const ERichCurveInterpMode InInt
 	, const FTrackScaleInfo& InScaleInfo) const
 {
 	FLinearColor FadedCurveColor = CurveColor;
-	FadedCurveColor.A = 0.2f;
+	FadedCurveColor.A = 0.05f;
 
 	switch (InInterpMode)
 	{
diff --git a/Engine/Plugins/Experimental/Avalanche/Source/AvalancheSequencer/Private/EaseCurveTool/Widgets/SAvaEaseCurvePresetGroupItem.cpp b/Engine/Plugins/Experimental/Avalanche/Source/AvalancheSequencer/Private/EaseCurveTool/Widgets/SAvaEaseCurvePresetGroupItem.cpp
--- a/Engine/Plugins/Experimental/Avalanche/Source/AvalancheSequencer/Private/EaseCurveTool/Widgets/SAvaEaseCurvePresetGroupItem.cpp
+++ b/Engine/Plugins/Experimental/Avalanche/Source/AvalancheSequencer/Private/EaseCurveTool/Widgets/SAvaEaseCurvePresetGroupItem.cpp
@@ -299,20 +299,20 @@ FText SAvaEaseCurvePresetGroupItem::GetQuickPresetIconToolTip() const
 
 	if (EaseCurveToolCommands.QuickEase->GetFirstValidChord()->IsValidChord())
 	{
-		CommandText = FText::Format(LOCTEXT("QuickEaseIconInputInOutTooltip", "{0} - Apply quick preset to Out (Leave) and In (Arrive) tangents\n")
-			, EaseCurveToolCommands.QuickEase->GetInputText());
+		CommandText = FText::Format(LOCTEXT("QuickEaseIconInOutTooltip", "{0}{1} - Apply to Out (Leave) and In (Arrive) tangents\n")
+			, CommandText, EaseCurveToolCommands.QuickEase->GetInputText());
 	}
 
 	if (EaseCurveToolCommands.QuickEaseIn->GetFirstValidChord()->IsValidChord())
 	{
-		CommandText = FText::Format(LOCTEXT("QuickEaseIconInputInTooltip", "{0} - Apply quick preset to In (Arrive) tangent only\n")
-			, EaseCurveToolCommands.QuickEaseIn->GetInputText());
+		CommandText = FText::Format(LOCTEXT("QuickEaseIconInTooltip", "{0}{1} - Apply to In (Arrive) tangent only\n")
+			, CommandText, EaseCurveToolCommands.QuickEaseIn->GetInputText());
 	}
 
 	if (EaseCurveToolCommands.QuickEaseOut->GetFirstValidChord()->IsValidChord())
 	{
-		CommandText = FText::Format(LOCTEXT("QuickEaseIconInputOutTooltip", "{0} - Apply quick preset to Out (Leave) tangent only\n")
-			, EaseCurveToolCommands.QuickEaseOut->GetInputText());
+		CommandText = FText::Format(LOCTEXT("QuickEaseIconOutTooltip", "{0}{1} - Apply to Out (Leave) tangent only\n")
+			, CommandText, EaseCurveToolCommands.QuickEaseOut->GetInputText());
 	}
 
 	return CommandText.IsEmpty() ? QuickEaseText : FText::Format(LOCTEXT("QuickEasePresetIconTooltip", "{0}\n\n{1}"), QuickEaseText, CommandText);
diff --git a/Engine/Plugins/Experimental/Avalanche/Source/AvalancheSequencer/Private/EaseCurveTool/Widgets/SAvaEaseCurveTool.cpp b/Engine/Plugins/Experimental/Avalanche/Source/AvalancheSequencer/Private/EaseCurveTool/Widgets/SAvaEaseCurveTool.cpp
--- a/Engine/Plugins/Experimental/Avalanche/Source/AvalancheSequencer/Private/EaseCurveTool/Widgets/SAvaEaseCurveTool.cpp
+++ b/Engine/Plugins/Experimental/Avalanche/Source/AvalancheSequencer/Private/EaseCurveTool/Widgets/SAvaEaseCurveTool.cpp
@@ -388,7 +388,7 @@ TSharedRef<SWidget> SAvaEaseCurveTool::ConstructTangentNumBox(const FText& InLab
 
 void SAvaEaseCurveTool::HandleEditorTangentsChanged(const FAvaEaseCurveTangents& InTangents) const
 {
-	SetTangents(InTangents, true, true, true);
+	SetTangents(InTangents, ToolOperation.Get(), true, true, true);
 }
 
 void SAvaEaseCurveTool::OnEditorDragStart() const
@@ -401,7 +401,8 @@ void SAvaEaseCurveTool::OnEditorDragEnd() const
 	EaseCurveTool->EndTransaction();
 }
 
-void SAvaEaseCurveTool::SetTangents(const FAvaEaseCurveTangents& InTangents, const bool bInSetEaseCurve, const bool bInBroadcastUpdate, const bool bInSetSequencerTangents) const
+void SAvaEaseCurveTool::SetTangents(const FAvaEaseCurveTangents& InTangents, FAvaEaseCurveTool::EOperation InOperation
+	, const bool bInSetEaseCurve, const bool bInBroadcastUpdate, const bool bInSetSequencerTangents) const
 {
 	if (CurvePresetWidget.IsValid())
 	{
@@ -414,7 +415,7 @@ void SAvaEaseCurveTool::SetTangents(const FAvaEaseCurveTangents& InTangents, con
 	// To change the graph UI tangents, we need to change the ease curve object tangents and the graph will reflect.
 	if (bInSetEaseCurve && EaseCurveTool.IsValid())
 	{
-		EaseCurveTool->SetEaseCurveTangents(InTangents, bInBroadcastUpdate, bInSetSequencerTangents);
+		EaseCurveTool->SetEaseCurveTangents(InTangents, InOperation, bInBroadcastUpdate, bInSetSequencerTangents);
 	}
 
 	if (GetDefault<UAvaEaseCurveToolSettings>()->GetAutoZoomToFit())
@@ -447,33 +448,33 @@ void SAvaEaseCurveTool::OnStartTangentSpinBoxChanged(const float InNewValue) con
 {
 	FAvaEaseCurveTangents NewTangents = EaseCurveTool->GetEaseCurveTangents();
 	NewTangents.Start = InNewValue;
-	SetTangents(NewTangents, true, true, true);
+	SetTangents(NewTangents, ToolOperation.Get(), true, true, true);
 }
 
 void SAvaEaseCurveTool::OnStartTangentWeightSpinBoxChanged(const float InNewValue) const
 {
 	FAvaEaseCurveTangents NewTangents = EaseCurveTool->GetEaseCurveTangents();
 	NewTangents.StartWeight = InNewValue;
-	SetTangents(NewTangents, true, true, true);
+	SetTangents(NewTangents, ToolOperation.Get(), true, true, true);
 }
 
 void SAvaEaseCurveTool::OnEndTangentSpinBoxChanged(const float InNewValue) const
 {
 	FAvaEaseCurveTangents NewTangents = EaseCurveTool->GetEaseCurveTangents();
 	NewTangents.End = InNewValue;
-	SetTangents(NewTangents, true, true, true);
+	SetTangents(NewTangents, ToolOperation.Get(), true, true, true);
 }
 
 void SAvaEaseCurveTool::OnEndTangentWeightSpinBoxChanged(const float InNewValue) const
 {
 	FAvaEaseCurveTangents NewTangents = EaseCurveTool->GetEaseCurveTangents();
 	NewTangents.EndWeight = InNewValue;
-	SetTangents(NewTangents, true, true, true);
+	SetTangents(NewTangents, ToolOperation.Get(), true, true, true);
 }
 
 void SAvaEaseCurveTool::OnPresetChanged(const TSharedPtr<FAvaEaseCurvePreset>& InPreset) const
 {
-	SetTangents(InPreset->Tangents, true, true, true);
+	SetTangents(InPreset->Tangents, ToolOperation.Get(), true, true, true);
 
 	FSlateApplication::Get().SetAllUserFocus(CurveEaseEditorWidget);
 }
@@ -538,31 +539,31 @@ void SAvaEaseCurveTool::BindCommands()
 		, FIsActionChecked::CreateSP(EaseCurveToolRef, &FAvaEaseCurveTool::IsToolOperation, FAvaEaseCurveTool::EOperation::In));
 
 	CommandList->MapAction(EaseCurveToolCommands.ResetTangents
-		, FExecuteAction::CreateSP(EaseCurveToolRef, &FAvaEaseCurveTool::ResetEaseCurveTangents, true, true));
+		, FExecuteAction::CreateSP(EaseCurveToolRef, &FAvaEaseCurveTool::ResetEaseCurveTangents, FAvaEaseCurveTool::EOperation::InOut));
 
 	CommandList->MapAction(EaseCurveToolCommands.ResetStartTangent
-		, FExecuteAction::CreateSP(EaseCurveToolRef, &FAvaEaseCurveTool::ResetEaseCurveTangents, true, false));
+		, FExecuteAction::CreateSP(EaseCurveToolRef, &FAvaEaseCurveTool::ResetEaseCurveTangents, FAvaEaseCurveTool::EOperation::Out));
 
 	CommandList->MapAction(EaseCurveToolCommands.ResetEndTangent
-		, FExecuteAction::CreateSP(EaseCurveToolRef, &FAvaEaseCurveTool::ResetEaseCurveTangents, false, true));
+		, FExecuteAction::CreateSP(EaseCurveToolRef, &FAvaEaseCurveTool::ResetEaseCurveTangents, FAvaEaseCurveTool::EOperation::In));
 
 	CommandList->MapAction(EaseCurveToolCommands.FlattenTangents
-		, FExecuteAction::CreateSP(EaseCurveToolRef, &FAvaEaseCurveTool::FlattenOrStraightenTangents, true, true, true));
+		, FExecuteAction::CreateSP(EaseCurveToolRef, &FAvaEaseCurveTool::FlattenOrStraightenTangents, FAvaEaseCurveTool::EOperation::InOut, true));
 
 	CommandList->MapAction(EaseCurveToolCommands.FlattenStartTangent
-		, FExecuteAction::CreateSP(EaseCurveToolRef, &FAvaEaseCurveTool::FlattenOrStraightenTangents, true, false, true));
+		, FExecuteAction::CreateSP(EaseCurveToolRef, &FAvaEaseCurveTool::FlattenOrStraightenTangents, FAvaEaseCurveTool::EOperation::Out, true));
 
 	CommandList->MapAction(EaseCurveToolCommands.FlattenEndTangent
-		, FExecuteAction::CreateSP(EaseCurveToolRef, &FAvaEaseCurveTool::FlattenOrStraightenTangents, false, true, true));
+		, FExecuteAction::CreateSP(EaseCurveToolRef, &FAvaEaseCurveTool::FlattenOrStraightenTangents, FAvaEaseCurveTool::EOperation::In, true));
 
 	CommandList->MapAction(EaseCurveToolCommands.StraightenTangents
-		, FExecuteAction::CreateSP(EaseCurveToolRef, &FAvaEaseCurveTool::FlattenOrStraightenTangents, true, true, false));
+		, FExecuteAction::CreateSP(EaseCurveToolRef, &FAvaEaseCurveTool::FlattenOrStraightenTangents, FAvaEaseCurveTool::EOperation::InOut, false));
 
 	CommandList->MapAction(EaseCurveToolCommands.StraightenStartTangent
-		, FExecuteAction::CreateSP(EaseCurveToolRef, &FAvaEaseCurveTool::FlattenOrStraightenTangents, true, false, false));
+		, FExecuteAction::CreateSP(EaseCurveToolRef, &FAvaEaseCurveTool::FlattenOrStraightenTangents, FAvaEaseCurveTool::EOperation::Out, false));
 
 	CommandList->MapAction(EaseCurveToolCommands.StraightenEndTangent
-		, FExecuteAction::CreateSP(EaseCurveToolRef, &FAvaEaseCurveTool::FlattenOrStraightenTangents, false, true, false));
+		, FExecuteAction::CreateSP(EaseCurveToolRef, &FAvaEaseCurveTool::FlattenOrStraightenTangents, FAvaEaseCurveTool::EOperation::In, false));
 
 	CommandList->MapAction(EaseCurveToolCommands.CopyTangents
 		, FExecuteAction::CreateSP(EaseCurveToolRef, &FAvaEaseCurveTool::CopyTangentsToClipboard)
diff --git a/Engine/Plugins/Experimental/Avalanche/Source/AvalancheSequencer/Private/EaseCurveTool/Widgets/SAvaEaseCurveTool.h b/Engine/Plugins/Experimental/Avalanche/Source/AvalancheSequencer/Private/EaseCurveTool/Widgets/SAvaEaseCurveTool.h
--- a/Engine/Plugins/Experimental/Avalanche/Source/AvalancheSequencer/Private/EaseCurveTool/Widgets/SAvaEaseCurveTool.h
+++ b/Engine/Plugins/Experimental/Avalanche/Source/AvalancheSequencer/Private/EaseCurveTool/Widgets/SAvaEaseCurveTool.h
@@ -33,7 +33,8 @@ public:
 
 	void Construct(const FArguments& InArgs, const TSharedRef<FAvaEaseCurveTool>& InEaseCurveTool);
 	
-	void SetTangents(const FAvaEaseCurveTangents& InTangents, const bool bInSetEaseCurve, const bool bInBroadcastUpdate, const bool bInSetSequencerTangents) const;
+	void SetTangents(const FAvaEaseCurveTangents& InTangents, FAvaEaseCurveTool::EOperation InOperation,
+		const bool bInSetEaseCurve, const bool bInBroadcastUpdate, const bool bInSetSequencerTangents) const;
 
 	float GetStartTangent() const;
 	float GetStartTangentWeight() const;
