[Motion Design]
 - Remove slate style macro definitions
 - Update style classes to inherit from FSlateStyleSet
 - Standardize style names, SetContentRoot, includes, etc

#changelist validated
#virtualized

diff --git a/Engine/Plugins/Experimental/Avalanche/Source/AvalancheEditor/Private/AvaEditorModule.cpp b/Engine/Plugins/Experimental/Avalanche/Source/AvalancheEditor/Private/AvaEditorModule.cpp
--- a/Engine/Plugins/Experimental/Avalanche/Source/AvalancheEditor/Private/AvaEditorModule.cpp
+++ b/Engine/Plugins/Experimental/Avalanche/Source/AvalancheEditor/Private/AvaEditorModule.cpp
@@ -5,7 +5,6 @@
 #include "AvaEditorActorUtils.h"
 #include "AvaEditorCommands.h"
 #include "AvaEditorIntegration.h"
-#include "AvaEditorStyle.h"
 #include "AvaShapeActor.h"
 #include "Components/LightComponentBase.h"
 #include "DynamicMeshes/AvaShapeDynMeshBase.h"
@@ -88,7 +87,6 @@ namespace UE::AvalancheEditor::Private
 
 void FAvaEditorModule::StartupModule()
 {
-	FAvaEditorStyle::Get(); // Initialises
 	FAvaEditorCommands::Register();
 
 	// Add the menu subsection
@@ -116,7 +114,6 @@ void FAvaEditorModule::StartupModule()
 
 void FAvaEditorModule::ShutdownModule()
 {
-	FAvaEditorStyle::Shutdown();
 	FAvaEditorCommands::Unregister();
 
 	FCoreDelegates::OnPostEngineInit.RemoveAll(this);
diff --git a/Engine/Plugins/Experimental/Avalanche/Source/AvalancheEditor/Private/AvaEditorStyle.cpp b/Engine/Plugins/Experimental/Avalanche/Source/AvalancheEditor/Private/AvaEditorStyle.cpp
--- a/Engine/Plugins/Experimental/Avalanche/Source/AvalancheEditor/Private/AvaEditorStyle.cpp
+++ b/Engine/Plugins/Experimental/Avalanche/Source/AvalancheEditor/Private/AvaEditorStyle.cpp
@@ -1,46 +1,27 @@
 // Copyright Epic Games, Inc. All Rights Reserved.
 
 #include "AvaEditorStyle.h"
+#include "Brushes/SlateColorBrush.h"
 #include "Brushes/SlateImageBrush.h"
-#include "Engine/Texture2D.h"
-#include "Framework/Application/SlateApplication.h"
+#include "Brushes/SlateNoResource.h"
+#include "Brushes/SlateRoundedBoxBrush.h"
+#include "Containers/StringFwd.h"
 #include "Interfaces/IPluginManager.h"
-#include "MediaPlate.h"
+#include "Layout/Margin.h"
+#include "Math/MathFwd.h"
+#include "Misc/Paths.h"
 #include "SlateOptMacros.h"
 #include "Styling/AppStyle.h"
+#include "Styling/CoreStyle.h"
 #include "Styling/SlateIconFinder.h"
 #include "Styling/SlateStyleMacros.h"
 #include "Styling/StyleColors.h"
 #include "Styling/SlateStyle.h"
 #include "Styling/SlateStyleRegistry.h"
-
-TSharedPtr<FAvaEditorStyle> FAvaEditorStyle::StyleInstance = nullptr;
-
-void FAvaEditorStyle::Shutdown()
-{
-	FSlateStyleRegistry::UnRegisterSlateStyle(*StyleInstance);
-	ensure(StyleInstance.IsUnique());
-	StyleInstance.Reset();
-}
-
-const FAvaEditorStyle& FAvaEditorStyle::Get()
-{
-	if (!StyleInstance.IsValid())
-	{
-		StyleInstance = MakeShared<FAvaEditorStyle>();
-		FSlateStyleRegistry::RegisterSlateStyle(*StyleInstance);
-	}
-
-	return *StyleInstance;
-}
+#include "Styling/SlateTypes.h"
 
 FAvaEditorStyle::FAvaEditorStyle()
-	: FSlateStyleSet(TEXT("AvalancheEditor"))
-{
-	Init();
-}
-
-namespace UE::AvalancheEditor::Private
+	: FSlateStyleSet(TEXT("AvaEditor"))
 {
 	const FVector2f Icon16(16.f);
 	const FVector2f Icon20(20.f);
@@ -48,23 +29,11 @@ namespace UE::AvalancheEditor::Private
 	const FVector2f Icon25(25.f);
 	const FVector2f Icon32(32.f);
 	const FVector2f Icon40(40.f);
-	const FVector2f Icon64(64.f);
-}
-
-BEGIN_SLATE_FUNCTION_BUILD_OPTIMIZATION
-
-void FAvaEditorStyle::Init()
-{
-	using namespace UE::AvalancheEditor::Private;
 
-	TSharedPtr<IPlugin> Plugin = IPluginManager::Get().FindPlugin(TEXT("Avalanche"));
+	TSharedPtr<IPlugin> Plugin = IPluginManager::Get().FindPlugin(UE_PLUGIN_NAME);
 	check(Plugin.IsValid());
 
-	if (Plugin.IsValid())
-	{
-		SetContentRoot(Plugin->GetBaseDir() / TEXT("Resources"));
-	}
-
+	SetContentRoot(Plugin->GetBaseDir() / TEXT("Resources"));
 	SetCoreContentRoot(FPaths::EngineContentDir() / TEXT("Editor/Slate"));
 
 	// Custom Class Icons
@@ -277,6 +246,11 @@ void FAvaEditorStyle::Init()
 	const FCheckBoxStyle& AppStyle_RadioButton = FAppStyle::GetWidgetStyle<FCheckBoxStyle>("RadioButton");
 	Set("AvalancheEditor.BlackRadioButton", FCheckBoxStyle(AppStyle_RadioButton)
 		.SetBackgroundImage(*GetBrush("AvalancheIcons.Editor.Radio.BlackBackground")));
+
+	FSlateStyleRegistry::RegisterSlateStyle(*this);
 }
 
-END_SLATE_FUNCTION_BUILD_OPTIMIZATION
+FAvaEditorStyle::~FAvaEditorStyle()
+{
+	FSlateStyleRegistry::UnRegisterSlateStyle(*this);
+}
diff --git a/Engine/Plugins/Experimental/Avalanche/Source/AvalancheEditor/Private/AvaEditorStyle.h b/Engine/Plugins/Experimental/Avalanche/Source/AvalancheEditor/Private/AvaEditorStyle.h
--- a/Engine/Plugins/Experimental/Avalanche/Source/AvalancheEditor/Private/AvaEditorStyle.h
+++ b/Engine/Plugins/Experimental/Avalanche/Source/AvalancheEditor/Private/AvaEditorStyle.h
@@ -2,20 +2,17 @@
 
 #pragma once
 
-#include "Styling/ISlateStyle.h"
 #include "Styling/SlateStyle.h"
 
-class FAvaEditorStyle : public FSlateStyleSet
+class FAvaEditorStyle final : public FSlateStyleSet
 {
 public:
-	static const FAvaEditorStyle& Get();
-
-	static void Shutdown();
+	static FAvaEditorStyle& Get()
+	{
+		static FAvaEditorStyle Instance;
+		return Instance;
+	}
 
 	FAvaEditorStyle();
-
-private:
-	static TSharedPtr<FAvaEditorStyle> StyleInstance;
-
-	void Init();
+	virtual ~FAvaEditorStyle() override;
 };
diff --git a/Engine/Plugins/Experimental/Avalanche/Source/AvalancheEditorCore/Private/LevelEditor/AvaLevelEditorStyle.cpp b/Engine/Plugins/Experimental/Avalanche/Source/AvalancheEditorCore/Private/LevelEditor/AvaLevelEditorStyle.cpp
--- a/Engine/Plugins/Experimental/Avalanche/Source/AvalancheEditorCore/Private/LevelEditor/AvaLevelEditorStyle.cpp
+++ b/Engine/Plugins/Experimental/Avalanche/Source/AvalancheEditorCore/Private/LevelEditor/AvaLevelEditorStyle.cpp
@@ -1,6 +1,7 @@
 // Copyright Epic Games, Inc. All Rights Reserved.
 
 #include "AvaLevelEditorStyle.h"
+#include "Brushes/SlateImageBrush.h"
 #include "Misc/Paths.h"
 #include "Styling/SlateStyleMacros.h"
 #include "Styling/SlateStyleRegistry.h"
@@ -10,8 +11,8 @@ FAvaLevelEditorStyle::FAvaLevelEditorStyle()
 {
 	const FVector2f Icon16x16(16.f, 16.f);
 
-	ContentRootDir     = FPaths::EngineContentDir() / TEXT("Editor/Slate");
-	CoreContentRootDir = FPaths::EngineContentDir() / TEXT("Slate");
+	SetContentRoot(FPaths::EngineContentDir() / TEXT("Editor/Slate"));
+	SetCoreContentRoot(FPaths::EngineContentDir() / TEXT("Slate"));
 
 	Set("AvaLevelEditor.CreateScene"    , new CORE_IMAGE_BRUSH_SVG("Starship/Common/plus", Icon16x16));
 	Set("AvaLevelEditor.ActivateScene"  , new CORE_IMAGE_BRUSH_SVG("Starship/Common/play", Icon16x16));
diff --git a/Engine/Plugins/Experimental/Avalanche/Source/AvalancheEditorCore/Private/LevelEditor/AvaLevelEditorStyle.h b/Engine/Plugins/Experimental/Avalanche/Source/AvalancheEditorCore/Private/LevelEditor/AvaLevelEditorStyle.h
--- a/Engine/Plugins/Experimental/Avalanche/Source/AvalancheEditorCore/Private/LevelEditor/AvaLevelEditorStyle.h
+++ b/Engine/Plugins/Experimental/Avalanche/Source/AvalancheEditorCore/Private/LevelEditor/AvaLevelEditorStyle.h
@@ -4,16 +4,15 @@
 
 #include "Styling/SlateStyle.h"
 
-class FAvaLevelEditorStyle : public FSlateStyleSet
+class FAvaLevelEditorStyle final : public FSlateStyleSet
 {
 public:
-	FAvaLevelEditorStyle();
-
-	virtual ~FAvaLevelEditorStyle() override;
-
 	static FAvaLevelEditorStyle& Get()
 	{
-		static FAvaLevelEditorStyle StyleSet;
-		return StyleSet;
+		static FAvaLevelEditorStyle Instance;
+		return Instance;
 	}
+
+	FAvaLevelEditorStyle();
+	virtual ~FAvaLevelEditorStyle() override;
 };
diff --git a/Engine/Plugins/Experimental/Avalanche/Source/AvalancheEffectorsEditor/Private/AvaEffectorsEditorStyle.cpp b/Engine/Plugins/Experimental/Avalanche/Source/AvalancheEffectorsEditor/Private/AvaEffectorsEditorStyle.cpp
--- a/Engine/Plugins/Experimental/Avalanche/Source/AvalancheEffectorsEditor/Private/AvaEffectorsEditorStyle.cpp
+++ b/Engine/Plugins/Experimental/Avalanche/Source/AvalancheEffectorsEditor/Private/AvaEffectorsEditorStyle.cpp
@@ -1,35 +1,34 @@
 // Copyright Epic Games, Inc. All Rights Reserved.
 
 #include "AvaEffectorsEditorStyle.h"
-
+#include "Brushes/SlateImageBrush.h"
 #include "Interfaces/IPluginManager.h"
-#include "Misc/Paths.h"
+#include "Math/MathFwd.h"
 #include "Styling/SlateStyleMacros.h"
 #include "Styling/StyleColors.h"
 #include "Styling/SlateTypes.h"
 #include "Styling/SlateStyleRegistry.h"
 
 FAvaEffectorsEditorStyle::FAvaEffectorsEditorStyle()
-	: FSlateStyleSet(UE_MODULE_NAME)
+	: FSlateStyleSet(TEXT("AvaEffectorsEditor"))
 {
 	const FVector2f Icon20x20(20.0f, 20.0f);
 	const FVector2f Icon32x32(32.0f, 32.0f);
-
+	
 	const TSharedPtr<IPlugin> Plugin = IPluginManager::Get().FindPlugin(UE_PLUGIN_NAME);
-
 	check(Plugin.IsValid());
-
-	ContentRootDir = FPaths::Combine(Plugin->GetBaseDir(), TEXT("Resources"));
+	
+	SetContentRoot(Plugin->GetBaseDir() / TEXT("Resources"));
 
 	Set("AvalancheEffectorsEditor.Tool_Actor_Effector", new IMAGE_BRUSH_SVG("Icons/ToolboxIcons/effector", Icon20x20));
 	Set("AvalancheEffectorsEditor.Tool_Actor_Cloner", new IMAGE_BRUSH_SVG("Icons/ToolboxIcons/cloner", Icon20x20));
 	Set("Tool_Actor_Effector", new IMAGE_BRUSH_SVG("Icons/ToolboxIcons/effector", Icon20x20));
 	Set("Tool_Actor_Cloner", new IMAGE_BRUSH_SVG("Icons/ToolboxIcons/cloner", Icon20x20));
-
+	
 	FSlateStyleRegistry::RegisterSlateStyle(*this);
 }
 
 FAvaEffectorsEditorStyle::~FAvaEffectorsEditorStyle()
 {
 	FSlateStyleRegistry::UnRegisterSlateStyle(*this);
-}
\ No newline at end of file
+}
diff --git a/Engine/Plugins/Experimental/Avalanche/Source/AvalancheEffectorsEditor/Private/AvaEffectorsEditorStyle.h b/Engine/Plugins/Experimental/Avalanche/Source/AvalancheEffectorsEditor/Private/AvaEffectorsEditorStyle.h
--- a/Engine/Plugins/Experimental/Avalanche/Source/AvalancheEffectorsEditor/Private/AvaEffectorsEditorStyle.h
+++ b/Engine/Plugins/Experimental/Avalanche/Source/AvalancheEffectorsEditor/Private/AvaEffectorsEditorStyle.h
@@ -4,16 +4,15 @@
 
 #include "Styling/SlateStyle.h"
 
-class FAvaEffectorsEditorStyle : public FSlateStyleSet
+class FAvaEffectorsEditorStyle final : public FSlateStyleSet
 {
 public:
-	FAvaEffectorsEditorStyle();
-
-	virtual ~FAvaEffectorsEditorStyle() override;
-
 	static FAvaEffectorsEditorStyle& Get()
 	{
-		static FAvaEffectorsEditorStyle StyleSet;
-		return StyleSet;
+		static FAvaEffectorsEditorStyle Instance;
+		return Instance;
 	}
-};
\ No newline at end of file
+
+	FAvaEffectorsEditorStyle();
+	virtual ~FAvaEffectorsEditorStyle() override;
+};
diff --git a/Engine/Plugins/Experimental/Avalanche/Source/AvalancheEffectorsEditor/Private/AvalancheEffectorsEditorModule.cpp b/Engine/Plugins/Experimental/Avalanche/Source/AvalancheEffectorsEditor/Private/AvalancheEffectorsEditorModule.cpp
--- a/Engine/Plugins/Experimental/Avalanche/Source/AvalancheEffectorsEditor/Private/AvalancheEffectorsEditorModule.cpp
+++ b/Engine/Plugins/Experimental/Avalanche/Source/AvalancheEffectorsEditor/Private/AvalancheEffectorsEditorModule.cpp
@@ -21,7 +21,6 @@ void FAvalancheEffectorsEditorModule::StartupModule()
 {
 	FCoreDelegates::OnPostEngineInit.AddRaw(this, &FAvalancheEffectorsEditorModule::PostEngineInit);
 
-	FAvaEffectorsEditorStyle::Get();
 	FAvaEffectorsEditorCommands::Register();
 
 	FAvaInteractiveToolsDelegates::GetRegisterToolsDelegate().AddRaw(this, &FAvalancheEffectorsEditorModule::RegisterTools);
diff --git a/Engine/Plugins/Experimental/Avalanche/Source/AvalancheInteractiveTools/Private/AvaInteractiveToolsCommands.cpp b/Engine/Plugins/Experimental/Avalanche/Source/AvalancheInteractiveTools/Private/AvaInteractiveToolsCommands.cpp
--- a/Engine/Plugins/Experimental/Avalanche/Source/AvalancheInteractiveTools/Private/AvaInteractiveToolsCommands.cpp
+++ b/Engine/Plugins/Experimental/Avalanche/Source/AvalancheInteractiveTools/Private/AvaInteractiveToolsCommands.cpp
@@ -10,7 +10,7 @@ FAvaInteractiveToolsCommands::FAvaInteractiveToolsCommands()
 		TEXT("AvalancheInteractiveTools")
 		, LOCTEXT("MotionDesignInteractiveTools", "Motion Design Interactive Tools")
 		, NAME_None
-		, FAvaInteractiveToolsStyle::GetStyleSetName()
+		, FAvaInteractiveToolsStyle::Get().GetStyleSetName()
 	)
 {
 }
diff --git a/Engine/Plugins/Experimental/Avalanche/Source/AvalancheInteractiveTools/Private/AvaInteractiveToolsStyle.cpp b/Engine/Plugins/Experimental/Avalanche/Source/AvalancheInteractiveTools/Private/AvaInteractiveToolsStyle.cpp
--- a/Engine/Plugins/Experimental/Avalanche/Source/AvalancheInteractiveTools/Private/AvaInteractiveToolsStyle.cpp
+++ b/Engine/Plugins/Experimental/Avalanche/Source/AvalancheInteractiveTools/Private/AvaInteractiveToolsStyle.cpp
@@ -1,90 +1,42 @@
 // Copyright Epic Games, Inc. All Rights Reserved.
 
 #include "AvaInteractiveToolsStyle.h"
+#include "Brushes/SlateImageBrush.h"
+#include "Containers/StringFwd.h"
 #include "Interfaces/IPluginManager.h"
+#include "Math/MathFwd.h"
+#include "Misc/Paths.h"
+#include "Styling/SlateStyleMacros.h"
 #include "Styling/SlateStyleRegistry.h"
 
-#define LOCTEXT_NAMESPACE "AvaInteractiveToolsStyle"
-
-TSharedPtr<FSlateStyleSet> FAvaInteractiveToolsStyle::StyleInstance;
-
-void FAvaInteractiveToolsStyle::Initialize()
-{
-	StyleInstance = Create();
-	FSlateStyleRegistry::RegisterSlateStyle(Get());
-}
-
-void FAvaInteractiveToolsStyle::Shutdown()
-{
-	FSlateStyleRegistry::UnRegisterSlateStyle(Get());
-	StyleInstance.Reset();
-}
-
-FName FAvaInteractiveToolsStyle::GetStyleSetName()
-{
-	static FName StyleSetName(TEXT("AvalancheInteractiveTools"));
-	return StyleSetName;
-}
-
-const ISlateStyle& FAvaInteractiveToolsStyle::Get()
-{
-	if (StyleInstance.IsValid() == false)
-	{
-		Initialize();
-	}
-
-	return *StyleInstance;
-}
-
-const FSlateBrush* FAvaInteractiveToolsStyle::GetBrush(FName PropertyName, const ANSICHAR* Specifier)
-{
-	return StyleInstance->GetBrush(PropertyName, Specifier);
-}
-
-#define IMAGE_BRUSH(RelativePath, ... ) FSlateImageBrush(Style->RootToContentDir(RelativePath, TEXT(".png")), __VA_ARGS__)
-#define IMAGE_BRUSH_SVG(RelativePath, ...) FSlateVectorImageBrush(Style->RootToContentDir(RelativePath, TEXT(".svg")), __VA_ARGS__)
-
-#define CORE_IMAGE_BRUSH(RelativePath, ...) FSlateImageBrush(Style->RootToCoreContentDir(RelativePath, TEXT(".png")), __VA_ARGS__)
-#define CORE_IMAGE_BRUSH_SVG(RelativePath, ...) FSlateVectorImageBrush(Style->RootToCoreContentDir(RelativePath, TEXT(".svg")), __VA_ARGS__)
-
-TSharedRef<FSlateStyleSet> FAvaInteractiveToolsStyle::Create()
+FAvaInteractiveToolsStyle::FAvaInteractiveToolsStyle()
+	: FSlateStyleSet(TEXT("AvaInteractiveTools"))
 {
 	const FVector2f Icon16x16(16.0f, 16.0f);
 	const FVector2f Icon20x20(20.0f, 20.0f);
 
-	TSharedRef<FSlateStyleSet> Style = MakeShared<FSlateStyleSet>(GetStyleSetName());
-
-	TSharedPtr<IPlugin> Plugin = IPluginManager::Get().FindPlugin(TEXT("Avalanche"));
+	const TSharedPtr<IPlugin> Plugin = IPluginManager::Get().FindPlugin(UE_PLUGIN_NAME);
 	check(Plugin.IsValid());
-	if (Plugin.IsValid())
-	{
-		Style->SetContentRoot(Plugin->GetBaseDir() / TEXT("Resources"));
-	}
-
-	Style->SetCoreContentRoot(FPaths::EngineContentDir() / TEXT("Editor/Slate"));
 
-	// Path Constants
-	const FString UEContentEditorSlatePath = FPaths::EngineContentDir() / TEXT("Editor/Slate");
-	const FString UEContentEditorSlateIconsPath = UEContentEditorSlatePath / TEXT("Icons");
+	SetContentRoot(Plugin->GetBaseDir() / TEXT("Resources"));
+	SetCoreContentRoot(FPaths::EngineContentDir() / TEXT("Editor/Slate"));
 
 	// Categories
-	Style->Set("AvalancheInteractiveTools.Category_2D",     new IMAGE_BRUSH("Icons/ToolboxIcons/rectangle", Icon20x20));
-	Style->Set("AvalancheInteractiveTools.Category_3D",     new IMAGE_BRUSH("Icons/ToolboxIcons/cube", Icon20x20));
-	Style->Set("AvalancheInteractiveTools.Category_Actor",  new CORE_IMAGE_BRUSH_SVG("Starship/AssetIcons/Actor_16", Icon16x16));
-	Style->Set("AvalancheInteractiveTools.Category_Layout", new IMAGE_BRUSH("Icons/ToolboxIcons/layoutgrid", Icon20x20));
+	Set("AvalancheInteractiveTools.Category_2D",     new IMAGE_BRUSH("Icons/ToolboxIcons/rectangle", Icon20x20));
+	Set("AvalancheInteractiveTools.Category_3D",     new IMAGE_BRUSH("Icons/ToolboxIcons/cube", Icon20x20));
+	Set("AvalancheInteractiveTools.Category_Actor",  new CORE_IMAGE_BRUSH_SVG("Starship/AssetIcons/Actor_16", Icon16x16));
+	Set("AvalancheInteractiveTools.Category_Layout", new IMAGE_BRUSH("Icons/ToolboxIcons/layoutgrid", Icon20x20));
 
 	// Actor Tools
-	Style->Set("AvalancheInteractiveTools.Tool_Actor_Null", new CORE_IMAGE_BRUSH(TEXT("Icons/SequencerIcons/icon_Sequencer_Move_24x"), Icon16x16));
-	Style->Set("AvalancheInteractiveTools.Tool_Actor_Spline", new IMAGE_BRUSH_SVG("Icons/ToolboxIcons/Toolbar_Spline", Icon20x20));
-	Style->Set("Tool_Actor_Null", new CORE_IMAGE_BRUSH(TEXT("Icons/SequencerIcons/icon_Sequencer_Move_24x"), Icon20x20));
-	Style->Set("Tool_Actor_Spline", new IMAGE_BRUSH_SVG("Icons/ToolboxIcons/Toolbar_Spline", Icon20x20));
+	Set("AvalancheInteractiveTools.Tool_Actor_Null", new CORE_IMAGE_BRUSH(TEXT("Icons/SequencerIcons/icon_Sequencer_Move_24x"), Icon16x16));
+	Set("AvalancheInteractiveTools.Tool_Actor_Spline", new IMAGE_BRUSH_SVG("Icons/ToolboxIcons/Toolbar_Spline", Icon20x20));
+	Set("Tool_Actor_Null", new CORE_IMAGE_BRUSH(TEXT("Icons/SequencerIcons/icon_Sequencer_Move_24x"), Icon20x20));
+	Set("Tool_Actor_Spline", new IMAGE_BRUSH_SVG("Icons/ToolboxIcons/Toolbar_Spline", Icon20x20));
 
-	return Style;
+	FSlateStyleRegistry::RegisterSlateStyle(*this);
 }
 
-#undef IMAGE_BRUSH
-#undef IMAGE_BRUSH_SVG
-#undef CORE_IMAGE_BRUSH
-#undef CORE_IMAGE_BRUSH_SVG
-
-#undef LOCTEXT_NAMESPACE
+FAvaInteractiveToolsStyle::~FAvaInteractiveToolsStyle()
+{
+	FSlateStyleRegistry::UnRegisterSlateStyle(*this);
+}
diff --git a/Engine/Plugins/Experimental/Avalanche/Source/AvalancheInteractiveTools/Private/AvaInteractiveToolsStyle.h b/Engine/Plugins/Experimental/Avalanche/Source/AvalancheInteractiveTools/Private/AvaInteractiveToolsStyle.h
--- a/Engine/Plugins/Experimental/Avalanche/Source/AvalancheInteractiveTools/Private/AvaInteractiveToolsStyle.h
+++ b/Engine/Plugins/Experimental/Avalanche/Source/AvalancheInteractiveTools/Private/AvaInteractiveToolsStyle.h
@@ -2,24 +2,17 @@
 
 #pragma once
 
-#include "Styling/ISlateStyle.h"
 #include "Styling/SlateStyle.h"
 
-struct FSlateBrush;
-
-class FAvaInteractiveToolsStyle
+class FAvaInteractiveToolsStyle final : public FSlateStyleSet
 {
 public:
-	static void Initialize();
-	static void Shutdown();
-
-	static FName GetStyleSetName();
-	static const ISlateStyle& Get();
-
-	static const FSlateBrush* GetBrush(FName PropertyName, const ANSICHAR* Specifier = nullptr);
-
-private:
-	static TSharedPtr<FSlateStyleSet> StyleInstance;
-
-	static TSharedRef<FSlateStyleSet> Create();
+	static FAvaInteractiveToolsStyle& Get()
+	{
+		static FAvaInteractiveToolsStyle Instance;
+		return Instance;
+	}
+
+	FAvaInteractiveToolsStyle();
+	virtual ~FAvaInteractiveToolsStyle() override;
 };
diff --git a/Engine/Plugins/Experimental/Avalanche/Source/AvalancheInteractiveTools/Private/AvalancheInteractiveToolsModule.cpp b/Engine/Plugins/Experimental/Avalanche/Source/AvalancheInteractiveTools/Private/AvalancheInteractiveToolsModule.cpp
--- a/Engine/Plugins/Experimental/Avalanche/Source/AvalancheInteractiveTools/Private/AvalancheInteractiveToolsModule.cpp
+++ b/Engine/Plugins/Experimental/Avalanche/Source/AvalancheInteractiveTools/Private/AvalancheInteractiveToolsModule.cpp
@@ -4,7 +4,6 @@
 #include "AvaInteractiveToolsCommands.h"
 #include "AvaInteractiveToolsDelegates.h"
 #include "AvaInteractiveToolsSettings.h"
-#include "AvaInteractiveToolsStyle.h"
 #include "Builders/AvaInteractiveToolsActorToolBuilder.h"
 #include "Interfaces/IPluginManager.h"
 #include "IPlacementModeModule.h"
@@ -23,7 +22,6 @@ namespace UE::AvalancheInteractiveTools::Private
 
 void FAvalancheInteractiveToolsModule::StartupModule()
 {
-	FAvaInteractiveToolsStyle::Initialize();
 	FAvaInteractiveToolsCommands::Register();
 	bHasActiveTool = false;
 
@@ -39,7 +37,6 @@ void FAvalancheInteractiveToolsModule::StartupModule()
 
 void FAvalancheInteractiveToolsModule::ShutdownModule()
 {
-	FAvaInteractiveToolsStyle::Shutdown();
 	FAvaInteractiveToolsCommands::Unregister();
 	FCoreDelegates::OnPostEngineInit.RemoveAll(this);
 }
diff --git a/Engine/Plugins/Experimental/Avalanche/Source/AvalancheLevelViewport/Private/AvaLevelViewportCommands.cpp b/Engine/Plugins/Experimental/Avalanche/Source/AvalancheLevelViewport/Private/AvaLevelViewportCommands.cpp
--- a/Engine/Plugins/Experimental/Avalanche/Source/AvalancheLevelViewport/Private/AvaLevelViewportCommands.cpp
+++ b/Engine/Plugins/Experimental/Avalanche/Source/AvalancheLevelViewport/Private/AvaLevelViewportCommands.cpp
@@ -10,7 +10,7 @@ FAvaLevelViewportCommands::FAvaLevelViewportCommands()
 		TEXT("AvalancheLevelViewport")
 		, LOCTEXT("MotionDesignLevelViewport", "Motion Design Level Viewport")
 		, NAME_None
-		, FAvaLevelViewportStyle::GetStyleSetName()
+		, FAvaLevelViewportStyle::Get().GetStyleSetName()
 	)
 {
 }
diff --git a/Engine/Plugins/Experimental/Avalanche/Source/AvalancheLevelViewport/Private/AvaLevelViewportStyle.cpp b/Engine/Plugins/Experimental/Avalanche/Source/AvalancheLevelViewport/Private/AvaLevelViewportStyle.cpp
--- a/Engine/Plugins/Experimental/Avalanche/Source/AvalancheLevelViewport/Private/AvaLevelViewportStyle.cpp
+++ b/Engine/Plugins/Experimental/Avalanche/Source/AvalancheLevelViewport/Private/AvaLevelViewportStyle.cpp
@@ -1,149 +1,110 @@
 // Copyright Epic Games, Inc. All Rights Reserved.
 
 #include "AvaLevelViewportStyle.h"
+#include "Brushes/SlateColorBrush.h"
+#include "Brushes/SlateImageBrush.h"
+#include "Containers/StringFwd.h"
 #include "Interfaces/IPluginManager.h"
+#include "Layout/Margin.h"
+#include "Math/MathFwd.h"
+#include "Misc/Paths.h"
 #include "Styling/AppStyle.h"
-#include "Styling/SlateStyleRegistry.h"
 #include "Styling/StyleColors.h"
+#include "Styling/SlateStyleMacros.h"
+#include "Styling/SlateStyleRegistry.h"
+#include "Styling/SlateTypes.h"
 #include "Styling/ToolBarStyle.h"
 
-#define LOCTEXT_NAMESPACE "AvaLevelViewportStyle"
-
-TSharedPtr<FSlateStyleSet> FAvaLevelViewportStyle::StyleInstance;
-
-void FAvaLevelViewportStyle::Initialize()
-{
-	StyleInstance = Create();
-	FSlateStyleRegistry::RegisterSlateStyle(Get());
-}
-
-void FAvaLevelViewportStyle::Shutdown()
-{
-	FSlateStyleRegistry::UnRegisterSlateStyle(Get());
-	StyleInstance.Reset();
-}
-
-FName FAvaLevelViewportStyle::GetStyleSetName()
-{
-	static FName StyleSetName(TEXT("AvalancheLevelViewport"));
-	return StyleSetName;
-}
-
-const ISlateStyle& FAvaLevelViewportStyle::Get()
-{
-	if (StyleInstance.IsValid() == false)
-	{
-		Initialize();
-	}
-
-	return *StyleInstance;
-}
-
-const FSlateBrush* FAvaLevelViewportStyle::GetBrush(FName PropertyName, const ANSICHAR* Specifier)
-{
-	return StyleInstance->GetBrush(PropertyName, Specifier);
-}
-
-#define IMAGE_BRUSH(RelativePath, ... ) FSlateImageBrush(Style->RootToContentDir(RelativePath, TEXT(".png")), __VA_ARGS__)
-#define IMAGE_BRUSH_SVG(RelativePath, ...) FSlateVectorImageBrush(Style->RootToContentDir(RelativePath, TEXT(".svg")), __VA_ARGS__)
-
-TSharedRef<FSlateStyleSet> FAvaLevelViewportStyle::Create()
+FAvaLevelViewportStyle::FAvaLevelViewportStyle()
+	: FSlateStyleSet(TEXT("AvaLevelViewport"))
 {
 	const FVector2f Icon16x16(16.0f, 16.0f);
 	const FVector2f Icon20x20(20.0f, 20.0f);
 	const FVector2f Icon22x22(22.0f, 22.0f);
 	const FVector2f Icon25x25(25.0f, 25.0f);
 
-	TSharedRef<FSlateStyleSet> Style = MakeShared<FSlateStyleSet>(GetStyleSetName());
-
-	TSharedPtr<IPlugin> Plugin = IPluginManager::Get().FindPlugin(TEXT("Avalanche"));
+	const TSharedPtr<IPlugin> Plugin = IPluginManager::Get().FindPlugin(UE_PLUGIN_NAME);
 	check(Plugin.IsValid());
-	if (Plugin.IsValid())
-	{
-		Style->SetContentRoot(Plugin->GetBaseDir() / TEXT("Resources"));
-	}
 
-	// Path Constants
-	const FString UEContentEditorSlatePath = FPaths::EngineContentDir() / TEXT("Editor/Slate");
-	const FString UEContentEditorSlateIconsPath = UEContentEditorSlatePath / TEXT("Icons");
+	SetContentRoot(Plugin->GetBaseDir() / TEXT("Resources"));
 
 	// Grid Icons
-	Style->Set("Button.ToggleGrid", new IMAGE_BRUSH_SVG("Icons/PaletteIcons/GridToggle", Icon22x22));
-	Style->Set("Button.ToggleGridAlwaysVisible", new IMAGE_BRUSH_SVG("Icons/PaletteIcons/GridAlwaysVisibleToggle", Icon22x22));
+	Set("Button.ToggleGrid", new IMAGE_BRUSH_SVG("Icons/PaletteIcons/GridToggle", Icon22x22));
+	Set("Button.ToggleGridAlwaysVisible", new IMAGE_BRUSH_SVG("Icons/PaletteIcons/GridAlwaysVisibleToggle", Icon22x22));
 
 	// Snap Icons
-	Style->Set("Button.ToggleSnap", new IMAGE_BRUSH_SVG("Icons/PaletteIcons/SnapToggle", Icon22x22));
-	Style->Set("Button.ToggleGrid", new IMAGE_BRUSH_SVG("Icons/PaletteIcons/SnapGridToggle", Icon22x22));
-	Style->Set("Button.ToggleScreen", new IMAGE_BRUSH_SVG("Icons/PaletteIcons/SnapScreenToggle", Icon22x22));
-	Style->Set("Button.ToggleActor", new IMAGE_BRUSH_SVG("Icons/PaletteIcons/SnapActorToggle", Icon22x22));
+	Set("Button.ToggleSnap", new IMAGE_BRUSH_SVG("Icons/PaletteIcons/SnapToggle", Icon22x22));
+	Set("Button.ToggleGrid", new IMAGE_BRUSH_SVG("Icons/PaletteIcons/SnapGridToggle", Icon22x22));
+	Set("Button.ToggleScreen", new IMAGE_BRUSH_SVG("Icons/PaletteIcons/SnapScreenToggle", Icon22x22));
+	Set("Button.ToggleActor", new IMAGE_BRUSH_SVG("Icons/PaletteIcons/SnapActorToggle", Icon22x22));
 
 	// Viewport
-	Style->Set("Button.GameView", new IMAGE_BRUSH_SVG("Icons/EditorIcons/GameView", Icon25x25));
-	Style->Set("Button.Visualizers", new IMAGE_BRUSH_SVG("Icons/EditorIcons/Visualizers", Icon25x25));
-	Style->Set("Button.Guides", new IMAGE_BRUSH_SVG("Icons/EditorIcons/Guides", Icon25x25));
-	Style->Set("Button.Billboards", new IMAGE_BRUSH("Icons/EditorIcons/Billboards", Icon25x25));
-	Style->Set("Button.IsolateActors", new IMAGE_BRUSH("Icons/EditorIcons/IsolateActors", Icon25x25));
-	Style->Set("Button.BoundingBoxes", new IMAGE_BRUSH_SVG("Icons/EditorIcons/Mode_Bounding-Box", Icon25x25));
-	Style->Set("Button.SafeFrames", new IMAGE_BRUSH_SVG("Icons/EditorIcons/Mode_Broadcast-Safe", Icon25x25));
-	Style->Set("Button.KeyPreview", new IMAGE_BRUSH_SVG("Icons/EditorIcons/Mode_Key-Preview", Icon25x25));
-	Style->Set("Button.Snapshot", new IMAGE_BRUSH_SVG("Icons/EditorIcons/Mode_Snapshot", Icon25x25));
-	Style->Set("Button.WireframeMode", new IMAGE_BRUSH_SVG("Icons/EditorIcons/Mode_Wireframe", Icon25x25));
-	Style->Set("Button.Mask.Toggle", new IMAGE_BRUSH("Icons/EditorIcons/Mode_Mask", Icon25x25));
+	Set("Button.GameView", new IMAGE_BRUSH_SVG("Icons/EditorIcons/GameView", Icon25x25));
+	Set("Button.Visualizers", new IMAGE_BRUSH_SVG("Icons/EditorIcons/Visualizers", Icon25x25));
+	Set("Button.Guides", new IMAGE_BRUSH_SVG("Icons/EditorIcons/Guides", Icon25x25));
+	Set("Button.Billboards", new IMAGE_BRUSH("Icons/EditorIcons/Billboards", Icon25x25));
+	Set("Button.IsolateActors", new IMAGE_BRUSH("Icons/EditorIcons/IsolateActors", Icon25x25));
+	Set("Button.BoundingBoxes", new IMAGE_BRUSH_SVG("Icons/EditorIcons/Mode_Bounding-Box", Icon25x25));
+	Set("Button.SafeFrames", new IMAGE_BRUSH_SVG("Icons/EditorIcons/Mode_Broadcast-Safe", Icon25x25));
+	Set("Button.KeyPreview", new IMAGE_BRUSH_SVG("Icons/EditorIcons/Mode_Key-Preview", Icon25x25));
+	Set("Button.Snapshot", new IMAGE_BRUSH_SVG("Icons/EditorIcons/Mode_Snapshot", Icon25x25));
+	Set("Button.WireframeMode", new IMAGE_BRUSH_SVG("Icons/EditorIcons/Mode_Wireframe", Icon25x25));
+	Set("Button.Mask.Toggle", new IMAGE_BRUSH("Icons/EditorIcons/Mode_Mask", Icon25x25));
 
 	// Editor
-	Style->Set("Button.SelectionLock", new IMAGE_BRUSH("Icons/EditorIcons/SelectionLock", Icon25x25));
-	Style->Set("Button.PivotMode", new IMAGE_BRUSH("Icons/EditorIcons/Pivot_Mode", Icon25x25));
+	Set("Button.SelectionLock", new IMAGE_BRUSH("Icons/EditorIcons/SelectionLock", Icon25x25));
+	Set("Button.PivotMode", new IMAGE_BRUSH("Icons/EditorIcons/Pivot_Mode", Icon25x25));
 
 	// Alignment
-	Style->Set("AvalancheIcons.Alignment.Translation.TopLeft",     new IMAGE_BRUSH("Icons/DetailsPanelIcons/TopLeft",     Icon16x16));
-	Style->Set("AvalancheIcons.Alignment.Translation.Top",         new IMAGE_BRUSH("Icons/DetailsPanelIcons/Top",         Icon16x16));
-	Style->Set("AvalancheIcons.Alignment.Translation.TopRight",    new IMAGE_BRUSH("Icons/DetailsPanelIcons/TopRight",    Icon16x16));
-	Style->Set("AvalancheIcons.Alignment.Translation.Left",        new IMAGE_BRUSH("Icons/DetailsPanelIcons/Left",        Icon16x16));
-	Style->Set("AvalancheIcons.Alignment.Translation.Center",      new IMAGE_BRUSH("Icons/DetailsPanelIcons/Center",      Icon16x16));
-	Style->Set("AvalancheIcons.Alignment.Translation.Right",       new IMAGE_BRUSH("Icons/DetailsPanelIcons/Right",       Icon16x16));
-	Style->Set("AvalancheIcons.Alignment.Translation.BottomLeft",  new IMAGE_BRUSH("Icons/DetailsPanelIcons/BottomLeft",  Icon16x16));
-	Style->Set("AvalancheIcons.Alignment.Translation.Bottom",      new IMAGE_BRUSH("Icons/DetailsPanelIcons/Bottom",      Icon16x16));
-	Style->Set("AvalancheIcons.Alignment.Translation.BottomRight", new IMAGE_BRUSH("Icons/DetailsPanelIcons/BottomRight", Icon16x16));
+	Set("AvalancheIcons.Alignment.Translation.TopLeft",     new IMAGE_BRUSH("Icons/DetailsPanelIcons/TopLeft",     Icon16x16));
+	Set("AvalancheIcons.Alignment.Translation.Top",         new IMAGE_BRUSH("Icons/DetailsPanelIcons/Top",         Icon16x16));
+	Set("AvalancheIcons.Alignment.Translation.TopRight",    new IMAGE_BRUSH("Icons/DetailsPanelIcons/TopRight",    Icon16x16));
+	Set("AvalancheIcons.Alignment.Translation.Left",        new IMAGE_BRUSH("Icons/DetailsPanelIcons/Left",        Icon16x16));
+	Set("AvalancheIcons.Alignment.Translation.Center",      new IMAGE_BRUSH("Icons/DetailsPanelIcons/Center",      Icon16x16));
+	Set("AvalancheIcons.Alignment.Translation.Right",       new IMAGE_BRUSH("Icons/DetailsPanelIcons/Right",       Icon16x16));
+	Set("AvalancheIcons.Alignment.Translation.BottomLeft",  new IMAGE_BRUSH("Icons/DetailsPanelIcons/BottomLeft",  Icon16x16));
+	Set("AvalancheIcons.Alignment.Translation.Bottom",      new IMAGE_BRUSH("Icons/DetailsPanelIcons/Bottom",      Icon16x16));
+	Set("AvalancheIcons.Alignment.Translation.BottomRight", new IMAGE_BRUSH("Icons/DetailsPanelIcons/BottomRight", Icon16x16));
 	
-	Style->Set("AvalancheIcons.Alignment.Translation.Back",        new IMAGE_BRUSH_SVG("Icons/DetailsPanelIcons/StackBack",   Icon22x22));
-	Style->Set("AvalancheIcons.Alignment.Translation.Center_X",    new IMAGE_BRUSH_SVG("Icons/DetailsPanelIcons/StackCenter", Icon22x22));
-	Style->Set("AvalancheIcons.Alignment.Translation.Front",       new IMAGE_BRUSH_SVG("Icons/DetailsPanelIcons/StackFront",  Icon22x22));
-
-	Style->Set("AvalancheIcons.Alignment.Center_YZ",   new IMAGE_BRUSH_SVG("Icons/PaletteIcons/AlignCenterVertAndHoriz", Icon22x22));
-	Style->Set("AvalancheIcons.Alignment.Left",        new IMAGE_BRUSH_SVG("Icons/PaletteIcons/AlignLeft",               Icon22x22));
-	Style->Set("AvalancheIcons.Alignment.Center_Y",    new IMAGE_BRUSH_SVG("Icons/PaletteIcons/AlignCenterHoriz",        Icon22x22));
-	Style->Set("AvalancheIcons.Alignment.Right",       new IMAGE_BRUSH_SVG("Icons/PaletteIcons/AlignRight",              Icon22x22));
-	Style->Set("AvalancheIcons.Alignment.Top",         new IMAGE_BRUSH_SVG("Icons/PaletteIcons/AlignTop",                Icon22x22));
-	Style->Set("AvalancheIcons.Alignment.Center_Z",    new IMAGE_BRUSH_SVG("Icons/PaletteIcons/AlignCenterVert",         Icon22x22));
-	Style->Set("AvalancheIcons.Alignment.Bottom",      new IMAGE_BRUSH_SVG("Icons/PaletteIcons/AlignBottom",             Icon22x22));
-	Style->Set("AvalancheIcons.Alignment.DistributeX", new IMAGE_BRUSH_SVG("Icons/PaletteIcons/DistributeDepth",         Icon22x22));
-	Style->Set("AvalancheIcons.Alignment.DistributeY", new IMAGE_BRUSH_SVG("Icons/PaletteIcons/DistributeHorizontal",    Icon22x22));
-	Style->Set("AvalancheIcons.Alignment.DistributeZ", new IMAGE_BRUSH_SVG("Icons/PaletteIcons/DistributeVertical",      Icon22x22));
+	Set("AvalancheIcons.Alignment.Translation.Back",        new IMAGE_BRUSH_SVG("Icons/DetailsPanelIcons/StackBack",   Icon22x22));
+	Set("AvalancheIcons.Alignment.Translation.Center_X",    new IMAGE_BRUSH_SVG("Icons/DetailsPanelIcons/StackCenter", Icon22x22));
+	Set("AvalancheIcons.Alignment.Translation.Front",       new IMAGE_BRUSH_SVG("Icons/DetailsPanelIcons/StackFront",  Icon22x22));
+
+	Set("AvalancheIcons.Alignment.Center_YZ",   new IMAGE_BRUSH_SVG("Icons/PaletteIcons/AlignCenterVertAndHoriz", Icon22x22));
+	Set("AvalancheIcons.Alignment.Left",        new IMAGE_BRUSH_SVG("Icons/PaletteIcons/AlignLeft",               Icon22x22));
+	Set("AvalancheIcons.Alignment.Center_Y",    new IMAGE_BRUSH_SVG("Icons/PaletteIcons/AlignCenterHoriz",        Icon22x22));
+	Set("AvalancheIcons.Alignment.Right",       new IMAGE_BRUSH_SVG("Icons/PaletteIcons/AlignRight",              Icon22x22));
+	Set("AvalancheIcons.Alignment.Top",         new IMAGE_BRUSH_SVG("Icons/PaletteIcons/AlignTop",                Icon22x22));
+	Set("AvalancheIcons.Alignment.Center_Z",    new IMAGE_BRUSH_SVG("Icons/PaletteIcons/AlignCenterVert",         Icon22x22));
+	Set("AvalancheIcons.Alignment.Bottom",      new IMAGE_BRUSH_SVG("Icons/PaletteIcons/AlignBottom",             Icon22x22));
+	Set("AvalancheIcons.Alignment.DistributeX", new IMAGE_BRUSH_SVG("Icons/PaletteIcons/DistributeDepth",         Icon22x22));
+	Set("AvalancheIcons.Alignment.DistributeY", new IMAGE_BRUSH_SVG("Icons/PaletteIcons/DistributeHorizontal",    Icon22x22));
+	Set("AvalancheIcons.Alignment.DistributeZ", new IMAGE_BRUSH_SVG("Icons/PaletteIcons/DistributeVertical",      Icon22x22));
 	
-	Style->Set("AvalancheIcons.Alignment.Rotation.Actor.Roll",   new IMAGE_BRUSH_SVG("Icons/PaletteIcons/RotationAlignActorRoll",   Icon22x22));
-	Style->Set("AvalancheIcons.Alignment.Rotation.Actor.Pitch",  new IMAGE_BRUSH_SVG("Icons/PaletteIcons/RotationAlignActorPitch",  Icon22x22));
-	Style->Set("AvalancheIcons.Alignment.Rotation.Actor.Yaw",    new IMAGE_BRUSH_SVG("Icons/PaletteIcons/RotationAlignActorYaw",    Icon22x22));
-	Style->Set("AvalancheIcons.Alignment.Rotation.Actor.All",    new IMAGE_BRUSH_SVG("Icons/PaletteIcons/RotationAlignActorAll",    Icon22x22));
-	Style->Set("AvalancheIcons.Alignment.Rotation.Camera.Roll",  new IMAGE_BRUSH_SVG("Icons/PaletteIcons/RotationAlignCameraRoll",  Icon22x22));
-	Style->Set("AvalancheIcons.Alignment.Rotation.Camera.Pitch", new IMAGE_BRUSH_SVG("Icons/PaletteIcons/RotationAlignCameraPitch", Icon22x22));
-	Style->Set("AvalancheIcons.Alignment.Rotation.Camera.Yaw",   new IMAGE_BRUSH_SVG("Icons/PaletteIcons/RotationAlignCameraYaw",   Icon22x22));
-	Style->Set("AvalancheIcons.Alignment.Rotation.Camera.All",   new IMAGE_BRUSH_SVG("Icons/PaletteIcons/RotationAlignCameraAll",   Icon22x22));
+	Set("AvalancheIcons.Alignment.Rotation.Actor.Roll",   new IMAGE_BRUSH_SVG("Icons/PaletteIcons/RotationAlignActorRoll",   Icon22x22));
+	Set("AvalancheIcons.Alignment.Rotation.Actor.Pitch",  new IMAGE_BRUSH_SVG("Icons/PaletteIcons/RotationAlignActorPitch",  Icon22x22));
+	Set("AvalancheIcons.Alignment.Rotation.Actor.Yaw",    new IMAGE_BRUSH_SVG("Icons/PaletteIcons/RotationAlignActorYaw",    Icon22x22));
+	Set("AvalancheIcons.Alignment.Rotation.Actor.All",    new IMAGE_BRUSH_SVG("Icons/PaletteIcons/RotationAlignActorAll",    Icon22x22));
+	Set("AvalancheIcons.Alignment.Rotation.Camera.Roll",  new IMAGE_BRUSH_SVG("Icons/PaletteIcons/RotationAlignCameraRoll",  Icon22x22));
+	Set("AvalancheIcons.Alignment.Rotation.Camera.Pitch", new IMAGE_BRUSH_SVG("Icons/PaletteIcons/RotationAlignCameraPitch", Icon22x22));
+	Set("AvalancheIcons.Alignment.Rotation.Camera.Yaw",   new IMAGE_BRUSH_SVG("Icons/PaletteIcons/RotationAlignCameraYaw",   Icon22x22));
+	Set("AvalancheIcons.Alignment.Rotation.Camera.All",   new IMAGE_BRUSH_SVG("Icons/PaletteIcons/RotationAlignCameraAll",   Icon22x22));
 
 	// Screen Icons
-	Style->Set("AvalancheIcons.Screen.SizeToScreen",        new IMAGE_BRUSH_SVG("Icons/PaletteIcons/SizeToScreen",        Icon22x22));
-	Style->Set("AvalancheIcons.Screen.SizeToScreenStretch", new IMAGE_BRUSH_SVG("Icons/PaletteIcons/SizeToScreenStretch", Icon22x22));
-	Style->Set("AvalancheIcons.Screen.FitToScreen",         new IMAGE_BRUSH_SVG("Icons/PaletteIcons/FitToScreen",         Icon22x22));
+	Set("AvalancheIcons.Screen.SizeToScreen",        new IMAGE_BRUSH_SVG("Icons/PaletteIcons/SizeToScreen",        Icon22x22));
+	Set("AvalancheIcons.Screen.SizeToScreenStretch", new IMAGE_BRUSH_SVG("Icons/PaletteIcons/SizeToScreenStretch", Icon22x22));
+	Set("AvalancheIcons.Screen.FitToScreen",         new IMAGE_BRUSH_SVG("Icons/PaletteIcons/FitToScreen",         Icon22x22));
 
 	// Color picker icons
-	Style->Set("AvalancheIcons.ColorPicker.SolidColors",    new IMAGE_BRUSH("Icons/EditorIcons/SolidColors",    Icon20x20));
-	Style->Set("AvalancheIcons.ColorPicker.LinearGradient", new IMAGE_BRUSH("Icons/EditorIcons/LinearGradient", Icon20x20));
+	Set("AvalancheIcons.ColorPicker.SolidColors",    new IMAGE_BRUSH("Icons/EditorIcons/SolidColors",    Icon20x20));
+	Set("AvalancheIcons.ColorPicker.LinearGradient", new IMAGE_BRUSH("Icons/EditorIcons/LinearGradient", Icon20x20));
 
 	// Post process icons
-	Style->Set("AvalancheIcons.PostProcess.RGB",   new IMAGE_BRUSH("Icons/ViewportIcons/RGBSquare", Icon20x20));
-	Style->Set("AvalancheIcons.PostProcess.Red",   new IMAGE_BRUSH("Icons/ViewportIcons/RedSquare",   Icon20x20));
-	Style->Set("AvalancheIcons.PostProcess.Green", new IMAGE_BRUSH("Icons/ViewportIcons/GreenSquare", Icon20x20));
-	Style->Set("AvalancheIcons.PostProcess.Blue",  new IMAGE_BRUSH("Icons/ViewportIcons/BlueSquare",  Icon20x20));
-	Style->Set("AvalancheIcons.PostProcess.Alpha", new IMAGE_BRUSH("Icons/ViewportIcons/AlphaSquare", Icon20x20));
+	Set("AvalancheIcons.PostProcess.RGB",   new IMAGE_BRUSH("Icons/ViewportIcons/RGBSquare", Icon20x20));
+	Set("AvalancheIcons.PostProcess.Red",   new IMAGE_BRUSH("Icons/ViewportIcons/RedSquare",   Icon20x20));
+	Set("AvalancheIcons.PostProcess.Green", new IMAGE_BRUSH("Icons/ViewportIcons/GreenSquare", Icon20x20));
+	Set("AvalancheIcons.PostProcess.Blue",  new IMAGE_BRUSH("Icons/ViewportIcons/BlueSquare",  Icon20x20));
+	Set("AvalancheIcons.PostProcess.Alpha", new IMAGE_BRUSH("Icons/ViewportIcons/AlphaSquare", Icon20x20));
 
 	// StatusBar ToolMenu
 	{
@@ -161,21 +122,21 @@ TSharedRef<FSlateStyleSet> FAvaLevelViewportStyle::Create()
     	StatusToolbarStyle.SetButtonPadding(0);
     	StatusToolbarStyle.SetCheckBoxPadding(0);
     	
-    	Style->Set("StatusBar", StatusToolbarStyle);
+    	Set("StatusBar", StatusToolbarStyle);
     }
 
 	FCheckBoxStyle CheckBoxStyle = FAppStyle::Get().GetWidgetStyle<FCheckBoxStyle>("DetailsView.SectionButton");
 	CheckBoxStyle.Padding.Left = 10.f;
 	CheckBoxStyle.Padding.Right = 10.f;
 
-	Style->Set("Avalanche.Alignment.Context", CheckBoxStyle);
+	Set("Avalanche.Alignment.Context", CheckBoxStyle);
 	
 	FButtonStyle ButtonStyle = FAppStyle::Get().GetWidgetStyle<FButtonStyle>("Button");
 	ButtonStyle.NormalPadding.Left = 4.f;
 	ButtonStyle.NormalPadding.Right = 4.f;
 	ButtonStyle.PressedPadding = ButtonStyle.NormalPadding;
 
-	Style->Set("Avalanche.Alignment.Button", ButtonStyle);
+	Set("Avalanche.Alignment.Button", ButtonStyle);
 
 	FButtonStyle GuidePresetMenuStyle = FAppStyle::GetWidgetStyle<FButtonStyle>("FlatButton");
 	GuidePresetMenuStyle.SetNormalPadding(2.f);
@@ -184,12 +145,12 @@ TSharedRef<FSlateStyleSet> FAvaLevelViewportStyle::Create()
 	GuidePresetMenuStyle.Hovered.TintColor = FSlateColor(EStyleColor::Highlight).GetSpecifiedColor();
 	GuidePresetMenuStyle.SetPressed(GuidePresetMenuStyle.Hovered);
 
-	Style->Set("Avalanche.Menu.GuidePreset.Button", GuidePresetMenuStyle);
+	Set("Avalanche.Menu.GuidePreset.Button", GuidePresetMenuStyle);
 
-	return Style;
+	FSlateStyleRegistry::RegisterSlateStyle(*this);
 }
 
-#undef IMAGE_BRUSH
-#undef IMAGE_BRUSH_SVG
-
-#undef LOCTEXT_NAMESPACE
+FAvaLevelViewportStyle::~FAvaLevelViewportStyle()
+{
+	FSlateStyleRegistry::UnRegisterSlateStyle(*this);
+}
diff --git a/Engine/Plugins/Experimental/Avalanche/Source/AvalancheLevelViewport/Private/AvaLevelViewportStyle.h b/Engine/Plugins/Experimental/Avalanche/Source/AvalancheLevelViewport/Private/AvaLevelViewportStyle.h
--- a/Engine/Plugins/Experimental/Avalanche/Source/AvalancheLevelViewport/Private/AvaLevelViewportStyle.h
+++ b/Engine/Plugins/Experimental/Avalanche/Source/AvalancheLevelViewport/Private/AvaLevelViewportStyle.h
@@ -2,24 +2,17 @@
 
 #pragma once
 
-#include "Styling/ISlateStyle.h"
 #include "Styling/SlateStyle.h"
 
-struct FSlateBrush;
-
-class FAvaLevelViewportStyle
+class FAvaLevelViewportStyle final : public FSlateStyleSet
 {
 public:
-	static void Initialize();
-	static void Shutdown();
-
-	static FName GetStyleSetName();
-	static const ISlateStyle& Get();
-
-	static const FSlateBrush* GetBrush(FName PropertyName, const ANSICHAR* Specifier = nullptr);
-
-private:
-	static TSharedPtr<FSlateStyleSet> StyleInstance;
-
-	static TSharedRef<FSlateStyleSet> Create();
+	static FAvaLevelViewportStyle& Get()
+	{
+		static FAvaLevelViewportStyle Instance;
+		return Instance;
+	}
+
+	FAvaLevelViewportStyle();
+	virtual ~FAvaLevelViewportStyle() override;
 };
diff --git a/Engine/Plugins/Experimental/Avalanche/Source/AvalancheLevelViewport/Private/AvalancheLevelViewportModule.cpp b/Engine/Plugins/Experimental/Avalanche/Source/AvalancheLevelViewport/Private/AvalancheLevelViewportModule.cpp
--- a/Engine/Plugins/Experimental/Avalanche/Source/AvalancheLevelViewport/Private/AvalancheLevelViewportModule.cpp
+++ b/Engine/Plugins/Experimental/Avalanche/Source/AvalancheLevelViewport/Private/AvalancheLevelViewportModule.cpp
@@ -27,7 +27,6 @@ namespace UE::AvalancheLevelViewport::Private
 
 void FAvalancheLevelViewportModule::StartupModule()
 {
-	FAvaLevelViewportStyle::Initialize();
 	FAvaLevelViewportCommands::Register();
 
 	ViewportCameraHistory = MakeShared<FAvaViewportCameraHistory>();
@@ -46,7 +45,6 @@ void FAvalancheLevelViewportModule::ShutdownModule()
 
 	ViewportCameraHistory.Reset();
 
-	FAvaLevelViewportStyle::Shutdown();
 	FAvaLevelViewportCommands::Unregister();
 
 	if (AvaLevelViewportClientCasterDelegateHandle.IsValid())
diff --git a/Engine/Plugins/Experimental/Avalanche/Source/AvalancheLevelViewport/Private/Widgets/SAvaLevelViewportStatusBarButtons.cpp b/Engine/Plugins/Experimental/Avalanche/Source/AvalancheLevelViewport/Private/Widgets/SAvaLevelViewportStatusBarButtons.cpp
--- a/Engine/Plugins/Experimental/Avalanche/Source/AvalancheLevelViewport/Private/Widgets/SAvaLevelViewportStatusBarButtons.cpp
+++ b/Engine/Plugins/Experimental/Avalanche/Source/AvalancheLevelViewport/Private/Widgets/SAvaLevelViewportStatusBarButtons.cpp
@@ -35,13 +35,14 @@
 
 namespace UE::Ava::LevelViewportStatusBarButtons::Private
 {
-	static const FSlateIcon RGBChannelIcon = FSlateIcon(FAvaLevelViewportStyle::GetStyleSetName(), "AvalancheIcons.PostProcess.RGB");
-	static const FSlateIcon BackplateIcon = FSlateIcon(FAppStyle::GetAppStyleSetName(), "Icons.Role");
-	static const FSlateIcon RedChannelIcon = FSlateIcon(FAvaLevelViewportStyle::GetStyleSetName(), "AvalancheIcons.PostProcess.Red");
-	static const FSlateIcon GreenChannelIcon = FSlateIcon(FAvaLevelViewportStyle::GetStyleSetName(), "AvalancheIcons.PostProcess.Green");
-	static const FSlateIcon BlueChannelIcon = FSlateIcon(FAvaLevelViewportStyle::GetStyleSetName(), "AvalancheIcons.PostProcess.Blue");
-	static const FSlateIcon AlphaChannelIcon = FSlateIcon(FAvaLevelViewportStyle::GetStyleSetName(), "AvalancheIcons.PostProcess.Alpha");
-	static const FSlateIcon CheckerboardIcon = FSlateIcon(FAppStyle::GetAppStyleSetName(), "Checker");
+	static const FName StyleSetName = FAvaLevelViewportStyle::Get().GetStyleSetName();
+	static const FSlateIcon RGBChannelIcon = FSlateIcon(StyleSetName, "AvalancheIcons.PostProcess.RGB");
+	static const FSlateIcon BackplateIcon = FSlateIcon(StyleSetName, "Icons.Role");
+	static const FSlateIcon RedChannelIcon = FSlateIcon(StyleSetName, "AvalancheIcons.PostProcess.Red");
+	static const FSlateIcon GreenChannelIcon = FSlateIcon(StyleSetName, "AvalancheIcons.PostProcess.Green");
+	static const FSlateIcon BlueChannelIcon = FSlateIcon(StyleSetName, "AvalancheIcons.PostProcess.Blue");
+	static const FSlateIcon AlphaChannelIcon = FSlateIcon(StyleSetName, "AvalancheIcons.PostProcess.Alpha");
+	static const FSlateIcon CheckerboardIcon = FSlateIcon(StyleSetName, "Checker");
 
 	static bool IsViewportPostProcessManagerEnabled(const TWeakPtr<SAvaLevelViewportFrame>& InViewportFrameWeak)
 	{
diff --git a/Engine/Plugins/Experimental/Avalanche/Source/AvalancheMRQEditor/Private/AvaMRQEditorStyle.cpp b/Engine/Plugins/Experimental/Avalanche/Source/AvalancheMRQEditor/Private/AvaMRQEditorStyle.cpp
--- a/Engine/Plugins/Experimental/Avalanche/Source/AvalancheMRQEditor/Private/AvaMRQEditorStyle.cpp
+++ b/Engine/Plugins/Experimental/Avalanche/Source/AvalancheMRQEditor/Private/AvaMRQEditorStyle.cpp
@@ -1,6 +1,7 @@
 // Copyright Epic Games, Inc. All Rights Reserved.
 
 #include "AvaMRQEditorStyle.h"
+#include "Brushes/SlateImageBrush.h"
 #include "Misc/Paths.h"
 #include "Styling/SlateStyleMacros.h"
 #include "Styling/SlateStyleRegistry.h"
@@ -10,10 +11,11 @@ FAvaMRQEditorStyle::FAvaMRQEditorStyle()
 {
 	const FVector2f Icon20(20.f);
 
-	ContentRootDir     = FPaths::EngineContentDir() / TEXT("Editor/Slate");
-	CoreContentRootDir = FPaths::EngineContentDir() / TEXT("Slate");
+	SetContentRoot(FPaths::EngineContentDir() / TEXT("Editor/Slate"));
+	SetCoreContentRoot(FPaths::EngineContentDir() / TEXT("Slate"));
 
 	Set("AvaMRQEditorCommands.RenderSelectedPages", new IMAGE_BRUSH_SVG("Starship/MainToolbar/cinematics", Icon20));
+
 	FSlateStyleRegistry::RegisterSlateStyle(*this);
 }
 
diff --git a/Engine/Plugins/Experimental/Avalanche/Source/AvalancheMRQEditor/Private/AvaMRQEditorStyle.h b/Engine/Plugins/Experimental/Avalanche/Source/AvalancheMRQEditor/Private/AvaMRQEditorStyle.h
--- a/Engine/Plugins/Experimental/Avalanche/Source/AvalancheMRQEditor/Private/AvaMRQEditorStyle.h
+++ b/Engine/Plugins/Experimental/Avalanche/Source/AvalancheMRQEditor/Private/AvaMRQEditorStyle.h
@@ -4,16 +4,15 @@
 
 #include "Styling/SlateStyle.h"
 
-class FAvaMRQEditorStyle : public FSlateStyleSet
+class FAvaMRQEditorStyle final : public FSlateStyleSet
 {
 public:
-	FAvaMRQEditorStyle();
-
-	virtual ~FAvaMRQEditorStyle() override;
-
 	static FAvaMRQEditorStyle& Get()
 	{
-		static FAvaMRQEditorStyle StyleSet;
-		return StyleSet;
+		static FAvaMRQEditorStyle Instance;
+		return Instance;
 	}
+
+	FAvaMRQEditorStyle();
+	virtual ~FAvaMRQEditorStyle() override;
 };
diff --git a/Engine/Plugins/Experimental/Avalanche/Source/AvalancheMaskEditor/Private/AvaMaskEditorCommands.h b/Engine/Plugins/Experimental/Avalanche/Source/AvalancheMaskEditor/Private/AvaMaskEditorCommands.h
--- a/Engine/Plugins/Experimental/Avalanche/Source/AvalancheMaskEditor/Private/AvaMaskEditorCommands.h
+++ b/Engine/Plugins/Experimental/Avalanche/Source/AvalancheMaskEditor/Private/AvaMaskEditorCommands.h
@@ -15,7 +15,7 @@ public:
 			TEXT("AvalancheMaskEditor"),
 			LOCTEXT("MotionDesignMaskEditor", "Motion Design Masking"),
 			NAME_None,
-			FAvalancheMaskEditorStyle::GetStyleSetName())
+			FAvaMaskEditorStyle::Get().GetStyleSetName())
 	{
 	}
 
diff --git a/Engine/Plugins/Experimental/Avalanche/Source/AvalancheMaskEditor/Private/AvaMaskEditorMode.cpp b/Engine/Plugins/Experimental/Avalanche/Source/AvalancheMaskEditor/Private/AvaMaskEditorMode.cpp
--- a/Engine/Plugins/Experimental/Avalanche/Source/AvalancheMaskEditor/Private/AvaMaskEditorMode.cpp
+++ b/Engine/Plugins/Experimental/Avalanche/Source/AvalancheMaskEditor/Private/AvaMaskEditorMode.cpp
@@ -77,23 +77,23 @@ void UAvaMaskEditorMode::Enter()
     	{
     		static FName ToolkitOverlayMenuName = UE::AvalancheMaskEditor::Internal::ToolkitOverlayMenuName;
     		UToolMenu* Menu = UToolMenus::Get()->ExtendMenu(ToolkitOverlayMenuName);
-    		Menu->SetStyleSet(&FAvalancheMaskEditorStyle::Get());
-    		Menu->StyleName = "AvalancheMaskEditor.ViewportOverlayToolbar";
+    		Menu->SetStyleSet(&FAvaMaskEditorStyle::Get());
+    		Menu->StyleName = TEXT("AvalancheMaskEditor.ViewportOverlayToolbar");
     		{
     			FToolMenuSection& Section = Menu->FindOrAddSection(TEXT("Default"));
     			{
     				FToolMenuEntry& Entry = Section.AddEntry(FToolMenuEntry::InitToolBarButton(FAvaMaskEditorCommands::Get().ToggleShowAllMasks));
-    				Entry.Icon.Set(FSlateIcon(FAvalancheMaskEditorStyle::GetStyleSetName(), TEXT("AvalancheMaskEditor.ToggleShowAllMasks")));
+    				Entry.Icon.Set(FSlateIcon(FAvaMaskEditorStyle::Get().GetStyleSetName(), TEXT("AvalancheMaskEditor.ToggleShowAllMasks")));
     			}
     
     			{
     				FToolMenuEntry& Entry = Section.AddEntry(FToolMenuEntry::InitToolBarButton(FAvaMaskEditorCommands::Get().ToggleEnableMask));
-    				Entry.Icon.Set(FSlateIcon(FAvalancheMaskEditorStyle::GetStyleSetName(), TEXT("AvalancheMaskEditor.ToggleDisableMask")));
+    				Entry.Icon.Set(FSlateIcon(FAvaMaskEditorStyle::Get().GetStyleSetName(), TEXT("AvalancheMaskEditor.ToggleDisableMask")));
     			}
     
     			{
     				FToolMenuEntry& Entry = Section.AddEntry(FToolMenuEntry::InitToolBarButton(FAvaMaskEditorCommands::Get().ToggleIsolateMask));
-    				Entry.Icon.Set(FSlateIcon(FAvalancheMaskEditorStyle::GetStyleSetName(), TEXT("AvalancheMaskEditor.ToggleIsolateMask")));
+    				Entry.Icon.Set(FSlateIcon(FAvaMaskEditorStyle::Get().GetStyleSetName(), TEXT("AvalancheMaskEditor.ToggleIsolateMask")));
     			}
     		}
     
@@ -118,7 +118,7 @@ void UAvaMaskEditorMode::Enter()
 				.Padding(FMargin(0.0f, 0.0f, 0.f, 15.f))
 				[
 					SNew(SBorder)
-					.BorderImage(FAppStyle::Get().GetBrush("EditorViewport.OverlayBrush"))
+					.BorderImage(FAppStyle::Get().GetBrush(TEXT("EditorViewport.OverlayBrush")))
 					.Padding(8.f)
 					[
 						SNew(SHorizontalBox)
@@ -128,7 +128,7 @@ void UAvaMaskEditorMode::Enter()
 						.Padding(FMargin(0.f, 0.f, 8.f, 0.f))
 						[
 							SNew(SImage)
-							.Image(FAvalancheMaskEditorStyle::Get().GetBrush(TEXT("AvalancheMaskEditor.ToggleMaskMode")))
+							.Image(FAvaMaskEditorStyle::Get().GetBrush(TEXT("AvalancheMaskEditor.ToggleMaskMode")))
 						]
 
 						+SHorizontalBox::Slot()
@@ -145,8 +145,8 @@ void UAvaMaskEditorMode::Enter()
 						.Padding(FMargin(2.0, 0.f, 0.f, 0.f))
 						[
 							SNew(SButton)
-							.ButtonStyle(FAppStyle::Get(), "PrimaryButton")
-							.TextStyle(FAppStyle::Get(), "DialogButtonText")
+							.ButtonStyle(FAppStyle::Get(), TEXT("PrimaryButton"))
+							.TextStyle(FAppStyle::Get(), TEXT("DialogButtonText"))
 							.Text(LOCTEXT("ExitEdit", "Exit"))
 							.ToolTipText(LOCTEXT("ExitTooltip", "Exit Mask Mode"))
 							.HAlign(HAlign_Center)
diff --git a/Engine/Plugins/Experimental/Avalanche/Source/AvalancheMaskEditor/Private/AvaMaskEditorModule.cpp b/Engine/Plugins/Experimental/Avalanche/Source/AvalancheMaskEditor/Private/AvaMaskEditorModule.cpp
--- a/Engine/Plugins/Experimental/Avalanche/Source/AvalancheMaskEditor/Private/AvaMaskEditorModule.cpp
+++ b/Engine/Plugins/Experimental/Avalanche/Source/AvalancheMaskEditor/Private/AvaMaskEditorModule.cpp
@@ -4,7 +4,6 @@
 
 #include "AvaMaskEditorCommands.h"
 #include "AvaMaskEditorMode.h"
-#include "AvaMaskEditorStyle.h"
 #include "AvaMaskEditorSubsystem.h"
 #include "EditorModeManager.h"
 #include "Templates/SharedPointer.h"
@@ -16,7 +15,6 @@
 
 void FAvalancheMaskEditorModule::StartupModule()
 {
-	FAvalancheMaskEditorStyle::Initialize();
 	FAvaMaskEditorCommands::Register();
 
 	CommandList = MakeShared<FUICommandList>();
@@ -29,7 +27,6 @@ void FAvalancheMaskEditorModule::ShutdownModule()
 	UToolMenus::UnRegisterStartupCallback(this);
 	UToolMenus::UnregisterOwner(this);
 	
-	FAvalancheMaskEditorStyle::Shutdown();
 	FAvaMaskEditorCommands::Unregister();
 }
 
@@ -73,7 +70,7 @@ void FAvalancheMaskEditorModule::RegisterMenus()
 			{
 				FToolMenuEntry& Entry = Section.AddEntry(FToolMenuEntry::InitToolBarButton(FAvaMaskEditorCommands::Get().ToggleMaskMode));				
 				Entry.Label.Set(FText::GetEmpty());
-				Entry.Icon.Set(FSlateIcon(FAvalancheMaskEditorStyle::GetStyleSetName(), TEXT("AvalancheMaskEditor.ToggleMaskMode.Small")));
+				Entry.Icon.Set(FSlateIcon(FAvaMaskEditorStyle::Get().GetStyleSetName(), TEXT("AvalancheMaskEditor.ToggleMaskMode.Small")));
 				
 				Entry.SetCommandList(CommandList);
 			}
diff --git a/Engine/Plugins/Experimental/Avalanche/Source/AvalancheMaskEditor/Private/AvaMaskEditorStyle.cpp b/Engine/Plugins/Experimental/Avalanche/Source/AvalancheMaskEditor/Private/AvaMaskEditorStyle.cpp
--- a/Engine/Plugins/Experimental/Avalanche/Source/AvalancheMaskEditor/Private/AvaMaskEditorStyle.cpp
+++ b/Engine/Plugins/Experimental/Avalanche/Source/AvalancheMaskEditor/Private/AvaMaskEditorStyle.cpp
@@ -1,68 +1,34 @@
-﻿// Copyright Epic Games, Inc. All Rights Reserved.
+// Copyright Epic Games, Inc. All Rights Reserved.
 
 #include "AvaMaskEditorStyle.h"
-
-#include "Framework/Application/SlateApplication.h"
+#include "Brushes/SlateImageBrush.h"
 #include "Interfaces/IPluginManager.h"
+#include "Math/MathFwd.h"
+#include "Styling/AppStyle.h"
 #include "Styling/SlateStyleMacros.h"
 #include "Styling/SlateStyleRegistry.h"
 #include "Styling/StyleColors.h"
 #include "Styling/ToolBarStyle.h"
 
-namespace UE::AvalancheMaskEditor::Private
-{
-	static FName NAME_StyleSet = TEXT("AvalancheMaskEditorStyle");
-}
-
-#define RootToContentDir Style->RootToContentDir
-
-TSharedPtr<FSlateStyleSet> FAvalancheMaskEditorStyle::StyleInstance = nullptr;
-
-void FAvalancheMaskEditorStyle::Initialize()
-{
-	if (!StyleInstance.IsValid())
-	{
-		StyleInstance = Create();
-		FSlateStyleRegistry::RegisterSlateStyle(*StyleInstance);
-	}
-}
-
-void FAvalancheMaskEditorStyle::Shutdown()
+FAvaMaskEditorStyle::FAvaMaskEditorStyle()
+	: FSlateStyleSet(TEXT("AvaMaskEditor"))
 {
-	FSlateStyleRegistry::UnRegisterSlateStyle(*StyleInstance);
-	ensure(StyleInstance.IsUnique());
-	StyleInstance.Reset();
-}
+	const FVector2D Icon16x16(16.0f, 16.0f);
+	const FVector2D Icon20x20(20.0f, 20.0f);
 
-FName FAvalancheMaskEditorStyle::GetStyleSetName()
-{
-	return UE::AvalancheMaskEditor::Private::NAME_StyleSet;
-}
+	const TSharedPtr<IPlugin> Plugin = IPluginManager::Get().FindPlugin(UE_PLUGIN_NAME);
+	check(Plugin.IsValid());
 
-const ISlateStyle& FAvalancheMaskEditorStyle::Get()
-{
-	return *StyleInstance;
-}
+	SetContentRoot(Plugin->GetBaseDir() / TEXT("Resources"));
 
-const FVector2D Icon16x16(16.0f, 16.0f);
-const FVector2D Icon20x20(20.0f, 20.0f);
-const FVector2D Icon32x32(32.0f, 32.0f);
-
-TSharedRef<FSlateStyleSet> FAvalancheMaskEditorStyle::Create()
-{
-	TSharedRef<FSlateStyleSet> Style = MakeShared<FSlateStyleSet>(UE::AvalancheMaskEditor::Private::NAME_StyleSet);
-
-	FString ContentRoot = IPluginManager::Get().FindPlugin(UE_PLUGIN_NAME)->GetBaseDir() / TEXT("Resources");
-	Style->SetContentRoot(ContentRoot);
-
-	Style->Set("AvalancheMaskEditor.ToggleMaskMode.Small", new IMAGE_BRUSH_SVG(TEXT("Icons/MaskIcons/Mode_On"), Icon16x16));
-	Style->Set("AvalancheMaskEditor.ToggleMaskMode", new IMAGE_BRUSH_SVG(TEXT("Icons/MaskIcons/Mode_On"), Icon20x20));
-	Style->Set("AvalancheMaskEditor.ToggleShowAllMasks", new IMAGE_BRUSH_SVG(TEXT("Icons/MaskIcons/Disable"), Icon20x20));
-	Style->Set("AvalancheMaskEditor.ToggleDisableMask", new IMAGE_BRUSH_SVG(TEXT("Icons/MaskIcons/Disable"), Icon20x20));
-	Style->Set("AvalancheMaskEditor.ToggleIsolateMask", new IMAGE_BRUSH_SVG(TEXT("Icons/MaskIcons/Disable"), Icon20x20));
+	Set("AvalancheMaskEditor.ToggleMaskMode.Small", new IMAGE_BRUSH_SVG(TEXT("Icons/MaskIcons/Mode_On"), Icon16x16));
+	Set("AvalancheMaskEditor.ToggleMaskMode", new IMAGE_BRUSH_SVG(TEXT("Icons/MaskIcons/Mode_On"), Icon20x20));
+	Set("AvalancheMaskEditor.ToggleShowAllMasks", new IMAGE_BRUSH_SVG(TEXT("Icons/MaskIcons/Disable"), Icon20x20));
+	Set("AvalancheMaskEditor.ToggleDisableMask", new IMAGE_BRUSH_SVG(TEXT("Icons/MaskIcons/Disable"), Icon20x20));
+	Set("AvalancheMaskEditor.ToggleIsolateMask", new IMAGE_BRUSH_SVG(TEXT("Icons/MaskIcons/Disable"), Icon20x20));
 
 	{
-		FToolBarStyle ViewportOverlayToolbarStyle = FAppStyle::Get().GetWidgetStyle<FToolBarStyle>("AssetEditorToolbar");
+		FToolBarStyle ViewportOverlayToolbarStyle = FAppStyle::Get().GetWidgetStyle<FToolBarStyle>(TEXT("AssetEditorToolbar"));
 
 		ViewportOverlayToolbarStyle.SetButtonPadding(       FMargin(0.0f, 0.0f));
 		ViewportOverlayToolbarStyle.SetCheckBoxPadding(     FMargin(0.0f, 0.0f));
@@ -76,16 +42,13 @@ TSharedRef<FSlateStyleSet> FAvalancheMaskEditorStyle::Create()
 		ViewportOverlayToolbarStyle.SetButtonPadding(0);
 		ViewportOverlayToolbarStyle.SetCheckBoxPadding(0);
 		
-		Style->Set("AvalancheMaskEditor.ViewportOverlayToolbar", ViewportOverlayToolbarStyle);
+		Set("AvalancheMaskEditor.ViewportOverlayToolbar", ViewportOverlayToolbarStyle);
 	}
 	
-	return Style;
+	FSlateStyleRegistry::RegisterSlateStyle(*this);
 }
 
-void FAvalancheMaskEditorStyle::ReloadTextures()
+FAvaMaskEditorStyle::~FAvaMaskEditorStyle()
 {
-	if (FSlateApplication::IsInitialized())
-	{
-		FSlateApplication::Get().GetRenderer()->ReloadTextureResources();
-	}
+	FSlateStyleRegistry::UnRegisterSlateStyle(*this);
 }
diff --git a/Engine/Plugins/Experimental/Avalanche/Source/AvalancheMaskEditor/Private/AvaMaskEditorStyle.h b/Engine/Plugins/Experimental/Avalanche/Source/AvalancheMaskEditor/Private/AvaMaskEditorStyle.h
--- a/Engine/Plugins/Experimental/Avalanche/Source/AvalancheMaskEditor/Private/AvaMaskEditorStyle.h
+++ b/Engine/Plugins/Experimental/Avalanche/Source/AvalancheMaskEditor/Private/AvaMaskEditorStyle.h
@@ -1,29 +1,18 @@
-﻿// Copyright Epic Games, Inc. All Rights Reserved.
+// Copyright Epic Games, Inc. All Rights Reserved.
 
 #pragma once
 
 #include "Styling/SlateStyle.h"
-#include "Templates/SharedPointer.h"
 
-class FSlateStyleSet;
-
-class FAvalancheMaskEditorStyle
+class FAvaMaskEditorStyle final : public FSlateStyleSet
 {
 public:
-	static void Initialize();
-	static void Shutdown();
-
-	/** reloads textures used by slate renderer */
-	static void ReloadTextures();
-
-	/** @return The Slate style set for the Shooter game */
-	static const ISlateStyle& Get();
-
-	static FName GetStyleSetName();
-
-private:
-	static TSharedRef<FSlateStyleSet> Create();
-
-private:
-	static TSharedPtr<FSlateStyleSet> StyleInstance;
+	static FAvaMaskEditorStyle& Get()
+	{
+		static FAvaMaskEditorStyle Instance;
+		return Instance;
+	}
+
+	FAvaMaskEditorStyle();
+	virtual ~FAvaMaskEditorStyle() override;
 };
diff --git a/Engine/Plugins/Experimental/Avalanche/Source/AvalancheMedia/Private/AvaMediaModule.cpp b/Engine/Plugins/Experimental/Avalanche/Source/AvalancheMedia/Private/AvaMediaModule.cpp
--- a/Engine/Plugins/Experimental/Avalanche/Source/AvalancheMedia/Private/AvaMediaModule.cpp
+++ b/Engine/Plugins/Experimental/Avalanche/Source/AvalancheMedia/Private/AvaMediaModule.cpp
@@ -4,7 +4,6 @@
 
 #include "Application/ThrottleManager.h"
 #include "AvaMediaSettings.h"
-#include "AvaMediaStyle.h"
 #include "Broadcast/OutputDevices/AvaBroadcastRenderTargetMediaUtils.h"
 #include "IMediaIOCoreModule.h"
 #include "Interfaces/IPluginManager.h"
@@ -37,7 +36,6 @@ void FAvaMediaModule::StartupModule()
 {
 	using namespace UE::AvaMediaModule::Private;
 	
-	FAvaMediaStyle::Initialize();
 	const TSharedPtr<IPlugin> Plugin = IPluginManager::Get().FindPlugin(UE_PLUGIN_NAME);
 	check(Plugin.IsValid());
 
@@ -157,8 +155,6 @@ void FAvaMediaModule::ShutdownModule()
 	
 	FCoreDelegates::OnPostEngineInit.RemoveAll(this);
 	FCoreDelegates::OnEnginePreExit.RemoveAll(this);
-	
-	FAvaMediaStyle::Shutdown();
 
 	if (IMediaIOCoreModule::IsAvailable())
 	{
diff --git a/Engine/Plugins/Experimental/Avalanche/Source/AvalancheMedia/Private/AvaMediaStyle.cpp b/Engine/Plugins/Experimental/Avalanche/Source/AvalancheMedia/Private/AvaMediaStyle.cpp
--- a/Engine/Plugins/Experimental/Avalanche/Source/AvalancheMedia/Private/AvaMediaStyle.cpp
+++ b/Engine/Plugins/Experimental/Avalanche/Source/AvalancheMedia/Private/AvaMediaStyle.cpp
@@ -1,77 +1,29 @@
 // Copyright Epic Games, Inc. All Rights Reserved.
 
 #include "AvaMediaStyle.h"
-
 #include "Brushes/SlateImageBrush.h"
 #include "Interfaces/IPluginManager.h"
-#include "Styling/SlateStyle.h"
+#include "Styling/SlateStyleMacros.h"
 #include "Styling/SlateStyleRegistry.h"
 
-TSharedPtr<FSlateStyleSet> FAvaMediaStyle::StyleInstance = nullptr;
-
-void FAvaMediaStyle::Initialize()
-{
-	if (!StyleInstance.IsValid())
-	{
-		StyleInstance = Create();
-		FSlateStyleRegistry::RegisterSlateStyle(*StyleInstance);
-	}
-}
-
-void FAvaMediaStyle::Shutdown()
+FAvaMediaStyle::FAvaMediaStyle()
+	: FSlateStyleSet(TEXT("AvaMedia"))
 {
-	FSlateStyleRegistry::UnRegisterSlateStyle(*StyleInstance);
-	ensure(StyleInstance.IsUnique());
-	StyleInstance.Reset();
-}
+	const FVector2D Icon16x16(16.0f, 16.0f);
+	const FVector2D Icon200x200(200.f, 200.f);
 
-FName FAvaMediaStyle::GetStyleSetName()
-{
-	static FName StyleSetName(TEXT("AvaMedia"));
-	return StyleSetName;
-}
-
-const FLinearColor& FAvaMediaStyle::GetColor(FName PropertyName, const ANSICHAR* Specifier)
-{
-	return StyleInstance->GetColor(PropertyName, Specifier);
-}
-
-const FSlateBrush* FAvaMediaStyle::GetBrush(FName PropertyName, const ANSICHAR* Specifier)
-{
-	return StyleInstance->GetBrush(PropertyName, Specifier);
-}
-
-#define IMAGE_BRUSH_SVG(RelativePath, ...) FSlateVectorImageBrush(Style->RootToContentDir(RelativePath, TEXT(".svg")), __VA_ARGS__)
-#define IMAGE_BRUSH(RelativePath, ...) FSlateImageBrush(Style->RootToContentDir(RelativePath, TEXT(".png")), __VA_ARGS__)
-#define BOX_BRUSH(RelativePath, ...) FSlateBoxBrush(Style->RootToContentDir(RelativePath, TEXT(".png")), __VA_ARGS__)
-#define BORDER_BRUSH(RelativePath, ...) FSlateBorderBrush(Style->RootToContentDir(RelativePath, TEXT(".png")), __VA_ARGS__)
-
-const FVector2D Icon16x16(16.0f, 16.0f);
-const FVector2D Icon200x200(200.f, 200.f);
-
-TSharedRef<FSlateStyleSet> FAvaMediaStyle::Create()
-{
-	TSharedRef<FSlateStyleSet> Style = MakeShared<FSlateStyleSet>("AvaMedia");
-	
-	TSharedPtr<IPlugin> Plugin = IPluginManager::Get().FindPlugin(UE_PLUGIN_NAME);
+	const TSharedPtr<IPlugin> Plugin = IPluginManager::Get().FindPlugin(UE_PLUGIN_NAME);
 	check(Plugin.IsValid());
-	if (Plugin.IsValid())
-	{
-		Style->SetContentRoot(FPaths::Combine(Plugin->GetBaseDir(), TEXT("Resources")));
-	}
+
+	SetContentRoot(Plugin->GetBaseDir() / TEXT("Resources"));
 	
-	Style->Set("AvaMedia.UnrealIcon", new IMAGE_BRUSH("Icons/MediaIcons/ue_logo", Icon200x200));
-	Style->Set("ClassIcon.AvaBroadcast", new IMAGE_BRUSH("Icons/MediaIcons/MediaOutput", Icon16x16));
+	Set("AvalancheMedia.UnrealIcon", new IMAGE_BRUSH("Icons/MediaIcons/ue_logo", Icon200x200));
+	Set("ClassIcon.AvalancheBroadcast", new IMAGE_BRUSH("Icons/MediaIcons/MediaOutput", Icon16x16));
 
-	return Style;
+	FSlateStyleRegistry::RegisterSlateStyle(*this);
 }
 
-#undef IMAGE_BRUSH_SVG
-#undef IMAGE_BRUSH
-#undef BOX_BRUSH
-#undef BORDER_BRUSH
-
-const ISlateStyle& FAvaMediaStyle::Get()
+FAvaMediaStyle::~FAvaMediaStyle()
 {
-	return *StyleInstance;
+	FSlateStyleRegistry::UnRegisterSlateStyle(*this);
 }
diff --git a/Engine/Plugins/Experimental/Avalanche/Source/AvalancheMedia/Private/AvaMediaStyle.h b/Engine/Plugins/Experimental/Avalanche/Source/AvalancheMedia/Private/AvaMediaStyle.h
--- a/Engine/Plugins/Experimental/Avalanche/Source/AvalancheMedia/Private/AvaMediaStyle.h
+++ b/Engine/Plugins/Experimental/Avalanche/Source/AvalancheMedia/Private/AvaMediaStyle.h
@@ -2,36 +2,17 @@
 
 #pragma once
 
-#include "Styling/ISlateStyle.h"
 #include "Styling/SlateStyle.h"
 
-class UTexture2D;
-struct FSlateBrush;
-
-class FAvaMediaStyle
+class FAvaMediaStyle final : public FSlateStyleSet
 {
 public:
-
-	static void Initialize();
-
-	static void Shutdown();
-	
-	static const ISlateStyle& Get();
-
-	static FName GetStyleSetName();
-
-	static const FLinearColor& GetColor(FName PropertyName, const ANSICHAR* Specifier = nullptr);
-
-	static const FSlateBrush* GetBrush(FName PropertyName, const ANSICHAR* Specifier = nullptr);
-
-	template< typename WidgetStyleType >
-	static const WidgetStyleType& GetWidgetStyle(FName PropertyName, const ANSICHAR* Specifier = nullptr)
+	static FAvaMediaStyle& Get()
 	{
-		return StyleInstance->GetWidgetStyle<WidgetStyleType>(PropertyName, Specifier);
+		static FAvaMediaStyle Instance;
+		return Instance;
 	}
 
-private:
-
-	static TSharedRef<FSlateStyleSet> Create();
-	static TSharedPtr<FSlateStyleSet> StyleInstance;
+	FAvaMediaStyle();
+	virtual ~FAvaMediaStyle() override;
 };
diff --git a/Engine/Plugins/Experimental/Avalanche/Source/AvalancheMedia/Private/Broadcast/Channel/Slate/SAvaBroadcastPlaceholderWidget.cpp b/Engine/Plugins/Experimental/Avalanche/Source/AvalancheMedia/Private/Broadcast/Channel/Slate/SAvaBroadcastPlaceholderWidget.cpp
--- a/Engine/Plugins/Experimental/Avalanche/Source/AvalancheMedia/Private/Broadcast/Channel/Slate/SAvaBroadcastPlaceholderWidget.cpp
+++ b/Engine/Plugins/Experimental/Avalanche/Source/AvalancheMedia/Private/Broadcast/Channel/Slate/SAvaBroadcastPlaceholderWidget.cpp
@@ -45,7 +45,7 @@ void SAvaBroadcastPlaceholderWidget::Construct(const FArguments& InArgs)
 			.Stretch(EStretch::ScaleToFit)
 			[
 				SNew(SImage)
-				.Image(FAvaMediaStyle::GetBrush(TEXT("AvaMedia.UnrealIcon")))
+				.Image(FAvaMediaStyle::Get().GetBrush(TEXT("AvaMedia.UnrealIcon")))
 			]
 		]
 
diff --git a/Engine/Plugins/Experimental/Avalanche/Source/AvalancheMediaEditor/Private/AvaMediaEditorModule.cpp b/Engine/Plugins/Experimental/Avalanche/Source/AvalancheMediaEditor/Private/AvaMediaEditorModule.cpp
--- a/Engine/Plugins/Experimental/Avalanche/Source/AvalancheMediaEditor/Private/AvaMediaEditorModule.cpp
+++ b/Engine/Plugins/Experimental/Avalanche/Source/AvalancheMediaEditor/Private/AvaMediaEditorModule.cpp
@@ -79,8 +79,6 @@ void FAvaMediaEditorModule::StartupModule()
 	FAvaPlaybackCommands::Register();
 	FAvaRundownCommands::Register();
 
-	FAvaMediaEditorStyle::Initialize();
-
 	if (FSlateApplication::IsInitialized())
 	{
 		AddEditorToolbarButtons();
@@ -156,7 +154,6 @@ void FAvaMediaEditorModule::ShutdownModule()
 		UnregisterCustomizations();
 	}
 
-	FAvaMediaEditorStyle::Shutdown();
 	FAvaPlaybackCommands::Unregister();
 	FAvaRundownCommands::Unregister();
 
@@ -199,13 +196,13 @@ FSlateIcon FAvaMediaEditorModule::GetToolbarBroadcastButtonIcon() const
 
 	if (MediaModule.IsMediaPlaybackClientStarted())
 	{
-		return FSlateIcon(FAvaMediaEditorStyle::GetStyleSetName(), "AvaMediaEditor.BroadcastClient", "AvaMediaEditor.BroadcastClient.Small");
+		return FSlateIcon(FAvaMediaEditorStyle::Get().GetStyleSetName(), "AvaMediaEditor.BroadcastClient", "AvaMediaEditor.BroadcastClient.Small");
 	}
 	else if (MediaModule.IsMediaPlaybackServerStarted())
 	{
-		return FSlateIcon(FAvaMediaEditorStyle::GetStyleSetName(), "AvaMediaEditor.BroadcastServer", "AvaMediaEditor.BroadcastServer.Small");
+		return FSlateIcon(FAvaMediaEditorStyle::Get().GetStyleSetName(), "AvaMediaEditor.BroadcastServer", "AvaMediaEditor.BroadcastServer.Small");
 	}
-	return FSlateIcon(FAvaMediaEditorStyle::GetStyleSetName(), "AvaMediaEditor.BroadcastIcon");
+	return FSlateIcon(FAvaMediaEditorStyle::Get().GetStyleSetName(), "AvaMediaEditor.BroadcastIcon");
 }
 
 bool FAvaMediaEditorModule::CanFilterSupportComparisonOperation(const FName& InFilterKey, ETextFilterComparisonOperation InOperation, EAvaRundownSearchListType InRundownSearchListType) const
diff --git a/Engine/Plugins/Experimental/Avalanche/Source/AvalancheMediaEditor/Private/AvaMediaEditorStyle.cpp b/Engine/Plugins/Experimental/Avalanche/Source/AvalancheMediaEditor/Private/AvaMediaEditorStyle.cpp
--- a/Engine/Plugins/Experimental/Avalanche/Source/AvalancheMediaEditor/Private/AvaMediaEditorStyle.cpp
+++ b/Engine/Plugins/Experimental/Avalanche/Source/AvalancheMediaEditor/Private/AvaMediaEditorStyle.cpp
@@ -1,191 +1,146 @@
 // Copyright Epic Games, Inc. All Rights Reserved.
 
 #include "AvaMediaEditorStyle.h"
-
+#include "Brushes/SlateColorBrush.h"
 #include "Brushes/SlateImageBrush.h"
+#include "Containers/StringFwd.h"
 #include "Interfaces/IPluginManager.h"
+#include "Layout/Margin.h"
+#include "Math/Color.h"
 #include "Styling/AppStyle.h"
 #include "Styling/SlateStyle.h"
+#include "Styling/SlateStyleMacros.h"
 #include "Styling/SlateStyleRegistry.h"
 #include "Styling/SlateTypes.h"
 #include "Styling/ToolBarStyle.h"
 
-TSharedPtr<FSlateStyleSet> FAvaMediaEditorStyle::StyleInstance = nullptr;
-
-void FAvaMediaEditorStyle::Initialize()
-{
-	if (!StyleInstance.IsValid())
-	{
-		StyleInstance = Create();
-		FSlateStyleRegistry::RegisterSlateStyle(*StyleInstance);
-	}
-}
-
-void FAvaMediaEditorStyle::Shutdown()
+FAvaMediaEditorStyle::FAvaMediaEditorStyle()
+	: FSlateStyleSet(TEXT("AvaMediaEditor"))
 {
-	FSlateStyleRegistry::UnRegisterSlateStyle(*StyleInstance);
-	ensure(StyleInstance.IsUnique());
-	StyleInstance.Reset();
-}
+	const FVector2D Icon12x12(12.0f, 12.0f);
+	const FVector2D Icon16x16(16.0f, 16.0f);
+	const FVector2D Icon20x20(20.0f, 20.0f);
+	const FVector2D Icon64x64(64.0f, 64.0f);
 
-FName FAvaMediaEditorStyle::GetStyleSetName()
-{
-	static FName StyleSetName(TEXT("AvaMediaEditor"));
-	return StyleSetName;
-}
-
-const FLinearColor& FAvaMediaEditorStyle::GetColor(FName PropertyName, const ANSICHAR* Specifier)
-{
-	return StyleInstance->GetColor(PropertyName, Specifier);
-}
-
-const FSlateBrush* FAvaMediaEditorStyle::GetBrush(FName PropertyName, const ANSICHAR* Specifier)
-{
-	return StyleInstance->GetBrush(PropertyName, Specifier);
-}
-
-#define IMAGE_BRUSH(RelativePath, ...) FSlateImageBrush(Style->RootToContentDir(RelativePath, TEXT(".png")), __VA_ARGS__)
-#define IMAGE_BRUSH_SVG(RelativePath, ...) FSlateVectorImageBrush(Style->RootToContentDir(RelativePath, TEXT(".svg")), __VA_ARGS__)
-#define BOX_BRUSH(RelativePath, ...) FSlateBoxBrush(Style->RootToContentDir(RelativePath, TEXT(".png") ), __VA_ARGS__ )
-#define CORE_FONT(...) FSlateFontInfo(FCoreStyle::GetDefaultFont(), __VA_ARGS__)
-
-const FVector2D Icon12x12(12.0f, 12.0f);
-const FVector2D Icon16x16(16.0f, 16.0f);
-const FVector2D Icon20x20(20.0f, 20.0f);
-const FVector2D Icon24x24(24.0f, 24.0f);
-const FVector2D Icon40x40(40.0f, 40.0f);
-const FVector2D Icon64x64(64.0f, 64.0f);
-
-TSharedRef<FSlateStyleSet> FAvaMediaEditorStyle::Create()
-{
-	TSharedRef<FSlateStyleSet> Style = MakeShared<FSlateStyleSet>("AvaMediaEditor");
-
-	TSharedPtr<IPlugin> Plugin = IPluginManager::Get().FindPlugin(UE_PLUGIN_NAME);
+	const TSharedPtr<IPlugin> Plugin = IPluginManager::Get().FindPlugin(UE_PLUGIN_NAME);
 	check(Plugin.IsValid());
-	if (Plugin.IsValid())
-	{
-		Style->SetContentRoot(FPaths::Combine(Plugin->GetBaseDir(), TEXT("Resources")));
-	}
+
+	SetContentRoot(Plugin->GetBaseDir() / TEXT("Resources"));
 
 	// Asset Class Icons
-	Style->Set("ClassIcon.AvaPlaybackGraph", new IMAGE_BRUSH("Icons/MediaIcons/AvaPlaybackIcon_16x"     , Icon16x16));
-	Style->Set("ClassThumbnail.AvaPlaybackGraph", new IMAGE_BRUSH("Icons/MediaIcons/AvaPlaybackIcon_64x", Icon64x64));
-	Style->Set("ClassIcon.AvaRundown", new IMAGE_BRUSH("Icons/MediaIcons/AvaRundownIcon_16x"     , Icon16x16));
-	Style->Set("ClassThumbnail.AvaRundown", new IMAGE_BRUSH("Icons/MediaIcons/AvaRundownIcon_64x", Icon64x64));
-	Style->Set("ClassIcon.AvaRundownMacroCollection", new IMAGE_BRUSH("Icons/MediaIcons/AvaRundownMacroCollectionIcon_16x"     , Icon16x16));
-	Style->Set("ClassThumbnail.AvaRundownMacroCollection", new IMAGE_BRUSH("Icons/MediaIcons/AvaRundownMacroCollectionIcon_64x", Icon64x64));
+	Set("ClassIcon.AvaPlaybackGraph", new IMAGE_BRUSH("Icons/MediaIcons/AvaPlaybackIcon_16x"     , Icon16x16));
+	Set("ClassThumbnail.AvaPlaybackGraph", new IMAGE_BRUSH("Icons/MediaIcons/AvaPlaybackIcon_64x", Icon64x64));
+	Set("ClassIcon.AvaRundown", new IMAGE_BRUSH("Icons/MediaIcons/AvaRundownIcon_16x"     , Icon16x16));
+	Set("ClassThumbnail.AvaRundown", new IMAGE_BRUSH("Icons/MediaIcons/AvaRundownIcon_64x", Icon64x64));
+	Set("ClassIcon.AvaRundownMacroCollection", new IMAGE_BRUSH("Icons/MediaIcons/AvaRundownMacroCollectionIcon_16x"     , Icon16x16));
+	Set("ClassThumbnail.AvaRundownMacroCollection", new IMAGE_BRUSH("Icons/MediaIcons/AvaRundownMacroCollectionIcon_64x", Icon64x64));
 
 	//Rundown Commands
-	Style->Set("AvaRundownCommands.AddPage"         , new IMAGE_BRUSH_SVG("Icons/MediaIcons/AddPage"      , Icon16x16));
-	Style->Set("AvaRundownCommands.AddTemplate"     , new IMAGE_BRUSH_SVG("Icons/MediaIcons/AddPage"      , Icon16x16));
-	Style->Set("AvaRundownCommands.CreatePageInstanceFromTemplate"  , new IMAGE_BRUSH_SVG("Icons/MediaIcons/AddPage"      , Icon16x16));
-	Style->Set("AvaRundownCommands.CreateComboTemplate", new IMAGE_BRUSH_SVG("Icons/MediaIcons/AddPage"   , Icon16x16));
-	Style->Set("AvaRundownCommands.RemovePage"      , new IMAGE_BRUSH_SVG("Icons/MediaIcons/RemovePage"   , Icon16x16));
-	Style->Set("AvaRundownCommands.RenumberPage"    , new IMAGE_BRUSH_SVG("Icons/MediaIcons/RenumberPage" , Icon16x16));
-	Style->Set("AvaRundownCommands.ReimportPage"    , new IMAGE_BRUSH_SVG("Icons/MediaIcons/ReimportPage" , Icon16x16));
-	Style->Set("AvaRundownCommands.EditPageSource"  , new IMAGE_BRUSH_SVG("Icons/MediaIcons/OpenAsset"	  , Icon16x16));
-	Style->Set("AvaRundownCommands.Play"            , new IMAGE_BRUSH_SVG("Icons/MediaIcons/Play"         , Icon16x16));
-	Style->Set("AvaRundownCommands.UpdateValues"    , new IMAGE_BRUSH_SVG("Icons/MediaIcons/UpdateValues" , Icon16x16));
-	Style->Set("AvaRundownCommands.Stop"            , new IMAGE_BRUSH_SVG("Icons/MediaIcons/Stop"         , Icon16x16));
-	Style->Set("AvaRundownCommands.ForceStop"       , new IMAGE_BRUSH_SVG("Icons/MediaIcons/Stop"         , Icon16x16));
-	Style->Set("AvaRundownCommands.Continue"        , new IMAGE_BRUSH_SVG("Icons/MediaIcons/Continue"     , Icon16x16));
-	Style->Set("AvaRundownCommands.PlayNext"        , new IMAGE_BRUSH_SVG("Icons/MediaIcons/PlayNext"     , Icon16x16));
-	Style->Set("AvaRundownCommands.PreviewFrame"    , new IMAGE_BRUSH_SVG("Icons/MediaIcons/PreviewFrame" , Icon16x16));
-	Style->Set("AvaRundownCommands.PreviewPlay"     , new IMAGE_BRUSH_SVG("Icons/MediaIcons/Play"         , Icon16x16));
-	Style->Set("AvaRundownCommands.PreviewStop"     , new IMAGE_BRUSH_SVG("Icons/MediaIcons/Stop"         , Icon16x16));
-	Style->Set("AvaRundownCommands.PreviewForceStop", new IMAGE_BRUSH_SVG("Icons/MediaIcons/Stop"         , Icon16x16));
-	Style->Set("AvaRundownCommands.PreviewContinue" , new IMAGE_BRUSH_SVG("Icons/MediaIcons/Continue"     , Icon16x16));
-	Style->Set("AvaRundownCommands.PreviewPlayNext" , new IMAGE_BRUSH_SVG("Icons/MediaIcons/PlayNext"     , Icon16x16));
-	Style->Set("AvaRundownCommands.TakeToProgram"   , new IMAGE_BRUSH_SVG("Icons/MediaIcons/Play"		   , Icon16x16));
+	Set("AvaRundownCommands.AddPage"         , new IMAGE_BRUSH_SVG("Icons/MediaIcons/AddPage"      , Icon16x16));
+	Set("AvaRundownCommands.AddTemplate"     , new IMAGE_BRUSH_SVG("Icons/MediaIcons/AddPage"      , Icon16x16));
+	Set("AvaRundownCommands.CreatePageInstanceFromTemplate"  , new IMAGE_BRUSH_SVG("Icons/MediaIcons/AddPage"      , Icon16x16));
+	Set("AvaRundownCommands.CreateComboTemplate", new IMAGE_BRUSH_SVG("Icons/MediaIcons/AddPage"   , Icon16x16));
+	Set("AvaRundownCommands.RemovePage"      , new IMAGE_BRUSH_SVG("Icons/MediaIcons/RemovePage"   , Icon16x16));
+	Set("AvaRundownCommands.RenumberPage"    , new IMAGE_BRUSH_SVG("Icons/MediaIcons/RenumberPage" , Icon16x16));
+	Set("AvaRundownCommands.ReimportPage"    , new IMAGE_BRUSH_SVG("Icons/MediaIcons/ReimportPage" , Icon16x16));
+	Set("AvaRundownCommands.EditPageSource"  , new IMAGE_BRUSH_SVG("Icons/MediaIcons/OpenAsset"	  , Icon16x16));
+	Set("AvaRundownCommands.Play"            , new IMAGE_BRUSH_SVG("Icons/MediaIcons/Play"         , Icon16x16));
+	Set("AvaRundownCommands.UpdateValues"    , new IMAGE_BRUSH_SVG("Icons/MediaIcons/UpdateValues" , Icon16x16));
+	Set("AvaRundownCommands.Stop"            , new IMAGE_BRUSH_SVG("Icons/MediaIcons/Stop"         , Icon16x16));
+	Set("AvaRundownCommands.ForceStop"       , new IMAGE_BRUSH_SVG("Icons/MediaIcons/Stop"         , Icon16x16));
+	Set("AvaRundownCommands.Continue"        , new IMAGE_BRUSH_SVG("Icons/MediaIcons/Continue"     , Icon16x16));
+	Set("AvaRundownCommands.PlayNext"        , new IMAGE_BRUSH_SVG("Icons/MediaIcons/PlayNext"     , Icon16x16));
+	Set("AvaRundownCommands.PreviewFrame"    , new IMAGE_BRUSH_SVG("Icons/MediaIcons/PreviewFrame" , Icon16x16));
+	Set("AvaRundownCommands.PreviewPlay"     , new IMAGE_BRUSH_SVG("Icons/MediaIcons/Play"         , Icon16x16));
+	Set("AvaRundownCommands.PreviewStop"     , new IMAGE_BRUSH_SVG("Icons/MediaIcons/Stop"         , Icon16x16));
+	Set("AvaRundownCommands.PreviewForceStop", new IMAGE_BRUSH_SVG("Icons/MediaIcons/Stop"         , Icon16x16));
+	Set("AvaRundownCommands.PreviewContinue" , new IMAGE_BRUSH_SVG("Icons/MediaIcons/Continue"     , Icon16x16));
+	Set("AvaRundownCommands.PreviewPlayNext" , new IMAGE_BRUSH_SVG("Icons/MediaIcons/PlayNext"     , Icon16x16));
+	Set("AvaRundownCommands.TakeToProgram"   , new IMAGE_BRUSH_SVG("Icons/MediaIcons/Play"		   , Icon16x16));
 
 	//Broadcast
-	Style->Set("AvaMediaEditor.OutputIcon"           , new IMAGE_BRUSH("Icons/MediaIcons/MediaOutput"     , Icon20x20));
-	Style->Set("AvaMediaEditor.BroadcastIcon"        , new IMAGE_BRUSH("Icons/MediaIcons/MediaOutput"     , Icon20x20));
-	Style->Set("AvaMediaEditor.BroadcastClient"      , new IMAGE_BRUSH("Icons/MediaIcons/BroadcastClient" , Icon64x64));
-	Style->Set("AvaMediaEditor.BroadcastServer"      , new IMAGE_BRUSH("Icons/MediaIcons/BroadcastServer" , Icon64x64));
-	Style->Set("AvaMediaEditor.BroadcastClient.Small", new IMAGE_BRUSH("Icons/MediaIcons/BroadcastClient" , Icon20x20));
-	Style->Set("AvaMediaEditor.BroadcastServer.Small", new IMAGE_BRUSH("Icons/MediaIcons/BroadcastServer" , Icon20x20));
-	Style->Set("AvaMediaEditor.BroadcastOffline"     , new IMAGE_BRUSH_SVG("Icons/MediaIcons/BroadcastOffline"	, Icon16x16));
-	Style->Set("AvaMediaEditor.BroadcastError"       , new IMAGE_BRUSH_SVG("Icons/MediaIcons/BroadcastError"		, Icon16x16));
-	Style->Set("AvaMediaEditor.BroadcastWarning"     , new IMAGE_BRUSH_SVG("Icons/MediaIcons/BroadcastWarning"	, Icon16x16));
-	Style->Set("AvaMediaEditor.BroadcastLive"        , new IMAGE_BRUSH_SVG("Icons/MediaIcons/BroadcastLive"		, Icon16x16)); 
-	Style->Set("AvaMediaEditor.BroadcastIdle"        , new IMAGE_BRUSH_SVG("Icons/MediaIcons/BroadcastIdle"		, Icon16x16));
+	Set("AvaMediaEditor.OutputIcon"           , new IMAGE_BRUSH("Icons/MediaIcons/MediaOutput"     , Icon20x20));
+	Set("AvaMediaEditor.BroadcastIcon"        , new IMAGE_BRUSH("Icons/MediaIcons/MediaOutput"     , Icon20x20));
+	Set("AvaMediaEditor.BroadcastClient"      , new IMAGE_BRUSH("Icons/MediaIcons/BroadcastClient" , Icon64x64));
+	Set("AvaMediaEditor.BroadcastServer"      , new IMAGE_BRUSH("Icons/MediaIcons/BroadcastServer" , Icon64x64));
+	Set("AvaMediaEditor.BroadcastClient.Small", new IMAGE_BRUSH("Icons/MediaIcons/BroadcastClient" , Icon20x20));
+	Set("AvaMediaEditor.BroadcastServer.Small", new IMAGE_BRUSH("Icons/MediaIcons/BroadcastServer" , Icon20x20));
+	Set("AvaMediaEditor.BroadcastOffline"     , new IMAGE_BRUSH_SVG("Icons/MediaIcons/BroadcastOffline"	, Icon16x16));
+	Set("AvaMediaEditor.BroadcastError"       , new IMAGE_BRUSH_SVG("Icons/MediaIcons/BroadcastError"		, Icon16x16));
+	Set("AvaMediaEditor.BroadcastWarning"     , new IMAGE_BRUSH_SVG("Icons/MediaIcons/BroadcastWarning"	, Icon16x16));
+	Set("AvaMediaEditor.BroadcastLive"        , new IMAGE_BRUSH_SVG("Icons/MediaIcons/BroadcastLive"		, Icon16x16)); 
+	Set("AvaMediaEditor.BroadcastIdle"        , new IMAGE_BRUSH_SVG("Icons/MediaIcons/BroadcastIdle"		, Icon16x16));
 
 	// Broadcast Channel Types
-	Style->Set("AvaMediaEditor.ChannelTypeProgram"	, new IMAGE_BRUSH_SVG("Icons/MediaIcons/ChannelTypeProgram"	, Icon16x16)); 
-	Style->Set("AvaMediaEditor.ChannelTypePreview"	, new IMAGE_BRUSH_SVG("Icons/MediaIcons/ChannelTypePreview"	, Icon16x16));
+	Set("AvaMediaEditor.ChannelTypeProgram"	, new IMAGE_BRUSH_SVG("Icons/MediaIcons/ChannelTypeProgram"	, Icon16x16)); 
+	Set("AvaMediaEditor.ChannelTypePreview"	, new IMAGE_BRUSH_SVG("Icons/MediaIcons/ChannelTypePreview"	, Icon16x16));
 	
 	// Media Status
 	/* For when we have SVGs
-	Style->Set("AvaMediaEditor.MediaAssetStatus", new IMAGE_BRUSH_SVG("Icons/MediaIcons/MediaLoaded", Icon16x16));
-	Style->Set("AvaMediaEditor.MediaSyncStatus",  new IMAGE_BRUSH_SVG("Icons/MediaIcons/MediaLoaded", Icon16x16));
-	Style->Set("AvaMediaEditor.MediaPreviewing",  new IMAGE_BRUSH_SVG("Icons/MediaIcons/MediaLoaded", Icon16x16));
-	Style->Set("AvaMediaEditor.MediaPlaying",     new IMAGE_BRUSH_SVG("Icons/MediaIcons/MediaLoaded", Icon16x16));
+	Set("AvaMediaEditor.MediaAssetStatus", new IMAGE_BRUSH_SVG("Icons/MediaIcons/MediaLoaded", Icon16x16));
+	Set("AvaMediaEditor.MediaSyncStatus",  new IMAGE_BRUSH_SVG("Icons/MediaIcons/MediaLoaded", Icon16x16));
+	Set("AvaMediaEditor.MediaPreviewing",  new IMAGE_BRUSH_SVG("Icons/MediaIcons/MediaLoaded", Icon16x16));
+	Set("AvaMediaEditor.MediaPlaying",     new IMAGE_BRUSH_SVG("Icons/MediaIcons/MediaLoaded", Icon16x16));
 	*/
-	Style->Set("AvaMediaEditor.MediaAssetStatus", new IMAGE_BRUSH("Icons/MediaIcons/MediaStatus",     Icon16x16));
-	Style->Set("AvaMediaEditor.MediaSyncStatus",  new IMAGE_BRUSH("Icons/MediaIcons/MediaSyncStatus", Icon16x16));
-	Style->Set("AvaMediaEditor.MediaPreviewing",  new IMAGE_BRUSH("Icons/MediaIcons/MediaPreviewing", Icon16x16));
-	Style->Set("AvaMediaEditor.MediaPlaying",     new IMAGE_BRUSH("Icons/MediaIcons/MediaPlaying",    Icon16x16));
+	Set("AvaMediaEditor.MediaAssetStatus", new IMAGE_BRUSH("Icons/MediaIcons/MediaStatus",     Icon16x16));
+	Set("AvaMediaEditor.MediaSyncStatus",  new IMAGE_BRUSH("Icons/MediaIcons/MediaSyncStatus", Icon16x16));
+	Set("AvaMediaEditor.MediaPreviewing",  new IMAGE_BRUSH("Icons/MediaIcons/MediaPreviewing", Icon16x16));
+	Set("AvaMediaEditor.MediaPlaying",     new IMAGE_BRUSH("Icons/MediaIcons/MediaPlaying",    Icon16x16));
 
 	// Media Output Status
-	Style->Set("AvaMediaEditor.MediaOutputOffline"   , new IMAGE_BRUSH_SVG("Icons/MediaIcons/MediaOutputOffline"	, Icon16x16));
-	Style->Set("AvaMediaEditor.MediaOutputIdle"      , new IMAGE_BRUSH_SVG("Icons/MediaIcons/MediaOutputIdle"		, Icon16x16));
-	Style->Set("AvaMediaEditor.MediaOutputPreparing" , new IMAGE_BRUSH_SVG("Icons/MediaIcons/MediaOutputPreparing", Icon16x16));
-	Style->Set("AvaMediaEditor.MediaOutputLive"      , new IMAGE_BRUSH_SVG("Icons/MediaIcons/MediaOutputLive"		, Icon16x16));
-	Style->Set("AvaMediaEditor.MediaOutputLiveWarn"  , new IMAGE_BRUSH_SVG("Icons/MediaIcons/MediaOutputLiveWarn"	, Icon16x16));
-	Style->Set("AvaMediaEditor.MediaOutputError"     , new IMAGE_BRUSH_SVG("Icons/MediaIcons/MediaOutputError"	, Icon16x16));
+	Set("AvaMediaEditor.MediaOutputOffline"   , new IMAGE_BRUSH_SVG("Icons/MediaIcons/MediaOutputOffline"	, Icon16x16));
+	Set("AvaMediaEditor.MediaOutputIdle"      , new IMAGE_BRUSH_SVG("Icons/MediaIcons/MediaOutputIdle"		, Icon16x16));
+	Set("AvaMediaEditor.MediaOutputPreparing" , new IMAGE_BRUSH_SVG("Icons/MediaIcons/MediaOutputPreparing", Icon16x16));
+	Set("AvaMediaEditor.MediaOutputLive"      , new IMAGE_BRUSH_SVG("Icons/MediaIcons/MediaOutputLive"		, Icon16x16));
+	Set("AvaMediaEditor.MediaOutputLiveWarn"  , new IMAGE_BRUSH_SVG("Icons/MediaIcons/MediaOutputLiveWarn"	, Icon16x16));
+	Set("AvaMediaEditor.MediaOutputError"     , new IMAGE_BRUSH_SVG("Icons/MediaIcons/MediaOutputError"	, Icon16x16));
 
 	// Motion Design Preview
-	Style->Set("AvaMediaEditor.Checkerboard" , new IMAGE_BRUSH("Images/AvaPreviewCheckerboard", Icon16x16, FLinearColor::White, ESlateBrushTileType::Both));
+	Set("AvaMediaEditor.Checkerboard" , new IMAGE_BRUSH("Images/AvaPreviewCheckerboard", Icon16x16, FLinearColor::White, ESlateBrushTileType::Both));
 
 	// Button styles
 	const FButtonStyle& AppStyle_SimpleButton = FAppStyle::Get().GetWidgetStyle<FButtonStyle>("SimpleButton");
 
-	Style->Set("AvaMediaEditor.BorderlessButton", FButtonStyle(AppStyle_SimpleButton));
+	Set("AvaMediaEditor.BorderlessButton", FButtonStyle(AppStyle_SimpleButton));
 
 	const FButtonStyle& MenuButtonStyleGreen = FAppStyle::Get().GetWidgetStyle<FButtonStyle>("FlatButton.Success");
-	Style->Set("AvaMediaEditor.ButtonGreen", MenuButtonStyleGreen);
+	Set("AvaMediaEditor.ButtonGreen", MenuButtonStyleGreen);
 	const FButtonStyle& MenuButtonStyleRed = FAppStyle::Get().GetWidgetStyle<FButtonStyle>("FlatButton.Danger");
-	Style->Set("AvaMediaEditor.ButtonRed", MenuButtonStyleRed);
+	Set("AvaMediaEditor.ButtonRed", MenuButtonStyleRed);
 
 	const FToolBarStyle& SlimToolbarStyle = FAppStyle::Get().GetWidgetStyle<FToolBarStyle>("SlimToolBar");
-	Style->Set("AvaMediaEditor.ToolBar", SlimToolbarStyle);
+	Set("AvaMediaEditor.ToolBar", SlimToolbarStyle);
 	const FButtonStyle& MenuButtonStyle = SlimToolbarStyle.ButtonStyle;
 
 	FToolBarStyle SlimToolbarStyleOverrideRedButton = FAppStyle::Get().GetWidgetStyle<FToolBarStyle>("SlimToolBar");
-	SlimToolbarStyleOverrideRedButton.SetButtonStyle(Style->GetWidgetStyle<FButtonStyle>("AvaMediaEditor.ButtonRed"));
+	SlimToolbarStyleOverrideRedButton.SetButtonStyle(GetWidgetStyle<FButtonStyle>("AvaMediaEditor.ButtonRed"));
 	SlimToolbarStyleOverrideRedButton.ButtonStyle.SetDisabled(MenuButtonStyle.Disabled);
-	Style->Set("AvaMediaEditor.ToolBarRedButtonOverride", SlimToolbarStyleOverrideRedButton);
+	Set("AvaMediaEditor.ToolBarRedButtonOverride", SlimToolbarStyleOverrideRedButton);
 
 	FToolBarStyle SlimToolbarStyleOverrideGreenButton = FAppStyle::Get().GetWidgetStyle<FToolBarStyle>("SlimToolBar");
-	SlimToolbarStyleOverrideGreenButton.SetButtonStyle(Style->GetWidgetStyle<FButtonStyle>("AvaMediaEditor.ButtonGreen"));
+	SlimToolbarStyleOverrideGreenButton.SetButtonStyle(GetWidgetStyle<FButtonStyle>("AvaMediaEditor.ButtonGreen"));
 	SlimToolbarStyleOverrideGreenButton.ButtonStyle.SetDisabled(MenuButtonStyle.Disabled);
-	Style->Set("AvaMediaEditor.ToolBarGreenButtonOverride", SlimToolbarStyleOverrideGreenButton);
+	Set("AvaMediaEditor.ToolBarGreenButtonOverride", SlimToolbarStyleOverrideGreenButton);
 
 	const FToolBarStyle& CalloutToolbarStyleOverride = FAppStyle::Get().GetWidgetStyle<FToolBarStyle>("CalloutToolbar");
-	Style->Set("AvaMediaEditor.CalloutToolbar", CalloutToolbarStyleOverride);
+	Set("AvaMediaEditor.CalloutToolbar", CalloutToolbarStyleOverride);
 
-	Style->Set("TableView.ActivePageBrush", new FSlateColorBrush(FLinearColor(0.05f, 0.1f, 0.05f, 0.5f)));
-	Style->Set("TableView.Hovered.ActivePageBrush", new FSlateColorBrush(FLinearColor(0.05f, 0.2f, 0.05f, 0.5f)));
-	Style->Set("TableView.Selected.ActivePageBrush", new FSlateColorBrush(FLinearColor(0.05f, 0.3f, 0.05f, 0.5f)));
-	Style->Set("TableView.DisabledPageBrush", new FSlateColorBrush(FLinearColor(0.1f, 0.05f, 0.05f, 0.5f)));
-	Style->Set("TableView.Hovered.DisabledPageBrush", new FSlateColorBrush(FLinearColor(0.2f, 0.05f, 0.05f, 0.5f)));
-	Style->Set("TableView.Selected.DisabledPageBrush", new FSlateColorBrush(FLinearColor(0.3f, 0.05f, 0.05f, 0.5f)));
+	Set("TableView.ActivePageBrush", new FSlateColorBrush(FLinearColor(0.05f, 0.1f, 0.05f, 0.5f)));
+	Set("TableView.Hovered.ActivePageBrush", new FSlateColorBrush(FLinearColor(0.05f, 0.2f, 0.05f, 0.5f)));
+	Set("TableView.Selected.ActivePageBrush", new FSlateColorBrush(FLinearColor(0.05f, 0.3f, 0.05f, 0.5f)));
+	Set("TableView.DisabledPageBrush", new FSlateColorBrush(FLinearColor(0.1f, 0.05f, 0.05f, 0.5f)));
+	Set("TableView.Hovered.DisabledPageBrush", new FSlateColorBrush(FLinearColor(0.2f, 0.05f, 0.05f, 0.5f)));
+	Set("TableView.Selected.DisabledPageBrush", new FSlateColorBrush(FLinearColor(0.3f, 0.05f, 0.05f, 0.5f)));
 
 	// Asset Colors - in the shades of yellow (hue = 36, saturation = 221)
-	Style->Set("AvaMediaEditor.AssetColors.Rundown", FLinearColor::MakeFromHSV8(36, 221, 170));
-	Style->Set("AvaMediaEditor.AssetColors.RundownMacroCollection", FLinearColor::MakeFromHSV8(36, 221, 124));
-	Style->Set("AvaMediaEditor.AssetColors.Playback", FLinearColor::MakeFromHSV8(36, 221, 105));
+	Set("AvaMediaEditor.AssetColors.Rundown", FLinearColor::MakeFromHSV8(36, 221, 170));
+	Set("AvaMediaEditor.AssetColors.RundownMacroCollection", FLinearColor::MakeFromHSV8(36, 221, 124));
+	Set("AvaMediaEditor.AssetColors.Playback", FLinearColor::MakeFromHSV8(36, 221, 105));
 
-	return Style;
+	FSlateStyleRegistry::RegisterSlateStyle(*this);
 }
 
-#undef IMAGE_BRUSH
-#undef IMAGE_BRUSH_SVG
-#undef BOX_BRUSH
-#undef CORE_FONT
-
-const ISlateStyle& FAvaMediaEditorStyle::Get()
+FAvaMediaEditorStyle::~FAvaMediaEditorStyle()
 {
-	return *StyleInstance;
+	FSlateStyleRegistry::UnRegisterSlateStyle(*this);
 }
diff --git a/Engine/Plugins/Experimental/Avalanche/Source/AvalancheMediaEditor/Private/AvaMediaEditorStyle.h b/Engine/Plugins/Experimental/Avalanche/Source/AvalancheMediaEditor/Private/AvaMediaEditorStyle.h
--- a/Engine/Plugins/Experimental/Avalanche/Source/AvalancheMediaEditor/Private/AvaMediaEditorStyle.h
+++ b/Engine/Plugins/Experimental/Avalanche/Source/AvalancheMediaEditor/Private/AvaMediaEditorStyle.h
@@ -2,37 +2,17 @@
 
 #pragma once
 
-#include "Styling/ISlateStyle.h"
 #include "Styling/SlateStyle.h"
 
-class UTexture2D;
-struct FSlateBrush;
-
-class FAvaMediaEditorStyle
+class FAvaMediaEditorStyle final : public FSlateStyleSet
 {
 public:
-
-	static void Initialize();
-
-	static void Shutdown();
-	
-	static const ISlateStyle& Get();
-
-	static FName GetStyleSetName();
-
-	static const FLinearColor& GetColor(FName PropertyName, const ANSICHAR* Specifier = nullptr);
-
-	static const FSlateBrush* GetBrush(FName PropertyName, const ANSICHAR* Specifier = nullptr);
-
-	template< typename WidgetStyleType >
-	static const WidgetStyleType& GetWidgetStyle(FName PropertyName, const ANSICHAR* Specifier = nullptr)
+	static FAvaMediaEditorStyle& Get()
 	{
-		return StyleInstance->GetWidgetStyle<WidgetStyleType>(PropertyName, Specifier);
+		static FAvaMediaEditorStyle Instance;
+		return Instance;
 	}
 
-private:
-
-	static TSharedRef<FSlateStyleSet> Create();
-	static TSharedPtr<FSlateStyleSet> StyleInstance;
-	static TMap<FName,UTexture2D*> Sprites;
+	FAvaMediaEditorStyle();
+	virtual ~FAvaMediaEditorStyle() override;
 };
diff --git a/Engine/Plugins/Experimental/Avalanche/Source/AvalancheMediaEditor/Private/Broadcast/TabFactories/AvaBroadcastChannelsTabFactory.cpp b/Engine/Plugins/Experimental/Avalanche/Source/AvalancheMediaEditor/Private/Broadcast/TabFactories/AvaBroadcastChannelsTabFactory.cpp
--- a/Engine/Plugins/Experimental/Avalanche/Source/AvalancheMediaEditor/Private/Broadcast/TabFactories/AvaBroadcastChannelsTabFactory.cpp
+++ b/Engine/Plugins/Experimental/Avalanche/Source/AvalancheMediaEditor/Private/Broadcast/TabFactories/AvaBroadcastChannelsTabFactory.cpp
@@ -15,7 +15,7 @@ FAvaBroadcastChannelsTabFactory::FAvaBroadcastChannelsTabFactory(const TSharedPt
 	: FAvaBroadcastTabFactory(TabID, InBroadcastEditor)
 {
 	TabLabel = LOCTEXT("BroadcastChannels_TabLabel", "Channels");
-	TabIcon = FSlateIcon(FAvaMediaEditorStyle::GetStyleSetName(), "AvaMediaEditor.BroadcastIcon");
+	TabIcon = FSlateIcon(FAvaMediaEditorStyle::Get().GetStyleSetName(), "AvaMediaEditor.BroadcastIcon");
 
 	bIsSingleton = true;
 
diff --git a/Engine/Plugins/Experimental/Avalanche/Source/AvalancheMediaEditor/Private/Rundown/AvaRundownCommands.h b/Engine/Plugins/Experimental/Avalanche/Source/AvalancheMediaEditor/Private/Rundown/AvaRundownCommands.h
--- a/Engine/Plugins/Experimental/Avalanche/Source/AvalancheMediaEditor/Private/Rundown/AvaRundownCommands.h
+++ b/Engine/Plugins/Experimental/Avalanche/Source/AvalancheMediaEditor/Private/Rundown/AvaRundownCommands.h
@@ -16,7 +16,7 @@ public:
 		: TCommands<FAvaRundownCommands>(TEXT("AvaRundownCommands")
 		, LOCTEXT("AvaRundownCommands", "Motion Design Rundown")
 		, NAME_None
-		, FAvaMediaEditorStyle::GetStyleSetName())
+		, FAvaMediaEditorStyle::Get().GetStyleSetName())
 	{
 	}
 
diff --git a/Engine/Plugins/Experimental/Avalanche/Source/AvalancheMediaEditor/Private/Rundown/TabFactories/AvaRundownChannelLayerStatusListTabFactory.cpp b/Engine/Plugins/Experimental/Avalanche/Source/AvalancheMediaEditor/Private/Rundown/TabFactories/AvaRundownChannelLayerStatusListTabFactory.cpp
--- a/Engine/Plugins/Experimental/Avalanche/Source/AvalancheMediaEditor/Private/Rundown/TabFactories/AvaRundownChannelLayerStatusListTabFactory.cpp
+++ b/Engine/Plugins/Experimental/Avalanche/Source/AvalancheMediaEditor/Private/Rundown/TabFactories/AvaRundownChannelLayerStatusListTabFactory.cpp
@@ -13,7 +13,7 @@ FAvaRundownChannelLayerStatusListTabFactory::FAvaRundownChannelLayerStatusListTa
 	: FAvaRundownTabFactory(TabID, InRundownEditor)
 {
 	TabLabel = LOCTEXT("ChannelLayerStatus_TabLabel", "Layer Status");
-	TabIcon = FSlateIcon(FAvaMediaEditorStyle::GetStyleSetName(), "AvaMediaEditor.BroadcastIcon");
+	TabIcon = FSlateIcon(FAvaMediaEditorStyle::Get().GetStyleSetName(), TEXT("AvaMediaEditor.BroadcastIcon"));
 
 	bIsSingleton = true;
 	bShouldAutosize = true;
diff --git a/Engine/Plugins/Experimental/Avalanche/Source/AvalancheMediaEditor/Private/Rundown/TabFactories/AvaRundownChannelStatusListTabFactory.cpp b/Engine/Plugins/Experimental/Avalanche/Source/AvalancheMediaEditor/Private/Rundown/TabFactories/AvaRundownChannelStatusListTabFactory.cpp
--- a/Engine/Plugins/Experimental/Avalanche/Source/AvalancheMediaEditor/Private/Rundown/TabFactories/AvaRundownChannelStatusListTabFactory.cpp
+++ b/Engine/Plugins/Experimental/Avalanche/Source/AvalancheMediaEditor/Private/Rundown/TabFactories/AvaRundownChannelStatusListTabFactory.cpp
@@ -13,7 +13,7 @@ FAvaRundownChannelStatusListTabFactory::FAvaRundownChannelStatusListTabFactory(c
 	: FAvaRundownTabFactory(TabID, InRundownEditor)
 {
 	TabLabel = LOCTEXT("ChannelStatus_TabLabel", "Channel Status");
-	TabIcon = FSlateIcon(FAvaMediaEditorStyle::GetStyleSetName(), "AvaMediaEditor.BroadcastIcon");
+	TabIcon = FSlateIcon(FAvaMediaEditorStyle::Get().GetStyleSetName(), TEXT("AvaMediaEditor.BroadcastIcon"));
 
 	bIsSingleton = true;
 	bShouldAutosize = true;
diff --git a/Engine/Plugins/Experimental/Avalanche/Source/AvalancheModifiersEditor/Private/AvaModifiersEditorStyle.cpp b/Engine/Plugins/Experimental/Avalanche/Source/AvalancheModifiersEditor/Private/AvaModifiersEditorStyle.cpp
--- a/Engine/Plugins/Experimental/Avalanche/Source/AvalancheModifiersEditor/Private/AvaModifiersEditorStyle.cpp
+++ b/Engine/Plugins/Experimental/Avalanche/Source/AvalancheModifiersEditor/Private/AvaModifiersEditorStyle.cpp
@@ -1,27 +1,28 @@
 // Copyright Epic Games, Inc. All Rights Reserved.
 
 #include "AvaModifiersEditorStyle.h"
+#include "Brushes/SlateImageBrush.h"
 #include "GenericPlatform/GenericPlatformFile.h"
 #include "HAL/PlatformFileManager.h"
 #include "Interfaces/IPluginManager.h"
-#include "Misc/Paths.h"
+#include "Math/MathFwd.h"
 #include "Modifiers/ActorModifierCoreBase.h"
-#include "Styling/SlateStyleMacros.h"
 #include "Styling/StyleColors.h"
+#include "Styling/SlateColor.h"
 #include "Styling/SlateTypes.h"
+#include "Styling/SlateStyleMacros.h"
 #include "Styling/SlateStyleRegistry.h"
 #include "Subsystems/ActorModifierCoreSubsystem.h"
 
 FAvaModifiersEditorStyle::FAvaModifiersEditorStyle()
-	: FSlateStyleSet(TEXT("AvalancheModifiersEditor"))
+	: FSlateStyleSet(TEXT("AvaModifiersEditor"))
 {
 	const FVector2f Icon16x16(16.f, 16.f);
 
-	const TSharedPtr<IPlugin> Plugin = IPluginManager::Get().FindPlugin(TEXT("Avalanche"));
-
+	const TSharedPtr<IPlugin> Plugin = IPluginManager::Get().FindPlugin(UE_PLUGIN_NAME);
 	check(Plugin.IsValid());
 
-	ContentRootDir = FPaths::Combine(Plugin->GetBaseDir(), TEXT("Resources"));
+	SetContentRoot(Plugin->GetBaseDir() / TEXT("Resources"));
 
 	Set("ClassIcon.AvaBaseModifier", new IMAGE_BRUSH_SVG("Icons/ModifierIcons/BaseModifier", Icon16x16));
 
@@ -82,4 +83,4 @@ FAvaModifiersEditorStyle::~FAvaModifiersEditorStyle()
 	UActorModifierCoreSubsystem::OnModifierClassRegisteredDelegate.RemoveAll(this);
 
 	FSlateStyleRegistry::UnRegisterSlateStyle(*this);
-}
\ No newline at end of file
+}
diff --git a/Engine/Plugins/Experimental/Avalanche/Source/AvalancheModifiersEditor/Private/AvaModifiersEditorStyle.h b/Engine/Plugins/Experimental/Avalanche/Source/AvalancheModifiersEditor/Private/AvaModifiersEditorStyle.h
--- a/Engine/Plugins/Experimental/Avalanche/Source/AvalancheModifiersEditor/Private/AvaModifiersEditorStyle.h
+++ b/Engine/Plugins/Experimental/Avalanche/Source/AvalancheModifiersEditor/Private/AvaModifiersEditorStyle.h
@@ -7,19 +7,18 @@
 
 struct FActorModifierCoreMetadata;
 
-class FAvaModifiersEditorStyle : public FSlateStyleSet
+class FAvaModifiersEditorStyle final : public FSlateStyleSet
 {
 public:
-	FAvaModifiersEditorStyle();
-
-	virtual ~FAvaModifiersEditorStyle() override;
-
 	static FAvaModifiersEditorStyle& Get()
 	{
-		static FAvaModifiersEditorStyle StyleSet;
-		return StyleSet;
+		static FAvaModifiersEditorStyle Instance;
+		return Instance;
 	}
 
+	FAvaModifiersEditorStyle();
+	virtual ~FAvaModifiersEditorStyle() override;
+
 	const FSlateColor& GetModifierCategoryColor(FName CategoryName);
 
 private:
diff --git a/Engine/Plugins/Experimental/Avalanche/Source/AvalancheOutliner/Private/AvaOutlinerModule.cpp b/Engine/Plugins/Experimental/Avalanche/Source/AvalancheOutliner/Private/AvaOutlinerModule.cpp
--- a/Engine/Plugins/Experimental/Avalanche/Source/AvalancheOutliner/Private/AvaOutlinerModule.cpp
+++ b/Engine/Plugins/Experimental/Avalanche/Source/AvalancheOutliner/Private/AvaOutlinerModule.cpp
@@ -21,7 +21,7 @@
 
 void FAvaOutlinerModule::StartupModule()
 {
-	FAvaOutlinerStyle::Initialize();
+	FAvaOutlinerStyle::Get();
 	FAvaOutlinerCommands::Register();
 	ItemProxyRegistry.RegisterItemProxyWithDefaultFactory<FAvaOutlinerMaterialProxy, 10>();
 	ItemProxyRegistry.RegisterItemProxyWithDefaultFactory<FAvaOutlinerComponentProxy, 20>();
@@ -32,7 +32,6 @@ void FAvaOutlinerModule::StartupModule()
 
 void FAvaOutlinerModule::ShutdownModule()
 {
-	FAvaOutlinerStyle::Shutdown();
 	FAvaOutlinerCommands::Unregister();
 }
 
diff --git a/Engine/Plugins/Experimental/Avalanche/Source/AvalancheOutliner/Private/AvaOutlinerStyle.cpp b/Engine/Plugins/Experimental/Avalanche/Source/AvalancheOutliner/Private/AvaOutlinerStyle.cpp
--- a/Engine/Plugins/Experimental/Avalanche/Source/AvalancheOutliner/Private/AvaOutlinerStyle.cpp
+++ b/Engine/Plugins/Experimental/Avalanche/Source/AvalancheOutliner/Private/AvaOutlinerStyle.cpp
@@ -1,89 +1,41 @@
 // Copyright Epic Games, Inc. All Rights Reserved.
 
 #include "AvaOutlinerStyle.h"
+#include "Brushes/SlateColorBrush.h"
+#include "Brushes/SlateImageBrush.h"
+#include "Brushes/SlateNoResource.h"
 #include "Interfaces/IPluginManager.h"
+#include "Misc/Paths.h"
 #include "Styling/AppStyle.h"
-#include "Styling/SlateStyleRegistry.h"
 #include "Styling/StyleColors.h"
+#include "Styling/SlateStyleMacros.h"
+#include "Styling/SlateStyleRegistry.h"
 #include "Styling/SlateTypes.h"
 #include "Styling/ToolBarStyle.h"
 
-TSharedPtr<FSlateStyleSet> FAvaOutlinerStyle::StyleInstance;
-
-void FAvaOutlinerStyle::Initialize()
-{
-	if (!StyleInstance.IsValid())
-	{
-		StyleInstance = Create();
-		FSlateStyleRegistry::RegisterSlateStyle(*StyleInstance);
-	}
-}
-
-void FAvaOutlinerStyle::Shutdown()
-{
-	FSlateStyleRegistry::UnRegisterSlateStyle(*StyleInstance);
-	ensure(StyleInstance.IsUnique());
-	StyleInstance.Reset();
-}
-
-FName FAvaOutlinerStyle::GetStyleSetName()
-{
-	static FName StyleSetName(TEXT("AvalancheOutliner"));
-	return StyleSetName;
-}
-
-const ISlateStyle& FAvaOutlinerStyle::Get()
-{
-	return *StyleInstance;
-}
-
-const FLinearColor& FAvaOutlinerStyle::GetColor(FName InPropertyName, const ANSICHAR* InSpecifier)
-{
-	return StyleInstance->GetColor(InPropertyName, InSpecifier);
-}
-
-const FSlateBrush* FAvaOutlinerStyle::GetBrush(FName InPropertyName, const ANSICHAR* InSpecifier)
-{
-	return StyleInstance->GetBrush(InPropertyName, InSpecifier);
-}
-
-#define BORDER_BRUSH(RelativePath, ...) FSlateBorderBrush(Style->RootToContentDir(RelativePath, TEXT(".png")), __VA_ARGS__)
-#define BOX_BRUSH(RelativePath, ...) FSlateBoxBrush(Style->RootToContentDir(RelativePath, TEXT(".png")), __VA_ARGS__)
-#define CORE_IMAGE_BRUSH(RelativePath, ...) FSlateImageBrush(Style->RootToCoreContentDir(RelativePath, TEXT(".png")), __VA_ARGS__)
-#define IMAGE_BRUSH(RelativePath, ...) FSlateImageBrush(Style->RootToContentDir(RelativePath, TEXT(".png")), __VA_ARGS__)
-
-const FVector2D Icon8x8(8.f, 8.f);
-const FVector2D Icon64x64(64.f, 64.f);
-const FVector2D Icon40x40(40.0f, 40.0f);
-const FVector2D Icon24x24(24.0f, 24.0f);
-const FVector2D Icon20x20(20.0f, 20.0f);
-const FVector2D Icon16x16(16.0f, 16.0f);
-const FVector2D Icon12x12(12.0f, 12.0f);
-
-TSharedRef<FSlateStyleSet> FAvaOutlinerStyle::Create()
+FAvaOutlinerStyle::FAvaOutlinerStyle()
+	: FSlateStyleSet(TEXT("AvaOutliner"))
 {
-	TSharedRef<FSlateStyleSet> Style = MakeShared<FSlateStyleSet>(GetStyleSetName());
-	const TSharedPtr<IPlugin> Plugin = IPluginManager::Get().FindPlugin(TEXT("Avalanche"));
+	const FVector2D Icon20x20(20.0f, 20.0f);
 
+	const TSharedPtr<IPlugin> Plugin = IPluginManager::Get().FindPlugin(UE_PLUGIN_NAME);
 	check(Plugin.IsValid());
-	if (Plugin.IsValid())
-	{
-		Style->SetContentRoot(FPaths::Combine(Plugin->GetBaseDir(), TEXT("Resources")));
-		Style->SetCoreContentRoot(FPaths::EngineContentDir() / TEXT("Slate"));
-	}
+
+	SetContentRoot(FPaths::Combine(Plugin->GetBaseDir(), TEXT("Resources")));
+	SetCoreContentRoot(FPaths::EngineContentDir() / TEXT("Slate"));
 	
-	Style->Set("AvalancheOutliner.FilterIcon", new IMAGE_BRUSH("Icons/OutlinerIcons/FilterIcon", Icon20x20));
+	Set("AvalancheOutliner.FilterIcon", new IMAGE_BRUSH("Icons/OutlinerIcons/FilterIcon", Icon20x20));
 	
 	// Table View Row Style
 	const FTableRowStyle& TableRowStyle = FAppStyle::Get().GetWidgetStyle<FTableRowStyle>("SceneOutliner.TableViewRow");
-	Style->Set("AvalancheOutliner.TableViewRow", FTableRowStyle(TableRowStyle)
+	Set("AvalancheOutliner.TableViewRow", FTableRowStyle(TableRowStyle)
 		.SetEvenRowBackgroundBrush(FSlateNoResource())
 		.SetOddRowBackgroundBrush(FSlateColorBrush(FStyleColors::Background)));
-	
-	return Style;
+
+	FSlateStyleRegistry::RegisterSlateStyle(*this);
 }
 
-#undef BORDER_BRUSH
-#undef BOX_BRUSH
-#undef CORE_IMAGE_BRUSH
-#undef IMAGE_BRUSH
+FAvaOutlinerStyle::~FAvaOutlinerStyle()
+{
+	FSlateStyleRegistry::UnRegisterSlateStyle(*this);
+}
diff --git a/Engine/Plugins/Experimental/Avalanche/Source/AvalancheOutliner/Private/AvaOutlinerStyle.h b/Engine/Plugins/Experimental/Avalanche/Source/AvalancheOutliner/Private/AvaOutlinerStyle.h
--- a/Engine/Plugins/Experimental/Avalanche/Source/AvalancheOutliner/Private/AvaOutlinerStyle.h
+++ b/Engine/Plugins/Experimental/Avalanche/Source/AvalancheOutliner/Private/AvaOutlinerStyle.h
@@ -2,31 +2,17 @@
 
 #pragma once
 
-#include "Styling/ISlateStyle.h"
 #include "Styling/SlateStyle.h"
 
-struct FLinearColor;
-struct FSlateBrush;
-
-class FAvaOutlinerStyle
+class FAvaOutlinerStyle final : public FSlateStyleSet
 {
 public:
-	static void Initialize();
-	static void Shutdown();
-
-	static FName GetStyleSetName();
-	static const ISlateStyle& Get();
-
-	static const FLinearColor& GetColor(FName InPropertyName, const ANSICHAR* InSpecifier = nullptr);
-	static const FSlateBrush* GetBrush(FName InPropertyName, const ANSICHAR* InSpecifier = nullptr);
-
-	template<typename InWidgetStyleType>
-	static const InWidgetStyleType& GetWidgetStyle(FName InPropertyName, const ANSICHAR* InSpecifier = nullptr)
+	static FAvaOutlinerStyle& Get()
 	{
-		return StyleInstance->GetWidgetStyle<InWidgetStyleType>(InPropertyName, InSpecifier);
+		static FAvaOutlinerStyle Instance;
+		return Instance;
 	}
 
-private:
-	static TSharedRef<FSlateStyleSet> Create();
-	static TSharedPtr<FSlateStyleSet> StyleInstance;
+	FAvaOutlinerStyle();
+	virtual ~FAvaOutlinerStyle() override;
 };
diff --git a/Engine/Plugins/Experimental/Avalanche/Source/AvalancheOutliner/Private/Slate/SAvaOutlinerItemFilters.cpp b/Engine/Plugins/Experimental/Avalanche/Source/AvalancheOutliner/Private/Slate/SAvaOutlinerItemFilters.cpp
--- a/Engine/Plugins/Experimental/Avalanche/Source/AvalancheOutliner/Private/Slate/SAvaOutlinerItemFilters.cpp
+++ b/Engine/Plugins/Experimental/Avalanche/Source/AvalancheOutliner/Private/Slate/SAvaOutlinerItemFilters.cpp
@@ -72,7 +72,7 @@ void SAvaOutlinerItemFilters::Construct(const FArguments& InArgs, const TSharedR
 			.OnClicked(this, &SAvaOutlinerItemFilters::ToggleShowItemFilters)
 			[
 				SNew(SImage)
-				.Image(FAvaOutlinerStyle::GetBrush(TEXT("AvalancheOutliner.FilterIcon")))
+				.Image(FAvaOutlinerStyle::Get().GetBrush(TEXT("AvalancheOutliner.FilterIcon")))
 				.DesiredSizeOverride(FVector2D(24.f))
 			]
 		];
diff --git a/Engine/Plugins/Experimental/Avalanche/Source/AvalancheSequencer/Private/EaseCurveTool/AvaEaseCurveStyle.cpp b/Engine/Plugins/Experimental/Avalanche/Source/AvalancheSequencer/Private/EaseCurveTool/AvaEaseCurveStyle.cpp
--- a/Engine/Plugins/Experimental/Avalanche/Source/AvalancheSequencer/Private/EaseCurveTool/AvaEaseCurveStyle.cpp
+++ b/Engine/Plugins/Experimental/Avalanche/Source/AvalancheSequencer/Private/EaseCurveTool/AvaEaseCurveStyle.cpp
@@ -1,24 +1,25 @@
 // Copyright Epic Games, Inc. All Rights Reserved.
 
 #include "AvaEaseCurveStyle.h"
-
-#include "Framework/Application/SlateApplication.h"
-#include "Interfaces/IPluginManager.h"
 #include "Brushes/SlateImageBrush.h"
+#include "Brushes/SlateRoundedBoxBrush.h"
+#include "Interfaces/IPluginManager.h"
+#include "Fonts/SlateFontInfo.h"
+#include "Layout/Margin.h"
 #include "Styling/AppStyle.h"
-#include "Styling/SlateStyle.h"
+#include "Styling/CoreStyle.h"
+#include "Styling/StyleColors.h"
 #include "Styling/SlateStyleMacros.h"
 #include "Styling/SlateStyleRegistry.h"
 #include "Styling/SlateTypes.h"
-#include "Styling/StyleColors.h"
 
 FAvaEaseCurveStyle::FAvaEaseCurveStyle()
-	: FSlateStyleSet(TEXT("AvaEaseCurveStyle"))
+	: FSlateStyleSet(TEXT("EaseCurveTool"))
 {
-	const TSharedPtr<IPlugin> Plugin = IPluginManager::Get().FindPlugin(TEXT("Avalanche"));
+	const TSharedPtr<IPlugin> Plugin = IPluginManager::Get().FindPlugin(UE_PLUGIN_NAME);
 	check(Plugin.IsValid());
 	
-	SetContentRoot(FPaths::Combine(Plugin->GetBaseDir(), TEXT("Resources")));
+	SetContentRoot(Plugin->GetBaseDir() / TEXT("Resources"));
 
 	Set("Preset.Selected", new FSlateRoundedBoxBrush(FStyleColors::Transparent, FVector4(4.f, 0.f, 0.f, 4.f), FStyleColors::Select.GetSpecifiedColor(), 1.f));
 
diff --git a/Engine/Plugins/Experimental/Avalanche/Source/AvalancheSequencer/Private/EaseCurveTool/AvaEaseCurveStyle.h b/Engine/Plugins/Experimental/Avalanche/Source/AvalancheSequencer/Private/EaseCurveTool/AvaEaseCurveStyle.h
--- a/Engine/Plugins/Experimental/Avalanche/Source/AvalancheSequencer/Private/EaseCurveTool/AvaEaseCurveStyle.h
+++ b/Engine/Plugins/Experimental/Avalanche/Source/AvalancheSequencer/Private/EaseCurveTool/AvaEaseCurveStyle.h
@@ -9,8 +9,8 @@ class FAvaEaseCurveStyle final : public FSlateStyleSet
 public:
 	static FAvaEaseCurveStyle& Get()
 	{
-		static FAvaEaseCurveStyle StyleSet;
-		return StyleSet;
+		static FAvaEaseCurveStyle Instance;
+		return Instance;
 	}
 
 	FAvaEaseCurveStyle();
diff --git a/Engine/Plugins/Experimental/Avalanche/Source/AvalancheShapesEditor/Private/AvaShapesEditorCommands.cpp b/Engine/Plugins/Experimental/Avalanche/Source/AvalancheShapesEditor/Private/AvaShapesEditorCommands.cpp
--- a/Engine/Plugins/Experimental/Avalanche/Source/AvalancheShapesEditor/Private/AvaShapesEditorCommands.cpp
+++ b/Engine/Plugins/Experimental/Avalanche/Source/AvalancheShapesEditor/Private/AvaShapesEditorCommands.cpp
@@ -10,7 +10,7 @@ FAvaShapesEditorCommands::FAvaShapesEditorCommands()
 		TEXT("AvalancheShapesEditor")
 		, LOCTEXT("MotionDesignShapesEditor", "Motion Design Shapes Editor")
 		, NAME_None
-		, FAvaShapesEditorStyle::GetStyleSetName()
+		, FAvaShapesEditorStyle::Get().GetStyleSetName()
 	)
 {
 }
diff --git a/Engine/Plugins/Experimental/Avalanche/Source/AvalancheShapesEditor/Private/AvaShapesEditorStyle.cpp b/Engine/Plugins/Experimental/Avalanche/Source/AvalancheShapesEditor/Private/AvaShapesEditorStyle.cpp
--- a/Engine/Plugins/Experimental/Avalanche/Source/AvalancheShapesEditor/Private/AvaShapesEditorStyle.cpp
+++ b/Engine/Plugins/Experimental/Avalanche/Source/AvalancheShapesEditor/Private/AvaShapesEditorStyle.cpp
@@ -1,120 +1,77 @@
 // Copyright Epic Games, Inc. All Rights Reserved.
 
 #include "AvaShapesEditorStyle.h"
+#include "Brushes/SlateImageBrush.h"
 #include "Interfaces/IPluginManager.h"
+#include "Math/MathFwd.h"
+#include "Misc/Paths.h"
+#include "Styling/SlateStyleMacros.h"
 #include "Styling/SlateStyleRegistry.h"
 
-#define LOCTEXT_NAMESPACE "AvaShapesEditorStyle"
-
-TSharedPtr<FSlateStyleSet> FAvaShapesEditorStyle::StyleInstance;
-
-void FAvaShapesEditorStyle::Initialize()
-{
-	StyleInstance = Create();
-	FSlateStyleRegistry::RegisterSlateStyle(Get());
-}
-
-void FAvaShapesEditorStyle::Shutdown()
-{
-	FSlateStyleRegistry::UnRegisterSlateStyle(Get());
-	StyleInstance.Reset();
-}
-
-FName FAvaShapesEditorStyle::GetStyleSetName()
-{
-	static FName StyleSetName(TEXT("AvalancheShapesEditor"));
-	return StyleSetName;
-}
-
-const ISlateStyle& FAvaShapesEditorStyle::Get()
-{
-	if (StyleInstance.IsValid() == false)
-	{
-		Initialize();
-	}
-
-	return *StyleInstance;
-}
-
-const FSlateBrush* FAvaShapesEditorStyle::GetBrush(FName PropertyName, const ANSICHAR* Specifier)
-{
-	return StyleInstance->GetBrush(PropertyName, Specifier);
-}
-
-#define IMAGE_BRUSH_SVG(RelativePath, ...) FSlateVectorImageBrush(Style->RootToContentDir(RelativePath, TEXT(".svg")), __VA_ARGS__)
-#define IMAGE_BRUSH(RelativePath, ... ) FSlateImageBrush(Style->RootToContentDir(RelativePath, TEXT(".png")), __VA_ARGS__)
-
-TSharedRef<FSlateStyleSet> FAvaShapesEditorStyle::Create()
+FAvaShapesEditorStyle::FAvaShapesEditorStyle()
+	: FSlateStyleSet(TEXT("AvaShapesEditor"))
 {
 	const FVector2f Icon16x16(16.0f, 16.0f);
 	const FVector2f Icon20x20(20.0f, 20.0f);
 
-	TSharedRef<FSlateStyleSet> Style = MakeShared<FSlateStyleSet>(GetStyleSetName());
-
-	TSharedPtr<IPlugin> Plugin = IPluginManager::Get().FindPlugin(TEXT("Avalanche"));
+	const TSharedPtr<IPlugin> Plugin = IPluginManager::Get().FindPlugin(UE_PLUGIN_NAME);
 	check(Plugin.IsValid());
-	if (Plugin.IsValid())
-	{
-		Style->SetContentRoot(Plugin->GetBaseDir() / TEXT("Resources"));
-	}
 
-	// Path Constants
-	const FString UEContentEditorSlatePath = FPaths::EngineContentDir() / TEXT("Editor/Slate");
-	const FString UEContentEditorSlateIconsPath = UEContentEditorSlatePath / TEXT("Icons");
+	SetContentRoot(Plugin->GetBaseDir() / TEXT("Resources"));
 
 	// Class Icons
-	Style->Set("ClassIcon.AvaShape2DArrowDynamicMesh", new IMAGE_BRUSH("Icons/ToolboxIcons/arrow", Icon16x16));
-	Style->Set("ClassIcon.AvaShapeChevronDynamicMesh", new IMAGE_BRUSH("Icons/ToolboxIcons/chevron", Icon16x16));
-	Style->Set("ClassIcon.AvaShapeConeDynamicMesh", new IMAGE_BRUSH("Icons/ToolboxIcons/cone", Icon16x16));
-	Style->Set("ClassIcon.AvaShapeCubeDynamicMesh", new IMAGE_BRUSH("Icons/ToolboxIcons/cube", Icon16x16));
-	Style->Set("ClassIcon.AvaShapeRoundedPolygonDynamicMesh", new IMAGE_BRUSH_SVG("Icons/ToolboxIcons/Toolbar_Primitive", Icon16x16));
-	Style->Set("ClassIcon.AvaShapeSphereDynamicMesh", new IMAGE_BRUSH("Icons/ToolboxIcons/sphere", Icon16x16));
-	Style->Set("ClassIcon.AvaShapeTorusDynamicMesh", new IMAGE_BRUSH("Icons/ToolboxIcons/torus", Icon16x16));
-	Style->Set("ClassIcon.AvaShapeStarDynamicMesh", new IMAGE_BRUSH_SVG("Icons/ToolboxIcons/Favorite", Icon16x16));
-	Style->Set("ClassIcon.AvaShapeLineDynamicMesh", new IMAGE_BRUSH_SVG("Icons/ToolboxIcons/minus", Icon16x16));
-	Style->Set("ClassIcon.AvaShapeRectangleDynamicMesh", new IMAGE_BRUSH("Icons/ToolboxIcons/rectangle", Icon16x16));
-	Style->Set("ClassIcon.AvaShapeNGonDynamicMesh", new IMAGE_BRUSH("Icons/ToolboxIcons/regularpolygon", Icon16x16));
-	Style->Set("ClassIcon.AvaShapeIrregularPolygonDynamicMesh", new IMAGE_BRUSH("Icons/ToolboxIcons/irregularpolygon", Icon16x16));
-	Style->Set("ClassIcon.AvaShapeEllipseDynamicMesh", new IMAGE_BRUSH("Icons/ToolboxIcons/circle", Icon16x16));
-	Style->Set("ClassIcon.AvaShapeRingDynamicMesh", new IMAGE_BRUSH("Icons/ToolboxIcons/ring", Icon16x16));
-	Style->Set("ClassIcon.AvaShapeShapeActor", new IMAGE_BRUSH("Icons/ToolboxIcons/regularpolygon", Icon16x16));
+	Set("ClassIcon.AvaShape2DArrowDynamicMesh", new IMAGE_BRUSH("Icons/ToolboxIcons/arrow", Icon16x16));
+	Set("ClassIcon.AvaShapeChevronDynamicMesh", new IMAGE_BRUSH("Icons/ToolboxIcons/chevron", Icon16x16));
+	Set("ClassIcon.AvaShapeConeDynamicMesh", new IMAGE_BRUSH("Icons/ToolboxIcons/cone", Icon16x16));
+	Set("ClassIcon.AvaShapeCubeDynamicMesh", new IMAGE_BRUSH("Icons/ToolboxIcons/cube", Icon16x16));
+	Set("ClassIcon.AvaShapeRoundedPolygonDynamicMesh", new IMAGE_BRUSH_SVG("Icons/ToolboxIcons/Toolbar_Primitive", Icon16x16));
+	Set("ClassIcon.AvaShapeSphereDynamicMesh", new IMAGE_BRUSH("Icons/ToolboxIcons/sphere", Icon16x16));
+	Set("ClassIcon.AvaShapeTorusDynamicMesh", new IMAGE_BRUSH("Icons/ToolboxIcons/torus", Icon16x16));
+	Set("ClassIcon.AvaShapeStarDynamicMesh", new IMAGE_BRUSH_SVG("Icons/ToolboxIcons/Favorite", Icon16x16));
+	Set("ClassIcon.AvaShapeLineDynamicMesh", new IMAGE_BRUSH_SVG("Icons/ToolboxIcons/minus", Icon16x16));
+	Set("ClassIcon.AvaShapeRectangleDynamicMesh", new IMAGE_BRUSH("Icons/ToolboxIcons/rectangle", Icon16x16));
+	Set("ClassIcon.AvaShapeNGonDynamicMesh", new IMAGE_BRUSH("Icons/ToolboxIcons/regularpolygon", Icon16x16));
+	Set("ClassIcon.AvaShapeIrregularPolygonDynamicMesh", new IMAGE_BRUSH("Icons/ToolboxIcons/irregularpolygon", Icon16x16));
+	Set("ClassIcon.AvaShapeEllipseDynamicMesh", new IMAGE_BRUSH("Icons/ToolboxIcons/circle", Icon16x16));
+	Set("ClassIcon.AvaShapeRingDynamicMesh", new IMAGE_BRUSH("Icons/ToolboxIcons/ring", Icon16x16));
+	Set("ClassIcon.AvaShapeShapeActor", new IMAGE_BRUSH("Icons/ToolboxIcons/regularpolygon", Icon16x16));
 
 	// 2D Tools
-	Style->Set("AvalancheShapesEditor.Tool_Shape_2DArrow",       new IMAGE_BRUSH("Icons/ToolboxIcons/arrow", Icon20x20));
-	Style->Set("AvalancheShapesEditor.Tool_Shape_Chevron",       new IMAGE_BRUSH("Icons/ToolboxIcons/chevron", Icon20x20));
-	Style->Set("AvalancheShapesEditor.Tool_Shape_Ellipse",       new IMAGE_BRUSH("Icons/ToolboxIcons/circle", Icon20x20));
-	Style->Set("AvalancheShapesEditor.Tool_Shape_IrregularPoly", new IMAGE_BRUSH("Icons/ToolboxIcons/irregularpolygon", Icon20x20));
-	Style->Set("AvalancheShapesEditor.Tool_Shape_Line",          new IMAGE_BRUSH_SVG("Icons/ToolboxIcons/minus", Icon20x20));
-	Style->Set("AvalancheShapesEditor.Tool_Shape_NGon",          new IMAGE_BRUSH("Icons/ToolboxIcons/regularpolygon", Icon20x20));
-	Style->Set("AvalancheShapesEditor.Tool_Shape_Rectangle",     new IMAGE_BRUSH("Icons/ToolboxIcons/rectangle", Icon20x20));
-	Style->Set("AvalancheShapesEditor.Tool_Shape_Ring",          new IMAGE_BRUSH("Icons/ToolboxIcons/ring", Icon20x20));
-	Style->Set("AvalancheShapesEditor.Tool_Shape_Star",          new IMAGE_BRUSH_SVG("Icons/ToolboxIcons/Favorite", Icon20x20));
-
-	Style->Set("Tool_Shape_2DArrow",       new IMAGE_BRUSH("Icons/ToolboxIcons/arrow", Icon20x20));
-	Style->Set("Tool_Shape_Chevron",       new IMAGE_BRUSH("Icons/ToolboxIcons/chevron", Icon20x20));
-	Style->Set("Tool_Shape_Ellipse",       new IMAGE_BRUSH("Icons/ToolboxIcons/circle", Icon20x20));
-	Style->Set("Tool_Shape_IrregularPoly", new IMAGE_BRUSH("Icons/ToolboxIcons/irregularpolygon", Icon20x20));
-	Style->Set("Tool_Shape_Line",          new IMAGE_BRUSH_SVG("Icons/ToolboxIcons/minus", Icon20x20));
-	Style->Set("Tool_Shape_NGon",          new IMAGE_BRUSH("Icons/ToolboxIcons/regularpolygon", Icon20x20));
-	Style->Set("Tool_Shape_Rectangle",     new IMAGE_BRUSH("Icons/ToolboxIcons/rectangle", Icon20x20));
-	Style->Set("Tool_Shape_Ring",          new IMAGE_BRUSH("Icons/ToolboxIcons/ring", Icon20x20));
-	Style->Set("Tool_Shape_Star",          new IMAGE_BRUSH_SVG("Icons/ToolboxIcons/Favorite", Icon20x20));
+	Set("AvalancheShapesEditor.Tool_Shape_2DArrow",       new IMAGE_BRUSH("Icons/ToolboxIcons/arrow", Icon20x20));
+	Set("AvalancheShapesEditor.Tool_Shape_Chevron",       new IMAGE_BRUSH("Icons/ToolboxIcons/chevron", Icon20x20));
+	Set("AvalancheShapesEditor.Tool_Shape_Ellipse",       new IMAGE_BRUSH("Icons/ToolboxIcons/circle", Icon20x20));
+	Set("AvalancheShapesEditor.Tool_Shape_IrregularPoly", new IMAGE_BRUSH("Icons/ToolboxIcons/irregularpolygon", Icon20x20));
+	Set("AvalancheShapesEditor.Tool_Shape_Line",          new IMAGE_BRUSH_SVG("Icons/ToolboxIcons/minus", Icon20x20));
+	Set("AvalancheShapesEditor.Tool_Shape_NGon",          new IMAGE_BRUSH("Icons/ToolboxIcons/regularpolygon", Icon20x20));
+	Set("AvalancheShapesEditor.Tool_Shape_Rectangle",     new IMAGE_BRUSH("Icons/ToolboxIcons/rectangle", Icon20x20));
+	Set("AvalancheShapesEditor.Tool_Shape_Ring",          new IMAGE_BRUSH("Icons/ToolboxIcons/ring", Icon20x20));
+	Set("AvalancheShapesEditor.Tool_Shape_Star",          new IMAGE_BRUSH_SVG("Icons/ToolboxIcons/Favorite", Icon20x20));
+
+	Set("Tool_Shape_2DArrow",       new IMAGE_BRUSH("Icons/ToolboxIcons/arrow", Icon20x20));
+	Set("Tool_Shape_Chevron",       new IMAGE_BRUSH("Icons/ToolboxIcons/chevron", Icon20x20));
+	Set("Tool_Shape_Ellipse",       new IMAGE_BRUSH("Icons/ToolboxIcons/circle", Icon20x20));
+	Set("Tool_Shape_IrregularPoly", new IMAGE_BRUSH("Icons/ToolboxIcons/irregularpolygon", Icon20x20));
+	Set("Tool_Shape_Line",          new IMAGE_BRUSH_SVG("Icons/ToolboxIcons/minus", Icon20x20));
+	Set("Tool_Shape_NGon",          new IMAGE_BRUSH("Icons/ToolboxIcons/regularpolygon", Icon20x20));
+	Set("Tool_Shape_Rectangle",     new IMAGE_BRUSH("Icons/ToolboxIcons/rectangle", Icon20x20));
+	Set("Tool_Shape_Ring",          new IMAGE_BRUSH("Icons/ToolboxIcons/ring", Icon20x20));
+	Set("Tool_Shape_Star",          new IMAGE_BRUSH_SVG("Icons/ToolboxIcons/Favorite", Icon20x20));
 
 	// 3D Tools
-	Style->Set("AvalancheShapesEditor.Tool_Shape_Cone",   new IMAGE_BRUSH("Icons/ToolboxIcons/cone", Icon20x20));
-	Style->Set("AvalancheShapesEditor.Tool_Shape_Cube",   new IMAGE_BRUSH("Icons/ToolboxIcons/cube", Icon20x20));
-	Style->Set("AvalancheShapesEditor.Tool_Shape_Sphere", new IMAGE_BRUSH("Icons/ToolboxIcons/sphere", Icon20x20));
-	Style->Set("AvalancheShapesEditor.Tool_Shape_Torus",  new IMAGE_BRUSH("Icons/ToolboxIcons/torus", Icon20x20));
+	Set("AvalancheShapesEditor.Tool_Shape_Cone",   new IMAGE_BRUSH("Icons/ToolboxIcons/cone", Icon20x20));
+	Set("AvalancheShapesEditor.Tool_Shape_Cube",   new IMAGE_BRUSH("Icons/ToolboxIcons/cube", Icon20x20));
+	Set("AvalancheShapesEditor.Tool_Shape_Sphere", new IMAGE_BRUSH("Icons/ToolboxIcons/sphere", Icon20x20));
+	Set("AvalancheShapesEditor.Tool_Shape_Torus",  new IMAGE_BRUSH("Icons/ToolboxIcons/torus", Icon20x20));
 
-	Style->Set("Tool_Shape_Cone",   new IMAGE_BRUSH("Icons/ToolboxIcons/cone", Icon20x20));
-	Style->Set("Tool_Shape_Cube",   new IMAGE_BRUSH("Icons/ToolboxIcons/cube", Icon20x20));
-	Style->Set("Tool_Shape_Sphere", new IMAGE_BRUSH("Icons/ToolboxIcons/sphere", Icon20x20));
-	Style->Set("Tool_Shape_Torus",  new IMAGE_BRUSH("Icons/ToolboxIcons/torus", Icon20x20));
+	Set("Tool_Shape_Cone",   new IMAGE_BRUSH("Icons/ToolboxIcons/cone", Icon20x20));
+	Set("Tool_Shape_Cube",   new IMAGE_BRUSH("Icons/ToolboxIcons/cube", Icon20x20));
+	Set("Tool_Shape_Sphere", new IMAGE_BRUSH("Icons/ToolboxIcons/sphere", Icon20x20));
+	Set("Tool_Shape_Torus",  new IMAGE_BRUSH("Icons/ToolboxIcons/torus", Icon20x20));
 
-	return Style;
+	FSlateStyleRegistry::RegisterSlateStyle(*this);
 }
 
-#undef IMAGE_BRUSH_SVG
-#undef IMAGE_BRUSH
-
-#undef LOCTEXT_NAMESPACE
+FAvaShapesEditorStyle::~FAvaShapesEditorStyle()
+{
+	FSlateStyleRegistry::UnRegisterSlateStyle(*this);
+}
diff --git a/Engine/Plugins/Experimental/Avalanche/Source/AvalancheShapesEditor/Private/AvaShapesEditorStyle.h b/Engine/Plugins/Experimental/Avalanche/Source/AvalancheShapesEditor/Private/AvaShapesEditorStyle.h
--- a/Engine/Plugins/Experimental/Avalanche/Source/AvalancheShapesEditor/Private/AvaShapesEditorStyle.h
+++ b/Engine/Plugins/Experimental/Avalanche/Source/AvalancheShapesEditor/Private/AvaShapesEditorStyle.h
@@ -2,24 +2,17 @@
 
 #pragma once
 
-#include "Styling/ISlateStyle.h"
 #include "Styling/SlateStyle.h"
 
-struct FSlateBrush;
-
-class FAvaShapesEditorStyle
+class FAvaShapesEditorStyle final : public FSlateStyleSet
 {
 public:
-	static void Initialize();
-	static void Shutdown();
-
-	static FName GetStyleSetName();
-	static const ISlateStyle& Get();
-
-	static const FSlateBrush* GetBrush(FName PropertyName, const ANSICHAR* Specifier = nullptr);
-
-private:
-	static TSharedPtr<FSlateStyleSet> StyleInstance;
-
-	static TSharedRef<FSlateStyleSet> Create();
+	static FAvaShapesEditorStyle& Get()
+	{
+		static FAvaShapesEditorStyle Instance;
+		return Instance;
+	}
+
+	FAvaShapesEditorStyle();
+	virtual ~FAvaShapesEditorStyle() override;
 };
diff --git a/Engine/Plugins/Experimental/Avalanche/Source/AvalancheShapesEditor/Private/AvalancheShapesEditorModule.cpp b/Engine/Plugins/Experimental/Avalanche/Source/AvalancheShapesEditor/Private/AvalancheShapesEditorModule.cpp
--- a/Engine/Plugins/Experimental/Avalanche/Source/AvalancheShapesEditor/Private/AvalancheShapesEditorModule.cpp
+++ b/Engine/Plugins/Experimental/Avalanche/Source/AvalancheShapesEditor/Private/AvalancheShapesEditorModule.cpp
@@ -4,7 +4,6 @@
 #include "AvaInteractiveToolsDelegates.h"
 #include "AvaShapeActor.h"
 #include "AvaShapesEditorCommands.h"
-#include "AvaShapesEditorStyle.h"
 #include "ComponentVisualizers.h"
 #include "Engine/Texture2D.h"
 #include "IAvalancheComponentVisualizersModule.h"
@@ -82,7 +81,6 @@ FAvalancheShapesEditorModule& FAvalancheShapesEditorModule::Get()
 
 void FAvalancheShapesEditorModule::StartupModule()
 {
-	FAvaShapesEditorStyle::Initialize();
 	FAvaShapesEditorCommands::Register();
 	FAvaInteractiveToolsDelegates::GetRegisterToolsDelegate().AddRaw(this, &FAvalancheShapesEditorModule::RegisterShapeTools);
 	FCoreDelegates::OnPostEngineInit.AddRaw(this, &FAvalancheShapesEditorModule::RegisterVisualizers);
@@ -119,7 +117,6 @@ void FAvalancheShapesEditorModule::StartupModule()
 
 void FAvalancheShapesEditorModule::ShutdownModule()
 {
-	FAvaShapesEditorStyle::Shutdown();
 	FAvaShapesEditorCommands::Unregister();
 	FAvaInteractiveToolsDelegates::GetRegisterToolsDelegate().RemoveAll(this);
 	FCoreDelegates::OnPostEngineInit.RemoveAll(this);
diff --git a/Engine/Plugins/Experimental/Avalanche/Source/AvalancheTagEditor/Private/AvaTagEditorStyle.cpp b/Engine/Plugins/Experimental/Avalanche/Source/AvalancheTagEditor/Private/AvaTagEditorStyle.cpp
--- a/Engine/Plugins/Experimental/Avalanche/Source/AvalancheTagEditor/Private/AvaTagEditorStyle.cpp
+++ b/Engine/Plugins/Experimental/Avalanche/Source/AvalancheTagEditor/Private/AvaTagEditorStyle.cpp
@@ -8,7 +8,7 @@
 #include "Styling/StyleColors.h"
 
 FAvaTagEditorStyle::FAvaTagEditorStyle()
-	: FSlateStyleSet(TEXT("AvaTagEditorStyle"))
+	: FSlateStyleSet(TEXT("AvaTagEditor"))
 {
 	const FTableRowStyle& TableRowStyle = FAppStyle::Get().GetWidgetStyle<FTableRowStyle>("TableView.Row");
 
diff --git a/Engine/Plugins/Experimental/Avalanche/Source/AvalancheTagEditor/Private/AvaTagEditorStyle.h b/Engine/Plugins/Experimental/Avalanche/Source/AvalancheTagEditor/Private/AvaTagEditorStyle.h
--- a/Engine/Plugins/Experimental/Avalanche/Source/AvalancheTagEditor/Private/AvaTagEditorStyle.h
+++ b/Engine/Plugins/Experimental/Avalanche/Source/AvalancheTagEditor/Private/AvaTagEditorStyle.h
@@ -4,16 +4,15 @@
 
 #include "Styling/SlateStyle.h"
 
-class FAvaTagEditorStyle : public FSlateStyleSet
+class FAvaTagEditorStyle final : public FSlateStyleSet
 {
 public:
-	FAvaTagEditorStyle();
-
-	virtual ~FAvaTagEditorStyle() override;
-
 	static FAvaTagEditorStyle& Get()
 	{
-		static FAvaTagEditorStyle StyleSet;
-		return StyleSet;
+		static FAvaTagEditorStyle Instance;
+		return Instance;
 	}
+
+	FAvaTagEditorStyle();
+	virtual ~FAvaTagEditorStyle() override;
 };
diff --git a/Engine/Plugins/Experimental/Avalanche/Source/AvalancheTextEditor/Private/AvaTextEditorCommands.cpp b/Engine/Plugins/Experimental/Avalanche/Source/AvalancheTextEditor/Private/AvaTextEditorCommands.cpp
--- a/Engine/Plugins/Experimental/Avalanche/Source/AvalancheTextEditor/Private/AvaTextEditorCommands.cpp
+++ b/Engine/Plugins/Experimental/Avalanche/Source/AvalancheTextEditor/Private/AvaTextEditorCommands.cpp
@@ -10,7 +10,7 @@ FAvaTextEditorCommands::FAvaTextEditorCommands()
 		TEXT("AvaTextEditor")
 		, LOCTEXT("AvaTextEditor", "Motion Design Text Editor")
 		, NAME_None
-		, FAvaTextEditorStyle::GetStyleSetName()
+		, FAvaTextEditorStyle::Get().GetStyleSetName()
 	)
 {
 }
diff --git a/Engine/Plugins/Experimental/Avalanche/Source/AvalancheTextEditor/Private/AvaTextEditorModule.cpp b/Engine/Plugins/Experimental/Avalanche/Source/AvalancheTextEditor/Private/AvaTextEditorModule.cpp
--- a/Engine/Plugins/Experimental/Avalanche/Source/AvalancheTextEditor/Private/AvaTextEditorModule.cpp
+++ b/Engine/Plugins/Experimental/Avalanche/Source/AvalancheTextEditor/Private/AvaTextEditorModule.cpp
@@ -6,7 +6,6 @@
 #include "AvaTextActor.h"
 #include "AvaTextDefs.h"
 #include "AvaTextEditorCommands.h"
-#include "AvaTextEditorStyle.h"
 #include "DetailsView/AvaLinearGradientSettingsCustomization.h"
 #include "DetailsView/AvaTextAlignmentCustomization.h"
 #include "DetailsView/AvaTextComponentCustomization.h"
@@ -27,7 +26,6 @@ void FAvaTextEditorModule::StartupModule()
 {
 	FCoreDelegates::OnPostEngineInit.AddRaw(this, &FAvaTextEditorModule::PostEngineInit);
 
-	FAvaTextEditorStyle::Initialize();
 	FAvaTextEditorCommands::Register();
 	RegisterCustomLayouts();
 	FAvaInteractiveToolsDelegates::GetRegisterToolsDelegate().AddRaw(this, &FAvaTextEditorModule::RegisterTools);
@@ -41,7 +39,6 @@ void FAvaTextEditorModule::ShutdownModule()
 {
 	FCoreDelegates::OnPostEngineInit.RemoveAll(this);
 
-	FAvaTextEditorStyle::Shutdown();
 	FAvaTextEditorCommands::Unregister();
 
 	if (UObjectInitialized() && !IsEngineExitRequested())
diff --git a/Engine/Plugins/Experimental/Avalanche/Source/AvalancheTextEditor/Private/AvaTextEditorStyle.cpp b/Engine/Plugins/Experimental/Avalanche/Source/AvalancheTextEditor/Private/AvaTextEditorStyle.cpp
--- a/Engine/Plugins/Experimental/Avalanche/Source/AvalancheTextEditor/Private/AvaTextEditorStyle.cpp
+++ b/Engine/Plugins/Experimental/Avalanche/Source/AvalancheTextEditor/Private/AvaTextEditorStyle.cpp
@@ -1,76 +1,31 @@
 // Copyright Epic Games, Inc. All Rights Reserved.
 
 #include "AvaTextEditorStyle.h"
-
 #include "Brushes/SlateImageBrush.h"
-#include "AvaTextDefs.h"
 #include "Interfaces/IPluginManager.h"
-#include "SlateOptMacros.h"
-#include "Styling/SlateStyle.h"
+#include "Math/MathFwd.h"
+#include "Styling/SlateStyleMacros.h"
 #include "Styling/SlateStyleRegistry.h"
 
-TSharedPtr<FSlateStyleSet> FAvaTextEditorStyle::StyleInstance = nullptr;
-
-void FAvaTextEditorStyle::Initialize()
-{
-	if (!StyleInstance.IsValid())
-	{
-		StyleInstance = Create();
-		FSlateStyleRegistry::RegisterSlateStyle(*StyleInstance);
-	}
-}
-
-void FAvaTextEditorStyle::Shutdown()
-{
-	FSlateStyleRegistry::UnRegisterSlateStyle(*StyleInstance);
-	ensure(StyleInstance.IsUnique());
-	StyleInstance.Reset();
-}
-
-FName FAvaTextEditorStyle::GetStyleSetName()
-{
-	static FName StyleSetName(TEXT("AvaTextEditor"));
-	return StyleSetName;
-}
-
-const ISlateStyle& FAvaTextEditorStyle::Get()
-{
-	return *StyleInstance;
-}
-
-const FSlateBrush* FAvaTextEditorStyle::GetBrush(FName PropertyName, const ANSICHAR* Specifier)
+FAvaTextEditorStyle::FAvaTextEditorStyle()
+	: FSlateStyleSet(TEXT("AvaTextEditor"))
 {
-	return StyleInstance->GetBrush(PropertyName, Specifier);
-}
-
-BEGIN_SLATE_FUNCTION_BUILD_OPTIMIZATION
-
-#define IMAGE_BRUSH_SVG(RelativePath, ...) FSlateVectorImageBrush(Style->RootToContentDir(RelativePath, TEXT(".svg")), __VA_ARGS__)
-#define IMAGE_BRUSH(RelativePath, ... ) FSlateImageBrush(Style->RootToContentDir(RelativePath, TEXT(".png")), __VA_ARGS__)
-
-TSharedRef<FSlateStyleSet> FAvaTextEditorStyle::Create()
-{
-	TSharedRef<FSlateStyleSet> Style = MakeShared<FSlateStyleSet>("AvaTextEditor");
+	const FVector2D Icon16x16(16.0f, 16.0f);
+	const FVector2D Icon20x20(20.0f, 20.0f);
 
-	TSharedPtr<IPlugin> Plugin = IPluginManager::Get().FindPlugin(TEXT("Avalanche"));
+	const TSharedPtr<IPlugin> Plugin = IPluginManager::Get().FindPlugin(UE_PLUGIN_NAME);
 	check(Plugin.IsValid());
-	if (Plugin.IsValid())
-	{
-		Style->SetContentRoot(Plugin->GetBaseDir() / TEXT("Resources"));
-	}
 
-	const FVector2D Icon16x16(16.0f, 16.0f);
-	const FVector2D Icon20x20(20.0f, 20.0f);
-	const FVector2D Icon22x22(22.0f, 22.0f);
+	SetContentRoot(Plugin->GetBaseDir() / TEXT("Resources"));
 	
-	Style->Set("ClassIcon.AvaTextActor", new IMAGE_BRUSH("Icons/ToolboxIcons/3d-text", Icon16x16));
-	Style->Set("AvaTextEditor.Tool_Actor_Text3D", new IMAGE_BRUSH_SVG("Icons/ToolboxIcons/Toolbar_Text", Icon20x20));
-	Style->Set("Tool_Actor_Text3D", new IMAGE_BRUSH_SVG("Icons/ToolboxIcons/Toolbar_Text", Icon20x20));
+	Set("ClassIcon.AvaTextActor", new IMAGE_BRUSH("Icons/ToolboxIcons/3d-text", Icon16x16));
+	Set("AvaTextEditor.Tool_Actor_Text3D", new IMAGE_BRUSH_SVG("Icons/ToolboxIcons/Toolbar_Text", Icon20x20));
+	Set("Tool_Actor_Text3D", new IMAGE_BRUSH_SVG("Icons/ToolboxIcons/Toolbar_Text", Icon20x20));
 
-	return Style;
+	FSlateStyleRegistry::RegisterSlateStyle(*this);
 }
 
-#undef IMAGE_BRUSH_SVG
-#undef IMAGE_BRUSH
-
-END_SLATE_FUNCTION_BUILD_OPTIMIZATION
+FAvaTextEditorStyle::~FAvaTextEditorStyle()
+{
+	FSlateStyleRegistry::UnRegisterSlateStyle(*this);
+}
diff --git a/Engine/Plugins/Experimental/Avalanche/Source/AvalancheTextEditor/Private/AvaTextEditorStyle.h b/Engine/Plugins/Experimental/Avalanche/Source/AvalancheTextEditor/Private/AvaTextEditorStyle.h
--- a/Engine/Plugins/Experimental/Avalanche/Source/AvalancheTextEditor/Private/AvaTextEditorStyle.h
+++ b/Engine/Plugins/Experimental/Avalanche/Source/AvalancheTextEditor/Private/AvaTextEditorStyle.h
@@ -2,22 +2,17 @@
 
 #pragma once
 
-#include "Styling/ISlateStyle.h"
 #include "Styling/SlateStyle.h"
 
-class FAvaTextEditorStyle
+class FAvaTextEditorStyle final : public FSlateStyleSet
 {
 public:
-	static void Initialize();
-	static void Shutdown();
+	static FAvaTextEditorStyle& Get()
+	{
+		static FAvaTextEditorStyle Instance;
+		return Instance;
+	}
 
-	static FName GetStyleSetName();
-	static const ISlateStyle& Get();
-
-	static const FSlateBrush* GetBrush(FName PropertyName, const ANSICHAR* Specifier = nullptr);
-
-private:
-	static TSharedPtr<FSlateStyleSet> StyleInstance;
-
-	static TSharedRef<FSlateStyleSet> Create();
+	FAvaTextEditorStyle();
+	virtual ~FAvaTextEditorStyle() override;
 };
diff --git a/Engine/Plugins/Experimental/Avalanche/Source/AvalancheTransitionEditor/Private/AvaTransitionCommands.cpp b/Engine/Plugins/Experimental/Avalanche/Source/AvalancheTransitionEditor/Private/AvaTransitionCommands.cpp
--- a/Engine/Plugins/Experimental/Avalanche/Source/AvalancheTransitionEditor/Private/AvaTransitionCommands.cpp
+++ b/Engine/Plugins/Experimental/Avalanche/Source/AvalancheTransitionEditor/Private/AvaTransitionCommands.cpp
@@ -1,6 +1,8 @@
 // Copyright Epic Games, Inc. All Rights Reserved.
 
 #include "AvaTransitionCommands.h"
+
+#include "AvaTransitionEditorStyle.h"
 #include "Styling/AppStyle.h"
 
 #define LOCTEXT_NAMESPACE "AvaTransitionEditorCommands"
@@ -9,7 +11,7 @@ FAvaTransitionEditorCommands::FAvaTransitionEditorCommands()
 	: TCommands(TEXT("AvaTransitionEditorCommands")
 	, LOCTEXT("AvaTransitionEditor", "Motion Design Transition Logic Editor")
 	, NAME_None
-	, TEXT("AvaTransitionEditorStyle"))
+	, FAvaTransitionEditorStyle::Get().GetStyleSetName())
 {
 }
 
diff --git a/Engine/Plugins/Experimental/Avalanche/Source/AvalancheTransitionEditor/Private/AvaTransitionEditorModule.cpp b/Engine/Plugins/Experimental/Avalanche/Source/AvalancheTransitionEditor/Private/AvaTransitionEditorModule.cpp
--- a/Engine/Plugins/Experimental/Avalanche/Source/AvalancheTransitionEditor/Private/AvaTransitionEditorModule.cpp
+++ b/Engine/Plugins/Experimental/Avalanche/Source/AvalancheTransitionEditor/Private/AvaTransitionEditorModule.cpp
@@ -29,9 +29,6 @@ IMPLEMENT_MODULE(FAvaTransitionEditorModule, AvalancheTransitionEditor)
 
 void FAvaTransitionEditorModule::StartupModule()
 {
-	// Trigger Registration of the Style on Module Startup
-	FAvaTransitionEditorStyle::Get();
-
 #if WITH_STATETREE_DEBUGGER
 	FAvaTransitionTraceModule::Startup();
 #endif
diff --git a/Engine/Plugins/Experimental/Avalanche/Source/AvalancheTransitionEditor/Private/AvaTransitionEditorStyle.cpp b/Engine/Plugins/Experimental/Avalanche/Source/AvalancheTransitionEditor/Private/AvaTransitionEditorStyle.cpp
--- a/Engine/Plugins/Experimental/Avalanche/Source/AvalancheTransitionEditor/Private/AvaTransitionEditorStyle.cpp
+++ b/Engine/Plugins/Experimental/Avalanche/Source/AvalancheTransitionEditor/Private/AvaTransitionEditorStyle.cpp
@@ -1,21 +1,26 @@
 // Copyright Epic Games, Inc. All Rights Reserved.
 
 #include "AvaTransitionEditorStyle.h"
+#include "Brushes/SlateImageBrush.h"
 #include "AvaTransitionTree.h"
 #include "Interfaces/IPluginManager.h"
+#include "Layout/Margin.h"
 #include "Misc/Paths.h"
 #include "Styling/SlateStyleMacros.h"
 #include "Styling/SlateStyleRegistry.h"
 
 FAvaTransitionEditorStyle::FAvaTransitionEditorStyle()
-	: FSlateStyleSet(TEXT("AvaTransitionEditorStyle"))
+	: FSlateStyleSet(TEXT("AvaTransitionEditor"))
 {
 	const FVector2f Icon16(16.f);
 	const FVector2f Icon20(20.f);
 	const FVector2f Icon64(64.f);
 
-	ContentRootDir     = FPaths::EngineContentDir() / TEXT("Editor/Slate");
-	CoreContentRootDir = FPaths::EngineContentDir() / TEXT("Slate");
+	const TSharedPtr<IPlugin> Plugin = IPluginManager::Get().FindPlugin(UE_PLUGIN_NAME);
+	check(Plugin.IsValid());
+
+	SetContentRoot(FPaths::EngineContentDir() / TEXT("Editor/Slate"));
+	SetCoreContentRoot(FPaths::EngineContentDir() / TEXT("Slate"));
 
 	// Re-use the Behavior Tree Icon
 	Set("ClassIcon.AvaTransitionTree"     , new IMAGE_BRUSH("Icons/AssetIcons/BehaviorTree_16x", Icon16));
@@ -31,11 +36,6 @@ FAvaTransitionEditorStyle::FAvaTransitionEditorStyle()
 	Set("AvaTransitionEditorCommands.ExportTransitionTree"  , new CORE_IMAGE_BRUSH_SVG("Starship/Common/export_20", Icon20));
 	Set("AvaTransitionEditorCommands.ToggleDebug"           , new IMAGE_BRUSH_SVG("Starship/Common/Debug"         , Icon20));
 
-	const TSharedPtr<IPlugin> Plugin = IPluginManager::Get().FindPlugin(UE_PLUGIN_NAME);
-	check(Plugin.IsValid());
-
-	ContentRootDir = (FPaths::Combine(Plugin->GetBaseDir(), TEXT("Resources")));
-
 	Set("DebugIndicatorBorder", new BOX_BRUSH("Images/NamespaceBorder", FMargin(0.25f)));
 
 	FSlateStyleRegistry::RegisterSlateStyle(*this);
diff --git a/Engine/Plugins/Experimental/Avalanche/Source/AvalancheTransitionEditor/Private/AvaTransitionEditorStyle.h b/Engine/Plugins/Experimental/Avalanche/Source/AvalancheTransitionEditor/Private/AvaTransitionEditorStyle.h
--- a/Engine/Plugins/Experimental/Avalanche/Source/AvalancheTransitionEditor/Private/AvaTransitionEditorStyle.h
+++ b/Engine/Plugins/Experimental/Avalanche/Source/AvalancheTransitionEditor/Private/AvaTransitionEditorStyle.h
@@ -4,18 +4,17 @@
 
 #include "Styling/SlateStyle.h"
 
-class FAvaTransitionEditorStyle : public FSlateStyleSet
+class FAvaTransitionEditorStyle final : public FSlateStyleSet
 {
 public:
-	FAvaTransitionEditorStyle();
-
-	virtual ~FAvaTransitionEditorStyle() override;
-
 	static FAvaTransitionEditorStyle& Get()
 	{
-		static FAvaTransitionEditorStyle StyleSet;
-		return StyleSet;
+		static FAvaTransitionEditorStyle Instance;
+		return Instance;
 	}
 
+	FAvaTransitionEditorStyle();
+	virtual ~FAvaTransitionEditorStyle() override;
+
 	static FLinearColor LerpColorSRGB(const FLinearColor& InA, const FLinearColor& InB, float InAlpha);
 };
