Motion Design Tags:
- UI improved to be shorter in width by default. Before this the rows would be filled so for Tag Collection Editor window filling a lot of space, the button to add a Tag Entry and the text boxes to edit these tags would be at opposite sides of the window

#changelist validated
#virtualized

diff --git a/Engine/Plugins/Experimental/Avalanche/Source/AvalancheTagEditor/Private/Customization/AvaTagCollectionCustomization.cpp b/Engine/Plugins/Experimental/Avalanche/Source/AvalancheTagEditor/Private/Customization/AvaTagCollectionCustomization.cpp
--- a/Engine/Plugins/Experimental/Avalanche/Source/AvalancheTagEditor/Private/Customization/AvaTagCollectionCustomization.cpp
+++ b/Engine/Plugins/Experimental/Avalanche/Source/AvalancheTagEditor/Private/Customization/AvaTagCollectionCustomization.cpp
@@ -9,7 +9,9 @@
 #include "IDetailChildrenBuilder.h"
 #include "PropertyCustomizationHelpers.h"
 #include "Widgets/Input/SComboButton.h"
+#include "Widgets/Layout/SBox.h"
 #include "Widgets/SBoxPanel.h"
+#include "Widgets/Text/STextBlock.h"
 
 #define LOCTEXT_NAMESPACE "AvaTagCollectionCustomization"
 
@@ -20,15 +22,32 @@ void FAvaTagCollectionCustomization::CustomizeDetails(IDetailLayoutBuilder& InDe
 
 	IDetailCategoryBuilder& Category = InDetailBuilder.EditCategory(TEXT("Tag"));
 
-	TSharedRef<SHorizontalBox> HeaderContentWidget = SNew(SHorizontalBox);
-	HeaderContentWidget->AddSlot()
-		.HAlign(HAlign_Right)
-		.VAlign(VAlign_Center)
+	TSharedRef<SBox> HeaderContentWidget = SNew(SBox)
+		.HAlign(HAlign_Left)
 		[
-			TagMapProperty->CreateDefaultPropertyButtonWidgets()
+			SNew(SBox)
+			.MinDesiredWidth(250.f)
+			[
+				SNew(SHorizontalBox)
+				+ SHorizontalBox::Slot()
+				.FillWidth(1.f)
+				.VAlign(VAlign_Center)
+				[
+					SNew(STextBlock)
+					.Text(Category.GetDisplayName())
+					.Font(IDetailLayoutBuilder::GetDetailFont())
+				]
+				+ SHorizontalBox::Slot()
+				.AutoWidth()
+				.HAlign(HAlign_Right)
+				.VAlign(VAlign_Center)
+				[
+					TagMapProperty->CreateDefaultPropertyButtonWidgets()
+				]
+			]
 		];
 
-	Category.HeaderContent(HeaderContentWidget);
+	Category.HeaderContent(HeaderContentWidget, /*bWholeRowContent*/true);
 	Category.AddCustomBuilder(MakeShared<FAvaTagMapBuilder>(TagMapProperty.ToSharedRef()), /*bForAdvanced*/false);
 }
 
@@ -73,20 +92,25 @@ void FAvaTagMapBuilder::GenerateChildContent(IDetailChildrenBuilder& InChildrenB
 			InChildrenBuilder.AddProperty(TagHandle.ToSharedRef())
 				.CustomWidget()
 				.WholeRowContent()
-				.HAlign(HAlign_Fill)
+				.HAlign(HAlign_Left)
 				[
-					SNew(SHorizontalBox)
-					+ SHorizontalBox::Slot()
-					.FillWidth(1.f)
-					[
-						TagNameHandle->CreatePropertyValueWidget()
-					]
-					+ SHorizontalBox::Slot()
-					.AutoWidth()
-					.HAlign(HAlign_Center)
-					.VAlign(VAlign_Center)
+					SNew(SBox)
+					.MinDesiredWidth(250.f)
 					[
-						CreatePropertyButtonsWidget(TagHandle)
+						SNew(SHorizontalBox)
+						+ SHorizontalBox::Slot()
+						.FillWidth(1.f)
+						.VAlign(VAlign_Center)
+						[
+							TagNameHandle->CreatePropertyValueWidget()
+						]
+						+ SHorizontalBox::Slot()
+						.AutoWidth()
+						.HAlign(HAlign_Right)
+						.VAlign(VAlign_Center)
+						[
+							CreatePropertyButtonsWidget(TagHandle)
+						]
 					]
 				];
 		}
