Motion Design Transition:
- Fixed issue with Scene State's Settings weak ref only set on BeginPlay, causing issues with TL Preview Actor to not correctly process the Scene State Tag Condition

#changelist validated
#virtualized

diff --git a/Engine/Plugins/Experimental/Avalanche/Source/Avalanche/Private/AvaScene.cpp b/Engine/Plugins/Experimental/Avalanche/Source/Avalanche/Private/AvaScene.cpp
--- a/Engine/Plugins/Experimental/Avalanche/Source/Avalanche/Private/AvaScene.cpp
+++ b/Engine/Plugins/Experimental/Avalanche/Source/Avalanche/Private/AvaScene.cpp
@@ -62,6 +62,11 @@ AAvaScene::AAvaScene()
 
 	SceneState = CreateDefaultSubobject<UAvaSceneState>(TEXT("SceneState"));
 
+	if (SceneState)
+	{
+		SceneState->SetSceneSettings(SceneSettings);
+	}
+
 	RemoteControlPreset = CreateDefaultSubobject<URemoteControlPreset>(TEXT("RemoteControlPreset"));
 
 	StartupCameraName = NAME_None;
@@ -245,16 +250,6 @@ void AAvaScene::RebuildSequenceTree()
 	IAvaSequenceProvider::RebuildSequenceTree();
 }
 
-void AAvaScene::BeginPlay()
-{
-	Super::BeginPlay();
-
-	if (SceneState)
-	{
-		SceneState->OnBeginPlay(SceneSettings);
-	}
-}
-
 void AAvaScene::PostActorCreated()
 {
 	Super::PostActorCreated();
diff --git a/Engine/Plugins/Experimental/Avalanche/Source/Avalanche/Private/AvaSceneState.cpp b/Engine/Plugins/Experimental/Avalanche/Source/Avalanche/Private/AvaSceneState.cpp
--- a/Engine/Plugins/Experimental/Avalanche/Source/Avalanche/Private/AvaSceneState.cpp
+++ b/Engine/Plugins/Experimental/Avalanche/Source/Avalanche/Private/AvaSceneState.cpp
@@ -4,7 +4,7 @@
 #include "AvaSceneSettings.h"
 #include "AvaTagHandle.h"
 
-void UAvaSceneState::OnBeginPlay(UAvaSceneSettings* InSceneSettings)
+void UAvaSceneState::SetSceneSettings(UAvaSceneSettings* InSceneSettings)
 {
 	SceneSettingsWeak = InSceneSettings;
 }
diff --git a/Engine/Plugins/Experimental/Avalanche/Source/Avalanche/Public/AvaScene.h b/Engine/Plugins/Experimental/Avalanche/Source/Avalanche/Public/AvaScene.h
--- a/Engine/Plugins/Experimental/Avalanche/Source/Avalanche/Public/AvaScene.h
+++ b/Engine/Plugins/Experimental/Avalanche/Source/Avalanche/Public/AvaScene.h
@@ -82,7 +82,6 @@ public:
 	//~ End IAvaSequenceProvider
 
 	//~ Begin AActor
-	virtual void BeginPlay() override;
 	virtual void PostActorCreated() override;
 #if WITH_EDITOR
 	virtual bool IsSelectable() const override { return false; }
diff --git a/Engine/Plugins/Experimental/Avalanche/Source/Avalanche/Public/AvaSceneState.h b/Engine/Plugins/Experimental/Avalanche/Source/Avalanche/Public/AvaSceneState.h
--- a/Engine/Plugins/Experimental/Avalanche/Source/Avalanche/Public/AvaSceneState.h
+++ b/Engine/Plugins/Experimental/Avalanche/Source/Avalanche/Public/AvaSceneState.h
@@ -15,7 +15,7 @@ class UAvaSceneState : public UObject
 	GENERATED_BODY()
 
 public:
-	void OnBeginPlay(UAvaSceneSettings* InSceneSettings);
+	void SetSceneSettings(UAvaSceneSettings* InSceneSettings);
 
 	AVALANCHE_API bool AddTagAttribute(const FAvaTagHandle& InTagHandle);
 
