[Motion Design]
 - Update style set names

diff --git a/Engine/Plugins/Experimental/Avalanche/Source/AvalancheEditor/Private/AvaEditorCommands.cpp b/Engine/Plugins/Experimental/Avalanche/Source/AvalancheEditor/Private/AvaEditorCommands.cpp
--- a/Engine/Plugins/Experimental/Avalanche/Source/AvalancheEditor/Private/AvaEditorCommands.cpp
+++ b/Engine/Plugins/Experimental/Avalanche/Source/AvalancheEditor/Private/AvaEditorCommands.cpp
@@ -9,7 +9,7 @@
 #define LOCTEXT_NAMESPACE "AvaEditorCommands"
 
 FAvaEditorCommands::FAvaEditorCommands()
-	: TCommands<FAvaEditorCommands>(TEXT("AvalancheEditor")
+	: TCommands<FAvaEditorCommands>(TEXT("AvaEditor")
 		, LOCTEXT("MotionDesignEditor", "Motion Design Editor")
 		, NAME_None
 		, FAvaEditorStyle::Get().GetStyleSetName()
diff --git a/Engine/Plugins/Experimental/Avalanche/Source/AvalancheEditor/Private/AvaEditorStyle.cpp b/Engine/Plugins/Experimental/Avalanche/Source/AvalancheEditor/Private/AvaEditorStyle.cpp
--- a/Engine/Plugins/Experimental/Avalanche/Source/AvalancheEditor/Private/AvaEditorStyle.cpp
+++ b/Engine/Plugins/Experimental/Avalanche/Source/AvalancheEditor/Private/AvaEditorStyle.cpp
@@ -154,33 +154,33 @@ FAvaEditorStyle::FAvaEditorStyle()
 	Set("AvalancheIcons.Lock3d", new IMAGE_BRUSH("Icons/DetailsPanelIcons/Lock3d", Icon16));
 	Set("AvalancheIcons.Unlock", new IMAGE_BRUSH("Icons/DetailsPanelIcons/Unlock", Icon16));
 
-	Set("AvalancheEditor.Thumbnail.Invalid", new FSlateRoundedBoxBrush(FStyleColors::Transparent, 4.0f, FLinearColor(1.0f, 0.2f, 0.2f, 1.0f), 1.0f));
-
-	Set("AvalancheEditor.StaticMeshToolsCategory", new CORE_IMAGE_BRUSH_SVG("Starship/AssetIcons/StaticMeshActor_16", Icon16));
-	Set("AvalancheEditor.CameraToolsCategory",     new CORE_IMAGE_BRUSH_SVG("Starship/AssetIcons/CameraActor_16", Icon16));
-	Set("AvalancheEditor.LightsToolsCategory",     new CORE_IMAGE_BRUSH_SVG("Starship/AssetIcons/PointLight_16", Icon16));
-
-	Set("AvalancheEditor.CubeTool",     new IMAGE_BRUSH("Icons/ToolboxIcons/cube",     Icon16));
-	Set("AvalancheEditor.SphereTool",   new IMAGE_BRUSH("Icons/ToolboxIcons/sphere",   Icon16));
-	Set("AvalancheEditor.CylinderTool", new IMAGE_BRUSH("Icons/ToolboxIcons/cylinder", Icon16));
-	Set("AvalancheEditor.ConeTool",     new IMAGE_BRUSH("Icons/ToolboxIcons/cone",     Icon20));
-	Set("AvalancheEditor.PlaneTool",    new IMAGE_BRUSH("Icons/ToolboxIcons/plane",    Icon20));
-
-	Set("AvalancheEditor.CameraTool",               new CORE_IMAGE_BRUSH_SVG("Starship/AssetIcons/CameraActor_16",       Icon16));
-	Set("AvalancheEditor.CineCameraTool",           new CORE_IMAGE_BRUSH_SVG("Starship/AssetIcons/CineCameraActor_16",   Icon16));
-	Set("AvalancheEditor.CameraRigCraneTool",       new CORE_IMAGE_BRUSH_SVG("Starship/AssetIcons/CameraRig_Crane_16",   Icon16));
-	Set("AvalancheEditor.CameraRigRailTool",        new CORE_IMAGE_BRUSH_SVG("Starship/AssetIcons/CameraRig_Rail_16",    Icon16));
-	Set("AvalancheEditor.CameraShakeSourceTool",    new CORE_IMAGE_BRUSH_SVG("Starship/AssetIcons/CameraShake",          Icon16));
-	Set("AvalancheEditor.AvaPostProcessVolumeTool", new CORE_IMAGE_BRUSH_SVG("Starship/AssetIcons/PostProcessVolume_16", Icon16));
-
-	Set("AvalancheEditor.PointLightTool",       new CORE_IMAGE_BRUSH_SVG("Starship/AssetIcons/PointLight_16",       Icon16));
-	Set("AvalancheEditor.DirectionalLightTool", new CORE_IMAGE_BRUSH_SVG("Starship/AssetIcons/DirectionalLight_16", Icon16));
-	Set("AvalancheEditor.RectLightTool",        new CORE_IMAGE_BRUSH_SVG("Starship/AssetIcons/RectLight_16",        Icon16));
-	Set("AvalancheEditor.SpotLightTool",        new CORE_IMAGE_BRUSH_SVG("Starship/AssetIcons/SpotLight_16",        Icon16));
-	Set("AvalancheEditor.SkyLightTool",         new CORE_IMAGE_BRUSH_SVG("Starship/AssetIcons/SkyLight_16",         Icon16));
+	Set("AvaEditor.Thumbnail.Invalid", new FSlateRoundedBoxBrush(FStyleColors::Transparent, 4.0f, FLinearColor(1.0f, 0.2f, 0.2f, 1.0f), 1.0f));
+
+	Set("AvaEditor.StaticMeshToolsCategory", new CORE_IMAGE_BRUSH_SVG("Starship/AssetIcons/StaticMeshActor_16", Icon16));
+	Set("AvaEditor.CameraToolsCategory",     new CORE_IMAGE_BRUSH_SVG("Starship/AssetIcons/CameraActor_16", Icon16));
+	Set("AvaEditor.LightsToolsCategory",     new CORE_IMAGE_BRUSH_SVG("Starship/AssetIcons/PointLight_16", Icon16));
+
+	Set("AvaEditor.CubeTool",     new IMAGE_BRUSH("Icons/ToolboxIcons/cube",     Icon16));
+	Set("AvaEditor.SphereTool",   new IMAGE_BRUSH("Icons/ToolboxIcons/sphere",   Icon16));
+	Set("AvaEditor.CylinderTool", new IMAGE_BRUSH("Icons/ToolboxIcons/cylinder", Icon16));
+	Set("AvaEditor.ConeTool",     new IMAGE_BRUSH("Icons/ToolboxIcons/cone",     Icon20));
+	Set("AvaEditor.PlaneTool",    new IMAGE_BRUSH("Icons/ToolboxIcons/plane",    Icon20));
+
+	Set("AvaEditor.CameraTool",               new CORE_IMAGE_BRUSH_SVG("Starship/AssetIcons/CameraActor_16",       Icon16));
+	Set("AvaEditor.CineCameraTool",           new CORE_IMAGE_BRUSH_SVG("Starship/AssetIcons/CineCameraActor_16",   Icon16));
+	Set("AvaEditor.CameraRigCraneTool",       new CORE_IMAGE_BRUSH_SVG("Starship/AssetIcons/CameraRig_Crane_16",   Icon16));
+	Set("AvaEditor.CameraRigRailTool",        new CORE_IMAGE_BRUSH_SVG("Starship/AssetIcons/CameraRig_Rail_16",    Icon16));
+	Set("AvaEditor.CameraShakeSourceTool",    new CORE_IMAGE_BRUSH_SVG("Starship/AssetIcons/CameraShake",          Icon16));
+	Set("AvaEditor.AvaPostProcessVolumeTool", new CORE_IMAGE_BRUSH_SVG("Starship/AssetIcons/PostProcessVolume_16", Icon16));
+
+	Set("AvaEditor.PointLightTool",       new CORE_IMAGE_BRUSH_SVG("Starship/AssetIcons/PointLight_16",       Icon16));
+	Set("AvaEditor.DirectionalLightTool", new CORE_IMAGE_BRUSH_SVG("Starship/AssetIcons/DirectionalLight_16", Icon16));
+	Set("AvaEditor.RectLightTool",        new CORE_IMAGE_BRUSH_SVG("Starship/AssetIcons/RectLight_16",        Icon16));
+	Set("AvaEditor.SpotLightTool",        new CORE_IMAGE_BRUSH_SVG("Starship/AssetIcons/SpotLight_16",        Icon16));
+	Set("AvaEditor.SkyLightTool",         new CORE_IMAGE_BRUSH_SVG("Starship/AssetIcons/SkyLight_16",         Icon16));
 	
 	// Colors
-	Set("AvalancheEditor.PalettesTab.ExpanderHeader", new FSlateColorBrush(FStyleColors::Header));
+	Set("AvaEditor.PalettesTab.ExpanderHeader", new FSlateColorBrush(FStyleColors::Header));
 
 	// Buttons
 	const FTextBlockStyle& AppStyle_ContentBrowserTopBarFont = FAppStyle::Get().GetWidgetStyle<FTextBlockStyle>("ContentBrowser.TopBar.Font");
@@ -188,42 +188,41 @@ FAvaEditorStyle::FAvaEditorStyle()
 	ButtonTextColor.A /= 2;
 	FLinearColor ButtonShadowColorAndOpacity = AppStyle_ContentBrowserTopBarFont.ShadowColorAndOpacity;
 	ButtonShadowColorAndOpacity.A /= 2;
-	Set("AvalancheEditor.Button.TextStyle", FTextBlockStyle(AppStyle_ContentBrowserTopBarFont)
+	Set("AvaEditor.Button.TextStyle", FTextBlockStyle(AppStyle_ContentBrowserTopBarFont)
 		.SetColorAndOpacity(ButtonTextColor)
 		.SetShadowColorAndOpacity(ButtonShadowColorAndOpacity));
 
 	const FButtonStyle& AppStyle_SimpleButton = FAppStyle::Get().GetWidgetStyle<FButtonStyle>("SimpleButton");
 
-	Set("AvalancheEditor.BorderlessButton", FButtonStyle(AppStyle_SimpleButton)
+	Set("AvaEditor.BorderlessButton", FButtonStyle(AppStyle_SimpleButton)
 		.SetNormalPadding(0.0f)
 		.SetPressedPadding(0.0f));
 
-	Set("AvalancheEditor.HighlightButton", FButtonStyle(AppStyle_SimpleButton)
+	Set("AvaEditor.HighlightButton", FButtonStyle(AppStyle_SimpleButton)
 		.SetNormal(FSlateColorBrush(FStyleColors::Secondary))
 		.SetHovered(FSlateColorBrush(FStyleColors::Hover))
 		.SetPressed(FSlateColorBrush(FStyleColors::Header))
 		.SetDisabled(FSlateColorBrush(FStyleColors::Dropdown)));
 
-
-	Set("AvalancheEditor.SuperBarButton", FButtonStyle(AppStyle_SimpleButton)
+	Set("AvaEditor.SuperBarButton", FButtonStyle(AppStyle_SimpleButton)
 		.SetNormal(FSlateNoResource())
 		.SetHovered(FSlateColorBrush(FStyleColors::Secondary))
 		.SetPressed(FSlateColorBrush(FStyleColors::Hover))
 		.SetDisabled(FSlateNoResource()));
 
-	Set("AvalancheEditor.SuperBarButton.Rounded", FButtonStyle(AppStyle_SimpleButton)
+	Set("AvaEditor.SuperBarButton.Rounded", FButtonStyle(AppStyle_SimpleButton)
 		.SetNormal(FSlateNoResource())
 		.SetHovered(FSlateRoundedBoxBrush(FStyleColors::Secondary, 4.0f))
 		.SetPressed(FSlateRoundedBoxBrush(FStyleColors::Hover, 4.0f))
 		.SetDisabled(FSlateNoResource()));
 
-	Set("AvalancheEditor.SuperBarSubMenuButton", FButtonStyle(AppStyle_SimpleButton)
+	Set("AvaEditor.SuperBarSubMenuButton", FButtonStyle(AppStyle_SimpleButton)
 		.SetNormal(FSlateNoResource())
 		.SetHovered(FSlateNoResource())
 		.SetPressed(FSlateNoResource())
 		.SetDisabled(FSlateNoResource()));
 
-	Set("AvalancheEditor.DarkButton", FButtonStyle(AppStyle_SimpleButton)
+	Set("AvaEditor.DarkButton", FButtonStyle(AppStyle_SimpleButton)
 		.SetNormal(FSlateColorBrush(FStyleColors::Recessed))
 		.SetHovered(FSlateColorBrush(FStyleColors::Hover))
 		.SetPressed(FSlateColorBrush(FStyleColors::Header))
@@ -237,14 +236,14 @@ FAvaEditorStyle::FAvaEditorStyle()
 
 	// Text
 	const FTextBlockStyle& AppStyle_GraphCompactNodeTitle = FAppStyle::Get().GetWidgetStyle<FTextBlockStyle>("Graph.CompactNode.Title");
-	Set("AvalancheEditor.FilterRow.And", FTextBlockStyle(AppStyle_GraphCompactNodeTitle)
+	Set("AvaEditor.FilterRow.And", FTextBlockStyle(AppStyle_GraphCompactNodeTitle)
 		.SetFont(DEFAULT_FONT("BoldCondensed", 16)));
-	Set("AvalancheEditor.FilterRow.Or", FTextBlockStyle(AppStyle_GraphCompactNodeTitle)
+	Set("AvaEditor.FilterRow.Or", FTextBlockStyle(AppStyle_GraphCompactNodeTitle)
 		.SetFont(DEFAULT_FONT("BoldCondensed", 18)));
 
 	// Check Boxes
 	const FCheckBoxStyle& AppStyle_RadioButton = FAppStyle::GetWidgetStyle<FCheckBoxStyle>("RadioButton");
-	Set("AvalancheEditor.BlackRadioButton", FCheckBoxStyle(AppStyle_RadioButton)
+	Set("AvaEditor.BlackRadioButton", FCheckBoxStyle(AppStyle_RadioButton)
 		.SetBackgroundImage(*GetBrush("AvalancheIcons.Editor.Radio.BlackBackground")));
 
 	FSlateStyleRegistry::RegisterSlateStyle(*this);
diff --git a/Engine/Plugins/Experimental/Avalanche/Source/AvalancheEditorCore/Private/LevelEditor/AvaLevelEditorCommands.cpp b/Engine/Plugins/Experimental/Avalanche/Source/AvalancheEditorCore/Private/LevelEditor/AvaLevelEditorCommands.cpp
--- a/Engine/Plugins/Experimental/Avalanche/Source/AvalancheEditorCore/Private/LevelEditor/AvaLevelEditorCommands.cpp
+++ b/Engine/Plugins/Experimental/Avalanche/Source/AvalancheEditorCore/Private/LevelEditor/AvaLevelEditorCommands.cpp
@@ -7,7 +7,7 @@
 
 FAvaLevelEditorCommands::FAvaLevelEditorCommands()
 	: TCommands<FAvaLevelEditorCommands>(TEXT("AvaLevelEditor")
-	, LOCTEXT("AvaLevelEditor", "Motion Design Level Editor")
+	, LOCTEXT("MotionDesignLevelEditor", "Motion Design Level Editor")
 	, NAME_None
 	, FAvaLevelEditorStyle::Get().GetStyleSetName())
 {
diff --git a/Engine/Plugins/Experimental/Avalanche/Source/AvalancheEffectorsEditor/Private/AvaEffectorsEditorCommands.cpp b/Engine/Plugins/Experimental/Avalanche/Source/AvalancheEffectorsEditor/Private/AvaEffectorsEditorCommands.cpp
--- a/Engine/Plugins/Experimental/Avalanche/Source/AvalancheEffectorsEditor/Private/AvaEffectorsEditorCommands.cpp
+++ b/Engine/Plugins/Experimental/Avalanche/Source/AvalancheEffectorsEditor/Private/AvaEffectorsEditorCommands.cpp
@@ -7,7 +7,7 @@
 
 FAvaEffectorsEditorCommands::FAvaEffectorsEditorCommands()
 	: TCommands<FAvaEffectorsEditorCommands>(
-		TEXT("AvalancheEffectorsEditor")
+		TEXT("AvaEffectorsEditor")
 		, LOCTEXT("MotionDesignEffectorsEditor", "Motion Design Effects Editor")
 		, NAME_None
 		, FAvaEffectorsEditorStyle::Get().GetStyleSetName()
diff --git a/Engine/Plugins/Experimental/Avalanche/Source/AvalancheEffectorsEditor/Private/AvaEffectorsEditorStyle.cpp b/Engine/Plugins/Experimental/Avalanche/Source/AvalancheEffectorsEditor/Private/AvaEffectorsEditorStyle.cpp
--- a/Engine/Plugins/Experimental/Avalanche/Source/AvalancheEffectorsEditor/Private/AvaEffectorsEditorStyle.cpp
+++ b/Engine/Plugins/Experimental/Avalanche/Source/AvalancheEffectorsEditor/Private/AvaEffectorsEditorStyle.cpp
@@ -20,8 +20,8 @@ FAvaEffectorsEditorStyle::FAvaEffectorsEditorStyle()
 	
 	SetContentRoot(Plugin->GetBaseDir() / TEXT("Resources"));
 
-	Set("AvalancheEffectorsEditor.Tool_Actor_Effector", new IMAGE_BRUSH_SVG("Icons/ToolboxIcons/effector", Icon20x20));
-	Set("AvalancheEffectorsEditor.Tool_Actor_Cloner", new IMAGE_BRUSH_SVG("Icons/ToolboxIcons/cloner", Icon20x20));
+	Set("AvaEffectorsEditor.Tool_Actor_Effector", new IMAGE_BRUSH_SVG("Icons/ToolboxIcons/effector", Icon20x20));
+	Set("AvaEffectorsEditor.Tool_Actor_Cloner", new IMAGE_BRUSH_SVG("Icons/ToolboxIcons/cloner", Icon20x20));
 	Set("Tool_Actor_Effector", new IMAGE_BRUSH_SVG("Icons/ToolboxIcons/effector", Icon20x20));
 	Set("Tool_Actor_Cloner", new IMAGE_BRUSH_SVG("Icons/ToolboxIcons/cloner", Icon20x20));
 	
diff --git a/Engine/Plugins/Experimental/Avalanche/Source/AvalancheInteractiveTools/Private/AvaInteractiveToolsCommands.cpp b/Engine/Plugins/Experimental/Avalanche/Source/AvalancheInteractiveTools/Private/AvaInteractiveToolsCommands.cpp
--- a/Engine/Plugins/Experimental/Avalanche/Source/AvalancheInteractiveTools/Private/AvaInteractiveToolsCommands.cpp
+++ b/Engine/Plugins/Experimental/Avalanche/Source/AvalancheInteractiveTools/Private/AvaInteractiveToolsCommands.cpp
@@ -7,7 +7,7 @@
 
 FAvaInteractiveToolsCommands::FAvaInteractiveToolsCommands()
 	: TCommands<FAvaInteractiveToolsCommands>(
-		TEXT("AvalancheInteractiveTools")
+		TEXT("AvaInteractiveTools")
 		, LOCTEXT("MotionDesignInteractiveTools", "Motion Design Interactive Tools")
 		, NAME_None
 		, FAvaInteractiveToolsStyle::Get().GetStyleSetName()
diff --git a/Engine/Plugins/Experimental/Avalanche/Source/AvalancheInteractiveTools/Private/AvaInteractiveToolsStyle.cpp b/Engine/Plugins/Experimental/Avalanche/Source/AvalancheInteractiveTools/Private/AvaInteractiveToolsStyle.cpp
--- a/Engine/Plugins/Experimental/Avalanche/Source/AvalancheInteractiveTools/Private/AvaInteractiveToolsStyle.cpp
+++ b/Engine/Plugins/Experimental/Avalanche/Source/AvalancheInteractiveTools/Private/AvaInteractiveToolsStyle.cpp
@@ -22,14 +22,14 @@ FAvaInteractiveToolsStyle::FAvaInteractiveToolsStyle()
 	SetCoreContentRoot(FPaths::EngineContentDir() / TEXT("Editor/Slate"));
 
 	// Categories
-	Set("AvalancheInteractiveTools.Category_2D",     new IMAGE_BRUSH("Icons/ToolboxIcons/rectangle", Icon20x20));
-	Set("AvalancheInteractiveTools.Category_3D",     new IMAGE_BRUSH("Icons/ToolboxIcons/cube", Icon20x20));
-	Set("AvalancheInteractiveTools.Category_Actor",  new CORE_IMAGE_BRUSH_SVG("Starship/AssetIcons/Actor_16", Icon16x16));
-	Set("AvalancheInteractiveTools.Category_Layout", new IMAGE_BRUSH("Icons/ToolboxIcons/layoutgrid", Icon20x20));
+	Set("AvaInteractiveTools.Category_2D",     new IMAGE_BRUSH("Icons/ToolboxIcons/rectangle", Icon20x20));
+	Set("AvaInteractiveTools.Category_3D",     new IMAGE_BRUSH("Icons/ToolboxIcons/cube", Icon20x20));
+	Set("AvaInteractiveTools.Category_Actor",  new CORE_IMAGE_BRUSH_SVG("Starship/AssetIcons/Actor_16", Icon16x16));
+	Set("AvaInteractiveTools.Category_Layout", new IMAGE_BRUSH("Icons/ToolboxIcons/layoutgrid", Icon20x20));
 
 	// Actor Tools
-	Set("AvalancheInteractiveTools.Tool_Actor_Null", new CORE_IMAGE_BRUSH(TEXT("Icons/SequencerIcons/icon_Sequencer_Move_24x"), Icon16x16));
-	Set("AvalancheInteractiveTools.Tool_Actor_Spline", new IMAGE_BRUSH_SVG("Icons/ToolboxIcons/Toolbar_Spline", Icon20x20));
+	Set("AvaInteractiveTools.Tool_Actor_Null", new CORE_IMAGE_BRUSH(TEXT("Icons/SequencerIcons/icon_Sequencer_Move_24x"), Icon16x16));
+	Set("AvaInteractiveTools.Tool_Actor_Spline", new IMAGE_BRUSH_SVG("Icons/ToolboxIcons/Toolbar_Spline", Icon20x20));
 	Set("Tool_Actor_Null", new CORE_IMAGE_BRUSH(TEXT("Icons/SequencerIcons/icon_Sequencer_Move_24x"), Icon20x20));
 	Set("Tool_Actor_Spline", new IMAGE_BRUSH_SVG("Icons/ToolboxIcons/Toolbar_Spline", Icon20x20));
 
diff --git a/Engine/Plugins/Experimental/Avalanche/Source/AvalancheLevelViewport/Private/AvaLevelViewportCommands.cpp b/Engine/Plugins/Experimental/Avalanche/Source/AvalancheLevelViewport/Private/AvaLevelViewportCommands.cpp
--- a/Engine/Plugins/Experimental/Avalanche/Source/AvalancheLevelViewport/Private/AvaLevelViewportCommands.cpp
+++ b/Engine/Plugins/Experimental/Avalanche/Source/AvalancheLevelViewport/Private/AvaLevelViewportCommands.cpp
@@ -7,7 +7,7 @@
 
 FAvaLevelViewportCommands::FAvaLevelViewportCommands()
 	: TCommands<FAvaLevelViewportCommands>(
-		TEXT("AvalancheLevelViewport")
+		TEXT("AvaLevelViewport")
 		, LOCTEXT("MotionDesignLevelViewport", "Motion Design Level Viewport")
 		, NAME_None
 		, FAvaLevelViewportStyle::Get().GetStyleSetName()
diff --git a/Engine/Plugins/Experimental/Avalanche/Source/AvalancheMRQEditor/Private/AvaMRQEditorCommands.cpp b/Engine/Plugins/Experimental/Avalanche/Source/AvalancheMRQEditor/Private/AvaMRQEditorCommands.cpp
--- a/Engine/Plugins/Experimental/Avalanche/Source/AvalancheMRQEditor/Private/AvaMRQEditorCommands.cpp
+++ b/Engine/Plugins/Experimental/Avalanche/Source/AvalancheMRQEditor/Private/AvaMRQEditorCommands.cpp
@@ -6,8 +6,8 @@
 #define LOCTEXT_NAMESPACE "AvaMRQEditorCommands"
 
 FAvaMRQEditorCommands::FAvaMRQEditorCommands()
-	: TCommands<FAvaMRQEditorCommands>(TEXT("AvaMRQEditorCommands")
-	, LOCTEXT("AvaMRQEditorCommands", "Motion Design MRQ Editor Commands")
+	: TCommands<FAvaMRQEditorCommands>(TEXT("AvaMRQEditor")
+	, LOCTEXT("MotionDesignMRQEditorCommands", "Motion Design MRQ Editor Commands")
 	, NAME_None
 	, FAvaMRQEditorStyle::Get().GetStyleSetName())
 {
diff --git a/Engine/Plugins/Experimental/Avalanche/Source/AvalancheMRQEditor/Private/AvaMRQEditorStyle.cpp b/Engine/Plugins/Experimental/Avalanche/Source/AvalancheMRQEditor/Private/AvaMRQEditorStyle.cpp
--- a/Engine/Plugins/Experimental/Avalanche/Source/AvalancheMRQEditor/Private/AvaMRQEditorStyle.cpp
+++ b/Engine/Plugins/Experimental/Avalanche/Source/AvalancheMRQEditor/Private/AvaMRQEditorStyle.cpp
@@ -14,7 +14,7 @@ FAvaMRQEditorStyle::FAvaMRQEditorStyle()
 	SetContentRoot(FPaths::EngineContentDir() / TEXT("Editor/Slate"));
 	SetCoreContentRoot(FPaths::EngineContentDir() / TEXT("Slate"));
 
-	Set("AvaMRQEditorCommands.RenderSelectedPages", new IMAGE_BRUSH_SVG("Starship/MainToolbar/cinematics", Icon20));
+	Set("AvaMRQEditor.RenderSelectedPages", new IMAGE_BRUSH_SVG("Starship/MainToolbar/cinematics", Icon20));
 
 	FSlateStyleRegistry::RegisterSlateStyle(*this);
 }
diff --git a/Engine/Plugins/Experimental/Avalanche/Source/AvalancheMaskEditor/Private/AvaMaskEditorCommands.h b/Engine/Plugins/Experimental/Avalanche/Source/AvalancheMaskEditor/Private/AvaMaskEditorCommands.h
--- a/Engine/Plugins/Experimental/Avalanche/Source/AvalancheMaskEditor/Private/AvaMaskEditorCommands.h
+++ b/Engine/Plugins/Experimental/Avalanche/Source/AvalancheMaskEditor/Private/AvaMaskEditorCommands.h
@@ -12,7 +12,7 @@ class FAvaMaskEditorCommands : public TCommands<FAvaMaskEditorCommands>
 public:
 	FAvaMaskEditorCommands()
 		: TCommands<FAvaMaskEditorCommands>(
-			TEXT("AvalancheMaskEditor"),
+			TEXT("AvaMaskEditor"),
 			LOCTEXT("MotionDesignMaskEditor", "Motion Design Masking"),
 			NAME_None,
 			FAvaMaskEditorStyle::Get().GetStyleSetName())
diff --git a/Engine/Plugins/Experimental/Avalanche/Source/AvalancheMaskEditor/Private/AvaMaskEditorStyle.cpp b/Engine/Plugins/Experimental/Avalanche/Source/AvalancheMaskEditor/Private/AvaMaskEditorStyle.cpp
--- a/Engine/Plugins/Experimental/Avalanche/Source/AvalancheMaskEditor/Private/AvaMaskEditorStyle.cpp
+++ b/Engine/Plugins/Experimental/Avalanche/Source/AvalancheMaskEditor/Private/AvaMaskEditorStyle.cpp
@@ -21,11 +21,11 @@ FAvaMaskEditorStyle::FAvaMaskEditorStyle()
 
 	SetContentRoot(Plugin->GetBaseDir() / TEXT("Resources"));
 
-	Set("AvalancheMaskEditor.ToggleMaskMode.Small", new IMAGE_BRUSH_SVG(TEXT("Icons/MaskIcons/Mode_On"), Icon16x16));
-	Set("AvalancheMaskEditor.ToggleMaskMode", new IMAGE_BRUSH_SVG(TEXT("Icons/MaskIcons/Mode_On"), Icon20x20));
-	Set("AvalancheMaskEditor.ToggleShowAllMasks", new IMAGE_BRUSH_SVG(TEXT("Icons/MaskIcons/Disable"), Icon20x20));
-	Set("AvalancheMaskEditor.ToggleDisableMask", new IMAGE_BRUSH_SVG(TEXT("Icons/MaskIcons/Disable"), Icon20x20));
-	Set("AvalancheMaskEditor.ToggleIsolateMask", new IMAGE_BRUSH_SVG(TEXT("Icons/MaskIcons/Disable"), Icon20x20));
+	Set("AvaMaskEditor.ToggleMaskMode.Small", new IMAGE_BRUSH_SVG(TEXT("Icons/MaskIcons/Mode_On"), Icon16x16));
+	Set("AvaMaskEditor.ToggleMaskMode", new IMAGE_BRUSH_SVG(TEXT("Icons/MaskIcons/Mode_On"), Icon20x20));
+	Set("AvaMaskEditor.ToggleShowAllMasks", new IMAGE_BRUSH_SVG(TEXT("Icons/MaskIcons/Disable"), Icon20x20));
+	Set("AvaMaskEditor.ToggleDisableMask", new IMAGE_BRUSH_SVG(TEXT("Icons/MaskIcons/Disable"), Icon20x20));
+	Set("AvaMaskEditor.ToggleIsolateMask", new IMAGE_BRUSH_SVG(TEXT("Icons/MaskIcons/Disable"), Icon20x20));
 
 	{
 		FToolBarStyle ViewportOverlayToolbarStyle = FAppStyle::Get().GetWidgetStyle<FToolBarStyle>(TEXT("AssetEditorToolbar"));
@@ -42,7 +42,7 @@ FAvaMaskEditorStyle::FAvaMaskEditorStyle()
 		ViewportOverlayToolbarStyle.SetButtonPadding(0);
 		ViewportOverlayToolbarStyle.SetCheckBoxPadding(0);
 		
-		Set("AvalancheMaskEditor.ViewportOverlayToolbar", ViewportOverlayToolbarStyle);
+		Set("AvaMaskEditor.ViewportOverlayToolbar", ViewportOverlayToolbarStyle);
 	}
 	
 	FSlateStyleRegistry::RegisterSlateStyle(*this);
diff --git a/Engine/Plugins/Experimental/Avalanche/Source/AvalancheMedia/Private/AvaMediaStyle.cpp b/Engine/Plugins/Experimental/Avalanche/Source/AvalancheMedia/Private/AvaMediaStyle.cpp
--- a/Engine/Plugins/Experimental/Avalanche/Source/AvalancheMedia/Private/AvaMediaStyle.cpp
+++ b/Engine/Plugins/Experimental/Avalanche/Source/AvalancheMedia/Private/AvaMediaStyle.cpp
@@ -17,7 +17,7 @@ FAvaMediaStyle::FAvaMediaStyle()
 
 	SetContentRoot(Plugin->GetBaseDir() / TEXT("Resources"));
 	
-	Set("AvalancheMedia.UnrealIcon", new IMAGE_BRUSH("Icons/MediaIcons/ue_logo", Icon200x200));
+	Set("AvaMedia.UnrealIcon", new IMAGE_BRUSH("Icons/MediaIcons/ue_logo", Icon200x200));
 	Set("ClassIcon.AvalancheBroadcast", new IMAGE_BRUSH("Icons/MediaIcons/MediaOutput", Icon16x16));
 
 	FSlateStyleRegistry::RegisterSlateStyle(*this);
diff --git a/Engine/Plugins/Experimental/Avalanche/Source/AvalancheMediaEditor/Private/AvaMediaEditorStyle.cpp b/Engine/Plugins/Experimental/Avalanche/Source/AvalancheMediaEditor/Private/AvaMediaEditorStyle.cpp
--- a/Engine/Plugins/Experimental/Avalanche/Source/AvalancheMediaEditor/Private/AvaMediaEditorStyle.cpp
+++ b/Engine/Plugins/Experimental/Avalanche/Source/AvalancheMediaEditor/Private/AvaMediaEditorStyle.cpp
@@ -36,27 +36,27 @@ FAvaMediaEditorStyle::FAvaMediaEditorStyle()
 	Set("ClassThumbnail.AvaRundownMacroCollection", new IMAGE_BRUSH("Icons/MediaIcons/AvaRundownMacroCollectionIcon_64x", Icon64x64));
 
 	//Rundown Commands
-	Set("AvaRundownCommands.AddPage"         , new IMAGE_BRUSH_SVG("Icons/MediaIcons/AddPage"      , Icon16x16));
-	Set("AvaRundownCommands.AddTemplate"     , new IMAGE_BRUSH_SVG("Icons/MediaIcons/AddPage"      , Icon16x16));
-	Set("AvaRundownCommands.CreatePageInstanceFromTemplate"  , new IMAGE_BRUSH_SVG("Icons/MediaIcons/AddPage"      , Icon16x16));
-	Set("AvaRundownCommands.CreateComboTemplate", new IMAGE_BRUSH_SVG("Icons/MediaIcons/AddPage"   , Icon16x16));
-	Set("AvaRundownCommands.RemovePage"      , new IMAGE_BRUSH_SVG("Icons/MediaIcons/RemovePage"   , Icon16x16));
-	Set("AvaRundownCommands.RenumberPage"    , new IMAGE_BRUSH_SVG("Icons/MediaIcons/RenumberPage" , Icon16x16));
-	Set("AvaRundownCommands.ReimportPage"    , new IMAGE_BRUSH_SVG("Icons/MediaIcons/ReimportPage" , Icon16x16));
-	Set("AvaRundownCommands.EditPageSource"  , new IMAGE_BRUSH_SVG("Icons/MediaIcons/OpenAsset"	  , Icon16x16));
-	Set("AvaRundownCommands.Play"            , new IMAGE_BRUSH_SVG("Icons/MediaIcons/Play"         , Icon16x16));
-	Set("AvaRundownCommands.UpdateValues"    , new IMAGE_BRUSH_SVG("Icons/MediaIcons/UpdateValues" , Icon16x16));
-	Set("AvaRundownCommands.Stop"            , new IMAGE_BRUSH_SVG("Icons/MediaIcons/Stop"         , Icon16x16));
-	Set("AvaRundownCommands.ForceStop"       , new IMAGE_BRUSH_SVG("Icons/MediaIcons/Stop"         , Icon16x16));
-	Set("AvaRundownCommands.Continue"        , new IMAGE_BRUSH_SVG("Icons/MediaIcons/Continue"     , Icon16x16));
-	Set("AvaRundownCommands.PlayNext"        , new IMAGE_BRUSH_SVG("Icons/MediaIcons/PlayNext"     , Icon16x16));
-	Set("AvaRundownCommands.PreviewFrame"    , new IMAGE_BRUSH_SVG("Icons/MediaIcons/PreviewFrame" , Icon16x16));
-	Set("AvaRundownCommands.PreviewPlay"     , new IMAGE_BRUSH_SVG("Icons/MediaIcons/Play"         , Icon16x16));
-	Set("AvaRundownCommands.PreviewStop"     , new IMAGE_BRUSH_SVG("Icons/MediaIcons/Stop"         , Icon16x16));
-	Set("AvaRundownCommands.PreviewForceStop", new IMAGE_BRUSH_SVG("Icons/MediaIcons/Stop"         , Icon16x16));
-	Set("AvaRundownCommands.PreviewContinue" , new IMAGE_BRUSH_SVG("Icons/MediaIcons/Continue"     , Icon16x16));
-	Set("AvaRundownCommands.PreviewPlayNext" , new IMAGE_BRUSH_SVG("Icons/MediaIcons/PlayNext"     , Icon16x16));
-	Set("AvaRundownCommands.TakeToProgram"   , new IMAGE_BRUSH_SVG("Icons/MediaIcons/Play"		   , Icon16x16));
+	Set("AvaRundown.AddPage"         , new IMAGE_BRUSH_SVG("Icons/MediaIcons/AddPage"      , Icon16x16));
+	Set("AvaRundown.AddTemplate"     , new IMAGE_BRUSH_SVG("Icons/MediaIcons/AddPage"      , Icon16x16));
+	Set("AvaRundown.CreatePageInstanceFromTemplate"  , new IMAGE_BRUSH_SVG("Icons/MediaIcons/AddPage"      , Icon16x16));
+	Set("AvaRundown.CreateComboTemplate", new IMAGE_BRUSH_SVG("Icons/MediaIcons/AddPage"   , Icon16x16));
+	Set("AvaRundown.RemovePage"      , new IMAGE_BRUSH_SVG("Icons/MediaIcons/RemovePage"   , Icon16x16));
+	Set("AvaRundown.RenumberPage"    , new IMAGE_BRUSH_SVG("Icons/MediaIcons/RenumberPage" , Icon16x16));
+	Set("AvaRundown.ReimportPage"    , new IMAGE_BRUSH_SVG("Icons/MediaIcons/ReimportPage" , Icon16x16));
+	Set("AvaRundown.EditPageSource"  , new IMAGE_BRUSH_SVG("Icons/MediaIcons/OpenAsset"	  , Icon16x16));
+	Set("AvaRundown.Play"            , new IMAGE_BRUSH_SVG("Icons/MediaIcons/Play"         , Icon16x16));
+	Set("AvaRundown.UpdateValues"    , new IMAGE_BRUSH_SVG("Icons/MediaIcons/UpdateValues" , Icon16x16));
+	Set("AvaRundown.Stop"            , new IMAGE_BRUSH_SVG("Icons/MediaIcons/Stop"         , Icon16x16));
+	Set("AvaRundown.ForceStop"       , new IMAGE_BRUSH_SVG("Icons/MediaIcons/Stop"         , Icon16x16));
+	Set("AvaRundown.Continue"        , new IMAGE_BRUSH_SVG("Icons/MediaIcons/Continue"     , Icon16x16));
+	Set("AvaRundown.PlayNext"        , new IMAGE_BRUSH_SVG("Icons/MediaIcons/PlayNext"     , Icon16x16));
+	Set("AvaRundown.PreviewFrame"    , new IMAGE_BRUSH_SVG("Icons/MediaIcons/PreviewFrame" , Icon16x16));
+	Set("AvaRundown.PreviewPlay"     , new IMAGE_BRUSH_SVG("Icons/MediaIcons/Play"         , Icon16x16));
+	Set("AvaRundown.PreviewStop"     , new IMAGE_BRUSH_SVG("Icons/MediaIcons/Stop"         , Icon16x16));
+	Set("AvaRundown.PreviewForceStop", new IMAGE_BRUSH_SVG("Icons/MediaIcons/Stop"         , Icon16x16));
+	Set("AvaRundown.PreviewContinue" , new IMAGE_BRUSH_SVG("Icons/MediaIcons/Continue"     , Icon16x16));
+	Set("AvaRundown.PreviewPlayNext" , new IMAGE_BRUSH_SVG("Icons/MediaIcons/PlayNext"     , Icon16x16));
+	Set("AvaRundown.TakeToProgram"   , new IMAGE_BRUSH_SVG("Icons/MediaIcons/Play"		   , Icon16x16));
 
 	//Broadcast
 	Set("AvaMediaEditor.OutputIcon"           , new IMAGE_BRUSH("Icons/MediaIcons/MediaOutput"     , Icon20x20));
diff --git a/Engine/Plugins/Experimental/Avalanche/Source/AvalancheMediaEditor/Private/Rundown/AvaRundownCommands.h b/Engine/Plugins/Experimental/Avalanche/Source/AvalancheMediaEditor/Private/Rundown/AvaRundownCommands.h
--- a/Engine/Plugins/Experimental/Avalanche/Source/AvalancheMediaEditor/Private/Rundown/AvaRundownCommands.h
+++ b/Engine/Plugins/Experimental/Avalanche/Source/AvalancheMediaEditor/Private/Rundown/AvaRundownCommands.h
@@ -13,8 +13,8 @@ class FAvaRundownCommands : public TCommands<FAvaRundownCommands>
 public:
 	
 	FAvaRundownCommands()
-		: TCommands<FAvaRundownCommands>(TEXT("AvaRundownCommands")
-		, LOCTEXT("AvaRundownCommands", "Motion Design Rundown")
+		: TCommands<FAvaRundownCommands>(TEXT("AvaRundown")
+		, LOCTEXT("MotionDesignRundownCommands", "Motion Design Rundown")
 		, NAME_None
 		, FAvaMediaEditorStyle::Get().GetStyleSetName())
 	{
diff --git a/Engine/Plugins/Experimental/Avalanche/Source/AvalancheMediaEditor/Public/Playback/AvaPlaybackCommands.h b/Engine/Plugins/Experimental/Avalanche/Source/AvalancheMediaEditor/Public/Playback/AvaPlaybackCommands.h
--- a/Engine/Plugins/Experimental/Avalanche/Source/AvalancheMediaEditor/Public/Playback/AvaPlaybackCommands.h
+++ b/Engine/Plugins/Experimental/Avalanche/Source/AvalancheMediaEditor/Public/Playback/AvaPlaybackCommands.h
@@ -10,7 +10,7 @@ class FAvaPlaybackCommands : public TCommands<FAvaPlaybackCommands>
 public:
 	
 	FAvaPlaybackCommands() : TCommands<FAvaPlaybackCommands>(TEXT("AvaPlayback")
-		, NSLOCTEXT("AvaPlaybackCommands", "AvaPlaybackCommands", "Motion Design Playback")
+		, NSLOCTEXT("MotionDesignPlaybackCommands", "MotionDesignPlaybackCommands", "Motion Design Playback")
 		, NAME_None
 		, FAppStyle::GetAppStyleSetName())
 	{
diff --git a/Engine/Plugins/Experimental/Avalanche/Source/AvalancheOutliner/Private/AvaOutlinerCommands.h b/Engine/Plugins/Experimental/Avalanche/Source/AvalancheOutliner/Private/AvaOutlinerCommands.h
--- a/Engine/Plugins/Experimental/Avalanche/Source/AvalancheOutliner/Private/AvaOutlinerCommands.h
+++ b/Engine/Plugins/Experimental/Avalanche/Source/AvalancheOutliner/Private/AvaOutlinerCommands.h
@@ -12,7 +12,7 @@ class FAvaOutlinerCommands
 {
 public:
 	FAvaOutlinerCommands()
-		: TCommands<FAvaOutlinerCommands>(TEXT("AvalancheOutliner")
+		: TCommands<FAvaOutlinerCommands>(TEXT("AvaOutliner")
 		, LOCTEXT("MotionDesignOutliner", "Motion Design Outliner")
 		, NAME_None
 		, FAppStyle::GetAppStyleSetName())
diff --git a/Engine/Plugins/Experimental/Avalanche/Source/AvalancheOutliner/Private/AvaOutlinerStyle.cpp b/Engine/Plugins/Experimental/Avalanche/Source/AvalancheOutliner/Private/AvaOutlinerStyle.cpp
--- a/Engine/Plugins/Experimental/Avalanche/Source/AvalancheOutliner/Private/AvaOutlinerStyle.cpp
+++ b/Engine/Plugins/Experimental/Avalanche/Source/AvalancheOutliner/Private/AvaOutlinerStyle.cpp
@@ -24,11 +24,11 @@ FAvaOutlinerStyle::FAvaOutlinerStyle()
 	SetContentRoot(FPaths::Combine(Plugin->GetBaseDir(), TEXT("Resources")));
 	SetCoreContentRoot(FPaths::EngineContentDir() / TEXT("Slate"));
 	
-	Set("AvalancheOutliner.FilterIcon", new IMAGE_BRUSH("Icons/OutlinerIcons/FilterIcon", Icon20x20));
+	Set("AvaOutliner.FilterIcon", new IMAGE_BRUSH("Icons/OutlinerIcons/FilterIcon", Icon20x20));
 	
 	// Table View Row Style
 	const FTableRowStyle& TableRowStyle = FAppStyle::Get().GetWidgetStyle<FTableRowStyle>("SceneOutliner.TableViewRow");
-	Set("AvalancheOutliner.TableViewRow", FTableRowStyle(TableRowStyle)
+	Set("AvaOutliner.TableViewRow", FTableRowStyle(TableRowStyle)
 		.SetEvenRowBackgroundBrush(FSlateNoResource())
 		.SetOddRowBackgroundBrush(FSlateColorBrush(FStyleColors::Background)));
 
diff --git a/Engine/Plugins/Experimental/Avalanche/Source/AvalancheOutliner/Private/Slate/SAvaOutlinerItemFilters.cpp b/Engine/Plugins/Experimental/Avalanche/Source/AvalancheOutliner/Private/Slate/SAvaOutlinerItemFilters.cpp
--- a/Engine/Plugins/Experimental/Avalanche/Source/AvalancheOutliner/Private/Slate/SAvaOutlinerItemFilters.cpp
+++ b/Engine/Plugins/Experimental/Avalanche/Source/AvalancheOutliner/Private/Slate/SAvaOutlinerItemFilters.cpp
@@ -72,7 +72,7 @@ void SAvaOutlinerItemFilters::Construct(const FArguments& InArgs, const TSharedR
 			.OnClicked(this, &SAvaOutlinerItemFilters::ToggleShowItemFilters)
 			[
 				SNew(SImage)
-				.Image(FAvaOutlinerStyle::Get().GetBrush(TEXT("AvalancheOutliner.FilterIcon")))
+				.Image(FAvaOutlinerStyle::Get().GetBrush(TEXT("AvaOutliner.FilterIcon")))
 				.DesiredSizeOverride(FVector2D(24.f))
 			]
 		];
diff --git a/Engine/Plugins/Experimental/Avalanche/Source/AvalancheOutliner/Private/Slate/SAvaOutlinerTreeRow.cpp b/Engine/Plugins/Experimental/Avalanche/Source/AvalancheOutliner/Private/Slate/SAvaOutlinerTreeRow.cpp
--- a/Engine/Plugins/Experimental/Avalanche/Source/AvalancheOutliner/Private/Slate/SAvaOutlinerTreeRow.cpp
+++ b/Engine/Plugins/Experimental/Avalanche/Source/AvalancheOutliner/Private/Slate/SAvaOutlinerTreeRow.cpp
@@ -23,7 +23,7 @@ void SAvaOutlinerTreeRow::Construct(const FArguments& InArgs
 	SetColorAndOpacity(TAttribute<FLinearColor>::CreateSP(&*InOutlinerView, &FAvaOutlinerView::GetItemBrushColor, Item));
 
 	SMultiColumnTableRow::Construct(FSuperRowType::FArguments()
-			.Style(&FAvaOutlinerStyle::Get().GetWidgetStyle<FTableRowStyle>("AvalancheOutliner.TableViewRow"))
+			.Style(&FAvaOutlinerStyle::Get().GetWidgetStyle<FTableRowStyle>("AvaOutliner.TableViewRow"))
 			.OnCanAcceptDrop(InOutlinerView, &FAvaOutlinerView::OnCanDrop)
 			.OnDragDetected(InOutlinerView, &FAvaOutlinerView::OnDragDetected, Item)
 			.OnDragEnter(InOutlinerView, &FAvaOutlinerView::OnDragEnter, Item)
diff --git a/Engine/Plugins/Experimental/Avalanche/Source/AvalancheSequencer/Private/Commands/AvaSequencerCommands.h b/Engine/Plugins/Experimental/Avalanche/Source/AvalancheSequencer/Private/Commands/AvaSequencerCommands.h
--- a/Engine/Plugins/Experimental/Avalanche/Source/AvalancheSequencer/Private/Commands/AvaSequencerCommands.h
+++ b/Engine/Plugins/Experimental/Avalanche/Source/AvalancheSequencer/Private/Commands/AvaSequencerCommands.h
@@ -9,8 +9,8 @@ class FAvaSequencerCommands : public TCommands<FAvaSequencerCommands>
 {
 public:
 	FAvaSequencerCommands()
-		: TCommands<FAvaSequencerCommands>(TEXT("AvalancheSequencerCommands")
-		, NSLOCTEXT("AvaSequencerCommands", "AvaSequencerCommands", "Motion Design Sequencer Commands")
+		: TCommands<FAvaSequencerCommands>(TEXT("AvaSequencerCommands")
+		, NSLOCTEXT("MotionDesignSequencerCommands", "MotionDesignSequencerCommands", "Motion Design Sequencer Commands")
 		, NAME_None
 		, FAppStyle::GetAppStyleSetName())
 	{
diff --git a/Engine/Plugins/Experimental/Avalanche/Source/AvalancheSequencer/Private/EaseCurveTool/AvaEaseCurveToolCommands.h b/Engine/Plugins/Experimental/Avalanche/Source/AvalancheSequencer/Private/EaseCurveTool/AvaEaseCurveToolCommands.h
--- a/Engine/Plugins/Experimental/Avalanche/Source/AvalancheSequencer/Private/EaseCurveTool/AvaEaseCurveToolCommands.h
+++ b/Engine/Plugins/Experimental/Avalanche/Source/AvalancheSequencer/Private/EaseCurveTool/AvaEaseCurveToolCommands.h
@@ -9,8 +9,8 @@ class FAvaEaseCurveToolCommands : public TCommands<FAvaEaseCurveToolCommands>
 {
 public:
 	FAvaEaseCurveToolCommands()
-		: TCommands<FAvaEaseCurveToolCommands>(TEXT("AvalancheEaseCurveToolCommands")
-			, NSLOCTEXT("AvaEaseCurveToolCommands", "AvaEaseCurveToolCommands", "Motion Design Curve Ease Tool Commands")
+		: TCommands<FAvaEaseCurveToolCommands>(TEXT("AvaEaseCurveTool")
+			, NSLOCTEXT("MotionDesignEaseCurveToolCommands", "MotionDesignEaseCurveToolCommands", "Motion Design Curve Ease Tool Commands")
 			, NAME_None
 			, FAppStyle::GetAppStyleSetName())
 	{
diff --git a/Engine/Plugins/Experimental/Avalanche/Source/AvalancheShapesEditor/Private/AvaShapesEditorCommands.cpp b/Engine/Plugins/Experimental/Avalanche/Source/AvalancheShapesEditor/Private/AvaShapesEditorCommands.cpp
--- a/Engine/Plugins/Experimental/Avalanche/Source/AvalancheShapesEditor/Private/AvaShapesEditorCommands.cpp
+++ b/Engine/Plugins/Experimental/Avalanche/Source/AvalancheShapesEditor/Private/AvaShapesEditorCommands.cpp
@@ -7,7 +7,7 @@
 
 FAvaShapesEditorCommands::FAvaShapesEditorCommands()
 	: TCommands<FAvaShapesEditorCommands>(
-		TEXT("AvalancheShapesEditor")
+		TEXT("AvaShapesEditor")
 		, LOCTEXT("MotionDesignShapesEditor", "Motion Design Shapes Editor")
 		, NAME_None
 		, FAvaShapesEditorStyle::Get().GetStyleSetName()
diff --git a/Engine/Plugins/Experimental/Avalanche/Source/AvalancheShapesEditor/Private/AvaShapesEditorStyle.cpp b/Engine/Plugins/Experimental/Avalanche/Source/AvalancheShapesEditor/Private/AvaShapesEditorStyle.cpp
--- a/Engine/Plugins/Experimental/Avalanche/Source/AvalancheShapesEditor/Private/AvaShapesEditorStyle.cpp
+++ b/Engine/Plugins/Experimental/Avalanche/Source/AvalancheShapesEditor/Private/AvaShapesEditorStyle.cpp
@@ -37,15 +37,15 @@ FAvaShapesEditorStyle::FAvaShapesEditorStyle()
 	Set("ClassIcon.AvaShapeShapeActor", new IMAGE_BRUSH("Icons/ToolboxIcons/regularpolygon", Icon16x16));
 
 	// 2D Tools
-	Set("AvalancheShapesEditor.Tool_Shape_2DArrow",       new IMAGE_BRUSH("Icons/ToolboxIcons/arrow", Icon20x20));
-	Set("AvalancheShapesEditor.Tool_Shape_Chevron",       new IMAGE_BRUSH("Icons/ToolboxIcons/chevron", Icon20x20));
-	Set("AvalancheShapesEditor.Tool_Shape_Ellipse",       new IMAGE_BRUSH("Icons/ToolboxIcons/circle", Icon20x20));
-	Set("AvalancheShapesEditor.Tool_Shape_IrregularPoly", new IMAGE_BRUSH("Icons/ToolboxIcons/irregularpolygon", Icon20x20));
-	Set("AvalancheShapesEditor.Tool_Shape_Line",          new IMAGE_BRUSH_SVG("Icons/ToolboxIcons/minus", Icon20x20));
-	Set("AvalancheShapesEditor.Tool_Shape_NGon",          new IMAGE_BRUSH("Icons/ToolboxIcons/regularpolygon", Icon20x20));
-	Set("AvalancheShapesEditor.Tool_Shape_Rectangle",     new IMAGE_BRUSH("Icons/ToolboxIcons/rectangle", Icon20x20));
-	Set("AvalancheShapesEditor.Tool_Shape_Ring",          new IMAGE_BRUSH("Icons/ToolboxIcons/ring", Icon20x20));
-	Set("AvalancheShapesEditor.Tool_Shape_Star",          new IMAGE_BRUSH_SVG("Icons/ToolboxIcons/Favorite", Icon20x20));
+	Set("AvaShapesEditor.Tool_Shape_2DArrow",       new IMAGE_BRUSH("Icons/ToolboxIcons/arrow", Icon20x20));
+	Set("AvaShapesEditor.Tool_Shape_Chevron",       new IMAGE_BRUSH("Icons/ToolboxIcons/chevron", Icon20x20));
+	Set("AvaShapesEditor.Tool_Shape_Ellipse",       new IMAGE_BRUSH("Icons/ToolboxIcons/circle", Icon20x20));
+	Set("AvaShapesEditor.Tool_Shape_IrregularPoly", new IMAGE_BRUSH("Icons/ToolboxIcons/irregularpolygon", Icon20x20));
+	Set("AvaShapesEditor.Tool_Shape_Line",          new IMAGE_BRUSH_SVG("Icons/ToolboxIcons/minus", Icon20x20));
+	Set("AvaShapesEditor.Tool_Shape_NGon",          new IMAGE_BRUSH("Icons/ToolboxIcons/regularpolygon", Icon20x20));
+	Set("AvaShapesEditor.Tool_Shape_Rectangle",     new IMAGE_BRUSH("Icons/ToolboxIcons/rectangle", Icon20x20));
+	Set("AvaShapesEditor.Tool_Shape_Ring",          new IMAGE_BRUSH("Icons/ToolboxIcons/ring", Icon20x20));
+	Set("AvaShapesEditor.Tool_Shape_Star",          new IMAGE_BRUSH_SVG("Icons/ToolboxIcons/Favorite", Icon20x20));
 
 	Set("Tool_Shape_2DArrow",       new IMAGE_BRUSH("Icons/ToolboxIcons/arrow", Icon20x20));
 	Set("Tool_Shape_Chevron",       new IMAGE_BRUSH("Icons/ToolboxIcons/chevron", Icon20x20));
@@ -58,10 +58,10 @@ FAvaShapesEditorStyle::FAvaShapesEditorStyle()
 	Set("Tool_Shape_Star",          new IMAGE_BRUSH_SVG("Icons/ToolboxIcons/Favorite", Icon20x20));
 
 	// 3D Tools
-	Set("AvalancheShapesEditor.Tool_Shape_Cone",   new IMAGE_BRUSH("Icons/ToolboxIcons/cone", Icon20x20));
-	Set("AvalancheShapesEditor.Tool_Shape_Cube",   new IMAGE_BRUSH("Icons/ToolboxIcons/cube", Icon20x20));
-	Set("AvalancheShapesEditor.Tool_Shape_Sphere", new IMAGE_BRUSH("Icons/ToolboxIcons/sphere", Icon20x20));
-	Set("AvalancheShapesEditor.Tool_Shape_Torus",  new IMAGE_BRUSH("Icons/ToolboxIcons/torus", Icon20x20));
+	Set("AvaShapesEditor.Tool_Shape_Cone",   new IMAGE_BRUSH("Icons/ToolboxIcons/cone", Icon20x20));
+	Set("AvaShapesEditor.Tool_Shape_Cube",   new IMAGE_BRUSH("Icons/ToolboxIcons/cube", Icon20x20));
+	Set("AvaShapesEditor.Tool_Shape_Sphere", new IMAGE_BRUSH("Icons/ToolboxIcons/sphere", Icon20x20));
+	Set("AvaShapesEditor.Tool_Shape_Torus",  new IMAGE_BRUSH("Icons/ToolboxIcons/torus", Icon20x20));
 
 	Set("Tool_Shape_Cone",   new IMAGE_BRUSH("Icons/ToolboxIcons/cone", Icon20x20));
 	Set("Tool_Shape_Cube",   new IMAGE_BRUSH("Icons/ToolboxIcons/cube", Icon20x20));
diff --git a/Engine/Plugins/Experimental/Avalanche/Source/AvalancheTextEditor/Private/AvaTextEditorCommands.cpp b/Engine/Plugins/Experimental/Avalanche/Source/AvalancheTextEditor/Private/AvaTextEditorCommands.cpp
--- a/Engine/Plugins/Experimental/Avalanche/Source/AvalancheTextEditor/Private/AvaTextEditorCommands.cpp
+++ b/Engine/Plugins/Experimental/Avalanche/Source/AvalancheTextEditor/Private/AvaTextEditorCommands.cpp
@@ -8,7 +8,7 @@
 FAvaTextEditorCommands::FAvaTextEditorCommands()
 	: TCommands<FAvaTextEditorCommands>(
 		TEXT("AvaTextEditor")
-		, LOCTEXT("AvaTextEditor", "Motion Design Text Editor")
+		, LOCTEXT("MotionDesignTextEditor", "Motion Design Text Editor")
 		, NAME_None
 		, FAvaTextEditorStyle::Get().GetStyleSetName()
 	)
diff --git a/Engine/Plugins/Experimental/Avalanche/Source/AvalancheTextEditor/Private/Font/AvaFontSelectorCommands.h b/Engine/Plugins/Experimental/Avalanche/Source/AvalancheTextEditor/Private/Font/AvaFontSelectorCommands.h
--- a/Engine/Plugins/Experimental/Avalanche/Source/AvalancheTextEditor/Private/Font/AvaFontSelectorCommands.h
+++ b/Engine/Plugins/Experimental/Avalanche/Source/AvalancheTextEditor/Private/Font/AvaFontSelectorCommands.h
@@ -12,7 +12,7 @@ class FAvaFontSelectorCommands
 {
 public:
 	FAvaFontSelectorCommands()
-		: TCommands<FAvaFontSelectorCommands>(TEXT("AvalancheFontSelector")
+		: TCommands<FAvaFontSelectorCommands>(TEXT("AvaFontSelector")
 		, LOCTEXT("MotionDesignFontSelector", "Motion Design Font Selector")
 		, NAME_None
 		, FAppStyle::GetAppStyleSetName())
diff --git a/Engine/Plugins/Experimental/Avalanche/Source/AvalancheTransitionEditor/Private/AvaTransitionCommands.cpp b/Engine/Plugins/Experimental/Avalanche/Source/AvalancheTransitionEditor/Private/AvaTransitionCommands.cpp
--- a/Engine/Plugins/Experimental/Avalanche/Source/AvalancheTransitionEditor/Private/AvaTransitionCommands.cpp
+++ b/Engine/Plugins/Experimental/Avalanche/Source/AvalancheTransitionEditor/Private/AvaTransitionCommands.cpp
@@ -8,8 +8,8 @@
 #define LOCTEXT_NAMESPACE "AvaTransitionEditorCommands"
 
 FAvaTransitionEditorCommands::FAvaTransitionEditorCommands() 
-	: TCommands(TEXT("AvaTransitionEditorCommands")
-	, LOCTEXT("AvaTransitionEditor", "Motion Design Transition Logic Editor")
+	: TCommands(TEXT("AvaTransitionEditor")
+	, LOCTEXT("MotionDesignTransitionEditor", "Motion Design Transition Logic Editor")
 	, NAME_None
 	, FAvaTransitionEditorStyle::Get().GetStyleSetName())
 {
diff --git a/Engine/Plugins/Experimental/Avalanche/Source/AvalancheTransitionEditor/Private/AvaTransitionEditorStyle.cpp b/Engine/Plugins/Experimental/Avalanche/Source/AvalancheTransitionEditor/Private/AvaTransitionEditorStyle.cpp
--- a/Engine/Plugins/Experimental/Avalanche/Source/AvalancheTransitionEditor/Private/AvaTransitionEditorStyle.cpp
+++ b/Engine/Plugins/Experimental/Avalanche/Source/AvalancheTransitionEditor/Private/AvaTransitionEditorStyle.cpp
@@ -29,12 +29,12 @@ FAvaTransitionEditorStyle::FAvaTransitionEditorStyle()
 	Set("Throbber.CircleChunk", new CORE_IMAGE_BRUSH("Common/Throbber_Piece", FVector2f(1.5f)));
 
 	// Editor Commands
-	Set("AvaTransitionEditorCommands.AddSiblingState"       , new CORE_IMAGE_BRUSH_SVG("Starship/Common/plus"     , Icon20));
-	Set("AvaTransitionEditorCommands.AddChildState"         , new CORE_IMAGE_BRUSH_SVG("Starship/Common/plus"     , Icon20));
-	Set("AvaTransitionEditorCommands.ImportTransitionTree"  , new CORE_IMAGE_BRUSH_SVG("Starship/Common/import_20", Icon20));
-	Set("AvaTransitionEditorCommands.ReimportTransitionTree", new CORE_IMAGE_BRUSH_SVG("Starship/Common/import_20", Icon20));
-	Set("AvaTransitionEditorCommands.ExportTransitionTree"  , new CORE_IMAGE_BRUSH_SVG("Starship/Common/export_20", Icon20));
-	Set("AvaTransitionEditorCommands.ToggleDebug"           , new IMAGE_BRUSH_SVG("Starship/Common/Debug"         , Icon20));
+	Set("AvaTransitionEditor.AddSiblingState"       , new CORE_IMAGE_BRUSH_SVG("Starship/Common/plus"     , Icon20));
+	Set("AvaTransitionEditor.AddChildState"         , new CORE_IMAGE_BRUSH_SVG("Starship/Common/plus"     , Icon20));
+	Set("AvaTransitionEditor.ImportTransitionTree"  , new CORE_IMAGE_BRUSH_SVG("Starship/Common/import_20", Icon20));
+	Set("AvaTransitionEditor.ReimportTransitionTree", new CORE_IMAGE_BRUSH_SVG("Starship/Common/import_20", Icon20));
+	Set("AvaTransitionEditor.ExportTransitionTree"  , new CORE_IMAGE_BRUSH_SVG("Starship/Common/export_20", Icon20));
+	Set("AvaTransitionEditor.ToggleDebug"           , new IMAGE_BRUSH_SVG("Starship/Common/Debug"         , Icon20));
 
 	Set("DebugIndicatorBorder", new BOX_BRUSH("Images/NamespaceBorder", FMargin(0.25f)));
 
