[Motion Design]
 - Update mask and alignment style names

diff --git a/Engine/Plugins/Experimental/Avalanche/Source/AvalancheMaskEditor/Private/AvaMaskEditorMode.cpp b/Engine/Plugins/Experimental/Avalanche/Source/AvalancheMaskEditor/Private/AvaMaskEditorMode.cpp
--- a/Engine/Plugins/Experimental/Avalanche/Source/AvalancheMaskEditor/Private/AvaMaskEditorMode.cpp
+++ b/Engine/Plugins/Experimental/Avalanche/Source/AvalancheMaskEditor/Private/AvaMaskEditorMode.cpp
@@ -78,22 +78,22 @@ void UAvaMaskEditorMode::Enter()
     		static FName ToolkitOverlayMenuName = UE::AvalancheMaskEditor::Internal::ToolkitOverlayMenuName;
     		UToolMenu* Menu = UToolMenus::Get()->ExtendMenu(ToolkitOverlayMenuName);
     		Menu->SetStyleSet(&FAvaMaskEditorStyle::Get());
-    		Menu->StyleName = TEXT("AvalancheMaskEditor.ViewportOverlayToolbar");
+    		Menu->StyleName = TEXT("AvaMaskEditor.ViewportOverlayToolbar");
     		{
     			FToolMenuSection& Section = Menu->FindOrAddSection(TEXT("Default"));
     			{
     				FToolMenuEntry& Entry = Section.AddEntry(FToolMenuEntry::InitToolBarButton(FAvaMaskEditorCommands::Get().ToggleShowAllMasks));
-    				Entry.Icon.Set(FSlateIcon(FAvaMaskEditorStyle::Get().GetStyleSetName(), TEXT("AvalancheMaskEditor.ToggleShowAllMasks")));
+    				Entry.Icon.Set(FSlateIcon(FAvaMaskEditorStyle::Get().GetStyleSetName(), TEXT("AvaMaskEditor.ToggleShowAllMasks")));
     			}
     
     			{
     				FToolMenuEntry& Entry = Section.AddEntry(FToolMenuEntry::InitToolBarButton(FAvaMaskEditorCommands::Get().ToggleEnableMask));
-    				Entry.Icon.Set(FSlateIcon(FAvaMaskEditorStyle::Get().GetStyleSetName(), TEXT("AvalancheMaskEditor.ToggleDisableMask")));
+    				Entry.Icon.Set(FSlateIcon(FAvaMaskEditorStyle::Get().GetStyleSetName(), TEXT("AvaMaskEditor.ToggleDisableMask")));
     			}
     
     			{
     				FToolMenuEntry& Entry = Section.AddEntry(FToolMenuEntry::InitToolBarButton(FAvaMaskEditorCommands::Get().ToggleIsolateMask));
-    				Entry.Icon.Set(FSlateIcon(FAvaMaskEditorStyle::Get().GetStyleSetName(), TEXT("AvalancheMaskEditor.ToggleIsolateMask")));
+    				Entry.Icon.Set(FSlateIcon(FAvaMaskEditorStyle::Get().GetStyleSetName(), TEXT("AvaMaskEditor.ToggleIsolateMask")));
     			}
     		}
     
@@ -128,7 +128,7 @@ void UAvaMaskEditorMode::Enter()
 						.Padding(FMargin(0.f, 0.f, 8.f, 0.f))
 						[
 							SNew(SImage)
-							.Image(FAvaMaskEditorStyle::Get().GetBrush(TEXT("AvalancheMaskEditor.ToggleMaskMode")))
+							.Image(FAvaMaskEditorStyle::Get().GetBrush(TEXT("AvaMaskEditor.ToggleMaskMode")))
 						]
 
 						+SHorizontalBox::Slot()
diff --git a/Engine/Plugins/Experimental/Avalanche/Source/AvalancheMaskEditor/Private/AvaMaskEditorModule.cpp b/Engine/Plugins/Experimental/Avalanche/Source/AvalancheMaskEditor/Private/AvaMaskEditorModule.cpp
--- a/Engine/Plugins/Experimental/Avalanche/Source/AvalancheMaskEditor/Private/AvaMaskEditorModule.cpp
+++ b/Engine/Plugins/Experimental/Avalanche/Source/AvalancheMaskEditor/Private/AvaMaskEditorModule.cpp
@@ -70,7 +70,7 @@ void FAvalancheMaskEditorModule::RegisterMenus()
 			{
 				FToolMenuEntry& Entry = Section.AddEntry(FToolMenuEntry::InitToolBarButton(FAvaMaskEditorCommands::Get().ToggleMaskMode));				
 				Entry.Label.Set(FText::GetEmpty());
-				Entry.Icon.Set(FSlateIcon(FAvaMaskEditorStyle::Get().GetStyleSetName(), TEXT("AvalancheMaskEditor.ToggleMaskMode.Small")));
+				Entry.Icon.Set(FSlateIcon(FAvaMaskEditorStyle::Get().GetStyleSetName(), TEXT("AvaMaskEditor.ToggleMaskMode.Small")));
 				
 				Entry.SetCommandList(CommandList);
 			}
diff --git a/Engine/Plugins/Experimental/Avalanche/Source/AvalancheTextEditor/Private/Slate/SAvaTextAlignmentWidget.cpp b/Engine/Plugins/Experimental/Avalanche/Source/AvalancheTextEditor/Private/Slate/SAvaTextAlignmentWidget.cpp
--- a/Engine/Plugins/Experimental/Avalanche/Source/AvalancheTextEditor/Private/Slate/SAvaTextAlignmentWidget.cpp
+++ b/Engine/Plugins/Experimental/Avalanche/Source/AvalancheTextEditor/Private/Slate/SAvaTextAlignmentWidget.cpp
@@ -92,7 +92,7 @@ END_SLATE_FUNCTION_BUILD_OPTIMIZATION
 
 TSharedRef<SButton> SAvaTextAlignmentWidget::GetHorizontalAlignmentButton(TSharedPtr<SButton>& OutButton, EText3DHorizontalTextAlignment InHorizontalAlignment, FName Image, FText Tooltip)
 {
-	const ISlateStyle* AvalancheEditorStyle = FSlateStyleRegistry::FindSlateStyle("AvalancheEditor");
+	const ISlateStyle* AvalancheEditorStyle = FSlateStyleRegistry::FindSlateStyle("AvaEditor");
 
 	OutButton = SNew(SButton)
 		.ButtonStyle(&FAppStyle::Get().GetWidgetStyle<FButtonStyle>("SimpleButton"))
@@ -111,7 +111,7 @@ TSharedRef<SButton> SAvaTextAlignmentWidget::GetHorizontalAlignmentButton(TShare
 
 TSharedRef<SButton> SAvaTextAlignmentWidget::GetVerticalAlignmentButton(TSharedPtr<SButton>& OutButton, EText3DVerticalTextAlignment InVerticalAlignment, FName Image, FText Tooltip)
 {
-	const ISlateStyle* AvalancheEditorStyle = FSlateStyleRegistry::FindSlateStyle("AvalancheEditor");
+	const ISlateStyle* AvalancheEditorStyle = FSlateStyleRegistry::FindSlateStyle("AvaEditor");
 
 	OutButton = SNew(SButton)
 		.ButtonStyle(&FAppStyle::Get().GetWidgetStyle<FButtonStyle>("SimpleButton"))
