[Motion Design] Fix Rundown RC Controller sorting in Page Details.

#rnx

diff --git a/Engine/Plugins/Experimental/Avalanche/Source/AvalancheMediaEditor/Private/Rundown/DetailsView/RemoteControl/Controllers/AvaRundownRCControllerItem.cpp b/Engine/Plugins/Experimental/Avalanche/Source/AvalancheMediaEditor/Private/Rundown/DetailsView/RemoteControl/Controllers/AvaRundownRCControllerItem.cpp
--- a/Engine/Plugins/Experimental/Avalanche/Source/AvalancheMediaEditor/Private/Rundown/DetailsView/RemoteControl/Controllers/AvaRundownRCControllerItem.cpp
+++ b/Engine/Plugins/Experimental/Avalanche/Source/AvalancheMediaEditor/Private/Rundown/DetailsView/RemoteControl/Controllers/AvaRundownRCControllerItem.cpp
@@ -14,6 +14,8 @@ FAvaRundownRCControllerItem::FAvaRundownRCControllerItem(int32 InInstanceIndex,
 	
 	if (InController)
 	{
+		DisplayIndex = InController->DisplayIndex;
+		
 		const FName DisplayName = InController->DisplayName.IsNone()
 			? InController->PropertyName
 			: InController->DisplayName;
