[GeometryMask] CIS fix

diff --git a/Engine/Plugins/Experimental/Avalanche/Source/AvalancheMask/Private/Mask2D/AvaMask2DReadModifier.cpp b/Engine/Plugins/Experimental/Avalanche/Source/AvalancheMask/Private/Mask2D/AvaMask2DReadModifier.cpp
--- a/Engine/Plugins/Experimental/Avalanche/Source/AvalancheMask/Private/Mask2D/AvaMask2DReadModifier.cpp
+++ b/Engine/Plugins/Experimental/Avalanche/Source/AvalancheMask/Private/Mask2D/AvaMask2DReadModifier.cpp
@@ -4,8 +4,8 @@
 
 #include "AvaMaskUtilities.h"
 #include "Components/PrimitiveComponent.h"
+#include "Engine/CanvasRenderTarget2D.h"
 #include "Engine/Engine.h"
-#include "Engine/TextureRenderTarget2D.h"
 #include "GeometryMaskCanvas.h"
 #include "GeometryMaskCanvasResource.h"
 #include "GeometryMaskReadComponent.h"
diff --git a/Engine/Plugins/Experimental/Avalanche/Source/AvalancheMask/Private/Mask2D/AvaMask2DWriteModifier.cpp b/Engine/Plugins/Experimental/Avalanche/Source/AvalancheMask/Private/Mask2D/AvaMask2DWriteModifier.cpp
--- a/Engine/Plugins/Experimental/Avalanche/Source/AvalancheMask/Private/Mask2D/AvaMask2DWriteModifier.cpp
+++ b/Engine/Plugins/Experimental/Avalanche/Source/AvalancheMask/Private/Mask2D/AvaMask2DWriteModifier.cpp
@@ -4,6 +4,7 @@
 
 #include "AvaMaskUtilities.h"
 #include "Components/PrimitiveComponent.h"
+#include "Engine/CanvasRenderTarget2D.h"
 #include "Engine/Engine.h"
 #include "Framework/AvaGizmoComponent.h"
 #include "GeometryMaskCanvas.h"
diff --git a/Engine/Plugins/Experimental/Avalanche/Source/AvalancheMask/Public/Mask2D/AvaMask2DBaseModifier.h b/Engine/Plugins/Experimental/Avalanche/Source/AvalancheMask/Public/Mask2D/AvaMask2DBaseModifier.h
--- a/Engine/Plugins/Experimental/Avalanche/Source/AvalancheMask/Public/Mask2D/AvaMask2DBaseModifier.h
+++ b/Engine/Plugins/Experimental/Avalanche/Source/AvalancheMask/Public/Mask2D/AvaMask2DBaseModifier.h
@@ -21,6 +21,7 @@ class IGeometryMaskWriteInterface;
 class UActorComponent;
 class UAvaMaskMaterialInstanceSubsystem;
 class UAvaObjectHandleSubsystem;
+class UCanvasRenderTarget2D;
 class UGeometryMaskCanvas;
 class UMaterialInstanceDynamic;
 class UMaterialInterface;
diff --git a/Engine/Plugins/Experimental/Avalanche/Source/AvalancheMask/Public/Mask2D/AvaMask2DReadModifier.h b/Engine/Plugins/Experimental/Avalanche/Source/AvalancheMask/Public/Mask2D/AvaMask2DReadModifier.h
--- a/Engine/Plugins/Experimental/Avalanche/Source/AvalancheMask/Public/Mask2D/AvaMask2DReadModifier.h
+++ b/Engine/Plugins/Experimental/Avalanche/Source/AvalancheMask/Public/Mask2D/AvaMask2DReadModifier.h
@@ -22,6 +22,7 @@ class UActorComponent;
 class UAvaMaskMaterialInstanceSubsystem;
 class UAvaObjectHandleSubsystem;
 class UGeometryMaskCanvas;
+class UCanvasRenderTarget2D;
 class UMaterialInstanceDynamic;
 class UMaterialInterface;
 class UPrimitiveComponent;
diff --git a/Engine/Plugins/Experimental/Avalanche/Source/AvalancheMask/Public/Mask2D/AvaMask2DWriteModifier.h b/Engine/Plugins/Experimental/Avalanche/Source/AvalancheMask/Public/Mask2D/AvaMask2DWriteModifier.h
--- a/Engine/Plugins/Experimental/Avalanche/Source/AvalancheMask/Public/Mask2D/AvaMask2DWriteModifier.h
+++ b/Engine/Plugins/Experimental/Avalanche/Source/AvalancheMask/Public/Mask2D/AvaMask2DWriteModifier.h
@@ -21,6 +21,7 @@ class IGeometryMaskWriteInterface;
 class UActorComponent;
 class UAvaMaskMaterialInstanceSubsystem;
 class UAvaObjectHandleSubsystem;
+class UCanvasRenderTarget2D;
 class UGeometryMaskCanvas;
 class UMaterialInstanceDynamic;
 class UMaterialInterface;
