[Motion Design] Fixing Playback Graph Editor. Adding missing core redirects.

#rnx

diff --git a/Engine/Plugins/Experimental/Avalanche/Config/DefaultAvalanche.ini b/Engine/Plugins/Experimental/Avalanche/Config/DefaultAvalanche.ini
--- a/Engine/Plugins/Experimental/Avalanche/Config/DefaultAvalanche.ini
+++ b/Engine/Plugins/Experimental/Avalanche/Config/DefaultAvalanche.ini
@@ -187,6 +187,12 @@
 +ClassRedirects=(OldName="/Script/AvalancheMediaEditor.AvalancheMediaEditorSettings",NewName="/Script/AvalancheMediaEditor.AvaMediaEditorSettings")
 +ClassRedirects=(OldName="/Script/AvalancheMediaEditor.AvaPlaybackGraph",NewName="/Script/AvalancheMediaEditor.AvaPlaybackEditorGraph")
 +ClassRedirects=(OldName="/Script/AvalancheMediaEditor.AvaPlaybackGraphSchema",NewName="/Script/AvalancheMediaEditor.AvaPlaybackEditorGraphSchema")
++ClassRedirects=(OldName="/Script/AvalancheMediaEditor.AvaPlaybackGraphNode",NewName="/Script/AvalancheMediaEditor.AvaPlaybackEditorGraphNode")
++ClassRedirects=(OldName="/Script/AvalancheMediaEditor.AvaPlaybackGraphNode_EventAction",NewName="/Script/AvalancheMediaEditor.AvaPlaybackEditorGraphNode_EventAction")
++ClassRedirects=(OldName="/Script/AvalancheMediaEditor.AvaPlaybackGraphNode_EventFlow",NewName="/Script/AvalancheMediaEditor.AvaPlaybackEditorGraphNode_EventFlow")
++ClassRedirects=(OldName="/Script/AvalancheMediaEditor.AvaPlaybackGraphNode_EventTrigger",NewName="/Script/AvalancheMediaEditor.AvaPlaybackEditorGraphNode_EventTrigger")
++ClassRedirects=(OldName="/Script/AvalancheMediaEditor.AvaPlaybackGraphNode_LevelPlayer",NewName="/Script/AvalancheMediaEditor.AvaPlaybackEditorGraphNode_LevelPlayer")
++ClassRedirects=(OldName="/Script/AvalancheMediaEditor.AvaPlaybackGraphNode_Root",NewName="/Script/AvalancheMediaEditor.AvaPlaybackEditorGraphNode_Root")
 +EnumRedirects=(OldName="/Script/AvalancheMediaEditor.EAvaRundownPageSet",NewName="/Script/AvalancheMediaEditor.EAvaRundownPageSet")
 +ClassRedirects=(OldName="/Script/AvalancheMediaEditor.AvaRundownEditorSettings",NewName="/Script/AvalancheMediaEditor.AvaRundownEditorSettings")
 +StructRedirects=(OldName="/Script/AvalancheMediaEditor.AvaPlaylistMsgBase",NewName="/Script/AvalancheMediaEditor.AvaRundownMsgBase")
