Motion Design Font:
- Fix cast checked as GetAsset does not guarantee non-null

#changelist validated
#virtualized

diff --git a/Engine/Plugins/Experimental/Avalanche/Source/AvalancheText/Private/Font/AvaFont.cpp b/Engine/Plugins/Experimental/Avalanche/Source/AvalancheText/Private/Font/AvaFont.cpp
--- a/Engine/Plugins/Experimental/Avalanche/Source/AvalancheText/Private/Font/AvaFont.cpp
+++ b/Engine/Plugins/Experimental/Avalanche/Source/AvalancheText/Private/Font/AvaFont.cpp
@@ -394,7 +394,7 @@ UFont* FAvaFont::GetFontByName(const FString& InFontName)
 
 	for (const FAssetData& AssetData : AssetDataArray)
 	{
-		if (UFont* const CurrFont = CastChecked<UFont>(AssetData.GetAsset()))
+		if (UFont* const CurrFont = Cast<UFont>(AssetData.GetAsset()))
 		{
 			FString CurrFontName;
 			UE::Avalanche::FontUtilities::Public::GetFontName(CurrFont, CurrFontName);
