Motion Design Transition:
- Introduced UX to identify when the Tree has been disabled, with option to enable it from within the Transition Editor (if not in read-only mode)
- Tree is disabled by default on new scenes.

#changelist validated
#virtualized

diff --git a/Engine/Plugins/Experimental/Avalanche/Source/AvalancheTransition/Public/AvaTransitionTree.h b/Engine/Plugins/Experimental/Avalanche/Source/AvalancheTransition/Public/AvaTransitionTree.h
--- a/Engine/Plugins/Experimental/Avalanche/Source/AvalancheTransition/Public/AvaTransitionTree.h
+++ b/Engine/Plugins/Experimental/Avalanche/Source/AvalancheTransition/Public/AvaTransitionTree.h
@@ -20,6 +20,11 @@ public:
 
 	void SetEnabled(bool bInEnabled);
 
+	static FName GetEnabledPropertyName()
+	{
+		return GET_MEMBER_NAME_CHECKED(UAvaTransitionTree, bEnabled);
+	}
+
 private:
 	/** The Layer this Transition Logic Tree deals with */
 	UPROPERTY()
diff --git a/Engine/Plugins/Experimental/Avalanche/Source/AvalancheTransitionEditor/Private/AvaTransitionEditor.cpp b/Engine/Plugins/Experimental/Avalanche/Source/AvalancheTransitionEditor/Private/AvaTransitionEditor.cpp
--- a/Engine/Plugins/Experimental/Avalanche/Source/AvalancheTransitionEditor/Private/AvaTransitionEditor.cpp
+++ b/Engine/Plugins/Experimental/Avalanche/Source/AvalancheTransitionEditor/Private/AvaTransitionEditor.cpp
@@ -136,6 +136,7 @@ void FAvaTransitionEditor::SetupReadOnlyMenuProfiles(FReadOnlyAssetEditorCustomi
 	}
 
 	TSharedRef<FAvaTransitionToolbar> ViewModelToolbar = EditorViewModel->GetToolbar();
+	ViewModelToolbar->SetReadOnlyProfileName(FAssetEditorToolkit::GetToolMenuToolbarName(), GetReadOnlyMenuProfileName());
 	ViewModelToolbar->SetupReadOnlyCustomization(InReadOnlyCustomization);
 }
 
diff --git a/Engine/Plugins/Experimental/Avalanche/Source/AvalancheTransitionEditor/Private/AvaTransitionEditorModule.cpp b/Engine/Plugins/Experimental/Avalanche/Source/AvalancheTransitionEditor/Private/AvaTransitionEditorModule.cpp
--- a/Engine/Plugins/Experimental/Avalanche/Source/AvalancheTransitionEditor/Private/AvaTransitionEditorModule.cpp
+++ b/Engine/Plugins/Experimental/Avalanche/Source/AvalancheTransitionEditor/Private/AvaTransitionEditorModule.cpp
@@ -115,6 +115,9 @@ void FAvaTransitionEditorModule::ValidateStateTree(UAvaTransitionTree* InTransit
 		return;
 	}
 
+	// Disable Tree by default if being set up for the first time
+	InTransitionTree->SetEnabled(false);
+
 	UAvaTransitionTreeEditorData* const EditorData = NewObject<UAvaTransitionTreeEditorData>(InTransitionTree, NAME_None, RF_Transactional);
 	check(EditorData);
 
diff --git a/Engine/Plugins/Experimental/Avalanche/Source/AvalancheTransitionEditor/Private/AvaTransitionTreeFactory.cpp b/Engine/Plugins/Experimental/Avalanche/Source/AvalancheTransitionEditor/Private/AvaTransitionTreeFactory.cpp
--- a/Engine/Plugins/Experimental/Avalanche/Source/AvalancheTransitionEditor/Private/AvaTransitionTreeFactory.cpp
+++ b/Engine/Plugins/Experimental/Avalanche/Source/AvalancheTransitionEditor/Private/AvaTransitionTreeFactory.cpp
@@ -28,6 +28,9 @@ UObject* UAvaTransitionTreeFactory::FactoryCreateNew(UClass* InClass, UObject* I
 		if (TransitionTree)
 		{
 			IAvaTransitionModule::Get().GetOnValidateTransitionTree().ExecuteIfBound(TransitionTree);
+
+			// Transition Tree Assets are always Enabled by Default
+			TransitionTree->SetEnabled(true);
 		}
 		return TransitionTree;
 	}
diff --git a/Engine/Plugins/Experimental/Avalanche/Source/AvalancheTransitionEditor/Private/Menu/AvaTransitionToolbar.cpp b/Engine/Plugins/Experimental/Avalanche/Source/AvalancheTransitionEditor/Private/Menu/AvaTransitionToolbar.cpp
--- a/Engine/Plugins/Experimental/Avalanche/Source/AvalancheTransitionEditor/Private/Menu/AvaTransitionToolbar.cpp
+++ b/Engine/Plugins/Experimental/Avalanche/Source/AvalancheTransitionEditor/Private/Menu/AvaTransitionToolbar.cpp
@@ -4,6 +4,7 @@
 #include "AvaTransitionCommands.h"
 #include "AvaTransitionEditorUtils.h"
 #include "AvaTransitionMenuContext.h"
+#include "AvaTransitionTree.h"
 #include "ToolMenu.h"
 #include "ToolMenuEntry.h"
 #include "ToolMenuSection.h"
@@ -11,6 +12,7 @@
 #include "Toolkits/AssetEditorToolkit.h"
 #include "ViewModels/AvaTransitionEditorViewModel.h"
 #include "ViewModels/AvaTransitionViewModelSharedData.h"
+#include "Widgets/AvaTransitionTreeStatus.h"
 #include "Widgets/Layout/SBox.h"
 #include "Widgets/Text/STextBlock.h"
 
@@ -18,11 +20,35 @@
 
 namespace UE::AvaTransitionEditor::Private
 {
-	static constexpr const TCHAR* TransitionLayerPickerName = TEXT("TransitionLogicLayerPicker");
+	FToolMenuEntry& AddLabelEntry(FToolMenuSection& InSection, FName InEntryName, const FText& InLabel)
+	{
+		TSharedRef<SWidget> LabelWidget = SNew(SBox)
+			.Padding(8.f, 0.f, 2.f, 0.f)
+			.VAlign(VAlign_Center)
+			[
+				SNew(STextBlock)
+				.Text(InLabel)
+				.Justification(ETextJustify::Right)
+			];
+		return InSection.AddEntry(FToolMenuEntry::InitWidget(InEntryName, LabelWidget, FText::GetEmpty()));
+	}
+}
+
+FAvaTransitionToolbar::FAvaTransitionToolbar(FAvaTransitionEditorViewModel& InOwner)
+	: Owner(InOwner)
+{
+}
+
+void FAvaTransitionToolbar::SetReadOnlyProfileName(FName InToolMenuToolbarName, FName InReadOnlyProfileName)
+{
+	ToolMenuToolbarName = InToolMenuToolbarName;
+	ReadOnlyProfileName = InReadOnlyProfileName;
 }
 
 void FAvaTransitionToolbar::ExtendEditorToolbar(UToolMenu* InToolbarMenu)
 {
+	using namespace UE::AvaTransitionEditor;
+
 	if (!InToolbarMenu)
 	{
 		return;
@@ -30,6 +56,15 @@ void FAvaTransitionToolbar::ExtendEditorToolbar(UToolMenu* InToolbarMenu)
 
 	const bool bReadOnly = Owner.GetSharedData()->IsReadOnly();
 
+	const FName PermissionOwner = TEXT("FAvaTransitionToolbar");
+
+	FNamePermissionList ReadOnlyPermissionList;
+
+	auto AllowInReadOnly = [&ReadOnlyPermissionList, &PermissionOwner](FToolMenuEntry& InEntry)
+		{
+			ReadOnlyPermissionList.AddAllowListItem(PermissionOwner, InEntry.Name);;
+		};
+
 	FToolMenuSection& Section = InToolbarMenu->FindOrAddSection(TEXT("TransitionLogic"));
 
 	const FAvaTransitionEditorCommands& Commands = FAvaTransitionEditorCommands::Get();
@@ -43,41 +78,53 @@ void FAvaTransitionToolbar::ExtendEditorToolbar(UToolMenu* InToolbarMenu)
 				return Owner.GetCompiler().GetCompileStatusIcon();
 			}))));
 
-	CompileButton.StyleNameOverride = "CalloutToolbar";
-
 	FToolMenuEntry& CompileOptions = Section.AddEntry(FToolMenuEntry::InitComboButton(TEXT("CompileComboButton")
 		, FUIAction()
 		, FNewToolMenuDelegate::CreateStatic(&FAvaTransitionCompiler::GenerateCompileOptionsMenu)
 		, LOCTEXT("CompileOptions_ToolbarTooltip", "Options to customize how State Trees compile")));
 
-	CompileOptions.StyleNameOverride = "CalloutToolbar";
 	CompileOptions.ToolBarData.bSimpleComboBox = true;
 
-	if (TSharedPtr<SWidget> LayerPicker = UE::AvaTransitionEditor::CreateTransitionLayerPicker(Owner.GetEditorData(), /*bInCompileOnLayerPicked*/false))
+	int32 SeparatorIndex = 0;
+	auto MakeSeparator = [&SeparatorIndex](FToolMenuSection& InSection)->FToolMenuEntry&
+		{
+			return InSection.AddSeparator(FName(TEXT("Separator"), SeparatorIndex++));
+		};
+
+	// Tree Status
+	{
+		TSharedRef<SWidget> TreeStatusWidget = SNew(SAvaTransitionTreeStatus, Owner.GetTransitionTree());
+		TreeStatusWidget->SetEnabled(!bReadOnly);
+
+		AllowInReadOnly(MakeSeparator(Section));
+
+		AllowInReadOnly(Private::AddLabelEntry(Section, TEXT("TreeStatusLabel"), LOCTEXT("StatusLabel", "Status")));
+
+		AllowInReadOnly(Section.AddEntry(FToolMenuEntry::InitWidget(TEXT("TreeStatusWidget")
+			, TreeStatusWidget
+			, FText::GetEmpty())));
+	}
+
+	// Layer Picker
+	if (TSharedPtr<SWidget> LayerPicker = CreateTransitionLayerPicker(Owner.GetEditorData(), /*bInCompileOnLayerPicked*/false))
 	{
 		LayerPicker->SetEnabled(!bReadOnly);
 
-		Section.AddSeparator(NAME_None);
+		AllowInReadOnly(MakeSeparator(Section));
 
-		TSharedRef<SWidget> TransitionLayerWidget = SNew(SBox)
-			.Padding(8.f, 0.f, 2.f, 0.f)
-			.VAlign(VAlign_Center)
-			[
-				SNew(STextBlock)
-				.Text(LOCTEXT("TransitionLayerText", "Layer"))
-				.Justification(ETextJustify::Right)
-			];
+		AllowInReadOnly(Private::AddLabelEntry(Section, TEXT("LayerPickerLabel"), LOCTEXT("LayerLabel", "Layer")));
 
-		Section.AddEntry(FToolMenuEntry::InitWidget("TransitionLayerWidget", TransitionLayerWidget, FText::GetEmpty()));
-		Section.AddEntry(FToolMenuEntry::InitWidget(UE::AvaTransitionEditor::Private::TransitionLayerPickerName
+		AllowInReadOnly(Section.AddEntry(FToolMenuEntry::InitWidget(TEXT("LayerPickerWidget")
 			, LayerPicker.ToSharedRef()
-			, LOCTEXT("TransitionLogicLabel", "Transition Logic Layer")));
+			, FText::GetEmpty())));
 	}
 
 #if WITH_STATETREE_DEBUGGER
-	Section.AddSeparator(NAME_None);
-	Section.AddEntry(FToolMenuEntry::InitToolBarButton(Commands.ToggleDebug));
+	AllowInReadOnly(MakeSeparator(Section));
+	AllowInReadOnly(Section.AddEntry(FToolMenuEntry::InitToolBarButton(Commands.ToggleDebug)));
 #endif
+
+	ApplyReadOnlyPermissionList(ReadOnlyPermissionList);
 }
 
 void FAvaTransitionToolbar::ExtendTreeToolbar(UToolMenu* InToolbarMenu)
@@ -107,15 +154,7 @@ void FAvaTransitionToolbar::ExtendTreeToolbar(UToolMenu* InToolbarMenu)
 
 void FAvaTransitionToolbar::SetupReadOnlyCustomization(FReadOnlyAssetEditorCustomization& InReadOnlyCustomization)
 {
-	const FName PermissionOwner = TEXT("FAvaTransitionToolbar");
-
-	const FAvaTransitionEditorCommands& Commands = FAvaTransitionEditorCommands::Get();
-
-	InReadOnlyCustomization.ToolbarPermissionList.AddAllowListItem(PermissionOwner, UE::AvaTransitionEditor::Private::TransitionLayerPickerName);
-
-#if WITH_STATETREE_DEBUGGER
-	InReadOnlyCustomization.ToolbarPermissionList.AddAllowListItem(PermissionOwner, Commands.ToggleDebug->GetCommandName());
-#endif
+	// Read Only Customizations not added here, as this is called before the Tool Menu is extended, and so the names are not known at this point
 }
 
 TSharedRef<SWidget> FAvaTransitionToolbar::GenerateTreeToolbarWidget()
@@ -123,13 +162,13 @@ TSharedRef<SWidget> FAvaTransitionToolbar::GenerateTreeToolbarWidget()
 	UToolMenus* const ToolMenus = UToolMenus::Get();
 	check(ToolMenus);
 
-	const FName ToolbarName = GetTreeToolbarName();
+	const FName TreeToolbarName = GetTreeToolbarName();
 
-	if (!ToolMenus->IsMenuRegistered(ToolbarName))
+	if (!ToolMenus->IsMenuRegistered(TreeToolbarName))
 	{
-		UToolMenu* const ToolBar = ToolMenus->RegisterMenu(ToolbarName, NAME_None, EMultiBoxType::SlimHorizontalToolBar);
-		ToolBar->StyleName = "CalloutToolbar";
-		ToolBar->AddDynamicSection("PopulateToolbar", FNewToolMenuDelegate::CreateStatic([](UToolMenu* InToolMenu)
+		UToolMenu* const Toolbar = ToolMenus->RegisterMenu(TreeToolbarName, NAME_None, EMultiBoxType::SlimHorizontalToolBar);
+		Toolbar->StyleName = "CalloutToolbar";
+		Toolbar->AddDynamicSection("PopulateToolbar", FNewToolMenuDelegate::CreateStatic([](UToolMenu* InToolMenu)
 		{
 			if (InToolMenu)
 			{
@@ -150,7 +189,21 @@ TSharedRef<SWidget> FAvaTransitionToolbar::GenerateTreeToolbarWidget()
 	ContextObject->SetEditorViewModel(StaticCastSharedRef<FAvaTransitionEditorViewModel>(Owner.AsShared()));
 
 	FToolMenuContext Context(Owner.GetCommandList(), Extender, ContextObject);
-	return ToolMenus->GenerateWidget(ToolbarName, Context);
+	return ToolMenus->GenerateWidget(TreeToolbarName, Context);
+}
+
+void FAvaTransitionToolbar::ApplyReadOnlyPermissionList(const FNamePermissionList& InPermissionList)
+{
+	UToolMenus* ToolMenus = UToolMenus::Get();
+	if (!ToolMenus)
+	{
+		return;
+	}
+
+	if (FToolMenuProfile* ToolbarProfile = ToolMenus->FindRuntimeMenuProfile(ToolMenuToolbarName, ReadOnlyProfileName))
+	{
+		ToolbarProfile->MenuPermissions.Append(InPermissionList);
+	}
 }
 
 #undef LOCTEXT_NAMESPACE
diff --git a/Engine/Plugins/Experimental/Avalanche/Source/AvalancheTransitionEditor/Private/Menu/AvaTransitionToolbar.h b/Engine/Plugins/Experimental/Avalanche/Source/AvalancheTransitionEditor/Private/Menu/AvaTransitionToolbar.h
--- a/Engine/Plugins/Experimental/Avalanche/Source/AvalancheTransitionEditor/Private/Menu/AvaTransitionToolbar.h
+++ b/Engine/Plugins/Experimental/Avalanche/Source/AvalancheTransitionEditor/Private/Menu/AvaTransitionToolbar.h
@@ -6,6 +6,7 @@
 #include "UObject/NameTypes.h"
 
 class FAvaTransitionEditorViewModel;
+class FNamePermissionList;
 class SWidget;
 class UToolMenu;
 struct FReadOnlyAssetEditorCustomization;
@@ -13,16 +14,15 @@ struct FReadOnlyAssetEditorCustomization;
 class FAvaTransitionToolbar : public TSharedFromThis<FAvaTransitionToolbar>
 {
 public:
-	explicit FAvaTransitionToolbar(FAvaTransitionEditorViewModel& InOwner)
-		: Owner(InOwner)
-	{
-	}
+	explicit FAvaTransitionToolbar(FAvaTransitionEditorViewModel& InOwner);
 
 	static FName GetTreeToolbarName()
 	{
 		return TEXT("AvaTransitionTreeToolbar");
 	}
 
+	void SetReadOnlyProfileName(FName InToolMenuToolbarName, FName InReadOnlyProfileName);
+
 	void ExtendEditorToolbar(UToolMenu* InToolbarMenu);
 
 	void ExtendTreeToolbar(UToolMenu* InToolbarMenu);
@@ -32,5 +32,10 @@ public:
 	TSharedRef<SWidget> GenerateTreeToolbarWidget();
 
 private:
+	void ApplyReadOnlyPermissionList(const FNamePermissionList& InPermissionList);
+
 	FAvaTransitionEditorViewModel& Owner;
+
+	FName ToolMenuToolbarName;
+	FName ReadOnlyProfileName;
 };
diff --git a/Engine/Plugins/Experimental/Avalanche/Source/AvalancheTransitionEditor/Private/Menu/Widgets/AvaTransitionTreeStatus.cpp b/Engine/Plugins/Experimental/Avalanche/Source/AvalancheTransitionEditor/Private/Menu/Widgets/AvaTransitionTreeStatus.cpp
new file mode 100644
--- /dev/null
+++ b/Engine/Plugins/Experimental/Avalanche/Source/AvalancheTransitionEditor/Private/Menu/Widgets/AvaTransitionTreeStatus.cpp
@@ -0,0 +1,148 @@
+// Copyright Epic Games, Inc. All Rights Reserved.
+
+#include "AvaTransitionTreeStatus.h"
+#include "AvaTransitionEditorStyle.h"
+#include "AvaTransitionTree.h"
+#include "ScopedTransaction.h"
+#include "Widgets/Images/SImage.h"
+#include "Widgets/Input/SButton.h"
+
+#define LOCTEXT_NAMESPACE "SAvaTransitionTreeStatus"
+
+namespace UE::AvaTransitionEditor::Private
+{
+	// Image that always display the brush with no Draw Effects even if disabled
+	class SStatusImage : public SImage
+	{
+		//~ Begin SWidget
+		virtual int32 OnPaint(const FPaintArgs& InArgs, const FGeometry& InAllottedGeometry, const FSlateRect& InCullingRect, FSlateWindowElementList& OutDrawElements, int32 InLayerId, const FWidgetStyle& InWidgetStyle, bool bInParentEnabled) const override
+		{
+			constexpr bool bIsParentEnabled = true;
+			return SImage::OnPaint(InArgs, InAllottedGeometry, InCullingRect, OutDrawElements, InLayerId, InWidgetStyle, bIsParentEnabled);
+		}
+		//~ End SWidget
+	};
+}
+
+void SAvaTransitionTreeStatus::Construct(const FArguments& InArgs, UAvaTransitionTree* InTransitionTree)
+{
+	TransitionTreeWeak = InTransitionTree;
+
+	ChildSlot
+		.HAlign(HAlign_Center)
+		.VAlign(VAlign_Center)
+		.Padding(2.f, 0.f)
+		[
+			SNew(SButton)
+			.ButtonStyle(FAppStyle::Get(), "SimpleButton")
+			.IsEnabled(this, &SAvaTransitionTreeStatus::IsStatusButtonEnabled)
+			.OnClicked(this, &SAvaTransitionTreeStatus::OnStatusButtonClicked)
+			.ToolTipText(this, &SAvaTransitionTreeStatus::GetStatusTooltipText)
+			.IsFocusable(false)
+			[
+				SNew(UE::AvaTransitionEditor::Private::SStatusImage)
+				.Image(this, &SAvaTransitionTreeStatus::GetStatusIcon)
+			]
+		];
+}
+
+bool SAvaTransitionTreeStatus::IsStatusButtonEnabled() const
+{
+	return GetStatus() == EStatus::Disabled;
+}
+
+FReply SAvaTransitionTreeStatus::OnStatusButtonClicked()
+{
+	UAvaTransitionTree* TransitionTree = TransitionTreeWeak.Get();
+	if (TransitionTree && !TransitionTree->IsEnabled())
+	{
+		FScopedTransaction Transaction(LOCTEXT("EnableTransitionTree", "Enable Transition Tree"));
+
+		FProperty* const EnabledProperty = FindFProperty<FBoolProperty>(UAvaTransitionTree::StaticClass(), UAvaTransitionTree::GetEnabledPropertyName());
+		TransitionTree->PreEditChange(EnabledProperty);
+
+		TransitionTree->SetEnabled(true);
+
+		FPropertyChangedEvent PropertyChangedEvent(EnabledProperty, EPropertyChangeType::ValueSet);
+		TransitionTree->PostEditChangeProperty(PropertyChangedEvent);
+	}
+	return FReply::Handled();
+}
+
+SAvaTransitionTreeStatus::EStatus SAvaTransitionTreeStatus::GetStatus() const
+{
+	UAvaTransitionTree* TransitionTree = TransitionTreeWeak.Get();
+	if (!TransitionTree)
+	{
+		return EStatus::Invalid;
+	}
+
+	return TransitionTree->IsEnabled()
+		? EStatus::Enabled
+		: EStatus::Disabled;
+}
+
+FText SAvaTransitionTreeStatus::GetStatusText(EStatus InStatus) const
+{
+	switch (InStatus)
+	{
+	case EStatus::Invalid:
+		return LOCTEXT("InvalidText", "invalid");
+
+	case EStatus::Disabled:
+		return LOCTEXT("DisabledText", "disabled");
+
+	case EStatus::Enabled:
+		return LOCTEXT("EnabledText", "enabled");
+	}
+
+	checkNoEntry()
+	return FText::GetEmpty();
+}
+
+FText SAvaTransitionTreeStatus::GetStatusTooltipText() const
+{
+	EStatus Status = GetStatus();
+
+	FFormatNamedArguments Arguments;
+	Arguments.Add(TEXT("Status"), GetStatusText(Status));
+
+	if (Status == EStatus::Disabled)
+	{
+		if (IsEnabled())
+		{
+			Arguments.Add(TEXT("Suggestion"), LOCTEXT("DisableSuggestionLabel_ReadWrite", "\nThe tree will never run. Click on this status icon to enable the transition tree."));
+		}
+		else
+		{
+			Arguments.Add(TEXT("Suggestion"), LOCTEXT("DisableSuggestionLabel_ReadOnly", "\nThe tree will never run."));
+		}
+	}
+	else
+	{
+		Arguments.Add(TEXT("Suggestion"), FText::GetEmpty());
+	}
+
+	return FText::Format(LOCTEXT("StatusTooltipTextFormat", "Tree is {Status}.{Suggestion}"), Arguments);
+}
+
+const FSlateBrush* SAvaTransitionTreeStatus::GetStatusIcon() const
+{
+	EStatus Status = GetStatus();
+	switch (Status)
+	{
+	case EStatus::Invalid:
+		return FAppStyle::Get().GetBrush(TEXT("Icons.ErrorWithColor"));
+
+	case EStatus::Disabled:
+		return FAppStyle::Get().GetBrush(TEXT("Icons.WarningWithColor"));
+
+	case EStatus::Enabled:
+		return FAppStyle::Get().GetBrush(TEXT("Icons.SuccessWithColor"));
+	}
+
+	checkNoEntry();
+	return nullptr;
+}
+
+#undef LOCTEXT_NAMESPACE
diff --git a/Engine/Plugins/Experimental/Avalanche/Source/AvalancheTransitionEditor/Private/Menu/Widgets/AvaTransitionTreeStatus.h b/Engine/Plugins/Experimental/Avalanche/Source/AvalancheTransitionEditor/Private/Menu/Widgets/AvaTransitionTreeStatus.h
new file mode 100644
--- /dev/null
+++ b/Engine/Plugins/Experimental/Avalanche/Source/AvalancheTransitionEditor/Private/Menu/Widgets/AvaTransitionTreeStatus.h
@@ -0,0 +1,39 @@
+// Copyright Epic Games, Inc. All Rights Reserved.
+
+#pragma once
+
+#include "UObject/WeakObjectPtr.h"
+#include "Widgets/SCompoundWidget.h"
+
+class UAvaTransitionTree;
+
+class SAvaTransitionTreeStatus : public SCompoundWidget
+{
+	enum class EStatus : uint8
+	{
+		Invalid,
+		Disabled,
+		Enabled,
+	};
+
+public:
+	SLATE_BEGIN_ARGS(SAvaTransitionTreeStatus) {}
+	SLATE_END_ARGS()
+
+	void Construct(const FArguments& InArgs, UAvaTransitionTree* InTransitionTree);
+
+private:
+	bool IsStatusButtonEnabled() const;
+
+	FReply OnStatusButtonClicked();
+
+	EStatus GetStatus() const;
+
+	FText GetStatusText(EStatus InStatus) const;
+
+	FText GetStatusTooltipText() const;
+
+	const FSlateBrush* GetStatusIcon() const;
+
+	TWeakObjectPtr<UAvaTransitionTree> TransitionTreeWeak;
+};
