[Motion Design] Playback graph asset creation is disabled by default.

Can be enabled with a cvar.

also:
- Moved Playback Graph and Rundown factories to private.
- Added a class header comment to playback graph so tooltip is updated.

#rnx

diff --git a/Engine/Plugins/Experimental/Avalanche/Source/AvalancheMedia/Public/Playback/AvaPlaybackGraph.h b/Engine/Plugins/Experimental/Avalanche/Source/AvalancheMedia/Public/Playback/AvaPlaybackGraph.h
--- a/Engine/Plugins/Experimental/Avalanche/Source/AvalancheMedia/Public/Playback/AvaPlaybackGraph.h
+++ b/Engine/Plugins/Experimental/Avalanche/Source/AvalancheMedia/Public/Playback/AvaPlaybackGraph.h
@@ -136,6 +136,11 @@ struct AVALANCHEMEDIA_API FAvaPlaybackPlayableGroup
 	}
 };
 
+/**
+ * A Playback Graph is used for playing Motion Design assets integrated with the broadcast framework.
+ * It allows the creation of a playback graph with some logic and inputs routed to player nodes, the results of which can be routed
+ * to broadcast channels. This is the lowest implementation layer that supports distributed playback (over message bus).
+ */
 UCLASS(NotBlueprintable, BlueprintType, ClassGroup = "Motion Design Playback", 
 	meta = (DisplayName = "Motion Design Playback Graph"))
 class AVALANCHEMEDIA_API UAvaPlaybackGraph : public UObject
diff --git a/Engine/Plugins/Experimental/Avalanche/Source/AvalancheMediaEditor/Private/Playback/AvaPlaybackGraphFactory.cpp b/Engine/Plugins/Experimental/Avalanche/Source/AvalancheMediaEditor/Private/Playback/AvaPlaybackGraphFactory.cpp
--- a/Engine/Plugins/Experimental/Avalanche/Source/AvalancheMediaEditor/Private/Playback/AvaPlaybackGraphFactory.cpp
+++ b/Engine/Plugins/Experimental/Avalanche/Source/AvalancheMediaEditor/Private/Playback/AvaPlaybackGraphFactory.cpp
@@ -8,6 +8,21 @@
 
 #define LOCTEXT_NAMESPACE "AvaPlaybackGraphFactory"
 
+namespace UE::AvaMediaEditor::PlaybackGraphFactory::Private
+{
+	/** 
+	 * CVar to specify if new playback graph assets can be created.
+	 * Default is false.
+	 */
+	bool bEnableCreateNewPlaybackGraph = false;
+	FAutoConsoleVariableRef CVarEnableCreateNewPlaybackGraph(
+		TEXT("MotionDesignPlaybackGraphFactory.EnableCreateNew"),
+		bEnableCreateNewPlaybackGraph,
+		TEXT("Specify if new playback graph assets can be created."),
+		ECVF_Default
+	);
+}
+
 UAvaPlaybackGraphFactory::UAvaPlaybackGraphFactory()
 {
 	// Provide the factory with information about how to handle our asset
@@ -20,26 +35,15 @@ UAvaPlaybackGraphFactory::~UAvaPlaybackGraphFactory()
 {
 }
 
-uint32 UAvaPlaybackGraphFactory::GetMenuCategories() const
-{
-	IAssetTools& AssetTools = FAssetToolsModule::GetModule().Get();
-	return AssetTools.FindAdvancedAssetCategory("MotionDesignCategory");
-}
-
-bool UAvaPlaybackGraphFactory::ConfigureProperties()
-{
-	return Super::ConfigureProperties();
-}
-
-bool UAvaPlaybackGraphFactory::ShouldShowInNewMenu() const
+bool UAvaPlaybackGraphFactory::CanCreateNew() const
 {
-	return Super::ShouldShowInNewMenu();
+	return UE::AvaMediaEditor::PlaybackGraphFactory::Private::bEnableCreateNewPlaybackGraph;	
 }
 
-UObject* UAvaPlaybackGraphFactory::FactoryCreateNew(UClass* Class, UObject* InParent, FName Name, EObjectFlags Flags
-  , UObject* Context, FFeedbackContext* Warn, FName CallingContext)
+uint32 UAvaPlaybackGraphFactory::GetMenuCategories() const
 {
-	return Super::FactoryCreateNew(Class, InParent, Name, Flags, Context, Warn, CallingContext);
+	const IAssetTools& AssetTools = FAssetToolsModule::GetModule().Get();
+	return AssetTools.FindAdvancedAssetCategory("MotionDesignCategory");
 }
 
 UObject* UAvaPlaybackGraphFactory::FactoryCreateNew(UClass* Class, UObject* InParent, FName Name, EObjectFlags Flags
@@ -53,16 +57,6 @@ UObject* UAvaPlaybackGraphFactory::FactoryCreateNew(UClass* Class, UObject* InPa
 	return Playback;
 }
 
-bool UAvaPlaybackGraphFactory::DoesSupportClass(UClass* Class)
-{
-	return Class == UAvaPlaybackGraph::StaticClass();
-}
-
-UClass* UAvaPlaybackGraphFactory::ResolveSupportedClass()
-{
-	return Super::ResolveSupportedClass();
-}
-
 FString UAvaPlaybackGraphFactory::GetDefaultNewAssetName() const
 {
 	return TEXT("NewPlaybackGraph");
diff --git a/Engine/Plugins/Experimental/Avalanche/Source/AvalancheMediaEditor/Private/Playback/AvaPlaybackGraphFactory.h b/Engine/Plugins/Experimental/Avalanche/Source/AvalancheMediaEditor/Private/Playback/AvaPlaybackGraphFactory.h
new file mode 100644
--- /dev/null
+++ b/Engine/Plugins/Experimental/Avalanche/Source/AvalancheMediaEditor/Private/Playback/AvaPlaybackGraphFactory.h
@@ -0,0 +1,26 @@
+// Copyright Epic Games, Inc. All Rights Reserved.
+
+#pragma once
+
+#include "Factories/Factory.h"
+#include "AvaPlaybackGraphFactory.generated.h"
+
+UCLASS()
+class UAvaPlaybackGraphFactory : public UFactory
+{
+	GENERATED_BODY()
+
+public:
+
+	UAvaPlaybackGraphFactory();
+	virtual ~UAvaPlaybackGraphFactory() override;
+
+protected:
+	
+	//~ Begin UFactory Interface
+	virtual bool CanCreateNew() const override;
+	virtual uint32 GetMenuCategories() const override;
+	virtual UObject* FactoryCreateNew(UClass* Class, UObject* InParent, FName Name, EObjectFlags Flags, UObject* Context, FFeedbackContext* Warn) override;
+	virtual FString GetDefaultNewAssetName() const override;
+	//~ Begin UFactory Interface
+};
diff --git a/Engine/Plugins/Experimental/Avalanche/Source/AvalancheMediaEditor/Private/Rundown/AvaRundownFactory.cpp b/Engine/Plugins/Experimental/Avalanche/Source/AvalancheMediaEditor/Private/Rundown/AvaRundownFactory.cpp
--- a/Engine/Plugins/Experimental/Avalanche/Source/AvalancheMediaEditor/Private/Rundown/AvaRundownFactory.cpp
+++ b/Engine/Plugins/Experimental/Avalanche/Source/AvalancheMediaEditor/Private/Rundown/AvaRundownFactory.cpp
@@ -20,26 +20,10 @@ UAvaRundownFactory::~UAvaRundownFactory()
 
 uint32 UAvaRundownFactory::GetMenuCategories() const
 {
-	IAssetTools& AssetTools = FAssetToolsModule::GetModule().Get();
+	const IAssetTools& AssetTools = FAssetToolsModule::GetModule().Get();
 	return AssetTools.FindAdvancedAssetCategory("MotionDesignCategory");
 }
 
-bool UAvaRundownFactory::ConfigureProperties()
-{
-	return Super::ConfigureProperties();
-}
-
-bool UAvaRundownFactory::ShouldShowInNewMenu() const
-{
-	return Super::ShouldShowInNewMenu();
-}
-
-UObject* UAvaRundownFactory::FactoryCreateNew(UClass* Class, UObject* InParent, FName Name, EObjectFlags Flags
-	, UObject* Context, FFeedbackContext* Warn, FName CallingContext)
-{
-	return Super::FactoryCreateNew(Class, InParent, Name, Flags, Context, Warn, CallingContext);
-}
-
 UObject* UAvaRundownFactory::FactoryCreateNew(UClass* Class, UObject* InParent, FName Name, EObjectFlags Flags
 	, UObject* Context, FFeedbackContext* Warn)
 {
@@ -51,16 +35,6 @@ UObject* UAvaRundownFactory::FactoryCreateNew(UClass* Class, UObject* InParent,
 	return Rundown;
 }
 
-bool UAvaRundownFactory::DoesSupportClass(UClass* Class)
-{
-	return Class == UAvaRundown::StaticClass();
-}
-
-UClass* UAvaRundownFactory::ResolveSupportedClass()
-{
-	return Super::ResolveSupportedClass();
-}
-
 FString UAvaRundownFactory::GetDefaultNewAssetName() const
 {
 	return TEXT("NewRundown");
diff --git a/Engine/Plugins/Experimental/Avalanche/Source/AvalancheMediaEditor/Private/Rundown/AvaRundownFactory.h b/Engine/Plugins/Experimental/Avalanche/Source/AvalancheMediaEditor/Private/Rundown/AvaRundownFactory.h
new file mode 100644
--- /dev/null
+++ b/Engine/Plugins/Experimental/Avalanche/Source/AvalancheMediaEditor/Private/Rundown/AvaRundownFactory.h
@@ -0,0 +1,25 @@
+// Copyright Epic Games, Inc. All Rights Reserved.
+
+#pragma once
+
+#include "Factories/Factory.h"
+#include "AvaRundownFactory.generated.h"
+
+UCLASS()
+class UAvaRundownFactory : public UFactory
+{
+	GENERATED_BODY()
+
+public:
+
+	UAvaRundownFactory();
+	virtual ~UAvaRundownFactory() override;
+
+protected:
+	
+	//~ Begin UFactory Interface
+	virtual uint32 GetMenuCategories() const override;
+	virtual UObject* FactoryCreateNew(UClass* Class, UObject* InParent, FName Name, EObjectFlags Flags, UObject* Context, FFeedbackContext* Warn) override;
+	virtual FString GetDefaultNewAssetName() const override;
+	//~ Begin UFactory Interface
+};
diff --git a/Engine/Plugins/Experimental/Avalanche/Source/AvalancheMediaEditor/Public/Playback/AvaPlaybackGraphFactory.h b/Engine/Plugins/Experimental/Avalanche/Source/AvalancheMediaEditor/Public/Playback/AvaPlaybackGraphFactory.h
deleted file mode 100644
diff --git a/Engine/Plugins/Experimental/Avalanche/Source/AvalancheMediaEditor/Public/Rundown/AvaRundownFactory.h b/Engine/Plugins/Experimental/Avalanche/Source/AvalancheMediaEditor/Public/Rundown/AvaRundownFactory.h
deleted file mode 100644
