[Avalanche/AvaText3DComponent/Text3DComponent]
Fix Alignment issue with AvaText3DComponent

diff --git a/Engine/Plugins/Experimental/Avalanche/Source/AvalancheText/Private/AvaText3DComponent.cpp b/Engine/Plugins/Experimental/Avalanche/Source/AvalancheText/Private/AvaText3DComponent.cpp
--- a/Engine/Plugins/Experimental/Avalanche/Source/AvalancheText/Private/AvaText3DComponent.cpp
+++ b/Engine/Plugins/Experimental/Avalanche/Source/AvalancheText/Private/AvaText3DComponent.cpp
@@ -135,17 +135,15 @@ void UAvaText3DComponent::PostEditImport()
 #if WITH_EDITOR
 void UAvaText3DComponent::PostEditChangeProperty(FPropertyChangedEvent& PropertyChangedEvent)
 {
-	Super::PostEditChangeProperty(PropertyChangedEvent);
-	
-	if (PropertyChangedEvent.Property)
-	{
-		const FName ChangedPropertyName = PropertyChangedEvent.MemberProperty->GetFName();
+	const FName ChangedPropertyName = PropertyChangedEvent.GetMemberPropertyName();
 
-		if (PropertyChangeFunctions.Contains(ChangedPropertyName))
-		{
-			PropertyChangeFunctions[ChangedPropertyName](this);
-		}
+	if (PropertyChangeFunctions.Contains(ChangedPropertyName))
+	{
+		PropertyChangeFunctions[ChangedPropertyName](this);
 	}
+
+	// Call the Super::PostEditChangeProperty after the PropertyChangeFunctions execute the right Setter/Refresh
+	Super::PostEditChangeProperty(PropertyChangedEvent);
 }
 
 void UAvaText3DComponent::PostEditComponentMove(bool bFinished)
@@ -503,17 +501,7 @@ void UAvaText3DComponent::SetupTextLayout()
 {
 	if (IsMarkedForLayoutRefresh())
 	{
-		SetKerning(Kerning);
-		SetWordSpacing(WordSpacing);
-		SetLineSpacing(LineSpacing);
-		SetScaleProportionally(bScaleProportionally);		
-		SetMaxWidth(MaxWidth);
-		SetMaxHeight(MaxHeight);
-		SetHasMaxWidth(bHasMaxWidth);
-		SetHasMaxHeight(bHasMaxHeight);
-		SetHorizontalAlignment(HorizontalAlignment);
-		SetVerticalAlignment(VerticalAlignment);
-
+		TriggerInternalRebuild(EText3DModifyFlags::Layout);
 		RemoveRefreshReason(EAvaTextRefreshReason::LayoutChange);
 	}
 }
diff --git a/Engine/Plugins/Experimental/Text3D/Source/Text3D/Private/Text3DComponent.cpp b/Engine/Plugins/Experimental/Text3D/Source/Text3D/Private/Text3DComponent.cpp
--- a/Engine/Plugins/Experimental/Text3D/Source/Text3D/Private/Text3DComponent.cpp
+++ b/Engine/Plugins/Experimental/Text3D/Source/Text3D/Private/Text3DComponent.cpp
@@ -1105,7 +1105,7 @@ void UText3DComponent::TriggerInternalRebuild(const EText3DModifyFlags InModifyF
 		MarkForGeometryUpdate();
 	}
 
-	if (EnumHasAnyFlags(ModifyFlags, EText3DModifyFlags::Layout))
+	if (EnumHasAnyFlags(InModifyFlags, EText3DModifyFlags::Layout))
 	{
 		MarkForLayoutUpdate();
 	}
