[Avalanche/Rundown]
Add filter for ComboPage.

diff --git a/Engine/Plugins/Experimental/Avalanche/Source/AvalancheMediaEditor/Private/AvaMediaEditorModule.cpp b/Engine/Plugins/Experimental/Avalanche/Source/AvalancheMediaEditor/Private/AvaMediaEditorModule.cpp
--- a/Engine/Plugins/Experimental/Avalanche/Source/AvalancheMediaEditor/Private/AvaMediaEditorModule.cpp
+++ b/Engine/Plugins/Experimental/Avalanche/Source/AvalancheMediaEditor/Private/AvaMediaEditorModule.cpp
@@ -20,6 +20,7 @@
 #include "Rundown/Customization/AvaRundownMacroKeyBindingCustomization.h"
 #include "Rundown/Factories/Filters/AvaRundownFactoriesUtils.h"
 #include "Rundown/Factories/Filters/AvaRundownFilterChannelExpressionFactory.h"
+#include "Rundown/Factories/Filters/AvaRundownFilterComboPageExpressionFactory.h"
 #include "Rundown/Factories/Filters/AvaRundownFilterIdExpressionFactory.h"
 #include "Rundown/Factories/Filters/AvaRundownFilterNameExpressionFactory.h"
 #include "Rundown/Factories/Filters/AvaRundownFilterPathExpressionFactory.h"
@@ -28,6 +29,7 @@
 #include "Rundown/Factories/Filters/IAvaRundownFilterExpressionFactory.h"
 #include "Rundown/Factories/Filters/IAvaRundownFilterSuggestionFactory.h"
 #include "Rundown/Factories/Suggestions/AvaRundownFilterChannelSuggestionFactory.h"
+#include "Rundown/Factories/Suggestions/AvaRundownFilterComboPageSuggestionFactory.h"
 #include "Rundown/Factories/Suggestions/AvaRundownFilterIdSuggestionFactory.h"
 #include "Rundown/Factories/Suggestions/AvaRundownFilterNameSuggestionFactory.h"
 #include "Rundown/Factories/Suggestions/AvaRundownFilterPathSuggestionFactory.h"
@@ -443,20 +445,22 @@ void FAvaMediaEditorModule::RegisterRundownFilterSuggestionFactory(InArgsType&&.
 
 void FAvaMediaEditorModule::RegisterRundownFilterExpressionFactories()
 {
-	RegisterRundownFilterExpressionFactory<FAvaRundownFilterNameExpressionFactory>();
+	RegisterRundownFilterExpressionFactory<FAvaRundownFilterChannelExpressionFactory>();
+	RegisterRundownFilterExpressionFactory<FAvaRundownFilterComboPageExpressionFactory>();
 	RegisterRundownFilterExpressionFactory<FAvaRundownFilterIdExpressionFactory>();
+	RegisterRundownFilterExpressionFactory<FAvaRundownFilterNameExpressionFactory>();
 	RegisterRundownFilterExpressionFactory<FAvaRundownFilterPathExpressionFactory>();
-	RegisterRundownFilterExpressionFactory<FAvaRundownFilterChannelExpressionFactory>();
 	RegisterRundownFilterExpressionFactory<FAvaRundownFilterStatusExpressionFactory>();
 	RegisterRundownFilterExpressionFactory<FAvaRundownFilterTransitionLayerExpressionFactory>();
 }
 
 void FAvaMediaEditorModule::RegisterRundownFilterSuggestionFactories()
 {
-	RegisterRundownFilterSuggestionFactory<FAvaRundownFilterNameSuggestionFactory>();
+	RegisterRundownFilterSuggestionFactory<FAvaRundownFilterChannelSuggestionFactory>();
+	RegisterRundownFilterSuggestionFactory<FAvaRundownFilterComboPageSuggestionFactory>();
 	RegisterRundownFilterSuggestionFactory<FAvaRundownFilterIdSuggestionFactory>();
+	RegisterRundownFilterSuggestionFactory<FAvaRundownFilterNameSuggestionFactory>();
 	RegisterRundownFilterSuggestionFactory<FAvaRundownFilterPathSuggestionFactory>();
-	RegisterRundownFilterSuggestionFactory<FAvaRundownFilterChannelSuggestionFactory>();
 	RegisterRundownFilterSuggestionFactory<FAvaRundownFilterStatusSuggestionFactory>();
 	RegisterRundownFilterSuggestionFactory<FAvaRundownFilterTransitionLayerSuggestionFactory>();
 }
diff --git a/Engine/Plugins/Experimental/Avalanche/Source/AvalancheMediaEditor/Private/Rundown/Factories/Filters/AvaRundownFilterChannelExpressionFactory.cpp b/Engine/Plugins/Experimental/Avalanche/Source/AvalancheMediaEditor/Private/Rundown/Factories/Filters/AvaRundownFilterChannelExpressionFactory.cpp
--- a/Engine/Plugins/Experimental/Avalanche/Source/AvalancheMediaEditor/Private/Rundown/Factories/Filters/AvaRundownFilterChannelExpressionFactory.cpp
+++ b/Engine/Plugins/Experimental/Avalanche/Source/AvalancheMediaEditor/Private/Rundown/Factories/Filters/AvaRundownFilterChannelExpressionFactory.cpp
@@ -6,10 +6,14 @@
 #include "Rundown/AvaRundownPage.h"
 #include "Rundown/Factories/Filters/AvaRundownFactoriesUtils.h"
 
-const FName FAvaRundownFilterChannelExpressionFactory::KeyName = FName(TEXT("CHANNEL"));
+namespace UE::AvaMediaEditor::Expression::Channel::Private
+{
+	static const FName KeyName = FName(TEXT("CHANNEL"));
+}
 
 FName FAvaRundownFilterChannelExpressionFactory::GetFilterIdentifier() const
 {
+	using namespace UE::AvaMediaEditor::Expression::Channel::Private;
 	return KeyName;
 }
 
diff --git a/Engine/Plugins/Experimental/Avalanche/Source/AvalancheMediaEditor/Private/Rundown/Factories/Filters/AvaRundownFilterChannelExpressionFactory.h b/Engine/Plugins/Experimental/Avalanche/Source/AvalancheMediaEditor/Private/Rundown/Factories/Filters/AvaRundownFilterChannelExpressionFactory.h
--- a/Engine/Plugins/Experimental/Avalanche/Source/AvalancheMediaEditor/Private/Rundown/Factories/Filters/AvaRundownFilterChannelExpressionFactory.h
+++ b/Engine/Plugins/Experimental/Avalanche/Source/AvalancheMediaEditor/Private/Rundown/Factories/Filters/AvaRundownFilterChannelExpressionFactory.h
@@ -7,8 +7,6 @@
 class FAvaRundownFilterChannelExpressionFactory : public IAvaRundownFilterExpressionFactory
 {
 public:
-	static const FName KeyName;
-
 	//~ Begin IAvaFilterExpressionFactory interface
 	virtual FName GetFilterIdentifier() const override;
 	virtual bool FilterExpression(const FAvaRundownPage& InItem, const FAvaRundownTextFilterArgs& InArgs) const override;
diff --git a/Engine/Plugins/Experimental/Avalanche/Source/AvalancheMediaEditor/Private/Rundown/Factories/Filters/AvaRundownFilterComboPageExpressionFactory.cpp b/Engine/Plugins/Experimental/Avalanche/Source/AvalancheMediaEditor/Private/Rundown/Factories/Filters/AvaRundownFilterComboPageExpressionFactory.cpp
new file mode 100644
--- /dev/null
+++ b/Engine/Plugins/Experimental/Avalanche/Source/AvalancheMediaEditor/Private/Rundown/Factories/Filters/AvaRundownFilterComboPageExpressionFactory.cpp
@@ -0,0 +1,35 @@
+// Copyright Epic Games, Inc. All Rights Reserved.
+
+#include "AvaRundownFilterComboPageExpressionFactory.h"
+
+#include "Misc/TextFilterUtils.h"
+#include "Rundown/AvaRundownPage.h"
+
+namespace UE::AvaMediaEditor::Expression::ComboPage::Private
+{
+	static const FName KeyName = FName(TEXT("ISCOMBOPAGE"));
+}
+
+FName FAvaRundownFilterComboPageExpressionFactory::GetFilterIdentifier() const
+{
+	using namespace UE::AvaMediaEditor::Expression::ComboPage::Private;
+	return KeyName;
+}
+
+bool FAvaRundownFilterComboPageExpressionFactory::FilterExpression(const FAvaRundownPage& InItem, const FAvaRundownTextFilterArgs& InArgs) const
+{
+	if (InArgs.ItemRundown)
+	{
+		const FAvaRundownPage& TemplatePage = InItem.ResolveTemplate(InArgs.ItemRundown);
+		const bool bIsPageComboTemplate = TemplatePage.IsComboTemplate();
+		const bool bDisplayComboTemplate = InArgs.ValueToCheck.CompareName(NAME_TRUE, InArgs.ComparisonMode);
+		return bDisplayComboTemplate ? bIsPageComboTemplate : !bIsPageComboTemplate;
+	}
+	// if there is something wrong just display everything
+	return true;
+}
+
+bool FAvaRundownFilterComboPageExpressionFactory::SupportsComparisonOperation(ETextFilterComparisonOperation InComparisonOperation, EAvaRundownSearchListType InPlaylistSearchListType) const
+{
+	return InComparisonOperation == ETextFilterComparisonOperation::Equal || InComparisonOperation == ETextFilterComparisonOperation::NotEqual;
+}
\ No newline at end of file
diff --git a/Engine/Plugins/Experimental/Avalanche/Source/AvalancheMediaEditor/Private/Rundown/Factories/Filters/AvaRundownFilterComboPageExpressionFactory.h b/Engine/Plugins/Experimental/Avalanche/Source/AvalancheMediaEditor/Private/Rundown/Factories/Filters/AvaRundownFilterComboPageExpressionFactory.h
new file mode 100644
--- /dev/null
+++ b/Engine/Plugins/Experimental/Avalanche/Source/AvalancheMediaEditor/Private/Rundown/Factories/Filters/AvaRundownFilterComboPageExpressionFactory.h
@@ -0,0 +1,15 @@
+// Copyright Epic Games, Inc. All Rights Reserved.
+
+#pragma once
+
+#include "Rundown/Factories/Filters/IAvaRundownFilterExpressionFactory.h"
+
+class FAvaRundownFilterComboPageExpressionFactory : public IAvaRundownFilterExpressionFactory
+{
+public:
+	//~ Begin IAvaFilterExpressionFactory interface
+	virtual FName GetFilterIdentifier() const override;
+	virtual bool FilterExpression(const FAvaRundownPage& InItem, const FAvaRundownTextFilterArgs& InArgs) const override;
+	virtual bool SupportsComparisonOperation(ETextFilterComparisonOperation InComparisonOperation, EAvaRundownSearchListType InPlaylistSearchListType) const override;
+	//~ End IAvaFilterExpressionFactory interface
+};
diff --git a/Engine/Plugins/Experimental/Avalanche/Source/AvalancheMediaEditor/Private/Rundown/Factories/Filters/AvaRundownFilterIdExpressionFactory.cpp b/Engine/Plugins/Experimental/Avalanche/Source/AvalancheMediaEditor/Private/Rundown/Factories/Filters/AvaRundownFilterIdExpressionFactory.cpp
--- a/Engine/Plugins/Experimental/Avalanche/Source/AvalancheMediaEditor/Private/Rundown/Factories/Filters/AvaRundownFilterIdExpressionFactory.cpp
+++ b/Engine/Plugins/Experimental/Avalanche/Source/AvalancheMediaEditor/Private/Rundown/Factories/Filters/AvaRundownFilterIdExpressionFactory.cpp
@@ -5,10 +5,14 @@
 #include "Misc/TextFilterUtils.h"
 #include "Rundown/AvaRundownPage.h"
 
-const FName FAvaRundownFilterIdExpressionFactory::KeyName = FName(TEXT("ID"));
+namespace UE::AvaMediaEditor::Expression::Id::Private
+{
+	static const FName KeyName = FName(TEXT("ID"));
+}
 
 FName FAvaRundownFilterIdExpressionFactory::GetFilterIdentifier() const
 {
+	using namespace UE::AvaMediaEditor::Expression::Id::Private;
 	return KeyName;
 }
 
diff --git a/Engine/Plugins/Experimental/Avalanche/Source/AvalancheMediaEditor/Private/Rundown/Factories/Filters/AvaRundownFilterIdExpressionFactory.h b/Engine/Plugins/Experimental/Avalanche/Source/AvalancheMediaEditor/Private/Rundown/Factories/Filters/AvaRundownFilterIdExpressionFactory.h
--- a/Engine/Plugins/Experimental/Avalanche/Source/AvalancheMediaEditor/Private/Rundown/Factories/Filters/AvaRundownFilterIdExpressionFactory.h
+++ b/Engine/Plugins/Experimental/Avalanche/Source/AvalancheMediaEditor/Private/Rundown/Factories/Filters/AvaRundownFilterIdExpressionFactory.h
@@ -7,8 +7,6 @@
 class FAvaRundownFilterIdExpressionFactory : public IAvaRundownFilterExpressionFactory
 {
 public:
-	static const FName KeyName;
-
 	//~ Begin IAvaFilterExpressionFactory interface
 	virtual FName GetFilterIdentifier() const override;
 	virtual bool FilterExpression(const FAvaRundownPage& InItem, const FAvaRundownTextFilterArgs& InArgs) const override;
diff --git a/Engine/Plugins/Experimental/Avalanche/Source/AvalancheMediaEditor/Private/Rundown/Factories/Filters/AvaRundownFilterNameExpressionFactory.cpp b/Engine/Plugins/Experimental/Avalanche/Source/AvalancheMediaEditor/Private/Rundown/Factories/Filters/AvaRundownFilterNameExpressionFactory.cpp
--- a/Engine/Plugins/Experimental/Avalanche/Source/AvalancheMediaEditor/Private/Rundown/Factories/Filters/AvaRundownFilterNameExpressionFactory.cpp
+++ b/Engine/Plugins/Experimental/Avalanche/Source/AvalancheMediaEditor/Private/Rundown/Factories/Filters/AvaRundownFilterNameExpressionFactory.cpp
@@ -5,10 +5,14 @@
 #include "Misc/TextFilterUtils.h"
 #include "Rundown/AvaRundownPage.h"
 
-const FName FAvaRundownFilterNameExpressionFactory::KeyName = FName(TEXT("NAME"));
+namespace UE::AvaMediaEditor::Expression::Name::Private
+{
+	static const FName KeyName = FName(TEXT("NAME"));
+}
 
 FName FAvaRundownFilterNameExpressionFactory::GetFilterIdentifier() const
 {
+	using namespace UE::AvaMediaEditor::Expression::Name::Private;
 	return KeyName;
 }
 
diff --git a/Engine/Plugins/Experimental/Avalanche/Source/AvalancheMediaEditor/Private/Rundown/Factories/Filters/AvaRundownFilterNameExpressionFactory.h b/Engine/Plugins/Experimental/Avalanche/Source/AvalancheMediaEditor/Private/Rundown/Factories/Filters/AvaRundownFilterNameExpressionFactory.h
--- a/Engine/Plugins/Experimental/Avalanche/Source/AvalancheMediaEditor/Private/Rundown/Factories/Filters/AvaRundownFilterNameExpressionFactory.h
+++ b/Engine/Plugins/Experimental/Avalanche/Source/AvalancheMediaEditor/Private/Rundown/Factories/Filters/AvaRundownFilterNameExpressionFactory.h
@@ -7,8 +7,6 @@
 class FAvaRundownFilterNameExpressionFactory : public IAvaRundownFilterExpressionFactory
 {
 public:
-	static const FName KeyName;
-
 	//~ Begin IAvaFilterExpressionFactory interface
 	virtual FName GetFilterIdentifier() const override;
 	virtual bool FilterExpression(const FAvaRundownPage& InItem, const FAvaRundownTextFilterArgs& InArgs) const override;
diff --git a/Engine/Plugins/Experimental/Avalanche/Source/AvalancheMediaEditor/Private/Rundown/Factories/Filters/AvaRundownFilterPathExpressionFactory.cpp b/Engine/Plugins/Experimental/Avalanche/Source/AvalancheMediaEditor/Private/Rundown/Factories/Filters/AvaRundownFilterPathExpressionFactory.cpp
--- a/Engine/Plugins/Experimental/Avalanche/Source/AvalancheMediaEditor/Private/Rundown/Factories/Filters/AvaRundownFilterPathExpressionFactory.cpp
+++ b/Engine/Plugins/Experimental/Avalanche/Source/AvalancheMediaEditor/Private/Rundown/Factories/Filters/AvaRundownFilterPathExpressionFactory.cpp
@@ -5,10 +5,14 @@
 #include "Misc/TextFilterUtils.h"
 #include "Rundown/AvaRundownPage.h"
 
-const FName FAvaRundownFilterPathExpressionFactory::KeyName = FName(TEXT("ASSET"));
+namespace UE::AvaMediaEditor::Expression::Path::Private
+{
+	static const FName KeyName = FName(TEXT("ASSET"));
+}
 
 FName FAvaRundownFilterPathExpressionFactory::GetFilterIdentifier() const
 {
+	using namespace UE::AvaMediaEditor::Expression::Path::Private;
 	return KeyName;
 }
 
diff --git a/Engine/Plugins/Experimental/Avalanche/Source/AvalancheMediaEditor/Private/Rundown/Factories/Filters/AvaRundownFilterPathExpressionFactory.h b/Engine/Plugins/Experimental/Avalanche/Source/AvalancheMediaEditor/Private/Rundown/Factories/Filters/AvaRundownFilterPathExpressionFactory.h
--- a/Engine/Plugins/Experimental/Avalanche/Source/AvalancheMediaEditor/Private/Rundown/Factories/Filters/AvaRundownFilterPathExpressionFactory.h
+++ b/Engine/Plugins/Experimental/Avalanche/Source/AvalancheMediaEditor/Private/Rundown/Factories/Filters/AvaRundownFilterPathExpressionFactory.h
@@ -7,8 +7,6 @@
 class FAvaRundownFilterPathExpressionFactory : public IAvaRundownFilterExpressionFactory
 {
 public:
-	static const FName KeyName;
-
 	//~ Begin IAvaFilterExpressionFactory interface
 	virtual FName GetFilterIdentifier() const override;
 	virtual bool FilterExpression(const FAvaRundownPage& InItem, const FAvaRundownTextFilterArgs& InArgs) const override;
diff --git a/Engine/Plugins/Experimental/Avalanche/Source/AvalancheMediaEditor/Private/Rundown/Factories/Filters/AvaRundownFilterStatusExpressionFactory.cpp b/Engine/Plugins/Experimental/Avalanche/Source/AvalancheMediaEditor/Private/Rundown/Factories/Filters/AvaRundownFilterStatusExpressionFactory.cpp
--- a/Engine/Plugins/Experimental/Avalanche/Source/AvalancheMediaEditor/Private/Rundown/Factories/Filters/AvaRundownFilterStatusExpressionFactory.cpp
+++ b/Engine/Plugins/Experimental/Avalanche/Source/AvalancheMediaEditor/Private/Rundown/Factories/Filters/AvaRundownFilterStatusExpressionFactory.cpp
@@ -5,10 +5,14 @@
 #include "Misc/TextFilterUtils.h"
 #include "Rundown/AvaRundownPage.h"
 
-const FName FAvaRundownFilterStatusExpressionFactory::KeyName = FName(TEXT("STATUS"));
+namespace UE::AvaMediaEditor::Expression::Status::Private
+{
+	static const FName KeyName = FName(TEXT("STATUS"));
+}
 
 FName FAvaRundownFilterStatusExpressionFactory::GetFilterIdentifier() const
 {
+	using namespace UE::AvaMediaEditor::Expression::Status::Private;
 	return KeyName;
 }
 
diff --git a/Engine/Plugins/Experimental/Avalanche/Source/AvalancheMediaEditor/Private/Rundown/Factories/Filters/AvaRundownFilterStatusExpressionFactory.h b/Engine/Plugins/Experimental/Avalanche/Source/AvalancheMediaEditor/Private/Rundown/Factories/Filters/AvaRundownFilterStatusExpressionFactory.h
--- a/Engine/Plugins/Experimental/Avalanche/Source/AvalancheMediaEditor/Private/Rundown/Factories/Filters/AvaRundownFilterStatusExpressionFactory.h
+++ b/Engine/Plugins/Experimental/Avalanche/Source/AvalancheMediaEditor/Private/Rundown/Factories/Filters/AvaRundownFilterStatusExpressionFactory.h
@@ -7,8 +7,6 @@
 class FAvaRundownFilterStatusExpressionFactory : public IAvaRundownFilterExpressionFactory
 {
 public:
-	static const FName KeyName;
-
 	//~ Begin IAvaFilterExpressionFactory interface
 	virtual FName GetFilterIdentifier() const override;
 	virtual bool FilterExpression(const FAvaRundownPage& InItem, const FAvaRundownTextFilterArgs& InArgs) const override;
diff --git a/Engine/Plugins/Experimental/Avalanche/Source/AvalancheMediaEditor/Private/Rundown/Factories/Filters/AvaRundownFilterTransitionLayerExpressionFactory.cpp b/Engine/Plugins/Experimental/Avalanche/Source/AvalancheMediaEditor/Private/Rundown/Factories/Filters/AvaRundownFilterTransitionLayerExpressionFactory.cpp
--- a/Engine/Plugins/Experimental/Avalanche/Source/AvalancheMediaEditor/Private/Rundown/Factories/Filters/AvaRundownFilterTransitionLayerExpressionFactory.cpp
+++ b/Engine/Plugins/Experimental/Avalanche/Source/AvalancheMediaEditor/Private/Rundown/Factories/Filters/AvaRundownFilterTransitionLayerExpressionFactory.cpp
@@ -5,10 +5,14 @@
 #include "Misc/TextFilterUtils.h"
 #include "Rundown/AvaRundownPage.h"
 
-const FName FAvaRundownFilterTransitionLayerExpressionFactory::KeyName = FName(TEXT("TRANSITIONLAYER"));
+namespace UE::AvaMediaEditor::Expression::Layer::Private
+{
+	static const FName KeyName = FName(TEXT("TRANSITIONLAYER"));
+}
 
 FName FAvaRundownFilterTransitionLayerExpressionFactory::GetFilterIdentifier() const
 {
+	using namespace UE::AvaMediaEditor::Expression::Layer::Private;
 	return KeyName;
 }
 
@@ -18,7 +22,7 @@ bool FAvaRundownFilterTransitionLayerExpressionFactory::FilterExpression(const F
 	{
 		return false;
 	}
-	return InArgs.ValueToCheck.CompareFString(InItem.GetTransitionLayer(InArgs.ItemRundown).ToString(), InArgs.ComparisonMode);
+	return InArgs.ValueToCheck.CompareFString(InItem.GetTransitionLayer(InArgs.ItemRundown).ToString().ToUpper(), InArgs.ComparisonMode);
 }
 
 bool FAvaRundownFilterTransitionLayerExpressionFactory::SupportsComparisonOperation(ETextFilterComparisonOperation InComparisonOperation, EAvaRundownSearchListType InRundownSearchListType) const
diff --git a/Engine/Plugins/Experimental/Avalanche/Source/AvalancheMediaEditor/Private/Rundown/Factories/Filters/AvaRundownFilterTransitionLayerExpressionFactory.h b/Engine/Plugins/Experimental/Avalanche/Source/AvalancheMediaEditor/Private/Rundown/Factories/Filters/AvaRundownFilterTransitionLayerExpressionFactory.h
--- a/Engine/Plugins/Experimental/Avalanche/Source/AvalancheMediaEditor/Private/Rundown/Factories/Filters/AvaRundownFilterTransitionLayerExpressionFactory.h
+++ b/Engine/Plugins/Experimental/Avalanche/Source/AvalancheMediaEditor/Private/Rundown/Factories/Filters/AvaRundownFilterTransitionLayerExpressionFactory.h
@@ -7,8 +7,6 @@
 class FAvaRundownFilterTransitionLayerExpressionFactory : public IAvaRundownFilterExpressionFactory
 {
 public:
-	static const FName KeyName;
-
 	//~ Begin IAvaFilterExpressionFactory interface
 	virtual FName GetFilterIdentifier() const override;
 	virtual bool FilterExpression(const FAvaRundownPage& InItem, const FAvaRundownTextFilterArgs& InArgs) const override;
diff --git a/Engine/Plugins/Experimental/Avalanche/Source/AvalancheMediaEditor/Private/Rundown/Factories/Suggestions/AvaRundownFilterChannelSuggestionFactory.cpp b/Engine/Plugins/Experimental/Avalanche/Source/AvalancheMediaEditor/Private/Rundown/Factories/Suggestions/AvaRundownFilterChannelSuggestionFactory.cpp
--- a/Engine/Plugins/Experimental/Avalanche/Source/AvalancheMediaEditor/Private/Rundown/Factories/Suggestions/AvaRundownFilterChannelSuggestionFactory.cpp
+++ b/Engine/Plugins/Experimental/Avalanche/Source/AvalancheMediaEditor/Private/Rundown/Factories/Suggestions/AvaRundownFilterChannelSuggestionFactory.cpp
@@ -8,14 +8,25 @@
 
 #define LOCTEXT_NAMESPACE "AvaRundownFilterChannelSuggestionFactory"
 
-const FName FAvaRundownFilterChannelSuggestionFactory::KeyName = FName(TEXT("CHANNEL"));
+namespace UE::AvaMediaEditor::Suggestion::Channel::Private
+{
+	static const FName KeyName = FName(TEXT("CHANNEL"));
+	static const FText CategoryLabel = LOCTEXT("ChannelCategoryLabel", "Ava-Rundown-Channel");
+}
+
+FName FAvaRundownFilterChannelSuggestionFactory::GetSuggestionIdentifier() const
+{
+	using namespace UE::AvaMediaEditor::Suggestion::Channel::Private;
+	return KeyName;
+}
 
 void FAvaRundownFilterChannelSuggestionFactory::AddSuggestion(const TSharedRef<FAvaRundownFilterSuggestionPayload>& InPayload)
 {
+	using namespace UE::AvaMediaEditor::Suggestion::Channel::Private;
+
 	const FAvaRundownPage& PageItem = UAvaRundown::GetPageSafe(InPayload->Rundown, InPayload->ItemPageId);
 	if (PageItem.IsValidPage())
 	{
-		const FText ChannelCategoryLabel = LOCTEXT("ChannelCategoryLabel", "Ava-Rundown-Channel");
 		const FString ChannelName = PageItem.GetChannelName().ToString();
 
 		FString ChannelNameSuggestion = FString::Printf(TEXT("Channel=%s"), *ChannelName);
@@ -24,7 +35,7 @@ void FAvaRundownFilterChannelSuggestionFactory::AddSuggestion(const TSharedRef<F
 		if (bIsFilterValueValid && !InPayload->FilterCache.Contains(ChannelNameSuggestion))
 		{
 			InPayload->FilterCache.Add(ChannelNameSuggestion);
-			InPayload->PossibleSuggestions.Add(FAssetSearchBoxSuggestion{MoveTemp(ChannelNameSuggestion), FText::FromString(ChannelName), ChannelCategoryLabel});
+			InPayload->PossibleSuggestions.Add(FAssetSearchBoxSuggestion{MoveTemp(ChannelNameSuggestion), FText::FromString(ChannelName), CategoryLabel});
 		}
 	}
 }
diff --git a/Engine/Plugins/Experimental/Avalanche/Source/AvalancheMediaEditor/Private/Rundown/Factories/Suggestions/AvaRundownFilterChannelSuggestionFactory.h b/Engine/Plugins/Experimental/Avalanche/Source/AvalancheMediaEditor/Private/Rundown/Factories/Suggestions/AvaRundownFilterChannelSuggestionFactory.h
--- a/Engine/Plugins/Experimental/Avalanche/Source/AvalancheMediaEditor/Private/Rundown/Factories/Suggestions/AvaRundownFilterChannelSuggestionFactory.h
+++ b/Engine/Plugins/Experimental/Avalanche/Source/AvalancheMediaEditor/Private/Rundown/Factories/Suggestions/AvaRundownFilterChannelSuggestionFactory.h
@@ -7,10 +7,8 @@
 class FAvaRundownFilterChannelSuggestionFactory : public IAvaRundownFilterSuggestionFactory
 {
 public:
-	static const FName KeyName;
-
 	//~ Begin IAvaFilterSuggestionFactory interface
-	virtual FName GetSuggestionIdentifier() const override { return KeyName; }
+	virtual FName GetSuggestionIdentifier() const override;
 	virtual bool IsSimpleSuggestion() const override { return false; }
 	virtual void AddSuggestion(const TSharedRef<FAvaRundownFilterSuggestionPayload>& InPayload) override;
 	virtual bool SupportSuggestionType(EAvaRundownSearchListType InSuggestionType) const override;
diff --git a/Engine/Plugins/Experimental/Avalanche/Source/AvalancheMediaEditor/Private/Rundown/Factories/Suggestions/AvaRundownFilterComboPageSuggestionFactory.cpp b/Engine/Plugins/Experimental/Avalanche/Source/AvalancheMediaEditor/Private/Rundown/Factories/Suggestions/AvaRundownFilterComboPageSuggestionFactory.cpp
new file mode 100644
--- /dev/null
+++ b/Engine/Plugins/Experimental/Avalanche/Source/AvalancheMediaEditor/Private/Rundown/Factories/Suggestions/AvaRundownFilterComboPageSuggestionFactory.cpp
@@ -0,0 +1,49 @@
+// Copyright Epic Games, Inc. All Rights Reserved.
+
+#include "AvaRundownFilterComboPageSuggestionFactory.h"
+
+#include "Rundown/Factories/Filters/AvaRundownFactoriesUtils.h"
+#include "SAssetSearchBox.h"
+
+#define LOCTEXT_NAMESPACE "AvaRundownFilterComboPageSuggestionFactory"
+
+namespace UE::AvaMediaEditor::Suggestion::ComboPage::Private
+{
+	static const FName KeyName = FName(TEXT("ISCOMBOPAGE"));
+	static const FText CategoryLabel = LOCTEXT("IsComboPageCategoryLabel", "Ava-Rundown-IsComboPage");
+	static const FString TrueSuggestion = TEXT("IsComboPage=True");
+	static const FText TrueDisplayName = LOCTEXT("ComboPageSuggestion_True","True");
+	static const FString FalseSuggestion = TEXT("IsComboPage=False");
+	static const FText FalseDisplayName = LOCTEXT("ComboPageSuggestion_False","False");
+}
+
+FName FAvaRundownFilterComboPageSuggestionFactory::GetSuggestionIdentifier() const
+{
+	using namespace UE::AvaMediaEditor::Suggestion::ComboPage::Private;
+	return KeyName;
+}
+
+void FAvaRundownFilterComboPageSuggestionFactory::AddSuggestion(const TSharedRef<FAvaRundownFilterSuggestionPayload>& InPayload)
+{
+	using namespace UE::AvaMediaEditor::Suggestion::ComboPage::Private;
+
+	const bool bIsFilterValueValidTrue = InPayload->FilterValue.IsEmpty() || TrueSuggestion.Contains(InPayload->FilterValue);
+	if (bIsFilterValueValidTrue && !InPayload->FilterCache.Contains(TrueSuggestion))
+	{
+		InPayload->FilterCache.Add(TrueSuggestion);
+		InPayload->PossibleSuggestions.Add(FAssetSearchBoxSuggestion{TrueSuggestion, TrueDisplayName, CategoryLabel});
+	}
+	const bool bIsFilterValueValidFalse = InPayload->FilterValue.IsEmpty() || FalseSuggestion.Contains(InPayload->FilterValue);
+	if (bIsFilterValueValidFalse && !InPayload->FilterCache.Contains(FalseSuggestion))
+	{
+		InPayload->FilterCache.Add(FalseSuggestion);
+		InPayload->PossibleSuggestions.Add(FAssetSearchBoxSuggestion{FalseSuggestion, FalseDisplayName, CategoryLabel});
+	}
+}
+
+bool FAvaRundownFilterComboPageSuggestionFactory::SupportSuggestionType(EAvaRundownSearchListType InSuggestionType) const
+{
+	return InSuggestionType != EAvaRundownSearchListType::None;
+}
+
+#undef LOCTEXT_NAMESPACE
diff --git a/Engine/Plugins/Experimental/Avalanche/Source/AvalancheMediaEditor/Private/Rundown/Factories/Suggestions/AvaRundownFilterComboPageSuggestionFactory.h b/Engine/Plugins/Experimental/Avalanche/Source/AvalancheMediaEditor/Private/Rundown/Factories/Suggestions/AvaRundownFilterComboPageSuggestionFactory.h
new file mode 100644
--- /dev/null
+++ b/Engine/Plugins/Experimental/Avalanche/Source/AvalancheMediaEditor/Private/Rundown/Factories/Suggestions/AvaRundownFilterComboPageSuggestionFactory.h
@@ -0,0 +1,16 @@
+// Copyright Epic Games, Inc. All Rights Reserved.
+
+#pragma once
+
+#include "Rundown/Factories/Filters/IAvaRundownFilterSuggestionFactory.h"
+
+class FAvaRundownFilterComboPageSuggestionFactory : public IAvaRundownFilterSuggestionFactory
+{
+public:
+	//~ Begin IAvaFilterSuggestionFactory interface
+	virtual FName GetSuggestionIdentifier() const override;
+	virtual bool IsSimpleSuggestion() const override { return true; }
+	virtual void AddSuggestion(const TSharedRef<FAvaRundownFilterSuggestionPayload>& InPayload) override;
+	virtual bool SupportSuggestionType(EAvaRundownSearchListType InSuggestionType) const override;
+	//~ End IAvaFilterSuggestionFactory interface
+};
diff --git a/Engine/Plugins/Experimental/Avalanche/Source/AvalancheMediaEditor/Private/Rundown/Factories/Suggestions/AvaRundownFilterIdSuggestionFactory.cpp b/Engine/Plugins/Experimental/Avalanche/Source/AvalancheMediaEditor/Private/Rundown/Factories/Suggestions/AvaRundownFilterIdSuggestionFactory.cpp
--- a/Engine/Plugins/Experimental/Avalanche/Source/AvalancheMediaEditor/Private/Rundown/Factories/Suggestions/AvaRundownFilterIdSuggestionFactory.cpp
+++ b/Engine/Plugins/Experimental/Avalanche/Source/AvalancheMediaEditor/Private/Rundown/Factories/Suggestions/AvaRundownFilterIdSuggestionFactory.cpp
@@ -7,17 +7,29 @@
 
 #define LOCTEXT_NAMESPACE "AvaRundownFilterIdSuggestionFactory"
 
-const FName FAvaRundownFilterIdSuggestionFactory::KeyName = FName(TEXT("ID"));
+namespace UE::AvaMediaEditor::Suggestion::Id::Private
+{
+	static const FName KeyName = FName(TEXT("ID"));
+	static const FText CategoryLabel = LOCTEXT("IdCategoryLabel", "Ava-Rundown-Id");
+	static const FString Suggestion = TEXT("Id");
+	static const FText DisplayName = LOCTEXT("IdSuggestion","Id");
+}
+
+FName FAvaRundownFilterIdSuggestionFactory::GetSuggestionIdentifier() const
+{
+	using namespace UE::AvaMediaEditor::Suggestion::Id::Private;
+	return KeyName;
+}
 
 void FAvaRundownFilterIdSuggestionFactory::AddSuggestion(const TSharedRef<FAvaRundownFilterSuggestionPayload>& InPayload)
 {
-	const FText IdCategoryLabel = LOCTEXT("IdCategoryLabel", "Ava-Rundown-Id");
-	FString IdNameSuggestion = TEXT("Id");
-	const bool bIsFilterValueValid = InPayload->FilterValue.IsEmpty() || IdNameSuggestion.Contains(InPayload->FilterValue);
-	if (bIsFilterValueValid && !InPayload->FilterCache.Contains(IdNameSuggestion))
+	using namespace UE::AvaMediaEditor::Suggestion::Id::Private;
+
+	const bool bIsFilterValueValid = InPayload->FilterValue.IsEmpty() || Suggestion.Contains(InPayload->FilterValue);
+	if (bIsFilterValueValid && !InPayload->FilterCache.Contains(Suggestion))
 	{
-		InPayload->FilterCache.Add(IdNameSuggestion);
-		InPayload->PossibleSuggestions.Add(FAssetSearchBoxSuggestion{MoveTemp(IdNameSuggestion), LOCTEXT("IdSuggestion","Id"), IdCategoryLabel});
+		InPayload->FilterCache.Add(Suggestion);
+		InPayload->PossibleSuggestions.Add(FAssetSearchBoxSuggestion{Suggestion, DisplayName, CategoryLabel});
 	}
 }
 
diff --git a/Engine/Plugins/Experimental/Avalanche/Source/AvalancheMediaEditor/Private/Rundown/Factories/Suggestions/AvaRundownFilterIdSuggestionFactory.h b/Engine/Plugins/Experimental/Avalanche/Source/AvalancheMediaEditor/Private/Rundown/Factories/Suggestions/AvaRundownFilterIdSuggestionFactory.h
--- a/Engine/Plugins/Experimental/Avalanche/Source/AvalancheMediaEditor/Private/Rundown/Factories/Suggestions/AvaRundownFilterIdSuggestionFactory.h
+++ b/Engine/Plugins/Experimental/Avalanche/Source/AvalancheMediaEditor/Private/Rundown/Factories/Suggestions/AvaRundownFilterIdSuggestionFactory.h
@@ -7,10 +7,8 @@
 class FAvaRundownFilterIdSuggestionFactory : public IAvaRundownFilterSuggestionFactory
 {
 public:
-	static const FName KeyName;
-
 	//~ Begin IAvaFilterSuggestionFactory interface
-	virtual FName GetSuggestionIdentifier() const override { return KeyName; }
+	virtual FName GetSuggestionIdentifier() const override;
 	virtual bool IsSimpleSuggestion() const override { return true; }
 	virtual void AddSuggestion(const TSharedRef<FAvaRundownFilterSuggestionPayload>& InPayload) override;
 	virtual bool SupportSuggestionType(EAvaRundownSearchListType InSuggestionType) const override;
diff --git a/Engine/Plugins/Experimental/Avalanche/Source/AvalancheMediaEditor/Private/Rundown/Factories/Suggestions/AvaRundownFilterNameSuggestionFactory.cpp b/Engine/Plugins/Experimental/Avalanche/Source/AvalancheMediaEditor/Private/Rundown/Factories/Suggestions/AvaRundownFilterNameSuggestionFactory.cpp
--- a/Engine/Plugins/Experimental/Avalanche/Source/AvalancheMediaEditor/Private/Rundown/Factories/Suggestions/AvaRundownFilterNameSuggestionFactory.cpp
+++ b/Engine/Plugins/Experimental/Avalanche/Source/AvalancheMediaEditor/Private/Rundown/Factories/Suggestions/AvaRundownFilterNameSuggestionFactory.cpp
@@ -8,14 +8,25 @@
 
 #define LOCTEXT_NAMESPACE "AvaRundownFilterNameSuggestionFactory"
 
-const FName FAvaRundownFilterNameSuggestionFactory::KeyName = FName(TEXT("NAME"));
+namespace UE::AvaMediaEditor::Suggestion::Name::Private
+{
+	static const FName KeyName = FName(TEXT("NAME"));
+	static const FText CategoryLabel = LOCTEXT("NameCategoryLabel", "Ava-Rundown-Name");
+}
+
+FName FAvaRundownFilterNameSuggestionFactory::GetSuggestionIdentifier() const
+{
+	using namespace UE::AvaMediaEditor::Suggestion::Name::Private;
+	return KeyName;
+}
 
 void FAvaRundownFilterNameSuggestionFactory::AddSuggestion(const TSharedRef<FAvaRundownFilterSuggestionPayload>& InPayload)
 {
+	using namespace UE::AvaMediaEditor::Suggestion::Name::Private;
+
 	const FAvaRundownPage& PageItem = UAvaRundown::GetPageSafe(InPayload->Rundown, InPayload->ItemPageId);
 	if (PageItem.IsValidPage())
 	{
-		const FText NameCategoryLabel = LOCTEXT("NameCategoryLabel", "Ava-Rundown-Name");
 		const FText PageText = PageItem.GetPageDescription();
 
 		FString PageName = TEXT("\"");
@@ -28,7 +39,7 @@ void FAvaRundownFilterNameSuggestionFactory::AddSuggestion(const TSharedRef<FAva
 		if (bIsFilterValueValid && !InPayload->FilterCache.Contains(NameSuggestion))
 		{
 			InPayload->FilterCache.Add(NameSuggestion);
-			InPayload->PossibleSuggestions.Add(FAssetSearchBoxSuggestion{MoveTemp(NameSuggestion), PageText, NameCategoryLabel});
+			InPayload->PossibleSuggestions.Add(FAssetSearchBoxSuggestion{MoveTemp(NameSuggestion), PageText, CategoryLabel});
 		}
 	}
 }
diff --git a/Engine/Plugins/Experimental/Avalanche/Source/AvalancheMediaEditor/Private/Rundown/Factories/Suggestions/AvaRundownFilterNameSuggestionFactory.h b/Engine/Plugins/Experimental/Avalanche/Source/AvalancheMediaEditor/Private/Rundown/Factories/Suggestions/AvaRundownFilterNameSuggestionFactory.h
--- a/Engine/Plugins/Experimental/Avalanche/Source/AvalancheMediaEditor/Private/Rundown/Factories/Suggestions/AvaRundownFilterNameSuggestionFactory.h
+++ b/Engine/Plugins/Experimental/Avalanche/Source/AvalancheMediaEditor/Private/Rundown/Factories/Suggestions/AvaRundownFilterNameSuggestionFactory.h
@@ -7,10 +7,8 @@
 class FAvaRundownFilterNameSuggestionFactory : public IAvaRundownFilterSuggestionFactory
 {
 public:
-	static const FName KeyName;
-
 	//~ Begin IAvaFilterSuggestionFactory interface
-	virtual FName GetSuggestionIdentifier() const override { return KeyName; }
+	virtual FName GetSuggestionIdentifier() const override;
 	virtual bool IsSimpleSuggestion() const override { return false; }
 	virtual void AddSuggestion(const TSharedRef<FAvaRundownFilterSuggestionPayload>& InPayload) override;
 	virtual bool SupportSuggestionType(EAvaRundownSearchListType InSuggestionType) const override;
diff --git a/Engine/Plugins/Experimental/Avalanche/Source/AvalancheMediaEditor/Private/Rundown/Factories/Suggestions/AvaRundownFilterPathSuggestionFactory.cpp b/Engine/Plugins/Experimental/Avalanche/Source/AvalancheMediaEditor/Private/Rundown/Factories/Suggestions/AvaRundownFilterPathSuggestionFactory.cpp
--- a/Engine/Plugins/Experimental/Avalanche/Source/AvalancheMediaEditor/Private/Rundown/Factories/Suggestions/AvaRundownFilterPathSuggestionFactory.cpp
+++ b/Engine/Plugins/Experimental/Avalanche/Source/AvalancheMediaEditor/Private/Rundown/Factories/Suggestions/AvaRundownFilterPathSuggestionFactory.cpp
@@ -8,14 +8,25 @@
 
 #define LOCTEXT_NAMESPACE "AvaRundownFilterPathSuggestionFactory"
 
-const FName FAvaRundownFilterPathSuggestionFactory::KeyName = FName(TEXT("ASSET"));
+namespace UE::AvaMediaEditor::Suggestion::Path::Private
+{
+	static const FName KeyName = FName(TEXT("ASSET"));
+	static const FText CategoryLabel = LOCTEXT("PathCategoryLabel", "Ava-Rundown-Asset");
+}
+
+FName FAvaRundownFilterPathSuggestionFactory::GetSuggestionIdentifier() const
+{
+	using namespace UE::AvaMediaEditor::Suggestion::Path::Private;
+	return KeyName;
+}
 
 void FAvaRundownFilterPathSuggestionFactory::AddSuggestion(const TSharedRef<FAvaRundownFilterSuggestionPayload>& InPayload)
 {
+	using namespace UE::AvaMediaEditor::Suggestion::Path::Private;
+
 	const FAvaRundownPage& PageItem = UAvaRundown::GetPageSafe(InPayload->Rundown, InPayload->ItemPageId);
 	if (PageItem.IsValidPage())
 	{
-		const FText PathCategoryLabel = LOCTEXT("PathCategoryLabel", "Ava-Rundown-Asset");
 		const FString AssetName = PageItem.GetAssetPath(InPayload->Rundown).GetAssetName();
 
 		FString AssetNameSuggestion = FString::Printf(TEXT("Asset=%s"), *AssetName);
@@ -24,7 +35,7 @@ void FAvaRundownFilterPathSuggestionFactory::AddSuggestion(const TSharedRef<FAva
 		if (bIsFilterValueValid && !InPayload->FilterCache.Contains(AssetNameSuggestion))
 		{
 			InPayload->FilterCache.Add(AssetNameSuggestion);
-			InPayload->PossibleSuggestions.Add(FAssetSearchBoxSuggestion{MoveTemp(AssetNameSuggestion), FText::FromString(AssetName), PathCategoryLabel});
+			InPayload->PossibleSuggestions.Add(FAssetSearchBoxSuggestion{MoveTemp(AssetNameSuggestion), FText::FromString(AssetName), CategoryLabel});
 		}
 	}
 }
diff --git a/Engine/Plugins/Experimental/Avalanche/Source/AvalancheMediaEditor/Private/Rundown/Factories/Suggestions/AvaRundownFilterPathSuggestionFactory.h b/Engine/Plugins/Experimental/Avalanche/Source/AvalancheMediaEditor/Private/Rundown/Factories/Suggestions/AvaRundownFilterPathSuggestionFactory.h
--- a/Engine/Plugins/Experimental/Avalanche/Source/AvalancheMediaEditor/Private/Rundown/Factories/Suggestions/AvaRundownFilterPathSuggestionFactory.h
+++ b/Engine/Plugins/Experimental/Avalanche/Source/AvalancheMediaEditor/Private/Rundown/Factories/Suggestions/AvaRundownFilterPathSuggestionFactory.h
@@ -7,10 +7,8 @@
 class FAvaRundownFilterPathSuggestionFactory : public IAvaRundownFilterSuggestionFactory
 {
 public:
-	static const FName KeyName;
-
 	//~ Begin IAvaFilterSuggestionFactory interface
-	virtual FName GetSuggestionIdentifier() const override { return KeyName; }
+	virtual FName GetSuggestionIdentifier() const override;
 	virtual bool IsSimpleSuggestion() const override { return false; }
 	virtual void AddSuggestion(const TSharedRef<FAvaRundownFilterSuggestionPayload>& InPayload) override;
 	virtual bool SupportSuggestionType(EAvaRundownSearchListType InSuggestionType) const override;
diff --git a/Engine/Plugins/Experimental/Avalanche/Source/AvalancheMediaEditor/Private/Rundown/Factories/Suggestions/AvaRundownFilterStatusSuggestionFactory.cpp b/Engine/Plugins/Experimental/Avalanche/Source/AvalancheMediaEditor/Private/Rundown/Factories/Suggestions/AvaRundownFilterStatusSuggestionFactory.cpp
--- a/Engine/Plugins/Experimental/Avalanche/Source/AvalancheMediaEditor/Private/Rundown/Factories/Suggestions/AvaRundownFilterStatusSuggestionFactory.cpp
+++ b/Engine/Plugins/Experimental/Avalanche/Source/AvalancheMediaEditor/Private/Rundown/Factories/Suggestions/AvaRundownFilterStatusSuggestionFactory.cpp
@@ -8,14 +8,25 @@
 
 #define LOCTEXT_NAMESPACE "AvaRundownFilterStatusSuggestionFactory"
 
-const FName FAvaRundownFilterStatusSuggestionFactory::KeyName = FName(TEXT("STATUS"));
+namespace UE::AvaMediaEditor::Suggestion::Status::Private
+{
+	static const FName KeyName = FName(TEXT("STATUS"));
+	static const FText CategoryLabel = LOCTEXT("StatusCategoryLabel", "Ava-Rundown-Status");
+}
+
+FName FAvaRundownFilterStatusSuggestionFactory::GetSuggestionIdentifier() const
+{
+	using namespace UE::AvaMediaEditor::Suggestion::Status::Private;
+	return KeyName;
+}
 
 void FAvaRundownFilterStatusSuggestionFactory::AddSuggestion(const TSharedRef<FAvaRundownFilterSuggestionPayload>& InPayload)
 {
+	using namespace UE::AvaMediaEditor::Suggestion::Status::Private;
+
 	const FAvaRundownPage& PageItem = UAvaRundown::GetPageSafe(InPayload->Rundown, InPayload->ItemPageId);
 	if (PageItem.IsValidPage())
 	{
-		const FText StatusCategoryLabel = LOCTEXT("StatusCategoryLabel", "Ava-Rundown-Status");
 		TArray<FAvaRundownChannelPageStatus> StatusPages = PageItem.GetPageContextualStatuses(InPayload->Rundown);
 		for (FAvaRundownChannelPageStatus Status : StatusPages)
 		{
@@ -27,7 +38,7 @@ void FAvaRundownFilterStatusSuggestionFactory::AddSuggestion(const TSharedRef<FA
 			if (bIsFilterValueValid && !InPayload->FilterCache.Contains(StatusNameSuggestion))
 			{
 				InPayload->FilterCache.Add(StatusNameSuggestion);
-				InPayload->PossibleSuggestions.Add(FAssetSearchBoxSuggestion{MoveTemp(StatusNameSuggestion), FText::FromString(StatusName), StatusCategoryLabel});
+				InPayload->PossibleSuggestions.Add(FAssetSearchBoxSuggestion{MoveTemp(StatusNameSuggestion), FText::FromString(StatusName), CategoryLabel});
 			}
 		}
 	}
diff --git a/Engine/Plugins/Experimental/Avalanche/Source/AvalancheMediaEditor/Private/Rundown/Factories/Suggestions/AvaRundownFilterStatusSuggestionFactory.h b/Engine/Plugins/Experimental/Avalanche/Source/AvalancheMediaEditor/Private/Rundown/Factories/Suggestions/AvaRundownFilterStatusSuggestionFactory.h
--- a/Engine/Plugins/Experimental/Avalanche/Source/AvalancheMediaEditor/Private/Rundown/Factories/Suggestions/AvaRundownFilterStatusSuggestionFactory.h
+++ b/Engine/Plugins/Experimental/Avalanche/Source/AvalancheMediaEditor/Private/Rundown/Factories/Suggestions/AvaRundownFilterStatusSuggestionFactory.h
@@ -7,10 +7,8 @@
 class FAvaRundownFilterStatusSuggestionFactory : public IAvaRundownFilterSuggestionFactory
 {
 public:
-	static const FName KeyName;
-
 	//~ Begin IAvaFilterSuggestionFactory interface
-	virtual FName GetSuggestionIdentifier() const override { return KeyName; }
+	virtual FName GetSuggestionIdentifier() const override;
 	virtual bool IsSimpleSuggestion() const override { return false; }
 	virtual void AddSuggestion(const TSharedRef<FAvaRundownFilterSuggestionPayload>& InPayload) override;
 	virtual bool SupportSuggestionType(EAvaRundownSearchListType InSuggestionType) const override;
diff --git a/Engine/Plugins/Experimental/Avalanche/Source/AvalancheMediaEditor/Private/Rundown/Factories/Suggestions/AvaRundownFilterTransitionLayerSuggestionFactory.cpp b/Engine/Plugins/Experimental/Avalanche/Source/AvalancheMediaEditor/Private/Rundown/Factories/Suggestions/AvaRundownFilterTransitionLayerSuggestionFactory.cpp
--- a/Engine/Plugins/Experimental/Avalanche/Source/AvalancheMediaEditor/Private/Rundown/Factories/Suggestions/AvaRundownFilterTransitionLayerSuggestionFactory.cpp
+++ b/Engine/Plugins/Experimental/Avalanche/Source/AvalancheMediaEditor/Private/Rundown/Factories/Suggestions/AvaRundownFilterTransitionLayerSuggestionFactory.cpp
@@ -8,14 +8,25 @@
 
 #define LOCTEXT_NAMESPACE "AvaRundownFilterLayerSuggestionFactory"
 
-const FName FAvaRundownFilterTransitionLayerSuggestionFactory::KeyName = FName(TEXT("TRANSITIONLAYER"));
+namespace UE::AvaMediaEditor::Suggestion::Layer::Private
+{
+	static const FName KeyName = FName(TEXT("TRANSITIONLAYER"));
+	static const FText CategoryLabel = LOCTEXT("TransitionLayerCategoryLabel", "Ava-Rundown-Transition-Layer");
+}
+
+FName FAvaRundownFilterTransitionLayerSuggestionFactory::GetSuggestionIdentifier() const
+{
+	using namespace UE::AvaMediaEditor::Suggestion::Layer::Private;
+	return KeyName;
+}
 
 void FAvaRundownFilterTransitionLayerSuggestionFactory::AddSuggestion(const TSharedRef<FAvaRundownFilterSuggestionPayload>& InPayload)
 {
+	using namespace UE::AvaMediaEditor::Suggestion::Layer::Private;
+
 	const FAvaRundownPage& PageItem = UAvaRundown::GetPageSafe(InPayload->Rundown, InPayload->ItemPageId);
 	if (PageItem.IsValidPage())
 	{
-		const FText TransitionLayerCategoryLabel = LOCTEXT("TransitionLayerCategoryLabel", "Ava-Rundown-Transition-Layer");
 		const FString PageLayer = PageItem.GetTransitionLayer(InPayload->Rundown).ToString();
 
 		FString TransitionLayerSuggestion = FString::Printf(TEXT("TransitionLayer=%s"), *PageLayer);
@@ -24,7 +35,7 @@ void FAvaRundownFilterTransitionLayerSuggestionFactory::AddSuggestion(const TSha
 		if (bIsFilterValueValid && !InPayload->FilterCache.Contains(TransitionLayerSuggestion))
 		{
 			InPayload->FilterCache.Add(TransitionLayerSuggestion);
-			InPayload->PossibleSuggestions.Add(FAssetSearchBoxSuggestion{MoveTemp(TransitionLayerSuggestion), FText::FromString(PageLayer), TransitionLayerCategoryLabel});
+			InPayload->PossibleSuggestions.Add(FAssetSearchBoxSuggestion{MoveTemp(TransitionLayerSuggestion), FText::FromString(PageLayer), CategoryLabel});
 		}
 	}
 }
diff --git a/Engine/Plugins/Experimental/Avalanche/Source/AvalancheMediaEditor/Private/Rundown/Factories/Suggestions/AvaRundownFilterTransitionLayerSuggestionFactory.h b/Engine/Plugins/Experimental/Avalanche/Source/AvalancheMediaEditor/Private/Rundown/Factories/Suggestions/AvaRundownFilterTransitionLayerSuggestionFactory.h
--- a/Engine/Plugins/Experimental/Avalanche/Source/AvalancheMediaEditor/Private/Rundown/Factories/Suggestions/AvaRundownFilterTransitionLayerSuggestionFactory.h
+++ b/Engine/Plugins/Experimental/Avalanche/Source/AvalancheMediaEditor/Private/Rundown/Factories/Suggestions/AvaRundownFilterTransitionLayerSuggestionFactory.h
@@ -7,10 +7,8 @@
 class FAvaRundownFilterTransitionLayerSuggestionFactory : public IAvaRundownFilterSuggestionFactory
 {
 public:
-	static const FName KeyName;
-
 	//~ Begin IAvaFilterSuggestionFactory interface
-	virtual FName GetSuggestionIdentifier() const override { return KeyName; }
+	virtual FName GetSuggestionIdentifier() const override;
 	virtual bool IsSimpleSuggestion() const override { return false; }
 	virtual void AddSuggestion(const TSharedRef<FAvaRundownFilterSuggestionPayload>& InPayload) override;
 	virtual bool SupportSuggestionType(EAvaRundownSearchListType InSuggestionType) const override;
