Motion Design:
- Cleanup deprecated/unused structs

#changelist validated
#virtualized

diff --git a/Engine/Plugins/Experimental/Avalanche/Source/Avalanche/Private/AvaScene.cpp b/Engine/Plugins/Experimental/Avalanche/Source/Avalanche/Private/AvaScene.cpp
--- a/Engine/Plugins/Experimental/Avalanche/Source/Avalanche/Private/AvaScene.cpp
+++ b/Engine/Plugins/Experimental/Avalanche/Source/Avalanche/Private/AvaScene.cpp
@@ -323,11 +323,6 @@ void AAvaScene::SetStartupCameraName(FName InName)
 {
 	StartupCameraName = InName;
 }
-
-void AAvaScene::SetViewportGuideData(const TConstArrayView<FAvaViewportGuideInfo_Deprecated>& InGuideData)
-{
-	GuideData = InGuideData;
-}
 #endif
 
 #undef LOCTEXT_NAMESPACE
diff --git a/Engine/Plugins/Experimental/Avalanche/Source/Avalanche/Public/AvaDefs.h b/Engine/Plugins/Experimental/Avalanche/Source/Avalanche/Public/AvaDefs.h
--- a/Engine/Plugins/Experimental/Avalanche/Source/Avalanche/Public/AvaDefs.h
+++ b/Engine/Plugins/Experimental/Avalanche/Source/Avalanche/Public/AvaDefs.h
@@ -41,13 +41,6 @@ enum class EAvaAnchors : uint8
 	Custom = UE::Avalanche::AnchorPoints::Custom
 };
 
-UENUM(BlueprintType)
-enum class EAvaLayoutPlacementCenter : uint8
-{
-	Bounds,
-	Pivot
-};
-
 UENUM(BlueprintType)
 enum class EAvaVerticalAlignment : uint8
 {
@@ -72,27 +65,6 @@ enum class EAvaDepthAlignment : uint8
 	Back
 };
 
-UENUM(BlueprintType)
-enum class EAvaLayoutSizeBounds : uint8
-{
-	Explicit,
-	ScaleHeightToChildren,
-	ScaleWidthToChildren,
-	ScaleToChildren,
-};
-
-UENUM(BlueprintType)
-enum class EAvaLayoutRelationshipBounds : uint8
-{
-	DoNotScale,
-	ScaleToHeight,
-	ScaleToWidth,
-	ScaleToFit,
-	StretchHeight,
-	StretchWidth,
-	Stretch
-};
-
 /** Specifies the method for finding a reference actor based on it's position in the parent's hierarchy. */
 UENUM(BlueprintType)
 enum class EAvaReferenceContainer : uint8
diff --git a/Engine/Plugins/Experimental/Avalanche/Source/Avalanche/Public/AvaScene.h b/Engine/Plugins/Experimental/Avalanche/Source/Avalanche/Public/AvaScene.h
--- a/Engine/Plugins/Experimental/Avalanche/Source/Avalanche/Public/AvaScene.h
+++ b/Engine/Plugins/Experimental/Avalanche/Source/Avalanche/Public/AvaScene.h
@@ -10,11 +10,6 @@
 #include "UObject/ObjectPtr.h"
 #include "UObject/WeakObjectPtr.h"
 #include "Viewport/Interaction/IAvaViewportDataProvider.h"
-
-#if WITH_EDITOR
-#include "Viewport/Interaction/AvaViewportGuide.h"
-#endif
-
 #include "AvaScene.generated.h"
 
 class IAvaSequencePlaybackObject;
@@ -100,8 +95,6 @@ public:
 	virtual FName GetStartupCameraName() const override { return StartupCameraName; }
 #if WITH_EDITOR
 	virtual void SetStartupCameraName(FName InName) override;
-	virtual TConstArrayView<FAvaViewportGuideInfo_Deprecated> GetViewportGuideData() const { return GuideData; }
-	virtual void SetViewportGuideData(const TConstArrayView<FAvaViewportGuideInfo_Deprecated>& InGuideData) override;
 #endif
 	//~ End IAvaViewportDataProvider
 
@@ -154,9 +147,4 @@ protected:
 	FSimpleMulticastDelegate OnTreeAnimationRebuilt;
 	
 	bool bPendingAnimTreeUpdate = false; 
-
-#if WITH_EDITORONLY_DATA
-	UPROPERTY()
-	TArray<FAvaViewportGuideInfo_Deprecated> GuideData;
-#endif
 };
diff --git a/Engine/Plugins/Experimental/Avalanche/Source/Avalanche/Public/Viewport/Interaction/AvaViewportGuide.h b/Engine/Plugins/Experimental/Avalanche/Source/Avalanche/Public/Viewport/Interaction/AvaViewportGuide.h
deleted file mode 100644
diff --git a/Engine/Plugins/Experimental/Avalanche/Source/Avalanche/Public/Viewport/Interaction/IAvaViewportDataProvider.h b/Engine/Plugins/Experimental/Avalanche/Source/Avalanche/Public/Viewport/Interaction/IAvaViewportDataProvider.h
--- a/Engine/Plugins/Experimental/Avalanche/Source/Avalanche/Public/Viewport/Interaction/IAvaViewportDataProvider.h
+++ b/Engine/Plugins/Experimental/Avalanche/Source/Avalanche/Public/Viewport/Interaction/IAvaViewportDataProvider.h
@@ -3,17 +3,10 @@
 #pragma once
 
 #include "UObject/Interface.h"
-#if WITH_EDITOR
-#include "Containers/ArrayView.h"
-#endif
 #include "IAvaViewportDataProvider.generated.h"
 
 class UObject;
 
-#if WITH_EDITOR
-struct FAvaViewportGuideInfo_Deprecated;
-#endif
-
 /** Interface for Objects that use UAvaSequence and need to be handled by IAvaSequencer */
 UINTERFACE(MinimalAPI, NotBlueprintType, meta = (CannotImplementInterfaceInBlueprint))
 class UAvaViewportDataProvider : public UInterface
@@ -32,8 +25,5 @@ public:
 
 #if WITH_EDITOR
 	virtual void SetStartupCameraName(FName InName) = 0;
-
-	virtual TConstArrayView<FAvaViewportGuideInfo_Deprecated> GetViewportGuideData() const = 0;
-	virtual void SetViewportGuideData(const TConstArrayView<FAvaViewportGuideInfo_Deprecated>& InGuideData) = 0;
 #endif
 };
diff --git a/Engine/Plugins/Experimental/Avalanche/Source/AvalancheViewport/Private/Interaction/AvaSnapOperation.cpp b/Engine/Plugins/Experimental/Avalanche/Source/AvalancheViewport/Private/Interaction/AvaSnapOperation.cpp
--- a/Engine/Plugins/Experimental/Avalanche/Source/AvalancheViewport/Private/Interaction/AvaSnapOperation.cpp
+++ b/Engine/Plugins/Experimental/Avalanche/Source/AvalancheViewport/Private/Interaction/AvaSnapOperation.cpp
@@ -23,7 +23,6 @@
 #include "Selection.h"
 #include "Selection/AvaSelectionProviderSubsystem.h"
 #include "Viewport/Interaction/AvaSnapPoint.h"
-#include "Viewport/Interaction/AvaViewportGuide.h"
 #include "Viewport/Interaction/IAvaSnapPointGenerator.h"
 #include "Viewport/Interaction/IAvaViewportDataProvider.h"
 #include "ViewportClient/IAvaViewportClient.h"
