[Motion Design]
 - Move ease curve presets to {Plugin}/Config/EaseCurves and setup mechanism to copy to from plugin to project config if project config ease curves is empty

diff --git a/Engine/Plugins/Experimental/Avalanche/Config/EaseCurves/Ease In Out.json b/Engine/Plugins/Experimental/Avalanche/Config/EaseCurves/Ease In Out.json
new file mode 100644
--- /dev/null
+++ b/Engine/Plugins/Experimental/Avalanche/Config/EaseCurves/Ease In Out.json	
@@ -0,0 +1,10 @@
+{
+	"Sine": "0.37, 0.0, 0.63, 1.0",
+	"Cubic": "0.65, 0.0, 0.35, 1.0",
+	"Quintic": "0.83, 0.0, 0.17, 1.0",
+	"Circular": "0.85, 0.0, 0.15, 1.0",
+	"Quadratic": "0.45, 0.0, 0.55, 1.0",
+	"Quartic": "0.76, 0.0, 0.24, 1.0",
+	"Exponential": "0.87, 0.0, 0.13, 1.0",
+	"Back": "0.68, -0.6, 0.32, 1.6"
+}
\ No newline at end of file
diff --git a/Engine/Plugins/Experimental/Avalanche/Config/EaseCurves/Ease In.json b/Engine/Plugins/Experimental/Avalanche/Config/EaseCurves/Ease In.json
new file mode 100644
--- /dev/null
+++ b/Engine/Plugins/Experimental/Avalanche/Config/EaseCurves/Ease In.json	
@@ -0,0 +1,10 @@
+{
+	"Sine": "0.12, 0.0, 0.39, 0.0",
+	"Cubic": "0.32, 0.0, 0.67, 0.0",
+	"Quintic": "0.64, 0.0, 0.78, 0.0",
+	"Circular": "0.55, 0.0, 1.0, 0.45",
+	"Quadratic": "0.11, 0.0, 0.5, 0.0",
+	"Quartic": "0.5, 0.0, 0.75, 0.0",
+	"Exponential": "0.7, 0.0, 0.84, 0.0",
+	"Back": "0.36, 0.0, 0.66, -0.56"
+}
\ No newline at end of file
diff --git a/Engine/Plugins/Experimental/Avalanche/Config/EaseCurves/Ease Out.json b/Engine/Plugins/Experimental/Avalanche/Config/EaseCurves/Ease Out.json
new file mode 100644
--- /dev/null
+++ b/Engine/Plugins/Experimental/Avalanche/Config/EaseCurves/Ease Out.json	
@@ -0,0 +1,10 @@
+{
+	"Sine": "0.61, 1.0, 0.88, 1.0",
+	"Cubic": "0.33, 1.0, 0.68, 1.0",
+	"Quintic": "0.22, 1.0, 0.36, 1.0",
+	"Circular": "0.0, 0.55, 0.45, 1.0",
+	"Quadratic": "0.5, 1.0, 0.89, 1.0",
+	"Quartic": "0.25, 1.0, 0.5, 1.0",
+	"Exponential": "0.16, 1.0, 0.3, 1.0",
+	"Back": "0.34, 1.56, 0.64, 1.0"
+}
\ No newline at end of file
diff --git a/Engine/Plugins/Experimental/Avalanche/Source/AvalancheSequencer/Private/AvaSequencerModule.cpp b/Engine/Plugins/Experimental/Avalanche/Source/AvalancheSequencer/Private/AvaSequencerModule.cpp
--- a/Engine/Plugins/Experimental/Avalanche/Source/AvalancheSequencer/Private/AvaSequencerModule.cpp
+++ b/Engine/Plugins/Experimental/Avalanche/Source/AvalancheSequencer/Private/AvaSequencerModule.cpp
@@ -13,7 +13,9 @@
 #include "Director/AvaSequenceDirectorBlueprint.h"
 #include "Director/AvaSequenceDirectorCompiler.h"
 #include "EaseCurveTool/AvaEaseCurveStyle.h"
+#include "EaseCurveTool/AvaEaseCurveSubsystem.h"
 #include "EaseCurveTool/AvaEaseCurveToolCommands.h"
+#include "Editor.h"
 #include "IAvaOutliner.h"
 #include "IAvaOutlinerModule.h"
 #include "ISequencerModule.h"
@@ -54,6 +56,8 @@ void FAvaSequencerModule::StartupModule()
 	RegisterDirectorCompiler();
 
 	FAvaSequenceTransitionCompiler::Get().Register();
+
+	EditorInitializedDelegate = FEditorDelegates::OnEditorInitialized.AddRaw(this, &FAvaSequencerModule::OnEditorInitialized);
 }
 
 void FAvaSequencerModule::ShutdownModule()
@@ -78,6 +82,8 @@ void FAvaSequencerModule::ShutdownModule()
 	UnregisterCustomLayouts();
 
 	FAvaSequenceTransitionCompiler::Get().Unregister();
+
+	FEditorDelegates::OnEditorInitialized.Remove(EditorInitializedDelegate);
 }
 
 void FAvaSequencerModule::RegisterSequenceEditor(ISequencerModule& InSequencerModule)
@@ -203,6 +209,13 @@ void FAvaSequencerModule::RegisterDirectorCompiler()
 	);
 }
 
+void FAvaSequencerModule::OnEditorInitialized(const double InDuration)
+{
+	// Copy default ease curve presets from the plugin to the project configured directory
+	// only if the project presets directory is empty.
+	UAvaEaseCurveSubsystem::Get().ResetToDefaultPresets(true);
+}
+
 IMPLEMENT_MODULE(FAvaSequencerModule, AvalancheSequencer)
 
 #undef LOCTEXT_NAMESPACE
diff --git a/Engine/Plugins/Experimental/Avalanche/Source/AvalancheSequencer/Private/AvaSequencerModule.h b/Engine/Plugins/Experimental/Avalanche/Source/AvalancheSequencer/Private/AvaSequencerModule.h
--- a/Engine/Plugins/Experimental/Avalanche/Source/AvalancheSequencer/Private/AvaSequencerModule.h
+++ b/Engine/Plugins/Experimental/Avalanche/Source/AvalancheSequencer/Private/AvaSequencerModule.h
@@ -118,6 +118,8 @@ private:
 
 	void RegisterDirectorCompiler();
 
+	void OnEditorInitialized(const double InDuration);
+
 	TSharedPtr<FAvaActorSubobjectSchema> ActorSubobjectSchema;
 
 	FDelegateHandle SequenceEditorHandle;
@@ -131,4 +133,6 @@ private:
 	FDelegateHandle OutlinerProxiesExtensionDelegateHandle;
 
 	TArray<const UClass*> CustomizedClasses;
+
+	FDelegateHandle EditorInitializedDelegate;
 };
diff --git a/Engine/Plugins/Experimental/Avalanche/Source/AvalancheSequencer/Private/EaseCurveTool/AvaEaseCurveSubsystem.cpp b/Engine/Plugins/Experimental/Avalanche/Source/AvalancheSequencer/Private/EaseCurveTool/AvaEaseCurveSubsystem.cpp
--- a/Engine/Plugins/Experimental/Avalanche/Source/AvalancheSequencer/Private/EaseCurveTool/AvaEaseCurveSubsystem.cpp
+++ b/Engine/Plugins/Experimental/Avalanche/Source/AvalancheSequencer/Private/EaseCurveTool/AvaEaseCurveSubsystem.cpp
@@ -6,7 +6,10 @@
 #include "Editor.h"
 #include "Engine/World.h"
 #include "HAL/FileManager.h"
+#include "HAL/PlatformFileManager.h"
+#include "Interfaces/IPluginManager.h"
 #include "Misc/FileHelper.h"
+#include "Misc/Paths.h"
 #include "Serialization/JsonSerializer.h"
 
 DEFINE_LOG_CATEGORY_STATIC(LogAvaEaseCurveSubsystem, Log, All);
@@ -20,8 +23,7 @@ namespace UE::EaseCurveTool::Json::Private
 		const UAvaEaseCurveToolSettings* const Settings = GetDefault<UAvaEaseCurveToolSettings>();
 		check(Settings);
 
-		FString PresetsPath = Settings->GetPresetsPath();
-		PresetsPath = FPaths::ConvertRelativePathToFull(PresetsPath);
+		const FString PresetsPath = FPaths::ConvertRelativePathToFull(UAvaEaseCurveSubsystem::ProjectPresetPath());
 
 		if (!FPaths::DirectoryExists(PresetsPath))
 		{
@@ -46,10 +48,9 @@ namespace UE::EaseCurveTool::Json::Private
 			InCategory = EaseCurveTool_NewPresetCategory;
 		}
 
-		PresetsPath = FPaths::Combine(PresetsPath, InCategory) + TEXT(".json");
-		OutFilePath = PresetsPath;
+		OutFilePath = FPaths::Combine(PresetsPath, InCategory) + TEXT(".json");
 
-		return FPaths::FileExists(PresetsPath);
+		return FPaths::FileExists(OutFilePath);
 	}
 
 	bool LoadCurvePresetsJson(const FString& InFilePath, TSharedPtr<FJsonObject>& OutRootObject)
@@ -305,18 +306,28 @@ void UAvaEaseCurveSubsystem::Initialize(FSubsystemCollectionBase& InCollection)
 
 void UAvaEaseCurveSubsystem::ExploreJsonPresetsFolder()
 {
-	const FString EaseCurvePresetsPath = FPaths::ConvertRelativePathToFull(GetDefault<UAvaEaseCurveToolSettings>()->GetPresetsPath());
+	const FString EaseCurvePresetsPath = FPaths::ConvertRelativePathToFull(ProjectPresetPath());
 	FPlatformProcess::ExploreFolder(*EaseCurvePresetsPath);
 }
 
-void UAvaEaseCurveSubsystem::ReloadPresetsFromJson()
+FString UAvaEaseCurveSubsystem::ProjectPresetPath()
+{
+	return FPaths::ProjectConfigDir() / TEXT("EaseCurves");
+}
+
+FString UAvaEaseCurveSubsystem::PluginPresetPath()
 {
-	const UAvaEaseCurveToolSettings* const EaseCurveToolSettings = GetDefault<UAvaEaseCurveToolSettings>();
-	check(EaseCurveToolSettings);
-	const FString EaseCurvePresetsPath = EaseCurveToolSettings->GetPresetsPath();
+	if (const TSharedPtr<IPlugin> Plugin = IPluginManager::Get().FindPlugin(UE_PLUGIN_NAME))
+	{
+		return Plugin->GetBaseDir() / TEXT("Config/EaseCurves");
+	}
+	return FString();
+}
 
+void UAvaEaseCurveSubsystem::ReloadPresetsFromJson()
+{
 	TArray<FString> JsonFiles;
-	IFileManager::Get().FindFilesRecursive(JsonFiles, *EaseCurvePresetsPath, TEXT("*.json"), true, false);
+	IFileManager::Get().FindFilesRecursive(JsonFiles, *ProjectPresetPath(), TEXT("*.json"), true, false);
 
 	// Load all Json files found as their own separate category. Category Name = File Name
 	Presets.Empty();
@@ -618,3 +629,32 @@ bool UAvaEaseCurveSubsystem::RenamePreset(const FString& InCategory, const FStri
 
 	return true;
 }
+
+void UAvaEaseCurveSubsystem::ResetToDefaultPresets(const bool bInOnlyIfNoProjectPresets)
+{
+	IPlatformFile& PlatformFile = FPlatformFileManager::Get().GetPlatformFile();
+
+	const FString ProjectPresetPath = FPaths::ConvertRelativePathToFull(UAvaEaseCurveSubsystem::ProjectPresetPath());
+
+	if (!bInOnlyIfNoProjectPresets)
+	{
+		PlatformFile.DeleteDirectoryRecursively(*ProjectPresetPath);
+		PlatformFile.CreateDirectory(*ProjectPresetPath);
+	}
+
+	TArray<FString> JsonFiles;
+	IFileManager::Get().FindFilesRecursive(JsonFiles, *ProjectPresetPath, TEXT("*.json"), true, false);
+
+	if (JsonFiles.IsEmpty())
+	{
+		const FString PluginPresetPath = FPaths::ConvertRelativePathToFull(UAvaEaseCurveSubsystem::PluginPresetPath());
+
+		if (PlatformFile.CopyDirectoryTree(*ProjectPresetPath, *PluginPresetPath, false))
+		{
+			UE_LOG(LogAvaEaseCurveSubsystem, Warning, TEXT("Motion Design ease curve tool project presets are empty. "
+				"Copied default presets to [Project]/Config/EaseCurves"));
+
+			ReloadPresetsFromJson();
+		}
+	}
+}
diff --git a/Engine/Plugins/Experimental/Avalanche/Source/AvalancheSequencer/Private/EaseCurveTool/AvaEaseCurveSubsystem.h b/Engine/Plugins/Experimental/Avalanche/Source/AvalancheSequencer/Private/EaseCurveTool/AvaEaseCurveSubsystem.h
--- a/Engine/Plugins/Experimental/Avalanche/Source/AvalancheSequencer/Private/EaseCurveTool/AvaEaseCurveSubsystem.h
+++ b/Engine/Plugins/Experimental/Avalanche/Source/AvalancheSequencer/Private/EaseCurveTool/AvaEaseCurveSubsystem.h
@@ -20,6 +20,9 @@ public:
 
 	static void ExploreJsonPresetsFolder();
 
+	static FString ProjectPresetPath();
+	static FString PluginPresetPath();
+
 	void ReloadPresetsFromJson();
 
 	TArray<FString> GetEaseCurveCategories() const;
@@ -48,6 +51,8 @@ public:
 
 	bool RenamePreset(const FString& InCategory, const FString& InPreset, const FString& InNewPresetName);
 
+	void ResetToDefaultPresets(const bool bInOnlyIfNoProjectPresets);
+
 	//~ Begin USubsystem
 	virtual void Initialize(FSubsystemCollectionBase& InCollection) override;
 	//~ End USubsystem
diff --git a/Engine/Plugins/Experimental/Avalanche/Source/AvalancheSequencer/Private/EaseCurveTool/AvaEaseCurveToolCommands.cpp b/Engine/Plugins/Experimental/Avalanche/Source/AvalancheSequencer/Private/EaseCurveTool/AvaEaseCurveToolCommands.cpp
--- a/Engine/Plugins/Experimental/Avalanche/Source/AvalancheSequencer/Private/EaseCurveTool/AvaEaseCurveToolCommands.cpp
+++ b/Engine/Plugins/Experimental/Avalanche/Source/AvalancheSequencer/Private/EaseCurveTool/AvaEaseCurveToolCommands.cpp
@@ -17,6 +17,14 @@ void FAvaEaseCurveToolCommands::RegisterCommands()
 		, EUserInterfaceActionType::Button
 		, FInputChord());
 
+	
+	UI_COMMAND(ResetToDefaultPresets
+		, "Reset To Default Presets"
+		, "Reset presets to the defaults.\n\n"
+			"*CAUTION* All directories and files inside '[Project]/Config/EaseCurves' will be replaced with defaults!"
+		, EUserInterfaceActionType::Button
+		, FInputChord());
+
 	FUICommandInfo::MakeCommandInfo(SharedThis
 		, Refresh
 		, TEXT("Refresh")
diff --git a/Engine/Plugins/Experimental/Avalanche/Source/AvalancheSequencer/Private/EaseCurveTool/AvaEaseCurveToolCommands.h b/Engine/Plugins/Experimental/Avalanche/Source/AvalancheSequencer/Private/EaseCurveTool/AvaEaseCurveToolCommands.h
--- a/Engine/Plugins/Experimental/Avalanche/Source/AvalancheSequencer/Private/EaseCurveTool/AvaEaseCurveToolCommands.h
+++ b/Engine/Plugins/Experimental/Avalanche/Source/AvalancheSequencer/Private/EaseCurveTool/AvaEaseCurveToolCommands.h
@@ -20,6 +20,8 @@ public:
 
 	TSharedPtr<FUICommandInfo> OpenToolSettings;
 
+	TSharedPtr<FUICommandInfo> ResetToDefaultPresets;
+
 	TSharedPtr<FUICommandInfo> Refresh;
 	TSharedPtr<FUICommandInfo> Apply;
 
diff --git a/Engine/Plugins/Experimental/Avalanche/Source/AvalancheSequencer/Private/EaseCurveTool/AvaEaseCurveToolSettings.cpp b/Engine/Plugins/Experimental/Avalanche/Source/AvalancheSequencer/Private/EaseCurveTool/AvaEaseCurveToolSettings.cpp
--- a/Engine/Plugins/Experimental/Avalanche/Source/AvalancheSequencer/Private/EaseCurveTool/AvaEaseCurveToolSettings.cpp
+++ b/Engine/Plugins/Experimental/Avalanche/Source/AvalancheSequencer/Private/EaseCurveTool/AvaEaseCurveToolSettings.cpp
@@ -8,6 +8,5 @@ UAvaEaseCurveToolSettings::UAvaEaseCurveToolSettings()
 	CategoryName = TEXT("Motion Design");
 	SectionName = TEXT("Ease Curve Tool");
 
-	PresetsPath = FPaths::ProjectConfigDir() / TEXT("EaseCurves");
 	NewPresetCategory = TEXT("Custom");
 }
diff --git a/Engine/Plugins/Experimental/Avalanche/Source/AvalancheSequencer/Private/EaseCurveTool/AvaEaseCurveToolSettings.h b/Engine/Plugins/Experimental/Avalanche/Source/AvalancheSequencer/Private/EaseCurveTool/AvaEaseCurveToolSettings.h
--- a/Engine/Plugins/Experimental/Avalanche/Source/AvalancheSequencer/Private/EaseCurveTool/AvaEaseCurveToolSettings.h
+++ b/Engine/Plugins/Experimental/Avalanche/Source/AvalancheSequencer/Private/EaseCurveTool/AvaEaseCurveToolSettings.h
@@ -38,9 +38,6 @@ public:
 	void SetAutoFlipTangents(const bool bInAutoFlipTangents) { bAutoFlipTangents = bInAutoFlipTangents; }
 	void ToggleAutoFlipTangents() { bAutoFlipTangents = !bAutoFlipTangents; }
 
-	FString GetPresetsPath() const { return PresetsPath; }
-	void SetPresetsPath(const FString& InPresetsPath) { PresetsPath = InPresetsPath; }
-
 	FString GetNewPresetCategory() const { return NewPresetCategory; }
 	void SetNewPresetCategory(const FString& InNewPresetCategory) { NewPresetCategory = InNewPresetCategory; }
 
@@ -72,10 +69,6 @@ private:
 	UPROPERTY(Config, EditAnywhere, Category = "EaseCurveTool", meta = (DisplayName = "Auto Flip Tangents"))
 	bool bAutoFlipTangents = true;
 
-	/** The location to search for ease curve preset Json files. */
-	UPROPERTY(Config, EditAnywhere, Category = "EaseCurveTool", meta = (DisplayName = "Presets Path"))
-	FString PresetsPath;
-
 	/** The name of the category to place newly created curve presets. */
 	UPROPERTY(Config, EditAnywhere, Category = "EaseCurveTool", meta = (DisplayName = "New Preset Category"))
 	FString NewPresetCategory;
diff --git a/Engine/Plugins/Experimental/Avalanche/Source/AvalancheSequencer/Private/EaseCurveTool/Widgets/SAvaEaseCurveTool.cpp b/Engine/Plugins/Experimental/Avalanche/Source/AvalancheSequencer/Private/EaseCurveTool/Widgets/SAvaEaseCurveTool.cpp
--- a/Engine/Plugins/Experimental/Avalanche/Source/AvalancheSequencer/Private/EaseCurveTool/Widgets/SAvaEaseCurveTool.cpp
+++ b/Engine/Plugins/Experimental/Avalanche/Source/AvalancheSequencer/Private/EaseCurveTool/Widgets/SAvaEaseCurveTool.cpp
@@ -8,6 +8,7 @@
 #include "EaseCurveTool/AvaEaseCurveTool.h"
 #include "EaseCurveTool/AvaEaseCurveToolCommands.h"
 #include "EaseCurveTool/AvaEaseCurveToolSettings.h"
+#include "EaseCurveTool/AvaEaseCurveSubsystem.h"
 #include "EaseCurveTool/Widgets/SAvaEaseCurveEditor.h"
 #include "EaseCurveTool/Widgets/SAvaEaseCurvePreset.h"
 #include "Editor.h"
@@ -15,6 +16,7 @@
 #include "Framework/MultiBox/MultiBoxBuilder.h"
 #include "IAvaSequencer.h"
 #include "Math/UnrealMathUtility.h"
+#include "Misc/MessageDialog.h"
 #include "SCurveEditor.h"
 #include "Styling/AppStyle.h"
 #include "Styling/StyleColors.h"
@@ -498,6 +500,8 @@ void SAvaEaseCurveTool::BindCommands()
 
 	CommandList->MapAction(EaseCurveToolCommands.OpenToolSettings, FExecuteAction::CreateSP(EaseCurveToolRef, &FAvaEaseCurveTool::OpenToolSettings));
 
+	CommandList->MapAction(EaseCurveToolCommands.ResetToDefaultPresets, FExecuteAction::CreateSP(this, &SAvaEaseCurveTool::ResetToDefaultPresets));
+
 	CommandList->MapAction(EaseCurveToolCommands.Refresh, FExecuteAction::CreateSP(EaseCurveToolRef, &FAvaEaseCurveTool::UpdateEaseCurveFromSequencerKeySelections));
 	
 	CommandList->MapAction(EaseCurveToolCommands.Apply, FExecuteAction::CreateSP(EaseCurveToolRef, &FAvaEaseCurveTool::ApplyEaseCurveToSequencerKeySelections));
@@ -860,6 +864,10 @@ void SAvaEaseCurveTool::MakeContextMenuSettings(FMenuBuilder& InMenuBuilder)
 	InMenuBuilder.AddSeparator();
 
 	InMenuBuilder.AddMenuEntry(EaseCurveToolCommands.ToggleAutoZoomToFit);
+
+	InMenuBuilder.AddSeparator();
+
+	InMenuBuilder.AddMenuEntry(EaseCurveToolCommands.ResetToDefaultPresets);
 }
 
 void SAvaEaseCurveTool::UndoAction()
@@ -932,4 +940,17 @@ FText SAvaEaseCurveTool::GetEndTooltipText() const
 		: LOCTEXT("LeaveTooltipText", "Leave");
 }
 
+void SAvaEaseCurveTool::ResetToDefaultPresets()
+{
+	const FText MessageBoxTitle = LOCTEXT("ResetToDefaultPresets", "Reset To Default Presets");
+	const EAppReturnType::Type Response = FMessageDialog::Open(EAppMsgType::YesNoCancel
+		, LOCTEXT("ConfirmResetToDefaultPresets", "Are you sure you want to reset to default presets?\n\n"
+			"*CAUTION* All directories and files inside '[Project]/Config/EaseCurves' will be lost!")
+		, MessageBoxTitle);
+	if (Response == EAppReturnType::Yes)
+	{
+		UAvaEaseCurveSubsystem::Get().ResetToDefaultPresets(false);
+	}
+}
+
 #undef LOCTEXT_NAMESPACE
diff --git a/Engine/Plugins/Experimental/Avalanche/Source/AvalancheSequencer/Private/EaseCurveTool/Widgets/SAvaEaseCurveTool.h b/Engine/Plugins/Experimental/Avalanche/Source/AvalancheSequencer/Private/EaseCurveTool/Widgets/SAvaEaseCurveTool.h
--- a/Engine/Plugins/Experimental/Avalanche/Source/AvalancheSequencer/Private/EaseCurveTool/Widgets/SAvaEaseCurveTool.h
+++ b/Engine/Plugins/Experimental/Avalanche/Source/AvalancheSequencer/Private/EaseCurveTool/Widgets/SAvaEaseCurveTool.h
@@ -81,6 +81,8 @@ protected:
 	FText GetEndText() const;
 	FText GetEndTooltipText() const;
 
+	void ResetToDefaultPresets();
+
 	//~ Begin SWidget
 	virtual FReply OnKeyDown(const FGeometry& MyGeometry, const FKeyEvent& InKeyEvent) override;
 	//~ End SWidget
