Motion Design
- Changed exposure compensation (bias) from 0 to 1 on the motion design post process volume.

diff --git a/Engine/Plugins/Experimental/Avalanche/Source/Avalanche/Public/Viewport/AvaPostProcessVolume.h b/Engine/Plugins/Experimental/Avalanche/Source/Avalanche/Public/Viewport/AvaPostProcessVolume.h
--- a/Engine/Plugins/Experimental/Avalanche/Source/Avalanche/Public/Viewport/AvaPostProcessVolume.h
+++ b/Engine/Plugins/Experimental/Avalanche/Source/Avalanche/Public/Viewport/AvaPostProcessVolume.h
@@ -1,4 +1,4 @@
-﻿// Copyright Epic Games, Inc. All Rights Reserved.
+// Copyright Epic Games, Inc. All Rights Reserved.
 
 #pragma once
 
@@ -22,7 +22,7 @@ public:
 		bUnbound = true;
 
 		Settings.bOverride_AutoExposureBias = 1;
-		Settings.AutoExposureBias = 0.f;
+		Settings.AutoExposureBias = 1.f;
 
 		Settings.bOverride_AutoExposureMinBrightness = 1;
 		Settings.AutoExposureMinBrightness = 1.f;
