Motion Design Level Viewport:
Cleaned up Color Picker
- RawPtr->SharedPtr
- FGCObject instead of multiple TStrongObjPtr
- Changed staticness of the delegates to an Color Picker object
- moved logic from Ava Editor to be within Level Viewport -- and unexported unneeded items

#changelist validated
#virtualized

diff --git a/Engine/Plugins/Experimental/Avalanche/Source/AvalancheEditor/Private/AvaEditorIntegration.cpp b/Engine/Plugins/Experimental/Avalanche/Source/AvalancheEditor/Private/AvaEditorIntegration.cpp
--- a/Engine/Plugins/Experimental/Avalanche/Source/AvalancheEditor/Private/AvaEditorIntegration.cpp
+++ b/Engine/Plugins/Experimental/Avalanche/Source/AvalancheEditor/Private/AvaEditorIntegration.cpp
@@ -10,7 +10,6 @@
 #include "MaterialDesigner/AvaLevelMaterialDesignerExtension.h"
 #include "OperatorStack/AvaOperatorStackExtension.h"
 #include "Outliner/AvaOutlinerExtension.h"
-#include "Palettes/AvaPaletteExtension.h"
 #include "RemoteControl/AvaRCExtension.h"
 #include "Scene/AvaSceneExtension.h"
 #include "Selection/AvaSelectionProviderExtension.h"
@@ -32,7 +31,6 @@ TSharedRef<IAvaEditor> FAvaLevelEditorIntegration::BuildEditor()
 		.AddExtension<FAvaRCExtension>()
 		.AddExtension<FAvaViewportExtension, FAvaLevelViewportExtension>()
 		.AddExtension<FAvaInteractiveToolsExtension>()
-		.AddExtension<FAvaPaletteExtension>()
 		.AddExtension<FAvaLevelMaterialDesignerExtension>()
 		.AddExtension<FAvaSceneExtension>()
 		.AddExtension<FAvaTransitionExtension>()
diff --git a/Engine/Plugins/Experimental/Avalanche/Source/AvalancheEditor/Private/AvaEditorModule.cpp b/Engine/Plugins/Experimental/Avalanche/Source/AvalancheEditor/Private/AvaEditorModule.cpp
--- a/Engine/Plugins/Experimental/Avalanche/Source/AvalancheEditor/Private/AvaEditorModule.cpp
+++ b/Engine/Plugins/Experimental/Avalanche/Source/AvalancheEditor/Private/AvaEditorModule.cpp
@@ -6,6 +6,8 @@
 #include "AvaEditorCommands.h"
 #include "AvaEditorIntegration.h"
 #include "AvaShapeActor.h"
+#include "ColorPicker/AvaViewportColorPickerActorClassRegistry.h"
+#include "ColorPicker/IAvaViewportColorPickerAdapter.h"
 #include "Components/LightComponentBase.h"
 #include "DynamicMeshes/AvaShapeDynMeshBase.h"
 #include "Editor/UnrealEdEngine.h"
@@ -23,7 +25,6 @@
 #include "Styling/SlateIconFinder.h"
 #include "UnrealEdGlobals.h"
 #include "Viewport/AvaViewportQualitySettings.h"
-#include "Widgets/AvaViewportColorPickerActorClassRegistry.h"
 
 // Details View
 #include "AssetToolsModule.h"
@@ -38,44 +39,38 @@ DEFINE_LOG_CATEGORY(AvaLog);
 
 #define LOCTEXT_NAMESPACE "AvalancheEditor"
 
-struct FAvaViewportColorPickerLightAdapter : public FAvaViewportColorPickerActorAdapter
+struct FAvaViewportColorPickerLightAdapter : IAvaViewportColorPickerAdapter
 {
-	virtual ~FAvaViewportColorPickerLightAdapter() override = default;
-
-	virtual FAvaColorChangeData GetColorData(const AActor* InActor) const override
+	//~ Begin IAvaViewportColorPickerActorAdapter
+	virtual bool GetColorData(const AActor& InActor, FAvaColorChangeData& OutColorData) const override
 	{
-		if (IsValid(InActor))
+		if (const ULightComponentBase* Component = InActor.FindComponentByClass<ULightComponentBase>())
 		{
-			if (const ULightComponentBase* Component = InActor->FindComponentByClass<ULightComponentBase>())
-			{
-				return {EAvaColorStyle::Solid, Component->GetLightColor(), FLinearColor::Black, /* bIsUnlit */ false};
-			}
+			OutColorData = { EAvaColorStyle::Solid, Component->GetLightColor(), FLinearColor::Black, /* bIsUnlit */ false };
+			return true;
 		}
-
-		return FAvaColorChangeData();
+		return false;
 	}
 
-	virtual void SetColorData(AActor* InActor, const FAvaColorChangeData& InColorData) const override
+	virtual void SetColorData(AActor& InActor, const FAvaColorChangeData& InColorData) const override
 	{
-		if (IsValid(InActor))
+		if (ULightComponentBase* Component = InActor.FindComponentByClass<ULightComponentBase>())
 		{
-			if (ULightComponentBase* Component = InActor->FindComponentByClass<ULightComponentBase>())
-			{
-				FProperty* LightColorProperty = Component->GetClass()->FindPropertyByName(GET_MEMBER_NAME_CHECKED(ULightComponentBase, LightColor));
+			FProperty* LightColorProperty = Component->GetClass()->FindPropertyByName(GET_MEMBER_NAME_CHECKED(ULightComponentBase, LightColor));
 
-				Component->PreEditChange(LightColorProperty);
+			Component->PreEditChange(LightColorProperty);
 
-				Component->LightColor = InColorData.PrimaryColor.ToFColorSRGB();
+			Component->LightColor = InColorData.PrimaryColor.ToFColorSRGB();
 
-				FPropertyChangedEvent PropertyChangedEvent = {
-					LightColorProperty,
-					EPropertyChangeType::Interactive
-				};
+			FPropertyChangedEvent PropertyChangedEvent = {
+				LightColorProperty,
+				EPropertyChangeType::Interactive
+			};
 
-				Component->PostEditChangeProperty(PropertyChangedEvent);
-			}
+			Component->PostEditChangeProperty(PropertyChangedEvent);
 		}
 	}
+	//~ End IAvaViewportColorPickerActorAdapter
 };
 
 namespace UE::AvalancheEditor::Private
@@ -106,10 +101,9 @@ void FAvaEditorModule::StartupModule()
 	IAvaOutlinerModule::Get().GetOnExtendOutlinerItemContextMenu()
     	.AddStatic(&FAvaOutlinerSVGActorContextMenu::OnExtendOutlinerContextMenu);
 
-	FAvaViewportColorPickerActorClassRegistry::RegisterClassAdapter(ALight::StaticClass(), MakeShared<FAvaViewportColorPickerLightAdapter>());
-
-	// Does not extend from ALight
-	FAvaViewportColorPickerActorClassRegistry::RegisterClassAdapter(ASkyLight::StaticClass(), MakeShared<FAvaViewportColorPickerLightAdapter>());
+	// Note: ASkylight Does not extend from ALight
+	FAvaViewportColorPickerActorClassRegistry::RegisterClassAdapter<ALight, FAvaViewportColorPickerLightAdapter>();
+	FAvaViewportColorPickerActorClassRegistry::RegisterClassAdapter<ASkyLight, FAvaViewportColorPickerLightAdapter>();
 }
 
 void FAvaEditorModule::ShutdownModule()
diff --git a/Engine/Plugins/Experimental/Avalanche/Source/AvalancheEditor/Private/Palettes/AvaPaletteExtension.cpp b/Engine/Plugins/Experimental/Avalanche/Source/AvalancheEditor/Private/Palettes/AvaPaletteExtension.cpp
deleted file mode 100644
diff --git a/Engine/Plugins/Experimental/Avalanche/Source/AvalancheEditor/Private/Palettes/AvaPaletteExtension.h b/Engine/Plugins/Experimental/Avalanche/Source/AvalancheEditor/Private/Palettes/AvaPaletteExtension.h
deleted file mode 100644
diff --git a/Engine/Plugins/Experimental/Avalanche/Source/AvalancheEditor/Private/Palettes/AvaPaletteTabSpawner.cpp b/Engine/Plugins/Experimental/Avalanche/Source/AvalancheEditor/Private/Palettes/AvaPaletteTabSpawner.cpp
deleted file mode 100644
diff --git a/Engine/Plugins/Experimental/Avalanche/Source/AvalancheEditor/Private/Palettes/AvaPaletteTabSpawner.h b/Engine/Plugins/Experimental/Avalanche/Source/AvalancheEditor/Private/Palettes/AvaPaletteTabSpawner.h
deleted file mode 100644
diff --git a/Engine/Plugins/Experimental/Avalanche/Source/AvalancheEditor/Private/Palettes/Widgets/SAvaPaletteBase.cpp b/Engine/Plugins/Experimental/Avalanche/Source/AvalancheEditor/Private/Palettes/Widgets/SAvaPaletteBase.cpp
deleted file mode 100644
diff --git a/Engine/Plugins/Experimental/Avalanche/Source/AvalancheEditor/Private/Palettes/Widgets/SAvaPaletteBase.h b/Engine/Plugins/Experimental/Avalanche/Source/AvalancheEditor/Private/Palettes/Widgets/SAvaPaletteBase.h
deleted file mode 100644
diff --git a/Engine/Plugins/Experimental/Avalanche/Source/AvalancheLevelViewport/Private/AvaLevelViewportModule.cpp b/Engine/Plugins/Experimental/Avalanche/Source/AvalancheLevelViewport/Private/AvaLevelViewportModule.cpp
new file mode 100644
--- /dev/null
+++ b/Engine/Plugins/Experimental/Avalanche/Source/AvalancheLevelViewport/Private/AvaLevelViewportModule.cpp
@@ -0,0 +1,70 @@
+// Copyright Epic Games, Inc. All Rights Reserved.
+
+#include "AvaLevelViewportModule.h"
+#include "AvaLevelViewportCommands.h"
+#include "AvaLevelViewportStyle.h"
+#include "AvaViewportCameraHistory.h"
+#include "AvaViewportUtils.h"
+#include "Modules/ModuleManager.h"
+#include "ToolMenus.h"
+#include "ViewportClient/AvaLevelViewportClient.h"
+
+DEFINE_LOG_CATEGORY(AvaLevelViewportLog);
+
+namespace UE::AvalancheLevelViewport::Private
+{
+	TSharedPtr<IAvaViewportClient> GetAsAvaLevelViewportClient(FEditorViewportClient* InViewportClient)
+	{
+		if (FAvaLevelViewportClient::IsAvaLevelViewportClient(InViewportClient))
+		{
+			return static_cast<FAvaLevelViewportClient*>(InViewportClient)->AsShared();
+		}
+
+		return nullptr;
+	}
+}
+
+void FAvaLevelViewportModule::StartupModule()
+{
+	FAvaLevelViewportCommands::Register();
+
+	ViewportCameraHistory = MakeShared<FAvaViewportCameraHistory>();
+
+	UToolMenus::RegisterStartupCallback(FSimpleMulticastDelegate::FDelegate::CreateRaw(this, &FAvaLevelViewportModule::RegisterMenus));
+
+	AvaLevelViewportClientCasterDelegateHandle = FAvaViewportUtils::RegisterViewportClientCaster(
+		&UE::AvalancheLevelViewport::Private::GetAsAvaLevelViewportClient
+	);
+}
+
+void FAvaLevelViewportModule::ShutdownModule()
+{
+	UToolMenus::UnRegisterStartupCallback(this);
+	UToolMenus::UnregisterOwner(this);
+
+	ViewportCameraHistory.Reset();
+
+	FAvaLevelViewportCommands::Unregister();
+
+	if (AvaLevelViewportClientCasterDelegateHandle.IsValid())
+	{
+		FAvaViewportUtils::UnregisterViewportClientCaster(AvaLevelViewportClientCasterDelegateHandle);
+		AvaLevelViewportClientCasterDelegateHandle.Reset();
+	}
+}
+
+void FAvaLevelViewportModule::RegisterMenus()
+{
+	// Owner will be used for cleanup in call to UToolMenus::UnregisterOwner
+	FToolMenuOwnerScoped OwnerScoped(this);
+	
+	// Extension point for SAvaLevelViewportStatusBar
+	{
+		static FName StatusBarMenuName = UE::AvaLevelViewport::Internal::StatusBarMenuName;
+		UToolMenu* Menu = UToolMenus::Get()->RegisterMenu(StatusBarMenuName, NAME_None, EMultiBoxType::ToolBar, false);
+		Menu->SetStyleSet(&FAvaLevelViewportStyle::Get());
+		Menu->StyleName = "StatusBar";
+	}
+}
+
+IMPLEMENT_MODULE(FAvaLevelViewportModule, AvalancheLevelViewport)
diff --git a/Engine/Plugins/Experimental/Avalanche/Source/AvalancheLevelViewport/Private/AvaLevelViewportModule.h b/Engine/Plugins/Experimental/Avalanche/Source/AvalancheLevelViewport/Private/AvaLevelViewportModule.h
new file mode 100644
--- /dev/null
+++ b/Engine/Plugins/Experimental/Avalanche/Source/AvalancheLevelViewport/Private/AvaLevelViewportModule.h
@@ -0,0 +1,31 @@
+// Copyright Epic Games, Inc. All Rights Reserved.
+
+#pragma once
+
+#include "AvaViewportCameraHistory.h"
+#include "Delegates/IDelegateInstance.h"
+#include "Logging/LogMacros.h"
+#include "Modules/ModuleInterface.h"
+#include "Templates/SharedPointer.h"
+
+DECLARE_LOG_CATEGORY_EXTERN(AvaLevelViewportLog, Log, All);
+
+namespace UE::AvaLevelViewport::Internal
+{
+	static FName StatusBarMenuName = TEXT("AvalancheLevelViewport.StatusBar");
+}
+
+class FAvaLevelViewportModule : public IModuleInterface
+{
+	//~ Begin IModuleInterface interface
+	virtual void StartupModule() override;
+	virtual void ShutdownModule() override;
+	//~ End IModuleInterface interface
+
+	void RegisterMenus();
+
+	FDelegateHandle AvaLevelViewportClientCasterDelegateHandle;
+
+	/** Handles viewport camera undo/redo */
+	TSharedPtr<FAvaViewportCameraHistory> ViewportCameraHistory;
+};
diff --git a/Engine/Plugins/Experimental/Avalanche/Source/AvalancheLevelViewport/Private/AvalancheLevelViewportModule.cpp b/Engine/Plugins/Experimental/Avalanche/Source/AvalancheLevelViewport/Private/AvalancheLevelViewportModule.cpp
deleted file mode 100644
diff --git a/Engine/Plugins/Experimental/Avalanche/Source/AvalancheLevelViewport/Private/AvalancheLevelViewportModule.h b/Engine/Plugins/Experimental/Avalanche/Source/AvalancheLevelViewport/Private/AvalancheLevelViewportModule.h
deleted file mode 100644
diff --git a/Engine/Plugins/Experimental/Avalanche/Source/AvalancheLevelViewport/Private/ColorPicker/AvaLevelColorPicker.cpp b/Engine/Plugins/Experimental/Avalanche/Source/AvalancheLevelViewport/Private/ColorPicker/AvaLevelColorPicker.cpp
new file mode 100644
--- /dev/null
+++ b/Engine/Plugins/Experimental/Avalanche/Source/AvalancheLevelViewport/Private/ColorPicker/AvaLevelColorPicker.cpp
@@ -0,0 +1,91 @@
+// Copyright Epic Games, Inc. All Rights Reserved.
+
+#include "AvaLevelColorPicker.h"
+#include "AvaDefs.h"
+#include "ColorPicker/AvaViewportColorPickerActorClassRegistry.h"
+#include "EditorModeManager.h"
+#include "Selection.h"
+#include "Selection/AvaEditorSelection.h"
+#include "Toolkits/IToolkitHost.h"
+
+FAvaLevelColorPicker::FAvaLevelColorPicker()
+{
+	OnSelectionChangedHandle = USelection::SelectionChangedEvent.AddRaw(this, &FAvaLevelColorPicker::UpdateFromSelection);
+}
+
+FAvaLevelColorPicker::~FAvaLevelColorPicker()
+{
+	USelection::SelectionChangedEvent.Remove(OnSelectionChangedHandle);
+	OnSelectionChangedHandle.Reset();
+}
+
+void FAvaLevelColorPicker::SetToolkitHost(const TSharedRef<IToolkitHost>& InToolkitHost)
+{
+	ToolkitHostWeak = InToolkitHost;
+	UpdateFromSelection(nullptr);
+}
+
+void FAvaLevelColorPicker::OnColorSelected(const FAvaColorChangeData& InNewColorData)
+{
+	LastColorData = InNewColorData;
+	ApplyColorToSelections(InNewColorData);
+}
+
+void FAvaLevelColorPicker::ApplyColorToSelections(const FAvaColorChangeData& InNewColorData)
+{
+	TSharedPtr<IToolkitHost> ToolkitHost = ToolkitHostWeak.Pin();
+	if (!ToolkitHost.IsValid())
+	{
+		return;
+	}
+
+	FAvaEditorSelection EditorSelection(ToolkitHost->GetEditorModeManager());
+	if (!EditorSelection.IsValid())
+	{
+		return;
+	}
+
+	TArray<AActor*> SelectedActors = EditorSelection.GetSelectedObjects<AActor>();
+
+	bool bColorApplied = false;
+
+	for (AActor* Actor : SelectedActors)
+	{
+		if (FAvaViewportColorPickerActorClassRegistry::ApplyColorDataToActor(Actor, InNewColorData))
+		{
+			bColorApplied = true;
+		}
+	}
+
+	if (bColorApplied)
+	{
+		GEditor->RedrawLevelEditingViewports();
+	}
+}
+
+void FAvaLevelColorPicker::UpdateFromSelection(UObject* InSelection)
+{
+	TSharedPtr<IToolkitHost> ToolkitHost = ToolkitHostWeak.Pin();
+	if (!ToolkitHost.IsValid())
+	{
+		return;
+	}
+
+	FAvaEditorSelection EditorSelection(ToolkitHost->GetEditorModeManager(), InSelection);
+	if (!EditorSelection.IsValid())
+	{
+		return;
+	}
+
+	TArray<AActor*> SelectedActors = EditorSelection.GetSelectedObjects<AActor>();
+	if (SelectedActors.Num() == 1)
+	{
+		const AActor* const ColorSource = SelectedActors[0];
+
+		FAvaColorChangeData NewColorData;
+		if (FAvaViewportColorPickerActorClassRegistry::GetColorDataFromActor(ColorSource, NewColorData))
+		{
+			LastColorData = NewColorData;
+		}
+	}
+}
diff --git a/Engine/Plugins/Experimental/Avalanche/Source/AvalancheLevelViewport/Private/ColorPicker/AvaLevelColorPicker.h b/Engine/Plugins/Experimental/Avalanche/Source/AvalancheLevelViewport/Private/ColorPicker/AvaLevelColorPicker.h
new file mode 100644
--- /dev/null
+++ b/Engine/Plugins/Experimental/Avalanche/Source/AvalancheLevelViewport/Private/ColorPicker/AvaLevelColorPicker.h
@@ -0,0 +1,36 @@
+// Copyright Epic Games, Inc. All Rights Reserved.
+
+#pragma once
+
+#include "AvaDefs.h"
+#include "Delegates/Delegate.h"
+#include "Templates/SharedPointer.h"
+
+class IToolkitHost;
+
+class FAvaLevelColorPicker
+{
+public:
+	FAvaLevelColorPicker();
+	~FAvaLevelColorPicker();
+
+	void SetToolkitHost(const TSharedRef<IToolkitHost>& InToolkitHost);
+
+	void OnColorSelected(const FAvaColorChangeData& InNewColorData);
+
+	const FAvaColorChangeData& GetLastColorData() const
+	{
+		return LastColorData;
+	}
+
+private:
+	void ApplyColorToSelections(const FAvaColorChangeData& InNewColorData);
+
+	void UpdateFromSelection(UObject* InSelection);
+
+	TWeakPtr<IToolkitHost> ToolkitHostWeak;
+
+	FAvaColorChangeData LastColorData;
+
+	FDelegateHandle OnSelectionChangedHandle;
+};
diff --git a/Engine/Plugins/Experimental/Avalanche/Source/AvalancheLevelViewport/Private/ColorPicker/AvaViewportColorPickerActorClassRegistry.cpp b/Engine/Plugins/Experimental/Avalanche/Source/AvalancheLevelViewport/Private/ColorPicker/AvaViewportColorPickerActorClassRegistry.cpp
new file mode 100644
--- /dev/null
+++ b/Engine/Plugins/Experimental/Avalanche/Source/AvalancheLevelViewport/Private/ColorPicker/AvaViewportColorPickerActorClassRegistry.cpp
@@ -0,0 +1,64 @@
+// Copyright Epic Games, Inc. All Rights Reserved.
+
+#include "ColorPicker/AvaViewportColorPickerActorClassRegistry.h"
+#include "AvaDefs.h"
+#include "ColorPicker/IAvaViewportColorPickerAdapter.h"
+#include "Containers/Map.h"
+#include "GameFramework/Actor.h"
+#include "Templates/SharedPointer.h"
+
+namespace UE::AvalancheLevelViewport::Private
+{
+	TMap<UClass*, TSharedRef<IAvaViewportColorPickerAdapter>> Adapters;
+
+	TSharedPtr<IAvaViewportColorPickerAdapter> FindAdapterForClass(const AActor* InActor)
+	{
+		UClass* ActorClass = IsValid(InActor) ? InActor->GetClass() : nullptr;
+		if (!ActorClass)
+		{
+			return nullptr;
+		}
+
+		for (UClass* Class = ActorClass; Class != AActor::StaticClass(); Class = Class->GetSuperClass())
+		{
+			if (const TSharedRef<IAvaViewportColorPickerAdapter>* Adapter = Adapters.Find(Class))
+			{
+				return *Adapter;
+			}
+		}
+
+		return nullptr;
+	}
+}
+
+void FAvaViewportColorPickerActorClassRegistry::RegisterClassAdapter(UClass* InClass, const TSharedRef<IAvaViewportColorPickerAdapter>& InAdapter)
+{
+	using namespace UE::AvalancheLevelViewport;
+	Private::Adapters.Add(InClass, InAdapter);
+}
+
+bool FAvaViewportColorPickerActorClassRegistry::GetColorDataFromActor(const AActor* InActor, FAvaColorChangeData& OutColorData)
+{
+	using namespace UE::AvalancheLevelViewport;
+
+	if (TSharedPtr<IAvaViewportColorPickerAdapter> Adapter = Private::FindAdapterForClass(InActor))
+	{
+		// Actor has already been checked for validity in FindAdapterForClass
+		return Adapter->GetColorData(*InActor, OutColorData);
+	}
+	return false;
+}
+
+bool FAvaViewportColorPickerActorClassRegistry::ApplyColorDataToActor(AActor* InActor, const FAvaColorChangeData& InColorData)
+{
+	using namespace UE::AvalancheLevelViewport;
+
+	if (TSharedPtr<IAvaViewportColorPickerAdapter> Adapter = Private::FindAdapterForClass(InActor))
+	{
+		// Actor has already been checked for validity in FindAdapterForClass
+		Adapter->SetColorData(*InActor, InColorData);
+		return true;
+	}
+
+	return false;
+}
diff --git a/Engine/Plugins/Experimental/Avalanche/Source/AvalancheLevelViewport/Private/Customizations/AvaLevelViewportComponentTransformDetails.cpp b/Engine/Plugins/Experimental/Avalanche/Source/AvalancheLevelViewport/Private/Customizations/AvaLevelViewportComponentTransformDetails.cpp
new file mode 100644
--- /dev/null
+++ b/Engine/Plugins/Experimental/Avalanche/Source/AvalancheLevelViewport/Private/Customizations/AvaLevelViewportComponentTransformDetails.cpp
@@ -0,0 +1,1992 @@
+// Copyright Epic Games, Inc. All Rights Reserved.
+
+#include "Customizations/AvaLevelViewportComponentTransformDetails.h"
+#include "Algo/Transform.h"
+#include "Components/SceneComponent.h"
+#include "DetailLayoutBuilder.h"
+#include "Editor.h"
+#include "Editor/UnrealEdEngine.h"
+#include "EditorModeManager.h"
+#include "Framework/Application/SlateApplication.h"
+#include "Framework/MultiBox/MultiBoxBuilder.h"
+#include "GameFramework/Actor.h"
+#include "HAL/PlatformApplicationMisc.h"
+#include "IDetailTreeNode.h"
+#include "IPropertyRowGenerator.h"
+#include "LevelEditorSubsystem.h"
+#include "Math/UnitConversion.h"
+#include "Modules/ModuleManager.h"
+#include "Misc/ConfigCacheIni.h"
+#include "Misc/CoreDelegates.h"
+#include "ScopedTransaction.h"
+#include "Settings/EditorProjectSettings.h"
+#include "SlateOptMacros.h"
+#include "Styling/AppStyle.h"
+#include "Textures/SlateIcon.h"
+#include "UObject/UnrealType.h"
+#include "UnrealEdGlobals.h"
+#include "ViewportClient/AvaLevelViewportClient.h"
+#include "Widgets/DeclarativeSyntaxSupport.h"
+#include "Widgets/Images/SImage.h"
+#include "Widgets/Input/NumericUnitTypeInterface.inl"
+#include "Widgets/Input/SCheckBox.h"
+#include "Widgets/Input/SComboButton.h"
+#include "Widgets/Input/SRotatorInputBox.h"
+#include "Widgets/Input/SVectorInputBox.h"
+#include "Widgets/Layout/SBox.h"
+#include "Widgets/SBoxPanel.h"
+#include "Widgets/Text/STextBlock.h"
+
+#define LOCTEXT_NAMESPACE "FAvaLevelViewportComponentTransformDetails"
+
+class FAvaLevelViewportScopedSwitchWorldForObject
+{
+public:
+	FAvaLevelViewportScopedSwitchWorldForObject(UObject* Object)
+		: PrevWorld(nullptr)
+	{
+		bool bRequiresPlayWorld = false;
+		if (GUnrealEd->PlayWorld && !GIsPlayInEditorWorld)
+		{
+			UPackage* ObjectPackage = Object->GetOutermost();
+			bRequiresPlayWorld = ObjectPackage->HasAnyPackageFlags(PKG_PlayInEditor);
+		}
+
+		if (bRequiresPlayWorld)
+		{
+			PrevWorld = SetPlayInEditorWorld(GUnrealEd->PlayWorld);
+		}
+	}
+
+	~FAvaLevelViewportScopedSwitchWorldForObject()
+	{
+		if (PrevWorld)
+		{
+			RestoreEditorWorld(PrevWorld);
+		}
+	}
+
+private:
+	UWorld* PrevWorld;
+};
+
+static USceneComponent* GetSceneComponentFromDetailsObject(UObject* InObject)
+{
+	AActor* Actor = Cast<AActor>(InObject);
+	if (Actor)
+	{
+		return Actor->GetRootComponent();
+	}
+
+	return Cast<USceneComponent>(InObject);
+}
+
+bool FAvaLevelViewportComponentTransformDetails::bUseSelectionAsParent = true;
+FOnUseSelectionAsParentChanged FAvaLevelViewportComponentTransformDetails::OnUseSelectionAsParentChangedEvent;
+
+FAvaLevelViewportComponentTransformDetails::FAvaLevelViewportComponentTransformDetails(TSharedPtr<FAvaLevelViewportClient> InAvaViewportClient)
+	: TNumericUnitTypeInterface(GetDefault<UEditorProjectAppearanceSettings>()->bDisplayUnitsOnComponentTransforms ? EUnit::Centimeters : EUnit::Unspecified)
+	, SelectedActorInfo(AssetSelectionUtils::BuildSelectedActorInfo(TArray<AActor*>()))
+	, SelectedObjects(TArray<TWeakObjectPtr<UObject>>())
+	, NotifyHook(nullptr)
+	, bPreserveScaleRatio(false)
+	, bEditingRotationInUI(false)
+	, bIsSliderTransaction(false)
+	, AvaViewportClientWeak(InAvaViewportClient)
+{
+	OnUseSelectionAsParentChangedEvent.AddRaw(this, &FAvaLevelViewportComponentTransformDetails::OnUseSelectionAsParentChanged);
+
+	GConfig->GetBool(TEXT("SelectionDetails"), TEXT("PreserveScaleRatio"), bPreserveScaleRatio, GEditorPerProjectIni);
+	FCoreUObjectDelegates::OnObjectsReplaced.AddRaw(this, &FAvaLevelViewportComponentTransformDetails::OnObjectsReplaced);
+
+	if (InAvaViewportClient.IsValid())
+	{
+		bCanUseSelectionAsParent = false;
+
+		if (FEditorModeTools* const ModeTools = InAvaViewportClient->GetModeTools())
+		{
+			EditorSelectionSetWeak = ModeTools->GetEditorSelectionSet();
+			if (EditorSelectionSetWeak.IsValid())
+			{
+				EditorSelectionSetWeak->OnChanged().AddRaw(this, &FAvaLevelViewportComponentTransformDetails::OnSelectionSetChanged);
+			}
+		}
+	}
+	else
+	{
+		bCanUseSelectionAsParent = false;
+	}
+
+	bLastUseSelectionAsParent = bUseSelectionAsParent;
+
+	GeneratePropertyHandles();
+
+	FCoreDelegates::OnEnginePreExit.AddRaw(this, &FAvaLevelViewportComponentTransformDetails::OnEnginePreExit);
+}
+
+FAvaLevelViewportComponentTransformDetails::~FAvaLevelViewportComponentTransformDetails()
+{
+	FCoreUObjectDelegates::OnObjectsReplaced.RemoveAll(this);
+	OnUseSelectionAsParentChangedEvent.RemoveAll(this);
+
+	if (UObjectInitialized() && EditorSelectionSetWeak.IsValid())
+	{
+		EditorSelectionSetWeak->OnChanged().RemoveAll(this);
+	}
+
+	FCoreDelegates::OnEnginePreExit.RemoveAll(this);
+
+	PropertyHandles.Empty();
+	TreeNodes.Empty();
+	PropertyRowGenerator.Reset();
+}
+
+TSharedRef<SWidget> FAvaLevelViewportComponentTransformDetails::BuildTransformFieldLabel(EAvaLevelViewportTransformFieldType TransformField)
+{
+	FText Label;
+
+	switch (TransformField)
+	{
+		case EAvaLevelViewportTransformFieldType::Rotation:
+			Label = LOCTEXT("RotationLabel", "Rotation");
+			break;
+			
+		case EAvaLevelViewportTransformFieldType::Scale:
+			Label = LOCTEXT("ScaleLabel", "Scale");
+			break;
+			
+		case EAvaLevelViewportTransformFieldType::Location:
+		default:
+			Label = LOCTEXT("LocationLabel", "Location");
+			break;
+	}
+
+	FMenuBuilder MenuBuilder(true, nullptr, nullptr);
+
+	FUIAction SetRelativeLocationAction
+	(
+		FExecuteAction::CreateSP(this, &FAvaLevelViewportComponentTransformDetails::OnSetAbsoluteTransform, TransformField, false),
+		FCanExecuteAction(),
+		FIsActionChecked::CreateSP(this, &FAvaLevelViewportComponentTransformDetails::IsAbsoluteTransformChecked, TransformField, false)
+	);
+
+	FUIAction SetWorldLocationAction
+	(
+		FExecuteAction::CreateSP(this, &FAvaLevelViewportComponentTransformDetails::OnSetAbsoluteTransform, TransformField, true),
+		FCanExecuteAction(),
+		FIsActionChecked::CreateSP(this, &FAvaLevelViewportComponentTransformDetails::IsAbsoluteTransformChecked, TransformField, true)
+	);
+
+	MenuBuilder.BeginSection(TEXT("TransformType"), FText::Format(LOCTEXT("TransformType", "{0} Type"), Label));
+
+	MenuBuilder.AddMenuEntry
+	(
+		FText::Format(LOCTEXT("RelativeLabel", "Relative"), Label),
+		FText::Format(LOCTEXT("RelativeLabel_ToolTip", "{0} is relative to its parent"), Label),
+		FSlateIcon(),
+		SetRelativeLocationAction,
+		NAME_None,
+		EUserInterfaceActionType::RadioButton
+	);
+
+	MenuBuilder.AddMenuEntry
+	(
+		FText::Format(LOCTEXT("WorldLabel", "World"), Label),
+		FText::Format(LOCTEXT("WorldLabel_ToolTip", "{0} is relative to the world"), Label),
+		FSlateIcon(),
+		SetWorldLocationAction,
+		NAME_None,
+		EUserInterfaceActionType::RadioButton
+	);
+
+	MenuBuilder.EndSection();
+
+	TSharedRef<SWidget> NameContent =
+		SNew(SComboButton)
+		.ContentPadding(0)
+		.MenuContent()
+		[
+			MenuBuilder.MakeWidget()
+		]
+		.ButtonContent()
+		[
+			SNew(SBox)
+			.Padding(FMargin(0.0f, 0.0f, 2.0f, 0.0f))
+			.MinDesiredWidth(50.f)
+			[
+				SNew(STextBlock)
+				.Text(this, &FAvaLevelViewportComponentTransformDetails::GetTransformFieldText, TransformField)
+				.Font(IDetailLayoutBuilder::GetDetailFont())
+			]
+		];
+
+	if (TransformField == EAvaLevelViewportTransformFieldType::Scale)
+	{
+		NameContent =
+			SNew(SHorizontalBox)
+			+ SHorizontalBox::Slot()
+			.VAlign(VAlign_Center)
+			[
+				NameContent
+			]
+			+ SHorizontalBox::Slot()
+			.AutoWidth()
+			.VAlign(VAlign_Center)
+			.Padding(FMargin(4.0f, 0.0f, 0.0f, 0.0f))
+			[
+				GetPreserveScaleRatioWidget()
+			];
+	}
+
+	return NameContent;
+}
+
+FText FAvaLevelViewportComponentTransformDetails::GetTransformFieldText(EAvaLevelViewportTransformFieldType TransformField) const
+{
+	switch (TransformField)
+	{
+	case EAvaLevelViewportTransformFieldType::Location:
+		return GetLocationText();
+		
+	case EAvaLevelViewportTransformFieldType::Rotation:
+		return GetRotationText();
+		
+	case EAvaLevelViewportTransformFieldType::Scale:
+		return GetScaleText();
+		
+	default:
+		return FText::GetEmpty();
+	}
+}
+
+bool FAvaLevelViewportComponentTransformDetails::OnCanCopy(EAvaLevelViewportTransformFieldType TransformField) const
+{
+	// We can only copy values if the whole field is set.  If multiple values are defined we do not copy since we are unable to determine the value
+	switch (TransformField)
+	{
+		case EAvaLevelViewportTransformFieldType::Location:
+			return CachedLocation.IsSet();
+
+		case EAvaLevelViewportTransformFieldType::Rotation:
+			return CachedRotation.IsSet();
+
+		case EAvaLevelViewportTransformFieldType::Scale:
+			return CachedScale.IsSet();
+
+		default:
+			return false;
+	}
+}
+
+void FAvaLevelViewportComponentTransformDetails::OnCopy(EAvaLevelViewportTransformFieldType TransformField)
+{
+	CacheTransform();
+
+	FString CopyStr;
+	switch (TransformField)
+	{
+	case EAvaLevelViewportTransformFieldType::Location:
+		CopyStr = FString::Printf(TEXT("(X=%f,Y=%f,Z=%f)"), CachedLocation.X.GetValue(), CachedLocation.Y.GetValue(), CachedLocation.Z.GetValue());
+		break;
+		
+	case EAvaLevelViewportTransformFieldType::Rotation:
+		CopyStr = FString::Printf(TEXT("(Pitch=%f,Yaw=%f,Roll=%f)"), CachedRotation.Y.GetValue(), CachedRotation.Z.GetValue(), CachedRotation.X.GetValue());
+		break;
+		
+	case EAvaLevelViewportTransformFieldType::Scale:
+		CopyStr = FString::Printf(TEXT("(X=%f,Y=%f,Z=%f)"), CachedScale.X.GetValue(), CachedScale.Y.GetValue(), CachedScale.Z.GetValue());
+		break;
+		
+	default:
+		break;
+	}
+
+	if (!CopyStr.IsEmpty())
+	{
+		FPlatformApplicationMisc::ClipboardCopy(*CopyStr);
+	}
+}
+
+void FAvaLevelViewportComponentTransformDetails::OnPaste(EAvaLevelViewportTransformFieldType TransformField)
+{
+	FString PastedText;
+	FPlatformApplicationMisc::ClipboardPaste(PastedText);
+
+	switch (TransformField)
+	{
+		case EAvaLevelViewportTransformFieldType::Location:
+		{
+			FVector Location;
+			if (Location.InitFromString(PastedText))
+			{
+				FScopedTransaction Transaction(LOCTEXT("PasteLocation", "Paste Location"));
+				OnBeginChange(EAvaLevelViewportTransformFieldType::Location);
+				OnSetTransform(EAvaLevelViewportTransformFieldType::Location, EAxisList::All, Location, /* bMirror */ false, /* bCommitted */ true);
+				OnEndChange(EAvaLevelViewportTransformFieldType::Location, EAxisList::All);
+			}
+		}
+		break;
+		
+	case EAvaLevelViewportTransformFieldType::Rotation:
+		{
+			FRotator Rotation;
+			PastedText.ReplaceInline(TEXT("Pitch="), TEXT("P="));
+			PastedText.ReplaceInline(TEXT("Yaw="), TEXT("Y="));
+			PastedText.ReplaceInline(TEXT("Roll="), TEXT("R="));
+			if (Rotation.InitFromString(PastedText))
+			{
+				FScopedTransaction Transaction(LOCTEXT("PasteRotation", "Paste Rotation"));
+				OnBeginChange(EAvaLevelViewportTransformFieldType::Rotation);
+				OnSetTransform(EAvaLevelViewportTransformFieldType::Rotation, EAxisList::All, Rotation.Euler(), /* bMirror */ false, /* bCommitted */ true);
+				OnEndChange(EAvaLevelViewportTransformFieldType::Rotation, EAxisList::All);
+			}
+		}
+		break;
+		
+	case EAvaLevelViewportTransformFieldType::Scale:
+		{
+			FVector Scale;
+			if (Scale.InitFromString(PastedText))
+			{
+				FScopedTransaction Transaction(LOCTEXT("PasteScale", "Paste Scale"));
+				OnBeginChange(EAvaLevelViewportTransformFieldType::Scale);
+				OnSetTransform(EAvaLevelViewportTransformFieldType::Scale, EAxisList::All, Scale, /* bMirror */ false, /* bCommitted */ true);
+				OnEndChange(EAvaLevelViewportTransformFieldType::Scale, EAxisList::All);
+			}
+		}
+		break;
+		
+	default:
+		break;
+	}
+}
+
+FUIAction FAvaLevelViewportComponentTransformDetails::CreateCopyAction(EAvaLevelViewportTransformFieldType TransformField) const
+{
+	FAvaLevelViewportComponentTransformDetails* const MutableThis = const_cast<FAvaLevelViewportComponentTransformDetails*>(this);
+	return FUIAction(FExecuteAction::CreateSP(MutableThis, &FAvaLevelViewportComponentTransformDetails::OnCopy, TransformField)
+		, FCanExecuteAction::CreateSP(this, &FAvaLevelViewportComponentTransformDetails::OnCanCopy, TransformField));
+}
+
+FUIAction FAvaLevelViewportComponentTransformDetails::CreatePasteAction(EAvaLevelViewportTransformFieldType TransformField) const
+{
+	FAvaLevelViewportComponentTransformDetails* const MutableThis = const_cast<FAvaLevelViewportComponentTransformDetails*>(this);
+	return FUIAction(FExecuteAction::CreateSP(MutableThis, &FAvaLevelViewportComponentTransformDetails::OnPaste, TransformField));
+}
+
+BEGIN_SLATE_FUNCTION_BUILD_OPTIMIZATION
+
+TSharedPtr<SWidget> FAvaLevelViewportComponentTransformDetails::GetTransformBody()
+{
+	FSlateFontInfo FontInfo = IDetailLayoutBuilder::GetDetailFont();
+
+	return SNew(SNumericVectorInputBox<FVector::FReal>)
+		.X(this, &FAvaLevelViewportComponentTransformDetails::GetLocationX)
+		.Y(this, &FAvaLevelViewportComponentTransformDetails::GetLocationY)
+		.Z(this, &FAvaLevelViewportComponentTransformDetails::GetLocationZ)
+		.bColorAxisLabels(true)
+		.IsEnabled(this, &FAvaLevelViewportComponentTransformDetails::GetIsEnabled)
+		.OnXChanged(this, &FAvaLevelViewportComponentTransformDetails::OnSetTransformAxis, ETextCommit::Default, EAvaLevelViewportTransformFieldType::Location, EAxisList::X, /* bCommitted */ false)
+		.OnYChanged(this, &FAvaLevelViewportComponentTransformDetails::OnSetTransformAxis, ETextCommit::Default, EAvaLevelViewportTransformFieldType::Location, EAxisList::Y, /* bCommitted */ false)
+		.OnZChanged(this, &FAvaLevelViewportComponentTransformDetails::OnSetTransformAxis, ETextCommit::Default, EAvaLevelViewportTransformFieldType::Location, EAxisList::Z, /* bCommitted */ false)
+		.OnXCommitted(this, &FAvaLevelViewportComponentTransformDetails::OnSetTransformAxis, EAvaLevelViewportTransformFieldType::Location, EAxisList::X, /* bCommitted */ true)
+		.OnYCommitted(this, &FAvaLevelViewportComponentTransformDetails::OnSetTransformAxis, EAvaLevelViewportTransformFieldType::Location, EAxisList::Y, /* bCommitted */ true)
+		.OnZCommitted(this, &FAvaLevelViewportComponentTransformDetails::OnSetTransformAxis, EAvaLevelViewportTransformFieldType::Location, EAxisList::Z, /* bCommitted */ true)
+		.Font(FontInfo)
+		.AllowSpin(true)
+		.SpinDelta(1)
+		.OnBeginSliderMovement(this, &FAvaLevelViewportComponentTransformDetails::OnBeginLocationSlider)
+		.OnEndSliderMovement(this, &FAvaLevelViewportComponentTransformDetails::OnEndLocationSlider);
+}
+
+TSharedPtr<SWidget> FAvaLevelViewportComponentTransformDetails::GetRotationBody()
+{
+	FSlateFontInfo FontInfo = IDetailLayoutBuilder::GetDetailFont();
+
+	return SNew(SRotatorInputBox)
+		.AllowSpin(SelectedObjects.Num() == 1)
+		.Roll(this, &FAvaLevelViewportComponentTransformDetails::GetRotationX)
+		.Pitch(this, &FAvaLevelViewportComponentTransformDetails::GetRotationY)
+		.Yaw(this, &FAvaLevelViewportComponentTransformDetails::GetRotationZ)
+		.bColorAxisLabels(true)
+		.IsEnabled(this, &FAvaLevelViewportComponentTransformDetails::GetIsEnabled)
+		.OnRollChanged(this, &FAvaLevelViewportComponentTransformDetails::OnSetTransformAxisFloat, ETextCommit::Default, EAvaLevelViewportTransformFieldType::Rotation, EAxisList::X, /* bCommitted */ false)
+		.OnPitchChanged(this, &FAvaLevelViewportComponentTransformDetails::OnSetTransformAxisFloat, ETextCommit::Default, EAvaLevelViewportTransformFieldType::Rotation, EAxisList::Y, /* bCommitted */ false)
+		.OnYawChanged(this, &FAvaLevelViewportComponentTransformDetails::OnSetTransformAxisFloat, ETextCommit::Default, EAvaLevelViewportTransformFieldType::Rotation, EAxisList::Z, /* bCommitted */ false)
+		.OnRollCommitted(this, &FAvaLevelViewportComponentTransformDetails::OnSetTransformAxisFloat, EAvaLevelViewportTransformFieldType::Rotation, EAxisList::X, /* bCommitted */ true)
+		.OnPitchCommitted(this, &FAvaLevelViewportComponentTransformDetails::OnSetTransformAxisFloat, EAvaLevelViewportTransformFieldType::Rotation, EAxisList::Y, /* bCommitted */ true)
+		.OnYawCommitted(this, &FAvaLevelViewportComponentTransformDetails::OnSetTransformAxisFloat, EAvaLevelViewportTransformFieldType::Rotation, EAxisList::Z, /* bCommitted */ true)
+		.Font(FontInfo)
+		.AllowSpin(true)
+		.OnYawBeginSliderMovement(this, &FAvaLevelViewportComponentTransformDetails::OnBeginRotationSlider)
+		.OnPitchBeginSliderMovement(this, &FAvaLevelViewportComponentTransformDetails::OnBeginRotationSlider)
+		.OnRollBeginSliderMovement(this, &FAvaLevelViewportComponentTransformDetails::OnBeginRotationSlider)
+		.OnYawEndSliderMovement(this, &FAvaLevelViewportComponentTransformDetails::OnEndRotationSlider)
+		.OnPitchEndSliderMovement(this, &FAvaLevelViewportComponentTransformDetails::OnEndRotationSlider)
+		.OnRollEndSliderMovement(this, &FAvaLevelViewportComponentTransformDetails::OnEndRotationSlider)
+		.MinSliderValue(-360.f)
+		.MaxSliderValue(360.f);
+}
+
+TSharedPtr<SWidget> FAvaLevelViewportComponentTransformDetails::GetScaleBody()
+{
+	FSlateFontInfo FontInfo = IDetailLayoutBuilder::GetDetailFont();
+
+	return SNew(SNumericVectorInputBox<FVector::FReal>)
+		.X(this, &FAvaLevelViewportComponentTransformDetails::GetScaleX)
+		.Y(this, &FAvaLevelViewportComponentTransformDetails::GetScaleY)
+		.Z(this, &FAvaLevelViewportComponentTransformDetails::GetScaleZ)
+		.bColorAxisLabels(true)
+		.IsEnabled(this, &FAvaLevelViewportComponentTransformDetails::GetIsEnabled)
+		.OnXChanged(this, &FAvaLevelViewportComponentTransformDetails::OnSetTransformAxis, ETextCommit::Default, EAvaLevelViewportTransformFieldType::Scale, EAxisList::X, /* bCommitted */ false)
+		.OnYChanged(this, &FAvaLevelViewportComponentTransformDetails::OnSetTransformAxis, ETextCommit::Default, EAvaLevelViewportTransformFieldType::Scale, EAxisList::Y, /* bCommitted */ false)
+		.OnZChanged(this, &FAvaLevelViewportComponentTransformDetails::OnSetTransformAxis, ETextCommit::Default, EAvaLevelViewportTransformFieldType::Scale, EAxisList::Z, /* bCommitted */ false)
+		.OnXCommitted(this, &FAvaLevelViewportComponentTransformDetails::OnSetTransformAxis, EAvaLevelViewportTransformFieldType::Scale, EAxisList::X, /* bCommitted */ true)
+		.OnYCommitted(this, &FAvaLevelViewportComponentTransformDetails::OnSetTransformAxis, EAvaLevelViewportTransformFieldType::Scale, EAxisList::Y, /* bCommitted */ true)
+		.OnZCommitted(this, &FAvaLevelViewportComponentTransformDetails::OnSetTransformAxis, EAvaLevelViewportTransformFieldType::Scale, EAxisList::Z, /* bCommitted */ true)
+		.ContextMenuExtenderX(this, &FAvaLevelViewportComponentTransformDetails::ExtendXScaleContextMenu)
+		.ContextMenuExtenderY(this, &FAvaLevelViewportComponentTransformDetails::ExtendYScaleContextMenu)
+		.ContextMenuExtenderZ(this, &FAvaLevelViewportComponentTransformDetails::ExtendZScaleContextMenu)
+		.Font(FontInfo)
+		.AllowSpin(true)
+		.SpinDelta(0.0025f)
+		.OnBeginSliderMovement(this, &FAvaLevelViewportComponentTransformDetails::OnBeginScaleSlider)
+		.OnEndSliderMovement(this, &FAvaLevelViewportComponentTransformDetails::OnEndScaleSlider);
+}
+
+TSharedPtr<SWidget> FAvaLevelViewportComponentTransformDetails::GetOptionsHeader() 
+{
+	FText Tooltip = LOCTEXT("UseSelectionAsParentTooltip", "Transforms selected actors as if the selection were the common attach parent.");
+
+	return SNew(SCheckBox)
+		.IsChecked(this, &FAvaLevelViewportComponentTransformDetails::IsUseSelectionAsParentChecked)
+		.OnCheckStateChanged(this, &FAvaLevelViewportComponentTransformDetails::OnUseSelectionAsParentChanged)
+		.ToolTipText(Tooltip);
+}
+
+TSharedPtr<SWidget> FAvaLevelViewportComponentTransformDetails::GetOptionsBody()
+{
+	FText Tooltip = LOCTEXT("UseSelectionAsParentTooltip", "Transforms selected actors as if the selection were the common attach parent.");
+	FSlateFontInfo FontInfo = IDetailLayoutBuilder::GetDetailFont();
+
+	return SNew(STextBlock)
+		.Text(LOCTEXT("UseSelectionAsParent", "Use Selection As Parent"))
+		.Font(FontInfo)
+		.ToolTipText(Tooltip);
+}
+
+TSharedRef<SWidget> FAvaLevelViewportComponentTransformDetails::GetPreserveScaleRatioWidget()
+{
+	// Add a checkbox to toggle between preserving the ratio of x,y,z components of scale when a value is entered
+	return SNew(SCheckBox)
+		.IsChecked(this, &FAvaLevelViewportComponentTransformDetails::IsPreserveScaleRatioChecked)
+		.IsEnabled(this, &FAvaLevelViewportComponentTransformDetails::GetIsEnabled)
+		.OnCheckStateChanged(this, &FAvaLevelViewportComponentTransformDetails::OnPreserveScaleRatioToggled)
+		.Style(FAppStyle::Get(), "TransparentCheckBox")
+		.ToolTipText(LOCTEXT("PreserveScaleToolTip", "When locked, scales uniformly based on the current xyz scale values so the object maintains its shape in each direction when scaled"))
+		[
+			SNew(SImage)
+			.Image(this, &FAvaLevelViewportComponentTransformDetails::GetPreserveScaleRatioImage)
+			.ColorAndOpacity(FSlateColor::UseForeground())
+		];
+}
+
+void FAvaLevelViewportComponentTransformDetails::SetTransform(EAvaLevelViewportTransformFieldType TransformField, const FVector& NewValue)
+{
+	FScopedTransaction Transaction(LOCTEXT("SetTransform", "Set Transform"));
+	OnBeginChange(TransformField);
+	OnSetTransform(TransformField, EAxisList::XYZ, NewValue, /* bMirror */ false, /* bCommitted */ true);
+	OnEndChange(TransformField, EAxisList::XYZ);
+}
+
+END_SLATE_FUNCTION_BUILD_OPTIMIZATION
+
+void FAvaLevelViewportComponentTransformDetails::OnSelectionChanged(const TArray<UObject*>& InSelectedObjects)
+{
+	SelectedObjects.Empty();
+	TArray<AActor*> ActorArray;
+
+	for (UObject* Object : InSelectedObjects)
+	{
+		SelectedObjects.Add(TWeakObjectPtr<UObject>(Object));
+
+		if (AActor* Actor = Cast<AActor>(Object))
+		{
+			ActorArray.Add(Actor);
+		}
+	}
+
+	SelectedActorInfo = AssetSelectionUtils::BuildSelectedActorInfo(ActorArray);
+
+	if (SelectedActorInfo.NumSelected == 0)
+	{
+		CachedLocation.X.Reset();
+		CachedLocation.Y.Reset();
+		CachedLocation.Z.Reset();
+
+		CachedRotation.X.Reset();
+		CachedRotation.Y.Reset();
+		CachedRotation.Z.Reset();
+
+		CachedScale.X.Reset();
+		CachedScale.Y.Reset();
+		CachedScale.Z.Reset();
+
+		bCanUseSelectionAsParent = false;
+	}
+	else if (SelectedActorInfo.NumSelected == 1)
+	{
+		bCanUseSelectionAsParent = false;
+	}
+	else
+	{
+		bCanUseSelectionAsParent = false;
+	}
+
+	TArray<UObject*> RootComponents;
+	RootComponents.Reserve(ActorArray.Num());
+
+	for (AActor* Actor : ActorArray)
+	{
+		RootComponents.Add(Actor->GetRootComponent());
+	}
+
+	UpdatePropertyHandlesObjects(RootComponents);
+
+	CacheTransform();
+}
+
+void FAvaLevelViewportComponentTransformDetails::OnSelectionSetChanged(const UTypedElementSelectionSet* InSelectionSet)
+{
+	OnSelectionChanged(InSelectionSet->GetSelectedObjects());
+}
+
+void FAvaLevelViewportComponentTransformDetails::Tick(float DeltaTime)
+{
+	if (bLastUseSelectionAsParent != bCanUseSelectionAsParent)
+	{
+		bLastUseSelectionAsParent = bCanUseSelectionAsParent;
+
+		CachedLocation.X.Reset();
+		CachedLocation.Y.Reset();
+		CachedLocation.Z.Reset();
+
+		CachedRotation.X.Reset();
+		CachedRotation.Y.Reset();
+		CachedRotation.Z.Reset();
+
+		CachedScale.X.Reset();
+		CachedScale.Y.Reset();
+		CachedScale.Z.Reset();
+	}
+
+	CacheTransform();
+
+	if (!FixedDisplayUnits.IsSet())
+	{
+		CacheCommonLocationUnits();
+	}
+}
+
+void FAvaLevelViewportComponentTransformDetails::CacheCommonLocationUnits()
+{
+	float LargestValue = 0.f;
+	if (CachedLocation.X.IsSet() && CachedLocation.X.GetValue() > LargestValue)
+	{
+		LargestValue = CachedLocation.X.GetValue();
+	}
+	if (CachedLocation.Y.IsSet() && CachedLocation.Y.GetValue() > LargestValue)
+	{
+		LargestValue = CachedLocation.Y.GetValue();
+	}
+	if (CachedLocation.Z.IsSet() && CachedLocation.Z.GetValue() > LargestValue)
+	{
+		LargestValue = CachedLocation.Z.GetValue();
+	}
+
+	SetupFixedDisplay(LargestValue);
+}
+
+void FAvaLevelViewportComponentTransformDetails::GeneratePropertyHandles()
+{
+	static const FName LocationName = "Location";
+	static const FName RotationName = "Rotation";
+	static const FName ScaleName = "Scale";
+	static const FName CategoryName = "TransformCommon";
+
+	FPropertyEditorModule& PropertyEditor = FModuleManager::Get().LoadModuleChecked<FPropertyEditorModule>(TEXT("PropertyEditor"));
+	PropertyRowGenerator = PropertyEditor.CreatePropertyRowGenerator(FPropertyRowGeneratorArgs());
+	PropertyRowGenerator->SetObjects({USceneComponent::StaticClass()->GetDefaultObject()});
+
+	TreeNodes.SetNum(3);
+	PropertyHandles.SetNum(3);
+
+	for (const TSharedRef<IDetailTreeNode>& CategoryNode : PropertyRowGenerator->GetRootTreeNodes())
+	{
+		if (CategoryNode->GetNodeName() != CategoryName)
+		{
+			continue;
+		}
+
+		TArray<TSharedRef<IDetailTreeNode>> ChildNodes;
+		CategoryNode->GetChildren(ChildNodes);
+
+		for (const TSharedRef<IDetailTreeNode>& ChildNode : ChildNodes)
+		{
+			if (ChildNode->GetNodeType() != EDetailNodeType::Item)
+			{
+				continue;
+			}
+
+			if (ChildNode->GetNodeName() == LocationName)
+			{
+				TreeNodes[0] = ChildNode;
+				PropertyHandles[0] = ChildNode->CreatePropertyHandle();
+			}
+			else if (ChildNode->GetNodeName() == RotationName)
+			{
+				TreeNodes[1] = ChildNode;
+				PropertyHandles[1] = ChildNode->CreatePropertyHandle();
+			}
+			else if (ChildNode->GetNodeName() == ScaleName)
+			{
+				TreeNodes[2] = ChildNode;
+				PropertyHandles[2] = ChildNode->CreatePropertyHandle();
+			}
+		}
+
+		break;
+	}
+
+	for (const TSharedPtr<IPropertyHandle>& PropertyHandle : PropertyHandles)
+	{
+		ensure(PropertyHandle.IsValid() && PropertyHandle->IsValidHandle());
+	}
+}
+
+void FAvaLevelViewportComponentTransformDetails::UpdatePropertyHandlesObjects(TArray<UObject*> NewSceneComponents)
+{
+	// Cached the old handles objects.
+	CachedHandlesObjects.Reset(NewSceneComponents.Num());
+	Algo::Transform(NewSceneComponents, CachedHandlesObjects, [](UObject* Obj) { return TWeakObjectPtr<UObject>(Obj); });
+
+	// If the new outer objects list is empty, the property handle is destroyed.
+	if (NewSceneComponents.IsEmpty())
+	{
+		NewSceneComponents = {USceneComponent::StaticClass()->GetDefaultObject()};
+	}
+
+	for (TSharedPtr<IPropertyHandle>& Handle : PropertyHandles)
+	{
+		if (Handle)
+		{
+			Handle->ReplaceOuterObjects(NewSceneComponents);
+		}
+	}
+}
+
+bool FAvaLevelViewportComponentTransformDetails::GetIsEnabled() const
+{
+	return SelectedActorInfo.NumSelected > 0;
+}
+
+const FSlateBrush* FAvaLevelViewportComponentTransformDetails::GetPreserveScaleRatioImage() const
+{
+	return bPreserveScaleRatio
+		? FAppStyle::GetBrush(TEXT("Icons.Lock"))
+		: FAppStyle::GetBrush(TEXT("Icons.Unlock")) ;
+}
+
+ECheckBoxState FAvaLevelViewportComponentTransformDetails::IsPreserveScaleRatioChecked() const
+{
+	return bPreserveScaleRatio
+		? ECheckBoxState::Checked
+		: ECheckBoxState::Unchecked;
+}
+
+void FAvaLevelViewportComponentTransformDetails::OnPreserveScaleRatioToggled(ECheckBoxState NewState)
+{
+	bPreserveScaleRatio = (NewState == ECheckBoxState::Checked);
+	GConfig->SetBool(TEXT("SelectionDetails"), TEXT("PreserveScaleRatio"), bPreserveScaleRatio, GEditorPerProjectIni);
+}
+
+FText FAvaLevelViewportComponentTransformDetails::GetLocationText() const
+{
+	return bAbsoluteLocation
+		? LOCTEXT("AbsoluteLocation", "Absolute Location")
+		: LOCTEXT("Location", "Location");
+}
+
+FText FAvaLevelViewportComponentTransformDetails::GetRotationText() const
+{
+	return bAbsoluteRotation
+		? LOCTEXT("AbsoluteRotation", "Absolute Rotation")
+		: LOCTEXT("Rotation", "Rotation");
+}
+
+FText FAvaLevelViewportComponentTransformDetails::GetScaleText() const
+{
+	return bAbsoluteScale
+		? LOCTEXT("AbsoluteScale", "Absolute Scale")
+		: LOCTEXT("Scale", "Scale");
+}
+
+void FAvaLevelViewportComponentTransformDetails::OnSetAbsoluteTransform(EAvaLevelViewportTransformFieldType TransformField, bool bAbsoluteEnabled)
+{
+	FBoolProperty* AbsoluteProperty = nullptr;
+	FText TransactionText;
+
+	switch (TransformField)
+	{
+		case EAvaLevelViewportTransformFieldType::Location:
+			AbsoluteProperty = FindFProperty<FBoolProperty>(USceneComponent::StaticClass(), USceneComponent::GetAbsoluteLocationPropertyName());
+			TransactionText = LOCTEXT("ToggleAbsoluteLocation", "Toggle Absolute Location");
+			break;
+		
+		case EAvaLevelViewportTransformFieldType::Rotation:
+			AbsoluteProperty = FindFProperty<FBoolProperty>(USceneComponent::StaticClass(), USceneComponent::GetAbsoluteRotationPropertyName());
+			TransactionText = LOCTEXT("ToggleAbsoluteRotation", "Toggle Absolute Rotation");
+			break;
+		
+		case EAvaLevelViewportTransformFieldType::Scale:
+			AbsoluteProperty = FindFProperty<FBoolProperty>(USceneComponent::StaticClass(), USceneComponent::GetAbsoluteScalePropertyName());
+			TransactionText = LOCTEXT("ToggleAbsoluteScale", "Toggle Absolute Scale");
+			break;
+		
+		default:
+			return;
+	}
+
+	bool bBeganTransaction = false;
+	TArray<UObject*> ModifiedObjectsLocal;
+	for (int32 ObjectIndex = 0; ObjectIndex < SelectedObjects.Num(); ++ObjectIndex)
+	{
+		TWeakObjectPtr<UObject> ObjectPtr = SelectedObjects[ObjectIndex];
+		if (ObjectPtr.IsValid())
+		{
+			UObject* Object = ObjectPtr.Get();
+			USceneComponent* SceneComponent = GetSceneComponentFromDetailsObject(Object);
+			if (SceneComponent)
+			{
+				bool bOldValue = TransformField == EAvaLevelViewportTransformFieldType::Location ? SceneComponent->IsUsingAbsoluteLocation() : (TransformField == EAvaLevelViewportTransformFieldType::Rotation ? SceneComponent->IsUsingAbsoluteRotation() : SceneComponent->IsUsingAbsoluteScale());
+
+				if (bOldValue == bAbsoluteEnabled)
+				{
+					// Already the desired value
+					continue;
+				}
+
+				if (!bBeganTransaction)
+				{
+					// NOTE: One transaction per change, not per actor
+					GEditor->BeginTransaction(TransactionText);
+					bBeganTransaction = true;
+				}
+
+				FAvaLevelViewportScopedSwitchWorldForObject WorldSwitcher(Object);
+
+				if (SceneComponent->HasAnyFlags(RF_DefaultSubObject))
+				{
+					// Default subobjects must be included in any undo/redo operations
+					SceneComponent->SetFlags(RF_Transactional);
+				}
+
+				SceneComponent->PreEditChange(AbsoluteProperty);
+
+				if (NotifyHook)
+				{
+					NotifyHook->NotifyPreChange(AbsoluteProperty);
+				}
+
+				switch (TransformField)
+				{
+				case EAvaLevelViewportTransformFieldType::Location:
+					SceneComponent->SetUsingAbsoluteLocation(bAbsoluteEnabled);
+
+					// Update RelativeLocation to maintain/stabilize position when switching between relative and world.
+					if (SceneComponent->GetAttachParent())
+					{
+						if (SceneComponent->IsUsingAbsoluteLocation())
+						{
+							SceneComponent->SetRelativeLocation_Direct(SceneComponent->GetComponentTransform().GetTranslation());
+						}
+						else
+						{
+							FTransform ParentToWorld = SceneComponent->GetAttachParent()->GetSocketTransform(SceneComponent->GetAttachSocketName());
+							FTransform RelativeTM = SceneComponent->GetComponentTransform().GetRelativeTransform(ParentToWorld);
+							SceneComponent->SetRelativeLocation_Direct(RelativeTM.GetTranslation());
+						}
+					}
+					break;
+					
+				case EAvaLevelViewportTransformFieldType::Rotation:
+					SceneComponent->SetUsingAbsoluteRotation(bAbsoluteEnabled);
+					break;
+					
+				case EAvaLevelViewportTransformFieldType::Scale:
+					SceneComponent->SetUsingAbsoluteScale(bAbsoluteEnabled);
+					break;
+				}
+
+				ModifiedObjectsLocal.Add(Object);
+			}
+		}
+	}
+
+	if (bBeganTransaction)
+	{
+		FPropertyChangedEvent PropertyChangedEvent(AbsoluteProperty, EPropertyChangeType::ValueSet, MakeArrayView(ModifiedObjectsLocal));
+
+		for (UObject* Object : ModifiedObjectsLocal)
+		{
+			USceneComponent* SceneComponent = GetSceneComponentFromDetailsObject(Object);
+
+			if (SceneComponent)
+			{
+				SceneComponent->PostEditChangeProperty(PropertyChangedEvent);
+			}
+		}
+
+		if (NotifyHook)
+		{
+			NotifyHook->NotifyPostChange(PropertyChangedEvent, AbsoluteProperty);
+		}
+
+		GEditor->EndTransaction();
+
+		GUnrealEd->RedrawLevelEditingViewports();
+	}
+}
+
+bool FAvaLevelViewportComponentTransformDetails::IsAbsoluteTransformChecked(EAvaLevelViewportTransformFieldType TransformField, bool bAbsoluteEnabled) const
+{
+	switch (TransformField)
+	{
+	case EAvaLevelViewportTransformFieldType::Location:
+		return bAbsoluteLocation == bAbsoluteEnabled;
+		
+	case EAvaLevelViewportTransformFieldType::Rotation:
+		return bAbsoluteRotation == bAbsoluteEnabled;
+		
+	case EAvaLevelViewportTransformFieldType::Scale:
+		return bAbsoluteScale == bAbsoluteEnabled;
+		
+	default:
+		return false;
+	}
+}
+
+struct FGetRootComponentArchetype
+{
+	static USceneComponent* Get(UObject* Object)
+	{
+		USceneComponent* const RootComponent = Object
+			? GetSceneComponentFromDetailsObject(Object)
+			: nullptr;
+		
+		return RootComponent
+			? Cast<USceneComponent>(RootComponent->GetArchetype())
+			: nullptr;
+	}
+};
+
+bool FAvaLevelViewportComponentTransformDetails::GetLocationResetVisibility() const
+{
+	const USceneComponent* Archetype = FGetRootComponentArchetype::Get(SelectedObjects[0].Get());
+	
+	const FVector Data = Archetype
+		? Archetype->GetRelativeLocation()
+		: FVector::ZeroVector;
+
+	// unset means multiple differing values, so show "Reset to Default" in that case
+	return CachedLocation.IsSet()
+			&& CachedLocation.X.GetValue() == Data.X
+			&& CachedLocation.Y.GetValue() == Data.Y
+			&& CachedLocation.Z.GetValue() == Data.Z
+		? false
+		: true;
+}
+
+void FAvaLevelViewportComponentTransformDetails::OnLocationResetClicked()
+{
+	const FText TransactionName = LOCTEXT("ResetLocation", "Reset Location");
+	FScopedTransaction Transaction(TransactionName);
+
+	const USceneComponent* Archetype = FGetRootComponentArchetype::Get(SelectedObjects[0].Get());
+	const FVector Data = Archetype ? Archetype->GetRelativeLocation() : FVector::ZeroVector;
+
+	OnBeginChange(EAvaLevelViewportTransformFieldType::Location);
+	OnSetTransform(EAvaLevelViewportTransformFieldType::Location, EAxisList::All, Data, /* bMirror */ false, /* bCommitted */ true);
+	OnEndChange(EAvaLevelViewportTransformFieldType::Location, EAxisList::All);
+}
+
+bool FAvaLevelViewportComponentTransformDetails::GetRotationResetVisibility() const
+{
+	const USceneComponent* Archetype = FGetRootComponentArchetype::Get(SelectedObjects[0].Get());
+	const FVector Data = Archetype ? Archetype->GetRelativeRotation().Euler() : FVector::ZeroVector;
+
+	// unset means multiple differing values, so show "Reset to Default" in that case
+	return CachedRotation.IsSet()
+			&& CachedRotation.X.GetValue() == Data.X
+			&& CachedRotation.Y.GetValue() == Data.Y
+			&& CachedRotation.Z.GetValue() == Data.Z
+		? false
+		: true;
+}
+
+void FAvaLevelViewportComponentTransformDetails::OnRotationResetClicked()
+{
+	const FText TransactionName = LOCTEXT("ResetRotation", "Reset Rotation");
+	FScopedTransaction Transaction(TransactionName);
+
+	const USceneComponent* Archetype = FGetRootComponentArchetype::Get(SelectedObjects[0].Get());
+	const FVector Data = Archetype ? Archetype->GetRelativeRotation().Euler() : FVector::ZeroVector;
+
+	OnBeginChange(EAvaLevelViewportTransformFieldType::Rotation);
+	OnSetTransform(EAvaLevelViewportTransformFieldType::Rotation, EAxisList::All, Data, /* bMirror */ false, /* bCommitted */ true);
+	OnEndChange(EAvaLevelViewportTransformFieldType::Rotation, EAxisList::All);
+}
+
+bool FAvaLevelViewportComponentTransformDetails::GetScaleResetVisibility() const
+{
+	const USceneComponent* Archetype = FGetRootComponentArchetype::Get(SelectedObjects[0].Get());
+	const FVector Data = Archetype ? Archetype->GetRelativeScale3D() : FVector(1.0f);
+
+	// unset means multiple differing values, so show "Reset to Default" in that case
+	return CachedScale.IsSet()
+			&& CachedScale.X.GetValue() == Data.X
+			&& CachedScale.Y.GetValue() == Data.Y
+			&& CachedScale.Z.GetValue() == Data.Z
+		? false
+		: true;
+}
+
+void FAvaLevelViewportComponentTransformDetails::OnScaleResetClicked()
+{
+	const FText TransactionName = LOCTEXT("ResetScale", "Reset Scale");
+	FScopedTransaction Transaction(TransactionName);
+
+	const USceneComponent* Archetype = FGetRootComponentArchetype::Get(SelectedObjects[0].Get());
+	const FVector Data = Archetype ? Archetype->GetRelativeScale3D() : FVector(1.0f);
+
+	OnBeginChange(EAvaLevelViewportTransformFieldType::Scale);
+	OnSetTransform(EAvaLevelViewportTransformFieldType::Scale, EAxisList::All, Data, /* bMirror */ false, /* bCommitted */ true);
+	OnEndChange(EAvaLevelViewportTransformFieldType::Scale, EAxisList::All);
+}
+
+void FAvaLevelViewportComponentTransformDetails::ExtendXScaleContextMenu(FMenuBuilder& MenuBuilder)
+{
+	MenuBuilder.BeginSection("ScaleOperations", LOCTEXT("ScaleOperations", "Scale Operations"));
+	MenuBuilder.AddMenuEntry(
+		LOCTEXT("MirrorValueX", "Mirror X"),
+		LOCTEXT("MirrorValueX_Tooltip", "Mirror scale value on the X axis"),
+		FSlateIcon(),
+		FUIAction(FExecuteAction::CreateSP(this, &FAvaLevelViewportComponentTransformDetails::OnXScaleMirrored), FCanExecuteAction())
+	);
+	MenuBuilder.EndSection();
+}
+
+void FAvaLevelViewportComponentTransformDetails::ExtendYScaleContextMenu(FMenuBuilder& MenuBuilder)
+{
+	MenuBuilder.BeginSection("ScaleOperations", LOCTEXT("ScaleOperations", "Scale Operations"));
+	MenuBuilder.AddMenuEntry(
+		LOCTEXT("MirrorValueY", "Mirror Y"),
+		LOCTEXT("MirrorValueY_Tooltip", "Mirror scale value on the Y axis"),
+		FSlateIcon(),
+		FUIAction(FExecuteAction::CreateSP(this, &FAvaLevelViewportComponentTransformDetails::OnYScaleMirrored), FCanExecuteAction())
+	);
+	MenuBuilder.EndSection();
+}
+
+void FAvaLevelViewportComponentTransformDetails::ExtendZScaleContextMenu(FMenuBuilder& MenuBuilder)
+{
+	MenuBuilder.BeginSection("ScaleOperations", LOCTEXT("ScaleOperations", "Scale Operations"));
+	MenuBuilder.AddMenuEntry(
+		LOCTEXT("MirrorValueZ", "Mirror Z"),
+		LOCTEXT("MirrorValueZ_Tooltip", "Mirror scale value on the Z axis"),
+		FSlateIcon(),
+		FUIAction(FExecuteAction::CreateSP(this, &FAvaLevelViewportComponentTransformDetails::OnZScaleMirrored), FCanExecuteAction())
+	);
+	MenuBuilder.EndSection();
+}
+
+void FAvaLevelViewportComponentTransformDetails::OnXScaleMirrored()
+{
+	FSlateApplication::Get().ClearKeyboardFocus(EFocusCause::Mouse);
+	FScopedTransaction Transaction(LOCTEXT("MirrorActorScaleX", "Mirror actor scale X"));
+	OnBeginChange(EAvaLevelViewportTransformFieldType::Scale);
+	OnSetTransform(EAvaLevelViewportTransformFieldType::Scale, EAxisList::X, FVector(1.0f), /* bMirror */ true, /* bCommitted */ true);
+	OnEndChange(EAvaLevelViewportTransformFieldType::Scale, EAxisList::X);
+}
+
+void FAvaLevelViewportComponentTransformDetails::OnYScaleMirrored()
+{
+	FSlateApplication::Get().ClearKeyboardFocus(EFocusCause::Mouse);
+	FScopedTransaction Transaction(LOCTEXT("MirrorActorScaleY", "Mirror actor scale Y"));
+	OnBeginChange(EAvaLevelViewportTransformFieldType::Scale);
+	OnSetTransform(EAvaLevelViewportTransformFieldType::Scale, EAxisList::Y, FVector(1.0f), /* bMirror */ true, /* bCommitted */ true);
+	OnEndChange(EAvaLevelViewportTransformFieldType::Scale, EAxisList::Y);
+}
+
+void FAvaLevelViewportComponentTransformDetails::OnZScaleMirrored()
+{
+	FSlateApplication::Get().ClearKeyboardFocus(EFocusCause::Mouse);
+	FScopedTransaction Transaction(LOCTEXT("MirrorActorScaleZ", "Mirror actor scale Z"));
+	OnBeginChange(EAvaLevelViewportTransformFieldType::Scale);
+	OnSetTransform(EAvaLevelViewportTransformFieldType::Scale, EAxisList::Z, FVector(1.0f), /* bMirror */ true, /* bCommitted */ true);
+	OnEndChange(EAvaLevelViewportTransformFieldType::Scale, EAxisList::Z);
+}
+
+void FAvaLevelViewportComponentTransformDetails::CacheTransform()
+{
+	if (bCanUseSelectionAsParent && bUseSelectionAsParent)
+	{
+		// @TODO: Get selection accumulated transform stuff
+
+		// Reset transforms
+		FVector PivotLocation = FVector::ZeroVector; // BlueprintEditorSP->GetAvaModeTools()->PivotLocation;
+		CachedLocation.X = PivotLocation.X;
+		CachedLocation.Y = PivotLocation.Y;
+		CachedLocation.Z = PivotLocation.Z;
+
+		FRotator SelectionRotation = FRotator::ZeroRotator; // BlueprintEditorSP->GetAvaModeTools()->GetAccumulatedTransform().Rotator();
+		CachedRotation.X = SelectionRotation.Roll;
+		CachedRotation.Y = SelectionRotation.Pitch;
+		CachedRotation.Z = SelectionRotation.Yaw;
+
+		FVector SelectionScale = FVector::ZeroVector; // BlueprintEditorSP->GetAvaModeTools()->GetAccumulatedTransform().GetScale3D();
+		CachedScale.X = SelectionScale.X;
+		CachedScale.Y = SelectionScale.Y;
+		CachedScale.Z = SelectionScale.Z;
+	}
+
+	FVector CurLoc;
+	FRotator CurRot;
+	FVector CurScale;
+	bool bFirstObject = true;
+	bool bSkipFirst = false;
+
+	if (SelectedObjects.Num() == 2)
+	{
+		AActor* First = Cast<AActor>(SelectedObjects[0].Get());
+		UActorComponent* Second = Cast<UActorComponent>(SelectedObjects[1].Get());
+
+		if (First && Second && Second->GetOwner() == First)
+		{
+			bSkipFirst = true;
+		}
+	}
+
+	for (int32 ObjectIndex = 0; ObjectIndex < SelectedObjects.Num(); ++ObjectIndex)
+	{
+		if (ObjectIndex == 0 && bSkipFirst)
+		{
+			continue;
+		}
+
+		TWeakObjectPtr<UObject> ObjectPtr = SelectedObjects[ObjectIndex];
+
+		if (ObjectPtr.IsValid())
+		{
+			UObject* Object = ObjectPtr.Get();
+			USceneComponent* SceneComponent = GetSceneComponentFromDetailsObject(Object);
+
+			if (bCanUseSelectionAsParent && bUseSelectionAsParent)
+			{
+				bAbsoluteLocation = SceneComponent ? SceneComponent->IsUsingAbsoluteLocation() : false;
+				bAbsoluteScale = SceneComponent ? SceneComponent->IsUsingAbsoluteScale() : false;
+				bAbsoluteRotation = SceneComponent ? SceneComponent->IsUsingAbsoluteRotation() : false;
+				return;
+			}
+
+			if (!SceneComponent)
+			{
+				continue;
+			}
+
+			FVector Loc = SceneComponent->GetRelativeLocation();
+			FRotator* FoundRotator = ObjectToRelativeRotationMap.Find(SceneComponent);
+			FRotator Rot = (bEditingRotationInUI && !Object->IsTemplate() && FoundRotator) ? *FoundRotator : SceneComponent->GetRelativeRotation();
+			FVector Scale = SceneComponent->GetRelativeScale3D();
+
+			if (bFirstObject)
+			{
+				// Cache the current values from the first actor to see if any values differ among other actors
+				CurLoc = Loc;
+				CurRot = Rot;
+				CurScale = Scale;
+
+				CachedLocation.Set(Loc);
+				CachedRotation.Set(Rot);
+				CachedScale.Set(Scale);
+
+				bAbsoluteLocation = SceneComponent->IsUsingAbsoluteLocation();
+				bAbsoluteScale = SceneComponent->IsUsingAbsoluteScale();
+				bAbsoluteRotation = SceneComponent->IsUsingAbsoluteRotation();
+
+				bFirstObject = false;
+			}
+			else if (CurLoc != Loc || CurRot != Rot || CurScale != Scale)
+			{
+				// Check which values differ and unset the different values
+				CachedLocation.X = Loc.X == CurLoc.X && CachedLocation.X.IsSet() ? Loc.X : TOptional<FVector::FReal>();
+				CachedLocation.Y = Loc.Y == CurLoc.Y && CachedLocation.Y.IsSet() ? Loc.Y : TOptional<FVector::FReal>();
+				CachedLocation.Z = Loc.Z == CurLoc.Z && CachedLocation.Z.IsSet() ? Loc.Z : TOptional<FVector::FReal>();
+
+				CachedRotation.X = Rot.Roll == CurRot.Roll && CachedRotation.X.IsSet() ? Rot.Roll : TOptional<float>();
+				CachedRotation.Y = Rot.Pitch == CurRot.Pitch && CachedRotation.Y.IsSet() ? Rot.Pitch : TOptional<float>();
+				CachedRotation.Z = Rot.Yaw == CurRot.Yaw && CachedRotation.Z.IsSet() ? Rot.Yaw : TOptional<float>();
+
+				CachedScale.X = Scale.X == CurScale.X && CachedScale.X.IsSet() ? Scale.X : TOptional<FVector::FReal>();
+				CachedScale.Y = Scale.Y == CurScale.Y && CachedScale.Y.IsSet() ? Scale.Y : TOptional<FVector::FReal>();
+				CachedScale.Z = Scale.Z == CurScale.Z && CachedScale.Z.IsSet() ? Scale.Z : TOptional<FVector::FReal>();
+
+				// If all values are unset all values are different and we can stop looking
+				const bool bAllValuesDiffer = !CachedLocation.IsSet() && !CachedRotation.IsSet() && !CachedScale.IsSet();
+
+				if (bAllValuesDiffer)
+				{
+					break;
+				}
+			}
+		}
+	}
+}
+
+FVector FAvaLevelViewportComponentTransformDetails::GetAxisFilteredVector(EAxisList::Type Axis, const FVector& NewValue, const FVector& OldValue)
+{
+	return FVector((Axis & EAxisList::X) ? NewValue.X : OldValue.X,
+		(Axis & EAxisList::Y) ? NewValue.Y : OldValue.Y,
+		(Axis & EAxisList::Z) ? NewValue.Z : OldValue.Z);
+}
+
+void FAvaLevelViewportComponentTransformDetails::OnBeginChange(EAvaLevelViewportTransformFieldType TransformField)
+{
+	ModifiedObjects.Empty();
+	OriginalLocations.Empty();
+	OriginalRotations.Empty();
+	OriginalScales.Empty();
+
+	FProperty* ValueProperty = nullptr;
+
+	switch (TransformField)
+	{
+		case EAvaLevelViewportTransformFieldType::Location:
+			ValueProperty = FindFProperty<FProperty>(USceneComponent::StaticClass(), USceneComponent::GetRelativeLocationPropertyName());
+			break;
+
+		case EAvaLevelViewportTransformFieldType::Rotation:
+			ValueProperty = FindFProperty<FProperty>(USceneComponent::StaticClass(), USceneComponent::GetRelativeRotationPropertyName());
+			break;
+
+		case EAvaLevelViewportTransformFieldType::Scale:
+			ValueProperty = FindFProperty<FProperty>(USceneComponent::StaticClass(), USceneComponent::GetRelativeScale3DPropertyName());
+			break;
+
+		default:
+			return;
+	}
+
+	for (int32 ObjectIndex = 0; ObjectIndex < SelectedObjects.Num(); ++ObjectIndex)
+	{
+		TWeakObjectPtr<UObject> ObjectPtr = SelectedObjects[ObjectIndex];
+
+		if (ObjectPtr.IsValid())
+		{
+			UObject* Object = ObjectPtr.Get();
+
+			AActor* Actor = Cast<AActor>(Object);
+
+			if (Actor)
+			{
+				Actor->Modify();
+				ModifiedObjects.Add(Actor);
+
+				if (USceneComponent* RootComponent = Actor->GetRootComponent())
+				{
+					RootComponent->Modify();
+					ModifiedObjects.Add(RootComponent);
+				}
+
+				continue;
+			}
+
+			USceneComponent* SceneComponent = GetSceneComponentFromDetailsObject(Object);
+
+			if (SceneComponent)
+			{
+				if (SceneComponent->HasAnyFlags(RF_DefaultSubObject))
+				{
+					// Default subobjects must be included in any undo/redo operations
+					SceneComponent->SetFlags(RF_Transactional);
+				}
+
+				SceneComponent->Modify();
+				SceneComponent->PreEditChange(ValueProperty);
+
+				if (NotifyHook)
+				{
+					NotifyHook->NotifyPreChange(ValueProperty);
+				}
+
+				ModifiedObjects.Add(SceneComponent);
+
+				switch (TransformField)
+				{
+					case EAvaLevelViewportTransformFieldType::Location:
+						OriginalLocations.Emplace(SceneComponent, SceneComponent->GetRelativeLocation());
+						break;
+
+					case EAvaLevelViewportTransformFieldType::Rotation:
+						OriginalRotations.Emplace(SceneComponent, SceneComponent->GetRelativeRotation());
+						break;
+
+					case EAvaLevelViewportTransformFieldType::Scale:
+						OriginalScales.Emplace(SceneComponent, SceneComponent->GetRelativeScale3D());
+						break;
+
+					default:
+						break;
+				}
+			}
+
+			GEditor->BroadcastBeginObjectMovement(*SceneComponent);
+			AActor* EditedActor = SceneComponent->GetOwner();
+
+			if (EditedActor && EditedActor->GetRootComponent() == SceneComponent)
+			{
+				GEditor->BroadcastBeginObjectMovement(*EditedActor);
+			}
+		}
+	}
+}
+
+void FAvaLevelViewportComponentTransformDetails::OnEndChange(EAvaLevelViewportTransformFieldType TransformField, EAxisList::Type Axis)
+{
+	FProperty* ValueProperty = nullptr;
+	FProperty* AxisProperty = nullptr;
+
+	switch (TransformField)
+	{
+		case EAvaLevelViewportTransformFieldType::Location:
+			ValueProperty = FindFProperty<FProperty>(USceneComponent::StaticClass(), USceneComponent::GetRelativeLocationPropertyName());
+
+			// Only set axis property for single axis set
+			if (Axis == EAxisList::X)
+			{
+				AxisProperty = FindFProperty<FDoubleProperty>(TBaseStructure<FVector>::Get(), GET_MEMBER_NAME_CHECKED(FVector, X));
+				check(AxisProperty != nullptr);
+			}
+			else if (Axis == EAxisList::Y)
+			{
+				AxisProperty = FindFProperty<FDoubleProperty>(TBaseStructure<FVector>::Get(), GET_MEMBER_NAME_CHECKED(FVector, Y));
+				check(AxisProperty != nullptr);
+			}
+			else if (Axis == EAxisList::Z)
+			{
+				AxisProperty = FindFProperty<FDoubleProperty>(TBaseStructure<FVector>::Get(), GET_MEMBER_NAME_CHECKED(FVector, Z));
+				check(AxisProperty != nullptr);
+			}
+			break;
+
+		case EAvaLevelViewportTransformFieldType::Rotation:
+			ValueProperty = FindFProperty<FProperty>(USceneComponent::StaticClass(), USceneComponent::GetRelativeRotationPropertyName());
+
+			// Only set axis property for single axis set
+			if (Axis == EAxisList::X)
+			{
+				AxisProperty = FindFProperty<FDoubleProperty>(TBaseStructure<FRotator>::Get(), GET_MEMBER_NAME_CHECKED(FRotator, Roll));
+				check(AxisProperty != nullptr);
+			}
+			else if (Axis == EAxisList::Y)
+			{
+				AxisProperty = FindFProperty<FDoubleProperty>(TBaseStructure<FRotator>::Get(), GET_MEMBER_NAME_CHECKED(FRotator, Pitch));
+				check(AxisProperty != nullptr);
+			}
+			else if (Axis == EAxisList::Z)
+			{
+				AxisProperty = FindFProperty<FDoubleProperty>(TBaseStructure<FRotator>::Get(), GET_MEMBER_NAME_CHECKED(FRotator, Yaw));
+				check(AxisProperty != nullptr);
+			}
+			break;
+
+		case EAvaLevelViewportTransformFieldType::Scale:
+			ValueProperty = FindFProperty<FProperty>(USceneComponent::StaticClass(), USceneComponent::GetRelativeScale3DPropertyName());
+
+			// If keep scale is set, don't set axis property
+			if (!bPreserveScaleRatio)
+			{
+				if (Axis == EAxisList::X)
+				{
+					AxisProperty = FindFProperty<FDoubleProperty>(TBaseStructure<FVector>::Get(), GET_MEMBER_NAME_CHECKED(FVector, X));
+					check(AxisProperty != nullptr);
+				}
+				else if (Axis == EAxisList::Y)
+				{
+					AxisProperty = FindFProperty<FDoubleProperty>(TBaseStructure<FVector>::Get(), GET_MEMBER_NAME_CHECKED(FVector, Y));
+					check(AxisProperty != nullptr);
+				}
+				else if (Axis == EAxisList::Z)
+				{
+					AxisProperty = FindFProperty<FDoubleProperty>(TBaseStructure<FVector>::Get(), GET_MEMBER_NAME_CHECKED(FVector, Z));
+					check(AxisProperty != nullptr);
+				}
+			}
+			break;
+
+		default:
+			return;
+	}
+
+	TArray<UObject*> ArrayViewArray;
+	Algo::Transform(ModifiedObjects, ArrayViewArray, [](const TWeakObjectPtr<UObject>& Ptr) { return Ptr.Get(); });
+
+	FPropertyChangedEvent PropertyChangedEvent(ValueProperty, EPropertyChangeType::ValueSet, MakeArrayView(ArrayViewArray));
+	FEditPropertyChain PropertyChain;
+
+	if (AxisProperty)
+	{
+		PropertyChain.AddHead(AxisProperty);
+	}
+	
+	PropertyChain.AddHead(ValueProperty);
+	
+	FPropertyChangedChainEvent PropertyChangedChainEvent(PropertyChain, PropertyChangedEvent);
+
+	for (int32 ObjectIndex = 0; ObjectIndex < ModifiedObjects.Num(); ++ObjectIndex)
+	{
+		TWeakObjectPtr<UObject> ObjectPtr = ModifiedObjects[ObjectIndex];
+
+		if (ObjectPtr.IsValid())
+		{
+			UObject* Object = ObjectPtr.Get();
+
+			AActor* Actor = Cast<AActor>(Object);
+
+			if (Actor)
+			{
+				Actor->PostEditChangeChainProperty(PropertyChangedChainEvent);
+				Actor->PostEditMove(true);
+				continue;
+			}
+
+			USceneComponent* SceneComponent = GetSceneComponentFromDetailsObject(Object);
+			USceneComponent* OldSceneComponent = SceneComponent;
+
+			if (SceneComponent)
+			{
+				AActor* EditedActor = SceneComponent->GetOwner();
+				FString SceneComponentPath = SceneComponent->GetPathName(EditedActor);
+
+				// This can invalidate OldSceneComponent
+				OldSceneComponent->PostEditChangeChainProperty(PropertyChangedChainEvent);
+
+				// Make sure we've got the latest one!
+				SceneComponent = FindObject<USceneComponent>(EditedActor, *SceneComponentPath);
+
+				if (SceneComponent)
+				{
+					// Broadcast when the actor is done moving
+					GEditor->BroadcastEndObjectMovement(*SceneComponent);
+
+					if (EditedActor && EditedActor->GetRootComponent() == SceneComponent)
+					{
+						GEditor->BroadcastEndObjectMovement(*EditedActor);
+					}
+				}
+			}
+		}
+	}
+
+	if (NotifyHook)
+	{
+		NotifyHook->NotifyPostChange(PropertyChangedEvent, ValueProperty);
+	}
+}
+
+void FAvaLevelViewportComponentTransformDetails::OnSetTransform(EAvaLevelViewportTransformFieldType TransformField, EAxisList::Type Axis, FVector NewValue, bool bMirror, bool bCommitted)
+{
+	if (bCanUseSelectionAsParent && bUseSelectionAsParent)
+	{
+		OnSetTransformSelectionParent(TransformField, Axis, NewValue, bMirror, bCommitted);
+		return;
+	}
+
+	for (int32 ObjectIndex = 0; ObjectIndex < SelectedObjects.Num(); ++ObjectIndex)
+	{
+		TWeakObjectPtr<UObject> ObjectPtr = SelectedObjects[ObjectIndex];
+
+		if (ObjectPtr.IsValid())
+		{
+			UObject* Object = ObjectPtr.Get();
+			USceneComponent* SceneComponent = GetSceneComponentFromDetailsObject(Object);
+
+			if (SceneComponent)
+			{
+				AActor* EditedActor = SceneComponent->GetOwner();
+
+				FVector OldComponentValue;
+				FVector NewComponentValue;
+
+				switch (TransformField)
+				{
+					case EAvaLevelViewportTransformFieldType::Location:
+						OldComponentValue = SceneComponent->GetRelativeLocation();
+						break;
+
+					case EAvaLevelViewportTransformFieldType::Rotation:
+						// Pull from the actual component or from the cache
+						OldComponentValue = SceneComponent->GetRelativeRotation().Euler();
+						if (bEditingRotationInUI && ObjectToRelativeRotationMap.Find(SceneComponent))
+						{
+							OldComponentValue = ObjectToRelativeRotationMap.Find(SceneComponent)->Euler();
+						}
+						break;
+
+					case EAvaLevelViewportTransformFieldType::Scale:
+						OldComponentValue = SceneComponent->GetRelativeScale3D();
+						break;
+				}
+
+				// Set the incoming value
+				if (bMirror)
+				{
+					NewComponentValue = GetAxisFilteredVector(Axis, -OldComponentValue, OldComponentValue);
+				}
+				else
+				{
+					NewComponentValue = GetAxisFilteredVector(Axis, NewValue, OldComponentValue);
+				}
+
+				// If we're committing during a slider transaction then we need to force it, in order that PostEditChangeChainProperty be called.
+				// Note: this will even happen if the slider hasn't changed the value.
+				bool bDoChange = bCommitted && bIsSliderTransaction;
+				bDoChange = bDoChange || (OldComponentValue != NewComponentValue);
+
+				if (bDoChange)
+				{
+					FAvaLevelViewportScopedSwitchWorldForObject WorldSwitcher(Object);
+
+					switch (TransformField)
+					{
+						case EAvaLevelViewportTransformFieldType::Location:
+							{
+								ObjectToRelativeRotationMap.FindOrAdd(SceneComponent) = SceneComponent->GetRelativeRotation();
+								SceneComponent->SetRelativeLocation(NewComponentValue);
+
+								// Also forcibly set it as the cache may have changed it slightly
+								SceneComponent->SetRelativeLocation_Direct(NewComponentValue);
+								break;
+							}
+						case EAvaLevelViewportTransformFieldType::Rotation:
+							{
+								FRotator NewRotation = FRotator::MakeFromEuler(NewComponentValue);
+								ObjectToRelativeRotationMap.FindOrAdd(SceneComponent) = NewRotation;
+								SceneComponent->SetRelativeRotationExact(NewRotation);
+								break;
+							}
+						case EAvaLevelViewportTransformFieldType::Scale:
+							{
+								if ((!bCanUseSelectionAsParent || !bUseSelectionAsParent) && bPreserveScaleRatio)
+								{
+									// If we set a single axis, scale the others
+									float Ratio = 0.0f;
+
+									switch (Axis)
+									{
+										case EAxisList::X:
+											if (bIsSliderTransaction)
+											{
+												Ratio = SliderScaleRatio.X == 0.0f ? SliderScaleRatio.Y : (SliderScaleRatio.Y / SliderScaleRatio.X);
+												NewComponentValue.Y = NewComponentValue.X * Ratio;
+
+												Ratio = SliderScaleRatio.X == 0.0f ? SliderScaleRatio.Z : (SliderScaleRatio.Z / SliderScaleRatio.X);
+												NewComponentValue.Z = NewComponentValue.X * Ratio;
+											}
+											else
+											{
+												Ratio = OldComponentValue.X == 0.0f ? NewComponentValue.Z : NewComponentValue.X / OldComponentValue.X;
+												NewComponentValue.Y *= Ratio;
+												NewComponentValue.Z *= Ratio;
+											}
+										break;
+										case EAxisList::Y:
+											if (bIsSliderTransaction)
+											{
+												Ratio = SliderScaleRatio.Y == 0.0f ? SliderScaleRatio.X : (SliderScaleRatio.X / SliderScaleRatio.Y);
+												NewComponentValue.X = NewComponentValue.Y * Ratio;
+
+												Ratio = SliderScaleRatio.Y == 0.0f ? SliderScaleRatio.Z : (SliderScaleRatio.Z / SliderScaleRatio.Y);
+												NewComponentValue.Z = NewComponentValue.Y * Ratio;
+											}
+											else
+											{
+												Ratio = OldComponentValue.Y == 0.0f ? NewComponentValue.Z : NewComponentValue.Y / OldComponentValue.Y;
+												NewComponentValue.X *= Ratio;
+												NewComponentValue.Z *= Ratio;
+											}
+										break;
+										case EAxisList::Z:
+											if (bIsSliderTransaction)
+											{
+												Ratio = SliderScaleRatio.Z == 0.0f ? SliderScaleRatio.X : (SliderScaleRatio.X / SliderScaleRatio.Z);
+												NewComponentValue.X = NewComponentValue.Z * Ratio;
+
+												Ratio = SliderScaleRatio.Z == 0.0f ? SliderScaleRatio.Y : (SliderScaleRatio.Y / SliderScaleRatio.Z);
+												NewComponentValue.Y = NewComponentValue.Z * Ratio;
+											}
+											else
+											{
+												Ratio = OldComponentValue.Z == 0.0f ? NewComponentValue.Z : NewComponentValue.Z / OldComponentValue.Z;
+												NewComponentValue.X *= Ratio;
+												NewComponentValue.Y *= Ratio;
+											}
+										break;
+										default:
+											// Do nothing, this set multiple axis at once
+											break;
+									}
+								}
+
+								SceneComponent->SetRelativeScale3D(NewComponentValue);
+								break;
+							}
+					}
+
+					SceneComponent->Modify();
+				}
+			}
+		}
+	}
+
+	CacheTransform();
+
+	if (TransformField == EAvaLevelViewportTransformFieldType::Location)
+	{
+		if (TSharedPtr<FAvaLevelViewportClient> ViewportClient = AvaViewportClientWeak.Pin())
+		{
+			if (FEditorModeTools* ModeTools = ViewportClient->GetModeTools())
+			{
+				FVector CurrentPivot = ModeTools->PivotLocation;
+
+				if (Axis & EAxisList::X)
+				{
+					CurrentPivot.X = NewValue.X;
+				}
+
+				if (Axis & EAxisList::Y)
+				{
+					CurrentPivot.Y = NewValue.Y;
+				}
+
+				if (Axis & EAxisList::Z)
+				{
+					CurrentPivot.Z = NewValue.Z;
+				}
+
+				ModeTools->SetPivotLocation(CurrentPivot, false);
+			}
+		}
+	}
+
+	if (ULevelEditorSubsystem* LevelEditorSubsystem = GEditor->GetEditorSubsystem<ULevelEditorSubsystem>())
+	{
+		LevelEditorSubsystem->EditorInvalidateViewports();
+	}
+}
+
+void FAvaLevelViewportComponentTransformDetails::OnSetTransformSelectionParent(EAvaLevelViewportTransformFieldType TransformField, EAxisList::Type Axis, FVector NewValue, bool bMirror, bool bCommitted)
+{
+	FVector OldSelectionValue;
+	FVector NewSelectionValue;
+	FTransform OldAccumulatedTransform;
+	FTransform NewAccumulatedTransform;
+
+	switch (TransformField)
+	{
+		case EAvaLevelViewportTransformFieldType::Location:
+			OldSelectionValue.X = CachedLocation.X.Get(0.f);
+			OldSelectionValue.Y = CachedLocation.Y.Get(0.f);
+			OldSelectionValue.Z = CachedLocation.Z.Get(0.f);
+			break;
+
+		case EAvaLevelViewportTransformFieldType::Rotation:
+			OldSelectionValue.X = CachedRotation.X.Get(0.f);
+			OldSelectionValue.Y = CachedRotation.Y.Get(0.f);
+			OldSelectionValue.Z = CachedRotation.Z.Get(0.f);
+			break;
+
+		case EAvaLevelViewportTransformFieldType::Scale:
+			OldSelectionValue.X = CachedScale.X.Get(0.f);
+			OldSelectionValue.Y = CachedScale.Y.Get(0.f);
+			OldSelectionValue.Z = CachedScale.Z.Get(0.f);
+			break;
+	}
+
+	if (!bMirror)
+	{
+		NewSelectionValue = GetAxisFilteredVector(Axis, NewValue, OldSelectionValue);
+	}
+
+	else
+	{
+		NewSelectionValue = GetAxisFilteredVector(Axis, -OldSelectionValue, OldSelectionValue);
+	}
+
+	if (TransformField == EAvaLevelViewportTransformFieldType::Scale && bPreserveScaleRatio)
+	{
+		// If we set a single axis, scale the others
+		float Ratio = 0.0f;
+
+		switch (Axis)
+		{
+			case EAxisList::X:
+				if (bIsSliderTransaction)
+				{
+					Ratio = SliderScaleRatio.X == 0.0f ? SliderScaleRatio.Y : (SliderScaleRatio.Y / SliderScaleRatio.X);
+					NewSelectionValue.Y = NewSelectionValue.X * Ratio;
+
+					Ratio = SliderScaleRatio.X == 0.0f ? SliderScaleRatio.Z : (SliderScaleRatio.Z / SliderScaleRatio.X);
+					NewSelectionValue.Z = NewSelectionValue.X * Ratio;
+				}
+				else
+				{
+					Ratio = OldSelectionValue.X == 0.0f ? NewSelectionValue.Z : NewSelectionValue.X / OldSelectionValue.X;
+					NewSelectionValue.Y *= Ratio;
+					NewSelectionValue.Z *= Ratio;
+				}
+			break;
+			case EAxisList::Y:
+				if (bIsSliderTransaction)
+				{
+					Ratio = SliderScaleRatio.Y == 0.0f ? SliderScaleRatio.X : (SliderScaleRatio.X / SliderScaleRatio.Y);
+					NewSelectionValue.X = NewSelectionValue.Y * Ratio;
+
+					Ratio = SliderScaleRatio.Y == 0.0f ? SliderScaleRatio.Z : (SliderScaleRatio.Z / SliderScaleRatio.Y);
+					NewSelectionValue.Z = NewSelectionValue.Y * Ratio;
+				}
+				else
+				{
+					Ratio = OldSelectionValue.Y == 0.0f ? NewSelectionValue.Z : NewSelectionValue.Y / OldSelectionValue.Y;
+					NewSelectionValue.X *= Ratio;
+					NewSelectionValue.Z *= Ratio;
+				}
+			break;
+			case EAxisList::Z:
+				if (bIsSliderTransaction)
+				{
+					Ratio = SliderScaleRatio.Z == 0.0f ? SliderScaleRatio.X : (SliderScaleRatio.X / SliderScaleRatio.Z);
+					NewSelectionValue.X = NewSelectionValue.Z * Ratio;
+
+					Ratio = SliderScaleRatio.Z == 0.0f ? SliderScaleRatio.Y : (SliderScaleRatio.Y / SliderScaleRatio.Z);
+					NewSelectionValue.Y = NewSelectionValue.Z * Ratio;
+				}
+				else
+				{
+					Ratio = OldSelectionValue.Z == 0.0f ? NewSelectionValue.Z : NewSelectionValue.Z / OldSelectionValue.Z;
+					NewSelectionValue.X *= Ratio;
+					NewSelectionValue.Y *= Ratio;
+				}
+			break;
+			default:
+				// Do nothing, this set multiple axis at once
+				break;
+		}
+	}
+
+	bool bTranslationChanged = false;
+	bool bRotationChanged = false;
+	bool bScaleChanged = false;
+
+	// @TODO Accumulated transform
+	FVector SelectionCenter = FVector::ZeroVector; // BlueprintEditorSP->GetAvaModeTools()->GetSelectionTransform().GetLocation();
+	OldAccumulatedTransform = FTransform::Identity; // BlueprintEditorSP->GetAvaModeTools()->GetAccumulatedTransform();
+	NewAccumulatedTransform = OldAccumulatedTransform;
+
+ 	switch (TransformField)
+	{
+		case EAvaLevelViewportTransformFieldType::Location:
+			{
+				CachedLocation.X = NewSelectionValue.X;
+				CachedLocation.Y = NewSelectionValue.Y;
+				CachedLocation.Z = NewSelectionValue.Z;
+				NewAccumulatedTransform.SetLocation(NewSelectionValue);
+				bTranslationChanged = true;
+				break;
+			}
+		case EAvaLevelViewportTransformFieldType::Rotation:
+			{
+				CachedRotation.X = NewSelectionValue.X;
+				CachedRotation.Y = NewSelectionValue.Y;
+				CachedRotation.Z = NewSelectionValue.Z;
+				NewAccumulatedTransform.SetRotation(FRotator(NewSelectionValue.Y, NewSelectionValue.Z, NewSelectionValue.X).Quaternion());
+				bRotationChanged = true;
+				break;
+			}
+		case EAvaLevelViewportTransformFieldType::Scale:
+			{
+				CachedScale.X = NewSelectionValue.X;
+				CachedScale.Y = NewSelectionValue.Y;
+				CachedScale.Z = NewSelectionValue.Z;
+				NewAccumulatedTransform.SetScale3D(NewSelectionValue);
+				bScaleChanged = true;
+				break;
+			}
+	}
+
+	for (int32 ObjectIndex = 0; ObjectIndex < SelectedObjects.Num(); ++ObjectIndex)
+	{
+		TWeakObjectPtr<UObject> ObjectPtr = SelectedObjects[ObjectIndex];
+
+		if (ObjectPtr.IsValid())
+		{
+			UObject* Object = ObjectPtr.Get();
+			USceneComponent* SceneComponent = GetSceneComponentFromDetailsObject(Object);
+
+			if (SceneComponent)
+			{
+				AActor* EditedActor = SceneComponent->GetOwner();
+
+				// If we're committing during a slider transaction then we need to force it, in order that PostEditChangeChainProperty be called.
+				// Note: this will even happen if the slider hasn't changed the value.
+				bool bDoChange = bCommitted && bIsSliderTransaction;
+				bDoChange = bDoChange || (OldSelectionValue != NewSelectionValue);
+
+				if (EditedActor && bDoChange)
+				{
+					// @TODO Accumulated transform
+					const FTransform* InitialActorTransform = nullptr; // BlueprintEditorSP->GetAvaModeTools()->GetSelectionInitialActorTransforms().Find(EditedActor);
+
+					if (!InitialActorTransform)
+						continue;
+
+					FTransform OldSceneTransform = SceneComponent->GetComponentTransform();
+					FTransform NewSceneTransform = (*InitialActorTransform) * NewAccumulatedTransform;
+					SceneComponent->SetWorldTransform(NewSceneTransform);
+					SceneComponent->Modify();
+					ObjectToRelativeRotationMap.FindOrAdd(SceneComponent) = SceneComponent->GetRelativeRotation();
+				}
+			}
+		}
+	}
+
+	// @TODO Accumulated transform
+	// BlueprintEditorSP->GetAvaModeTools()->SetAccumulatedTransform(NewAccumulatedTransform);
+
+	CacheTransform();
+	
+	if (ULevelEditorSubsystem* LevelEditorSubsystem = GEditor->GetEditorSubsystem<ULevelEditorSubsystem>())
+	{
+		LevelEditorSubsystem->EditorInvalidateViewports();
+	}
+}
+
+void FAvaLevelViewportComponentTransformDetails::OnSetTransformAxis(FVector::FReal NewValue, ETextCommit::Type CommitInfo, EAvaLevelViewportTransformFieldType TransformField, EAxisList::Type Axis, bool bCommitted)
+{
+	FVector NewVector = GetAxisFilteredVector(Axis, FVector(NewValue), FVector::ZeroVector);
+
+	OnSetTransform(TransformField, Axis, NewVector, /* bMirror */ false, bCommitted);
+}
+
+void FAvaLevelViewportComponentTransformDetails::OnSetTransformAxisFloat(float NewValue, ETextCommit::Type CommitInfo, EAvaLevelViewportTransformFieldType TransformField, EAxisList::Type Axis, bool bCommitted)
+{
+	FVector NewVector = GetAxisFilteredVector(Axis, FVector(NewValue), FVector::ZeroVector);
+
+	OnSetTransform(EAvaLevelViewportTransformFieldType::Rotation, Axis, NewVector, /* bMirror */ false, bCommitted);
+}
+
+void FAvaLevelViewportComponentTransformDetails::BeginSliderTransaction(FText ActorTransaction, FText ComponentTransaction) const
+{
+	bool bBeganTransaction = false;
+	for (TWeakObjectPtr<UObject> ObjectPtr : SelectedObjects)
+	{
+		if (ObjectPtr.IsValid())
+		{
+			UObject* Object = ObjectPtr.Get();
+
+			// Start a new transaction when a slider begins to change
+			// We'll end it when the slider is released
+			// NOTE: One transaction per change, not per actor
+			if (!bBeganTransaction)
+			{
+				if (Object->IsA<AActor>())
+				{
+					GEditor->BeginTransaction(ActorTransaction);
+				}
+				else
+				{
+					GEditor->BeginTransaction(ComponentTransaction);
+				}
+
+				bBeganTransaction = true;
+			}
+
+			USceneComponent* SceneComponent = GetSceneComponentFromDetailsObject(Object);
+			if (SceneComponent)
+			{
+				FAvaLevelViewportScopedSwitchWorldForObject WorldSwitcher(Object);
+
+				if (SceneComponent->HasAnyFlags(RF_DefaultSubObject))
+				{
+					// Default subobjects must be included in any undo/redo operations
+					SceneComponent->SetFlags(RF_Transactional);
+				}
+
+				// Call modify but not PreEdit, we don't do the proper "Edit" until it's committed
+				SceneComponent->Modify();
+			}
+		}
+	}
+
+	// Just in case we couldn't start a new transaction for some reason
+	if (!bBeganTransaction)
+	{
+		GEditor->BeginTransaction(ActorTransaction);
+	}
+}
+
+void FAvaLevelViewportComponentTransformDetails::OnBeginRotationSlider()
+{
+	FText ActorTransaction = LOCTEXT("OnSetRotation", "Set Rotation");
+	FText ComponentTransaction = LOCTEXT("OnSetRotation_ComponentDirect", "Modify Component(s)");
+	BeginSliderTransaction(ActorTransaction, ComponentTransaction);
+
+	bEditingRotationInUI = true;
+	bIsSliderTransaction = true;
+
+	for (TWeakObjectPtr<UObject> ObjectPtr : SelectedObjects)
+	{
+		if (ObjectPtr.IsValid())
+		{
+			UObject* Object = ObjectPtr.Get();
+
+			USceneComponent* SceneComponent = GetSceneComponentFromDetailsObject(Object);
+			if (SceneComponent)
+			{
+				FAvaLevelViewportScopedSwitchWorldForObject WorldSwitcher(Object);
+
+				// Add/update cached rotation value prior to slider interaction
+				ObjectToRelativeRotationMap.FindOrAdd(SceneComponent) = SceneComponent->GetRelativeRotation();
+			}
+		}
+	}
+
+	OnBeginChange(EAvaLevelViewportTransformFieldType::Rotation);
+}
+
+void FAvaLevelViewportComponentTransformDetails::OnEndRotationSlider(float NewValue)
+{
+	OnEndChange(EAvaLevelViewportTransformFieldType::Rotation, EAxisList::None);
+
+	// Commit gets called right before this, only need to end the transaction
+	bEditingRotationInUI = false;
+	bIsSliderTransaction = false;
+	GEditor->EndTransaction();
+}
+
+void FAvaLevelViewportComponentTransformDetails::OnBeginLocationSlider()
+{
+	bIsSliderTransaction = true;
+	FText ActorTransaction = LOCTEXT("OnSetLocation", "Set Location");
+	FText ComponentTransaction = LOCTEXT("OnSetLocation_ComponentDirect", "Modify Component Location");
+	BeginSliderTransaction(ActorTransaction, ComponentTransaction);
+
+	OnBeginChange(EAvaLevelViewportTransformFieldType::Location);
+}
+
+void FAvaLevelViewportComponentTransformDetails::OnEndLocationSlider(FVector::FReal NewValue)
+{
+	OnEndChange(EAvaLevelViewportTransformFieldType::Location, EAxisList::All);
+
+	bIsSliderTransaction = false;
+	GEditor->EndTransaction();
+}
+
+void FAvaLevelViewportComponentTransformDetails::OnBeginScaleSlider()
+{
+	// Assumption: slider isn't usable if multiple objects are selected
+	SliderScaleRatio.X = CachedScale.X.GetValue();
+	SliderScaleRatio.Y = CachedScale.Y.GetValue();
+	SliderScaleRatio.Z = CachedScale.Z.GetValue();
+
+	bIsSliderTransaction = true;
+	FText ActorTransaction = LOCTEXT("OnSetScale", "Set Scale");
+	FText ComponentTransaction = LOCTEXT("OnSetScale_ComponentDirect", "Modify Component Scale");
+	BeginSliderTransaction(ActorTransaction, ComponentTransaction);
+
+	OnBeginChange(EAvaLevelViewportTransformFieldType::Scale);
+}
+
+void FAvaLevelViewportComponentTransformDetails::OnEndScaleSlider(FVector::FReal NewValue)
+{
+	OnEndChange(EAvaLevelViewportTransformFieldType::Scale, EAxisList::All);
+
+	bIsSliderTransaction = false;
+	GEditor->EndTransaction();
+}
+
+void FAvaLevelViewportComponentTransformDetails::OnObjectsReplaced(const TMap<UObject*, UObject*>& ReplacementMap)
+{
+	TArray<UObject*> NewSceneComponents;
+	for (const TWeakObjectPtr<UObject> Obj : CachedHandlesObjects)
+	{
+		if (UObject* Replacement = ReplacementMap.FindRef(Obj.GetEvenIfUnreachable()))
+		{
+			NewSceneComponents.Add(Replacement);
+		}
+	}
+
+	if (NewSceneComponents.Num())
+	{
+		UpdatePropertyHandlesObjects(NewSceneComponents);
+	}
+}
+
+void FAvaLevelViewportComponentTransformDetails::OnUseSelectionAsParentChanged(ECheckBoxState NewState)
+{
+	OnUseSelectionAsParentChangedEvent.Broadcast(NewState == ECheckBoxState::Checked);
+}
+
+void FAvaLevelViewportComponentTransformDetails::OnUseSelectionAsParentChanged(bool bNewValue)
+{
+	bUseSelectionAsParent = bNewValue;
+	bLastUseSelectionAsParent = bUseSelectionAsParent;
+
+	CachedLocation.X.Reset();
+	CachedLocation.Y.Reset();
+	CachedLocation.Z.Reset();
+
+	CachedRotation.X.Reset();
+	CachedRotation.Y.Reset();
+	CachedRotation.Z.Reset();
+
+	CachedScale.X.Reset();
+	CachedScale.Y.Reset();
+	CachedScale.Z.Reset();
+
+	CacheTransform();
+
+	if (TransformBox.IsValid())
+	{
+		TransformBox->SetContent(GetTransformBody().ToSharedRef());
+	}
+
+	if (RotationBox.IsValid())
+	{
+		RotationBox->SetContent(GetRotationBody().ToSharedRef());
+	}
+
+	if (ScaleBox.IsValid())
+	{
+		ScaleBox->SetContent(GetScaleBody().ToSharedRef());
+	}
+}
+
+void FAvaLevelViewportComponentTransformDetails::OnEnginePreExit()
+{
+	PropertyHandles.Empty();
+	TreeNodes.Empty();
+	PropertyRowGenerator.Reset();
+}
+
+#undef LOCTEXT_NAMESPACE
diff --git a/Engine/Plugins/Experimental/Avalanche/Source/AvalancheLevelViewport/Private/Customizations/AvaLevelViewportComponentTransformDetails.h b/Engine/Plugins/Experimental/Avalanche/Source/AvalancheLevelViewport/Private/Customizations/AvaLevelViewportComponentTransformDetails.h
new file mode 100644
--- /dev/null
+++ b/Engine/Plugins/Experimental/Avalanche/Source/AvalancheLevelViewport/Private/Customizations/AvaLevelViewportComponentTransformDetails.h
@@ -0,0 +1,415 @@
+// Copyright Epic Games, Inc. All Rights Reserved.
+
+#pragma once
+
+#include "AssetSelection.h"
+#include "Framework/Commands/UIAction.h"
+#include "IDetailCustomNodeBuilder.h"
+#include "UObject/WeakObjectPtr.h"
+#include "Widgets/Input/NumericTypeInterface.h"
+
+class FAvaLevelViewportClient;
+class FDetailWidgetRow;
+class FMenuBuilder;
+class FNotifyHook;
+class IDetailCategoryBuilder;
+class IDetailChildrenBuilder;
+class IDetailLayoutBuilder;
+class IDetailTreeNode;
+class IPropertyHandle;
+class IPropertyRowGenerator;
+class SBox;
+class SWidget;
+class UTypedElementSelectionSet;
+struct FSlateBrush;
+
+enum class EAvaLevelViewportTransformFieldType : uint8
+{
+	Location,
+	Rotation,
+	Scale,
+	All
+};
+
+DECLARE_EVENT_OneParam(FAvaLevelViewportComponentTransformDetails, FOnUseSelectionAsParentChanged, bool);
+
+/**
+ * Manages the Transform section of a details view
+ */
+class FAvaLevelViewportComponentTransformDetails : public TSharedFromThis<FAvaLevelViewportComponentTransformDetails>,
+	public IDetailCustomNodeBuilder, public TNumericUnitTypeInterface<FVector::FReal>
+{
+public:
+	FAvaLevelViewportComponentTransformDetails(TSharedPtr<FAvaLevelViewportClient> InAvaViewportClient);
+	virtual ~FAvaLevelViewportComponentTransformDetails() override;
+
+	static bool IsUseSelectionAsParent() { return bUseSelectionAsParent; }
+
+	/**
+	 * Caches the representation of the actor transform for the user input boxes
+	 */
+	void CacheTransform();
+
+	virtual void GenerateHeaderRowContent(FDetailWidgetRow& NodeRow) override {}
+	virtual void GenerateChildContent(IDetailChildrenBuilder& ChildrenBuilder) override {}
+	virtual bool RequiresTick() const override { return true; }
+	virtual FName GetName() const override { return "Transform"; }
+	virtual bool InitiallyCollapsed() const override { return false; }
+	virtual void Tick(float DeltaTime) override;
+	virtual void SetOnRebuildChildren(FSimpleDelegate OnRebuildChildren) override {}
+
+	void OnSelectionChanged(const TArray<UObject*>& InSelectedObjects);
+	void OnSelectionSetChanged(const UTypedElementSelectionSet* InSelectionSet);
+
+	TSharedPtr<SWidget> GetTransformBody();
+	TSharedPtr<SWidget> GetRotationBody();
+	TSharedPtr<SWidget> GetScaleBody();
+	TSharedPtr<SWidget> GetOptionsHeader();
+	TSharedPtr<SWidget> GetOptionsBody();
+	TSharedRef<SWidget> GetPreserveScaleRatioWidget();
+
+	int32 GetSelectedActorCount() const { return SelectedActorInfo.NumSelected; }
+
+	/** Returns the cached list of selected objects. */
+	TConstArrayView<TWeakObjectPtr<UObject>> GetSelectedObjects() const { return SelectedObjects; }
+
+	void SetTransform(EAvaLevelViewportTransformFieldType TransformField, const FVector& NewValue);
+
+	static FOnUseSelectionAsParentChanged OnUseSelectionAsParentChangedEvent;
+
+	/** Called when the "Reset to Default" button for the location has been clicked */
+	void OnLocationResetClicked();
+
+	/** Called when the "Reset to Default" button for the rotation has been clicked */
+	void OnRotationResetClicked();
+
+	/** Called when the "Reset to Default" button for the scale has been clicked */
+	void OnScaleResetClicked();
+
+	TConstArrayView<TSharedPtr<IPropertyHandle>> GetPropertyHandles() const { return PropertyHandles; }
+
+private:
+
+	/** @return Whether the transform details panel should be enabled (editable) or not (read-only / greyed out) */
+	bool GetIsEnabled() const;
+
+	/** Sets a vector based on two source vectors and an axis list */
+	FVector GetAxisFilteredVector(EAxisList::Type Axis, const FVector& NewValue, const FVector& OldValue);
+
+	/* Called when a change starts (direct ui commit or slider begin) */
+	void OnBeginChange(EAvaLevelViewportTransformFieldType TransformField);
+
+	/* Called when a change ends (direct ui commit or slider end) */
+	void OnEndChange(EAvaLevelViewportTransformFieldType TransformField, EAxisList::Type Axis);
+
+	/**
+	 * Sets the selected object(s) axis to passed in value
+	 *
+	 * @param TransformField	The field (location/rotation/scale) to modify
+	 * @param Axis				Bitfield of which axis to set, can be multiple
+	 * @param NewValue			The new vector values, it only uses the ones with specified axis
+	 * @param bMirror			If true, set the value to it's inverse instead of using NewValue
+	 * @param bCommitted		True if the value was committed, false is the value comes from the slider
+	 */
+	void OnSetTransform(EAvaLevelViewportTransformFieldType TransformField, EAxisList::Type Axis, FVector NewValue, bool bMirror, bool bCommitted);
+
+	// Transform components based the selection transform, not the component transform.
+	void OnSetTransformSelectionParent(EAvaLevelViewportTransformFieldType TransformField, EAxisList::Type Axis, FVector NewValue, bool bMirror, bool bCommitted);
+
+	/**
+	 * Sets a single axis value, called from UI
+	 *
+	 * @param TransformField	The field (location/rotation/scale) to modify
+	 * @param NewValue			The new translation value
+	 * @param CommitInfo		Whether or not this was committed from pressing enter or losing focus
+	 * @param Axis				Bitfield of which axis to set, can be multiple
+	 * @param bCommitted		true if the value was committed, false is the value comes from the slider
+	 */
+	void OnSetTransformAxis(FVector::FReal NewValue, ETextCommit::Type CommitInfo, EAvaLevelViewportTransformFieldType TransformField, EAxisList::Type Axis, bool bCommitted);
+
+	/**
+	 * Sets a single axis float value, called from UI
+	 *
+	 * @param TransformField	The field (location/rotation/scale) to modify
+	 * @param NewValue			The new translation value
+	 * @param CommitInfo		Whether or not this was committed from pressing enter or losing focus
+	 * @param Axis				Bitfield of which axis to set, can be multiple
+	 * @param bCommitted		true if the value was committed, false is the value comes from the slider
+	 */
+	void OnSetTransformAxisFloat(float NewValue, ETextCommit::Type CommitInfo, EAvaLevelViewportTransformFieldType TransformField, EAxisList::Type Axis, bool bCommitted);
+
+	/**
+	 * Helper to begin a new transaction for a slider interaction.
+	 * @param ActorTransaction		The name to give the transaction when changing an actor transform.
+	 * @param ComponentTransaction	The name to give the transaction when directly editing a component transform.
+	 */
+	void BeginSliderTransaction(FText ActorTransaction, FText ComponentTransaction) const;
+
+	/** Called when the one of the axis sliders for object rotation begins to change for the first time */
+	void OnBeginRotationSlider();
+
+	/** Called when the one of the axis sliders for object rotation is released */
+	void OnEndRotationSlider(float NewValue);
+
+	/** Called when one of the axis sliders for object location begins to change */
+	void OnBeginLocationSlider();
+
+	/** Called when one of the axis sliders for object location is released */
+	void OnEndLocationSlider(FVector::FReal NewValue);
+
+	/** Called when one of the axis sliders for object scale begins to change */
+	void OnBeginScaleSlider();
+
+	/** Called when one of the axis sliders for object scale is released */
+	void OnEndScaleSlider(FVector::FReal NewValue);
+
+	/** @return Icon to use in the preserve scale ratio check box */
+	const FSlateBrush* GetPreserveScaleRatioImage() const;
+
+	/** @return The state of the preserve scale ratio checkbox */
+	ECheckBoxState IsPreserveScaleRatioChecked() const;
+
+	/**
+	 * Called when the preserve scale ratio checkbox is toggled
+	 */
+	void OnPreserveScaleRatioToggled(ECheckBoxState NewState);
+
+	/**
+	 * Builds a transform field label
+	 *
+	 * @param TransformField The field to build the label for
+	 * @return The label widget
+	 */
+	TSharedRef<SWidget> BuildTransformFieldLabel(EAvaLevelViewportTransformFieldType TransformField);
+
+	/**
+	 * Gets display text for a transform field
+	 *
+	 * @param TransformField	The field to get the text for
+	 * @return The text label for TransformField
+	 */
+	FText GetTransformFieldText(EAvaLevelViewportTransformFieldType TransformField) const;
+
+	/**
+	 * @return the text to display for translation
+	 */
+	FText GetLocationText() const;
+
+	/**
+	 * @return the text to display for rotation
+	 */
+	FText GetRotationText() const;
+
+	/**
+	 * @return the text to display for scale
+	 */
+	FText GetScaleText() const;
+
+	/**
+	 * Sets relative transform on the specified field
+	 *
+	 * @param TransformField the field that should be set to relative
+	 * @param bAbsoluteEnabled true to absolute, false for relative 
+	 */
+	void OnSetAbsoluteTransform(EAvaLevelViewportTransformFieldType TransformField, bool bAbsoluteEnabled);
+
+	/** @return true if Absolute flag of transform type matches passed in bCheckAbsolute*/
+	bool IsAbsoluteTransformChecked(EAvaLevelViewportTransformFieldType TransformField, bool bAbsoluteEnabled=true) const;
+
+	/** @return true of copy is enabled for the specified field */
+	bool OnCanCopy(EAvaLevelViewportTransformFieldType TransformField) const;
+
+	/**
+	 * Copies the specified transform field to the clipboard
+	 */
+	void OnCopy(EAvaLevelViewportTransformFieldType TransformField);
+
+	/**
+	 * Pastes the specified transform field from the clipboard
+	 */
+	void OnPaste(EAvaLevelViewportTransformFieldType TransformField);
+
+	/**
+	 * Creates a UI action for copying a specified transform field
+	 */
+	FUIAction CreateCopyAction(EAvaLevelViewportTransformFieldType TransformField) const;
+
+	/**
+	 * Creates a UI action for pasting a specified transform field
+	 */
+	FUIAction CreatePasteAction(EAvaLevelViewportTransformFieldType TransformField) const;
+
+	/** Extend the context menu for the X component */
+	void ExtendXScaleContextMenu(FMenuBuilder& MenuBuilder);
+	/** Extend the context menu for the Y component */
+	void ExtendYScaleContextMenu(FMenuBuilder& MenuBuilder);
+	/** Extend the context menu for the Z component */
+	void ExtendZScaleContextMenu(FMenuBuilder& MenuBuilder);
+
+	/** Called when the X axis scale is mirrored */
+	void OnXScaleMirrored();
+	/** Called when the Y axis scale is mirrored */
+	void OnYScaleMirrored();
+	/** Called when the Z axis scale is mirrored */
+	void OnZScaleMirrored();
+
+	/** @return The X component of location */
+	TOptional<FVector::FReal> GetLocationX() const { return CachedLocation.X; }
+	/** @return The Y component of location */
+	TOptional<FVector::FReal> GetLocationY() const { return CachedLocation.Y; }
+	/** @return The Z component of location */
+	TOptional<FVector::FReal> GetLocationZ() const { return CachedLocation.Z; }
+	/** @return The visibility of the "Reset to Default" button for the location component */
+	bool GetLocationResetVisibility() const;
+
+	/** @return The X component of rotation */
+	TOptional<float> GetRotationX() const { return CachedRotation.X; }
+	/** @return The Y component of rotation */
+	TOptional<float> GetRotationY() const { return CachedRotation.Y; }
+	/** @return The Z component of rotation */
+	TOptional<float> GetRotationZ() const { return CachedRotation.Z; }
+	/** @return The visibility of the "Reset to Default" button for the rotation component */
+	bool GetRotationResetVisibility() const;
+
+	/** @return The X component of scale */
+	TOptional<FVector::FReal> GetScaleX() const { return CachedScale.X; }
+	/** @return The Y component of scale */
+	TOptional<FVector::FReal> GetScaleY() const { return CachedScale.Y; }
+	/** @return The Z component of scale */
+	TOptional<FVector::FReal> GetScaleZ() const { return CachedScale.Z; }
+	/** @return The visibility of the "Reset to Default" button for the scale component */
+	bool GetScaleResetVisibility() const;
+
+	/** Cache a single unit to display all location components in */
+	void CacheCommonLocationUnits();
+
+	/** Generate a property handles from a the scene component properties. */
+	void GeneratePropertyHandles();
+
+	/** Update the outer objects of the property handles generated from this transform. */
+	void UpdatePropertyHandlesObjects(TArray<UObject*> NewSceneComponents);
+
+	/** Delegate called when the OnObjectsReplaced is broadcast by the engine. */
+	void OnObjectsReplaced(const TMap<UObject*, UObject*>& ReplacementMap);
+
+	ECheckBoxState IsUseSelectionAsParentChecked() const { return bUseSelectionAsParent ? ECheckBoxState::Checked : ECheckBoxState::Unchecked; }
+	void OnUseSelectionAsParentChanged(ECheckBoxState NewState);
+
+	void OnUseSelectionAsParentChanged(bool bNewValue);
+
+	void OnEnginePreExit();
+
+private:
+	/** A vector where it may optionally be unset */
+	template <typename NumericType>
+	struct FOptionalVector
+	{
+		/**
+		 * Sets the value from an FVector
+		 */
+		void Set(const FVector& InVec)
+		{
+			X = InVec.X;
+			Y = InVec.Y;
+			Z = InVec.Z;
+		}
+
+		/**
+		 * Sets the value from an FRotator
+		 */
+		void Set(const FRotator& InRot)
+		{
+			X = InRot.Roll;
+			Y = InRot.Pitch;
+			Z = InRot.Yaw;
+		}
+
+		/** @return Whether or not the value is set */
+		bool IsSet() const
+		{
+			// The vector is set if all values are set
+			return X.IsSet() && Y.IsSet() && Z.IsSet();
+		}
+
+		TOptional<NumericType> X;
+		TOptional<NumericType> Y;
+		TOptional<NumericType> Z;
+	};
+
+	FSelectedActorInfo SelectedActorInfo;
+
+	/** Copy of selected actor references in the details view */
+	TArray< TWeakObjectPtr<UObject> > SelectedObjects;
+
+	/** Cache translation value of the selected set */
+	FOptionalVector<FVector::FReal> CachedLocation;
+
+	/** Cache rotation value of the selected set */
+	FOptionalVector<float> CachedRotation;
+
+	/** Cache scale value of the selected set */
+	FOptionalVector<FVector::FReal> CachedScale;
+
+	/** Notify hook to use */
+	FNotifyHook* NotifyHook;
+
+	/** Mapping from object to relative rotation values which are not affected by Quat->Rotator conversions during transform calculations */
+	TMap< UObject*, FRotator > ObjectToRelativeRotationMap;
+
+	/** Whether or not we are in absolute translation mode */
+	bool bAbsoluteLocation;
+
+	/** Whether or not we are in absolute rotation mode */
+	bool bAbsoluteRotation;
+
+	/** Whether or not we are in absolute scale mode */
+	bool bAbsoluteScale;
+
+	/** Whether or not to preserve scale ratios */
+	bool bPreserveScaleRatio;
+
+	/** Scale ratio to use when we are using the sliders with bPreserveScaleRatio set. */
+	FVector SliderScaleRatio;
+
+	/** Flag to indicate we are currently editing the rotation in the UI, so we should rely on the cached value in objectToRelativeRotationMap, not the value from the object */
+	bool bEditingRotationInUI;
+
+	/** Flag to indicate we are currently performing a slider transaction */
+	bool bIsSliderTransaction;
+
+	/** Bitmask to indicate which fields should be hidden (if any) */
+	uint8 DisplayedFieldType;
+
+	/** Used to generated property handles. */
+	TSharedPtr<IPropertyRowGenerator> PropertyRowGenerator;
+
+	/** Tree nodes for the property handles. */
+	TArray<TSharedPtr<IDetailTreeNode>> TreeNodes;
+
+	/** Holds this transform's property handles. */
+	TArray<TSharedPtr<IPropertyHandle>> PropertyHandles;
+
+	/** Holds the property handles' outer objects. Used to update the handles' objects when the actor construction script runs. */
+	TArray<TWeakObjectPtr<UObject>> CachedHandlesObjects;
+
+	// Link to the viewport client for which this is displaying details
+	TWeakPtr<FAvaLevelViewportClient> AvaViewportClientWeak;
+
+	// If true, use the multi-select as an ad-hoc parent
+	static bool bUseSelectionAsParent;
+	bool bLastUseSelectionAsParent;
+
+	// If true, we can use the multi-select as a parent
+	bool bCanUseSelectionAsParent;
+
+	TSharedPtr<SBox> TransformBox;
+	TSharedPtr<SBox> RotationBox;
+	TSharedPtr<SBox> ScaleBox;
+	TArray<TWeakObjectPtr<UObject>> ModifiedObjects;
+	TMap<USceneComponent*, FVector> OriginalLocations;
+	TMap<USceneComponent*, FRotator> OriginalRotations;
+	TMap<USceneComponent*, FVector> OriginalScales;
+
+	/** Editor Selection Set used to listen to Selection Changes */
+	TWeakObjectPtr<UTypedElementSelectionSet> EditorSelectionSetWeak;
+};
diff --git a/Engine/Plugins/Experimental/Avalanche/Source/AvalancheLevelViewport/Private/Widgets/AvaLevelViewportComponentTransformDetails.cpp b/Engine/Plugins/Experimental/Avalanche/Source/AvalancheLevelViewport/Private/Widgets/AvaLevelViewportComponentTransformDetails.cpp
deleted file mode 100644
diff --git a/Engine/Plugins/Experimental/Avalanche/Source/AvalancheLevelViewport/Private/Widgets/AvaLevelViewportComponentTransformDetails.h b/Engine/Plugins/Experimental/Avalanche/Source/AvalancheLevelViewport/Private/Widgets/AvaLevelViewportComponentTransformDetails.h
deleted file mode 100644
diff --git a/Engine/Plugins/Experimental/Avalanche/Source/AvalancheLevelViewport/Private/Widgets/AvaViewportColorPickerDelegates.cpp b/Engine/Plugins/Experimental/Avalanche/Source/AvalancheLevelViewport/Private/Widgets/AvaViewportColorPickerDelegates.cpp
deleted file mode 100644
diff --git a/Engine/Plugins/Experimental/Avalanche/Source/AvalancheLevelViewport/Private/Widgets/SAvaLevelViewportActorColorMenu.cpp b/Engine/Plugins/Experimental/Avalanche/Source/AvalancheLevelViewport/Private/Widgets/SAvaLevelViewportActorColorMenu.cpp
--- a/Engine/Plugins/Experimental/Avalanche/Source/AvalancheLevelViewport/Private/Widgets/SAvaLevelViewportActorColorMenu.cpp
+++ b/Engine/Plugins/Experimental/Avalanche/Source/AvalancheLevelViewport/Private/Widgets/SAvaLevelViewportActorColorMenu.cpp
@@ -3,9 +3,9 @@
 #include "Widgets/SAvaLevelViewportActorColorMenu.h"
 #include "AvaLevelViewportStyle.h"
 #include "Brushes/SlateImageBrush.h"
+#include "ColorPicker/AvaViewportColorPickerActorClassRegistry.h"
 #include "CoreGlobals.h"
 #include "EditorModeManager.h"
-#include "Selection.h"
 #include "Engine/Engine.h"
 #include "Engine/Texture2D.h"
 #include "Framework/Application/SlateApplication.h"
@@ -16,12 +16,11 @@
 #include "Materials/MaterialInstanceDynamic.h"
 #include "Materials/MaterialInterface.h"
 #include "Misc/ConfigCacheIni.h"
+#include "Selection.h"
 #include "SlateMaterialBrush.h"
 #include "Styling/StyleColors.h"
 #include "Toolkits/IToolkitHost.h"
 #include "UObject/Package.h"
-#include "Widgets/AvaViewportColorPickerActorClassRegistry.h"
-#include "Widgets/AvaViewportColorPickerDelegates.h"
 #include "Widgets/Colors/SColorBlock.h"
 #include "Widgets/Images/SImage.h"
 #include "Widgets/Input/SEditableTextBox.h"
@@ -75,41 +74,7 @@ void SAvaLevelViewportActorColorMenu::OnColorStyleChanged()
 
 void SAvaLevelViewportActorColorMenu::BroadcastColorChange()
 {
-	if (FAvaViewportColorPickerDelegates::GetOnColorPicked().IsBound())
-	{
-		if (TSharedPtr<IToolkitHost> ToolkitHost = ToolkitHostWeak.Pin())
-		{
-			FAvaViewportColorPickerDelegates::BroadcastColorPicked(ToolkitHost.ToSharedRef(), {ColorStyle, ActiveColor, InactiveColor, bIsUnlit});
-		}
-	}
-}
-
-void SAvaLevelViewportActorColorMenu::OnColorSourceBroadcast(const TSharedRef<IToolkitHost>& InToolkitHost, const FAvaColorChangeData& InNewColorData)
-{
-	if (!ToolkitHostWeak.HasSameObject(&*InToolkitHost))
-	{
-		return;
-	}
-
-	bIsUnlit = InNewColorData.bIsUnlit;
-
-	switch (InNewColorData.ColorStyle)
-	{
-		case EAvaColorStyle::Solid:
-			SetColorStyle_Direct(EAvaColorStyle::Solid);
-			SetActiveColorRGB(InNewColorData.PrimaryColor);
-			break;
-
-		case EAvaColorStyle::LinearGradient:
-			SetColorStyle_Direct(EAvaColorStyle::LinearGradient);
-			SetActiveColorRGB(InNewColorData.PrimaryColor);
-			SetInactiveColorRGB(InNewColorData.SecondaryColor);
-			break;
-
-		default:
-			// Invalid color
-			return;
-	}
+	ColorPicker.OnColorSelected({ColorStyle, ActiveColor, InactiveColor, bIsUnlit});
 }
 
 void SAvaLevelViewportActorColorMenu::SetActiveColorRGB_Direct(FLinearColor InRGB)
@@ -175,42 +140,7 @@ FSlateColor SAvaLevelViewportActorColorMenu::GetUnlitButtonColor() const
 
 TSharedRef<SAvaLevelViewportActorColorMenu> SAvaLevelViewportActorColorMenu::CreateMenu(const TSharedRef<IToolkitHost>& InToolkitHost)
 {
-	TSharedRef<SAvaLevelViewportActorColorMenu> NewMenu = SNew(SAvaLevelViewportActorColorMenu, InToolkitHost);
-#
-	if (USelection* ActorSelection = InToolkitHost->GetEditorModeManager().GetSelectedActors())
-	{
-		TArray<AActor*> SelectedActors;
-		ActorSelection->GetSelectedObjects(SelectedActors);
-		FAvaColorChangeData NewColorData;
-
-		for (AActor* Actor : SelectedActors)
-		{
-			if (FAvaViewportColorPickerActorClassRegistry::GetColorDataFromActor(Actor, NewColorData))
-			{
-				NewMenu->SetActiveColorRGB_Direct(NewColorData.PrimaryColor);
-				NewMenu->SetInactiveColorRGB_Direct(NewColorData.SecondaryColor);
-				NewMenu->SetColorStyle_Direct(NewColorData.ColorStyle);
-				break;
-			}
-		}
-	}
-
-	return NewMenu;
-}
-
-SAvaLevelViewportActorColorMenu::~SAvaLevelViewportActorColorMenu()
-{
-	if (HueBrush)
-	{
-		delete HueBrush;
-	}
-
-	if (SatValueBrush)
-	{
-		delete SatValueBrush;
-	}
-
-	FAvaViewportColorPickerDelegates::GetOnColorSourceSelected().RemoveAll(this);
+	return SNew(SAvaLevelViewportActorColorMenu, InToolkitHost);
 }
 
 void SAvaLevelViewportActorColorMenu::PrivateRegisterAttributes(struct FSlateAttributeDescriptor::FInitializer&)
@@ -219,30 +149,13 @@ void SAvaLevelViewportActorColorMenu::PrivateRegisterAttributes(struct FSlateAtt
 
 void SAvaLevelViewportActorColorMenu::Construct(const FArguments& InArgs, const TSharedRef<IToolkitHost>& InToolkitHost)
 {
-	ToolkitHostWeak = InToolkitHost;
+	ColorPicker.SetToolkitHost(InToolkitHost);
 
-	HueMaterial = TStrongObjectPtr<UMaterial>(FindObject<UMaterial>(nullptr, TEXT("Material'/Avalanche/ColorPickerHueMaterial.ColorPickerHueMaterial'")));
+	HueMaterial = LoadObject<UMaterial>(nullptr, TEXT("Material'/Avalanche/ColorPickerHueMaterial.ColorPickerHueMaterial'"));
+	SatValueMaterial = LoadObject<UMaterial>(nullptr, TEXT("Material'/Avalanche/ColorPickerBrightnessSaturdationMaterial.ColorPickerBrightnessSaturdationMaterial'"));
+	SelectedValueImage = LoadObject<UTexture2D>(nullptr, TEXT("Texture2D'/Avalanche/EditorResources/ColorSelectionIcon.ColorSelectionIcon'"));
 
-	if (!HueMaterial.IsValid())
-	{
-		HueMaterial = TStrongObjectPtr<UMaterial>(LoadObject<UMaterial>(nullptr, TEXT("Material'/Avalanche/ColorPickerHueMaterial.ColorPickerHueMaterial'")));
-	}
-
-	SatValueMaterial = TStrongObjectPtr<UMaterial>(FindObject<UMaterial>(nullptr, TEXT("Material'/Avalanche/ColorPickerBrightnessSaturdationMaterial.ColorPickerBrightnessSaturdationMaterial'")));
-
-	if (!SatValueMaterial.IsValid())
-	{
-		SatValueMaterial = TStrongObjectPtr<UMaterial>(LoadObject<UMaterial>(nullptr, TEXT("Material'/Avalanche/ColorPickerBrightnessSaturdationMaterial.ColorPickerBrightnessSaturdationMaterial'")));
-	}
-
-	SelectedValueImage = TStrongObjectPtr<UTexture2D>(FindObject<UTexture2D>(nullptr, TEXT("Texture2D'/Avalanche/EditorResources/ColorSelectionIcon.ColorSelectionIcon'")));
-
-	if (!SelectedValueImage.IsValid())
-	{
-		SelectedValueImage = TStrongObjectPtr<UTexture2D>(LoadObject<UTexture2D>(nullptr, TEXT("Texture2D'/Avalanche/EditorResources/ColorSelectionIcon.ColorSelectionIcon'")));
-	}
-
-	if (!HueMaterial.IsValid() || !SatValueMaterial.IsValid())
+	if (!HueMaterial || !SatValueMaterial)
 	{
 		return;
 	}
@@ -257,14 +170,14 @@ void SAvaLevelViewportActorColorMenu::Construct(const FArguments& InArgs, const
 	ActiveColor = FLinearColor::White;
 	InactiveColor = FLinearColor::White;
 	ColorStyle = EAvaColorStyle::Solid;
-	
+
 	ActiveThemeIndex = 0;
 
-	SatValueMaterialDynamic = TStrongObjectPtr<UMaterialInstanceDynamic>(UMaterialInstanceDynamic::Create(SatValueMaterial.Get(), GetTransientPackage()));
+	SatValueMaterialDynamic = UMaterialInstanceDynamic::Create(SatValueMaterial.Get(), GetTransientPackage());
 
-	HueBrush = new FSlateMaterialBrush(*HueMaterial.Get(), FVector2f(200.f, 200.f));
-	SatValueBrush = new FSlateMaterialBrush(*SatValueMaterialDynamic.Get(), FVector2f(80.f, 80.f));
-	SelectedValueBrush = new FSlateImageBrush(SelectedValueImage.Get(), FVector2f(32.f, 32.f));
+	HueBrush = MakeShared<FSlateMaterialBrush>(*HueMaterial.Get(), FVector2f(200.f, 200.f));
+	SatValueBrush = MakeShared<FSlateMaterialBrush>(*SatValueMaterialDynamic.Get(), FVector2f(80.f, 80.f));
+	SelectedValueBrush = MakeShared<FSlateImageBrush>(SelectedValueImage.Get(), FVector2f(32.f, 32.f));
 
 	const ISlateStyle& AvaLevelViewportStyle = FAvaLevelViewportStyle::Get();
 
@@ -272,7 +185,6 @@ void SAvaLevelViewportActorColorMenu::Construct(const FArguments& InArgs, const
 	[
 		SNew(SBox)
 		.Padding(FMargin(3.f, 3.f, 3.f, 10.f))
-		//.HeightOverride(200.f)
 		[
 			SNew(SConstraintCanvas)
 			+ SConstraintCanvas::Slot()
@@ -281,7 +193,7 @@ void SAvaLevelViewportActorColorMenu::Construct(const FArguments& InArgs, const
 			.Offset(FMargin(0.f, 0.f, 200.f, 200.f))
 			[
 				SNew(SImage)
-				.Image(HueBrush)
+				.Image(HueBrush.Get())
 			]
 			+ SConstraintCanvas::Slot()
 			.Alignment(FVector2D(0.f))
@@ -289,7 +201,7 @@ void SAvaLevelViewportActorColorMenu::Construct(const FArguments& InArgs, const
 			.Offset(FMargin(60.f, 60.f, 80.f, 80.f))
 			[
 				SNew(SImage)
-				.Image(SatValueBrush)
+				.Image(SatValueBrush.Get())
 			]
 			+ SConstraintCanvas::Slot()
 			.Alignment(FVector2D(0.f))
@@ -298,7 +210,7 @@ void SAvaLevelViewportActorColorMenu::Construct(const FArguments& InArgs, const
 			.Expose(HueSlot)
 			[
 				SNew(SImage)
-				.Image(SelectedValueBrush)
+				.Image(SelectedValueBrush.Get())
 			]
 			+ SConstraintCanvas::Slot()
 			.Alignment(FVector2D(0.f))
@@ -307,7 +219,7 @@ void SAvaLevelViewportActorColorMenu::Construct(const FArguments& InArgs, const
 			.Expose(SatValueSlot)
 			[
 				SNew(SImage)
-				.Image(SelectedValueBrush)
+				.Image(SelectedValueBrush.Get())
 			]
 			+ SConstraintCanvas::Slot()
 			.Alignment(FVector2D(0.f))
@@ -423,13 +335,12 @@ void SAvaLevelViewportActorColorMenu::Construct(const FArguments& InArgs, const
 		]
 	];
 
-	SetColorStyle_Direct(FAvaViewportColorPickerDelegates::GetLastColorData().ColorStyle);
-	SetActiveColorRGB_Direct(FAvaViewportColorPickerDelegates::GetLastColorData().PrimaryColor);
-	SetInactiveColorRGB_Direct(FAvaViewportColorPickerDelegates::GetLastColorData().SecondaryColor);
+	const FAvaColorChangeData& LastColorData = ColorPicker.GetLastColorData();
+	SetColorStyle_Direct(LastColorData.ColorStyle);
+	SetActiveColorRGB_Direct(LastColorData.PrimaryColor);
+	SetInactiveColorRGB_Direct(LastColorData.SecondaryColor);
 
 	UpdateColorPalette();
-
-	FAvaViewportColorPickerDelegates::GetOnColorSourceSelected().AddRaw(this, &SAvaLevelViewportActorColorMenu::OnColorSourceBroadcast);
 }
 
 const FVector2f SAvaLevelViewportActorColorMenu::WheelMiddle = FVector2f(100.f, 100.f);
@@ -556,9 +467,22 @@ FAvaColorTheme* SAvaLevelViewportActorColorMenu::GetTheme(int32 InThemeIndex)
 	return nullptr;
 }
 
+FString SAvaLevelViewportActorColorMenu::GetReferencerName() const
+{
+	return TEXT("SAvaLevelViewportActorColorMenu");
+}
+
+void SAvaLevelViewportActorColorMenu::AddReferencedObjects(FReferenceCollector& InCollector)
+{
+	InCollector.AddReferencedObject(HueMaterial);
+	InCollector.AddReferencedObject(SatValueMaterial);
+	InCollector.AddReferencedObject(SatValueMaterialDynamic);
+	InCollector.AddReferencedObject(SelectedValueImage);
+}
+
 void SAvaLevelViewportActorColorMenu::ApplyActiveColor()
 {
-	if (SatValueMaterialDynamic.IsValid())
+	if (IsValid(SatValueMaterialDynamic))
 	{
 		FLinearColor HSVColor(Hue * 360.0f, 1.f, 1.f);
 		HSVColor = HSVColor.HSVToLinearRGB();
@@ -1246,7 +1170,7 @@ TSharedRef<SWidget> SAvaLevelViewportActorColorMenu::BuildRenameColorMenu(int32
 				}));
 
 		ColorMenuBuilder.AddMenuEntry(LOCTEXT("RemovePaletteColorLabel", "Remove"), LOCTEXT("RemovePaletteColorTooltip", "Remove this color from the palette."), FSlateIcon(),
-			FUIAction(FExecuteAction::CreateLambda([this, Theme, InColorIdx, CurrentColor]()
+			FUIAction(FExecuteAction::CreateSPLambda(this, [this, Theme, InColorIdx, CurrentColor]()
 				{
 					TArray<FAvaColorInfo>& ThemeColors = Theme->GetColors();
 
@@ -1275,8 +1199,7 @@ TSharedRef<SWidget> SAvaLevelViewportActorColorMenu::BuildRenameColorMenu(int32
 
 TSharedRef<SWidget> SAvaLevelViewportActorColorMenu::BuildThemesMenu()
 {
-	constexpr bool bInShouldCloseWindowAfterMenuSelection = true;
-	FMenuBuilder ThemeMenuBuilder(bInShouldCloseWindowAfterMenuSelection, nullptr);
+	FMenuBuilder ThemeMenuBuilder(/*bInShouldCloseWindowAfterMenuSelection*/false, nullptr);
 	{
 		static FText RemoveThemeLabelText = LOCTEXT("RemoveThemeLabel", "Remove");
 		static FText RemoveThemeTooltipText = LOCTEXT("RemoveThemeTooltip", "Remove this theme.");
@@ -1294,7 +1217,7 @@ TSharedRef<SWidget> SAvaLevelViewportActorColorMenu::BuildThemesMenu()
 				.MinDesiredWidth(100.f)
 				.IsReadOnly(false)
 				.ClearKeyboardFocusOnCommit(true)
-				.OnTextCommitted(FOnTextCommitted::CreateLambda(
+				.OnTextCommitted(FOnTextCommitted::CreateSPLambda(this,
 					[this](const FText& NewName, ETextCommit::Type CommitType)
 					{
 						if (!ColorThemes.IsValidIndex(ActiveThemeIndex))
@@ -1310,7 +1233,7 @@ TSharedRef<SWidget> SAvaLevelViewportActorColorMenu::BuildThemesMenu()
 			);
 
 			ThemeMenuBuilder.AddMenuEntry(RemoveThemeLabelText, RemoveThemeTooltipText, FSlateIcon(),
-				FUIAction(FExecuteAction::CreateLambda([this]()
+				FUIAction(FExecuteAction::CreateSPLambda(this, [this]()
 					{
 						if (!ColorThemes.IsValidIndex(ActiveThemeIndex))
 						{
@@ -1327,7 +1250,7 @@ TSharedRef<SWidget> SAvaLevelViewportActorColorMenu::BuildThemesMenu()
 						SaveColorThemesToIni();
 						UpdateColorPalette();
 					}),
-					FCanExecuteAction::CreateLambda([this]() { return ColorThemes.Num() > 1; }))
+					FCanExecuteAction::CreateSPLambda(this, [this]() { return ColorThemes.Num() > 1; }))
 			);
 
 			ThemeMenuBuilder.EndSection();
@@ -1342,12 +1265,12 @@ TSharedRef<SWidget> SAvaLevelViewportActorColorMenu::BuildThemesMenu()
 				ActiveThemeIndex == ThemeIdx 
 					? LOCTEXT("CurrentThemeTooltip", "This is the current theme.") 
 					: LOCTEXT("ChangeToTheme", "Change to this theme."),
-				FNewMenuDelegate::CreateLambda([this, ThemeIdx](FMenuBuilder& SubMenuBuilder)
+				FNewMenuDelegate::CreateSPLambda(this, [this, ThemeIdx](FMenuBuilder& SubMenuBuilder)
 					{
 						SubMenuBuilder.BeginSection("ThemeOptions", LOCTEXT("Options", "Options"));
 
 						SubMenuBuilder.AddMenuEntry(RemoveThemeLabelText, RemoveThemeTooltipText, FSlateIcon(),
-							FUIAction(FExecuteAction::CreateLambda([this, ThemeIdx]()
+							FUIAction(FExecuteAction::CreateSPLambda(this, [this, ThemeIdx]()
 								{
 									if (!ColorThemes.IsValidIndex(ThemeIdx))
 									{
@@ -1364,7 +1287,7 @@ TSharedRef<SWidget> SAvaLevelViewportActorColorMenu::BuildThemesMenu()
 									SaveColorThemesToIni();
 									UpdateColorPalette();
 								}),
-								FCanExecuteAction::CreateLambda([this]() { return ColorThemes.Num() > 1; }))
+								FCanExecuteAction::CreateSPLambda(this, [this]() { return ColorThemes.Num() > 1; }))
 						);
 
 						SubMenuBuilder.EndSection();
@@ -1376,7 +1299,7 @@ TSharedRef<SWidget> SAvaLevelViewportActorColorMenu::BuildThemesMenu()
 							for (int32 ColorIdx = 0; ColorIdx < ColorThemes[ThemeIdx].GetColors().Num(); ++ColorIdx)
 							{
 								SubMenuBuilder.AddSubMenu(
-									FUIAction(FExecuteAction::CreateLambda([this, ThemeIdx, ColorIdx]()
+									FUIAction(FExecuteAction::CreateSPLambda(this, [this, ThemeIdx, ColorIdx]()
 										{
 											if (!ColorThemes.IsValidIndex(ThemeIdx))
 											{
@@ -1411,10 +1334,10 @@ TSharedRef<SWidget> SAvaLevelViewportActorColorMenu::BuildThemesMenu()
 											SNew(STextBlock)
 											.Text(ColorThemes[ThemeIdx].GetColors()[ColorIdx].Label)
 										],
-									FNewMenuDelegate::CreateLambda([this, ThemeIdx, ColorIdx](FMenuBuilder& SubMenuBuilder)
+									FNewMenuDelegate::CreateSPLambda(this, [this, ThemeIdx, ColorIdx](FMenuBuilder& SubMenuBuilder)
 										{
 											SubMenuBuilder.AddMenuEntry(LOCTEXT("RemoveColorLabel", "Remove"), LOCTEXT("RemoveColorTooltip", "Remove this color."), FSlateIcon(),
-												FUIAction(FExecuteAction::CreateLambda([this, ThemeIdx, ColorIdx]()
+												FUIAction(FExecuteAction::CreateSPLambda(this, [this, ThemeIdx, ColorIdx]()
 													{
 														if (!ColorThemes.IsValidIndex(ThemeIdx))
 														{
@@ -1440,7 +1363,7 @@ TSharedRef<SWidget> SAvaLevelViewportActorColorMenu::BuildThemesMenu()
 					}				
 				),
 				FUIAction(
-					FExecuteAction::CreateLambda([this,ThemeIdx]() 
+					FExecuteAction::CreateSPLambda(this, [this,ThemeIdx]() 
 						{
 							if (!ColorThemes.IsValidIndex(ThemeIdx))
 							{
@@ -1450,17 +1373,20 @@ TSharedRef<SWidget> SAvaLevelViewportActorColorMenu::BuildThemesMenu()
 							ActiveThemeIndex = ThemeIdx;
 							UpdateColorPalette();
 						}),
-					FCanExecuteAction::CreateLambda([this, ThemeIdx]()
+					FCanExecuteAction::CreateSPLambda(this, [this, ThemeIdx]()
 						{
 							return ActiveThemeIndex != ThemeIdx;
 						}),
-					FGetActionCheckState::CreateLambda([this, ThemeIdx]()
+					FGetActionCheckState::CreateSPLambda(this, [this, ThemeIdx]()
 						{
 							return ActiveThemeIndex == ThemeIdx ? ECheckBoxState::Checked : ECheckBoxState::Unchecked;
 						})
 				),
 				FName(ColorThemes[ThemeIdx].GetName()),
-				EUserInterfaceActionType::RadioButton
+				EUserInterfaceActionType::RadioButton,
+				/*bInOpenSubMenuOnClick*/false,
+				FSlateIcon(),
+				/*bInShouldCloseWindowAfterMenuSelection*/false
 			);
 		}
 
@@ -1477,7 +1403,7 @@ TSharedRef<SWidget> SAvaLevelViewportActorColorMenu::BuildThemesMenu()
 			.MinDesiredWidth(100.f)
 			.IsReadOnly(false)
 			.ClearKeyboardFocusOnCommit(true)
-			.OnTextCommitted(FOnTextCommitted::CreateLambda(
+			.OnTextCommitted(FOnTextCommitted::CreateSPLambda(this,
 				[this](const FText& NewName, ETextCommit::Type CommitType)
 				{
 					if (CommitType == ETextCommit::OnEnter)
@@ -1640,21 +1566,6 @@ void SAvaLevelViewportActorColorMenu::SetValueValue(float InNewValue, bool InbAd
 	}
 }
 
-void SAvaLevelViewportActorColorMenu::BroadcastColorSourceChange(const TSharedRef<IToolkitHost>& InToolkitHost, const FAvaColorChangeData& InNewColorData)
-{
-	if (InNewColorData.ColorStyle == EAvaColorStyle::None)
-	{
-		return;
-	}
-	
-	if (!FAvaViewportColorPickerDelegates::GetOnColorSourceSelected().IsBound())
-	{
-		return;
-	}
-
-	FAvaViewportColorPickerDelegates::BroadcastColorSourceSelected(InToolkitHost, InNewColorData);
-}
-
 void SAvaLevelViewportActorColorMenu::SetColorStyle(EAvaColorStyle InNewStyle)
 {
 	if (ColorStyle == InNewStyle || ColorStyle == EAvaColorStyle::None)
diff --git a/Engine/Plugins/Experimental/Avalanche/Source/AvalancheLevelViewport/Private/Widgets/SAvaLevelViewportActorColorMenu.h b/Engine/Plugins/Experimental/Avalanche/Source/AvalancheLevelViewport/Private/Widgets/SAvaLevelViewportActorColorMenu.h
--- a/Engine/Plugins/Experimental/Avalanche/Source/AvalancheLevelViewport/Private/Widgets/SAvaLevelViewportActorColorMenu.h
+++ b/Engine/Plugins/Experimental/Avalanche/Source/AvalancheLevelViewport/Private/Widgets/SAvaLevelViewportActorColorMenu.h
@@ -4,6 +4,7 @@
 
 #include "Widgets/SCompoundWidget.h"
 #include "AvaDefs.h"
+#include "ColorPicker/AvaLevelColorPicker.h"
 #include "Materials/Material.h"
 #include "Templates/SharedPointer.h"
 #include "UObject/StrongObjectPtr.h"
@@ -67,7 +68,7 @@ protected:
 	TArray<FAvaColorInfo> Colors;
 };
 
-class SAvaLevelViewportActorColorMenu : public SCompoundWidget
+class SAvaLevelViewportActorColorMenu : public SCompoundWidget, public FGCObject
 {
 	SLATE_DECLARE_WIDGET(SAvaLevelViewportActorColorMenu, SCompoundWidget)
 
@@ -80,10 +81,6 @@ public:
 	static FAvaColorTheme* FindTheme(const FString& InName, bool bInAddIfNotFound = false);
 	static FAvaColorTheme* GetTheme(int32 InThemeIndex);
 
-	static void BroadcastColorSourceChange(const TSharedRef<IToolkitHost>& InToolkitHost, const FAvaColorChangeData& InNewColorData);
-
-	virtual ~SAvaLevelViewportActorColorMenu() override;
-
 	void Construct(const FArguments& InArgs, const TSharedRef<IToolkitHost>& InToolkitHost);
 
 	virtual FReply OnMouseButtonDown(const FGeometry& InMyGeometry, const FPointerEvent& InMouseEvent) override;
@@ -115,7 +112,7 @@ public:
 	void SetSaturationValue(float InNewValue, bool InbAddToTheme);
 	void SetValueValue(float InNewValue, bool InbAddToTheme);
 
-protected:
+private:
 	static const FVector2f WheelMiddle;
 	static const inline int32 MaxColorsForAutoAdd = 21;
 	static constexpr int32 ColorsPerPaletteRow = 3;
@@ -125,23 +122,28 @@ protected:
 	static void LoadColorThemesFromIni();
 	static void SaveColorThemesToIni();
 
-	TWeakPtr<IToolkitHost> ToolkitHostWeak;
-	TStrongObjectPtr<UMaterial> HueMaterial;
-	TStrongObjectPtr<UMaterial> SatValueMaterial;
-	TStrongObjectPtr<UMaterialInstanceDynamic> SatValueMaterialDynamic;
-	TStrongObjectPtr<UTexture2D> SelectedValueImage;
-	FSlateMaterialBrush* HueBrush;
-	FSlateMaterialBrush* SatValueBrush;
-	FSlateImageBrush* SelectedValueBrush;
+	FAvaLevelColorPicker ColorPicker;
+
+	TObjectPtr<UMaterial> HueMaterial;
+	TObjectPtr<UMaterial> SatValueMaterial;
+	TObjectPtr<UTexture2D> SelectedValueImage;
+	TObjectPtr<UMaterialInstanceDynamic> SatValueMaterialDynamic;
+
+	TSharedPtr<FSlateMaterialBrush> HueBrush;
+	TSharedPtr<FSlateMaterialBrush> SatValueBrush;
+	TSharedPtr<FSlateImageBrush> SelectedValueBrush;
 
 	float Red;
 	float Green;
 	float Blue;
+
 	float Hue;
 	float Saturation;
 	float Value;
+
 	FLinearColor ActiveColor;
 	FLinearColor InactiveColor;
+
 	int32 ActiveThemeIndex;
 	EAvaColorStyle ColorStyle = EAvaColorStyle::Solid;
 	bool bIsUnlit = true;
@@ -149,8 +151,10 @@ protected:
 	bool bIsRGB = true;
 	TSharedPtr<STextBlock> ColorFormatButtonText;
 	TSharedPtr<SBox> ColorTextInputWidgetBox;
-	SConstraintCanvas::FSlot* HueSlot;
-	SConstraintCanvas::FSlot* SatValueSlot;
+
+	SConstraintCanvas::FSlot* HueSlot = nullptr;
+	SConstraintCanvas::FSlot* SatValueSlot = nullptr;
+
 	TSharedPtr<SGridPanel> ColorPalette;
 	TSharedPtr<SButton> ThemesButton;
 	TSharedPtr<SWidgetSwitcher> ColorStyleSwitcher;
@@ -158,6 +162,11 @@ protected:
 	bool bEditingHue;
 	bool bEditingSaturationValue;
 
+	//~ Begin FGCObject
+	virtual FString GetReferencerName() const override;
+	virtual void AddReferencedObjects(FReferenceCollector& InCollector) override;
+	//~ End FGCObject
+
 	void ApplyActiveColor();
 	void UpdateColorPalette();
 	void AddColorToTheme(const FLinearColor& InNewColor);
@@ -212,7 +221,6 @@ protected:
 	
 	void OnColorStyleChanged();
 	void BroadcastColorChange();
-	void OnColorSourceBroadcast(const TSharedRef<IToolkitHost>& InToolkitHost, const FAvaColorChangeData& InNewColorData);
 
 	void SetActiveColorRGB_Direct(FLinearColor InRGB);
 	void SetActiveColorHSV_Direct(FLinearColor InHSV);
diff --git a/Engine/Plugins/Experimental/Avalanche/Source/AvalancheLevelViewport/Private/Widgets/SAvaLevelViewportStatusBarButtons.cpp b/Engine/Plugins/Experimental/Avalanche/Source/AvalancheLevelViewport/Private/Widgets/SAvaLevelViewportStatusBarButtons.cpp
--- a/Engine/Plugins/Experimental/Avalanche/Source/AvalancheLevelViewport/Private/Widgets/SAvaLevelViewportStatusBarButtons.cpp
+++ b/Engine/Plugins/Experimental/Avalanche/Source/AvalancheLevelViewport/Private/Widgets/SAvaLevelViewportStatusBarButtons.cpp
@@ -2,8 +2,8 @@
 
 #include "Widgets/SAvaLevelViewportStatusBarButtons.h"
 #include "AvaEditorWidgetUtils.h"
-#include "AvalancheLevelViewportModule.h"
 #include "AvaLevelViewportCommands.h"
+#include "AvaLevelViewportModule.h"
 #include "AvaLevelViewportStyle.h"
 #include "AvaViewportPostProcessManager.h"
 #include "AvaViewportSettings.h"
diff --git a/Engine/Plugins/Experimental/Avalanche/Source/AvalancheLevelViewport/Private/Widgets/SAvaLevelViewportStatusBarTransformSettings.cpp b/Engine/Plugins/Experimental/Avalanche/Source/AvalancheLevelViewport/Private/Widgets/SAvaLevelViewportStatusBarTransformSettings.cpp
--- a/Engine/Plugins/Experimental/Avalanche/Source/AvalancheLevelViewport/Private/Widgets/SAvaLevelViewportStatusBarTransformSettings.cpp
+++ b/Engine/Plugins/Experimental/Avalanche/Source/AvalancheLevelViewport/Private/Widgets/SAvaLevelViewportStatusBarTransformSettings.cpp
@@ -3,13 +3,13 @@
 #include "Widgets/SAvaLevelViewportStatusBarTransformSettings.h"
 #include "AvaLevelViewportCommands.h"
 #include "CustomDetailsViewSequencer.h"
+#include "Customizations/AvaLevelViewportComponentTransformDetails.h"
 #include "IDetailKeyframeHandler.h"
-#include "SAvaLevelViewportFrame.h"
+#include "LevelEditor/AvaLevelEditorUtils.h"
 #include "Modules/ModuleManager.h"
 #include "PropertyEditorModule.h"
-#include "LevelEditor/AvaLevelEditorUtils.h"
+#include "SAvaLevelViewportFrame.h"
 #include "ViewportClient/AvaLevelViewportClient.h"
-#include "Widgets/AvaLevelViewportComponentTransformDetails.h"
 #include "Widgets/Layout/SBox.h"
 #include "Widgets/Layout/SWidgetSwitcher.h"
 #include "Widgets/SAvaLevelViewportStatusBar.h"
diff --git a/Engine/Plugins/Experimental/Avalanche/Source/AvalancheLevelViewport/Public/ColorPicker/AvaViewportColorPickerActorClassRegistry.h b/Engine/Plugins/Experimental/Avalanche/Source/AvalancheLevelViewport/Public/ColorPicker/AvaViewportColorPickerActorClassRegistry.h
new file mode 100644
--- /dev/null
+++ b/Engine/Plugins/Experimental/Avalanche/Source/AvalancheLevelViewport/Public/ColorPicker/AvaViewportColorPickerActorClassRegistry.h
@@ -0,0 +1,34 @@
+// Copyright Epic Games, Inc. All Rights Reserved.
+
+#pragma once
+
+#include "AvaViewportColorPickerAdapter.h"
+#include "Concepts/StaticClassProvider.h"
+#include "Templates/SharedPointer.h"
+
+class AActor;
+class UClass;
+struct FAvaColorChangeData;
+
+struct FAvaViewportColorPickerActorClassRegistry
+{
+	template<typename InClassType
+		UE_REQUIRES(TModels_V<CStaticClassProvider, InClassType>)>
+	static void RegisterDefaultClassAdapter()
+	{
+		RegisterClassAdapter<InClassType, TAvaViewportColorPickerActorAdapter<InClassType>>();
+	}
+
+	template<typename InClassType, typename InAdapterType, typename... InArgTypes
+		UE_REQUIRES(TModels_V<CStaticClassProvider, InClassType> && std::is_base_of_v<IAvaViewportColorPickerAdapter, InAdapterType>)>
+	static void RegisterClassAdapter(InArgTypes&&... InArgs)
+	{
+		RegisterClassAdapter(InClassType::StaticClass(), MakeShared<InAdapterType>(Forward<InArgTypes>(InArgs)...));
+	}
+
+	AVALANCHELEVELVIEWPORT_API static void RegisterClassAdapter(UClass* InClass, const TSharedRef<IAvaViewportColorPickerAdapter>& InAdapter);
+
+	AVALANCHELEVELVIEWPORT_API static bool GetColorDataFromActor(const AActor* InActor, FAvaColorChangeData& OutColorData);
+
+	static bool ApplyColorDataToActor(AActor* InActor, const FAvaColorChangeData& InColorData);
+};
diff --git a/Engine/Plugins/Experimental/Avalanche/Source/AvalancheLevelViewport/Public/ColorPicker/AvaViewportColorPickerAdapter.h b/Engine/Plugins/Experimental/Avalanche/Source/AvalancheLevelViewport/Public/ColorPicker/AvaViewportColorPickerAdapter.h
new file mode 100644
--- /dev/null
+++ b/Engine/Plugins/Experimental/Avalanche/Source/AvalancheLevelViewport/Public/ColorPicker/AvaViewportColorPickerAdapter.h
@@ -0,0 +1,48 @@
+﻿// Copyright Epic Games, Inc. All Rights Reserved.
+
+#pragma once
+
+#include "AvaDefs.h"
+#include "IAvaViewportColorPickerAdapter.h"
+
+/**
+ * Describes a type that implements:
+ * - GetColorData const member function that returns FAvaColorChangeData
+ * - SetColorData that accepts a const FAvaColorChangeData as a parameter
+ */
+struct CAvaViewportColorPickable
+{
+	template <typename T>
+	auto Requires(const T& InConstActor, FAvaColorChangeData& OutColorData, T& InMutableActor, const FAvaColorChangeData& InColorData)->decltype(
+		OutColorData = InConstActor.GetColorData(),
+		InMutableActor.SetColorData(InColorData)
+	);
+};
+
+/** Adapter class for a type that can be casted from an Actor (i.e. an Actor Derived type or an IInterface type) */
+template<typename T
+	UE_REQUIRES(TModels_V<CAvaViewportColorPickable, T>)>
+struct TAvaViewportColorPickerActorAdapter : IAvaViewportColorPickerAdapter
+{
+	virtual ~TAvaViewportColorPickerActorAdapter() override = default;
+
+	//~ Begin IAvaViewportColorPickerActorAdapter
+	virtual bool GetColorData(const AActor& InActor, FAvaColorChangeData& OutColorData) const override
+	{
+		if (const T* CastedActor = Cast<T>(&InActor))
+		{
+			OutColorData = CastedActor->GetColorData();
+			return true;
+		}
+		return false;
+	}
+
+	virtual void SetColorData(AActor& InActor, const FAvaColorChangeData& InColorData) const override
+	{
+		if (T* CastedActor = Cast<T>(&InActor))
+		{
+			CastedActor->SetColorData(InColorData);
+		}
+	}
+	//~ End IAvaViewportColorPickerActorAdapter
+};
diff --git a/Engine/Plugins/Experimental/Avalanche/Source/AvalancheLevelViewport/Public/ColorPicker/IAvaViewportColorPickerAdapter.h b/Engine/Plugins/Experimental/Avalanche/Source/AvalancheLevelViewport/Public/ColorPicker/IAvaViewportColorPickerAdapter.h
new file mode 100644
--- /dev/null
+++ b/Engine/Plugins/Experimental/Avalanche/Source/AvalancheLevelViewport/Public/ColorPicker/IAvaViewportColorPickerAdapter.h
@@ -0,0 +1,15 @@
+// Copyright Epic Games, Inc. All Rights Reserved.
+
+#pragma once
+
+class AActor;
+struct FAvaColorChangeData;
+
+struct IAvaViewportColorPickerAdapter
+{
+	virtual ~IAvaViewportColorPickerAdapter() = default;
+
+	virtual bool GetColorData(const AActor& InActor, FAvaColorChangeData& OutColorData) const = 0;
+
+	virtual void SetColorData(AActor& InActor, const FAvaColorChangeData& InColorData) const = 0;
+};
diff --git a/Engine/Plugins/Experimental/Avalanche/Source/AvalancheLevelViewport/Public/Widgets/AvaViewportColorPickerActorClassRegistry.cpp b/Engine/Plugins/Experimental/Avalanche/Source/AvalancheLevelViewport/Public/Widgets/AvaViewportColorPickerActorClassRegistry.cpp
deleted file mode 100644
diff --git a/Engine/Plugins/Experimental/Avalanche/Source/AvalancheLevelViewport/Public/Widgets/AvaViewportColorPickerActorClassRegistry.h b/Engine/Plugins/Experimental/Avalanche/Source/AvalancheLevelViewport/Public/Widgets/AvaViewportColorPickerActorClassRegistry.h
deleted file mode 100644
diff --git a/Engine/Plugins/Experimental/Avalanche/Source/AvalancheLevelViewport/Public/Widgets/AvaViewportColorPickerDelegates.h b/Engine/Plugins/Experimental/Avalanche/Source/AvalancheLevelViewport/Public/Widgets/AvaViewportColorPickerDelegates.h
deleted file mode 100644
diff --git a/Engine/Plugins/Experimental/Avalanche/Source/AvalancheShapesEditor/Private/AvalancheShapesEditorModule.cpp b/Engine/Plugins/Experimental/Avalanche/Source/AvalancheShapesEditor/Private/AvalancheShapesEditorModule.cpp
--- a/Engine/Plugins/Experimental/Avalanche/Source/AvalancheShapesEditor/Private/AvalancheShapesEditorModule.cpp
+++ b/Engine/Plugins/Experimental/Avalanche/Source/AvalancheShapesEditor/Private/AvalancheShapesEditorModule.cpp
@@ -10,7 +10,10 @@
 #include "IDetailsView.h"
 #include "ISequencerModule.h"
 #include "TrackEditors/AvaShapeRectCornerTrackEditor.h"
-#include "Widgets/AvaViewportColorPickerActorClassRegistry.h"
+
+// Color Picker
+#include "ColorPicker/AvaViewportColorPickerActorClassRegistry.h"
+#include "ColorPicker/AvaViewportColorPickerAdapter.h"
 
 // Meshes
 #include "DynamicMeshes/AvaShape2DArrowDynMesh.h"
@@ -112,7 +115,7 @@ void FAvalancheShapesEditorModule::StartupModule()
 	ISequencerModule& SequencerModule = FModuleManager::LoadModuleChecked<ISequencerModule>("Sequencer");
 	TrackEditorHandle = SequencerModule.RegisterTrackEditor(FOnCreateTrackEditor::CreateStatic(&FAvaShapeRectCornerTrackEditor::CreateTrackEditor));
 
-	FAvaViewportColorPickerActorClassRegistry::RegisterClassAdapter(AAvaShapeActor::StaticClass(), MakeShared<TAvaViewportColorPickerActorAdapter<AAvaShapeActor>>());
+	FAvaViewportColorPickerActorClassRegistry::RegisterDefaultClassAdapter<AAvaShapeActor>();
 }
 
 void FAvalancheShapesEditorModule::ShutdownModule()
diff --git a/Engine/Plugins/Experimental/Avalanche/Source/AvalancheTextEditor/Private/AvaTextEditorModule.cpp b/Engine/Plugins/Experimental/Avalanche/Source/AvalancheTextEditor/Private/AvaTextEditorModule.cpp
--- a/Engine/Plugins/Experimental/Avalanche/Source/AvalancheTextEditor/Private/AvaTextEditorModule.cpp
+++ b/Engine/Plugins/Experimental/Avalanche/Source/AvalancheTextEditor/Private/AvaTextEditorModule.cpp
@@ -6,6 +6,8 @@
 #include "AvaTextActor.h"
 #include "AvaTextDefs.h"
 #include "AvaTextEditorCommands.h"
+#include "ColorPicker/AvaViewportColorPickerActorClassRegistry.h"
+#include "ColorPicker/AvaViewportColorPickerAdapter.h"
 #include "DetailsView/AvaLinearGradientSettingsCustomization.h"
 #include "DetailsView/AvaTextAlignmentCustomization.h"
 #include "DetailsView/AvaTextComponentCustomization.h"
@@ -20,7 +22,6 @@
 #include "PropertyEditorModule.h"
 #include "Tool/AvaTextActorTool.h"
 #include "Visualizer/AvaTextVisualizer.h"
-#include "Widgets/AvaViewportColorPickerActorClassRegistry.h"
 
 void FAvaTextEditorModule::StartupModule()
 {
@@ -32,7 +33,7 @@ void FAvaTextEditorModule::StartupModule()
 
 	RegisterDynamicMaterialPropertyGenerator();
 
-	FAvaViewportColorPickerActorClassRegistry::RegisterClassAdapter(AAvaTextActor::StaticClass(), MakeShared<TAvaViewportColorPickerActorAdapter<AAvaTextActor>>());
+	FAvaViewportColorPickerActorClassRegistry::RegisterDefaultClassAdapter<AAvaTextActor>();
 }
 
 void FAvaTextEditorModule::ShutdownModule()
diff --git a/Engine/Plugins/Experimental/Avalanche/Source/AvalancheViewport/Private/AvaViewportAxisUtils.h b/Engine/Plugins/Experimental/Avalanche/Source/AvalancheViewport/Private/AvaViewportAxisUtils.h
new file mode 100644
--- /dev/null
+++ b/Engine/Plugins/Experimental/Avalanche/Source/AvalancheViewport/Private/AvaViewportAxisUtils.h
@@ -0,0 +1,43 @@
+// Copyright Epic Games, Inc. All Rights Reserved.
+
+#pragma once
+
+#include "Math/MathFwd.h"
+#include "Math/Axis.h"
+
+class AActor;
+
+template<typename T>
+struct TAvaAxisList
+{
+	T Horizontal;
+	T Vertical;
+	T Depth;
+};
+
+struct FAvaViewportAxisOrientation
+{
+	int32 Index;
+	bool bCodirectional;
+};
+
+using FAvaViewportAxisMap = TAvaAxisList<FAvaViewportAxisOrientation>;
+
+class FAvaViewportAxisUtils
+{
+public:
+	static constexpr FAvaViewportAxisMap WorldAxisIndexList = {
+		{1, true},
+		{2, true},
+		{0, true}
+	};
+		
+	/* Takes a viewport and actor's rotation and uses it to determine which of the actor's axes most closely match the camera's. */
+	[[nodiscard]] static FAvaViewportAxisMap CreateViewportAxisMap(const FTransform& InViewportCameraTransform, const AActor& InActor);
+
+	[[nodiscard]] static FAvaViewportAxisMap CreateViewportAxisMap(const FTransform& InViewportCameraTransform);
+
+	[[nodiscard]] static FAvaViewportAxisMap CreateViewportAxisMap(const TAvaAxisList<FVector>& InReferenceVectors, const TAvaAxisList<FVector>& InAxisVectors);
+
+	static EAxis::Type GetRotationAxisForVectorComponent(int32 InComponent);
+};
diff --git a/Engine/Plugins/Experimental/Avalanche/Source/AvalancheViewport/Public/AvaViewportAxisUtils.h b/Engine/Plugins/Experimental/Avalanche/Source/AvalancheViewport/Public/AvaViewportAxisUtils.h
deleted file mode 100644
