Motion Design
- Some visualizer icons are now displayed even when activated.

diff --git a/Engine/Plugins/Experimental/Avalanche/Source/AvalancheEffectorsEditor/Private/Effector/AvaEffectorActorVis.cpp b/Engine/Plugins/Experimental/Avalanche/Source/AvalancheEffectorsEditor/Private/Effector/AvaEffectorActorVis.cpp
--- a/Engine/Plugins/Experimental/Avalanche/Source/AvalancheEffectorsEditor/Private/Effector/AvaEffectorActorVis.cpp
+++ b/Engine/Plugins/Experimental/Avalanche/Source/AvalancheEffectorsEditor/Private/Effector/AvaEffectorActorVis.cpp
@@ -165,18 +165,12 @@ void FAvaEffectorActorVisualizer::DrawVisualizationEditing(const UActorComponent
 
 	if (EffectorActor->GetType() != ECEClonerEffectorType::Plane)
 	{
-		if (!bEditingInnerZone)
-		{
-			DrawZoneButton(EffectorActor, InView, InPDI, InOutIconIndex, true, FAvaVisualizerBase::Inactive);
-			InOutIconIndex++;
-		}
-	}
-
-	if (!bEditingOuterZone)
-	{
-		DrawZoneButton(EffectorActor, InView, InPDI, InOutIconIndex, false, bEditingOuterZone ? FAvaVisualizerBase::Active : FAvaVisualizerBase::Inactive);
+		DrawZoneButton(EffectorActor, InView, InPDI, InOutIconIndex, true, FAvaVisualizerBase::Inactive);
 		InOutIconIndex++;
 	}
+
+	DrawZoneButton(EffectorActor, InView, InPDI, InOutIconIndex, false, bEditingOuterZone ? FAvaVisualizerBase::Active : FAvaVisualizerBase::Inactive);
+	InOutIconIndex++;
 }
 
 void FAvaEffectorActorVisualizer::DrawVisualizationNotEditing(const UActorComponent* InComponent, const FSceneView* InView, FPrimitiveDrawInterface* InPDI, int32& InOutIconIndex)
diff --git a/Engine/Plugins/Experimental/Avalanche/Source/AvalancheShapesEditor/Private/Visualizers/AvaShape3DDynMeshVis.cpp b/Engine/Plugins/Experimental/Avalanche/Source/AvalancheShapesEditor/Private/Visualizers/AvaShape3DDynMeshVis.cpp
--- a/Engine/Plugins/Experimental/Avalanche/Source/AvalancheShapesEditor/Private/Visualizers/AvaShape3DDynMeshVis.cpp
+++ b/Engine/Plugins/Experimental/Avalanche/Source/AvalancheShapesEditor/Private/Visualizers/AvaShape3DDynMeshVis.cpp
@@ -191,10 +191,7 @@ void FAvaShape3DDynamicMeshVisualizer::DrawVisualizationEditing(const UActorComp
 	{
 		for (const AvaAlignment Alignment : SupportedAlignments)
 		{
-			if (SizeDragAnchor != Alignment)
-			{
-				DrawSizeButton(DynMesh, InView, InPDI, Alignment);
-			}
+			DrawSizeButton(DynMesh, InView, InPDI, Alignment);
 		}
 	}
 }
diff --git a/Engine/Plugins/Experimental/Avalanche/Source/AvalancheShapesEditor/Private/Visualizers/AvaShapeConeDynMeshVis.cpp b/Engine/Plugins/Experimental/Avalanche/Source/AvalancheShapesEditor/Private/Visualizers/AvaShapeConeDynMeshVis.cpp
--- a/Engine/Plugins/Experimental/Avalanche/Source/AvalancheShapesEditor/Private/Visualizers/AvaShapeConeDynMeshVis.cpp
+++ b/Engine/Plugins/Experimental/Avalanche/Source/AvalancheShapesEditor/Private/Visualizers/AvaShapeConeDynMeshVis.cpp
@@ -297,7 +297,7 @@ void FAvaShapeConeDynamicMeshVisualizer::DrawVisualizationNotEditing(const UActo
 void FAvaShapeConeDynamicMeshVisualizer::DrawVisualizationEditing(const UActorComponent* InComponent,
 	const FSceneView* InView, FPrimitiveDrawInterface* InPDI, int32& InOutIconIndex)
 {
-	Super::DrawVisualizationNotEditing(InComponent, InView, InPDI, InOutIconIndex);
+	Super::DrawVisualizationEditing(InComponent, InView, InPDI, InOutIconIndex);
 
 	const FMeshType* DynMesh = Cast<FMeshType>(InComponent);
 
diff --git a/Engine/Plugins/Experimental/Avalanche/Source/AvalancheShapesEditor/Private/Visualizers/AvaShapeIrregularPolygonDynMeshVis.cpp b/Engine/Plugins/Experimental/Avalanche/Source/AvalancheShapesEditor/Private/Visualizers/AvaShapeIrregularPolygonDynMeshVis.cpp
--- a/Engine/Plugins/Experimental/Avalanche/Source/AvalancheShapesEditor/Private/Visualizers/AvaShapeIrregularPolygonDynMeshVis.cpp
+++ b/Engine/Plugins/Experimental/Avalanche/Source/AvalancheShapesEditor/Private/Visualizers/AvaShapeIrregularPolygonDynMeshVis.cpp
@@ -156,16 +156,8 @@ void FAvaShapeIrregularPolygonDynamicMeshVisualizer::DrawVisualizationEditing(co
 
 	for (int32 PointIterIdx = 0; PointIterIdx < DynMesh->GetNumPoints(); ++PointIterIdx)
 	{
-		if (PointIterIdx != EditingPointIdx)
-		{
-			DrawPointButton(DynMesh, InView, InPDI, Inactive, PointIterIdx);
-		}
-
-		if (PointIterIdx != EditingBevelIdx)
-		{
-			DrawBevelButton(DynMesh, InView, InPDI, Inactive, PointIterIdx);
-		}
-
+		DrawPointButton(DynMesh, InView, InPDI, Inactive, PointIterIdx);
+		DrawBevelButton(DynMesh, InView, InPDI, Inactive, PointIterIdx);
 		DrawBreakButton(DynMesh, InView, InPDI, Inactive, PointIterIdx);
 	}
 }
diff --git a/Engine/Plugins/Experimental/Avalanche/Source/AvalancheShapesEditor/Private/Visualizers/AvaShapeLineDynMeshVis.cpp b/Engine/Plugins/Experimental/Avalanche/Source/AvalancheShapesEditor/Private/Visualizers/AvaShapeLineDynMeshVis.cpp
--- a/Engine/Plugins/Experimental/Avalanche/Source/AvalancheShapesEditor/Private/Visualizers/AvaShapeLineDynMeshVis.cpp
+++ b/Engine/Plugins/Experimental/Avalanche/Source/AvalancheShapesEditor/Private/Visualizers/AvaShapeLineDynMeshVis.cpp
@@ -96,15 +96,8 @@ void FAvaShapeLineDynamicMeshVisualizer::DrawVisualizationEditing(const UActorCo
 		return;
 	}
 
-	if (!bEditingStart)
-	{
-		DrawStartButton(DynMesh, InView, InPDI, Inactive);
-	}
-
-	if (!bEditingEnd)
-	{
-		DrawEndButton(DynMesh, InView, InPDI, Inactive);
-	}
+	DrawStartButton(DynMesh, InView, InPDI, Inactive);
+	DrawEndButton(DynMesh, InView, InPDI, Inactive);
 }
 
 FVector FAvaShapeLineDynamicMeshVisualizer::GetStartWorldLocation(const FMeshType* InDynMesh) const
diff --git a/Engine/Plugins/Experimental/Avalanche/Source/AvalancheShapesEditor/Private/Visualizers/AvaShapeRectangleDynMeshVis.cpp b/Engine/Plugins/Experimental/Avalanche/Source/AvalancheShapesEditor/Private/Visualizers/AvaShapeRectangleDynMeshVis.cpp
--- a/Engine/Plugins/Experimental/Avalanche/Source/AvalancheShapesEditor/Private/Visualizers/AvaShapeRectangleDynMeshVis.cpp
+++ b/Engine/Plugins/Experimental/Avalanche/Source/AvalancheShapesEditor/Private/Visualizers/AvaShapeRectangleDynMeshVis.cpp
@@ -162,32 +162,14 @@ void FAvaShapeRectangleDynamicMeshVisualizer::DrawVisualizationEditing(const UAc
 	DrawGlobalBevelButton(DynMesh, InView, InPDI, InOutIconIndex, bEditingGlobalBevelSize ? Active : Inactive);
 	++InOutIconIndex;
 
-	if (!bEditingGlobalBevelSize)
+	if (ShouldDrawExtraHandles(InComponent, InView))
 	{
-		if (Corner != EAvaAnchors::TopLeft)
-		{
-			DrawBevelButton(DynMesh, InView, InPDI, Inactive, EAvaAnchors::TopLeft);
-		}
-		if (Corner != EAvaAnchors::TopRight)
-		{
-			DrawBevelButton(DynMesh, InView, InPDI, Inactive, EAvaAnchors::TopRight);
-		}
-		if (Corner != EAvaAnchors::BottomLeft)
-		{
-			DrawBevelButton(DynMesh, InView, InPDI, Inactive, EAvaAnchors::BottomLeft);
-		}
-		if (Corner != EAvaAnchors::BottomRight)
-		{
-			DrawBevelButton(DynMesh, InView, InPDI, Inactive, EAvaAnchors::BottomRight);
-		}
-		if (SlantSide != EAvaAnchors::Left)
-		{
-			DrawSlantButton(DynMesh, InView, InPDI, Inactive, EAvaAnchors::Left);
-		}
-		if (SlantSide != EAvaAnchors::Right)
-		{
-			DrawSlantButton(DynMesh, InView, InPDI, Inactive, EAvaAnchors::Right);
-		}
+		DrawBevelButton(DynMesh, InView, InPDI, Inactive, EAvaAnchors::TopLeft);
+		DrawBevelButton(DynMesh, InView, InPDI, Inactive, EAvaAnchors::TopRight);
+		DrawBevelButton(DynMesh, InView, InPDI, Inactive, EAvaAnchors::BottomLeft);
+		DrawBevelButton(DynMesh, InView, InPDI, Inactive, EAvaAnchors::BottomRight);
+		DrawSlantButton(DynMesh, InView, InPDI, Inactive, EAvaAnchors::Left);
+		DrawSlantButton(DynMesh, InView, InPDI, Inactive, EAvaAnchors::Right);
 	}
 }
 
diff --git a/Engine/Plugins/Experimental/Avalanche/Source/AvalancheShapesEditor/Private/Visualizers/AvaShapeSphereDynMeshVis.cpp b/Engine/Plugins/Experimental/Avalanche/Source/AvalancheShapesEditor/Private/Visualizers/AvaShapeSphereDynMeshVis.cpp
--- a/Engine/Plugins/Experimental/Avalanche/Source/AvalancheShapesEditor/Private/Visualizers/AvaShapeSphereDynMeshVis.cpp
+++ b/Engine/Plugins/Experimental/Avalanche/Source/AvalancheShapesEditor/Private/Visualizers/AvaShapeSphereDynMeshVis.cpp
@@ -319,7 +319,7 @@ void FAvaShapeSphereDynamicMeshVisualizer::DrawVisualizationNotEditing(const UAc
 void FAvaShapeSphereDynamicMeshVisualizer::DrawVisualizationEditing(const UActorComponent* InComponent,
 	const FSceneView* InView, FPrimitiveDrawInterface* InPDI, int32& InOutIconIndex)
 {
-	Super::DrawVisualizationNotEditing(InComponent, InView, InPDI, InOutIconIndex);
+	Super::DrawVisualizationEditing(InComponent, InView, InPDI, InOutIconIndex);
 
 	const FMeshType* DynMesh = Cast<FMeshType>(InComponent);
 
diff --git a/Engine/Plugins/Experimental/Avalanche/Source/AvalancheTextEditor/Private/Visualizer/AvaTextVisualizer.cpp b/Engine/Plugins/Experimental/Avalanche/Source/AvalancheTextEditor/Private/Visualizer/AvaTextVisualizer.cpp
--- a/Engine/Plugins/Experimental/Avalanche/Source/AvalancheTextEditor/Private/Visualizer/AvaTextVisualizer.cpp
+++ b/Engine/Plugins/Experimental/Avalanche/Source/AvalancheTextEditor/Private/Visualizer/AvaTextVisualizer.cpp
@@ -999,15 +999,8 @@ void FAvaTextVisualizer::DrawVisualizationEditing(const UActorComponent* InCompo
 	DrawScaleProportionallyButton(AvaTextComponent, InView, InPDI, InOutIconIndex, AvaTextComponent->ScalesProportionally() ? Enabled : Disabled);
 	++InOutIconIndex;
 
-	if (AvaTextComponent->HasMaxWidth() && !bEditingWidth)
-	{
-		DrawMaxTextWidthHandle(AvaTextComponent, InView, InPDI, Inactive);
-	}
-
-	if (AvaTextComponent->HasMaxHeight() && !bEditingHeight)
-	{
-		DrawMaxTextHeightHandle(AvaTextComponent, InView, InPDI, Inactive);
-	}
+	DrawMaxTextWidthHandle(AvaTextComponent, InView, InPDI, Inactive);
+	DrawMaxTextHeightHandle(AvaTextComponent, InView, InPDI, Inactive);
 	
 	if (AvaTextComponent->GetColoringStyle() == EAvaTextColoringStyle::Gradient)
 	{
