Motion Design
- Removed post process from camera menu (duplicated from status bar menu)
- Fixed missing icons in status bar post process menu.

diff --git a/Engine/Plugins/Experimental/Avalanche/Source/AvalancheLevelViewport/Private/SAvaLevelViewport.h b/Engine/Plugins/Experimental/Avalanche/Source/AvalancheLevelViewport/Private/SAvaLevelViewport.h
--- a/Engine/Plugins/Experimental/Avalanche/Source/AvalancheLevelViewport/Private/SAvaLevelViewport.h
+++ b/Engine/Plugins/Experimental/Avalanche/Source/AvalancheLevelViewport/Private/SAvaLevelViewport.h
@@ -218,8 +218,6 @@ private:
 	void AddGuidePresetCurrentMenu(UToolMenu* InMenu);
 	void AddGuidePresetSavedMenu(UToolMenu* InMenu);
 
-	void AddPostProcessMenuEntries(UToolMenu* InMenu);
-
 	void AddCameraZoomMenuEntries(UToolMenu* InMenu);
 
 	void AddVisualizerEntries(UToolMenu* InMenu);
diff --git a/Engine/Plugins/Experimental/Avalanche/Source/AvalancheLevelViewport/Private/SAvaLevelViewport_Menus.cpp b/Engine/Plugins/Experimental/Avalanche/Source/AvalancheLevelViewport/Private/SAvaLevelViewport_Menus.cpp
--- a/Engine/Plugins/Experimental/Avalanche/Source/AvalancheLevelViewport/Private/SAvaLevelViewport_Menus.cpp
+++ b/Engine/Plugins/Experimental/Avalanche/Source/AvalancheLevelViewport/Private/SAvaLevelViewport_Menus.cpp
@@ -215,15 +215,6 @@ void SAvaLevelViewport::FillCameraMenu(UToolMenu* InMenu)
 		FNewToolMenuDelegate::CreateSP(this, &SAvaLevelViewport::AddGuidePresetMenuEntries)
 	);
 
-	static const FText PostProcessLabel = LOCTEXT("ViewportPostProcess", "Post Process");
-
-	VirtualSection.AddSubMenu(
-		"ViewportPostProcess",
-		PostProcessLabel,
-		FText(),
-		FNewToolMenuDelegate::CreateSP(this, &SAvaLevelViewport::AddPostProcessMenuEntries)
-	);
-
 	static const FText CameraZoomLabel = LOCTEXT("ViewportCameraZoom", "Camera Zoom");
 
 	VirtualSection.AddSubMenu(
@@ -752,90 +743,6 @@ void SAvaLevelViewport::AddGuidePresetSavedMenu(UToolMenu* InMenu)
 	}
 }
 
-void SAvaLevelViewport::AddPostProcessMenuEntries(UToolMenu* InMenu)
-{
-	const FAvaLevelViewportCommands& LevelViewportCommands = FAvaLevelViewportCommands::Get();
-
-	FToolMenuSection& PostProcessSection = InMenu->AddSection("ViewportPostProcessType", LOCTEXT("ViewportPostProcessType", "Post Process Type"));
-
-	PostProcessSection.AddEntry(FToolMenuEntry::InitMenuEntry(
-		LevelViewportCommands.TogglePostProcessNone
-	));
-
-	PostProcessSection.AddEntry(FToolMenuEntry::InitMenuEntry(
-		LevelViewportCommands.TogglePostProcessBackplate
-	));
-
-	PostProcessSection.AddEntry(FToolMenuEntry::InitMenuEntry(
-		LevelViewportCommands.TogglePostProcessChannelRed
-	));
-
-	PostProcessSection.AddEntry(FToolMenuEntry::InitMenuEntry(
-		LevelViewportCommands.TogglePostProcessChannelGreen
-	));
-
-	PostProcessSection.AddEntry(FToolMenuEntry::InitMenuEntry(
-		LevelViewportCommands.TogglePostProcessChannelBlue
-	));
-
-	PostProcessSection.AddEntry(FToolMenuEntry::InitMenuEntry(
-		LevelViewportCommands.TogglePostProcessChannelAlpha
-	));
-
-	PostProcessSection.AddEntry(FToolMenuEntry::InitMenuEntry(
-		LevelViewportCommands.TogglePostProcessCheckerboard
-	));
-
-	if (!PostProcessOpacitySlider.IsValid())
-	{
-		PostProcessOpacitySlider = SNew(SSpinBox<float>)
-			.ClearKeyboardFocusOnCommit(true)
-			.MaxFractionalDigits(3)
-			.MinDesiredWidth(50.f)
-			.OnBeginSliderMovement(this, &SAvaLevelViewport::OnBackplateOpacitySliderBegin)
-			.OnEndSliderMovement(this, &SAvaLevelViewport::OnBackplateOpacitySliderEnd)
-			.OnValueCommitted(this, &SAvaLevelViewport::OnBackplateOpacityCommitted)
-			.Value(this, &SAvaLevelViewport::GetBackplateOpacity)
-			.MinValue(0.f)
-			.MinSliderValue(0.f)
-			.MaxValue(1.f)
-			.MaxSliderValue(1.f);
-	}
-
-	if (!BackplateTextureSelector.IsValid())
-	{
-		BackplateTextureSelector = SNew(SObjectPropertyEntryBox)
-			.AllowClear(true)
-			.AllowedClass(UTexture::StaticClass())
-			.DisplayBrowse(true)
-			.DisplayThumbnail(true)
-			.DisplayCompactSize(true)
-			.DisplayUseSelected(true)
-			.ThumbnailPool(UThumbnailManager::Get().GetSharedThumbnailPool())
-			.EnableContentPicker(true)
-			.ObjectPath(this, &SAvaLevelViewport::GetBackplateTextureObjectPath)
-			.OnObjectChanged(this, &SAvaLevelViewport::OnBackplateTextureChanged)
-			.OnShouldSetAsset(FOnShouldSetAsset::CreateLambda([](const FAssetData& InAssetData) { return false; }))
-			.IsEnabled(this, &SAvaLevelViewport::IsPostProcessBackplateEnabled);
-	}
-
-	FToolMenuSection& PostProcessSettingsSection = InMenu->AddSection("ViewportPostProcessSettings", LOCTEXT("ViewportPostProcessSettings", "Post Process Settings"));
-
-	PostProcessSettingsSection.AddEntry(FToolMenuEntry::InitWidget(
-		"PostProcessOpacity",
-		PostProcessOpacitySlider.ToSharedRef(),
-		LOCTEXT("PostProcessOpacity", "Opacity"),
-		true
-	));
-
-	PostProcessSettingsSection.AddEntry(FToolMenuEntry::InitWidget(
-		"PostProcessTexture", 
-		BackplateTextureSelector.ToSharedRef(),
-		LOCTEXT("PostProcessTexture", "Texture"),
-		true
-	));
-}
-
 void SAvaLevelViewport::AddCameraZoomMenuEntries(UToolMenu* InMenu)
 {
 	const FAvaLevelViewportCommands& LevelViewportCommands = FAvaLevelViewportCommands::Get();
diff --git a/Engine/Plugins/Experimental/Avalanche/Source/AvalancheLevelViewport/Private/Widgets/SAvaLevelViewportStatusBarButtons.cpp b/Engine/Plugins/Experimental/Avalanche/Source/AvalancheLevelViewport/Private/Widgets/SAvaLevelViewportStatusBarButtons.cpp
--- a/Engine/Plugins/Experimental/Avalanche/Source/AvalancheLevelViewport/Private/Widgets/SAvaLevelViewportStatusBarButtons.cpp
+++ b/Engine/Plugins/Experimental/Avalanche/Source/AvalancheLevelViewport/Private/Widgets/SAvaLevelViewportStatusBarButtons.cpp
@@ -35,14 +35,15 @@
 
 namespace UE::Ava::LevelViewportStatusBarButtons::Private
 {
-	static const FName StyleSetName = FAvaLevelViewportStyle::Get().GetStyleSetName();
-	static const FSlateIcon RGBChannelIcon = FSlateIcon(StyleSetName, "AvalancheIcons.PostProcess.RGB");
-	static const FSlateIcon BackplateIcon = FSlateIcon(StyleSetName, "Icons.Role");
-	static const FSlateIcon RedChannelIcon = FSlateIcon(StyleSetName, "AvalancheIcons.PostProcess.Red");
-	static const FSlateIcon GreenChannelIcon = FSlateIcon(StyleSetName, "AvalancheIcons.PostProcess.Green");
-	static const FSlateIcon BlueChannelIcon = FSlateIcon(StyleSetName, "AvalancheIcons.PostProcess.Blue");
-	static const FSlateIcon AlphaChannelIcon = FSlateIcon(StyleSetName, "AvalancheIcons.PostProcess.Alpha");
-	static const FSlateIcon CheckerboardIcon = FSlateIcon(StyleSetName, "Checker");
+	static const FName AvaLevelViewportStyleName = FAvaLevelViewportStyle::Get().GetStyleSetName();
+	static const FName AppStyleSetName = FAppStyle::Get().GetStyleSetName();
+	static const FSlateIcon RGBChannelIcon = FSlateIcon(AvaLevelViewportStyleName, "AvalancheIcons.PostProcess.RGB");
+	static const FSlateIcon BackplateIcon = FSlateIcon(AppStyleSetName, "Icons.Role");
+	static const FSlateIcon RedChannelIcon = FSlateIcon(AvaLevelViewportStyleName, "AvalancheIcons.PostProcess.Red");
+	static const FSlateIcon GreenChannelIcon = FSlateIcon(AvaLevelViewportStyleName, "AvalancheIcons.PostProcess.Green");
+	static const FSlateIcon BlueChannelIcon = FSlateIcon(AvaLevelViewportStyleName, "AvalancheIcons.PostProcess.Blue");
+	static const FSlateIcon AlphaChannelIcon = FSlateIcon(AvaLevelViewportStyleName, "AvalancheIcons.PostProcess.Alpha");
+	static const FSlateIcon CheckerboardIcon = FSlateIcon(AppStyleSetName, "Checker");
 
 	static bool IsViewportPostProcessManagerEnabled(const TWeakPtr<SAvaLevelViewportFrame>& InViewportFrameWeak)
 	{
