[Motion Design]
 - Update ease curve settings menu to use UToolMenu

diff --git a/Engine/Plugins/Experimental/Avalanche/Source/AvalancheSequencer/Private/EaseCurveTool/Widgets/SAvaEaseCurveTool.cpp b/Engine/Plugins/Experimental/Avalanche/Source/AvalancheSequencer/Private/EaseCurveTool/Widgets/SAvaEaseCurveTool.cpp
--- a/Engine/Plugins/Experimental/Avalanche/Source/AvalancheSequencer/Private/EaseCurveTool/Widgets/SAvaEaseCurveTool.cpp
+++ b/Engine/Plugins/Experimental/Avalanche/Source/AvalancheSequencer/Private/EaseCurveTool/Widgets/SAvaEaseCurveTool.cpp
@@ -20,6 +20,7 @@
 #include "SCurveEditor.h"
 #include "Styling/AppStyle.h"
 #include "Styling/StyleColors.h"
+#include "ToolMenu.h"
 #include "ToolMenus.h"
 #include "Widgets/Images/SImage.h"
 #include "Widgets/Input/SButton.h"
@@ -626,7 +627,7 @@ void SAvaEaseCurveTool::BindCommands()
 
 TSharedRef<SWidget> SAvaEaseCurveTool::CreateContextMenuContent()
 {
-	UToolMenus* ToolMenus = UToolMenus::Get();
+	UToolMenus* const ToolMenus = UToolMenus::Get();
 	if (!ToolMenus)
 	{
 		return SNullWidget::NullWidget;
@@ -645,7 +646,7 @@ TSharedRef<SWidget> SAvaEaseCurveTool::CreateContextMenuContent()
 		Section.AddSubMenu(TEXT("Settings"),
 			LOCTEXT("SettingsSubMenuLabel", "Settings"),
 			LOCTEXT("SettingsSubMenuToolTip", ""),
-			FNewMenuDelegate::CreateSP(this, &SAvaEaseCurveTool::MakeContextMenuSettings),
+			FNewToolMenuDelegate::CreateSP(this, &SAvaEaseCurveTool::MakeContextMenuSettings),
 			false,
 			FSlateIcon(FAppStyle::Get().GetStyleSetName(), TEXT("Icons.Toolbar.Settings")));
 
@@ -745,17 +746,24 @@ TSharedRef<SWidget> SAvaEaseCurveTool::CreateContextMenuContent()
 	return ToolMenus->GenerateWidget(MenuName, FToolMenuContext(CommandList));
 }
 
-void SAvaEaseCurveTool::MakeContextMenuSettings(FMenuBuilder& InMenuBuilder)
+void SAvaEaseCurveTool::MakeContextMenuSettings(UToolMenu* const InToolMenu)
 {
+	if (!IsValid(InToolMenu))
+	{
+		return;
+	}
+
 	const FAvaEaseCurveToolCommands& EaseCurveToolCommands = FAvaEaseCurveToolCommands::Get();
 
-	InMenuBuilder.AddMenuEntry(EaseCurveToolCommands.OpenToolSettings);
+	FToolMenuSection& Section = InToolMenu->FindOrAddSection(TEXT("EaseCurveToolSettings"), LOCTEXT("EaseCurveToolSettingsActions", "Settings"));
 
-	InMenuBuilder.AddSeparator();
+	Section.AddMenuEntry(EaseCurveToolCommands.OpenToolSettings);
 
-	InMenuBuilder.AddMenuEntry(EaseCurveToolCommands.ToggleAutoFlipTangents);
+	Section.AddSeparator(NAME_None);
 
-	InMenuBuilder.AddSeparator();
+	Section.AddMenuEntry(EaseCurveToolCommands.ToggleAutoFlipTangents);
+
+	Section.AddSeparator(NAME_None);
 
 	// Graph Size
 	{
@@ -806,7 +814,8 @@ void SAvaEaseCurveTool::MakeContextMenuSettings(FMenuBuilder& InMenuBuilder)
 					EaseCurveToolSettings->SetGraphSize(CurrentGraphSize);
 					EaseCurveToolSettings->SaveConfig();
 				});
-		InMenuBuilder.AddWidget(GraphSizeWidget, LOCTEXT("ToolSizeLabel", "Tool Size"));
+
+		Section.AddEntry(FToolMenuEntry::InitWidget(TEXT("ToolSize"), GraphSizeWidget, LOCTEXT("ToolSizeLabel", "Tool Size")));
 	}
 
 	// Grid Size
@@ -858,16 +867,17 @@ void SAvaEaseCurveTool::MakeContextMenuSettings(FMenuBuilder& InMenuBuilder)
 				{
 					SetEaseCurveToolGridSize(InNewValue);
 				});
-		InMenuBuilder.AddWidget(GridSizeWidget, LOCTEXT("GridSizeLabel", "Grid Size"));
+
+		Section.AddEntry(FToolMenuEntry::InitWidget(TEXT("GridSize"), GridSizeWidget, LOCTEXT("GridSizeLabel", "Grid Size")));
 	}
 
-	InMenuBuilder.AddSeparator();
+	Section.AddSeparator(NAME_None);
 
-	InMenuBuilder.AddMenuEntry(EaseCurveToolCommands.ToggleAutoZoomToFit);
+	Section.AddMenuEntry(EaseCurveToolCommands.ToggleAutoZoomToFit);
 
-	InMenuBuilder.AddSeparator();
+	Section.AddSeparator(NAME_None);
 
-	InMenuBuilder.AddMenuEntry(EaseCurveToolCommands.ResetToDefaultPresets);
+	Section.AddMenuEntry(EaseCurveToolCommands.ResetToDefaultPresets);
 }
 
 void SAvaEaseCurveTool::UndoAction()
diff --git a/Engine/Plugins/Experimental/Avalanche/Source/AvalancheSequencer/Private/EaseCurveTool/Widgets/SAvaEaseCurveTool.h b/Engine/Plugins/Experimental/Avalanche/Source/AvalancheSequencer/Private/EaseCurveTool/Widgets/SAvaEaseCurveTool.h
--- a/Engine/Plugins/Experimental/Avalanche/Source/AvalancheSequencer/Private/EaseCurveTool/Widgets/SAvaEaseCurveTool.h
+++ b/Engine/Plugins/Experimental/Avalanche/Source/AvalancheSequencer/Private/EaseCurveTool/Widgets/SAvaEaseCurveTool.h
@@ -14,6 +14,7 @@ class FText;
 class SAvaEaseCurveEditor;
 class SAvaEaseCurvePreset;
 class UCurveBase;
+class UToolMenu;
 struct FKeyHandle;
 struct FRichCurve;
 
@@ -68,7 +69,7 @@ protected:
 	void BindCommands();
 
 	TSharedRef<SWidget> CreateContextMenuContent();
-	void MakeContextMenuSettings(FMenuBuilder& InMenuBuilder);
+	void MakeContextMenuSettings(UToolMenu* const InToolMenu);
 
 	void UndoAction();
 	void RedoAction();
