Motion Design:
Fix issue with Scene Tags not being able to be added/removed dynamically in RC.
RC, by design, will set a value and there's no current way to add/remove scene attributes without disturbing the entire array of handles.
To circumvent this, the scene attribute object implementation is being introduced.
This change removes the existing Scene Tags property given its limitations in favor of Scene Tag Objects.

#changelist validated
#virtualized

diff --git a/Engine/Plugins/Experimental/Avalanche/Source/Avalanche/Private/Attributes/AvaSceneTagAttribute.cpp b/Engine/Plugins/Experimental/Avalanche/Source/Avalanche/Private/Attributes/AvaSceneTagAttribute.cpp
new file mode 100644
--- /dev/null
+++ b/Engine/Plugins/Experimental/Avalanche/Source/Avalanche/Private/Attributes/AvaSceneTagAttribute.cpp
@@ -0,0 +1,13 @@
+// Copyright Epic Games, Inc. All Rights Reserved.
+
+#include "Attributes/AvaSceneTagAttribute.h"
+
+bool UAvaSceneTagAttribute::ContainsTag(const FAvaTagHandle& InTagHandle) const
+{
+	return Tag.MatchesTag(InTagHandle);
+}
+
+bool UAvaSceneTagContainerAttribute::ContainsTag(const FAvaTagHandle& InTagHandle) const
+{
+	return TagContainer.ContainsTag(InTagHandle);
+}
diff --git a/Engine/Plugins/Experimental/Avalanche/Source/Avalanche/Private/AvaSceneSettings.cpp b/Engine/Plugins/Experimental/Avalanche/Source/Avalanche/Private/AvaSceneSettings.cpp
deleted file mode 100644
diff --git a/Engine/Plugins/Experimental/Avalanche/Source/Avalanche/Private/AvaSceneState.cpp b/Engine/Plugins/Experimental/Avalanche/Source/Avalanche/Private/AvaSceneState.cpp
--- a/Engine/Plugins/Experimental/Avalanche/Source/Avalanche/Private/AvaSceneState.cpp
+++ b/Engine/Plugins/Experimental/Avalanche/Source/Avalanche/Private/AvaSceneState.cpp
@@ -1,6 +1,7 @@
 // Copyright Epic Games, Inc. All Rights Reserved.
 
 #include "AvaSceneState.h"
+#include "Attributes/AvaSceneTagAttribute.h"
 #include "AvaSceneSettings.h"
 #include "AvaTagHandle.h"
 
@@ -42,6 +43,21 @@ bool UAvaSceneState::ContainsTagAttribute(const FAvaTagHandle& InTagHandle) cons
 		return true;
 	}
 
-	UAvaSceneSettings* SceneSettings = SceneSettingsWeak.Get();
-	return SceneSettings && SceneSettings->GetTagAttributes().ContainsTag(InTagHandle);
+	const UAvaSceneSettings* SceneSettings = SceneSettingsWeak.Get();
+	if (!SceneSettings)
+	{
+		return false;
+	}
+
+	bool bFoundTag = false;
+
+	SceneSettings->ForEachSceneAttributeOfType<UAvaSceneTagAttributeBase>(
+		[&InTagHandle, &bFoundTag](const UAvaSceneTagAttributeBase& InTagAttribute)
+		{
+			bFoundTag = InTagAttribute.ContainsTag(InTagHandle);
+			// continue iteration if tag not found
+			return !bFoundTag;
+		});
+
+	return bFoundTag;
 }
diff --git a/Engine/Plugins/Experimental/Avalanche/Source/Avalanche/Public/Attributes/AvaSceneAttribute.h b/Engine/Plugins/Experimental/Avalanche/Source/Avalanche/Public/Attributes/AvaSceneAttribute.h
new file mode 100644
--- /dev/null
+++ b/Engine/Plugins/Experimental/Avalanche/Source/Avalanche/Public/Attributes/AvaSceneAttribute.h
@@ -0,0 +1,13 @@
+// Copyright Epic Games, Inc. All Rights Reserved.
+
+#pragma once
+
+#include "UObject/Object.h"
+#include "AvaSceneAttribute.generated.h"
+
+/** Scene Attributes are objects that are added to a Scene to describe the scene to other systems */
+UCLASS(MinimalAPI, Abstract, EditInlineNew, DisplayName="Motion Design Scene Attribute")
+class UAvaSceneAttribute : public UObject
+{
+	GENERATED_BODY()
+};
diff --git a/Engine/Plugins/Experimental/Avalanche/Source/Avalanche/Public/Attributes/AvaSceneTagAttribute.h b/Engine/Plugins/Experimental/Avalanche/Source/Avalanche/Public/Attributes/AvaSceneTagAttribute.h
new file mode 100644
--- /dev/null
+++ b/Engine/Plugins/Experimental/Avalanche/Source/Avalanche/Public/Attributes/AvaSceneTagAttribute.h
@@ -0,0 +1,51 @@
+// Copyright Epic Games, Inc. All Rights Reserved.
+
+#pragma once
+
+#include "AvaSceneAttribute.h"
+#include "AvaTagHandle.h"
+#include "AvaTagHandleContainer.h"
+#include "AvaSceneTagAttribute.generated.h"
+
+/** Base implementation of Scene Attributes that represent a tag in some form */
+UCLASS(MinimalAPI, Abstract)
+class UAvaSceneTagAttributeBase : public UAvaSceneAttribute
+{
+	GENERATED_BODY()
+
+public:
+	virtual bool ContainsTag(const FAvaTagHandle& InTagHandle) const
+	{
+		return false;
+	}
+};
+
+/** Scene Tag Attributes are scene attributes that hold a tag container handle */
+UCLASS(MinimalAPI, DisplayName="Tag Attribute")
+class UAvaSceneTagAttribute : public UAvaSceneTagAttributeBase
+{
+	GENERATED_BODY()
+
+public:
+	//~ Begin UAvaSceneTagAttributeBase
+	virtual bool ContainsTag(const FAvaTagHandle& InTagHandle) const override;
+	//~ End UAvaSceneTagAttributeBase
+
+	UPROPERTY(EditAnywhere, Category="Attributes")
+	FAvaTagHandle Tag;
+};
+
+/** Scene Tag Attributes are scene attributes that hold a tag container handle */
+UCLASS(MinimalAPI, DisplayName="Tag Container Attribute")
+class UAvaSceneTagContainerAttribute : public UAvaSceneTagAttributeBase
+{
+	GENERATED_BODY()
+
+public:
+	//~ Begin UAvaSceneTagAttributeBase
+	virtual bool ContainsTag(const FAvaTagHandle& InTagHandle) const override;
+	//~ End UAvaSceneTagAttributeBase
+
+	UPROPERTY(EditAnywhere, Category="Attributes")
+	FAvaTagHandleContainer TagContainer;
+};
diff --git a/Engine/Plugins/Experimental/Avalanche/Source/Avalanche/Public/AvaSceneSettings.h b/Engine/Plugins/Experimental/Avalanche/Source/Avalanche/Public/AvaSceneSettings.h
--- a/Engine/Plugins/Experimental/Avalanche/Source/Avalanche/Public/AvaSceneSettings.h
+++ b/Engine/Plugins/Experimental/Avalanche/Source/Avalanche/Public/AvaSceneSettings.h
@@ -2,7 +2,8 @@
 
 #pragma once
 
-#include "AvaTagHandleContainer.h"
+#include "Attributes/AvaSceneAttribute.h"
+#include "Containers/Array.h"
 #include "UObject/Object.h"
 #include "AvaSceneSettings.generated.h"
 
@@ -13,16 +14,26 @@ class UAvaSceneSettings : public UObject
 	GENERATED_BODY()
 
 public:
-	const FAvaTagHandleContainer& GetTagAttributes() const
+	/**
+	 * Iterate each valid Scene Attribute of the given Type
+	 * The callable should return true to continue iteration, and false to stop it
+	 */
+	template<typename InAttributeType>
+	void ForEachSceneAttributeOfType(TFunctionRef<bool(const InAttributeType&)> InCallable) const
 	{
-		return TagAttributes;
+		for (const UAvaSceneAttribute* SceneAttribute : SceneAttributes)
+		{
+			if (const InAttributeType* CastedAttribute = Cast<InAttributeType>(SceneAttribute))
+			{
+				if (!InCallable(*CastedAttribute))
+				{
+					break;
+				}
+			}
+		}
 	}
 
-	UFUNCTION()
-	void SetTagAttributes(const FAvaTagHandleContainer& InTagAttributes);
-
 private:
-	/** Tags describing the Scene */
-	UPROPERTY(EditAnywhere, Setter="SetTagAttributes", Category="Scene", meta=(AllowPrivateAccess="true"))
-	FAvaTagHandleContainer TagAttributes;
+	UPROPERTY(EditAnywhere, Instanced, Category="Attributes",  meta=(ShowOnlyInnerProperties))
+	TArray<TObjectPtr<UAvaSceneAttribute>> SceneAttributes;
 };
diff --git a/Engine/Plugins/Experimental/Avalanche/Source/Avalanche/Public/AvaSceneState.h b/Engine/Plugins/Experimental/Avalanche/Source/Avalanche/Public/AvaSceneState.h
--- a/Engine/Plugins/Experimental/Avalanche/Source/Avalanche/Public/AvaSceneState.h
+++ b/Engine/Plugins/Experimental/Avalanche/Source/Avalanche/Public/AvaSceneState.h
@@ -25,7 +25,7 @@ public:
 
 private:
 	/** Weak pointer to the same scene's Settings. Used to query for default tag attributes of the scene */
-	TWeakObjectPtr<UAvaSceneSettings> SceneSettingsWeak;
+	TWeakObjectPtr<const UAvaSceneSettings> SceneSettingsWeak;
 
 	/** Active Tag Attributes that were added to the scene, separate from Scene Settings */
 	TSet<FAvaTag> ActiveTagAttributes;
