Motion Design:
- Moved base Attributes to their own runtime/editor module pair
- Added a UI Customization for Attributes that aims to greatly reduce the nesting in Details View.
- By default any property under the "Attributes" category will be inlined to avoid further nesting through the category

#changelist validated
#virtualized

diff --git a/Engine/Plugins/Experimental/Avalanche/Avalanche.uplugin b/Engine/Plugins/Experimental/Avalanche/Avalanche.uplugin
--- a/Engine/Plugins/Experimental/Avalanche/Avalanche.uplugin
+++ b/Engine/Plugins/Experimental/Avalanche/Avalanche.uplugin
@@ -198,6 +198,16 @@
 			"Name": "AvalancheTagEditor",
 			"Type": "Editor",
 			"LoadingPhase": "Default"
+		},
+		{
+			"Name": "AvalancheAttribute",
+			"Type": "Runtime",
+			"LoadingPhase": "Default"
+		},
+		{
+			"Name": "AvalancheAttributeEditor",
+			"Type": "Editor",
+			"LoadingPhase": "Default"
 		}
 	],
 	"Plugins": [
@@ -302,4 +312,4 @@
 			"Enabled": true
 		}
 	]
-}
+}
\ No newline at end of file
diff --git a/Engine/Plugins/Experimental/Avalanche/Source/Avalanche/Avalanche.Build.cs b/Engine/Plugins/Experimental/Avalanche/Source/Avalanche/Avalanche.Build.cs
--- a/Engine/Plugins/Experimental/Avalanche/Source/Avalanche/Avalanche.Build.cs
+++ b/Engine/Plugins/Experimental/Avalanche/Source/Avalanche/Avalanche.Build.cs
@@ -11,6 +11,7 @@ public class Avalanche : ModuleRules
 		PublicDependencyModuleNames.AddRange(
 			new string[]
 			{
+				"AvalancheAttribute",
 				"AvalancheCore",
 				"AvalancheSceneTree",
 				"AvalancheSequence",
@@ -24,7 +25,7 @@ public class Avalanche : ModuleRules
 				"GeometryScriptingCore",
 				"Niagara",
 				"RawMesh",
-				"Text3D",
+				"Text3D", 
 			}
 		);
 
diff --git a/Engine/Plugins/Experimental/Avalanche/Source/Avalanche/Private/Attributes/AvaSceneTagAttribute.cpp b/Engine/Plugins/Experimental/Avalanche/Source/Avalanche/Private/Attributes/AvaSceneTagAttribute.cpp
deleted file mode 100644
diff --git a/Engine/Plugins/Experimental/Avalanche/Source/Avalanche/Private/AvaSceneState.cpp b/Engine/Plugins/Experimental/Avalanche/Source/Avalanche/Private/AvaSceneState.cpp
--- a/Engine/Plugins/Experimental/Avalanche/Source/Avalanche/Private/AvaSceneState.cpp
+++ b/Engine/Plugins/Experimental/Avalanche/Source/Avalanche/Private/AvaSceneState.cpp
@@ -1,9 +1,9 @@
 // Copyright Epic Games, Inc. All Rights Reserved.
 
 #include "AvaSceneState.h"
-#include "Attributes/AvaSceneTagAttribute.h"
 #include "AvaSceneSettings.h"
 #include "AvaTagHandle.h"
+#include "Tags/AvaTagAttribute.h"
 
 void UAvaSceneState::SetSceneSettings(UAvaSceneSettings* InSceneSettings)
 {
@@ -51,8 +51,8 @@ bool UAvaSceneState::ContainsTagAttribute(const FAvaTagHandle& InTagHandle) cons
 
 	bool bFoundTag = false;
 
-	SceneSettings->ForEachSceneAttributeOfType<UAvaSceneTagAttributeBase>(
-		[&InTagHandle, &bFoundTag](const UAvaSceneTagAttributeBase& InTagAttribute)
+	SceneSettings->ForEachSceneAttributeOfType<UAvaTagAttributeBase>(
+		[&InTagHandle, &bFoundTag](const UAvaTagAttributeBase& InTagAttribute)
 		{
 			bFoundTag = InTagAttribute.ContainsTag(InTagHandle);
 			// continue iteration if tag not found
diff --git a/Engine/Plugins/Experimental/Avalanche/Source/Avalanche/Public/Attributes/AvaSceneAttribute.h b/Engine/Plugins/Experimental/Avalanche/Source/Avalanche/Public/Attributes/AvaSceneAttribute.h
deleted file mode 100644
diff --git a/Engine/Plugins/Experimental/Avalanche/Source/Avalanche/Public/Attributes/AvaSceneTagAttribute.h b/Engine/Plugins/Experimental/Avalanche/Source/Avalanche/Public/Attributes/AvaSceneTagAttribute.h
deleted file mode 100644
diff --git a/Engine/Plugins/Experimental/Avalanche/Source/Avalanche/Public/AvaSceneSettings.h b/Engine/Plugins/Experimental/Avalanche/Source/Avalanche/Public/AvaSceneSettings.h
--- a/Engine/Plugins/Experimental/Avalanche/Source/Avalanche/Public/AvaSceneSettings.h
+++ b/Engine/Plugins/Experimental/Avalanche/Source/Avalanche/Public/AvaSceneSettings.h
@@ -2,7 +2,7 @@
 
 #pragma once
 
-#include "Attributes/AvaSceneAttribute.h"
+#include "AvaAttribute.h"
 #include "Containers/Array.h"
 #include "UObject/Object.h"
 #include "AvaSceneSettings.generated.h"
@@ -14,6 +14,11 @@ class UAvaSceneSettings : public UObject
 	GENERATED_BODY()
 
 public:
+	static FName GetSceneAttributesName()
+	{
+		return GET_MEMBER_NAME_CHECKED(UAvaSceneSettings, SceneAttributes);
+	}
+
 	/**
 	 * Iterate each valid Scene Attribute of the given Type
 	 * The callable should return true to continue iteration, and false to stop it
@@ -21,7 +26,7 @@ public:
 	template<typename InAttributeType>
 	void ForEachSceneAttributeOfType(TFunctionRef<bool(const InAttributeType&)> InCallable) const
 	{
-		for (const UAvaSceneAttribute* SceneAttribute : SceneAttributes)
+		for (const UAvaAttribute* SceneAttribute : SceneAttributes)
 		{
 			if (const InAttributeType* CastedAttribute = Cast<InAttributeType>(SceneAttribute))
 			{
@@ -34,6 +39,6 @@ public:
 	}
 
 private:
-	UPROPERTY(EditAnywhere, Instanced, Category="Attributes",  meta=(ShowOnlyInnerProperties))
-	TArray<TObjectPtr<UAvaSceneAttribute>> SceneAttributes;
+	UPROPERTY(EditAnywhere, Instanced, Category="Scene Attributes")
+	TArray<TObjectPtr<UAvaAttribute>> SceneAttributes;
 };
diff --git a/Engine/Plugins/Experimental/Avalanche/Source/AvalancheAttribute/AvalancheAttribute.Build.cs b/Engine/Plugins/Experimental/Avalanche/Source/AvalancheAttribute/AvalancheAttribute.Build.cs
new file mode 100644
--- /dev/null
+++ b/Engine/Plugins/Experimental/Avalanche/Source/AvalancheAttribute/AvalancheAttribute.Build.cs
@@ -0,0 +1,29 @@
+// Copyright Epic Games, Inc. All Rights Reserved.
+
+using UnrealBuildTool;
+
+public class AvalancheAttribute : ModuleRules
+{
+    public AvalancheAttribute(ReadOnlyTargetRules Target) : base(Target)
+    {
+        PCHUsage = ModuleRules.PCHUsageMode.UseExplicitOrSharedPCHs;
+
+        PublicDependencyModuleNames.AddRange(
+            new string[]
+            {
+                "AvalancheTag",
+                "Core",
+                "CoreUObject",
+            }
+        );
+
+        PrivateDependencyModuleNames.AddRange(
+            new string[]
+            {
+                "Engine",
+                "Slate",
+                "SlateCore",
+            }
+        );
+    }
+}
diff --git a/Engine/Plugins/Experimental/Avalanche/Source/AvalancheAttribute/Private/AvaAttributeModule.cpp b/Engine/Plugins/Experimental/Avalanche/Source/AvalancheAttribute/Private/AvaAttributeModule.cpp
new file mode 100644
--- /dev/null
+++ b/Engine/Plugins/Experimental/Avalanche/Source/AvalancheAttribute/Private/AvaAttributeModule.cpp
@@ -0,0 +1,5 @@
+// Copyright Epic Games, Inc. All Rights Reserved.
+
+#include "Modules/ModuleManager.h"
+
+IMPLEMENT_MODULE(FDefaultModuleImpl, AvalancheAttribute)
diff --git a/Engine/Plugins/Experimental/Avalanche/Source/AvalancheAttribute/Private/Tags/AvaTagAttribute.cpp b/Engine/Plugins/Experimental/Avalanche/Source/AvalancheAttribute/Private/Tags/AvaTagAttribute.cpp
new file mode 100644
--- /dev/null
+++ b/Engine/Plugins/Experimental/Avalanche/Source/AvalancheAttribute/Private/Tags/AvaTagAttribute.cpp
@@ -0,0 +1,17 @@
+// Copyright Epic Games, Inc. All Rights Reserved.
+
+#include "Tags/AvaTagAttribute.h"
+
+#define LOCTEXT_NAMESPACE "AvaTagAttribute"
+
+FText UAvaTagAttribute::GetDisplayName() const
+{
+	return FText::Format(LOCTEXT("DisplayName", "Tag Attribute: {0}"), FText::FromName(Tag.ToName()));
+}
+
+bool UAvaTagAttribute::ContainsTag(const FAvaTagHandle& InTagHandle) const
+{
+	return Tag.MatchesTag(InTagHandle);
+}
+
+#undef LOCTEXT_NAMESPACE
diff --git a/Engine/Plugins/Experimental/Avalanche/Source/AvalancheAttribute/Private/Tags/AvaTagContainerAttribute.cpp b/Engine/Plugins/Experimental/Avalanche/Source/AvalancheAttribute/Private/Tags/AvaTagContainerAttribute.cpp
new file mode 100644
--- /dev/null
+++ b/Engine/Plugins/Experimental/Avalanche/Source/AvalancheAttribute/Private/Tags/AvaTagContainerAttribute.cpp
@@ -0,0 +1,17 @@
+// Copyright Epic Games, Inc. All Rights Reserved.
+
+#include "Tags/AvaTagContainerAttribute.h"
+
+#define LOCTEXT_NAMESPACE "AvaTagContainerAttribute"
+
+FText UAvaTagContainerAttribute::GetDisplayName() const
+{
+	return FText::Format(LOCTEXT("DisplayName", "Tag Container Attribute: {0}"), FText::FromString(TagContainer.ToString()));
+}
+
+bool UAvaTagContainerAttribute::ContainsTag(const FAvaTagHandle& InTagHandle) const
+{
+	return TagContainer.ContainsTag(InTagHandle);
+}
+
+#undef LOCTEXT_NAMESPACE
diff --git a/Engine/Plugins/Experimental/Avalanche/Source/AvalancheAttribute/Public/AvaAttribute.h b/Engine/Plugins/Experimental/Avalanche/Source/AvalancheAttribute/Public/AvaAttribute.h
new file mode 100644
--- /dev/null
+++ b/Engine/Plugins/Experimental/Avalanche/Source/AvalancheAttribute/Public/AvaAttribute.h
@@ -0,0 +1,20 @@
+// Copyright Epic Games, Inc. All Rights Reserved.
+
+#pragma once
+
+#include "UObject/Object.h"
+#include "AvaAttribute.generated.h"
+
+/** Attributes are objects that are added to other objects to describe said object */
+UCLASS(MinimalAPI, Abstract, EditInlineNew)
+class UAvaAttribute : public UObject
+{
+	GENERATED_BODY()
+
+public:
+	/** Gets the display name to use for the attribute */
+	virtual FText GetDisplayName() const
+	{
+		return FText::GetEmpty();
+	}
+};
diff --git a/Engine/Plugins/Experimental/Avalanche/Source/AvalancheAttribute/Public/Tags/AvaTagAttribute.h b/Engine/Plugins/Experimental/Avalanche/Source/AvalancheAttribute/Public/Tags/AvaTagAttribute.h
new file mode 100644
--- /dev/null
+++ b/Engine/Plugins/Experimental/Avalanche/Source/AvalancheAttribute/Public/Tags/AvaTagAttribute.h
@@ -0,0 +1,26 @@
+// Copyright Epic Games, Inc. All Rights Reserved.
+
+#pragma once
+
+#include "AvaTagAttributeBase.h"
+#include "AvaTagHandle.h"
+#include "AvaTagAttribute.generated.h"
+
+/** Attribute that holds a tag handle, for single tag referencing */
+UCLASS(MinimalAPI, DisplayName="Tag Attribute")
+class UAvaTagAttribute : public UAvaTagAttributeBase
+{
+	GENERATED_BODY()
+
+public:
+	//~ Begin UAvaAttribute
+	AVALANCHEATTRIBUTE_API virtual FText GetDisplayName() const override;
+	//~ End UAvaAttribute
+
+	//~ Begin UAvaTagAttributeBase
+	AVALANCHEATTRIBUTE_API virtual bool ContainsTag(const FAvaTagHandle& InTagHandle) const override;
+	//~ End UAvaTagAttributeBase
+
+	UPROPERTY(EditAnywhere, Category="Attributes")
+	FAvaTagHandle Tag;
+};
diff --git a/Engine/Plugins/Experimental/Avalanche/Source/AvalancheAttribute/Public/Tags/AvaTagAttributeBase.h b/Engine/Plugins/Experimental/Avalanche/Source/AvalancheAttribute/Public/Tags/AvaTagAttributeBase.h
new file mode 100644
--- /dev/null
+++ b/Engine/Plugins/Experimental/Avalanche/Source/AvalancheAttribute/Public/Tags/AvaTagAttributeBase.h
@@ -0,0 +1,21 @@
+// Copyright Epic Games, Inc. All Rights Reserved.
+
+#pragma once
+
+#include "AvaAttribute.h"
+#include "AvaTagAttributeBase.generated.h"
+
+struct FAvaTagHandle;
+
+/** Base implementation of an attribute that represent one or more tags in some form */
+UCLASS(MinimalAPI, Abstract)
+class UAvaTagAttributeBase : public UAvaAttribute
+{
+	GENERATED_BODY()
+
+public:
+	virtual bool ContainsTag(const FAvaTagHandle& InTagHandle) const
+	{
+		return false;
+	}
+};
diff --git a/Engine/Plugins/Experimental/Avalanche/Source/AvalancheAttribute/Public/Tags/AvaTagContainerAttribute.h b/Engine/Plugins/Experimental/Avalanche/Source/AvalancheAttribute/Public/Tags/AvaTagContainerAttribute.h
new file mode 100644
--- /dev/null
+++ b/Engine/Plugins/Experimental/Avalanche/Source/AvalancheAttribute/Public/Tags/AvaTagContainerAttribute.h
@@ -0,0 +1,26 @@
+// Copyright Epic Games, Inc. All Rights Reserved.
+
+#pragma once
+
+#include "AvaTagAttributeBase.h"
+#include "AvaTagHandleContainer.h"
+#include "AvaTagContainerAttribute.generated.h"
+
+/** Attribute that holds a tag container handle, for multi-tag referencing */
+UCLASS(MinimalAPI, DisplayName="Tag Container Attribute")
+class UAvaTagContainerAttribute : public UAvaTagAttributeBase
+{
+	GENERATED_BODY()
+
+public:
+	//~ Begin UAvaAttribute
+	AVALANCHEATTRIBUTE_API virtual FText GetDisplayName() const override;
+	//~ End UAvaAttribute
+
+	//~ Begin UAvaTagAttributeBase
+	AVALANCHEATTRIBUTE_API virtual bool ContainsTag(const FAvaTagHandle& InTagHandle) const override;
+	//~ End UAvaTagAttributeBase
+
+	UPROPERTY(EditAnywhere, Category="Attributes")
+	FAvaTagHandleContainer TagContainer;
+};
diff --git a/Engine/Plugins/Experimental/Avalanche/Source/AvalancheAttributeEditor/AvalancheAttributeEditor.Build.cs b/Engine/Plugins/Experimental/Avalanche/Source/AvalancheAttributeEditor/AvalancheAttributeEditor.Build.cs
new file mode 100644
--- /dev/null
+++ b/Engine/Plugins/Experimental/Avalanche/Source/AvalancheAttributeEditor/AvalancheAttributeEditor.Build.cs
@@ -0,0 +1,31 @@
+// Copyright Epic Games, Inc. All Rights Reserved.
+
+using UnrealBuildTool;
+
+public class AvalancheAttributeEditor : ModuleRules
+{
+    public AvalancheAttributeEditor(ReadOnlyTargetRules Target) : base(Target)
+    {
+        PCHUsage = ModuleRules.PCHUsageMode.UseExplicitOrSharedPCHs;
+
+        PublicDependencyModuleNames.AddRange(
+            new string[]
+            {
+                "Core",
+            }
+        );
+
+        PrivateDependencyModuleNames.AddRange(
+            new string[]
+            {
+                "AvalancheAttribute",
+                "CoreUObject",
+                "Engine",
+                "PropertyEditor",
+                "Slate",
+                "SlateCore", 
+                "UnrealEd",
+            }
+        );
+    }
+}
diff --git a/Engine/Plugins/Experimental/Avalanche/Source/AvalancheAttributeEditor/Private/AvaAttributeEditorModule.cpp b/Engine/Plugins/Experimental/Avalanche/Source/AvalancheAttributeEditor/Private/AvaAttributeEditorModule.cpp
new file mode 100644
--- /dev/null
+++ b/Engine/Plugins/Experimental/Avalanche/Source/AvalancheAttributeEditor/Private/AvaAttributeEditorModule.cpp
@@ -0,0 +1,52 @@
+// Copyright Epic Games, Inc. All Rights Reserved.
+
+#include "AvaAttributeEditorModule.h"
+#include "Customizations/AvaAttributeNodeBuilder.h"
+#include "DetailCategoryBuilder.h"
+#include "DetailLayoutBuilder.h"
+#include "IDetailChildrenBuilder.h"
+#include "PropertyCustomizationHelpers.h"
+#include "Widgets/Text/STextBlock.h"
+
+IMPLEMENT_MODULE(FAvaAttributeEditorModule, AvalancheAttributeEditor)
+
+void FAvaAttributeEditorModule::CustomizeAttributes(const TSharedRef<IPropertyHandle>& InAttributesHandle, IDetailLayoutBuilder& InDetailBuilder)
+{
+	InAttributesHandle->MarkHiddenByCustomization();
+
+	IDetailCategoryBuilder& AttributesCategory = InDetailBuilder.EditCategory(InAttributesHandle->GetDefaultCategoryName());
+
+	TSharedRef<SWidget> HeaderContentWidget = SNew(SHorizontalBox)
+		+ SHorizontalBox::Slot()
+		.FillWidth(1.f)
+		.VAlign(VAlign_Center)
+		[
+			SNew(STextBlock)
+			.Text(AttributesCategory.GetDisplayName())
+			.Font(IDetailLayoutBuilder::GetDetailFont())
+		]
+		+ SHorizontalBox::Slot()
+		.AutoWidth()
+		.HAlign(HAlign_Right)
+		.VAlign(VAlign_Center)
+		[
+			InAttributesHandle->CreateDefaultPropertyButtonWidgets()
+		];
+
+	AttributesCategory.HeaderContent(HeaderContentWidget, /*bWholeRowContent*/true);
+
+	TSharedRef<FDetailArrayBuilder> AttributeArrayBuilder = MakeShared<FDetailArrayBuilder>(InAttributesHandle
+		, /*GenerateHeader*/false
+		, /*DisplayResetToDefault*/false
+		, /*bDisplayElementNum*/false);
+
+	TWeakPtr<IPropertyUtilities> PropertyUtilitiesWeak = InDetailBuilder.GetPropertyUtilities();
+
+	AttributeArrayBuilder->OnGenerateArrayElementWidget(FOnGenerateArrayElementWidget::CreateLambda(
+		[PropertyUtilitiesWeak](TSharedRef<IPropertyHandle> InAttributeHandle, int32 InArrayIndex, IDetailChildrenBuilder& InChildrenBuilder)
+		{
+			InChildrenBuilder.AddCustomBuilder(MakeShared<FAvaAttributeNodeBuilder>(InAttributeHandle, PropertyUtilitiesWeak));
+		}));
+
+	AttributesCategory.AddCustomBuilder(AttributeArrayBuilder);
+}
diff --git a/Engine/Plugins/Experimental/Avalanche/Source/AvalancheAttributeEditor/Private/AvaAttributeEditorModule.h b/Engine/Plugins/Experimental/Avalanche/Source/AvalancheAttributeEditor/Private/AvaAttributeEditorModule.h
new file mode 100644
--- /dev/null
+++ b/Engine/Plugins/Experimental/Avalanche/Source/AvalancheAttributeEditor/Private/AvaAttributeEditorModule.h
@@ -0,0 +1,13 @@
+// Copyright Epic Games, Inc. All Rights Reserved.
+
+#pragma once
+
+#include "IAvaAttributeEditorModule.h"
+
+class FAvaAttributeEditorModule : public IAvaAttributeEditorModule
+{
+public:
+	//~ Begin IAvaAttributeEditorModule
+	virtual void CustomizeAttributes(const TSharedRef<IPropertyHandle>& InAttributesHandle, IDetailLayoutBuilder& InDetailBuilder) override;
+	//~ End IAvaAttributeEditorModule
+};
diff --git a/Engine/Plugins/Experimental/Avalanche/Source/AvalancheAttributeEditor/Private/Customizations/AvaAttributeNodeBuilder.cpp b/Engine/Plugins/Experimental/Avalanche/Source/AvalancheAttributeEditor/Private/Customizations/AvaAttributeNodeBuilder.cpp
new file mode 100644
--- /dev/null
+++ b/Engine/Plugins/Experimental/Avalanche/Source/AvalancheAttributeEditor/Private/Customizations/AvaAttributeNodeBuilder.cpp
@@ -0,0 +1,101 @@
+// Copyright Epic Games, Inc. All Rights Reserved.
+
+#include "AvaAttributeNodeBuilder.h"
+#include "DetailWidgetRow.h"
+#include "DragDrop/AvaArrayItemDragDropHandler.h"
+#include "IDetailChildrenBuilder.h"
+#include "ObjectEditorUtils.h"
+#include "Widgets/SAvaAttributePicker.h"
+
+FName FAvaAttributeNodeBuilder::GetName() const
+{
+	return TEXT("FAvaAttributeNodeBuilder");
+}
+
+void FAvaAttributeNodeBuilder::GenerateHeaderRowContent(FDetailWidgetRow& InNodeRow)
+{
+	TSharedRef<SWidget> WholeRowContent = SNew(SHorizontalBox)
+		+ SHorizontalBox::Slot()
+		.FillWidth(1.f)
+		.HAlign(HAlign_Fill)
+		[
+			SNew(SAvaAttributePicker, AttributeHandle)
+		]
+		+ SHorizontalBox::Slot()
+		.AutoWidth()
+		[
+			AttributeHandle->CreateDefaultPropertyButtonWidgets()
+		];
+
+	InNodeRow.WholeRowContent()
+		.HAlign(HAlign_Left)
+		[
+			SNew(SBox)
+			.MinDesiredWidth(250.f)
+			[
+				WholeRowContent
+			]
+		];
+
+	InNodeRow.DragDropHandler(MakeShared<FAvaArrayItemDragDropHandler>(AttributeHandle, WholeRowContent, PropertyUtilitiesWeak));
+}
+
+void FAvaAttributeNodeBuilder::GenerateChildContent(IDetailChildrenBuilder& InChildrenBuilder)
+{
+	GenerateChildContentRecursive(AttributeHandle, InChildrenBuilder, TEXT("Attributes"));
+}
+
+TSharedPtr<IPropertyHandle> FAvaAttributeNodeBuilder::GetPropertyHandle() const
+{
+	return AttributeHandle;
+}
+
+void FAvaAttributeNodeBuilder::GenerateChildContentRecursive(const TSharedRef<IPropertyHandle>& InParentHandle, IDetailChildrenBuilder& InChildrenBuilder, FName InDefaultCategoryName)
+{
+	uint32 NumChildren = 0;
+	InParentHandle->GetNumChildren(NumChildren);
+
+	for (uint32 ChildIndex = 0; ChildIndex < NumChildren; ++ChildIndex)
+	{
+		TSharedPtr<IPropertyHandle> ChildHandle = InParentHandle->GetChildHandle(ChildIndex);
+		if (!ChildHandle.IsValid())
+		{
+			return;
+		}
+
+		if (ChildHandle->GetProperty())
+		{
+			InChildrenBuilder.AddProperty(ChildHandle.ToSharedRef());
+			continue;
+		}
+
+		// If handle is a category handle, find its category name through its children (currently, no way to get the category name directly)
+		if (ChildHandle->IsCategoryHandle())
+		{
+			TOptional<FName> CategoryName;
+
+			uint32 NumGrandChildren = 0;
+			ChildHandle->GetNumChildren(NumGrandChildren);
+
+			for (uint32 GrandChildIndex = 0; GrandChildIndex < NumGrandChildren; ++GrandChildIndex)
+			{
+				TSharedPtr<IPropertyHandle> GrandChildHandle = ChildHandle->GetChildHandle(GrandChildIndex);
+				if (FProperty* Property = GrandChildHandle->GetProperty())
+				{
+					CategoryName = FObjectEditorUtils::GetCategoryFName(Property);
+					break;
+				}
+			}
+
+			// If a category was found, but it does not match the default category, add it to the builder
+			if (CategoryName.IsSet() && *CategoryName != InDefaultCategoryName)
+			{
+				InChildrenBuilder.AddProperty(ChildHandle.ToSharedRef());
+				continue;
+			}
+		}
+
+		// If the handle was a handle to the default category, or it's a whole other type of handle, recurse its children 
+		GenerateChildContentRecursive(ChildHandle.ToSharedRef(), InChildrenBuilder, InDefaultCategoryName);
+	}
+}
diff --git a/Engine/Plugins/Experimental/Avalanche/Source/AvalancheAttributeEditor/Private/Customizations/AvaAttributeNodeBuilder.h b/Engine/Plugins/Experimental/Avalanche/Source/AvalancheAttributeEditor/Private/Customizations/AvaAttributeNodeBuilder.h
new file mode 100644
--- /dev/null
+++ b/Engine/Plugins/Experimental/Avalanche/Source/AvalancheAttributeEditor/Private/Customizations/AvaAttributeNodeBuilder.h
@@ -0,0 +1,33 @@
+// Copyright Epic Games, Inc. All Rights Reserved.
+
+#pragma once
+
+#include "IDetailCustomNodeBuilder.h"
+#include "Templates/SharedPointer.h"
+
+class IPropertyHandle;
+class IPropertyUtilities;
+
+class FAvaAttributeNodeBuilder : public IDetailCustomNodeBuilder
+{
+public:
+	explicit FAvaAttributeNodeBuilder(const TSharedRef<IPropertyHandle>& InAttributeHandle, const TWeakPtr<IPropertyUtilities>& InPropertyUtilitiesWeak)
+		: AttributeHandle(InAttributeHandle)
+		, PropertyUtilitiesWeak(InPropertyUtilitiesWeak)
+	{
+	}
+
+private:
+	//~ Begin IDetailCustomNodeBuilder
+	virtual FName GetName() const override;
+	virtual void GenerateHeaderRowContent(FDetailWidgetRow& InNodeRow) override;
+	virtual void GenerateChildContent(IDetailChildrenBuilder& InChildrenBuilder) override;
+	virtual TSharedPtr<IPropertyHandle> GetPropertyHandle() const override;
+	//~ End IDetailCustomNodeBuilder
+
+	void GenerateChildContentRecursive(const TSharedRef<IPropertyHandle>& InParentHandle, IDetailChildrenBuilder& InChildrenBuilder, FName InDefaultCategoryName);
+
+	TSharedRef<IPropertyHandle> AttributeHandle;
+
+	TWeakPtr<IPropertyUtilities> PropertyUtilitiesWeak;
+};
diff --git a/Engine/Plugins/Experimental/Avalanche/Source/AvalancheAttributeEditor/Private/DragDrop/AvaArrayItemDragDropHandler.cpp b/Engine/Plugins/Experimental/Avalanche/Source/AvalancheAttributeEditor/Private/DragDrop/AvaArrayItemDragDropHandler.cpp
new file mode 100644
--- /dev/null
+++ b/Engine/Plugins/Experimental/Avalanche/Source/AvalancheAttributeEditor/Private/DragDrop/AvaArrayItemDragDropHandler.cpp
@@ -0,0 +1,184 @@
+// Copyright Epic Games, Inc. All Rights Reserved.
+
+#include "AvaArrayItemDragDropHandler.h"
+#include "DragAndDrop/DecoratedDragDropOp.h"
+#include "IPropertyUtilities.h"
+#include "PropertyHandle.h"
+#include "ScopedTransaction.h"
+#include "UObject/UnrealType.h"
+
+#define LOCTEXT_NAMESPACE "AvaArrayItemDragDropHandler"
+
+namespace UE::AvaEditor::Private
+{
+	int32 CalculateTargetIndex(int32 InSourceIndex, int32 InTargetIndex, EItemDropZone DropZone)
+	{
+		if (DropZone == EItemDropZone::BelowItem)
+		{
+			// If the drop zone is below, then move it to the next item's index
+			InTargetIndex++;
+		}
+
+		if (InSourceIndex < InTargetIndex)
+		{
+			// If the item is moved down the list, then all the other elements below it are shifted up one
+			InTargetIndex--;
+		}
+
+		return ensure(InTargetIndex >= 0) ? InTargetIndex : 0;
+	}
+
+	bool IsValidDrop(const TSharedPtr<IPropertyHandle>& InTargetItemHandle, const TSharedPtr<IPropertyHandle>& InDraggedItemHandle, EItemDropZone InDropZone)
+	{
+		// Can't drop onto another array item; need to drop above or below. Likewise, cannot drop above/below itself (redundant operation)
+		if (InDropZone == EItemDropZone::OntoItem || InTargetItemHandle == InDraggedItemHandle)
+		{
+			return false;
+		}
+
+		const int32 SourceIndex = InDraggedItemHandle->GetIndexInArray();
+		const int32 TargetIndex = CalculateTargetIndex(SourceIndex, InTargetItemHandle->GetArrayIndex(), InDropZone);
+
+		if (SourceIndex != TargetIndex)
+		{
+			TSharedPtr<IPropertyHandle> TargetParentHandle  = InTargetItemHandle->GetParentHandle();
+			TSharedPtr<IPropertyHandle> DraggedParentHandle = InDraggedItemHandle->GetParentHandle();
+
+			// Ensure that these two property handles belong to the same array parent property
+			return TargetParentHandle.IsValid() && DraggedParentHandle.IsValid()
+				&& TargetParentHandle->IsSamePropertyNode(DraggedParentHandle)
+				&& TargetParentHandle->AsArray().IsValid();
+		}
+
+		return false;
+	}
+}
+
+class FAvaArrayItemDragDropOp : public FDecoratedDragDropOp
+{
+public:
+	DRAG_DROP_OPERATOR_TYPE(FAvaAttributeElementDragDropOp, FDecoratedDragDropOp)
+
+	FAvaArrayItemDragDropOp(const TWeakPtr<IPropertyHandle>& InArrayItemHandleWeak)
+	{
+		ArrayItemHandleWeak = InArrayItemHandleWeak;
+		MouseCursor = EMouseCursor::GrabHandClosed;
+	}
+
+	void Init()
+	{
+		SetValidDrop(false);
+		SetupDefaults();
+		Construct();
+	}
+
+	void SetValidDrop(bool bIsValidDrop)
+	{
+		if (bIsValidDrop)
+		{
+			CurrentHoverText = LOCTEXT("PlaceRowHere", "Place Row Here");
+			CurrentIconBrush = FAppStyle::GetBrush("Graph.ConnectorFeedback.OK");
+		}
+		else
+		{
+			CurrentHoverText = LOCTEXT("CannotPlaceRowHere", "Cannot Place Row Here");
+			CurrentIconBrush = FAppStyle::GetBrush("Graph.ConnectorFeedback.Error");
+		}
+	}
+
+	TWeakPtr<IPropertyHandle> ArrayItemHandleWeak;
+};
+
+FAvaArrayItemDragDropHandler::FAvaArrayItemDragDropHandler(const TSharedRef<IPropertyHandle>& InArrayItemHandle, const TSharedRef<SWidget>& InRowWidget, const TWeakPtr<IPropertyUtilities>& InPropertyUtilities)
+	: ArrayItemHandleWeak(InArrayItemHandle)
+	, RowWidgetWeak(InRowWidget)
+	, PropertyUtilitiesWeak(InPropertyUtilities)
+{
+}
+
+TSharedPtr<FDragDropOperation> FAvaArrayItemDragDropHandler::CreateDragDropOperation() const
+{
+	return MakeShared<FAvaArrayItemDragDropOp>(ArrayItemHandleWeak);
+}
+
+TOptional<EItemDropZone> FAvaArrayItemDragDropHandler::CanAcceptDrop(const FDragDropEvent& InDragDropEvent, EItemDropZone InDropZone) const
+{
+	TSharedPtr<FAvaArrayItemDragDropOp> ArrayItemDragDropOp = InDragDropEvent.GetOperationAs<FAvaArrayItemDragDropOp>();
+	TSharedPtr<IPropertyHandle> ArrayItemHandle = ArrayItemHandleWeak.Pin();
+	TSharedPtr<SWidget> RowWidget = RowWidgetWeak.Pin();
+
+	if (!RowWidget.IsValid() || !ArrayItemHandle.IsValid() || !ArrayItemDragDropOp.IsValid())
+	{
+		return TOptional<EItemDropZone>();
+	}
+
+	TSharedPtr<IPropertyHandle> DraggedArrayItemHandle = ArrayItemDragDropOp->ArrayItemHandleWeak.Pin();
+	if (!DraggedArrayItemHandle.IsValid())
+	{
+		return TOptional<EItemDropZone>();
+	}
+
+	const FGeometry& Geometry = RowWidget->GetTickSpaceGeometry();
+	const float LocalPointerY = Geometry.AbsoluteToLocal(InDragDropEvent.GetScreenSpacePosition()).Y;
+
+	const EItemDropZone OverrideDropZone = LocalPointerY < Geometry.GetLocalSize().Y * 0.5f
+		? EItemDropZone::AboveItem
+		: EItemDropZone::BelowItem;
+
+	if (!UE::AvaEditor::Private::IsValidDrop(ArrayItemHandle, DraggedArrayItemHandle, OverrideDropZone))
+	{
+		ArrayItemDragDropOp->SetValidDrop(false);
+		return TOptional<EItemDropZone>();
+	}
+
+	ArrayItemDragDropOp->SetValidDrop(true);
+	return OverrideDropZone;
+}
+
+bool FAvaArrayItemDragDropHandler::AcceptDrop(const FDragDropEvent& InDragDropEvent, EItemDropZone InDropZone) const
+{
+	TSharedPtr<FAvaArrayItemDragDropOp> ArrayItemDragDropOp = InDragDropEvent.GetOperationAs<FAvaArrayItemDragDropOp>();
+	TSharedPtr<IPropertyHandle> ArrayItemHandle = ArrayItemHandleWeak.Pin();
+	TSharedPtr<SWidget> RowWidget = RowWidgetWeak.Pin();
+
+	if (!RowWidget.IsValid() || !ArrayItemHandle.IsValid() || !ArrayItemDragDropOp.IsValid())
+	{
+		return false;
+	}
+
+	TSharedPtr<IPropertyHandle> DraggedArrayItemHandle = ArrayItemDragDropOp->ArrayItemHandleWeak.Pin();
+	TSharedPtr<IPropertyHandle> ParentHandle = ArrayItemHandle->GetParentHandle();
+
+	if (!DraggedArrayItemHandle.IsValid() || !ParentHandle.IsValid())
+	{
+		return false;
+	}
+
+	TSharedPtr<IPropertyHandleArray> ParentArrayHandle = ParentHandle->AsArray();
+
+	if (!ParentArrayHandle.IsValid())
+	{
+		return false;
+	}
+
+	const int32 SourceIndex = DraggedArrayItemHandle->GetArrayIndex();
+	const int32 TargetIndex = UE::AvaEditor::Private::CalculateTargetIndex(SourceIndex, ArrayItemHandle->GetArrayIndex(), InDropZone);
+
+	FScopedTransaction Transaction(LOCTEXT("MoveRow", "Move Row"));
+
+	ParentHandle->NotifyPreChange();
+	ParentArrayHandle->MoveElementTo(SourceIndex, TargetIndex);
+	ParentHandle->NotifyPostChange(EPropertyChangeType::ArrayMove);
+
+	// IPropertyHandle::NotifyFinishedChangingProperties is not called as it hard codes the change type to ValueSet, when this is an Array Move
+	// Instead, PropertyUtilities function is called directly with the correct change type
+	if (TSharedPtr<IPropertyUtilities> PropertyUtilities = PropertyUtilitiesWeak.Pin())
+	{
+		FPropertyChangedEvent MoveEvent(ParentHandle->GetProperty(), EPropertyChangeType::ArrayMove);
+		PropertyUtilities->NotifyFinishedChangingProperties(MoveEvent);
+	}
+
+	return true;
+}
+
+#undef LOCTEXT_NAMESPACE
diff --git a/Engine/Plugins/Experimental/Avalanche/Source/AvalancheAttributeEditor/Private/DragDrop/AvaArrayItemDragDropHandler.h b/Engine/Plugins/Experimental/Avalanche/Source/AvalancheAttributeEditor/Private/DragDrop/AvaArrayItemDragDropHandler.h
new file mode 100644
--- /dev/null
+++ b/Engine/Plugins/Experimental/Avalanche/Source/AvalancheAttributeEditor/Private/DragDrop/AvaArrayItemDragDropHandler.h
@@ -0,0 +1,39 @@
+// Copyright Epic Games, Inc. All Rights Reserved.
+
+#pragma once
+
+#include "IDetailDragDropHandler.h"
+#include "Templates/SharedPointer.h"
+
+class IPropertyHandle;
+class IPropertyUtilities;
+class SWidget;
+
+/**
+ * Detail Drag Drop Handler for Custom Builders to re-order the array items
+ *
+ * Its existence is due to the following,
+ * 1) The Property Editor API only adds a reorder handle widget if the FProperty of GetPropertyNode is able to re-order (owner is an array property that can re-order)
+ * 2) The Array item property handle in question is being presented through a Custom Builder instead of directly as a property
+ * 3) Custom Builder customizations don't have a Property Node set (null), so the Property Editor API does not pick up the property to check with (1)
+ * 4) The only other way for the Property Editor API to add a reorder handle widget is by having a custom drag drop handler
+ * 5) There's no Array Detail Drag Drop Handler built-in; the default array reordering is handled differently (and is private) rather than through a IDetailDragDropHandler
+ */
+class FAvaArrayItemDragDropHandler : public IDetailDragDropHandler
+{
+public:
+	explicit FAvaArrayItemDragDropHandler(const TSharedRef<IPropertyHandle>& InArrayItemHandle, const TSharedRef<SWidget>& InRowWidget, const TWeakPtr<IPropertyUtilities>& InPropertyUtilities);
+
+	//~ Begin IDetailDragDropHandler
+	virtual TSharedPtr<FDragDropOperation> CreateDragDropOperation() const override;
+	virtual TOptional<EItemDropZone> CanAcceptDrop(const FDragDropEvent& InDragDropEvent, EItemDropZone InDropZone) const override;
+	virtual bool AcceptDrop(const FDragDropEvent& InDragDropEvent, EItemDropZone InDropZone) const override;
+	//~ End IDetailDragDropHandler
+
+private:
+	TWeakPtr<IPropertyHandle> ArrayItemHandleWeak;
+
+	TWeakPtr<SWidget> RowWidgetWeak;
+
+	TWeakPtr<IPropertyUtilities> PropertyUtilitiesWeak;
+};
diff --git a/Engine/Plugins/Experimental/Avalanche/Source/AvalancheAttributeEditor/Private/Widgets/SAvaAttributePicker.cpp b/Engine/Plugins/Experimental/Avalanche/Source/AvalancheAttributeEditor/Private/Widgets/SAvaAttributePicker.cpp
new file mode 100644
--- /dev/null
+++ b/Engine/Plugins/Experimental/Avalanche/Source/AvalancheAttributeEditor/Private/Widgets/SAvaAttributePicker.cpp
@@ -0,0 +1,264 @@
+// Copyright Epic Games, Inc. All Rights Reserved.
+
+#include "SAvaAttributePicker.h"
+#include "AvaAttribute.h"
+#include "ClassViewerFilter.h"
+#include "ClassViewerModule.h"
+#include "DetailLayoutBuilder.h"
+#include "Modules/ModuleManager.h"
+#include "ScopedTransaction.h"
+#include "Styling/SlateIconFinder.h"
+#include "UObject/ConstructorHelpers.h"
+#include "Widgets/Images/SImage.h"
+#include "Widgets/Input/SComboButton.h"
+#include "Widgets/SBoxPanel.h"
+#include "Widgets/Text/STextBlock.h"
+
+#define LOCTEXT_NAMESPACE "SAvaAttributePicker"
+
+class FAvaAttributeClassFilter : public IClassViewerFilter
+{
+public:
+	static constexpr EClassFlags DisallowedClassFlags = CLASS_Deprecated | CLASS_Abstract;
+
+	//~ Begin IClassViewerFilter
+	virtual bool IsClassAllowed(const FClassViewerInitializationOptions&, const UClass* InClass, TSharedRef<FClassViewerFilterFuncs>) override
+	{
+		return InClass
+			&& InClass->IsChildOf<UAvaAttribute>()
+			&& !InClass->HasAnyClassFlags(DisallowedClassFlags);
+	}
+	virtual bool IsUnloadedClassAllowed(const FClassViewerInitializationOptions&, const TSharedRef<const IUnloadedBlueprintData> InUnloadedClassData, TSharedRef<FClassViewerFilterFuncs>) override
+	{
+		return InUnloadedClassData->IsChildOf(UAvaAttribute::StaticClass())
+			&& !InUnloadedClassData->HasAnyClassFlags(DisallowedClassFlags);
+	}
+	//~ End IClassViewerFilter
+};
+
+void SAvaAttributePicker::Construct(const FArguments& InArgs, const TSharedRef<IPropertyHandle>& InAttributeHandle)
+{
+	AttributeHandle = InAttributeHandle;
+
+	int32 ArrayIndex = InAttributeHandle->GetArrayIndex();
+
+	ChildSlot
+	[
+		SNew(SHorizontalBox)
+		+ SHorizontalBox::Slot()
+		.AutoWidth()
+		.VAlign(VAlign_Center)
+		[
+			SNew(STextBlock)
+			.Text(INVTEXT("["))
+			.ColorAndOpacity(FSlateColor::UseSubduedForeground())
+		]
+		+ SHorizontalBox::Slot()
+		.Padding(3.f, 0.f, 3.f, 0.f)
+		.AutoWidth()
+		.VAlign(VAlign_Center)
+		[
+			SNew(STextBlock)
+			.Text(FText::AsNumber(ArrayIndex))
+			.Font(IDetailLayoutBuilder::GetDetailFont())
+		]
+		+ SHorizontalBox::Slot()
+		.AutoWidth()
+		.VAlign(VAlign_Center)
+		[
+			SNew(STextBlock)
+			.Text(INVTEXT("]"))
+			.ColorAndOpacity(FSlateColor::UseSubduedForeground())
+		]
+		+ SHorizontalBox::Slot()
+		.Padding(5.f, 0.f, 0.f, 0.f)
+		.FillWidth(1.f)
+		.VAlign(VAlign_Center)
+		[
+			SAssignNew(PickerComboButton, SComboButton)
+			.IsEnabled(this, &SAvaAttributePicker::IsPropertyEnabled)
+			.ContentPadding(1)
+			.OnGetMenuContent(this, &SAvaAttributePicker::GenerateClassPicker)
+			.ButtonContent()
+			[
+				SNew(SHorizontalBox)
+				+ SHorizontalBox::Slot()
+				.AutoWidth()
+				.VAlign(VAlign_Center)
+				[
+					SNew(SImage)
+					.Image(this, &SAvaAttributePicker::GetIcon)
+				]
+				+ SHorizontalBox::Slot()
+				.Padding(4.f, 0.f, 0.f, 0.f)
+				.FillWidth(1.f)
+				.VAlign(VAlign_Center)
+				[
+					SNew(STextBlock)
+					.Text(this, &SAvaAttributePicker::GetDisplayName)
+					.Font(IDetailLayoutBuilder::GetDetailFont())
+				]
+			]
+		]
+	];
+}
+
+bool SAvaAttributePicker::IsPropertyEnabled() const
+{
+	return !AttributeHandle->IsEditConst();
+}
+
+UAvaAttribute* SAvaAttributePicker::GetAttribute() const
+{
+	UObject* Object = nullptr;
+	if (AttributeHandle->GetValue(Object) == FPropertyAccess::Success)
+	{
+		return Cast<UAvaAttribute>(Object);
+	}
+	return nullptr;
+}
+
+FText SAvaAttributePicker::GetDisplayName() const
+{
+	if (UAvaAttribute* Attribute = GetAttribute())
+	{
+		FText DisplayName = Attribute->GetDisplayName();
+		if (!DisplayName.IsEmpty())
+		{
+			return DisplayName;
+		}
+		return Attribute->GetClass()->GetDisplayNameText();
+	}
+	return FText::GetEmpty();
+}
+
+const FSlateBrush* SAvaAttributePicker::GetIcon() const
+{
+	if (UAvaAttribute* Attribute = GetAttribute())
+	{
+		return FSlateIconFinder::FindIconBrushForClass(Attribute->GetClass());
+	}
+	return nullptr;
+}
+
+TSharedRef<SWidget> SAvaAttributePicker::GenerateClassPicker()
+{
+	FClassViewerInitializationOptions Options;
+	Options.PropertyHandle = AttributeHandle;
+	Options.NameTypeToDisplay = EClassViewerNameTypeToDisplay::DisplayName;
+	Options.bShowBackgroundBorder = false;
+	Options.bShowUnloadedBlueprints = true;
+	Options.ClassFilters.Add(MakeShared<FAvaAttributeClassFilter>());
+	Options.bShowNoneOption = true;
+
+	FOnClassPicked OnClassPicked(FOnClassPicked::CreateSP(this, &SAvaAttributePicker::OnClassPicked));
+	return FModuleManager::LoadModuleChecked<FClassViewerModule>("ClassViewer").CreateClassViewer(Options, OnClassPicked);
+}
+
+void SAvaAttributePicker::OnClassPicked(UClass* InClassPicked)
+{
+	// Class can be null, but never a class that isn't an attribute
+	if (!ensure(!InClassPicked || InClassPicked->IsChildOf<UAvaAttribute>()))
+	{
+		return;
+	}
+
+	// Gather current values to cleanup later (named previous to clearly differentiate with the new values after they've been set)
+	TArray<FString> PreviousValues;
+	if (FObjectProperty* ObjectProperty = CastFieldChecked<FObjectProperty>(AttributeHandle->GetProperty()))
+	{
+		if (ObjectProperty->HasAnyPropertyFlags(CPF_InstancedReference))
+		{
+			AttributeHandle->GetPerObjectValues(PreviousValues);
+		}
+	}
+
+	TArray<FString> NewValues = GenerateNewValues(PreviousValues, InClassPicked);
+	ensure(PreviousValues.IsEmpty() || PreviousValues.Num() == NewValues.Num());
+
+	// SetPerObjectValues already transacts, but old subobject renames must also happen within the same transaction,
+	// so create one covering this entire scope
+	FScopedTransaction Transaction(LOCTEXT("SetAttributeClass", "Set Attribute Class"));
+
+	AttributeHandle->SetPerObjectValues(NewValues);
+
+	for (int32 ValueIndex = 0; ValueIndex < PreviousValues.Num(); ++ValueIndex)
+	{
+		FString& PreviousValue = PreviousValues[ValueIndex];
+		if (!PreviousValue.IsEmpty() && PreviousValue != NewValues[ValueIndex])
+		{
+			// Move the old subobject to the transient package so GetObjectsWithOuter will not return it
+			ConstructorHelpers::StripObjectClass(PreviousValue);
+			if (UObject* SubObject = StaticFindObject(UObject::StaticClass(), nullptr, *PreviousValue))
+			{
+				SubObject->Rename(nullptr, GetTransientOuterForRename(SubObject->GetClass()), REN_DontCreateRedirectors);
+			}
+		}
+	}
+
+	AttributeHandle->RequestRebuildChildren();
+	PickerComboButton->SetIsOpen(false);
+}
+
+TArray<FString> SAvaAttributePicker::GenerateNewValues(TConstArrayView<FString> InCurrentValues, UClass* InClassPicked) const
+{
+	const FString NoneString = FName(NAME_None).ToString();
+
+	TArray<FString> NewValues;
+
+	if (!InClassPicked)
+	{
+		int32 PerObjectValueCount = AttributeHandle->GetNumPerObjectValues();
+		NewValues.Init(NoneString, PerObjectValueCount);
+		return NewValues;
+	}
+
+	const FString ClassPath = InClassPicked->GetClassPathName().ToString();
+
+	TArray<UObject*> Outers;
+	AttributeHandle->GetOuterObjects(Outers);
+
+	NewValues.Reserve(Outers.Num());
+
+	for (int32 ObjectIndex = 0; ObjectIndex < Outers.Num(); ++ObjectIndex)
+	{
+		FString& NewValue = NewValues.Emplace_GetRef();
+		const FString* CurrentValue = nullptr;
+
+		if (InCurrentValues.IsValidIndex(ObjectIndex))
+		{
+			CurrentValue = &InCurrentValues[ObjectIndex];
+		}
+
+		FStringView CurrentClassName;
+		int32 ClassEndIndex;
+		if (CurrentValue && *CurrentValue != NoneString && CurrentValue->FindChar(TEXT('\''), ClassEndIndex))
+		{
+			CurrentClassName = FStringView(*CurrentValue).Left(ClassEndIndex);
+		}
+
+		// Re-use old object if the same class was picked
+		if (CurrentValue &&  ClassPath == CurrentClassName)
+		{
+			NewValue = *CurrentValue;
+		}
+		else
+		{
+			UObject* Outer = Outers[ObjectIndex];
+			EObjectFlags MaskedOuterFlags = Outer->GetMaskedFlags(RF_PropagateToSubObjects);
+
+			UObject* InstancedObject = NewObject<UObject>(Outer, InClassPicked, NAME_None, MaskedOuterFlags);
+
+			// Wrap the value in quotes before setting - in some cases for editinline-instanced values, the outer object path
+			// can potentially contain a space token (e.g. if the outer object was instanced as a Blueprint template object
+			// based on a user-facing variable name). While technically such characters should not be allowed, historically there
+			// has been no issue with most tokens in the INVALID_OBJECTNAME_CHARACTERS set being present in in-memory object names,
+			// other than some systems failing to resolve the object's path if the string representation of the value is not quoted.
+			NewValue = FString::Printf(TEXT("\"%s\""), *InstancedObject->GetPathName().ReplaceQuotesWithEscapedQuotes());
+		}
+	}
+
+	return NewValues;
+}
+
+#undef LOCTEXT_NAMESPACE
diff --git a/Engine/Plugins/Experimental/Avalanche/Source/AvalancheAttributeEditor/Private/Widgets/SAvaAttributePicker.h b/Engine/Plugins/Experimental/Avalanche/Source/AvalancheAttributeEditor/Private/Widgets/SAvaAttributePicker.h
new file mode 100644
--- /dev/null
+++ b/Engine/Plugins/Experimental/Avalanche/Source/AvalancheAttributeEditor/Private/Widgets/SAvaAttributePicker.h
@@ -0,0 +1,43 @@
+// Copyright Epic Games, Inc. All Rights Reserved.
+
+#pragma once
+
+#include "Containers/ArrayView.h"
+#include "Containers/ContainersFwd.h"
+#include "Templates/SharedPointer.h"
+#include "Widgets/SCompoundWidget.h"
+
+class FText;
+class IPropertyHandle;
+class SComboButton;
+class UAvaAttribute;
+class UClass;
+struct FSlateBrush;
+
+class SAvaAttributePicker : public SCompoundWidget
+{
+public:
+	SLATE_BEGIN_ARGS(SAvaAttributePicker) {}
+	SLATE_END_ARGS()
+
+	void Construct(const FArguments& InArgs, const TSharedRef<IPropertyHandle>& InAttributeHandle);
+
+private:
+	bool IsPropertyEnabled() const;
+
+	UAvaAttribute* GetAttribute() const;
+
+	FText GetDisplayName() const;
+
+	const FSlateBrush* GetIcon() const;
+
+	TSharedRef<SWidget> GenerateClassPicker();
+
+	void OnClassPicked(UClass* InClassPicked);
+
+	TArray<FString> GenerateNewValues(TConstArrayView<FString> InCurrentValues, UClass* InClassPicked) const;
+
+	TSharedPtr<IPropertyHandle> AttributeHandle;
+
+	TSharedPtr<SComboButton> PickerComboButton;
+};
diff --git a/Engine/Plugins/Experimental/Avalanche/Source/AvalancheAttributeEditor/Public/IAvaAttributeEditorModule.h b/Engine/Plugins/Experimental/Avalanche/Source/AvalancheAttributeEditor/Public/IAvaAttributeEditorModule.h
new file mode 100644
--- /dev/null
+++ b/Engine/Plugins/Experimental/Avalanche/Source/AvalancheAttributeEditor/Public/IAvaAttributeEditorModule.h
@@ -0,0 +1,27 @@
+// Copyright Epic Games, Inc. All Rights Reserved.
+
+#pragma once
+
+#include "Modules/ModuleManager.h"
+#include "Templates/SharedPointerFwd.h"
+
+class IDetailLayoutBuilder;
+class IPropertyHandle;
+
+class IAvaAttributeEditorModule : public IModuleInterface
+{
+	static constexpr const TCHAR* ModuleName = TEXT("AvalancheAttributeEditor");
+
+public:
+	static bool IsLoaded()
+	{
+		return FModuleManager::Get().IsModuleLoaded(ModuleName);
+	}
+
+	static IAvaAttributeEditorModule& Get()
+	{
+		return FModuleManager::Get().LoadModuleChecked<IAvaAttributeEditorModule>(ModuleName);
+	}
+
+	virtual void CustomizeAttributes(const TSharedRef<IPropertyHandle>& InAttributesHandle, IDetailLayoutBuilder& InDetailBuilder) = 0;
+};
diff --git a/Engine/Plugins/Experimental/Avalanche/Source/AvalancheEditor/AvalancheEditor.Build.cs b/Engine/Plugins/Experimental/Avalanche/Source/AvalancheEditor/AvalancheEditor.Build.cs
--- a/Engine/Plugins/Experimental/Avalanche/Source/AvalancheEditor/AvalancheEditor.Build.cs
+++ b/Engine/Plugins/Experimental/Avalanche/Source/AvalancheEditor/AvalancheEditor.Build.cs
@@ -29,6 +29,7 @@ public class AvalancheEditor : ModuleRules
 				"ApplicationCore",
 				"AssetRegistry",
 				"AssetTools",
+				"AvalancheAttributeEditor",
 				"AvalancheComponentVisualizers",
 				"AvalancheCore",
 				"AvalancheEditorCore",
@@ -110,7 +111,7 @@ public class AvalancheEditor : ModuleRules
 				"UMGEditor",
 				"UnrealEd",
 				"XmlParser",
-				"zlib"
+				"zlib",
 			}
 		);
 
diff --git a/Engine/Plugins/Experimental/Avalanche/Source/AvalancheEditor/Private/AvaEditorModule.cpp b/Engine/Plugins/Experimental/Avalanche/Source/AvalancheEditor/Private/AvaEditorModule.cpp
--- a/Engine/Plugins/Experimental/Avalanche/Source/AvalancheEditor/Private/AvaEditorModule.cpp
+++ b/Engine/Plugins/Experimental/Avalanche/Source/AvalancheEditor/Private/AvaEditorModule.cpp
@@ -28,12 +28,14 @@
 
 // Details View
 #include "AssetToolsModule.h"
-#include "DetailsPanel/DMMaterialInterfaceTypeCustomizer.h"
+#include "AvaSceneSettings.h"
 #include "DetailView/Customizations/AvaAnchorAlignmentPropertyTypeCustomization.h"
 #include "DetailView/Customizations/AvaCategoryHiderCustomization.h"
 #include "DetailView/Customizations/AvaMeshesDetailCustomization.h"
+#include "DetailView/Customizations/AvaSceneSettingsCustomization.h"
 #include "DetailView/Customizations/AvaVectorPropertyTypeCustomization.h"
 #include "DetailView/Customizations/AvaViewportQualitySettingsPropertyTypeCustomization.h"
+#include "DetailsPanel/DMMaterialInterfaceTypeCustomizer.h"
 
 DEFINE_LOG_CATEGORY(AvaLog);
 
@@ -258,6 +260,9 @@ void FAvaEditorModule::RegisterCustomLayouts()
 	PropertyModule.RegisterCustomClassLayout(UAvaShapeDynamicMeshBase::StaticClass()->GetFName(), FOnGetDetailCustomizationInstance::CreateStatic(
 		&FAvaMeshesDetailCustomization::MakeInstance));
 
+	PropertyModule.RegisterCustomClassLayout(UAvaSceneSettings::StaticClass()->GetFName(), FOnGetDetailCustomizationInstance::CreateStatic(
+		&FAvaSceneSettingsCustomization::MakeInstance));
+
 	const TSharedRef<FDMMaterialInterfaceTypeIdentifier> MaterialPropertyTypeIdentifier = MakeShared<FDMMaterialInterfaceTypeIdentifier>();
 	
 	PropertyModule.RegisterCustomPropertyTypeLayout(UMaterialInterface::StaticClass()->GetFName(), FOnGetPropertyTypeCustomizationInstance::CreateStatic(
@@ -289,6 +294,8 @@ void FAvaEditorModule::UnregisterCustomLayouts()
 		PropertyModule.UnregisterCustomClassLayout(AAvaShapeActor::StaticClass()->GetFName());
 		PropertyModule.UnregisterCustomClassLayout(UAvaShapeDynamicMeshBase::StaticClass()->GetFName());
 
+		PropertyModule.UnregisterCustomClassLayout(UAvaSceneSettings::StaticClass()->GetFName());
+
 		PropertyModule.UnregisterCustomPropertyTypeLayout(UMaterialInterface::StaticClass()->GetFName());
 
 		PropertyModule.UnregisterCustomPropertyTypeLayout(TEXT("Vector"));
diff --git a/Engine/Plugins/Experimental/Avalanche/Source/AvalancheEditor/Private/DetailView/Customizations/AvaSceneSettingsCustomization.cpp b/Engine/Plugins/Experimental/Avalanche/Source/AvalancheEditor/Private/DetailView/Customizations/AvaSceneSettingsCustomization.cpp
new file mode 100644
--- /dev/null
+++ b/Engine/Plugins/Experimental/Avalanche/Source/AvalancheEditor/Private/DetailView/Customizations/AvaSceneSettingsCustomization.cpp
@@ -0,0 +1,19 @@
+// Copyright Epic Games, Inc. All Rights Reserved.
+
+#include "AvaSceneSettingsCustomization.h"
+#include "AvaSceneSettings.h"
+#include "DetailLayoutBuilder.h"
+#include "IAvaAttributeEditorModule.h"
+
+TSharedRef<IDetailCustomization> FAvaSceneSettingsCustomization::MakeInstance()
+{
+	return MakeShared<FAvaSceneSettingsCustomization>();
+}
+
+void FAvaSceneSettingsCustomization::CustomizeDetails(IDetailLayoutBuilder& InDetailBuilder)
+{
+	if (TSharedPtr<IPropertyHandle> AttributesHandle = InDetailBuilder.GetProperty(UAvaSceneSettings::GetSceneAttributesName()))
+	{
+		IAvaAttributeEditorModule::Get().CustomizeAttributes(AttributesHandle.ToSharedRef(), InDetailBuilder);	
+	}
+}
diff --git a/Engine/Plugins/Experimental/Avalanche/Source/AvalancheEditor/Private/DetailView/Customizations/AvaSceneSettingsCustomization.h b/Engine/Plugins/Experimental/Avalanche/Source/AvalancheEditor/Private/DetailView/Customizations/AvaSceneSettingsCustomization.h
new file mode 100644
--- /dev/null
+++ b/Engine/Plugins/Experimental/Avalanche/Source/AvalancheEditor/Private/DetailView/Customizations/AvaSceneSettingsCustomization.h
@@ -0,0 +1,16 @@
+// Copyright Epic Games, Inc. All Rights Reserved.
+
+#pragma once
+
+#include "IDetailCustomization.h"
+
+class FAvaSceneSettingsCustomization : public IDetailCustomization
+{
+public:
+	static TSharedRef<IDetailCustomization> MakeInstance();
+
+private:
+	//~ Begin IDetailCustomization
+	virtual void CustomizeDetails(IDetailLayoutBuilder& InDetailBuilder) override;
+	//~ End IDetailCustomization
+};
