Motion Design
- Exposure Compensation set back to 0 from 1 in the Ava Post Process Volume.

diff --git a/Engine/Plugins/Experimental/Avalanche/Source/Avalanche/Public/Viewport/AvaPostProcessVolume.h b/Engine/Plugins/Experimental/Avalanche/Source/Avalanche/Public/Viewport/AvaPostProcessVolume.h
--- a/Engine/Plugins/Experimental/Avalanche/Source/Avalanche/Public/Viewport/AvaPostProcessVolume.h
+++ b/Engine/Plugins/Experimental/Avalanche/Source/Avalanche/Public/Viewport/AvaPostProcessVolume.h
@@ -22,7 +22,7 @@ public:
 		bUnbound = true;
 
 		Settings.bOverride_AutoExposureBias = 1;
-		Settings.AutoExposureBias = 1.f;
+		Settings.AutoExposureBias = 0.f;
 
 		Settings.bOverride_AutoExposureMinBrightness = 1;
 		Settings.AutoExposureMinBrightness = 1.f;
