Motion Designer:
- Fix Ava Tag Picker crash when re-arranging multi enter conditions in a single state in TL
- Fix Outliner registering a spawn action (with no attach parent) for an Actor that's already been registered

#changelist validated
#virtualized

diff --git a/Engine/Plugins/Experimental/Avalanche/Source/AvalancheOutliner/Private/AvaOutliner.cpp b/Engine/Plugins/Experimental/Avalanche/Source/AvalancheOutliner/Private/AvaOutliner.cpp
--- a/Engine/Plugins/Experimental/Avalanche/Source/AvalancheOutliner/Private/AvaOutliner.cpp
+++ b/Engine/Plugins/Experimental/Avalanche/Source/AvalancheOutliner/Private/AvaOutliner.cpp
@@ -1153,7 +1153,7 @@ void FAvaOutliner::OnActorSpawned(AActor* InActor)
 	const bool bSpawnNotifyIgnored = EnumHasAnyFlags(IgnoreNotifyFlags, EAvaOutlinerIgnoreNotifyFlags::Spawn);
 	if (!bSpawnNotifyIgnored && IsValid(InActor) && InActor->GetWorld() == GetWorld())
 	{
-		if (CanProcessActorSpawn(InActor))
+		if (CanProcessActorSpawn(InActor) && !FindItem(FAvaOutlinerItemId(InActor)))
 		{
 			FAvaOutlinerAddItemParams Params;
 			Params.Item             = FindOrAdd<FAvaOutlinerActor>(InActor);
diff --git a/Engine/Plugins/Experimental/Avalanche/Source/AvalancheTagEditor/Private/Customization/SAvaTagPicker.cpp b/Engine/Plugins/Experimental/Avalanche/Source/AvalancheTagEditor/Private/Customization/SAvaTagPicker.cpp
--- a/Engine/Plugins/Experimental/Avalanche/Source/AvalancheTagEditor/Private/Customization/SAvaTagPicker.cpp
+++ b/Engine/Plugins/Experimental/Avalanche/Source/AvalancheTagEditor/Private/Customization/SAvaTagPicker.cpp
@@ -147,6 +147,10 @@ const UAvaTagCollection* SAvaTagPicker::GetOrLoadTagCollection() const
 	StructPropertyHandle->EnumerateConstRawData(
 		[&TagCollection, this](const void* InStructRawData, const int32 InDataIndex, const int32 InNumData)->bool
 		{
+			if (!InStructRawData)
+			{
+				return true;
+			}
 			const UAvaTagCollection* CurrentTagCollection = TagCustomizer->GetOrLoadTagCollection(InStructRawData);
 			if (InDataIndex == 0)
 			{
@@ -248,7 +252,7 @@ bool SAvaTagPicker::IsTagHandleSelected(FAvaTagHandle InTagHandle) const
 	StructPropertyHandle->EnumerateConstRawData(
 		[&bTagHandleSelected, &InTagHandle, this](const void* InStructRawData, const int32 InDataIndex, const int32 InNumData)->bool
 		{
-			if (TagCustomizer->ContainsTagHandle(InStructRawData, InTagHandle))
+			if (InStructRawData && TagCustomizer->ContainsTagHandle(InStructRawData, InTagHandle))
 			{
 				bTagHandleSelected = true;
 				return false;
@@ -288,6 +292,10 @@ FText SAvaTagPicker::GetValueDisplayText() const
 	StructPropertyHandle->EnumerateConstRawData(
 		[&ValueName, &bMultipleValues, this](const void* InStructRawData, const int32 InDataIndex, const int32 InNumData)->bool
 		{
+			if (!InStructRawData)
+			{
+				return true;	
+			}
 			FName CurrentValueName = TagCustomizer->GetDisplayValueName(InStructRawData);
 			if (InDataIndex == 0)
 			{
