[MotionDesign] moved impl from header in AvaCineCameraActor, AvaPostProcessVolume

diff --git a/Engine/Plugins/Experimental/Avalanche/Source/Avalanche/Private/Viewport/AvaCineCameraActor.cpp b/Engine/Plugins/Experimental/Avalanche/Source/Avalanche/Private/Viewport/AvaCineCameraActor.cpp
new file mode 100644
--- /dev/null
+++ b/Engine/Plugins/Experimental/Avalanche/Source/Avalanche/Private/Viewport/AvaCineCameraActor.cpp
@@ -0,0 +1,33 @@
+﻿// Copyright Epic Games, Inc. All Rights Reserved.
+
+#include "Viewport/AvaCineCameraActor.h"
+
+#include "CineCameraComponent.h"
+
+AAvaCineCameraActor::AAvaCineCameraActor(const FObjectInitializer& ObjectInitializer)
+	: Super(ObjectInitializer)
+	, DefaultCameraDistance(500.0f) // Same value as UAvaEditorSettings::CameraDistance
+{
+	/* todo: when the reset to defaults issue will be solved, Configure(GetDefaultCameraDistance()) could be called here to to initialize defaults */
+}
+
+void AAvaCineCameraActor::Configure(float InCameraDistance)
+{
+	// initialize default camera distance - todo: we might use this value later for a custom reset?
+	DefaultCameraDistance = InCameraDistance;
+		
+	if (UCineCameraComponent* const CineCameraComp = GetCineCameraComponent())
+	{
+		CineCameraComp->SetFieldOfView(90.0f);
+
+		// smaller camera mesh, todo: this will be substituted with a custom visualizer, or other solution, to visualize frustum + direction
+		CineCameraComp->SetWorldScale3D(FVector(0.5f));
+
+		// using default ManualFocusDistance value from Default CineCameraComponent, since that is handled by UAvaEditorSettings
+		CineCameraComp->FocusSettings.ManualFocusDistance = DefaultCameraDistance;
+
+		FVector CameraPosition = FVector::ZeroVector;
+		CameraPosition.X = -DefaultCameraDistance;
+		SetActorLocation(CameraPosition);
+	}
+}
diff --git a/Engine/Plugins/Experimental/Avalanche/Source/Avalanche/Private/Viewport/AvaPostProcessVolume.cpp b/Engine/Plugins/Experimental/Avalanche/Source/Avalanche/Private/Viewport/AvaPostProcessVolume.cpp
new file mode 100644
--- /dev/null
+++ b/Engine/Plugins/Experimental/Avalanche/Source/Avalanche/Private/Viewport/AvaPostProcessVolume.cpp
@@ -0,0 +1,30 @@
+﻿// Copyright Epic Games, Inc. All Rights Reserved.
+
+#include "Viewport/AvaPostProcessVolume.h"
+
+AAvaPostProcessVolume::AAvaPostProcessVolume()
+{
+	BlendWeight = 1.f;
+	bUnbound = true;
+
+	Settings.bOverride_AutoExposureBias = 1;
+	Settings.AutoExposureBias = 0.f;
+
+	Settings.bOverride_AutoExposureMinBrightness = 1;
+	Settings.AutoExposureMinBrightness = 1.f;
+
+	Settings.bOverride_AutoExposureMaxBrightness = 1;
+	Settings.AutoExposureMaxBrightness = 1.f;
+
+	Settings.bOverride_ExpandGamut = 1;
+	Settings.ExpandGamut = 0.f;
+
+	Settings.bOverride_ToneCurveAmount = 1;
+	Settings.ToneCurveAmount = 0.f;
+
+	Settings.bOverride_BlueCorrection = 1;
+	Settings.BlueCorrection = 0.f;
+
+	Settings.bOverride_VignetteIntensity = 1;
+	Settings.VignetteIntensity = 0.f;
+}
diff --git a/Engine/Plugins/Experimental/Avalanche/Source/Avalanche/Public/Viewport/AvaCineCameraActor.h b/Engine/Plugins/Experimental/Avalanche/Source/Avalanche/Public/Viewport/AvaCineCameraActor.h
--- a/Engine/Plugins/Experimental/Avalanche/Source/Avalanche/Public/Viewport/AvaCineCameraActor.h
+++ b/Engine/Plugins/Experimental/Avalanche/Source/Avalanche/Public/Viewport/AvaCineCameraActor.h
@@ -3,7 +3,6 @@
 #pragma once
 
 #include "CineCameraActor.h"
-#include "CineCameraComponent.h"
 #include "AvaCineCameraActor.generated.h"
 
 /** 
@@ -13,69 +12,23 @@
  * In particular, Motion Design Editor configuration property "Camera Distance" is used to initialize camera position and manual focus.
  * See: Editor Preferences --> Motion Design --> Editor Settings --> Camera Distance
  */
-UCLASS(DisplayName = "Motion Design Cine Camera Actor")
-class AVALANCHE_API AAvaCineCameraActor : public ACineCameraActor
+UCLASS(MinimalAPI, DisplayName = "Motion Design Cine Camera Actor")
+class AAvaCineCameraActor : public ACineCameraActor
 {
 	friend class UAvaEditorSettings;
 	
 	GENERATED_BODY()
 
 public:
-	AAvaCineCameraActor(const FObjectInitializer& ObjectInitializer) : Super(ObjectInitializer)
-	{
-		/* todo: when the reset to defaults issue will be solved, Configure(GetDefaultCameraDistance()) could be called here to to initialize defaults */
-	}
+	AAvaCineCameraActor(const FObjectInitializer& ObjectInitializer);
 
 	/**
 	 * Initialize the camera with Avalanche scene default values: field of view, camera position, focus distance.
 	 * @param InCameraDistance: camera distance value, used to initialize camera position and manual focus distance
 	 */
-	void Configure(float InCameraDistance)
-	{
-		// initialize default camera distance - todo: we might use this value later for a custom reset?
-		DefaultCameraDistance = InCameraDistance;
-		
-		if (UCineCameraComponent* const CineCameraComp = GetCineCameraComponent())
-		{
-			CineCameraComp->SetFieldOfView(90.0f);
-
-			// smaller camera mesh, todo: this will be substituted with a custom visualizer, or other solution, to visualize frustum + direction
-			CineCameraComp->SetWorldScale3D(FVector(0.5f));
-
-			// using default ManualFocusDistance value from Default CineCameraComponent, since that is handled by UAvaEditorSettings
-			CineCameraComp->FocusSettings.ManualFocusDistance = DefaultCameraDistance;
-
-			FVector CameraPosition = FVector::ZeroVector;
-			CameraPosition.X = -DefaultCameraDistance;
-			SetActorLocation(CameraPosition);
-		}
-	}
+	AVALANCHE_API void Configure(float InCameraDistance);
 
 private:
-	/**
-	 * Updates some defaults of the Mutable CDO's CineCameraComponent, based on the provided CameraDistance value
-	 */
-	static void SetDefaultCameraDistance(float InCameraDistance)
-	{
-		AAvaCineCameraActor* const DefaultObject = GetMutableDefault<AAvaCineCameraActor>();
-
-		FVector CameraPosition = FVector::ZeroVector;
-		CameraPosition.X = -InCameraDistance;
-		DefaultObject->SetActorLocation(CameraPosition);
-		
-		if (UCineCameraComponent* const DefaultCineCameraComponent = DefaultObject->GetCineCameraComponent())
-		{
-			DefaultCineCameraComponent->FocusSettings.ManualFocusDistance = InCameraDistance;
-		}
-
-		DefaultObject->DefaultCameraDistance = InCameraDistance;
-	}
-
-	static float GetDefaultCameraDistance()
-	{
-		return GetMutableDefault<AAvaCineCameraActor>()->DefaultCameraDistance;
-	}
-
 	/** The CameraDistance value used when configuring this AvaCineCamera*/
 	UPROPERTY()
 	float DefaultCameraDistance;
diff --git a/Engine/Plugins/Experimental/Avalanche/Source/Avalanche/Public/Viewport/AvaPostProcessVolume.h b/Engine/Plugins/Experimental/Avalanche/Source/Avalanche/Public/Viewport/AvaPostProcessVolume.h
--- a/Engine/Plugins/Experimental/Avalanche/Source/Avalanche/Public/Viewport/AvaPostProcessVolume.h
+++ b/Engine/Plugins/Experimental/Avalanche/Source/Avalanche/Public/Viewport/AvaPostProcessVolume.h
@@ -16,30 +16,5 @@ class AAvaPostProcessVolume : public APostProcessVolume
 	GENERATED_BODY()
 
 public:
-	AAvaPostProcessVolume()
-	{
-		BlendWeight = 1.f;
-		bUnbound = true;
-
-		Settings.bOverride_AutoExposureBias = 1;
-		Settings.AutoExposureBias = 0.f;
-
-		Settings.bOverride_AutoExposureMinBrightness = 1;
-		Settings.AutoExposureMinBrightness = 1.f;
-
-		Settings.bOverride_AutoExposureMaxBrightness = 1;
-		Settings.AutoExposureMaxBrightness = 1.f;
-
-		Settings.bOverride_ExpandGamut = 1;
-		Settings.ExpandGamut = 0.f;
-
-		Settings.bOverride_ToneCurveAmount = 1;
-		Settings.ToneCurveAmount = 0.f;
-
-		Settings.bOverride_BlueCorrection = 1;
-		Settings.BlueCorrection = 0.f;
-
-		Settings.bOverride_VignetteIntensity = 1;
-		Settings.VignetteIntensity = 0.f;
-	}
+	AAvaPostProcessVolume();
 };
diff --git a/Engine/Plugins/Experimental/Avalanche/Source/AvalancheEditor/Private/AvaEditorSettings.cpp b/Engine/Plugins/Experimental/Avalanche/Source/AvalancheEditor/Private/AvaEditorSettings.cpp
--- a/Engine/Plugins/Experimental/Avalanche/Source/AvalancheEditor/Private/AvaEditorSettings.cpp
+++ b/Engine/Plugins/Experimental/Avalanche/Source/AvalancheEditor/Private/AvaEditorSettings.cpp
@@ -66,10 +66,4 @@ void UAvaEditorSettings::PostEditChangeProperty(FPropertyChangedEvent& PropertyC
 	OnChanged.Broadcast(this, PropertyChangedEvent.GetPropertyName());
 }
 
-// todo: when the reset to defaults issue will be solved, this function could be used from both PostInitProperties and PECP to update AAvaCineCameraActor defaults
-void UAvaEditorSettings::UpdateAvaCineCameraDefaults() const
-{
-	AAvaCineCameraActor::SetDefaultCameraDistance(CameraDistance);
-}
-
 #undef LOCTEXT_NAMESPACE
diff --git a/Engine/Plugins/Experimental/Avalanche/Source/AvalancheEditor/Private/AvaEditorSettings.h b/Engine/Plugins/Experimental/Avalanche/Source/AvalancheEditor/Private/AvaEditorSettings.h
--- a/Engine/Plugins/Experimental/Avalanche/Source/AvalancheEditor/Private/AvaEditorSettings.h
+++ b/Engine/Plugins/Experimental/Avalanche/Source/AvalancheEditor/Private/AvaEditorSettings.h
@@ -92,10 +92,4 @@ public:
 	//~ End UObject
 
 	void OpenEditorSettingsWindow() const;
-
-	/**
-	 * AvaCineCamera uses CameraDistance property to setup manual focus in its default object.
-	 * This function updates that default value to match the one from AvaEditorSettings
-	 */
-	void UpdateAvaCineCameraDefaults() const;
 };
