[Motion Design] Rundown Transition Logic: Fixing logic that no-TL pages kick playing pages.

This is regression introduced in 30961363, we added the distinction between playing pages and exiting pages so that playing pages can be included in take out transitions and  react to other pages being taken out. This is broke the case where there is a no-TL enter page that should kick out all playing. So, need to add this case now. I decided to deal with it in the page transition. I could have also dealt with it in the playable transition. Not sure which is best.

Note: 30961363 had made the flag EAvaPlayableTransitionFlags::ExitOnly obsolete because of the explicit distinction between "playing" and "exiting" playabled. So, I am recycling that flag.

#rnx

diff --git a/Engine/Plugins/Experimental/Avalanche/Source/AvalancheMedia/Private/Playable/Transition/AvaPlayableLocalTransition.cpp b/Engine/Plugins/Experimental/Avalanche/Source/AvalancheMedia/Private/Playable/Transition/AvaPlayableLocalTransition.cpp
--- a/Engine/Plugins/Experimental/Avalanche/Source/AvalancheMedia/Private/Playable/Transition/AvaPlayableLocalTransition.cpp
+++ b/Engine/Plugins/Experimental/Avalanche/Source/AvalancheMedia/Private/Playable/Transition/AvaPlayableLocalTransition.cpp
@@ -292,9 +292,12 @@ bool UAvaPlayableLocalTransition::Start()
 		++ArrayIndex;
 	}
 
+	const EAvaPlayableTransitionEntryRole PlayingPlayablesRole = EnumHasAnyFlags(TransitionFlags, EAvaPlayableTransitionFlags::TreatPlayingAsExiting) ?
+		EAvaPlayableTransitionEntryRole::Exit : EAvaPlayableTransitionEntryRole::Playing;
+
 	for (UAvaPlayable* Playable : PlayingPlayables)
 	{
-		if (!Helper.AddTransitionBehaviorInstance(Playable, EAvaPlayableTransitionEntryRole::Playing))
+		if (!Helper.AddTransitionBehaviorInstance(Playable, PlayingPlayablesRole))
 		{
 			// No Transition Tree:
 			// Stop exit/playing playable without any transition.
diff --git a/Engine/Plugins/Experimental/Avalanche/Source/AvalancheMedia/Private/Rundown/Transition/AvaRundownPageTransition.cpp b/Engine/Plugins/Experimental/Avalanche/Source/AvalancheMedia/Private/Rundown/Transition/AvaRundownPageTransition.cpp
--- a/Engine/Plugins/Experimental/Avalanche/Source/AvalancheMedia/Private/Rundown/Transition/AvaRundownPageTransition.cpp
+++ b/Engine/Plugins/Experimental/Avalanche/Source/AvalancheMedia/Private/Rundown/Transition/AvaRundownPageTransition.cpp
@@ -484,10 +484,10 @@ void UAvaRundownPageTransition::MakePlayableTransition()
 	AddPlayersToBuilder(TransitionBuilder, ExitPlayersWeak, TEXT("Exit"), EAvaPlayableTransitionEntryRole::Exit);
 	PlayableTransition = TransitionBuilder.MakeTransition(this);
 
-	// Mark exit only transitions to ensure they properly create the null behavior instances. 
-	if (PlayableTransition && EnterPlayersWeak.IsEmpty())
+	// For non-TL enter pages, we need to kick out the playing pages too. 
+	if (PlayableTransition && HasEnterPagesWithNoTransitionLogic())
 	{
-		PlayableTransition->SetTransitionFlags(EAvaPlayableTransitionFlags::ExitOnly);
+		PlayableTransition->SetTransitionFlags(EAvaPlayableTransitionFlags::TreatPlayingAsExiting);
 	}
 }
 
diff --git a/Engine/Plugins/Experimental/Avalanche/Source/AvalancheMedia/Public/AvaMediaDefines.h b/Engine/Plugins/Experimental/Avalanche/Source/AvalancheMedia/Public/AvaMediaDefines.h
--- a/Engine/Plugins/Experimental/Avalanche/Source/AvalancheMedia/Public/AvaMediaDefines.h
+++ b/Engine/Plugins/Experimental/Avalanche/Source/AvalancheMedia/Public/AvaMediaDefines.h
@@ -124,8 +124,8 @@ enum class EAvaPlayableSequenceEventType : uint8
 enum class EAvaPlayableTransitionFlags : uint8
 {
 	None = 0,
-	/** This transition contains only exit playables. */
-	ExitOnly = 1 << 0,
+	/** Playing playables will be treated as exit playables. */
+	TreatPlayingAsExiting = 1 << 0,
 };
 ENUM_CLASS_FLAGS(EAvaPlayableTransitionFlags);
 
