[Motion Design]
 - Remove remaining CoreMinimal.h headers

diff --git a/Engine/Plugins/Experimental/Avalanche/Source/AvalancheMedia/Private/AvaMediaRenderTargetUtils.cpp b/Engine/Plugins/Experimental/Avalanche/Source/AvalancheMedia/Private/AvaMediaRenderTargetUtils.cpp
--- a/Engine/Plugins/Experimental/Avalanche/Source/AvalancheMedia/Private/AvaMediaRenderTargetUtils.cpp
+++ b/Engine/Plugins/Experimental/Avalanche/Source/AvalancheMedia/Private/AvaMediaRenderTargetUtils.cpp
@@ -1,10 +1,12 @@
 // Copyright Epic Games, Inc. All Rights Reserved.
 
 #include "AvaMediaRenderTargetUtils.h"
-
 #include "Engine/Engine.h"
 #include "Engine/TextureRenderTarget2D.h"
+#include "PixelFormat.h"
 #include "UObject/Package.h"
+#include "UObject/NameTypes.h"
+#include "UObject/NoExportTypes.h"
 
 namespace UE::AvaMediaRenderTargetUtils
 {
diff --git a/Engine/Plugins/Experimental/Avalanche/Source/AvalancheMedia/Private/Broadcast/AvaBroadcast.cpp b/Engine/Plugins/Experimental/Avalanche/Source/AvalancheMedia/Private/Broadcast/AvaBroadcast.cpp
--- a/Engine/Plugins/Experimental/Avalanche/Source/AvalancheMedia/Private/Broadcast/AvaBroadcast.cpp
+++ b/Engine/Plugins/Experimental/Avalanche/Source/AvalancheMedia/Private/Broadcast/AvaBroadcast.cpp
@@ -1,14 +1,16 @@
 // Copyright Epic Games, Inc. All Rights Reserved.
 
 #include "Broadcast/AvaBroadcast.h"
-
 #include "Async/Async.h"
 #include "AvaMediaSerializationUtils.h"
+#include "Containers/UnrealString.h"
+#include "Delegates/IDelegateInstance.h"
 #include "Formatters/XmlArchiveInputFormatter.h"
 #include "Formatters/XmlArchiveOutputFormatter.h"
 #include "HAL/FileManager.h"
 #include "IAvaMediaModule.h"
 #include "Misc/Paths.h"
+#include "UObject/NameTypes.h"
 #include "UObject/Package.h"
 
 DEFINE_LOG_CATEGORY(LogAvaBroadcast);
diff --git a/Engine/Plugins/Experimental/Avalanche/Source/AvalancheMedia/Private/Broadcast/Channel/AvaBroadcastOutputChannel.cpp b/Engine/Plugins/Experimental/Avalanche/Source/AvalancheMedia/Private/Broadcast/Channel/AvaBroadcastOutputChannel.cpp
--- a/Engine/Plugins/Experimental/Avalanche/Source/AvalancheMedia/Private/Broadcast/Channel/AvaBroadcastOutputChannel.cpp
+++ b/Engine/Plugins/Experimental/Avalanche/Source/AvalancheMedia/Private/Broadcast/Channel/AvaBroadcastOutputChannel.cpp
@@ -1,12 +1,12 @@
 // Copyright Epic Games, Inc. All Rights Reserved.
 
 #include "Broadcast/Channel/AvaBroadcastOutputChannel.h"
-
 #include "AvaMediaRenderTargetUtils.h"
 #include "AvaMediaSettings.h"
 #include "Blueprint/UserWidget.h"
 #include "Blueprint/WidgetBlueprintGeneratedClass.h"
 #include "Broadcast/AvaBroadcast.h"
+#include "Broadcast/AvaBroadcastProfile.h"
 #include "Broadcast/Channel/IAvaBroadcastOutputChannelViewInterface.h"
 #include "Broadcast/IAvaBroadcastSettings.h"
 #include "Broadcast/OutputDevices/AvaBroadcastOutputUtils.h"
@@ -14,17 +14,25 @@
 #include "Engine/TextureRenderTarget2D.h"
 #include "Extensions/UserWidgetExtension.h"
 #include "Framework/Application/SlateApplication.h"
+#include "GameTime.h"
+#include "GlobalShader.h"
 #include "IAvaMediaModule.h"
 #include "MediaCapture.h"
 #include "MediaOutput.h"
 #include "MediaShaders.h"
+#include "PixelFormat.h"
 #include "Playable/AvaPlayableGroup.h"
 #include "Playback/IAvaPlaybackClient.h"
+#include "RenderGraphBuilder.h"
 #include "RenderResource.h"
 #include "RHICommandList.h"
-#include "ScreenPass.h"
+#include "RHIResources.h"
+#include "SceneView.h"
 #include "Slate/WidgetRenderer.h"
+#include "StereoRenderTargetManager.h"
+#include "ScreenPass.h"
 #include "TextureResource.h"
+#include "UObject/NameTypes.h"
 #include "UObject/Package.h"
 
 DEFINE_LOG_CATEGORY_STATIC(LogAvaOutputChannel, Log, All);
diff --git a/Engine/Plugins/Experimental/Avalanche/Source/AvalancheMedia/Private/Broadcast/Channel/Slate/SAvaBroadcastPlaceholderWidget.cpp b/Engine/Plugins/Experimental/Avalanche/Source/AvalancheMedia/Private/Broadcast/Channel/Slate/SAvaBroadcastPlaceholderWidget.cpp
--- a/Engine/Plugins/Experimental/Avalanche/Source/AvalancheMedia/Private/Broadcast/Channel/Slate/SAvaBroadcastPlaceholderWidget.cpp
+++ b/Engine/Plugins/Experimental/Avalanche/Source/AvalancheMedia/Private/Broadcast/Channel/Slate/SAvaBroadcastPlaceholderWidget.cpp
@@ -1,8 +1,8 @@
 // Copyright Epic Games, Inc. All Rights Reserved.
 
 #include "Broadcast/Channel/Slate/SAvaBroadcastPlaceholderWidget.h"
-
 #include "AvaMediaStyle.h"
+#include "Internationalization/Text.h"
 #include "Widgets/Images/SImage.h"
 #include "Widgets/Images/SThrobber.h"
 #include "Widgets/Layout/SBox.h"
diff --git a/Engine/Plugins/Experimental/Avalanche/Source/AvalancheMedia/Private/Broadcast/OutputDevices/AvaBroadcastDisplayDeviceProvider.cpp b/Engine/Plugins/Experimental/Avalanche/Source/AvalancheMedia/Private/Broadcast/OutputDevices/AvaBroadcastDisplayDeviceProvider.cpp
--- a/Engine/Plugins/Experimental/Avalanche/Source/AvalancheMedia/Private/Broadcast/OutputDevices/AvaBroadcastDisplayDeviceProvider.cpp
+++ b/Engine/Plugins/Experimental/Avalanche/Source/AvalancheMedia/Private/Broadcast/OutputDevices/AvaBroadcastDisplayDeviceProvider.cpp
@@ -1,9 +1,10 @@
 // Copyright Epic Games, Inc. All Rights Reserved.
 
 #include "AvaBroadcastDisplayDeviceProvider.h"
-
 #include "Broadcast/OutputDevices/AvaBroadcastDisplayDeviceManager.h"
+#include "Internationalization/Text.h"
 #include "MediaIOCoreCommonDisplayMode.h"
+#include "MediaIOCoreDefinitions.h"
 
 #define LOCTEXT_NAMESPACE "AvaBroadcastDisplayDeviceProvider"
 
diff --git a/Engine/Plugins/Experimental/Avalanche/Source/AvalancheMedia/Private/Broadcast/OutputDevices/AvaBroadcastDisplayDeviceProvider.h b/Engine/Plugins/Experimental/Avalanche/Source/AvalancheMedia/Private/Broadcast/OutputDevices/AvaBroadcastDisplayDeviceProvider.h
--- a/Engine/Plugins/Experimental/Avalanche/Source/AvalancheMedia/Private/Broadcast/OutputDevices/AvaBroadcastDisplayDeviceProvider.h
+++ b/Engine/Plugins/Experimental/Avalanche/Source/AvalancheMedia/Private/Broadcast/OutputDevices/AvaBroadcastDisplayDeviceProvider.h
@@ -2,9 +2,17 @@
 
 #pragma once
 
-#include "CoreMinimal.h"
+#include "Containers/Array.h"
 #include "IMediaIOCoreDeviceProvider.h"
-#include "UObject/ObjectMacros.h"
+
+class FText;
+struct FMediaIOConfiguration;
+struct FMediaIOConnection;
+struct FMediaIODevice;
+struct FMediaIOInputConfiguration;
+struct FMediaIOMode;
+struct FMediaIOOutputConfiguration;
+struct FMediaIOVideoTimecodeConfiguration;
 
 /**
  * Implementation of IMediaIOCoreDeviceProvider for UAvaBroadcastDisplayMediaOutput.
diff --git a/Engine/Plugins/Experimental/Avalanche/Source/AvalancheMedia/Private/Broadcast/OutputDevices/Slate/SAvaBroadcastCaptureImage.cpp b/Engine/Plugins/Experimental/Avalanche/Source/AvalancheMedia/Private/Broadcast/OutputDevices/Slate/SAvaBroadcastCaptureImage.cpp
--- a/Engine/Plugins/Experimental/Avalanche/Source/AvalancheMedia/Private/Broadcast/OutputDevices/Slate/SAvaBroadcastCaptureImage.cpp
+++ b/Engine/Plugins/Experimental/Avalanche/Source/AvalancheMedia/Private/Broadcast/OutputDevices/Slate/SAvaBroadcastCaptureImage.cpp
@@ -1,6 +1,17 @@
 // Copyright Epic Games, Inc. All Rights Reserved.
 
 #include "Broadcast/OutputDevices/Slate/SAvaBroadcastCaptureImage.h"
+#include "Layout/FlowDirection.h"
+#include "Layout/Geometry.h"
+#include "Layout/SlateRect.h"
+#include "Math/TransformCalculus2D.h"
+#include "Rendering/DrawElements.h"
+#include "Rendering/DrawElementTypes.h"
+#include "Rendering/SlateRenderTransform.h"
+#include "Styling/SlateBrush.h"
+#include "Styling/WidgetStyle.h"
+#include "Types/PaintArgs.h"
+#include "UObject/NoExportTypes.h"
 
 void SAvaBroadcastCaptureImage::Construct(const FArguments& InArgs)
 {
diff --git a/Engine/Plugins/Experimental/Avalanche/Source/AvalancheMedia/Private/Playback/Nodes/AvaPlaybackNode.cpp b/Engine/Plugins/Experimental/Avalanche/Source/AvalancheMedia/Private/Playback/Nodes/AvaPlaybackNode.cpp
--- a/Engine/Plugins/Experimental/Avalanche/Source/AvalancheMedia/Private/Playback/Nodes/AvaPlaybackNode.cpp
+++ b/Engine/Plugins/Experimental/Avalanche/Source/AvalancheMedia/Private/Playback/Nodes/AvaPlaybackNode.cpp
@@ -1,11 +1,13 @@
 // Copyright Epic Games, Inc. All Rights Reserved.
 
 #include "Playback/Nodes/AvaPlaybackNode.h"
-
 #include "EngineUtils.h"
+#include "Internationalization/Text.h"
 #include "Misc/App.h"
 #include "Playback/AvaPlaybackGraph.h"
 #include "Playback/IAvaPlaybackGraphEditor.h"
+#include "UObject/NameTypes.h"
+#include "UObject/UObjectGlobals.h"
 
 #if WITH_EDITOR
 #include "EdGraph/EdGraph.h"
diff --git a/Engine/Plugins/Experimental/Avalanche/Source/AvalancheMedia/Private/Playback/Nodes/AvaPlaybackNodeCombiner.cpp b/Engine/Plugins/Experimental/Avalanche/Source/AvalancheMedia/Private/Playback/Nodes/AvaPlaybackNodeCombiner.cpp
--- a/Engine/Plugins/Experimental/Avalanche/Source/AvalancheMedia/Private/Playback/Nodes/AvaPlaybackNodeCombiner.cpp
+++ b/Engine/Plugins/Experimental/Avalanche/Source/AvalancheMedia/Private/Playback/Nodes/AvaPlaybackNodeCombiner.cpp
@@ -1,6 +1,13 @@
 // Copyright Epic Games, Inc. All Rights Reserved.
 
 #include "Playback/Nodes/AvaPlaybackNodeCombiner.h"
+#include "AvaMediaDefines.h"
+#include "Containers/UnrealString.h"
+#include "Internationalization/Text.h"
+
+#if WITH_EDITOR
+#include "UObject/NameTypes.h"
+#endif
 
 #define LOCTEXT_NAMESPACE "AvaPlaybackNodeCombiner"
 
diff --git a/Engine/Plugins/Experimental/Avalanche/Source/AvalancheMedia/Private/Playback/Nodes/AvaPlaybackNodeLevelPlayer.cpp b/Engine/Plugins/Experimental/Avalanche/Source/AvalancheMedia/Private/Playback/Nodes/AvaPlaybackNodeLevelPlayer.cpp
--- a/Engine/Plugins/Experimental/Avalanche/Source/AvalancheMedia/Private/Playback/Nodes/AvaPlaybackNodeLevelPlayer.cpp
+++ b/Engine/Plugins/Experimental/Avalanche/Source/AvalancheMedia/Private/Playback/Nodes/AvaPlaybackNodeLevelPlayer.cpp
@@ -1,8 +1,9 @@
 // Copyright Epic Games, Inc. All Rights Reserved.
 
 #include "Playback/Nodes/AvaPlaybackNodeLevelPlayer.h"
-
 #include "Engine/World.h"
+#include "Framework/AvaSoftAssetPtr.h"
+#include "Internationalization/Text.h"
 #include "Playback/AvaPlaybackGraph.h"
 
 #define LOCTEXT_NAMESPACE "AvaPlaybackNodeLevelPlayer"
diff --git a/Engine/Plugins/Experimental/Avalanche/Source/AvalancheMedia/Private/Playback/Nodes/AvaPlaybackNodePlayer.cpp b/Engine/Plugins/Experimental/Avalanche/Source/AvalancheMedia/Private/Playback/Nodes/AvaPlaybackNodePlayer.cpp
--- a/Engine/Plugins/Experimental/Avalanche/Source/AvalancheMedia/Private/Playback/Nodes/AvaPlaybackNodePlayer.cpp
+++ b/Engine/Plugins/Experimental/Avalanche/Source/AvalancheMedia/Private/Playback/Nodes/AvaPlaybackNodePlayer.cpp
@@ -1,8 +1,10 @@
 // Copyright Epic Games, Inc. All Rights Reserved.
 
 #include "Playback/Nodes/AvaPlaybackNodePlayer.h"
-
+#include "AvaMediaDefines.h"
 #include "Broadcast/AvaBroadcast.h"
+#include "Engine/TextureRenderTarget2D.h"
+#include "Internationalization/Text.h"
 #include "Playable/AvaPlayable.h"
 #include "Playback/AvaPlaybackGraph.h"
 #include "Playback/Nodes/Events/AvaPlaybackNodeEvent.h"
diff --git a/Engine/Plugins/Experimental/Avalanche/Source/AvalancheMedia/Private/Playback/Nodes/AvaPlaybackNodeSwitcher.cpp b/Engine/Plugins/Experimental/Avalanche/Source/AvalancheMedia/Private/Playback/Nodes/AvaPlaybackNodeSwitcher.cpp
--- a/Engine/Plugins/Experimental/Avalanche/Source/AvalancheMedia/Private/Playback/Nodes/AvaPlaybackNodeSwitcher.cpp
+++ b/Engine/Plugins/Experimental/Avalanche/Source/AvalancheMedia/Private/Playback/Nodes/AvaPlaybackNodeSwitcher.cpp
@@ -1,6 +1,13 @@
 // Copyright Epic Games, Inc. All Rights Reserved.
 
 #include "Playback/Nodes/AvaPlaybackNodeSwitcher.h"
+#include "AvaMediaDefines.h"
+#include "Containers/UnrealString.h"
+#include "Internationalization/Text.h"
+
+#if WITH_EDITOR
+#include "UObject/NameTypes.h"
+#endif
 
 #define LOCTEXT_NAMESPACE "AvaPlaybackNodeSwitcher"
 
diff --git a/Engine/Plugins/Experimental/Avalanche/Source/AvalancheMedia/Private/Playback/Nodes/Events/Actions/AvaPlaybackNode_PlayAnim.cpp b/Engine/Plugins/Experimental/Avalanche/Source/AvalancheMedia/Private/Playback/Nodes/Events/Actions/AvaPlaybackNode_PlayAnim.cpp
--- a/Engine/Plugins/Experimental/Avalanche/Source/AvalancheMedia/Private/Playback/Nodes/Events/Actions/AvaPlaybackNode_PlayAnim.cpp
+++ b/Engine/Plugins/Experimental/Avalanche/Source/AvalancheMedia/Private/Playback/Nodes/Events/Actions/AvaPlaybackNode_PlayAnim.cpp
@@ -1,14 +1,18 @@
 // Copyright Epic Games, Inc. All Rights Reserved.
 
 #include "AvaPlaybackNode_PlayAnim.h"
-
 #include "Async/Async.h"
+#include "AvaMediaDefines.h"
 #include "AvaScene.h"
 #include "AvaSequence.h"
 #include "Broadcast/AvaBroadcast.h"
 #include "Engine/Level.h"
+#include "Internationalization/Text.h"
 #include "Playback/AvaPlaybackGraph.h"
+#include "Playback/Nodes/AvaPlaybackNode.h"
 #include "Playback/Nodes/AvaPlaybackNodeLevelPlayer.h"
+#include "Playback/Nodes/Events/Actions/AvaPlaybackAnimations.h"
+#include "UObject/NoExportTypes.h"
 
 #define LOCTEXT_NAMESPACE "AvaPlaybackNode_PlayAnim"
 
diff --git a/Engine/Plugins/Experimental/Avalanche/Source/AvalancheMedia/Private/Playback/Nodes/Events/Actions/AvaPlaybackNode_PlayAnim.h b/Engine/Plugins/Experimental/Avalanche/Source/AvalancheMedia/Private/Playback/Nodes/Events/Actions/AvaPlaybackNode_PlayAnim.h
--- a/Engine/Plugins/Experimental/Avalanche/Source/AvalancheMedia/Private/Playback/Nodes/Events/Actions/AvaPlaybackNode_PlayAnim.h
+++ b/Engine/Plugins/Experimental/Avalanche/Source/AvalancheMedia/Private/Playback/Nodes/Events/Actions/AvaPlaybackNode_PlayAnim.h
@@ -2,18 +2,25 @@
 
 #pragma once
 
-#include "CoreMinimal.h"
+#include "Containers/Array.h"
+#include "Containers/Map.h"
+#include "Containers/Set.h"
 #include "Playback/Nodes/Events/Actions/AvaPlaybackAnimations.h"
 #include "Playback/Nodes/Events/AvaPlaybackNodeAction.h"
 #include "AvaPlaybackNode_PlayAnim.generated.h"
 
+class FText;
+class UAvaPlaybackNode;
+struct FAvaPlaybackAnimations;
+struct FAvaPlaybackEventParameters;
+struct FSoftObjectPath;
+
 UCLASS()
 class AVALANCHEMEDIA_API UAvaPlaybackNode_PlayAnim : public UAvaPlaybackNodeAction
 {
 	GENERATED_BODY()
 
 public:
-	
 	virtual FText GetNodeDisplayNameText() const override;
 	virtual FText GetNodeTooltipText() const override;
 	virtual void OnEventTriggered(const FAvaPlaybackEventParameters& InEventParameters) override;
@@ -21,8 +28,8 @@ public:
 	virtual void PreDryRun() override;
 	virtual void DryRun(const TArray<UAvaPlaybackNode*>& InAncestors) override;
 	virtual void PostDryRun() override;
+
 protected:
-	
 	UPROPERTY(VisibleAnywhere, Category = "Motion Design")
 	TMap<FSoftObjectPath, FAvaPlaybackAnimations> AnimationMap;
 
diff --git a/Engine/Plugins/Experimental/Avalanche/Source/AvalancheMedia/Private/Playback/Nodes/Events/Actions/AvaPlaybackNode_PreloadPlayer.cpp b/Engine/Plugins/Experimental/Avalanche/Source/AvalancheMedia/Private/Playback/Nodes/Events/Actions/AvaPlaybackNode_PreloadPlayer.cpp
--- a/Engine/Plugins/Experimental/Avalanche/Source/AvalancheMedia/Private/Playback/Nodes/Events/Actions/AvaPlaybackNode_PreloadPlayer.cpp
+++ b/Engine/Plugins/Experimental/Avalanche/Source/AvalancheMedia/Private/Playback/Nodes/Events/Actions/AvaPlaybackNode_PreloadPlayer.cpp
@@ -1,8 +1,9 @@
 // Copyright Epic Games, Inc. All Rights Reserved.
 
 #include "AvaPlaybackNode_PreloadPlayer.h"
-
+#include "AvaMediaDefines.h"
 #include "Broadcast/AvaBroadcast.h"
+#include "Internationalization/Text.h"
 #include "Playback/AvaPlaybackGraph.h"
 
 #define LOCTEXT_NAMESPACE "AvaPlaybackNode_PreloadPlayer"
diff --git a/Engine/Plugins/Experimental/Avalanche/Source/AvalancheMedia/Private/Playback/Nodes/Events/Actions/AvaPlaybackNode_PreloadPlayer.h b/Engine/Plugins/Experimental/Avalanche/Source/AvalancheMedia/Private/Playback/Nodes/Events/Actions/AvaPlaybackNode_PreloadPlayer.h
--- a/Engine/Plugins/Experimental/Avalanche/Source/AvalancheMedia/Private/Playback/Nodes/Events/Actions/AvaPlaybackNode_PreloadPlayer.h
+++ b/Engine/Plugins/Experimental/Avalanche/Source/AvalancheMedia/Private/Playback/Nodes/Events/Actions/AvaPlaybackNode_PreloadPlayer.h
@@ -2,17 +2,18 @@
 
 #pragma once
 
-#include "CoreMinimal.h"
 #include "Playback/Nodes/Events/AvaPlaybackNodeAction.h"
 #include "AvaPlaybackNode_PreloadPlayer.generated.h"
 
+class FText;
+struct FAvaPlaybackEventParameters;
+
 UCLASS()
 class AVALANCHEMEDIA_API UAvaPlaybackNode_PreloadPlayer : public UAvaPlaybackNodeAction
 {
 	GENERATED_BODY()
 
 public:
-	
 	virtual FText GetNodeDisplayNameText() const override;
 	virtual FText GetNodeTooltipText() const override;
 	virtual void OnEventTriggered(const FAvaPlaybackEventParameters& InEventParameters) override;
diff --git a/Engine/Plugins/Experimental/Avalanche/Source/AvalancheMedia/Private/Playback/Nodes/Events/AvaPlaybackNodeAction.cpp b/Engine/Plugins/Experimental/Avalanche/Source/AvalancheMedia/Private/Playback/Nodes/Events/AvaPlaybackNodeAction.cpp
--- a/Engine/Plugins/Experimental/Avalanche/Source/AvalancheMedia/Private/Playback/Nodes/Events/AvaPlaybackNodeAction.cpp
+++ b/Engine/Plugins/Experimental/Avalanche/Source/AvalancheMedia/Private/Playback/Nodes/Events/AvaPlaybackNodeAction.cpp
@@ -1,6 +1,8 @@
 // Copyright Epic Games, Inc. All Rights Reserved.
 
 #include "Playback/Nodes/Events/AvaPlaybackNodeAction.h"
+#include "Internationalization/Text.h"
+#include "Playback/Nodes/AvaPlaybackNode.h"
 
 FText UAvaPlaybackNodeAction::GetNodeCategoryText() const
 {
diff --git a/Engine/Plugins/Experimental/Avalanche/Source/AvalancheMedia/Private/Playback/Nodes/Events/AvaPlaybackNodeEvent.cpp b/Engine/Plugins/Experimental/Avalanche/Source/AvalancheMedia/Private/Playback/Nodes/Events/AvaPlaybackNodeEvent.cpp
--- a/Engine/Plugins/Experimental/Avalanche/Source/AvalancheMedia/Private/Playback/Nodes/Events/AvaPlaybackNodeEvent.cpp
+++ b/Engine/Plugins/Experimental/Avalanche/Source/AvalancheMedia/Private/Playback/Nodes/Events/AvaPlaybackNodeEvent.cpp
@@ -1,6 +1,7 @@
 // Copyright Epic Games, Inc. All Rights Reserved.
 
 #include "Playback/Nodes/Events/AvaPlaybackNodeEvent.h"
+#include "AvaMediaDefines.h"
 
 void UAvaPlaybackNodeEvent::TickEvent(float DeltaTime, FAvaPlaybackEventParameters& OutEventParameters)
 {
diff --git a/Engine/Plugins/Experimental/Avalanche/Source/AvalancheMedia/Private/Playback/Nodes/Events/AvaPlaybackNodeFlow.cpp b/Engine/Plugins/Experimental/Avalanche/Source/AvalancheMedia/Private/Playback/Nodes/Events/AvaPlaybackNodeFlow.cpp
--- a/Engine/Plugins/Experimental/Avalanche/Source/AvalancheMedia/Private/Playback/Nodes/Events/AvaPlaybackNodeFlow.cpp
+++ b/Engine/Plugins/Experimental/Avalanche/Source/AvalancheMedia/Private/Playback/Nodes/Events/AvaPlaybackNodeFlow.cpp
@@ -1,6 +1,8 @@
 // Copyright Epic Games, Inc. All Rights Reserved.
 
 #include "Playback/Nodes/Events/AvaPlaybackNodeFlow.h"
+#include "Internationalization/Text.h"
+#include "Playback/Nodes/AvaPlaybackNode.h"
 
 FText UAvaPlaybackNodeFlow::GetNodeCategoryText() const
 {
diff --git a/Engine/Plugins/Experimental/Avalanche/Source/AvalancheMedia/Private/Playback/Nodes/Events/AvaPlaybackNodeTrigger.cpp b/Engine/Plugins/Experimental/Avalanche/Source/AvalancheMedia/Private/Playback/Nodes/Events/AvaPlaybackNodeTrigger.cpp
--- a/Engine/Plugins/Experimental/Avalanche/Source/AvalancheMedia/Private/Playback/Nodes/Events/AvaPlaybackNodeTrigger.cpp
+++ b/Engine/Plugins/Experimental/Avalanche/Source/AvalancheMedia/Private/Playback/Nodes/Events/AvaPlaybackNodeTrigger.cpp
@@ -1,6 +1,9 @@
 // Copyright Epic Games, Inc. All Rights Reserved.
 
 #include "Playback/Nodes/Events/AvaPlaybackNodeTrigger.h"
+#include "AvaMediaDefines.h"
+#include "Internationalization/Text.h"
+#include "Playback/Nodes/AvaPlaybackNode.h"
 
 FText UAvaPlaybackNodeTrigger::GetNodeCategoryText() const
 {
diff --git a/Engine/Plugins/Experimental/Avalanche/Source/AvalancheMedia/Private/Playback/Nodes/Events/FlowControl/AvaPlaybackNode_Hub.cpp b/Engine/Plugins/Experimental/Avalanche/Source/AvalancheMedia/Private/Playback/Nodes/Events/FlowControl/AvaPlaybackNode_Hub.cpp
--- a/Engine/Plugins/Experimental/Avalanche/Source/AvalancheMedia/Private/Playback/Nodes/Events/FlowControl/AvaPlaybackNode_Hub.cpp
+++ b/Engine/Plugins/Experimental/Avalanche/Source/AvalancheMedia/Private/Playback/Nodes/Events/FlowControl/AvaPlaybackNode_Hub.cpp
@@ -1,6 +1,8 @@
 // Copyright Epic Games, Inc. All Rights Reserved.
 
 #include "AvaPlaybackNode_Hub.h"
+#include "AvaMediaDefines.h"
+#include "Internationalization/Text.h"
 
 #define LOCTEXT_NAMESPACE "AvaPlaybackNode_Hub"
 
diff --git a/Engine/Plugins/Experimental/Avalanche/Source/AvalancheMedia/Private/Playback/Nodes/Events/FlowControl/AvaPlaybackNode_Hub.h b/Engine/Plugins/Experimental/Avalanche/Source/AvalancheMedia/Private/Playback/Nodes/Events/FlowControl/AvaPlaybackNode_Hub.h
--- a/Engine/Plugins/Experimental/Avalanche/Source/AvalancheMedia/Private/Playback/Nodes/Events/FlowControl/AvaPlaybackNode_Hub.h
+++ b/Engine/Plugins/Experimental/Avalanche/Source/AvalancheMedia/Private/Playback/Nodes/Events/FlowControl/AvaPlaybackNode_Hub.h
@@ -2,17 +2,18 @@
 
 #pragma once
 
-#include "CoreMinimal.h"
 #include "Playback/Nodes/Events/AvaPlaybackNodeFlow.h"
 #include "AvaPlaybackNode_Hub.generated.h"
 
+class FText;
+struct FAvaPlaybackEventParameters;
+
 UCLASS()
 class AVALANCHEMEDIA_API UAvaPlaybackNode_Hub : public UAvaPlaybackNodeFlow
 {
 	GENERATED_BODY()
 
 public:
-	
 	virtual FText GetNodeDisplayNameText() const override;
 	virtual FText GetNodeTooltipText() const override;
 	virtual int32 GetMaxChildNodes() const override { return MaxAllowedChildNodes; }
diff --git a/Engine/Plugins/Experimental/Avalanche/Source/AvalancheMedia/Private/Playback/Nodes/Events/Triggers/AvaPlaybackNode_BeginPlay.cpp b/Engine/Plugins/Experimental/Avalanche/Source/AvalancheMedia/Private/Playback/Nodes/Events/Triggers/AvaPlaybackNode_BeginPlay.cpp
--- a/Engine/Plugins/Experimental/Avalanche/Source/AvalancheMedia/Private/Playback/Nodes/Events/Triggers/AvaPlaybackNode_BeginPlay.cpp
+++ b/Engine/Plugins/Experimental/Avalanche/Source/AvalancheMedia/Private/Playback/Nodes/Events/Triggers/AvaPlaybackNode_BeginPlay.cpp
@@ -1,6 +1,7 @@
 // Copyright Epic Games, Inc. All Rights Reserved.
 
 #include "AvaPlaybackNode_BeginPlay.h"
+#include "Internationalization/Text.h"
 
 #define LOCTEXT_NAMESPACE "AvaPlaybackNode_BeginPlay"
 
diff --git a/Engine/Plugins/Experimental/Avalanche/Source/AvalancheMedia/Private/Playback/Nodes/Events/Triggers/AvaPlaybackNode_BeginPlay.h b/Engine/Plugins/Experimental/Avalanche/Source/AvalancheMedia/Private/Playback/Nodes/Events/Triggers/AvaPlaybackNode_BeginPlay.h
--- a/Engine/Plugins/Experimental/Avalanche/Source/AvalancheMedia/Private/Playback/Nodes/Events/Triggers/AvaPlaybackNode_BeginPlay.h
+++ b/Engine/Plugins/Experimental/Avalanche/Source/AvalancheMedia/Private/Playback/Nodes/Events/Triggers/AvaPlaybackNode_BeginPlay.h
@@ -2,19 +2,18 @@
 
 #pragma once
 
-#include "CoreMinimal.h"
 #include "Playback/Nodes/Events/AvaPlaybackNodeTrigger.h"
 #include "AvaPlaybackNode_BeginPlay.generated.h"
 
+class FText;
+
 UCLASS()
 class AVALANCHEMEDIA_API UAvaPlaybackNode_BeginPlay : public UAvaPlaybackNodeTrigger
 {
 	GENERATED_BODY()
 
 public:
-
 	virtual FText GetNodeDisplayNameText() const override;
 	virtual FText GetNodeTooltipText() const override;
 	virtual void NotifyPlaybackStateChanged(bool bPlaying) override;
-	
 };
diff --git a/Engine/Plugins/Experimental/Avalanche/Source/AvalancheMedia/Private/Playback/Nodes/Events/Triggers/AvaPlaybackNode_KeyInput.cpp b/Engine/Plugins/Experimental/Avalanche/Source/AvalancheMedia/Private/Playback/Nodes/Events/Triggers/AvaPlaybackNode_KeyInput.cpp
--- a/Engine/Plugins/Experimental/Avalanche/Source/AvalancheMedia/Private/Playback/Nodes/Events/Triggers/AvaPlaybackNode_KeyInput.cpp
+++ b/Engine/Plugins/Experimental/Avalanche/Source/AvalancheMedia/Private/Playback/Nodes/Events/Triggers/AvaPlaybackNode_KeyInput.cpp
@@ -4,6 +4,7 @@
 
 #include "Framework/Application/SlateApplication.h"
 #include "Input/Events.h"
+#include "Internationalization/Text.h"
 #include "Playback/AvaPlaybackGraph.h"
 
 #define LOCTEXT_NAMESPACE "AvaPlaybackNode_KeyInput"
@@ -70,7 +71,7 @@ FText UAvaPlaybackNode_KeyInput::GetNodeTooltipText() const
 
 bool UAvaPlaybackNode_KeyInput::HandleKeyDownEvent(FSlateApplication& SlateApp, const FKeyEvent& InKeyEvent)
 {
-	//Dont process input while not Playing
+	// Don't process input while not Playing
 	UAvaPlaybackGraph* const Playback = GetPlayback();
 	if (!Playback || !Playback->IsPlaying())
 	{
diff --git a/Engine/Plugins/Experimental/Avalanche/Source/AvalancheMedia/Private/Playback/Nodes/Events/Triggers/AvaPlaybackNode_KeyInput.h b/Engine/Plugins/Experimental/Avalanche/Source/AvalancheMedia/Private/Playback/Nodes/Events/Triggers/AvaPlaybackNode_KeyInput.h
--- a/Engine/Plugins/Experimental/Avalanche/Source/AvalancheMedia/Private/Playback/Nodes/Events/Triggers/AvaPlaybackNode_KeyInput.h
+++ b/Engine/Plugins/Experimental/Avalanche/Source/AvalancheMedia/Private/Playback/Nodes/Events/Triggers/AvaPlaybackNode_KeyInput.h
@@ -2,12 +2,15 @@
 
 #pragma once
 
-#include "CoreMinimal.h"
 #include "Framework/Application/IInputProcessor.h"
 #include "Framework/Commands/InputChord.h"
 #include "Playback/Nodes/Events/AvaPlaybackNodeTrigger.h"
+#include "Templates/SharedPointer.h"
 #include "AvaPlaybackNode_KeyInput.generated.h"
 
+class ICursor;
+class FSlateApplication;
+class FText;
 struct FKeyEvent;
 
 UCLASS()
@@ -18,18 +21,15 @@ class AVALANCHEMEDIA_API UAvaPlaybackNode_KeyInput : public UAvaPlaybackNodeTrig
 	class FEventInputProcessor: public IInputProcessor
 	{
 	public:
-		
 		FEventInputProcessor(UAvaPlaybackNode_KeyInput* InNode) : Node(InNode) {}
 		virtual void Tick(const float DeltaTime, FSlateApplication& SlateApp, TSharedRef<ICursor> Cursor) override {}
 		virtual bool HandleKeyDownEvent(FSlateApplication& SlateApp, const FKeyEvent& InKeyEvent) override;
 		
 	protected:
-		
-		UAvaPlaybackNode_KeyInput* Node;
+		UAvaPlaybackNode_KeyInput* Node = nullptr;
 	};
 
 public:
-	
 	virtual void PostAllocateNode() override;
 	virtual void BeginDestroy() override;
 
@@ -39,8 +39,7 @@ public:
 	bool HandleKeyDownEvent(FSlateApplication& SlateApp, const FKeyEvent& InKeyEvent);
 
 protected:
-
-	TSharedPtr<FEventInputProcessor> InputProcessor;
+	TSharedPtr<UAvaPlaybackNode_KeyInput::FEventInputProcessor> InputProcessor;
 	
 	UPROPERTY(EditAnywhere, Category = "Motion Design", meta = (ShowOnlyInnerProperties))
 	FInputChord InputChord;
diff --git a/Engine/Plugins/Experimental/Avalanche/Source/AvalancheMedia/Public/AvaMediaRenderTargetUtils.h b/Engine/Plugins/Experimental/Avalanche/Source/AvalancheMedia/Public/AvaMediaRenderTargetUtils.h
--- a/Engine/Plugins/Experimental/Avalanche/Source/AvalancheMedia/Public/AvaMediaRenderTargetUtils.h
+++ b/Engine/Plugins/Experimental/Avalanche/Source/AvalancheMedia/Public/AvaMediaRenderTargetUtils.h
@@ -2,11 +2,14 @@
 
 #pragma once
 
-#include "CoreMinimal.h"
+#include "HAL/Platform.h"
 #include "Math/MathFwd.h"
 #include "PixelFormat.h"
 
+class FName;
 class UTextureRenderTarget2D;
+struct FLinearColor;
+enum EPixelFormat : uint8;
 
 namespace UE::AvaMediaRenderTargetUtils
 {
diff --git a/Engine/Plugins/Experimental/Avalanche/Source/AvalancheMedia/Public/AvaMediaSettings.h b/Engine/Plugins/Experimental/Avalanche/Source/AvalancheMedia/Public/AvaMediaSettings.h
--- a/Engine/Plugins/Experimental/Avalanche/Source/AvalancheMedia/Public/AvaMediaSettings.h
+++ b/Engine/Plugins/Experimental/Avalanche/Source/AvalancheMedia/Public/AvaMediaSettings.h
@@ -2,11 +2,15 @@
 
 #pragma once
 
-#include "CoreMinimal.h"
+#include "Containers/Array.h"
+#include "Containers/EnumAsByte.h"
+#include "Containers/StringFwd.h"
 #include "Engine/DeveloperSettings.h"
 #include "Framework/AvaInstanceSettings.h"
 #include "Logging/LogVerbosity.h"
+#include "Math/MathFwd.h"
 #include "PixelFormat.h"
+#include "UObject/SoftObjectPtr.h"
 #include "AvaMediaSettings.generated.h"
 
 class UUserWidget;
diff --git a/Engine/Plugins/Experimental/Avalanche/Source/AvalancheMedia/Public/Broadcast/AvaBroadcast.h b/Engine/Plugins/Experimental/Avalanche/Source/AvalancheMedia/Public/Broadcast/AvaBroadcast.h
--- a/Engine/Plugins/Experimental/Avalanche/Source/AvalancheMedia/Public/Broadcast/AvaBroadcast.h
+++ b/Engine/Plugins/Experimental/Avalanche/Source/AvalancheMedia/Public/Broadcast/AvaBroadcast.h
@@ -4,10 +4,14 @@
 
 #include "AvaBroadcastProfile.h"
 #include "AvaMediaDefines.h"
-#include "CoreMinimal.h"
+#include "Containers/Array.h"
+#include "Containers/StringFwd.h"
 #include "UObject/Object.h"
 #include "AvaBroadcast.generated.h"
 
+class FName;
+class FDelegateHandle;
+
 DECLARE_MULTICAST_DELEGATE_OneParam(FOnAvaBroadcastChanged, EAvaBroadcastChange)
 DECLARE_MULTICAST_DELEGATE_OneParam(FOnAvaBroadcastChannelsListChanged, const FAvaBroadcastProfile& /*InProfile*/);
 
@@ -22,7 +26,6 @@ class AVALANCHEMEDIA_API UAvaBroadcast : public UObject
 	GENERATED_BODY()
 	
 public:
-
 	static UAvaBroadcast& Get();
 
 	virtual void BeginDestroy() override;
@@ -92,7 +95,6 @@ public:
 	void RemoveChangeListener(const void* InUserObject);
 
 public:
-
 	int32 GetChannelNameCount() const;
 
 	/**
@@ -151,7 +153,6 @@ public:
 	FOnAvaBroadcastChannelsListChanged& GetOnChannelsListChanged() { return OnChannelsListChanged; }
 
 protected:
-
 	// Todo: Proposal - make the events static and pass in the UAvaBroadcast that caused the event.
 	// Reason: Attempt at gradually getting rid of the global UAvaBroadcast object. Channel events are also global.
 	FOnAvaBroadcastChanged OnBroadcastChanged;
@@ -175,7 +176,6 @@ protected:
 	TArray<int32> BuildChannelIndices() const;
 
 protected:
-	
 	UPROPERTY()
 	FName CurrentProfile;
 
diff --git a/Engine/Plugins/Experimental/Avalanche/Source/AvalancheMedia/Public/Broadcast/AvaBroadcastDevice.h b/Engine/Plugins/Experimental/Avalanche/Source/AvalancheMedia/Public/Broadcast/AvaBroadcastDevice.h
--- a/Engine/Plugins/Experimental/Avalanche/Source/AvalancheMedia/Public/Broadcast/AvaBroadcastDevice.h
+++ b/Engine/Plugins/Experimental/Avalanche/Source/AvalancheMedia/Public/Broadcast/AvaBroadcastDevice.h
@@ -2,8 +2,11 @@
 
 #pragma once
 
-#include "CoreMinimal.h"
+#include "AvaBroadcastDevice.h"
+#include "Containers/UnrealString.h"
 #include "MediaIOCoreDefinitions.h"
+#include "Internationalization/Text.h"
+#include "UObject/NameTypes.h"
 
 #define LOCTEXT_NAMESPACE "AvaBroadcastDevice"
 
@@ -47,7 +50,6 @@ struct FAvaBroadcastDevice
 	const FString& GetServerName() const { return ServerName;}
 	
 protected:
-	
 	FMediaIODevice Device;
 
 	/** Name of the device provider for this device. */
diff --git a/Engine/Plugins/Experimental/Avalanche/Source/AvalancheMedia/Public/Broadcast/Channel/AvaBroadcastMediaOutputInfo.h b/Engine/Plugins/Experimental/Avalanche/Source/AvalancheMedia/Public/Broadcast/Channel/AvaBroadcastMediaOutputInfo.h
--- a/Engine/Plugins/Experimental/Avalanche/Source/AvalancheMedia/Public/Broadcast/Channel/AvaBroadcastMediaOutputInfo.h
+++ b/Engine/Plugins/Experimental/Avalanche/Source/AvalancheMedia/Public/Broadcast/Channel/AvaBroadcastMediaOutputInfo.h
@@ -2,7 +2,9 @@
 
 #pragma once
 
-#include "CoreMinimal.h"
+#include "Containers/UnrealString.h"
+#include "Misc/Guid.h"
+#include "UObject/NameTypes.h"
 #include "AvaBroadcastMediaOutputInfo.generated.h"
 
 /**
@@ -17,7 +19,7 @@ struct AVALANCHEMEDIA_API FAvaBroadcastMediaOutputInfo
 public:
 	/**
 	 * Unique identifier for this output.
-	 * Allows easier management fo client/server status and configuration replication.
+	 * Allows easier management for client/server status and configuration replication.
 	 */
 	UPROPERTY()
 	FGuid Guid;
@@ -54,4 +56,4 @@ public:
 
 	/** Determines if the given server name is considered "remote", i.e. from a different instance. */
 	static bool IsRemote(const FString& InServerName);
-};
\ No newline at end of file
+};
diff --git a/Engine/Plugins/Experimental/Avalanche/Source/AvalancheMedia/Public/Broadcast/Channel/AvaBroadcastOutputChannel.h b/Engine/Plugins/Experimental/Avalanche/Source/AvalancheMedia/Public/Broadcast/Channel/AvaBroadcastOutputChannel.h
--- a/Engine/Plugins/Experimental/Avalanche/Source/AvalancheMedia/Public/Broadcast/Channel/AvaBroadcastOutputChannel.h
+++ b/Engine/Plugins/Experimental/Avalanche/Source/AvalancheMedia/Public/Broadcast/Channel/AvaBroadcastOutputChannel.h
@@ -4,20 +4,27 @@
 
 #include "AvaBroadcastMediaOutputInfo.h"
 #include "AvaMediaDefines.h"
-#include "CoreMinimal.h"
+#include "Containers/Array.h"
+#include "Containers/Map.h"
+#include "Containers/StringFwd.h"
 #include "Engine/TextureRenderTarget2D.h"
 #include "Slate/SAvaBroadcastPlaceholderWidget.h"
 #include "Viewport/AvaViewportQualitySettings.h"
+#include "UObject/ObjectPtr.h"
 #include "AvaBroadcastOutputChannel.generated.h"
 
 class FAudioDeviceHandle;
+class FName;
 class FWidgetRenderer;
 class SVirtualWindow;
+class SWidget;
 class UAvaPlayableGroup;
 class UMediaCapture;
 class UMediaOutput;
+class UTextureRenderTarget2D;
 class UUserWidget;
 struct FAvaBroadcastProfile;
+enum EPixelFormat : uint8;
 
 /**
  *	Creates a bridge between the UMediaCapture's event
diff --git a/Engine/Plugins/Experimental/Avalanche/Source/AvalancheMedia/Public/Broadcast/Channel/Slate/SAvaBroadcastPlaceholderWidget.h b/Engine/Plugins/Experimental/Avalanche/Source/AvalancheMedia/Public/Broadcast/Channel/Slate/SAvaBroadcastPlaceholderWidget.h
--- a/Engine/Plugins/Experimental/Avalanche/Source/AvalancheMedia/Public/Broadcast/Channel/Slate/SAvaBroadcastPlaceholderWidget.h
+++ b/Engine/Plugins/Experimental/Avalanche/Source/AvalancheMedia/Public/Broadcast/Channel/Slate/SAvaBroadcastPlaceholderWidget.h
@@ -2,15 +2,15 @@
 
 #pragma once
 
-#include "CoreMinimal.h"
+#include "Templates/SharedPointer.h"
 #include "Widgets/SCompoundWidget.h"
 
+class FText;
 class STextBlock;
 
 class AVALANCHEMEDIA_API SAvaBroadcastPlaceholderWidget : public SCompoundWidget
 {
 public:
-	
 	SLATE_BEGIN_ARGS(SAvaBroadcastPlaceholderWidget) {}
 	SLATE_END_ARGS()
 
@@ -20,6 +20,5 @@ public:
 	void SetChannelName(const FText& InChannelName);
 	
 protected:
-
 	TSharedPtr<STextBlock> ChannelText;
 };
diff --git a/Engine/Plugins/Experimental/Avalanche/Source/AvalancheMedia/Public/Broadcast/OutputDevices/AvaBroadcastDisplayDeviceManager.h b/Engine/Plugins/Experimental/Avalanche/Source/AvalancheMedia/Public/Broadcast/OutputDevices/AvaBroadcastDisplayDeviceManager.h
--- a/Engine/Plugins/Experimental/Avalanche/Source/AvalancheMedia/Public/Broadcast/OutputDevices/AvaBroadcastDisplayDeviceManager.h
+++ b/Engine/Plugins/Experimental/Avalanche/Source/AvalancheMedia/Public/Broadcast/OutputDevices/AvaBroadcastDisplayDeviceManager.h
@@ -2,7 +2,8 @@
 
 #pragma once
 
-#include "CoreMinimal.h"
+#include "Containers/Array.h"
+#include "Containers/UnrealString.h"
 #include "GenericPlatform/GenericApplication.h"
 #include "Misc/FrameRate.h"
 
@@ -49,10 +50,10 @@ public:
 	static const TArray<FAvaBroadcastMonitorInfo>&  GetCachedMonitors(bool bForceUpdate = false);
 
 	/**
-	 *	Returns a string with a display name for the monitor.
+	 * Returns a string with a display name for the monitor.
 	 **/
 	static FString GetMonitorDisplayName(const FAvaBroadcastMonitorInfo& InMonitorInfo);
 	
 private:
 	static TArray<FAvaBroadcastMonitorInfo> CachedMonitorInfo;
-};
\ No newline at end of file
+};
diff --git a/Engine/Plugins/Experimental/Avalanche/Source/AvalancheMedia/Public/Broadcast/OutputDevices/Slate/SAvaBroadcastCaptureImage.h b/Engine/Plugins/Experimental/Avalanche/Source/AvalancheMedia/Public/Broadcast/OutputDevices/Slate/SAvaBroadcastCaptureImage.h
--- a/Engine/Plugins/Experimental/Avalanche/Source/AvalancheMedia/Public/Broadcast/OutputDevices/Slate/SAvaBroadcastCaptureImage.h
+++ b/Engine/Plugins/Experimental/Avalanche/Source/AvalancheMedia/Public/Broadcast/OutputDevices/Slate/SAvaBroadcastCaptureImage.h
@@ -2,9 +2,14 @@
 
 #pragma once
 
-#include "CoreMinimal.h"
+#include "Misc/Attribute.h"
 #include "Widgets/Images/SImage.h"
 
+class FSlateRect;
+class FSlateWindowElementList;
+class FWidgetStyle;
+struct FGeometry;
+
 class AVALANCHEMEDIA_API SAvaBroadcastCaptureImage : public SImage
 {
 public:
diff --git a/Engine/Plugins/Experimental/Avalanche/Source/AvalancheMedia/Public/IAvaMediaModule.h b/Engine/Plugins/Experimental/Avalanche/Source/AvalancheMedia/Public/IAvaMediaModule.h
--- a/Engine/Plugins/Experimental/Avalanche/Source/AvalancheMedia/Public/IAvaMediaModule.h
+++ b/Engine/Plugins/Experimental/Avalanche/Source/AvalancheMedia/Public/IAvaMediaModule.h
@@ -2,7 +2,7 @@
 
 #pragma once
 
-#include "CoreMinimal.h"
+#include "Containers/StringFwd.h"
 #include "Modules/ModuleInterface.h"
 #include "Modules/ModuleManager.h"
 
@@ -12,6 +12,7 @@ class FAvaPlaybackManager;
 class FAvaPlaybackServer;
 class FAvaRundownManagedInstanceCache;
 class FCommonViewportClient;
+class FName;
 class IAvaBroadcastDeviceProviderProxyManager;
 class IAvaBroadcastSettings;
 class IAvaMediaSyncProvider;
@@ -36,13 +37,13 @@ class AVALANCHEMEDIA_API IAvaMediaModule : public IModuleInterface
 public:
 	static bool IsModuleLoaded()
 	{
-		static const FName ModuleName = TEXT( "AvalancheMedia" );
+		static const FName ModuleName = TEXT("AvalancheMedia");
 		return FModuleManager::Get().IsModuleLoaded(ModuleName);
 	}
 	
 	static IAvaMediaModule& Get()
 	{
-		static const FName ModuleName = TEXT( "AvalancheMedia" );
+		static const FName ModuleName = TEXT("AvalancheMedia");
 		return FModuleManager::LoadModuleChecked<IAvaMediaModule>(ModuleName);
 	};
 
diff --git a/Engine/Plugins/Experimental/Avalanche/Source/AvalancheMedia/Public/Playback/IAvaPlaybackGraphEditor.h b/Engine/Plugins/Experimental/Avalanche/Source/AvalancheMedia/Public/Playback/IAvaPlaybackGraphEditor.h
--- a/Engine/Plugins/Experimental/Avalanche/Source/AvalancheMedia/Public/Playback/IAvaPlaybackGraphEditor.h
+++ b/Engine/Plugins/Experimental/Avalanche/Source/AvalancheMedia/Public/Playback/IAvaPlaybackGraphEditor.h
@@ -4,18 +4,19 @@
 
 #if WITH_EDITOR
 
-#include "CoreMinimal.h"
+#include "Containers/Set.h"
+#include "Math/MathFwd.h"
 
 class UEdGraphNode;
 class UEdGraph;
 class UAvaPlaybackNode;
 class UAvaPlaybackGraph;
+class UObject;
 class FSlateRect;
 
 class IAvaPlaybackGraphEditor
 {
 public:
-	
 	virtual ~IAvaPlaybackGraphEditor() = default;
 	
 	virtual UEdGraph* CreatePlaybackGraph(UAvaPlaybackGraph* InPlayback) = 0;
@@ -42,7 +43,6 @@ public:
 
 	virtual bool CanPasteNodes() const = 0;
 	virtual void PasteNodesHere(const FVector2D& Location) = 0;
-	
 };
 
 #endif
diff --git a/Engine/Plugins/Experimental/Avalanche/Source/AvalancheMedia/Public/Playback/Nodes/AvaPlaybackNode.h b/Engine/Plugins/Experimental/Avalanche/Source/AvalancheMedia/Public/Playback/Nodes/AvaPlaybackNode.h
--- a/Engine/Plugins/Experimental/Avalanche/Source/AvalancheMedia/Public/Playback/Nodes/AvaPlaybackNode.h
+++ b/Engine/Plugins/Experimental/Avalanche/Source/AvalancheMedia/Public/Playback/Nodes/AvaPlaybackNode.h
@@ -3,12 +3,18 @@
 #pragma once
 
 #include "AvaMediaDefines.h"
-#include "CoreMinimal.h"
+#include "Containers/Array.h"
+#include "Delegates/IDelegateInstance.h"
 #include "UObject/Object.h"
+#include "UObject/ObjectPtr.h"
 #include "AvaPlaybackNode.generated.h"
 
+class FName;
+class FText;
+class FReferenceCollector;
 class UAvaPlaybackGraph;
 class UEdGraphNode;
+struct FPropertyChangedEvent;
 
 /**
  * Base Class for all the Nodes found in the Motion Design Playback Graph
@@ -150,7 +156,6 @@ public:
 #endif
 	
 protected:
-
 	UPROPERTY()
 	TArray<TObjectPtr<UAvaPlaybackNode>> ChildNodes;
 
diff --git a/Engine/Plugins/Experimental/Avalanche/Source/AvalancheMedia/Public/Playback/Nodes/AvaPlaybackNodeCombiner.h b/Engine/Plugins/Experimental/Avalanche/Source/AvalancheMedia/Public/Playback/Nodes/AvaPlaybackNodeCombiner.h
--- a/Engine/Plugins/Experimental/Avalanche/Source/AvalancheMedia/Public/Playback/Nodes/AvaPlaybackNodeCombiner.h
+++ b/Engine/Plugins/Experimental/Avalanche/Source/AvalancheMedia/Public/Playback/Nodes/AvaPlaybackNodeCombiner.h
@@ -3,9 +3,15 @@
 #pragma once
 
 #include "AvaPlaybackNode.h"
-#include "CoreMinimal.h"
 #include "AvaPlaybackNodeCombiner.generated.h"
 
+class FText;
+struct FAvaPlaybackChannelParameters;
+
+#if WITH_EDITOR
+class FName;
+#endif
+
 UCLASS()
 class AVALANCHEMEDIA_API UAvaPlaybackNodeCombiner : public UAvaPlaybackNode
 {
@@ -25,7 +31,6 @@ class AVALANCHEMEDIA_API UAvaPlaybackNodeCombiner : public UAvaPlaybackNode
 #endif
 	
 protected:
-
 	UPROPERTY(EditAnywhere, Category = "Motion Design")
 	TArray<bool> EnabledIndices;	
 };
diff --git a/Engine/Plugins/Experimental/Avalanche/Source/AvalancheMedia/Public/Playback/Nodes/AvaPlaybackNodeLevelPlayer.h b/Engine/Plugins/Experimental/Avalanche/Source/AvalancheMedia/Public/Playback/Nodes/AvaPlaybackNodeLevelPlayer.h
--- a/Engine/Plugins/Experimental/Avalanche/Source/AvalancheMedia/Public/Playback/Nodes/AvaPlaybackNodeLevelPlayer.h
+++ b/Engine/Plugins/Experimental/Avalanche/Source/AvalancheMedia/Public/Playback/Nodes/AvaPlaybackNodeLevelPlayer.h
@@ -3,10 +3,12 @@
 #pragma once
 
 #include "AvaPlaybackNodePlayer.h"
-#include "CoreMinimal.h"
+#include "UObject/SoftObjectPtr.h"
+#include "UObject/NoExportTypes.h"
 #include "AvaPlaybackNodeLevelPlayer.generated.h"
 
 class UWorld;
+struct FAvaSoftAssetPtr;
 
 UCLASS()
 class AVALANCHEMEDIA_API UAvaPlaybackNodeLevelPlayer : public UAvaPlaybackNodePlayer
@@ -14,7 +16,6 @@ class AVALANCHEMEDIA_API UAvaPlaybackNodeLevelPlayer : public UAvaPlaybackNodePl
 	GENERATED_BODY()
 
 public:
-	
 	UAvaPlaybackNodeLevelPlayer();
 	
 	virtual void RefreshNode(bool bDryRunGraph) override;
@@ -30,8 +31,6 @@ public:
 	virtual FAvaSoftAssetPtr GetAssetPtr() const override;
 
 protected:
-	
 	UPROPERTY(EditAnywhere, Category = "Motion Design")
 	TSoftObjectPtr<UWorld> LevelAsset;
 };
-
diff --git a/Engine/Plugins/Experimental/Avalanche/Source/AvalancheMedia/Public/Playback/Nodes/AvaPlaybackNodePlayer.h b/Engine/Plugins/Experimental/Avalanche/Source/AvalancheMedia/Public/Playback/Nodes/AvaPlaybackNodePlayer.h
--- a/Engine/Plugins/Experimental/Avalanche/Source/AvalancheMedia/Public/Playback/Nodes/AvaPlaybackNodePlayer.h
+++ b/Engine/Plugins/Experimental/Avalanche/Source/AvalancheMedia/Public/Playback/Nodes/AvaPlaybackNodePlayer.h
@@ -3,10 +3,14 @@
 #pragma once
 
 #include "AvaPlaybackNode.h"
-#include "CoreMinimal.h"
+#include "Containers/Array.h"
+#include "Framework/AvaSoftAssetPtr.h"
 #include "AvaPlaybackNodePlayer.generated.h"
 
+class FText;
 class UTextureRenderTarget2D;
+struct FAvaPlaybackChannelParameters;
+struct FPropertyChangedEvent;
 
 UCLASS(Abstract)
 class AVALANCHEMEDIA_API UAvaPlaybackNodePlayer : public UAvaPlaybackNode
@@ -55,4 +59,3 @@ protected:
 	TArray<int32> ChannelIndices;
 	TArray<int32> LastTickChannelIndices;
 };
-
diff --git a/Engine/Plugins/Experimental/Avalanche/Source/AvalancheMedia/Public/Playback/Nodes/AvaPlaybackNodeSwitcher.h b/Engine/Plugins/Experimental/Avalanche/Source/AvalancheMedia/Public/Playback/Nodes/AvaPlaybackNodeSwitcher.h
--- a/Engine/Plugins/Experimental/Avalanche/Source/AvalancheMedia/Public/Playback/Nodes/AvaPlaybackNodeSwitcher.h
+++ b/Engine/Plugins/Experimental/Avalanche/Source/AvalancheMedia/Public/Playback/Nodes/AvaPlaybackNodeSwitcher.h
@@ -3,9 +3,12 @@
 #pragma once
 
 #include "AvaPlaybackNode.h"
-#include "CoreMinimal.h"
 #include "AvaPlaybackNodeSwitcher.generated.h"
 
+class FName;
+class FText;
+struct FAvaPlaybackChannelParameters;
+
 UCLASS()
 class AVALANCHEMEDIA_API UAvaPlaybackNodeSwitcher : public UAvaPlaybackNode
 {
@@ -25,7 +28,6 @@ class AVALANCHEMEDIA_API UAvaPlaybackNodeSwitcher : public UAvaPlaybackNode
 #endif
 	
 protected:
-
 	UPROPERTY(EditAnywhere, Category = "Motion Design")
 	bool bEnabled = true;
 	
diff --git a/Engine/Plugins/Experimental/Avalanche/Source/AvalancheMedia/Public/Playback/Nodes/Events/Actions/AvaPlaybackAnimations.h b/Engine/Plugins/Experimental/Avalanche/Source/AvalancheMedia/Public/Playback/Nodes/Events/Actions/AvaPlaybackAnimations.h
--- a/Engine/Plugins/Experimental/Avalanche/Source/AvalancheMedia/Public/Playback/Nodes/Events/Actions/AvaPlaybackAnimations.h
+++ b/Engine/Plugins/Experimental/Avalanche/Source/AvalancheMedia/Public/Playback/Nodes/Events/Actions/AvaPlaybackAnimations.h
@@ -3,7 +3,8 @@
 #pragma once
 
 #include "AvaSequenceShared.h"
-#include "CoreMinimal.h"
+#include "Containers/Set.h"
+#include "UObject/NameTypes.h"
 #include "AvaPlaybackAnimations.generated.h"
 
 UENUM()
@@ -78,4 +79,4 @@ struct FAvaPlaybackAnimations
 	
 	UPROPERTY(EditAnywhere, Category = "Motion Design")
 	TSet<FAvaPlaybackAnimPlaySettings> AvailableAnimations;
-};
\ No newline at end of file
+};
diff --git a/Engine/Plugins/Experimental/Avalanche/Source/AvalancheMedia/Public/Playback/Nodes/Events/AvaPlaybackNodeAction.h b/Engine/Plugins/Experimental/Avalanche/Source/AvalancheMedia/Public/Playback/Nodes/Events/AvaPlaybackNodeAction.h
--- a/Engine/Plugins/Experimental/Avalanche/Source/AvalancheMedia/Public/Playback/Nodes/Events/AvaPlaybackNodeAction.h
+++ b/Engine/Plugins/Experimental/Avalanche/Source/AvalancheMedia/Public/Playback/Nodes/Events/AvaPlaybackNodeAction.h
@@ -3,9 +3,10 @@
 #pragma once
 
 #include "AvaPlaybackNodeEvent.h"
-#include "CoreMinimal.h"
 #include "AvaPlaybackNodeAction.generated.h"
 
+class FText;
+
 /*
  * Node part of the Event Nodes that handles Logic Execution
  */
@@ -15,7 +16,6 @@ class AVALANCHEMEDIA_API UAvaPlaybackNodeAction : public UAvaPlaybackNodeEvent
 	GENERATED_BODY()
 
 public:
-
 	virtual FText GetNodeCategoryText() const override final;
 	
 	virtual int32 GetMinChildNodes() const override final { return 1; }
diff --git a/Engine/Plugins/Experimental/Avalanche/Source/AvalancheMedia/Public/Playback/Nodes/Events/AvaPlaybackNodeEvent.h b/Engine/Plugins/Experimental/Avalanche/Source/AvalancheMedia/Public/Playback/Nodes/Events/AvaPlaybackNodeEvent.h
--- a/Engine/Plugins/Experimental/Avalanche/Source/AvalancheMedia/Public/Playback/Nodes/Events/AvaPlaybackNodeEvent.h
+++ b/Engine/Plugins/Experimental/Avalanche/Source/AvalancheMedia/Public/Playback/Nodes/Events/AvaPlaybackNodeEvent.h
@@ -2,10 +2,12 @@
 
 #pragma once
 
-#include "CoreMinimal.h"
 #include "Playback/Nodes/AvaPlaybackNode.h"
 #include "AvaPlaybackNodeEvent.generated.h"
 
+struct FAvaPlaybackChannelParameters;
+struct FAvaPlaybackEventParameters;
+
 /*
  * Base Class for all Event Nodes
  */
@@ -15,7 +17,6 @@ class AVALANCHEMEDIA_API UAvaPlaybackNodeEvent : public UAvaPlaybackNode
 	GENERATED_BODY()
 
 public:
-
 	//UAvaPlaybackNode Interface
 	virtual void Tick(float DeltaTime, FAvaPlaybackChannelParameters& ChannelParameters) override final {}
 	//~UAvaPlaybackNode Interface
diff --git a/Engine/Plugins/Experimental/Avalanche/Source/AvalancheMedia/Public/Playback/Nodes/Events/AvaPlaybackNodeFlow.h b/Engine/Plugins/Experimental/Avalanche/Source/AvalancheMedia/Public/Playback/Nodes/Events/AvaPlaybackNodeFlow.h
--- a/Engine/Plugins/Experimental/Avalanche/Source/AvalancheMedia/Public/Playback/Nodes/Events/AvaPlaybackNodeFlow.h
+++ b/Engine/Plugins/Experimental/Avalanche/Source/AvalancheMedia/Public/Playback/Nodes/Events/AvaPlaybackNodeFlow.h
@@ -3,9 +3,10 @@
 #pragma once
 
 #include "AvaPlaybackNodeEvent.h"
-#include "CoreMinimal.h"
 #include "AvaPlaybackNodeFlow.generated.h"
 
+class FText;
+
 /*
  * Node part of the Event Nodes that Controls the Flow of Events
  */
@@ -15,7 +16,6 @@ class AVALANCHEMEDIA_API UAvaPlaybackNodeFlow : public UAvaPlaybackNodeEvent
 	GENERATED_BODY()
 
 public:
-
 	virtual FText GetNodeCategoryText() const override final;
 	
 	virtual int32 GetMinChildNodes() const override final { return 1; }
diff --git a/Engine/Plugins/Experimental/Avalanche/Source/AvalancheMedia/Public/Playback/Nodes/Events/AvaPlaybackNodeTrigger.h b/Engine/Plugins/Experimental/Avalanche/Source/AvalancheMedia/Public/Playback/Nodes/Events/AvaPlaybackNodeTrigger.h
--- a/Engine/Plugins/Experimental/Avalanche/Source/AvalancheMedia/Public/Playback/Nodes/Events/AvaPlaybackNodeTrigger.h
+++ b/Engine/Plugins/Experimental/Avalanche/Source/AvalancheMedia/Public/Playback/Nodes/Events/AvaPlaybackNodeTrigger.h
@@ -3,9 +3,11 @@
 #pragma once
 
 #include "AvaPlaybackNodeEvent.h"
-#include "CoreMinimal.h"
 #include "AvaPlaybackNodeTrigger.generated.h"
 
+class FText;
+struct FAvaPlaybackEventParameters;
+
 /*
  * Node part of the Event Nodes that Listens to an Event and Triggers the Event Flow
  */
@@ -15,7 +17,6 @@ class AVALANCHEMEDIA_API UAvaPlaybackNodeTrigger : public UAvaPlaybackNodeEvent
 	GENERATED_BODY()
 
 public:
-
 	virtual FText GetNodeCategoryText() const override final;
 	
 	virtual int32 GetMinChildNodes() const override final { return 0; }
@@ -26,6 +27,5 @@ public:
 	virtual void TickEvent(float DeltaTime, FAvaPlaybackEventParameters& OutEventParameters) override;
 	
 private:
-
 	bool bEventTriggered = false;
 };
diff --git a/Engine/Plugins/Experimental/Avalanche/Source/AvalancheMediaEditor/Private/Broadcast/AvaBroadcastEditor.cpp b/Engine/Plugins/Experimental/Avalanche/Source/AvalancheMediaEditor/Private/Broadcast/AvaBroadcastEditor.cpp
--- a/Engine/Plugins/Experimental/Avalanche/Source/AvalancheMediaEditor/Private/Broadcast/AvaBroadcastEditor.cpp
+++ b/Engine/Plugins/Experimental/Avalanche/Source/AvalancheMediaEditor/Private/Broadcast/AvaBroadcastEditor.cpp
@@ -1,19 +1,29 @@
 // Copyright Epic Games, Inc. All Rights Reserved.
 
 #include "Broadcast/AvaBroadcastEditor.h"
-
 #include "AppModes/AvaBroadcastDefaultMode.h"
 #include "AvaMediaSettings.h"
 #include "Broadcast/AvaBroadcast.h"
+#include "Broadcast/ChannelGrid/AvaBroadcastOutputTileItem.h"
+#include "Containers/UnrealString.h"
 #include "ChannelGrid/Slate/SAvaBroadcastProfileEntry.h"
 #include "Framework/Application/SlateApplication.h"
 #include "Framework/MultiBox/MultiBoxBuilder.h"
+#include "Framework/MultiBox/MultiBoxExtender.h"
 #include "IAvaMediaEditorModule.h"
 #include "IAvaMediaModule.h"
+#include "Input/Reply.h"
+#include "Internationalization/Text.h"
+#include "MediaOutput.h"
 #include "Misc/MessageDialog.h"
 #include "Playback/IAvaPlaybackClient.h"
 #include "ScopedTransaction.h"
+#include "Styling/SlateBrush.h"
+#include "UObject/NameTypes.h"
 #include "WorkflowOrientedApp/ApplicationMode.h"
+#include "WorkflowOrientedApp/WorkflowCentricApplication.h"
+#include "UObject/NoExportTypes.h"
+#include "UObject/UnrealType.h"
 
 #define LOCTEXT_NAMESPACE "AvaBroadcastEditor"
 
diff --git a/Engine/Plugins/Experimental/Avalanche/Source/AvalancheMediaEditor/Private/Broadcast/ChannelGrid/AvaBroadcastOutputTileItem.cpp b/Engine/Plugins/Experimental/Avalanche/Source/AvalancheMediaEditor/Private/Broadcast/ChannelGrid/AvaBroadcastOutputTileItem.cpp
--- a/Engine/Plugins/Experimental/Avalanche/Source/AvalancheMediaEditor/Private/Broadcast/ChannelGrid/AvaBroadcastOutputTileItem.cpp
+++ b/Engine/Plugins/Experimental/Avalanche/Source/AvalancheMediaEditor/Private/Broadcast/ChannelGrid/AvaBroadcastOutputTileItem.cpp
@@ -1,14 +1,26 @@
 // Copyright Epic Games, Inc. All Rights Reserved.
 
 #include "Broadcast/ChannelGrid/AvaBroadcastOutputTileItem.h"
-
+#include "AvaMediaDefines.h"
 #include "AvaMediaEditorStyle.h"
 #include "Broadcast/AvaBroadcast.h"
+#include "Broadcast/Channel/AvaBroadcastOutputChannel.h"
 #include "Broadcast/OutputDevices/AvaBroadcastOutputUtils.h"
 #include "ClassIconFinder.h"
 #include "DragDropOps/AvaBroadcastOutputTileItemDragDropOp.h"
+#include "Input/Events.h"
+#include "Input/Reply.h"
+#include "Layout/Geometry.h"
 #include "MediaOutput.h"
+#include "Rundown/ChannelStatus/SAvaRundownChannelStatus.h"
+#include "Styling/SlateBrush.h"
+#include "Widgets/Images/SImage.h"
 #include "Widgets/Layout/SScaleBox.h"
+#include "Widgets/SBoxPanel.h"
+#include "Widgets/SWidget.h"
+#include "Widgets/Text/STextBlock.h"
+#include "UObject/NoExportTypes.h"
+#include "UObject/UnrealType.h"
 
 #define LOCTEXT_NAMESPACE "AvaBroadcastOutputTileItem"
 
diff --git a/Engine/Plugins/Experimental/Avalanche/Source/AvalancheMediaEditor/Private/Broadcast/ChannelGrid/DragDropOps/AvaBroadcastOutputTileItemDragDropOp.cpp b/Engine/Plugins/Experimental/Avalanche/Source/AvalancheMediaEditor/Private/Broadcast/ChannelGrid/DragDropOps/AvaBroadcastOutputTileItemDragDropOp.cpp
--- a/Engine/Plugins/Experimental/Avalanche/Source/AvalancheMediaEditor/Private/Broadcast/ChannelGrid/DragDropOps/AvaBroadcastOutputTileItemDragDropOp.cpp
+++ b/Engine/Plugins/Experimental/Avalanche/Source/AvalancheMediaEditor/Private/Broadcast/ChannelGrid/DragDropOps/AvaBroadcastOutputTileItemDragDropOp.cpp
@@ -1,10 +1,11 @@
 // Copyright Epic Games, Inc. All Rights Reserved.
 
 #include "AvaBroadcastOutputTileItemDragDropOp.h"
-
 #include "Broadcast/AvaBroadcast.h"
+#include "Input/Reply.h"
 #include "MediaOutput.h"
 #include "ScopedTransaction.h"
+#include "UObject/NameTypes.h"
 
 #define LOCTEXT_NAMESPACE "AvaBroadcastOutputTileItemDragDropOp"
 
diff --git a/Engine/Plugins/Experimental/Avalanche/Source/AvalancheMediaEditor/Private/Broadcast/ChannelGrid/DragDropOps/AvaBroadcastOutputTileItemDragDropOp.h b/Engine/Plugins/Experimental/Avalanche/Source/AvalancheMediaEditor/Private/Broadcast/ChannelGrid/DragDropOps/AvaBroadcastOutputTileItemDragDropOp.h
--- a/Engine/Plugins/Experimental/Avalanche/Source/AvalancheMediaEditor/Private/Broadcast/ChannelGrid/DragDropOps/AvaBroadcastOutputTileItemDragDropOp.h
+++ b/Engine/Plugins/Experimental/Avalanche/Source/AvalancheMediaEditor/Private/Broadcast/ChannelGrid/DragDropOps/AvaBroadcastOutputTileItemDragDropOp.h
@@ -3,13 +3,14 @@
 #pragma once
 
 #include "Broadcast/ChannelGrid/AvaBroadcastOutputTileItem.h"
-#include "CoreMinimal.h"
 #include "DragAndDrop/DecoratedDragDropOp.h"
 
+class FName;
+class FReply;
+
 class FAvaBroadcastOutputTileItemDragDropOp : public FDecoratedDragDropOp
 {
 public:
-
 	DRAG_DROP_OPERATOR_TYPE(FAvaOutputTileDragDropOp, FDecoratedDragDropOp)
 	
 	static TSharedRef<FAvaBroadcastOutputTileItemDragDropOp> New(const FAvaBroadcastOutputTileItemPtr& InOutputClassItem, bool bInIsDuplicating);
@@ -21,7 +22,6 @@ public:
 	FReply OnChannelDrop(FName InTargetChannelName);
 	
 protected:
-	
 	void Init(const FAvaBroadcastOutputTileItemPtr& InOutputTileItem, bool bInIsDuplicating);
 
 	//Keep Reference Count while Drag Dropping
diff --git a/Engine/Plugins/Experimental/Avalanche/Source/AvalancheMediaEditor/Private/Broadcast/ChannelGrid/Slate/SAvaBroadcastChannel.cpp b/Engine/Plugins/Experimental/Avalanche/Source/AvalancheMediaEditor/Private/Broadcast/ChannelGrid/Slate/SAvaBroadcastChannel.cpp
--- a/Engine/Plugins/Experimental/Avalanche/Source/AvalancheMediaEditor/Private/Broadcast/ChannelGrid/Slate/SAvaBroadcastChannel.cpp
+++ b/Engine/Plugins/Experimental/Avalanche/Source/AvalancheMediaEditor/Private/Broadcast/ChannelGrid/Slate/SAvaBroadcastChannel.cpp
@@ -1,22 +1,27 @@
 // Copyright Epic Games, Inc. All Rights Reserved.
 
 #include "SAvaBroadcastChannel.h"
-
-#include "AvaMediaDefines.h"
 #include "AvaMediaEditorStyle.h"
 #include "AvaMediaEditorUtils.h"
 #include "Broadcast/AvaBroadcast.h"
 #include "Broadcast/AvaBroadcastEditor.h"
+#include "Broadcast/Channel/AvaBroadcastOutputChannel.h"
 #include "Broadcast/ChannelGrid/DragDropOps/AvaBroadcastOutputTileItemDragDropOp.h"
 #include "Broadcast/OutputDevices/DragDropOps/AvaBroadcastOutputTreeItemDragDropOp.h"
 #include "Broadcast/OutputDevices/Slate/SAvaBroadcastCaptureImage.h"
 #include "Brushes/SlateImageBrush.h"
 #include "Framework/Commands/GenericCommands.h"
 #include "Framework/MultiBox/MultiBoxBuilder.h"
+#include "Input/DragAndDrop.h"
+#include "Input/Events.h"
+#include "Internationalization/Text.h"
 #include "IStructureDetailsView.h"
+#include "Layout/Geometry.h"
+#include "Layout/Visibility.h"
 #include "MediaOutput.h"
 #include "PropertyEditorModule.h"
 #include "ScopedTransaction.h"
+#include "Styling/SlateBrush.h"
 #include "Styling/SlateTypes.h"
 #include "Widgets/Input/SButton.h"
 #include "Widgets/Input/SMenuAnchor.h"
@@ -24,7 +29,10 @@
 #include "Widgets/Layout/SScrollBox.h"
 #include "Widgets/Layout/SSeparator.h"
 #include "Widgets/Text/SInlineEditableTextBlock.h"
+#include "Widgets/Views/ITableRow.h"
 #include "Widgets/Views/SListView.h"
+#include "Widgets/Views/STableViewBase.h"
+#include "UObject/NameTypes.h"
 
 #define LOCTEXT_NAMESPACE "SAvaBroadcastChannel"
 
diff --git a/Engine/Plugins/Experimental/Avalanche/Source/AvalancheMediaEditor/Private/Broadcast/ChannelGrid/Slate/SAvaBroadcastChannel.h b/Engine/Plugins/Experimental/Avalanche/Source/AvalancheMediaEditor/Private/Broadcast/ChannelGrid/Slate/SAvaBroadcastChannel.h
--- a/Engine/Plugins/Experimental/Avalanche/Source/AvalancheMediaEditor/Private/Broadcast/ChannelGrid/Slate/SAvaBroadcastChannel.h
+++ b/Engine/Plugins/Experimental/Avalanche/Source/AvalancheMediaEditor/Private/Broadcast/ChannelGrid/Slate/SAvaBroadcastChannel.h
@@ -2,21 +2,31 @@
 
 #pragma once
 
+#include "AvaMediaDefines.h"
 #include "Broadcast/Channel/AvaBroadcastOutputChannel.h"
 #include "Broadcast/ChannelGrid/AvaBroadcastOutputTileItem.h"
-#include "CoreMinimal.h"
+#include "Misc/Attribute.h"
+#include "Templates/SharedPointer.h"
+#include "Types/SlateEnums.h"
 #include "Widgets/SCompoundWidget.h"
 
 class FAvaBroadcastEditor;
+class FName;
+class FReply;
+class FText;
 class FUICommandList;
 class IStructureDetailsView;
 class ITableRow;
 class SAvaBroadcastChannel;
 class SInlineEditableTextBlock;
+class SlateBrush;
 class SMenuAnchor;
 class STableViewBase;
+class SWidget;
 class UMediaOutput;
+struct EVisibility;
 struct FAvaBroadcastOutputChannel;
+struct FGeometry;
 struct FSlateImageBrush;
 template<typename ItemType> class SListView;
 
@@ -26,7 +36,6 @@ DECLARE_DELEGATE_OneParam(FAvaBroadcastOnChannelMaximizeClicked, const TSharedRe
 class SAvaBroadcastChannel : public SCompoundWidget
 {
 public:
-	
 	SLATE_BEGIN_ARGS(SAvaBroadcastChannel) {}
 		SLATE_ARGUMENT(FName, ChannelName)
 		SLATE_ATTRIBUTE(bool, CanMaximize)
diff --git a/Engine/Plugins/Experimental/Avalanche/Source/AvalancheMediaEditor/Private/Broadcast/ChannelGrid/Slate/SAvaBroadcastChannels.cpp b/Engine/Plugins/Experimental/Avalanche/Source/AvalancheMediaEditor/Private/Broadcast/ChannelGrid/Slate/SAvaBroadcastChannels.cpp
--- a/Engine/Plugins/Experimental/Avalanche/Source/AvalancheMediaEditor/Private/Broadcast/ChannelGrid/Slate/SAvaBroadcastChannels.cpp
+++ b/Engine/Plugins/Experimental/Avalanche/Source/AvalancheMediaEditor/Private/Broadcast/ChannelGrid/Slate/SAvaBroadcastChannels.cpp
@@ -6,6 +6,7 @@
 #include "Broadcast/AvaBroadcast.h"
 #include "Broadcast/AvaBroadcastEditor.h"
 #include "Framework/MultiBox/MultiBoxBuilder.h"
+#include "Input/Reply.h"
 #include "SAvaBroadcastChannel.h"
 #include "ScopedTransaction.h"
 #include "Widgets/Layout/SGridPanel.h"
diff --git a/Engine/Plugins/Experimental/Avalanche/Source/AvalancheMediaEditor/Private/Broadcast/ChannelGrid/Slate/SAvaBroadcastChannels.h b/Engine/Plugins/Experimental/Avalanche/Source/AvalancheMediaEditor/Private/Broadcast/ChannelGrid/Slate/SAvaBroadcastChannels.h
--- a/Engine/Plugins/Experimental/Avalanche/Source/AvalancheMediaEditor/Private/Broadcast/ChannelGrid/Slate/SAvaBroadcastChannels.h
+++ b/Engine/Plugins/Experimental/Avalanche/Source/AvalancheMediaEditor/Private/Broadcast/ChannelGrid/Slate/SAvaBroadcastChannels.h
@@ -4,12 +4,15 @@
 
 #include "Animation/CurveSequence.h"
 #include "AvaMediaDefines.h"
-#include "CoreMinimal.h"
+#include "Containers/Map.h"
+#include "Templates/SharedPointer.h"
 #include "Widgets/SCompoundWidget.h"
 
 class FAvaBroadcastEditor;
+class FReply;
 class SAvaBroadcastChannel;
 class SGridPanel;
+class SWidget;
 class UAvaBroadcast;
 
 class SAvaBroadcastChannels : public SCompoundWidget
@@ -30,7 +33,6 @@ class SAvaBroadcastChannels : public SCompoundWidget
 	};
 	
 public:
-	
 	SLATE_BEGIN_ARGS(SAvaBroadcastChannels) {}
 	SLATE_END_ARGS()
 	
@@ -56,7 +58,6 @@ public:
 	virtual FReply OnMouseButtonDown(const FGeometry& MyGeometry, const FPointerEvent& MouseEvent) override;;
 	
 protected:
-	
 	TWeakObjectPtr<UAvaBroadcast> BroadcastWeak;
 	
 	TWeakPtr<FAvaBroadcastEditor> BroadcastEditorWeak;
diff --git a/Engine/Plugins/Experimental/Avalanche/Source/AvalancheMediaEditor/Private/Broadcast/ChannelGrid/Slate/SAvaBroadcastProfileEntry.cpp b/Engine/Plugins/Experimental/Avalanche/Source/AvalancheMediaEditor/Private/Broadcast/ChannelGrid/Slate/SAvaBroadcastProfileEntry.cpp
--- a/Engine/Plugins/Experimental/Avalanche/Source/AvalancheMediaEditor/Private/Broadcast/ChannelGrid/Slate/SAvaBroadcastProfileEntry.cpp
+++ b/Engine/Plugins/Experimental/Avalanche/Source/AvalancheMediaEditor/Private/Broadcast/ChannelGrid/Slate/SAvaBroadcastProfileEntry.cpp
@@ -1,10 +1,15 @@
 // Copyright Epic Games, Inc. All Rights Reserved.
 
 #include "SAvaBroadcastProfileEntry.h"
-
 #include "Broadcast/AvaBroadcast.h"
 #include "Framework/Application/SlateApplication.h"
+#include "Input/Events.h"
+#include "Input/Reply.h"
+#include "Internationalization/Text.h"
+#include "Layout/Geometry.h"
 #include "ScopedTransaction.h"
+#include "Styling/SlateBrush.h"
+#include "Styling/SlateTypes.h"
 #include "Widgets/Images/SImage.h"
 #include "Widgets/Input/SButton.h"
 #include "Widgets/Layout/SBox.h"
@@ -12,6 +17,7 @@
 #include "Widgets/Layout/SSpacer.h"
 #include "Widgets/SBoxPanel.h"
 #include "Widgets/Text/SInlineEditableTextBlock.h"
+#include "UObject/NameTypes.h"
 
 #define LOCTEXT_NAMESPACE "SAvaBroadcastProfileEntry"
 
diff --git a/Engine/Plugins/Experimental/Avalanche/Source/AvalancheMediaEditor/Private/Broadcast/ChannelGrid/Slate/SAvaBroadcastProfileEntry.h b/Engine/Plugins/Experimental/Avalanche/Source/AvalancheMediaEditor/Private/Broadcast/ChannelGrid/Slate/SAvaBroadcastProfileEntry.h
--- a/Engine/Plugins/Experimental/Avalanche/Source/AvalancheMediaEditor/Private/Broadcast/ChannelGrid/Slate/SAvaBroadcastProfileEntry.h
+++ b/Engine/Plugins/Experimental/Avalanche/Source/AvalancheMediaEditor/Private/Broadcast/ChannelGrid/Slate/SAvaBroadcastProfileEntry.h
@@ -2,18 +2,25 @@
 
 #pragma once
 
-#include "CoreMinimal.h"
 #include "Styling/SlateTypes.h"
+#include "Templates/SharedPointer.h"
+#include "Types/SlateEnums.h"
 #include "Widgets/SCompoundWidget.h"
 
+class FName;
+class FReply;
+class FText;
 class SInlineEditableTextBlock;
+struct FButtonStyle;
+struct FGeometry;
+struct FPointerEvent;
+struct FSlateBrush;
 
 DECLARE_DELEGATE_RetVal_OneParam(FReply, FOnAvaBroadcastProfileEntrySelected, FName ProfileName);
 
 class SAvaBroadcastProfileEntry : public SCompoundWidget
 {
 public:
-	
 	SLATE_BEGIN_ARGS(SAvaBroadcastProfileEntry){}
 		SLATE_EVENT(FOnAvaBroadcastProfileEntrySelected, OnProfileEntrySelected)
 	SLATE_END_ARGS()
@@ -39,7 +46,6 @@ public:
 	const FSlateBrush* GetBorderImage() const;
 
 protected:
-
 	FName ProfileName;
 	
 	FText ProfileNameText;
diff --git a/Engine/Plugins/Experimental/Avalanche/Source/AvalancheMediaEditor/Private/Broadcast/DetailsView/SAvaBroadcastDetailsView.cpp b/Engine/Plugins/Experimental/Avalanche/Source/AvalancheMediaEditor/Private/Broadcast/DetailsView/SAvaBroadcastDetailsView.cpp
--- a/Engine/Plugins/Experimental/Avalanche/Source/AvalancheMediaEditor/Private/Broadcast/DetailsView/SAvaBroadcastDetailsView.cpp
+++ b/Engine/Plugins/Experimental/Avalanche/Source/AvalancheMediaEditor/Private/Broadcast/DetailsView/SAvaBroadcastDetailsView.cpp
@@ -5,6 +5,8 @@
 #include "Broadcast/AvaBroadcastEditor.h"
 #include "Broadcast/Channel/AvaBroadcastOutputChannel.h"
 #include "Broadcast/ChannelGrid/AvaBroadcastOutputTileItem.h"
+#include "DetailsViewArgs.h"
+#include "Layout/Visibility.h"
 #include "MediaOutput.h"
 #include "Modules/ModuleManager.h"
 #include "PropertyEditorModule.h"
diff --git a/Engine/Plugins/Experimental/Avalanche/Source/AvalancheMediaEditor/Private/Broadcast/DetailsView/SAvaBroadcastDetailsView.h b/Engine/Plugins/Experimental/Avalanche/Source/AvalancheMediaEditor/Private/Broadcast/DetailsView/SAvaBroadcastDetailsView.h
--- a/Engine/Plugins/Experimental/Avalanche/Source/AvalancheMediaEditor/Private/Broadcast/DetailsView/SAvaBroadcastDetailsView.h
+++ b/Engine/Plugins/Experimental/Avalanche/Source/AvalancheMediaEditor/Private/Broadcast/DetailsView/SAvaBroadcastDetailsView.h
@@ -2,17 +2,17 @@
 
 #pragma once
 
-#include "CoreMinimal.h"
+#include "Templates/SharedPointer.h"
 #include "Widgets/SCompoundWidget.h"
 
 class FAvaBroadcastEditor;
 class FAvaBroadcastOutputTileItem;
 class IDetailsView;
+struct EVisibility;
 
 class SAvaBroadcastDetailsView : public SCompoundWidget
 {
 public:
-	
 	SLATE_BEGIN_ARGS(SAvaBroadcastDetailsView) {}
 	SLATE_END_ARGS()
 	
@@ -25,7 +25,6 @@ public:
 	EVisibility GetEmptySelectionTextVisibility() const;
 	
 protected:
-
 	TWeakPtr<FAvaBroadcastEditor> BroadcastEditorWeak;
 	
 	TSharedPtr<IDetailsView> DetailsView;
diff --git a/Engine/Plugins/Experimental/Avalanche/Source/AvalancheMediaEditor/Private/Broadcast/OutputDevices/AvaBroadcastDisplayMediaOutputFactory.cpp b/Engine/Plugins/Experimental/Avalanche/Source/AvalancheMediaEditor/Private/Broadcast/OutputDevices/AvaBroadcastDisplayMediaOutputFactory.cpp
--- a/Engine/Plugins/Experimental/Avalanche/Source/AvalancheMediaEditor/Private/Broadcast/OutputDevices/AvaBroadcastDisplayMediaOutputFactory.cpp
+++ b/Engine/Plugins/Experimental/Avalanche/Source/AvalancheMediaEditor/Private/Broadcast/OutputDevices/AvaBroadcastDisplayMediaOutputFactory.cpp
@@ -1,9 +1,12 @@
 // Copyright Epic Games, Inc. All Rights Reserved.
 
 #include "AvaBroadcastDisplayMediaOutputFactory.h"
-
 #include "AssetTypeCategories.h"
 #include "Broadcast/OutputDevices/AvaBroadcastDisplayMediaOutput.h"
+#include "Misc/FeedbackContext.h"
+#include "UObject/NoExportTypes.h"
+#include "UObject/NameTypes.h"
+#include "UObject/ObjectMacros.h"
 
 UAvaBroadcastDisplayMediaOutputFactory::UAvaBroadcastDisplayMediaOutputFactory()
 {
diff --git a/Engine/Plugins/Experimental/Avalanche/Source/AvalancheMediaEditor/Private/Broadcast/OutputDevices/AvaBroadcastDisplayMediaOutputFactory.h b/Engine/Plugins/Experimental/Avalanche/Source/AvalancheMediaEditor/Private/Broadcast/OutputDevices/AvaBroadcastDisplayMediaOutputFactory.h
--- a/Engine/Plugins/Experimental/Avalanche/Source/AvalancheMediaEditor/Private/Broadcast/OutputDevices/AvaBroadcastDisplayMediaOutputFactory.h
+++ b/Engine/Plugins/Experimental/Avalanche/Source/AvalancheMediaEditor/Private/Broadcast/OutputDevices/AvaBroadcastDisplayMediaOutputFactory.h
@@ -2,10 +2,15 @@
 
 #pragma once
 
-#include "CoreMinimal.h"
 #include "Factories/Factory.h"
 #include "AvaBroadcastDisplayMediaOutputFactory.generated.h"
 
+class FFeedbackContext;
+class FName;
+class UClass;
+class UObject;
+enum EObjectFlags;
+
 /**
  * Implements a factory for UAvaBroadcastDisplayMediaOutput objects.
  */
diff --git a/Engine/Plugins/Experimental/Avalanche/Source/AvalancheMediaEditor/Private/Playback/AvaPlaybackGraphEditor.cpp b/Engine/Plugins/Experimental/Avalanche/Source/AvalancheMediaEditor/Private/Playback/AvaPlaybackGraphEditor.cpp
--- a/Engine/Plugins/Experimental/Avalanche/Source/AvalancheMediaEditor/Private/Playback/AvaPlaybackGraphEditor.cpp
+++ b/Engine/Plugins/Experimental/Avalanche/Source/AvalancheMediaEditor/Private/Playback/AvaPlaybackGraphEditor.cpp
@@ -4,13 +4,19 @@
 
 #include "AppModes/AvaPlaybackDefaultMode.h"
 #include "Broadcast/AvaBroadcastEditor.h"
+#include "EdGraph/EdGraph.h"
+#include "EdGraph/EdGraphNode.h"
 #include "EdGraphUtilities.h"
 #include "Framework/Commands/GenericCommands.h"
+#include "Framework/Commands/UICommandList.h"
 #include "Framework/MultiBox/MultiBoxBuilder.h"
+#include "GraphEditor.h"
 #include "GraphEditorActions.h"
 #include "HAL/PlatformApplicationMisc.h"
 #include "IAvaMediaEditorModule.h"
+#include "Internationalization/Text.h"
 #include "Kismet2/BlueprintEditorUtils.h"
+#include "Math/Color.h"
 #include "Playback/AvaPlaybackCommands.h"
 #include "Playback/AvaPlaybackGraph.h"
 #include "Playback/Graph/AvaPlaybackEditorGraph.h"
@@ -21,6 +27,8 @@
 #include "Playback/Nodes/AvaPlaybackNode.h"
 #include "ScopedTransaction.h"
 #include "SNodePanel.h"
+#include "UObject/NameTypes.h"
+#include "UObject/NoExportTypes.h"
 #include "UObject/UObjectIterator.h"
 #include "WorkflowOrientedApp/ApplicationMode.h"
 
diff --git a/Engine/Plugins/Experimental/Avalanche/Source/AvalancheMediaEditor/Private/Playback/DetailsView/SAvaPlaybackDetailsView.cpp b/Engine/Plugins/Experimental/Avalanche/Source/AvalancheMediaEditor/Private/Playback/DetailsView/SAvaPlaybackDetailsView.cpp
--- a/Engine/Plugins/Experimental/Avalanche/Source/AvalancheMediaEditor/Private/Playback/DetailsView/SAvaPlaybackDetailsView.cpp
+++ b/Engine/Plugins/Experimental/Avalanche/Source/AvalancheMediaEditor/Private/Playback/DetailsView/SAvaPlaybackDetailsView.cpp
@@ -1,7 +1,7 @@
 // Copyright Epic Games, Inc. All Rights Reserved.
 
 #include "SAvaPlaybackDetailsView.h"
-
+#include "DetailsViewArgs.h"
 #include "Modules/ModuleManager.h"
 #include "Playback/AvaPlaybackGraphEditor.h"
 #include "PropertyEditorModule.h"
diff --git a/Engine/Plugins/Experimental/Avalanche/Source/AvalancheMediaEditor/Private/Playback/DetailsView/SAvaPlaybackDetailsView.h b/Engine/Plugins/Experimental/Avalanche/Source/AvalancheMediaEditor/Private/Playback/DetailsView/SAvaPlaybackDetailsView.h
--- a/Engine/Plugins/Experimental/Avalanche/Source/AvalancheMediaEditor/Private/Playback/DetailsView/SAvaPlaybackDetailsView.h
+++ b/Engine/Plugins/Experimental/Avalanche/Source/AvalancheMediaEditor/Private/Playback/DetailsView/SAvaPlaybackDetailsView.h
@@ -2,16 +2,17 @@
 
 #pragma once
 
-#include "CoreMinimal.h"
+#include "Containers/Array.h"
+#include "Templates/SharedPointer.h"
 #include "Widgets/SCompoundWidget.h"
 
 class FAvaPlaybackGraphEditor;
 class IDetailsView;
+class UObject;
 
 class SAvaPlaybackDetailsView : public SCompoundWidget
 {
 public:
-	
 	SLATE_BEGIN_ARGS(SAvaPlaybackDetailsView){}
 	SLATE_END_ARGS()
 	
@@ -21,7 +22,6 @@ public:
 	void OnPlaybackNodeSelectionChanged(const TArray<UObject*>& InSelectedObjects);
 
 protected:
-
 	TWeakPtr<FAvaPlaybackGraphEditor> PlaybackEditorWeak;
 	
 	TSharedPtr<IDetailsView> DetailsView;
diff --git a/Engine/Plugins/Experimental/Avalanche/Source/AvalancheMediaEditor/Private/Playback/Graph/AvaPlaybackEditorGraphSchema.cpp b/Engine/Plugins/Experimental/Avalanche/Source/AvalancheMediaEditor/Private/Playback/Graph/AvaPlaybackEditorGraphSchema.cpp
--- a/Engine/Plugins/Experimental/Avalanche/Source/AvalancheMediaEditor/Private/Playback/Graph/AvaPlaybackEditorGraphSchema.cpp
+++ b/Engine/Plugins/Experimental/Avalanche/Source/AvalancheMediaEditor/Private/Playback/Graph/AvaPlaybackEditorGraphSchema.cpp
@@ -1,8 +1,12 @@
 // Copyright Epic Games, Inc. All Rights Reserved.
 
 #include "Playback/Graph/AvaPlaybackEditorGraphSchema.h"
-
+#include "EdGraph/EdGraph.h"
+#include "EdGraph/EdGraphNode.h"
+#include "EdGraph/EdGraphPin.h"
+#include "EdGraph/EdGraphSchema.h"
 #include "GraphEditorActions.h"
+#include "Math/Color.h"
 #include "Playback/AvaPlaybackGraph.h"
 #include "Playback/Graph/AvaPlaybackEditorGraph.h"
 #include "Playback/Graph/Nodes/AvaPlaybackEditorGraphNode_Root.h"
@@ -13,6 +17,8 @@
 #include "Playback/Nodes/AvaPlaybackNodeRoot.h"
 #include "ScopedTransaction.h"
 #include "ToolMenu.h"
+#include "UObject/NameTypes.h"
+#include "UObject/NoExportTypes.h"
 #include "UObject/UObjectIterator.h"
 
 #define LOCTEXT_NAMESPACE "AvaPlaybackEditorGraphSchema"
diff --git a/Engine/Plugins/Experimental/Avalanche/Source/AvalancheMediaEditor/Private/Playback/Graph/Nodes/AvaPlaybackEditorGraphNode.cpp b/Engine/Plugins/Experimental/Avalanche/Source/AvalancheMediaEditor/Private/Playback/Graph/Nodes/AvaPlaybackEditorGraphNode.cpp
--- a/Engine/Plugins/Experimental/Avalanche/Source/AvalancheMediaEditor/Private/Playback/Graph/Nodes/AvaPlaybackEditorGraphNode.cpp
+++ b/Engine/Plugins/Experimental/Avalanche/Source/AvalancheMediaEditor/Private/Playback/Graph/Nodes/AvaPlaybackEditorGraphNode.cpp
@@ -1,10 +1,13 @@
 // Copyright Epic Games, Inc. All Rights Reserved.
 
 #include "Playback/Graph/Nodes/AvaPlaybackEditorGraphNode.h"
-
 #include "AvaMediaEditorSettings.h"
+#include "EdGraph/EdGraphNode.h"
+#include "EdGraph/EdGraphPin.h"
+#include "EdGraph/EdGraphSchema.h"
 #include "Framework/Commands/GenericCommands.h"
 #include "GraphEditorActions.h"
+#include "Math/Color.h"
 #include "Playback/AvaPlaybackCommands.h"
 #include "Playback/AvaPlaybackGraph.h"
 #include "Playback/Graph/AvaPlaybackEditorGraph.h"
@@ -13,7 +16,10 @@
 #include "Playback/Graph/Nodes/Slate/SAvaPlaybackEditorGraphNode.h"
 #include "Playback/Nodes/AvaPlaybackNode.h"
 #include "ScopedTransaction.h"
+#include "SGraphNode.h"
 #include "ToolMenu.h"
+#include "UObject/NameTypes.h"
+#include "UObject/NoExportTypes.h"
 
 #define LOCTEXT_NAMESPACE "AvaPlaybackEditorGraphNode"
 
diff --git a/Engine/Plugins/Experimental/Avalanche/Source/AvalancheMediaEditor/Private/Playback/Graph/Nodes/Slate/SAvaPlaybackEditorGraphNode.cpp b/Engine/Plugins/Experimental/Avalanche/Source/AvalancheMediaEditor/Private/Playback/Graph/Nodes/Slate/SAvaPlaybackEditorGraphNode.cpp
--- a/Engine/Plugins/Experimental/Avalanche/Source/AvalancheMediaEditor/Private/Playback/Graph/Nodes/Slate/SAvaPlaybackEditorGraphNode.cpp
+++ b/Engine/Plugins/Experimental/Avalanche/Source/AvalancheMediaEditor/Private/Playback/Graph/Nodes/Slate/SAvaPlaybackEditorGraphNode.cpp
@@ -1,12 +1,16 @@
 // Copyright Epic Games, Inc. All Rights Reserved.
 
 #include "Playback/Graph/Nodes/Slate/SAvaPlaybackEditorGraphNode.h"
-
+#include "EdGraph/EdGraphPin.h"
 #include "GraphEditorSettings.h"
+#include "Input/Reply.h"
 #include "KismetPins/SGraphPinExec.h"
+#include "Layout/Visibility.h"
 #include "Playback/Graph/AvaPlaybackEditorGraphSchema.h"
 #include "Playback/Graph/Nodes/AvaPlaybackEditorGraphNode.h"
+#include "SGraphPin.h"
 #include "Widgets/SBoxPanel.h"
+#include "Layout/Margin.h"
 
 #define LOCTEXT_NAMESPACE "SAvaPlaybackEditorGraphNode"
 
diff --git a/Engine/Plugins/Experimental/Avalanche/Source/AvalancheMediaEditor/Private/Playback/Graph/Nodes/Slate/SAvaPlaybackEditorGraphNode_Player.cpp b/Engine/Plugins/Experimental/Avalanche/Source/AvalancheMediaEditor/Private/Playback/Graph/Nodes/Slate/SAvaPlaybackEditorGraphNode_Player.cpp
--- a/Engine/Plugins/Experimental/Avalanche/Source/AvalancheMediaEditor/Private/Playback/Graph/Nodes/Slate/SAvaPlaybackEditorGraphNode_Player.cpp
+++ b/Engine/Plugins/Experimental/Avalanche/Source/AvalancheMediaEditor/Private/Playback/Graph/Nodes/Slate/SAvaPlaybackEditorGraphNode_Player.cpp
@@ -1,7 +1,6 @@
 // Copyright Epic Games, Inc. All Rights Reserved.
 
 #include "SAvaPlaybackEditorGraphNode_Player.h"
-
 #include "Broadcast/Channel/AvaBroadcastOutputChannel.h"
 #include "Broadcast/OutputDevices/Slate/SAvaBroadcastCaptureImage.h"
 #include "Brushes/SlateImageBrush.h"
@@ -9,6 +8,7 @@
 #include "GraphEditorSettings.h"
 #include "Playback/Graph/Nodes/AvaPlaybackEditorGraphNode.h"
 #include "Playback/Nodes/AvaPlaybackNodePlayer.h"
+#include "Styling/SlateBrush.h"
 #include "Widgets/Layout/SBox.h"
 #include "Widgets/Layout/SScaleBox.h"
 #include "Widgets/SBoxPanel.h"
diff --git a/Engine/Plugins/Experimental/Avalanche/Source/AvalancheMediaEditor/Private/Playback/Graph/Nodes/Slate/SAvaPlaybackEditorGraphNode_Player.h b/Engine/Plugins/Experimental/Avalanche/Source/AvalancheMediaEditor/Private/Playback/Graph/Nodes/Slate/SAvaPlaybackEditorGraphNode_Player.h
--- a/Engine/Plugins/Experimental/Avalanche/Source/AvalancheMediaEditor/Private/Playback/Graph/Nodes/Slate/SAvaPlaybackEditorGraphNode_Player.h
+++ b/Engine/Plugins/Experimental/Avalanche/Source/AvalancheMediaEditor/Private/Playback/Graph/Nodes/Slate/SAvaPlaybackEditorGraphNode_Player.h
@@ -2,10 +2,15 @@
 
 #pragma once
 
-#include "CoreMinimal.h"
 #include "Playback/Graph/Nodes/Slate/SAvaPlaybackEditorGraphNode.h"
+#include "Templates/SharedPointer.h"
+#include "Templates/UniquePtr.h"
+#include "UObject/WeakObjectPtrTemplates.h"
 
 class UAvaPlaybackNodePlayer;
+class SWidget;
+struct FSlateBrush;
+class SVerticalBox;
 
 class SAvaPlaybackEditorGraphNode_Player : public SAvaPlaybackEditorGraphNode
 {
@@ -26,7 +31,6 @@ public:
 	//~SGraphNode Interface
 
 protected:
-
 	TWeakObjectPtr<UAvaPlaybackNodePlayer> PlayerNode;
 
 	struct FPreviewBrush;
diff --git a/Engine/Plugins/Experimental/Avalanche/Source/AvalancheMediaEditor/Private/Playback/Graph/Nodes/Slate/SAvaPlaybackEditorGraphNode_Root.cpp b/Engine/Plugins/Experimental/Avalanche/Source/AvalancheMediaEditor/Private/Playback/Graph/Nodes/Slate/SAvaPlaybackEditorGraphNode_Root.cpp
--- a/Engine/Plugins/Experimental/Avalanche/Source/AvalancheMediaEditor/Private/Playback/Graph/Nodes/Slate/SAvaPlaybackEditorGraphNode_Root.cpp
+++ b/Engine/Plugins/Experimental/Avalanche/Source/AvalancheMediaEditor/Private/Playback/Graph/Nodes/Slate/SAvaPlaybackEditorGraphNode_Root.cpp
@@ -1,9 +1,11 @@
 // Copyright Epic Games, Inc. All Rights Reserved.
 
 #include "SAvaPlaybackEditorGraphNode_Root.h"
-
+#include "EdGraph/EdGraphPin.h"
 #include "Playback/Graph/AvaPlaybackEditorGraphSchema.h"
 #include "Playback/Graph/Pins/Slate/SAvaPlaybackEditorGraphPin_Channel.h"
+#include "SGraphNode.h"
+#include "SGraphPin.h"
 
 TSharedPtr<SGraphPin> SAvaPlaybackEditorGraphNode_Root::CreatePinWidget(UEdGraphPin* Pin) const
 {
diff --git a/Engine/Plugins/Experimental/Avalanche/Source/AvalancheMediaEditor/Private/Playback/Graph/Nodes/Slate/SAvaPlaybackEditorGraphNode_Root.h b/Engine/Plugins/Experimental/Avalanche/Source/AvalancheMediaEditor/Private/Playback/Graph/Nodes/Slate/SAvaPlaybackEditorGraphNode_Root.h
--- a/Engine/Plugins/Experimental/Avalanche/Source/AvalancheMediaEditor/Private/Playback/Graph/Nodes/Slate/SAvaPlaybackEditorGraphNode_Root.h
+++ b/Engine/Plugins/Experimental/Avalanche/Source/AvalancheMediaEditor/Private/Playback/Graph/Nodes/Slate/SAvaPlaybackEditorGraphNode_Root.h
@@ -2,15 +2,16 @@
 
 #pragma once
 
-#include "CoreMinimal.h"
+#include "Templates/SharedPointer.h"
 #include "Playback/Graph/Nodes/Slate/SAvaPlaybackEditorGraphNode.h"
 
+class SGraphPin;
 class UAvaPlaybackEditorGraphNode_Root;
+class UEdGraphPin;
 
 class SAvaPlaybackEditorGraphNode_Root : public SAvaPlaybackEditorGraphNode
 {
 public:
-
 	/** Hook that allows derived classes to supply their own SGraphPin derivatives for any pin. */
 	virtual TSharedPtr<SGraphPin> CreatePinWidget(UEdGraphPin* Pin) const override;
 };
diff --git a/Engine/Plugins/Experimental/Avalanche/Source/AvalancheMediaEditor/Private/Playback/Graph/Pins/Slate/SAvaPlaybackEditorGraphPin_Channel.cpp b/Engine/Plugins/Experimental/Avalanche/Source/AvalancheMediaEditor/Private/Playback/Graph/Pins/Slate/SAvaPlaybackEditorGraphPin_Channel.cpp
--- a/Engine/Plugins/Experimental/Avalanche/Source/AvalancheMediaEditor/Private/Playback/Graph/Pins/Slate/SAvaPlaybackEditorGraphPin_Channel.cpp
+++ b/Engine/Plugins/Experimental/Avalanche/Source/AvalancheMediaEditor/Private/Playback/Graph/Pins/Slate/SAvaPlaybackEditorGraphPin_Channel.cpp
@@ -1,12 +1,19 @@
 // Copyright Epic Games, Inc. All Rights Reserved.
 
 #include "SAvaPlaybackEditorGraphPin_Channel.h"
-
 #include "AvaMediaEditorUtils.h"
 #include "Broadcast/AvaBroadcast.h"
+#include "Broadcast/Channel/AvaBroadcastOutputChannel.h"
+#include "EdGraph/EdGraphPin.h"
+#include "Internationalization/Text.h"
 #include "Playback/Graph/AvaPlaybackEditorGraphSchema.h"
+#include "SGraphPin.h"
+#include "Styling/SlateColor.h"
+#include "Styling/SlateBrush.h"
 #include "Widgets/Images/SImage.h"
 #include "Widgets/SBoxPanel.h"
+#include "Widgets/Text/STextBlock.h"
+#include "Widgets/SWidget.h"
 
 void SAvaPlaybackEditorGraphPin_Channel::Construct(const FArguments& InArgs, UEdGraphPin* InPin)
 {
@@ -45,13 +52,13 @@ void SAvaPlaybackEditorGraphPin_Channel::UpdateChannelState()
 TSharedRef<SWidget> SAvaPlaybackEditorGraphPin_Channel::GetDefaultValueWidget()
 {
 	return SNew(SHorizontalBox)
-		+SHorizontalBox::Slot()
+		+ SHorizontalBox::Slot()
 		.AutoWidth()
 		[
 			SNew(SImage)
 			.Image(this, &SAvaPlaybackEditorGraphPin_Channel::GetChannelStatusBrush)
 		]
-		+SHorizontalBox::Slot()
+		+ SHorizontalBox::Slot()
 		.Padding(2.f, 0.f, 0.f, 0.f)
 		.HAlign(EHorizontalAlignment::HAlign_Left)
 		.VAlign(EVerticalAlignment::VAlign_Center)
diff --git a/Engine/Plugins/Experimental/Avalanche/Source/AvalancheMediaEditor/Private/Playback/Graph/Pins/Slate/SAvaPlaybackEditorGraphPin_Channel.h b/Engine/Plugins/Experimental/Avalanche/Source/AvalancheMediaEditor/Private/Playback/Graph/Pins/Slate/SAvaPlaybackEditorGraphPin_Channel.h
--- a/Engine/Plugins/Experimental/Avalanche/Source/AvalancheMediaEditor/Private/Playback/Graph/Pins/Slate/SAvaPlaybackEditorGraphPin_Channel.h
+++ b/Engine/Plugins/Experimental/Avalanche/Source/AvalancheMediaEditor/Private/Playback/Graph/Pins/Slate/SAvaPlaybackEditorGraphPin_Channel.h
@@ -2,15 +2,19 @@
 
 #pragma once
 
-#include "CoreMinimal.h"
 #include "SGraphPin.h"
+#include "Templates/SharedPointer.h"
 
 struct FAvaBroadcastOutputChannel;
+class UEdGraphPin;
+class SWidget;
+class FText;
+struct FSlateBrush;
+struct FSlateColor;
 
 class SAvaPlaybackEditorGraphPin_Channel : public SGraphPin
 {
 public:
-
 	SLATE_BEGIN_ARGS(SAvaPlaybackEditorGraphPin_Channel){}
 	SLATE_END_ARGS()
 	
@@ -30,7 +34,6 @@ public:
 	virtual FSlateColor GetPinColor() const override;
 	
 protected:
-
 	const FSlateBrush* ChannelStateBrush = nullptr;
 	
 	FText ChannelStateText;
diff --git a/Engine/Plugins/Experimental/Avalanche/Source/AvalancheMediaEditor/Private/Playback/Graph/SchemaActions/AvaPlaybackAction_NewNode.cpp b/Engine/Plugins/Experimental/Avalanche/Source/AvalancheMediaEditor/Private/Playback/Graph/SchemaActions/AvaPlaybackAction_NewNode.cpp
--- a/Engine/Plugins/Experimental/Avalanche/Source/AvalancheMediaEditor/Private/Playback/Graph/SchemaActions/AvaPlaybackAction_NewNode.cpp
+++ b/Engine/Plugins/Experimental/Avalanche/Source/AvalancheMediaEditor/Private/Playback/Graph/SchemaActions/AvaPlaybackAction_NewNode.cpp
@@ -1,7 +1,8 @@
 // Copyright Epic Games, Inc. All Rights Reserved.
 
 #include "Playback/Graph/SchemaActions/AvaPlaybackAction_NewNode.h"
-
+#include "EdGraph/EdGraphNode.h"
+#include "Internationalization/Text.h"
 #include "Playback/AvaPlaybackGraph.h"
 #include "Playback/Graph/AvaPlaybackEditorGraph.h"
 #include "Playback/Nodes/AvaPlaybackNode.h"
@@ -9,6 +10,11 @@
 
 #define LOCTEXT_NAMESPACE "AvaPlaybackAction_NewNode"
 
+FAvaPlaybackAction_NewNode::FAvaPlaybackAction_NewNode(FText InNodeCategory, FText InMenuDesc, FText InToolTip, const int32 InGrouping)
+	: FEdGraphSchemaAction(MoveTemp(InNodeCategory), MoveTemp(InMenuDesc), MoveTemp(InToolTip), InGrouping)
+{
+}
+
 UEdGraphNode* FAvaPlaybackAction_NewNode::PerformAction(UEdGraph* ParentGraph
 	, UEdGraphPin* FromPin
 	, const FVector2D Location
diff --git a/Engine/Plugins/Experimental/Avalanche/Source/AvalancheMediaEditor/Private/Rundown/ChannelStatus/SAvaRundownChannelStatus.cpp b/Engine/Plugins/Experimental/Avalanche/Source/AvalancheMediaEditor/Private/Rundown/ChannelStatus/SAvaRundownChannelStatus.cpp
--- a/Engine/Plugins/Experimental/Avalanche/Source/AvalancheMediaEditor/Private/Rundown/ChannelStatus/SAvaRundownChannelStatus.cpp
+++ b/Engine/Plugins/Experimental/Avalanche/Source/AvalancheMediaEditor/Private/Rundown/ChannelStatus/SAvaRundownChannelStatus.cpp
@@ -1,13 +1,16 @@
 // Copyright Epic Games, Inc. All Rights Reserved.
 
 #include "SAvaRundownChannelStatus.h"
-
+#include "AvaMediaDefines.h"
 #include "AvaMediaEditorUtils.h"
 #include "Broadcast/AvaBroadcast.h"
+#include "Broadcast/Channel/AvaBroadcastOutputChannel.h"
+#include "Internationalization/Text.h"
 #include "Widgets/Images/SImage.h"
 #include "Widgets/Layout/SBox.h"
 #include "Widgets/SBoxPanel.h"
 #include "Widgets/Text/STextBlock.h"
+#include "UObject/NameTypes.h"
 
 void SAvaRundownChannelStatus::Construct(const FArguments& InArgs, const FAvaBroadcastOutputChannel& InChannel)
 {
diff --git a/Engine/Plugins/Experimental/Avalanche/Source/AvalancheMediaEditor/Private/Rundown/ChannelStatus/SAvaRundownChannelStatus.h b/Engine/Plugins/Experimental/Avalanche/Source/AvalancheMediaEditor/Private/Rundown/ChannelStatus/SAvaRundownChannelStatus.h
--- a/Engine/Plugins/Experimental/Avalanche/Source/AvalancheMediaEditor/Private/Rundown/ChannelStatus/SAvaRundownChannelStatus.h
+++ b/Engine/Plugins/Experimental/Avalanche/Source/AvalancheMediaEditor/Private/Rundown/ChannelStatus/SAvaRundownChannelStatus.h
@@ -3,11 +3,14 @@
 #pragma once
 
 #include "Broadcast/Channel/AvaBroadcastOutputChannel.h"
-#include "CoreMinimal.h"
 #include "Widgets/SCompoundWidget.h"
+#include "Templates/SharedPointer.h"
 
 class SImage;
 class UAvaBroadcast;
+class FName;
+struct FAvaBroadcastOutputChannel;
+enum class EAvaBroadcastChannelChange : uint8;
 
 /*
  * Widget displaying the Status of a Single Channel
@@ -15,7 +18,6 @@ class UAvaBroadcast;
 class SAvaRundownChannelStatus : public SCompoundWidget
 {
 public:
-	
 	SLATE_BEGIN_ARGS(SAvaRundownChannelStatus) {}
 	SLATE_END_ARGS()
 
@@ -26,7 +28,6 @@ public:
 	void OnChannelChanged(const FAvaBroadcastOutputChannel& InChannel, EAvaBroadcastChannelChange InChange);
 	
 protected:
-
 	FName ChannelName;
 
 	TSharedPtr<SImage> ChannelStatusIcon;
diff --git a/Engine/Plugins/Experimental/Avalanche/Source/AvalancheMediaEditor/Private/Rundown/ChannelStatus/SAvaRundownChannelStatusList.cpp b/Engine/Plugins/Experimental/Avalanche/Source/AvalancheMediaEditor/Private/Rundown/ChannelStatus/SAvaRundownChannelStatusList.cpp
--- a/Engine/Plugins/Experimental/Avalanche/Source/AvalancheMediaEditor/Private/Rundown/ChannelStatus/SAvaRundownChannelStatusList.cpp
+++ b/Engine/Plugins/Experimental/Avalanche/Source/AvalancheMediaEditor/Private/Rundown/ChannelStatus/SAvaRundownChannelStatusList.cpp
@@ -1,11 +1,13 @@
 // Copyright Epic Games, Inc. All Rights Reserved.
 
 #include "SAvaRundownChannelStatusList.h"
-
+#include "AvaMediaDefines.h"
 #include "Broadcast/AvaBroadcast.h"
+#include "Misc/EnumClassFlags.h"
 #include "SAvaRundownChannelStatus.h"
 #include "Widgets/Layout/SGridPanel.h"
 #include "Widgets/Layout/SWrapBox.h"
+#include "UObject/NameTypes.h"
 
 void SAvaRundownChannelStatusList::Construct(const FArguments& InArgs)
 {
diff --git a/Engine/Plugins/Experimental/Avalanche/Source/AvalancheMediaEditor/Private/Rundown/ChannelStatus/SAvaRundownChannelStatusList.h b/Engine/Plugins/Experimental/Avalanche/Source/AvalancheMediaEditor/Private/Rundown/ChannelStatus/SAvaRundownChannelStatusList.h
--- a/Engine/Plugins/Experimental/Avalanche/Source/AvalancheMediaEditor/Private/Rundown/ChannelStatus/SAvaRundownChannelStatusList.h
+++ b/Engine/Plugins/Experimental/Avalanche/Source/AvalancheMediaEditor/Private/Rundown/ChannelStatus/SAvaRundownChannelStatusList.h
@@ -2,13 +2,16 @@
 
 #pragma once
 
-#include "AvaMediaDefines.h"
-#include "CoreMinimal.h"
+#include "Containers/Map.h"
+#include "Templates/SharedPointer.h"
 #include "Widgets/SCompoundWidget.h"
+#include "UObject/WeakObjectPtrTemplates.h"
 
 class SAvaRundownChannelStatus;
 class SWrapBox;
 class UAvaBroadcast;
+class FName;
+enum class EAvaBroadcastChange : uint8;
 
 /*
  * Widget containing the Status of all Channels in Broadcast. It contains a list of SAvaRundownChannelStatus widgets
@@ -16,7 +19,6 @@ class UAvaBroadcast;
 class SAvaRundownChannelStatusList : public SCompoundWidget
 {
 public:
-	
 	SLATE_BEGIN_ARGS(SAvaRundownChannelStatusList){}
 	SLATE_END_ARGS()
 
@@ -30,7 +32,6 @@ public:
 	void OnBroadcastChanged(EAvaBroadcastChange InChange);
 
 protected:
-
 	TSharedPtr<SWrapBox> WrapBox;
 
 	TWeakObjectPtr<UAvaBroadcast> BroadcastWeak;
diff --git a/Engine/Plugins/Experimental/Avalanche/Source/AvalancheMediaEditor/Private/Rundown/DetailsView/RemoteControl/Properties/AvaRundownPageRemoteControlWidgetUtils.h b/Engine/Plugins/Experimental/Avalanche/Source/AvalancheMediaEditor/Private/Rundown/DetailsView/RemoteControl/Properties/AvaRundownPageRemoteControlWidgetUtils.h
--- a/Engine/Plugins/Experimental/Avalanche/Source/AvalancheMediaEditor/Private/Rundown/DetailsView/RemoteControl/Properties/AvaRundownPageRemoteControlWidgetUtils.h
+++ b/Engine/Plugins/Experimental/Avalanche/Source/AvalancheMediaEditor/Private/Rundown/DetailsView/RemoteControl/Properties/AvaRundownPageRemoteControlWidgetUtils.h
@@ -2,7 +2,9 @@
 
 #pragma once
 
-#include "CoreMinimal.h"
+#include "Containers/Array.h"
+#include "Containers/StringFwd.h"
+#include "Templates/SharedPointer.h"
 
 class IDetailTreeNode;
 class IPropertyHandle;
diff --git a/Engine/Plugins/Experimental/Avalanche/Source/AvalancheMediaEditor/Private/Rundown/DetailsView/RemoteControl/Properties/SAvaRundownPageRemoteControlProps.cpp b/Engine/Plugins/Experimental/Avalanche/Source/AvalancheMediaEditor/Private/Rundown/DetailsView/RemoteControl/Properties/SAvaRundownPageRemoteControlProps.cpp
--- a/Engine/Plugins/Experimental/Avalanche/Source/AvalancheMediaEditor/Private/Rundown/DetailsView/RemoteControl/Properties/SAvaRundownPageRemoteControlProps.cpp
+++ b/Engine/Plugins/Experimental/Avalanche/Source/AvalancheMediaEditor/Private/Rundown/DetailsView/RemoteControl/Properties/SAvaRundownPageRemoteControlProps.cpp
@@ -4,6 +4,9 @@
 
 #include "AvaRundownRCPropertyItem.h"
 #include "Playable/AvaPlayableRemoteControl.h"
+#include "Playable/AvaPlayableRemoteControlValues.h"
+#include "RemoteControlEntity.h"
+#include "RemoteControlPreset.h"
 #include "Rundown/AvaRundownEditor.h"
 #include "Rundown/AvaRundownManagedInstanceCache.h"
 #include "Rundown/AvaRundownPage.h"
@@ -12,6 +15,7 @@
 #include "Widgets/Views/ITableRow.h"
 #include "Widgets/Views/SListView.h"
 #include "Widgets/Views/STableViewBase.h"
+#include "UObject/NameTypes.h"
 
 #define LOCTEXT_NAMESPACE "SAvaRundownPageRemoteControlProps"
 
diff --git a/Engine/Plugins/Experimental/Avalanche/Source/AvalancheMediaEditor/Private/Rundown/DetailsView/RemoteControl/Properties/SAvaRundownPageRemoteControlProps.h b/Engine/Plugins/Experimental/Avalanche/Source/AvalancheMediaEditor/Private/Rundown/DetailsView/RemoteControl/Properties/SAvaRundownPageRemoteControlProps.h
--- a/Engine/Plugins/Experimental/Avalanche/Source/AvalancheMediaEditor/Private/Rundown/DetailsView/RemoteControl/Properties/SAvaRundownPageRemoteControlProps.h
+++ b/Engine/Plugins/Experimental/Avalanche/Source/AvalancheMediaEditor/Private/Rundown/DetailsView/RemoteControl/Properties/SAvaRundownPageRemoteControlProps.h
@@ -2,8 +2,10 @@
 
 #pragma once
 
-#include "CoreMinimal.h"
-#include "RemoteControlPreset.h"
+#include "Containers/Array.h"
+#include "Containers/Map.h"
+#include "Containers/Set.h"
+#include "Templates/SharedPointer.h"
 #include "Widgets/SCompoundWidget.h"
 #include "Widgets/Views/SListView.h"
 #include "Widgets/Views/ITableRow.h"
@@ -11,9 +13,16 @@
 class FAvaRundownEditor;
 class FAvaRundownManagedInstance;
 class FAvaRundownRCPropertyItem;
+class FName;
+class ITableRow;
 class SAvaRundownPageRemoteControlProps;
+class STableViewBase;
+class SWidget;
+class URemoteControlPreset;
 struct FAvaPlayableRemoteControlValue;
 struct FAvaRundownPage;
+struct FGuid;
+struct FRemoteControlEntity;
 
 using FAvaRundownRCPropertyItemPtr = TSharedPtr<FAvaRundownRCPropertyItem>;
 
diff --git a/Engine/Plugins/Experimental/Avalanche/Source/AvalancheMediaEditor/Private/Rundown/DetailsView/SAvaRundownPageDetails.cpp b/Engine/Plugins/Experimental/Avalanche/Source/AvalancheMediaEditor/Private/Rundown/DetailsView/SAvaRundownPageDetails.cpp
--- a/Engine/Plugins/Experimental/Avalanche/Source/AvalancheMediaEditor/Private/Rundown/DetailsView/SAvaRundownPageDetails.cpp
+++ b/Engine/Plugins/Experimental/Avalanche/Source/AvalancheMediaEditor/Private/Rundown/DetailsView/SAvaRundownPageDetails.cpp
@@ -1,9 +1,10 @@
 // Copyright Epic Games, Inc. All Rights Reserved.
 
 #include "SAvaRundownPageDetails.h"
-
 #include "Async/Async.h"
 #include "IAvaMediaModule.h"
+#include "Input/Reply.h"
+#include "Internationalization/Text.h"
 #include "RemoteControl/Controllers/SAvaRundownRCControllerPanel.h"
 #include "Rundown/AvaRundown.h"
 #include "Rundown/AvaRundownEditor.h"
@@ -12,6 +13,7 @@
 #include "Rundown/DetailsView/RemoteControl/Properties/SAvaRundownPageRemoteControlProps.h"
 #include "Rundown/Pages/Slate/SAvaRundownInstancedPageList.h"
 #include "Rundown/Pages/Slate/SAvaRundownPageList.h"
+#include "Styling/SlateBrush.h"
 #include "Widgets/Images/SImage.h"
 #include "Widgets/Input/SButton.h"
 #include "Widgets/Input/SEditableTextBox.h"
diff --git a/Engine/Plugins/Experimental/Avalanche/Source/AvalancheMediaEditor/Private/Rundown/DetailsView/SAvaRundownPageDetails.h b/Engine/Plugins/Experimental/Avalanche/Source/AvalancheMediaEditor/Private/Rundown/DetailsView/SAvaRundownPageDetails.h
--- a/Engine/Plugins/Experimental/Avalanche/Source/AvalancheMediaEditor/Private/Rundown/DetailsView/SAvaRundownPageDetails.h
+++ b/Engine/Plugins/Experimental/Avalanche/Source/AvalancheMediaEditor/Private/Rundown/DetailsView/SAvaRundownPageDetails.h
@@ -2,14 +2,20 @@
 
 #pragma once
 
-#include "CoreMinimal.h"
+#include "Containers/Array.h"
 #include "Rundown/AvaRundownDefines.h"
+#include "Templates/SharedPointer.h"
+#include "Types/SlateEnums.h"
 #include "Widgets/SCompoundWidget.h"
 
 class FAvaRundownEditor;
+class FReply;
+class FText;
 class SAvaRundownPageRemoteControlProps;
 class SAvaRundownRCControllerPanel;
 struct FAvaRundownPage;
+struct FSoftObjectPath;
+struct FSlateBrush;
 
 class SAvaRundownPageDetails : public SCompoundWidget
 {
diff --git a/Engine/Plugins/Experimental/Avalanche/Source/AvalancheMediaEditor/Private/Rundown/Pages/Columns/Slate/SAvaRundownPageChannelSelector.cpp b/Engine/Plugins/Experimental/Avalanche/Source/AvalancheMediaEditor/Private/Rundown/Pages/Columns/Slate/SAvaRundownPageChannelSelector.cpp
--- a/Engine/Plugins/Experimental/Avalanche/Source/AvalancheMediaEditor/Private/Rundown/Pages/Columns/Slate/SAvaRundownPageChannelSelector.cpp
+++ b/Engine/Plugins/Experimental/Avalanche/Source/AvalancheMediaEditor/Private/Rundown/Pages/Columns/Slate/SAvaRundownPageChannelSelector.cpp
@@ -1,12 +1,15 @@
 // Copyright Epic Games, Inc. All Rights Reserved.
 
 #include "SAvaRundownPageChannelSelector.h"
-
 #include "Broadcast/AvaBroadcast.h"
+#include "Internationalization/Text.h"
 #include "Rundown/Pages/PageViews/IAvaRundownInstancedPageView.h"
 #include "Rundown/Pages/PageViews/AvaRundownInstancedPageViewImpl.h"
 #include "Rundown/Pages/PageViews/AvaRundownPageViewImpl.h"
 #include "Rundown/Pages/Slate/SAvaRundownPageViewRow.h"
+#include "UObject/NameTypes.h"
+#include "Widgets/SWidget.h"
+#include "Widgets/Text/STextBlock.h"
 
 void SAvaRundownPageChannelSelector::Construct(const FArguments& InArgs, const FAvaRundownPageViewRef& InPageView, const TSharedPtr<SAvaRundownPageViewRow>& InRow)
 {
diff --git a/Engine/Plugins/Experimental/Avalanche/Source/AvalancheMediaEditor/Private/Rundown/Pages/Columns/Slate/SAvaRundownPageChannelSelector.h b/Engine/Plugins/Experimental/Avalanche/Source/AvalancheMediaEditor/Private/Rundown/Pages/Columns/Slate/SAvaRundownPageChannelSelector.h
--- a/Engine/Plugins/Experimental/Avalanche/Source/AvalancheMediaEditor/Private/Rundown/Pages/Columns/Slate/SAvaRundownPageChannelSelector.h
+++ b/Engine/Plugins/Experimental/Avalanche/Source/AvalancheMediaEditor/Private/Rundown/Pages/Columns/Slate/SAvaRundownPageChannelSelector.h
@@ -2,13 +2,18 @@
 
 #pragma once
 
-#include "CoreMinimal.h"
+#include "Containers/Array.h"
 #include "Rundown/AvaRundownDefines.h"
 #include "Widgets/Input/SComboBox.h"
 #include "Widgets/SCompoundWidget.h"
+#include "Templates/SharedPointer.h"
+#include "Types/SlateEnums.h"
 
 class IAvaRundownInstancedPageView;
 class SAvaRundownPageViewRow;
+class FName;
+class FText;
+class SWidget;
 
 class SAvaRundownPageChannelSelector : public SCompoundWidget
 {
diff --git a/Engine/Plugins/Experimental/Avalanche/Source/AvalancheMediaEditor/Private/Rundown/Pages/Columns/Slate/SAvaRundownPageId.cpp b/Engine/Plugins/Experimental/Avalanche/Source/AvalancheMediaEditor/Private/Rundown/Pages/Columns/Slate/SAvaRundownPageId.cpp
--- a/Engine/Plugins/Experimental/Avalanche/Source/AvalancheMediaEditor/Private/Rundown/Pages/Columns/Slate/SAvaRundownPageId.cpp
+++ b/Engine/Plugins/Experimental/Avalanche/Source/AvalancheMediaEditor/Private/Rundown/Pages/Columns/Slate/SAvaRundownPageId.cpp
@@ -1,7 +1,7 @@
 // Copyright Epic Games, Inc. All Rights Reserved.
 
 #include "SAvaRundownPageId.h"
-
+#include "Internationalization/Text.h"
 #include "Rundown/AvaRundown.h"
 #include "ScopedTransaction.h"
 #include "Widgets/Text/SInlineEditableTextBlock.h"
diff --git a/Engine/Plugins/Experimental/Avalanche/Source/AvalancheMediaEditor/Private/Rundown/Pages/Columns/Slate/SAvaRundownPageId.h b/Engine/Plugins/Experimental/Avalanche/Source/AvalancheMediaEditor/Private/Rundown/Pages/Columns/Slate/SAvaRundownPageId.h
--- a/Engine/Plugins/Experimental/Avalanche/Source/AvalancheMediaEditor/Private/Rundown/Pages/Columns/Slate/SAvaRundownPageId.h
+++ b/Engine/Plugins/Experimental/Avalanche/Source/AvalancheMediaEditor/Private/Rundown/Pages/Columns/Slate/SAvaRundownPageId.h
@@ -2,17 +2,18 @@
 
 #pragma once
 
-#include "CoreMinimal.h"
 #include "Rundown/AvaRundownDefines.h"
 #include "Rundown/Pages/PageViews/IAvaRundownPageView.h"
+#include "Templates/SharedPointer.h"
+#include "Types/SlateEnums.h"
 #include "Widgets/SCompoundWidget.h"
 
 class SInlineEditableTextBlock;
+class FText;
 
 class SAvaRundownPageId : public SCompoundWidget
 {
 public:
-	
 	SLATE_BEGIN_ARGS(SAvaRundownPageId){}
 	SLATE_END_ARGS()
 
@@ -33,7 +34,6 @@ public:
 	void RenumberPageId(const FText& InText, const FAvaRundownPageViewPtr& InPageView);
 	
 protected:
-
 	TWeakPtr<IAvaRundownPageView> PageViewWeak;
 	
 	TSharedPtr<SInlineEditableTextBlock> InlineTextBlock;
diff --git a/Engine/Plugins/Experimental/Avalanche/Source/AvalancheMediaEditor/Private/Rundown/Pages/Columns/Slate/SAvaRundownPageName.cpp b/Engine/Plugins/Experimental/Avalanche/Source/AvalancheMediaEditor/Private/Rundown/Pages/Columns/Slate/SAvaRundownPageName.cpp
--- a/Engine/Plugins/Experimental/Avalanche/Source/AvalancheMediaEditor/Private/Rundown/Pages/Columns/Slate/SAvaRundownPageName.cpp
+++ b/Engine/Plugins/Experimental/Avalanche/Source/AvalancheMediaEditor/Private/Rundown/Pages/Columns/Slate/SAvaRundownPageName.cpp
@@ -1,8 +1,8 @@
 // Copyright Epic Games, Inc. All Rights Reserved.
 
 #include "SAvaRundownPageName.h"
-
 #include "Framework/Application/SlateApplication.h"
+#include "Internationalization/Text.h"
 #include "Rundown/Pages/PageViews/IAvaRundownPageView.h"
 #include "Rundown/Pages/Slate/SAvaRundownPageViewRow.h"
 #include "Widgets/Text/SInlineEditableTextBlock.h"
diff --git a/Engine/Plugins/Experimental/Avalanche/Source/AvalancheMediaEditor/Private/Rundown/Pages/Columns/Slate/SAvaRundownPageName.h b/Engine/Plugins/Experimental/Avalanche/Source/AvalancheMediaEditor/Private/Rundown/Pages/Columns/Slate/SAvaRundownPageName.h
--- a/Engine/Plugins/Experimental/Avalanche/Source/AvalancheMediaEditor/Private/Rundown/Pages/Columns/Slate/SAvaRundownPageName.h
+++ b/Engine/Plugins/Experimental/Avalanche/Source/AvalancheMediaEditor/Private/Rundown/Pages/Columns/Slate/SAvaRundownPageName.h
@@ -2,10 +2,12 @@
 
 #pragma once
 
-#include "CoreMinimal.h"
 #include "Rundown/AvaRundownDefines.h"
+#include "Templates/SharedPointer.h"
+#include "Types/SlateEnums.h"
 #include "Widgets/SCompoundWidget.h"
 
+class FText;
 class IAvaRundownPageView;
 class SAvaRundownPageViewRow;
 class SInlineEditableTextBlock;
@@ -14,7 +16,6 @@ enum class EAvaRundownPageActionState : uint8;
 class SAvaRundownPageName : public SCompoundWidget
 {
 public:
-	
 	SLATE_BEGIN_ARGS(SAvaRundownPageName){}
 	SLATE_END_ARGS()
 
@@ -35,7 +36,6 @@ public:
 	bool IsReadOnly() const;
 	
 protected:
-
 	TWeakPtr<IAvaRundownPageView> PageViewWeak;
 	
 	TWeakPtr<SAvaRundownPageViewRow> PageViewRowWeak;
diff --git a/Engine/Plugins/Experimental/Avalanche/Source/AvalancheMediaEditor/Private/Rundown/Pages/Slate/SAvaRundownPageViewRow.cpp b/Engine/Plugins/Experimental/Avalanche/Source/AvalancheMediaEditor/Private/Rundown/Pages/Slate/SAvaRundownPageViewRow.cpp
--- a/Engine/Plugins/Experimental/Avalanche/Source/AvalancheMediaEditor/Private/Rundown/Pages/Slate/SAvaRundownPageViewRow.cpp
+++ b/Engine/Plugins/Experimental/Avalanche/Source/AvalancheMediaEditor/Private/Rundown/Pages/Slate/SAvaRundownPageViewRow.cpp
@@ -1,15 +1,23 @@
 // Copyright Epic Games, Inc. All Rights Reserved.
 
 #include "Rundown/Pages/Slate/SAvaRundownPageViewRow.h"
-
 #include "AvaMediaEditorStyle.h"
+#include "Input/DragAndDrop.h"
+#include "Input/Events.h"
+#include "Input/Reply.h"
+#include "Layout/Geometry.h"
+#include "Misc/Attribute.h"
 #include "Rundown/Pages/Columns/IAvaRundownPageViewColumn.h"
+#include "Rundown/Pages/PageViews/IAvaRundownPageView.h"
 #include "SAvaRundownPageList.h"
 #include "ScopedTransaction.h"
+#include "Styling/SlateBrush.h"
 #include "Widgets/Layout/SBorder.h"
 #include "Widgets/SNullWidget.h"
+#include "Widgets/SWidget.h"
 #include "Widgets/Text/STextBlock.h"
 #include "Widgets/Views/SListView.h"
+#include "UObject/NameTypes.h"
 
 #define LOCTEXT_NAMESPACE "SAvaRundownPageViewRow"
 
diff --git a/Engine/Plugins/Experimental/Avalanche/Source/AvalancheMediaEditor/Private/Rundown/Preview/SAvaRundownPagePreview.cpp b/Engine/Plugins/Experimental/Avalanche/Source/AvalancheMediaEditor/Private/Rundown/Preview/SAvaRundownPagePreview.cpp
--- a/Engine/Plugins/Experimental/Avalanche/Source/AvalancheMediaEditor/Private/Rundown/Preview/SAvaRundownPagePreview.cpp
+++ b/Engine/Plugins/Experimental/Avalanche/Source/AvalancheMediaEditor/Private/Rundown/Preview/SAvaRundownPagePreview.cpp
@@ -11,15 +11,19 @@
 #include "Framework/MultiBox/MultiBoxBuilder.h"
 #include "Framework/Notifications/NotificationManager.h"
 #include "IAvaMediaEditorModule.h"
+#include "Internationalization/Text.h"
 #include "ISettingsModule.h"
+#include "Layout/Visibility.h"
 #include "Rundown/AvaRundown.h"
 #include "Rundown/AvaRundownCommands.h"
 #include "Rundown/AvaRundownEditor.h"
 #include "Rundown/AvaRundownEditorSettings.h"
 #include "Rundown/Preview/SAvaRundownPreviewChannelSelector.h"
+#include "Styling/SlateBrush.h"
 #include "Widgets/Layout/SScaleBox.h"
 #include "Widgets/Layout/SSpacer.h"
 #include "Widgets/Notifications/SNotificationList.h"
+#include "Widgets/SWidget.h"
 
 #define LOCTEXT_NAMESPACE "SAvaRundownPagePreview"
 
diff --git a/Engine/Plugins/Experimental/Avalanche/Source/AvalancheMediaEditor/Private/Rundown/Preview/SAvaRundownPagePreview.h b/Engine/Plugins/Experimental/Avalanche/Source/AvalancheMediaEditor/Private/Rundown/Preview/SAvaRundownPagePreview.h
--- a/Engine/Plugins/Experimental/Avalanche/Source/AvalancheMediaEditor/Private/Rundown/Preview/SAvaRundownPagePreview.h
+++ b/Engine/Plugins/Experimental/Avalanche/Source/AvalancheMediaEditor/Private/Rundown/Preview/SAvaRundownPagePreview.h
@@ -2,13 +2,19 @@
 
 #pragma once
 
-#include "CoreMinimal.h"
+#include "Math/MathFwd.h"
+#include "Templates/SharedPointer.h"
+#include "Templates/UniquePtr.h"
 #include "Widgets/SCompoundWidget.h"
 
 class FAvaRundownEditor;
 class FMenuBuilder;
+class FText;
 class FUICommandList;
+class SWidget;
 class UTextureRenderTarget2D;
+struct EVisibility;
+struct FSlateBrush;
 
 class SAvaRundownPagePreview : public SCompoundWidget
 {
@@ -43,4 +49,4 @@ protected:
 
 	struct FPreviewBrush;
 	TUniquePtr<FPreviewBrush> PreviewBrush;
-};
\ No newline at end of file
+};
diff --git a/Engine/Plugins/Experimental/Avalanche/Source/AvalancheMediaEditor/Private/Rundown/ShowControl/SAvaRundownShowControl.cpp b/Engine/Plugins/Experimental/Avalanche/Source/AvalancheMediaEditor/Private/Rundown/ShowControl/SAvaRundownShowControl.cpp
--- a/Engine/Plugins/Experimental/Avalanche/Source/AvalancheMediaEditor/Private/Rundown/ShowControl/SAvaRundownShowControl.cpp
+++ b/Engine/Plugins/Experimental/Avalanche/Source/AvalancheMediaEditor/Private/Rundown/ShowControl/SAvaRundownShowControl.cpp
@@ -1,13 +1,15 @@
 // Copyright Epic Games, Inc. All Rights Reserved.
 
 #include "SAvaRundownShowControl.h"
-
+#include "Framework/Commands/UICommandList.h"
 #include "Framework/MultiBox/MultiBoxBuilder.h"
+#include "Internationalization/Text.h"
 #include "Playback/AvaPlaybackGraph.h"
 #include "Rundown/AvaRundownCommands.h"
 #include "Rundown/AvaRundownEditor.h"
 #include "Rundown/Pages/Slate/SAvaRundownInstancedPageList.h"
 #include "Rundown/Pages/Slate/SAvaRundownPageList.h"
+#include "Widgets/SWidget.h"
 
 #define LOCTEXT_NAMESPACE "SAvaRundownShowControl"
 
diff --git a/Engine/Plugins/Experimental/Avalanche/Source/AvalancheMediaEditor/Private/Rundown/ShowControl/SAvaRundownShowControl.h b/Engine/Plugins/Experimental/Avalanche/Source/AvalancheMediaEditor/Private/Rundown/ShowControl/SAvaRundownShowControl.h
--- a/Engine/Plugins/Experimental/Avalanche/Source/AvalancheMediaEditor/Private/Rundown/ShowControl/SAvaRundownShowControl.h
+++ b/Engine/Plugins/Experimental/Avalanche/Source/AvalancheMediaEditor/Private/Rundown/ShowControl/SAvaRundownShowControl.h
@@ -2,11 +2,13 @@
 
 #pragma once
 
-#include "CoreMinimal.h"
+#include "Templates/SharedPointer.h"
 #include "Widgets/SCompoundWidget.h"
 
 class FAvaRundownEditor;
 class FUICommandList;
+class FText;
+class SWidget;
 
 class SAvaRundownShowControl : public SCompoundWidget
 {
diff --git a/Engine/Plugins/Experimental/Avalanche/Source/AvalancheMediaEditor/Public/Broadcast/AvaBroadcastEditor.h b/Engine/Plugins/Experimental/Avalanche/Source/AvalancheMediaEditor/Public/Broadcast/AvaBroadcastEditor.h
--- a/Engine/Plugins/Experimental/Avalanche/Source/AvalancheMediaEditor/Public/Broadcast/AvaBroadcastEditor.h
+++ b/Engine/Plugins/Experimental/Avalanche/Source/AvalancheMediaEditor/Public/Broadcast/AvaBroadcastEditor.h
@@ -3,19 +3,28 @@
 #pragma once
 
 #include "AvaMediaDefines.h"
-#include "CoreMinimal.h"
 #include "WorkflowOrientedApp/WorkflowCentricApplication.h"
+#include "Templates/SharedPointer.h"
+#include "UObject/WeakObjectPtrTemplates.h"
 
 class FAvaBroadcastOutputTileItem;
+class FName;
+class FReply;
+class FText;
+class FToolBarBuilder;
+class FExtender;
 class UAvaBroadcast;
 class UMediaOutput;
+class UObject;
+struct FLinearColor;
+struct FPropertyChangedEvent;
+struct FSlateBrush;
 
 class FAvaBroadcastEditor : public FWorkflowCentricApplication
 {
 	struct FPrivateToken { explicit FPrivateToken() = default; };
 	
 public:
-
 	explicit FAvaBroadcastEditor(FPrivateToken) {}
 	
 	virtual ~FAvaBroadcastEditor() override;
@@ -29,7 +38,6 @@ public:
 	FOnOutputItemSelectionChanged OnOutputTileSelectionChanged;
 	
 protected:
-	
 	void InitBroadcastEditor(UAvaBroadcast* InBroadcast);
 
 	void OnBroadcastChanged(EAvaBroadcastChange ChangedEvent);
@@ -55,7 +63,6 @@ protected:
 	//~ End FAssetEditorToolkit Interface
 
 public:
-	
 	UAvaBroadcast* GetBroadcastObject() const;
 
 	void ExtendToolBar(TSharedPtr<FExtender> Extender);
@@ -69,7 +76,6 @@ public:
 	FReply OnProfileSelected(FName InProfileName);
 	
 protected:
-	
 	void RegisterApplicationModes();
 	void CreateDefaultCommands();
 
@@ -78,7 +84,6 @@ protected:
 	static void StartPlaybackClientAction();
 
 protected:
-	
 	static TSharedPtr<FAvaBroadcastEditor> BroadcastEditor;
 	
 	TWeakObjectPtr<UAvaBroadcast> BroadcastWeak;
diff --git a/Engine/Plugins/Experimental/Avalanche/Source/AvalancheMediaEditor/Public/Broadcast/ChannelGrid/AvaBroadcastOutputTileItem.h b/Engine/Plugins/Experimental/Avalanche/Source/AvalancheMediaEditor/Public/Broadcast/ChannelGrid/AvaBroadcastOutputTileItem.h
--- a/Engine/Plugins/Experimental/Avalanche/Source/AvalancheMediaEditor/Public/Broadcast/ChannelGrid/AvaBroadcastOutputTileItem.h
+++ b/Engine/Plugins/Experimental/Avalanche/Source/AvalancheMediaEditor/Public/Broadcast/ChannelGrid/AvaBroadcastOutputTileItem.h
@@ -3,18 +3,29 @@
 #pragma once
 
 #include "Broadcast/Channel/AvaBroadcastOutputChannel.h"
-#include "CoreMinimal.h"
+#include "Delegates/IDelegateInstance.h"
+#include "Internationalization/Text.h"
+#include "Templates/SharedPointer.h"
+#include "UObject/NameTypes.h"
+#include "UObject/WeakObjectPtrTemplates.h"
 
 class FAvaBroadcastOutputTileItem;
+class FReply;
+class SWidget;
+class UObject;
 class UMediaOutput;
 struct FAvaBroadcastOutputChannel;
+struct FGeometry;
+struct FPointerEvent;
+struct FPropertyChangedEvent;
+struct FSlateBrush;
+enum class EAvaBroadcastChange : uint8;
 
 using FAvaBroadcastOutputTileItemPtr = TSharedPtr<FAvaBroadcastOutputTileItem>;
 
 class FAvaBroadcastOutputTileItem : public TSharedFromThis<FAvaBroadcastOutputTileItem>
 {
 public:
-
 	FAvaBroadcastOutputTileItem(FName InChannelName, UMediaOutput* InMediaOutput);
 	~FAvaBroadcastOutputTileItem();
 	
@@ -35,7 +46,6 @@ public:
 	FReply OnDragDetected(const FGeometry& MyGeometry, const FPointerEvent& MouseEvent);
 	
 protected:
-
 	void OnMediaOutputPropertyChanged(UObject* InObject, FPropertyChangedEvent& PropertyChangedEvent);
 	void OnChannelChanged(const FAvaBroadcastOutputChannel& InChannel, EAvaBroadcastChannelChange InChange);
 	void OnMediaOutputStateChanged(const FAvaBroadcastOutputChannel& InChannel, const UMediaOutput* InMediaOutput);
@@ -46,7 +56,6 @@ protected:
 	FText FindLatestDisplayText() const;
 
 protected:
-	
 	FName ChannelName = NAME_None;
 
 	TWeakObjectPtr<UMediaOutput> MediaOutput;
@@ -59,4 +68,3 @@ protected:
 
 	FDelegateHandle BroadcastChangedHandle;
 };
-
diff --git a/Engine/Plugins/Experimental/Avalanche/Source/AvalancheMediaEditor/Public/IAvaMediaEditorModule.h b/Engine/Plugins/Experimental/Avalanche/Source/AvalancheMediaEditor/Public/IAvaMediaEditorModule.h
--- a/Engine/Plugins/Experimental/Avalanche/Source/AvalancheMediaEditor/Public/IAvaMediaEditorModule.h
+++ b/Engine/Plugins/Experimental/Avalanche/Source/AvalancheMediaEditor/Public/IAvaMediaEditorModule.h
@@ -2,18 +2,23 @@
 
 #pragma once
 
-#include "CoreMinimal.h"
+#include "Containers/Array.h"
 #include "Modules/ModuleInterface.h"
 #include "Modules/ModuleManager.h"
 #include "Toolkits/AssetEditorToolkit.h"
-
-DECLARE_LOG_CATEGORY_EXTERN(LogAvaMediaEditor, Log, All);
+#include "Templates/SharedPointer.h"
 
 class FAvaRundownServer;
 class IAvaRundownFilterSuggestionFactory;
-enum class EAvaRundownSearchListType : uint8;
+class FExtensibilityManager;
+class FName;
 struct FAvaRundownPage;
 struct FAvaRundownTextFilterArgs;
+struct FSlateIcon;
+enum class EAvaRundownSearchListType : uint8;
+enum class ETextFilterComparisonOperation : uint8;
+
+DECLARE_LOG_CATEGORY_EXTERN(LogAvaMediaEditor, Log, All);
 
 class IAvaMediaEditorModule : public IModuleInterface
 {
diff --git a/Engine/Plugins/Experimental/Avalanche/Source/AvalancheMediaEditor/Public/Playback/AvaPlaybackGraphEditor.h b/Engine/Plugins/Experimental/Avalanche/Source/AvalancheMediaEditor/Public/Playback/AvaPlaybackGraphEditor.h
--- a/Engine/Plugins/Experimental/Avalanche/Source/AvalancheMediaEditor/Public/Playback/AvaPlaybackGraphEditor.h
+++ b/Engine/Plugins/Experimental/Avalanche/Source/AvalancheMediaEditor/Public/Playback/AvaPlaybackGraphEditor.h
@@ -2,31 +2,50 @@
 
 #pragma once
 
-#include "CoreMinimal.h"
+#include "Containers/Array.h"
+#include "Containers/Map.h"
+#include "Containers/Set.h"
+#include "Containers/StringFwd.h"
+#include "Math/MathFwd.h"
 #include "Playback/IAvaPlaybackGraphEditor.h"
 #include "Templates/SharedPointer.h"
+#include "Templates/SubclassOf.h"
+#include "Toolkits/IToolkit.h"
+#include "Toolkits/IToolkitHost.h"
+#include "Types/SlateEnums.h"
 #include "WorkflowOrientedApp/WorkflowCentricApplication.h"
+#include "UObject/WeakObjectPtrTemplates.h"
 
+class FExtender;
+class FName;
 class FSlateRect;
+class FText;
+class FToolBarBuilder;
+class FUICommandList;
 class SGraphEditor;
 class UAvaPlaybackEditorGraphNode;
+class UAvaPlaybackGraph;
+class UAvaPlaybackNode;
+class UEdGraph;
+class UEdGraphNode;
+class UObject;
+struct FLinearColor;
 
 class FAvaPlaybackGraphEditor
 	: public FWorkflowCentricApplication
 	, public IAvaPlaybackGraphEditor
 {
 public:
-
 	void InitPlaybackEditor(const EToolkitMode::Type InMode
 		, const TSharedPtr<IToolkitHost>& InitToolkitHost
 		, UAvaPlaybackGraph* InPlayback);
 
-	//IToolkit Interface
+	//~ Begin IToolkit
 	virtual FName GetToolkitFName() const override;
 	virtual FText GetBaseToolkitName() const override;
 	virtual FString GetWorldCentricTabPrefix() const override;
 	virtual FLinearColor GetWorldCentricTabColorScale() const override;
-	//~IToolkit Interface
+	//~ End IToolkit
 
 	UAvaPlaybackGraph* GetPlaybackObject() const;
 	
@@ -34,12 +53,11 @@ public:
 	void FillPlayToolBar(FToolBarBuilder& ToolBarBuilder);
 
 public:
-
 	static void CacheOverrideNodeClasses();
 	
 	TSharedRef<SGraphEditor> CreateGraphEditor();
 	
-	//IAvaPlaybackGraphEditor Interface
+	//~ Begin IAvaPlaybackGraphEditor
 	virtual UEdGraph* CreatePlaybackGraph(UAvaPlaybackGraph* InPlayback) override;
 	virtual void SetupPlaybackNode(UEdGraph* InGraph, UAvaPlaybackNode* InPlaybackNode, bool bSelectNewNode) override;
 	virtual void CompilePlaybackNodesFromGraphNodes(UAvaPlaybackGraph* InPlayback) override;
@@ -47,7 +65,7 @@ public:
 	virtual void RefreshNode(UEdGraphNode& InGraphNode) override;
 	virtual bool GetBoundsForSelectedNodes(FSlateRect& Rect, float Padding) override;
 	virtual TSet<UObject*> GetSelectedNodes() const override;
-	//~IAvaPlaybackGraphEditor Interface
+	//~ End IAvaPlaybackGraphEditor
 	
 	/** Called when the selection changes in the GraphEditor */
 	void OnSelectedNodesChanged(const TSet<UObject*>& NewSelection);
@@ -65,14 +83,12 @@ public:
 	FOnPlaybackSelectionChanged OnPlaybackSelectionChanged;
 	
 protected:
-	
 	void RegisterApplicationModes();
 	
 	void CreateDefaultCommands();
 	void CreateGraphCommands();
 
 public:
-	
 	bool CanAddInputPin() const;
 	void AddInputPin();
 
@@ -113,7 +129,6 @@ public:
 	void OnDistributeNodesV();
 	
 protected:
-	
 	TWeakObjectPtr<UAvaPlaybackGraph> PlaybackGraphWeak;
 	
 	TSharedPtr<FUICommandList> GraphEditorCommands;
diff --git a/Engine/Plugins/Experimental/Avalanche/Source/AvalancheMediaEditor/Public/Playback/Graph/AvaPlaybackEditorGraph.h b/Engine/Plugins/Experimental/Avalanche/Source/AvalancheMediaEditor/Public/Playback/Graph/AvaPlaybackEditorGraph.h
--- a/Engine/Plugins/Experimental/Avalanche/Source/AvalancheMediaEditor/Public/Playback/Graph/AvaPlaybackEditorGraph.h
+++ b/Engine/Plugins/Experimental/Avalanche/Source/AvalancheMediaEditor/Public/Playback/Graph/AvaPlaybackEditorGraph.h
@@ -2,8 +2,8 @@
 
 #pragma once
 
-#include "CoreMinimal.h"
 #include "EdGraph/EdGraph.h"
+#include "Templates/SubclassOf.h"
 #include "AvaPlaybackEditorGraph.generated.h"
 
 class UAvaPlaybackGraph;
@@ -15,7 +15,6 @@ class UAvaPlaybackEditorGraph : public UEdGraph
 	GENERATED_BODY()
 
 public:
-	
 	UAvaPlaybackGraph* GetPlaybackGraph() const;
 	
 	UAvaPlaybackEditorGraphNode* CreatePlaybackEditorGraphNode(TSubclassOf<UAvaPlaybackEditorGraphNode> NewNodeClass
diff --git a/Engine/Plugins/Experimental/Avalanche/Source/AvalancheMediaEditor/Public/Playback/Graph/AvaPlaybackEditorGraphSchema.h b/Engine/Plugins/Experimental/Avalanche/Source/AvalancheMediaEditor/Public/Playback/Graph/AvaPlaybackEditorGraphSchema.h
--- a/Engine/Plugins/Experimental/Avalanche/Source/AvalancheMediaEditor/Public/Playback/Graph/AvaPlaybackEditorGraphSchema.h
+++ b/Engine/Plugins/Experimental/Avalanche/Source/AvalancheMediaEditor/Public/Playback/Graph/AvaPlaybackEditorGraphSchema.h
@@ -2,25 +2,35 @@
 
 #pragma once
 
-#include "CoreMinimal.h"
+#include "Containers/Array.h"
+#include "Containers/StringFwd.h"
 #include "EdGraph/EdGraphSchema.h"
+#include "Math/MathFwd.h"
+#include "Templates/SubclassOf.h"
+#include "Templates/SharedPointer.h"
 #include "AvaPlaybackEditorGraphSchema.generated.h"
 
 class IAvaPlaybackGraphEditor;
+class FName;
 class UAvaPlaybackGraph;
 class UAvaPlaybackNode;
+class UEdGraph;
+class UEdGraphNode;
+class UEdGraphPin;
+class UGraphNodeContextMenuContext;
+class UToolMenu;
+struct FAssetData;
+struct FEdGraphPinType;
+struct FGraphContextMenuBuilder;
+struct FLinearColor;
+struct FPinConnectionResponse;
 
 UCLASS()
 class UAvaPlaybackEditorGraphSchema : public UEdGraphSchema
 {
 	GENERATED_BODY()
 
-protected:
-
-	TSharedPtr<IAvaPlaybackGraphEditor> GetPlaybackGraphEditor(const UEdGraph* Graph) const;
-
 public:
-
 	static const FLinearColor ActivePinColor;
 	static const FLinearColor InactivePinColor;
 	
@@ -33,7 +43,7 @@ public:
 	/** Check whether connecting these pins would cause a loop */
 	bool ConnectionCausesLoop(const UEdGraphPin* InputPin, const UEdGraphPin* OutputPin) const;
 	
-	//UEdGraphSchema Interface.
+	//~ Begin UEdGraphSchema
 	virtual void GetGraphContextActions(FGraphContextMenuBuilder& ContextMenuBuilder) const override;
 	virtual void GetContextMenuActions(UToolMenu* Menu, class UGraphNodeContextMenuContext* Context) const override;
 	virtual void CreateDefaultNodesForGraph(UEdGraph& Graph) const override;
@@ -50,7 +60,7 @@ public:
 	virtual void DroppedAssetsOnNode(const TArray<FAssetData>& Assets, const FVector2D& GraphPosition, UEdGraphNode* Node) const override;
 	virtual int32 GetNodeSelectionCount(const UEdGraph* Graph) const override;
 	virtual TSharedPtr<FEdGraphSchemaAction> GetCreateCommentAction() const override;
-	//~UEdGraphSchema
+	//~ End UEdGraphSchema
 
 	static void CachePlaybackNodeClasses();
 	
@@ -59,7 +69,7 @@ public:
 	void GetCommentAction(FGraphActionMenuBuilder& ActionMenuBuilder, const UEdGraph* CurrentGraph = nullptr) const;
 
 protected:
-
 	static TArray<TSubclassOf<UAvaPlaybackNode>> PlaybackNodeClasses;
-	
+
+	TSharedPtr<IAvaPlaybackGraphEditor> GetPlaybackGraphEditor(const UEdGraph* Graph) const;
 };
diff --git a/Engine/Plugins/Experimental/Avalanche/Source/AvalancheMediaEditor/Public/Playback/Graph/Nodes/AvaPlaybackEditorGraphNode.h b/Engine/Plugins/Experimental/Avalanche/Source/AvalancheMediaEditor/Public/Playback/Graph/Nodes/AvaPlaybackEditorGraphNode.h
--- a/Engine/Plugins/Experimental/Avalanche/Source/AvalancheMediaEditor/Public/Playback/Graph/Nodes/AvaPlaybackEditorGraphNode.h
+++ b/Engine/Plugins/Experimental/Avalanche/Source/AvalancheMediaEditor/Public/Playback/Graph/Nodes/AvaPlaybackEditorGraphNode.h
@@ -2,13 +2,24 @@
 
 #pragma once
 
-#include "CoreMinimal.h"
+#include "Containers/Array.h"
+#include "Containers/Set.h"
 #include "EdGraph/EdGraphNode.h"
 #include "Templates/SubclassOf.h"
+#include "Templates/SharedPointer.h"
+#include "UObject/ObjectPtr.h"
 #include "AvaPlaybackEditorGraphNode.generated.h"
 
+class FName;
+class SGraphNode;
 class UAvaPlaybackGraph;
 class UAvaPlaybackNode;
+class UEdGraphPin;
+class UEdGraphSchema;
+class UGraphNodeContextMenuContext;
+class UToolMenu;
+class UObject;
+struct FLinearColor;
 
 UCLASS()
 class AVALANCHEMEDIAEDITOR_API UAvaPlaybackEditorGraphNode : public UEdGraphNode
@@ -16,7 +27,6 @@ class AVALANCHEMEDIAEDITOR_API UAvaPlaybackEditorGraphNode : public UEdGraphNode
 	GENERATED_BODY()
 
 public:
-
 	UAvaPlaybackGraph* GetPlayback() const;
 	
 	void SetPlaybackNode(UAvaPlaybackNode* InPlaybackNode);
@@ -49,14 +59,12 @@ public:
 	}
 	
 protected:
-	
 	/** Create all of the input pins required */
 	void CreateInputPins();
 	void CreateOutputPin();
 
 public:
-	
-	/** Is this the undeletable root node */
+	/** Is this the non-deletable root node */
 	bool IsRootNode() const;
 	
 	UEdGraphPin* GetOutputPin() const;
@@ -75,13 +83,13 @@ public:
 	 */
 	void InsertNewNode(UEdGraphPin* FromPin, UEdGraphPin* NewLinkPin, TSet<UEdGraphNode*>& OutNodeList);
 
-	//UObject interface
+	//~ Begin UObject
 	virtual void PostLoad() override;
 	virtual void PostEditImport() override;
 	virtual void PostDuplicate(bool bDuplicateForPIE) override;
-	//~UObject interface
+	//~ End UObject
 	
-	//UEdGraphNode
+	//~ Begin UEdGraphNode
 	virtual FText GetNodeTitle(ENodeTitleType::Type TitleType) const override;
 	virtual FText GetTooltipText() const override;
 	virtual void GetNodeContextMenuActions(UToolMenu* Menu, UGraphNodeContextMenuContext* Context) const override;
@@ -91,7 +99,7 @@ public:
 	virtual bool CanCreateUnderSpecifiedSchema(const UEdGraphSchema* Schema) const override;
 	virtual FLinearColor GetNodeTitleColor() const override;
 	virtual TSharedPtr<SGraphNode> CreateVisualWidget() override;
-	//~UEdGraphNode
+	//~ End UEdGraphNode
 	
 	virtual FName GetInputPinName(int32 InputPinIndex) const;
 	virtual FName GetInputPinCategory(int32 InputPinIndex) const;
@@ -101,7 +109,6 @@ public:
 	virtual FName GetOutputPinCategory() const;
 	
 protected:
-	
 	// The Playback Node this represents
 	UPROPERTY(VisibleAnywhere, Instanced, Category=Playback)
 	TObjectPtr<UAvaPlaybackNode> PlaybackNode;
diff --git a/Engine/Plugins/Experimental/Avalanche/Source/AvalancheMediaEditor/Public/Playback/Graph/Nodes/Slate/SAvaPlaybackEditorGraphNode.h b/Engine/Plugins/Experimental/Avalanche/Source/AvalancheMediaEditor/Public/Playback/Graph/Nodes/Slate/SAvaPlaybackEditorGraphNode.h
--- a/Engine/Plugins/Experimental/Avalanche/Source/AvalancheMediaEditor/Public/Playback/Graph/Nodes/Slate/SAvaPlaybackEditorGraphNode.h
+++ b/Engine/Plugins/Experimental/Avalanche/Source/AvalancheMediaEditor/Public/Playback/Graph/Nodes/Slate/SAvaPlaybackEditorGraphNode.h
@@ -2,16 +2,20 @@
 
 #pragma once
 
-#include "CoreMinimal.h"
 #include "SGraphNode.h"
+#include "Templates/SharedPointer.h"
+#include "UObject/WeakObjectPtrTemplates.h"
 
+class FReply;
+class SGraphPin;
 class SVerticalBox;
 class UAvaPlaybackEditorGraphNode;
+class UEdGraphPin;
+struct EVisibility;
 
 class AVALANCHEMEDIAEDITOR_API SAvaPlaybackEditorGraphNode : public SGraphNode
 {
 public:
-	
 	SLATE_BEGIN_ARGS(SAvaPlaybackEditorGraphNode) {}
 	SLATE_END_ARGS()
 
@@ -20,15 +24,13 @@ public:
 	virtual void PostConstruct() {};
 	
 protected:
-	
-	//SGraphNode Interface
+	//~ Begin SGraphNode
 	virtual void CreateOutputSideAddButton(TSharedPtr<SVerticalBox> OutputBox) override;
 	virtual EVisibility IsAddPinButtonVisible() const override;
 	virtual FReply OnAddPin() override;
 	virtual TSharedPtr<SGraphPin> CreatePinWidget(UEdGraphPin* Pin) const override;
-	//~SGraphNode Interface
+	//~ End SGraphNode
 	
 protected:
-
 	TWeakObjectPtr<UAvaPlaybackEditorGraphNode> PlaybackGraphNode;
 };
diff --git a/Engine/Plugins/Experimental/Avalanche/Source/AvalancheMediaEditor/Public/Playback/Graph/SchemaActions/AvaPlaybackAction_NewNode.h b/Engine/Plugins/Experimental/Avalanche/Source/AvalancheMediaEditor/Public/Playback/Graph/SchemaActions/AvaPlaybackAction_NewNode.h
--- a/Engine/Plugins/Experimental/Avalanche/Source/AvalancheMediaEditor/Public/Playback/Graph/SchemaActions/AvaPlaybackAction_NewNode.h
+++ b/Engine/Plugins/Experimental/Avalanche/Source/AvalancheMediaEditor/Public/Playback/Graph/SchemaActions/AvaPlaybackAction_NewNode.h
@@ -2,46 +2,41 @@
 
 #pragma once
 
-#include "CoreMinimal.h"
 #include "EdGraph/EdGraphSchema.h"
+#include "Math/MathFwd.h"
+#include "Playback/Nodes/AvaPlaybackNode.h"
 #include "Templates/SubclassOf.h"
 #include "AvaPlaybackAction_NewNode.generated.h"
 
+class FText;
 class UAvaPlaybackNode;
+class UEdGraph;
+class UEdGraphNode;
 
 USTRUCT()
 struct AVALANCHEMEDIAEDITOR_API FAvaPlaybackAction_NewNode : public FEdGraphSchemaAction
 {
-	
 	GENERATED_BODY()
 
 public:
-	
 	FAvaPlaybackAction_NewNode() = default;
 	
-	FAvaPlaybackAction_NewNode(FText InNodeCategory, FText InMenuDesc, FText InToolTip, const int32 InGrouping)
-		: FEdGraphSchemaAction(MoveTemp(InNodeCategory), MoveTemp(InMenuDesc), MoveTemp(InToolTip), InGrouping)
-	{
-	}
+	FAvaPlaybackAction_NewNode(FText InNodeCategory, FText InMenuDesc, FText InToolTip, const int32 InGrouping);
 
-	//~ Begin FEdGraphSchemaAction Interface
+	//~ Begin FEdGraphSchemaAction
 	virtual UEdGraphNode* PerformAction(class UEdGraph* ParentGraph
 		, UEdGraphPin* FromPin
 		, const FVector2D Location
 		, bool bSelectNewNode = true) override;
-	//~ End FEdGraphSchemaAction Interface
+	//~ End FEdGraphSchemaAction
 
 	void SetPlaybackNodeClass(TSubclassOf<UAvaPlaybackNode> InPlaybackNodeClass);
 
 protected:
-	
 	/** Connects new node to output of selected nodes */
 	void ConnectToSelectedNodes(UAvaPlaybackNode* NewNode, UEdGraph* ParentGraph) const;
 
-protected:
-	
 	/** Class of node we want to create */
 	UPROPERTY()
 	TSubclassOf<UAvaPlaybackNode> PlaybackNodeClass;
-	
 };
diff --git a/Engine/Plugins/Experimental/Avalanche/Source/AvalancheMediaEditor/Public/Rundown/Pages/Columns/IAvaRundownPageViewColumn.h b/Engine/Plugins/Experimental/Avalanche/Source/AvalancheMediaEditor/Public/Rundown/Pages/Columns/IAvaRundownPageViewColumn.h
--- a/Engine/Plugins/Experimental/Avalanche/Source/AvalancheMediaEditor/Public/Rundown/Pages/Columns/IAvaRundownPageViewColumn.h
+++ b/Engine/Plugins/Experimental/Avalanche/Source/AvalancheMediaEditor/Public/Rundown/Pages/Columns/IAvaRundownPageViewColumn.h
@@ -3,11 +3,14 @@
 #pragma once
 
 #include "AvaType.h"
-#include "CoreMinimal.h"
 #include "Rundown/AvaRundownDefines.h"
+#include "Templates/SharedPointer.h"
 #include "Widgets/Views/SHeaderRow.h"
 
+class FName;
+class FText;
 class SAvaRundownPageViewRow;
+class SWidget;
 
 class IAvaRundownPageViewColumn : public IAvaTypeCastable, public TSharedFromThis<IAvaRundownPageViewColumn>
 {
diff --git a/Engine/Plugins/Experimental/Avalanche/Source/AvalancheMediaEditor/Public/Rundown/Pages/PageViews/IAvaRundownPageView.h b/Engine/Plugins/Experimental/Avalanche/Source/AvalancheMediaEditor/Public/Rundown/Pages/PageViews/IAvaRundownPageView.h
--- a/Engine/Plugins/Experimental/Avalanche/Source/AvalancheMediaEditor/Public/Rundown/Pages/PageViews/IAvaRundownPageView.h
+++ b/Engine/Plugins/Experimental/Avalanche/Source/AvalancheMediaEditor/Public/Rundown/Pages/PageViews/IAvaRundownPageView.h
@@ -3,11 +3,14 @@
 #pragma once
 
 #include "AvaType.h"
-#include "CoreMinimal.h"
+#include "Delegates/Delegate.h"
+#include "Templates/SharedPointer.h"
 
 class FReply;
+class FText;
 class UAvaRundown;
 struct FAssetData;
+struct FSoftObjectPath;
 
 enum class EAvaRundownPageActionState : uint8
 {
diff --git a/Engine/Plugins/Experimental/Avalanche/Source/AvalancheMediaEditor/Public/Rundown/Pages/Slate/SAvaRundownPageViewRow.h b/Engine/Plugins/Experimental/Avalanche/Source/AvalancheMediaEditor/Public/Rundown/Pages/Slate/SAvaRundownPageViewRow.h
--- a/Engine/Plugins/Experimental/Avalanche/Source/AvalancheMediaEditor/Public/Rundown/Pages/Slate/SAvaRundownPageViewRow.h
+++ b/Engine/Plugins/Experimental/Avalanche/Source/AvalancheMediaEditor/Public/Rundown/Pages/Slate/SAvaRundownPageViewRow.h
@@ -2,12 +2,23 @@
 
 #pragma once
 
-#include "CoreMinimal.h"
+#include "Containers/Array.h"
+#include "Input/DragAndDrop.h"
 #include "Rundown/AvaRundownDefines.h"
+#include "Templates/SharedPointer.h"
 #include "Widgets/Views/SListView.h"
 #include "Widgets/Views/STableRow.h"
+#include "UObject/NameTypes.h"
+#include "Misc/Optional.h"
 
+class FDragDropEvent;
+class FReply;
+class IAvaRundownPageView;
 class SAvaRundownPageList;
+class SWidget;
+struct FGeometry;
+struct FPointerEvent;
+struct FSlateBrush;
 
 enum class EAvaRundownRowState
 {
@@ -20,8 +31,6 @@ class FAvaRundownPageViewRowDragDropOp : public FDragDropOperation, public TShar
 public:
 	DRAG_DROP_OPERATOR_TYPE(FAvaRundownPageViewRowDragDropOp, FDragDropOperation)
 
-	FName Color;
-
 	virtual TSharedPtr<SWidget> GetDefaultDecorator() const override;
 
 	static TSharedRef<FAvaRundownPageViewRowDragDropOp> New(const TSharedPtr<SAvaRundownPageList>& InPageList);
@@ -29,6 +38,8 @@ public:
 	TSharedPtr<SAvaRundownPageList> GetPageList() const { return PageListWeak.Pin(); }
 	const TArray<int32>& GetDraggedIds() const { return DraggedIds; }
 
+	FName Color;
+
 protected:
 	TWeakPtr<SAvaRundownPageList> PageListWeak;
 	TArray<int32> DraggedIds;
