Motion Design:
- Remove platform constraint of only Windows and Linux

#changelist validated
#virtualized

diff --git a/Engine/Plugins/Experimental/Avalanche/Avalanche.uplugin b/Engine/Plugins/Experimental/Avalanche/Avalanche.uplugin
--- a/Engine/Plugins/Experimental/Avalanche/Avalanche.uplugin
+++ b/Engine/Plugins/Experimental/Avalanche/Avalanche.uplugin
@@ -1,7 +1,7 @@
 {
 	"FileVersion": 3,
 	"Version": 0.1,
-	"VersionName": "0.1.0-alpha",
+	"VersionName": "0.1",
 	"FriendlyName": "Motion Design",
 	"Description": "Compositing, designer and broadcasting tool.\n\nPlugin Dependencies: Advanced Renamer, Custom Details View, Dynamic Material, Geometry Cache, Geometry Scripting, Media Compositing, Media IO Framework, Mesh Modeling Toolset Exp, Remote Control, SVG Importer, Text3D and ActorModifierCore.",
 	"Category": "Experimental",
@@ -11,12 +11,14 @@
 	"MarketplaceURL": "",
 	"SupportURL": "",
 	"CanContainContent": true,
-	"IsBetaVersion": true,
+	"IsBetaVersion": false,
+	"IsExperimentalVersion": true,
 	"Installed": false,
 	"EnabledByDefault": false,
 	"SupportedTargetPlatforms": [
 		"Win64",
-		"Linux"
+		"Linux",
+		"Mac"
 	],
 	"Modules": [
 		{
