[Motion Design]
 - Add missing headers for Horde

diff --git a/Engine/Plugins/Experimental/Avalanche/Source/AvalancheMediaEditor/Private/Rundown/DetailsView/RemoteControl/Properties/SAvaRundownRCPropertyItemRow.cpp b/Engine/Plugins/Experimental/Avalanche/Source/AvalancheMediaEditor/Private/Rundown/DetailsView/RemoteControl/Properties/SAvaRundownRCPropertyItemRow.cpp
--- a/Engine/Plugins/Experimental/Avalanche/Source/AvalancheMediaEditor/Private/Rundown/DetailsView/RemoteControl/Properties/SAvaRundownRCPropertyItemRow.cpp
+++ b/Engine/Plugins/Experimental/Avalanche/Source/AvalancheMediaEditor/Private/Rundown/DetailsView/RemoteControl/Properties/SAvaRundownRCPropertyItemRow.cpp
@@ -1,13 +1,16 @@
 // Copyright Epic Games, Inc. All Rights Reserved.
 
 #include "SAvaRundownRCPropertyItemRow.h"
-
 #include "AvaRundownPageRemoteControlWidgetUtils.h"
 #include "AvaRundownRCPropertyItem.h"
 #include "IDetailTreeNode.h"
+#include "Internationalization/Text.h"
 #include "IPropertyRowGenerator.h"
 #include "Modules/ModuleManager.h"
 #include "PropertyEditorModule.h"
+#include "RemoteControlEntity.h"
+#include "RemoteControlField.h"
+#include "Rundown/DetailsView/RemoteControl/Properties/AvaRundownPageRemoteControlWidgetUtils.h"
 #include "Widgets/Layout/SScissorRectBox.h"
 #include "Widgets/SNullWidget.h"
 #include "Widgets/Text/STextBlock.h"
@@ -129,7 +132,7 @@ TSharedRef<SWidget> SAvaRundownRCPropertyItemRow::CreateValue()
 							[
 								ValueWidget.ToSharedRef()
 							]
-						+ SHorizontalBox::Slot()
+							+ SHorizontalBox::Slot()
 							.VAlign(EVerticalAlignment::VAlign_Center)
 							.Padding(3.f, 0.f, 0.f, 0.f)
 							[
diff --git a/Engine/Plugins/Experimental/Avalanche/Source/AvalancheMediaEditor/Private/Rundown/DetailsView/RemoteControl/Properties/SAvaRundownRCPropertyItemRow.h b/Engine/Plugins/Experimental/Avalanche/Source/AvalancheMediaEditor/Private/Rundown/DetailsView/RemoteControl/Properties/SAvaRundownRCPropertyItemRow.h
--- a/Engine/Plugins/Experimental/Avalanche/Source/AvalancheMediaEditor/Private/Rundown/DetailsView/RemoteControl/Properties/SAvaRundownRCPropertyItemRow.h
+++ b/Engine/Plugins/Experimental/Avalanche/Source/AvalancheMediaEditor/Private/Rundown/DetailsView/RemoteControl/Properties/SAvaRundownRCPropertyItemRow.h
@@ -2,16 +2,21 @@
 
 #pragma once
 
+#include "AvaRundownRCPropertyItem.h"
+#include "IPropertyRowGenerator.h"
+#include "Rundown/DetailsView/RemoteControl/Properties/AvaRundownRCPropertyItem.h"
 #include "SAvaRundownPageRemoteControlProps.h"
+#include "Templates/SharedPointer.h"
 #include "Widgets/Views/STableRow.h"
+#include "Widgets/Layout/SBox.h"
+#include "Widgets/SWidget.h"
 
-class SBox;
-class SWidget;
+class FText;
 
 class SAvaRundownRCPropertyItemRow : public SMultiColumnTableRow<FAvaRundownRCPropertyItemPtr>
 {
 public:
-	SLATE_BEGIN_ARGS(SAvaRundownRCPropertyItemRow) { }
+	SLATE_BEGIN_ARGS(SAvaRundownRCPropertyItemRow) {}
 	SLATE_END_ARGS()
 
 	void Construct(const FArguments& InArgs, TSharedRef<SAvaRundownPageRemoteControlProps> InPropertyPanel,
