[Avalanche] Mask modifier fix for settings synchronization

diff --git a/Engine/Plugins/Experimental/Avalanche/Source/AvalancheMask/Private/Mask2D/AvaMask2DBaseModifier.cpp b/Engine/Plugins/Experimental/Avalanche/Source/AvalancheMask/Private/Mask2D/AvaMask2DBaseModifier.cpp
--- a/Engine/Plugins/Experimental/Avalanche/Source/AvalancheMask/Private/Mask2D/AvaMask2DBaseModifier.cpp
+++ b/Engine/Plugins/Experimental/Avalanche/Source/AvalancheMask/Private/Mask2D/AvaMask2DBaseModifier.cpp
@@ -12,6 +12,7 @@
 #include "GeometryMaskCanvas.h"
 #include "GeometryMaskReadComponent.h"
 #include "GeometryMaskSubsystem.h"
+#include "GeometryMaskWorldSubsystem.h"
 #include "GeometryMaskWriteComponent.h"
 #include "Handling/AvaHandleUtilities.h"
 #include "Handling/AvaObjectHandleSubsystem.h"
@@ -46,6 +47,16 @@ namespace UE::AvaMask::Private
 
 		return InNew.Emplace(InKey, InValueFactory());
 	}
+
+	// @note: It's crucial these remain in sync with the defaults in UGeometryMaskCanvas
+	namespace CanvasPropertyDefaults
+	{
+		static bool bApplyBlur = false;
+		static double BlurStrength = 16;
+		static bool bApplyFeather = false;
+		static int32 OuterFeatherRadius = 16;
+		static int32 InnerFeatherRadius = 16;
+	}
 }
 
 #if WITH_EDITOR
@@ -59,6 +70,7 @@ const TAvaPropertyChangeDispatcher<UAvaMask2DBaseModifier> UAvaMask2DBaseModifie
 	{ GET_MEMBER_NAME_CHECKED(UAvaMask2DBaseModifier, bUseFeathering), &UAvaMask2DBaseModifier::OnFeatherChanged },
 	{ GET_MEMBER_NAME_CHECKED(UAvaMask2DBaseModifier, OuterFeatherRadius), &UAvaMask2DBaseModifier::OnFeatherChanged },
 	{ GET_MEMBER_NAME_CHECKED(UAvaMask2DBaseModifier, InnerFeatherRadius), &UAvaMask2DBaseModifier::OnFeatherChanged },
+	{ GET_MEMBER_NAME_CHECKED(UAvaMask2DBaseModifier, CanvasWeak), &UAvaMask2DBaseModifier::OnCanvasChanged }
 };
 #endif
 
@@ -163,6 +175,66 @@ void UAvaMask2DBaseModifier::OnFeatherChanged()
 	MarkModifierDirty();
 }
 
+void UAvaMask2DBaseModifier::OnCanvasChanged()
+{
+	OnBlurChanged();
+	OnFeatherChanged();
+}
+
+void UAvaMask2DBaseModifier::CanvasParamsToLocal()
+{
+	if (UGeometryMaskCanvas* Canvas = GetCurrentCanvas())
+	{
+		bUseBlur = Canvas->IsBlurApplied();
+		BlurStrength = Canvas->GetBlurStrength();
+
+		bUseFeathering = Canvas->IsFeatherApplied();
+		OuterFeatherRadius = Canvas->GetOuterFeatherRadius();
+		InnerFeatherRadius = Canvas->GetInnerFeatherRadius();
+	}
+}
+
+void UAvaMask2DBaseModifier::LocalParamsToCanvas()
+{
+	if (UGeometryMaskCanvas* Canvas = GetCurrentCanvas())
+	{
+		// Only set if not defaults
+		
+		const bool bModifierUseBlur = bUseBlur;
+		if (bModifierUseBlur != UE::AvaMask::Private::CanvasPropertyDefaults::bApplyBlur)
+		{
+			Canvas->SetApplyBlur(bModifierUseBlur);
+		}
+
+		const double ModifierBlurStrength = BlurStrength;		
+		if (!FMath::IsNearlyEqual(ModifierBlurStrength, UE::AvaMask::Private::CanvasPropertyDefaults::BlurStrength))
+		{
+			const double CanvasBlurStrength = Canvas->GetBlurStrength();
+			Canvas->SetBlurStrength(FMath::Max(ModifierBlurStrength, CanvasBlurStrength));
+		}
+
+		const bool bModifierUseFeathering = bUseFeathering;
+		if (bModifierUseFeathering != UE::AvaMask::Private::CanvasPropertyDefaults::bApplyFeather)
+		{
+			Canvas->SetApplyFeather(bModifierUseFeathering);
+		}
+
+		const int32 ModifierOuterFeatherRadius = OuterFeatherRadius;
+		if (ModifierOuterFeatherRadius != UE::AvaMask::Private::CanvasPropertyDefaults::OuterFeatherRadius)
+		{
+			const int32 CanvasOuterFeatherRadius = Canvas->GetOuterFeatherRadius();
+			Canvas->SetOuterFeatherRadius(FMath::Max(ModifierOuterFeatherRadius, CanvasOuterFeatherRadius));
+		}
+
+		const int32 ModifierInnerFeatherRadius = InnerFeatherRadius;
+		if (ModifierInnerFeatherRadius != UE::AvaMask::Private::CanvasPropertyDefaults::InnerFeatherRadius)
+		{
+			const int32 CanvasInnerFeatherRadius = Canvas->GetInnerFeatherRadius();
+			Canvas->SetInnerFeatherRadius(FMath::Max(ModifierInnerFeatherRadius, CanvasInnerFeatherRadius));
+		}
+	}
+}
+
 #if WITH_EDITOR
 void UAvaMask2DBaseModifier::PostEditChangeProperty(FPropertyChangedEvent& InPropertyChangedEvent)
 {
@@ -185,6 +257,11 @@ void UAvaMask2DBaseModifier::OnModifierAdded(EActorModifierCoreEnableReason InRe
 	{
 		SetupChannelName();
 	}
+
+	if (InReason == EActorModifierCoreEnableReason::Load)
+	{
+		LocalParamsToCanvas();
+	}
 }
 
 void UAvaMask2DBaseModifier::OnModifierRemoved(EActorModifierCoreDisableReason InReason)
@@ -193,6 +270,11 @@ void UAvaMask2DBaseModifier::OnModifierRemoved(EActorModifierCoreDisableReason I
 
 	TRACE_BOOKMARK(TEXT("UAvaMask2DModifier::OnModifierRemoved"));
 
+	if (InReason == EActorModifierCoreDisableReason::Destroyed)
+	{
+		CanvasParamsToLocal();
+	}
+
 	TArray<AActor*> ActorDataActors;
 	ActorDataActors.Reserve(ActorData.Num());	
 	Algo::Transform(ActorData, ActorDataActors, [](const TPair<TWeakObjectPtr<AActor>, FAvaMask2DActorData>& InActorDataPair)
@@ -269,15 +351,7 @@ void UAvaMask2DBaseModifier::RestorePreState()
 
 void UAvaMask2DBaseModifier::Apply()
 {
-	if (UGeometryMaskCanvas* Canvas = GetCurrentCanvas())
-	{
-		Canvas->SetApplyBlur(bUseBlur);
-		Canvas->SetBlurStrength(BlurStrength);
-
-		Canvas->SetApplyFeather(bUseFeathering);
-		Canvas->SetOuterFeatherRadius(OuterFeatherRadius);
-		Canvas->SetInnerFeatherRadius(InnerFeatherRadius);
-	}
+	LocalParamsToCanvas();
 }
 
 void UAvaMask2DBaseModifier::SaveActorPreState(AActor* InActor, FAvaMask2DActorData& InActorData)
@@ -321,8 +395,6 @@ void UAvaMask2DBaseModifier::OnSceneTreeTrackedActorParentChanged(
 {
 	Super::OnSceneTreeTrackedActorParentChanged(InIdx, InPreviousParentActor, InNewParentActor);
 
-	// @todo: remove stuff from child writers/readers if necessary
-
 	// We don't use the provided parent here, instead we just use the event to trigger custom parent discovery
 	if (bUseParentChannel)
 	{
diff --git a/Engine/Plugins/Experimental/Avalanche/Source/AvalancheMask/Public/Mask2D/AvaMask2DBaseModifier.h b/Engine/Plugins/Experimental/Avalanche/Source/AvalancheMask/Public/Mask2D/AvaMask2DBaseModifier.h
--- a/Engine/Plugins/Experimental/Avalanche/Source/AvalancheMask/Public/Mask2D/AvaMask2DBaseModifier.h
+++ b/Engine/Plugins/Experimental/Avalanche/Source/AvalancheMask/Public/Mask2D/AvaMask2DBaseModifier.h
@@ -126,6 +126,13 @@ protected:
 	void OnInvertedChanged();
 	void OnBlurChanged();
 	void OnFeatherChanged();
+	void OnCanvasChanged();
+
+	/** Stores specific properties on the canvas locally. */	
+	void CanvasParamsToLocal();
+
+	/** Apply locally stored parameters to the canvas. */
+	void LocalParamsToCanvas();
 
 	void SetupChannelName();
 	virtual void SetupMaskComponent(UActorComponent* InComponent);
@@ -176,19 +183,19 @@ protected:
 	UPROPERTY(EditInstanceOnly, BlueprintReadWrite, Getter = "IsInverted", Setter = "SetIsInverted", Category = "Mask2D", meta = (AllowPrivateAccess = "true"))
 	bool bInverted = false;
 
-	UPROPERTY(EditInstanceOnly, BlueprintReadWrite, Getter = "UseBlur", Setter = "UseBlur", Category = "Mask2D|Shared", meta = (AllowPrivateAccess = "true"))
+	UPROPERTY(BlueprintReadWrite, Getter = "UseBlur", Setter = "UseBlur", Category = "Mask2D|Shared", meta = (AllowPrivateAccess = "true"))
 	bool bUseBlur = false;
 
-	UPROPERTY(EditInstanceOnly, BlueprintReadWrite, Getter, Setter, Category = "Mask2D|Shared", meta = (ClampMin = 0.0, AllowPrivateAccess = "true", EditCondition = "bUseBlur", EditConditionHides))
-	float BlurStrength = 1.0f;
+	UPROPERTY(BlueprintReadWrite, Getter, Setter, Category = "Mask2D|Shared", meta = (ClampMin = 0.0, AllowPrivateAccess = "true", EditCondition = "bUseBlur", EditConditionHides))
+	float BlurStrength = 16.0f;
 
-	UPROPERTY(EditInstanceOnly, BlueprintReadWrite, Getter = "UseFeathering", Setter = "UseFeathering", Category = "Mask2D|Shared", meta = (AllowPrivateAccess = "true"))
+	UPROPERTY(BlueprintReadWrite, Getter = "UseFeathering", Setter = "UseFeathering", Category = "Mask2D|Shared", meta = (AllowPrivateAccess = "true"))
 	bool bUseFeathering = false;
 
-	UPROPERTY(EditInstanceOnly, BlueprintReadWrite, Getter, Setter, Category = "Mask2D|Shared", meta = (ClampMin = 0, AllowPrivateAccess = "true", EditCondition = "bUseFeathering", EditConditionHides))
+	UPROPERTY(BlueprintReadWrite, Getter, Setter, Category = "Mask2D|Shared", meta = (ClampMin = 0, AllowPrivateAccess = "true", EditCondition = "bUseFeathering", EditConditionHides))
 	int32 OuterFeatherRadius = 16;
 
-	UPROPERTY(EditInstanceOnly, BlueprintReadWrite, Getter, Setter, Category = "Mask2D|Shared", meta = (ClampMin = 0, AllowPrivateAccess = "true", EditCondition = "bUseFeathering", EditConditionHides))
+	UPROPERTY(BlueprintReadWrite, Getter, Setter, Category = "Mask2D|Shared", meta = (ClampMin = 0, AllowPrivateAccess = "true", EditCondition = "bUseFeathering", EditConditionHides))
 	int32 InnerFeatherRadius = 16;
 
 #if WITH_EDITOR
@@ -218,7 +225,7 @@ protected:
 	FName LastResolvedCanvasName;
 
 	/** Reference to the Canvas used */
-	UPROPERTY(Transient, DuplicateTransient)
+	UPROPERTY(EditAnywhere, Transient, DuplicateTransient, Category = "Mask2D|Shared", NoClear,	meta = (DisplayName = "Canvas", AllowPrivateAccess = "true", ShowInnerProperties, NoResetToDefault))
 	TWeakObjectPtr<UGeometryMaskCanvas> CanvasWeak;
 
 	/** Reference to the underlying canvas texture */
diff --git a/Engine/Plugins/Experimental/Avalanche/Source/AvalancheMaskEditor/Private/AvaMaskEditorModule.cpp b/Engine/Plugins/Experimental/Avalanche/Source/AvalancheMaskEditor/Private/AvaMaskEditorModule.cpp
--- a/Engine/Plugins/Experimental/Avalanche/Source/AvalancheMaskEditor/Private/AvaMaskEditorModule.cpp
+++ b/Engine/Plugins/Experimental/Avalanche/Source/AvalancheMaskEditor/Private/AvaMaskEditorModule.cpp
@@ -5,7 +5,11 @@
 #include "AvaMaskEditorCommands.h"
 #include "AvaMaskEditorMode.h"
 #include "AvaMaskEditorSubsystem.h"
+#include "Details/AvaMask2DModifierDetails.h"
 #include "EditorModeManager.h"
+#include "Mask2D/AvaMask2DReadModifier.h"
+#include "Mask2D/AvaMask2DWriteModifier.h"
+#include "PropertyEditorModule.h"
 #include "Templates/SharedPointer.h"
 #include "ToolMenuEntry.h"
 #include "ToolMenus.h"
@@ -15,6 +19,17 @@
 
 void FAvalancheMaskEditorModule::StartupModule()
 {
+	// Details
+	{
+		FPropertyEditorModule& PropertyModule = FModuleManager::GetModuleChecked<FPropertyEditorModule>("PropertyEditor");
+
+		FOnGetDetailCustomizationInstance MaskCustomization = FOnGetDetailCustomizationInstance::CreateStatic(&FAvaMask2DModifierDetails::MakeInstance);
+		PropertyModule.RegisterCustomClassLayout(UAvaMask2DReadModifier::StaticClass()->GetFName(), MaskCustomization);
+		PropertyModule.RegisterCustomClassLayout(UAvaMask2DWriteModifier::StaticClass()->GetFName(), MaskCustomization);
+		
+		PropertyModule.NotifyCustomizationModuleChanged();
+	}
+
 	FAvaMaskEditorCommands::Register();
 
 	CommandList = MakeShared<FUICommandList>();
@@ -24,9 +39,20 @@ void FAvalancheMaskEditorModule::StartupModule()
 
 void FAvalancheMaskEditorModule::ShutdownModule()
 {
+	// Details
+	if (FModuleManager::Get().IsModuleLoaded(TEXT("PropertyEditor")))
+	{
+		if (FPropertyEditorModule* PropertyModule = FModuleManager::GetModulePtr<FPropertyEditorModule>("PropertyEditor"))
+		{
+			PropertyModule->UnregisterCustomClassLayout(UAvaMask2DReadModifier::StaticClass()->GetFName());
+			PropertyModule->UnregisterCustomClassLayout(UAvaMask2DWriteModifier::StaticClass()->GetFName());
+			PropertyModule->NotifyCustomizationModuleChanged();
+		}
+	}
+	
 	UToolMenus::UnRegisterStartupCallback(this);
 	UToolMenus::UnregisterOwner(this);
-	
+
 	FAvaMaskEditorCommands::Unregister();
 }
 
diff --git a/Engine/Plugins/Experimental/Avalanche/Source/AvalancheMaskEditor/Private/Details/AvaMask2DModifierDetails.cpp b/Engine/Plugins/Experimental/Avalanche/Source/AvalancheMaskEditor/Private/Details/AvaMask2DModifierDetails.cpp
new file mode 100644
--- /dev/null
+++ b/Engine/Plugins/Experimental/Avalanche/Source/AvalancheMaskEditor/Private/Details/AvaMask2DModifierDetails.cpp
@@ -0,0 +1,72 @@
+// Copyright Epic Games, Inc. All Rights Reserved.
+
+#include "Details/AvaMask2DModifierDetails.h"
+
+#include "Algo/RemoveIf.h"
+#include "DetailCategoryBuilder.h"
+#include "DetailLayoutBuilder.h"
+#include "DetailWidgetRow.h"
+#include "GeometryMaskCanvas.h"
+#include "PropertyAccess.h"
+
+void FAvaMask2DModifierDetails::CustomizeDetails(IDetailLayoutBuilder& InDetailBuilder)
+{
+	IDetailCategoryBuilder& MaskCategory = InDetailBuilder.EditCategory(TEXT("Mask2D"));
+	MaskCategory.SetIsEmpty(true);
+
+	TFunction<TSharedPtr<IPropertyHandle>(const TSharedPtr<IPropertyHandle>&)> GetInnerRootHandle;
+	GetInnerRootHandle = [&GetInnerRootHandle](const TSharedPtr<IPropertyHandle>& InPropertyHandle)
+	{
+		uint32 NumChildren;
+		InPropertyHandle->GetNumChildren(NumChildren);
+		if (NumChildren == 1)
+		{
+			return GetInnerRootHandle(InPropertyHandle->GetChildHandle(0));
+		}
+
+		// No children, or more than one, so return self as the "last valid" handle
+		return InPropertyHandle;
+	};
+
+	// Effects
+	{
+		IDetailCategoryBuilder& EffectsCategory = InDetailBuilder.EditCategory(TEXT("Effects"));
+
+		static const FName CanvasPropertyPath = TEXT("CanvasWeak");
+		TSharedRef<IPropertyHandle> CanvasPropertyHandle = InDetailBuilder.GetProperty(CanvasPropertyPath, UAvaMask2DBaseModifier::StaticClass());
+
+		EffectsCategory.InitiallyCollapsed(false);
+		EffectsCategory.SetSortOrder(MaskCategory.GetSortOrder() + 1);
+
+		if (CanvasPropertyHandle->IsValidHandle())
+		{
+			CanvasPropertyHandle->MarkHiddenByCustomization();
+
+			static TArray<FName, TInlineAllocator<3>> CanvasExcludePropertyNames = {
+				TEXT("CanvasName"),
+				TEXT("CanvasResource"),
+				TEXT("ColorChannel")
+			};
+
+			const TSharedPtr<IPropertyHandle> CanvasInnerHandle = GetInnerRootHandle(CanvasPropertyHandle);
+
+			uint32 NumChildren = 0;
+			if (CanvasInnerHandle->GetNumChildren(NumChildren) == FPropertyAccess::Success)
+			{
+				for (uint32 ChildIdx = 0; ChildIdx < NumChildren; ++ChildIdx)
+				{
+					TSharedPtr<IPropertyHandle> ChildHandle = CanvasInnerHandle->GetChildHandle(ChildIdx);
+					if (!ensure(ChildHandle->GetProperty()))
+					{
+						continue;
+					}
+					
+					if (!CanvasExcludePropertyNames.Contains(ChildHandle->GetProperty()->NamePrivate))
+					{
+						EffectsCategory.AddProperty(ChildHandle);
+					}
+				}
+			}
+		}
+	}
+}
diff --git a/Engine/Plugins/Experimental/Avalanche/Source/AvalancheMaskEditor/Private/Details/AvaMask2DModifierDetails.h b/Engine/Plugins/Experimental/Avalanche/Source/AvalancheMaskEditor/Private/Details/AvaMask2DModifierDetails.h
new file mode 100644
--- /dev/null
+++ b/Engine/Plugins/Experimental/Avalanche/Source/AvalancheMaskEditor/Private/Details/AvaMask2DModifierDetails.h
@@ -0,0 +1,19 @@
+// Copyright Epic Games, Inc. All Rights Reserved.
+
+#pragma once
+
+#include "IDetailCustomization.h"
+#include "Mask2D/AvaMask2DBaseModifier.h"
+
+class FAvaMask2DModifierDetails
+	: public TSharedFromThis<FAvaMask2DModifierDetails>
+	, public IDetailCustomization
+{
+public:
+	static TSharedRef<IDetailCustomization> MakeInstance()
+	{
+		return MakeShared<FAvaMask2DModifierDetails>();
+	}
+
+	virtual void CustomizeDetails(IDetailLayoutBuilder& InDetailBuilder) override;
+};
diff --git a/Engine/Plugins/Experimental/GeometryMask/Source/GeometryMask/Private/GeometryMaskCanvas.cpp b/Engine/Plugins/Experimental/GeometryMask/Source/GeometryMask/Private/GeometryMaskCanvas.cpp
--- a/Engine/Plugins/Experimental/GeometryMask/Source/GeometryMask/Private/GeometryMaskCanvas.cpp
+++ b/Engine/Plugins/Experimental/GeometryMask/Source/GeometryMask/Private/GeometryMaskCanvas.cpp
@@ -11,6 +11,12 @@
 #include "GeometryMaskModule.h"
 #include "IGeometryMaskWriteInterface.h"
 
+const FName UGeometryMaskCanvas::ApplyBlurPropertyName = GET_MEMBER_NAME_CHECKED(UGeometryMaskCanvas, bApplyBlur);
+const FName UGeometryMaskCanvas::BlurStrengthPropertyName = GET_MEMBER_NAME_CHECKED(UGeometryMaskCanvas, BlurStrength);
+const FName UGeometryMaskCanvas::ApplyFeatherPropertyName = GET_MEMBER_NAME_CHECKED(UGeometryMaskCanvas, bApplyFeather);
+const FName UGeometryMaskCanvas::OuterFeatherRadiusPropertyName = GET_MEMBER_NAME_CHECKED(UGeometryMaskCanvas, InnerFeatherRadius);
+const FName UGeometryMaskCanvas::InnerFeatherRadiusPropertyName = GET_MEMBER_NAME_CHECKED(UGeometryMaskCanvas, OuterFeatherRadius);
+
 void UGeometryMaskCanvas::BeginDestroy()
 {
 	FreeResource();
@@ -18,6 +24,27 @@ void UGeometryMaskCanvas::BeginDestroy()
 	UObject::BeginDestroy();
 }
 
+#if WITH_EDITOR
+void UGeometryMaskCanvas::PostEditChangeProperty(FPropertyChangedEvent& InPropertyChangedEvent)
+{
+	UObject::PostEditChangeProperty(InPropertyChangedEvent);
+
+	static TSet<FName> UpdateRenderParameterProperties =
+	{
+		ApplyBlurPropertyName,
+		BlurStrengthPropertyName,
+		ApplyFeatherPropertyName,
+		OuterFeatherRadiusPropertyName,
+		InnerFeatherRadiusPropertyName
+	};
+
+	if (UpdateRenderParameterProperties.Contains(InPropertyChangedEvent.GetPropertyName()))
+	{
+		UpdateRenderParameters();
+	}
+}
+#endif
+
 const TArray<TWeakInterfacePtr<IGeometryMaskWriteInterface>>& UGeometryMaskCanvas::GetWriters() const
 {
 	return Writers;
@@ -51,7 +78,7 @@ void UGeometryMaskCanvas::AddWriters(const TArray<TScriptInterface<IGeometryMask
 	RemoveInvalidWriters();
 	for (const TScriptInterface<IGeometryMaskWriteInterface>& Writer : InWriters)
 	{
-		AddWriter(Writer);		
+		AddWriter(Writer);
 	}
 }
 
@@ -212,6 +239,8 @@ void UGeometryMaskCanvas::AssignResource(
 	ColorChannel = InColorChannel;
 
 	CanvasResource->OnDrawToCanvas().AddUObject(this, &UGeometryMaskCanvas::OnDrawToCanvas);
+
+	UpdateRenderParameters();
 }
 
 void UGeometryMaskCanvas::FreeResource()
@@ -227,6 +256,31 @@ void UGeometryMaskCanvas::FreeResource()
 	}
 }
 
+FName UGeometryMaskCanvas::GetApplyBlurPropertyName()
+{
+	return ApplyBlurPropertyName;
+}
+
+FName UGeometryMaskCanvas::GetBlurStrengthPropertyName()
+{
+	return BlurStrengthPropertyName;
+}
+
+FName UGeometryMaskCanvas::GetApplyFeatherPropertyName()
+{
+	return ApplyFeatherPropertyName;
+}
+
+FName UGeometryMaskCanvas::GetOuterFeatherRadiusPropertyName()
+{
+	return OuterFeatherRadiusPropertyName;
+}
+
+FName UGeometryMaskCanvas::GetInnerFeatherRadiusPropertyName()
+{
+	return InnerFeatherRadiusPropertyName;
+}
+
 void UGeometryMaskCanvas::SortWriters()
 {
 	Algo::SortBy(Writers, [](const TWeakInterfacePtr<IGeometryMaskWriteInterface>& InWriter)
diff --git a/Engine/Plugins/Experimental/GeometryMask/Source/GeometryMask/Private/GeometryMaskCanvasResource.cpp b/Engine/Plugins/Experimental/GeometryMask/Source/GeometryMask/Private/GeometryMaskCanvasResource.cpp
--- a/Engine/Plugins/Experimental/GeometryMask/Source/GeometryMask/Private/GeometryMaskCanvasResource.cpp
+++ b/Engine/Plugins/Experimental/GeometryMask/Source/GeometryMask/Private/GeometryMaskCanvasResource.cpp
@@ -225,7 +225,7 @@ void UGeometryMaskCanvasResource::UpdateRenderParameters(
 
 	if (UCanvasRenderTarget2D* Texture = GetRenderTargetTexture())
 	{
-		if (bInApplyBlur || bInApplyFeather)
+		if (bApplyBlur || bApplyDF)
 		{
 			// MSAA not supported for these effects (they need UAV access)
 			Texture->SetSampleCount(ETextureRenderTargetSampleCount::RTSC_1);
diff --git a/Engine/Plugins/Experimental/GeometryMask/Source/GeometryMask/Private/GeometryMaskWriteComponent.cpp b/Engine/Plugins/Experimental/GeometryMask/Source/GeometryMask/Private/GeometryMaskWriteComponent.cpp
--- a/Engine/Plugins/Experimental/GeometryMask/Source/GeometryMask/Private/GeometryMaskWriteComponent.cpp
+++ b/Engine/Plugins/Experimental/GeometryMask/Source/GeometryMask/Private/GeometryMaskWriteComponent.cpp
@@ -68,6 +68,9 @@ void UGeometryMaskWriteMeshComponent::DrawToCanvas(FCanvas* InCanvas)
 		TArray<FName> KeyNames;
 		CachedMeshData.GenerateKeyArray(KeyNames);
 
+		TMap<FName, TWeakObjectPtr<USceneComponent>> ComponentsToKeep;
+		ComponentsToKeep.Reserve(KeyNames.Num());
+
 		TRACE_CPUPROFILER_EVENT_SCOPE(UGeometryMaskWriteMeshComponent::DrawToCanvas::WriteToCanvas);
 
 		ParallelFor(KeyNames.Num(), [&](int32 TaskIdx)
@@ -76,14 +79,18 @@ void UGeometryMaskWriteMeshComponent::DrawToCanvas(FCanvas* InCanvas)
 
 			const FName KeyName = KeyNames[TaskIdx];
 			FTransform LocalToWorld;
-			
-			if (const USceneComponent* Component = CachedComponentsWeak[KeyName].Get())
+
+			if (CachedComponentsWeak.Contains(KeyName))
 			{
-				LocalToWorld = Component->GetComponentToWorld();
+				TWeakObjectPtr<USceneComponent>& CachedComponentWeak = CachedComponentsWeak[KeyName];
+				if (const USceneComponent* Component = CachedComponentWeak.Get())
+				{
+					ComponentsToKeep.Emplace(KeyName, CachedComponentWeak);
+					LocalToWorld = Component->GetComponentToWorld();
+				}
 			}
 			else
 			{
-				// @todo: remove from cache
 				return;
 			}
 
@@ -139,6 +146,8 @@ void UGeometryMaskWriteMeshComponent::DrawToCanvas(FCanvas* InCanvas)
 				InCanvas->PopTransform();
 			}
 		});
+
+		CachedComponentsWeak = ComponentsToKeep;
 	}
 }
 
@@ -198,8 +207,14 @@ void UGeometryMaskWriteMeshComponent::UpdateCachedStaticMeshData(TConstArrayView
 	{
 		if (const UStaticMesh* StaticMesh = InStaticMeshComponent->GetStaticMesh())
 		{
-			if (const FGeometryMaskBatchElementData* CachedData = CachedMeshData.Find(StaticMesh->GetFName()))
+			const FName ComponentKey = StaticMesh->GetFName();
+			if (const FGeometryMaskBatchElementData* CachedData = CachedMeshData.Find(ComponentKey))
 			{
+				if (!CachedComponentsWeak.Contains(ComponentKey))
+				{
+					return false; // Cached component invalid, don't remove from build list to ensure it's re-cached
+				}
+				
 				if (CachedData->Vertices.Num() == StaticMesh->GetNumVertices(0))
 				{
 					return true; // Cached, remove from "to build" list							
@@ -242,19 +257,18 @@ void UGeometryMaskWriteMeshComponent::UpdateCachedStaticMeshData(TConstArrayView
 				if (FStaticMeshRenderData* RenderData = StaticMesh->GetRenderData())
 				{
 					if (RenderData->LODResources.IsEmpty())
-						{
-							continue;
-						}
-
-						if (RenderData->LODResources[0].GetNumVertices() == 0)
-						{
-							continue;
-						}
-
-						StaticMeshObjectNames.Add(StaticMesh->GetFName());
-						StaticMeshResources.Add(&RenderData->LODResources[0]);
-						CachedComponentsWeak.Add(StaticMesh->GetFName(), StaticMeshComponent);
+					{
+						continue;
 					}
+
+					if (RenderData->LODResources[0].GetNumVertices() == 0)
+					{
+						continue;
+					}
+
+					StaticMeshObjectNames.Add(StaticMesh->GetFName());
+					StaticMeshResources.Add(&RenderData->LODResources[0]);
+					CachedComponentsWeak.Add(StaticMesh->GetFName(), StaticMeshComponent);
 				}
 			}
 		}
@@ -282,6 +296,7 @@ void UGeometryMaskWriteMeshComponent::UpdateCachedStaticMeshData(TConstArrayView
 				}
 			});
 		}
+	}
 }
 
 void UGeometryMaskWriteMeshComponent::UpdateCachedDynamicMeshData(TConstArrayView<UPrimitiveComponent*> InPrimitiveComponents)
@@ -311,8 +326,14 @@ void UGeometryMaskWriteMeshComponent::UpdateCachedDynamicMeshData(TConstArrayVie
 	{
 		if (const UDynamicMesh* DynamicMesh = InDynamicMeshComponent->GetDynamicMesh())
 		{
-			if (const FGeometryMaskBatchElementData* CachedData = CachedMeshData.Find(InDynamicMeshComponent->GetFName()))
+			const FName ComponentKey = InDynamicMeshComponent->GetFName();
+			if (const FGeometryMaskBatchElementData* CachedData = CachedMeshData.Find(ComponentKey))
 			{
+				if (!CachedComponentsWeak.Contains(ComponentKey))
+				{
+					return false; // Cached component invalid, don't remove from build list to ensure it's re-cached
+				}
+				
 				if (CachedData->Vertices.Num() == DynamicMesh->GetMeshRef().VertexCount()
 					&& CachedData->ChangeStamp == DynamicMesh->GetMeshRef().GetChangeStamp())
 				{
diff --git a/Engine/Plugins/Experimental/GeometryMask/Source/GeometryMask/Public/GeometryMaskCanvas.h b/Engine/Plugins/Experimental/GeometryMask/Source/GeometryMask/Public/GeometryMaskCanvas.h
--- a/Engine/Plugins/Experimental/GeometryMask/Source/GeometryMask/Public/GeometryMaskCanvas.h
+++ b/Engine/Plugins/Experimental/GeometryMask/Source/GeometryMask/Public/GeometryMaskCanvas.h
@@ -31,7 +31,11 @@ class GEOMETRYMASK_API UGeometryMaskCanvas : public UObject
 
 public:
 	virtual void BeginDestroy() override;
-	
+
+#if WITH_EDITOR
+	virtual void PostEditChangeProperty(FPropertyChangedEvent& InPropertyChangedEvent) override;
+#endif
+
 	/** Returns all writers. */
     const TArray<TWeakInterfacePtr<IGeometryMaskWriteInterface>>& GetWriters() const;
 
@@ -98,6 +102,20 @@ public:
 	void AssignResource(UGeometryMaskCanvasResource* InResource, EGeometryMaskColorChannel InColorChannel);
 	void FreeResource();
 
+public:
+	static FName GetApplyBlurPropertyName();
+	static FName GetBlurStrengthPropertyName();
+	static FName GetApplyFeatherPropertyName();
+	static FName GetOuterFeatherRadiusPropertyName();
+	static FName GetInnerFeatherRadiusPropertyName();
+
+private:
+	static const FName ApplyBlurPropertyName;
+	static const FName BlurStrengthPropertyName;
+	static const FName ApplyFeatherPropertyName;
+	static const FName OuterFeatherRadiusPropertyName;
+	static const FName InnerFeatherRadiusPropertyName;
+
 private:
 	/** Sorts writers by various criteria for proper rendering order. */
 	void SortWriters();
@@ -114,7 +132,8 @@ private:
 private:
 	// Allow the subsystem to set the CanvasName
 	friend class UGeometryMaskSubsystem;
-	
+	friend class UGeometryMaskWorldSubsystem;
+
 	FOnGeometryMaskCanvasActivated OnActivatedDelegate;
 	FOnGeometryMaskCanvasDeactivated OnDeactivatedDelegate;
 
