[Motion Design]
 - UAvaOutlinerSubsystem: Fix GEngine broadcast attached event on outliner hierarchy attached changed

diff --git a/Engine/Plugins/Experimental/Avalanche/Source/AvalancheOutliner/Private/AvaOutlinerSubsystem.cpp b/Engine/Plugins/Experimental/Avalanche/Source/AvalancheOutliner/Private/AvaOutlinerSubsystem.cpp
--- a/Engine/Plugins/Experimental/Avalanche/Source/AvalancheOutliner/Private/AvaOutlinerSubsystem.cpp
+++ b/Engine/Plugins/Experimental/Avalanche/Source/AvalancheOutliner/Private/AvaOutlinerSubsystem.cpp
@@ -32,7 +32,7 @@ void UAvaOutlinerSubsystem::BroadcastActorHierarchyChanged(AActor* InActor
 	switch (InChangeType)
 	{
 	case EAvaOutlinerHierarchyChangeType::Attached:
-		GEngine->BroadcastLevelActorDetached(InActor, InParentActor);
+		GEngine->BroadcastLevelActorAttached(InActor, InParentActor);
 		break;
 
 	case EAvaOutlinerHierarchyChangeType::Detached:
