[Motion Design] Hide UAvaText3DCharacterTransform Transform Category

diff --git a/Engine/Plugins/Experimental/Avalanche/Source/AvalancheText/Public/AvaTextCharacterTransform.h b/Engine/Plugins/Experimental/Avalanche/Source/AvalancheText/Public/AvaTextCharacterTransform.h
--- a/Engine/Plugins/Experimental/Avalanche/Source/AvalancheText/Public/AvaTextCharacterTransform.h
+++ b/Engine/Plugins/Experimental/Avalanche/Source/AvalancheText/Public/AvaTextCharacterTransform.h
@@ -9,7 +9,7 @@
 /**
  * This Component is kept for compatibility with old AAvaTextActor assets
  */
-UCLASS(ClassGroup=(Text3D), HideCategories = (Collision, Tags, Activation, Cooking, Rendering, Physics, Mobility, LOD, AssetUserData, Navigation), meta = (BlueprintSpawnableComponent))
+UCLASS(ClassGroup=(Text3D), HideCategories = (Collision, Tags, Activation, Cooking, Rendering, Physics, Mobility, LOD, AssetUserData, Navigation, Transform), meta = (BlueprintSpawnableComponent))
 class AVALANCHETEXT_API UAvaTextCharacterTransform : public UText3DCharacterTransform
 {
 	GENERATED_BODY()
