[Motion Design] Remove redundant Package used to import Font Faces

diff --git a/Engine/Plugins/Experimental/Avalanche/Source/AvalancheTextEditor/Private/Font/AvaFontManagerSubsystem.cpp b/Engine/Plugins/Experimental/Avalanche/Source/AvalancheTextEditor/Private/Font/AvaFontManagerSubsystem.cpp
--- a/Engine/Plugins/Experimental/Avalanche/Source/AvalancheTextEditor/Private/Font/AvaFontManagerSubsystem.cpp
+++ b/Engine/Plugins/Experimental/Avalanche/Source/AvalancheTextEditor/Private/Font/AvaFontManagerSubsystem.cpp
@@ -437,22 +437,18 @@ void UAvaFontManagerSubsystem::SetupFontFamilyTypefaces(UFont* InFont, const FSy
 			continue;
 		}
 
-		UPackage* StyleFontFacePackage = CreatePackage(*FontFacePackage->GetName());
-		StyleFontFacePackage->AddToRoot();
-		StyleFontFacePackage->FullyLoad();
-
 		if (!bIsTempAsset)
 		{
 			FAssetRegistryModule::AssetCreated(NewFontFace);
-			StyleFontFacePackage->SetDirtyFlag(true);
+			FontFacePackage->SetDirtyFlag(true);
 
-			const FString StyleFontFaceAssetFileName = FPackageName::LongPackageNameToFilename(StyleFontFacePackage->GetPathName(), FPackageName::GetAssetPackageExtension());
+			const FString StyleFontFaceAssetFileName = FPackageName::LongPackageNameToFilename(FontFacePackage->GetPathName(), FPackageName::GetAssetPackageExtension());
 
 			if (!FPlatformFileManager::Get().GetPlatformFile().IsReadOnly(*StyleFontFaceAssetFileName))
 			{
 				FSavePackageArgs SavePackageArgs;
 				SavePackageArgs.TopLevelFlags = RF_Public | RF_Standalone;
-				UPackage::SavePackage(StyleFontFacePackage, NewFontFace, *StyleFontFaceAssetFileName, SavePackageArgs);
+				UPackage::SavePackage(FontFacePackage, NewFontFace, *StyleFontFaceAssetFileName, SavePackageArgs);
 			}
 		}
 
