Motion Design Modifiers:
- public api Cleanup

#changelist validated
#virtualized

diff --git a/Engine/Plugins/Experimental/Avalanche/Source/AvalancheModifiers/Private/AvaModifiersActorUtils.h b/Engine/Plugins/Experimental/Avalanche/Source/AvalancheModifiers/Private/AvaModifiersActorUtils.h
new file mode 100644
--- /dev/null
+++ b/Engine/Plugins/Experimental/Avalanche/Source/AvalancheModifiers/Private/AvaModifiersActorUtils.h
@@ -0,0 +1,44 @@
+// Copyright Epic Games, Inc. All Rights Reserved.
+
+#pragma once
+
+#include "Containers/Array.h"
+#include "Math/Box.h"
+#include "Math/OrientedBox.h"
+#include "Math/TransformNonVectorized.h"
+#include "UObject/WeakObjectPtrTemplates.h"
+#include "AvaModifiersActorUtils.generated.h"
+
+enum class EAvaAxis : uint8;
+class AActor;
+
+/** Axis to use */
+UENUM(BlueprintType, meta=(Bitflags, UseEnumValuesAsMaskValuesInEditor="true"))
+enum class EAvaModifiersAxis : uint8
+{
+	None = 0 UMETA(Hidden),
+	X = 1 << 0,
+	Y = 1 << 1,
+	Z = 1 << 2,
+};
+ENUM_CLASS_FLAGS(EAvaModifiersAxis);
+
+// All operations that can be reused or shared in modifiers should go here
+struct FAvaModifiersActorUtils
+{
+	/** Begin Bounds */
+	static FBox GetActorsBounds(const TSet<TWeakObjectPtr<AActor>>& InActors, const FTransform& InReferenceTransform, bool bInSkipHidden = false);
+	static FBox GetActorsBounds(AActor* InActor, bool bInIncludeChildren, bool bInSkipHidden = false);
+	static FBox GetActorBounds(const AActor* InActor);
+	static FOrientedBox GetOrientedBox(const FBox& InLocalBox, const FTransform& InWorldTransform);
+	static FVector GetVectorAxis(int32 InAxis);
+	static bool IsAxisVectorEquals(const FVector& InVectorA, const FVector& InVectorB, int32 InCompareAxis);
+	/** End Bounds */
+
+	/** Begin Outliner */
+	static bool IsActorNotIsolated(const AActor* InActor);
+	/** End Outliner */
+	
+	static FRotator FindLookAtRotation(const FVector& InEyePosition, const FVector& InTargetPosition, const EAvaAxis InAxis, const bool bInFlipAxis);
+	static bool IsActorVisible(const AActor* InActor);
+};
\ No newline at end of file
diff --git a/Engine/Plugins/Experimental/Avalanche/Source/AvalancheModifiers/Private/AvaModifiersParametricMaterial.cpp b/Engine/Plugins/Experimental/Avalanche/Source/AvalancheModifiers/Private/AvaModifiersParametricMaterial.cpp
--- a/Engine/Plugins/Experimental/Avalanche/Source/AvalancheModifiers/Private/AvaModifiersParametricMaterial.cpp
+++ b/Engine/Plugins/Experimental/Avalanche/Source/AvalancheModifiers/Private/AvaModifiersParametricMaterial.cpp
@@ -16,19 +16,6 @@ FAvaModifiersParametricMaterial::FAvaModifiersParametricMaterial()
 	DefaultMaterial = LoadedMaterial;
 }
 
-FAvaModifiersParametricMaterial::FAvaModifiersParametricMaterial(const FAvaModifiersParametricMaterial& Other)
-{
-	MaskColor = Other.MaskColor;
-	InstanceMaterial = nullptr;
-}
-
-FAvaModifiersParametricMaterial& FAvaModifiersParametricMaterial::operator=(const FAvaModifiersParametricMaterial& Other)
-{
-	MaskColor = Other.MaskColor;
-	InstanceMaterial = nullptr;
-	return *this;
-}
-
 UMaterial* FAvaModifiersParametricMaterial::GetDefaultMaterial() const
 {
 	return DefaultMaterial;
diff --git a/Engine/Plugins/Experimental/Avalanche/Source/AvalancheModifiers/Public/AvaModifiersActorUtils.h b/Engine/Plugins/Experimental/Avalanche/Source/AvalancheModifiers/Public/AvaModifiersActorUtils.h
deleted file mode 100644
diff --git a/Engine/Plugins/Experimental/Avalanche/Source/AvalancheModifiers/Public/AvaModifiersParametricMaterial.h b/Engine/Plugins/Experimental/Avalanche/Source/AvalancheModifiers/Public/AvaModifiersParametricMaterial.h
--- a/Engine/Plugins/Experimental/Avalanche/Source/AvalancheModifiers/Public/AvaModifiersParametricMaterial.h
+++ b/Engine/Plugins/Experimental/Avalanche/Source/AvalancheModifiers/Public/AvaModifiersParametricMaterial.h
@@ -11,32 +11,31 @@ class UMaterialInstanceDynamic;
 class UObject;
 
 /** Use this if you want a parametric material */
-USTRUCT(BlueprintType)
-struct AVALANCHEMODIFIERS_API FAvaModifiersParametricMaterial
+USTRUCT()
+struct FAvaModifiersParametricMaterial
 {
 	GENERATED_BODY()
-	
-	FAvaModifiersParametricMaterial();
 
-	FAvaModifiersParametricMaterial(const FAvaModifiersParametricMaterial& Other);
-	FAvaModifiersParametricMaterial& operator=(const FAvaModifiersParametricMaterial& Other);
+	AVALANCHEMODIFIERS_API FAvaModifiersParametricMaterial();
+
+	AVALANCHEMODIFIERS_API UMaterial* GetDefaultMaterial() const;
+
+	AVALANCHEMODIFIERS_API UMaterialInstanceDynamic* GetMaterial() const;
+
+	AVALANCHEMODIFIERS_API void ApplyChanges(UObject* Outer = nullptr);
 
 	UPROPERTY()
 	FLinearColor MaskColor;
 
-	UMaterial* GetDefaultMaterial() const;
-	UMaterialInstanceDynamic* GetMaterial() const;
-	void ApplyChanges(UObject* Outer = nullptr);
+private:
+	void ApplyParams() const;
+	void CreateAndApply(UObject* Outer = nullptr);
+	void EnsureCurrentMaterial(UObject* Outer = nullptr);
+	UMaterial* LoadResource() const;
 
-protected:
 	UPROPERTY()
 	TObjectPtr<UMaterial> DefaultMaterial;
 
 	UPROPERTY(Transient, Instanced)
 	TObjectPtr<UMaterialInstanceDynamic> InstanceMaterial;
-
-	void ApplyParams() const;
-	void CreateAndApply(UObject* Outer = nullptr);
-	void EnsureCurrentMaterial(UObject* Outer = nullptr);
-	UMaterial* LoadResource() const;
 };
\ No newline at end of file
diff --git a/Engine/Plugins/Experimental/Avalanche/Source/AvalancheModifiers/Public/Extensions/AvaRenderStateUpdateModifierExtension.h b/Engine/Plugins/Experimental/Avalanche/Source/AvalancheModifiers/Public/Extensions/AvaRenderStateUpdateModifierExtension.h
--- a/Engine/Plugins/Experimental/Avalanche/Source/AvalancheModifiers/Public/Extensions/AvaRenderStateUpdateModifierExtension.h
+++ b/Engine/Plugins/Experimental/Avalanche/Source/AvalancheModifiers/Public/Extensions/AvaRenderStateUpdateModifierExtension.h
@@ -16,7 +16,7 @@ class UAvaRenderStateUpdateHandler : public UInterface
 };
 
 /** Implement this interface to handle extension event */
-class AVALANCHEMODIFIERS_API IAvaRenderStateUpdateHandler
+class IAvaRenderStateUpdateHandler
 {
 	GENERATED_BODY()
 
diff --git a/Engine/Plugins/Experimental/Avalanche/Source/AvalancheModifiers/Public/Extensions/AvaSceneTreeUpdateModifierExtension.h b/Engine/Plugins/Experimental/Avalanche/Source/AvalancheModifiers/Public/Extensions/AvaSceneTreeUpdateModifierExtension.h
--- a/Engine/Plugins/Experimental/Avalanche/Source/AvalancheModifiers/Public/Extensions/AvaSceneTreeUpdateModifierExtension.h
+++ b/Engine/Plugins/Experimental/Avalanche/Source/AvalancheModifiers/Public/Extensions/AvaSceneTreeUpdateModifierExtension.h
@@ -12,7 +12,7 @@ enum class EAvaOutlinerHierarchyChangeType : uint8;
 enum class EAvaReferenceContainer : uint8;
 
 USTRUCT(BlueprintType)
-struct AVALANCHEMODIFIERS_API FAvaSceneTreeActor
+struct FAvaSceneTreeActor
 {
 	friend class FAvaSceneTreeUpdateModifierExtension;
 	
@@ -73,7 +73,7 @@ class UAvaSceneTreeUpdateHandler : public UInterface
 };
 
 /** Implement this interface to handle extension event */
-class AVALANCHEMODIFIERS_API IAvaSceneTreeUpdateHandler
+class IAvaSceneTreeUpdateHandler
 {
 	GENERATED_BODY()
 	
diff --git a/Engine/Plugins/Experimental/Avalanche/Source/AvalancheModifiers/Public/Extensions/AvaTransformUpdateModifierExtension.h b/Engine/Plugins/Experimental/Avalanche/Source/AvalancheModifiers/Public/Extensions/AvaTransformUpdateModifierExtension.h
--- a/Engine/Plugins/Experimental/Avalanche/Source/AvalancheModifiers/Public/Extensions/AvaTransformUpdateModifierExtension.h
+++ b/Engine/Plugins/Experimental/Avalanche/Source/AvalancheModifiers/Public/Extensions/AvaTransformUpdateModifierExtension.h
@@ -13,7 +13,7 @@ class UAvaTransformUpdateHandler : public UInterface
 };
 
 /** Implement this interface to handle extension event */
-class AVALANCHEMODIFIERS_API IAvaTransformUpdateHandler
+class IAvaTransformUpdateHandler
 {
 	GENERATED_BODY()
 
diff --git a/Engine/Plugins/Experimental/Avalanche/Source/AvalancheModifiers/Public/Modifiers/AvaMirrorModifier.h b/Engine/Plugins/Experimental/Avalanche/Source/AvalancheModifiers/Public/Modifiers/AvaMirrorModifier.h
--- a/Engine/Plugins/Experimental/Avalanche/Source/AvalancheModifiers/Public/Modifiers/AvaMirrorModifier.h
+++ b/Engine/Plugins/Experimental/Avalanche/Source/AvalancheModifiers/Public/Modifiers/AvaMirrorModifier.h
@@ -79,6 +79,7 @@ protected:
 	UPROPERTY(EditInstanceOnly, BlueprintReadWrite, Category="Mirror", meta=(AllowPrivateAccess="true"))
 	bool bShowMirrorFrame = false;
 
+private:
 	UPROPERTY(Transient, DuplicateTransient, TextExportTransient)
 	FAvaModifierPreviewPlane PreviewPlane;
 #endif
diff --git a/Engine/Plugins/Experimental/Avalanche/Source/AvalancheModifiers/Public/Modifiers/AvaPlaneCutModifier.h b/Engine/Plugins/Experimental/Avalanche/Source/AvalancheModifiers/Public/Modifiers/AvaPlaneCutModifier.h
--- a/Engine/Plugins/Experimental/Avalanche/Source/AvalancheModifiers/Public/Modifiers/AvaPlaneCutModifier.h
+++ b/Engine/Plugins/Experimental/Avalanche/Source/AvalancheModifiers/Public/Modifiers/AvaPlaneCutModifier.h
@@ -75,6 +75,7 @@ protected:
 	UPROPERTY(EditInstanceOnly, BlueprintReadWrite, Category="PlaneCut", meta=(AllowPrivateAccess="true"))
 	bool bUsePreview = false;
 
+private:
 	UPROPERTY(Transient, DuplicateTransient, TextExportTransient)
 	FAvaModifierPreviewPlane PreviewPlane;
 #endif
diff --git a/Engine/Plugins/Experimental/Avalanche/Source/AvalancheModifiers/Public/Shared/AvaBooleanModifierShared.h b/Engine/Plugins/Experimental/Avalanche/Source/AvalancheModifiers/Public/Shared/AvaBooleanModifierShared.h
--- a/Engine/Plugins/Experimental/Avalanche/Source/AvalancheModifiers/Public/Shared/AvaBooleanModifierShared.h
+++ b/Engine/Plugins/Experimental/Avalanche/Source/AvalancheModifiers/Public/Shared/AvaBooleanModifierShared.h
@@ -19,7 +19,7 @@ struct FAvaBooleanModifierSharedChannel
 };
 
 USTRUCT(BlueprintType)
-struct AVALANCHEMODIFIERS_API FAvaBooleanModifierSharedChannelInfo
+struct FAvaBooleanModifierSharedChannelInfo
 {
 	GENERATED_BODY()
 
