Motion Design:
- Added option to copy from a preset config when rendering selected pages in Rundown

#changelist validated
#virtualized

diff --git a/Engine/Plugins/Experimental/Avalanche/Source/AvalancheMRQEditor/AvalancheMRQEditor.Build.cs b/Engine/Plugins/Experimental/Avalanche/Source/AvalancheMRQEditor/AvalancheMRQEditor.Build.cs
--- a/Engine/Plugins/Experimental/Avalanche/Source/AvalancheMRQEditor/AvalancheMRQEditor.Build.cs
+++ b/Engine/Plugins/Experimental/Avalanche/Source/AvalancheMRQEditor/AvalancheMRQEditor.Build.cs
@@ -19,6 +19,7 @@ public class AvalancheMRQEditor : ModuleRules
                 "AvalancheSequencer",
                 "Core",
                 "CoreUObject",
+                "DeveloperSettings",
                 "Engine",
                 "MovieRenderPipelineCore",
                 "MovieRenderPipelineEditor",
diff --git a/Engine/Plugins/Experimental/Avalanche/Source/AvalancheMRQEditor/Private/AvaMRQEditorRundowntUtils.cpp b/Engine/Plugins/Experimental/Avalanche/Source/AvalancheMRQEditor/Private/AvaMRQEditorRundowntUtils.cpp
--- a/Engine/Plugins/Experimental/Avalanche/Source/AvalancheMRQEditor/Private/AvaMRQEditorRundowntUtils.cpp
+++ b/Engine/Plugins/Experimental/Avalanche/Source/AvalancheMRQEditor/Private/AvaMRQEditorRundowntUtils.cpp
@@ -1,6 +1,7 @@
 // Copyright Epic Games, Inc. All Rights Reserved.
 
 #include "AvaMRQEditorRundownUtils.h"
+#include "AvaMRQEditorSettings.h"
 #include "AvaMRQRundownPageSetting.h"
 #include "AvaSceneSubsystem.h"
 #include "IAvaSceneInterface.h"
@@ -34,6 +35,11 @@ namespace UE::AvaMRQEditor::Private
 		{
 			QueueSubsystem = GEditor ? GEditor->GetEditorSubsystem<UMoviePipelineQueueSubsystem>() : nullptr;
 			ensureMsgf(QueueSubsystem, TEXT("Not able to access UMoviePipelineQueueSubsystem (returned null)"));
+
+			// todo: optionally prompt the user for a different preset
+			const UAvaMRQEditorSettings* MRQEditorSettings = GetDefault<UAvaMRQEditorSettings>();
+			check(MRQEditorSettings);
+			PresetConfig = MRQEditorSettings->PresetConfig.LoadSynchronous();
 		}
 
 		~FAvaMRQScopedRender()
@@ -64,6 +70,11 @@ namespace UE::AvaMRQEditor::Private
 			}
 			if (UMoviePipelineExecutorJob* Job = Queue->AllocateNewJob(UMoviePipelineExecutorJob::StaticClass()))
 			{
+				UMoviePipelinePrimaryConfig* Config = Job->GetConfiguration();
+				if (PresetConfig && Config)
+				{
+					Config->CopyFrom(PresetConfig);
+				}
 				++JobAllocationCount;
 				return Job;
 			}
@@ -91,6 +102,8 @@ namespace UE::AvaMRQEditor::Private
 		uint32 JobAllocationCount = 0;
 
 		TArray<TWeakObjectPtr<UWorld>> RootedWorlds;
+
+		UMoviePipelineConfigBase* PresetConfig = nullptr;
 	};
 
 	void RenderSequence(FAvaMRQScopedRender& InScopedRender, UWorld* InWorld, const UAvaRundown& InRundown, const FAvaRundownPage& InPage, UAvaSequence* InSequence)
diff --git a/Engine/Plugins/Experimental/Avalanche/Source/AvalancheMRQEditor/Private/AvaMRQEditorSettings.cpp b/Engine/Plugins/Experimental/Avalanche/Source/AvalancheMRQEditor/Private/AvaMRQEditorSettings.cpp
new file mode 100644
--- /dev/null
+++ b/Engine/Plugins/Experimental/Avalanche/Source/AvalancheMRQEditor/Private/AvaMRQEditorSettings.cpp
@@ -0,0 +1,9 @@
+// Copyright Epic Games, Inc. All Rights Reserved.
+
+#include "AvaMRQEditorSettings.h"
+
+UAvaMRQEditorSettings::UAvaMRQEditorSettings()
+{
+	CategoryName = TEXT("Motion Design");
+	SectionName  = TEXT("Movie Render Queue");
+}
diff --git a/Engine/Plugins/Experimental/Avalanche/Source/AvalancheMRQEditor/Private/AvaMRQEditorSettings.h b/Engine/Plugins/Experimental/Avalanche/Source/AvalancheMRQEditor/Private/AvaMRQEditorSettings.h
new file mode 100644
--- /dev/null
+++ b/Engine/Plugins/Experimental/Avalanche/Source/AvalancheMRQEditor/Private/AvaMRQEditorSettings.h
@@ -0,0 +1,21 @@
+// Copyright Epic Games, Inc. All Rights Reserved.
+
+#pragma once
+
+#include "Engine/DeveloperSettings.h"
+#include "UObject/SoftObjectPtr.h"
+#include "AvaMRQEditorSettings.generated.h"
+
+class UMoviePipelinePrimaryConfig;
+
+UCLASS(config=EditorPerProjectUserSettings, meta=(DisplayName="Movie Render Queue Settings"))
+class UAvaMRQEditorSettings : public UDeveloperSettings
+{
+	GENERATED_BODY()
+
+public:
+	UAvaMRQEditorSettings();
+
+	UPROPERTY(Config, EditAnywhere, Category="Motion Design")
+	TSoftObjectPtr<UMoviePipelinePrimaryConfig> PresetConfig;
+};
