Motion Design:
- trivial rename of MRQ Settings

#changelist validated
#virtualized

diff --git a/Engine/Plugins/Experimental/Avalanche/Source/AvalancheMRQEditor/Private/AvaMRQEditorSettings.h b/Engine/Plugins/Experimental/Avalanche/Source/AvalancheMRQEditor/Private/AvaMRQEditorSettings.h
--- a/Engine/Plugins/Experimental/Avalanche/Source/AvalancheMRQEditor/Private/AvaMRQEditorSettings.h
+++ b/Engine/Plugins/Experimental/Avalanche/Source/AvalancheMRQEditor/Private/AvaMRQEditorSettings.h
@@ -8,7 +8,7 @@
 
 class UMoviePipelinePrimaryConfig;
 
-UCLASS(config=EditorPerProjectUserSettings, meta=(DisplayName="Movie Render Queue Settings"))
+UCLASS(config=EditorPerProjectUserSettings, meta=(DisplayName="Movie Render Queue"))
 class UAvaMRQEditorSettings : public UDeveloperSettings
 {
 	GENERATED_BODY()
