Avalanche
- Cameras spawned through the scene defaults are now possed by default
- Post process volume is now parented to the scene null.

diff --git a/Engine/Plugins/Experimental/Avalanche/Source/AvalancheEditor/Private/Scene/AvaSceneDefaults.cpp b/Engine/Plugins/Experimental/Avalanche/Source/AvalancheEditor/Private/Scene/AvaSceneDefaults.cpp
--- a/Engine/Plugins/Experimental/Avalanche/Source/AvalancheEditor/Private/Scene/AvaSceneDefaults.cpp
+++ b/Engine/Plugins/Experimental/Avalanche/Source/AvalancheEditor/Private/Scene/AvaSceneDefaults.cpp
@@ -17,17 +17,20 @@
 #include "Framework/Application/SlateApplication.h"
 #include "Framework/AvaNullActor.h"
 #include "GameFramework/Actor.h"
+#include "IAvaEditor.h"
 #include "Scene/SAvaSceneDefaultActorResponses.h"
 #include "Templates/SharedPointer.h"
 #include "Viewport/AvaCineCameraActor.h"
 #include "Viewport/AvaPostProcessVolume.h"
+#include "Viewport/AvaViewportExtension.h"
+#include "ViewportClient/IAvaViewportClient.h"
 #include "Widgets/SWindow.h"
 
 #define LOCTEXT_NAMESPACE "AvaSceneDefaults"
 
 using FAvaSceneDefaultActorList = TMap<FName, TSharedRef<FAvaSceneDefaultActorResponse>>;
 
-DECLARE_DELEGATE_TwoParams(FAvaSceneDefaultActorInitDelegate, const FAvaSceneDefaultActorList&, AActor*)
+DECLARE_DELEGATE_ThreeParams(FAvaSceneDefaultActorInitDelegate, TSharedRef<IAvaEditor>, const FAvaSceneDefaultActorList&, AActor*)
 
 namespace UE::AvalancheEditor::Private
 {
@@ -49,11 +52,11 @@ namespace UE::AvalancheEditor::Private
 		static FName CineCamera = "CineCamera";
 	};
 
-	void InitDefaultSceneActor_DefaultRoot(const FAvaSceneDefaultActorList& InActorResponses, AActor* InActor);
-	void InitDefaultSceneActor_DirectionalLight(const FAvaSceneDefaultActorList& InActorResponses, AActor* InActor);
-	void InitDefaultSceneActor_SkyLight(const FAvaSceneDefaultActorList& InActorResponses, AActor* InActor);
-	void InitDefaultSceneActor_PostProcessVolume(const FAvaSceneDefaultActorList& InActorResponses, AActor* InActor);
-	void InitDefaultSceneActor_Camera(const FAvaSceneDefaultActorList& InActorResponses, AActor* InActor);
+	void InitDefaultSceneActor_DefaultRoot(TSharedRef<IAvaEditor> InEditor, const FAvaSceneDefaultActorList& InActorResponses, AActor* InActor);
+	void InitDefaultSceneActor_DirectionalLight(TSharedRef<IAvaEditor> InEditor, const FAvaSceneDefaultActorList& InActorResponses, AActor* InActor);
+	void InitDefaultSceneActor_SkyLight(TSharedRef<IAvaEditor> InEditor, const FAvaSceneDefaultActorList& InActorResponses, AActor* InActor);
+	void InitDefaultSceneActor_PostProcessVolume(TSharedRef<IAvaEditor> InEditor, const FAvaSceneDefaultActorList& InActorResponses, AActor* InActor);
+	void InitDefaultSceneActor_Camera(TSharedRef<IAvaEditor> InEditor, const FAvaSceneDefaultActorList& InActorResponses, AActor* InActor);
 
 	static const TMap<FName, FAvaSceneDefaultActor> DefaultSceneActors = {
 		{
@@ -98,7 +101,7 @@ namespace UE::AvalancheEditor::Private
 		}
 	};
 
-	AActor* GetActorFromResponses(const FAvaSceneDefaultActorList& InActorResponses, FName InName)
+	AActor* GetActorFromResponses(TSharedRef<IAvaEditor> InEditor, const FAvaSceneDefaultActorList& InActorResponses, FName InName)
 	{
 		if (const TSharedRef<FAvaSceneDefaultActorResponse>* ActorInput = InActorResponses.Find(InName))
 		{
@@ -108,7 +111,7 @@ namespace UE::AvalancheEditor::Private
 		return nullptr;
 	}
 
-	void InitDefaultSceneActor_DefaultRoot(const FAvaSceneDefaultActorList& InActorResponses, AActor* InActor)
+	void InitDefaultSceneActor_DefaultRoot(TSharedRef<IAvaEditor> InEditor, const FAvaSceneDefaultActorList& InActorResponses, AActor* InActor)
 	{
 		if (!InActor)
 		{
@@ -119,7 +122,7 @@ namespace UE::AvalancheEditor::Private
 		InActor->SetActorLabel(TEXT("Default Scene"));
 	}
 
-	void InitDefaultSceneActor_DirectionalLight(const FAvaSceneDefaultActorList& InActorResponses, AActor* InActor)
+	void InitDefaultSceneActor_DirectionalLight(TSharedRef<IAvaEditor> InEditor, const FAvaSceneDefaultActorList& InActorResponses, AActor* InActor)
 	{
 		ADirectionalLight* DirectionalLight = Cast<ADirectionalLight>(InActor);
 
@@ -136,7 +139,7 @@ namespace UE::AvalancheEditor::Private
 		LightComponent->SetLightColor(DefaultSceneProfile.DirectionalLightColor.ToFColor(true));
 		LightComponent->SetMobility(EComponentMobility::Movable);
 
-		if (AActor* DefaultRoot = GetActorFromResponses(InActorResponses, DefaultSceneActorNames::DefaultRoot))
+		if (AActor* DefaultRoot = GetActorFromResponses(InEditor, InActorResponses, DefaultSceneActorNames::DefaultRoot))
 		{
 			static FTransform Transform = {
 				FRotator::ZeroRotator,
@@ -149,7 +152,7 @@ namespace UE::AvalancheEditor::Private
 		}
 	}
 
-	void InitDefaultSceneActor_SkyLight(const FAvaSceneDefaultActorList& InActorResponses, AActor* InActor)
+	void InitDefaultSceneActor_SkyLight(TSharedRef<IAvaEditor> InEditor, const FAvaSceneDefaultActorList& InActorResponses, AActor* InActor)
 	{
 		ASkyLight* SkyLight = Cast<ASkyLight>(InActor);
 
@@ -176,7 +179,7 @@ namespace UE::AvalancheEditor::Private
 		SkyLightComponent->SetCubemap(Cast<UTextureCube>(LoadedObject));
 		SkyLightComponent->SetVisibility(DefaultSceneProfile.bUseSkyLighting, true);
 
-		if (AActor* DefaultRoot = GetActorFromResponses(InActorResponses, DefaultSceneActorNames::DefaultRoot))
+		if (AActor* DefaultRoot = GetActorFromResponses(InEditor, InActorResponses, DefaultSceneActorNames::DefaultRoot))
 		{
 			static FTransform Transform = {
 				FRotator::ZeroRotator,
@@ -189,28 +192,27 @@ namespace UE::AvalancheEditor::Private
 		}
 	}
 
-	void InitDefaultSceneActor_PostProcessVolume(const FAvaSceneDefaultActorList& InActorResponses, AActor* InActor)
+	void InitDefaultSceneActor_PostProcessVolume(TSharedRef<IAvaEditor> InEditor, const FAvaSceneDefaultActorList& InActorResponses, AActor* InActor)
 	{
 		if (!InActor)
 		{
 			return;
 		}
 
-		FTransform PostProcessVolumeTransform = FTransform::Identity;
-
-		float PostProcessXOffset = UAvaEditorSettings::Get()->CameraDistance;
-
-		if (const UCubeBuilder* const DefaultCubeBuilder = Cast<UCubeBuilder>(UCubeBuilder::StaticClass()->GetDefaultObject()))
+		if (AActor* DefaultRoot = GetActorFromResponses(InEditor, InActorResponses, DefaultSceneActorNames::DefaultRoot))
 		{
-			PostProcessXOffset += DefaultCubeBuilder->X * 0.5f;
-		}
-
-		PostProcessVolumeTransform.SetLocation(FVector(-PostProcessXOffset, 0, 0));
+			static FTransform Transform = {
+				FRotator::ZeroRotator,
+				FVector(0, 0, 150),
+				FVector::OneVector
+			};
 
-		InActor->SetActorTransform(PostProcessVolumeTransform);
+			InActor->AttachToActor(DefaultRoot, FAttachmentTransformRules::KeepWorldTransform);
+			InActor->SetActorRelativeTransform(Transform);
+		}
 	}
 
-	void InitDefaultSceneActor_Camera(const FAvaSceneDefaultActorList& InActorResponses, AActor* InActor)
+	void InitDefaultSceneActor_Camera(TSharedRef<IAvaEditor> InEditor, const FAvaSceneDefaultActorList& InActorResponses, AActor* InActor)
 	{
 		AAvaCineCameraActor* Camera = Cast<AAvaCineCameraActor>(InActor);
 
@@ -219,7 +221,7 @@ namespace UE::AvalancheEditor::Private
 			return;
 		}
 
-		if (AActor* DefaultRoot = GetActorFromResponses(InActorResponses, DefaultSceneActorNames::DefaultRoot))
+		if (AActor* DefaultRoot = GetActorFromResponses(InEditor, InActorResponses, DefaultSceneActorNames::DefaultRoot))
 		{
 			FTransform Transform = FTransform::Identity;
 			Transform.SetLocation(FVector(-UAvaEditorSettings::Get()->CameraDistance, 0, 0));
@@ -229,6 +231,20 @@ namespace UE::AvalancheEditor::Private
 		}
 
 		Camera->Configure(UAvaEditorSettings::Get()->CameraDistance);
+		Camera->SetLockLocation(true);
+
+		TSharedPtr<FAvaViewportExtension> ViewportExtension = InEditor->FindExtension<FAvaViewportExtension>();
+
+		for (const TSharedPtr<IAvaViewportClient>& ViewportClient : ViewportExtension->GetViewportClients())
+		{
+			if (!ViewportClient->IsAvalancheViewport())
+			{
+				continue;
+			}
+
+			ViewportClient->SetViewTarget(Camera);
+			break;
+		}
 	}
 
 	TMap<FName, TSharedRef<FAvaSceneDefaultActorResponse>> GeneratorInitialActorResponses(UWorld* InWorld)
@@ -387,7 +403,7 @@ namespace UE::AvalancheEditor::Private
 		}
 	}
 
-	void ApplyActorSettings(TMap<FName, TSharedRef<FAvaSceneDefaultActorResponse>>& InResponses)
+	void ApplyActorSettings(TSharedRef<IAvaEditor> InEditor, TMap<FName, TSharedRef<FAvaSceneDefaultActorResponse>>& InResponses)
 	{
 		for (TPair<FName, TSharedRef<FAvaSceneDefaultActorResponse>>& ResponsePair : InResponses)
 		{
@@ -403,12 +419,12 @@ namespace UE::AvalancheEditor::Private
 				continue;
 			}
 
-			DefaultSceneActor->InitDelegate.ExecuteIfBound(InResponses, ResponsePair.Value->SelectedActor.Get());
+			DefaultSceneActor->InitDelegate.ExecuteIfBound(InEditor, InResponses, ResponsePair.Value->SelectedActor.Get());
 		}
 	}
 }
 
-void FAvaSceneDefaults::CreateDefaultScene(UWorld* InWorld)
+void FAvaSceneDefaults::CreateDefaultScene(TSharedRef<IAvaEditor> InEditor, UWorld* InWorld)
 {
 	if (!InWorld)
 	{
@@ -422,7 +438,7 @@ void FAvaSceneDefaults::CreateDefaultScene(UWorld* InWorld)
 	if (ShowActorResponsesWindow(InWorld, ActorResponses))
 	{
 		CreateNewSceneActors(InWorld, ActorResponses);
-		ApplyActorSettings(ActorResponses);
+		ApplyActorSettings(InEditor, ActorResponses);
 	}
 }
 
diff --git a/Engine/Plugins/Experimental/Avalanche/Source/AvalancheEditor/Private/Scene/AvaSceneDefaults.h b/Engine/Plugins/Experimental/Avalanche/Source/AvalancheEditor/Private/Scene/AvaSceneDefaults.h
--- a/Engine/Plugins/Experimental/Avalanche/Source/AvalancheEditor/Private/Scene/AvaSceneDefaults.h
+++ b/Engine/Plugins/Experimental/Avalanche/Source/AvalancheEditor/Private/Scene/AvaSceneDefaults.h
@@ -4,11 +4,13 @@
 
 #include "Containers/Array.h"
 #include "Internationalization/Text.h"
+#include "Templates/SharedPointerFwd.h"
 #include "UObject/UnrealNames.h"
 #include "UObject/WeakObjectPtr.h"
 #include "UObject/WeakObjectPtrTemplates.h"
 
 class AActor;
+class IAvaEditor;
 class UWorld;
 
 enum class EAvaSceneDefaultActorResponse : uint8
@@ -31,5 +33,5 @@ struct FAvaSceneDefaultActorResponse
 class FAvaSceneDefaults
 {
 public:
-	static void CreateDefaultScene(UWorld* InWorld);
+	static void CreateDefaultScene(TSharedRef<IAvaEditor> InEditor, UWorld* InWorld);
 };
diff --git a/Engine/Plugins/Experimental/Avalanche/Source/AvalancheEditor/Private/Scene/AvaSceneExtension.cpp b/Engine/Plugins/Experimental/Avalanche/Source/AvalancheEditor/Private/Scene/AvaSceneExtension.cpp
--- a/Engine/Plugins/Experimental/Avalanche/Source/AvalancheEditor/Private/Scene/AvaSceneExtension.cpp
+++ b/Engine/Plugins/Experimental/Avalanche/Source/AvalancheEditor/Private/Scene/AvaSceneExtension.cpp
@@ -15,10 +15,10 @@ void FAvaSceneExtension::ExtendToolbarMenu(UToolMenu& InMenu)
 {
 	FToolMenuSection& Section = InMenu.FindOrAddSection(DefaultSectionName);
 
-	FToolMenuEntry& Entry = Section.AddEntry(FToolMenuEntry::InitToolBarButton(TEXT("SpawnDefaultsButton")
+	FToolMenuEntry& Entry = Section.AddEntry(FToolMenuEntry::InitToolBarButton(TEXT("CreateDefaultsButton")
 		, FExecuteAction::CreateSP(this, &FAvaSceneExtension::SpawnDefaultScene)
-		, LOCTEXT("SpawnDefaultSceneLabel", "Spawn Defaults")
-		, LOCTEXT("SpawnDefaultSceneTooltip", "Open the spawn defaults menu to add a basic scene setup.")
+		, LOCTEXT("CreateDefaultSceneLabel", "Create Defaults")
+		, LOCTEXT("CreateDefaultSceneTooltip", "Open the Spawn Defaults menu to add a basic scene setup.")
 		, FSlateIcon(FAppStyle::GetAppStyleSetName(), "GraphEditor.SpawnActor_16x")));
 
 	Entry.StyleNameOverride = "CalloutToolbar";
@@ -31,7 +31,15 @@ void FAvaSceneExtension::RegisterTabSpawners(const TSharedRef<IAvaEditor>& InEdi
 
 void FAvaSceneExtension::SpawnDefaultScene()
 {
+	TSharedPtr<IAvaEditor> Editor = GetEditor();
+
+	if (!Editor.IsValid())
+	{
+		return;
+	}
+
 	UWorld* const World = GetWorld();
+
 	if (!IsValid(World))
 	{
 		return;
@@ -39,7 +47,7 @@ void FAvaSceneExtension::SpawnDefaultScene()
 
 	FScopedTransaction Transaction(LOCTEXT("SpawnDefaultScene", "Spawn Default Scene"));
 
-	FAvaSceneDefaults::CreateDefaultScene(World);
+	FAvaSceneDefaults::CreateDefaultScene(Editor.ToSharedRef(), World);
 }
 
 #undef LOCTEXT_NAMESPACE
