Motion Design
- Update for compatibility with PropertyHandle change.

diff --git a/Engine/Plugins/Experimental/Avalanche/Source/AvalancheEditor/Private/DetailView/Customizations/AvaViewportQualitySettingsPropertyTypeCustomization.cpp b/Engine/Plugins/Experimental/Avalanche/Source/AvalancheEditor/Private/DetailView/Customizations/AvaViewportQualitySettingsPropertyTypeCustomization.cpp
--- a/Engine/Plugins/Experimental/Avalanche/Source/AvalancheEditor/Private/DetailView/Customizations/AvaViewportQualitySettingsPropertyTypeCustomization.cpp
+++ b/Engine/Plugins/Experimental/Avalanche/Source/AvalancheEditor/Private/DetailView/Customizations/AvaViewportQualitySettingsPropertyTypeCustomization.cpp
@@ -30,6 +30,8 @@ void FAvaViewportQualitySettingsPropertyTypeCustomization::CustomizeChildren(TSh
 
 	const TSharedRef<FDetailArrayBuilder> FeaturesArrayBuilder = MakeShared<FDetailArrayBuilder>(FeaturesProperty.ToSharedRef(), /*InGenerateHeader*/ false, /*InDisplayResetToDefault*/ true, /*InDisplayElementNum*/ false);
 
+	static const bool bDisplayDefaultPropertyButtons = false;
+
 	FeaturesArrayBuilder->OnGenerateArrayElementWidget(FOnGenerateArrayElementWidget::CreateLambda([](TSharedRef<IPropertyHandle> InElementPropertyHandle, int32 InArrayIndex, IDetailChildrenBuilder& InChildrenBuilder)
 	{
 		TSharedPtr<IPropertyHandle> NameProperty  = InElementPropertyHandle->GetChildHandle(GET_MEMBER_NAME_CHECKED(FAvaViewportQualitySettingsFeature, Name));
@@ -52,7 +54,7 @@ void FAvaViewportQualitySettingsPropertyTypeCustomization::CustomizeChildren(TSh
 			]
 			.ValueContent()
 			[
-				ValueProperty->CreatePropertyValueWidget(/*bDisplayPropertyButtons*/false)
+				ValueProperty->CreatePropertyValueWidget(bDisplayDefaultPropertyButtons)
 			];
 	}));
 
diff --git a/Engine/Plugins/Experimental/Avalanche/Source/AvalancheSequencer/Private/Customization/AvaSequenceTimeCustomization.cpp b/Engine/Plugins/Experimental/Avalanche/Source/AvalancheSequencer/Private/Customization/AvaSequenceTimeCustomization.cpp
--- a/Engine/Plugins/Experimental/Avalanche/Source/AvalancheSequencer/Private/Customization/AvaSequenceTimeCustomization.cpp
+++ b/Engine/Plugins/Experimental/Avalanche/Source/AvalancheSequencer/Private/Customization/AvaSequenceTimeCustomization.cpp
@@ -16,7 +16,7 @@ void FAvaSequenceTimeCustomization::CustomizeHeader(TSharedRef<IPropertyHandle>
 	{
 		TSharedPtr<IPropertyHandle> ChildPropertyHandle = InPropertyHandle->GetChildHandle(InPropertyName);
 
-		TSharedRef<SWidget> PropertyValueWidget = ChildPropertyHandle->CreatePropertyValueWidget(/*bDisplayDefaultPropertyButtons*/false);
+		TSharedRef<SWidget> PropertyValueWidget = ChildPropertyHandle->CreatePropertyValueWidget();
 
 		TAttribute<EVisibility>::FGetter VisibilityGetter = TAttribute<EVisibility>::FGetter::CreateLambda([ChildPropertyHandle]
 			{
@@ -47,6 +47,8 @@ void FAvaSequenceTimeCustomization::CustomizeHeader(TSharedRef<IPropertyHandle>
 
 	TSharedPtr<IPropertyHandle> HasTimeHandle = InPropertyHandle->GetChildHandle(GET_MEMBER_NAME_CHECKED(FAvaSequenceTime, bHasTimeConstraint));
 
+	static const bool bDisplayDefaultPropertyButtons = false;
+
 	InHeaderRow
 		.NameContent()
 		[
@@ -58,7 +60,7 @@ void FAvaSequenceTimeCustomization::CustomizeHeader(TSharedRef<IPropertyHandle>
 				+ SHorizontalBox::Slot()
 				.AutoWidth()
 				[
-					HasTimeHandle->CreatePropertyValueWidget(/*bDisplayDefaultPropertyButtons*/false)
+					HasTimeHandle->CreatePropertyValueWidget(bDisplayDefaultPropertyButtons)
 				]
 				+ SHorizontalBox::Slot()
 				.AutoWidth()
@@ -71,7 +73,7 @@ void FAvaSequenceTimeCustomization::CustomizeHeader(TSharedRef<IPropertyHandle>
 				.Padding(0.f, 0.f, 5.f, 0.f)
 				.HAlign(HAlign_Right)
 				[
-					PositionTypeHandle->CreatePropertyValueWidget(/*bDisplayDefaultPropertyButtons*/false)
+					PositionTypeHandle->CreatePropertyValueWidget(bDisplayDefaultPropertyButtons)
 				]
 			]
 		]
