[Motion Design] Playable transitions with no transition tree will not wait for sequences by default.

#rnx

diff --git a/Engine/Plugins/Experimental/Avalanche/Source/Avalanche/Public/Framework/AvaInstanceSettings.h b/Engine/Plugins/Experimental/Avalanche/Source/Avalanche/Public/Framework/AvaInstanceSettings.h
--- a/Engine/Plugins/Experimental/Avalanche/Source/Avalanche/Public/Framework/AvaInstanceSettings.h
+++ b/Engine/Plugins/Experimental/Avalanche/Source/Avalanche/Public/Framework/AvaInstanceSettings.h
@@ -15,4 +15,12 @@ struct FAvaInstanceSettings
 	/** Enable loading dependent levels as sub-playables. */
 	UPROPERTY(config, EditAnywhere, Category = Settings)
 	bool bEnableLoadSubPlayables = true;
+
+	/**
+	 * For default playable transitions (when there is no transition tree),
+	 * wait for the sequences to finish before ending the transition.
+	 */
+	UPROPERTY(config, EditAnywhere, Category = Settings)
+	bool bDefaultPlayableTransitionWaitForSequences = false;
+	
 };
diff --git a/Engine/Plugins/Experimental/Avalanche/Source/AvalancheMedia/Private/Playable/Transition/AvaPlayableLocalTransition.cpp b/Engine/Plugins/Experimental/Avalanche/Source/AvalancheMedia/Private/Playable/Transition/AvaPlayableLocalTransition.cpp
--- a/Engine/Plugins/Experimental/Avalanche/Source/AvalancheMedia/Private/Playable/Transition/AvaPlayableLocalTransition.cpp
+++ b/Engine/Plugins/Experimental/Avalanche/Source/AvalancheMedia/Private/Playable/Transition/AvaPlayableLocalTransition.cpp
@@ -8,6 +8,8 @@
 #include "Behavior/IAvaTransitionBehavior.h"
 #include "Execution/AvaTransitionExecutorBuilder.h"
 #include "Execution/IAvaTransitionExecutor.h"
+#include "Framework/AvaInstanceSettings.h"
+#include "IAvaMediaModule.h"
 #include "IAvaSceneInterface.h"
 #include "IAvaSequenceProvider.h"
 #include "Playable/AvaPlayable.h"
@@ -431,7 +433,9 @@ void UAvaPlayableLocalTransition::PostTransitionExecutorPhase()
 		}
 	}
 
-	if (bSequenceStarted)
+	const FAvaInstanceSettings& PlaybackInstanceSettings = IAvaMediaModule::Get().GetAvaInstanceSettings();
+
+	if (bSequenceStarted && PlaybackInstanceSettings.bDefaultPlayableTransitionWaitForSequences)
 	{
 		StartWaitOnPostExecutorSequences();
 	}
