Motion Design Sequence:
- Fix issues with deleting the last sequence in the tree causing an ensure to fire off and to also not select a newly default created sequence

#changelist validated
#virtualized

diff --git a/Engine/Plugins/Experimental/Avalanche/Source/Avalanche/Private/AvaScene.cpp b/Engine/Plugins/Experimental/Avalanche/Source/Avalanche/Private/AvaScene.cpp
--- a/Engine/Plugins/Experimental/Avalanche/Source/Avalanche/Private/AvaScene.cpp
+++ b/Engine/Plugins/Experimental/Avalanche/Source/Avalanche/Private/AvaScene.cpp
@@ -178,10 +178,6 @@ bool AAvaScene::AddSequence(UAvaSequence* InSequence)
 void AAvaScene::RemoveSequence(UAvaSequence* InSequence)
 {
 	Animations.Remove(InSequence);
-	if (InSequence)
-	{
-		InSequence->OnSequenceRemoved();
-	}
 	ScheduleRebuildSequenceTree();
 }
 
diff --git a/Engine/Plugins/Experimental/Avalanche/Source/AvalancheEditor/Private/Sequencer/AvaLevelSequencerExtension.cpp b/Engine/Plugins/Experimental/Avalanche/Source/AvalancheEditor/Private/Sequencer/AvaLevelSequencerExtension.cpp
--- a/Engine/Plugins/Experimental/Avalanche/Source/AvalancheEditor/Private/Sequencer/AvaLevelSequencerExtension.cpp
+++ b/Engine/Plugins/Experimental/Avalanche/Source/AvalancheEditor/Private/Sequencer/AvaLevelSequencerExtension.cpp
@@ -153,8 +153,19 @@ TSharedPtr<ISequencer> FAvaLevelSequencerExtension::GetExternalSequencer() const
 
 void FAvaLevelSequencerExtension::OnViewedSequenceChanged(UAvaSequence* InOldSequence, UAvaSequence* InNewSequence)
 {
-	// If one of these sequences is null (i.e. a ULevelSequence opening, or opening for first time where old is null), re-open the Asset Editor
-	if (!InOldSequence || !InNewSequence)
+	IAvaSequenceProvider* SequenceProvider = GetSequenceProvider();
+	if (!SequenceProvider)
+	{
+		return;
+	}
+
+	UObject* const Outer = SequenceProvider->ToUObject();
+
+	bool bIsOldOwnedByProvider = InOldSequence && InOldSequence->GetOuter() == Outer;
+	bool bIsNewOwnedByProvider = InNewSequence && InNewSequence->GetOuter() == Outer;
+
+	// Only avoid re-opening the asset editor if it's changing between two sequences that are owned by the provider
+	if (!bIsOldOwnedByProvider || !bIsNewOwnedByProvider)
 	{
 		OpenSequencerWithViewedSequence();	
 	}
diff --git a/Engine/Plugins/Experimental/Avalanche/Source/AvalancheSequencer/Private/AvaSequencer.cpp b/Engine/Plugins/Experimental/Avalanche/Source/AvalancheSequencer/Private/AvaSequencer.cpp
--- a/Engine/Plugins/Experimental/Avalanche/Source/AvalancheSequencer/Private/AvaSequencer.cpp
+++ b/Engine/Plugins/Experimental/Avalanche/Source/AvalancheSequencer/Private/AvaSequencer.cpp
@@ -911,17 +911,33 @@ void FAvaSequencer::DeleteSequence_Execute()
 	check(Outer);
 	Outer->Modify();
 
+	TArray<UAvaSequence*> RemovedSequences;
+	RemovedSequences.Reserve(SelectedItems.Num());
+
+	// Remove the Selected Sequences from the List (not marked as garbage yet)
 	for (const FAvaSequenceItemPtr& Item : SelectedItems)
 	{
 		check(Item.IsValid());
-		SequenceProvider->RemoveSequence(Item->GetSequence());
-		SetViewedSequence(GetDefaultSequence());
+		if (UAvaSequence* Sequence = Item->GetSequence())
+		{
+			SequenceProvider->RemoveSequence(Sequence);
+			RemovedSequences.Add(Sequence);
+		}
+	}
+
+	// Set the Viewed Sequence to the Default one
+	SetViewedSequence(GetDefaultSequence());
+
+	// Once a new viewed sequence is set, the removed sequences can now be marked as garbage 
+	for (UAvaSequence* Sequence : RemovedSequences)
+	{
+		Sequence->OnSequenceRemoved();
 	}
 
 	if (UBlueprint* const Blueprint = Cast<UBlueprint>(Outer))
 	{
 		FBlueprintEditorUtils::MarkBlueprintAsStructurallyModified(Blueprint);
-	}	
+	}
 }
 
 bool FAvaSequencer::RelabelSequence_CanExecute() const
@@ -1376,7 +1392,7 @@ void FAvaSequencer::SetViewedSequence(UAvaSequence* InSequenceToView)
 		return;
 	}
 
-	UAvaSequence* OldSequence = ViewedSequenceWeak.Get();
+	UAvaSequence* OldSequence = ViewedSequenceWeak.Get(/*bEvenIfPendingKill*/true);
 	ViewedSequenceWeak = InSequenceToView;
 	NotifyViewedSequenceChanged(OldSequence);
 }
@@ -1579,7 +1595,14 @@ void FAvaSequencer::NotifyOnSequenceTreeChanged()
 			RemainingItems.Append(Item->GetChildren());
 		}
 	}
-	
+
+	// Ensure the new item representing the Viewed Sequence is selected
+	UAvaSequence* ViewedSequence = GetViewedSequence();
+	if (ViewedSequence && SequenceTree.IsValid())
+	{
+		SequenceTree->OnPostSetViewedSequence(ViewedSequence);
+	}
+
 	SequenceTreeView->RequestTreeRefresh();
 }
 
diff --git a/Engine/Plugins/Experimental/Avalanche/Source/AvalancheSequencer/Private/SequenceTree/Widgets/SAvaSequenceTree.cpp b/Engine/Plugins/Experimental/Avalanche/Source/AvalancheSequencer/Private/SequenceTree/Widgets/SAvaSequenceTree.cpp
--- a/Engine/Plugins/Experimental/Avalanche/Source/AvalancheSequencer/Private/SequenceTree/Widgets/SAvaSequenceTree.cpp
+++ b/Engine/Plugins/Experimental/Avalanche/Source/AvalancheSequencer/Private/SequenceTree/Widgets/SAvaSequenceTree.cpp
@@ -286,18 +286,16 @@ void SAvaSequenceTree::OnSelectionChanged(FAvaSequenceItemPtr InSelectedItem, ES
 	}
 
 	TSharedPtr<FAvaSequencer> Sequencer = SequencerWeak.Pin();
-
-	if (!Sequencer.IsValid())
+	if (!Sequencer.IsValid() || !InSelectedItem.IsValid())
 	{
 		return;
 	}
 
-	UAvaSequence* const SelectedSequence = InSelectedItem.IsValid()
-		? InSelectedItem->GetSequence()
-		: nullptr;
-
-	TGuardValue<bool> Guard(bSyncingSelection, true);
-	Sequencer->SetViewedSequence(SelectedSequence);
+	if (UAvaSequence* const SelectedSequence = InSelectedItem->GetSequence())
+	{
+		TGuardValue<bool> Guard(bSyncingSelection, true);
+		Sequencer->SetViewedSequence(SelectedSequence);
+	}
 }
 
 void SAvaSequenceTree::OnSearchChanged(const FText& InSearchText)
