[MotionDesign] Camera history capacity bug

diff --git a/Engine/Plugins/Experimental/Avalanche/Source/AvalancheLevelViewport/Private/AvaViewportCameraHistory.cpp b/Engine/Plugins/Experimental/Avalanche/Source/AvalancheLevelViewport/Private/AvaViewportCameraHistory.cpp
--- a/Engine/Plugins/Experimental/Avalanche/Source/AvalancheLevelViewport/Private/AvaViewportCameraHistory.cpp
+++ b/Engine/Plugins/Experimental/Avalanche/Source/AvalancheLevelViewport/Private/AvaViewportCameraHistory.cpp
@@ -93,7 +93,8 @@ void FAvaViewportCameraHistory::PostEngineInit()
 void FAvaViewportCameraHistory::Reset()
 {
 	CameraTransformHistory.Reset();
-	CameraTransformHistoryIndex = 0;
+	CameraTransformHistory.SetNumUninitialized(UE::AvalancheEditor::Internal::CameraUndoHistoryCapacity, false);
+	CameraTransformHistoryIndex = INDEX_NONE;
 	CameraTransformHistoryHeadIndex = 1;
 }
 
@@ -126,10 +127,12 @@ void FAvaViewportCameraHistory::OnBeginCameraTransform(UObject& InCameraObject)
 
 	if (AActor* CameraActor = Cast<AActor>(&InCameraObject))
 	{
-		if (CameraTransformHistory.IsEmpty())
+		// First move must be added
+		if (CameraTransformHistoryIndex == INDEX_NONE)
 		{
-			CameraTransformHistoryIndex = CameraTransformHistory.Add({ CameraActor, CameraActor->GetActorTransform() });
-			CameraTransformHistoryHeadIndex = 1;
+			CameraTransformHistoryIndex = 0;
+			CameraTransformHistory[CameraTransformHistoryIndex] = { CameraActor, CameraActor->GetActorTransform() };
+			CameraTransformHistoryHeadIndex = CameraTransformHistoryIndex + 1;
 		}
 	}
 }
@@ -146,21 +149,19 @@ void FAvaViewportCameraHistory::OnEndCameraTransform(UObject& InCameraObject)
 		// New item replaces next
 		int32 CameraTransformIndexToAdd = UE::AvalancheEditor::Private::WrapIndex(CameraTransformHistoryIndex + 1, CameraTransformHistory.Max());
 
-		CameraTransformIndexToAdd = CameraTransformHistory.Insert({ CameraActor, CameraActor->GetActorTransform() }, CameraTransformIndexToAdd);
+		CameraTransformHistory[CameraTransformIndexToAdd] = { CameraActor, CameraActor->GetActorTransform() };
 		
 		CameraTransformHistoryIndex = CameraTransformIndexToAdd;
 		CameraTransformHistoryHeadIndex = UE::AvalancheEditor::Private::WrapIndex(CameraTransformHistoryIndex + 1, CameraTransformHistory.Max());
 
-		UE_LOG(LogAva, Display, TEXT("Saved camera transform at %i, head %i, %s"), CameraTransformHistoryIndex, CameraTransformHistoryHeadIndex, *CameraActor->GetActorTransform().Rotator().ToString());
+		UE_LOG(LogAva, Verbose, TEXT("Saved camera transform at %i, head %i, %s"), CameraTransformHistoryIndex, CameraTransformHistoryHeadIndex, *CameraActor->GetActorTransform().Rotator().ToString());
 	}
 }
 
 void FAvaViewportCameraHistory::ExecuteCameraTransformUndo()
 {
 	const int32 CameraTransformIndexToRestore = UE::AvalancheEditor::Private::WrapIndex(CameraTransformHistoryIndex - 1, CameraTransformHistory.Max());
-	if (CameraTransformHistory.IsEmpty()
-		|| !CameraTransformHistory.IsValidIndex(CameraTransformIndexToRestore)
-		|| (CameraTransformHistory.Num() > 1 && CameraTransformIndexToRestore == CameraTransformHistoryHeadIndex))
+	if (CameraTransformIndexToRestore == CameraTransformHistoryHeadIndex)
 	{
 		return;
 	}
@@ -172,16 +173,14 @@ void FAvaViewportCameraHistory::ExecuteCameraTransformUndo()
 		CameraActor->SetActorTransform(UndoPair.Value);
 		NotifyUndo();
 
-		UE_LOG(LogAva, Display, TEXT("Undo camera transform at %i, %s"), CameraTransformIndexToRestore, *CameraActor->GetActorTransform().Rotator().ToString());
+		UE_LOG(LogAva, Verbose, TEXT("Undo camera transform at %i, %s"), CameraTransformIndexToRestore, *CameraActor->GetActorTransform().Rotator().ToString());
 	}
 }
 
 void FAvaViewportCameraHistory::ExecuteCameraTransformRedo()
 {
 	const int32 CameraTransformIndexToRestore = UE::AvalancheEditor::Private::WrapIndex(CameraTransformHistoryIndex + 1, CameraTransformHistory.Max());
-	if (CameraTransformHistory.IsEmpty()
-		|| !CameraTransformHistory.IsValidIndex(CameraTransformIndexToRestore)
-		|| CameraTransformIndexToRestore == CameraTransformHistoryHeadIndex)
+	if (CameraTransformIndexToRestore == CameraTransformHistoryHeadIndex)
 	{
 		return;
 	}
@@ -193,7 +192,7 @@ void FAvaViewportCameraHistory::ExecuteCameraTransformRedo()
 		CameraActor->SetActorTransform(RedoPair.Value);
 		NotifyRedo();
 
-		UE_LOG(LogAva, Display, TEXT("Redo camera transform at %i, %s"), CameraTransformIndexToRestore, *CameraActor->GetActorTransform().Rotator().ToString());
+		UE_LOG(LogAva, Verbose, TEXT("Redo camera transform at %i, %s"), CameraTransformIndexToRestore, *CameraActor->GetActorTransform().Rotator().ToString());
 	}
 }
 
diff --git a/Engine/Plugins/Experimental/Avalanche/Source/AvalancheLevelViewport/Private/AvaViewportCameraHistory.h b/Engine/Plugins/Experimental/Avalanche/Source/AvalancheLevelViewport/Private/AvaViewportCameraHistory.h
--- a/Engine/Plugins/Experimental/Avalanche/Source/AvalancheLevelViewport/Private/AvaViewportCameraHistory.h
+++ b/Engine/Plugins/Experimental/Avalanche/Source/AvalancheLevelViewport/Private/AvaViewportCameraHistory.h
@@ -59,7 +59,7 @@ private:
 	TArray<TPair<TWeakObjectPtr<AActor>, FTransform>, TFixedAllocator<UE::AvalancheEditor::Internal::CameraUndoHistoryCapacity>> CameraTransformHistory;
 
 	// Stores current position in the camera transform history, to allow redo of actions ahead of this position
-	int32 CameraTransformHistoryIndex = 0;
+	int32 CameraTransformHistoryIndex = INDEX_NONE;
 
 	// Stores current position of the array "end", needed due to circular storage
 	int32 CameraTransformHistoryHeadIndex = 1;
