[Motion Design] CIS issue

diff --git a/Engine/Plugins/Experimental/Avalanche/Source/AvalancheLevelViewport/Private/AvaViewportCameraHistory.cpp b/Engine/Plugins/Experimental/Avalanche/Source/AvalancheLevelViewport/Private/AvaViewportCameraHistory.cpp
--- a/Engine/Plugins/Experimental/Avalanche/Source/AvalancheLevelViewport/Private/AvaViewportCameraHistory.cpp
+++ b/Engine/Plugins/Experimental/Avalanche/Source/AvalancheLevelViewport/Private/AvaViewportCameraHistory.cpp
@@ -93,7 +93,7 @@ void FAvaViewportCameraHistory::PostEngineInit()
 void FAvaViewportCameraHistory::Reset()
 {
 	CameraTransformHistory.Reset();
-	CameraTransformHistory.SetNumUninitialized(UE::AvalancheEditor::Internal::CameraUndoHistoryCapacity, false);
+	CameraTransformHistory.SetNumUninitialized(UE::AvalancheEditor::Internal::CameraUndoHistoryCapacity, EAllowShrinking::No);
 	CameraTransformHistoryIndex = INDEX_NONE;
 	CameraTransformHistoryHeadIndex = 1;
 }
