MotionDesign : Modifiers
- Fix modifier scene extension invalid world check

#changelist validated
#virtualized

diff --git a/Engine/Plugins/Experimental/Avalanche/Source/AvalancheModifiers/Private/Extensions/AvaSceneTreeUpdateModifierExtension.cpp b/Engine/Plugins/Experimental/Avalanche/Source/AvalancheModifiers/Private/Extensions/AvaSceneTreeUpdateModifierExtension.cpp
--- a/Engine/Plugins/Experimental/Avalanche/Source/AvalancheModifiers/Private/Extensions/AvaSceneTreeUpdateModifierExtension.cpp
+++ b/Engine/Plugins/Experimental/Avalanche/Source/AvalancheModifiers/Private/Extensions/AvaSceneTreeUpdateModifierExtension.cpp
@@ -331,17 +331,15 @@ TArray<TWeakObjectPtr<AActor>> FAvaSceneTreeUpdateModifierExtension::GetDirectCh
 	TArray<AActor*> DirectChildren;
 
 #if WITH_EDITOR
-	const UWorld* const World = InActor->GetWorld();
-	check(IsValid(World));
-
-	const UAvaOutlinerSubsystem* const OutlinerSubsystem = World->GetSubsystem<UAvaOutlinerSubsystem>();
-	if (IsValid(OutlinerSubsystem))
+	if (const UWorld* const World = InActor->GetWorld())
 	{
-		const TSharedPtr<IAvaOutliner> AvaOutliner = OutlinerSubsystem->GetOutliner();
-		if (AvaOutliner.IsValid())
+		if (const UAvaOutlinerSubsystem* const OutlinerSubsystem = World->GetSubsystem<UAvaOutlinerSubsystem>())
 		{
-			DirectChildren = FAvaOutlinerUtils::EditorOutlinerChildActors(AvaOutliner, InActor);
-			bIsOutlinerAttachedActors = true;
+			if (const TSharedPtr<IAvaOutliner> AvaOutliner = OutlinerSubsystem->GetOutliner())
+			{
+				DirectChildren = FAvaOutlinerUtils::EditorOutlinerChildActors(AvaOutliner, InActor);
+				bIsOutlinerAttachedActors = true;
+			}
 		}
 	}
 #endif
