[Motion Design]
 - Fix camera shake source tool missing icon caused from engine icon change

diff --git a/Engine/Plugins/Experimental/Avalanche/Source/AvalancheEditor/Private/AvaEditorStyle.cpp b/Engine/Plugins/Experimental/Avalanche/Source/AvalancheEditor/Private/AvaEditorStyle.cpp
--- a/Engine/Plugins/Experimental/Avalanche/Source/AvalancheEditor/Private/AvaEditorStyle.cpp
+++ b/Engine/Plugins/Experimental/Avalanche/Source/AvalancheEditor/Private/AvaEditorStyle.cpp
@@ -170,7 +170,7 @@ FAvaEditorStyle::FAvaEditorStyle()
 	Set("AvaEditor.CineCameraTool",           new CORE_IMAGE_BRUSH_SVG("Starship/AssetIcons/CineCameraActor_16",   Icon16));
 	Set("AvaEditor.CameraRigCraneTool",       new CORE_IMAGE_BRUSH_SVG("Starship/AssetIcons/CameraRig_Crane_16",   Icon16));
 	Set("AvaEditor.CameraRigRailTool",        new CORE_IMAGE_BRUSH_SVG("Starship/AssetIcons/CameraRig_Rail_16",    Icon16));
-	Set("AvaEditor.CameraShakeSourceTool",    new CORE_IMAGE_BRUSH_SVG("Starship/AssetIcons/CameraShake",          Icon16));
+	Set("AvaEditor.CameraShakeSourceTool",    new CORE_IMAGE_BRUSH_SVG("Starship/Common/CameraShake",              Icon16));
 	Set("AvaEditor.AvaPostProcessVolumeTool", new CORE_IMAGE_BRUSH_SVG("Starship/AssetIcons/PostProcessVolume_16", Icon16));
 
 	Set("AvaEditor.PointLightTool",       new CORE_IMAGE_BRUSH_SVG("Starship/AssetIcons/PointLight_16",       Icon16));
