Motion Design:
- Made TransitionTree Toolbar button callout (to display label for clarity)

#changelist validated
#virtualized

diff --git a/Engine/Plugins/Experimental/Avalanche/Source/AvalancheEditor/Private/Transition/AvaTransitionExtension.cpp b/Engine/Plugins/Experimental/Avalanche/Source/AvalancheEditor/Private/Transition/AvaTransitionExtension.cpp
--- a/Engine/Plugins/Experimental/Avalanche/Source/AvalancheEditor/Private/Transition/AvaTransitionExtension.cpp
+++ b/Engine/Plugins/Experimental/Avalanche/Source/AvalancheEditor/Private/Transition/AvaTransitionExtension.cpp
@@ -125,16 +125,18 @@ void FAvaTransitionExtension::Deactivate()
 
 void FAvaTransitionExtension::ExtendToolbarMenu(UToolMenu& InMenu)
 {
-	FToolMenuSection& Section = InMenu.FindOrAddSection(TEXT("TransitionLogic"));
+	FToolMenuSection& Section = InMenu.FindOrAddSection(DefaultSectionName);
 
 	FSlateIcon TransitionTreeIcon = FSlateIconFinder::FindCustomIconForClass(UAvaTransitionTree::StaticClass(), TEXT("ClassThumbnail"));
 
-	Section.AddEntry(FToolMenuEntry::InitToolBarButton(TEXT("TransitionLogicButton")
+	FToolMenuEntry& Entry = Section.AddEntry(FToolMenuEntry::InitToolBarButton(TEXT("TransitionLogicButton")
 		, FExecuteAction::CreateSP(this, &FAvaTransitionExtension::OpenTransitionEditor)
 		, LOCTEXT("TransitionLogicLabel", "Transition Logic")
 		, LOCTEXT("TransitionLogicTooltip", "Opens the Transition Logic Editor for the given Scene")
 		, TransitionTreeIcon));
 
+	Entry.StyleNameOverride = TEXT("CalloutToolbar");
+
 	Section.AddEntry(FToolMenuEntry::InitComboButton(TEXT("TransitionLogicComboButton")
 		, FUIAction()
 		, FNewToolMenuDelegate::CreateSP(this, &FAvaTransitionExtension::GenerateTransitionLogicOptions)
