MotionDesign : PropertyAnimator
- Only allow preview if sequence matches root sequence of sequencer to avoid confusion for user

#changelist validated
#virtualized

diff --git a/Engine/Plugins/Experimental/Avalanche/Source/AvalanchePropertyAnimator/Private/TimeSources/AvaPropertyAnimatorSequenceTimeSource.cpp b/Engine/Plugins/Experimental/Avalanche/Source/AvalanchePropertyAnimator/Private/TimeSources/AvaPropertyAnimatorSequenceTimeSource.cpp
--- a/Engine/Plugins/Experimental/Avalanche/Source/AvalanchePropertyAnimator/Private/TimeSources/AvaPropertyAnimatorSequenceTimeSource.cpp
+++ b/Engine/Plugins/Experimental/Avalanche/Source/AvalanchePropertyAnimator/Private/TimeSources/AvaPropertyAnimatorSequenceTimeSource.cpp
@@ -92,7 +92,14 @@ double UAvaPropertyAnimatorSequenceTimeSource::GetTimeElapsed()
 	// Use sequencer global time if sequencer active
 	if (const TSharedPtr<ISequencer> Sequencer = GetSequencer())
 	{
-		return Sequencer->GetGlobalTime().AsSeconds();
+		// Scrub to global time only if root sequence is the selected sequence
+		if (const UMovieSceneSequence* RootSequence = Sequencer->GetRootMovieSceneSequence())
+		{
+			if (RootSequence->GetName() == SequenceName)
+			{
+				return Sequencer->GetGlobalTime().AsSeconds();
+			}
+		}
 	}
 #endif
 
