Motion Design:
- Removed MRQ Rundown from auto-adding additional settings

#changelist validated
#virtualized

diff --git a/Engine/Plugins/Experimental/Avalanche/Source/AvalancheMRQEditor/Private/AvaMRQEditorRundowntUtils.cpp b/Engine/Plugins/Experimental/Avalanche/Source/AvalancheMRQEditor/Private/AvaMRQEditorRundowntUtils.cpp
--- a/Engine/Plugins/Experimental/Avalanche/Source/AvalancheMRQEditor/Private/AvaMRQEditorRundowntUtils.cpp
+++ b/Engine/Plugins/Experimental/Avalanche/Source/AvalancheMRQEditor/Private/AvaMRQEditorRundowntUtils.cpp
@@ -6,9 +6,6 @@
 #include "AvaSceneSubsystem.h"
 #include "IAvaSceneInterface.h"
 #include "IAvaSequenceProvider.h"
-#include "MoviePipelineDeferredPasses.h"
-#include "MoviePipelineImageSequenceOutput.h"
-#include "MoviePipelineOutputSetting.h"
 #include "MoviePipelinePIEExecutor.h"
 #include "MoviePipelineQueueSubsystem.h"
 #include "Playback/AvaPlaybackUtils.h"
@@ -128,16 +125,6 @@ namespace UE::AvaMRQEditor::Private
 			OutSettings = Cast<InSettingType>(MRQConfig->FindOrAddSettingByClass(InSettingType::StaticClass()));
 		};
 
-		UMoviePipelineImageSequenceOutput_PNG* OutputPNG;
-		CreatePipelineSetting(OutputPNG);
-
-		UMoviePipelineDeferredPassBase* DeferredPass;
-		CreatePipelineSetting(DeferredPass);
-
-		UMoviePipelineOutputSetting* OutputSettings;
-		CreatePipelineSetting(OutputSettings);
-		OutputSettings->FileNameFormat = TEXT("{job_name}.{frame_number}");;
-
 		UAvaMRQRundownPageSetting* RundownPageSetting;
 		CreatePipelineSetting(RundownPageSetting);
 		RundownPageSetting->RundownPage.Rundown = &InRundown;
