MotionDesign : TranslucentSortPriority Modifier
- Adding offset and step support per level to avoid clash when multiple levels are used within the same world during TL

#changelist validated
#virtualized

diff --git a/Engine/Plugins/Experimental/Avalanche/Source/AvalancheModifiers/Private/Modifiers/AvaTranslucentPriorityModifier.cpp b/Engine/Plugins/Experimental/Avalanche/Source/AvalancheModifiers/Private/Modifiers/AvaTranslucentPriorityModifier.cpp
--- a/Engine/Plugins/Experimental/Avalanche/Source/AvalancheModifiers/Private/Modifiers/AvaTranslucentPriorityModifier.cpp
+++ b/Engine/Plugins/Experimental/Avalanche/Source/AvalancheModifiers/Private/Modifiers/AvaTranslucentPriorityModifier.cpp
@@ -63,6 +63,23 @@ void UAvaTranslucentPriorityModifier::OnModifierAdded(EActorModifierCoreEnableRe
 		CameraActorWeak = GetDefaultCameraActor();
 		Mode = EAvaTranslucentPriorityModifierMode::AutoOutlinerTree;
 	}
+
+	if (UAvaTranslucentPriorityModifierShared* SharedObject = GetShared<UAvaTranslucentPriorityModifierShared>(true))
+	{
+		SortPriorityOffset = SharedObject->GetSortPriorityOffset();
+		SortPriorityStep = SharedObject->GetSortPriorityStep();
+		SharedObject->OnLevelGlobalsChangedDelegate.AddUObject(this, &UAvaTranslucentPriorityModifier::OnGlobalSortPriorityOffsetChanged);
+	}
+}
+
+void UAvaTranslucentPriorityModifier::OnModifierRemoved(EActorModifierCoreDisableReason InReason)
+{
+	Super::OnModifierRemoved(InReason);
+
+	if (UAvaTranslucentPriorityModifierShared* SharedObject = GetShared<UAvaTranslucentPriorityModifierShared>(false))
+	{
+		SharedObject->OnLevelGlobalsChangedDelegate.RemoveAll(this);
+	}
 }
 
 void UAvaTranslucentPriorityModifier::SavePreState()
@@ -139,8 +156,14 @@ void UAvaTranslucentPriorityModifier::Apply()
 		CachedSortedComponentStates = SharedObject->GetSortedComponentStates(this);
 	}
 
+	// Get global sort offset and step for this level
+	const int32 GlobalOffset = SharedObject->GetSortPriorityOffset();
+	const int32 GlobalStep = SharedObject->GetSortPriorityStep();
+
 	if (Mode == EAvaTranslucentPriorityModifierMode::Manual)
 	{
+		int32 TranslucentSortPriority = GlobalOffset + SortPriority;
+
 		// Sets all components with the same priority
 		for (const FAvaTranslucentPriorityModifierComponentState* SortedComponentState : CachedSortedComponentStates)
 		{
@@ -148,42 +171,50 @@ void UAvaTranslucentPriorityModifier::Apply()
 			{
 				if (UPrimitiveComponent* Component = SortedComponentState->PrimitiveComponentWeak.Get())
 				{
-					Component->SetTranslucentSortPriority(SortPriority);
+					Component->SetTranslucentSortPriority(TranslucentSortPriority);
+
+					TranslucentSortPriority += GlobalStep;
 				}
 			}
 		}
 	}
 	else
 	{
-		int32 TranslucentSortPriority = 1;
+		int32 TranslucentSortPriority = GlobalOffset;
 
 		// Increment sort priority for each component that this modifier handles
 		for (const FAvaTranslucentPriorityModifierComponentState* SortedComponentState : CachedSortedComponentStates)
 		{
-			if (!SortedComponentState
-				|| !SortedComponentState->PrimitiveComponentWeak.IsValid()
-				|| !SortedComponentState->ModifierWeak.IsValid())
+			if (!SortedComponentState)
 			{
 				continue;
 			}
 
 			UPrimitiveComponent* Component = SortedComponentState->PrimitiveComponentWeak.Get();
+			UAvaTranslucentPriorityModifier* ComponentModifier = SortedComponentState->ModifierWeak.Get();
+
+			if (!Component || !ComponentModifier)
+			{
+				continue;
+			}
 
-			if (SortedComponentState->ModifierWeak == this)
+			// This modifier handles this component
+			if (ComponentModifier == this)
 			{
 				Component->SetTranslucentSortPriority(TranslucentSortPriority);
 			}
-			else if (UAvaTranslucentPriorityModifier* OtherModifier = SortedComponentState->ModifierWeak.Get())
+			// Another modifier handles this component
+			else
 			{
 				if (Component->TranslucencySortPriority != TranslucentSortPriority)
 				{
 					// Cache to avoid doing the same query for the same result
-					OtherModifier->CachedSortedComponentStates = CachedSortedComponentStates;
-					OtherModifier->MarkModifierDirty();
+					ComponentModifier->CachedSortedComponentStates = CachedSortedComponentStates;
+					ComponentModifier->MarkModifierDirty();
 				}
 			}
 
-			TranslucentSortPriority++;
+			TranslucentSortPriority += GlobalStep;
 		}
 	}
 
@@ -245,24 +276,24 @@ void UAvaTranslucentPriorityModifier::PostEditChangeProperty(FPropertyChangedEve
 
 	const FName MemberName = PropertyChangedEvent.GetMemberPropertyName();
 
-	static const FName ModeName = GET_MEMBER_NAME_CHECKED(UAvaTranslucentPriorityModifier, Mode);
-	static const FName CameraActorName = GET_MEMBER_NAME_CHECKED(UAvaTranslucentPriorityModifier, CameraActorWeak);
-	static const FName SortPriorityName = GET_MEMBER_NAME_CHECKED(UAvaTranslucentPriorityModifier, SortPriority);
-	static const FName IncludeChildrenName = GET_MEMBER_NAME_CHECKED(UAvaTranslucentPriorityModifier, bIncludeChildren);
-
-	if (MemberName == ModeName)
+	if (MemberName == GET_MEMBER_NAME_CHECKED(UAvaTranslucentPriorityModifier, Mode))
 	{
 		OnModeChanged();
 	}
-	else if (MemberName == CameraActorName)
+	else if (MemberName == GET_MEMBER_NAME_CHECKED(UAvaTranslucentPriorityModifier, CameraActorWeak))
 	{
 		OnCameraActorChanged();
 	}
-	else if (MemberName == SortPriorityName)
+	else if (MemberName == GET_MEMBER_NAME_CHECKED(UAvaTranslucentPriorityModifier, SortPriority))
 	{
 		OnSortPriorityChanged();
 	}
-	else if (MemberName == IncludeChildrenName)
+	else if (MemberName == GET_MEMBER_NAME_CHECKED(UAvaTranslucentPriorityModifier, SortPriorityOffset)
+		|| MemberName == GET_MEMBER_NAME_CHECKED(UAvaTranslucentPriorityModifier, SortPriorityStep))
+	{
+		OnSortPriorityLevelGlobalsChanged();
+	}
+	else if (MemberName == GET_MEMBER_NAME_CHECKED(UAvaTranslucentPriorityModifier, bIncludeChildren))
 	{
 		OnIncludeChildrenChanged();
 	}
@@ -312,6 +343,28 @@ void UAvaTranslucentPriorityModifier::SetSortPriority(int32 InSortPriority)
 	OnSortPriorityChanged();
 }
 
+void UAvaTranslucentPriorityModifier::SetSortPriorityOffset(int32 InOffset)
+{
+	if (SortPriorityOffset == InOffset)
+	{
+		return;
+	}
+
+	SortPriorityOffset = InOffset;
+	OnSortPriorityLevelGlobalsChanged();
+}
+
+void UAvaTranslucentPriorityModifier::SetSortPriorityStep(int32 InStep)
+{
+	if (SortPriorityStep == InStep)
+	{
+		return;
+	}
+
+	SortPriorityStep = InStep;
+	OnSortPriorityLevelGlobalsChanged();
+}
+
 void UAvaTranslucentPriorityModifier::SetIncludeChildren(bool bInIncludeChildren)
 {
 	if (bIncludeChildren == bInIncludeChildren)
@@ -428,6 +481,25 @@ void UAvaTranslucentPriorityModifier::OnSortPriorityChanged()
 	}
 }
 
+void UAvaTranslucentPriorityModifier::OnSortPriorityLevelGlobalsChanged() const
+{
+	if (UAvaTranslucentPriorityModifierShared* SharedObject = GetShared<UAvaTranslucentPriorityModifierShared>(false))
+	{
+		SharedObject->SetSortPriorityOffset(SortPriorityOffset);
+		SharedObject->SetSortPriorityStep(SortPriorityStep);
+	}
+}
+
+void UAvaTranslucentPriorityModifier::OnGlobalSortPriorityOffsetChanged()
+{
+	if (const UAvaTranslucentPriorityModifierShared* SharedObject = GetShared<UAvaTranslucentPriorityModifierShared>(false))
+	{
+		SortPriorityOffset = SharedObject->GetSortPriorityOffset();
+		SortPriorityStep = SharedObject->GetSortPriorityStep();
+		MarkModifierDirty();
+	}
+}
+
 void UAvaTranslucentPriorityModifier::OnIncludeChildrenChanged()
 {
 	MarkModifierDirty();
diff --git a/Engine/Plugins/Experimental/Avalanche/Source/AvalancheModifiers/Private/Shared/AvaTranslucentPriorityModifierShared.cpp b/Engine/Plugins/Experimental/Avalanche/Source/AvalancheModifiers/Private/Shared/AvaTranslucentPriorityModifierShared.cpp
--- a/Engine/Plugins/Experimental/Avalanche/Source/AvalancheModifiers/Private/Shared/AvaTranslucentPriorityModifierShared.cpp
+++ b/Engine/Plugins/Experimental/Avalanche/Source/AvalancheModifiers/Private/Shared/AvaTranslucentPriorityModifierShared.cpp
@@ -2,6 +2,7 @@
 
 #include "Shared/AvaTranslucentPriorityModifierShared.h"
 #include "AvaActorUtils.h"
+#include "AvaSceneItem.h"
 #include "AvaSceneTree.h"
 #include "AvaSceneTreeNode.h"
 #include "Camera/CameraActor.h"
@@ -199,7 +200,7 @@ TArray<const FAvaTranslucentPriorityModifierComponentState*> UAvaTranslucentPrio
 		}
 
 		// Sort current modifiers by distance X to a specific camera
-		SortedComponentStates.Sort([CameraActor](const FAvaTranslucentPriorityModifierComponentState& InComponentA, const FAvaTranslucentPriorityModifierComponentState& InComponentB)->bool
+		SortedComponentStates.StableSort([CameraActor](const FAvaTranslucentPriorityModifierComponentState& InComponentA, const FAvaTranslucentPriorityModifierComponentState& InComponentB)->bool
 		{
 			const FVector AForwardComponentLocation = CameraActor->GetActorForwardVector() * InComponentA.GetComponentLocation();
 			const float ADist = FVector::Distance(AForwardComponentLocation, CameraActor->GetActorLocation());
@@ -244,7 +245,7 @@ TArray<const FAvaTranslucentPriorityModifierComponentState*> UAvaTranslucentPrio
 		TSharedPtr<IAvaOutliner> AvaOutliner = FAvaOutlinerUtils::EditorGetOutliner(World);
 		if (AvaOutliner.IsValid())
 		{
-			SortedComponentStates.Sort([AvaOutliner](const FAvaTranslucentPriorityModifierComponentState& InComponentA, const FAvaTranslucentPriorityModifierComponentState& InComponentB)->bool
+			SortedComponentStates.StableSort([AvaOutliner](const FAvaTranslucentPriorityModifierComponentState& InComponentA, const FAvaTranslucentPriorityModifierComponentState& InComponentB)->bool
 			{
 				const FAvaOutlinerItemPtr OutlinerItemA = AvaOutliner->FindItem(InComponentA.GetOwningActor());
 				const FAvaOutlinerItemPtr OutlinerItemB = AvaOutliner->FindItem(InComponentB.GetOwningActor());
@@ -260,7 +261,7 @@ TArray<const FAvaTranslucentPriorityModifierComponentState*> UAvaTranslucentPrio
 			{
 				const FAvaSceneTree& SceneTree = SceneInterface->GetSceneTree();
 
-				SortedComponentStates.Sort([&SceneTree, World](const FAvaTranslucentPriorityModifierComponentState& InComponentA, const FAvaTranslucentPriorityModifierComponentState& InComponentB)->bool
+				SortedComponentStates.StableSort([&SceneTree, World](const FAvaTranslucentPriorityModifierComponentState& InComponentA, const FAvaTranslucentPriorityModifierComponentState& InComponentB)->bool
 				{
 					const FAvaSceneTreeNode* SceneItemA = SceneTree.FindTreeNode(FAvaSceneItem(InComponentA.GetOwningActor(), World));
 					const FAvaSceneTreeNode* SceneItemB = SceneTree.FindTreeNode(FAvaSceneItem(InComponentB.GetOwningActor(), World));
@@ -276,6 +277,28 @@ TArray<const FAvaTranslucentPriorityModifierComponentState*> UAvaTranslucentPrio
 	return SortedComponentStates;
 }
 
+void UAvaTranslucentPriorityModifierShared::SetSortPriorityOffset(int32 InOffset)
+{
+	if (SortPriorityOffset == InOffset)
+	{
+		return;
+	}
+
+	SortPriorityOffset = InOffset;
+	OnLevelGlobalsChangedDelegate.Broadcast();
+}
+
+void UAvaTranslucentPriorityModifierShared::SetSortPriorityStep(int32 InStep)
+{
+	if (SortPriorityStep == InStep)
+	{
+		return;
+	}
+
+	SortPriorityStep = InStep;
+	OnLevelGlobalsChangedDelegate.Broadcast();
+}
+
 void UAvaTranslucentPriorityModifierShared::PostLoad()
 {
 	Super::PostLoad();
diff --git a/Engine/Plugins/Experimental/Avalanche/Source/AvalancheModifiers/Private/Tests/AvaModifiersTranslucentPriority.spec.cpp b/Engine/Plugins/Experimental/Avalanche/Source/AvalancheModifiers/Private/Tests/AvaModifiersTranslucentPriority.spec.cpp
--- a/Engine/Plugins/Experimental/Avalanche/Source/AvalancheModifiers/Private/Tests/AvaModifiersTranslucentPriority.spec.cpp
+++ b/Engine/Plugins/Experimental/Avalanche/Source/AvalancheModifiers/Private/Tests/AvaModifiersTranslucentPriority.spec.cpp
@@ -43,6 +43,8 @@ void AvalancheModifiersTranslucentPriority::Define()
 		UActorModifierCoreStack* ModifierStack = ModifierTestUtils->GenerateModifierStackForActor(ModifiedActor);
 		TranslucentPriorityModifier = Cast<UAvaTranslucentPriorityModifier>(ModifierSubsystem->InsertModifier(ModifierStack, InsertOp));
 		TranslucentPriorityModifier->SetMode(EAvaTranslucentPriorityModifierMode::Manual);
+		TranslucentPriorityModifier->SetSortPriorityOffset(0);
+		TranslucentPriorityModifier->SetSortPriorityStep(0);
 		TranslucentPriorityModifier->SetSortPriority(ModifierSortPriority);
 	});
 
diff --git a/Engine/Plugins/Experimental/Avalanche/Source/AvalancheModifiers/Public/Modifiers/AvaTranslucentPriorityModifier.h b/Engine/Plugins/Experimental/Avalanche/Source/AvalancheModifiers/Public/Modifiers/AvaTranslucentPriorityModifier.h
--- a/Engine/Plugins/Experimental/Avalanche/Source/AvalancheModifiers/Public/Modifiers/AvaTranslucentPriorityModifier.h
+++ b/Engine/Plugins/Experimental/Avalanche/Source/AvalancheModifiers/Public/Modifiers/AvaTranslucentPriorityModifier.h
@@ -54,6 +54,18 @@ public:
 	    return SortPriority;
     }
 
+	void SetSortPriorityOffset(int32 InOffset);
+	int32 GetSortPriorityOffset() const
+	{
+		return SortPriorityOffset;
+	}
+
+	void SetSortPriorityStep(int32 InStep);
+	int32 GetSortPriorityStep() const
+	{
+		return SortPriorityStep;
+	}
+
 	void SetIncludeChildren(bool bInIncludeChildren);
 	bool GetIncludeChildren() const
 	{
@@ -64,6 +76,7 @@ protected:
 	//~ Begin UActorModifierCoreBase
 	virtual void OnModifierCDOSetup(FActorModifierCoreMetadata& InMetadata) override;
 	virtual void OnModifierAdded(EActorModifierCoreEnableReason InReason) override;
+	virtual void OnModifierRemoved(EActorModifierCoreDisableReason InReason) override;
 	virtual void SavePreState() override;
 	virtual void RestorePreState() override;
 	virtual void Apply() override;
@@ -88,8 +101,11 @@ protected:
 	void OnModeChanged();
 	void OnCameraActorChanged();
 	void OnSortPriorityChanged();
+	void OnSortPriorityLevelGlobalsChanged() const;
 	void OnIncludeChildrenChanged();
 
+	void OnGlobalSortPriorityOffsetChanged();
+
 	ACameraActor* GetDefaultCameraActor() const;
 
 	/** The sort mode we are currently in */
@@ -104,6 +120,14 @@ protected:
 	UPROPERTY(EditInstanceOnly, BlueprintReadWrite, Setter, Getter, Category="TranslucentPriority", meta=(EditCondition="Mode == EAvaTranslucentPriorityModifierMode::Manual", EditConditionHides, AllowPrivateAccess="true"))
 	int32 SortPriority = 0;
 
+	/** Sort priority offset shared across all modifiers in this same level */
+	UPROPERTY(Transient, EditInstanceOnly, BlueprintReadWrite, Setter, Getter, Category="TranslucentPriority", meta=(AllowPrivateAccess="true"))
+	int32 SortPriorityOffset = 0;
+
+	/** Sort priority incremental step shared across all modifiers in this same level */
+	UPROPERTY(Transient, EditInstanceOnly, BlueprintReadWrite, Setter, Getter, Category="TranslucentPriority", meta=(AllowPrivateAccess="true"))
+	int32 SortPriorityStep = 1;
+
 	/** If true, will include children too and update their sort priority */
 	UPROPERTY(EditAnywhere, BlueprintReadWrite, Setter="SetIncludeChildren", Getter="GetIncludeChildren", Category="TranslucentPriority", meta=(AllowPrivateAccess="true"))
 	bool bIncludeChildren = false;
diff --git a/Engine/Plugins/Experimental/Avalanche/Source/AvalancheModifiers/Public/Shared/AvaTranslucentPriorityModifierShared.h b/Engine/Plugins/Experimental/Avalanche/Source/AvalancheModifiers/Public/Shared/AvaTranslucentPriorityModifierShared.h
--- a/Engine/Plugins/Experimental/Avalanche/Source/AvalancheModifiers/Public/Shared/AvaTranslucentPriorityModifierShared.h
+++ b/Engine/Plugins/Experimental/Avalanche/Source/AvalancheModifiers/Public/Shared/AvaTranslucentPriorityModifierShared.h
@@ -67,6 +67,9 @@ class AVALANCHEMODIFIERS_API UAvaTranslucentPriorityModifierShared : public UAct
 	GENERATED_BODY()
 
 public:
+	DECLARE_MULTICAST_DELEGATE(FOnLevelGlobalsChanged)
+	FOnLevelGlobalsChanged OnLevelGlobalsChangedDelegate;
+
 	/** Replaces the component linked to this context by new components set */
 	void SetComponentsState(UAvaTranslucentPriorityModifier* InModifierContext, const TSet<TWeakObjectPtr<UPrimitiveComponent>>& InComponents);
 
@@ -88,6 +91,20 @@ public:
 	/** Get sorted components state based on modifier context */
 	TArray<const FAvaTranslucentPriorityModifierComponentState*> GetSortedComponentStates(UAvaTranslucentPriorityModifier* InModifierContext) const;
 
+	int32 GetSortPriorityOffset() const
+	{
+		return SortPriorityOffset;
+	}
+
+	void SetSortPriorityOffset(int32 InOffset);
+
+	int32 GetSortPriorityStep() const
+	{
+		return SortPriorityStep;
+	}
+
+	void SetSortPriorityStep(int32 InStep);
+
 private:
 	//~ Begin UObject
 	virtual void PostLoad() override;
@@ -95,4 +112,12 @@ private:
 
 	UPROPERTY()
 	TSet<FAvaTranslucentPriorityModifierComponentState> ComponentStates;
+
+	/** Offset for the whole level */
+	UPROPERTY()
+	int32 SortPriorityOffset = 0;
+
+	/** Incremental step for the whole level */
+	UPROPERTY()
+	int32 SortPriorityStep = 1;
 };
