Motion Design
- "Tool Box" removed from the Motion Design editor mode name. So long! :salute:

diff --git a/Engine/Plugins/Experimental/Avalanche/Source/AvalancheInteractiveTools/Private/EdMode/AvaInteractiveToolsEdMode.cpp b/Engine/Plugins/Experimental/Avalanche/Source/AvalancheInteractiveTools/Private/EdMode/AvaInteractiveToolsEdMode.cpp
--- a/Engine/Plugins/Experimental/Avalanche/Source/AvalancheInteractiveTools/Private/EdMode/AvaInteractiveToolsEdMode.cpp
+++ b/Engine/Plugins/Experimental/Avalanche/Source/AvalancheInteractiveTools/Private/EdMode/AvaInteractiveToolsEdMode.cpp
@@ -16,7 +16,7 @@
 UAvaInteractiveToolsEdMode::UAvaInteractiveToolsEdMode()
 {
 	Info = FEditorModeInfo(IAvalancheInteractiveToolsModule::EM_AvaInteractiveToolsEdModeId,
-		LOCTEXT("AvaInteractiveToolsEdModeName", "Motion Design Tool Box"),
+		LOCTEXT("AvaInteractiveToolsEdModeName", "Motion Design"),
 		FSlateIcon(),
 		true);
 
