[Motion Design]
 - Ease curve tool analytics

diff --git a/Engine/Plugins/Experimental/Avalanche/Source/AvalancheSequencer/Private/EaseCurveTool/AvaEaseCurveSubsystem.cpp b/Engine/Plugins/Experimental/Avalanche/Source/AvalancheSequencer/Private/EaseCurveTool/AvaEaseCurveSubsystem.cpp
--- a/Engine/Plugins/Experimental/Avalanche/Source/AvalancheSequencer/Private/EaseCurveTool/AvaEaseCurveSubsystem.cpp
+++ b/Engine/Plugins/Experimental/Avalanche/Source/AvalancheSequencer/Private/EaseCurveTool/AvaEaseCurveSubsystem.cpp
@@ -1,6 +1,7 @@
 // Copyright Epic Games, Inc. All Rights Reserved.
 
 #include "AvaEaseCurveSubsystem.h"
+#include "Containers/UnrealString.h"
 #include "EaseCurveTool/AvaEaseCurvePreset.h"
 #include "EaseCurveTool/AvaEaseCurveToolSettings.h"
 #include "Editor.h"
@@ -634,22 +635,22 @@ void UAvaEaseCurveSubsystem::ResetToDefaultPresets(const bool bInOnlyIfNoProject
 {
 	IPlatformFile& PlatformFile = FPlatformFileManager::Get().GetPlatformFile();
 
-	const FString ProjectPresetPath = FPaths::ConvertRelativePathToFull(UAvaEaseCurveSubsystem::ProjectPresetPath());
+	const FString SourceProjectPresetPath = FPaths::ConvertRelativePathToFull(UAvaEaseCurveSubsystem::ProjectPresetPath());
 
 	if (!bInOnlyIfNoProjectPresets)
 	{
-		PlatformFile.DeleteDirectoryRecursively(*ProjectPresetPath);
-		PlatformFile.CreateDirectory(*ProjectPresetPath);
+		PlatformFile.DeleteDirectoryRecursively(*SourceProjectPresetPath);
+		PlatformFile.CreateDirectory(*SourceProjectPresetPath);
 	}
 
 	TArray<FString> JsonFiles;
-	IFileManager::Get().FindFilesRecursive(JsonFiles, *ProjectPresetPath, TEXT("*.json"), true, false);
+	IFileManager::Get().FindFilesRecursive(JsonFiles, *SourceProjectPresetPath, TEXT("*.json"), true, false);
 
 	if (JsonFiles.IsEmpty())
 	{
-		const FString PluginPresetPath = FPaths::ConvertRelativePathToFull(UAvaEaseCurveSubsystem::PluginPresetPath());
+		const FString DestPluginPresetPath = FPaths::ConvertRelativePathToFull(UAvaEaseCurveSubsystem::PluginPresetPath());
 
-		if (PlatformFile.CopyDirectoryTree(*ProjectPresetPath, *PluginPresetPath, false))
+		if (PlatformFile.CopyDirectoryTree(*SourceProjectPresetPath, *DestPluginPresetPath, false))
 		{
 			UE_LOG(LogAvaEaseCurveSubsystem, Warning, TEXT("Motion Design ease curve tool project presets are empty. "
 				"Copied default presets to [Project]/Config/EaseCurves"));
@@ -658,3 +659,24 @@ void UAvaEaseCurveSubsystem::ResetToDefaultPresets(const bool bInOnlyIfNoProject
 		}
 	}
 }
+
+const TMap<FString, TArray<FString>>& UAvaEaseCurveSubsystem::GetDefaultCategoryPresetNames()
+{
+	static const TArray<FString> Names = {
+		TEXT("Sine"), 
+		TEXT("Cubic"), 
+		TEXT("Quintic"), 
+		TEXT("Circular"), 
+		TEXT("Quadratic"), 
+		TEXT("Quartic"), 
+		TEXT("Exponential") 
+	};
+
+	static const TMap<FString, TArray<FString>> CategoryPresetNames = {
+		{ TEXT("Ease In Out"), Names },
+		{ TEXT("Ease In"), Names },
+		{ TEXT("Ease Out"), Names }
+	};
+
+	return CategoryPresetNames;
+}
diff --git a/Engine/Plugins/Experimental/Avalanche/Source/AvalancheSequencer/Private/EaseCurveTool/AvaEaseCurveSubsystem.h b/Engine/Plugins/Experimental/Avalanche/Source/AvalancheSequencer/Private/EaseCurveTool/AvaEaseCurveSubsystem.h
--- a/Engine/Plugins/Experimental/Avalanche/Source/AvalancheSequencer/Private/EaseCurveTool/AvaEaseCurveSubsystem.h
+++ b/Engine/Plugins/Experimental/Avalanche/Source/AvalancheSequencer/Private/EaseCurveTool/AvaEaseCurveSubsystem.h
@@ -4,6 +4,8 @@
 
 #include "AvaEaseCurveTangents.h"
 #include "Containers/Array.h"
+#include "Containers/Map.h"
+#include "Containers/StringFwd.h"
 #include "EditorSubsystem.h"
 #include "Templates/SharedPointer.h"
 #include "AvaEaseCurveSubsystem.generated.h"
@@ -23,6 +25,12 @@ public:
 	static FString ProjectPresetPath();
 	static FString PluginPresetPath();
 
+	/**
+	 * Mainly used to compare clicked presets names for engine analytics, thus returns static category
+	 * and preset names of the default presets instead of looking on disk for the actual default names.
+	 */
+	static const TMap<FString, TArray<FString>>& GetDefaultCategoryPresetNames();
+
 	void ReloadPresetsFromJson();
 
 	TArray<FString> GetEaseCurveCategories() const;
diff --git a/Engine/Plugins/Experimental/Avalanche/Source/AvalancheSequencer/Private/EaseCurveTool/AvaEaseCurveTool.cpp b/Engine/Plugins/Experimental/Avalanche/Source/AvalancheSequencer/Private/EaseCurveTool/AvaEaseCurveTool.cpp
--- a/Engine/Plugins/Experimental/Avalanche/Source/AvalancheSequencer/Private/EaseCurveTool/AvaEaseCurveTool.cpp
+++ b/Engine/Plugins/Experimental/Avalanche/Source/AvalancheSequencer/Private/EaseCurveTool/AvaEaseCurveTool.cpp
@@ -5,11 +5,12 @@
 #include "AssetToolsModule.h"
 #include "AvaSequencer.h"
 #include "CurveEditor.h"
+#include "Dialogs/DlgPickAssetPath.h"
 #include "EaseCurveTool/AvaEaseCurveToolCommands.h"
 #include "EaseCurveTool/AvaEaseCurveToolSettings.h"
 #include "EaseCurveTool/Widgets/SAvaEaseCurveEditor.h"
 #include "EaseCurveTool/Widgets/SAvaEaseCurveTool.h"
-#include "Dialogs/DlgPickAssetPath.h"
+#include "EngineAnalytics.h"
 #include "Factories/CurveFactory.h"
 #include "Framework/Notifications/NotificationManager.h"
 #include "HAL/PlatformApplicationMisc.h"
@@ -276,6 +277,24 @@ void FAvaEaseCurveTool::ApplyQuickEaseToSequencerKeySelections(const EOperation
 	{
 		ToolWidget->ZoomToFit();
 	}
+
+	if (FEngineAnalytics::IsAvailable())
+	{
+		FString ParamValue;
+		switch (InOperation)
+		{
+		case EOperation::InOut:
+			ParamValue = TEXT("InOut");
+			break;
+		case EOperation::In:
+			ParamValue = TEXT("In");
+			break;
+		case EOperation::Out:
+			ParamValue = TEXT("Out");
+			break;
+		}
+		FEngineAnalytics::GetProvider().RecordEvent(TEXT("Editor.Usage.MotionDesign.EaseCurveTool"), TEXT("QuickEase"), ParamValue);
+	}
 }
 
 void FAvaEaseCurveTool::SetSequencerKeySelectionTangents(const FAvaEaseCurveTangents& InTangents, const EOperation InOperation)
diff --git a/Engine/Plugins/Experimental/Avalanche/Source/AvalancheSequencer/Private/EaseCurveTool/Widgets/SAvaEaseCurveTool.cpp b/Engine/Plugins/Experimental/Avalanche/Source/AvalancheSequencer/Private/EaseCurveTool/Widgets/SAvaEaseCurveTool.cpp
--- a/Engine/Plugins/Experimental/Avalanche/Source/AvalancheSequencer/Private/EaseCurveTool/Widgets/SAvaEaseCurveTool.cpp
+++ b/Engine/Plugins/Experimental/Avalanche/Source/AvalancheSequencer/Private/EaseCurveTool/Widgets/SAvaEaseCurveTool.cpp
@@ -12,6 +12,7 @@
 #include "EaseCurveTool/Widgets/SAvaEaseCurveEditor.h"
 #include "EaseCurveTool/Widgets/SAvaEaseCurvePreset.h"
 #include "Editor.h"
+#include "EngineAnalytics.h"
 #include "Framework/Commands/GenericCommands.h"
 #include "Framework/MultiBox/MultiBoxBuilder.h"
 #include "IAvaSequencer.h"
@@ -480,6 +481,21 @@ void SAvaEaseCurveTool::OnPresetChanged(const TSharedPtr<FAvaEaseCurvePreset>& I
 	SetTangents(InPreset->Tangents, ToolOperation.Get(), true, true, true);
 
 	FSlateApplication::Get().SetAllUserFocus(CurveEaseEditorWidget);
+
+	if (FEngineAnalytics::IsAvailable())
+	{
+		// Only send analytics for default presets
+		const TMap<FString, TArray<FString>>& DefaultPresetNames = UAvaEaseCurveSubsystem::GetDefaultCategoryPresetNames();
+		if (DefaultPresetNames.Contains(InPreset->Category)
+			&& DefaultPresetNames[InPreset->Category].Contains(InPreset->Name))
+		{
+			TArray<FAnalyticsEventAttribute> Attributes;
+			Attributes.Emplace(TEXT("Category"), InPreset->Category);
+			Attributes.Emplace(TEXT("Name"), InPreset->Name);
+
+			FEngineAnalytics::GetProvider().RecordEvent(TEXT("Editor.Usage.MotionDesign.EaseCurveTool.SetTangentsPreset"), Attributes);
+		}
+	}
 }
 
 void SAvaEaseCurveTool::OnQuickPresetChanged(const TSharedPtr<FAvaEaseCurvePreset>& InPreset) const
