[Motion Design] Multiple fixes for "reimport page":
- Fix bug when multiple pages are selected. Previous code had a bug and it would only update the ActivePage (first page in the selection). The new code in UpdateDefaultRemoteControlValues properly updates all the given pages.
- Update the default RC values even if the page details tab is closed. Moved the UpdateDefaultRemoteControlValues out of the RC tab widgets. It is now called from the page list and the RC tabs will only refresh UI.
- For Combo Pages: UpdateDefaultRemoteControlValues will also update the sub templates. This should reduce the possibilities of having sub-templates with state values.

#rnx

diff --git a/Engine/Plugins/Experimental/Avalanche/Source/AvalancheMediaEditor/Private/Rundown/AvaRundownEditor.cpp b/Engine/Plugins/Experimental/Avalanche/Source/AvalancheMediaEditor/Private/Rundown/AvaRundownEditor.cpp
--- a/Engine/Plugins/Experimental/Avalanche/Source/AvalancheMediaEditor/Private/Rundown/AvaRundownEditor.cpp
+++ b/Engine/Plugins/Experimental/Avalanche/Source/AvalancheMediaEditor/Private/Rundown/AvaRundownEditor.cpp
@@ -985,37 +985,6 @@ void FAvaRundownEditor::MarkAsModified()
 	}
 }
 
-TArray<TSharedPtr<FAvaRundownManagedInstance>> FAvaRundownEditor::GetManagedInstancesForPage(const UAvaRundown* InRundown, const FAvaRundownPage& InPage)
-{
-	const TArray<FSoftObjectPath> AssetPaths = InPage.GetAssetPaths(InRundown);
-	
-	TArray<TSharedPtr<FAvaRundownManagedInstance>> ManagedInstances;
-	ManagedInstances.Reserve(AssetPaths.Num());
-	
-	FAvaRundownManagedInstanceCache& ManagedInstanceCache = IAvaMediaModule::Get().GetManagedInstanceCache();
-	
-	for (const FSoftObjectPath& AssetPath : AssetPaths)
-	{		
-		if (TSharedPtr<FAvaRundownManagedInstance> ManagedInstance = ManagedInstanceCache.GetOrLoadInstance(AssetPath))
-		{
-			ManagedInstances.Add(ManagedInstance);
-		}
-	}
-	return ManagedInstances;
-}
-
-bool FAvaRundownEditor::MergeDefaultRemoteControlValues(const TArray<TSharedPtr<FAvaRundownManagedInstance>>& InManagedInstances, FAvaPlayableRemoteControlValues& OutMergedValues)
-{
-	bool bAllUniqueIds = true;
-	
-	for (const TSharedPtr<FAvaRundownManagedInstance>& ManagedInstance : InManagedInstances)
-	{
-		bAllUniqueIds &= OutMergedValues.Merge(ManagedInstance->GetDefaultRemoteControlValues());
-	}
-	
-	return bAllUniqueIds;
-}
-
 void FAvaRundownEditor::OnActiveSubListChanged()
 {
 	UAvaRundown* Rundown = AvaRundown.Get();
diff --git a/Engine/Plugins/Experimental/Avalanche/Source/AvalancheMediaEditor/Private/Rundown/AvaRundownEditorUtils.cpp b/Engine/Plugins/Experimental/Avalanche/Source/AvalancheMediaEditor/Private/Rundown/AvaRundownEditorUtils.cpp
--- a/Engine/Plugins/Experimental/Avalanche/Source/AvalancheMediaEditor/Private/Rundown/AvaRundownEditorUtils.cpp
+++ b/Engine/Plugins/Experimental/Avalanche/Source/AvalancheMediaEditor/Private/Rundown/AvaRundownEditorUtils.cpp
@@ -10,11 +10,12 @@
 #include "Framework/Application/SlateApplication.h"
 #include "HAL/FileManager.h"
 #include "IAssetTools.h"
-#include "IAvaMediaEditorModule.h"
+#include "IAvaMediaModule.h"
 #include "IContentBrowserSingleton.h"
 #include "JsonObjectConverter.h"
 #include "Misc/PathViews.h"
 #include "Rundown/AvaRundown.h"
+#include "Rundown/AvaRundownManagedInstanceCache.h"
 #include "StructDeserializer.h"
 #include "StructSerializer.h"
 
@@ -623,4 +624,75 @@ UAvaRundown* UE::AvaRundownEditor::Utils::SaveDuplicateRundown(UAvaRundown* InSo
 	return Cast<UAvaRundown>(IAssetTools::Get().DuplicateAsset(InAssetName, InPackagePath, InSourceRundown));
 }
 
+TArray<TSharedPtr<FAvaRundownManagedInstance>> UE::AvaRundownEditor::Utils::GetManagedInstancesForPage(const UAvaRundown* InRundown, const FAvaRundownPage& InPage)
+{
+	const TArray<FSoftObjectPath> AssetPaths = InPage.GetAssetPaths(InRundown);
+	
+	TArray<TSharedPtr<FAvaRundownManagedInstance>> ManagedInstances;
+	ManagedInstances.Reserve(AssetPaths.Num());
+	
+	FAvaRundownManagedInstanceCache& ManagedInstanceCache = IAvaMediaModule::Get().GetManagedInstanceCache();
+	
+	for (const FSoftObjectPath& AssetPath : AssetPaths)
+	{		
+		if (TSharedPtr<FAvaRundownManagedInstance> ManagedInstance = ManagedInstanceCache.GetOrLoadInstance(AssetPath))
+		{
+			ManagedInstances.Add(ManagedInstance);
+		}
+	}
+	return ManagedInstances;
+}
+
+bool UE::AvaRundownEditor::Utils::MergeDefaultRemoteControlValues(const TArray<TSharedPtr<FAvaRundownManagedInstance>>& InManagedInstances, FAvaPlayableRemoteControlValues& OutMergedValues)
+{
+	bool bAllUniqueIds = true;
+	
+	for (const TSharedPtr<FAvaRundownManagedInstance>& ManagedInstance : InManagedInstances)
+	{
+		if (ManagedInstance.IsValid())
+		{
+			bAllUniqueIds &= OutMergedValues.Merge(ManagedInstance->GetDefaultRemoteControlValues());
+		}
+	}
+	
+	return bAllUniqueIds;
+}
+
+EAvaPlayableRemoteControlChanges UE::AvaRundownEditor::Utils::UpdateDefaultRemoteControlValues(UAvaRundown* InRundown, const TArray<int32>& InSelectedPageIds)
+{
+	EAvaPlayableRemoteControlChanges Changes = EAvaPlayableRemoteControlChanges::None;
+
+	if (!InRundown)
+	{
+		return Changes;
+	}
+	
+	for (const int32 PageId : InSelectedPageIds)
+	{
+		FAvaRundownPage& Page = InRundown->GetPage(PageId);
+
+		if (Page.IsValidPage())
+		{
+			TArray<TSharedPtr<FAvaRundownManagedInstance>> ManagedInstances = GetManagedInstancesForPage(InRundown, Page);
+					
+			if (!ManagedInstances.IsEmpty())
+			{
+				FAvaPlayableRemoteControlValues MergedDefaultRCValues;
+				MergeDefaultRemoteControlValues(ManagedInstances, MergedDefaultRCValues);
+
+				// Using the rundown API for event propagation.
+				constexpr bool bUpdateDefaults = true;
+				Changes |= InRundown->UpdateRemoteControlValues(PageId, MergedDefaultRCValues, bUpdateDefaults);
+			}
+
+			// Combo templates will also update the values of the sub-templates.
+			if (Page.IsComboTemplate())
+			{
+				Changes |= UpdateDefaultRemoteControlValues(InRundown, Page.GetCombinedTemplateIds());
+			}
+		}
+	}
+	return Changes;
+}
+
 #undef LOCTEXT_NAMESPACE
diff --git a/Engine/Plugins/Experimental/Avalanche/Source/AvalancheMediaEditor/Private/Rundown/AvaRundownEditorUtils.h b/Engine/Plugins/Experimental/Avalanche/Source/AvalancheMediaEditor/Private/Rundown/AvaRundownEditorUtils.h
--- a/Engine/Plugins/Experimental/Avalanche/Source/AvalancheMediaEditor/Private/Rundown/AvaRundownEditorUtils.h
+++ b/Engine/Plugins/Experimental/Avalanche/Source/AvalancheMediaEditor/Private/Rundown/AvaRundownEditorUtils.h
@@ -8,6 +8,7 @@
 #include "Templates/SharedPointer.h"
 #include "XmlSerializationDefines.h"
 
+class FAvaRundownManagedInstance;
 class UAvaRundown;
 
 namespace UE::AvaRundownEditor::Utils
@@ -107,4 +108,8 @@ namespace UE::AvaRundownEditor::Utils
 		const FAvaRundownPageInsertPosition& InInsertPosition = FAvaRundownPageInsertPosition());
 	
 	UAvaRundown* SaveDuplicateRundown(UAvaRundown* InSourceRundown, const FString& InAssetName, const FString& InPackagePath);
+
+	TArray<TSharedPtr<FAvaRundownManagedInstance>> GetManagedInstancesForPage(const UAvaRundown* InRundown, const FAvaRundownPage& InPage);
+	bool MergeDefaultRemoteControlValues(const TArray<TSharedPtr<FAvaRundownManagedInstance>>& InManagedInstances, FAvaPlayableRemoteControlValues& OutMergedValues);
+	EAvaPlayableRemoteControlChanges UpdateDefaultRemoteControlValues(UAvaRundown* InRundown, const TArray<int32>& InSelectedPageIds);
 }
\ No newline at end of file
diff --git a/Engine/Plugins/Experimental/Avalanche/Source/AvalancheMediaEditor/Private/Rundown/DetailsView/RemoteControl/Controllers/SAvaRundownRCControllerPanel.cpp b/Engine/Plugins/Experimental/Avalanche/Source/AvalancheMediaEditor/Private/Rundown/DetailsView/RemoteControl/Controllers/SAvaRundownRCControllerPanel.cpp
--- a/Engine/Plugins/Experimental/Avalanche/Source/AvalancheMediaEditor/Private/Rundown/DetailsView/RemoteControl/Controllers/SAvaRundownRCControllerPanel.cpp
+++ b/Engine/Plugins/Experimental/Avalanche/Source/AvalancheMediaEditor/Private/Rundown/DetailsView/RemoteControl/Controllers/SAvaRundownRCControllerPanel.cpp
@@ -10,6 +10,7 @@
 #include "Playable/AvaPlayableRemoteControl.h"
 #include "RCVirtualProperty.h"
 #include "Rundown/AvaRundownEditor.h"
+#include "Rundown/AvaRundownEditorUtils.h"
 #include "Rundown/AvaRundownManagedInstanceCache.h"
 #include "Rundown/AvaRundownPage.h"
 #include "Widgets/Views/SHeaderRow.h"
@@ -226,14 +227,13 @@ void SAvaRundownRCControllerPanel::Refresh(const TArray<int32>& InSelectedPageId
        	ManagedInstances.Reset();
 		return;
 	}
-		
-	ManagedInstances = FAvaRundownEditor::GetManagedInstancesForPage(Rundown, Page);
+
+	ManagedInstances = UE::AvaRundownEditor::Utils::GetManagedInstancesForPage(Rundown, Page);
 
 	int32 NumItems = 0;
 	for (const TSharedPtr<FAvaRundownManagedInstance>& ManagedInstance : ManagedInstances)
 	{
-		const URemoteControlPreset* Preset = ManagedInstance ? ManagedInstance->GetRemoteControlPreset() : nullptr;
-		if (Preset)
+		if (const URemoteControlPreset* Preset = ManagedInstance ? ManagedInstance->GetRemoteControlPreset() : nullptr)
 		{
 			NumItems += Preset->GetNumControllers();
 		}
@@ -265,49 +265,26 @@ TSharedRef<ITableRow> SAvaRundownRCControllerPanel::OnGenerateControllerRow(FAva
 	return InItem->CreateWidget(SharedThis(this), InOwnerTable);
 }
 
-void SAvaRundownRCControllerPanel::UpdateDefaultValuesAndRefresh()
+void SAvaRundownRCControllerPanel::UpdateDefaultValuesAndRefresh(const TArray<int32>& InSelectedPageIds)
 {
 	// Remark: The RC values might be already updated in SAvaRundownPageRemoteControlProps.
 	// But order of callback is not guaranteed. Code could reach here first, so
-	// it needs to update and refresh just in case. Calling UpdateRemoteControlValues
+	// it needs to update and refresh just in case. Calling UpdateDefaultRemoteControlValues
 	// multiple time (from different code paths) is harmless (fast if nothing changed).
 	
-	const TSharedPtr<FAvaRundownEditor> RundownEditor = RundownEditorWeak.Pin();
-	if (!RundownEditor)
-	{
-		return;
-	}
-	
-	UAvaRundown* Rundown = RundownEditor->GetRundown();
-	if (!Rundown)
-	{
-		return;
-	}
-	
-	const FAvaRundownPage& Page = GetActivePage(Rundown);
-	if (!Page.IsValidPage())
-	{
-		return;
-	}
-	
-	ManagedInstances = FAvaRundownEditor::GetManagedInstancesForPage(Rundown, ActivePageId);
-					
-	if (ManagedInstances.IsEmpty())
-	{
-		return;
-	}
-	
-	FAvaPlayableRemoteControlValues MergedDefaultRCValues;
-	FAvaRundownEditor::MergeDefaultRemoteControlValues(ManagedInstances, MergedDefaultRCValues);
-
-	// Using the rundown API for event propagation.
-	constexpr bool bUpdateDefaults = true;	
-	if (Rundown->UpdateRemoteControlValues(ActivePageId, MergedDefaultRCValues, bUpdateDefaults) != EAvaPlayableRemoteControlChanges::None)
+	if (const TSharedPtr<FAvaRundownEditor> RundownEditor = RundownEditorWeak.Pin())
 	{
-		RundownEditor->MarkAsModified();
+		if (UAvaRundown* Rundown = RundownEditor->GetRundown())
+		{
+			using namespace UE::AvaRundownEditor::Utils;
+			if (UpdateDefaultRemoteControlValues(Rundown, InSelectedPageIds) != EAvaPlayableRemoteControlChanges::None)
+			{
+				RundownEditor->MarkAsModified();
+			}
+		}
 	}
 	
-	Refresh({ActivePageId});
+	Refresh(InSelectedPageIds);
 }
 
 void SAvaRundownRCControllerPanel::OnRemoteControlControllerModified(URemoteControlPreset* InPreset, const TSet<FGuid>& InModifiedControllerIds)
diff --git a/Engine/Plugins/Experimental/Avalanche/Source/AvalancheMediaEditor/Private/Rundown/DetailsView/RemoteControl/Controllers/SAvaRundownRCControllerPanel.h b/Engine/Plugins/Experimental/Avalanche/Source/AvalancheMediaEditor/Private/Rundown/DetailsView/RemoteControl/Controllers/SAvaRundownRCControllerPanel.h
--- a/Engine/Plugins/Experimental/Avalanche/Source/AvalancheMediaEditor/Private/Rundown/DetailsView/RemoteControl/Controllers/SAvaRundownRCControllerPanel.h
+++ b/Engine/Plugins/Experimental/Avalanche/Source/AvalancheMediaEditor/Private/Rundown/DetailsView/RemoteControl/Controllers/SAvaRundownRCControllerPanel.h
@@ -59,12 +59,12 @@ private:
 	
 	void RefreshForManagedInstance(int32 InInstanceIndex, const FAvaRundownManagedInstance& InManagedInstance, const FAvaRundownPage& InPage);
 	
-	/** Update the current page's remote control values from the defaults then refresh the widget. */
-	void UpdateDefaultValuesAndRefresh();
+	/** Update the given pages remote control values from the defaults then refresh the widget. */
+	void UpdateDefaultValuesAndRefresh(const TArray<int32>& InSelectedPageIds);
 
-	void OnRemoteControlControllerAdded(URemoteControlPreset* InPreset, const FName NewControllerName, const FGuid& InControllerId) { UpdateDefaultValuesAndRefresh(); }
-	void OnRemoteControlControllerRemoved(URemoteControlPreset* InPreset, const FGuid& InControllerId) { UpdateDefaultValuesAndRefresh(); }
-	void OnRemoteControlControllerRenamed(URemoteControlPreset* InPreset, const FName InOldLabel, const FName InNewLabel) { UpdateDefaultValuesAndRefresh(); }
+	void OnRemoteControlControllerAdded(URemoteControlPreset* InPreset, const FName NewControllerName, const FGuid& InControllerId) { UpdateDefaultValuesAndRefresh({ActivePageId}); }
+	void OnRemoteControlControllerRemoved(URemoteControlPreset* InPreset, const FGuid& InControllerId) { UpdateDefaultValuesAndRefresh({ActivePageId}); }
+	void OnRemoteControlControllerRenamed(URemoteControlPreset* InPreset, const FName InOldLabel, const FName InNewLabel) { UpdateDefaultValuesAndRefresh({ActivePageId}); }
 	void OnRemoteControlControllerModified(URemoteControlPreset* InPreset, const TSet<FGuid>& InModifiedControllerIds);
 	void BindRemoteControlDelegates(URemoteControlPreset* InPreset);
 
diff --git a/Engine/Plugins/Experimental/Avalanche/Source/AvalancheMediaEditor/Private/Rundown/DetailsView/RemoteControl/Properties/SAvaRundownPageRemoteControlProps.cpp b/Engine/Plugins/Experimental/Avalanche/Source/AvalancheMediaEditor/Private/Rundown/DetailsView/RemoteControl/Properties/SAvaRundownPageRemoteControlProps.cpp
--- a/Engine/Plugins/Experimental/Avalanche/Source/AvalancheMediaEditor/Private/Rundown/DetailsView/RemoteControl/Properties/SAvaRundownPageRemoteControlProps.cpp
+++ b/Engine/Plugins/Experimental/Avalanche/Source/AvalancheMediaEditor/Private/Rundown/DetailsView/RemoteControl/Properties/SAvaRundownPageRemoteControlProps.cpp
@@ -8,14 +8,15 @@
 #include "RemoteControlEntity.h"
 #include "RemoteControlPreset.h"
 #include "Rundown/AvaRundownEditor.h"
+#include "Rundown/AvaRundownEditorUtils.h"
 #include "Rundown/AvaRundownManagedInstanceCache.h"
 #include "Rundown/AvaRundownPage.h"
 #include "SAvaRundownRCPropertyItemRow.h"
 #include "SlateOptMacros.h"
+#include "UObject/NameTypes.h"
 #include "Widgets/Views/ITableRow.h"
 #include "Widgets/Views/SListView.h"
 #include "Widgets/Views/STableViewBase.h"
-#include "UObject/NameTypes.h"
 
 #define LOCTEXT_NAMESPACE "SAvaRundownPageRemoteControlProps"
 
@@ -73,45 +74,19 @@ SAvaRundownPageRemoteControlProps::~SAvaRundownPageRemoteControlProps()
 
 void SAvaRundownPageRemoteControlProps::UpdateDefaultValuesAndRefresh(const TArray<int32>& InSelectedPageIds)
 {
-	// Remark:
-	// This is used in the "reimport page" work flow. We want to be able to
-	// reimport multiple pages at the same time.
-	
 	if (const TSharedPtr<FAvaRundownEditor> RundownEditor = RundownEditorWeak.Pin())
 	{
-		UAvaRundown* Rundown = RundownEditor->GetRundown();
-
-		if (IsValid(Rundown))
+		if (UAvaRundown* Rundown = RundownEditor->GetRundown())
 		{
-			EAvaPlayableRemoteControlChanges Changes = EAvaPlayableRemoteControlChanges::None;
-			for (const int32 PageId : InSelectedPageIds)
-			{
-				FAvaRundownPage& Page = RundownEditor->GetRundown()->GetPage(PageId);
-
-				if (Page.IsValidPage())
-				{
-					ManagedInstances = FAvaRundownEditor::GetManagedInstancesForPage(Rundown, Page);
-					
-					if (!ManagedInstances.IsEmpty())
-					{
-						FAvaPlayableRemoteControlValues MergedDefaultRCValues;
-						FAvaRundownEditor::MergeDefaultRemoteControlValues(ManagedInstances, MergedDefaultRCValues);
-
-						// Using the rundown API for event propagation.
-						constexpr bool bUpdateDefaults = true;
-						Changes |= Rundown->UpdateRemoteControlValues(ActivePageId, MergedDefaultRCValues, bUpdateDefaults);
-					}
-				}
-			}
-			
-			if (Changes != EAvaPlayableRemoteControlChanges::None)
+			using namespace UE::AvaRundownEditor::Utils;
+			if (UpdateDefaultRemoteControlValues(Rundown, InSelectedPageIds) != EAvaPlayableRemoteControlChanges::None)
 			{
 				RundownEditor->MarkAsModified();
 			}
-			
-			Refresh(InSelectedPageIds);
 		}
 	}
+
+	Refresh(InSelectedPageIds);
 }
 
 TSharedRef<ITableRow> SAvaRundownPageRemoteControlProps::OnGenerateControllerRow(FAvaRundownRCPropertyItemPtr InItem, const TSharedRef<STableViewBase>& InOwnerTable)
@@ -119,7 +94,6 @@ TSharedRef<ITableRow> SAvaRundownPageRemoteControlProps::OnGenerateControllerRow
 	return InItem->CreateWidget(SharedThis(this), InOwnerTable);
 }
 
-
 void SAvaRundownPageRemoteControlProps::RefreshTable(const TSet<FGuid>& InEntityIds)
 {
 	for (const FAvaRundownRCPropertyItemPtr& PropertyItem : PropertyItems)
@@ -157,21 +131,18 @@ void SAvaRundownPageRemoteControlProps::Refresh(const TArray<int32>& InSelectedP
 
 			if (FAvaRundownPage* ActivePage = GetActivePage())
 			{
-				ManagedInstances = FAvaRundownEditor::GetManagedInstancesForPage(Rundown, *ActivePage);
+				using namespace UE::AvaRundownEditor::Utils;
+				ManagedInstances = GetManagedInstancesForPage(Rundown, *ActivePage);
 
 				for (const TSharedPtr<FAvaRundownManagedInstance>& ManagedInstance : ManagedInstances)
 				{
-					URemoteControlPreset* Preset = ManagedInstance ? ManagedInstance->GetRemoteControlPreset() : nullptr;
-					if (Preset)
-					{
-						BindRemoteControlDelegates(Preset);	
-					}
+					BindRemoteControlDelegates(ManagedInstance ? ManagedInstance->GetRemoteControlPreset() : nullptr);	
 				}
 
 				if (!ManagedInstances.IsEmpty())
 				{
 					FAvaPlayableRemoteControlValues MergedDefaultRCValues;
-					FAvaRundownEditor::MergeDefaultRemoteControlValues(ManagedInstances, MergedDefaultRCValues);
+					MergeDefaultRemoteControlValues(ManagedInstances, MergedDefaultRCValues);
 					
 					// Prune any extra stale values. This happens if templates are changed.
 					if (ActivePage->PruneRemoteControlValues(MergedDefaultRCValues) != EAvaPlayableRemoteControlChanges::None)
@@ -190,6 +161,7 @@ void SAvaRundownPageRemoteControlProps::Refresh(const TArray<int32>& InSelectedP
 		};
 
 		TArray<FAvaPropertyDetails> NewItems;
+		int32 TotalNumExposedEntities = 0;
 
 		for (const TSharedPtr<FAvaRundownManagedInstance>& ManagedInstance : ManagedInstances)
 		{
@@ -200,7 +172,7 @@ void SAvaRundownPageRemoteControlProps::Refresh(const TArray<int32>& InSelectedP
 			}
 			
 			const TArray<TWeakPtr<FRemoteControlEntity>> ExposedEntities = RemoteControlPreset->GetExposedEntities<FRemoteControlEntity>();
-			NewItems.Reserve(ExposedEntities.Num());
+			NewItems.Reserve(TotalNumExposedEntities += ExposedEntities.Num());
 			
 			for (const TWeakPtr<FRemoteControlEntity>& EntityWeakPtr : ExposedEntities)
 			{
diff --git a/Engine/Plugins/Experimental/Avalanche/Source/AvalancheMediaEditor/Private/Rundown/DetailsView/SAvaRundownPageDetails.cpp b/Engine/Plugins/Experimental/Avalanche/Source/AvalancheMediaEditor/Private/Rundown/DetailsView/SAvaRundownPageDetails.cpp
--- a/Engine/Plugins/Experimental/Avalanche/Source/AvalancheMediaEditor/Private/Rundown/DetailsView/SAvaRundownPageDetails.cpp
+++ b/Engine/Plugins/Experimental/Avalanche/Source/AvalancheMediaEditor/Private/Rundown/DetailsView/SAvaRundownPageDetails.cpp
@@ -189,22 +189,12 @@ SAvaRundownPageDetails::~SAvaRundownPageDetails()
 
 void SAvaRundownPageDetails::OnPageEvent(const TArray<int32>& InSelectedPageIds, UE::AvaRundown::EPageEvent InPageEvent)
 {
-	if (InPageEvent == UE::AvaRundown::EPageEvent::SelectionChanged)
+	if (InPageEvent == UE::AvaRundown::EPageEvent::SelectionChanged || InPageEvent == UE::AvaRundown::EPageEvent::ReimportRequest)
 	{
 		OnPageSelectionChanged(InSelectedPageIds);
 		RemoteControlProps->Refresh(InSelectedPageIds);
 		RCControllerPanel->Refresh(InSelectedPageIds);
 	}
-	else if (InPageEvent == UE::AvaRundown::EPageEvent::ReimportRequest)
-	{
-		// Note: all widgets are refreshed in this event handler to ensure it is done
-		// after the cache is invalidated.
-		OnPageSelectionChanged(InSelectedPageIds);
-		// Only call SAvaRundownPageRemoteControlProps::UpdateDefaultValuesAndRefresh to have it update the values of all selected pages.
-		RemoteControlProps->UpdateDefaultValuesAndRefresh(InSelectedPageIds);
-		// Pages values already updated above, only need to refresh UI.
-		RCControllerPanel->Refresh(InSelectedPageIds);
-	}
 }
 
 void SAvaRundownPageDetails::OnPageSelectionChanged(const TArray<int32>& InSelectedPageIds)
diff --git a/Engine/Plugins/Experimental/Avalanche/Source/AvalancheMediaEditor/Private/Rundown/Pages/Slate/SAvaRundownPageList.cpp b/Engine/Plugins/Experimental/Avalanche/Source/AvalancheMediaEditor/Private/Rundown/Pages/Slate/SAvaRundownPageList.cpp
--- a/Engine/Plugins/Experimental/Avalanche/Source/AvalancheMediaEditor/Private/Rundown/Pages/Slate/SAvaRundownPageList.cpp
+++ b/Engine/Plugins/Experimental/Avalanche/Source/AvalancheMediaEditor/Private/Rundown/Pages/Slate/SAvaRundownPageList.cpp
@@ -641,12 +641,17 @@ void SAvaRundownPageList::ReimportSelectedPage() const
 	{
 		if (const TSharedPtr<FAvaRundownEditor> RundownEditor = GetRundownEditor())
 		{
-			UAvaRundown* Rundown = RundownEditor->GetRundown();
-			if (IsValid(Rundown))
+			if (UAvaRundown* Rundown = RundownEditor->GetRundown(); IsValid(Rundown))
 			{
 				// Enforce invalidation of the Motion Design Managed Instance Cache for the selected page(s).
 				Rundown->InvalidateManagedInstanceCacheForPages(SelectedPageIds);
 				Rundown->UpdateAssetForPages(SelectedPageIds, true);
+				
+				using namespace UE::AvaRundownEditor::Utils;
+				if (UpdateDefaultRemoteControlValues(Rundown, SelectedPageIds) != EAvaPlayableRemoteControlChanges::None)
+				{
+					RundownEditor->MarkAsModified();
+				}
 			}
 
 			RundownEditor->GetOnPageEvent().Broadcast(SelectedPageIds, UE::AvaRundown::EPageEvent::ReimportRequest);
diff --git a/Engine/Plugins/Experimental/Avalanche/Source/AvalancheMediaEditor/Public/Rundown/AvaRundownEditor.h b/Engine/Plugins/Experimental/Avalanche/Source/AvalancheMediaEditor/Public/Rundown/AvaRundownEditor.h
--- a/Engine/Plugins/Experimental/Avalanche/Source/AvalancheMediaEditor/Public/Rundown/AvaRundownEditor.h
+++ b/Engine/Plugins/Experimental/Avalanche/Source/AvalancheMediaEditor/Public/Rundown/AvaRundownEditor.h
@@ -47,9 +47,6 @@ public:
 	DECLARE_MULTICAST_DELEGATE_TwoParams(FOnPageEvent, const TArray<int32>&, UE::AvaRundown::EPageEvent);
 	FOnPageEvent& GetOnPageEvent() { return OnPageEvent; }
 
-	static TArray<TSharedPtr<FAvaRundownManagedInstance>> GetManagedInstancesForPage(const UAvaRundown* InRundown, const FAvaRundownPage& InPage);
-	static bool MergeDefaultRemoteControlValues(const TArray<TSharedPtr<FAvaRundownManagedInstance>>& InManagedInstances, FAvaPlayableRemoteControlValues& OutMergedValues);
-
 	bool IsKeyRelevant(const FKeyEvent& InKeyEvent) const;
 	bool HandleKeyDownEvent(const FKeyEvent& InKeyEvent);
 
