[Motion Design]
 - Set Motion Design tool box tab icon
 - Change Motion Design enter/exit mode button icons
 - Fix name of editor mode in window menu entry

diff --git a/Engine/Plugins/Experimental/Avalanche/Source/AvalancheEditorCore/AvalancheEditorCore.Build.cs b/Engine/Plugins/Experimental/Avalanche/Source/AvalancheEditorCore/AvalancheEditorCore.Build.cs
--- a/Engine/Plugins/Experimental/Avalanche/Source/AvalancheEditorCore/AvalancheEditorCore.Build.cs
+++ b/Engine/Plugins/Experimental/Avalanche/Source/AvalancheEditorCore/AvalancheEditorCore.Build.cs
@@ -28,6 +28,7 @@ public class AvalancheEditorCore : ModuleRules
 				"EditorFramework",
 				"Engine",
 				"LevelEditor",
+				"Projects",
 				"StatusBar",
 				"ToolMenus",
 				"TypedElementFramework",
diff --git a/Engine/Plugins/Experimental/Avalanche/Source/AvalancheEditorCore/Private/LevelEditor/AvaLevelEditorStyle.cpp b/Engine/Plugins/Experimental/Avalanche/Source/AvalancheEditorCore/Private/LevelEditor/AvaLevelEditorStyle.cpp
--- a/Engine/Plugins/Experimental/Avalanche/Source/AvalancheEditorCore/Private/LevelEditor/AvaLevelEditorStyle.cpp
+++ b/Engine/Plugins/Experimental/Avalanche/Source/AvalancheEditorCore/Private/LevelEditor/AvaLevelEditorStyle.cpp
@@ -2,6 +2,7 @@
 
 #include "AvaLevelEditorStyle.h"
 #include "Brushes/SlateImageBrush.h"
+#include "Interfaces/IPluginManager.h"
 #include "Misc/Paths.h"
 #include "Styling/SlateStyleMacros.h"
 #include "Styling/SlateStyleRegistry.h"
@@ -11,12 +12,15 @@ FAvaLevelEditorStyle::FAvaLevelEditorStyle()
 {
 	const FVector2f Icon16x16(16.f, 16.f);
 
-	SetContentRoot(FPaths::EngineContentDir() / TEXT("Editor/Slate"));
-	SetCoreContentRoot(FPaths::EngineContentDir() / TEXT("Slate"));
+	const TSharedPtr<IPlugin> Plugin = IPluginManager::Get().FindPlugin(UE_PLUGIN_NAME);
+	check(Plugin.IsValid());
 
-	Set("AvaLevelEditor.CreateScene"    , new CORE_IMAGE_BRUSH_SVG("Starship/Common/plus", Icon16x16));
-	Set("AvaLevelEditor.ActivateScene"  , new CORE_IMAGE_BRUSH_SVG("Starship/Common/play", Icon16x16));
-	Set("AvaLevelEditor.DeactivateScene", new CORE_IMAGE_BRUSH_SVG("Starship/Common/stop", Icon16x16));
+	SetContentRoot(Plugin->GetBaseDir() / TEXT("Resources"));
+	SetCoreContentRoot(FPaths::EngineContentDir() / TEXT("Slate"));
+	
+	Set("AvaLevelEditor.CreateScene", new IMAGE_BRUSH_SVG("Icons/EditorIcons/LevelEditor_CreateScene_40", Icon16x16));
+	Set("AvaLevelEditor.ActivateScene", new IMAGE_BRUSH_SVG("Icons/EditorIcons/LevelEditor_ActivateScene_40", Icon16x16));
+	Set("AvaLevelEditor.DeactivateScene", new IMAGE_BRUSH_SVG("Icons/EditorIcons/LevelEditor_DeactivateScene_40", Icon16x16));
 
 	FSlateStyleRegistry::RegisterSlateStyle(*this);
 }
diff --git a/Engine/Plugins/Experimental/Avalanche/Source/AvalancheInteractiveTools/Private/EdMode/AvaInteractiveToolsEdMode.cpp b/Engine/Plugins/Experimental/Avalanche/Source/AvalancheInteractiveTools/Private/EdMode/AvaInteractiveToolsEdMode.cpp
--- a/Engine/Plugins/Experimental/Avalanche/Source/AvalancheInteractiveTools/Private/EdMode/AvaInteractiveToolsEdMode.cpp
+++ b/Engine/Plugins/Experimental/Avalanche/Source/AvalancheInteractiveTools/Private/EdMode/AvaInteractiveToolsEdMode.cpp
@@ -17,7 +17,7 @@ UAvaInteractiveToolsEdMode::UAvaInteractiveToolsEdMode()
 {
 	Info = FEditorModeInfo(IAvalancheInteractiveToolsModule::EM_AvaInteractiveToolsEdModeId,
 		LOCTEXT("AvaInteractiveToolsEdModeName", "Motion Design"),
-		FSlateIcon(),
+		FSlateIcon(FAppStyle::GetAppStyleSetName(), TEXT("LevelEditor.Tabs.EditorModes")),
 		true);
 
 	LastActiveTool = "";
diff --git a/Engine/Plugins/Experimental/Avalanche/Source/AvalancheInteractiveTools/Private/EdMode/AvaInteractiveToolsEdModeToolkit.cpp b/Engine/Plugins/Experimental/Avalanche/Source/AvalancheInteractiveTools/Private/EdMode/AvaInteractiveToolsEdModeToolkit.cpp
--- a/Engine/Plugins/Experimental/Avalanche/Source/AvalancheInteractiveTools/Private/EdMode/AvaInteractiveToolsEdModeToolkit.cpp
+++ b/Engine/Plugins/Experimental/Avalanche/Source/AvalancheInteractiveTools/Private/EdMode/AvaInteractiveToolsEdModeToolkit.cpp
@@ -4,6 +4,7 @@
 #include "AvaInteractiveToolsEdMode.h"
 #include "AvalancheInteractiveToolsModule.h"
 #include "IAvalancheInteractiveToolsModule.h"
+#include "Toolkits/AssetEditorModeUILayer.h"
 #include "Tools/UEdMode.h"
 #include "Widgets/Layout/SBorder.h"
 #include "Widgets/Layout/SBox.h"
@@ -99,6 +100,25 @@ void FAvaInteractiveToolsEdModeToolkit::InvokeUI()
 	InlineContentHolder->SetContent(GetInlineContent().ToSharedRef());
 }
 
+void FAvaInteractiveToolsEdModeToolkit::RequestModeUITabs()
+{
+	if (const TSharedPtr<FAssetEditorModeUILayer> ModeUILayerPtr = ModeUILayer.Pin())
+	{
+		PrimaryTabInfo.OnSpawnTab = FOnSpawnTab::CreateSP(SharedThis(this), &FAvaInteractiveToolsEdModeToolkit::CreatePrimaryModePanel);
+		PrimaryTabInfo.TabLabel = LOCTEXT("MotionDesignToolboxTab", "Motion Design");
+		PrimaryTabInfo.TabTooltip = LOCTEXT("MotionDesignToolboxTabTooltipText", "Opens the Motion Design tab.");
+		ModeUILayerPtr->SetModePanelInfo(UAssetEditorUISubsystem::TopLeftTabID, PrimaryTabInfo);
+
+		if (!HasIntegratedToolPalettes() && !HasToolkitBuilder())
+		{
+			ToolbarInfo.OnSpawnTab = FOnSpawnTab::CreateSP(SharedThis(this), &FAvaInteractiveToolsEdModeToolkit::MakeModeToolbarTab);
+			ToolbarInfo.TabLabel = LOCTEXT("MotionDesignToolbarTab", "Motion Design Toolbar");
+			ToolbarInfo.TabTooltip = LOCTEXT("MotionDesignToolbarTabTooltipText", "Opens the toolbar for the Motion Design toolbox.");
+			ModeUILayerPtr->SetModePanelInfo(UAssetEditorUISubsystem::VerticalToolbarID, ToolbarInfo);
+		}
+	}
+}
+
 void FAvaInteractiveToolsEdModeToolkit::RegisterPalettes()
 {
 	UAvaInteractiveToolsEdMode* AvaInteractiveToolsMode = Cast<UAvaInteractiveToolsEdMode>(GetScriptableEditorMode().Get());
diff --git a/Engine/Plugins/Experimental/Avalanche/Source/AvalancheInteractiveTools/Private/EdMode/AvaInteractiveToolsEdModeToolkit.h b/Engine/Plugins/Experimental/Avalanche/Source/AvalancheInteractiveTools/Private/EdMode/AvaInteractiveToolsEdModeToolkit.h
--- a/Engine/Plugins/Experimental/Avalanche/Source/AvalancheInteractiveTools/Private/EdMode/AvaInteractiveToolsEdModeToolkit.h
+++ b/Engine/Plugins/Experimental/Avalanche/Source/AvalancheInteractiveTools/Private/EdMode/AvaInteractiveToolsEdModeToolkit.h
@@ -7,7 +7,6 @@
 class FAvaInteractiveToolsEdModeToolkit : public FModeToolkit
 {
 public:
-
 	FAvaInteractiveToolsEdModeToolkit();
 	virtual ~FAvaInteractiveToolsEdModeToolkit() override;
 
@@ -29,6 +28,11 @@ public:
 	virtual void InvokeUI() override;
 	//~ End FModeToolkit
 
+protected:
+	//~ Begin FModeToolkit
+	virtual void RequestModeUITabs() override;
+	//~ End FModeToolkit
+
 private:
 	TSharedPtr<SWidget> ToolkitWidget;
 
