[Motion Design] Viewport pilot actor default keybinds changed

diff --git a/Engine/Plugins/Experimental/Avalanche/Source/AvalancheLevelViewport/Private/AvaLevelViewportCommands.cpp b/Engine/Plugins/Experimental/Avalanche/Source/AvalancheLevelViewport/Private/AvaLevelViewportCommands.cpp
--- a/Engine/Plugins/Experimental/Avalanche/Source/AvalancheLevelViewport/Private/AvaLevelViewportCommands.cpp
+++ b/Engine/Plugins/Experimental/Avalanche/Source/AvalancheLevelViewport/Private/AvaLevelViewportCommands.cpp
@@ -167,13 +167,13 @@ void FAvaLevelViewportCommands::RegisterCameraCommands()
 		, "Undo Viewport Camera Transform"
 		, "Undoes the previous viewport camera move, if any."
 		, EUserInterfaceActionType::Button
-		, FInputChord(EKeys::Z, EModifierKey::Shift))
+		, FInputChord(EKeys::LeftBracket))
 
 	UI_COMMAND(CameraTransformRedo
 		, "Redo Viewport Camera Transform"
 		, "Redoes the previous undone viewport camera move, if any."
 		, EUserInterfaceActionType::Button
-		, FInputChord(EKeys::Y, EModifierKey::Shift))
+		, FInputChord(EKeys::RightBracket))
 }
 
 void FAvaLevelViewportCommands::RegisterGridCommands()
