Motion Design
- Fixed bug with visualizers not receiving tracking events when in a camera piloting transaction.

diff --git a/Engine/Plugins/Experimental/Avalanche/Source/AvalancheComponentVisualizers/Private/AvaComponentVisualizersEdMode.cpp b/Engine/Plugins/Experimental/Avalanche/Source/AvalancheComponentVisualizers/Private/AvaComponentVisualizersEdMode.cpp
--- a/Engine/Plugins/Experimental/Avalanche/Source/AvalancheComponentVisualizers/Private/AvaComponentVisualizersEdMode.cpp
+++ b/Engine/Plugins/Experimental/Avalanche/Source/AvalancheComponentVisualizers/Private/AvaComponentVisualizersEdMode.cpp
@@ -94,15 +94,21 @@ void UAvaComponentVisualizersEdMode::ModeTick(float DeltaTime)
 {
 	Super::ModeTick(DeltaTime);
 
-	UE::AvalancheComponentVisualizers::Private::FAvaViewportClientAndVisualizer VCnV(this);
+	UE::AvalancheComponentVisualizers::Private::FAvaViewportClientAndVisualizer VisualizerAndClient(this);
 
-	if (VCnV.ActiveVisualizer.IsValid())
+	if (VisualizerAndClient.ActiveVisualizer.IsValid())
 	{
+		if (VisualizerAndClient.EditorViewportClient && !VisualizerAndClient.EditorViewportClient->IsTracking()
+			&& VisualizerAndClient.ActiveVisualizer->IsTracking())
+		{
+			VisualizerAndClient.ActiveVisualizer->TrackingStopped(VisualizerAndClient.EditorViewportClient, /* bDidMove */ true);
+		}
+
 		FEditorModeTools* ModeManager = GetModeManager();
 		const UE::Widget::EWidgetMode CurrentModeStart = ModeManager->GetWidgetMode();
 		UE::Widget::EWidgetMode CurrentModeModified = ModeManager->GetWidgetMode();
 
-		if (VCnV.ActiveVisualizer->GetWidgetMode(VCnV.EditorViewportClient, CurrentModeModified))
+		if (VisualizerAndClient.ActiveVisualizer->GetWidgetMode(VisualizerAndClient.EditorViewportClient, CurrentModeModified))
 		{
 			if (CurrentModeStart == CurrentModeModified)
 			{
@@ -119,7 +125,7 @@ void UAvaComponentVisualizersEdMode::ModeTick(float DeltaTime)
 
 			CurrentModeModified = WidgetMode;
 
-			if (VCnV.ActiveVisualizer->GetWidgetMode(VCnV.EditorViewportClient, CurrentModeModified))
+			if (VisualizerAndClient.ActiveVisualizer->GetWidgetMode(VisualizerAndClient.EditorViewportClient, CurrentModeModified))
 			{
 				if (WidgetMode == CurrentModeModified)
 				{
diff --git a/Engine/Plugins/Experimental/Avalanche/Source/AvalancheComponentVisualizers/Private/AvaVisBase.cpp b/Engine/Plugins/Experimental/Avalanche/Source/AvalancheComponentVisualizers/Private/AvaVisBase.cpp
--- a/Engine/Plugins/Experimental/Avalanche/Source/AvalancheComponentVisualizers/Private/AvaVisBase.cpp
+++ b/Engine/Plugins/Experimental/Avalanche/Source/AvalancheComponentVisualizers/Private/AvaVisBase.cpp
@@ -58,6 +58,11 @@ bool FAvaVisualizerBase::HandleInputDelta(FEditorViewportClient* InViewportClien
 		return false;
 	}
 
+	if (!bTracking)
+	{
+		TrackingStartedInternal(InViewportClient);
+	}
+
 	FVector RootComponentScale = RootComponent->GetComponentScale();
 
 	if (RootComponentScale != FVector::OneVector)
@@ -165,34 +170,12 @@ void FAvaVisualizerBase::EndEditing()
 
 void FAvaVisualizerBase::TrackingStarted(FEditorViewportClient* InViewportClient)
 {
-	if (!GetEditedComponent())
-	{
-		return;
-	}
-
-	StoreInitialValues();
-	StartTransaction();
-	bTracking = true;
-
-	if (TSharedPtr<IAvaViewportClient> AvaViewportClient = FAvaViewportUtils::GetAsAvaViewportClient(InViewportClient))
-	{
-		SnapOperation = AvaViewportClient->StartSnapOperation();
-
-		if (SnapOperation.IsValid())
-		{
-			SnapOperation->GenerateComponentSnapPoints(GetEditedComponent());
-			GenerateContextSensitiveSnapPoints();
-			SnapOperation->FinaliseSnapPoints();
-		}
-	}
+	TrackingStartedInternal(InViewportClient);
 }
 
 void FAvaVisualizerBase::TrackingStopped(FEditorViewportClient* InViewportClient, bool bInDidMove)
 {
-	EndTransaction();
-	AddSnapDataBinding();
-	bTracking = false;
-	SnapOperation.Reset();
+	TrackingStoppedInternal(InViewportClient);
 }
 
 void FAvaVisualizerBase::DrawVisualization(const UActorComponent* InComponent, const FSceneView* InView,
@@ -501,6 +484,43 @@ bool FAvaVisualizerBase::ShouldDrawExtraHandles(const UActorComponent* InCompone
 	return IsMouseOverComponent(InComponent, InView);
 }
 
+void FAvaVisualizerBase::TrackingStartedInternal(FEditorViewportClient* InViewportClient)
+{
+	if (!GetEditedComponent())
+	{
+		return;
+	}
+
+	if (bTracking)
+	{
+		return;
+	}
+
+	StoreInitialValues();
+	StartTransaction();
+	bTracking = true;
+
+	if (TSharedPtr<IAvaViewportClient> AvaViewportClient = FAvaViewportUtils::GetAsAvaViewportClient(InViewportClient))
+	{
+		SnapOperation = AvaViewportClient->StartSnapOperation();
+
+		if (SnapOperation.IsValid())
+		{
+			SnapOperation->GenerateComponentSnapPoints(GetEditedComponent());
+			GenerateContextSensitiveSnapPoints();
+			SnapOperation->FinaliseSnapPoints();
+		}
+	}
+}
+
+void FAvaVisualizerBase::TrackingStoppedInternal(FEditorViewportClient* InViewportClient)
+{
+	EndTransaction();
+	AddSnapDataBinding();
+	bTracking = false;
+	SnapOperation.Reset();
+}
+
 void FAvaVisualizerBase::NotifyPropertyModified(UObject* InObject, FProperty* InProperty, 
 	EPropertyChangeType::Type InPropertyChangeType, FProperty* InMemberProperty)
 {
diff --git a/Engine/Plugins/Experimental/Avalanche/Source/AvalancheComponentVisualizers/Public/AvaVisBase.h b/Engine/Plugins/Experimental/Avalanche/Source/AvalancheComponentVisualizers/Public/AvaVisBase.h
--- a/Engine/Plugins/Experimental/Avalanche/Source/AvalancheComponentVisualizers/Public/AvaVisBase.h
+++ b/Engine/Plugins/Experimental/Avalanche/Source/AvalancheComponentVisualizers/Public/AvaVisBase.h
@@ -186,4 +186,7 @@ protected:
 	virtual void AddSnapDataBinding() {}
 
 	virtual void GenerateContextSensitiveSnapPoints() {}
+
+	virtual void TrackingStartedInternal(FEditorViewportClient* InViewportClient);
+	virtual void TrackingStoppedInternal(FEditorViewportClient* InViewportClient);
 };
