Motion Design:
- Added PostPlaced override fucntion in Shape and Text factories to record analytics event

#changelist validated
#virtualized

diff --git a/Engine/Plugins/Experimental/Avalanche/Source/AvalancheEffectorsEditor/AvalancheEffectorsEditor.Build.cs b/Engine/Plugins/Experimental/Avalanche/Source/AvalancheEffectorsEditor/AvalancheEffectorsEditor.Build.cs
--- a/Engine/Plugins/Experimental/Avalanche/Source/AvalancheEffectorsEditor/AvalancheEffectorsEditor.Build.cs
+++ b/Engine/Plugins/Experimental/Avalanche/Source/AvalancheEffectorsEditor/AvalancheEffectorsEditor.Build.cs
@@ -18,6 +18,7 @@ public class AvalancheEffectorsEditor : ModuleRules
 				"AvalancheShapes",
 				"AvalancheShapesEditor",
 				"ClonerEffector",
+				"ClonerEffectorEditor",
 				"ComponentVisualizers",
 				"Core",
 				"CoreUObject",
diff --git a/Engine/Plugins/Experimental/Avalanche/Source/AvalancheEffectorsEditor/Private/Cloner/AvaClonerActorTool.cpp b/Engine/Plugins/Experimental/Avalanche/Source/AvalancheEffectorsEditor/Private/Cloner/AvaClonerActorTool.cpp
--- a/Engine/Plugins/Experimental/Avalanche/Source/AvalancheEffectorsEditor/Private/Cloner/AvaClonerActorTool.cpp
+++ b/Engine/Plugins/Experimental/Avalanche/Source/AvalancheEffectorsEditor/Private/Cloner/AvaClonerActorTool.cpp
@@ -5,6 +5,7 @@
 #include "AvaInteractiveToolsSettings.h"
 #include "Builders/AvaInteractiveToolsToolBuilder.h"
 #include "Cloner/CEClonerActor.h"
+#include "Cloner/CEClonerActorFactory.h"
 
 UAvaClonerActorTool::UAvaClonerActorTool()
 {
@@ -32,6 +33,7 @@ FAvaInteractiveToolsToolParameters UAvaClonerActorTool::GetToolParameters() cons
 			{
 				return UAvaInteractiveToolsToolBuilder::CreateToolBuilder<UAvaClonerActorTool>(InEdMode);
 			}),
-		ActorClass
+		ActorClass,
+		CreateActorFactory<UCEClonerActorFactory>()
 	};
 }
diff --git a/Engine/Plugins/Experimental/Avalanche/Source/AvalancheEffectorsEditor/Private/Effector/AvaEffectorActorTool.cpp b/Engine/Plugins/Experimental/Avalanche/Source/AvalancheEffectorsEditor/Private/Effector/AvaEffectorActorTool.cpp
--- a/Engine/Plugins/Experimental/Avalanche/Source/AvalancheEffectorsEditor/Private/Effector/AvaEffectorActorTool.cpp
+++ b/Engine/Plugins/Experimental/Avalanche/Source/AvalancheEffectorsEditor/Private/Effector/AvaEffectorActorTool.cpp
@@ -5,6 +5,7 @@
 #include "AvaInteractiveToolsSettings.h"
 #include "Builders/AvaInteractiveToolsToolBuilder.h"
 #include "Effector/CEEffectorActor.h"
+#include "Effector/CEEffectorActorFactory.h"
 
 UAvaEffectorActorTool::UAvaEffectorActorTool()
 {
@@ -32,6 +33,7 @@ FAvaInteractiveToolsToolParameters UAvaEffectorActorTool::GetToolParameters() co
 			{
 				return UAvaInteractiveToolsToolBuilder::CreateToolBuilder<UAvaEffectorActorTool>(InEdMode);
 			}),
-		ActorClass
+		ActorClass,
+		CreateActorFactory<UCEEffectorActorFactory>()
 	};
 }
diff --git a/Engine/Plugins/Experimental/Avalanche/Source/AvalancheInteractiveTools/Public/Tools/AvaInteractiveToolsToolBase.h b/Engine/Plugins/Experimental/Avalanche/Source/AvalancheInteractiveTools/Public/Tools/AvaInteractiveToolsToolBase.h
--- a/Engine/Plugins/Experimental/Avalanche/Source/AvalancheInteractiveTools/Public/Tools/AvaInteractiveToolsToolBase.h
+++ b/Engine/Plugins/Experimental/Avalanche/Source/AvalancheInteractiveTools/Public/Tools/AvaInteractiveToolsToolBase.h
@@ -107,6 +107,14 @@ public:
 	FViewport* GetViewport(EAvaViewportStatus InViewportStatus) const;
 
 protected:
+	/** Helper to create an Actor Factory of a given class */
+	template<typename InActorFactoryClass
+		UE_REQUIRES(std::is_base_of_v<UActorFactory, InActorFactoryClass>)>
+	static InActorFactoryClass* CreateActorFactory()
+	{
+		return NewObject<InActorFactoryClass>(GetTransientPackage(), InActorFactoryClass::StaticClass(), NAME_None, RF_Standalone);
+	}
+
 	UPROPERTY()
 	UAvaSingleClickAndDragBehavior* LeftClickBehavior = nullptr;
 
diff --git a/Engine/Plugins/Experimental/Avalanche/Source/AvalancheSVGEditor/Private/Tool/AvaSVGActorTool.cpp b/Engine/Plugins/Experimental/Avalanche/Source/AvalancheSVGEditor/Private/Tool/AvaSVGActorTool.cpp
--- a/Engine/Plugins/Experimental/Avalanche/Source/AvalancheSVGEditor/Private/Tool/AvaSVGActorTool.cpp
+++ b/Engine/Plugins/Experimental/Avalanche/Source/AvalancheSVGEditor/Private/Tool/AvaSVGActorTool.cpp
@@ -3,6 +3,7 @@
 #include "Tool/AvaSVGActorTool.h"
 #include "AvaInteractiveToolsSettings.h"
 #include "Builders/AvaInteractiveToolsToolBuilder.h"
+#include "Factories/SVGActorFactory.h"
 #include "SVGActor.h"
 #include "SVGImporterEditorCommands.h"
 
@@ -32,6 +33,7 @@ FAvaInteractiveToolsToolParameters UAvaSVGActorTool::GetToolParameters() const
 			{
 				return UAvaInteractiveToolsToolBuilder::CreateToolBuilder<UAvaSVGActorTool>(InEdMode);
 			}),
-		ActorClass
+		ActorClass,
+		CreateActorFactory<USVGActorFactory>()
 	};
 }
diff --git a/Engine/Plugins/Experimental/Avalanche/Source/AvalancheShapesEditor/Private/AvaShapeFactory.cpp b/Engine/Plugins/Experimental/Avalanche/Source/AvalancheShapesEditor/Private/AvaShapeFactory.cpp
--- a/Engine/Plugins/Experimental/Avalanche/Source/AvalancheShapesEditor/Private/AvaShapeFactory.cpp
+++ b/Engine/Plugins/Experimental/Avalanche/Source/AvalancheShapesEditor/Private/AvaShapeFactory.cpp
@@ -2,11 +2,13 @@
 
 #include "AvaShapeFactory.h"
 #include "AvaShapeActor.h"
-#include "DynamicMeshes/AvaShapeDynMeshBase.h"
 #include "DynamicMeshes/AvaShape2DDynMeshBase.h"
 #include "DynamicMeshes/AvaShape3DDynMeshBase.h"
+#include "DynamicMeshes/AvaShapeDynMeshBase.h"
 #include "Engine/Level.h"
 #include "Engine/World.h"
+#include "EngineAnalytics.h"
+#include "Subsystems/PlacementSubsystem.h"
 #include "Tools/AvaShapesEditorShapeToolBase.h"
 
 UAvaShapeFactory::UAvaShapeFactory()
@@ -81,3 +83,18 @@ AActor* UAvaShapeFactory::SpawnActor(UObject* InAsset, ULevel* InLevel, const FT
 
 	return ShapeActor;
 }
+
+void UAvaShapeFactory::PostPlaceAsset(TArrayView<const FTypedElementHandle> InHandle, const FAssetPlacementInfo& InPlacementInfo, const FPlacementOptions& InPlacementOptions)
+{
+	Super::PostPlaceAsset(InHandle, InPlacementInfo, InPlacementOptions);
+
+	if (!InPlacementOptions.bIsCreatingPreviewElements && FEngineAnalytics::IsAvailable())
+	{
+		TArray<FAnalyticsEventAttribute> Attributes;
+		if (MeshClass)
+		{
+			Attributes.Emplace(TEXT("MeshClass"), MeshClass->GetName());	
+		}
+		FEngineAnalytics::GetProvider().RecordEvent(TEXT("Editor.Usage.MotionDesign.PlaceShape"), Attributes);
+	}
+}
diff --git a/Engine/Plugins/Experimental/Avalanche/Source/AvalancheShapesEditor/Public/AvaShapeFactory.h b/Engine/Plugins/Experimental/Avalanche/Source/AvalancheShapesEditor/Public/AvaShapeFactory.h
--- a/Engine/Plugins/Experimental/Avalanche/Source/AvalancheShapesEditor/Public/AvaShapeFactory.h
+++ b/Engine/Plugins/Experimental/Avalanche/Source/AvalancheShapesEditor/Public/AvaShapeFactory.h
@@ -23,9 +23,8 @@ protected:
 
 	//~ Begin UActorFactory
 	virtual bool CanCreateActorFrom(const FAssetData& AssetData, FText& OutErrorMsg) override;
-
 	virtual AActor* GetDefaultActor(const FAssetData& AssetData) override;
-
 	virtual AActor* SpawnActor(UObject* InAsset, ULevel* InLevel, const FTransform& InTransform, const FActorSpawnParameters& InSpawnParams) override;
+	virtual void PostPlaceAsset(TArrayView<const FTypedElementHandle> InHandle, const FAssetPlacementInfo& InPlacementInfo, const FPlacementOptions& InPlacementOptions) override;
 	//~ End UActorFactory
 };
diff --git a/Engine/Plugins/Experimental/Avalanche/Source/AvalancheShapesEditor/Public/Tools/AvaShapesEditorShapeToolBase.h b/Engine/Plugins/Experimental/Avalanche/Source/AvalancheShapesEditor/Public/Tools/AvaShapesEditorShapeToolBase.h
--- a/Engine/Plugins/Experimental/Avalanche/Source/AvalancheShapesEditor/Public/Tools/AvaShapesEditorShapeToolBase.h
+++ b/Engine/Plugins/Experimental/Avalanche/Source/AvalancheShapesEditor/Public/Tools/AvaShapesEditorShapeToolBase.h
@@ -32,7 +32,7 @@ protected:
 		UE_REQUIRES(std::derived_from<InMeshClass, UAvaShapeDynamicMeshBase>)>
 	static UAvaShapeFactory* CreateFactory()
 	{
-		UAvaShapeFactory* Factory = NewObject<UAvaShapeFactory>(GetTransientPackage(), UAvaShapeFactory::StaticClass(), NAME_None, RF_Standalone);
+		UAvaShapeFactory* Factory = CreateActorFactory<UAvaShapeFactory>();
 		Factory->SetMeshClass(InMeshClass::StaticClass());
 		return Factory;
 	}
diff --git a/Engine/Plugins/Experimental/Avalanche/Source/AvalancheTextEditor/Private/AvaTextActorFactory.cpp b/Engine/Plugins/Experimental/Avalanche/Source/AvalancheTextEditor/Private/AvaTextActorFactory.cpp
new file mode 100644
--- /dev/null
+++ b/Engine/Plugins/Experimental/Avalanche/Source/AvalancheTextEditor/Private/AvaTextActorFactory.cpp
@@ -0,0 +1,20 @@
+// Copyright Epic Games, Inc. All Rights Reserved.
+
+#include "AvaTextActorFactory.h"
+#include "AvaTextActor.h"
+#include "EngineAnalytics.h"
+#include "Subsystems/PlacementSubsystem.h"
+
+UAvaTextActorFactory::UAvaTextActorFactory()
+{
+	NewActorClass = AAvaTextActor::StaticClass();
+}
+
+void UAvaTextActorFactory::PostPlaceAsset(TArrayView<const FTypedElementHandle> InHandle, const FAssetPlacementInfo& InPlacementInfo, const FPlacementOptions& InPlacementOptions)
+{
+	Super::PostPlaceAsset(InHandle, InPlacementInfo, InPlacementOptions);
+	if (!InPlacementOptions.bIsCreatingPreviewElements && FEngineAnalytics::IsAvailable())
+	{
+		FEngineAnalytics::GetProvider().RecordEvent(TEXT("Editor.Usage.MotionDesign.PlaceText"));
+	}
+}
diff --git a/Engine/Plugins/Experimental/Avalanche/Source/AvalancheTextEditor/Private/AvaTextActorFactory.h b/Engine/Plugins/Experimental/Avalanche/Source/AvalancheTextEditor/Private/AvaTextActorFactory.h
new file mode 100644
--- /dev/null
+++ b/Engine/Plugins/Experimental/Avalanche/Source/AvalancheTextEditor/Private/AvaTextActorFactory.h
@@ -0,0 +1,19 @@
+// Copyright Epic Games, Inc. All Rights Reserved.
+
+#pragma once
+
+#include "ActorFactories/ActorFactory.h"
+#include "AvaTextActorFactory.generated.h"
+
+UCLASS()
+class UAvaTextActorFactory : public UActorFactory
+{
+	GENERATED_BODY()
+
+public:
+	UAvaTextActorFactory();
+
+	//~ Begin UActorFactory
+	virtual void PostPlaceAsset(TArrayView<const FTypedElementHandle> InHandle, const FAssetPlacementInfo& InPlacementInfo, const FPlacementOptions& InPlacementOptions) override;
+	//~ End UActorFactory
+};
diff --git a/Engine/Plugins/Experimental/Avalanche/Source/AvalancheTextEditor/Private/Tool/AvaTextActorTool.cpp b/Engine/Plugins/Experimental/Avalanche/Source/AvalancheTextEditor/Private/Tool/AvaTextActorTool.cpp
--- a/Engine/Plugins/Experimental/Avalanche/Source/AvalancheTextEditor/Private/Tool/AvaTextActorTool.cpp
+++ b/Engine/Plugins/Experimental/Avalanche/Source/AvalancheTextEditor/Private/Tool/AvaTextActorTool.cpp
@@ -3,6 +3,7 @@
 #include "Tool/AvaTextActorTool.h"
 #include "AvaInteractiveToolsSettings.h"
 #include "AvaTextActor.h"
+#include "AvaTextActorFactory.h"
 #include "AvaTextEditorCommands.h"
 #include "Builders/AvaInteractiveToolsToolBuilder.h"
 #include "IAvalancheComponentVisualizersModule.h"
@@ -33,6 +34,7 @@ FAvaInteractiveToolsToolParameters UAvaTextActorTool::GetToolParameters() const
 			{
 				return UAvaInteractiveToolsToolBuilder::CreateToolBuilder<UAvaTextActorTool>(InEdMode);
 			}),
-		ActorClass
+		ActorClass,
+		CreateActorFactory<UAvaTextActorFactory>()
 	};
 }
diff --git a/Engine/Plugins/Experimental/ClonerEffector/Source/ClonerEffectorEditor/ClonerEffectorEditor.Build.cs b/Engine/Plugins/Experimental/ClonerEffector/Source/ClonerEffectorEditor/ClonerEffectorEditor.Build.cs
--- a/Engine/Plugins/Experimental/ClonerEffector/Source/ClonerEffectorEditor/ClonerEffectorEditor.Build.cs
+++ b/Engine/Plugins/Experimental/ClonerEffector/Source/ClonerEffectorEditor/ClonerEffectorEditor.Build.cs
@@ -15,7 +15,8 @@ public class ClonerEffectorEditor : ModuleRules
 		        "CoreUObject",
 		        "Engine",
 		        "Slate",
-		        "SlateCore"
+		        "SlateCore",
+		        "UnrealEd",
 	        }
         );
 
@@ -25,7 +26,6 @@ public class ClonerEffectorEditor : ModuleRules
 		        "ClonerEffector",
 		        "InputCore",
 		        "Projects",
-		        "UnrealEd" 
 	        }
         );
 
diff --git a/Engine/Plugins/Experimental/ClonerEffector/Source/ClonerEffectorEditor/Private/Cloner/CEClonerActorFactory.cpp b/Engine/Plugins/Experimental/ClonerEffector/Source/ClonerEffectorEditor/Private/Cloner/CEClonerActorFactory.cpp
new file mode 100644
--- /dev/null
+++ b/Engine/Plugins/Experimental/ClonerEffector/Source/ClonerEffectorEditor/Private/Cloner/CEClonerActorFactory.cpp
@@ -0,0 +1,21 @@
+// Copyright Epic Games, Inc. All Rights Reserved.
+
+#include "Cloner/CEClonerActorFactory.h"
+#include "Cloner/CEClonerActor.h"
+#include "EngineAnalytics.h"
+#include "Subsystems/PlacementSubsystem.h"
+
+UCEClonerActorFactory::UCEClonerActorFactory()
+{
+	NewActorClass = ACEClonerActor::StaticClass();
+}
+
+void UCEClonerActorFactory::PostPlaceAsset(TArrayView<const FTypedElementHandle> InHandle, const FAssetPlacementInfo& InPlacementInfo, const FPlacementOptions& InPlacementOptions)
+{
+	Super::PostPlaceAsset(InHandle, InPlacementInfo, InPlacementOptions);
+
+	if (!InPlacementOptions.bIsCreatingPreviewElements && FEngineAnalytics::IsAvailable())
+	{
+		FEngineAnalytics::GetProvider().RecordEvent(TEXT("Editor.Usage.ClonerEffector.PlaceCloner"));
+	}
+}
diff --git a/Engine/Plugins/Experimental/ClonerEffector/Source/ClonerEffectorEditor/Private/Effector/CEEffectorActorFactory.cpp b/Engine/Plugins/Experimental/ClonerEffector/Source/ClonerEffectorEditor/Private/Effector/CEEffectorActorFactory.cpp
new file mode 100644
--- /dev/null
+++ b/Engine/Plugins/Experimental/ClonerEffector/Source/ClonerEffectorEditor/Private/Effector/CEEffectorActorFactory.cpp
@@ -0,0 +1,21 @@
+// Copyright Epic Games, Inc. All Rights Reserved.
+
+#include "Effector/CEEffectorActorFactory.h"
+#include "Effector/CEEffectorActor.h"
+#include "EngineAnalytics.h"
+#include "Subsystems/PlacementSubsystem.h"
+
+UCEEffectorActorFactory::UCEEffectorActorFactory()
+{
+	NewActorClass = ACEEffectorActor::StaticClass();
+}
+
+void UCEEffectorActorFactory::PostPlaceAsset(TArrayView<const FTypedElementHandle> InHandle, const FAssetPlacementInfo& InPlacementInfo, const FPlacementOptions& InPlacementOptions)
+{
+	Super::PostPlaceAsset(InHandle, InPlacementInfo, InPlacementOptions);
+
+	if (!InPlacementOptions.bIsCreatingPreviewElements && FEngineAnalytics::IsAvailable())
+	{
+		FEngineAnalytics::GetProvider().RecordEvent(TEXT("Editor.Usage.ClonerEffector.PlaceEffector"));
+	}
+}
diff --git a/Engine/Plugins/Experimental/ClonerEffector/Source/ClonerEffectorEditor/Public/Cloner/CEClonerActorFactory.h b/Engine/Plugins/Experimental/ClonerEffector/Source/ClonerEffectorEditor/Public/Cloner/CEClonerActorFactory.h
new file mode 100644
--- /dev/null
+++ b/Engine/Plugins/Experimental/ClonerEffector/Source/ClonerEffectorEditor/Public/Cloner/CEClonerActorFactory.h
@@ -0,0 +1,19 @@
+// Copyright Epic Games, Inc. All Rights Reserved.
+
+#pragma once
+
+#include "ActorFactories/ActorFactory.h"
+#include "CEClonerActorFactory.generated.h"
+
+UCLASS(MinimalAPI)
+class UCEClonerActorFactory : public UActorFactory
+{
+	GENERATED_BODY()
+
+public:
+	UCEClonerActorFactory();
+
+	//~ Begin UActorFactory
+	virtual void PostPlaceAsset(TArrayView<const FTypedElementHandle> InHandle, const FAssetPlacementInfo& InPlacementInfo, const FPlacementOptions& InPlacementOptions) override;
+	//~ End UActorFactory
+};
diff --git a/Engine/Plugins/Experimental/ClonerEffector/Source/ClonerEffectorEditor/Public/Effector/CEEffectorActorFactory.h b/Engine/Plugins/Experimental/ClonerEffector/Source/ClonerEffectorEditor/Public/Effector/CEEffectorActorFactory.h
new file mode 100644
--- /dev/null
+++ b/Engine/Plugins/Experimental/ClonerEffector/Source/ClonerEffectorEditor/Public/Effector/CEEffectorActorFactory.h
@@ -0,0 +1,19 @@
+// Copyright Epic Games, Inc. All Rights Reserved.
+
+#pragma once
+
+#include "ActorFactories/ActorFactory.h"
+#include "CEEffectorActorFactory.generated.h"
+
+UCLASS(MinimalAPI)
+class UCEEffectorActorFactory : public UActorFactory
+{
+	GENERATED_BODY()
+
+public:
+	UCEEffectorActorFactory();
+
+	//~ Begin UActorFactory
+	virtual void PostPlaceAsset(TArrayView<const FTypedElementHandle> InHandle, const FAssetPlacementInfo& InPlacementInfo, const FPlacementOptions& InPlacementOptions) override;
+	//~ End UActorFactory
+};
diff --git a/Engine/Plugins/Experimental/SVGImporter/Source/SVGImporterEditor/Private/Factories/SVGActorFactory.cpp b/Engine/Plugins/Experimental/SVGImporter/Source/SVGImporterEditor/Private/Factories/SVGActorFactory.cpp
--- a/Engine/Plugins/Experimental/SVGImporter/Source/SVGImporterEditor/Private/Factories/SVGActorFactory.cpp
+++ b/Engine/Plugins/Experimental/SVGImporter/Source/SVGImporterEditor/Private/Factories/SVGActorFactory.cpp
@@ -1,9 +1,11 @@
 // Copyright Epic Games, Inc. All Rights Reserved.
 
 #include "Factories/SVGActorFactory.h"
+#include "EngineAnalytics.h"
 #include "LevelEditorViewport.h"
 #include "SVGActor.h"
 #include "SVGData.h"
+#include "Subsystems/PlacementSubsystem.h"
 
 USVGActorFactory::USVGActorFactory(const FObjectInitializer& ObjectInitializer)
 	: Super(ObjectInitializer)
@@ -43,3 +45,13 @@ AActor* USVGActorFactory::SpawnActor(UObject* InAsset, ULevel* InLevel, const FT
 	SVGActor->Initialize();
 	return SVGActor;
 }
+
+void USVGActorFactory::PostPlaceAsset(TArrayView<const FTypedElementHandle> InHandle, const FAssetPlacementInfo& InPlacementInfo, const FPlacementOptions& InPlacementOptions)
+{
+	Super::PostPlaceAsset(InHandle, InPlacementInfo, InPlacementOptions);
+
+	if (!InPlacementOptions.bIsCreatingPreviewElements && FEngineAnalytics::IsAvailable())
+	{
+		FEngineAnalytics::GetProvider().RecordEvent(TEXT("Editor.Usage.SVGImporter.PlaceSVG"));
+	}
+}
diff --git a/Engine/Plugins/Experimental/SVGImporter/Source/SVGImporterEditor/Public/Factories/SVGActorFactory.h b/Engine/Plugins/Experimental/SVGImporter/Source/SVGImporterEditor/Public/Factories/SVGActorFactory.h
--- a/Engine/Plugins/Experimental/SVGImporter/Source/SVGImporterEditor/Public/Factories/SVGActorFactory.h
+++ b/Engine/Plugins/Experimental/SVGImporter/Source/SVGImporterEditor/Public/Factories/SVGActorFactory.h
@@ -17,5 +17,6 @@ protected:
 	virtual bool CanCreateActorFrom(const FAssetData& AssetData, FText& OutErrorMsg) override;
 	virtual AActor* GetDefaultActor(const FAssetData& AssetData) override;
 	virtual AActor* SpawnActor(UObject* InAsset, ULevel* InLevel, const FTransform& InTransform, const FActorSpawnParameters& InSpawnParams) override;
+	virtual void PostPlaceAsset(TArrayView<const FTypedElementHandle> InHandle, const FAssetPlacementInfo& InPlacementInfo, const FPlacementOptions& InPlacementOptions) override;
 	//~ End UActorFactory
 };
