[Motion Design] GizmoComponent OnRegister bug

diff --git a/Engine/Plugins/Experimental/Avalanche/Source/Avalanche/Private/Framework/AvaGizmoComponent.cpp b/Engine/Plugins/Experimental/Avalanche/Source/Avalanche/Private/Framework/AvaGizmoComponent.cpp
--- a/Engine/Plugins/Experimental/Avalanche/Source/Avalanche/Private/Framework/AvaGizmoComponent.cpp
+++ b/Engine/Plugins/Experimental/Avalanche/Source/Avalanche/Private/Framework/AvaGizmoComponent.cpp
@@ -74,9 +74,9 @@ void UAvaGizmoComponent::OnRegister()
 			PostRegisterComponentsHandle = World->AddOnPostRegisterAllActorComponentsHandler(
 				FOnPostRegisterAllActorComponents::FDelegate::CreateUObject(this, &UAvaGizmoComponent::OnPostRegisterParentComponents));
 		}
-		
-		Super::OnRegister();
 	}
+
+	Super::OnRegister();
 }
 
 void UAvaGizmoComponent::OnUnregister()
@@ -89,10 +89,10 @@ void UAvaGizmoComponent::OnUnregister()
 		{
 			World->RemoveOnPostRegisterAllActorComponentsHandler(PostRegisterComponentsHandle);
 		}
-
-		Super::OnUnregister();
 	}
 
+	Super::OnUnregister();
+
 	RestoreComponentValues();
 }
 
@@ -132,30 +132,38 @@ void UAvaGizmoComponent::ApplyGizmoValues()
 		IsValid(Owner))
 	{
 		bAreGizmoValuesApplied = true;
-	
-		Owner->ForEachComponent<UPrimitiveComponent>(bApplyToChildActors, [this](UPrimitiveComponent* InComponent)
+
+		TArray<UPrimitiveComponent*> PrimitiveComponents;
+		Owner->GetComponents<UPrimitiveComponent>(PrimitiveComponents, bApplyToChildActors);
+
+		for (UPrimitiveComponent* Component : PrimitiveComponents)
 		{
-			InComponent->bVisibleInReflectionCaptures = false;
-			InComponent->bVisibleInRealTimeSkyCaptures = false;
-			InComponent->bVisibleInRayTracing = false;
-			InComponent->SetVisibleInRayTracing(false);
-			InComponent->SetVisibleInSceneCaptureOnly(false);
+			if (!IsValid(Component))
+			{
+				continue;
+			}
+		
+			Component->bVisibleInReflectionCaptures = false;
+			Component->bVisibleInRealTimeSkyCaptures = false;
+			Component->bVisibleInRayTracing = false;
+			Component->SetVisibleInRayTracing(false);
+			Component->SetVisibleInSceneCaptureOnly(false);
 
-			InComponent->SetAffectDistanceFieldLighting(false);
-			InComponent->SetAffectDynamicIndirectLighting(false);
-			InComponent->SetAffectIndirectLightingWhileHidden(false);
+			Component->SetAffectDistanceFieldLighting(false);
+			Component->SetAffectDynamicIndirectLighting(false);
+			Component->SetAffectIndirectLightingWhileHidden(false);
 
-			InComponent->SetCastShadow(bCastShadow);
-			InComponent->SetVisibility(bIsVisibleInEditor, true);
-			InComponent->SetHiddenInGame(bIsHiddenInGame);
+			Component->SetCastShadow(bCastShadow);
+			Component->SetVisibility(bIsVisibleInEditor, true);
+			Component->SetHiddenInGame(bIsHiddenInGame);
 
-			InComponent->SetRenderCustomDepth(bSetStencil);
-			InComponent->SetCustomDepthStencilValue(StencilId);	
+			Component->SetRenderCustomDepth(bSetStencil);
+			Component->SetCustomDepthStencilValue(StencilId);	
 
-			InComponent->SetRenderInMainPass(bRenderInMainPass);
-			InComponent->SetRenderInDepthPass(bRenderDepth);
+			Component->SetRenderInMainPass(bRenderInMainPass);
+			Component->SetRenderInDepthPass(bRenderDepth);
 
-			if (UDynamicMeshComponent* DynamicMeshComponent = Cast<UDynamicMeshComponent>(InComponent))
+			if (UDynamicMeshComponent* DynamicMeshComponent = Cast<UDynamicMeshComponent>(Component))
 			{
 				DynamicMeshComponent->SetEnableWireframeRenderPass(bDrawWireframe);
 				DynamicMeshComponent->WireframeColor = WireframeColor;
@@ -163,12 +171,12 @@ void UAvaGizmoComponent::ApplyGizmoValues()
 
 			if (Material)
 			{
-				for (int32 MaterialIdx = 0; MaterialIdx < InComponent->GetNumMaterials(); ++MaterialIdx)
+				for (int32 MaterialIdx = 0; MaterialIdx < Component->GetNumMaterials(); ++MaterialIdx)
 				{
-					InComponent->SetMaterial(MaterialIdx, Material);
+					Component->SetMaterial(MaterialIdx, Material);
 				}
 			}
-		});
+		}
 
 		if (!FUObjectThreadContext::Get().IsRoutingPostLoad)
 		{
@@ -209,55 +217,63 @@ void UAvaGizmoComponent::RestoreComponentValues()
 		IsValid(Owner))
 	{
 		bAreGizmoValuesApplied = false;
-		
-		Owner->ForEachComponent<UPrimitiveComponent>(bApplyToChildActors, [this, &Owner](UPrimitiveComponent* InComponent)
+
+		TArray<UPrimitiveComponent*> PrimitiveComponents;
+		Owner->GetComponents<UPrimitiveComponent>(PrimitiveComponents, bApplyToChildActors);
+
+		for (UPrimitiveComponent* Component : PrimitiveComponents)
 		{
+			if (!IsValid(Component))
+			{
+				continue;
+			}
+			
 			FSoftComponentReference ComponentReference;
-			ComponentReference.PathToComponent = InComponent->GetPathName(Owner);
+			ComponentReference.PathToComponent = Component->GetPathName(Owner);
 			
 			if (FAvaGizmoComponentPrimitiveValues* StoredComponentValues = ComponentValues.Find(ComponentReference))
 			{
-				InComponent->bVisibleInReflectionCaptures = StoredComponentValues->bVisibleInReflectionCaptures;
-				InComponent->bVisibleInRealTimeSkyCaptures = StoredComponentValues->bVisibleInRealTimeSkyCaptures;
-				InComponent->bVisibleInRayTracing = StoredComponentValues->bVisibleInRayTracing;
+				Component->bVisibleInReflectionCaptures = StoredComponentValues->bVisibleInReflectionCaptures;
+				Component->bVisibleInRealTimeSkyCaptures = StoredComponentValues->bVisibleInRealTimeSkyCaptures;
+				Component->bVisibleInRayTracing = StoredComponentValues->bVisibleInRayTracing;
 
-				InComponent->SetAffectDistanceFieldLighting(StoredComponentValues->bAffectDistanceFieldLighting);
-				InComponent->SetAffectDynamicIndirectLighting(StoredComponentValues->bAffectDynamicIndirectLighting);
-				InComponent->SetAffectIndirectLightingWhileHidden(StoredComponentValues->bAffectIndirectLightingWhileHidden);
+				Component->SetAffectDistanceFieldLighting(StoredComponentValues->bAffectDistanceFieldLighting);
+				Component->SetAffectDynamicIndirectLighting(StoredComponentValues->bAffectDynamicIndirectLighting);
+				Component->SetAffectIndirectLightingWhileHidden(StoredComponentValues->bAffectIndirectLightingWhileHidden);
 										
-				InComponent->CastShadow = StoredComponentValues->bCastShadow;
-				InComponent->SetVisibility(StoredComponentValues->bIsVisible);
-				InComponent->bHiddenInGame = StoredComponentValues->bIsHiddenInGame;
+				Component->CastShadow = StoredComponentValues->bCastShadow;
+				Component->SetVisibility(StoredComponentValues->bIsVisible);
+				Component->bHiddenInGame = StoredComponentValues->bIsHiddenInGame;
 
-				InComponent->SetRenderCustomDepth(StoredComponentValues->bRendersCustomDepth);
-				InComponent->SetCustomDepthStencilValue(StoredComponentValues->CustomStencilId);
+				Component->SetRenderCustomDepth(StoredComponentValues->bRendersCustomDepth);
+				Component->SetCustomDepthStencilValue(StoredComponentValues->CustomStencilId);
 
-				InComponent->SetRenderInMainPass(StoredComponentValues->bRendersInMainPass);
-				InComponent->SetRenderInDepthPass(StoredComponentValues->bRendersDepth);
+				Component->SetRenderInMainPass(StoredComponentValues->bRendersInMainPass);
+				Component->SetRenderInDepthPass(StoredComponentValues->bRendersDepth);
 
-				if (UDynamicMeshComponent* DynamicMeshComponent = Cast<UDynamicMeshComponent>(InComponent))
+				if (UDynamicMeshComponent* DynamicMeshComponent = Cast<UDynamicMeshComponent>(Component))
 				{
 					DynamicMeshComponent->SetEnableWireframeRenderPass(StoredComponentValues->bDrawWireframe);
 					DynamicMeshComponent->WireframeColor = StoredComponentValues->WireframeColor;
 				}
 
-				const int32 ActualNumMaterials = InComponent->GetNumMaterials();
+				const int32 ActualNumMaterials = Component->GetNumMaterials();
 				const int32 ExpectedNumMaterials = StoredComponentValues->Materials.Num();
 
 				if (ActualNumMaterials != ExpectedNumMaterials)
 				{
 					UE_LOG(LogTemp, Warning, TEXT("Material count mismatch when restoring component '%s' - got %i, expected %i"),
-						*InComponent->GetName(),
+						*Component->GetName(),
 						ActualNumMaterials,
 						ExpectedNumMaterials);
 				}
 
 				for (int32 MaterialIdx = 0; MaterialIdx < FMath::Min(ActualNumMaterials, ExpectedNumMaterials); ++MaterialIdx)
 				{
-					InComponent->SetMaterial(MaterialIdx, StoredComponentValues->Materials[MaterialIdx]);
+					Component->SetMaterial(MaterialIdx, StoredComponentValues->Materials[MaterialIdx]);
 				}
 			}
-		});
+		}
 
 		if (!FUObjectThreadContext::Get().IsRoutingPostLoad)
 		{
@@ -465,49 +481,57 @@ void UAvaGizmoComponent::StoreComponentValues()
 		if (const AActor* Owner = GetOwner();
 			IsValid(Owner))
 		{
+			TArray<UPrimitiveComponent*> PrimitiveComponents;
+        	Owner->GetComponents<UPrimitiveComponent>(PrimitiveComponents, bApplyToChildActors);
+
 			bool bWasOneOrMorePrimitiveComponents = false;
-			Owner->ForEachComponent<UPrimitiveComponent>(bApplyToChildActors, [this, &Owner, &bWasOneOrMorePrimitiveComponents](UPrimitiveComponent* InComponent)
-			{
+        	for (UPrimitiveComponent* Component : PrimitiveComponents)
+        	{
+        		if (!IsValid(Component))
+        		{
+        			continue;
+        		}
+        		
 				bWasOneOrMorePrimitiveComponents = true;
 				
 				FSoftComponentReference ComponentReference;
-				ComponentReference.PathToComponent = InComponent->GetPathName(Owner);
+				ComponentReference.PathToComponent = Component->GetPathName(Owner);
 				
 				FAvaGizmoComponentPrimitiveValues& StoredComponentValues = ComponentValues.FindOrAdd(ComponentReference);
 				
-				StoredComponentValues.bVisibleInReflectionCaptures = InComponent->bVisibleInReflectionCaptures;
-				StoredComponentValues.bVisibleInRealTimeSkyCaptures = InComponent->bVisibleInRealTimeSkyCaptures;
-				StoredComponentValues.bVisibleInRayTracing = InComponent->bVisibleInRayTracing;
+				StoredComponentValues.bVisibleInReflectionCaptures = Component->bVisibleInReflectionCaptures;
+				StoredComponentValues.bVisibleInRealTimeSkyCaptures = Component->bVisibleInRealTimeSkyCaptures;
+				StoredComponentValues.bVisibleInRayTracing = Component->bVisibleInRayTracing;
 				
-				StoredComponentValues.bAffectDistanceFieldLighting = InComponent->bAffectDistanceFieldLighting;
-				StoredComponentValues.bAffectDynamicIndirectLighting = InComponent->bAffectDynamicIndirectLighting;
-				StoredComponentValues.bAffectIndirectLightingWhileHidden = InComponent->bAffectIndirectLightingWhileHidden;
+				StoredComponentValues.bAffectDistanceFieldLighting = Component->bAffectDistanceFieldLighting;
+				StoredComponentValues.bAffectDynamicIndirectLighting = Component->bAffectDynamicIndirectLighting;
+				StoredComponentValues.bAffectIndirectLightingWhileHidden = Component->bAffectIndirectLightingWhileHidden;
 				
-				StoredComponentValues.bCastShadow = InComponent->CastShadow;
-				StoredComponentValues.bIsVisible = InComponent->IsVisible();
-				StoredComponentValues.bIsHiddenInGame = InComponent->bHiddenInGame;
+				StoredComponentValues.bCastShadow = Component->CastShadow;
+				StoredComponentValues.bIsVisible = Component->IsVisible();
+				StoredComponentValues.bIsHiddenInGame = Component->bHiddenInGame;
 
-				StoredComponentValues.bRendersCustomDepth = InComponent->bRenderCustomDepth > 0;
-				StoredComponentValues.CustomStencilId = InComponent->CustomDepthStencilValue;
+				StoredComponentValues.bRendersCustomDepth = Component->bRenderCustomDepth > 0;
+				StoredComponentValues.CustomStencilId = Component->CustomDepthStencilValue;
 
-				StoredComponentValues.bRendersInMainPass = InComponent->bRenderInMainPass;
-				StoredComponentValues.bRendersDepth = InComponent->bRenderInDepthPass;
+				StoredComponentValues.bRendersInMainPass = Component->bRenderInMainPass;
+				StoredComponentValues.bRendersDepth = Component->bRenderInDepthPass;
 
 				// @todo: move this impl specific stuff to a Handler/Provider subsystem
-				if (UDynamicMeshComponent* DynamicMeshComponent = Cast<UDynamicMeshComponent>(InComponent))
+				if (UDynamicMeshComponent* DynamicMeshComponent = Cast<UDynamicMeshComponent>(Component))
 				{
 					StoredComponentValues.bDrawWireframe = DynamicMeshComponent->GetEnableWireframeRenderPass();
 					StoredComponentValues.WireframeColor = DynamicMeshComponent->WireframeColor;
 				}
 
 				TArray<UMaterialInterface*> UsedMaterials;
-				InComponent->GetUsedMaterials(UsedMaterials);
+				Component->GetUsedMaterials(UsedMaterials);
 
 				Algo::Transform(UsedMaterials, StoredComponentValues.Materials, [](UMaterialInterface* InMaterial)
 				{
 					return InMaterial;
 				});
-			});
+			}
 
 			if (!bWasOneOrMorePrimitiveComponents)
 			{
