Motion Design:
- Fix for image scaling in the "texture to scene" option.

diff --git a/Engine/Plugins/Experimental/Avalanche/Source/AvalancheEditor/Private/MaterialDesigner/AvaLevelMaterialDesignerExtension.cpp b/Engine/Plugins/Experimental/Avalanche/Source/AvalancheEditor/Private/MaterialDesigner/AvaLevelMaterialDesignerExtension.cpp
--- a/Engine/Plugins/Experimental/Avalanche/Source/AvalancheEditor/Private/MaterialDesigner/AvaLevelMaterialDesignerExtension.cpp
+++ b/Engine/Plugins/Experimental/Avalanche/Source/AvalancheEditor/Private/MaterialDesigner/AvaLevelMaterialDesignerExtension.cpp
@@ -93,7 +93,7 @@ TSharedRef<FExtender> FAvaLevelMaterialDesignerExtension::OnExtendContentBrowser
 						LOCTEXT("AddTextureToScene", "Add Texture To Scene"),
 						LOCTEXT("AddTextureToSceneTooltip", "Creates a parametric shape in the scene the same size as the texture and creates a Material Designer asset on using this texture."),
 						FSlateIconFinder::FindIconForClass(UMaterial::StaticClass()),
-						FUIAction(FExecuteAction::CreateStatic(&FAvaLevelMaterialDesignerExtension::OpenHeaderViewForAsset, InSelectedAssets[0]))
+						FUIAction(FExecuteAction::CreateStatic(&FAvaLevelMaterialDesignerExtension::AddTextureToSene, InSelectedAssets[0]))
 					);
 				})
 			);
@@ -103,7 +103,7 @@ TSharedRef<FExtender> FAvaLevelMaterialDesignerExtension::OnExtendContentBrowser
 	return Extender;
 }
 
-void FAvaLevelMaterialDesignerExtension::OpenHeaderViewForAsset(FAssetData InAssetData)
+void FAvaLevelMaterialDesignerExtension::AddTextureToSene(FAssetData InAssetData)
 {
 	const TArray<TSharedPtr<IAvaViewportClient>> LevelViewportClients = FAvaLevelViewportExtension::GetLevelEditorViewportClients();
 
diff --git a/Engine/Plugins/Experimental/Avalanche/Source/AvalancheEditor/Private/MaterialDesigner/AvaLevelMaterialDesignerExtension.h b/Engine/Plugins/Experimental/Avalanche/Source/AvalancheEditor/Private/MaterialDesigner/AvaLevelMaterialDesignerExtension.h
--- a/Engine/Plugins/Experimental/Avalanche/Source/AvalancheEditor/Private/MaterialDesigner/AvaLevelMaterialDesignerExtension.h
+++ b/Engine/Plugins/Experimental/Avalanche/Source/AvalancheEditor/Private/MaterialDesigner/AvaLevelMaterialDesignerExtension.h
@@ -31,5 +31,5 @@ private:
 
 	static TSharedRef<FExtender> OnExtendContentBrowserAssetSelectionMenu(const TArray<FAssetData>& InSelectedAssets);
 
-	static void OpenHeaderViewForAsset(FAssetData InAssetData);
+	static void AddTextureToSene(FAssetData InAssetData);
 };
diff --git a/Engine/Plugins/Experimental/Avalanche/Source/AvalancheEditor/Private/MaterialDesigner/AvaMaterialDesignerTextureAssetFactory.cpp b/Engine/Plugins/Experimental/Avalanche/Source/AvalancheEditor/Private/MaterialDesigner/AvaMaterialDesignerTextureAssetFactory.cpp
--- a/Engine/Plugins/Experimental/Avalanche/Source/AvalancheEditor/Private/MaterialDesigner/AvaMaterialDesignerTextureAssetFactory.cpp
+++ b/Engine/Plugins/Experimental/Avalanche/Source/AvalancheEditor/Private/MaterialDesigner/AvaMaterialDesignerTextureAssetFactory.cpp
@@ -187,6 +187,7 @@ AActor* UAvaMaterialDesignerTextureAssetFactory::SpawnActor(UObject* InAsset, UL
 
 	TextureValue->SetValue(Texture);
 	RectangleMesh->SetMaterial(UAvaShapeDynamicMeshBase::MESH_INDEX_PRIMARY, NewInstance);
+	RectangleMesh->SetMaterialUVMode(UAvaShapeDynamicMeshBase::MESH_INDEX_PRIMARY, EAvaShapeUVMode::Stretch);
 
 	return InNewActor;
 }
