Motion Design
- Renamed "Backplate" visualizer to "Background" visualizer.

diff --git a/Engine/Plugins/Experimental/Avalanche/Source/Avalanche/Public/Viewport/Interaction/AvaViewportPostProcessInfo.h b/Engine/Plugins/Experimental/Avalanche/Source/Avalanche/Public/Viewport/Interaction/AvaViewportPostProcessInfo.h
--- a/Engine/Plugins/Experimental/Avalanche/Source/Avalanche/Public/Viewport/Interaction/AvaViewportPostProcessInfo.h
+++ b/Engine/Plugins/Experimental/Avalanche/Source/Avalanche/Public/Viewport/Interaction/AvaViewportPostProcessInfo.h
@@ -11,7 +11,7 @@ UENUM()
 enum class EAvaViewportPostProcessType_Deprecated : uint8
 {
 	None,
-	Backplate,
+	Background,
 	RedChannel,
 	GreenChannel,
 	BlueChannel,
diff --git a/Engine/Plugins/Experimental/Avalanche/Source/AvalancheLevelViewport/Private/AvaLevelViewportCommands.cpp b/Engine/Plugins/Experimental/Avalanche/Source/AvalancheLevelViewport/Private/AvaLevelViewportCommands.cpp
--- a/Engine/Plugins/Experimental/Avalanche/Source/AvalancheLevelViewport/Private/AvaLevelViewportCommands.cpp
+++ b/Engine/Plugins/Experimental/Avalanche/Source/AvalancheLevelViewport/Private/AvaLevelViewportCommands.cpp
@@ -70,9 +70,9 @@ void FAvaLevelViewportCommands::RegisterViewportCommands()
 		, EUserInterfaceActionType::Check
 		, FInputChord())
 
-	UI_COMMAND(TogglePostProcessBackplate
-		, "Backplate"
-		, "Switch to the backplate post process renderer."
+	UI_COMMAND(TogglePostProcessBackground
+		, "Background"
+		, "Switch to the background post process renderer."
 		, EUserInterfaceActionType::Check
 		, FInputChord())
 
diff --git a/Engine/Plugins/Experimental/Avalanche/Source/AvalancheLevelViewport/Private/SAvaLevelViewport.cpp b/Engine/Plugins/Experimental/Avalanche/Source/AvalancheLevelViewport/Private/SAvaLevelViewport.cpp
--- a/Engine/Plugins/Experimental/Avalanche/Source/AvalancheLevelViewport/Private/SAvaLevelViewport.cpp
+++ b/Engine/Plugins/Experimental/Avalanche/Source/AvalancheLevelViewport/Private/SAvaLevelViewport.cpp
@@ -117,10 +117,10 @@ void SAvaLevelViewport::BindCommands()
 	);
 
 	CommandListRef.MapAction(
-		AvaLevelViewportCommands.TogglePostProcessBackplate,
-		FExecuteAction::CreateSP(this, &SAvaLevelViewport::ExecuteTogglePostProcessBackplate),
-		FCanExecuteAction::CreateSP(this, &SAvaLevelViewport::CanTogglePostProcessBackplate),
-		FIsActionChecked::CreateSP(this, &SAvaLevelViewport::IsPostProcessBackplateEnabled)
+		AvaLevelViewportCommands.TogglePostProcessBackground,
+		FExecuteAction::CreateSP(this, &SAvaLevelViewport::ExecuteTogglePostProcessBackground),
+		FCanExecuteAction::CreateSP(this, &SAvaLevelViewport::CanTogglePostProcessBackground),
+		FIsActionChecked::CreateSP(this, &SAvaLevelViewport::IsPostProcessBackgroundEnabled)
 	);
 
 	CommandListRef.MapAction(
diff --git a/Engine/Plugins/Experimental/Avalanche/Source/AvalancheLevelViewport/Private/SAvaLevelViewport.h b/Engine/Plugins/Experimental/Avalanche/Source/AvalancheLevelViewport/Private/SAvaLevelViewport.h
--- a/Engine/Plugins/Experimental/Avalanche/Source/AvalancheLevelViewport/Private/SAvaLevelViewport.h
+++ b/Engine/Plugins/Experimental/Avalanche/Source/AvalancheLevelViewport/Private/SAvaLevelViewport.h
@@ -130,9 +130,9 @@ public:
 	bool CanTogglePostProcessNone() const;
 	void ExecuteTogglePostProcessNone();
 
-	bool IsPostProcessBackplateEnabled() const;
-	bool CanTogglePostProcessBackplate() const;
-	void ExecuteTogglePostProcessBackplate();
+	bool IsPostProcessBackgroundEnabled() const;
+	bool CanTogglePostProcessBackground() const;
+	void ExecuteTogglePostProcessBackground();
 
 	bool IsPostProcessChannelRedEnabled() const;
 	bool CanTogglePostProcessChannelRed() const;
@@ -190,7 +190,7 @@ private:
 	TArray<TSharedPtr<SAvaLevelViewportGuide>> Guides;
 	TSharedPtr<SAvaLevelViewportSnapIndicators> SnapIndicators;
 	FPanelExtensionFactory PanelExtensionFactory;
-	TSharedPtr<SWidget> BackplateTextureSelector;
+	TSharedPtr<SWidget> BackgroundTextureSelector;
 	TSharedPtr<SWidget> PostProcessOpacitySlider;
 	TSharedPtr<FScopedTransaction> PostProcessInfoTransaction;
 
@@ -237,14 +237,14 @@ private:
 	void BeginPostProcessInfoTransaction();
 	void EndPostProcessInfoTransaction();
 
-	FString GetBackplateTextureObjectPath() const;
-	void OnBackplateTextureChanged(const FAssetData& InAssetData);
+	FString GetBackgroundTextureObjectPath() const;
+	void OnBackgroundTextureChanged(const FAssetData& InAssetData);
 
-	float GetBackplateOpacity() const;
-	void OnBackplateOpacitySliderBegin();
-	void OnBackplateOpacitySliderEnd(float InValue);
-	void OnBackplateOpacityChanged(float InValue);
-	void OnBackplateOpacityCommitted(float InValue, ETextCommit::Type InCommitType);
+	float GetBackgroundOpacity() const;
+	void OnBackgroundOpacitySliderBegin();
+	void OnBackgroundOpacitySliderEnd(float InValue);
+	void OnBackgroundOpacityChanged(float InValue);
+	void OnBackgroundOpacityCommitted(float InValue, ETextCommit::Type InCommitType);
 
 	int32 GetVirtualSizeX() const;
 	int32 GetVirtualSizeY() const;
diff --git a/Engine/Plugins/Experimental/Avalanche/Source/AvalancheLevelViewport/Private/SAvaLevelViewport_Menus.cpp b/Engine/Plugins/Experimental/Avalanche/Source/AvalancheLevelViewport/Private/SAvaLevelViewport_Menus.cpp
--- a/Engine/Plugins/Experimental/Avalanche/Source/AvalancheLevelViewport/Private/SAvaLevelViewport_Menus.cpp
+++ b/Engine/Plugins/Experimental/Avalanche/Source/AvalancheLevelViewport/Private/SAvaLevelViewport_Menus.cpp
@@ -867,7 +867,7 @@ FReply SAvaLevelViewport::OnChildActorLockButtonClicked()
 	return FReply::Handled();
 }
 
-FString SAvaLevelViewport::GetBackplateTextureObjectPath() const
+FString SAvaLevelViewport::GetBackgroundTextureObjectPath() const
 {
 	TSharedPtr<FAvaLevelViewportClient> ViewportClient = GetAvaLevelViewportClient();
 
@@ -891,7 +891,7 @@ FString SAvaLevelViewport::GetBackplateTextureObjectPath() const
 	return PostProcessInfo->Texture.ToString();
 }
 
-void SAvaLevelViewport::OnBackplateTextureChanged(const FAssetData& InAssetData)
+void SAvaLevelViewport::OnBackgroundTextureChanged(const FAssetData& InAssetData)
 {
 	TSharedPtr<FAvaLevelViewportClient> ViewportClient = GetAvaLevelViewportClient();
 
@@ -921,7 +921,7 @@ void SAvaLevelViewport::OnBackplateTextureChanged(const FAssetData& InAssetData)
 	EndPostProcessInfoTransaction();
 }
 
-float SAvaLevelViewport::GetBackplateOpacity() const
+float SAvaLevelViewport::GetBackgroundOpacity() const
 {
 	TSharedPtr<FAvaLevelViewportClient> ViewportClient = GetAvaLevelViewportClient();
 
@@ -976,17 +976,17 @@ void SAvaLevelViewport::EndPostProcessInfoTransaction()
 	PostProcessInfoTransaction.Reset();
 }
 
-void SAvaLevelViewport::OnBackplateOpacitySliderBegin()
+void SAvaLevelViewport::OnBackgroundOpacitySliderBegin()
 {
 	BeginPostProcessInfoTransaction();
 }
 
-void SAvaLevelViewport::OnBackplateOpacitySliderEnd(float InValue)
+void SAvaLevelViewport::OnBackgroundOpacitySliderEnd(float InValue)
 {
 	EndPostProcessInfoTransaction();
 }
 
-void SAvaLevelViewport::OnBackplateOpacityChanged(float InValue)
+void SAvaLevelViewport::OnBackgroundOpacityChanged(float InValue)
 {
 	TSharedPtr<FAvaLevelViewportClient> ViewportClient = GetAvaLevelViewportClient();
 
@@ -1004,7 +1004,7 @@ void SAvaLevelViewport::OnBackplateOpacityChanged(float InValue)
 	ViewportClient->Invalidate();
 }
 
-void SAvaLevelViewport::OnBackplateOpacityCommitted(float InValue, ETextCommit::Type InCommitType)
+void SAvaLevelViewport::OnBackgroundOpacityCommitted(float InValue, ETextCommit::Type InCommitType)
 {
 	if (InCommitType == ETextCommit::OnCleared)
 	{
diff --git a/Engine/Plugins/Experimental/Avalanche/Source/AvalancheLevelViewport/Private/SAvaLevelViewport_StatusBar.cpp b/Engine/Plugins/Experimental/Avalanche/Source/AvalancheLevelViewport/Private/SAvaLevelViewport_StatusBar.cpp
--- a/Engine/Plugins/Experimental/Avalanche/Source/AvalancheLevelViewport/Private/SAvaLevelViewport_StatusBar.cpp
+++ b/Engine/Plugins/Experimental/Avalanche/Source/AvalancheLevelViewport/Private/SAvaLevelViewport_StatusBar.cpp
@@ -64,7 +64,7 @@ void SAvaLevelViewport::ExecuteTogglePostProcessNone()
 	EndPostProcessInfoTransaction();
 }
 
-bool SAvaLevelViewport::IsPostProcessBackplateEnabled() const
+bool SAvaLevelViewport::IsPostProcessBackgroundEnabled() const
 {
 	TSharedPtr<FAvaLevelViewportClient> ViewportClient = GetAvaLevelViewportClient();
 
@@ -78,10 +78,10 @@ bool SAvaLevelViewport::IsPostProcessBackplateEnabled() const
 		return false;
 	}
 
-	return ViewportClient->GetPostProcessManager()->GetType() == EAvaViewportPostProcessType::Backplate;
+	return ViewportClient->GetPostProcessManager()->GetType() == EAvaViewportPostProcessType::Background;
 }
 
-bool SAvaLevelViewport::CanTogglePostProcessBackplate() const
+bool SAvaLevelViewport::CanTogglePostProcessBackground() const
 {
 	TSharedPtr<FAvaLevelViewportClient> ViewportClient = GetAvaLevelViewportClient();
 
@@ -95,20 +95,20 @@ bool SAvaLevelViewport::CanTogglePostProcessBackplate() const
 		return false;
 	}
 
-	if (IsPostProcessBackplateEnabled())
+	if (IsPostProcessBackgroundEnabled())
 	{
 		return true;
 	}
 
-	if (TSharedPtr<FAvaViewportPostProcessVisualizer> BackplateVisualizer = ViewportClient->GetPostProcessManager()->GetVisualizer(EAvaViewportPostProcessType::Backplate))
+	if (TSharedPtr<FAvaViewportPostProcessVisualizer> BackgroundVisualizer = ViewportClient->GetPostProcessManager()->GetVisualizer(EAvaViewportPostProcessType::Background))
 	{
-		return BackplateVisualizer->CanActivate(/* bInSilent */ true);
+		return BackgroundVisualizer->CanActivate(/* bInSilent */ true);
 	}
 
 	return false;
 }
 
-void SAvaLevelViewport::ExecuteTogglePostProcessBackplate()
+void SAvaLevelViewport::ExecuteTogglePostProcessBackground()
 {
 	TSharedPtr<FAvaLevelViewportClient> ViewportClient = GetAvaLevelViewportClient();
 
@@ -124,7 +124,7 @@ void SAvaLevelViewport::ExecuteTogglePostProcessBackplate()
 
 	BeginPostProcessInfoTransaction();
 
-	ViewportClient->GetPostProcessManager()->SetType(EAvaViewportPostProcessType::Backplate);
+	ViewportClient->GetPostProcessManager()->SetType(EAvaViewportPostProcessType::Background);
 	ViewportClient->Invalidate();
 
 	EndPostProcessInfoTransaction();
diff --git a/Engine/Plugins/Experimental/Avalanche/Source/AvalancheLevelViewport/Private/Widgets/SAvaLevelViewportStatusBarButtons.cpp b/Engine/Plugins/Experimental/Avalanche/Source/AvalancheLevelViewport/Private/Widgets/SAvaLevelViewportStatusBarButtons.cpp
--- a/Engine/Plugins/Experimental/Avalanche/Source/AvalancheLevelViewport/Private/Widgets/SAvaLevelViewportStatusBarButtons.cpp
+++ b/Engine/Plugins/Experimental/Avalanche/Source/AvalancheLevelViewport/Private/Widgets/SAvaLevelViewportStatusBarButtons.cpp
@@ -38,7 +38,7 @@ namespace UE::Ava::LevelViewportStatusBarButtons::Private
 	static const FName AvaLevelViewportStyleName = FAvaLevelViewportStyle::Get().GetStyleSetName();
 	static const FName AppStyleSetName = FAppStyle::Get().GetStyleSetName();
 	static const FSlateIcon RGBChannelIcon = FSlateIcon(AvaLevelViewportStyleName, "AvalancheIcons.PostProcess.RGB");
-	static const FSlateIcon BackplateIcon = FSlateIcon(AppStyleSetName, "Icons.Role");
+	static const FSlateIcon BackgroundIcon = FSlateIcon(AppStyleSetName, "Icons.Role");
 	static const FSlateIcon RedChannelIcon = FSlateIcon(AvaLevelViewportStyleName, "AvalancheIcons.PostProcess.Red");
 	static const FSlateIcon GreenChannelIcon = FSlateIcon(AvaLevelViewportStyleName, "AvalancheIcons.PostProcess.Green");
 	static const FSlateIcon BlueChannelIcon = FSlateIcon(AvaLevelViewportStyleName, "AvalancheIcons.PostProcess.Blue");
@@ -179,19 +179,19 @@ void SAvaLevelViewportStatusBarButtons::CreateContextMenuWigets()
 			.ClearKeyboardFocusOnCommit(true)
 			.MaxFractionalDigits(3)
 			.MinDesiredWidth(50.f)
-			.OnBeginSliderMovement(LevelViewport, &SAvaLevelViewport::OnBackplateOpacitySliderBegin)
-			.OnEndSliderMovement(LevelViewport, &SAvaLevelViewport::OnBackplateOpacitySliderEnd)
-			.OnValueCommitted(LevelViewport, &SAvaLevelViewport::OnBackplateOpacityCommitted)
-			.Value(LevelViewport, &SAvaLevelViewport::GetBackplateOpacity)
+			.OnBeginSliderMovement(LevelViewport, &SAvaLevelViewport::OnBackgroundOpacitySliderBegin)
+			.OnEndSliderMovement(LevelViewport, &SAvaLevelViewport::OnBackgroundOpacitySliderEnd)
+			.OnValueCommitted(LevelViewport, &SAvaLevelViewport::OnBackgroundOpacityCommitted)
+			.Value(LevelViewport, &SAvaLevelViewport::GetBackgroundOpacity)
 			.MinValue(0.f)
 			.MinSliderValue(0.f)
 			.MaxValue(1.f)
 			.MaxSliderValue(1.f);
 	}
 
-	if (!BackplateTextureSelector.IsValid())
+	if (!BackgroundTextureSelector.IsValid())
 	{
-		BackplateTextureSelector = SNew(SObjectPropertyEntryBox)
+		BackgroundTextureSelector = SNew(SObjectPropertyEntryBox)
 			.AllowClear(true)
 			.AllowedClass(UTexture::StaticClass())
 			.DisplayBrowse(true)
@@ -200,8 +200,8 @@ void SAvaLevelViewportStatusBarButtons::CreateContextMenuWigets()
 			.DisplayUseSelected(true)
 			.ThumbnailPool(UThumbnailManager::Get().GetSharedThumbnailPool())
 			.EnableContentPicker(true)
-			.ObjectPath(LevelViewport, &SAvaLevelViewport::GetBackplateTextureObjectPath)
-			.OnObjectChanged(LevelViewport, &SAvaLevelViewport::OnBackplateTextureChanged)
+			.ObjectPath(LevelViewport, &SAvaLevelViewport::GetBackgroundTextureObjectPath)
+			.OnObjectChanged(LevelViewport, &SAvaLevelViewport::OnBackgroundTextureChanged)
 			.OnShouldSetAsset(FOnShouldSetAsset::CreateLambda([](const FAssetData& InAssetData) { return false; }));
 	}
 
@@ -476,7 +476,7 @@ const FSlateBrush* SAvaLevelViewportStatusBarButtons::GetPostProcessIcon() const
 	using namespace UE::Ava::LevelViewportStatusBarButtons::Private;
 
 	static const FSlateBrush* RGBBrush = RGBChannelIcon.GetIcon();
-	static const FSlateBrush* BackplateBrush = BackplateIcon.GetIcon();
+	static const FSlateBrush* BackgroundBrush = BackgroundIcon.GetIcon();
 	static const FSlateBrush* RedChannelBrush = RedChannelIcon.GetIcon();
 	static const FSlateBrush* GreenChannelBrush = GreenChannelIcon.GetIcon();
 	static const FSlateBrush* BlueChannelBrush = BlueChannelIcon.GetIcon();
@@ -494,8 +494,8 @@ const FSlateBrush* SAvaLevelViewportStatusBarButtons::GetPostProcessIcon() const
 				case EAvaViewportPostProcessType::None:
 					return RGBBrush;
 
-				case EAvaViewportPostProcessType::Backplate:
-					return BackplateBrush;
+				case EAvaViewportPostProcessType::Background:
+					return BackgroundBrush;
 
 				case EAvaViewportPostProcessType::RedChannel:
 					return RedChannelBrush;
@@ -536,15 +536,15 @@ TSharedRef<SWidget> SAvaLevelViewportStatusBarButtons::GetPostProcessMenuContent
 
 	check(Menus);
 
-	static const FName BackplateMenuName = TEXT("AvaLevelViewport.StatusBar.PostProcess.Backplate");
+	static const FName BackgroundMenuName = TEXT("AvaLevelViewport.StatusBar.PostProcess.Background");
 
-	UToolMenu* ContextMenu = Menus->FindMenu(BackplateMenuName);
+	UToolMenu* ContextMenu = Menus->FindMenu(BackgroundMenuName);
 
 	if (!ContextMenu)
 	{
 		using namespace UE::Ava::LevelViewportStatusBarButtons::Private;
 
-		ContextMenu = Menus->RegisterMenu(BackplateMenuName, NAME_None, EMultiBoxType::Menu);
+		ContextMenu = Menus->RegisterMenu(BackgroundMenuName, NAME_None, EMultiBoxType::Menu);
 
 		const FAvaLevelViewportCommands& CommandsRef = FAvaLevelViewportCommands::Get();
 
@@ -564,17 +564,17 @@ TSharedRef<SWidget> SAvaLevelViewportStatusBarButtons::GetPostProcessMenuContent
 			EUserInterfaceActionType::Check
 		));
 
-		FToolUIAction BackplateAction;
-		BackplateAction.ExecuteAction = FToolMenuExecuteAction::CreateSP(this, &SAvaLevelViewportStatusBarButtons::TogglePostProcessMenu, EAvaViewportPostProcessType::Backplate);
-		BackplateAction.GetActionCheckState = FToolMenuGetActionCheckState::CreateSP(this, &SAvaLevelViewportStatusBarButtons::GetPostProcessActiveMenu, EAvaViewportPostProcessType::Backplate);
-		BackplateAction.CanExecuteAction = FToolMenuCanExecuteAction::CreateSP(this, &SAvaLevelViewportStatusBarButtons::GetPostProcessEnabledMenu);
+		FToolUIAction BackgroundAction;
+		BackgroundAction.ExecuteAction = FToolMenuExecuteAction::CreateSP(this, &SAvaLevelViewportStatusBarButtons::TogglePostProcessMenu, EAvaViewportPostProcessType::Background);
+		BackgroundAction.GetActionCheckState = FToolMenuGetActionCheckState::CreateSP(this, &SAvaLevelViewportStatusBarButtons::GetPostProcessActiveMenu, EAvaViewportPostProcessType::Background);
+		BackgroundAction.CanExecuteAction = FToolMenuCanExecuteAction::CreateSP(this, &SAvaLevelViewportStatusBarButtons::GetPostProcessEnabledMenu);
 
 		EffectsSection.AddEntry(FToolMenuEntry::InitMenuEntry(
-			"Backplate", 
-			LOCTEXT("Backplate", "Backplate"), 
-			CommandsRef.TogglePostProcessBackplate->GetDescription(),
-			BackplateIcon,
-			FToolUIActionChoice(BackplateAction),
+			"Background", 
+			LOCTEXT("Background", "Background"), 
+			CommandsRef.TogglePostProcessBackground->GetDescription(),
+			BackgroundIcon,
+			FToolUIActionChoice(BackgroundAction),
 			EUserInterfaceActionType::Check
 		));
 
@@ -660,11 +660,11 @@ TSharedRef<SWidget> SAvaLevelViewportStatusBarButtons::GetPostProcessMenuContent
 			));
 		}
 
-		if (BackplateTextureSelector.IsValid())
+		if (BackgroundTextureSelector.IsValid())
 		{
 			OptionsSection.AddEntry(FToolMenuEntry::InitWidget(
 				"PostProcessTexture",
-				BackplateTextureSelector.ToSharedRef(),
+				BackgroundTextureSelector.ToSharedRef(),
 				LOCTEXT("PostProcessTexture", "Texture"),
 				true
 			));
diff --git a/Engine/Plugins/Experimental/Avalanche/Source/AvalancheLevelViewport/Private/Widgets/SAvaLevelViewportStatusBarButtons.h b/Engine/Plugins/Experimental/Avalanche/Source/AvalancheLevelViewport/Private/Widgets/SAvaLevelViewportStatusBarButtons.h
--- a/Engine/Plugins/Experimental/Avalanche/Source/AvalancheLevelViewport/Private/Widgets/SAvaLevelViewportStatusBarButtons.h
+++ b/Engine/Plugins/Experimental/Avalanche/Source/AvalancheLevelViewport/Private/Widgets/SAvaLevelViewportStatusBarButtons.h
@@ -24,7 +24,7 @@ public:
 protected:
 	TWeakPtr<SAvaLevelViewportFrame> ViewportFrameWeak;
 
-	TSharedPtr<SWidget> BackplateTextureSelector;
+	TSharedPtr<SWidget> BackgroundTextureSelector;
 
 	TSharedPtr<SWidget> PostProcessOpacitySlider;
 
diff --git a/Engine/Plugins/Experimental/Avalanche/Source/AvalancheLevelViewport/Public/AvaLevelViewportCommands.h b/Engine/Plugins/Experimental/Avalanche/Source/AvalancheLevelViewport/Public/AvaLevelViewportCommands.h
--- a/Engine/Plugins/Experimental/Avalanche/Source/AvalancheLevelViewport/Public/AvaLevelViewportCommands.h
+++ b/Engine/Plugins/Experimental/Avalanche/Source/AvalancheLevelViewport/Public/AvaLevelViewportCommands.h
@@ -29,7 +29,7 @@ public:
 	TSharedPtr<FUICommandInfo> ToggleChildActorLock;
 	TSharedPtr<FUICommandInfo> ToggleShapeEditorOverlay;
 	TSharedPtr<FUICommandInfo> TogglePostProcessNone;
-	TSharedPtr<FUICommandInfo> TogglePostProcessBackplate;
+	TSharedPtr<FUICommandInfo> TogglePostProcessBackground;
 	TSharedPtr<FUICommandInfo> TogglePostProcessChannelRed;
 	TSharedPtr<FUICommandInfo> TogglePostProcessChannelGreen;
 	TSharedPtr<FUICommandInfo> TogglePostProcessChannelBlue;
diff --git a/Engine/Plugins/Experimental/Avalanche/Source/AvalancheViewport/Private/AvaViewportPostProcessManager.cpp b/Engine/Plugins/Experimental/Avalanche/Source/AvalancheViewport/Private/AvaViewportPostProcessManager.cpp
--- a/Engine/Plugins/Experimental/Avalanche/Source/AvalancheViewport/Private/AvaViewportPostProcessManager.cpp
+++ b/Engine/Plugins/Experimental/Avalanche/Source/AvalancheViewport/Private/AvaViewportPostProcessManager.cpp
@@ -9,7 +9,7 @@
 #include "Viewport/Interaction/IAvaViewportDataProvider.h"
 #include "Viewport/Interaction/IAvaViewportDataProxy.h"
 #include "ViewportClient/IAvaViewportClient.h"
-#include "Visualizers/AvaViewportBackplateVisualizer.h"
+#include "Visualizers/AvaViewportBackgroundVisualizer.h"
 #include "Visualizers/AvaViewportChannelVisualizer.h"
 #include "Visualizers/AvaViewportCheckerboardVisualizer.h"
 
@@ -17,7 +17,7 @@ FAvaViewportPostProcessManager::FAvaViewportPostProcessManager(TSharedRef<IAvaVi
 {
 	AvaViewportClientWeak = InAvaViewportClient;
 
-	Visualizers.Emplace(EAvaViewportPostProcessType::Backplate,    MakeShared<FAvaViewportBackplateVisualizer>(InAvaViewportClient));
+	Visualizers.Emplace(EAvaViewportPostProcessType::Background,    MakeShared<FAvaViewportBackgroundVisualizer>(InAvaViewportClient));
 	Visualizers.Emplace(EAvaViewportPostProcessType::RedChannel,   MakeShared<FAvaViewportChannelVisualizer>(InAvaViewportClient, EAvaViewportPostProcessType::RedChannel));
 	Visualizers.Emplace(EAvaViewportPostProcessType::GreenChannel, MakeShared<FAvaViewportChannelVisualizer>(InAvaViewportClient, EAvaViewportPostProcessType::GreenChannel));
 	Visualizers.Emplace(EAvaViewportPostProcessType::BlueChannel,  MakeShared<FAvaViewportChannelVisualizer>(InAvaViewportClient, EAvaViewportPostProcessType::BlueChannel));
diff --git a/Engine/Plugins/Experimental/Avalanche/Source/AvalancheViewport/Private/AvaViewportSettings.cpp b/Engine/Plugins/Experimental/Avalanche/Source/AvalancheViewport/Private/AvaViewportSettings.cpp
--- a/Engine/Plugins/Experimental/Avalanche/Source/AvalancheViewport/Private/AvaViewportSettings.cpp
+++ b/Engine/Plugins/Experimental/Avalanche/Source/AvalancheViewport/Private/AvaViewportSettings.cpp
@@ -9,7 +9,7 @@ UAvaViewportSettings::UAvaViewportSettings()
 
 	bEnableViewportOverlay = true;
 	bEnableBoundingBoxes = true;
-	ViewportBackplateMaterial = TSoftObjectPtr<UMaterial>(FSoftObjectPath(TEXT("/Script/Engine.Material'/Avalanche/EditorResources/M_AvalancheBackplate.M_AvalancheBackplate'")));
+	ViewportBackgroundMaterial = TSoftObjectPtr<UMaterial>(FSoftObjectPath(TEXT("/Script/Engine.Material'/Avalanche/EditorResources/M_AvalancheBackground.M_AvalancheBackground'")));
 	ViewportCheckerboardMaterial = TSoftObjectPtr<UMaterial>(FSoftObjectPath(TEXT("/Script/Engine.Material'/Avalanche/EditorResources/M_AvalancheCheckerboard.M_AvalancheCheckerboard'")));
 	ViewportCheckerboardColor0 = FLinearColor(FVector3f(0.048172f));
 	ViewportCheckerboardColor1 = FLinearColor(FVector3f(0.177888f));
diff --git a/Engine/Plugins/Experimental/Avalanche/Source/AvalancheViewport/Private/Visualizers/AvaViewportBackgroundVisualizer.cpp b/Engine/Plugins/Experimental/Avalanche/Source/AvalancheViewport/Private/Visualizers/AvaViewportBackgroundVisualizer.cpp
new file mode 100644
--- /dev/null
+++ b/Engine/Plugins/Experimental/Avalanche/Source/AvalancheViewport/Private/Visualizers/AvaViewportBackgroundVisualizer.cpp
@@ -0,0 +1,206 @@
+// Copyright Epic Games, Inc. All Rights Reserved.
+
+#include "Visualizers/AvaViewportBackgroundVisualizer.h"
+#include "AvaViewportPostProcessManager.h"
+#include "AvaViewportSettings.h"
+#include "AvaViewportUtils.h"
+#include "AvaVisibleArea.h"
+#include "Engine/Texture.h"
+#include "Materials/Material.h"
+#include "Materials/MaterialInstanceDynamic.h"
+#include "UObject/Package.h"
+#include "Viewport/Interaction/AvaViewportPostProcessInfo.h"
+#include "Viewport/Interaction/IAvaViewportDataProvider.h"
+#include "Viewport/Interaction/IAvaViewportDataProxy.h"
+#include "ViewportClient/IAvaViewportClient.h"
+
+#define LOCTEXT_NAMESPACE "AvaViewportBackgroundVisualizer"
+
+namespace UE::AvalancheViewport::Private
+{
+	const FString BackgroundReferencerName = FString(TEXT("AvaViewportBackgroundVisualizer"));
+	const FName TextureObjectName = FName(TEXT("TextureObject"));
+	const FName TextureOffsetName = FName(TEXT("TextureOffset"));
+	const FName TextureScaleName = FName(TEXT("TextureScale"));
+}
+
+FAvaViewportBackgroundVisualizer::FAvaViewportBackgroundVisualizer(TSharedRef<IAvaViewportClient> InAvaViewportClient)
+	: FAvaViewportPostProcessVisualizer(InAvaViewportClient)
+{
+	bRequiresTonemapperSetting = true;
+
+	TextureOffset = FVector::ZeroVector;
+	TextureScale = FVector::ZeroVector;
+
+	const UAvaViewportSettings* ViewportSettings = GetDefault<UAvaViewportSettings>();
+
+	if (!ViewportSettings)
+	{
+		return;
+	}
+
+	UMaterial* BackgroundMaterial = ViewportSettings->ViewportBackgroundMaterial.LoadSynchronous();
+
+	if (!BackgroundMaterial)
+	{
+		return;
+	}
+
+	PostProcessBaseMaterial = BackgroundMaterial;
+	PostProcessMaterial = UMaterialInstanceDynamic::Create(BackgroundMaterial, GetTransientPackage());
+}
+
+UTexture* FAvaViewportBackgroundVisualizer::GetTexture() const
+{
+	return Texture;
+}
+
+void FAvaViewportBackgroundVisualizer::SetTexture(UTexture* InTexture)
+{
+	if (Texture == InTexture)
+	{
+		return;
+	}
+
+	SetTextureInternal(InTexture);
+
+	UpdatePostProcessInfo();
+	UpdatePostProcessMaterial();
+}
+
+void FAvaViewportBackgroundVisualizer::AddReferencedObjects(FReferenceCollector& InCollector)
+{
+	Super::AddReferencedObjects(InCollector);
+
+	if (Texture)
+	{
+		InCollector.AddReferencedObject(Texture);
+	}
+}
+
+FString FAvaViewportBackgroundVisualizer::GetReferencerName() const
+{
+	return UE::AvalancheViewport::Private::BackgroundReferencerName;
+}
+
+void FAvaViewportBackgroundVisualizer::UpdateForViewport(const FAvaVisibleArea& InVisibleArea, const FVector2f& InWidgetSize, 
+	const FVector2f& InCameraOffset)
+{
+	if (FMath::IsNearlyZero(PostProcessOpacity) || !Texture || !PostProcessMaterial)
+	{
+		return;
+	}
+
+	if (!InVisibleArea.IsValid())
+	{
+		return;
+	}
+
+	if (!FAvaViewportUtils::IsValidViewportSize(InWidgetSize))
+	{
+		return;
+	}
+
+	const FVector2f ImageSize = {Texture->GetSurfaceWidth(), Texture->GetSurfaceHeight()};
+
+	if (!FAvaViewportUtils::IsValidViewportSize(ImageSize))
+	{
+		return;
+	}
+
+	using namespace UE::AvalancheViewport::Private;
+
+	const float ImageAspectRatio = ImageSize.X / ImageSize.Y;
+	const float WidgetAspectRatio = InWidgetSize.X / InWidgetSize.Y;
+	const float ViewportAspectRatio = InVisibleArea.AbsoluteSize.X / InVisibleArea.AbsoluteSize.Y;
+	const FVector2f WidgetBasedScale = InVisibleArea.AbsoluteSize / InWidgetSize;
+
+	const FVector2f Scale = WidgetBasedScale / InVisibleArea.GetVisibleAreaFraction();
+
+	if (!FMath::IsNearlyEqual(TextureScale.X, Scale.X)
+		|| !FMath::IsNearlyEqual(TextureScale.Y, Scale.Y))
+	{
+		TextureScale.X = Scale.X;
+		TextureScale.Y = Scale.Y;
+		TextureOffset.Z = 0.f;
+		PostProcessMaterial->SetVectorParameterValue(TextureScaleName, TextureScale);
+	}
+
+	FVector2f Offset = ((InWidgetSize - InVisibleArea.AbsoluteSize) * 0.5f)
+		+ InCameraOffset / InVisibleArea.GetVisibleAreaFraction();
+
+	if (!FMath::IsNearlyEqual(WidgetAspectRatio, ViewportAspectRatio))
+	{
+		if (WidgetAspectRatio > ViewportAspectRatio)
+		{
+			Offset.X -= (Scale.Y - 1) * InVisibleArea.AbsoluteSize.X * 0.5f;
+			Offset.Y -= (Scale.Y - 1) * InVisibleArea.AbsoluteSize.Y * 0.5f;
+		}
+		else
+		{
+			Offset.X -= (Scale.X - 1) * InVisibleArea.AbsoluteSize.X * 0.5f;
+			Offset.Y -= (Scale.X - 1) * InVisibleArea.AbsoluteSize.Y * 0.5f;
+		}
+	}		
+
+	if (!FMath::IsNearlyEqual(TextureOffset.X, Offset.X)
+		|| !FMath::IsNearlyEqual(TextureOffset.Y, Offset.Y))
+	{
+		TextureOffset.X = Offset.X;
+		TextureOffset.Y = Offset.Y;
+		TextureOffset.Z = 0.f;
+		PostProcessMaterial->SetVectorParameterValue(TextureOffsetName, TextureOffset);
+	}
+}
+
+void FAvaViewportBackgroundVisualizer::LoadPostProcessInfo(const FAvaViewportPostProcessInfo& InPostProcessInfo)
+{
+	Super::LoadPostProcessInfo(InPostProcessInfo);
+
+	SetTextureInternal(InPostProcessInfo.Texture.LoadSynchronous());
+}
+
+void FAvaViewportBackgroundVisualizer::UpdatePostProcessInfo(FAvaViewportPostProcessInfo& InPostProcessInfo) const
+{
+	Super::UpdatePostProcessInfo(InPostProcessInfo);
+
+	InPostProcessInfo.Texture = Texture;
+}
+
+void FAvaViewportBackgroundVisualizer::UpdatePostProcessMaterial()
+{
+	if (!PostProcessMaterial)
+	{
+		return;
+	}
+
+	Super::UpdatePostProcessMaterial();
+
+	using namespace UE::AvalancheViewport::Private;
+
+	PostProcessMaterial->SetTextureParameterValue(TextureObjectName, Texture);
+}
+
+bool FAvaViewportBackgroundVisualizer::SetupPostProcessSettings(FPostProcessSettings& InPostProcessSettings) const
+{
+	if (!IsValid(Texture))
+	{
+		return false;
+	}
+
+	return FAvaViewportPostProcessVisualizer::SetupPostProcessSettings(InPostProcessSettings);
+}
+
+void FAvaViewportBackgroundVisualizer::SetTextureInternal(UTexture* InTexture)
+{
+	if (!IsValid(InTexture))
+	{
+		Texture = nullptr;
+	}
+	else
+	{
+		Texture = InTexture;
+	}
+}
+
+#undef LOCTEXT_NAMESPACE
diff --git a/Engine/Plugins/Experimental/Avalanche/Source/AvalancheViewport/Private/Visualizers/AvaViewportBackplateVisualizer.cpp b/Engine/Plugins/Experimental/Avalanche/Source/AvalancheViewport/Private/Visualizers/AvaViewportBackplateVisualizer.cpp
deleted file mode 100644
diff --git a/Engine/Plugins/Experimental/Avalanche/Source/AvalancheViewport/Private/Visualizers/AvaViewportCheckerboardVisualizer.cpp b/Engine/Plugins/Experimental/Avalanche/Source/AvalancheViewport/Private/Visualizers/AvaViewportCheckerboardVisualizer.cpp
--- a/Engine/Plugins/Experimental/Avalanche/Source/AvalancheViewport/Private/Visualizers/AvaViewportCheckerboardVisualizer.cpp
+++ b/Engine/Plugins/Experimental/Avalanche/Source/AvalancheViewport/Private/Visualizers/AvaViewportCheckerboardVisualizer.cpp
@@ -7,7 +7,7 @@
 #include "Materials/MaterialInstanceDynamic.h"
 #include "UObject/Package.h"
 
-#define LOCTEXT_NAMESPACE "AvaViewportBackplateVisualizer"
+#define LOCTEXT_NAMESPACE "AvaViewportBackgroundVisualizer"
 
 namespace UE::AvalancheViewport::Private
 {
diff --git a/Engine/Plugins/Experimental/Avalanche/Source/AvalancheViewport/Public/AvaViewportPostProcessManager.h b/Engine/Plugins/Experimental/Avalanche/Source/AvalancheViewport/Public/AvaViewportPostProcessManager.h
--- a/Engine/Plugins/Experimental/Avalanche/Source/AvalancheViewport/Public/AvaViewportPostProcessManager.h
+++ b/Engine/Plugins/Experimental/Avalanche/Source/AvalancheViewport/Public/AvaViewportPostProcessManager.h
@@ -6,7 +6,7 @@
 #include "Templates/SharedPointer.h"
 #include "AvaViewportPostProcessManager.generated.h"
 
-class FAvaViewportBackplateVisualizer;
+class FAvaViewportBackgroundVisualizer;
 class FAvaViewportChannelVisualizer;
 class FAvaViewportPostProcessVisualizer;
 class FSceneView;
@@ -17,7 +17,7 @@ UENUM()
 enum class EAvaViewportPostProcessType : uint8
 {
 	None,
-	Backplate,
+	Background,
 	RedChannel,
 	GreenChannel,
 	BlueChannel,
diff --git a/Engine/Plugins/Experimental/Avalanche/Source/AvalancheViewport/Public/AvaViewportSettings.h b/Engine/Plugins/Experimental/Avalanche/Source/AvalancheViewport/Public/AvaViewportSettings.h
--- a/Engine/Plugins/Experimental/Avalanche/Source/AvalancheViewport/Public/AvaViewportSettings.h
+++ b/Engine/Plugins/Experimental/Avalanche/Source/AvalancheViewport/Public/AvaViewportSettings.h
@@ -62,7 +62,7 @@ public:
 	bool bEnableBoundingBoxes;
 
 	UPROPERTY(Config, EditAnywhere, BlueprintReadWrite, Category = "Viewport")
-	TSoftObjectPtr<UMaterial> ViewportBackplateMaterial;
+	TSoftObjectPtr<UMaterial> ViewportBackgroundMaterial;
 
 	UPROPERTY(Config, EditAnywhere, BlueprintReadWrite, Category = "Viewport")
 	TSoftObjectPtr<UMaterial> ViewportCheckerboardMaterial;
diff --git a/Engine/Plugins/Experimental/Avalanche/Source/AvalancheViewport/Public/Visualizers/AvaViewportBackgroundVisualizer.h b/Engine/Plugins/Experimental/Avalanche/Source/AvalancheViewport/Public/Visualizers/AvaViewportBackgroundVisualizer.h
new file mode 100644
--- /dev/null
+++ b/Engine/Plugins/Experimental/Avalanche/Source/AvalancheViewport/Public/Visualizers/AvaViewportBackgroundVisualizer.h
@@ -0,0 +1,50 @@
+// Copyright Epic Games, Inc. All Rights Reserved.
+
+#pragma once
+
+#include "Visualizers/AvaViewportPostProcessVisualizer.h"
+#include "Math/Vector.h"
+
+class UTexture;
+struct FAvaViewportPostProcessInfo;
+
+class AVALANCHEVIEWPORT_API FAvaViewportBackgroundVisualizer : public FAvaViewportPostProcessVisualizer
+{
+public:
+	UE_AVA_INHERITS_WITH_SUPER(FAvaViewportBackgroundVisualizer, FAvaViewportPostProcessVisualizer)
+
+	FAvaViewportBackgroundVisualizer(TSharedRef<IAvaViewportClient> InAvaViewportClient);
+	virtual ~FAvaViewportBackgroundVisualizer() override = default;
+
+	UTexture* GetTexture() const;
+	void SetTexture(UTexture* InTexture);
+
+	//~ Begin FAvaViewportPostProcessVisualizer
+	virtual void UpdateForViewport(const FAvaVisibleArea& InVisibleArea, const FVector2f& InWidgetSize, const FVector2f& InCameraOffset) override;
+	//~ End FAvaViewportPostProcessVisualizer
+
+	//~ Begin FGCObject
+	virtual void AddReferencedObjects(FReferenceCollector& InCollector) override;
+	virtual FString GetReferencerName() const override;
+	//~ End FGCObject
+
+protected:
+	TObjectPtr<UTexture> Texture;
+
+	/** Offset from the top-left of the viewport where the texture should start in pixel space. */
+	FVector TextureOffset;
+
+	/** Scale of the texture. Larger values mean a larger texture. */
+	FVector TextureScale;
+
+	using FAvaViewportPostProcessVisualizer::UpdatePostProcessInfo;
+
+	void SetTextureInternal(UTexture* InTexture);
+
+	//~ Begin FAvaViewportPostProcessVisualizer
+	virtual void LoadPostProcessInfo(const FAvaViewportPostProcessInfo& InPostProcessInfo) override;
+	virtual void UpdatePostProcessInfo(FAvaViewportPostProcessInfo& InPostProcessInfo) const override;
+	virtual void UpdatePostProcessMaterial() override;
+	virtual bool SetupPostProcessSettings(FPostProcessSettings& InPostProcessSettings) const override;
+	//~ End FAvaViewportPostProcessVisualizer
+};
\ No newline at end of file
diff --git a/Engine/Plugins/Experimental/Avalanche/Source/AvalancheViewport/Public/Visualizers/AvaViewportBackplateVisualizer.h b/Engine/Plugins/Experimental/Avalanche/Source/AvalancheViewport/Public/Visualizers/AvaViewportBackplateVisualizer.h
deleted file mode 100644
