[Motion Design] Mask EditorMode uses specialized modifiers vs. base class

diff --git a/Engine/Plugins/Experimental/Avalanche/Source/AvalancheMask/Private/Handling/AvaMaskActorMaterialCollectionHandle.cpp b/Engine/Plugins/Experimental/Avalanche/Source/AvalancheMask/Private/Handling/AvaMaskActorMaterialCollectionHandle.cpp
--- a/Engine/Plugins/Experimental/Avalanche/Source/AvalancheMask/Private/Handling/AvaMaskActorMaterialCollectionHandle.cpp
+++ b/Engine/Plugins/Experimental/Avalanche/Source/AvalancheMask/Private/Handling/AvaMaskActorMaterialCollectionHandle.cpp
@@ -370,7 +370,10 @@ FStructView FAvaMaskActorMaterialCollectionHandle::GetMaterialHandleData(
 	if (InParentHandleData)
 	{
 		const FAvaMask2DComponentMaterialPath ComponentMaterialPath{ InComponent, InSlotIdx };
-		return InParentHandleData->ComponentMaterialData[ComponentMaterialPath];
+		if (InParentHandleData->ComponentMaterialData.Contains(ComponentMaterialPath))
+		{
+			return InParentHandleData->ComponentMaterialData[ComponentMaterialPath]; 
+		}
 	}
 
 	return nullptr;
diff --git a/Engine/Plugins/Experimental/Avalanche/Source/AvalancheMask/Private/Handling/AvaMaskDesignedMaterialHandle.cpp b/Engine/Plugins/Experimental/Avalanche/Source/AvalancheMask/Private/Handling/AvaMaskDesignedMaterialHandle.cpp
--- a/Engine/Plugins/Experimental/Avalanche/Source/AvalancheMask/Private/Handling/AvaMaskDesignedMaterialHandle.cpp
+++ b/Engine/Plugins/Experimental/Avalanche/Source/AvalancheMask/Private/Handling/AvaMaskDesignedMaterialHandle.cpp
@@ -80,7 +80,7 @@ bool FAvaMaskDesignedMaterialHandle::SetMaskParameters(
 {
 	if (UDynamicMaterialInstance* MaterialInstance = WeakDesignedMaterial.Get())
 	{
-		UE_LOG(LogAvaMask, Display, TEXT("SetParameters: Texture:%s, Channel:%s"), InTexture ? *InTexture->GetName() : TEXT("(None)"), *UE::AvaMask::Internal::MaskChannelEnumToVector[InChannel].ToString());
+		UE_LOG(LogAvaMask, VeryVerbose, TEXT("SetParameters: Texture:%s, Channel:%s"), InTexture ? *InTexture->GetName() : TEXT("(None)"), *UE::AvaMask::Internal::MaskChannelEnumToVector[InChannel].ToString());
 
 		UTexture* DefaultTexture = nullptr;
 		MaterialInstance->GetTextureParameterDefaultValue(UE::AvaMask::Internal::TextureParameterInfo, DefaultTexture);
diff --git a/Engine/Plugins/Experimental/Avalanche/Source/AvalancheMask/Private/Handling/AvaMaskParametricMaterialHandle.cpp b/Engine/Plugins/Experimental/Avalanche/Source/AvalancheMask/Private/Handling/AvaMaskParametricMaterialHandle.cpp
--- a/Engine/Plugins/Experimental/Avalanche/Source/AvalancheMask/Private/Handling/AvaMaskParametricMaterialHandle.cpp
+++ b/Engine/Plugins/Experimental/Avalanche/Source/AvalancheMask/Private/Handling/AvaMaskParametricMaterialHandle.cpp
@@ -72,7 +72,7 @@ bool FAvaMaskParametricMaterialHandle::SetMaskParameters(
 		FLinearColor PaddingValueV = FLinearColor(FMath::Max(0, InPadding.X), FMath::Max(0, InPadding.Y), 0, 0);
 		FLinearColor FeatherValueV = FLinearColor(bInApplyFeathering ? 1.0f : 0.0f, FMath::Max(0, InOuterFeathering), FMath::Max(0, InInnerFeathering), FMath::Max(0, FMath::Max(InOuterFeathering, InInnerFeathering)));
 		
-		UE_LOG(LogAvaMask, Display, TEXT("SetParameters:\nTexture:%s\nChannel:%s\nFeather:%s\n"),
+		UE_LOG(LogAvaMask, VeryVerbose, TEXT("SetParameters:\nTexture:%s\nChannel:%s\nFeather:%s\n"),
 			InTexture ? *InTexture->GetName() : TEXT("(None)"),
 			*UE::AvaMask::Internal::MaskChannelEnumToVector[InChannel].ToString(),
 			*FeatherValueV.ToString());
diff --git a/Engine/Plugins/Experimental/Avalanche/Source/AvalancheMask/Private/Handling/IAvaMaskMaterialHandle.h b/Engine/Plugins/Experimental/Avalanche/Source/AvalancheMask/Private/Handling/IAvaMaskMaterialHandle.h
--- a/Engine/Plugins/Experimental/Avalanche/Source/AvalancheMask/Private/Handling/IAvaMaskMaterialHandle.h
+++ b/Engine/Plugins/Experimental/Avalanche/Source/AvalancheMask/Private/Handling/IAvaMaskMaterialHandle.h
@@ -123,7 +123,7 @@ bool TAvaMaskMaterialHandle<HandleDataType>::SetMaskParameters(UTexture* InTextu
 {
 	if (UMaterialInstanceDynamic* MaterialInstance = GetMaterialInstance())
 	{
-		UE_LOG(LogAvaMask, Display, TEXT("SetParameters: Texture:%s, Channel:%s"), InTexture ? *InTexture->GetName() : TEXT("(None)"), *UE::AvaMask::Internal::MaskChannelEnumToVector[InChannel].ToString());
+		UE_LOG(LogAvaMask, VeryVerbose, TEXT("SetParameters: Texture:%s, Channel:%s"), InTexture ? *InTexture->GetName() : TEXT("(None)"), *UE::AvaMask::Internal::MaskChannelEnumToVector[InChannel].ToString());
 
 		UTexture* DefaultTexture = nullptr;
 		MaterialInstance->GetTextureParameterDefaultValue(UE::AvaMask::Internal::TextureParameterInfo, DefaultTexture);
@@ -159,5 +159,7 @@ bool TAvaMaskMaterialHandle<HandleDataType>::HasRequiredParameters(TArray<FStrin
 			, OutMissingParameterNames);
 	}
 
+	UE_LOG(LogAvaMask, Warning, TEXT("HasRequiredParameters: Material was invalid."));
+
 	return false;
 }
diff --git a/Engine/Plugins/Experimental/Avalanche/Source/AvalancheMaskEditor/Private/AvaMaskEditorMode.cpp b/Engine/Plugins/Experimental/Avalanche/Source/AvalancheMaskEditor/Private/AvaMaskEditorMode.cpp
--- a/Engine/Plugins/Experimental/Avalanche/Source/AvalancheMaskEditor/Private/AvaMaskEditorMode.cpp
+++ b/Engine/Plugins/Experimental/Avalanche/Source/AvalancheMaskEditor/Private/AvaMaskEditorMode.cpp
@@ -17,6 +17,8 @@
 #include "IGeometryMaskWriteInterface.h"
 #include "LevelEditor.h"
 #include "Mask2D/AvaMask2DBaseModifier.h"
+#include "Mask2D/AvaMask2DReadModifier.h"
+#include "Mask2D/AvaMask2DWriteModifier.h"
 #include "Modifiers/ActorModifierCoreStack.h"
 #include "Modules/ModuleManager.h"
 #include "Selection.h"
@@ -327,7 +329,7 @@ UGeometryMaskCanvas* UAvaMaskEditorMode::GetCanvasReferencedByActor(const AActor
 
 	if (UGeometryMaskCanvas* Canvas = GEngine->GetEngineSubsystem<UGeometryMaskSubsystem>()->GetNamedCanvas(CanvasName))
 	{
-		return Canvas;		
+		return Canvas;
 	}
 
 	return nullptr;
@@ -359,31 +361,26 @@ bool UAvaMaskEditorMode::AddMaskToSelected(const TArray<AActor*>& InMaskingActor
 	
 	AActor* ParentActor = GetActorToParentTo();
 	
-	UAvaMask2DBaseModifier* ParentMaskModifier = FindOrAddMaskModifier(ParentActor);
-	if (!ParentMaskModifier)
+	UAvaMask2DReadModifier* ParentMaskReadModifier = FindOrAddMaskModifier<UAvaMask2DReadModifier>(ParentActor);
+	if (!ParentMaskReadModifier)
 	{
 		return false;
 	}
 
 	const UGeometryMaskCanvas* ParentCanvas = GetCanvasReferencedByActor(ParentActor);
-	const FName ChannelName = ParentMaskModifier->GetChannel();
+	const FName ChannelName = ParentMaskReadModifier->GetChannel();
 	const EGeometryMaskColorChannel ColorChannel = ParentCanvas ? ParentCanvas->GetColorChannel() : EGeometryMaskColorChannel::Red;
 
-	// @todo: replace SetMode
-	// ParentMaskModifier->SetMode(EAvaMask2DMode::Read);
-
 	for (AActor* PlacedActor : InMaskingActors)
 	{
 		// Parent to Actor
 		PlacedActor->AttachToActor(ParentActor, FAttachmentTransformRules::KeepWorldTransform);
 		
 		// Setup modifier
-		if (UAvaMask2DBaseModifier* MaskModifier = FindOrAddMaskModifier(PlacedActor))
+		if (UAvaMask2DWriteModifier* MaskWriteModifier = FindOrAddMaskModifier<UAvaMask2DWriteModifier>(PlacedActor))
 		{
-			// @todo: replace SetMode
-			// MaskModifier->SetMode(EAvaMask2DMode::Write);
-			MaskModifier->SetChannel(ChannelName);
-			MaskModifier->SetUseParentChannel(true);
+			MaskWriteModifier->SetChannel(ChannelName);
+			MaskWriteModifier->SetUseParentChannel(true);
 
 			// Flush unused canvases, in case temporary actors were created/destroyed
 			if (UGeometryMaskSubsystem* Subsystem = GEngine->GetEngineSubsystem<UGeometryMaskSubsystem>())
@@ -399,7 +396,7 @@ bool UAvaMaskEditorMode::AddMaskToSelected(const TArray<AActor*>& InMaskingActor
 	return true;
 }
 
-UAvaMask2DBaseModifier* UAvaMaskEditorMode::FindOrAddMaskModifier(AActor* InActor)
+UAvaMask2DBaseModifier* UAvaMaskEditorMode::FindOrAddMaskModifier(AActor* InActor, const TSubclassOf<UAvaMask2DBaseModifier>& InMaskModifierType)
 {
 	if (const UActorModifierCoreSubsystem* ModifierSubsystem = UActorModifierCoreSubsystem::Get())
 	{
@@ -409,30 +406,33 @@ UAvaMask2DBaseModifier* UAvaMaskEditorMode::FindOrAddMaskModifier(AActor* InActo
 			ModifierStack = ModifierSubsystem->AddActorModifierStack(InActor);
 		}
 
-		const FName MaskModifierName = GetDefault<UAvaMask2DBaseModifier>()->GetModifierName();
-		UAvaMask2DBaseModifier* MaskModifier = nullptr;
-		
+		const FName MaskModifierName = GetDefault<UAvaBaseModifier>(InMaskModifierType)->GetModifierName();
+
 		TArray<UAvaMask2DBaseModifier*> FoundModifiers;
-		ModifierStack->GetClassModifiers<UAvaMask2DBaseModifier>(FoundModifiers);
-		if (!FoundModifiers.IsEmpty())
+		if (UActorModifierCoreBase* FoundModifier = ModifierStack->FindModifier(InMaskModifierType))
 		{
-			MaskModifier = FoundModifiers.Last();
+			return Cast<UAvaMask2DBaseModifier>(FoundModifier);
 		}
-		else
-		{
-			FActorModifierCoreStackInsertOp InsertOp;
-			InsertOp.NewModifierName = MaskModifierName;
+		
+		FActorModifierCoreStackInsertOp InsertOp;
+		InsertOp.NewModifierName = MaskModifierName;
 			
-			MaskModifier = Cast<UAvaMask2DBaseModifier>(ModifierSubsystem->InsertModifier(ModifierStack, InsertOp));
-			if (!MaskModifier)
+		UAvaMask2DBaseModifier* MaskModifier = Cast<UAvaMask2DBaseModifier>(ModifierSubsystem->InsertModifier(ModifierStack, InsertOp));
+		if (!MaskModifier)
+		{
+			if (InsertOp.FailReason)
 			{
-				UE_LOG(LogAvalancheMaskEditor, Error, TEXT("Error inserting Mask modifier."));
+				UE_LOG(LogAvalancheMaskEditor, Error, TEXT("Error inserting Mask modifier: %s"), *InsertOp.FailReason->ToString());
 			}
 			else
 			{
-				MaskModifier->SetChannel(MaskModifier->GenerateUniqueMaskName());
+				UE_LOG(LogAvalancheMaskEditor, Error, TEXT("Error inserting Mask modifier."));
 			}
 		}
+		else
+		{
+			MaskModifier->SetChannel(MaskModifier->GenerateUniqueMaskName());
+		}
 
 		return MaskModifier;
 	}
@@ -446,23 +446,20 @@ bool UAvaMaskEditorMode::CanMaskSelected(AActor* InSelectedActor)
 	{
 		if (const UActorModifierCoreSubsystem* ModifierSubsystem = UActorModifierCoreSubsystem::Get())
 		{
-			/*
-			 * Todo : this will always be none since UAvaMask2DBaseModifier is abstract and not registered as valid modifier,
-			 * should change this to UAvaMask2DReadModifier or UAvaMask2DWriteModifier based on context and avoid using abstract base class
-			 */
-
-			const FName MaskModifierName = ModifierSubsystem->GetRegisteredModifierName(UAvaMask2DBaseModifier::StaticClass());
+			const FName MaskReadModifierName = ModifierSubsystem->GetRegisteredModifierName(UAvaMask2DReadModifier::StaticClass());
 
 			if (const UActorModifierCoreStack* ExistingModifierStack = ModifierSubsystem->GetActorModifierStack(InActor))
 			{
-				if (ExistingModifierStack->ContainsModifier(UAvaMask2DBaseModifier::StaticClass()))
+				// Can't double-mask, and can't apply a mask to this object if it's itself being masked
+				if (ExistingModifierStack->ContainsModifier(UAvaMask2DReadModifier::StaticClass())
+					|| ExistingModifierStack->ContainsModifier(UAvaMask2DWriteModifier::StaticClass()))
 				{
 					return true;
 				}
 			}
- 
+
 			const TSet<FName> AllowedModifiers = ModifierSubsystem->GetAllowedModifiers(InActor);
-			if (AllowedModifiers.Contains(MaskModifierName))
+			if (AllowedModifiers.Contains(MaskReadModifierName))
 			{
 				return true;
 			}
diff --git a/Engine/Plugins/Experimental/Avalanche/Source/AvalancheMaskEditor/Private/AvaMaskEditorMode.h b/Engine/Plugins/Experimental/Avalanche/Source/AvalancheMaskEditor/Private/AvaMaskEditorMode.h
--- a/Engine/Plugins/Experimental/Avalanche/Source/AvalancheMaskEditor/Private/AvaMaskEditorMode.h
+++ b/Engine/Plugins/Experimental/Avalanche/Source/AvalancheMaskEditor/Private/AvaMaskEditorMode.h
@@ -67,7 +67,14 @@ private:
 
 	bool AddMaskToSelected(const TArray<AActor*>& InMaskingActors);
 
-	UAvaMask2DBaseModifier* FindOrAddMaskModifier(AActor* InActor);
+	UAvaMask2DBaseModifier* FindOrAddMaskModifier(AActor* InActor, const TSubclassOf<UAvaMask2DBaseModifier>& InMaskModifierType);
+
+	template <typename MaskModifierType
+		UE_REQUIRES(std::is_base_of_v<UAvaMask2DBaseModifier, MaskModifierType>)>
+	MaskModifierType* FindOrAddMaskModifier(AActor* InActor)
+	{
+		return Cast<MaskModifierType>(FindOrAddMaskModifier(InActor, MaskModifierType::StaticClass()));
+	}
 
 	/** Uses USelection if SelectedActor not specified. */
 	bool CanMaskSelected(AActor* InSelectedActor = nullptr);
