[Motion Design] Moved SAvaMultiComboButton to common location

diff --git a/Engine/Plugins/Experimental/Avalanche/Source/AvalancheEditorCore/Private/Widgets/SAvaMultiComboButton.cpp b/Engine/Plugins/Experimental/Avalanche/Source/AvalancheEditorCore/Private/Widgets/SAvaMultiComboButton.cpp
new file mode 100644
--- /dev/null
+++ b/Engine/Plugins/Experimental/Avalanche/Source/AvalancheEditorCore/Private/Widgets/SAvaMultiComboButton.cpp
@@ -0,0 +1,226 @@
+// Copyright Epic Games, Inc. All Rights Reserved.
+
+/**
+ * This is a copy+paste of the engine SComboButton because it does not support some required
+ * functionality (the ability to use right click to open the menu and left click to activate
+ * the button).
+ */
+
+#include "Widgets/SAvaMultiComboButton.h"
+#include "Widgets/Layout/SBorder.h"
+#include "Widgets/Images/SImage.h"
+#include "Widgets/Input/SButton.h"
+#include "Framework/Application/SlateApplication.h"
+
+void SAvaMultiComboButton::Construct( const FArguments& InArgs )
+{
+	check(InArgs._ComboButtonStyle);
+	Style = InArgs._ComboButtonStyle;
+	// Work out which values we should use based on whether we were given an override, or should use the style's version
+	const FButtonStyle* const OurButtonStyle = InArgs._ButtonStyle ? InArgs._ButtonStyle : &InArgs._ComboButtonStyle->ButtonStyle;
+
+	MenuBorderBrush = &InArgs._ComboButtonStyle->MenuBorderBrush;
+	MenuBorderPadding = InArgs._ComboButtonStyle->MenuBorderPadding;
+	
+	OnComboBoxOpened = InArgs._OnComboBoxOpened;
+	ContentWidgetPtr = InArgs._MenuContent.Widget;
+	bIsFocusable = InArgs._IsFocusable;
+
+	const bool bHasDownArrowShadow = !InArgs._ComboButtonStyle->ShadowOffset.IsZero();
+
+	SMenuAnchor::Construct( SMenuAnchor::FArguments()
+		.Placement(InArgs._MenuPlacement)
+		.Method(InArgs._Method)
+		.OnMenuOpenChanged(InArgs._OnMenuOpenChanged)
+		.OnGetMenuContent(InArgs._OnGetMenuContent)
+		.IsCollapsedByParent(InArgs._CollapseMenuOnParentFocus)
+		[
+			SAssignNew(ButtonPtr, SButton )
+			.ButtonStyle( OurButtonStyle )
+			.ClickMethod( EButtonClickMethod::MouseDown )
+			.OnClicked(InArgs._OnButtonClicked)
+			.ToolTipText( this, &SAvaMultiComboButton::GetFilteredToolTipText, InArgs._ToolTipText)
+			.ContentPadding( InArgs._ContentPadding.IsSet() ? InArgs._ContentPadding : InArgs._ComboButtonStyle->ContentPadding )
+			.ForegroundColor( InArgs._ForegroundColor )
+			.ButtonColorAndOpacity( InArgs._ButtonColorAndOpacity )
+			.IsFocusable( InArgs._IsFocusable )
+			// We set the button to not be accessible as there are issues interacting with the menus that pop up e.g in SComboBox
+			.AccessibleParams(FAccessibleWidgetData(EAccessibleBehavior::NotAccessible, EAccessibleBehavior::NotAccessible, false))
+			[
+				// Button and down arrow on the right
+				// +-------------------+---+
+				// | Button Content    | v |
+				// +-------------------+---+
+				SAssignNew( HBox, SHorizontalBox )
+				+ SHorizontalBox::Slot()
+				.Expose( ButtonContentSlot )
+				.FillWidth( 1 )
+				.HAlign( InArgs._HAlign )
+				.VAlign( InArgs._VAlign )
+				[
+					InArgs._ButtonContent.Widget
+				]
+
+				+SHorizontalBox::Slot()
+				.AutoWidth()
+				.HAlign( HAlign_Right )
+				.VAlign(InArgs._HasDownArrow ? (EVerticalAlignment) InArgs._ComboButtonStyle->DownArrowAlign : VAlign_Center)
+				.Padding(InArgs._HasDownArrow ? InArgs._ComboButtonStyle->DownArrowPadding : FMargin(0))
+				[
+					SNew(SOverlay)
+					// drop shadow
+					+ SOverlay::Slot()
+					.VAlign(VAlign_Top)
+					.Padding(FMargin(InArgs._ComboButtonStyle->ShadowOffset.X, InArgs._ComboButtonStyle->ShadowOffset.Y, 0, 0))
+					[
+						SAssignNew(ShadowImage, SImage)
+						.Visibility( InArgs._HasDownArrow && bHasDownArrowShadow ? EVisibility::Visible : EVisibility::Collapsed )
+						.Image( &InArgs._ComboButtonStyle->DownArrowImage )
+						.ColorAndOpacity( InArgs._ComboButtonStyle->ShadowColorAndOpacity )
+					]
+					+ SOverlay::Slot()
+					.VAlign(VAlign_Top)
+					[
+						SAssignNew(ForegroundArrowImage,SImage)
+						.Visibility( InArgs._HasDownArrow ? EVisibility::Visible : EVisibility::Collapsed )
+						.Image( &InArgs._ComboButtonStyle->DownArrowImage )
+						// Inherit tinting from parent
+						.ColorAndOpacity( FSlateColor::UseForeground() )
+					]
+				]
+			]
+		]
+	);
+	
+	// The menu that pops up when we press the button.
+	// We keep this content around, and then put it into a new window when we need to pop
+	// it up.
+	SetMenuContent( InArgs._MenuContent.Widget );
+}
+
+FText SAvaMultiComboButton::GetFilteredToolTipText(TAttribute<FText> ToolTipText) const
+{
+	if (IsOpen())
+	{
+		return FText::GetEmpty();
+	}
+
+	return ToolTipText.Get();
+}
+
+FReply SAvaMultiComboButton::OnButtonClicked()
+{
+	// Button was clicked; show the popup.
+	// Do nothing if clicking on the button also dismissed the menu, because we will end up doing the same thing twice.
+	// Don't explicitly focus the menu, here, we're going to do it in the button reply so that it's properly focused
+	// to the correct user.
+	SetIsOpen( ShouldOpenDueToClick(), false );
+
+	// If the menu is open, execute the related delegate.
+	if( IsOpen() && OnComboBoxOpened.IsBound() )
+	{
+		OnComboBoxOpened.Execute();
+	}
+
+	// Focusing any newly-created widgets must occur after they have been added to the UI root.
+	FReply ButtonClickedReply = FReply::Handled();
+	
+	// Don't try to focus the menu if the menu is closing
+	if (!IsOpen())
+	{
+		return ButtonClickedReply;
+	}
+
+	TSharedPtr<SWidget> WidgetToFocus = WidgetToFocusPtr.Pin();
+	
+	if (bIsFocusable)
+	{
+		if (!WidgetToFocus.IsValid())
+		{
+			// no explicitly focused widget, try to focus the content that is a child of the border
+			if (WrappedContent.IsValid() && WrappedContent->GetChildren()->Num() == 1)
+			{
+				WidgetToFocus = WrappedContent->GetChildren()->GetChildAt(0);
+			}
+		}
+
+		if (!WidgetToFocus.IsValid())
+		{
+			// no explicitly focused widget, try to focus the content
+			WidgetToFocus = MenuContent;
+		}
+
+		if (!WidgetToFocus.IsValid())
+		{
+			// no content, so try to focus the original widget set on construction
+			WidgetToFocus = ContentWidgetPtr.Pin();
+		}
+	}
+
+	if (WidgetToFocus.IsValid())
+	{
+		ButtonClickedReply.SetUserFocus(WidgetToFocus.ToSharedRef(), EFocusCause::SetDirectly);
+	}
+
+	return ButtonClickedReply;
+}
+
+FReply SAvaMultiComboButton::OnKeyDown(const FGeometry& MyGeometry, const FKeyEvent& InKeyEvent)
+{
+	FReply Reply = FReply::Unhandled();
+
+	if (FSlateApplication::Get().GetNavigationActionFromKey(InKeyEvent) == EUINavigationAction::Accept)
+	{
+		// Handle click button with controller.
+		ButtonPtr->SimulateClick();
+		Reply = FReply::Handled();
+	}
+
+	return Reply;
+}
+
+void SAvaMultiComboButton::SetMenuContent(TSharedRef<SWidget> InContent)
+{
+	WrappedContent = MenuContent =
+		SNew(SBorder)
+		.BorderImage(MenuBorderBrush)
+		.Padding(MenuBorderPadding)
+		[
+			InContent
+		];
+}
+
+void SAvaMultiComboButton::SetOnGetMenuContent(FOnGetContent InOnGetMenuContent)
+{
+	OnGetMenuContent = InOnGetMenuContent;
+}
+
+void SAvaMultiComboButton::SetButtonContentPadding(FMargin InPadding)
+{
+	check(ButtonPtr);
+	ButtonPtr->SetContentPadding(InPadding);
+}
+
+void SAvaMultiComboButton::SetHasDownArrow(bool InHasArrowDown)
+{
+	const bool bHasDownArrowShadow = !Style->ShadowOffset.IsZero();
+
+	check(HBox && HBox->NumSlots() >= 2);
+	HBox->GetSlot(1).SetVerticalAlignment(InHasArrowDown ? (EVerticalAlignment)Style->DownArrowAlign : VAlign_Center);
+	HBox->GetSlot(1).SetPadding(InHasArrowDown ? Style->DownArrowPadding : FMargin(0));
+
+	check(ForegroundArrowImage);
+	ForegroundArrowImage->SetVisibility(InHasArrowDown ? EVisibility::Visible : EVisibility::Collapsed);
+	check(ShadowImage);
+	ShadowImage->SetVisibility(InHasArrowDown && bHasDownArrowShadow ? EVisibility::Visible : EVisibility::Collapsed);
+}
+
+FReply SAvaMultiComboButton::OnMouseButtonDown(const FGeometry& InGeometry, const FPointerEvent& InMouseEvent)
+{
+	if (InMouseEvent.GetEffectingButton() == EKeys::RightMouseButton)
+	{
+		return OnButtonClicked();
+	}
+
+	return SMenuAnchor::OnMouseButtonDown(InGeometry, InMouseEvent);
+}
diff --git a/Engine/Plugins/Experimental/Avalanche/Source/AvalancheEditorCore/Public/Widgets/SAvaMultiComboButton.h b/Engine/Plugins/Experimental/Avalanche/Source/AvalancheEditorCore/Public/Widgets/SAvaMultiComboButton.h
new file mode 100644
--- /dev/null
+++ b/Engine/Plugins/Experimental/Avalanche/Source/AvalancheEditorCore/Public/Widgets/SAvaMultiComboButton.h
@@ -0,0 +1,165 @@
+// Copyright Epic Games, Inc. All Rights Reserved.
+
+/**
+ * This is a copy+paste of the engine SComboButton because it does not support some required
+ * functionality (the ability to use right click to open the menu and left click to activate
+ * the button).
+ */
+
+#pragma once
+
+#include "Layout/Margin.h"
+#include "Styling/SlateColor.h"
+#include "Widgets/DeclarativeSyntaxSupport.h"
+#include "Input/Reply.h"
+#include "Widgets/SWidget.h"
+#include "Styling/SlateTypes.h"
+#include "Styling/AppStyle.h"
+#include "Widgets/SBoxPanel.h"
+#include "Widgets/Input/SComboButton.h"
+#include "Widgets/Input/SMenuAnchor.h"
+
+class SButton;
+class SImage;
+
+/**
+ * A button that, when clicked, brings up a popup.
+ */
+class SAvaMultiComboButton : public SMenuAnchor
+{
+public:
+
+	SLATE_BEGIN_ARGS( SAvaMultiComboButton )
+		: _ComboButtonStyle(&FAppStyle::Get().GetWidgetStyle< FComboButtonStyle >( "ComboButton" ))
+		, _ButtonStyle(nullptr)
+		, _ButtonContent()
+		, _MenuContent()
+		, _IsFocusable(true)
+		, _HasDownArrow(true)
+		, _ForegroundColor(FSlateColor::UseStyle())
+		, _ButtonColorAndOpacity(FLinearColor::White)
+		, _MenuPlacement(MenuPlacement_ComboBox)
+		, _HAlign(HAlign_Fill)
+		, _VAlign(VAlign_Center)
+		, _Method()
+		, _CollapseMenuOnParentFocus(false)
+		{}
+
+		SLATE_STYLE_ARGUMENT( FComboButtonStyle, ComboButtonStyle )
+
+		/** The visual style of the button (overrides ComboButtonStyle) */
+		SLATE_STYLE_ARGUMENT( FButtonStyle, ButtonStyle )
+		
+		SLATE_NAMED_SLOT( FArguments, ButtonContent )
+		
+		/** Optional static menu content.  If the menu content needs to be dynamically built, use the OnGetMenuContent event */
+		SLATE_NAMED_SLOT( FArguments, MenuContent )
+		
+		/** Sets an event handler to generate a widget dynamically when the menu is needed. */
+		SLATE_EVENT( FOnGetContent, OnGetMenuContent )
+		SLATE_EVENT( FOnIsOpenChanged, OnMenuOpenChanged )
+		
+		SLATE_EVENT( FOnComboBoxOpened, OnComboBoxOpened )
+
+		SLATE_EVENT(FOnClicked, OnButtonClicked)
+
+		SLATE_ARGUMENT( bool, IsFocusable )
+		SLATE_ARGUMENT( bool, HasDownArrow )
+
+		SLATE_ATTRIBUTE( FSlateColor, ForegroundColor )
+		SLATE_ATTRIBUTE( FSlateColor, ButtonColorAndOpacity )
+		SLATE_ATTRIBUTE( FMargin, ContentPadding )
+		SLATE_ATTRIBUTE( EMenuPlacement, MenuPlacement )
+		SLATE_ARGUMENT( EHorizontalAlignment, HAlign )
+		SLATE_ARGUMENT( EVerticalAlignment, VAlign )
+
+		/** Spawn a new window or reuse current window for this combo*/
+		SLATE_ARGUMENT( TOptional<EPopupMethod>, Method )
+
+		/** True if this combo's menu should be collapsed when our parent receives focus, false (default) otherwise */
+		SLATE_ARGUMENT(bool, CollapseMenuOnParentFocus)
+
+	SLATE_END_ARGS()
+
+	// SMenuAnchor interface
+	AVALANCHEEDITORCORE_API virtual void SetMenuContent(TSharedRef<SWidget> InContent) override;
+	// End of SMenuAnchor interface
+
+	/** See the OnGetMenuContent event */
+	AVALANCHEEDITORCORE_API void SetOnGetMenuContent( FOnGetContent InOnGetMenuContent );
+
+	/**
+	 * Construct the widget from a declaration
+	 *
+	 * @param InArgs  The declaration from which to construct
+	 */
+	AVALANCHEEDITORCORE_API void Construct(const FArguments& InArgs);
+
+	void SetMenuContentWidgetToFocus( TWeakPtr<SWidget> InWidgetToFocusPtr )
+	{
+		WidgetToFocusPtr = InWidgetToFocusPtr;
+	}
+
+	/** See the padding for button content. */
+	AVALANCHEEDITORCORE_API void SetButtonContentPadding(FMargin InPadding);
+
+	/** add/remove the expanding arrow. */
+	AVALANCHEEDITORCORE_API void SetHasDownArrow(bool InHasArrowDown);
+
+	//~ Begin SWidget
+	AVALANCHEEDITORCORE_API virtual FReply OnMouseButtonDown(const FGeometry& InGeometry, const FPointerEvent& InMouseEvent) override;
+	//~ End SWidget
+
+protected:
+	/**
+	 * Handle the button being clicked by summoning the ComboButton.
+	 */
+	AVALANCHEEDITORCORE_API virtual FReply OnButtonClicked();
+	AVALANCHEEDITORCORE_API virtual FReply OnKeyDown(const FGeometry& MyGeometry, const FKeyEvent& InKeyEvent) override;
+
+	/**
+	 * Called to query the tool tip text for this widget, but will return an empty text when the menu is already open
+	 *
+	 * @param	ToolTipText	Tool tip text to display, if possible
+	 *
+	 * @return	Tool tip text, or an empty text if filtered out
+	 */
+	AVALANCHEEDITORCORE_API FText GetFilteredToolTipText(TAttribute<FText> ToolTipText) const;
+
+protected:
+	/** Area where the button's content resides */
+	SHorizontalBox::FSlot* ButtonContentSlot;
+
+	/** Delegate to execute when the combo list is opened */
+	FOnComboBoxOpened OnComboBoxOpened;
+
+	TWeakPtr<SWidget> WidgetToFocusPtr;
+
+	/** Brush to use to add a "menu border" around the drop-down content */
+	const FSlateBrush* MenuBorderBrush;
+
+	/** Padding to use to add a "menu border" around the drop-down content */
+	FMargin MenuBorderPadding;
+
+	/** The content widget, if any, set by the user on creation */
+	TWeakPtr<SWidget> ContentWidgetPtr;
+
+	/** Can this button be focused? */
+	bool bIsFocusable;
+
+private:
+	/** The button widget within the ComboButton */
+	TSharedPtr<SButton> ButtonPtr;
+
+	/** The box containing the arrow */
+	TSharedPtr<SHorizontalBox> HBox;
+
+	/** The arrow image shadow */
+	TSharedPtr<SImage> ShadowImage;
+
+	/** The arrow image */
+	TSharedPtr<SImage> ForegroundArrowImage;
+
+	/** The ComboButton style */
+	const FComboButtonStyle* Style;
+};
diff --git a/Engine/Plugins/Experimental/Avalanche/Source/AvalancheLevelViewport/Private/Widgets/Engine/SAvaMultiComboButton.cpp b/Engine/Plugins/Experimental/Avalanche/Source/AvalancheLevelViewport/Private/Widgets/Engine/SAvaMultiComboButton.cpp
deleted file mode 100644
diff --git a/Engine/Plugins/Experimental/Avalanche/Source/AvalancheLevelViewport/Private/Widgets/Engine/SAvaMultiComboButton.h b/Engine/Plugins/Experimental/Avalanche/Source/AvalancheLevelViewport/Private/Widgets/Engine/SAvaMultiComboButton.h
deleted file mode 100644
diff --git a/Engine/Plugins/Experimental/Avalanche/Source/AvalancheLevelViewport/Private/Widgets/SAvaLevelViewportStatusBar.h b/Engine/Plugins/Experimental/Avalanche/Source/AvalancheLevelViewport/Private/Widgets/SAvaLevelViewportStatusBar.h
--- a/Engine/Plugins/Experimental/Avalanche/Source/AvalancheLevelViewport/Private/Widgets/SAvaLevelViewportStatusBar.h
+++ b/Engine/Plugins/Experimental/Avalanche/Source/AvalancheLevelViewport/Private/Widgets/SAvaLevelViewportStatusBar.h
@@ -9,7 +9,7 @@
 #include "Styling/StyleColors.h"
 #include "Templates/SharedPointer.h"
 #include "Widgets/SCompoundWidget.h"
-#include "Widgets/Engine/SAvaMultiComboButton.h"
+#include "Widgets/SAvaMultiComboButton.h"
 #include "Widgets/Images/SImage.h"
 #include "Widgets/Input/SButton.h"
 #include "Widgets/Input/SComboButton.h"
