Motion Design:
- Changed Null default spawn location to be outside of view to not disturb default scene with default billboards

#changelist validated
#virtualized

diff --git a/Engine/Plugins/Experimental/Avalanche/Source/AvalancheEditor/Private/Scene/AvaSceneDefaults.cpp b/Engine/Plugins/Experimental/Avalanche/Source/AvalancheEditor/Private/Scene/AvaSceneDefaults.cpp
--- a/Engine/Plugins/Experimental/Avalanche/Source/AvalancheEditor/Private/Scene/AvaSceneDefaults.cpp
+++ b/Engine/Plugins/Experimental/Avalanche/Source/AvalancheEditor/Private/Scene/AvaSceneDefaults.cpp
@@ -118,7 +118,7 @@ namespace UE::AvalancheEditor::Private
 			return;
 		}
 
-		InActor->SetActorTransform(FTransform::Identity);
+		InActor->SetActorLocation(FVector(-500, 0, 250));
 		InActor->SetActorLabel(TEXT("Default Scene"));
 	}
 
