[Motion Design] Mask statusbar has context menu

diff --git a/Engine/Plugins/Experimental/Avalanche/Source/AvalancheMaskEditor/AvalancheMaskEditor.Build.cs b/Engine/Plugins/Experimental/Avalanche/Source/AvalancheMaskEditor/AvalancheMaskEditor.Build.cs
--- a/Engine/Plugins/Experimental/Avalanche/Source/AvalancheMaskEditor/AvalancheMaskEditor.Build.cs
+++ b/Engine/Plugins/Experimental/Avalanche/Source/AvalancheMaskEditor/AvalancheMaskEditor.Build.cs
@@ -17,6 +17,7 @@ public class AvalancheMaskEditor : ModuleRules
 			new string[] {
 				"ActorModifierCore",
 				"Avalanche",
+				"AvalancheEditorCore",
 				"AvalancheMask",
 				"AvalancheModifiers",
 				"AvalancheShapes",
diff --git a/Engine/Plugins/Experimental/Avalanche/Source/AvalancheMaskEditor/Private/AvaMaskEditorCommands.cpp b/Engine/Plugins/Experimental/Avalanche/Source/AvalancheMaskEditor/Private/AvaMaskEditorCommands.cpp
--- a/Engine/Plugins/Experimental/Avalanche/Source/AvalancheMaskEditor/Private/AvaMaskEditorCommands.cpp
+++ b/Engine/Plugins/Experimental/Avalanche/Source/AvalancheMaskEditor/Private/AvaMaskEditorCommands.cpp
@@ -6,6 +6,13 @@
 
 void FAvaMaskEditorCommands::RegisterCommands()
 {
+	UI_COMMAND(
+		ShowVisualizeMasks,
+		"ShowVisualizeMasks",
+		"Visualize Masks",
+		EUserInterfaceActionType::Button,
+		FInputChord());
+	
 	UI_COMMAND(
 		ToggleMaskMode,
 		"ToggleMaskMode",
diff --git a/Engine/Plugins/Experimental/Avalanche/Source/AvalancheMaskEditor/Private/AvaMaskEditorCommands.h b/Engine/Plugins/Experimental/Avalanche/Source/AvalancheMaskEditor/Private/AvaMaskEditorCommands.h
--- a/Engine/Plugins/Experimental/Avalanche/Source/AvalancheMaskEditor/Private/AvaMaskEditorCommands.h
+++ b/Engine/Plugins/Experimental/Avalanche/Source/AvalancheMaskEditor/Private/AvaMaskEditorCommands.h
@@ -22,6 +22,7 @@ public:
 	virtual void RegisterCommands() override;
 
 public:
+	TSharedPtr<FUICommandInfo> ShowVisualizeMasks;
 	TSharedPtr<FUICommandInfo> ToggleMaskMode;
 	TSharedPtr<FUICommandInfo> ToggleShowAllMasks;
 	TSharedPtr<FUICommandInfo> ToggleIsolateMask;
diff --git a/Engine/Plugins/Experimental/Avalanche/Source/AvalancheMaskEditor/Private/AvaMaskEditorModule.cpp b/Engine/Plugins/Experimental/Avalanche/Source/AvalancheMaskEditor/Private/AvaMaskEditorModule.cpp
--- a/Engine/Plugins/Experimental/Avalanche/Source/AvalancheMaskEditor/Private/AvaMaskEditorModule.cpp
+++ b/Engine/Plugins/Experimental/Avalanche/Source/AvalancheMaskEditor/Private/AvaMaskEditorModule.cpp
@@ -14,6 +14,8 @@
 #include "ToolMenuEntry.h"
 #include "ToolMenus.h"
 #include "ToolMenuSection.h"
+#include "Widgets/Images/SImage.h"
+#include "Widgets/SAvaMultiComboButton.h"
 
 #define LOCTEXT_NAMESPACE "AvalancheMaskEditor"
 
@@ -66,6 +68,39 @@ void FAvalancheMaskEditorModule::ToggleEditorMode()
 	GLevelEditorModeTools().ActivateMode(UAvaMaskEditorMode::EM_MotionDesignMaskEditorModeId, true);
 }
 
+TSharedRef<SWidget> FAvalancheMaskEditorModule::GetStatusBarWidgetMenuContent()
+{
+	UToolMenus* ToolMenus = UToolMenus::Get();
+	
+	static const FName MenuName = TEXT("AvaMaskEditor.StatusBar");
+	UToolMenu* ContextMenu = ToolMenus->FindMenu(MenuName);
+	if (!ContextMenu)
+	{
+		ContextMenu = ToolMenus->RegisterMenu(MenuName, NAME_None, EMultiBoxType::Menu);
+		FToolMenuSection& Section = ContextMenu->AddSection(TEXT("Mask"));
+		Section.AddEntry(FToolMenuEntry::InitMenuEntryWithCommandList(
+			FAvaMaskEditorCommands::Get().ShowVisualizeMasks,
+			GetCommandList(),
+			LOCTEXT("ShowVisualizeMasks", "Visualize Masks")));
+	}
+
+	if (!ContextMenu)
+	{
+		return SNullWidget::NullWidget;
+	}
+
+	return ToolMenus->GenerateWidget(ContextMenu);
+}
+
+FReply FAvalancheMaskEditorModule::OnToggleMaskModeClicked()
+{
+	check(CommandList.IsValid());
+	
+	return GetCommandList()->ExecuteAction(FAvaMaskEditorCommands::Get().ToggleMaskMode.ToSharedRef())
+		? FReply::Handled()
+		: FReply::Unhandled();
+}
+
 void FAvalancheMaskEditorModule::RegisterMenus()
 {
 	// Owner will be used for cleanup in call to UToolMenus::UnregisterOwner
@@ -94,11 +129,39 @@ void FAvalancheMaskEditorModule::RegisterMenus()
 		{
 			FToolMenuSection& Section = ToolbarMenu->FindOrAddSection("ModeToggles");
 			{
-				FToolMenuEntry& Entry = Section.AddEntry(FToolMenuEntry::InitToolBarButton(FAvaMaskEditorCommands::Get().ToggleMaskMode));				
-				Entry.Label.Set(FText::GetEmpty());
-				Entry.Icon.Set(FSlateIcon(FAvaMaskEditorStyle::Get().GetStyleSetName(), TEXT("AvaMaskEditor.ToggleMaskMode.Small")));
-				
-				Entry.SetCommandList(CommandList);
+				static const FButtonStyle* Style = &FAppStyle::Get().GetWidgetStyle<FButtonStyle>("SimpleButton");
+				static const FMargin Margin = FMargin(1.0f, 2.0f, 1.0f, 2.0f);
+				static const FVector2D ImageSize = FVector2D(16.0f, 16.0f);
+
+				TSharedRef<SWidget> MaskStatusBarWidget =
+					SNew(SAvaMultiComboButton)
+						.ButtonStyle(Style)
+						.ContentPadding(Margin)
+						.ToolTipText(FAvaMaskEditorCommands::Get().ToggleMaskMode->GetDescription())
+						.HasDownArrow(false)
+						.OnGetMenuContent(FOnGetContent::CreateRaw(this, &FAvalancheMaskEditorModule::GetStatusBarWidgetMenuContent))
+						.OnButtonClicked(FOnClicked::CreateRaw(this, &FAvalancheMaskEditorModule::OnToggleMaskModeClicked))
+						.ButtonContent()
+						[
+							SNew(SOverlay)
+							+ SOverlay::Slot()
+							.Padding(0.f, 0.f, ImageSize.X - 1.f, 0.f)
+							[
+								SNew(SImage)
+								.Image(FSlateIcon(FAvaMaskEditorStyle::Get().GetStyleSetName(), TEXT("AvaMaskEditor.ToggleMaskMode.Small")).GetIcon())
+								.DesiredSizeOverride(ImageSize)
+							]
+							+ SOverlay::Slot()
+							.Padding(ImageSize.X - 1.f, 0.f, 0.f, 0.f)
+							[
+								SNew(SImage)
+								.Image(FAppStyle::Get().GetBrush("Icons.ChevronDown"))
+								.DesiredSizeOverride(ImageSize)
+								.ColorAndOpacity(FSlateColor(FLinearColor(0.5f, 0.5f, 0.5f, 1.f)))
+							]
+						];
+
+				Section.AddEntry(FToolMenuEntry::InitWidget(TEXT("AvaMaskEditor.StatusBar.Toggle"), MaskStatusBarWidget, FText::GetEmpty(), true));
 			}
 		}
 	}
diff --git a/Engine/Plugins/Experimental/Avalanche/Source/AvalancheMaskEditor/Private/AvaMaskEditorModule.h b/Engine/Plugins/Experimental/Avalanche/Source/AvalancheMaskEditor/Private/AvaMaskEditorModule.h
--- a/Engine/Plugins/Experimental/Avalanche/Source/AvalancheMaskEditor/Private/AvaMaskEditorModule.h
+++ b/Engine/Plugins/Experimental/Avalanche/Source/AvalancheMaskEditor/Private/AvaMaskEditorModule.h
@@ -4,8 +4,11 @@
 
 #include "Modules/ModuleInterface.h"
 
-#include "Templates/SharedPointerFwd.h"
 #include "Framework/Commands/UICommandList.h"
+#include "Input/Reply.h"
+#include "Templates/SharedPointerFwd.h"
+
+class SWidget;
 
 class FAvalancheMaskEditorModule
 	: public IModuleInterface
@@ -22,6 +25,9 @@ private:
 	void RegisterMenus();	
 	void ToggleEditorMode();
 
+	TSharedRef<SWidget> GetStatusBarWidgetMenuContent();
+	FReply OnToggleMaskModeClicked();
+
 private:
 	TSharedPtr<FUICommandList> CommandList;
 };
diff --git a/Engine/Plugins/Experimental/Avalanche/Source/AvalancheMaskEditor/Private/AvaMaskEditorSubsystem.cpp b/Engine/Plugins/Experimental/Avalanche/Source/AvalancheMaskEditor/Private/AvaMaskEditorSubsystem.cpp
--- a/Engine/Plugins/Experimental/Avalanche/Source/AvalancheMaskEditor/Private/AvaMaskEditorSubsystem.cpp
+++ b/Engine/Plugins/Experimental/Avalanche/Source/AvalancheMaskEditor/Private/AvaMaskEditorSubsystem.cpp
@@ -25,7 +25,12 @@ void UAvaMaskEditorSubsystem::Initialize(FSubsystemCollectionBase& Collection)
 {
 	/* Commands for this subsystem */
 	const TSharedPtr<FUICommandList> CommandList = FModuleManager::Get().LoadModuleChecked<FAvalancheMaskEditorModule>(UE_MODULE_NAME).GetCommandList();
-	
+
+	CommandList->MapAction(
+		FAvaMaskEditorCommands::Get().ShowVisualizeMasks,
+		FExecuteAction::CreateUObject(this, &UAvaMaskEditorSubsystem::ShowVisualizeMasks),
+		FCanExecuteAction());
+
 	CommandList->MapAction(
 		FAvaMaskEditorCommands::Get().ToggleMaskMode,
 		FExecuteAction::CreateUObject(this, &UAvaMaskEditorSubsystem::ToggleEditorMode),
@@ -62,6 +67,12 @@ void UAvaMaskEditorSubsystem::Deinitialize()
 	Super::Deinitialize();
 }
 
+void UAvaMaskEditorSubsystem::ShowVisualizeMasks()
+{
+	static constexpr const TCHAR* Cmd = TEXT("GeometryMask.Visualize");
+	GEngine->Exec(nullptr, Cmd);
+}
+
 void UAvaMaskEditorSubsystem::ToggleEditorMode()
 {
 	GLevelEditorModeTools().ActivateMode(UAvaMaskEditorMode::EM_MotionDesignMaskEditorModeId, true);
diff --git a/Engine/Plugins/Experimental/Avalanche/Source/AvalancheMaskEditor/Private/AvaMaskEditorSubsystem.h b/Engine/Plugins/Experimental/Avalanche/Source/AvalancheMaskEditor/Private/AvaMaskEditorSubsystem.h
--- a/Engine/Plugins/Experimental/Avalanche/Source/AvalancheMaskEditor/Private/AvaMaskEditorSubsystem.h
+++ b/Engine/Plugins/Experimental/Avalanche/Source/AvalancheMaskEditor/Private/AvaMaskEditorSubsystem.h
@@ -42,6 +42,8 @@ public:
 	// ~End UEditorSubsystem
 
 private:
+	void ShowVisualizeMasks();
+	
 	void ToggleEditorMode();
 	bool IsEditorModeActive();
 
