MotionDesign : TranslucentPriorityModifier
- Enabling include children by default

#changelist validated
#virtualized

diff --git a/Engine/Plugins/Experimental/Avalanche/Source/AvalancheModifiers/Public/Modifiers/AvaTranslucentPriorityModifier.h b/Engine/Plugins/Experimental/Avalanche/Source/AvalancheModifiers/Public/Modifiers/AvaTranslucentPriorityModifier.h
--- a/Engine/Plugins/Experimental/Avalanche/Source/AvalancheModifiers/Public/Modifiers/AvaTranslucentPriorityModifier.h
+++ b/Engine/Plugins/Experimental/Avalanche/Source/AvalancheModifiers/Public/Modifiers/AvaTranslucentPriorityModifier.h
@@ -130,7 +130,7 @@ protected:
 
 	/** If true, will include children too and update their sort priority */
 	UPROPERTY(EditAnywhere, BlueprintReadWrite, Setter="SetIncludeChildren", Getter="GetIncludeChildren", Category="TranslucentPriority", meta=(AllowPrivateAccess="true"))
-	bool bIncludeChildren = false;
+	bool bIncludeChildren = true;
 
 	/** The components this modifier is managing */
 	UPROPERTY()
