[Motion Design] Privatizing/Internalizing pass on AvalancheMedia module: Broadcast.

#rnx

diff --git a/Engine/Plugins/Experimental/Avalanche/Source/AvalancheMedia/Internal/AvaMediaMessageUtils.h b/Engine/Plugins/Experimental/Avalanche/Source/AvalancheMedia/Internal/AvaMediaMessageUtils.h
new file mode 100644
--- /dev/null
+++ b/Engine/Plugins/Experimental/Avalanche/Source/AvalancheMedia/Internal/AvaMediaMessageUtils.h
@@ -0,0 +1,15 @@
+// Copyright Epic Games, Inc. All Rights Reserved.
+
+#pragma once
+
+#include "HAL/Platform.h"
+
+namespace UE::AvaMediaMessageUtils
+{
+	/**
+	 * Some of the message bus bridge have limitations to the size of the messages.
+	 * This function returns the current largest message size that can "safely" be sent
+	 * through the message bus without causing (known) issues.
+	 */
+	AVALANCHEMEDIA_API uint32 GetSafeMessageSizeLimit();
+}
diff --git a/Engine/Plugins/Experimental/Avalanche/Source/AvalancheMedia/Internal/AvaMediaRenderTargetUtils.h b/Engine/Plugins/Experimental/Avalanche/Source/AvalancheMedia/Internal/AvaMediaRenderTargetUtils.h
new file mode 100644
--- /dev/null
+++ b/Engine/Plugins/Experimental/Avalanche/Source/AvalancheMedia/Internal/AvaMediaRenderTargetUtils.h
@@ -0,0 +1,48 @@
+// Copyright Epic Games, Inc. All Rights Reserved.
+
+#pragma once
+
+#include "HAL/Platform.h"
+#include "Math/MathFwd.h"
+#include "PixelFormat.h"
+
+class FName;
+class UTextureRenderTarget2D;
+struct FLinearColor;
+enum EPixelFormat : uint8;
+
+namespace UE::AvaMediaRenderTargetUtils
+{
+	/** Creates a Render Target with default settings like bForceLinearGamma to be false. */
+	AVALANCHEMEDIA_API UTextureRenderTarget2D* CreateDefaultRenderTarget(FName InBaseName);
+
+	/** Update the render target with the given parameters. */
+	AVALANCHEMEDIA_API void UpdateRenderTarget(UTextureRenderTarget2D* InRenderTarget, const FIntPoint& InSize, EPixelFormat InFormat, const FLinearColor& InClearColor);
+	
+	/**
+	 * Setup the render target gamma appropriately depending on the format.
+	 * @return true if the TargetGamma was modified, false otherwise.
+	 */
+	AVALANCHEMEDIA_API bool SetupRenderTargetGamma(UTextureRenderTarget2D* InRenderTarget);
+	
+	/** Returns the size of the render target. */ 
+	AVALANCHEMEDIA_API FIntPoint GetRenderTargetSize(const UTextureRenderTarget2D* InRenderTarget);
+
+	/** Returns true if render target's format is floating point. */
+	AVALANCHEMEDIA_API bool IsFloatFormat(const UTextureRenderTarget2D* InRenderTarget);
+
+	/** Returns the number of color components for the given pixel format. */
+	AVALANCHEMEDIA_API int32 GetNumColorComponents(EPixelFormat InFormat);
+
+	/** Returns true if the given format has an alpha channel. */
+	AVALANCHEMEDIA_API bool HasAlpha(EPixelFormat InFormat);
+
+	/**
+	 * Returns the number of bits per color component (assuming all the same) for the given pixel format.
+	 * Known exceptions: PF_R5G6B5_UNORM, PF_FloatR11G11B10. But those are not render target formats.
+	 */
+	AVALANCHEMEDIA_API int32 GetNumColorChannelBits(EPixelFormat InFormat);
+
+	/** Returns the number of bits for the alpha channel. */
+	AVALANCHEMEDIA_API int32 GetNumAlphaChannelBits(EPixelFormat InFormat);
+}
diff --git a/Engine/Plugins/Experimental/Avalanche/Source/AvalancheMedia/Internal/AvaMediaSerializationUtils.h b/Engine/Plugins/Experimental/Avalanche/Source/AvalancheMedia/Internal/AvaMediaSerializationUtils.h
new file mode 100644
--- /dev/null
+++ b/Engine/Plugins/Experimental/Avalanche/Source/AvalancheMedia/Internal/AvaMediaSerializationUtils.h
@@ -0,0 +1,37 @@
+// Copyright Epic Games, Inc. All Rights Reserved.
+
+#pragma once
+
+#include "Containers/ArrayView.h"
+#include "Containers/UnrealString.h"
+#include "Templates/Function.h"
+
+class FProperty;
+class FStructuredArchiveFormatter;
+class UObject;
+
+namespace UE::AvaMediaSerializationUtils
+{
+	/** Serialize the given object with the given formatter. */
+	AVALANCHEMEDIA_API void SerializeObject(FStructuredArchiveFormatter& Formatter, UObject* InObject, const TFunction<bool(FProperty*)>& InShouldSkipProperty = TFunction<bool(FProperty*)>());
+
+	/**
+	 * The following functions are meant to be used to convert the serialized values of FJsonStruct(De)SerializerBackend
+	 * to/from FString.
+	 * 
+	 * @remark FJsonStructSerializerBackend and FJsonStructDeserializerBackend are currently hardcoded to use UCS2CHAR.
+	 */
+	namespace JsonValueConversion
+	{
+		/** Converts a serialized value stored as raw bytes into a string. In place version. */
+		AVALANCHEMEDIA_API void BytesToString(const TArray<uint8>& InValueAsBytes, FString& OutValueAsString);
+
+		/** This ArrayView version is meant to work with FMemoryReaderView. */
+		inline TArrayView<const uint8> ValueToConstBytesView(const FString& InValueAsString)
+		{
+			// Len() excludes the terminating character, and it is desired so it matches
+			// what FJsonStructSerializerBackend does.
+			return TArrayView<const uint8>(reinterpret_cast<const uint8*>(*InValueAsString), InValueAsString.Len() * sizeof(TCHAR));
+		}
+	}
+}
diff --git a/Engine/Plugins/Experimental/Avalanche/Source/AvalancheMedia/Internal/Broadcast/AvaBroadcastDevice.h b/Engine/Plugins/Experimental/Avalanche/Source/AvalancheMedia/Internal/Broadcast/AvaBroadcastDevice.h
new file mode 100644
--- /dev/null
+++ b/Engine/Plugins/Experimental/Avalanche/Source/AvalancheMedia/Internal/Broadcast/AvaBroadcastDevice.h
@@ -0,0 +1,62 @@
+// Copyright Epic Games, Inc. All Rights Reserved.
+
+#pragma once
+
+#include "AvaBroadcastDevice.h"
+#include "Containers/UnrealString.h"
+#include "MediaIOCoreDefinitions.h"
+#include "Internationalization/Text.h"
+#include "UObject/NameTypes.h"
+
+#define LOCTEXT_NAMESPACE "AvaBroadcastDevice"
+
+struct FAvaBroadcastDevice
+{
+	FAvaBroadcastDevice(const FMediaIODevice& InDevice, const FName& InDeviceProviderName, const FString& InServerName)
+		: Device(InDevice)
+		, DeviceProviderName(InDeviceProviderName)
+		, ServerName(InServerName)
+	{
+	}
+	
+	bool IsValid() const
+	{
+		return Device.DeviceName != NAME_None && Device.IsValid();
+	}
+	
+	bool operator==(const FAvaBroadcastDevice& Other) const
+	{
+		return Device.DeviceName == Other.Device.DeviceName
+			&& Device.DeviceIdentifier == Other.Device.DeviceIdentifier
+			&& ServerName == Other.ServerName;
+	}
+	
+	FText GetDisplayNameText() const
+	{
+		return FText::Format(LOCTEXT("AvaMediaDeviceDisplayName", "{0}")
+			, FText::FromName(Device.DeviceName));
+	}
+
+	friend uint32 GetTypeHash(const FAvaBroadcastDevice& InDevice)
+	{
+		uint32 Hash = GetTypeHash(InDevice.Device.DeviceName);
+		Hash = HashCombine(Hash, InDevice.Device.DeviceIdentifier);
+		Hash = HashCombine(Hash, GetTypeHash(InDevice.ServerName));
+		return Hash;
+	}
+
+	const FMediaIODevice& GetDevice() const { return Device; }
+	const FName& GetDeviceProviderName() const { return DeviceProviderName;}
+	const FString& GetServerName() const { return ServerName;}
+	
+protected:
+	FMediaIODevice Device;
+
+	/** Name of the device provider for this device. */
+	FName DeviceProviderName;
+	
+	/** Name of the server hosting this device. */
+	FString ServerName;
+};
+
+#undef LOCTEXT_NAMESPACE
diff --git a/Engine/Plugins/Experimental/Avalanche/Source/AvalancheMedia/Internal/Broadcast/OutputDevices/AvaBroadcastDeviceProviderData.h b/Engine/Plugins/Experimental/Avalanche/Source/AvalancheMedia/Internal/Broadcast/OutputDevices/AvaBroadcastDeviceProviderData.h
new file mode 100644
--- /dev/null
+++ b/Engine/Plugins/Experimental/Avalanche/Source/AvalancheMedia/Internal/Broadcast/OutputDevices/AvaBroadcastDeviceProviderData.h
@@ -0,0 +1,89 @@
+// Copyright Epic Games, Inc. All Rights Reserved.
+
+#pragma once
+
+#include "IMediaIOCoreDeviceProvider.h"
+#include "AvaBroadcastDeviceProviderData.generated.h"
+
+USTRUCT()
+struct AVALANCHEMEDIA_API FAvaBroadcastDeviceProviderData
+{
+	GENERATED_BODY()
+
+public:
+	UPROPERTY()
+	FName Name;
+
+	UPROPERTY()
+	TArray<FMediaIOConnection> Connections;
+
+	UPROPERTY()
+	TArray<FMediaIOConfiguration> Configurations;
+
+	UPROPERTY()
+	TArray<FMediaIODevice> Devices;
+
+	UPROPERTY()
+	TArray<FMediaIOMode> Modes;
+
+	UPROPERTY()
+	TArray<FMediaIOInputConfiguration> InputConfigurations;
+
+	UPROPERTY()
+	TArray<FMediaIOOutputConfiguration> OutputConfigurations;
+
+	UPROPERTY()
+	TArray<FMediaIOVideoTimecodeConfiguration> TimecodeConfigurations;
+	
+	UPROPERTY()
+	FMediaIOConfiguration DefaultConfiguration;
+
+	UPROPERTY()
+	FMediaIOMode DefaultMode;
+
+	UPROPERTY()
+	FMediaIOInputConfiguration DefaultInputConfiguration;
+
+	UPROPERTY()
+	FMediaIOOutputConfiguration DefaultOutputConfiguration;
+
+	UPROPERTY()
+	FMediaIOVideoTimecodeConfiguration DefaultTimecodeConfiguration;
+	
+	UPROPERTY()
+	bool bShowInputTransportInSelector = true;
+
+	UPROPERTY()
+	bool bShowOutputTransportInSelector = true;
+
+	UPROPERTY()
+	bool bShowInputKeyInSelector = true;
+
+	UPROPERTY()
+	bool bShowOutputKeyInSelector = true;
+
+	UPROPERTY()
+	bool bShowReferenceInSelector = true;
+
+public:
+	void PopulateFrom(const FName& InName, const IMediaIOCoreDeviceProvider* InDeviceProvider);
+	void ApplyServerName(const FString& InServerName);
+};
+
+USTRUCT()
+struct FAvaBroadcastDeviceProviderDataList
+{
+	GENERATED_BODY()
+
+public:
+	UPROPERTY()
+	FString ServerName;
+	
+	UPROPERTY()
+	TArray<FAvaBroadcastDeviceProviderData> DeviceProviders;
+
+	bool SaveToJson();
+	bool SaveToXml();
+	bool LoadFromJson();
+	void Populate(const FString& InAssignedServerName);
+};
\ No newline at end of file
diff --git a/Engine/Plugins/Experimental/Avalanche/Source/AvalancheMedia/Internal/Broadcast/OutputDevices/AvaBroadcastOutputUtils.h b/Engine/Plugins/Experimental/Avalanche/Source/AvalancheMedia/Internal/Broadcast/OutputDevices/AvaBroadcastOutputUtils.h
new file mode 100644
--- /dev/null
+++ b/Engine/Plugins/Experimental/Avalanche/Source/AvalancheMedia/Internal/Broadcast/OutputDevices/AvaBroadcastOutputUtils.h
@@ -0,0 +1,32 @@
+// Copyright Epic Games, Inc. All Rights Reserved.
+
+#pragma once
+
+#include "Containers/UnrealString.h"
+#include "UObject/NameTypes.h"
+
+class UMediaOutput;
+class UObject;
+struct FAvaBroadcastOutputData;
+
+namespace UE::AvaBroadcastOutputUtils
+{
+	/** Get the device's name from the given MediaOutput. */
+	AVALANCHEMEDIA_API FString GetDeviceName(const UMediaOutput* InMediaOutput);
+
+	/** Returns true if the class has MediaIOCustomLayout meta data. */
+	AVALANCHEMEDIA_API bool HasDeviceProviderName(const UClass* InMediaOutputClass);
+
+	/** Retrieves the device provider name from the MediaIOCustomLayout class meta data. */
+	AVALANCHEMEDIA_API FName GetDeviceProviderName(const UClass* InMediaOutputClass);
+
+	/** Retrieves the device provider name from the MediaIOCustomLayout class meta data. */
+	AVALANCHEMEDIA_API FName GetDeviceProviderName(const UMediaOutput* InMediaOutput);
+
+	/** Serializes the given MediaOutput object into a FAvaBroadcastOutputData. */ 
+	AVALANCHEMEDIA_API FAvaBroadcastOutputData CreateMediaOutputData(UMediaOutput* InMediaOutput);
+
+	/** Create a MediaOutput object from the given data. */
+	AVALANCHEMEDIA_API UMediaOutput* CreateMediaOutput(const FAvaBroadcastOutputData& InMediaOutputData, UObject* InOuter);
+}
+
diff --git a/Engine/Plugins/Experimental/Avalanche/Source/AvalancheMedia/Internal/Broadcast/OutputDevices/AvaBroadcastRenderTargetMediaUtils.h b/Engine/Plugins/Experimental/Avalanche/Source/AvalancheMedia/Internal/Broadcast/OutputDevices/AvaBroadcastRenderTargetMediaUtils.h
new file mode 100644
--- /dev/null
+++ b/Engine/Plugins/Experimental/Avalanche/Source/AvalancheMedia/Internal/Broadcast/OutputDevices/AvaBroadcastRenderTargetMediaUtils.h
@@ -0,0 +1,29 @@
+// Copyright Epic Games, Inc. All Rights Reserved.
+
+#pragma once
+
+#include "RHI.h"
+
+class FRHICommandListImmediate;
+class UTextureRenderTarget2D;
+
+namespace UE::AvaBroadcastRenderTargetMediaUtils
+{
+	static inline const TCHAR* VirtualShaderMountPoint = TEXT("/Plugin/MotionDesign");
+
+	/**
+	 * Clears the given render target to it's specified clear color.
+	 * @remark called from the main thread.
+	 */
+	AVALANCHEMEDIA_API void ClearRenderTarget(UTextureRenderTarget2D* InRenderTarget);
+
+	/**
+	 * Copy a texture from source to destination.
+	 * Handles different sizes, will rescale.
+	 * Handles different formats, will convert.
+	 */
+	AVALANCHEMEDIA_API void CopyTexture(FRHICommandListImmediate& InRHICmdList, FTextureRHIRef InSourceTexture, FTextureRHIRef InDestTarget);
+	
+	/** Implementation of texture conversion from SCS_FinalColorSDR back to SCS_FinalToneCurveHDR. */
+	AVALANCHEMEDIA_API void ConvertTextureRGBGamma(FRHICommandListImmediate& InRHICmdList, FTextureRHIRef InSourceTexture, FTextureRHIRef InDestTarget, bool bInSrgbToLinear, float InGamma);
+}
\ No newline at end of file
diff --git a/Engine/Plugins/Experimental/Avalanche/Source/AvalancheMedia/Internal/Broadcast/OutputDevices/Slate/SAvaBroadcastCaptureImage.h b/Engine/Plugins/Experimental/Avalanche/Source/AvalancheMedia/Internal/Broadcast/OutputDevices/Slate/SAvaBroadcastCaptureImage.h
new file mode 100644
--- /dev/null
+++ b/Engine/Plugins/Experimental/Avalanche/Source/AvalancheMedia/Internal/Broadcast/OutputDevices/Slate/SAvaBroadcastCaptureImage.h
@@ -0,0 +1,41 @@
+// Copyright Epic Games, Inc. All Rights Reserved.
+
+#pragma once
+
+#include "Misc/Attribute.h"
+#include "Widgets/Images/SImage.h"
+
+class FSlateRect;
+class FSlateWindowElementList;
+class FWidgetStyle;
+struct FGeometry;
+
+class AVALANCHEMEDIA_API SAvaBroadcastCaptureImage : public SImage
+{
+public:
+	SLATE_BEGIN_ARGS(SAvaBroadcastCaptureImage)
+		: _ImageArgs()
+		, _EnableGammaCorrection(true)
+	{}
+		SLATE_ARGUMENT(SImage::FArguments, ImageArgs)
+		SLATE_ARGUMENT(bool, EnableGammaCorrection)
+		SLATE_ATTRIBUTE(bool, EnableBlending)
+		SLATE_ATTRIBUTE(bool, ShouldInvertAlpha)
+	SLATE_END_ARGS()
+
+	/** Constructs this widget with InArgs */
+	void Construct(const FArguments& InArgs);
+	
+	virtual int32 OnPaint(const FPaintArgs& Args
+		, const FGeometry& AllottedGeometry
+		, const FSlateRect& MyCullingRect
+		, FSlateWindowElementList& OutDrawElements
+		, int32 LayerId
+		, const FWidgetStyle& InWidgetStyle
+		, bool bParentEnabled) const override;
+
+protected:
+	bool bEnableGammaCorrection = true;
+	TAttribute<bool> EnableBlending;
+	TAttribute<bool> ShouldInvertAlpha;
+};
diff --git a/Engine/Plugins/Experimental/Avalanche/Source/AvalancheMedia/Private/Broadcast/Channel/AvaBroadcastOutputChannel.cpp b/Engine/Plugins/Experimental/Avalanche/Source/AvalancheMedia/Private/Broadcast/Channel/AvaBroadcastOutputChannel.cpp
--- a/Engine/Plugins/Experimental/Avalanche/Source/AvalancheMedia/Private/Broadcast/Channel/AvaBroadcastOutputChannel.cpp
+++ b/Engine/Plugins/Experimental/Avalanche/Source/AvalancheMedia/Private/Broadcast/Channel/AvaBroadcastOutputChannel.cpp
@@ -8,6 +8,7 @@
 #include "Broadcast/AvaBroadcast.h"
 #include "Broadcast/AvaBroadcastProfile.h"
 #include "Broadcast/Channel/IAvaBroadcastOutputChannelViewInterface.h"
+#include "Broadcast/Channel/Slate/SAvaBroadcastPlaceholderWidget.h"
 #include "Broadcast/IAvaBroadcastSettings.h"
 #include "Broadcast/OutputDevices/AvaBroadcastOutputUtils.h"
 #include "Broadcast/OutputDevices/AvaBroadcastRenderTargetMediaUtils.h"
@@ -29,7 +30,6 @@
 #include "RHIResources.h"
 #include "SceneView.h"
 #include "Slate/WidgetRenderer.h"
-#include "StereoRenderTargetManager.h"
 #include "ScreenPass.h"
 #include "TextureResource.h"
 #include "UObject/NameTypes.h"
diff --git a/Engine/Plugins/Experimental/Avalanche/Source/AvalancheMedia/Private/Broadcast/Channel/Slate/SAvaBroadcastPlaceholderWidget.h b/Engine/Plugins/Experimental/Avalanche/Source/AvalancheMedia/Private/Broadcast/Channel/Slate/SAvaBroadcastPlaceholderWidget.h
new file mode 100644
--- /dev/null
+++ b/Engine/Plugins/Experimental/Avalanche/Source/AvalancheMedia/Private/Broadcast/Channel/Slate/SAvaBroadcastPlaceholderWidget.h
@@ -0,0 +1,24 @@
+// Copyright Epic Games, Inc. All Rights Reserved.
+
+#pragma once
+
+#include "Templates/SharedPointer.h"
+#include "Widgets/SCompoundWidget.h"
+
+class FText;
+class STextBlock;
+
+class SAvaBroadcastPlaceholderWidget : public SCompoundWidget
+{
+public:
+	SLATE_BEGIN_ARGS(SAvaBroadcastPlaceholderWidget) {}
+	SLATE_END_ARGS()
+
+	/** Constructs this widget with InArgs */
+	void Construct(const FArguments& InArgs);
+
+	void SetChannelName(const FText& InChannelName);
+	
+protected:
+	TSharedPtr<STextBlock> ChannelText;
+};
diff --git a/Engine/Plugins/Experimental/Avalanche/Source/AvalancheMedia/Private/Broadcast/OutputDevices/AvaBroadcastDisplayDeviceManager.h b/Engine/Plugins/Experimental/Avalanche/Source/AvalancheMedia/Private/Broadcast/OutputDevices/AvaBroadcastDisplayDeviceManager.h
new file mode 100644
--- /dev/null
+++ b/Engine/Plugins/Experimental/Avalanche/Source/AvalancheMedia/Private/Broadcast/OutputDevices/AvaBroadcastDisplayDeviceManager.h
@@ -0,0 +1,59 @@
+// Copyright Epic Games, Inc. All Rights Reserved.
+
+#pragma once
+
+#include "Containers/Array.h"
+#include "Containers/UnrealString.h"
+#include "GenericPlatform/GenericApplication.h"
+#include "Misc/FrameRate.h"
+
+/**
+ * Contains information on physical display adapter,
+ * a.k.a. graphics card.
+ **/
+struct FAvaBroadcastDisplayAdapterInfo
+{
+	FString Description;			// ex: "Nvidia GeForce RTX ..."
+	uint32 VendorId = 0;
+	uint32 DeviceId = 0;
+	uint32 SubSysId = 0;
+	uint32 Revision = 0;
+	uint64 DedicatedVideoMemory = 0;
+};
+
+/**
+ * Contains info on a physical monitor connected to the display device.
+ * Similar to FMonitorInfo, except we have the adapter description
+ * and display frequency.
+ */
+struct FAvaBroadcastMonitorInfo
+{
+	FAvaBroadcastDisplayAdapterInfo AdapterInfo;
+	FString Name;					// ex: DISPLAY1, DISPLAY2, etc (not manufacturer name)
+	int32 Width = 0;				// Horizontal resolution in pixels.
+	int32 Height = 0;				// Vertical resolution in pixels.
+	FFrameRate DisplayFrequency = FFrameRate(0, 0);	// Default invalid.
+	FPlatformRect DisplayRect;
+	FPlatformRect WorkArea;
+	bool bIsPrimary = false;
+};
+
+/**
+ * Provides additional information on the display devices that is not available
+ * in FDisplayMetrics.
+ */
+class FAvaBroadcastDisplayDeviceManager
+{
+public:
+	static void EnumMonitors(TArray<FAvaBroadcastMonitorInfo>& OutMonitorInfo);
+
+	static const TArray<FAvaBroadcastMonitorInfo>&  GetCachedMonitors(bool bForceUpdate = false);
+
+	/**
+	 * Returns a string with a display name for the monitor.
+	 **/
+	static FString GetMonitorDisplayName(const FAvaBroadcastMonitorInfo& InMonitorInfo);
+	
+private:
+	static TArray<FAvaBroadcastMonitorInfo> CachedMonitorInfo;
+};
diff --git a/Engine/Plugins/Experimental/Avalanche/Source/AvalancheMedia/Private/Broadcast/OutputDevices/AvaBroadcastDisplayMediaCapture.h b/Engine/Plugins/Experimental/Avalanche/Source/AvalancheMedia/Private/Broadcast/OutputDevices/AvaBroadcastDisplayMediaCapture.h
new file mode 100644
--- /dev/null
+++ b/Engine/Plugins/Experimental/Avalanche/Source/AvalancheMedia/Private/Broadcast/OutputDevices/AvaBroadcastDisplayMediaCapture.h
@@ -0,0 +1,39 @@
+// Copyright Epic Games, Inc. All Rights Reserved.
+
+#pragma once
+
+#include "HAL/CriticalSection.h"
+#include "MediaCapture.h"
+#include "AvaBroadcastDisplayMediaCapture.generated.h"
+
+DECLARE_LOG_CATEGORY_EXTERN(LogAvaBroadcastDisplayMedia, Log, All);
+
+UCLASS()
+class UAvaBroadcastDisplayMediaCapture : public UMediaCapture
+{
+	GENERATED_BODY()
+
+public:
+	virtual ~UAvaBroadcastDisplayMediaCapture() override; // so we can pimpl.
+
+protected:
+	virtual bool ShouldCaptureRHIResource() const override { return true; }
+	virtual void OnRHIResourceCaptured_RenderingThread(const FCaptureBaseData& InBaseData, TSharedPtr<FMediaCaptureUserData, ESPMode::ThreadSafe> InUserData, FTextureRHIRef InTexture) override;
+
+	virtual bool InitializeCapture() override;
+	virtual bool PostInitializeCaptureViewport(TSharedPtr<FSceneViewport>& InSceneViewport) override;
+	virtual bool PostInitializeCaptureRenderTarget(UTextureRenderTarget2D* InRenderTarget) override;
+	virtual void StopCaptureImpl(bool bAllowPendingFrameToBeProcess) override;
+
+	UPROPERTY(Transient)
+	TObjectPtr<UTextureRenderTarget2D> RenderTarget;
+
+private:
+	bool StartNewCapture(const FIntPoint& InSourceTargetSize, EPixelFormat InSourceTargetFormat);
+
+	FCriticalSection CaptureInstanceCriticalSection;
+
+	// Private capture instance
+	class FCaptureInstance;
+	FCaptureInstance* CaptureInstance = nullptr; // lazy pimpl
+};
diff --git a/Engine/Plugins/Experimental/Avalanche/Source/AvalancheMedia/Private/Broadcast/OutputDevices/AvaBroadcastRenderTargetMediaCapture.h b/Engine/Plugins/Experimental/Avalanche/Source/AvalancheMedia/Private/Broadcast/OutputDevices/AvaBroadcastRenderTargetMediaCapture.h
new file mode 100644
--- /dev/null
+++ b/Engine/Plugins/Experimental/Avalanche/Source/AvalancheMedia/Private/Broadcast/OutputDevices/AvaBroadcastRenderTargetMediaCapture.h
@@ -0,0 +1,32 @@
+// Copyright Epic Games, Inc. All Rights Reserved.
+
+#pragma once
+
+#include "HAL/CriticalSection.h"
+#include "MediaCapture.h"
+#include "AvaBroadcastRenderTargetMediaCapture.generated.h"
+
+DECLARE_LOG_CATEGORY_EXTERN(LogAvaBroadcastRenderTargetMedia, Log, All);
+
+UCLASS()
+class UAvaBroadcastRenderTargetMediaCapture : public UMediaCapture
+{
+	GENERATED_BODY()
+
+protected:
+	virtual bool ShouldCaptureRHIResource() const override { return true; }
+	virtual void OnRHIResourceCaptured_RenderingThread(const FCaptureBaseData& InBaseData, TSharedPtr<FMediaCaptureUserData, ESPMode::ThreadSafe> InUserData, FTextureRHIRef InTexture) override;
+
+	virtual bool InitializeCapture() override;
+	virtual bool PostInitializeCaptureViewport(TSharedPtr<FSceneViewport>& InSceneViewport) override;
+	virtual bool PostInitializeCaptureRenderTarget(UTextureRenderTarget2D* InRenderTarget) override;
+	virtual void StopCaptureImpl(bool bAllowPendingFrameToBeProcess) override;
+
+	UPROPERTY(Transient)
+	TObjectPtr<UTextureRenderTarget2D> RenderTarget;
+
+private:
+	bool StartNewCapture(const FIntPoint& InSourceTargetSize, EPixelFormat InSourceTargetFormat);
+	
+	FCriticalSection RenderTargetCriticalSection;
+};
diff --git a/Engine/Plugins/Experimental/Avalanche/Source/AvalancheMedia/Private/Broadcast/OutputDevices/AvaBroadcastRenderTargetMediaOutput.h b/Engine/Plugins/Experimental/Avalanche/Source/AvalancheMedia/Private/Broadcast/OutputDevices/AvaBroadcastRenderTargetMediaOutput.h
new file mode 100644
--- /dev/null
+++ b/Engine/Plugins/Experimental/Avalanche/Source/AvalancheMedia/Private/Broadcast/OutputDevices/AvaBroadcastRenderTargetMediaOutput.h
@@ -0,0 +1,52 @@
+// Copyright Epic Games, Inc. All Rights Reserved.
+
+#pragma once
+
+#include "MediaOutput.h"
+#include "UObject/SoftObjectPtr.h"
+#include "AvaBroadcastRenderTargetMediaOutput.generated.h"
+
+class UTextureRenderTarget2D;
+
+/**
+ * Capture a Media to a render target.
+ */
+UCLASS(BlueprintType, ClassGroup = "Motion Design Broadcast",
+	meta = (DisplayName = "Motion Design Render Target Media Output", MediaIOCustomLayout = "AvaRenderTarget"))
+class UAvaBroadcastRenderTargetMediaOutput : public UMediaOutput
+{
+	GENERATED_BODY()
+
+public:
+	/** Invert the key (alpha channel). */
+	UPROPERTY(BlueprintReadWrite, EditAnywhere, Category="Media")
+	bool bInvertKeyOutput = true;
+
+	/**
+	 * Specify the render target to be capturing to.
+	 */
+	UPROPERTY(BlueprintReadWrite, EditAnywhere, Category="Media")
+	TSoftObjectPtr<UTextureRenderTarget2D> RenderTarget;
+
+	/**
+	 * The source name is a property that exists to provide the "device name" for
+	 * displaying in the broadcast editor. It is derived from the specified render target
+	 * path. */
+	UPROPERTY()
+	FString SourceName;
+	
+	//~ UMediaOutput interface
+public:
+	virtual bool Validate(FString& FailureReason) const override;
+	virtual FIntPoint GetRequestedSize() const override;
+	virtual EPixelFormat GetRequestedPixelFormat() const override;
+	virtual EMediaCaptureConversionOperation GetConversionOperation(EMediaCaptureSourceType InSourceType) const override;
+
+#if WITH_EDITOR
+	virtual void PostEditChangeProperty(struct FPropertyChangedEvent& PropertyChangedEvent) override;
+#endif
+
+protected:
+	virtual UMediaCapture* CreateMediaCaptureImpl() override;
+	//~ End UMediaOutput interface
+};
diff --git a/Engine/Plugins/Experimental/Avalanche/Source/AvalancheMedia/Public/AvaMediaMessageUtils.h b/Engine/Plugins/Experimental/Avalanche/Source/AvalancheMedia/Public/AvaMediaMessageUtils.h
deleted file mode 100644
diff --git a/Engine/Plugins/Experimental/Avalanche/Source/AvalancheMedia/Public/AvaMediaRenderTargetUtils.h b/Engine/Plugins/Experimental/Avalanche/Source/AvalancheMedia/Public/AvaMediaRenderTargetUtils.h
deleted file mode 100644
diff --git a/Engine/Plugins/Experimental/Avalanche/Source/AvalancheMedia/Public/AvaMediaSerializationUtils.h b/Engine/Plugins/Experimental/Avalanche/Source/AvalancheMedia/Public/AvaMediaSerializationUtils.h
deleted file mode 100644
diff --git a/Engine/Plugins/Experimental/Avalanche/Source/AvalancheMedia/Public/Broadcast/AvaBroadcastDevice.h b/Engine/Plugins/Experimental/Avalanche/Source/AvalancheMedia/Public/Broadcast/AvaBroadcastDevice.h
deleted file mode 100644
diff --git a/Engine/Plugins/Experimental/Avalanche/Source/AvalancheMedia/Public/Broadcast/Channel/AvaBroadcastOutputChannel.h b/Engine/Plugins/Experimental/Avalanche/Source/AvalancheMedia/Public/Broadcast/Channel/AvaBroadcastOutputChannel.h
--- a/Engine/Plugins/Experimental/Avalanche/Source/AvalancheMedia/Public/Broadcast/Channel/AvaBroadcastOutputChannel.h
+++ b/Engine/Plugins/Experimental/Avalanche/Source/AvalancheMedia/Public/Broadcast/Channel/AvaBroadcastOutputChannel.h
@@ -8,7 +8,6 @@
 #include "Containers/Map.h"
 #include "Containers/StringFwd.h"
 #include "Engine/TextureRenderTarget2D.h"
-#include "Slate/SAvaBroadcastPlaceholderWidget.h"
 #include "Viewport/AvaViewportQualitySettings.h"
 #include "UObject/ObjectPtr.h"
 #include "AvaBroadcastOutputChannel.generated.h"
@@ -16,6 +15,7 @@
 class FAudioDeviceHandle;
 class FName;
 class FWidgetRenderer;
+class SAvaBroadcastPlaceholderWidget;
 class SVirtualWindow;
 class SWidget;
 class UAvaPlayableGroup;
diff --git a/Engine/Plugins/Experimental/Avalanche/Source/AvalancheMedia/Public/Broadcast/Channel/Slate/SAvaBroadcastPlaceholderWidget.h b/Engine/Plugins/Experimental/Avalanche/Source/AvalancheMedia/Public/Broadcast/Channel/Slate/SAvaBroadcastPlaceholderWidget.h
deleted file mode 100644
diff --git a/Engine/Plugins/Experimental/Avalanche/Source/AvalancheMedia/Public/Broadcast/OutputDevices/AvaBroadcastDeviceProviderData.h b/Engine/Plugins/Experimental/Avalanche/Source/AvalancheMedia/Public/Broadcast/OutputDevices/AvaBroadcastDeviceProviderData.h
deleted file mode 100644
diff --git a/Engine/Plugins/Experimental/Avalanche/Source/AvalancheMedia/Public/Broadcast/OutputDevices/AvaBroadcastDisplayDeviceManager.h b/Engine/Plugins/Experimental/Avalanche/Source/AvalancheMedia/Public/Broadcast/OutputDevices/AvaBroadcastDisplayDeviceManager.h
deleted file mode 100644
diff --git a/Engine/Plugins/Experimental/Avalanche/Source/AvalancheMedia/Public/Broadcast/OutputDevices/AvaBroadcastDisplayMediaCapture.h b/Engine/Plugins/Experimental/Avalanche/Source/AvalancheMedia/Public/Broadcast/OutputDevices/AvaBroadcastDisplayMediaCapture.h
deleted file mode 100644
diff --git a/Engine/Plugins/Experimental/Avalanche/Source/AvalancheMedia/Public/Broadcast/OutputDevices/AvaBroadcastDisplayMediaOutput.h b/Engine/Plugins/Experimental/Avalanche/Source/AvalancheMedia/Public/Broadcast/OutputDevices/AvaBroadcastDisplayMediaOutput.h
--- a/Engine/Plugins/Experimental/Avalanche/Source/AvalancheMedia/Public/Broadcast/OutputDevices/AvaBroadcastDisplayMediaOutput.h
+++ b/Engine/Plugins/Experimental/Avalanche/Source/AvalancheMedia/Public/Broadcast/OutputDevices/AvaBroadcastDisplayMediaOutput.h
@@ -9,9 +9,9 @@
 /**
  * Output Media to a display adapter.
  */
-UCLASS(BlueprintType, ClassGroup = "Motion Design Broadcast",
+UCLASS(MinimalAPI, ClassGroup = "Motion Design Broadcast",
 	meta = (DisplayName = "Motion Design Broadcast Display Media Output", MediaIOCustomLayout = "AvaDisplay"))
-class AVALANCHEMEDIA_API UAvaBroadcastDisplayMediaOutput : public UMediaOutput
+class UAvaBroadcastDisplayMediaOutput : public UMediaOutput
 {
 	GENERATED_BODY()
 
diff --git a/Engine/Plugins/Experimental/Avalanche/Source/AvalancheMedia/Public/Broadcast/OutputDevices/AvaBroadcastOutputUtils.h b/Engine/Plugins/Experimental/Avalanche/Source/AvalancheMedia/Public/Broadcast/OutputDevices/AvaBroadcastOutputUtils.h
deleted file mode 100644
diff --git a/Engine/Plugins/Experimental/Avalanche/Source/AvalancheMedia/Public/Broadcast/OutputDevices/AvaBroadcastRenderTargetMediaCapture.h b/Engine/Plugins/Experimental/Avalanche/Source/AvalancheMedia/Public/Broadcast/OutputDevices/AvaBroadcastRenderTargetMediaCapture.h
deleted file mode 100644
diff --git a/Engine/Plugins/Experimental/Avalanche/Source/AvalancheMedia/Public/Broadcast/OutputDevices/AvaBroadcastRenderTargetMediaOutput.h b/Engine/Plugins/Experimental/Avalanche/Source/AvalancheMedia/Public/Broadcast/OutputDevices/AvaBroadcastRenderTargetMediaOutput.h
deleted file mode 100644
diff --git a/Engine/Plugins/Experimental/Avalanche/Source/AvalancheMedia/Public/Broadcast/OutputDevices/AvaBroadcastRenderTargetMediaUtils.h b/Engine/Plugins/Experimental/Avalanche/Source/AvalancheMedia/Public/Broadcast/OutputDevices/AvaBroadcastRenderTargetMediaUtils.h
deleted file mode 100644
diff --git a/Engine/Plugins/Experimental/Avalanche/Source/AvalancheMedia/Public/Broadcast/OutputDevices/Slate/SAvaBroadcastCaptureImage.h b/Engine/Plugins/Experimental/Avalanche/Source/AvalancheMedia/Public/Broadcast/OutputDevices/Slate/SAvaBroadcastCaptureImage.h
deleted file mode 100644
