[Avalanche/AvalancheEditor/AvalancheShapes]
Fix Category section properties for Text and Shapes

diff --git a/Engine/Plugins/Experimental/Avalanche/Source/AvalancheEditor/Private/AvaEditorModule.cpp b/Engine/Plugins/Experimental/Avalanche/Source/AvalancheEditor/Private/AvaEditorModule.cpp
--- a/Engine/Plugins/Experimental/Avalanche/Source/AvalancheEditor/Private/AvaEditorModule.cpp
+++ b/Engine/Plugins/Experimental/Avalanche/Source/AvalancheEditor/Private/AvaEditorModule.cpp
@@ -227,15 +227,13 @@ void FAvaEditorModule::RegisterPropertyEditorCategories()
 
 		Section = PropertyModule.FindOrCreateSection("AvaTextActor", "Style", LOCTEXT("Style", "Style"));
 		Section->AddCategory("Style");
+		Section->AddCategory("Materials");
 
 		Section = PropertyModule.FindOrCreateSection("AvaTextActor", "Geometry", LOCTEXT("Geometry", "Geometry"));
 		Section->AddCategory("Geometry");
 
 		Section = PropertyModule.FindOrCreateSection("AvaTextActor", "Layout", LOCTEXT("Layout", "Layout"));
 		Section->AddCategory("Layout");
-
-		Section = PropertyModule.FindOrCreateSection("AvaTextActor", "Materials", LOCTEXT("Materials", "Text Materials"));
-		Section->AddCategory("Materials");
 	}
 
 	// AvaCineCameraActor Sections
diff --git a/Engine/Plugins/Experimental/Avalanche/Source/AvalancheShapes/Public/DynamicMeshes/AvaShapeDynMeshBase.h b/Engine/Plugins/Experimental/Avalanche/Source/AvalancheShapes/Public/DynamicMeshes/AvaShapeDynMeshBase.h
--- a/Engine/Plugins/Experimental/Avalanche/Source/AvalancheShapes/Public/DynamicMeshes/AvaShapeDynMeshBase.h
+++ b/Engine/Plugins/Experimental/Avalanche/Source/AvalancheShapes/Public/DynamicMeshes/AvaShapeDynMeshBase.h
@@ -340,22 +340,22 @@ protected:
 #endif
 
 	// enable mesh size property editing
-	UPROPERTY(EditDefaultsOnly, BlueprintReadOnly, Category="Shape")
+	UPROPERTY(EditDefaultsOnly, BlueprintReadOnly, Category="Shape", meta=(DisplayPriority=1))
 	bool bAllowEditSize;
 
 	// the type of size you want to handle
-	UPROPERTY(EditInstanceOnly, BlueprintReadWrite, Category="Shape", Transient, meta=(DisplayName="Size Type", EditCondition="bAllowEditSize", EditConditionHides, AllowPrivateAccess="true"))
+	UPROPERTY(EditInstanceOnly, BlueprintReadWrite, Category="Shape", Transient, meta=(DisplayName="Size Type", DisplayPriority=1, EditCondition="bAllowEditSize", EditConditionHides, AllowPrivateAccess="true"))
 	ESizeType SizeType = ESizeType::UnrealUnit;
 
 	// Uniform scaled size of the mesh
-	UPROPERTY(EditInstanceOnly, BlueprintReadWrite, Transient, Category="Shape", meta=(ClampMin="0.0", DisplayName="Uniform Scaled Size", Units="times", AllowPrivateAccess="true"))
+	UPROPERTY(EditInstanceOnly, BlueprintReadWrite, Transient, Category="Shape", meta=(ClampMin="0.0", DisplayName="Uniform Scaled Size", DisplayPriority=1, Units="times", AllowPrivateAccess="true"))
 	float UniformScaledSize = 1.f;
 
-	UPROPERTY(BlueprintReadWrite, AdvancedDisplay, Category = "Shape", meta=(DisplayName="Material Vertex Color", AllowPrivateAccess="true"))
+	UPROPERTY(BlueprintReadWrite, AdvancedDisplay, Category = "Shape", meta=(DisplayName="Material Vertex Color", DisplayPriority=1, AllowPrivateAccess="true"))
 	FLinearColor VertexColor;
 
 	// use primary material for every slot available
-	UPROPERTY(EditInstanceOnly, BlueprintReadWrite, Category="Material", meta=(DisplayName="Use Single Material", AllowPrivateAccess="true"))
+	UPROPERTY(EditInstanceOnly, BlueprintReadWrite, Category="Material", meta=(DisplayName="Use Single Material", DisplayPriority=1, AllowPrivateAccess="true"))
 	bool bUsePrimaryMaterialEverywhere;
 
 	TArray<FAvaSnapPoint> LocalSnapPoints;
