MotionDesign : Shapes
- Cleaning public facing API

#changelist validated
#virtualized

diff --git a/Engine/Plugins/Experimental/Avalanche/Source/AvalancheShapes/Private/DynamicMeshes/AvaShape2DArrowDynMesh.cpp b/Engine/Plugins/Experimental/Avalanche/Source/AvalancheShapes/Private/DynamicMeshes/AvaShape2DArrowDynMesh.cpp
--- a/Engine/Plugins/Experimental/Avalanche/Source/AvalancheShapes/Private/DynamicMeshes/AvaShape2DArrowDynMesh.cpp
+++ b/Engine/Plugins/Experimental/Avalanche/Source/AvalancheShapes/Private/DynamicMeshes/AvaShape2DArrowDynMesh.cpp
@@ -6,86 +6,79 @@
 
 const FString UAvaShape2DArrowDynamicMesh::MeshName = TEXT("2DArrow");
 
-bool UAvaShape2DArrowDynamicMesh::SetRatioArrowLine(float InRatio)
+void UAvaShape2DArrowDynamicMesh::SetRatioArrowLine(float InRatio)
 {
 	if (RatioArrowLine == InRatio)
 	{
-		return false;
+		return;
 	}
 
 	if (InRatio < 0.01f || InRatio > 1.f)
 	{
-		return false;
+		return;
 	}
 
 	RatioArrowLine = InRatio;
 	OnRatioArrowLineChanged();
-
-	return true;
 }
 
-bool UAvaShape2DArrowDynamicMesh::SetRatioLineHeight(float InRatio)
+void UAvaShape2DArrowDynamicMesh::SetRatioLineHeight(float InRatio)
 {
 	if (RatioLineHeight == InRatio)
 	{
-		return false;
+		return;
 	}
 
 	if (InRatio < 0.01f || InRatio > 1.f)
 	{
-		return false;
+		return;
 	}
-	
+
 	RatioLineHeight = InRatio;
 	OnRatioLineHeightChanged();
-
-	return true;
 }
 
-bool UAvaShape2DArrowDynamicMesh::SetRatioArrowY(float InRatio)
+void UAvaShape2DArrowDynamicMesh::SetRatioArrowY(float InRatio)
 {
 	if (RatioArrowY == InRatio)
 	{
-		return false;
+		return;
 	}
 
 	if (InRatio < 0.01f || InRatio > 1.f)
 	{
-		return false;
+		return;
 	}
 
 	RatioArrowY = InRatio;
 	OnRatioArrowYChanged();
-	return true;
 }
 
-bool UAvaShape2DArrowDynamicMesh::SetRatioLineY(float InRatio)
+void UAvaShape2DArrowDynamicMesh::SetRatioLineY(float InRatio)
 {
 	if (RatioLineY == InRatio)
 	{
-		return false;
+		return;
 	}
 
 	if (InRatio < 0.01f || InRatio > 1.f)
 	{
-		return false;
+		return;
 	}
 
 	RatioLineY = InRatio;
 	OnRatioLineYChanged();
-	return true;
 }
 
-bool UAvaShape2DArrowDynamicMesh::SetBothSideArrows(bool bBothSide)
+void UAvaShape2DArrowDynamicMesh::SetBothSideArrows(bool bBothSide)
 {
 	if (bBothSideArrows == bBothSide)
 	{
-		return false;
+		return;
 	}
 
 	bBothSideArrows = bBothSide;
 	OnBothSideArrowsChanged();
-	return true;
 }
 
 #if WITH_EDITOR
@@ -102,7 +95,7 @@ void UAvaShape2DArrowDynamicMesh::PostEditChangeProperty(FPropertyChangedEvent&
 	if (PropertyChangedEvent.MemberProperty->GetFName() == RatioArrowLineName)
 	{
 		OnRatioArrowLineChanged();
-	} 
+	}
 	else if (PropertyChangedEvent.MemberProperty->GetFName() == RatioLineHeightName)
 	{
 		OnRatioLineHeightChanged();
diff --git a/Engine/Plugins/Experimental/Avalanche/Source/AvalancheShapes/Private/DynamicMeshes/AvaShape2DDynMeshBase.cpp b/Engine/Plugins/Experimental/Avalanche/Source/AvalancheShapes/Private/DynamicMeshes/AvaShape2DDynMeshBase.cpp
--- a/Engine/Plugins/Experimental/Avalanche/Source/AvalancheShapes/Private/DynamicMeshes/AvaShape2DDynMeshBase.cpp
+++ b/Engine/Plugins/Experimental/Avalanche/Source/AvalancheShapes/Private/DynamicMeshes/AvaShape2DDynMeshBase.cpp
@@ -9,9 +9,9 @@
 #include "Operations/MinimalHoleFiller.h"
 #include "Operations/OffsetMeshRegion.h"
 
-bool UAvaShape2DDynMeshBase::SetSize3D(const FVector& InSize)
+void UAvaShape2DDynMeshBase::SetSize3D(const FVector& InSize)
 {
-	return SetSize2D({InSize.Y, InSize.Z});
+	SetSize2D({InSize.Y, InSize.Z});
 }
 
 void UAvaShape2DDynMeshBase::PostLoad()
@@ -127,45 +127,43 @@ bool UAvaShape2DDynMeshBase::CreateMesh(FAvaShapeMesh& InMesh)
 	return Super::CreateMesh(InMesh);
 }
 
-bool UAvaShape2DDynMeshBase::SetPixelSize2D(const FVector2D& InPixelSize2D)
+void UAvaShape2DDynMeshBase::SetPixelSize2D(const FVector2D& InPixelSize2D)
 {
 	if (!bAllowEditSize)
 	{
-		return false;
+		return;
 	}
 
 	if (PixelSize2D == InPixelSize2D)
 	{
-		return false;
+		return;
 	}
 
 	if (InPixelSize2D.GetMin() < UAvaShapeDynamicMeshBase::MinSizeValue)
 	{
-		return false;
+		return;
 	}
 
 	PixelSize2D = InPixelSize2D;
 	OnPixelSizeChanged();
-	return true;
 }
 
 // Special case. We can update more easily via c++ than Details panel!
-bool UAvaShape2DDynMeshBase::SetSize2D(const FVector2D& InSize2D)
+void UAvaShape2DDynMeshBase::SetSize2D(const FVector2D& InSize2D)
 {
 	if (!bAllowEditSize)
 	{
-		return false;
+		return;
 	}
 
 	const FVector2D NewSize = FVector2D::Max(FVector2D(UAvaShapeDynamicMeshBase::MinSizeValue), InSize2D);
 	if (Size2D.Equals(NewSize))
 	{
-		return false;
+		return;
 	}
 
 	Size2D = NewSize;
 	OnSizeChanged();
-	return true;
 }
 
 void UAvaShape2DDynMeshBase::OnRegisteredMeshes()
diff --git a/Engine/Plugins/Experimental/Avalanche/Source/AvalancheShapes/Private/DynamicMeshes/AvaShape3DDynMeshBase.cpp b/Engine/Plugins/Experimental/Avalanche/Source/AvalancheShapes/Private/DynamicMeshes/AvaShape3DDynMeshBase.cpp
--- a/Engine/Plugins/Experimental/Avalanche/Source/AvalancheShapes/Private/DynamicMeshes/AvaShape3DDynMeshBase.cpp
+++ b/Engine/Plugins/Experimental/Avalanche/Source/AvalancheShapes/Private/DynamicMeshes/AvaShape3DDynMeshBase.cpp
@@ -86,44 +86,42 @@ void UAvaShape3DDynMeshBase::GetBounds(FVector& Origin, FVector& BoxExtent, FVec
 	Pivot = FVector(0, 0, 0);
 }
 
-bool UAvaShape3DDynMeshBase::SetPixelSize3D(const FVector& InPixelSize)
+void UAvaShape3DDynMeshBase::SetPixelSize3D(const FVector& InPixelSize)
 {
 	if (!bAllowEditSize)
 	{
-		return false;
+		return;
 	}
 
 	if (PixelSize3D == InPixelSize)
 	{
-		return false;
+		return;
 	}
 
 	if (InPixelSize.GetMin() < UAvaShapeDynamicMeshBase::MinSizeValue)
 	{
-		return false;
+		return;
 	}
 
 	PixelSize3D = InPixelSize;
 	OnPixelSizeChanged();
-	return true;
 }
 
-bool UAvaShape3DDynMeshBase::SetSize3D(const FVector& InSize)
+void UAvaShape3DDynMeshBase::SetSize3D(const FVector& InSize)
 {
 	if (!bAllowEditSize)
 	{
-		return false;
+		return;
 	}
 
 	const FVector NewSize = FVector::Max(FVector(UAvaShapeDynamicMeshBase::MinSizeValue), InSize);
 	if (Size3D.Equals(NewSize))
 	{
-		return false;
+		return;
 	}
 
 	Size3D = NewSize;
 	OnSizeChanged();
-	return true;
 }
 
 void UAvaShape3DDynMeshBase::OnRegisteredMeshes()
diff --git a/Engine/Plugins/Experimental/Avalanche/Source/AvalancheShapes/Private/DynamicMeshes/AvaShapeChevronDynMesh.cpp b/Engine/Plugins/Experimental/Avalanche/Source/AvalancheShapes/Private/DynamicMeshes/AvaShapeChevronDynMesh.cpp
--- a/Engine/Plugins/Experimental/Avalanche/Source/AvalancheShapes/Private/DynamicMeshes/AvaShapeChevronDynMesh.cpp
+++ b/Engine/Plugins/Experimental/Avalanche/Source/AvalancheShapes/Private/DynamicMeshes/AvaShapeChevronDynMesh.cpp
@@ -19,22 +19,20 @@ void UAvaShapeChevronDynamicMesh::PostEditChangeProperty(FPropertyChangedEvent&
 }
 #endif
 
-bool UAvaShapeChevronDynamicMesh::SetRatioChevron(float InRatio)
+void UAvaShapeChevronDynamicMesh::SetRatioChevron(float InRatio)
 {
 	if (RatioChevron == InRatio)
 	{
-		return false;
+		return;
 	}
 
 	if (InRatio < 0.f || InRatio > 0.99f)
 	{
-		return false;
+		return;
 	}
 
 	RatioChevron = InRatio;
 	OnRatioChevronChanged();
-
-	return true;
 }
 
 void UAvaShapeChevronDynamicMesh::OnRatioChevronChanged()
diff --git a/Engine/Plugins/Experimental/Avalanche/Source/AvalancheShapes/Private/DynamicMeshes/AvaShapeConeDynMesh.cpp b/Engine/Plugins/Experimental/Avalanche/Source/AvalancheShapes/Private/DynamicMeshes/AvaShapeConeDynMesh.cpp
--- a/Engine/Plugins/Experimental/Avalanche/Source/AvalancheShapes/Private/DynamicMeshes/AvaShapeConeDynMesh.cpp
+++ b/Engine/Plugins/Experimental/Avalanche/Source/AvalancheShapes/Private/DynamicMeshes/AvaShapeConeDynMesh.cpp
@@ -7,76 +7,68 @@
 
 const FString UAvaShapeConeDynamicMesh::MeshName = TEXT("Cone");
 
-bool UAvaShapeConeDynamicMesh::SetNumSides(uint8 InNumSides)
+void UAvaShapeConeDynamicMesh::SetNumSides(uint8 InNumSides)
 {
 	if (NumSides == InNumSides)
 	{
-		return false;
+		return;
 	}
 
 	if (InNumSides < UAvaShapeConeDynamicMesh::MinNumSides || InNumSides > UAvaShapeConeDynamicMesh::MaxNumSides)
 	{
-		return false;
+		return;
 	}
 
 	NumSides = InNumSides;
 	OnNumSidesChanged();
-
-	return true;
 }
 
-bool UAvaShapeConeDynamicMesh::SetTopRadius(float InTopRadius)
+void UAvaShapeConeDynamicMesh::SetTopRadius(float InTopRadius)
 {
 	if (TopRadius == InTopRadius)
 	{
-		return false;
+		return;
 	}
 
 	if (InTopRadius < 0.f || InTopRadius > 1.f)
 	{
-		return false;
+		return;
 	}
 
 	TopRadius = InTopRadius;
 	OnTopRadiusChanged();
-
-	return true;
 }
 
-bool UAvaShapeConeDynamicMesh::SetAngleDegree(float InDegree)
+void UAvaShapeConeDynamicMesh::SetAngleDegree(float InDegree)
 {
 	if (AngleDegree == InDegree)
 	{
-		return false;
+		return;
 	}
 
 	if (InDegree < 0.f || InDegree > 360.f)
 	{
-		return false;
+		return;
 	}
 
 	AngleDegree = InDegree;
 	OnAngleDegreeChanged();
-
-	return true;
 }
 
-bool UAvaShapeConeDynamicMesh::SetStartDegree(float InDegree)
+void UAvaShapeConeDynamicMesh::SetStartDegree(float InDegree)
 {
 	if (StartDegree == InDegree)
 	{
-		return false;
+		return;
 	}
 
 	if (InDegree < 0.f || InDegree > 360.f)
 	{
-		return false;
+		return;
 	}
 
 	StartDegree = InDegree;
 	OnStartDegreeChanged();
-
-	return true;
 }
 
 #if WITH_EDITOR
diff --git a/Engine/Plugins/Experimental/Avalanche/Source/AvalancheShapes/Private/DynamicMeshes/AvaShapeCubeDynMesh.cpp b/Engine/Plugins/Experimental/Avalanche/Source/AvalancheShapes/Private/DynamicMeshes/AvaShapeCubeDynMesh.cpp
--- a/Engine/Plugins/Experimental/Avalanche/Source/AvalancheShapes/Private/DynamicMeshes/AvaShapeCubeDynMesh.cpp
+++ b/Engine/Plugins/Experimental/Avalanche/Source/AvalancheShapes/Private/DynamicMeshes/AvaShapeCubeDynMesh.cpp
@@ -8,58 +8,52 @@
 
 const FString UAvaShapeCubeDynamicMesh::MeshName = TEXT("Cube");
 
-bool UAvaShapeCubeDynamicMesh::SetSegment(float InSegment)
+void UAvaShapeCubeDynamicMesh::SetSegment(float InSegment)
 {
 	if (Segment == InSegment)
 	{
-		return false;
+		return;
 	}
 
 	if (InSegment < 0.01f)
 	{
-		return false;
+		return;
 	}
 
 	Segment = InSegment;
 	OnSegmentChanged();
-
-	return true;
 }
 
-bool UAvaShapeCubeDynamicMesh::SetBevelSizeRatio(float InBevel)
+void UAvaShapeCubeDynamicMesh::SetBevelSizeRatio(float InBevel)
 {
 	if (BevelSizeRatio == InBevel)
 	{
-		return false;
+		return;
 	}
 
 	if (InBevel < 0 || InBevel > GetMaxBevelSize())
 	{
-		return false;
+		return;
 	}
 
 	BevelSizeRatio = InBevel;
 	OnBevelSizeChanged();
-
-	return true;
 }
 
-bool UAvaShapeCubeDynamicMesh::SetBevelNum(uint8 InBevel)
+void UAvaShapeCubeDynamicMesh::SetBevelNum(uint8 InBevel)
 {
 	if (BevelNum == InBevel)
 	{
-		return false;
+		return;
 	}
 
 	if (InBevel < UAvaShapeCubeDynamicMesh::MinBevelNum || InBevel > UAvaShapeCubeDynamicMesh::MaxBevelNum)
 	{
-		return false;
+		return;
 	}
 
 	BevelNum = InBevel;
 	OnBevelNumChanged();
-
-	return true;
 }
 
 #if WITH_EDITOR
diff --git a/Engine/Plugins/Experimental/Avalanche/Source/AvalancheShapes/Private/DynamicMeshes/AvaShapeDynMeshBase.cpp b/Engine/Plugins/Experimental/Avalanche/Source/AvalancheShapes/Private/DynamicMeshes/AvaShapeDynMeshBase.cpp
--- a/Engine/Plugins/Experimental/Avalanche/Source/AvalancheShapes/Private/DynamicMeshes/AvaShapeDynMeshBase.cpp
+++ b/Engine/Plugins/Experimental/Avalanche/Source/AvalancheShapes/Private/DynamicMeshes/AvaShapeDynMeshBase.cpp
@@ -130,11 +130,11 @@ void UAvaShapeDynamicMeshBase::SetupMeshes()
 
 bool UAvaShapeDynamicMeshBase::SetAlignmentSize(AActor* InActor, const FVector& InSizeMultiplier)
 {
-	if (AAvaShapeActor* ShapeActor = Cast<AAvaShapeActor>(InActor))
+	if (const AAvaShapeActor* ShapeActor = Cast<AAvaShapeActor>(InActor))
 	{
 		if (UAvaShapeDynamicMeshBase* DynMesh = ShapeActor->GetDynamicMesh())
 		{
-			return DynMesh->SetSize3D(DynMesh->GetSize3D() * InSizeMultiplier);
+			DynMesh->SetSize3D(DynMesh->GetSize3D() * InSizeMultiplier);
 		}
 	}
 
@@ -404,6 +404,17 @@ void UAvaShapeDynamicMeshBase::SetUsePrimaryMaterialEverywhere(bool bInUse)
 	OnUsePrimaryMaterialEverywhereChanged();
 }
 
+void UAvaShapeDynamicMeshBase::SetUniformScaledSize(float InSize)
+{
+	if (UniformScaledSize == InSize)
+	{
+		return;
+	}
+
+	UniformScaledSize = InSize;
+	OnScaledSizeChanged();
+}
+
 void UAvaShapeDynamicMeshBase::OnUsePrimaryMaterialEverywhereChanged()
 {
 	if (!IsValidMeshIndex(MESH_INDEX_PRIMARY))
diff --git a/Engine/Plugins/Experimental/Avalanche/Source/AvalancheShapes/Private/DynamicMeshes/AvaShapeEllipseDynMesh.cpp b/Engine/Plugins/Experimental/Avalanche/Source/AvalancheShapes/Private/DynamicMeshes/AvaShapeEllipseDynMesh.cpp
--- a/Engine/Plugins/Experimental/Avalanche/Source/AvalancheShapes/Private/DynamicMeshes/AvaShapeEllipseDynMesh.cpp
+++ b/Engine/Plugins/Experimental/Avalanche/Source/AvalancheShapes/Private/DynamicMeshes/AvaShapeEllipseDynMesh.cpp
@@ -5,58 +5,52 @@
 
 const FString UAvaShapeEllipseDynamicMesh::MeshName = TEXT("Ellipse");
 
-bool UAvaShapeEllipseDynamicMesh::SetNumSides(uint8 InNumSides)
+void UAvaShapeEllipseDynamicMesh::SetNumSides(uint8 InNumSides)
 {
 	if (NumSides == InNumSides)
 	{
-		return false;
+		return;
 	}
 
 	if (InNumSides < UAvaShapeEllipseDynamicMesh::MinNumSides || InNumSides > UAvaShapeEllipseDynamicMesh::MaxNumSides)
 	{
-		return false;
+		return;
 	}
 
 	NumSides = InNumSides;
 	OnNumSidesChanged();
-
-	return true;
 }
 
-bool UAvaShapeEllipseDynamicMesh::SetAngleDegree(float InDegree)
+void UAvaShapeEllipseDynamicMesh::SetAngleDegree(float InDegree)
 {
 	if (AngleDegree == InDegree)
 	{
-		return false;
+		return;
 	}
 
 	if (InDegree < 0.f || InDegree > 360.f)
 	{
-		return false;
+		return;
 	}
 
 	AngleDegree = InDegree;
 	OnAngleDegreeChanged();
-
-	return true;
 }
 
-bool UAvaShapeEllipseDynamicMesh::SetStartDegree(float InDegree)
+void UAvaShapeEllipseDynamicMesh::SetStartDegree(float InDegree)
 {
 	if (StartDegree == InDegree)
 	{
-		return false;
+		return;
 	}
 
 	if (InDegree < 0.f || InDegree > 360.f)
 	{
-		return false;
+		return;
 	}
 
 	StartDegree = InDegree;
 	OnStartDegreeChanged();
-
-	return true;
 }
 
 #if WITH_EDITOR
diff --git a/Engine/Plugins/Experimental/Avalanche/Source/AvalancheShapes/Private/DynamicMeshes/AvaShapeIrregularPolygonDynMesh.cpp b/Engine/Plugins/Experimental/Avalanche/Source/AvalancheShapes/Private/DynamicMeshes/AvaShapeIrregularPolygonDynMesh.cpp
--- a/Engine/Plugins/Experimental/Avalanche/Source/AvalancheShapes/Private/DynamicMeshes/AvaShapeIrregularPolygonDynMesh.cpp
+++ b/Engine/Plugins/Experimental/Avalanche/Source/AvalancheShapes/Private/DynamicMeshes/AvaShapeIrregularPolygonDynMesh.cpp
@@ -158,7 +158,7 @@ bool UAvaShapeIrregularPolygonDynamicMesh::DoesLineIntersectBorder(const FVector
 
 		if (DoLinesIntersect(CheckOrigin, CheckEnd, Origin, End))
 		{
-			return true;	
+			return true;
 		}
 	}
 
@@ -337,13 +337,22 @@ void UAvaShapeIrregularPolygonDynamicMesh::RestorePoints()
 	Points.Append(PreEditPoints);
 }
 
-bool UAvaShapeIrregularPolygonDynamicMesh::SetPoints(const TArray<FVector2D>& InPoints)
+void UAvaShapeIrregularPolygonDynamicMesh::SetPoints(const TArray<FVector2D>& InPoints)
 {
 	BackupPoints();
 	Points.Empty();
 	Points.Append(InPoints);
 
-	return CheckNewPointsArray();
+	CheckNewPointsArray();
+}
+
+void UAvaShapeIrregularPolygonDynamicMesh::SetPoints(const TArray<FAvaShapeRoundedCorner>& InPoints)
+{
+	BackupPoints();
+	Points.Empty();
+	Points.Append(InPoints);
+
+	CheckNewPointsArray();
 }
 
 bool UAvaShapeIrregularPolygonDynamicMesh::SetLocation(int32 PointIdx, const FVector2D& InPoint)
@@ -514,17 +523,15 @@ bool UAvaShapeIrregularPolygonDynamicMesh::ShiftPoints(const FVector2D& Amount)
 	return true;
 }
 
-bool UAvaShapeIrregularPolygonDynamicMesh::SetGlobalBevelSize(float InBevelSize)
+void UAvaShapeIrregularPolygonDynamicMesh::SetGlobalBevelSize(float InBevelSize)
 {
 	if (InBevelSize < 0 || InBevelSize > 1)
 	{
-		return false;
+		return;
 	}
 
 	GlobalBevelSize = InBevelSize;
 	OnGlobalBevelSizeChanged();
-
-	return true;
 }
 
 void UAvaShapeIrregularPolygonDynamicMesh::SetGlobalBevelSubdivisions(uint8 InBevelSubdivisions)
@@ -547,7 +554,7 @@ void UAvaShapeIrregularPolygonDynamicMesh::PreEditChange(FProperty* PropertyAbou
 	{
 		return;
 	}
-	
+
 	const FName PropertyName = PropertyAboutToChange->GetFName();
 
 	static const FName PointsName = GET_MEMBER_NAME_CHECKED(UAvaShapeIrregularPolygonDynamicMesh, Points);
@@ -571,7 +578,7 @@ void UAvaShapeIrregularPolygonDynamicMesh::PreEditChange(FProperty* PropertyAbou
 void UAvaShapeIrregularPolygonDynamicMesh::PostEditChangeProperty(FPropertyChangedEvent& PropertyChangedEvent)
 {
 	Super::PostEditChangeProperty(PropertyChangedEvent);
-	
+
 	const FName MemberName = PropertyChangedEvent.GetMemberPropertyName();
 	const FName PropertyName = PropertyChangedEvent.GetPropertyName();
 
@@ -803,7 +810,7 @@ void UAvaShapeIrregularPolygonDynamicMesh::OnBevelSubdivisionsChanged(int32 Poin
 {
 	if (!Points.IsValidIndex(PointIdx))
 	{
-		return;	
+		return;
 	}
 
 	if (Points[PointIdx].Settings.BevelSubdivisions > 0 && Points[PointIdx].Settings.BevelSize == 0.f)
diff --git a/Engine/Plugins/Experimental/Avalanche/Source/AvalancheShapes/Private/DynamicMeshes/AvaShapeLineDynMesh.cpp b/Engine/Plugins/Experimental/Avalanche/Source/AvalancheShapes/Private/DynamicMeshes/AvaShapeLineDynMesh.cpp
--- a/Engine/Plugins/Experimental/Avalanche/Source/AvalancheShapes/Private/DynamicMeshes/AvaShapeLineDynMesh.cpp
+++ b/Engine/Plugins/Experimental/Avalanche/Source/AvalancheShapes/Private/DynamicMeshes/AvaShapeLineDynMesh.cpp
@@ -4,40 +4,36 @@
 
 const FString UAvaShapeLineDynamicMesh::MeshName = TEXT("Line");
 
-bool UAvaShapeLineDynamicMesh::SetLineWidth(float InLineWidth)
+void UAvaShapeLineDynamicMesh::SetLineWidth(float InLineWidth)
 {
 	if (InLineWidth == LineWidth)
 	{
-		return false;
+		return;
 	}
 
 	if (InLineWidth < 0.f)
 	{
-		return false;
+		return;
 	}
 
 	LineWidth = InLineWidth;
 	OnLineWidthChanged();
-
-	return true;
 }
 
-bool UAvaShapeLineDynamicMesh::SetVector(const FVector2D& InVector)
+void UAvaShapeLineDynamicMesh::SetVector(const FVector2D& InVector)
 {
 	if (InVector == Vector)
 	{
-		return false;
+		return;
 	}
 
 	if (InVector == FVector2D::ZeroVector)
 	{
-		return false;
+		return;
 	}
 
 	Vector = InVector;
 	OnVectorChanged();
-
-	return true;
 }
 
 void UAvaShapeLineDynamicMesh::GenerateBorderVertices(TArray<FVector2D>& BorderVertices)
diff --git a/Engine/Plugins/Experimental/Avalanche/Source/AvalancheShapes/Private/DynamicMeshes/AvaShapeNGonDynMesh.cpp b/Engine/Plugins/Experimental/Avalanche/Source/AvalancheShapes/Private/DynamicMeshes/AvaShapeNGonDynMesh.cpp
--- a/Engine/Plugins/Experimental/Avalanche/Source/AvalancheShapes/Private/DynamicMeshes/AvaShapeNGonDynMesh.cpp
+++ b/Engine/Plugins/Experimental/Avalanche/Source/AvalancheShapes/Private/DynamicMeshes/AvaShapeNGonDynMesh.cpp
@@ -4,21 +4,20 @@
 
 const FString UAvaShapeNGonDynamicMesh::MeshName = TEXT("RegularPolygon");
 
-bool UAvaShapeNGonDynamicMesh::SetNumSides(uint8 InNumSides)
+void UAvaShapeNGonDynamicMesh::SetNumSides(uint8 InNumSides)
 {
 	if (NumSides == InNumSides)
 	{
-		return false;
+		return;
 	}
 
 	if (InNumSides < UAvaShapeNGonDynamicMesh::MinNumSides || InNumSides > UAvaShapeNGonDynamicMesh::MaxNumSides)
 	{
-		return false;
+		return;
 	}
 
 	NumSides = InNumSides;
 	OnNumSidesChanged();
-	return true;
 }
 
 #if WITH_EDITOR
diff --git a/Engine/Plugins/Experimental/Avalanche/Source/AvalancheShapes/Private/DynamicMeshes/AvaShapeRectangleDynMesh.cpp b/Engine/Plugins/Experimental/Avalanche/Source/AvalancheShapes/Private/DynamicMeshes/AvaShapeRectangleDynMesh.cpp
--- a/Engine/Plugins/Experimental/Avalanche/Source/AvalancheShapes/Private/DynamicMeshes/AvaShapeRectangleDynMesh.cpp
+++ b/Engine/Plugins/Experimental/Avalanche/Source/AvalancheShapes/Private/DynamicMeshes/AvaShapeRectangleDynMesh.cpp
@@ -6,30 +6,26 @@
 
 const FString UAvaShapeRectangleDynamicMesh::MeshName = TEXT("Rectangle");
 
-bool UAvaShapeRectangleDynamicMesh::SetHorizontalAlignment(EAvaHorizontalAlignment InHorizontalAlignment)
+void UAvaShapeRectangleDynamicMesh::SetHorizontalAlignment(EAvaHorizontalAlignment InHorizontalAlignment)
 {
 	if (HorizontalAlignment == InHorizontalAlignment)
 	{
-		return false;
+		return;
 	}
 
 	HorizontalAlignment = InHorizontalAlignment;
 	OnAlignmentChanged();
-
-	return true;
 }
 
-bool UAvaShapeRectangleDynamicMesh::SetVerticalAlignment(EAvaVerticalAlignment InVerticalAlignment)
+void UAvaShapeRectangleDynamicMesh::SetVerticalAlignment(EAvaVerticalAlignment InVerticalAlignment)
 {
 	if (VerticalAlignment == InVerticalAlignment)
 	{
-		return false;
+		return;
 	}
 
 	VerticalAlignment = InVerticalAlignment;
 	OnAlignmentChanged();
-
-	return true;
 }
 
 void UAvaShapeRectangleDynamicMesh::SetBottomRightBevelSubdivisions(uint8 InBevelSubdivisions)
@@ -91,7 +87,7 @@ FVector2D UAvaShapeRectangleDynamicMesh::GetValidRangeLeftSlantAngle() const
 	BottomLeftCorner.X = -MaxX;
 
 	const float MinLeftSlant = -90 + FMath::RadiansToDegrees(FMath::Atan2(TopLeftCorner.Y - BottomLeftCorner.Y, TopLeftCorner.X - BottomLeftCorner.X));
-	
+
 	return FVector2D(MinLeftSlant, MaxLeftSlant);
 }
 
@@ -102,19 +98,19 @@ FVector2D UAvaShapeRectangleDynamicMesh::GetValidRangeRightSlantAngle() const
 	// Flat slant = 0 deg
 	FVector2D TopRightCorner(Size2D.X, Size2D.Y);
 	FVector2D BottomRightCorner(Size2D.X, 0.f);
-	
+
 	// Max Slant
 	TopRightCorner.X = Size2D.X + MaxX;
 	BottomRightCorner.X = Size2D.X;
 
 	const float MaxLeftSlant = 90 - FMath::RadiansToDegrees(FMath::Atan2(TopRightCorner.Y - BottomRightCorner.Y, TopRightCorner.X - BottomRightCorner.X));
-	
+
 	// Min Slant
 	TopRightCorner.X = Size2D.X;
 	BottomRightCorner.X = Size2D.X - MaxX;
 
 	const float MinLeftSlant = -90 + FMath::RadiansToDegrees(FMath::Atan2(TopRightCorner.Y - BottomRightCorner.Y, TopRightCorner.X - BottomRightCorner.X));
-	
+
 	return FVector2D(MinLeftSlant, MaxLeftSlant);
 }
 
@@ -131,15 +127,15 @@ void UAvaShapeRectangleDynamicMesh::GetValidSlantAngle(float& OutLeftSlant, floa
 	{
 		return;
 	}
-	
+
 	const float MaxSizeX = Size2D.X;
-	
+
 	// check corners do not exceed max size
 	const float LeftSlantSizeX = Size2D.Y * FMath::Tan(FMath::DegreesToRadians(OutLeftSlant));
 	const float RightSlantSizeX = -Size2D.Y * FMath::Tan(FMath::DegreesToRadians(OutRightSlant));
-	
+
 	const float SizeX = FMath::Abs(LeftSlantSizeX + RightSlantSizeX);
-	
+
 	if (SizeX >= MaxSizeX)
 	{
 		const float SizeRatio = MaxSizeX / SizeX;
@@ -150,53 +146,47 @@ void UAvaShapeRectangleDynamicMesh::GetValidSlantAngle(float& OutLeftSlant, floa
 	}
 }
 
-bool UAvaShapeRectangleDynamicMesh::SetLeftSlant(float InSlant)
+void UAvaShapeRectangleDynamicMesh::SetLeftSlant(float InSlant)
 {
 	if (LeftSlant == InSlant)
 	{
-		return false;
+		return;
 	}
 
 	if (InSlant < MinSlantAngle || InSlant > MaxSlantAngle)
 	{
-		return false;
+		return;
 	}
 
 	LeftSlant = InSlant;
 	OnLeftSlantChanged();
-
-	return true;
 }
 
-bool UAvaShapeRectangleDynamicMesh::SetRightSlant(float InSlant)
+void UAvaShapeRectangleDynamicMesh::SetRightSlant(float InSlant)
 {
 	if (RightSlant == InSlant)
 	{
-		return false;
+		return;
 	}
 
 	if (InSlant < MinSlantAngle || InSlant > MaxSlantAngle)
 	{
-		return false;
+		return;
 	}
 
 	RightSlant = InSlant;
 	OnRightSlantChanged();
-
-	return true;
 }
 
-bool UAvaShapeRectangleDynamicMesh::SetGlobalBevelSize(float InBevelSize)
+void UAvaShapeRectangleDynamicMesh::SetGlobalBevelSize(float InBevelSize)
 {
 	if (InBevelSize < 0 || InBevelSize > GetMaximumBevelSize())
 	{
-		return false;
+		return;
 	}
 
 	GlobalBevelSize = InBevelSize;
 	OnGlobalBevelSizeChanged();
-
-	return true;
 }
 
 void UAvaShapeRectangleDynamicMesh::SetGlobalBevelSubdivisions(uint8 InGlobalBevelSubdivisions)
@@ -238,35 +228,31 @@ void UAvaShapeRectangleDynamicMesh::SetBottomRight(const FAvaShapeRectangleCorne
 	SetBottomRightBevelSubdivisions(InCornerSettings.BevelSubdivisions);
 }
 
-bool UAvaShapeRectangleDynamicMesh::SetTopLeftCornerType(EAvaShapeCornerType InType)
+void UAvaShapeRectangleDynamicMesh::SetTopLeftCornerType(EAvaShapeCornerType InType)
 {
 	if (TopLeft.Type == InType)
 	{
-		return false;
+		return;
 	}
 
 	TopLeft.Type = InType;
 	OnTopLeftCornerTypeChanged();
-
-	return true;
 }
 
-bool UAvaShapeRectangleDynamicMesh::SetTopLeftBevelSize(float InSize)
+void UAvaShapeRectangleDynamicMesh::SetTopLeftBevelSize(float InSize)
 {
 	if (InSize == TopLeft.BevelSize)
 	{
-		return false;
+		return;
 	}
 
 	if (InSize < 0 || InSize > GetMaximumBevelSize())
 	{
-		return false;
+		return;
 	}
 
 	TopLeft.BevelSize = InSize;
 	OnTopLeftBevelSizeChanged();
-
-	return true;
 }
 
 void UAvaShapeRectangleDynamicMesh::SetTopLeftBevelSubdivisions(uint8 InBevelSubdivisions)
@@ -285,35 +271,31 @@ void UAvaShapeRectangleDynamicMesh::SetTopLeftBevelSubdivisions(uint8 InBevelSub
 	OnTopLeftBevelSubdivisionsChanged();
 }
 
-bool UAvaShapeRectangleDynamicMesh::SetTopRightCornerType(EAvaShapeCornerType InType)
+void UAvaShapeRectangleDynamicMesh::SetTopRightCornerType(EAvaShapeCornerType InType)
 {
 	if (TopRight.Type == InType)
 	{
-		return false;
+		return;
 	}
 
 	TopRight.Type = InType;
 	OnTopRightCornerTypeChanged();
-
-	return true;
 }
 
-bool UAvaShapeRectangleDynamicMesh::SetTopRightBevelSize(float InSize)
+void UAvaShapeRectangleDynamicMesh::SetTopRightBevelSize(float InSize)
 {
 	if (InSize == TopRight.BevelSize)
 	{
-		return false;
+		return;
 	}
 
 	if (InSize < 0 || InSize > GetMaximumBevelSize())
 	{
-		return false;
+		return;
 	}
 
 	TopRight.BevelSize = InSize;
 	OnTopRightBevelSizeChanged();
-
-	return true;
 }
 
 void UAvaShapeRectangleDynamicMesh::SetTopRightBevelSubdivisions(uint8 InBevelSubdivisions)
@@ -332,35 +314,31 @@ void UAvaShapeRectangleDynamicMesh::SetTopRightBevelSubdivisions(uint8 InBevelSu
 	OnTopRightBevelSubdivisionsChanged();
 }
 
-bool UAvaShapeRectangleDynamicMesh::SetBottomLeftCornerType(EAvaShapeCornerType InType)
+void UAvaShapeRectangleDynamicMesh::SetBottomLeftCornerType(EAvaShapeCornerType InType)
 {
 	if (BottomLeft.Type == InType)
 	{
-		return false;
+		return;
 	}
-	
+
 	BottomLeft.Type = InType;
 	OnBottomLeftCornerTypeChanged();
-
-	return true;
 }
 
-bool UAvaShapeRectangleDynamicMesh::SetBottomLeftBevelSize(float InSize)
+void UAvaShapeRectangleDynamicMesh::SetBottomLeftBevelSize(float InSize)
 {
 	if (InSize == BottomLeft.BevelSize)
 	{
-		return false;
+		return;
 	}
 
 	if (InSize < 0 || InSize > GetMaximumBevelSize())
 	{
-		return false;
+		return;
 	}
 
 	BottomLeft.BevelSize = InSize;
 	OnBottomLeftBevelSizeChanged();
-
-	return true;
 }
 
 void UAvaShapeRectangleDynamicMesh::SetBottomLeftBevelSubdivisions(uint8 InBevelSubdivisions)
@@ -379,35 +357,31 @@ void UAvaShapeRectangleDynamicMesh::SetBottomLeftBevelSubdivisions(uint8 InBevel
 	OnBottomLeftBevelSubdivisionsChanged();
 }
 
-bool UAvaShapeRectangleDynamicMesh::SetBottomRightCornerType(EAvaShapeCornerType InType)
+void UAvaShapeRectangleDynamicMesh::SetBottomRightCornerType(EAvaShapeCornerType InType)
 {
 	if (BottomRight.Type == InType)
 	{
-		return false;
+		return;
 	}
 
 	BottomRight.Type = InType;
 	OnBottomRightCornerTypeChanged();
-
-	return true;
 }
 
-bool UAvaShapeRectangleDynamicMesh::SetBottomRightBevelSize(float InSize)
+void UAvaShapeRectangleDynamicMesh::SetBottomRightBevelSize(float InSize)
 {
 	if (InSize == BottomRight.BevelSize)
 	{
-		return false;
+		return;
 	}
 
 	if (InSize < 0 || InSize > GetMaximumBevelSize())
 	{
-		return false;
+		return;
 	}
 
 	BottomRight.BevelSize = InSize;
 	OnBottomRightBevelSizeChanged();
-
-	return true;
 }
 
 void UAvaShapeRectangleDynamicMesh::OnAlignmentChanged()
@@ -482,7 +456,7 @@ void UAvaShapeRectangleDynamicMesh::OnTopLeftBevelSubdivisionsChanged()
 		if (TopLeft.BevelSize == 0.f)
 		{
 			TopLeft.BevelSize = GetMaximumBevelSize() / 4.f;
-		}	
+		}
 	}
 
 	MarkAllMeshesDirty();
@@ -528,9 +502,9 @@ void UAvaShapeRectangleDynamicMesh::OnBottomLeftBevelSubdivisionsChanged()
 		if (BottomLeft.BevelSize == 0.f)
 		{
 			BottomLeft.BevelSize = GetMaximumBevelSize() / 4.f;
-		}	
+		}
 	}
-	
+
 	MarkAllMeshesDirty();
 }
 
@@ -574,9 +548,9 @@ void UAvaShapeRectangleDynamicMesh::OnTopRightBevelSubdivisionsChanged()
 		if (TopRight.BevelSize == 0.f)
 		{
 			TopRight.BevelSize = GetMaximumBevelSize() / 4.f;
-		}	
+		}
 	}
-	
+
 	MarkAllMeshesDirty();
 }
 
@@ -620,9 +594,9 @@ void UAvaShapeRectangleDynamicMesh::OnBottomRightBevelSubdivisionsChanged()
 		if (BottomRight.BevelSize == 0.f)
 		{
 			BottomRight.BevelSize = GetMaximumBevelSize() / 4.f;
-		}	
+		}
 	}
-	
+
 	MarkAllMeshesDirty();
 }
 
@@ -664,7 +638,7 @@ bool UAvaShapeRectangleDynamicMesh::GenerateBaseMeshSections(FAvaShapeMesh& Prim
 	}
 
 	const float MaxBevelSize = GetMaximumBevelSize();
-	
+
 	FAvaShapeRectangleCornerSettings TopLeftCopy = TopLeft;
 	FAvaShapeRectangleCornerSettings TopRightCopy = TopRight;
 	FAvaShapeRectangleCornerSettings BottomLeftCopy = BottomLeft;
@@ -674,7 +648,7 @@ bool UAvaShapeRectangleDynamicMesh::GenerateBaseMeshSections(FAvaShapeMesh& Prim
 	TopRightCopy.BevelSize = FMath::Min(TopRightCopy.BevelSize, MaxBevelSize);
 	BottomLeftCopy.BevelSize = FMath::Min(BottomLeftCopy.BevelSize, MaxBevelSize);
 	BottomRightCopy.BevelSize = FMath::Min(BottomRightCopy.BevelSize, MaxBevelSize);
-	
+
 	float LeftSlantCopy;
 	float RightSlantCopy;
 	GetValidSlantAngle(LeftSlantCopy, RightSlantCopy);
@@ -760,7 +734,7 @@ bool UAvaShapeRectangleDynamicMesh::GenerateBaseMeshSections(FAvaShapeMesh& Prim
 
 		if (!bTopRightCornerCurveOut)
 		{
-			TopRightCorner.Location.X -= BottomRightCopy.BevelSize;	
+			TopRightCorner.Location.X -= BottomRightCopy.BevelSize;
 		}
 	}
 
@@ -1273,7 +1247,7 @@ void UAvaShapeRectangleDynamicMesh::PostEditChangeProperty(FPropertyChangedEvent
 void UAvaShapeRectangleDynamicMesh::OnSizeChanged()
 {
 	Super::OnSizeChanged();
-	
+
 	//Only Mark Meshes Dirty any of the corners has bevels
 	if (TopLeft.IsBeveled()
 		|| TopRight.IsBeveled()
diff --git a/Engine/Plugins/Experimental/Avalanche/Source/AvalancheShapes/Private/DynamicMeshes/AvaShapeRingDynMesh.cpp b/Engine/Plugins/Experimental/Avalanche/Source/AvalancheShapes/Private/DynamicMeshes/AvaShapeRingDynMesh.cpp
--- a/Engine/Plugins/Experimental/Avalanche/Source/AvalancheShapes/Private/DynamicMeshes/AvaShapeRingDynMesh.cpp
+++ b/Engine/Plugins/Experimental/Avalanche/Source/AvalancheShapes/Private/DynamicMeshes/AvaShapeRingDynMesh.cpp
@@ -4,81 +4,73 @@
 
 const FString UAvaShapeRingDynamicMesh::MeshName = TEXT("Ring");
 
-bool UAvaShapeRingDynamicMesh::SetNumSides(uint8 InNumSides)
+void UAvaShapeRingDynamicMesh::SetNumSides(uint8 InNumSides)
 {
 	if (NumSides == InNumSides)
 	{
-		return false;
+		return;
 	}
 
 	if (InNumSides < UAvaShapeRingDynamicMesh::MinNumSides || InNumSides > UAvaShapeRingDynamicMesh::MaxNumSides)
 	{
-		return false;
+		return;
 	}
 
 	NumSides = InNumSides;
 	OnNumSidesChanged();
-
-	return true;
 }
 
-bool UAvaShapeRingDynamicMesh::SetInnerSize(float InInnerSize)
+void UAvaShapeRingDynamicMesh::SetInnerSize(float InInnerSize)
 {
 	if (InnerSize == InInnerSize)
 	{
-		return false;
+		return;
 	}
 
 	if (InInnerSize < 0.01f)
 	{
-		return false;
+		return;
 	}
 
 	if (InInnerSize > 0.99f)
 	{
-		return false;
+		return;
 	}
 
 	InnerSize = InInnerSize;
 	OnInnerSizeChanged();
-
-	return true;
 }
 
-bool UAvaShapeRingDynamicMesh::SetAngleDegree(float InDegree)
+void UAvaShapeRingDynamicMesh::SetAngleDegree(float InDegree)
 {
 	if (AngleDegree == InDegree)
 	{
-		return false;
+		return;
 	}
 
 	if (InDegree < 0.f || InDegree > 360.f)
 	{
-		return false;
+		return;
 	}
 
 	AngleDegree = InDegree;
 	OnAngleDegreeChanged();
-
-	return true;
 }
 
-bool UAvaShapeRingDynamicMesh::SetStartDegree(float InDegree)
+void UAvaShapeRingDynamicMesh::SetStartDegree(float InDegree)
 {
 	if (StartDegree == InDegree)
 	{
-		return false;
+		return;
 	}
 
 	if (InDegree < 0.f || InDegree > 360.f)
 	{
-		return false;
+		return;
 	}
 
 	StartDegree = InDegree;
 	OnStartDegreeChanged();
-
-	return true;
 }
 
 #if WITH_EDITOR
diff --git a/Engine/Plugins/Experimental/Avalanche/Source/AvalancheShapes/Private/DynamicMeshes/AvaShapeRoundedPolygonDynMesh.cpp b/Engine/Plugins/Experimental/Avalanche/Source/AvalancheShapes/Private/DynamicMeshes/AvaShapeRoundedPolygonDynMesh.cpp
--- a/Engine/Plugins/Experimental/Avalanche/Source/AvalancheShapes/Private/DynamicMeshes/AvaShapeRoundedPolygonDynMesh.cpp
+++ b/Engine/Plugins/Experimental/Avalanche/Source/AvalancheShapes/Private/DynamicMeshes/AvaShapeRoundedPolygonDynMesh.cpp
@@ -5,22 +5,20 @@
 #include "DynamicMeshes/AvaShapeRectangleDynMesh.h"
 #include "Kismet/KismetMathLibrary.h"
 
-bool UAvaShapeRoundedPolygonDynamicMesh::SetBevelSize(float InBevelSize)
+void UAvaShapeRoundedPolygonDynamicMesh::SetBevelSize(float InBevelSize)
 {
 	if (BevelSize == InBevelSize)
 	{
-		return false;
+		return;
 	}
 
 	if (InBevelSize < 0.f || InBevelSize >= FMath::Min(Size2D.X, Size2D.Y))
 	{
-		return false;
+		return;
 	}
 
 	BevelSize = InBevelSize;
 	OnRoundedRadiusChanged();
-
-	return true;
 }
 
 void UAvaShapeRoundedPolygonDynamicMesh::SetBevelSubdivisions(uint8 InBevelSubdivisions)
@@ -65,7 +63,7 @@ void UAvaShapeRoundedPolygonDynamicMesh::PostEditChangeProperty(FPropertyChanged
 	Super::PostEditChangeProperty(PropertyChangedEvent);
 
 	const FName MemberName = PropertyChangedEvent.GetMemberPropertyName();
-	
+
 	static const FName BevelSizeName = GET_MEMBER_NAME_CHECKED(UAvaShapeRoundedPolygonDynamicMesh, BevelSize);
 	static const FName BevelSubdivisionsName = GET_MEMBER_NAME_CHECKED(UAvaShapeRoundedPolygonDynamicMesh, BevelSubdivisions);
 
@@ -83,7 +81,7 @@ void UAvaShapeRoundedPolygonDynamicMesh::PostEditChangeProperty(FPropertyChanged
 void UAvaShapeRoundedPolygonDynamicMesh::OnSizeChanged()
 {
 	Super::OnSizeChanged();
-	
+
 	// Only Mark Meshes Dirty if we have bevels on it
 	if (BevelSize > 0.f && BevelSubdivisions > 0)
 	{
diff --git a/Engine/Plugins/Experimental/Avalanche/Source/AvalancheShapes/Private/DynamicMeshes/AvaShapeSphereDynMesh.cpp b/Engine/Plugins/Experimental/Avalanche/Source/AvalancheShapes/Private/DynamicMeshes/AvaShapeSphereDynMesh.cpp
--- a/Engine/Plugins/Experimental/Avalanche/Source/AvalancheShapes/Private/DynamicMeshes/AvaShapeSphereDynMesh.cpp
+++ b/Engine/Plugins/Experimental/Avalanche/Source/AvalancheShapes/Private/DynamicMeshes/AvaShapeSphereDynMesh.cpp
@@ -7,106 +7,96 @@
 
 const FString UAvaShapeSphereDynamicMesh::MeshName = TEXT("Sphere");
 
-bool UAvaShapeSphereDynamicMesh::SetNumSides(uint8 InNumSides)
+void UAvaShapeSphereDynamicMesh::SetNumSides(uint8 InNumSides)
 {
 	if (NumSides == InNumSides)
 	{
-		return false;
+		return;
 	}
 
 	if (InNumSides < UAvaShapeSphereDynamicMesh::MinNumSides || InNumSides > UAvaShapeSphereDynamicMesh::MaxNumSides)
 	{
-		return false;
+		return;
 	}
 
 	NumSides = InNumSides;
 	OnNumSidesChanged();
-
-	return true;
 }
 
-bool UAvaShapeSphereDynamicMesh::SetStartLatitude(float InDegree)
+void UAvaShapeSphereDynamicMesh::SetStartLatitude(float InDegree)
 {
 	if (StartLatitude == InDegree)
 	{
-		return false;
+		return;
 	}
 
 	if (InDegree < 0.0f || InDegree > 360.f)
 	{
-		return false;
+		return;
 	}
 
 	StartLatitude = InDegree;
 	OnStartLatitudeChanged();
-
-	return true;
 }
 
-bool UAvaShapeSphereDynamicMesh::SetLatitudeDegree(float InDegree)
+void UAvaShapeSphereDynamicMesh::SetLatitudeDegree(float InDegree)
 {
 	if (LatitudeDegree == InDegree)
 	{
-		return false;
+		return;
 	}
 
 	if (InDegree < 0.0f || InDegree > 360.f)
 	{
-		return false;
+		return;
 	}
 
 	LatitudeDegree = InDegree;
 	OnLatitudeDegreeChanged();
-
-	return true;
 }
 
-bool UAvaShapeSphereDynamicMesh::SetStartLongitude(float InDegree)
+void UAvaShapeSphereDynamicMesh::SetStartLongitude(float InDegree)
 {
 	if (StartLongitude == InDegree)
 	{
-		return false;
+		return;
 	}
 
 	if (InDegree < 0.0f || InDegree > 180.f)
 	{
-		return false;
+		return;
 	}
 
 	if (InDegree > EndLongitude)
 	{
-		return false;
+		return;
 	}
 
 	PreEditStartLongitude = StartLongitude;
 	StartLongitude = InDegree;
 	OnStartLongitudeChanged();
-
-	return true;
 }
 
-bool UAvaShapeSphereDynamicMesh::SetEndLongitude(float InDegree)
+void UAvaShapeSphereDynamicMesh::SetEndLongitude(float InDegree)
 {
 	if (EndLongitude == InDegree)
 	{
-		return false;
+		return;
 	}
 
 	if (InDegree < 0.01f || InDegree > 180.f)
 	{
-		return false;
+		return;
 	}
 
 	if (InDegree < StartLongitude)
 	{
-		return false;
+		return;
 	}
 
 	PreEditEndLongitude = EndLongitude;
 	EndLongitude = InDegree;
 	OnEndLongitudeChanged();
-
-	return true;
 }
 
 #if WITH_EDITOR
diff --git a/Engine/Plugins/Experimental/Avalanche/Source/AvalancheShapes/Private/DynamicMeshes/AvaShapeStarDynMesh.cpp b/Engine/Plugins/Experimental/Avalanche/Source/AvalancheShapes/Private/DynamicMeshes/AvaShapeStarDynMesh.cpp
--- a/Engine/Plugins/Experimental/Avalanche/Source/AvalancheShapes/Private/DynamicMeshes/AvaShapeStarDynMesh.cpp
+++ b/Engine/Plugins/Experimental/Avalanche/Source/AvalancheShapes/Private/DynamicMeshes/AvaShapeStarDynMesh.cpp
@@ -4,47 +4,43 @@
 
 const FString UAvaShapeStarDynamicMesh::MeshName = TEXT("Star");
 
-bool UAvaShapeStarDynamicMesh::SetNumPoints(uint8 InNumPoints)
+void UAvaShapeStarDynamicMesh::SetNumPoints(uint8 InNumPoints)
 {
 	if (NumPoints == InNumPoints)
 	{
-		return false;
+		return;
 	}
 
 	if (InNumPoints < UAvaShapeStarDynamicMesh::MinNumPoints || InNumPoints > UAvaShapeStarDynamicMesh::MaxNumPoints)
 	{
-		return false;
+		return;
 	}
 
 	NumPoints = InNumPoints;
 	OnNumSidesChanged();
-
-	return true;
 }
 
-bool UAvaShapeStarDynamicMesh::SetInnerSize(float InInnerSize)
+void UAvaShapeStarDynamicMesh::SetInnerSize(float InInnerSize)
 {
 	if (InnerSize == InInnerSize)
 	{
-		return false;
+		return;
 	}
 
 	if (InInnerSize < 0.0f || InInnerSize > 0.99f)
 	{
-		return false;
+		return;
 	}
 
 	InnerSize = InInnerSize;
 	OnInnerSizeChanged();
-
-	return true;
 }
 
 #if WITH_EDITOR
 void UAvaShapeStarDynamicMesh::PostEditChangeProperty(FPropertyChangedEvent& PropertyChangedEvent)
 {
 	Super::PostEditChangeProperty(PropertyChangedEvent);
-	
+
 	static FName NumPointsName = GET_MEMBER_NAME_CHECKED(UAvaShapeStarDynamicMesh, NumPoints);
 	static FName InnerSizeName = GET_MEMBER_NAME_CHECKED(UAvaShapeStarDynamicMesh, InnerSize);
 
diff --git a/Engine/Plugins/Experimental/Avalanche/Source/AvalancheShapes/Private/DynamicMeshes/AvaShapeTorusDynMesh.cpp b/Engine/Plugins/Experimental/Avalanche/Source/AvalancheShapes/Private/DynamicMeshes/AvaShapeTorusDynMesh.cpp
--- a/Engine/Plugins/Experimental/Avalanche/Source/AvalancheShapes/Private/DynamicMeshes/AvaShapeTorusDynMesh.cpp
+++ b/Engine/Plugins/Experimental/Avalanche/Source/AvalancheShapes/Private/DynamicMeshes/AvaShapeTorusDynMesh.cpp
@@ -7,94 +7,84 @@
 
 const FString UAvaShapeTorusDynamicMesh::MeshName = TEXT("Torus");
 
-bool UAvaShapeTorusDynamicMesh::SetNumSides(uint8 InNumSides)
+void UAvaShapeTorusDynamicMesh::SetNumSides(uint8 InNumSides)
 {
 	if (NumSides == InNumSides)
 	{
-		return false;
+		return;
 	}
 
 	if (InNumSides < UAvaShapeTorusDynamicMesh::MinNumSides || InNumSides > UAvaShapeTorusDynamicMesh::MaxNumSides)
 	{
-		return false;
+		return;
 	}
 
 	NumSides = InNumSides;
 	OnNumSidesChanged();
-
-	return true;
 }
 
-bool UAvaShapeTorusDynamicMesh::SetNumSlices(uint8 InNumSlices)
+void UAvaShapeTorusDynamicMesh::SetNumSlices(uint8 InNumSlices)
 {
 	if (NumSlices == InNumSlices)
 	{
-		return false;
+		return;
 	}
 
 	if (InNumSlices < UAvaShapeTorusDynamicMesh::MinNumSlices || InNumSlices > UAvaShapeTorusDynamicMesh::MaxNumSlices)
 	{
-		return false;
+		return;
 	}
 
 	NumSlices = InNumSlices;
 	OnNumSlicesChanged();
-
-	return true;
 }
 
-bool UAvaShapeTorusDynamicMesh::SetInnerSize(float InInnerSize)
+void UAvaShapeTorusDynamicMesh::SetInnerSize(float InInnerSize)
 {
 	if (InnerSize == InInnerSize)
 	{
-		return false;
+		return;
 	}
 
 	if (InInnerSize < 0.5f || InInnerSize > 0.99f)
 	{
-		return false;
+		return;
 	}
 
 	InnerSize = InInnerSize;
 	OnInnerSizeChanged();
-
-	return true;
 }
 
-bool UAvaShapeTorusDynamicMesh::SetAngleDegree(float InAngleDegree)
+void UAvaShapeTorusDynamicMesh::SetAngleDegree(float InAngleDegree)
 {
 	if (AngleDegree == InAngleDegree)
 	{
-		return false;
+		return;
 	}
 
 	if (InAngleDegree < 0.f || InAngleDegree > 360.f)
 	{
-		return false;
+		return;
 	}
 
 	AngleDegree = InAngleDegree;
 	OnAngleDegreeChanged();
-
-	return true;
 }
 
-bool UAvaShapeTorusDynamicMesh::SetStartDegree(float InStartDegree)
+void UAvaShapeTorusDynamicMesh::SetStartDegree(float InStartDegree)
 {
 	if (StartDegree == InStartDegree)
 	{
-		return false;
+		return;
 	}
 
 	if (InStartDegree < 0.f || InStartDegree > 360.f)
 	{
-		return false;
+		return;
 	}
 
 	StartDegree = InStartDegree;
 	OnStartDegreeChanged();
-
-	return true;
 }
 
 #if WITH_EDITOR
diff --git a/Engine/Plugins/Experimental/Avalanche/Source/AvalancheShapes/Public/AvaShapeUVParameters.h b/Engine/Plugins/Experimental/Avalanche/Source/AvalancheShapes/Public/AvaShapeUVParameters.h
--- a/Engine/Plugins/Experimental/Avalanche/Source/AvalancheShapes/Public/AvaShapeUVParameters.h
+++ b/Engine/Plugins/Experimental/Avalanche/Source/AvalancheShapes/Public/AvaShapeUVParameters.h
@@ -7,7 +7,7 @@
 #include "AvaShapeUVParameters.generated.h"
 
 USTRUCT(BlueprintType)
-struct AVALANCHESHAPES_API FAvaShapeMaterialUVParameters
+struct FAvaShapeMaterialUVParameters
 {
 	friend class UAvaShapeDynamicMeshBase;
 	friend class FAvaShapeDynamicMeshVisualizer;
diff --git a/Engine/Plugins/Experimental/Avalanche/Source/AvalancheShapes/Public/DynamicMeshes/AvaShape2DArrowDynMesh.h b/Engine/Plugins/Experimental/Avalanche/Source/AvalancheShapes/Public/DynamicMeshes/AvaShape2DArrowDynMesh.h
--- a/Engine/Plugins/Experimental/Avalanche/Source/AvalancheShapes/Public/DynamicMeshes/AvaShape2DArrowDynMesh.h
+++ b/Engine/Plugins/Experimental/Avalanche/Source/AvalancheShapes/Public/DynamicMeshes/AvaShape2DArrowDynMesh.h
@@ -5,16 +5,17 @@
 #include "DynamicMeshes/AvaShape2DDynMeshBase.h"
 #include "AvaShape2DArrowDynMesh.generated.h"
 
-UCLASS(ClassGroup="Shape", BlueprintType, CustomConstructor, Within=AvaShapeActor)
-class AVALANCHESHAPES_API UAvaShape2DArrowDynamicMesh : public UAvaShape2DDynMeshBase
+UCLASS(MinimalAPI, ClassGroup="Shape", BlueprintType, CustomConstructor, Within=AvaShapeActor)
+class UAvaShape2DArrowDynamicMesh : public UAvaShape2DDynMeshBase
 {
 	GENERATED_BODY()
 
 public:
+	static const FString MeshName;
+
 	UAvaShape2DArrowDynamicMesh()
 		: UAvaShape2DArrowDynamicMesh(FVector2D(50.f, 50.f))
-	{
-	}
+	{}
 
 	UAvaShape2DArrowDynamicMesh(const FVector2D& Size2D, const FLinearColor& InVertexColor = FLinearColor::White)
 		: UAvaShape2DDynMeshBase(Size2D, InVertexColor)
@@ -23,36 +24,31 @@ public:
 		, RatioArrowY(0.5f)
 		, RatioLineY(0.5f)
 		, bBothSideArrows(false)
-	{
-	}
+	{}
 
-	static const FString MeshName;
 	virtual const FString& GetMeshName() const override { return MeshName; }
 
-	UFUNCTION()
-	bool SetRatioArrowLine(float InRatio);
+	AVALANCHESHAPES_API void SetRatioArrowLine(float InRatio);
 	float GetRatioArrowLine() const { return RatioArrowLine; }
 
-	UFUNCTION()
-	bool SetRatioLineHeight(float InRatio);
+	AVALANCHESHAPES_API void SetRatioLineHeight(float InRatio);
 	float GetRatioLineHeight() const { return RatioLineHeight; }
 
-	UFUNCTION()
-	bool SetRatioArrowY(float InRatio);
+	AVALANCHESHAPES_API void SetRatioArrowY(float InRatio);
 	float GetRatioArrowY() const { return RatioArrowY; }
 
-	UFUNCTION()
-	bool SetRatioLineY(float InRatio);
+	AVALANCHESHAPES_API void SetRatioLineY(float InRatio);
 	float GetRatioLineY() const { return RatioLineY; }
 
-	UFUNCTION()
-	bool SetBothSideArrows(bool bBothSide);
+	AVALANCHESHAPES_API void SetBothSideArrows(bool bBothSide);
 	bool IsBothSideArrows() const { return bBothSideArrows; }
 
 protected:
+	//~ Begin UObject
 #if WITH_EDITOR
 	virtual void PostEditChangeProperty(FPropertyChangedEvent& PropertyChangedEvent) override;
 #endif
+	//~ End UObject
 
 	virtual void OnRatioArrowLineChanged();
 	virtual void OnRatioLineHeightChanged();
@@ -64,22 +60,22 @@ protected:
 	virtual bool CreateMesh(FAvaShapeMesh& InMesh) override;
 
 	// represents the ratio for the arrow and the line, 0.6 means 60% arrow and 40% line
-	UPROPERTY(EditInstanceOnly, BlueprintReadWrite, Category="Shape", meta=(ClampMin="0.0", ClampMax="1.0", DisplayName="Ratio Arrow Line", AllowPrivateAccess="true"))
+	UPROPERTY(EditInstanceOnly, BlueprintReadWrite, Category="Shape", Setter, Getter, meta=(ClampMin="0.0", ClampMax="1.0", DisplayName="Ratio Arrow Line", AllowPrivateAccess="true"))
 	float RatioArrowLine;
 
 	// represents the ratio for the line height, 1 means 100% of the height available
-	UPROPERTY(EditInstanceOnly, BlueprintReadWrite, Category="Shape", meta=(ClampMin="0.0", ClampMax="1.0", DisplayName="Ratio Line Height", AllowPrivateAccess="true"))
+	UPROPERTY(EditInstanceOnly, BlueprintReadWrite, Category="Shape", Setter, Getter, meta=(ClampMin="0.0", ClampMax="1.0", DisplayName="Ratio Line Height", AllowPrivateAccess="true"))
 	float RatioLineHeight;
 
 	// represents the ratio for the arrow end, 0 means arrow point will be at bottom, 1 means arrow point will be at top
-	UPROPERTY(EditInstanceOnly, BlueprintReadWrite, Category="Shape", meta=(ClampMin="0.0", ClampMax="1.0", DisplayName="Ratio Arrow Y", AllowPrivateAccess="true"))
+	UPROPERTY(EditInstanceOnly, BlueprintReadWrite, Category="Shape", Setter, Getter, meta=(ClampMin="0.0", ClampMax="1.0", DisplayName="Ratio Arrow Y", AllowPrivateAccess="true"))
 	float RatioArrowY;
 
 	// represents the ratio for the arrow end, 0 means arrow point will be at bottom, 1 means arrow point will be at top
-	UPROPERTY(EditInstanceOnly, BlueprintReadWrite, Category="Shape", meta=(ClampMin="0.0", ClampMax="1.0", DisplayName="Ratio Line Y", AllowPrivateAccess="true"))
+	UPROPERTY(EditInstanceOnly, BlueprintReadWrite, Category="Shape", Setter, Getter, meta=(ClampMin="0.0", ClampMax="1.0", DisplayName="Ratio Line Y", AllowPrivateAccess="true"))
 	float RatioLineY;
 
 	// whether there should be an arrow on both side, arrows will have same ratio
-	UPROPERTY(EditInstanceOnly, BlueprintReadWrite, Category="Shape", meta=(DisplayName="Both Side Arrows", AllowPrivateAccess="true"))
+	UPROPERTY(EditInstanceOnly, BlueprintReadWrite, Category="Shape", Setter="SetBothSideArrows", Getter="IsBothSideArrows", meta=(DisplayName="Both Side Arrows", AllowPrivateAccess="true"))
 	bool bBothSideArrows;
 };
diff --git a/Engine/Plugins/Experimental/Avalanche/Source/AvalancheShapes/Public/DynamicMeshes/AvaShape2DDynMeshBase.h b/Engine/Plugins/Experimental/Avalanche/Source/AvalancheShapes/Public/DynamicMeshes/AvaShape2DDynMeshBase.h
--- a/Engine/Plugins/Experimental/Avalanche/Source/AvalancheShapes/Public/DynamicMeshes/AvaShape2DDynMeshBase.h
+++ b/Engine/Plugins/Experimental/Avalanche/Source/AvalancheShapes/Public/DynamicMeshes/AvaShape2DDynMeshBase.h
@@ -8,8 +8,8 @@
 
 struct FAvaShapeCachedVertex2D;
 
-UCLASS(ClassGroup="Shape", Abstract, BlueprintType, CustomConstructor)
-class AVALANCHESHAPES_API UAvaShape2DDynMeshBase : public UAvaShapeDynamicMeshBase
+UCLASS(MinimalAPI, ClassGroup="Shape", Abstract, BlueprintType, CustomConstructor)
+class UAvaShape2DDynMeshBase : public UAvaShapeDynamicMeshBase
 {
 	GENERATED_BODY()
 
@@ -18,29 +18,8 @@ class AVALANCHESHAPES_API UAvaShape2DDynMeshBase : public UAvaShapeDynamicMeshBa
 	friend class FAvaVectorPropertyTypeCustomization;
 
 public:
-	virtual ~UAvaShape2DDynMeshBase() override = default;
-
-	UFUNCTION()
-	bool SetPixelSize2D(const FVector2D& InPixelSize2D);
-	const FVector2D& GetPixelSize2D() const { return PixelSize2D; }
-
-	UFUNCTION()
-	bool SetSize2D(const FVector2D& InSize2D);
-	const FVector2D& GetSize2D() const { return Size2D; }
-
-	virtual bool SetSize3D(const FVector& InSize) override;
-	virtual const FVector& GetSize3D() const override { return Size3D; }
-
-	virtual void PostLoad() override;
-#if WITH_EDITOR
-	virtual void PostEditChangeProperty(FPropertyChangedEvent& PropertyChangedEvent) override;
-#endif
-
-	virtual void GetBounds(FVector& Origin, FVector& BoxExtent, FVector& Pivot) const override;
-
-protected:
 	UAvaShape2DDynMeshBase()
-		: UAvaShape2DDynMeshBase(FVector2D(50.f, 50.f))
+	: UAvaShape2DDynMeshBase(FVector2D(50.f, 50.f))
 	{}
 
 	UAvaShape2DDynMeshBase(const FVector2D& InExtent,
@@ -51,6 +30,36 @@ protected:
 		, Size3D(FVector(0.f, InExtent.X, InExtent.Y))
 	{}
 
+	virtual ~UAvaShape2DDynMeshBase() override = default;
+
+	AVALANCHESHAPES_API void SetPixelSize2D(const FVector2D& InPixelSize2D);
+	const FVector2D& GetPixelSize2D() const
+	{
+		return PixelSize2D;
+	}
+
+	AVALANCHESHAPES_API void SetSize2D(const FVector2D& InSize2D);
+	const FVector2D& GetSize2D() const
+	{
+		return Size2D;
+	}
+
+	AVALANCHESHAPES_API virtual void SetSize3D(const FVector& InSize) override;
+	virtual const FVector& GetSize3D() const override
+	{
+		return Size3D;
+	}
+
+	virtual void GetBounds(FVector& Origin, FVector& BoxExtent, FVector& Pivot) const override;
+
+protected:
+	//~ Begin UObject
+	virtual void PostLoad() override;
+#if WITH_EDITOR
+	virtual void PostEditChangeProperty(FPropertyChangedEvent& PropertyChangedEvent) override;
+#endif
+	//~ End UObject
+
 	virtual void OnRegisteredMeshes() override;
 	virtual void OnPixelSizeChanged() override;
 	virtual void OnScaledSizeChanged() override;
@@ -99,11 +108,11 @@ protected:
 	bool bDoNotRecenterVertices;
 
 	// pixel size of the mesh, will only be available in editor
-	UPROPERTY(EditInstanceOnly, BlueprintReadWrite, Transient, Category="Shape", meta=(ClampMin="0.0", DisplayName="Pixel Size", DisplayAfter="SizeType", AllowPreserveRatio, EditCondition="bAllowEditSize && SizeType == ESizeType::Pixel", EditConditionHides, AllowPrivateAccess="true"))
+	UPROPERTY(EditInstanceOnly, BlueprintReadWrite, Transient, Setter, Getter, Category="Shape", meta=(ClampMin="0.0", DisplayName="Pixel Size", DisplayAfter="SizeType", AllowPreserveRatio, EditCondition="bAllowEditSize && SizeType == ESizeType::Pixel", EditConditionHides, AllowPrivateAccess="true"))
 	FVector2D PixelSize2D = FVector2D::ZeroVector;
 
 	// total size in 2D from 0 to mesh size and not origin
-	UPROPERTY(EditInstanceOnly, BlueprintReadWrite, Category="Shape", meta=(ClampMin="0.0", DisplayName="Mesh Size", DisplayAfter="SizeType", AllowPreserveRatio, Units="Centimeters", EditCondition="bAllowEditSize && SizeType == ESizeType::UnrealUnit", EditConditionHides, AllowPrivateAccess="true"))
+	UPROPERTY(EditInstanceOnly, BlueprintReadWrite, Setter, Getter, Category="Shape", meta=(ClampMin="0.0", DisplayName="Mesh Size", DisplayAfter="SizeType", AllowPreserveRatio, Units="Centimeters", EditCondition="bAllowEditSize && SizeType == ESizeType::UnrealUnit", EditConditionHides, AllowPrivateAccess="true"))
 	FVector2D Size2D = FVector2D::ZeroVector;
 
 	UPROPERTY(Transient)
diff --git a/Engine/Plugins/Experimental/Avalanche/Source/AvalancheShapes/Public/DynamicMeshes/AvaShape3DDynMeshBase.h b/Engine/Plugins/Experimental/Avalanche/Source/AvalancheShapes/Public/DynamicMeshes/AvaShape3DDynMeshBase.h
--- a/Engine/Plugins/Experimental/Avalanche/Source/AvalancheShapes/Public/DynamicMeshes/AvaShape3DDynMeshBase.h
+++ b/Engine/Plugins/Experimental/Avalanche/Source/AvalancheShapes/Public/DynamicMeshes/AvaShape3DDynMeshBase.h
@@ -7,8 +7,8 @@
 
 struct FAvaShapeCachedVertex3D;
 
-UCLASS(ClassGroup="Shape", Abstract, BlueprintType, CustomConstructor)
-class AVALANCHESHAPES_API UAvaShape3DDynMeshBase : public UAvaShapeDynamicMeshBase
+UCLASS(MinimalAPI, ClassGroup="Shape", Abstract, BlueprintType, CustomConstructor)
+class UAvaShape3DDynMeshBase : public UAvaShapeDynamicMeshBase
 {
 	GENERATED_BODY()
 
@@ -18,28 +18,33 @@ class AVALANCHESHAPES_API UAvaShape3DDynMeshBase : public UAvaShapeDynamicMeshBa
 public:
 	UAvaShape3DDynMeshBase()
 	: UAvaShape3DDynMeshBase(FVector(50.f, 50.f, 50.f))
-	{
-	}
+	{}
 
 	UAvaShape3DDynMeshBase(const FVector& InSize, const FLinearColor& InVertexColor = FLinearColor::White)
 		: UAvaShapeDynamicMeshBase(InVertexColor)
 		, Size3D(InSize)
+	{}
+
+	AVALANCHESHAPES_API void SetPixelSize3D(const FVector& InPixelSize);
+	const FVector& GetPixelSize3D() const
 	{
+		return PixelSize3D;
 	}
 
-	UFUNCTION()
-	bool SetPixelSize3D(const FVector& InPixelSize);
-	const FVector& GetPixelSize3D() const { return PixelSize3D; }
-
-	virtual bool SetSize3D(const FVector& InSize) override;
-	virtual const FVector& GetSize3D() const override { return Size3D; }
+	AVALANCHESHAPES_API virtual void SetSize3D(const FVector& InSize) override;
+	virtual const FVector& GetSize3D() const override
+	{
+		return Size3D;
+	}
 
 	virtual void GetBounds(FVector& Origin, FVector& BoxExtent, FVector& Pivot) const override;
 
 protected:
+	//~ Begin UObject
 #if WITH_EDITOR
 	virtual void PostEditChangeProperty(FPropertyChangedEvent& PropertyChangedEvent) override;
 #endif
+	//~ End UObject
 
 	virtual void OnRegisteredMeshes() override;
 	virtual void OnPixelSizeChanged() override;
@@ -78,13 +83,13 @@ protected:
 	void AddTriangle(FAvaShapeMesh& InMesh, int32 A, int32 B, int32 C);
 
 	// pixel size of the mesh, will only be available in editor
-	UPROPERTY(EditInstanceOnly, BlueprintReadWrite, Transient, Category="Shape", meta=(ClampMin="0.0", DisplayName="Pixel Size", DisplayAfter="SizeType", AllowPreserveRatio, EditCondition="bAllowEditSize && SizeType == ESizeType::Pixel", EditConditionHides, AllowPrivateAccess="true"))
+	UPROPERTY(EditInstanceOnly, BlueprintReadWrite, Transient, Setter, Getter, Category="Shape", meta=(ClampMin="0.0", DisplayName="Pixel Size", DisplayAfter="SizeType", AllowPreserveRatio, EditCondition="bAllowEditSize && SizeType == ESizeType::Pixel", EditConditionHides, AllowPrivateAccess="true"))
 	FVector PixelSize3D = FVector::ZeroVector;
 
 	/*
 	 * Corresponds to the total size from 0 to mesh size
 	 */
-	UPROPERTY(EditInstanceOnly, BlueprintReadWrite, Category="Shape", meta=(ClampMin="0.0", DisplayName="Mesh Size", DisplayAfter="SizeType", AllowPreserveRatio, Units="Centimeters", EditCondition="bAllowEditSize && SizeType == ESizeType::UnrealUnit", EditConditionHides, AllowPrivateAccess="true"))
+	UPROPERTY(EditInstanceOnly, BlueprintReadWrite, Setter, Getter, Category="Shape", meta=(ClampMin="0.0", DisplayName="Mesh Size", DisplayAfter="SizeType", AllowPreserveRatio, Units="Centimeters", EditCondition="bAllowEditSize && SizeType == ESizeType::UnrealUnit", EditConditionHides, AllowPrivateAccess="true"))
 	FVector Size3D = FVector::ZeroVector;
 
 	UPROPERTY(Transient)
diff --git a/Engine/Plugins/Experimental/Avalanche/Source/AvalancheShapes/Public/DynamicMeshes/AvaShapeChevronDynMesh.h b/Engine/Plugins/Experimental/Avalanche/Source/AvalancheShapes/Public/DynamicMeshes/AvaShapeChevronDynMesh.h
--- a/Engine/Plugins/Experimental/Avalanche/Source/AvalancheShapes/Public/DynamicMeshes/AvaShapeChevronDynMesh.h
+++ b/Engine/Plugins/Experimental/Avalanche/Source/AvalancheShapes/Public/DynamicMeshes/AvaShapeChevronDynMesh.h
@@ -5,15 +5,17 @@
 #include "DynamicMeshes/AvaShape2DDynMeshBase.h"
 #include "AvaShapeChevronDynMesh.generated.h"
 
-UCLASS(ClassGroup="Shape", BlueprintType, CustomConstructor, Within=AvaShapeActor)
-class AVALANCHESHAPES_API UAvaShapeChevronDynamicMesh : public UAvaShape2DDynMeshBase
+UCLASS(MinimalAPI, ClassGroup="Shape", BlueprintType, CustomConstructor, Within=AvaShapeActor)
+class UAvaShapeChevronDynamicMesh : public UAvaShape2DDynMeshBase
 {
 	GENERATED_BODY()
 
 public:
+	static const FString MeshName;
+
 	UAvaShapeChevronDynamicMesh()
 		: UAvaShapeChevronDynamicMesh(FVector2D(50.f, 50.f))
-	{ }
+	{}
 
 	UAvaShapeChevronDynamicMesh(
 		const FVector2D& Size2D,
@@ -21,20 +23,25 @@ public:
 		const float InRatioChevron = 0.5)
 		: UAvaShape2DDynMeshBase(Size2D, InVertexColor)
 		, RatioChevron(InRatioChevron)
+	{}
+
+	virtual const FString& GetMeshName() const override
 	{
+		return MeshName;
 	}
 
-	static const FString MeshName;
-	virtual const FString& GetMeshName() const override { return MeshName; }
-
-	UFUNCTION()
-	bool SetRatioChevron(float InRatio);
-	float GetRatioChevron() const { return RatioChevron; }
+	AVALANCHESHAPES_API void SetRatioChevron(float InRatio);
+	float GetRatioChevron() const
+	{
+		return RatioChevron;
+	}
 
 protected:
+	//~ Begin UObject
 #if WITH_EDITOR
 	virtual void PostEditChangeProperty(FPropertyChangedEvent& PropertyChangedEvent) override;
 #endif
+	//~ End UObject
 
 	virtual void OnRatioChevronChanged();
 
@@ -42,6 +49,6 @@ protected:
 	virtual bool CreateMesh(FAvaShapeMesh& InMesh) override;
 
 	// represents the ratio for the chevron size
-	UPROPERTY(EditInstanceOnly, BlueprintReadWrite, Category="Shape", meta=(ClampMin="0.0", ClampMax="0.99", DisplayName="Ratio Chevron", AllowPrivateAccess="true"))
+	UPROPERTY(EditInstanceOnly, BlueprintReadWrite, Setter, Getter, Category="Shape", meta=(ClampMin="0.0", ClampMax="0.99", DisplayName="Ratio Chevron", AllowPrivateAccess="true"))
 	float RatioChevron;
 };
diff --git a/Engine/Plugins/Experimental/Avalanche/Source/AvalancheShapes/Public/DynamicMeshes/AvaShapeConeDynMesh.h b/Engine/Plugins/Experimental/Avalanche/Source/AvalancheShapes/Public/DynamicMeshes/AvaShapeConeDynMesh.h
--- a/Engine/Plugins/Experimental/Avalanche/Source/AvalancheShapes/Public/DynamicMeshes/AvaShapeConeDynMesh.h
+++ b/Engine/Plugins/Experimental/Avalanche/Source/AvalancheShapes/Public/DynamicMeshes/AvaShapeConeDynMesh.h
@@ -5,14 +5,16 @@
 #include "AvaShape3DDynMeshBase.h"
 #include "AvaShapeConeDynMesh.generated.h"
 
-UCLASS(ClassGroup="Shape", BlueprintType, CustomConstructor, Within=AvaShapeActor)
-class AVALANCHESHAPES_API UAvaShapeConeDynamicMesh : public UAvaShape3DDynMeshBase
+UCLASS(MinimalAPI, ClassGroup="Shape", BlueprintType, CustomConstructor, Within=AvaShapeActor)
+class UAvaShapeConeDynamicMesh : public UAvaShape3DDynMeshBase
 {
 	GENERATED_BODY()
 
 	friend class FAvaShapeConeDynamicMeshVisualizer;
 
 public:
+	static const FString MeshName;
+
 	static constexpr uint8 MinNumSides = 3;
 	static constexpr uint8 MaxNumSides = 128;
 
@@ -23,8 +25,7 @@ public:
 
 	UAvaShapeConeDynamicMesh()
 		: UAvaShapeConeDynamicMesh(FVector(50.f, 50.f, 50.f))
-	{
-	}
+	{}
 
 	UAvaShapeConeDynamicMesh(
 		const FVector& InSize,
@@ -37,43 +38,48 @@ public:
 		float InStartDegree = 0.f)
 		: UAvaShape3DDynMeshBase(InSize, InVertexColor)
 		, NumSides(InNumSides)
-		, Height(InHeight)
-		, BaseRadius(InBaseRadius)
 		, TopRadius(InTopRadius)
 		, AngleDegree(InAngleDegree)
 		, StartDegree(InStartDegree)
+		, Height(InHeight)
+		, BaseRadius(InBaseRadius)
+	{}
+
+	virtual const FString& GetMeshName() const override
 	{
+		return MeshName;
 	}
 
-	static const FString MeshName;
-	virtual const FString& GetMeshName() const override { return MeshName; };
-
-	UFUNCTION()
-	bool SetNumSides(uint8 InNumSides);
-	uint8 GetNumSides() const { return NumSides; }
-
-	UFUNCTION()
-	float GetBaseRadius() const { return BaseRadius; }
-
-	UFUNCTION()
-	bool SetTopRadius(float InTopRadius);
-	float GetTopRadius() const { return TopRadius; }
+	AVALANCHESHAPES_API void SetNumSides(uint8 InNumSides);
+	uint8 GetNumSides() const
+	{
+		return NumSides;
+	}
 
-	UFUNCTION()
-	float GetHeight() const { return Height; }
+	AVALANCHESHAPES_API void SetTopRadius(float InTopRadius);
+	float GetTopRadius() const
+	{
+		return TopRadius;
+	}
 
-	UFUNCTION()
-	bool SetAngleDegree(float InDegree);
-	float GetAngleDegree() const { return AngleDegree; }
+	AVALANCHESHAPES_API void SetAngleDegree(float InDegree);
+	float GetAngleDegree() const
+	{
+		return AngleDegree;
+	}
 
-	UFUNCTION()
-	bool SetStartDegree(float InDegree);
-	float GetStartDegree() const { return StartDegree; }
+	AVALANCHESHAPES_API void SetStartDegree(float InDegree);
+	float GetStartDegree() const
+	{
+		return StartDegree;
+	}
 
 protected:
+	//~ Begin UObject
 #if WITH_EDITOR
 	virtual void PostEditChangeProperty(FPropertyChangedEvent& PropertyChangedEvent) override;
 #endif
+	//~ End UObject
 
 	virtual void OnHeightChanged();
 	virtual void OnBaseRadiusChanged();
@@ -103,26 +109,27 @@ protected:
 	virtual bool CreateUVs(FAvaShapeMesh& InMesh, FAvaShapeMaterialUVParameters& InParams) override;
 
 	// The number of sides around the base of the cone
-	UPROPERTY(EditInstanceOnly, BlueprintReadWrite, Category="Shape", meta=(ClampMin="3.0", ClampMax="128.0", DisplayName="Sides", AllowPrivateAccess="true"))
+	UPROPERTY(EditInstanceOnly, BlueprintReadWrite, Setter, Getter, Category="Shape", meta=(ClampMin="3.0", ClampMax="128.0", DisplayName="Sides", AllowPrivateAccess="true"))
 	uint8 NumSides;
 
-	// The height of the cone from the base to the top
-	UPROPERTY()
-	float Height;
-
-	// the radius of the cone base
-	UPROPERTY()
-	float BaseRadius;
-
 	// the ratio for the radius of the cone top
-	UPROPERTY(EditInstanceOnly, BlueprintReadWrite, Category="Shape", meta=(ClampMin="0.0", ClampMax="1.0", DisplayName="Top radius ratio", AllowPrivateAccess="true"))
+	UPROPERTY(EditInstanceOnly, BlueprintReadWrite, Setter, Getter, Category="Shape", meta=(ClampMin="0.0", ClampMax="1.0", DisplayName="Top radius ratio", AllowPrivateAccess="true"))
 	float TopRadius;
 
 	// represents the base angle in degree for the cone
-	UPROPERTY(EditInstanceOnly, BlueprintReadWrite, Category="Shape", meta=(ClampMin="0.0", ClampMax="360.0", DisplayName="Angle degree", AllowPrivateAccess="true"))
+	UPROPERTY(EditInstanceOnly, BlueprintReadWrite, Setter, Getter, Category="Shape", meta=(ClampMin="0.0", ClampMax="360.0", DisplayName="Angle degree", AllowPrivateAccess="true"))
 	float AngleDegree;
 
 	// represents the starting angle in degree for the cone
-	UPROPERTY(EditInstanceOnly, BlueprintReadWrite, Category="Shape", meta=(ClampMin="0.0", ClampMax="360.0", DisplayName="Start Degree", AllowPrivateAccess="true"))
+	UPROPERTY(EditInstanceOnly, BlueprintReadWrite, Setter, Getter, Category="Shape", meta=(ClampMin="0.0", ClampMax="360.0", DisplayName="Start Degree", AllowPrivateAccess="true"))
 	float StartDegree;
+
+private:
+	// The height of the cone from the base to the top
+	UPROPERTY()
+	float Height;
+
+	// the radius of the cone base
+	UPROPERTY()
+	float BaseRadius;
 };
diff --git a/Engine/Plugins/Experimental/Avalanche/Source/AvalancheShapes/Public/DynamicMeshes/AvaShapeCubeDynMesh.h b/Engine/Plugins/Experimental/Avalanche/Source/AvalancheShapes/Public/DynamicMeshes/AvaShapeCubeDynMesh.h
--- a/Engine/Plugins/Experimental/Avalanche/Source/AvalancheShapes/Public/DynamicMeshes/AvaShapeCubeDynMesh.h
+++ b/Engine/Plugins/Experimental/Avalanche/Source/AvalancheShapes/Public/DynamicMeshes/AvaShapeCubeDynMesh.h
@@ -5,14 +5,16 @@
 #include "DynamicMeshes/AvaShape3DDynMeshBase.h"
 #include "AvaShapeCubeDynMesh.generated.h"
 
-UCLASS(ClassGroup="Shape", BlueprintType, CustomConstructor, Within=AvaShapeActor)
-class AVALANCHESHAPES_API UAvaShapeCubeDynamicMesh : public UAvaShape3DDynMeshBase
+UCLASS(MinimalAPI, ClassGroup="Shape", BlueprintType, CustomConstructor, Within=AvaShapeActor)
+class UAvaShapeCubeDynamicMesh : public UAvaShape3DDynMeshBase
 {
 	GENERATED_BODY()
 
 	friend class FAvaShapeCubeDynamicMeshVisualizer;
 
 public:
+	static const FString MeshName;
+
 	static constexpr uint8 MinBevelNum = 1;
 	static constexpr uint8 MaxBevelNum = 8;
 
@@ -24,7 +26,7 @@ public:
 
 	UAvaShapeCubeDynamicMesh()
 		: UAvaShapeCubeDynamicMesh(FVector(50.f, 50.f, 50.f))
-	{ }
+	{}
 
 	UAvaShapeCubeDynamicMesh(
 		const FVector& InSize,
@@ -33,31 +35,40 @@ public:
 		float InBevel = 0.f,
 		uint8 InBevelNum = 1)
 		: UAvaShape3DDynMeshBase(InSize, InVertexColor)
-		, Segment(InSegment)
 		, BevelSizeRatio(InBevel)
 		, BevelNum(InBevelNum)
+		, Segment(InSegment)
+	{}
+
+	virtual const FString& GetMeshName() const override
 	{
+		return MeshName;
 	}
 
-	static const FString MeshName;
-	virtual const FString& GetMeshName() const override { return MeshName; }
-
-	UFUNCTION()
-	bool SetSegment(float InSegment);
-	float GetSegment() const { return Segment; }
+	AVALANCHESHAPES_API void SetSegment(float InSegment);
+	float GetSegment() const
+	{
+		return Segment;
+	}
 
-	UFUNCTION()
-	bool SetBevelSizeRatio(float InBevel);
-	float GetBevelSizeRatio() const { return BevelSizeRatio; }
+	AVALANCHESHAPES_API void SetBevelSizeRatio(float InBevel);
+	float GetBevelSizeRatio() const
+	{
+		return BevelSizeRatio;
+	}
 
-	UFUNCTION()
-	bool SetBevelNum(uint8 InBevel);
-	uint8 GetBevelNum() const { return BevelNum; }
+	AVALANCHESHAPES_API void SetBevelNum(uint8 InBevel);
+	uint8 GetBevelNum() const
+	{
+		return BevelNum;
+	}
 
 protected:
+	//~ Begin UObject
 #if WITH_EDITOR
 	virtual void PostEditChangeProperty(FPropertyChangedEvent& PropertyChangedEvent) override;
 #endif
+	//~ End UObject
 
 	virtual void OnSegmentChanged();
 	virtual void OnBevelSizeChanged();
@@ -93,15 +104,16 @@ protected:
 
 	virtual bool CreateUVs(FAvaShapeMesh& InMesh, FAvaShapeMaterialUVParameters& InParams) override;
 
-	// segment size ratio to multiply with mesh size
-	UPROPERTY()
-	float Segment;
-
 	// represents the bevel size applied on each face of the cube, 0 means no bevels
-	UPROPERTY(EditInstanceOnly, BlueprintReadWrite, Category="Shape", meta=(ClampMin="0.0", DisplayName="Bevel Size", AllowPrivateAccess="true"))
+	UPROPERTY(EditInstanceOnly, BlueprintReadWrite, Setter, Getter, Category="Shape", meta=(ClampMin="0.0", DisplayName="Bevel Size", AllowPrivateAccess="true"))
 	float BevelSizeRatio;
 
 	// represents the bevel number of division, only valid when bevel size is greater than zero
-	UPROPERTY(EditInstanceOnly, BlueprintReadWrite, Category="Shape", meta=(ClampMin="1.0", ClampMax="8.0", DisplayName="Bevel Num", AllowPrivateAccess="true"))
+	UPROPERTY(EditInstanceOnly, BlueprintReadWrite, Setter, Getter, Category="Shape", meta=(ClampMin="1.0", ClampMax="8.0", DisplayName="Bevel Num", AllowPrivateAccess="true"))
 	uint8 BevelNum;
+
+private:
+	// segment size ratio to multiply with mesh size
+	UPROPERTY()
+	float Segment;
 };
diff --git a/Engine/Plugins/Experimental/Avalanche/Source/AvalancheShapes/Public/DynamicMeshes/AvaShapeDynMeshBase.h b/Engine/Plugins/Experimental/Avalanche/Source/AvalancheShapes/Public/DynamicMeshes/AvaShapeDynMeshBase.h
--- a/Engine/Plugins/Experimental/Avalanche/Source/AvalancheShapes/Public/DynamicMeshes/AvaShapeDynMeshBase.h
+++ b/Engine/Plugins/Experimental/Avalanche/Source/AvalancheShapes/Public/DynamicMeshes/AvaShapeDynMeshBase.h
@@ -22,8 +22,8 @@ DECLARE_MULTICAST_DELEGATE_OneParam(FMaskEnabledDelegate, AActor* /** NewMaskAct
 DECLARE_MULTICAST_DELEGATE_OneParam(FMaskDisabledDelegate, AActor* /** OldMaskActor */);
 DECLARE_MULTICAST_DELEGATE_TwoParams(FMaskVisibilityDelegate, const UWorld* /** CurrentWorld */, bool /** bMaskActorVisible */);
 
-UCLASS(ClassGroup="Shape", Abstract, BlueprintType, CustomConstructor, EditInlineNew, DefaultToInstanced)
-class AVALANCHESHAPES_API UAvaShapeDynamicMeshBase
+UCLASS(MinimalAPI, ClassGroup="Shape", Abstract, BlueprintType, CustomConstructor, EditInlineNew, DefaultToInstanced)
+class UAvaShapeDynamicMeshBase
 	: public UActorComponent
 	, public IAvaInteractiveToolsModeDetailsObject
 	, public IAvaGizmoObjectInterface
@@ -51,90 +51,109 @@ public:
 	static inline FMaskDisabledDelegate OnMaskDisabled;
 	static inline FMaskVisibilityDelegate OnMaskVisibility;
 
+	UAvaShapeDynamicMeshBase(const FObjectInitializer& ObjectInitializer)
+		: UAvaShapeDynamicMeshBase(FLinearColor::White)
+	{}
+
+	UAvaShapeDynamicMeshBase(const FLinearColor& InVertexColor = FLinearColor::White,
+		float InUniformScaledSize = 1.f, bool bInAllowEditSize = true);
+
 	// Converts the Outer of this object to a AAvaShapeActor.
-	AAvaShapeActor* GetShapeActor() const;
+	AVALANCHESHAPES_API AAvaShapeActor* GetShapeActor() const;
 
 	// Gets the Dynamic Mesh component from the Shape actor, this will load it if its nullptr
-	UDynamicMeshComponent* GetShapeMeshComponent() const;
+	AVALANCHESHAPES_API UDynamicMeshComponent* GetShapeMeshComponent() const;
 
 	/** Get the name of the shape */
 	virtual const FString& GetMeshName() const;
 
 	/** Can we change the size of this shape */
-	bool AllowsSizeEditing() const { return bAllowEditSize; }
+	bool AllowsSizeEditing() const
+	{
+		return bAllowEditSize;
+	}
 
 	// override in child classes, represents the mesh size
-	UFUNCTION()
-	virtual const FVector& GetSize3D() const { return FVector::ZeroVector; }
-
-	UFUNCTION()
-	virtual bool SetSize3D(const FVector& InSize) { return false; }
+	virtual void SetSize3D(const FVector& InSize) {}
+	virtual const FVector& GetSize3D() const
+	{
+		return FVector::ZeroVector;
+	}
 
 	// Primary material override
-	bool GetUsePrimaryMaterialEverywhere() const { return bUsePrimaryMaterialEverywhere; }
-	void SetUsePrimaryMaterialEverywhere(bool bInUse);
+	AVALANCHESHAPES_API void SetUsePrimaryMaterialEverywhere(bool bInUse);
+	bool GetUsePrimaryMaterialEverywhere() const
+	{
+		return bUsePrimaryMaterialEverywhere;
+	}
+
+	AVALANCHESHAPES_API void SetUniformScaledSize(float InSize);
+	float GetUniformScaledSize() const
+	{
+		return UniformScaledSize;
+	}
 
 	/** Checks if shape size is render-able and not minimal */
 	bool IsMeshSizeValid() const;
 
 	/** Checks if the mesh section is valid */
-	bool IsValidMeshIndex(int32 MeshIndex) const;
+	AVALANCHESHAPES_API bool IsValidMeshIndex(int32 MeshIndex) const;
 
-	UMaterialInterface* GetMaterial(int32 MeshIndex) const;
-	bool SetMaterial(int32 MeshIndex, UMaterialInterface* NewMaterial);
+	AVALANCHESHAPES_API UMaterialInterface* GetMaterial(int32 MeshIndex) const;
+	AVALANCHESHAPES_API bool SetMaterial(int32 MeshIndex, UMaterialInterface* NewMaterial);
 
-	const FAvaShapeParametricMaterial* GetParametricMaterial(int32 MeshIndex) const;
-	FAvaShapeParametricMaterial* GetParametricMaterial(int32 MeshIndex);
+	AVALANCHESHAPES_API const FAvaShapeParametricMaterial* GetParametricMaterial(int32 MeshIndex) const;
+	AVALANCHESHAPES_API FAvaShapeParametricMaterial* GetParametricMaterial(int32 MeshIndex);
 
 	// Internal materials are not copied over, only settings
-	bool SetParametricMaterial(int32 MeshIndex, const FAvaShapeParametricMaterial& NewMaterialParams);
+	AVALANCHESHAPES_API bool SetParametricMaterial(int32 MeshIndex, const FAvaShapeParametricMaterial& NewMaterialParams);
 
 	// Allows custom uv params per mesh section instead of using the primary one
 	bool SetOverridePrimaryUVParams(int32 MeshIndex, bool bOverride);
 
 	// Helper function to quickly get what you need instead of using GetMeshData()
-	EAvaShapeUVMode GetMaterialUVMode(int32 MeshIndex) const;
-	bool SetMaterialUVMode(int32 MeshIndex, EAvaShapeUVMode InUVMode);
+	AVALANCHESHAPES_API EAvaShapeUVMode GetMaterialUVMode(int32 MeshIndex) const;
+	AVALANCHESHAPES_API bool SetMaterialUVMode(int32 MeshIndex, EAvaShapeUVMode InUVMode);
 
 	EAvaAnchors GetMaterialUVAnchorPreset(int32 MeshIndex) const;
 	bool SetMaterialUVAnchorPreset(int32 MeshIndex, EAvaAnchors InUVAnchorPreset);
 
-	float GetMaterialUVRotation(int32 MeshIndex) const;
-	bool SetMaterialUVRotation(int32 MeshIndex, float InUVRotation);
+	AVALANCHESHAPES_API float GetMaterialUVRotation(int32 MeshIndex) const;
+	AVALANCHESHAPES_API bool SetMaterialUVRotation(int32 MeshIndex, float InUVRotation);
 
-	const FVector2D& GetMaterialUVAnchor(int32 MeshIndex) const;
-	bool SetMaterialUVAnchor(int32 MeshIndex, const FVector2D& InUVAnchor);
+	AVALANCHESHAPES_API const FVector2D& GetMaterialUVAnchor(int32 MeshIndex) const;
+	AVALANCHESHAPES_API bool SetMaterialUVAnchor(int32 MeshIndex, const FVector2D& InUVAnchor);
 
-	const FVector2D& GetMaterialUVScale(int32 MeshIndex) const;
-	bool SetMaterialUVScale(int32 MeshIndex, const FVector2D& InUVScale);
+	AVALANCHESHAPES_API const FVector2D& GetMaterialUVScale(int32 MeshIndex) const;
+	AVALANCHESHAPES_API bool SetMaterialUVScale(int32 MeshIndex, const FVector2D& InUVScale);
 
-	const FVector2D& GetMaterialUVOffset(int32 MeshIndex) const;
-	bool SetMaterialUVOffset(int32 MeshIndex, const FVector2D& InUVOffset);
+	AVALANCHESHAPES_API const FVector2D& GetMaterialUVOffset(int32 MeshIndex) const;
+	AVALANCHESHAPES_API bool SetMaterialUVOffset(int32 MeshIndex, const FVector2D& InUVOffset);
 
-	bool GetMaterialHorizontalFlip(int32 MeshIndex) const;
-	bool SetMaterialHorizontalFlip(int32 MeshIndex, bool InHorizontalFlip);
+	AVALANCHESHAPES_API bool GetMaterialHorizontalFlip(int32 MeshIndex) const;
+	AVALANCHESHAPES_API bool SetMaterialHorizontalFlip(int32 MeshIndex, bool InHorizontalFlip);
 
-	bool GetMaterialVerticalFlip(int32 MeshIndex) const;
-	bool SetMaterialVerticalFlip(int32 MeshIndex, bool InVerticalFlip);
+	AVALANCHESHAPES_API bool GetMaterialVerticalFlip(int32 MeshIndex) const;
+	AVALANCHESHAPES_API bool SetMaterialVerticalFlip(int32 MeshIndex, bool InVerticalFlip);
 
-	const FAvaShapeMaterialUVParameters* GetInUseMaterialUVParams(int32 MeshIndex) const;
-	FAvaShapeMaterialUVParameters* GetInUseMaterialUVParams(int32 MeshIndex);
+	AVALANCHESHAPES_API const FAvaShapeMaterialUVParameters* GetInUseMaterialUVParams(int32 MeshIndex) const;
+	AVALANCHESHAPES_API FAvaShapeMaterialUVParameters* GetInUseMaterialUVParams(int32 MeshIndex);
 
 	const FAvaShapeMaterialUVParameters* GetMaterialUVParams(int32 MeshIndex) const;
 	bool SetMaterialUVParams(int32 MeshIndex, const FAvaShapeMaterialUVParameters& InParams);
 
 	bool SetMaterialType(int32 MeshIndex, EMaterialType Type);
-	bool IsMaterialType(int32 MeshIndex, EMaterialType Type);
+	AVALANCHESHAPES_API bool IsMaterialType(int32 MeshIndex, EMaterialType Type);
 
 	bool HasMeshRegenWorldLocation() const { return bHasNewMeshRegenWorldLocation; }
 	const FVector& GetMeshRegenWorldLocation() const { return MeshRegenWorldLocation; }
-	void SetMeshRegenWorldLocation(const FVector& NewLocation, bool bImmediateUpdate = false);
+	AVALANCHESHAPES_API void SetMeshRegenWorldLocation(const FVector& NewLocation, bool bImmediateUpdate = false);
 
 	// Generate a list of 3d-space snap points for this shape
 	virtual TArray<FAvaSnapPoint> GetLocalSnapPoints() const;
 	void GetLocalSnapPoints(TArray<FAvaSnapPoint>& Points) const;
 
-	FTransform GetTransform() const;
+	AVALANCHESHAPES_API FTransform GetTransform() const;
 
 	/** Clear the dynamic mesh section with a specific index */
 	bool ClearDynamicMeshSection(int32 MeshIndex);
@@ -143,7 +162,7 @@ public:
 	bool ClearDynamicMesh();
 
 	/** Converts dynamic mesh to static mesh */
-	bool ExportToStaticMesh(UStaticMesh* DestinationMesh);
+	AVALANCHESHAPES_API bool ExportToStaticMesh(UStaticMesh* DestinationMesh);
 
 	// Gets the bounds of the shape, override this in child classes for custom bounds
 	// Origin is the center of the box, BoxExtent is half the size, pivot is the default location of the pivot for this shape
@@ -153,10 +172,10 @@ public:
 	virtual FAvaColorChangeData GetActiveColor() const;
 
 	/** Get the registered meshes indexes */
-	TSet<int32> GetMeshesIndexes() const;
+	AVALANCHESHAPES_API TSet<int32> GetMeshesIndexes() const;
 
 	/** Returns the name of the meshes */
-	TArray<FName> GetMeshDataNames() const;
+	AVALANCHESHAPES_API TArray<FName> GetMeshDataNames() const;
 
 	/** find a registered mesh and gets a pointer to it */
 	FAvaShapeMeshData* GetMeshData(int32 MeshIndex);
@@ -205,13 +224,6 @@ protected:
 	// When a parametric material is updated
 	void OnParametricMaterialChanged(FAvaShapeParametricMaterial& InMaterial);
 
-	UAvaShapeDynamicMeshBase(const FObjectInitializer& ObjectInitializer)
-		: UAvaShapeDynamicMeshBase(FLinearColor::White)
-	{ }
-
-	UAvaShapeDynamicMeshBase(const FLinearColor& InVertexColor = FLinearColor::White,
-		float InUniformScaledSize = 1.f, bool bInAllowEditSize = true);
-
 	/** Register all meshes once, calls RegisterMeshes() */
 	void SetupMeshes();
 
@@ -348,14 +360,14 @@ protected:
 	ESizeType SizeType = ESizeType::UnrealUnit;
 
 	// Uniform scaled size of the mesh
-	UPROPERTY(EditInstanceOnly, BlueprintReadWrite, Transient, Category="Shape", meta=(ClampMin="0.0", DisplayName="Uniform Scaled Size", DisplayPriority=1, Units="times", AllowPrivateAccess="true"))
+	UPROPERTY(EditInstanceOnly, BlueprintReadWrite, Transient, Setter, Getter, Category="Shape", meta=(DisplayPriority=1, ClampMin="0.0", DisplayName="Uniform Scaled Size", Units="times", AllowPrivateAccess="true"))
 	float UniformScaledSize = 1.f;
 
-	UPROPERTY(BlueprintReadWrite, AdvancedDisplay, Category = "Shape", meta=(DisplayName="Material Vertex Color", DisplayPriority=1, AllowPrivateAccess="true"))
+	UPROPERTY(BlueprintReadWrite, AdvancedDisplay, Category="Shape", meta=(DisplayPriority=1, DisplayName="Material Vertex Color", AllowPrivateAccess="true"))
 	FLinearColor VertexColor;
 
 	// use primary material for every slot available
-	UPROPERTY(EditInstanceOnly, BlueprintReadWrite, Category="Material", meta=(DisplayName="Use Single Material", DisplayPriority=1, AllowPrivateAccess="true"))
+	UPROPERTY(EditInstanceOnly, BlueprintReadWrite, Setter="SetUsePrimaryMaterialEverywhere", Getter="GetUsePrimaryMaterialEverywhere", Category="Material", meta=(DisplayPriority=1, DisplayName="Use Single Material", AllowPrivateAccess="true"))
 	bool bUsePrimaryMaterialEverywhere;
 
 	TArray<FAvaSnapPoint> LocalSnapPoints;
diff --git a/Engine/Plugins/Experimental/Avalanche/Source/AvalancheShapes/Public/DynamicMeshes/AvaShapeEllipseDynMesh.h b/Engine/Plugins/Experimental/Avalanche/Source/AvalancheShapes/Public/DynamicMeshes/AvaShapeEllipseDynMesh.h
--- a/Engine/Plugins/Experimental/Avalanche/Source/AvalancheShapes/Public/DynamicMeshes/AvaShapeEllipseDynMesh.h
+++ b/Engine/Plugins/Experimental/Avalanche/Source/AvalancheShapes/Public/DynamicMeshes/AvaShapeEllipseDynMesh.h
@@ -5,20 +5,21 @@
 #include "DynamicMeshes/AvaShape2DDynMeshBase.h"
 #include "AvaShapeEllipseDynMesh.generated.h"
 
-UCLASS(ClassGroup="Shape", BlueprintType, CustomConstructor, Within=AvaShapeActor)
-class AVALANCHESHAPES_API UAvaShapeEllipseDynamicMesh : public UAvaShape2DDynMeshBase
+UCLASS(MinimalAPI, ClassGroup="Shape", BlueprintType, CustomConstructor, Within=AvaShapeActor)
+class UAvaShapeEllipseDynamicMesh : public UAvaShape2DDynMeshBase
 {
 	GENERATED_BODY()
 
 	friend class FAvaShapeEllipseDynamicMeshVisualizer;
 
 public:
+	static const FString MeshName;
 	static constexpr uint8 MinNumSides = 3;
 	static constexpr uint8 MaxNumSides = 128;
 
 	UAvaShapeEllipseDynamicMesh()
 		: UAvaShapeEllipseDynamicMesh(FVector2D(50.f, 50.f))
-	{ }
+	{}
 
 	UAvaShapeEllipseDynamicMesh(const FVector2D& Size2D, const FLinearColor& InVertexColor = FLinearColor::White,
 		uint8 InNumSides = 64, float InAngleDegree = 360.f, float InStartDegree = 0.f)
@@ -30,25 +31,35 @@ public:
 		bDoNotRecenterVertices = true;
 	}
 
-	static const FString MeshName;
-	virtual const FString& GetMeshName() const override { return MeshName; }
+	virtual const FString& GetMeshName() const override
+	{
+		return MeshName;
+	}
 
-	UFUNCTION()
-	bool SetNumSides(uint8 InNumSides);
-	uint8 GetNumSides() const { return NumSides; }
+	AVALANCHESHAPES_API void SetNumSides(uint8 InNumSides);
+	uint8 GetNumSides() const
+	{
+		return NumSides;
+	}
 
-	UFUNCTION()
-	bool SetAngleDegree(float InDegree);
-	float GetAngleDegree() const { return AngleDegree; }
+	AVALANCHESHAPES_API void SetAngleDegree(float InDegree);
+	float GetAngleDegree() const
+	{
+		return AngleDegree;
+	}
 
-	UFUNCTION()
-	bool SetStartDegree(float InDegree);
-	float GetStartDegree() const { return StartDegree; }
+	AVALANCHESHAPES_API void SetStartDegree(float InDegree);
+	float GetStartDegree() const
+	{
+		return StartDegree;
+	}
 
 protected:
+	//~ Begin UObject
 #if WITH_EDITOR
 	virtual void PostEditChangeProperty(FPropertyChangedEvent& PropertyChangedEvent) override;
 #endif
+	//~ End UObject
 
 	virtual void OnNumSidesChanged();
 	virtual void OnAngleDegreeChanged();
@@ -58,14 +69,14 @@ protected:
 	virtual bool CreateMesh(FAvaShapeMesh& InMesh) override;
 
 	// The number of sides for the shape
-	UPROPERTY(EditInstanceOnly, BlueprintReadWrite, Category="Shape", meta=(ClampMin="3.0", ClampMax="128.0", DisplayName="Sides", AllowPrivateAccess="true"))
+	UPROPERTY(EditInstanceOnly, BlueprintReadWrite, Setter, Getter, Category="Shape", meta=(ClampMin="3.0", ClampMax="128.0", DisplayName="Sides", AllowPrivateAccess="true"))
 	uint8 NumSides;
 
 	// represents the angle in degree for the ellipse
-	UPROPERTY(EditInstanceOnly, BlueprintReadWrite, Category="Shape", meta=(ClampMin="0.0", ClampMax="360.0", DisplayName="Angle degree", AllowPrivateAccess="true"))
+	UPROPERTY(EditInstanceOnly, BlueprintReadWrite, Setter, Getter, Category="Shape", meta=(ClampMin="0.0", ClampMax="360.0", DisplayName="Angle degree", AllowPrivateAccess="true"))
 	float AngleDegree;
 
 	// represents the starting angle in degree for the ellipse
-	UPROPERTY(EditInstanceOnly, BlueprintReadWrite, Category="Shape", meta=(ClampMin="0.0", ClampMax="360.0", DisplayName="Start Degree", AllowPrivateAccess="true"))
+	UPROPERTY(EditInstanceOnly, BlueprintReadWrite, Setter, Getter, Category="Shape", meta=(ClampMin="0.0", ClampMax="360.0", DisplayName="Start Degree", AllowPrivateAccess="true"))
 	float StartDegree;
 };
diff --git a/Engine/Plugins/Experimental/Avalanche/Source/AvalancheShapes/Public/DynamicMeshes/AvaShapeIrregularPolygonDynMesh.h b/Engine/Plugins/Experimental/Avalanche/Source/AvalancheShapes/Public/DynamicMeshes/AvaShapeIrregularPolygonDynMesh.h
--- a/Engine/Plugins/Experimental/Avalanche/Source/AvalancheShapes/Public/DynamicMeshes/AvaShapeIrregularPolygonDynMesh.h
+++ b/Engine/Plugins/Experimental/Avalanche/Source/AvalancheShapes/Public/DynamicMeshes/AvaShapeIrregularPolygonDynMesh.h
@@ -7,7 +7,7 @@
 #include "AvaShapeIrregularPolygonDynMesh.generated.h"
 
 USTRUCT(BlueprintType)
-struct AVALANCHESHAPES_API FAvaShapeRoundedCornerSettings
+struct FAvaShapeRoundedCornerSettings
 {
 	GENERATED_BODY()
 
@@ -19,14 +19,13 @@ struct AVALANCHESHAPES_API FAvaShapeRoundedCornerSettings
 };
 
 USTRUCT(BlueprintType)
-struct AVALANCHESHAPES_API FAvaShapeRoundedCorner
+struct FAvaShapeRoundedCorner
 {
 	GENERATED_BODY()
 
 	FAvaShapeRoundedCorner()
 		: FAvaShapeRoundedCorner(FVector2D::ZeroVector)
-	{
-	}
+	{}
 
 	FAvaShapeRoundedCorner(const FVector2D& InLocation)
 	{
@@ -45,22 +44,22 @@ struct AVALANCHESHAPES_API FAvaShapeRoundedCorner
 	FAvaShapeRoundedCornerMetrics CornerMetrics;
 };
 
-UCLASS(ClassGroup="Shape", BlueprintType, CustomConstructor, Within=AvaShapeActor)
-class AVALANCHESHAPES_API UAvaShapeIrregularPolygonDynamicMesh : public UAvaShape2DDynMeshBase
+UCLASS(MinimalAPI, ClassGroup="Shape", BlueprintType, CustomConstructor, Within=AvaShapeActor)
+class UAvaShapeIrregularPolygonDynamicMesh : public UAvaShape2DDynMeshBase
 {
 	GENERATED_BODY()
 
 	friend class FAvaShapeIrregularPolygonDynamicMeshVisualizer;
 
 public:
+	static const FString MeshName;
 	static const float MinPointDistance;
 
 	static bool DoLinesIntersect(const FVector2D Origin1, const FVector2D End1, const FVector2D Origin2, const FVector2D End2);
 
 	UAvaShapeIrregularPolygonDynamicMesh(const FObjectInitializer& ObjectInitializer)
      		: UAvaShapeIrregularPolygonDynamicMesh(FLinearColor::White)
-	{
-    }
+	{}
 
 	UAvaShapeIrregularPolygonDynamicMesh(const FLinearColor& InVertexColor)
 		: UAvaShape2DDynMeshBase(FVector2D::Zero(), InVertexColor)
@@ -69,40 +68,60 @@ public:
 		bAllowEditSize = false;
 	}
 
-	static const FString MeshName;
-	virtual const FString& GetMeshName() const override { return MeshName; }
+	virtual const FString& GetMeshName() const override
+	{
+		return MeshName;
+	}
 
-	UFUNCTION()
-	bool SetGlobalBevelSize(float InBevelSize);
-	float GetGlobalBevelSize() const { return GlobalBevelSize; }
+	AVALANCHESHAPES_API void SetGlobalBevelSize(float InBevelSize);
+	float GetGlobalBevelSize() const
+	{
+		return GlobalBevelSize;
+	}
 
-	UFUNCTION()
-	void SetGlobalBevelSubdivisions(uint8 InBevelSubdivisions);
-	uint8 GetGlobalBevelSubdivisions() const { return GlobalBevelSubdivisions; }
+	AVALANCHESHAPES_API void SetGlobalBevelSubdivisions(uint8 InBevelSubdivisions);
+	uint8 GetGlobalBevelSubdivisions() const
+	{
+		return GlobalBevelSubdivisions;
+	}
 
-	bool CanAddPoint(const FVector2D& InPoint);
-	bool AddPoint(const FVector2D& InPoint);
+	AVALANCHESHAPES_API void SetPoints(const TArray<FVector2D>& InPoints);
+	AVALANCHESHAPES_API void SetPoints(const TArray<FAvaShapeRoundedCorner>& InPoints);
+	const TArray<FAvaShapeRoundedCorner>& GetPoints() const
+	{
+		return Points;
+	}
+
+	const FAvaShapeRoundedCorner& GetPoint(int32 PointIdx) const
+	{
+		return Points[PointIdx];
+	}
+
+	int32 GetNumPoints() const
+	{
+		return Points.Num();
+	}
+
+	AVALANCHESHAPES_API bool CanAddPoint(const FVector2D& InPoint);
+	AVALANCHESHAPES_API bool AddPoint(const FVector2D& InPoint);
 	bool RemovePoint(int32 PointIdx);
 	bool RemoveFirstPoint();
 	bool RemoveLastPoint();
 	bool RemoveAllPoints();
-	const FAvaShapeRoundedCorner& GetPoint(int32 PointIdx) const { return Points[PointIdx]; }
-	const TArray<FAvaShapeRoundedCorner>& GetPoints() const { return Points; }
-	int32 GetNumPoints() const { return Points.Num(); }
-	bool SetPoints(const TArray<FVector2D>& InPoints);
-
-	void RecalculateActorPosition();
+	AVALANCHESHAPES_API void RecalculateActorPosition();
 
 protected:
+	// Begin UObject
 #if WITH_EDITOR
 	virtual void PreEditChange(FProperty* PropertyAboutToChange) override;
 	virtual void PostEditChangeProperty(FPropertyChangedEvent& PropertyChangedEvent) override;
 #endif
+	//~ End UObject
 
 	void RecalculateExtent();
 
-	bool SetLocation(int32 PointIdx, const FVector2D& InPoint);
-	bool SetBevelSize(int32 PointIdx, float InBevelSize);
+	AVALANCHESHAPES_API bool SetLocation(int32 PointIdx, const FVector2D& InPoint);
+	AVALANCHESHAPES_API bool SetBevelSize(int32 PointIdx, float InBevelSize);
 	bool SetBevelSubdivisions(int32 PointIdx, uint8 InBevelSubdivisions);
 
 	bool ShiftPoints(const FVector2D& Amount);
@@ -117,12 +136,12 @@ protected:
 
 	bool IsLocationInsideShape(const FVector2D& Location);
 
-	void BackupPoints();
+	AVALANCHESHAPES_API void BackupPoints();
 
-	float GetMaxBevelSizeForPoint(int32 PointIdx) const;
+	AVALANCHESHAPES_API float GetMaxBevelSizeForPoint(int32 PointIdx) const;
 
 	// Breaks the line, adding a new point, in between InPointIdx and InPointIdx+1
-	bool BreakSide(int32 InPointIdx);
+	AVALANCHESHAPES_API bool BreakSide(int32 InPointIdx);
 
 	void RestorePoints();
 	bool CheckNewPointsArray();
@@ -142,13 +161,13 @@ protected:
 	virtual bool ClearMesh() override;
 	virtual bool CreateMesh(FAvaShapeMesh& InMesh) override;
 
-	UPROPERTY(EditInstanceOnly, BlueprintReadWrite, Category="Shape", meta=(ClampMin="0.0",ClampMax="1.0", AllowPrivateAccess="true"))
+	UPROPERTY(EditInstanceOnly, BlueprintReadWrite, Setter, Getter, Category="Shape", meta=(ClampMin="0.0",ClampMax="1.0", AllowPrivateAccess="true"))
 	float GlobalBevelSize = 0.f;
 
-	UPROPERTY(EditInstanceOnly, BlueprintReadWrite, Category="Shape", meta=(ClampMax="64.0", AllowPrivateAccess="true"))
+	UPROPERTY(EditInstanceOnly, BlueprintReadWrite, Setter, Getter, Category="Shape", meta=(ClampMax="64.0", AllowPrivateAccess="true"))
 	uint8 GlobalBevelSubdivisions = DefaultSubdivisions;
 
-	UPROPERTY(BlueprintReadWrite, EditInstanceOnly, Category="Shape", meta=(EditFixedOrder, AllowPrivateAccess="true"))
+	UPROPERTY(EditInstanceOnly, BlueprintReadWrite, Setter, Getter, Category="Shape", meta=(EditFixedOrder, AllowPrivateAccess="true"))
 	TArray<FAvaShapeRoundedCorner> Points;
 	TArray<FAvaShapeRoundedCorner> PreEditPoints;
 };
diff --git a/Engine/Plugins/Experimental/Avalanche/Source/AvalancheShapes/Public/DynamicMeshes/AvaShapeLineDynMesh.h b/Engine/Plugins/Experimental/Avalanche/Source/AvalancheShapes/Public/DynamicMeshes/AvaShapeLineDynMesh.h
--- a/Engine/Plugins/Experimental/Avalanche/Source/AvalancheShapes/Public/DynamicMeshes/AvaShapeLineDynMesh.h
+++ b/Engine/Plugins/Experimental/Avalanche/Source/AvalancheShapes/Public/DynamicMeshes/AvaShapeLineDynMesh.h
@@ -5,18 +5,19 @@
 #include "AvaShapeRoundedPolygonDynMesh.h"
 #include "AvaShapeLineDynMesh.generated.h"
 
-UCLASS(ClassGroup="Shape", BlueprintType, CustomConstructor, Within=AvaShapeActor)
-class AVALANCHESHAPES_API UAvaShapeLineDynamicMesh : public UAvaShapeRoundedPolygonDynamicMesh
+UCLASS(MinimalAPI, ClassGroup="Shape", BlueprintType, CustomConstructor, Within=AvaShapeActor)
+class UAvaShapeLineDynamicMesh : public UAvaShapeRoundedPolygonDynamicMesh
 {
 	GENERATED_BODY()
 
 	friend class FAvaShapeLineDynamicMeshVisualizer;
 
 public:
+	static const FString MeshName;
+
 	UAvaShapeLineDynamicMesh()
 		: UAvaShapeLineDynamicMesh(FVector2D(50.f, 1.f))
-	{
-	}
+	{}
 
 	UAvaShapeLineDynamicMesh(const FVector2D& Size2D, const FLinearColor& InVertexColor = FLinearColor::White,
 		float InBevelSize = 0.f, uint8 InBevelSubdivisions = 0, float InLineWidth = 5.f)
@@ -27,21 +28,29 @@ public:
 		bAllowEditSize = false;
 	}
 
-	static const FString MeshName;
-	virtual const FString& GetMeshName() const override { return MeshName; }
+	virtual const FString& GetMeshName() const override
+	{
+		return MeshName;
+	}
 
-	UFUNCTION()
-	bool SetLineWidth(float InLineWidth);
-	float GetLineWidth() const { return LineWidth; }
+	AVALANCHESHAPES_API void SetLineWidth(float InLineWidth);
+	float GetLineWidth() const
+	{
+		return LineWidth;
+	}
 
-	UFUNCTION()
-	bool SetVector(const FVector2D& InVector);
-	const FVector2D& GetVector() const { return Vector; }
+	AVALANCHESHAPES_API void SetVector(const FVector2D& InVector);
+	const FVector2D& GetVector() const
+	{
+		return Vector;
+	}
 
 protected:
+	//~ Begin UObject
 #if WITH_EDITOR
 	virtual void PostEditChangeProperty(FPropertyChangedEvent& PropertyChangedEvent) override;
 #endif
+	//~ End UObject
 
 	void OnLineWidthChanged();
 	void OnVectorChanged();
@@ -53,9 +62,9 @@ protected:
 	virtual bool IsMeshVisible(int32 MeshIndex) override;
 	virtual void GenerateBorderVertices(TArray<FVector2D>& BorderVertices) override;
 
-	UPROPERTY(EditInstanceOnly, BlueprintReadWrite, Category="Shape", meta=(ClampMin="0.0", AllowPrivateAccess="true"))
+	UPROPERTY(EditInstanceOnly, BlueprintReadWrite, Setter, Getter, Category="Shape", meta=(ClampMin="0.0", AllowPrivateAccess="true"))
 	float LineWidth;
 
-	UPROPERTY(EditInstanceOnly, BlueprintReadWrite, Category="Shape", meta=(AllowPrivateAccess="true"))
+	UPROPERTY(EditInstanceOnly, BlueprintReadWrite, Setter, Getter, Category="Shape", meta=(AllowPrivateAccess="true"))
 	FVector2D Vector;
 };
diff --git a/Engine/Plugins/Experimental/Avalanche/Source/AvalancheShapes/Public/DynamicMeshes/AvaShapeNGonDynMesh.h b/Engine/Plugins/Experimental/Avalanche/Source/AvalancheShapes/Public/DynamicMeshes/AvaShapeNGonDynMesh.h
--- a/Engine/Plugins/Experimental/Avalanche/Source/AvalancheShapes/Public/DynamicMeshes/AvaShapeNGonDynMesh.h
+++ b/Engine/Plugins/Experimental/Avalanche/Source/AvalancheShapes/Public/DynamicMeshes/AvaShapeNGonDynMesh.h
@@ -5,45 +5,50 @@
 #include "AvaShapeRoundedPolygonDynMesh.h"
 #include "AvaShapeNGonDynMesh.generated.h"
 
-UCLASS(ClassGroup="Shape", BlueprintType, CustomConstructor, Within=AvaShapeActor)
-class AVALANCHESHAPES_API UAvaShapeNGonDynamicMesh : public UAvaShapeRoundedPolygonDynamicMesh
+UCLASS(MinimalAPI, ClassGroup="Shape", BlueprintType, CustomConstructor, Within=AvaShapeActor)
+class UAvaShapeNGonDynamicMesh : public UAvaShapeRoundedPolygonDynamicMesh
 {
 	GENERATED_BODY()
 
 	friend class FAvaShapeNGonDynamicMeshVisualizer;
 
 public:
+	static const FString MeshName;
 	static inline constexpr float MinNumSides = 3.f;
 	static inline constexpr float MaxNumSides = 128.f;
 
 	UAvaShapeNGonDynamicMesh()
 		: UAvaShapeNGonDynamicMesh(FVector2D(50.f, 50.f))
-	{
-	}
+	{}
 
 	UAvaShapeNGonDynamicMesh(const FVector2D& Size2D, const FLinearColor& InVertexColor = FLinearColor::White,
 		float InBevelSize = 0.f, uint8 InBevelSubdivisions = 0, uint8 InNumSides = 5)
 		: UAvaShapeRoundedPolygonDynamicMesh(Size2D, InVertexColor, InBevelSize, InBevelSubdivisions)
 		, NumSides(InNumSides)
+	{}
+
+	virtual const FString& GetMeshName() const override
 	{
+		return MeshName;
 	}
 
-	static const FString MeshName;
-	virtual const FString& GetMeshName() const override { return MeshName; }
-
-	UFUNCTION()
-	bool SetNumSides(uint8 InNumSides);
-	uint8 GetNumSides() const { return NumSides; }
+	AVALANCHESHAPES_API void SetNumSides(uint8 InNumSides);
+	uint8 GetNumSides() const
+	{
+		return NumSides;
+	}
 
 protected:
+	//~ Begin UObject
 #if WITH_EDITOR
 	virtual void PostEditChangeProperty(FPropertyChangedEvent& PropertyChangedEvent) override;
 #endif
+	//~ End UObject
 
 	virtual void OnNumSidesChanged();
 
 	virtual void GenerateBorderVertices(TArray<FVector2D>& BorderVertices) override;
 
-	UPROPERTY(EditInstanceOnly, BlueprintReadWrite, Category="Shape", meta=(ClampMin="3.0", ClampMax="128.0", DisplayName="Sides", AllowPrivateAccess="true"))
+	UPROPERTY(EditInstanceOnly, BlueprintReadWrite, Setter, Getter, Category="Shape", meta=(ClampMin="3.0", ClampMax="128.0", DisplayName="Sides", AllowPrivateAccess="true"))
 	uint8 NumSides;
 };
diff --git a/Engine/Plugins/Experimental/Avalanche/Source/AvalancheShapes/Public/DynamicMeshes/AvaShapeRectangleDynMesh.h b/Engine/Plugins/Experimental/Avalanche/Source/AvalancheShapes/Public/DynamicMeshes/AvaShapeRectangleDynMesh.h
--- a/Engine/Plugins/Experimental/Avalanche/Source/AvalancheShapes/Public/DynamicMeshes/AvaShapeRectangleDynMesh.h
+++ b/Engine/Plugins/Experimental/Avalanche/Source/AvalancheShapes/Public/DynamicMeshes/AvaShapeRectangleDynMesh.h
@@ -6,7 +6,7 @@
 #include "AvaShapeRectangleDynMesh.generated.h"
 
 USTRUCT(BlueprintType)
-struct AVALANCHESHAPES_API FAvaShapeRectangleCornerSettings
+struct FAvaShapeRectangleCornerSettings
 {
 	GENERATED_BODY()
 
@@ -28,112 +28,175 @@ struct AVALANCHESHAPES_API FAvaShapeRectangleCornerSettings
 	FVector2D CornerPositionCache = FVector2D::ZeroVector;
 };
 
-UCLASS(ClassGroup="Shape", BlueprintType, CustomConstructor, Within=AvaShapeActor)
-class AVALANCHESHAPES_API UAvaShapeRectangleDynamicMesh : public UAvaShape2DDynMeshBase
+UCLASS(MinimalAPI, ClassGroup="Shape", BlueprintType, CustomConstructor, Within=AvaShapeActor)
+class UAvaShapeRectangleDynamicMesh : public UAvaShape2DDynMeshBase
 {
 	GENERATED_BODY()
 
 	friend class FAvaShapeRectangleDynamicMeshVisualizer;
 
 public:
+	static const FString MeshName;
 	static inline constexpr float MinSlantAngle = -45.f;
 	static inline constexpr float MaxSlantAngle = 45.f;
 	static inline constexpr float CornerMinMargin = 0.1f;
 
 	UAvaShapeRectangleDynamicMesh()
 		: UAvaShapeRectangleDynamicMesh(FVector2D(50.f, 50.f))
-	{
-	}
+	{}
 
     explicit UAvaShapeRectangleDynamicMesh(const FVector2D& Size2D, const FLinearColor& InVertexColor = FLinearColor::White)
 		: UAvaShape2DDynMeshBase(Size2D, InVertexColor)
+	{}
+
+	virtual const FString& GetMeshName() const override
 	{
+		return MeshName;
 	}
 
-	static const FString MeshName;
-	virtual const FString& GetMeshName() const override { return MeshName; }
-
-	UFUNCTION()
-	bool SetHorizontalAlignment(EAvaHorizontalAlignment InHorizontalAlignment);
-	EAvaHorizontalAlignment GetHorizontalAlignment() const { return HorizontalAlignment; }
+	AVALANCHESHAPES_API void SetHorizontalAlignment(EAvaHorizontalAlignment InHorizontalAlignment);
+	EAvaHorizontalAlignment GetHorizontalAlignment() const
+	{
+		return HorizontalAlignment;
+	}
 
-	UFUNCTION()
-	bool SetVerticalAlignment(EAvaVerticalAlignment InVerticalAlignment);
-	EAvaVerticalAlignment GetVerticalAlignment() const { return VerticalAlignment; }
+	AVALANCHESHAPES_API void SetVerticalAlignment(EAvaVerticalAlignment InVerticalAlignment);
+	EAvaVerticalAlignment GetVerticalAlignment() const
+	{
+		return VerticalAlignment;
+	}
 
-	UFUNCTION()
-	bool SetLeftSlant(float InSlant);
-	float GetLeftSlant() const { return LeftSlant; }
+	AVALANCHESHAPES_API void SetLeftSlant(float InSlant);
+	float GetLeftSlant() const
+	{
+		return LeftSlant;
+	}
 
-	UFUNCTION()
-	bool SetRightSlant(float InSlant);
-	float GetRightSlant() const { return RightSlant; }
+	AVALANCHESHAPES_API void SetRightSlant(float InSlant);
+	float GetRightSlant() const
+	{
+		return RightSlant;
+	}
 
-	UFUNCTION()
-	bool SetGlobalBevelSize(float InBevelSize);
-	float GetGlobalBevelSize() const { return GlobalBevelSize; }
+	AVALANCHESHAPES_API void SetGlobalBevelSize(float InBevelSize);
+	float GetGlobalBevelSize() const
+	{
+		return GlobalBevelSize;
+	}
 
 	/** Controls all corners bevel subdivisions */
-	UFUNCTION()
-	void SetGlobalBevelSubdivisions(uint8 InGlobalBevelSubdivisions);
-	uint8 GetGlobalBevelSubdivisions() const { return GlobalBevelSubdivisions; }
+	AVALANCHESHAPES_API void SetGlobalBevelSubdivisions(uint8 InGlobalBevelSubdivisions);
+	uint8 GetGlobalBevelSubdivisions() const
+	{
+		return GlobalBevelSubdivisions;
+	}
 
-	UFUNCTION()
-	void SetTopLeft(const FAvaShapeRectangleCornerSettings& InCornerSettings);
+	AVALANCHESHAPES_API void SetTopLeft(const FAvaShapeRectangleCornerSettings& InCornerSettings);
+	const FAvaShapeRectangleCornerSettings& GetTopLeft() const
+	{
+		return TopLeft;
+	}
 
-	UFUNCTION()
-	void SetTopRight(const FAvaShapeRectangleCornerSettings& InCornerSettings);
+	AVALANCHESHAPES_API void SetTopRight(const FAvaShapeRectangleCornerSettings& InCornerSettings);
+	const FAvaShapeRectangleCornerSettings& GetTopRight() const
+	{
+		return TopRight;
+	}
 
-	UFUNCTION()
-	void SetBottomLeft(const FAvaShapeRectangleCornerSettings& InCornerSettings);
+	AVALANCHESHAPES_API void SetBottomLeft(const FAvaShapeRectangleCornerSettings& InCornerSettings);
+	const FAvaShapeRectangleCornerSettings& GetBottomLeft() const
+	{
+		return BottomLeft;
+	}
 
-	UFUNCTION()
-	void SetBottomRight(const FAvaShapeRectangleCornerSettings& InCornerSettings);
+	AVALANCHESHAPES_API void SetBottomRight(const FAvaShapeRectangleCornerSettings& InCornerSettings);
+	const FAvaShapeRectangleCornerSettings& GetBottomRight() const
+	{
+		return BottomRight;
+	}
 
-	EAvaShapeCornerType GetTopLeftCornerType() const { return TopLeft.Type; }
-	bool SetTopLeftCornerType(EAvaShapeCornerType InType);
+	AVALANCHESHAPES_API void SetTopLeftCornerType(EAvaShapeCornerType InType);
+	EAvaShapeCornerType GetTopLeftCornerType() const
+	{
+		return TopLeft.Type;
+	}
 
-	float GetTopLeftBevelSize() const { return TopLeft.BevelSize; }
-	bool SetTopLeftBevelSize(float InSize);
+	AVALANCHESHAPES_API void SetTopLeftBevelSize(float InSize);
+	float GetTopLeftBevelSize() const
+	{
+		return TopLeft.BevelSize;
+	}
 
 	/** Bevel subdivisions for top left corner */
-	uint8 GetTopLeftBevelSubdivisions() const { return TopLeft.BevelSubdivisions; }
-	void SetTopLeftBevelSubdivisions(uint8 InBevelSubdivisions);
+	AVALANCHESHAPES_API void SetTopLeftBevelSubdivisions(uint8 InBevelSubdivisions);
+	uint8 GetTopLeftBevelSubdivisions() const
+	{
+		return TopLeft.BevelSubdivisions;
+	}
 
-	EAvaShapeCornerType GetBottomLeftCornerType() const { return BottomLeft.Type; }
-	bool SetBottomLeftCornerType(EAvaShapeCornerType InType);
+	AVALANCHESHAPES_API void SetBottomLeftCornerType(EAvaShapeCornerType InType);
+	EAvaShapeCornerType GetBottomLeftCornerType() const
+	{
+		return BottomLeft.Type;
+	}
 
-	float GetBottomLeftBevelSize() const { return BottomLeft.BevelSize; }
-	bool SetBottomLeftBevelSize(float InSize);
+	AVALANCHESHAPES_API void SetBottomLeftBevelSize(float InSize);
+	float GetBottomLeftBevelSize() const
+	{
+		return BottomLeft.BevelSize;
+	}
 
 	/** Bevel subdivisions for bottom left corner */
-	uint8 GetBottomLeftBevelSubdivisions() const { return BottomLeft.BevelSubdivisions; }
-	void SetBottomLeftBevelSubdivisions(uint8 InBevelSubdivisions);
+	AVALANCHESHAPES_API void SetBottomLeftBevelSubdivisions(uint8 InBevelSubdivisions);
+	uint8 GetBottomLeftBevelSubdivisions() const
+	{
+		return BottomLeft.BevelSubdivisions;
+	}
 
-	EAvaShapeCornerType GetTopRightCornerType() const { return TopRight.Type; }
-	bool SetTopRightCornerType(EAvaShapeCornerType InType);
+	AVALANCHESHAPES_API void SetTopRightCornerType(EAvaShapeCornerType InType);
+	EAvaShapeCornerType GetTopRightCornerType() const
+	{
+		return TopRight.Type;
+	}
 
-	float GetTopRightBevelSize() const { return TopRight.BevelSize; }
-	bool SetTopRightBevelSize(float InSize);
+	AVALANCHESHAPES_API void SetTopRightBevelSize(float InSize);
+	float GetTopRightBevelSize() const
+	{
+		return TopRight.BevelSize;
+	}
 
 	/** Bevel subdivisions for top right corner */
-	uint8 GetTopRightBevelSubdivisions() const { return TopRight.BevelSubdivisions; }
-	void SetTopRightBevelSubdivisions(uint8 InBevelSubdivisions);
+	AVALANCHESHAPES_API void SetTopRightBevelSubdivisions(uint8 InBevelSubdivisions);
+	uint8 GetTopRightBevelSubdivisions() const
+	{
+		return TopRight.BevelSubdivisions;
+	}
 
-	EAvaShapeCornerType GetBottomRightCornerType() const { return BottomRight.Type; }
-	bool SetBottomRightCornerType(EAvaShapeCornerType InType);
+	AVALANCHESHAPES_API void SetBottomRightCornerType(EAvaShapeCornerType InType);
+	EAvaShapeCornerType GetBottomRightCornerType() const
+	{
+		return BottomRight.Type;
+	}
 
-	float GetBottomRightBevelSize() const { return BottomRight.BevelSize; }
-	bool SetBottomRightBevelSize(float InSize);
+	AVALANCHESHAPES_API void SetBottomRightBevelSize(float InSize);
+	float GetBottomRightBevelSize() const
+	{
+		return BottomRight.BevelSize;
+	}
 
 	/** Bevel subdivisions for bottom right corner */
-	uint8 GetBottomRightBevelSubdivisions() const { return BottomRight.BevelSubdivisions; }
-	void SetBottomRightBevelSubdivisions(uint8 InBevelSubdivisions);
+	AVALANCHESHAPES_API void SetBottomRightBevelSubdivisions(uint8 InBevelSubdivisions);
+	uint8 GetBottomRightBevelSubdivisions() const
+	{
+		return BottomRight.BevelSubdivisions;
+	}
 
 protected:
+	//~ Begin UObject
 #if WITH_EDITOR
 	virtual void PostEditChangeProperty(FPropertyChangedEvent& PropertyChangedEvent) override;
 #endif
+	//~ End UObject
 
 	float GetMaximumBevelSize() const;
 
@@ -172,35 +235,35 @@ protected:
 
 	virtual void OnSizeChanged() override;
 
-	UPROPERTY(EditInstanceOnly, BlueprintReadWrite, Category="Shape", meta=(AllowPrivateAccess="true"))
+	UPROPERTY(EditInstanceOnly, BlueprintReadWrite, Setter, Getter, Category="Shape", meta=(AllowPrivateAccess="true"))
 	EAvaHorizontalAlignment HorizontalAlignment = EAvaHorizontalAlignment::Center;
 
-	UPROPERTY(EditInstanceOnly, BlueprintReadWrite, Category="Shape", meta=(AllowPrivateAccess="true"))
+	UPROPERTY(EditInstanceOnly, BlueprintReadWrite, Setter, Getter, Category="Shape", meta=(AllowPrivateAccess="true"))
 	EAvaVerticalAlignment VerticalAlignment = EAvaVerticalAlignment::Center;
 
 	/** Angle in degrees for the left slant of the rectangle */
-	UPROPERTY(EditInstanceOnly, BlueprintReadWrite, Category="Shape", meta=(ClampMin="-45.0",ClampMax="45.0", AllowPrivateAccess="true"))
+	UPROPERTY(EditInstanceOnly, BlueprintReadWrite, Setter, Getter, Category="Shape", meta=(ClampMin="-45.0",ClampMax="45.0", AllowPrivateAccess="true"))
 	float LeftSlant = 0.f;
 
 	/** Angle in degrees for the right slant of the rectangle */
-	UPROPERTY(EditInstanceOnly, BlueprintReadWrite, Category="Shape", meta=(ClampMin="-45.0",ClampMax="45.0", AllowPrivateAccess="true"))
+	UPROPERTY(EditInstanceOnly, BlueprintReadWrite, Setter, Getter, Category="Shape", meta=(ClampMin="-45.0",ClampMax="45.0", AllowPrivateAccess="true"))
 	float RightSlant = 0.f;
 
-	UPROPERTY(EditInstanceOnly, BlueprintReadWrite, Category="Shape", meta=(ClampMin="0.0", AllowPrivateAccess="true"))
+	UPROPERTY(EditInstanceOnly, BlueprintReadWrite, Setter, Getter, Category="Shape", meta=(ClampMin="0.0", AllowPrivateAccess="true"))
 	float GlobalBevelSize = 0.f;
 
-	UPROPERTY(EditInstanceOnly, BlueprintReadWrite, Category="Shape", meta=(ClampMax="128.0", AllowPrivateAccess="true"))
+	UPROPERTY(EditInstanceOnly, BlueprintReadWrite, Setter, Getter, Category="Shape", meta=(ClampMax="128.0", AllowPrivateAccess="true"))
 	uint8 GlobalBevelSubdivisions = DefaultSubdivisions;
 
-	UPROPERTY(EditInstanceOnly, BlueprintReadWrite, Category="Shape", meta=(AllowPrivateAccess="true"))
+	UPROPERTY(EditInstanceOnly, BlueprintReadWrite, Setter, Getter, Category="Shape", meta=(AllowPrivateAccess="true"))
 	FAvaShapeRectangleCornerSettings TopLeft;
 
-	UPROPERTY(EditInstanceOnly, BlueprintReadWrite, Category="Shape", meta=(AllowPrivateAccess="true"))
+	UPROPERTY(EditInstanceOnly, BlueprintReadWrite, Setter, Getter, Category="Shape", meta=(AllowPrivateAccess="true"))
 	FAvaShapeRectangleCornerSettings TopRight;
 
-	UPROPERTY(EditInstanceOnly, BlueprintReadWrite, Category="Shape", meta=(AllowPrivateAccess="true"))
+	UPROPERTY(EditInstanceOnly, BlueprintReadWrite, Setter, Getter, Category="Shape", meta=(AllowPrivateAccess="true"))
 	FAvaShapeRectangleCornerSettings BottomLeft;
 
-	UPROPERTY(EditInstanceOnly, BlueprintReadWrite, Category="Shape", meta=(AllowPrivateAccess="true"))
+	UPROPERTY(EditInstanceOnly, BlueprintReadWrite, Setter, Getter, Category="Shape", meta=(AllowPrivateAccess="true"))
 	FAvaShapeRectangleCornerSettings BottomRight;
 };
diff --git a/Engine/Plugins/Experimental/Avalanche/Source/AvalancheShapes/Public/DynamicMeshes/AvaShapeRingDynMesh.h b/Engine/Plugins/Experimental/Avalanche/Source/AvalancheShapes/Public/DynamicMeshes/AvaShapeRingDynMesh.h
--- a/Engine/Plugins/Experimental/Avalanche/Source/AvalancheShapes/Public/DynamicMeshes/AvaShapeRingDynMesh.h
+++ b/Engine/Plugins/Experimental/Avalanche/Source/AvalancheShapes/Public/DynamicMeshes/AvaShapeRingDynMesh.h
@@ -5,21 +5,21 @@
 #include "DynamicMeshes/AvaShape2DDynMeshBase.h"
 #include "AvaShapeRingDynMesh.generated.h"
 
-UCLASS(ClassGroup="Shape", BlueprintType, CustomConstructor, Within=AvaShapeActor)
-class AVALANCHESHAPES_API UAvaShapeRingDynamicMesh : public UAvaShape2DDynMeshBase
+UCLASS(MinimalAPI, ClassGroup="Shape", BlueprintType, CustomConstructor, Within=AvaShapeActor)
+class UAvaShapeRingDynamicMesh : public UAvaShape2DDynMeshBase
 {
 	GENERATED_BODY()
 
 	friend class FAvaShapeRingDynamicMeshVisualizer;
 
 public:
+	static const FString MeshName;
 	static inline constexpr uint8 MinNumSides = 3;
 	static inline constexpr uint8 MaxNumSides = 128;
 
 	UAvaShapeRingDynamicMesh()
 		: UAvaShapeRingDynamicMesh(FVector2D(50.f, 50.f))
-	{
-	}
+	{}
 
 	UAvaShapeRingDynamicMesh(const FVector2D& Size2D, const FLinearColor& InVertexColor = FLinearColor::White,
 		uint8 InNumSides = 64, float InInnerSize = 0.8f, float InAngleDegree = 360.f, float InStartDegree = 0.f)
@@ -32,29 +32,41 @@ public:
 		bDoNotRecenterVertices = true;
 	}
 
-	static const FString MeshName;
-	virtual const FString& GetMeshName() const override { return MeshName; }
+	virtual const FString& GetMeshName() const override
+	{
+		return MeshName;
+	}
 
-	UFUNCTION()
-	bool SetNumSides(uint8 InNumSides);
-	uint8 GetNumSides() const { return NumSides; }
+	AVALANCHESHAPES_API void SetNumSides(uint8 InNumSides);
+	uint8 GetNumSides() const
+	{
+		return NumSides;
+	}
 
-	UFUNCTION()
-	bool SetInnerSize(float InInnerSize);
-	float GetInnerSize() const { return InnerSize; }
+	AVALANCHESHAPES_API void SetInnerSize(float InInnerSize);
+	float GetInnerSize() const
+	{
+		return InnerSize;
+	}
 
-	UFUNCTION()
-	bool SetAngleDegree(float InDegree);
-	float GetAngleDegree() const { return AngleDegree; }
+	AVALANCHESHAPES_API void SetAngleDegree(float InDegree);
+	float GetAngleDegree() const
+	{
+		return AngleDegree;
+	}
 
-	UFUNCTION()
-	bool SetStartDegree(float InDegree);
-	float GetStartDegree() const { return StartDegree; }
+	AVALANCHESHAPES_API void SetStartDegree(float InDegree);
+	float GetStartDegree() const
+	{
+		return StartDegree;
+	}
 
 protected:
+	//~ Begin UObject
 #if WITH_EDITOR
 	virtual void PostEditChangeProperty(FPropertyChangedEvent& PropertyChangedEvent) override;
 #endif
+	//~ End UObject
 
 	virtual void OnNumSidesChanged();
 	virtual void OnInnerSizeChanged();
@@ -64,17 +76,17 @@ protected:
 	virtual bool IsMeshVisible(int32 MeshIndex) override;
 	virtual bool CreateMesh(FAvaShapeMesh& InMesh) override;
 
-	UPROPERTY(EditInstanceOnly, BlueprintReadWrite, Category="Shape", meta=(ClampMin="3.0", ClampMax="128.0", DisplayName="Sides", AllowPrivateAccess="true"))
+	UPROPERTY(EditInstanceOnly, BlueprintReadWrite, Setter, Getter, Category="Shape", meta=(ClampMin="3.0", ClampMax="128.0", DisplayName="Sides", AllowPrivateAccess="true"))
 	uint8 NumSides;
 
-	UPROPERTY(EditInstanceOnly, BlueprintReadWrite, Category="Shape", meta=(ClampMin="0.01", ClampMax="0.99",DisplayName="Inner Size", AllowPrivateAccess="true"))
+	UPROPERTY(EditInstanceOnly, BlueprintReadWrite, Setter, Getter, Category="Shape", meta=(ClampMin="0.01", ClampMax="0.99",DisplayName="Inner Size", AllowPrivateAccess="true"))
 	float InnerSize;
 
 	// represents the angle in degree for the ring
-	UPROPERTY(EditInstanceOnly, BlueprintReadWrite, Category="Shape", meta=(ClampMin="0.0", ClampMax="360.0", DisplayName="Angle degree", AllowPrivateAccess="true"))
+	UPROPERTY(EditInstanceOnly, BlueprintReadWrite, Setter, Getter, Category="Shape", meta=(ClampMin="0.0", ClampMax="360.0", DisplayName="Angle degree", AllowPrivateAccess="true"))
 	float AngleDegree;
 
 	// represents the starting angle in degree for the ring
-	UPROPERTY(EditInstanceOnly, BlueprintReadWrite, Category="Shape", meta=(ClampMin="0.0", ClampMax="360.0", DisplayName="Start Degree", AllowPrivateAccess="true"))
+	UPROPERTY(EditInstanceOnly, BlueprintReadWrite, Setter, Getter, Category="Shape", meta=(ClampMin="0.0", ClampMax="360.0", DisplayName="Start Degree", AllowPrivateAccess="true"))
 	float StartDegree;
 };
diff --git a/Engine/Plugins/Experimental/Avalanche/Source/AvalancheShapes/Public/DynamicMeshes/AvaShapeRoundedPolygonDynMesh.h b/Engine/Plugins/Experimental/Avalanche/Source/AvalancheShapes/Public/DynamicMeshes/AvaShapeRoundedPolygonDynMesh.h
--- a/Engine/Plugins/Experimental/Avalanche/Source/AvalancheShapes/Public/DynamicMeshes/AvaShapeRoundedPolygonDynMesh.h
+++ b/Engine/Plugins/Experimental/Avalanche/Source/AvalancheShapes/Public/DynamicMeshes/AvaShapeRoundedPolygonDynMesh.h
@@ -6,7 +6,7 @@
 #include "DynamicMeshes/AvaShape2DDynMeshBase.h"
 #include "AvaShapeRoundedPolygonDynMesh.generated.h"
 
-struct AVALANCHESHAPES_API FAvaShapeRoundedCornerMetrics
+struct FAvaShapeRoundedCornerMetrics
 {
 	float Angle;
 	FAvaShapeCachedVertex2D Start;
@@ -18,8 +18,8 @@ struct AVALANCHESHAPES_API FAvaShapeRoundedCornerMetrics
 	bool bValid = false;
 };
 
-UCLASS(ClassGroup="Shape", Abstract, BlueprintType, CustomConstructor)
-class AVALANCHESHAPES_API UAvaShapeRoundedPolygonDynamicMesh : public UAvaShape2DDynMeshBase
+UCLASS(MinimalAPI, ClassGroup="Shape", Abstract, BlueprintType, CustomConstructor)
+class UAvaShapeRoundedPolygonDynamicMesh : public UAvaShape2DDynMeshBase
 {
 	GENERATED_BODY()
 
@@ -44,18 +44,24 @@ public:
 		, BevelSubdivisions(InBevelSubdivisions)
 	{}
 
-	UFUNCTION()
-	bool SetBevelSize(float InBevelSize);
-	float GetBevelSize() const { return BevelSize; }
+	AVALANCHESHAPES_API void SetBevelSize(float InBevelSize);
+	float GetBevelSize() const
+	{
+		return BevelSize;
+	}
 
-	UFUNCTION()
-	void SetBevelSubdivisions(uint8 InBevelSubdivisions);
-	uint8 GetBevelSubdivisions() const { return BevelSubdivisions; }
+	AVALANCHESHAPES_API void SetBevelSubdivisions(uint8 InBevelSubdivisions);
+	uint8 GetBevelSubdivisions() const
+	{
+		return BevelSubdivisions;
+	}
 
 protected:
+	//~ Begin UObject
 #if WITH_EDITOR
 	virtual void PostEditChangeProperty(FPropertyChangedEvent& PropertyChangedEvent) override;
 #endif
+	//~ End UObject
 
 	virtual void OnSizeChanged() override;
 
@@ -71,9 +77,9 @@ protected:
 	virtual bool UseCenteredVertex() { return false; }
 	virtual bool CreateMesh(FAvaShapeMesh& InMesh) override;
 
-	UPROPERTY(EditInstanceOnly, BlueprintReadWrite, Category="Shape", meta=(ClampMin="0.0", AllowPrivateAccess="true"))
+	UPROPERTY(EditInstanceOnly, BlueprintReadWrite, Setter, Getter, Category="Shape", meta=(ClampMin="0.0", AllowPrivateAccess="true"))
 	float BevelSize = 0.f;
 
-	UPROPERTY(EditInstanceOnly, BlueprintReadWrite, Category="Shape", meta=(ClampMax="64.0", AllowPrivateAccess="true"))
+	UPROPERTY(EditInstanceOnly, BlueprintReadWrite, Setter, Getter, Category="Shape", meta=(ClampMax="64.0", AllowPrivateAccess="true"))
 	uint8 BevelSubdivisions = UAvaShapeDynamicMeshBase::DefaultSubdivisions;
 };
diff --git a/Engine/Plugins/Experimental/Avalanche/Source/AvalancheShapes/Public/DynamicMeshes/AvaShapeSphereDynMesh.h b/Engine/Plugins/Experimental/Avalanche/Source/AvalancheShapes/Public/DynamicMeshes/AvaShapeSphereDynMesh.h
--- a/Engine/Plugins/Experimental/Avalanche/Source/AvalancheShapes/Public/DynamicMeshes/AvaShapeSphereDynMesh.h
+++ b/Engine/Plugins/Experimental/Avalanche/Source/AvalancheShapes/Public/DynamicMeshes/AvaShapeSphereDynMesh.h
@@ -5,14 +5,15 @@
 #include "DynamicMeshes/AvaShape3DDynMeshBase.h"
 #include "AvaShapeSphereDynMesh.generated.h"
 
-UCLASS(ClassGroup="Shape", BlueprintType, CustomConstructor, Within=AvaShapeActor)
-class AVALANCHESHAPES_API UAvaShapeSphereDynamicMesh : public UAvaShape3DDynMeshBase
+UCLASS(MinimalAPI, ClassGroup="Shape", BlueprintType, CustomConstructor, Within=AvaShapeActor)
+class UAvaShapeSphereDynamicMesh : public UAvaShape3DDynMeshBase
 {
 	GENERATED_BODY()
 
 	friend class FAvaShapeSphereDynamicMeshVisualizer;
 
 public:
+	static const FString MeshName;
 	static constexpr uint8 MinNumSides = 4;
 	static constexpr uint8 MaxNumSides = 128;
 
@@ -23,7 +24,7 @@ public:
 
 	UAvaShapeSphereDynamicMesh()
 		: UAvaShapeSphereDynamicMesh(FVector(50.f, 50.f, 50.f))
-	{ }
+	{}
 
 	UAvaShapeSphereDynamicMesh(
 		const FVector& InSize,
@@ -43,40 +44,50 @@ public:
 		, LatitudeDegree(InLatitudeDegree)
 		, NumSides(InNumSides)
 		, Radius(InRadius)
+	{}
+
+	virtual const FString& GetMeshName() const override
 	{
+		return MeshName;
 	}
 
-	static const FString MeshName;
-	virtual const FString& GetMeshName() const override { return MeshName; }
-
-	UFUNCTION()
-	bool SetNumSides(uint8 InNumSides);
-	uint8 GetNumSides() const { return NumSides; }
-
-	UFUNCTION()
-	float GetRadius() const { return Radius; }
+	AVALANCHESHAPES_API void SetNumSides(uint8 InNumSides);
+	uint8 GetNumSides() const
+	{
+		return NumSides;
+	}
 
-	UFUNCTION()
-	bool SetStartLatitude(float InDegree);
-	float GetStartLatitude() const { return StartLatitude; }
+	AVALANCHESHAPES_API void SetStartLatitude(float InDegree);
+	float GetStartLatitude() const
+	{
+		return StartLatitude;
+	}
 
-	UFUNCTION()
-	bool SetLatitudeDegree(float InDegree);
-	float GetLatitudeDegree() const { return LatitudeDegree; }
+	AVALANCHESHAPES_API void SetLatitudeDegree(float InDegree);
+	float GetLatitudeDegree() const
+	{
+		return LatitudeDegree;
+	}
 
-	UFUNCTION()
-	bool SetStartLongitude(float InDegree);
-	float GetStartLongitude() const { return StartLongitude; }
+	AVALANCHESHAPES_API void SetStartLongitude(float InDegree);
+	float GetStartLongitude() const
+	{
+		return StartLongitude;
+	}
 
-	UFUNCTION()
-	bool SetEndLongitude(float InDegree);
-	float GetEndLongitude() const { return EndLongitude; }
+	AVALANCHESHAPES_API void SetEndLongitude(float InDegree);
+	float GetEndLongitude() const
+	{
+		return EndLongitude;
+	}
 
 protected:
+	//~ Begin UObject
 #if WITH_EDITOR
 	virtual void PreEditChange(FProperty* PropertyAboutToChange) override;
 	virtual void PostEditChangeProperty(FPropertyChangedEvent& PropertyChangedEvent) override;
 #endif
+	//~ End UObject
 
 	virtual void OnNumSidesChanged();
 	virtual void OnRadiusChanged();
@@ -106,27 +117,28 @@ protected:
 	virtual bool CreateUVs(FAvaShapeMesh& InMesh, FAvaShapeMaterialUVParameters& InParams) override;
 
 	// represents the longitude (Z) angle in degree for the sphere at the start
-	UPROPERTY(EditInstanceOnly, BlueprintReadWrite, Category="Shape", meta=(ClampMin="0.0", ClampMax="180.0", DisplayName="Start Longitude degree", AllowPrivateAccess="true"))
+	UPROPERTY(EditInstanceOnly, BlueprintReadWrite, Setter, Getter, Category="Shape", meta=(ClampMin="0.0", ClampMax="180.0", DisplayName="Start Longitude degree", AllowPrivateAccess="true"))
 	float StartLongitude;
 	float PreEditStartLongitude;
 
 	// represents the longitude (Z) angle in degree for the sphere at the end
-	UPROPERTY(EditInstanceOnly, BlueprintReadWrite, Category="Shape", meta=(ClampMin="0.0", ClampMax="180.0", DisplayName="End Longitude degree", AllowPrivateAccess="true"))
+	UPROPERTY(EditInstanceOnly, BlueprintReadWrite, Setter, Getter, Category="Shape", meta=(ClampMin="0.0", ClampMax="180.0", DisplayName="End Longitude degree", AllowPrivateAccess="true"))
 	float EndLongitude;
 	float PreEditEndLongitude;
 
 	// represents the latitude (Y) angle in degree for the sphere at the start
-	UPROPERTY(EditInstanceOnly, BlueprintReadWrite, Category="Shape", meta=(ClampMin="0.0", ClampMax="360.0", DisplayName="Start Latitude degree", AllowPrivateAccess="true"))
+	UPROPERTY(EditInstanceOnly, BlueprintReadWrite, Setter, Getter, Category="Shape", meta=(ClampMin="0.0", ClampMax="360.0", DisplayName="Start Latitude degree", AllowPrivateAccess="true"))
 	float StartLatitude;
 
 	// represents the total latitude (Y) angle in degree for the sphere
-	UPROPERTY(EditInstanceOnly, BlueprintReadWrite, Category="Shape", meta=(ClampMin="0.0", ClampMax="360.0", DisplayName="Latitude degree", AllowPrivateAccess="true"))
+	UPROPERTY(EditInstanceOnly, BlueprintReadWrite, Setter, Getter, Category="Shape", meta=(ClampMin="0.0", ClampMax="360.0", DisplayName="Latitude degree", AllowPrivateAccess="true"))
 	float LatitudeDegree;
 
 	// represents the precision of the sphere mesh
-	UPROPERTY(EditInstanceOnly, BlueprintReadWrite, Category="Shape", meta=(ClampMin="4.0", ClampMax="128.0", DisplayName="Sides", AllowPrivateAccess="true"))
+	UPROPERTY(EditInstanceOnly, BlueprintReadWrite, Setter, Getter, Category="Shape", meta=(ClampMin="4.0", ClampMax="128.0", DisplayName="Sides", AllowPrivateAccess="true"))
 	uint8 NumSides;
 
+private:
 	// represents the radius ratio of the sphere
 	UPROPERTY()
 	float Radius;
diff --git a/Engine/Plugins/Experimental/Avalanche/Source/AvalancheShapes/Public/DynamicMeshes/AvaShapeStarDynMesh.h b/Engine/Plugins/Experimental/Avalanche/Source/AvalancheShapes/Public/DynamicMeshes/AvaShapeStarDynMesh.h
--- a/Engine/Plugins/Experimental/Avalanche/Source/AvalancheShapes/Public/DynamicMeshes/AvaShapeStarDynMesh.h
+++ b/Engine/Plugins/Experimental/Avalanche/Source/AvalancheShapes/Public/DynamicMeshes/AvaShapeStarDynMesh.h
@@ -5,45 +5,52 @@
 #include "AvaShapeRoundedPolygonDynMesh.h"
 #include "AvaShapeStarDynMesh.generated.h"
 
-UCLASS(ClassGroup="Shape", BlueprintType, CustomConstructor, Within=AvaShapeActor)
-class AVALANCHESHAPES_API UAvaShapeStarDynamicMesh : public UAvaShapeRoundedPolygonDynamicMesh
+UCLASS(MinimalAPI, ClassGroup="Shape", BlueprintType, CustomConstructor, Within=AvaShapeActor)
+class UAvaShapeStarDynamicMesh : public UAvaShapeRoundedPolygonDynamicMesh
 {
 	GENERATED_BODY()
 
 	friend class FAvaShapeStarDynamicMeshVisualizer;
 
 public:
+	static const FString MeshName;
 	static inline constexpr uint8 MinNumPoints = 2;
 	static inline constexpr uint8 MaxNumPoints = 128;
 
 	UAvaShapeStarDynamicMesh()
 		: UAvaShapeStarDynamicMesh(FVector2D(50.f, 50.f))
-	{
-	}
+	{}
 
 	UAvaShapeStarDynamicMesh(const FVector2D& Size2D, const FLinearColor& InVertexColor = FLinearColor::White,
 		float InBevelSize = 0.f, uint8 InBevelSubdivisions = 0, uint8 InNumPoints = 5, float InInnerSize = 0.5f)
 		: UAvaShapeRoundedPolygonDynamicMesh(Size2D, InVertexColor, InBevelSize, InBevelSubdivisions)
 		, NumPoints(InNumPoints)
 		, InnerSize(InInnerSize)
+	{}
+
+	virtual const FString& GetMeshName() const override
 	{
+		return MeshName;
 	}
 
-	static const FString MeshName;
-	virtual const FString& GetMeshName() const override { return MeshName; }
-
-	UFUNCTION()
-	bool SetNumPoints(uint8 InNumPoints);
-	uint8 GetNumPoints() const { return NumPoints; }
+	AVALANCHESHAPES_API void SetNumPoints(uint8 InNumPoints);
+	uint8 GetNumPoints() const
+	{
+		return NumPoints;
+	}
 
-	UFUNCTION()
-	bool SetInnerSize(float InInnerSize);
-	float GetInnerSize() const { return InnerSize; }
+	AVALANCHESHAPES_API void SetInnerSize(float InInnerSize);
+	float GetInnerSize() const
+	{
+		return InnerSize;
+	}
 
 protected:
+	//~ Begin UObject
 #if WITH_EDITOR
 	virtual void PostEditChangeProperty(FPropertyChangedEvent& PropertyChangedEvent) override;
 #endif
+	//~ End UObject
 
 	virtual void OnNumSidesChanged();
 	virtual void OnInnerSizeChanged();
@@ -53,9 +60,9 @@ protected:
 	virtual bool UseCenteredVertex() override { return true; }
 	virtual bool IsMeshVisible(int32 MeshIndex) override;
 
-	UPROPERTY(EditInstanceOnly, BlueprintReadWrite, Category="Shape", meta=(ClampMin="2.0", ClampMax="128.0", DisplayName="Points", AllowPrivateAccess="true"))
+	UPROPERTY(EditInstanceOnly, BlueprintReadWrite, Setter, Getter, Category="Shape", meta=(ClampMin="2.0", ClampMax="128.0", DisplayName="Points", AllowPrivateAccess="true"))
 	uint8 NumPoints = 0;
 
-	UPROPERTY(EditInstanceOnly, BlueprintReadWrite, Category="Shape", meta=(ClampMin="0.0", ClampMax="0.99", DisplayName="Inner Size Fraction", AllowPrivateAccess="true"))
+	UPROPERTY(EditInstanceOnly, BlueprintReadWrite, Setter, Getter, Category="Shape", meta=(ClampMin="0.0", ClampMax="0.99", DisplayName="Inner Size Fraction", AllowPrivateAccess="true"))
 	float InnerSize = 0.f;
 };
diff --git a/Engine/Plugins/Experimental/Avalanche/Source/AvalancheShapes/Public/DynamicMeshes/AvaShapeTorusDynMesh.h b/Engine/Plugins/Experimental/Avalanche/Source/AvalancheShapes/Public/DynamicMeshes/AvaShapeTorusDynMesh.h
--- a/Engine/Plugins/Experimental/Avalanche/Source/AvalancheShapes/Public/DynamicMeshes/AvaShapeTorusDynMesh.h
+++ b/Engine/Plugins/Experimental/Avalanche/Source/AvalancheShapes/Public/DynamicMeshes/AvaShapeTorusDynMesh.h
@@ -5,14 +5,15 @@
 #include "AvaShape3DDynMeshBase.h"
 #include "AvaShapeTorusDynMesh.generated.h"
 
-UCLASS(ClassGroup="Shape", BlueprintType, CustomConstructor, Within=AvaShapeActor)
-class AVALANCHESHAPES_API UAvaShapeTorusDynamicMesh : public UAvaShape3DDynMeshBase
+UCLASS(MinimalAPI, ClassGroup="Shape", BlueprintType, CustomConstructor, Within=AvaShapeActor)
+class UAvaShapeTorusDynamicMesh : public UAvaShape3DDynMeshBase
 {
 	GENERATED_BODY()
 
 	friend class FAvaShapeTorusDynamicMeshVisualizer;
 
 public:
+	static const FString MeshName;
 	static constexpr uint8 MinNumSlices = 3;
 	static constexpr uint8 MaxNumSlices = 128;
 	static constexpr uint8 MinNumSides = 4;
@@ -23,8 +24,7 @@ public:
 
 	UAvaShapeTorusDynamicMesh()
 		: UAvaShapeTorusDynamicMesh(FVector(50.f, 50.f, 50.f))
-	{
-	}
+	{}
 
 	UAvaShapeTorusDynamicMesh(
 		const FVector& InSize,
@@ -40,36 +40,49 @@ public:
 		, InnerSize(InInnerSize)
 		, AngleDegree(InAngleDegree)
 		, StartDegree(InStartDegree)
+	{}
+
+	virtual const FString& GetMeshName() const override
 	{
+		return MeshName;
 	}
 
-	static const FString MeshName;
-	virtual const FString& GetMeshName() const override { return MeshName; }
-
-	UFUNCTION()
-	bool SetNumSides(uint8 InNumSides);
-	uint8 GetNumSides() const { return NumSides; }
+	AVALANCHESHAPES_API void SetNumSides(uint8 InNumSides);
+	uint8 GetNumSides() const
+	{
+		return NumSides;
+	}
 
-	UFUNCTION()
-	bool SetNumSlices(uint8 InNumSlices);
-	uint8 GetNumSlices() const { return NumSlices; }
+	AVALANCHESHAPES_API void SetNumSlices(uint8 InNumSlices);
+	uint8 GetNumSlices() const
+	{
+		return NumSlices;
+	}
 
-	UFUNCTION()
-	bool SetInnerSize(float InInnerSize);
-	float GetInnerSize() const { return InnerSize; }
+	AVALANCHESHAPES_API void SetInnerSize(float InInnerSize);
+	float GetInnerSize() const
+	{
+		return InnerSize;
+	}
 
-	UFUNCTION()
-	bool SetAngleDegree(float InAngleDegree);
-	float GetAngleDegree() const { return AngleDegree; }
+	AVALANCHESHAPES_API void SetAngleDegree(float InAngleDegree);
+	float GetAngleDegree() const
+	{
+		return AngleDegree;
+	}
 
-	UFUNCTION()
-	bool SetStartDegree(float InStartDegree);
-	float GetStartDegree() const { return StartDegree; }
+	AVALANCHESHAPES_API void SetStartDegree(float InStartDegree);
+	float GetStartDegree() const
+	{
+		return StartDegree;
+	}
 
 protected:
+	//~ Begin UObject
 #if WITH_EDITOR
 	virtual void PostEditChangeProperty(FPropertyChangedEvent& PropertyChangedEvent) override;
 #endif
+	//~ End UObject
 
 	virtual bool CreateBaseUVs(FAvaShapeMesh& BaseMesh, FAvaShapeMaterialUVParameters& InUseParams);
 	virtual bool GenerateBaseMeshSections(FAvaShapeMesh& BaseMesh);
@@ -92,22 +105,22 @@ protected:
 	virtual bool CreateUVs(FAvaShapeMesh& InMesh, FAvaShapeMaterialUVParameters& InParams) override;
 
 	// represents the number of slices composing the tube
-	UPROPERTY(EditInstanceOnly, BlueprintReadWrite, Category="Shape", meta=(ClampMin="3.0", ClampMax="128.0", DisplayName="Slices", AllowPrivateAccess="true"))
+	UPROPERTY(EditInstanceOnly, BlueprintReadWrite, Setter, Getter, Category="Shape", meta=(ClampMin="3.0", ClampMax="128.0", DisplayName="Slices", AllowPrivateAccess="true"))
 	uint8 NumSlices;
 
 	// represents the precision of each circle composing a slice
-	UPROPERTY(EditInstanceOnly, BlueprintReadWrite, Category="Shape", meta=(ClampMin="4.0", ClampMax="128.0", DisplayName="Sides", AllowPrivateAccess="true"))
+	UPROPERTY(EditInstanceOnly, BlueprintReadWrite, Setter, Getter, Category="Shape", meta=(ClampMin="4.0", ClampMax="128.0", DisplayName="Sides", AllowPrivateAccess="true"))
 	uint8 NumSides;
 
 	// represents the size ratio available inside the torus
-	UPROPERTY(EditInstanceOnly, BlueprintReadWrite, Category="Shape", meta=(ClampMin="0.5", ClampMax="0.99", DisplayName="Inner Size", AllowPrivateAccess="true"))
+	UPROPERTY(EditInstanceOnly, BlueprintReadWrite, Setter, Getter, Category="Shape", meta=(ClampMin="0.5", ClampMax="0.99", DisplayName="Inner Size", AllowPrivateAccess="true"))
 	float InnerSize;
 
 	// represents the tube angle in degree for the torus
-	UPROPERTY(EditInstanceOnly, BlueprintReadWrite, Category="Shape", meta=(ClampMin="0.0", ClampMax="360.0", DisplayName="Angle degree", AllowPrivateAccess="true"))
+	UPROPERTY(EditInstanceOnly, BlueprintReadWrite, Setter, Getter, Category="Shape", meta=(ClampMin="0.0", ClampMax="360.0", DisplayName="Angle degree", AllowPrivateAccess="true"))
 	float AngleDegree;
 
 	// represents the starting angle in degree for the torus
-	UPROPERTY(EditInstanceOnly, BlueprintReadWrite, Category="Shape", meta=(ClampMin="0.0", ClampMax="360.0", DisplayName="Start Degree", AllowPrivateAccess="true"))
+	UPROPERTY(EditInstanceOnly, BlueprintReadWrite, Setter, Getter, Category="Shape", meta=(ClampMin="0.0", ClampMax="360.0", DisplayName="Start Degree", AllowPrivateAccess="true"))
 	float StartDegree;
 };
diff --git a/Engine/Plugins/Experimental/Avalanche/Source/AvalancheShapesEditor/Private/Visualizers/AvaShape2DDynMeshVis.cpp b/Engine/Plugins/Experimental/Avalanche/Source/AvalancheShapesEditor/Private/Visualizers/AvaShape2DDynMeshVis.cpp
--- a/Engine/Plugins/Experimental/Avalanche/Source/AvalancheShapesEditor/Private/Visualizers/AvaShape2DDynMeshVis.cpp
+++ b/Engine/Plugins/Experimental/Avalanche/Source/AvalancheShapesEditor/Private/Visualizers/AvaShape2DDynMeshVis.cpp
@@ -48,7 +48,7 @@ bool FAvaShape2DDynamicMeshVisualizer::VisProxyHandleClick(FEditorViewportClient
 	HComponentVisProxy* InVisProxy, const FViewportClick& InClick)
 {
 	if (InClick.GetKey() != EKeys::LeftMouseButton)
-	{ 
+	{
 		EndEditing();
 		return Super::VisProxyHandleClick(InViewportClient, InVisProxy, InClick);
 	}
@@ -342,7 +342,7 @@ bool FAvaShape2DDynamicMeshVisualizer::HandleInputDeltaInternal(FEditorViewportC
 		{
 			// multiply by 0.01 to allow for more precision in the panning
         	constexpr float UVPanDeltaScale = 1.f;
-			
+
 			// control uniform scale when pressing modifier key
 			const bool bUniformScale = FSlateApplication::Get().GetModifierKeys().IsShiftDown();
 
@@ -364,10 +364,10 @@ bool FAvaShape2DDynamicMeshVisualizer::HandleInputDeltaInternal(FEditorViewportC
 						UVOffset.Y += InAccumulatedTranslation.Z * UVPanDeltaScale / Size2D.Y;
 					}
 
+					DynMesh2D->Modify();
 					if (DynMesh2D->SetMaterialUVOffset(UVSectionIdx, UVOffset))
 					{
 						bHasBeenModified = true;
-						DynMesh2D->Modify();
 						NotifyPropertyChainModified(DynMesh2D, UVOffsetProperty, EPropertyChangeType::Interactive, UVSectionIdx, {UVParamsProperty, MeshDataProperty});
 					}
 				}
@@ -380,10 +380,10 @@ bool FAvaShape2DDynamicMeshVisualizer::HandleInputDeltaInternal(FEditorViewportC
 					float UVRotation = InitialPrimaryUVRotation;
 					UVRotation -= InAccumulatedRotation.Roll;
 
+					DynMesh2D->Modify();
 					if (DynMesh2D->SetMaterialUVRotation(UVSectionIdx, UVRotation))
 					{
 						bHasBeenModified = true;
-						DynMesh2D->Modify();
 						NotifyPropertyChainModified(DynMesh2D, UVRotationProperty, EPropertyChangeType::Interactive, UVSectionIdx, {UVParamsProperty, MeshDataProperty});
 					}
 				}
@@ -394,7 +394,7 @@ bool FAvaShape2DDynamicMeshVisualizer::HandleInputDeltaInternal(FEditorViewportC
 				if (GetViewportWidgetAxisList(InViewportClient) & EAxisList::YZ)
 				{
 					FVector2D UVScale = InitialPrimaryUVScale;
-					
+
 					if (GetViewportWidgetAxisList(InViewportClient) & EAxisList::Y)
 					{
 						UVScale.X -= InAccumulatedScale.Y * UVPanDeltaScale;
@@ -415,10 +415,10 @@ bool FAvaShape2DDynamicMeshVisualizer::HandleInputDeltaInternal(FEditorViewportC
 						}
 					}
 
+					DynMesh2D->Modify();
 					if (DynMesh2D->SetMaterialUVScale(UVSectionIdx, UVScale))
 					{
 						bHasBeenModified = true;
-						DynMesh2D->Modify();
 						NotifyPropertyChainModified(DynMesh2D, UVScaleProperty, EPropertyChangeType::Interactive, UVSectionIdx, {UVParamsProperty, MeshDataProperty});
 					}
 				}
@@ -443,10 +443,10 @@ bool FAvaShape2DDynamicMeshVisualizer::HandleInputDeltaInternal(FEditorViewportC
 						UVAnchor.Y = FMath::Clamp(UVAnchor.Y - InAccumulatedTranslation.Z / Size2D.Y, 0.f, 1.f);
 					}
 
+					DynMesh2D->Modify();
 					if (DynMesh2D->SetMaterialUVAnchor(UVSectionIdx, UVAnchor))
 					{
 						bHasBeenModified = true;
-						DynMesh2D->Modify();
 						NotifyPropertyChainModified(DynMesh2D, UVAnchorProperty, EPropertyChangeType::Interactive, UVSectionIdx, {UVParamsProperty, MeshDataProperty});
 					}
 				}
@@ -455,10 +455,10 @@ bool FAvaShape2DDynamicMeshVisualizer::HandleInputDeltaInternal(FEditorViewportC
 
 		return true;
 	}
-	
+
 	// ALT + drag
 	const bool bExtendBothSide = IsExtendBothSidesEnabled(DynMesh2D);
-	
+
 	if (SizeDragAnchor != INDEX_NONE)
 	{
 		if (GetViewportWidgetMode(InViewportClient) == UE::Widget::WM_Translate)
@@ -468,7 +468,7 @@ bool FAvaShape2DDynamicMeshVisualizer::HandleInputDeltaInternal(FEditorViewportC
 				FVector2D Size2D       = InitialSize;
 				const FTransform Transform   = InitialTransform;
 				FVector Location       = FVector::ZeroVector;
-				
+
 				const FVector SizePosition3D = GetLocationFromAlignment(SizeDragAnchor, FVector(0.f, InitialSize.X, InitialSize.Y));
 				FVector2D SizePosition = FVector2D(SizePosition3D.Y, SizePosition3D.Z);
 				const FVector2D InitialSizePosition = SizePosition;
@@ -478,7 +478,7 @@ bool FAvaShape2DDynamicMeshVisualizer::HandleInputDeltaInternal(FEditorViewportC
 				const EAvaVerticalAlignment Vertical = GetVAlignment(SizeDragAnchor);
 
 				// Update translation
-				
+
 				if (Horizontal != EAvaHorizontalAlignment::Center)
 				{
 					// Y
@@ -487,7 +487,7 @@ bool FAvaShape2DDynamicMeshVisualizer::HandleInputDeltaInternal(FEditorViewportC
 						SizePosition.X += InAccumulatedTranslation.Y;
 					}
 				}
-				
+
 				if (Vertical != EAvaVerticalAlignment::Center)
 				{
 					// Z
@@ -526,7 +526,7 @@ bool FAvaShape2DDynamicMeshVisualizer::HandleInputDeltaInternal(FEditorViewportC
 				if (!bExtendBothSide)
 				{
 					FVector2D SizeChange = (Size2D - InitialSize) / 2.f;
-					
+
 					if (Horizontal != EAvaHorizontalAlignment::Center)
 					{
 						// Y
@@ -547,19 +547,18 @@ bool FAvaShape2DDynamicMeshVisualizer::HandleInputDeltaInternal(FEditorViewportC
 				}
 
 				Size2D = FVector2D::Max(Size2D, UAvaShapeDynamicMeshBase::MinSize2D);
-				if (DynMesh2D->SetSize2D(Size2D))
+				DynMesh2D->Modify();
+				DynMesh2D->SetSize2D(Size2D);
+
+				if (DynMesh2D->GetSize2D() != PrevShapeSize)
 				{
-					if (DynMesh2D->GetSize2D() != PrevShapeSize)
+					PrevShapeSize = Size2D;
+					if (!bExtendBothSide)
 					{
-						PrevShapeSize = Size2D;
-						if (!bExtendBothSide)
-						{
-							Location = Transform.TransformPosition(Location);
-							DynMesh2D->SetMeshRegenWorldLocation(Location, true);
-							bHasBeenModified = true;
-							DynMesh2D->Modify();
-							NotifyPropertiesModified(DynMesh2D, {MeshRegenWorldLocationProperty}, EPropertyChangeType::Interactive);
-						}
+						Location = Transform.TransformPosition(Location);
+						DynMesh2D->SetMeshRegenWorldLocation(Location, true);
+						bHasBeenModified = true;
+						NotifyPropertiesModified(DynMesh2D, {MeshRegenWorldLocationProperty}, EPropertyChangeType::Interactive);
 					}
 				}
 			}
@@ -597,11 +596,13 @@ bool FAvaShape2DDynamicMeshVisualizer::HandleModifiedClick(FEditorViewportClient
 			switch (UVMode)
 			{
 				case EAvaShapeUVMode::Stretch:
+					HitProxyDynamicMesh->Modify();
 					HitProxyDynamicMesh->SetMaterialUVMode(UVHitProxy->SectionIdx,
 						EAvaShapeUVMode::Uniform);
 					break;
 
 				case EAvaShapeUVMode::Uniform:
+					HitProxyDynamicMesh->Modify();
 					HitProxyDynamicMesh->SetMaterialUVMode(UVHitProxy->SectionIdx, EAvaShapeUVMode::Stretch);
 					break;
 
@@ -610,7 +611,6 @@ bool FAvaShape2DDynamicMeshVisualizer::HandleModifiedClick(FEditorViewportClient
 					break;
 			}
 
-			HitProxyDynamicMesh->Modify();
 			NotifyPropertyChainModified(HitProxyDynamicMesh, UVModeProperty, EPropertyChangeType::Interactive, UVHitProxy->SectionIdx, {UVParamsProperty, MeshDataProperty});
 		}
 		// Toggle flip flags
@@ -621,26 +621,26 @@ bool FAvaShape2DDynamicMeshVisualizer::HandleModifiedClick(FEditorViewportClient
 
 			if (!bHorizFlip && !bVerticalFlip)
 			{
-				HitProxyDynamicMesh->SetMaterialHorizontalFlip(UVHitProxy->SectionIdx, true);
 				HitProxyDynamicMesh->Modify();
+				HitProxyDynamicMesh->SetMaterialHorizontalFlip(UVHitProxy->SectionIdx, true);
 				NotifyPropertyChainModified(HitProxyDynamicMesh, UVHorizFlipProperty, EPropertyChangeType::Interactive, UVHitProxy->SectionIdx, {UVParamsProperty, MeshDataProperty});
 			}
 			else if (bHorizFlip && !bVerticalFlip)
 			{
-				HitProxyDynamicMesh->SetMaterialVerticalFlip(UVHitProxy->SectionIdx, true);
 				HitProxyDynamicMesh->Modify();
+				HitProxyDynamicMesh->SetMaterialVerticalFlip(UVHitProxy->SectionIdx, true);
 				NotifyPropertyChainModified(HitProxyDynamicMesh, UVVertFlipProperty, EPropertyChangeType::Interactive, UVHitProxy->SectionIdx, {UVParamsProperty, MeshDataProperty});
 			}
 			else if (bHorizFlip && bVerticalFlip)
 			{
-				HitProxyDynamicMesh->SetMaterialHorizontalFlip(UVHitProxy->SectionIdx, false);
 				HitProxyDynamicMesh->Modify();
+				HitProxyDynamicMesh->SetMaterialHorizontalFlip(UVHitProxy->SectionIdx, false);
 				NotifyPropertyChainModified(HitProxyDynamicMesh, UVHorizFlipProperty, EPropertyChangeType::Interactive, UVHitProxy->SectionIdx, {UVParamsProperty, MeshDataProperty});
 			}
 			else
 			{
-				HitProxyDynamicMesh->SetMaterialVerticalFlip(UVHitProxy->SectionIdx, false);
 				HitProxyDynamicMesh->Modify();
+				HitProxyDynamicMesh->SetMaterialVerticalFlip(UVHitProxy->SectionIdx, false);
 				NotifyPropertyChainModified(HitProxyDynamicMesh, UVVertFlipProperty, EPropertyChangeType::Interactive, UVHitProxy->SectionIdx, {UVParamsProperty, MeshDataProperty});
 			}
 		}
@@ -673,6 +673,7 @@ bool FAvaShape2DDynamicMeshVisualizer::ResetValue(FEditorViewportClient* InViewp
 
 	FScopedTransaction Transaction(NSLOCTEXT("AvaShapeVisualizer", "VisualizerResetValue", "Visualizer Reset Value"));
 	HitProxyDynamicMesh->SetFlags(RF_Transactional);
+	HitProxyDynamicMesh->Modify();
 
 	if (bEditingUVAnchor)
 	{
@@ -704,8 +705,6 @@ bool FAvaShape2DDynamicMeshVisualizer::ResetValue(FEditorViewportClient* InViewp
 		}
 	}
 
-	HitProxyDynamicMesh->Modify();
-
 	return true;
 }
 
diff --git a/Engine/Plugins/Experimental/Avalanche/Source/AvalancheShapesEditor/Private/Visualizers/AvaShape3DDynMeshVis.cpp b/Engine/Plugins/Experimental/Avalanche/Source/AvalancheShapesEditor/Private/Visualizers/AvaShape3DDynMeshVis.cpp
--- a/Engine/Plugins/Experimental/Avalanche/Source/AvalancheShapesEditor/Private/Visualizers/AvaShape3DDynMeshVis.cpp
+++ b/Engine/Plugins/Experimental/Avalanche/Source/AvalancheShapesEditor/Private/Visualizers/AvaShape3DDynMeshVis.cpp
@@ -348,16 +348,15 @@ bool FAvaShape3DDynamicMeshVisualizer::HandleInputDeltaInternal(FEditorViewportC
 				}
 
 				Size3D = FVector::Max(Size3D, UAvaShapeDynamicMeshBase::MinSize3D);
-				if (DynMesh3D->SetSize3D(Size3D))
+				DynMesh3D->Modify();
+				DynMesh3D->SetSize3D(Size3D);
+
+				if (!bExtendBothSide)
 				{
-					if (!bExtendBothSide)
-					{
-						Location = Transform.TransformPosition(Location);
-						DynMesh3D->SetMeshRegenWorldLocation(Location, true);
-						bHasBeenModified = true;
-						DynMesh3D->Modify();
-						NotifyPropertiesModified(DynMesh3D, {MeshRegenWorldLocationProperty}, EPropertyChangeType::Interactive);
-					}
+					Location = Transform.TransformPosition(Location);
+					DynMesh3D->SetMeshRegenWorldLocation(Location, true);
+					bHasBeenModified = true;
+					NotifyPropertiesModified(DynMesh3D, {MeshRegenWorldLocationProperty}, EPropertyChangeType::Interactive);
 				}
 			}
 		}
diff --git a/Engine/Plugins/Experimental/Avalanche/Source/AvalancheShapesEditor/Private/Visualizers/AvaShapeConeDynMeshVis.cpp b/Engine/Plugins/Experimental/Avalanche/Source/AvalancheShapesEditor/Private/Visualizers/AvaShapeConeDynMeshVis.cpp
--- a/Engine/Plugins/Experimental/Avalanche/Source/AvalancheShapesEditor/Private/Visualizers/AvaShapeConeDynMeshVis.cpp
+++ b/Engine/Plugins/Experimental/Avalanche/Source/AvalancheShapesEditor/Private/Visualizers/AvaShapeConeDynMeshVis.cpp
@@ -158,7 +158,7 @@ bool FAvaShapeConeDynamicMeshVisualizer::GetWidgetAxisList(const FEditorViewport
 		OutAxisList = EAxisList::Type::Screen;
 		return true;
 	}
-	
+
 	if (bEditingAngleDegree)
 	{
 		OutAxisList = EAxisList::Type::Z;
@@ -203,8 +203,8 @@ bool FAvaShapeConeDynamicMeshVisualizer::ResetValue(FEditorViewportClient* InVie
 			{
 				FScopedTransaction Transaction(NSLOCTEXT("AvaShapeVisualizer", "VisualizerResetValue", "Visualizer Reset Value"));
 				HitProxyDynamicMesh->SetFlags(RF_Transactional);
-				HitProxyDynamicMesh->SetNumSides(64);
 				HitProxyDynamicMesh->Modify();
+				HitProxyDynamicMesh->SetNumSides(64);
 				NotifyPropertyModified(HitProxyDynamicMesh, NumSidesProperty, EPropertyChangeType::ValueSet);
 			}
 		}
@@ -223,8 +223,8 @@ bool FAvaShapeConeDynamicMeshVisualizer::ResetValue(FEditorViewportClient* InVie
 			{
 				FScopedTransaction Transaction(NSLOCTEXT("AvaShapeVisualizer", "VisualizerResetValue", "Visualizer Reset Value"));
 				HitProxyDynamicMesh->SetFlags(RF_Transactional);
-				HitProxyDynamicMesh->SetTopRadius(0.f);
 				HitProxyDynamicMesh->Modify();
+				HitProxyDynamicMesh->SetTopRadius(0.f);
 				NotifyPropertyModified(HitProxyDynamicMesh, TopRadiusProperty, EPropertyChangeType::ValueSet);
 			}
 		}
@@ -243,8 +243,8 @@ bool FAvaShapeConeDynamicMeshVisualizer::ResetValue(FEditorViewportClient* InVie
 			{
 				FScopedTransaction Transaction(NSLOCTEXT("AvaShapeVisualizer", "VisualizerResetValue", "Visualizer Reset Value"));
 				HitProxyDynamicMesh->SetFlags(RF_Transactional);
-				HitProxyDynamicMesh->SetAngleDegree(360.f);
 				HitProxyDynamicMesh->Modify();
+				HitProxyDynamicMesh->SetAngleDegree(360.f);
 				NotifyPropertyModified(HitProxyDynamicMesh, AngleDegreeProperty, EPropertyChangeType::ValueSet);
 			}
 		}
@@ -286,7 +286,7 @@ void FAvaShapeConeDynamicMeshVisualizer::DrawVisualizationNotEditing(const UActo
 
 	DrawNumSidesButton(DynMesh, InView, InPDI, InOutIconIndex, Inactive);
 	++InOutIconIndex;
-	
+
 	DrawTopRadiusButton(DynMesh, InView, InPDI, InOutIconIndex, Inactive);
 	++InOutIconIndex;
 
@@ -308,7 +308,7 @@ void FAvaShapeConeDynamicMeshVisualizer::DrawVisualizationEditing(const UActorCo
 
 	DrawNumSidesButton(DynMesh, InView, InPDI, InOutIconIndex, bEditingNumSides ? Active : Inactive);
 	++InOutIconIndex;
-	
+
 	DrawTopRadiusButton(DynMesh, InView, InPDI, InOutIconIndex, bEditingTopRadius ? Active : Inactive);
 	++InOutIconIndex;
 
@@ -335,14 +335,12 @@ bool FAvaShapeConeDynamicMeshVisualizer::HandleInputDeltaInternal(FEditorViewpor
 				if (GetViewportWidgetAxisList(InViewportClient) & EAxisList::Y)
 				{
 					int32 NumSides = InitialNumSides;
-					NumSides       = FMath::Clamp(NumSides + static_cast<int32>(InAccumulatedScale.Y), 1, 128);
-
-					if (DynMesh->SetNumSides(NumSides))
-					{
-						bHasBeenModified = true;
-						DynMesh->Modify();
-						NotifyPropertyModified(DynMesh, NumSidesProperty, EPropertyChangeType::Interactive);
-					}
+					NumSides = FMath::Clamp(NumSides + static_cast<int32>(InAccumulatedScale.Y), 1, 128);
+					DynMesh->Modify();
+					DynMesh->SetNumSides(NumSides);
+
+					bHasBeenModified = true;
+					NotifyPropertyModified(DynMesh, NumSidesProperty, EPropertyChangeType::Interactive);
 				}
 			}
 			return true;
@@ -354,14 +352,12 @@ bool FAvaShapeConeDynamicMeshVisualizer::HandleInputDeltaInternal(FEditorViewpor
 				if (GetViewportWidgetAxisList(InViewportClient) & EAxisList::Y)
 				{
 					float TopRadius = InitialTopRadius;
-					TopRadius       = FMath::Clamp(TopRadius + (InAccumulatedTranslation.Y / DynMesh->GetSize3D().X * 2.f), 0.0, 1.f);
-					
-					if (DynMesh->SetTopRadius(TopRadius))
-					{
-						DynMesh->Modify();
-						bHasBeenModified = true;
-						NotifyPropertyModified(DynMesh, TopRadiusProperty, EPropertyChangeType::Interactive);
-					}
+					TopRadius = FMath::Clamp(TopRadius + (InAccumulatedTranslation.Y / DynMesh->GetSize3D().X * 2.f), 0.0, 1.f);
+					DynMesh->Modify();
+					DynMesh->SetTopRadius(TopRadius);
+
+                    bHasBeenModified = true;
+                    NotifyPropertyModified(DynMesh, TopRadiusProperty, EPropertyChangeType::Interactive);
 				}
 			}
 			return true;
@@ -373,14 +369,12 @@ bool FAvaShapeConeDynamicMeshVisualizer::HandleInputDeltaInternal(FEditorViewpor
 				if (GetViewportWidgetAxisList(InViewportClient) & EAxisList::Z)
 				{
 					float AngleDegree = InitialAngleDegree;
-					AngleDegree       = FMath::Clamp(AngleDegree + static_cast<int32>(InAccumulatedRotation.Yaw), 0.f, 360.f);
-					
-					if (DynMesh->SetAngleDegree(AngleDegree))
-					{
-						DynMesh->Modify();
-						bHasBeenModified = true;
-						NotifyPropertyModified(DynMesh, AngleDegreeProperty, EPropertyChangeType::Interactive);
-					}
+					AngleDegree = FMath::Clamp(AngleDegree + static_cast<int32>(InAccumulatedRotation.Yaw), 0.f, 360.f);
+					DynMesh->Modify();
+					DynMesh->SetAngleDegree(AngleDegree);
+
+					bHasBeenModified = true;
+					NotifyPropertyModified(DynMesh, AngleDegreeProperty, EPropertyChangeType::Interactive);
 				}
 			}
 		}
diff --git a/Engine/Plugins/Experimental/Avalanche/Source/AvalancheShapesEditor/Private/Visualizers/AvaShapeCubeDynMeshVis.cpp b/Engine/Plugins/Experimental/Avalanche/Source/AvalancheShapesEditor/Private/Visualizers/AvaShapeCubeDynMeshVis.cpp
--- a/Engine/Plugins/Experimental/Avalanche/Source/AvalancheShapesEditor/Private/Visualizers/AvaShapeCubeDynMeshVis.cpp
+++ b/Engine/Plugins/Experimental/Avalanche/Source/AvalancheShapesEditor/Private/Visualizers/AvaShapeCubeDynMeshVis.cpp
@@ -167,8 +167,8 @@ bool FAvaShapeCubeDynamicMeshVisualizer::ResetValue(FEditorViewportClient* InVie
 			{
 				FScopedTransaction Transaction(NSLOCTEXT("AvaShapeVisualizer", "VisualizerResetValue", "Visualizer Reset Value"));
 				HitProxyDynamicMesh->SetFlags(RF_Transactional);
-				HitProxyDynamicMesh->SetBevelNum(1);
 				HitProxyDynamicMesh->Modify();
+				HitProxyDynamicMesh->SetBevelNum(1);
 				NotifyPropertyModified(HitProxyDynamicMesh, BevelNumProperty, EPropertyChangeType::ValueSet);
 			}
 		}
@@ -186,8 +186,8 @@ bool FAvaShapeCubeDynamicMeshVisualizer::ResetValue(FEditorViewportClient* InVie
 			{
 				FScopedTransaction Transaction(NSLOCTEXT("AvaShapeVisualizer", "VisualizerResetValue", "Visualizer Reset Value"));
 				HitProxyDynamicMesh->SetFlags(RF_Transactional);
-				HitProxyDynamicMesh->SetBevelSizeRatio(0.f);
 				HitProxyDynamicMesh->Modify();
+				HitProxyDynamicMesh->SetBevelSizeRatio(0.f);
 				NotifyPropertyModified(HitProxyDynamicMesh, BevelSizeRatioProperty, EPropertyChangeType::ValueSet);
 			}
 		}
@@ -271,14 +271,12 @@ bool FAvaShapeCubeDynamicMeshVisualizer::HandleInputDeltaInternal(FEditorViewpor
 				if (GetViewportWidgetAxisList(InViewportClient) & EAxisList::Y)
 				{
 					int32 BevelNum = InitialBevelNum;
-					BevelNum       = FMath::Clamp(BevelNum + static_cast<int32>(InAccumulatedScale.Y), 1, 255);
-
-					if (DynMesh->SetBevelNum(BevelNum))
-					{
-						bHasBeenModified = true;
-						DynMesh->Modify();
-						NotifyPropertyModified(DynMesh, BevelNumProperty, EPropertyChangeType::Interactive);
-					}
+					BevelNum = FMath::Clamp(BevelNum + static_cast<int32>(InAccumulatedScale.Y), 1, 255);
+					DynMesh->Modify();
+					DynMesh->SetBevelNum(BevelNum);
+
+					bHasBeenModified = true;
+					NotifyPropertyModified(DynMesh, BevelNumProperty, EPropertyChangeType::Interactive);
 				}
 			}
 
@@ -292,14 +290,12 @@ bool FAvaShapeCubeDynamicMeshVisualizer::HandleInputDeltaInternal(FEditorViewpor
 				if (GetViewportWidgetAxisList(InViewportClient) & EAxisList::Y)
 				{
 					float BevelSize = InitialBevelSizeRatio;
-					BevelSize       = FMath::Clamp(BevelSize + InAccumulatedTranslation.Y, 0.f, DynMesh->GetMaxBevelSize());
-					
-					if (DynMesh->SetBevelSizeRatio(BevelSize))
-					{
-						DynMesh->Modify();
-						bHasBeenModified = true;
-						NotifyPropertyModified(DynMesh, BevelSizeRatioProperty, EPropertyChangeType::Interactive);
-					}
+					BevelSize = FMath::Clamp(BevelSize + InAccumulatedTranslation.Y, 0.f, DynMesh->GetMaxBevelSize());
+					DynMesh->Modify();
+					DynMesh->SetBevelSizeRatio(BevelSize);
+
+					bHasBeenModified = true;
+					NotifyPropertyModified(DynMesh, BevelSizeRatioProperty, EPropertyChangeType::Interactive);
 				}
 			}
 
diff --git a/Engine/Plugins/Experimental/Avalanche/Source/AvalancheShapesEditor/Private/Visualizers/AvaShapeEllipseDynMeshVis.cpp b/Engine/Plugins/Experimental/Avalanche/Source/AvalancheShapesEditor/Private/Visualizers/AvaShapeEllipseDynMeshVis.cpp
--- a/Engine/Plugins/Experimental/Avalanche/Source/AvalancheShapesEditor/Private/Visualizers/AvaShapeEllipseDynMeshVis.cpp
+++ b/Engine/Plugins/Experimental/Avalanche/Source/AvalancheShapesEditor/Private/Visualizers/AvaShapeEllipseDynMeshVis.cpp
@@ -170,14 +170,12 @@ bool FAvaShapeEllipseDynamicMeshVisualizer::HandleInputDeltaInternal(FEditorView
 					if (DynMesh)
 					{
 						int32 NumSides = InitialNumSides;
-						NumSides       = FMath::Clamp(NumSides + static_cast<int32>(InAccumulatedScale.Y), 3, 128);
-
-						if (DynMesh->SetNumSides(NumSides))
-						{
-							bHasBeenModified = true;
-							DynMesh->Modify();
-							NotifyPropertyModified(DynMesh, NumSidesProperty, EPropertyChangeType::Interactive);
-						}
+						NumSides = FMath::Clamp(NumSides + static_cast<int32>(InAccumulatedScale.Y), 3, 128);
+						DynMesh->Modify();
+						DynMesh->SetNumSides(NumSides);
+
+						bHasBeenModified = true;
+						NotifyPropertyModified(DynMesh, NumSidesProperty, EPropertyChangeType::Interactive);
 					}
 				}
 				else
@@ -225,8 +223,8 @@ bool FAvaShapeEllipseDynamicMeshVisualizer::ResetValue(FEditorViewportClient* In
 		{
 			FScopedTransaction Transaction(NSLOCTEXT("AvaShapeVisualizer", "VisualizerResetValue", "Visualizer Reset Value"));
 			HitProxyDynamicMesh->SetFlags(RF_Transactional);
-			HitProxyDynamicMesh->SetNumSides(5);
 			HitProxyDynamicMesh->Modify();
+			HitProxyDynamicMesh->SetNumSides(5);
 			NotifyPropertyModified(HitProxyDynamicMesh, NumSidesProperty, EPropertyChangeType::ValueSet);
 		}
 	}
diff --git a/Engine/Plugins/Experimental/Avalanche/Source/AvalancheShapesEditor/Private/Visualizers/AvaShapeIrregularPolygonDynMeshVis.cpp b/Engine/Plugins/Experimental/Avalanche/Source/AvalancheShapesEditor/Private/Visualizers/AvaShapeIrregularPolygonDynMeshVis.cpp
--- a/Engine/Plugins/Experimental/Avalanche/Source/AvalancheShapesEditor/Private/Visualizers/AvaShapeIrregularPolygonDynMeshVis.cpp
+++ b/Engine/Plugins/Experimental/Avalanche/Source/AvalancheShapesEditor/Private/Visualizers/AvaShapeIrregularPolygonDynMeshVis.cpp
@@ -530,14 +530,12 @@ bool FAvaShapeIrregularPolygonDynamicMeshVisualizer::HandleInputDeltaInternal(FE
 					if (GetViewportWidgetAxisList(InViewportClient) & EAxisList::Y)
 					{
 						float BevelSize = InitialGlobalBevelSize;
-						BevelSize       = FMath::Max(BevelSize + InAccumulatedTranslation.Y / DynMesh->GetSize2D().X * 2.f, 0.f);
+						BevelSize = FMath::Max(BevelSize + InAccumulatedTranslation.Y / DynMesh->GetSize2D().X * 2.f, 0.f);
+						DynMesh->Modify();
+						DynMesh->SetGlobalBevelSize(BevelSize);
 
-						if (DynMesh->SetGlobalBevelSize(BevelSize))
-						{
-							bHasBeenModified = true;
-							DynMesh->Modify();
-							NotifyPropertyModified(DynMesh, GlobalBevelSizeProperty, EPropertyChangeType::Interactive);
-						}
+						bHasBeenModified = true;
+						NotifyPropertyModified(DynMesh, GlobalBevelSizeProperty, EPropertyChangeType::Interactive);
 					}
 				}
 				else if (EditingPointIdx != INDEX_NONE)
@@ -558,10 +556,10 @@ bool FAvaShapeIrregularPolygonDynamicMeshVisualizer::HandleInputDeltaInternal(FE
 
 						SnapLocation2D(InViewportClient, InitialTransform, NewPointLocation);
 
+						DynMesh->Modify();
 						if (DynMesh->SetLocation(EditingPointIdx, NewPointLocation))
 						{
 							bHasBeenModified = true;
-							DynMesh->Modify();
 							NotifyPropertyModified(DynMesh, LocationProperty, EPropertyChangeType::Interactive, PointsProperty);
 						}
 					}
@@ -572,25 +570,24 @@ bool FAvaShapeIrregularPolygonDynamicMeshVisualizer::HandleInputDeltaInternal(FE
 					{
 						const bool bSetGlobalBevelCorner = FSlateApplication::Get().GetModifierKeys().IsAltDown();
 						const float MaxBevelSize = DynMesh->GetMaxBevelSizeForPoint(EditingBevelIdx);
-						float BevelSize          = InitialPointBevelSize;
+						float BevelSize = InitialPointBevelSize;
 
 						BevelSize = FMath::Clamp(BevelSize + InAccumulatedTranslation.Y / (MaxBevelSize * 0.75f), 0.f, 1.f);
-						
+
 						if (bSetGlobalBevelCorner)
 						{
-							if (DynMesh->SetGlobalBevelSize(BevelSize))
-							{
-								bHasBeenModified = true;
-								DynMesh->Modify();
-								NotifyPropertyModified(DynMesh, GlobalBevelSizeProperty, EPropertyChangeType::Interactive);
-							}
+							DynMesh->Modify();
+							DynMesh->SetGlobalBevelSize(BevelSize);
+
+							bHasBeenModified = true;
+							NotifyPropertyModified(DynMesh, GlobalBevelSizeProperty, EPropertyChangeType::Interactive);
 						}
 						else
 						{
+							DynMesh->Modify();
 							if (DynMesh->SetBevelSize(EditingBevelIdx, BevelSize))
 							{
 								bHasBeenModified = true;
-								DynMesh->Modify();
 								NotifyPropertyModified(DynMesh, BevelSizeProperty, EPropertyChangeType::Interactive, PointsProperty);
 							}
 						}
@@ -662,8 +659,8 @@ bool FAvaShapeIrregularPolygonDynamicMeshVisualizer::ResetValue(FEditorViewportC
 		{
 			FScopedTransaction Transaction(LOCTEXT("VisualizerResetValue", "Visualizer Reset Value"));
 			HitProxyDynamicMesh->SetFlags(RF_Transactional);
-			HitProxyDynamicMesh->SetGlobalBevelSize(0.f);
 			HitProxyDynamicMesh->Modify();
+			HitProxyDynamicMesh->SetGlobalBevelSize(0.f);
 			NotifyPropertyModified(HitProxyDynamicMesh, GlobalBevelSizeProperty, EPropertyChangeType::ValueSet);
 		}
 	}
@@ -677,8 +674,8 @@ bool FAvaShapeIrregularPolygonDynamicMeshVisualizer::ResetValue(FEditorViewportC
 		{
 			FScopedTransaction Transaction(LOCTEXT("VisualizerResetValue", "Visualizer Reset Value"));
 			HitProxyDynamicMesh->SetFlags(RF_Transactional);
-			HitProxyDynamicMesh->SetBevelSize(BevelHitProxy->PointIdx, 0.f);
 			HitProxyDynamicMesh->Modify();
+			HitProxyDynamicMesh->SetBevelSize(BevelHitProxy->PointIdx, 0.f);
 			NotifyPropertyModified(HitProxyDynamicMesh, BevelSizeProperty, EPropertyChangeType::ValueSet, PointsProperty);
 		}
 	}
diff --git a/Engine/Plugins/Experimental/Avalanche/Source/AvalancheShapesEditor/Private/Visualizers/AvaShapeNGonDynMeshVis.cpp b/Engine/Plugins/Experimental/Avalanche/Source/AvalancheShapesEditor/Private/Visualizers/AvaShapeNGonDynMeshVis.cpp
--- a/Engine/Plugins/Experimental/Avalanche/Source/AvalancheShapesEditor/Private/Visualizers/AvaShapeNGonDynMeshVis.cpp
+++ b/Engine/Plugins/Experimental/Avalanche/Source/AvalancheShapesEditor/Private/Visualizers/AvaShapeNGonDynMeshVis.cpp
@@ -173,14 +173,12 @@ bool FAvaShapeNGonDynamicMeshVisualizer::HandleInputDeltaInternal(FEditorViewpor
 					if (GetViewportWidgetAxisList(InViewportClient) & EAxisList::Y)
 					{
 						int32 NumSides = InitialNumSides;
-						NumSides       = FMath::Clamp(NumSides + static_cast<int32>(InAccumulatedScale.Y), 3, 128);
-
-						if (DynMesh->SetNumSides(NumSides))
-						{
-							DynMesh->Modify();
-							bHasBeenModified = true;
-							NotifyPropertyModified(DynMesh, NumSidesProperty, EPropertyChangeType::Interactive);
-						}
+						NumSides = FMath::Clamp(NumSides + static_cast<int32>(InAccumulatedScale.Y), 3, 128);
+						DynMesh->Modify();
+						DynMesh->SetNumSides(NumSides);
+
+						bHasBeenModified = true;
+						NotifyPropertyModified(DynMesh, NumSidesProperty, EPropertyChangeType::Interactive);
 					}
 				}
 			}
@@ -229,8 +227,8 @@ bool FAvaShapeNGonDynamicMeshVisualizer::ResetValue(FEditorViewportClient* InVie
 		{
 			FScopedTransaction Transaction(NSLOCTEXT("AvaShapeVisualizer", "VisualizerResetValue", "Visualizer Reset Value"));
 			HitProxyDynamicMesh->SetFlags(RF_Transactional);
-			HitProxyDynamicMesh->SetNumSides(5);
 			HitProxyDynamicMesh->Modify();
+			HitProxyDynamicMesh->SetNumSides(5);
 			NotifyPropertyModified(HitProxyDynamicMesh, NumSidesProperty, EPropertyChangeType::ValueSet);
 		}
 	}
diff --git a/Engine/Plugins/Experimental/Avalanche/Source/AvalancheShapesEditor/Private/Visualizers/AvaShapeRectangleDynMeshVis.cpp b/Engine/Plugins/Experimental/Avalanche/Source/AvalancheShapesEditor/Private/Visualizers/AvaShapeRectangleDynMeshVis.cpp
--- a/Engine/Plugins/Experimental/Avalanche/Source/AvalancheShapesEditor/Private/Visualizers/AvaShapeRectangleDynMeshVis.cpp
+++ b/Engine/Plugins/Experimental/Avalanche/Source/AvalancheShapesEditor/Private/Visualizers/AvaShapeRectangleDynMeshVis.cpp
@@ -558,14 +558,12 @@ bool FAvaShapeRectangleDynamicMeshVisualizer::HandleInputDeltaInternal(FEditorVi
 					if (GetViewportWidgetAxisList(InViewportClient) & EAxisList::Y)
 					{
 						float BevelSize = InitialGlobalBevelSize;
-						BevelSize       = FMath::Max(BevelSize + InAccumulatedTranslation.Y * 2.f, 0.f);
+						BevelSize = FMath::Max(BevelSize + InAccumulatedTranslation.Y * 2.f, 0.f);
+						DynMesh->Modify();
+						DynMesh->SetGlobalBevelSize(BevelSize);
 
-						if (DynMesh->SetGlobalBevelSize(BevelSize))
-						{
-							bHasBeenModified = true;
-							DynMesh->Modify();
-							NotifyPropertyModified(DynMesh, GlobalBevelSizeProperty, EPropertyChangeType::Interactive);
-						}
+						bHasBeenModified = true;
+						NotifyPropertyModified(DynMesh, GlobalBevelSizeProperty, EPropertyChangeType::Interactive);
 					}
 				}
 			}
@@ -603,19 +601,19 @@ bool FAvaShapeRectangleDynamicMeshVisualizer::HandleInputDeltaInternal(FEditorVi
 
 								if (bSetGlobalBevelCorner)
 								{
-									if (DynMesh->SetGlobalBevelSize(BevelSize))
-									{
-										bHasBeenModified = true;
-										DynMesh->Modify();
-										NotifyPropertyModified(DynMesh, GlobalBevelSizeProperty, EPropertyChangeType::Interactive);
-									}
+									DynMesh->Modify();
+									DynMesh->SetGlobalBevelSize(BevelSize);
+
+									bHasBeenModified = true;
+									NotifyPropertyModified(DynMesh, GlobalBevelSizeProperty, EPropertyChangeType::Interactive);
 								}
-								else if (DynMesh->SetTopLeftBevelSize(BevelSize))
+								else
 								{
-									bHasBeenModified = true;
 									DynMesh->Modify();
-									NotifyPropertyModified(DynMesh, BevelSizeProperty, EPropertyChangeType::Interactive,
-										TopLeftCornerSettingsProperty);
+									DynMesh->SetTopLeftBevelSize(BevelSize);
+
+									bHasBeenModified = true;
+									NotifyPropertyModified(DynMesh, BevelSizeProperty, EPropertyChangeType::Interactive, TopLeftCornerSettingsProperty);
 								}
 
 								break;
@@ -626,19 +624,19 @@ bool FAvaShapeRectangleDynamicMeshVisualizer::HandleInputDeltaInternal(FEditorVi
 
 								if (bSetGlobalBevelCorner)
 								{
-									if (DynMesh->SetGlobalBevelSize(BevelSize))
-									{
-										bHasBeenModified = true;
-										DynMesh->Modify();
-										NotifyPropertyModified(DynMesh, GlobalBevelSizeProperty, EPropertyChangeType::Interactive);
-									}
+									DynMesh->Modify();
+									DynMesh->SetGlobalBevelSize(BevelSize);
+
+									bHasBeenModified = true;
+									NotifyPropertyModified(DynMesh, GlobalBevelSizeProperty, EPropertyChangeType::Interactive);
 								}
-								else if (DynMesh->SetTopRightBevelSize(BevelSize))
+								else
 								{
-									bHasBeenModified = true;
 									DynMesh->Modify();
-									NotifyPropertyModified(DynMesh, BevelSizeProperty, EPropertyChangeType::Interactive,
-										TopRightCornerSettingsProperty);
+									DynMesh->SetTopRightBevelSize(BevelSize);
+
+									bHasBeenModified = true;
+									NotifyPropertyModified(DynMesh, BevelSizeProperty, EPropertyChangeType::Interactive, TopRightCornerSettingsProperty);
 								}
 
 								break;
@@ -649,19 +647,19 @@ bool FAvaShapeRectangleDynamicMeshVisualizer::HandleInputDeltaInternal(FEditorVi
 
 								if (bSetGlobalBevelCorner)
 								{
-									if (DynMesh->SetGlobalBevelSize(BevelSize))
-									{
-										bHasBeenModified = true;
-										DynMesh->Modify();
-										NotifyPropertyModified(DynMesh, GlobalBevelSizeProperty, EPropertyChangeType::Interactive);
-									}
+									DynMesh->Modify();
+									DynMesh->SetGlobalBevelSize(BevelSize);
+
+									bHasBeenModified = true;
+									NotifyPropertyModified(DynMesh, GlobalBevelSizeProperty, EPropertyChangeType::Interactive);
 								}
-								else if (DynMesh->SetBottomLeftBevelSize(BevelSize))
+								else
 								{
-									bHasBeenModified = true;
 									DynMesh->Modify();
-									NotifyPropertyModified(DynMesh, BevelSizeProperty, EPropertyChangeType::Interactive,
-										BottomLeftCornerSettingsProperty);
+									DynMesh->SetBottomLeftBevelSize(BevelSize);
+
+									bHasBeenModified = true;
+									NotifyPropertyModified(DynMesh, BevelSizeProperty, EPropertyChangeType::Interactive, BottomLeftCornerSettingsProperty);
 								}
 
 								break;
@@ -672,19 +670,19 @@ bool FAvaShapeRectangleDynamicMeshVisualizer::HandleInputDeltaInternal(FEditorVi
 
 								if (bSetGlobalBevelCorner)
 								{
-									if (DynMesh->SetGlobalBevelSize(BevelSize))
-									{
-										bHasBeenModified = true;
-										DynMesh->Modify();
-										NotifyPropertyModified(DynMesh, GlobalBevelSizeProperty, EPropertyChangeType::Interactive);
-									}
+									DynMesh->Modify();
+									DynMesh->SetGlobalBevelSize(BevelSize);
+
+									bHasBeenModified = true;
+									NotifyPropertyModified(DynMesh, GlobalBevelSizeProperty, EPropertyChangeType::Interactive);
 								}
-								else if (DynMesh->SetBottomRightBevelSize(BevelSize))
+								else
 								{
-									bHasBeenModified = true;
 									DynMesh->Modify();
-									NotifyPropertyModified(DynMesh, BevelSizeProperty, EPropertyChangeType::Interactive,
-										BottomRightCornerSettingsProperty);
+									DynMesh->SetBottomRightBevelSize(BevelSize);
+
+									bHasBeenModified = true;
+									NotifyPropertyModified(DynMesh, BevelSizeProperty, EPropertyChangeType::Interactive, BottomRightCornerSettingsProperty);
 								}
 
 								break;
@@ -725,26 +723,22 @@ bool FAvaShapeRectangleDynamicMeshVisualizer::HandleInputDeltaInternal(FEditorVi
 							case EAvaAnchors::Left:
 								Slant = InitialLeftSlant;
 								Slant = FMath::Clamp(Slant + Delta * UAvaShapeRectangleDynamicMesh::MaxSlantAngle, UAvaShapeRectangleDynamicMesh::MinSlantAngle, UAvaShapeRectangleDynamicMesh::MaxSlantAngle);
+								DynMesh->Modify();
+								DynMesh->SetLeftSlant(Slant);
 
-								if (DynMesh->SetLeftSlant(Slant))
-								{
-									bHasBeenModified = true;
-									DynMesh->Modify();
-									NotifyPropertyModified(DynMesh, LeftSlantProperty, EPropertyChangeType::Interactive);
-								}
+								bHasBeenModified = true;
+								NotifyPropertyModified(DynMesh, LeftSlantProperty, EPropertyChangeType::Interactive);
 
 								break;
 
 							case EAvaAnchors::Right:
 								Slant = InitialRightSlant;
 								Slant = FMath::Clamp(Slant + Delta * UAvaShapeRectangleDynamicMesh::MaxSlantAngle, UAvaShapeRectangleDynamicMesh::MinSlantAngle, UAvaShapeRectangleDynamicMesh::MaxSlantAngle);
+								DynMesh->Modify();
+								DynMesh->SetRightSlant(Slant);
 
-								if (DynMesh->SetRightSlant(Slant))
-								{
-									bHasBeenModified = true;
-									DynMesh->Modify();
-									NotifyPropertyModified(DynMesh, RightSlantProperty, EPropertyChangeType::Interactive);
-								}
+								bHasBeenModified = true;
+								NotifyPropertyModified(DynMesh, RightSlantProperty, EPropertyChangeType::Interactive);
 
 								break;
 							default:
@@ -952,8 +946,8 @@ bool FAvaShapeRectangleDynamicMeshVisualizer::ResetValue(FEditorViewportClient*
 		{
 			FScopedTransaction Transaction(NSLOCTEXT("AvaShapeVisualizer", "VisualizerResetValue", "Visualizer Reset Value"));
 			HitProxyDynamicMesh->SetFlags(RF_Transactional);
-			HitProxyDynamicMesh->SetGlobalBevelSize(0.f);
 			HitProxyDynamicMesh->Modify();
+			HitProxyDynamicMesh->SetGlobalBevelSize(0.f);
 			NotifyPropertyModified(HitProxyDynamicMesh, GlobalBevelSizeProperty, EPropertyChangeType::ValueSet);
 		}
 	}
@@ -966,6 +960,7 @@ bool FAvaShapeRectangleDynamicMeshVisualizer::ResetValue(FEditorViewportClient*
 		{
 			FScopedTransaction Transaction(NSLOCTEXT("AvaShapeVisualizer", "VisualizerResetValue", "Visualizer Reset Value"));
 			HitProxyDynamicMesh->SetFlags(RF_Transactional);
+			HitProxyDynamicMesh->Modify();
 
 			switch (CornerHitProxy->Corner)
 			{
@@ -1001,8 +996,6 @@ bool FAvaShapeRectangleDynamicMeshVisualizer::ResetValue(FEditorViewportClient*
 					// Nothing to do
 					break;
 			}
-
-			HitProxyDynamicMesh->Modify();
 		}
 	}
 	else if (InHitProxy->IsA(HAvaShapeRectangleSlantHitProxy::StaticGetType()))
diff --git a/Engine/Plugins/Experimental/Avalanche/Source/AvalancheShapesEditor/Private/Visualizers/AvaShapeRingDynMeshVis.cpp b/Engine/Plugins/Experimental/Avalanche/Source/AvalancheShapesEditor/Private/Visualizers/AvaShapeRingDynMeshVis.cpp
--- a/Engine/Plugins/Experimental/Avalanche/Source/AvalancheShapesEditor/Private/Visualizers/AvaShapeRingDynMeshVis.cpp
+++ b/Engine/Plugins/Experimental/Avalanche/Source/AvalancheShapesEditor/Private/Visualizers/AvaShapeRingDynMeshVis.cpp
@@ -243,14 +243,12 @@ bool FAvaShapeRingDynamicMeshVisualizer::HandleInputDeltaInternal(FEditorViewpor
 				if (GetViewportWidgetAxisList(InViewportClient) & EAxisList::Y)
 				{
 					int32 NumSides = InitialNumSides;
-					NumSides       = FMath::Clamp(NumSides + static_cast<int32>(InAccumulatedScale.Y), 3, 128);
-
-					if (DynMesh->SetNumSides(NumSides))
-					{
-						bHasBeenModified = true;
-						DynMesh->Modify();
-						NotifyPropertyModified(DynMesh, NumSidesProperty, EPropertyChangeType::Interactive);
-					}
+					NumSides = FMath::Clamp(NumSides + static_cast<int32>(InAccumulatedScale.Y), 3, 128);
+					DynMesh->Modify();
+					DynMesh->SetNumSides(NumSides);
+
+					bHasBeenModified = true;
+					NotifyPropertyModified(DynMesh, NumSidesProperty, EPropertyChangeType::Interactive);
 				}
 			}
 
@@ -264,14 +262,12 @@ bool FAvaShapeRingDynamicMeshVisualizer::HandleInputDeltaInternal(FEditorViewpor
 				if (GetViewportWidgetAxisList(InViewportClient) & EAxisList::YZ)
 				{
 					float InnerSize = InitialInnerSize;
-					InnerSize       = FMath::Clamp(InnerSize - InAccumulatedTranslation.Y / DynMesh->GetSize2D().X * 2.f, 0.01, 0.99);
-
-					if (DynMesh->SetInnerSize(InnerSize))
-					{
-						bHasBeenModified = true;
-						DynMesh->Modify();
-						NotifyPropertyModified(DynMesh, InnerSizeProperty, EPropertyChangeType::Interactive);
-					}
+					InnerSize = FMath::Clamp(InnerSize - InAccumulatedTranslation.Y / DynMesh->GetSize2D().X * 2.f, 0.01, 0.99);
+					DynMesh->Modify();
+					DynMesh->SetInnerSize(InnerSize);
+
+					bHasBeenModified = true;
+					NotifyPropertyModified(DynMesh, InnerSizeProperty, EPropertyChangeType::Interactive);
 				}
 			}
 
@@ -321,8 +317,8 @@ bool FAvaShapeRingDynamicMeshVisualizer::ResetValue(FEditorViewportClient* InVie
 			{
 				FScopedTransaction Transaction(NSLOCTEXT("AvaShapeVisualizer", "VisualizerResetValue", "Visualizer Reset Value"));
 				HitProxyDynamicMesh->SetFlags(RF_Transactional);
-				HitProxyDynamicMesh->SetNumSides(5);
 				HitProxyDynamicMesh->Modify();
+				HitProxyDynamicMesh->SetNumSides(5);
 				NotifyPropertyModified(HitProxyDynamicMesh, NumSidesProperty, EPropertyChangeType::ValueSet);
 			}
 		}
@@ -339,8 +335,8 @@ bool FAvaShapeRingDynamicMeshVisualizer::ResetValue(FEditorViewportClient* InVie
 			{
 				FScopedTransaction Transaction(NSLOCTEXT("AvaShapeVisualizer", "VisualizerResetValue", "Visualizer Reset Value"));
 				HitProxyDynamicMesh->SetFlags(RF_Transactional);
-				HitProxyDynamicMesh->SetInnerSize(0.5);
 				HitProxyDynamicMesh->Modify();
+				HitProxyDynamicMesh->SetInnerSize(0.5);
 				NotifyPropertyModified(HitProxyDynamicMesh, InnerSizeProperty, EPropertyChangeType::ValueSet);
 			}
 		}
diff --git a/Engine/Plugins/Experimental/Avalanche/Source/AvalancheShapesEditor/Private/Visualizers/AvaShapeRoundedPolygonDynMeshVis.cpp b/Engine/Plugins/Experimental/Avalanche/Source/AvalancheShapesEditor/Private/Visualizers/AvaShapeRoundedPolygonDynMeshVis.cpp
--- a/Engine/Plugins/Experimental/Avalanche/Source/AvalancheShapesEditor/Private/Visualizers/AvaShapeRoundedPolygonDynMeshVis.cpp
+++ b/Engine/Plugins/Experimental/Avalanche/Source/AvalancheShapesEditor/Private/Visualizers/AvaShapeRoundedPolygonDynMeshVis.cpp
@@ -203,14 +203,12 @@ bool FAvaShapeRoundedPolygonDynamicMeshVisualizer::HandleInputDeltaInternal(FEdi
 					if (GetViewportWidgetAxisList(InViewportClient) & EAxisList::Y)
 					{
 						float BevelSize = InitialBevelSize;
-						BevelSize       = FMath::Max(BevelSize + InAccumulatedTranslation.Y * 1.f, 0.f);
-
-						if (DynMesh->SetBevelSize(BevelSize))
-						{
-							DynMesh->Modify();
-							bHasBeenModified = true;
-							NotifyPropertyModified(DynMesh, BevelSizeProperty, EPropertyChangeType::Interactive);
-						}
+						BevelSize = FMath::Max(BevelSize + InAccumulatedTranslation.Y * 1.f, 0.f);
+						DynMesh->Modify();
+						DynMesh->SetBevelSize(BevelSize);
+
+						bHasBeenModified = true;
+						NotifyPropertyModified(DynMesh, BevelSizeProperty, EPropertyChangeType::Interactive);
 					}
 				}
 			}
@@ -259,8 +257,8 @@ bool FAvaShapeRoundedPolygonDynamicMeshVisualizer::ResetValue(FEditorViewportCli
 		{
 			FScopedTransaction Transaction(NSLOCTEXT("AvaShapeVisualizer", "VisualizerResetValue", "Visualizer Reset Value"));
 			HitProxyDynamicMesh->SetFlags(RF_Transactional);
-			HitProxyDynamicMesh->SetBevelSize(0.f);
 			HitProxyDynamicMesh->Modify();
+			HitProxyDynamicMesh->SetBevelSize(0.f);
 			NotifyPropertyModified(HitProxyDynamicMesh, BevelSizeProperty, EPropertyChangeType::ValueSet);
 		}
 	}
diff --git a/Engine/Plugins/Experimental/Avalanche/Source/AvalancheShapesEditor/Private/Visualizers/AvaShapeSphereDynMeshVis.cpp b/Engine/Plugins/Experimental/Avalanche/Source/AvalancheShapesEditor/Private/Visualizers/AvaShapeSphereDynMeshVis.cpp
--- a/Engine/Plugins/Experimental/Avalanche/Source/AvalancheShapesEditor/Private/Visualizers/AvaShapeSphereDynMeshVis.cpp
+++ b/Engine/Plugins/Experimental/Avalanche/Source/AvalancheShapesEditor/Private/Visualizers/AvaShapeSphereDynMeshVis.cpp
@@ -219,8 +219,8 @@ bool FAvaShapeSphereDynamicMeshVisualizer::ResetValue(FEditorViewportClient* InV
 			{
 				FScopedTransaction Transaction(NSLOCTEXT("AvaShapeVisualizer", "VisualizerResetValue", "Visualizer Reset Value"));
 				HitProxyDynamicMesh->SetFlags(RF_Transactional);
-				HitProxyDynamicMesh->SetNumSides(32);
 				HitProxyDynamicMesh->Modify();
+				HitProxyDynamicMesh->SetNumSides(32);
 				NotifyPropertyModified(HitProxyDynamicMesh, NumSidesProperty, EPropertyChangeType::ValueSet);
 			}
 		}
@@ -242,8 +242,8 @@ bool FAvaShapeSphereDynamicMeshVisualizer::ResetValue(FEditorViewportClient* InV
 				{
 					FScopedTransaction Transaction(NSLOCTEXT("AvaShapeVisualizer", "VisualizerResetValue", "Visualizer Reset Value"));
 					HitProxyDynamicMesh->SetFlags(RF_Transactional);
-					HitProxyDynamicMesh->SetStartLongitude(0.f);
 					HitProxyDynamicMesh->Modify();
+					HitProxyDynamicMesh->SetStartLongitude(0.f);
 					NotifyPropertyModified(HitProxyDynamicMesh, StartLongitudeProperty, EPropertyChangeType::ValueSet);
 				}
 				// end longitude
@@ -251,8 +251,8 @@ bool FAvaShapeSphereDynamicMeshVisualizer::ResetValue(FEditorViewportClient* InV
 				{
 					FScopedTransaction Transaction(NSLOCTEXT("AvaShapeVisualizer", "VisualizerResetValue", "Visualizer Reset Value"));
 					HitProxyDynamicMesh->SetFlags(RF_Transactional);
-					HitProxyDynamicMesh->SetEndLongitude(180.f);
 					HitProxyDynamicMesh->Modify();
+					HitProxyDynamicMesh->SetEndLongitude(180.f);
 					NotifyPropertyModified(HitProxyDynamicMesh, EndLongitudeProperty, EPropertyChangeType::ValueSet);
 				}
 				// latitude
@@ -260,8 +260,8 @@ bool FAvaShapeSphereDynamicMeshVisualizer::ResetValue(FEditorViewportClient* InV
 				{
 					FScopedTransaction Transaction(NSLOCTEXT("AvaShapeVisualizer", "VisualizerResetValue", "Visualizer Reset Value"));
 					HitProxyDynamicMesh->SetFlags(RF_Transactional);
-					HitProxyDynamicMesh->SetLatitudeDegree(0.f);
 					HitProxyDynamicMesh->Modify();
+					HitProxyDynamicMesh->SetLatitudeDegree(0.f);
 					NotifyPropertyModified(HitProxyDynamicMesh, LatitudeDegreeProperty, EPropertyChangeType::ValueSet);
 				}
 			}
@@ -360,14 +360,12 @@ bool FAvaShapeSphereDynamicMeshVisualizer::HandleInputDeltaInternal(FEditorViewp
 				if (GetViewportWidgetAxisList(InViewportClient) & EAxisList::Y)
 				{
 					int32 NumSides = InitialNumSides;
-					NumSides       = FMath::Clamp(NumSides + static_cast<int32>(InAccumulatedScale.Y), 4, 255);
-
-					if (DynMesh->SetNumSides(NumSides))
-					{
-						bHasBeenModified = true;
-						DynMesh->Modify();
-						NotifyPropertyModified(DynMesh, NumSidesProperty, EPropertyChangeType::Interactive);
-					}
+					NumSides = FMath::Clamp(NumSides + static_cast<int32>(InAccumulatedScale.Y), 4, 255);
+					DynMesh->Modify();
+					DynMesh->SetNumSides(NumSides);
+
+					bHasBeenModified = true;
+					NotifyPropertyModified(DynMesh, NumSidesProperty, EPropertyChangeType::Interactive);
 				}
 			}
 			return true;
@@ -383,13 +381,11 @@ bool FAvaShapeSphereDynamicMeshVisualizer::HandleInputDeltaInternal(FEditorViewp
 					EndLongitude = FMath::Clamp(EndLongitude + ((InAccumulatedTranslation.Z * 2) / DynMesh->GetSize3D().Z), -1.f, 1.f);
 					// get angle from ratio
 					EndLongitude = FMath::GetMappedRangeValueClamped(FVector2D(-90, 90), FVector2D(180, 0), FMath::RadiansToDegrees(FMath::Asin(EndLongitude)));
-					
-					if (DynMesh->SetEndLongitude(EndLongitude))
-					{
-						DynMesh->Modify();
-						bHasBeenModified = true;
-						NotifyPropertyModified(DynMesh, EndLongitudeProperty, EPropertyChangeType::Interactive);
-					}
+					DynMesh->Modify();
+					DynMesh->SetEndLongitude(EndLongitude);
+
+					bHasBeenModified = true;
+					NotifyPropertyModified(DynMesh, EndLongitudeProperty, EPropertyChangeType::Interactive);
 				}
 			}
 			return true;
@@ -405,13 +401,11 @@ bool FAvaShapeSphereDynamicMeshVisualizer::HandleInputDeltaInternal(FEditorViewp
 					StartLongitude = FMath::Clamp(StartLongitude + ((InAccumulatedTranslation.Z * 2) / DynMesh->GetSize3D().Z), -1.f, 1.f);
 					// get angle from ratio
 					StartLongitude = FMath::GetMappedRangeValueClamped(FVector2D(90, -90), FVector2D(0, 180), FMath::RadiansToDegrees(FMath::Asin(StartLongitude)));
-					
-					if (DynMesh->SetStartLongitude(StartLongitude))
-					{
-						DynMesh->Modify();
-						bHasBeenModified = true;
-						NotifyPropertyModified(DynMesh, StartLongitudeProperty, EPropertyChangeType::Interactive);
-					}
+					DynMesh->Modify();
+					DynMesh->SetStartLongitude(StartLongitude);
+
+					bHasBeenModified = true;
+					NotifyPropertyModified(DynMesh, StartLongitudeProperty, EPropertyChangeType::Interactive);
 				}
 			}
 			return true;
@@ -423,14 +417,12 @@ bool FAvaShapeSphereDynamicMeshVisualizer::HandleInputDeltaInternal(FEditorViewp
 				if (GetViewportWidgetAxisList(InViewportClient) & EAxisList::Z)
 				{
 					float LatDegree = InitialLatitudeDegree;
-					LatDegree       = FMath::Clamp(LatDegree + static_cast<int32>(InAccumulatedRotation.Yaw), 0.f, 360.f);
-					
-					if (DynMesh->SetLatitudeDegree(LatDegree))
-					{
-						DynMesh->Modify();
-						bHasBeenModified = true;
-						NotifyPropertyModified(DynMesh, LatitudeDegreeProperty, EPropertyChangeType::Interactive);
-					}
+					LatDegree = FMath::Clamp(LatDegree + static_cast<int32>(InAccumulatedRotation.Yaw), 0.f, 360.f);
+					DynMesh->Modify();
+					DynMesh->SetLatitudeDegree(LatDegree);
+
+					bHasBeenModified = true;
+					NotifyPropertyModified(DynMesh, LatitudeDegreeProperty, EPropertyChangeType::Interactive);
 				}
 			}
 		}
diff --git a/Engine/Plugins/Experimental/Avalanche/Source/AvalancheShapesEditor/Private/Visualizers/AvaShapeStarDynMeshVis.cpp b/Engine/Plugins/Experimental/Avalanche/Source/AvalancheShapesEditor/Private/Visualizers/AvaShapeStarDynMeshVis.cpp
--- a/Engine/Plugins/Experimental/Avalanche/Source/AvalancheShapesEditor/Private/Visualizers/AvaShapeStarDynMeshVis.cpp
+++ b/Engine/Plugins/Experimental/Avalanche/Source/AvalancheShapesEditor/Private/Visualizers/AvaShapeStarDynMeshVis.cpp
@@ -243,14 +243,12 @@ bool FAvaShapeStarDynamicMeshVisualizer::HandleInputDeltaInternal(FEditorViewpor
 				if (GetViewportWidgetAxisList(InViewportClient) & EAxisList::Y)
 				{
 					int32 NumPoints = InitialNumPoints;
-					NumPoints       = FMath::Clamp(NumPoints + static_cast<int32>(InAccumulatedScale.Y), 3, 128);
-
-					if (DynMesh->SetNumPoints(NumPoints))
-					{
-						bHasBeenModified = true;
-						DynMesh->Modify();
-						NotifyPropertyModified(DynMesh, NumPointsProperty, EPropertyChangeType::Interactive);
-					}
+					NumPoints = FMath::Clamp(NumPoints + static_cast<int32>(InAccumulatedScale.Y), 3, 128);
+					DynMesh->Modify();
+					DynMesh->SetNumPoints(NumPoints);
+
+					bHasBeenModified = true;
+					NotifyPropertyModified(DynMesh, NumPointsProperty, EPropertyChangeType::Interactive);
 				}
 			}
 
@@ -264,14 +262,12 @@ bool FAvaShapeStarDynamicMeshVisualizer::HandleInputDeltaInternal(FEditorViewpor
 				if (GetViewportWidgetAxisList(InViewportClient) & EAxisList::YZ)
 				{
 					float InnerSize = InitialInnerSize;
-					InnerSize       = FMath::Clamp(InnerSize - InAccumulatedTranslation.Y / DynMesh->GetSize2D().X * 2.f, 0.01, 0.99);
-
-					if (DynMesh->SetInnerSize(InnerSize))
-					{
-						bHasBeenModified = true;
-						DynMesh->Modify();
-						NotifyPropertyModified(DynMesh, InnerSizeProperty, EPropertyChangeType::Interactive);
-					}
+					InnerSize = FMath::Clamp(InnerSize - InAccumulatedTranslation.Y / DynMesh->GetSize2D().X * 2.f, 0.01, 0.99);
+					DynMesh->Modify();
+					DynMesh->SetInnerSize(InnerSize);
+
+					bHasBeenModified = true;
+					NotifyPropertyModified(DynMesh, InnerSizeProperty, EPropertyChangeType::Interactive);
 				}
 			}
 
@@ -323,8 +319,8 @@ bool FAvaShapeStarDynamicMeshVisualizer::ResetValue(FEditorViewportClient* InVie
 			{
 				FScopedTransaction Transaction(NSLOCTEXT("AvaShapeVisualizer", "VisualizerResetValue", "Visualizer Reset Value"));
 				HitProxyDynamicMesh->SetFlags(RF_Transactional);
-				HitProxyDynamicMesh->SetNumPoints(5);
 				HitProxyDynamicMesh->Modify();
+				HitProxyDynamicMesh->SetNumPoints(5);
 				NotifyPropertyModified(HitProxyDynamicMesh, NumPointsProperty, EPropertyChangeType::ValueSet);
 			}
 		}
@@ -342,8 +338,8 @@ bool FAvaShapeStarDynamicMeshVisualizer::ResetValue(FEditorViewportClient* InVie
 			{
 				FScopedTransaction Transaction(NSLOCTEXT("AvaShapeVisualizer", "VisualizerResetValue", "Visualizer Reset Value"));
 				HitProxyDynamicMesh->SetFlags(RF_Transactional);
-				HitProxyDynamicMesh->SetInnerSize(0.5);
 				HitProxyDynamicMesh->Modify();
+				HitProxyDynamicMesh->SetInnerSize(0.5);
 				NotifyPropertyModified(HitProxyDynamicMesh, InnerSizeProperty, EPropertyChangeType::ValueSet);
 			}
 		}
diff --git a/Engine/Plugins/Experimental/Avalanche/Source/AvalancheShapesEditor/Private/Visualizers/AvaShapeTorusDynMeshVis.cpp b/Engine/Plugins/Experimental/Avalanche/Source/AvalancheShapesEditor/Private/Visualizers/AvaShapeTorusDynMeshVis.cpp
--- a/Engine/Plugins/Experimental/Avalanche/Source/AvalancheShapesEditor/Private/Visualizers/AvaShapeTorusDynMeshVis.cpp
+++ b/Engine/Plugins/Experimental/Avalanche/Source/AvalancheShapesEditor/Private/Visualizers/AvaShapeTorusDynMeshVis.cpp
@@ -214,8 +214,8 @@ bool FAvaShapeTorusDynamicMeshVisualizer::ResetValue(FEditorViewportClient* InVi
 			{
 				FScopedTransaction Transaction(NSLOCTEXT("AvaShapeVisualizer", "VisualizerResetValue", "Visualizer Reset Value"));
 				HitProxyDynamicMesh->SetFlags(RF_Transactional);
-				HitProxyDynamicMesh->SetNumSlices(32);
 				HitProxyDynamicMesh->Modify();
+				HitProxyDynamicMesh->SetNumSlices(32);
 				NotifyPropertyModified(HitProxyDynamicMesh, NumSlicesProperty, EPropertyChangeType::ValueSet);
 			}
 		}
@@ -234,8 +234,8 @@ bool FAvaShapeTorusDynamicMeshVisualizer::ResetValue(FEditorViewportClient* InVi
 			{
 				FScopedTransaction Transaction(NSLOCTEXT("AvaShapeVisualizer", "VisualizerResetValue", "Visualizer Reset Value"));
 				HitProxyDynamicMesh->SetFlags(RF_Transactional);
-				HitProxyDynamicMesh->SetNumSides(32);
 				HitProxyDynamicMesh->Modify();
+				HitProxyDynamicMesh->SetNumSides(32);
 				NotifyPropertyModified(HitProxyDynamicMesh, NumSidesProperty, EPropertyChangeType::ValueSet);
 			}
 		}
@@ -254,8 +254,8 @@ bool FAvaShapeTorusDynamicMeshVisualizer::ResetValue(FEditorViewportClient* InVi
 			{
 				FScopedTransaction Transaction(NSLOCTEXT("AvaShapeVisualizer", "VisualizerResetValue", "Visualizer Reset Value"));
 				HitProxyDynamicMesh->SetFlags(RF_Transactional);
-				HitProxyDynamicMesh->SetInnerSize(0.75f);
 				HitProxyDynamicMesh->Modify();
+				HitProxyDynamicMesh->SetInnerSize(0.75f);
 				NotifyPropertyModified(HitProxyDynamicMesh, InnerSizeProperty, EPropertyChangeType::ValueSet);
 			}
 		}
@@ -274,8 +274,8 @@ bool FAvaShapeTorusDynamicMeshVisualizer::ResetValue(FEditorViewportClient* InVi
 			{
 				FScopedTransaction Transaction(NSLOCTEXT("AvaShapeVisualizer", "VisualizerResetValue", "Visualizer Reset Value"));
 				HitProxyDynamicMesh->SetFlags(RF_Transactional);
-				HitProxyDynamicMesh->SetAngleDegree(360.f);
 				HitProxyDynamicMesh->Modify();
+				HitProxyDynamicMesh->SetAngleDegree(360.f);
 				NotifyPropertyModified(HitProxyDynamicMesh, AngleDegreeProperty, EPropertyChangeType::ValueSet);
 			}
 		}
@@ -373,14 +373,12 @@ bool FAvaShapeTorusDynamicMeshVisualizer::HandleInputDeltaInternal(FEditorViewpo
 				if (GetViewportWidgetAxisList(InViewportClient) & EAxisList::Z)
 				{
 					int32 NumSlices = InitialNumSlices;
-					NumSlices       = FMath::Clamp(NumSlices + static_cast<int32>(InAccumulatedScale.Z), 3, 255);
-
-					if (DynMesh->SetNumSlices(NumSlices))
-					{
-						bHasBeenModified = true;
-						DynMesh->Modify();
-						NotifyPropertyModified(DynMesh, NumSlicesProperty, EPropertyChangeType::Interactive);
-					}
+					NumSlices = FMath::Clamp(NumSlices + static_cast<int32>(InAccumulatedScale.Z), 3, 255);
+					DynMesh->Modify();
+					DynMesh->SetNumSlices(NumSlices);
+
+					bHasBeenModified = true;
+					NotifyPropertyModified(DynMesh, NumSlicesProperty, EPropertyChangeType::Interactive);
 				}
 			}
 			return true;
@@ -392,14 +390,12 @@ bool FAvaShapeTorusDynamicMeshVisualizer::HandleInputDeltaInternal(FEditorViewpo
 				if (GetViewportWidgetAxisList(InViewportClient) & EAxisList::X)
 				{
 					int32 NumSides = InitialNumSides;
-					NumSides       = FMath::Clamp(NumSides + static_cast<int32>(InAccumulatedScale.X), 3, 255);
-
-					if (DynMesh->SetNumSides(NumSides))
-					{
-						bHasBeenModified = true;
-						DynMesh->Modify();
-						NotifyPropertyModified(DynMesh, NumSidesProperty, EPropertyChangeType::Interactive);
-					}
+					NumSides = FMath::Clamp(NumSides + static_cast<int32>(InAccumulatedScale.X), 3, 255);
+					DynMesh->Modify();
+					DynMesh->SetNumSides(NumSides);
+
+					bHasBeenModified = true;
+					NotifyPropertyModified(DynMesh, NumSidesProperty, EPropertyChangeType::Interactive);
 				}
 			}
 			return true;
@@ -411,14 +407,12 @@ bool FAvaShapeTorusDynamicMeshVisualizer::HandleInputDeltaInternal(FEditorViewpo
 				if (GetViewportWidgetAxisList(InViewportClient) & EAxisList::Z)
 				{
 					float InnerSize = InitialInnerSize;
-					InnerSize       = FMath::Clamp(InnerSize + InAccumulatedTranslation.Z / DynMesh->GetSize3D().Z, 0.5, 1.f);
-					
-					if (DynMesh->SetInnerSize(InnerSize))
-					{
-						DynMesh->Modify();
-						bHasBeenModified = true;
-						NotifyPropertyModified(DynMesh, InnerSizeProperty, EPropertyChangeType::Interactive);
-					}
+					InnerSize = FMath::Clamp(InnerSize + InAccumulatedTranslation.Z / DynMesh->GetSize3D().Z, 0.5, 1.f);
+					DynMesh->Modify();
+					DynMesh->SetInnerSize(InnerSize);
+
+					bHasBeenModified = true;
+					NotifyPropertyModified(DynMesh, InnerSizeProperty, EPropertyChangeType::Interactive);
 				}
 			}
 			return true;
@@ -430,14 +424,12 @@ bool FAvaShapeTorusDynamicMeshVisualizer::HandleInputDeltaInternal(FEditorViewpo
 				if (GetViewportWidgetAxisList(InViewportClient) & EAxisList::X)
 				{
 					float AngleDegree = InitialAngleDegree;
-					AngleDegree       = FMath::Clamp(AngleDegree + static_cast<int32>(InAccumulatedRotation.Roll), 0.f, 360.f);
-					
-					if (DynMesh->SetAngleDegree(AngleDegree))
-					{
-						DynMesh->Modify();
-						bHasBeenModified = true;
-						NotifyPropertyModified(DynMesh, AngleDegreeProperty, EPropertyChangeType::Interactive);
-					}
+					AngleDegree = FMath::Clamp(AngleDegree + static_cast<int32>(InAccumulatedRotation.Roll), 0.f, 360.f);
+					DynMesh->Modify();
+					DynMesh->SetAngleDegree(AngleDegree);
+
+					bHasBeenModified = true;
+					NotifyPropertyModified(DynMesh, AngleDegreeProperty, EPropertyChangeType::Interactive);
 				}
 			}
 		}
diff --git a/Engine/Plugins/Experimental/Avalanche/Source/AvalancheShapesEditor/Public/AvaShapeFactory.h b/Engine/Plugins/Experimental/Avalanche/Source/AvalancheShapesEditor/Public/AvaShapeFactory.h
--- a/Engine/Plugins/Experimental/Avalanche/Source/AvalancheShapesEditor/Public/AvaShapeFactory.h
+++ b/Engine/Plugins/Experimental/Avalanche/Source/AvalancheShapesEditor/Public/AvaShapeFactory.h
@@ -9,7 +9,7 @@
 class UAvaShapeDynamicMeshBase;
 
 UCLASS()
-class AVALANCHESHAPESEDITOR_API UAvaShapeFactory : public UActorFactory
+class UAvaShapeFactory : public UActorFactory
 {
 	GENERATED_BODY()
 
diff --git a/Engine/Plugins/Experimental/Avalanche/Source/AvalancheShapesEditor/Public/Tools/AvaShapesEditorShapeAreaToolBase.h b/Engine/Plugins/Experimental/Avalanche/Source/AvalancheShapesEditor/Public/Tools/AvaShapesEditorShapeAreaToolBase.h
--- a/Engine/Plugins/Experimental/Avalanche/Source/AvalancheShapesEditor/Public/Tools/AvaShapesEditorShapeAreaToolBase.h
+++ b/Engine/Plugins/Experimental/Avalanche/Source/AvalancheShapesEditor/Public/Tools/AvaShapesEditorShapeAreaToolBase.h
@@ -9,7 +9,7 @@ class AActor;
 class AAvaShapeActor;
 
 UCLASS(Abstract)
-class AVALANCHESHAPESEDITOR_API UAvaShapesEditorShapeAreaToolBase : public UAvaShapesEditorShapeToolBase
+class UAvaShapesEditorShapeAreaToolBase : public UAvaShapesEditorShapeToolBase
 {
 	GENERATED_BODY()
 
diff --git a/Engine/Plugins/Experimental/Avalanche/Source/AvalancheShapesEditor/Public/Tools/AvaShapesEditorShapeToolBase.h b/Engine/Plugins/Experimental/Avalanche/Source/AvalancheShapesEditor/Public/Tools/AvaShapesEditorShapeToolBase.h
--- a/Engine/Plugins/Experimental/Avalanche/Source/AvalancheShapesEditor/Public/Tools/AvaShapesEditorShapeToolBase.h
+++ b/Engine/Plugins/Experimental/Avalanche/Source/AvalancheShapesEditor/Public/Tools/AvaShapesEditorShapeToolBase.h
@@ -12,7 +12,7 @@ class AAvaShapeActor;
 class UAvaShapeDynamicMeshBase;
 
 UCLASS(Abstract)
-class AVALANCHESHAPESEDITOR_API UAvaShapesEditorShapeToolBase : public UAvaInteractiveToolsActorToolBase
+class UAvaShapesEditorShapeToolBase : public UAvaInteractiveToolsActorToolBase
 {
 	GENERATED_BODY()
 
