[Motion Design] AvalancheMedia Public API clean up.

#rnx

diff --git a/Engine/Plugins/Experimental/Avalanche/Source/AvalancheMedia/Internal/Playable/AvaPlayableRemoteControl.h b/Engine/Plugins/Experimental/Avalanche/Source/AvalancheMedia/Internal/Playable/AvaPlayableRemoteControl.h
new file mode 100644
--- /dev/null
+++ b/Engine/Plugins/Experimental/Avalanche/Source/AvalancheMedia/Internal/Playable/AvaPlayableRemoteControl.h
@@ -0,0 +1,76 @@
+// Copyright Epic Games, Inc. All Rights Reserved.
+
+#pragma once
+
+#include "CoreTypes.h"
+#include "RemoteControlPreset.h"
+#include "AvaPlayableRemoteControl.generated.h"
+
+AVALANCHEMEDIA_API DECLARE_LOG_CATEGORY_EXTERN(LogAvaPlayableRemoteControl, Log, All);
+
+/**
+ * Return code for Remote Control operations.
+ */
+UENUM()
+enum class EAvaPlayableRemoteControlResult : int8
+{
+	/** Operation completed successfully */
+	Completed = 0,
+	/** Operation was not necessary. */
+	UpToDate = 1,	
+	/** Input parameter is invalid */ 
+	InvalidParameter = -1,
+	/** Object reference read access failed. */
+	ReadAccessDenied = -2,
+	/** Object reference write access failed. */
+	WriteAccessDenied = -3,
+	/** Serialization from property failed. */
+	ReadPropertyFailed = -4,
+	/** Serialization to property failed. */
+	WritePropertyFailed = -5
+};
+
+namespace UE::AvaPlayableRemoteControl
+{
+	/** Gets the value (as bytes) of the given property (remote control entity). Returns true if the operation was successful, false otherwise. */
+	AVALANCHEMEDIA_API EAvaPlayableRemoteControlResult GetValueOfEntity(const TSharedPtr<const FRemoteControlEntity>& InRemoteControlEntity, TArray<uint8>& OutValue);
+	AVALANCHEMEDIA_API EAvaPlayableRemoteControlResult GetValueOfEntity(const TSharedPtr<const FRemoteControlEntity>& InRemoteControlEntity, FString& OutValue);
+	
+	/** Set the value (as bytes) of the given property (remote control entity). Returns true if the operation was successful, false otherwise. */
+	AVALANCHEMEDIA_API EAvaPlayableRemoteControlResult SetValueOfEntity(const TSharedPtr<FRemoteControlEntity>& InRemoteControlEntity, const TArrayView<const uint8>& InValue);
+	AVALANCHEMEDIA_API EAvaPlayableRemoteControlResult SetValueOfEntity(const TSharedPtr<FRemoteControlEntity>& InRemoteControlEntity, const FString& InValue);
+
+	/** Gets the value (as bytes) of the given controller. Returns true if the operation was successful, false otherwise. */
+	AVALANCHEMEDIA_API EAvaPlayableRemoteControlResult GetValueOfController(URCVirtualPropertyBase* InController, TArray<uint8>& OutValue);
+	AVALANCHEMEDIA_API EAvaPlayableRemoteControlResult GetValueOfController(URCVirtualPropertyBase* InController, FString& OutValue);
+	
+	/** Set the value (as bytes) of the given controller. Returns true if the operation was successful, false otherwise. */
+	AVALANCHEMEDIA_API EAvaPlayableRemoteControlResult SetValueOfController(URCVirtualPropertyBase* InController, const TArrayView<const uint8>& InValue);
+	AVALANCHEMEDIA_API EAvaPlayableRemoteControlResult SetValueOfController(URCVirtualPropertyBase* InController, const FString& InValue);
+
+	/** Retrieve the entity ids that are controlled by the given controller. */
+	AVALANCHEMEDIA_API bool GetEntitiesControlledByController(const URemoteControlPreset* InRemoteControlPreset, const URCVirtualPropertyBase* InVirtualProperty, TSet<FGuid>& OutEntityIds);
+
+	inline bool Failed(EAvaPlayableRemoteControlResult InResult)
+	{
+		// Negative result codes are failure.
+		return InResult < EAvaPlayableRemoteControlResult::Completed;
+	}
+
+	AVALANCHEMEDIA_API FString EnumToString(EAvaPlayableRemoteControlResult InValue);
+	
+	/**
+	 * Utility class to temporarily push the enabled state of the controller's behaviours.
+	 * This class is intended to be used temporarily on local stack. Lifetime of virtual property is not managed.
+	 */
+	class AVALANCHEMEDIA_API FScopedPushControllerBehavioursEnable
+	{
+	public:
+		FScopedPushControllerBehavioursEnable(URCVirtualPropertyBase* InVirtualProperty, bool bInBehavioursEnabled);
+		~FScopedPushControllerBehavioursEnable();
+		
+	private:
+		URCVirtualPropertyBase* VirtualProperty = nullptr;
+		TArray<bool> PreviousBehavioursEnabled;
+	};
+}
\ No newline at end of file
diff --git a/Engine/Plugins/Experimental/Avalanche/Source/AvalancheMedia/Internal/Playback/AvaPlaybackUtils.h b/Engine/Plugins/Experimental/Avalanche/Source/AvalancheMedia/Internal/Playback/AvaPlaybackUtils.h
new file mode 100644
--- /dev/null
+++ b/Engine/Plugins/Experimental/Avalanche/Source/AvalancheMedia/Internal/Playback/AvaPlaybackUtils.h
@@ -0,0 +1,60 @@
+// Copyright Epic Games, Inc. All Rights Reserved.
+
+#pragma once
+
+#include "Containers/Array.h"
+
+class UPackage;
+struct FAssetData;
+
+class AVALANCHEMEDIA_API FAvaPlaybackUtils
+{
+public:
+	/**
+	 * Reset the loaders of any loaded packages so the client can overwrite the files.
+	*/
+	static void FlushPackageLoading(UPackage* InPackage);
+
+	/**
+	 * Checks if the given packages has been deleted on disk.
+	 * @remark With Editor only.
+	 * @return true if the package was deleted, false if the package still exists on disk.
+	 */
+	static bool IsPackageDeleted(const UPackage* InExistingPackage);
+
+	/**
+	 *	Purge all the objects in memory owned by the given packages.
+	 */
+	static void PurgePackages(const TArray<UPackage*>& InExistingPackages);
+
+	/**
+	 * Reloads the given package.
+	 * @remark With Editor only.
+	 * @return true if the package was reloaded.
+	 */
+	static bool ReloadPackages(const TArray<UPackage*>& InExistingPackages);
+
+	/**
+	 * @brief Determines if the asset is a map by checking the file extension.
+	 * @remark The file must exist on disk.
+	 * @param InPackageName Package name.
+	 * @return true if the file on disk is a .umap file.
+	 */
+	static bool IsMapAsset(const FString& InPackageName);
+
+	/**
+	 * @brief Determines if the asset is a playable (can be used as template) asset, using the asset class.
+	 * @param InAssetData Asset Data
+	 * @return true if the asset is an ava playable.
+	 */
+	static bool IsPlayableAsset(const FAssetData& InAssetData);
+};
+
+namespace UE::AvaPlayback::Utils
+{
+	/**
+	 *	Returns a compactly formatted time stamp information for the current frame.
+	 *	This is to used for logging and tracing.
+	 */
+	AVALANCHEMEDIA_API FString GetBriefFrameInfo();
+}
\ No newline at end of file
diff --git a/Engine/Plugins/Experimental/Avalanche/Source/AvalancheMedia/Internal/Playback/Nodes/AvaPlaybackNodeLevelPlayer.h b/Engine/Plugins/Experimental/Avalanche/Source/AvalancheMedia/Internal/Playback/Nodes/AvaPlaybackNodeLevelPlayer.h
new file mode 100644
--- /dev/null
+++ b/Engine/Plugins/Experimental/Avalanche/Source/AvalancheMedia/Internal/Playback/Nodes/AvaPlaybackNodeLevelPlayer.h
@@ -0,0 +1,36 @@
+// Copyright Epic Games, Inc. All Rights Reserved.
+
+#pragma once
+
+#include "Playback/Nodes/AvaPlaybackNodePlayer.h"
+#include "UObject/NoExportTypes.h"
+#include "UObject/SoftObjectPtr.h"
+#include "AvaPlaybackNodeLevelPlayer.generated.h"
+
+class UWorld;
+struct FAvaSoftAssetPtr;
+
+UCLASS(MinimalAPI)
+class UAvaPlaybackNodeLevelPlayer : public UAvaPlaybackNodePlayer
+{
+	GENERATED_BODY()
+
+public:
+	UAvaPlaybackNodeLevelPlayer();
+	
+	virtual void RefreshNode(bool bDryRunGraph) override;
+	virtual void PostLoad() override;
+	
+	TSoftObjectPtr<UWorld> GetAsset() const { return LevelAsset; }
+	void SetAsset(const TSoftObjectPtr<UWorld>& InAsset);
+	
+	void UpdateDisplayNameText();
+
+	virtual const FSoftObjectPath& GetAssetPath() const override { return LevelAsset.ToSoftObjectPath();}
+
+	virtual FAvaSoftAssetPtr GetAssetPtr() const override;
+
+protected:
+	UPROPERTY(EditAnywhere, Category = "Motion Design")
+	TSoftObjectPtr<UWorld> LevelAsset;
+};
diff --git a/Engine/Plugins/Experimental/Avalanche/Source/AvalancheMedia/Internal/Rundown/AvaRundownManagedInstance.h b/Engine/Plugins/Experimental/Avalanche/Source/AvalancheMedia/Internal/Rundown/AvaRundownManagedInstance.h
new file mode 100644
--- /dev/null
+++ b/Engine/Plugins/Experimental/Avalanche/Source/AvalancheMedia/Internal/Rundown/AvaRundownManagedInstance.h
@@ -0,0 +1,55 @@
+// Copyright Epic Games, Inc. All Rights Reserved.
+
+#pragma once
+
+#include "AvaType.h"
+#include "Playable/AvaPlayableRemoteControlValues.h"
+#include "UObject/GCObject.h"
+#include "UObject/SoftObjectPath.h"
+
+class FAvaRundownManagedInstanceCache;
+class IAvaSceneInterface;
+class URemoteControlPreset;
+
+class AVALANCHEMEDIA_API FAvaRundownManagedInstance : public IAvaTypeCastable, public FGCObject, public TSharedFromThis<FAvaRundownManagedInstance>
+{
+public:
+	UE_AVA_INHERITS(FAvaRundownManagedInstance, IAvaTypeCastable);
+	
+	FAvaRundownManagedInstance(FAvaRundownManagedInstanceCache* InParentCache, const FSoftObjectPath& InAssetPath);
+	virtual ~FAvaRundownManagedInstance() override;
+
+	/** Returns true if the instanced asset is valid. */
+	virtual bool IsValid() const { return false; }
+	
+	/** Get the Managed Motion Design Asset's (top level) Remote Control Preset. */
+	virtual URemoteControlPreset* GetRemoteControlPreset() const { return nullptr; }
+	
+	/** Access a copy of the Remote Control values from the source (a.k.a. template). */
+	const FAvaPlayableRemoteControlValues& GetDefaultRemoteControlValues() const { return DefaultRemoteControlValues; }
+
+	/** Get the Ava scene interface from the asset. */
+	virtual IAvaSceneInterface* GetSceneInterface() const { return nullptr; }
+
+	const FSoftObjectPath& GetSourceAssetPath() const { return SourceAssetPath; }
+	
+protected:
+	void InvalidateFromCache();
+
+	void RegisterSourceRemoteControlPresetDelegates(URemoteControlPreset* InSourceRemoteControlPreset);
+	void UnregisterSourceRemoteControlPresetDelegates(URemoteControlPreset* InSourceRemoteControlPreset) const;
+	
+	void OnSourceRemoteControlEntitiesExposed(URemoteControlPreset* InSourcePreset, const FGuid& InEntityId);
+	void OnSourceRemoteControlEntitiesUnexposed(URemoteControlPreset* InSourcePreset, const FGuid& InEntityId);
+	void OnSourceRemoteControlEntitiesUpdated(URemoteControlPreset* InSourcePreset, const TSet<FGuid>& InModifiedEntities);
+	void OnSourceRemoteControlExposedPropertiesModified(URemoteControlPreset* InSourcePreset, const TSet<FGuid>& InModifiedProperties);
+	void OnSourceRemoteControlControllerAdded(URemoteControlPreset* InSourcePreset, const FName NewControllerName, const FGuid& InControllerId);
+	void OnSourceRemoteControlControllerRemoved(URemoteControlPreset* InSourcePreset, const FGuid& InControllerId);
+	void OnSourceRemoteControlControllerRenamed(URemoteControlPreset* InSourcePreset, const FName InOldLabel, const FName InNewLabel);
+	void OnSourceRemoteControlControllerModified(URemoteControlPreset* InSourcePreset, const TSet<FGuid>& InModifiedControllerIds);
+
+protected:
+	FAvaRundownManagedInstanceCache* ParentCache;
+	FSoftObjectPath SourceAssetPath;
+	FAvaPlayableRemoteControlValues DefaultRemoteControlValues;
+};
\ No newline at end of file
diff --git a/Engine/Plugins/Experimental/Avalanche/Source/AvalancheMedia/Internal/Rundown/AvaRundownManagedInstanceCache.h b/Engine/Plugins/Experimental/Avalanche/Source/AvalancheMedia/Internal/Rundown/AvaRundownManagedInstanceCache.h
new file mode 100644
--- /dev/null
+++ b/Engine/Plugins/Experimental/Avalanche/Source/AvalancheMedia/Internal/Rundown/AvaRundownManagedInstanceCache.h
@@ -0,0 +1,86 @@
+// Copyright Epic Games, Inc. All Rights Reserved.
+
+#pragma once
+
+#include "Rundown/AvaRundownManagedInstance.h"
+
+class IAvaMediaSyncProvider;
+class UPackage;
+
+/**
+ * Cache for the managed Motion Design instances.
+ *
+ * This is LRU cache. 
+ */
+class AVALANCHEMEDIA_API FAvaRundownManagedInstanceCache
+{
+public:
+	FAvaRundownManagedInstanceCache();
+	virtual ~FAvaRundownManagedInstanceCache();
+
+	/**
+	 *	Retrieves elements from the cache or load a new one if not present.
+	 *	If the cache size is exceeded, the oldest elements are going to be flushed.
+	 */
+	TSharedPtr<FAvaRundownManagedInstance> GetOrLoadInstance(const FSoftObjectPath& InAssetPath);
+
+	/**
+	 * Invalidates a cached entry without deleting it immediately.
+	 * It will be deleted and reloaded on the next access query.
+	 */
+	void InvalidateNoDelete(const FSoftObjectPath& InAssetPath);
+
+	/**
+	 * Invalidates a cached entry. The entry may be deleted as a result.
+	 */
+	void Invalidate(const FSoftObjectPath& InAssetPath);
+
+	/**
+	 *	Delegate called when an entry is invalidated. Reason for an entry to be invalidated is if the source
+	 *	has been modified.
+	 */
+	DECLARE_MULTICAST_DELEGATE_OneParam(FOnEntryInvalidated, const FSoftObjectPath& /*InAssetPath*/);
+	FOnEntryInvalidated OnEntryInvalidated;
+	
+	/**
+	 * Get the maximum size of the cache beyond witch it will start flushing elements.
+	 */
+	int32 GetMaximumCacheSize() const;
+	
+	/**
+	 *	Flush specified entry from the cache.
+	 */
+	void Flush(const FSoftObjectPath& InAssetPath);
+
+	/**
+	 * Flush all unused entries from the cache.
+	 */
+	void Flush();
+
+	/**
+	 * Trim the cache elements that exceed the cache capacity according to LRU replacement policy.
+	 */
+	void TrimCache();
+
+	/**
+	 *	Perform any pending actions, such has removing invalidated entries. This may lead to
+	 *	objects being deleted.
+	 */
+	void FinishPendingActions();
+
+private:
+	void RemovePendingInvalidatedPaths();
+	void RemoveEntry(const FSoftObjectPath& InAssetPath);
+	void RemoveEntries(TFunctionRef<bool(const FSoftObjectPath&, const TSharedPtr<FAvaRundownManagedInstance>&)> InRemovePredicate, bool bInNotify);
+
+	void OnPackageSaved(const FString& InPackageFileName, UPackage* InPackage, FObjectPostSaveContext InObjectSaveContext);
+	void OnAvaSyncPackageModified(IAvaMediaSyncProvider* InAvaMediaSyncProvider, const FName& InPackageName);
+	void OnAssetRemoved(const FAssetData& InAssetData);
+	void OnPackageModified(const FName& InPackageName);
+	void OnSettingChanged(UObject* , struct FPropertyChangedEvent&);
+
+private:
+	TMap<FSoftObjectPath, TSharedPtr<FAvaRundownManagedInstance>> Instances;
+	TSet<FSoftObjectPath> PendingInvalidatedPaths;
+	TArray<FSoftObjectPath> OrderQueue;
+};
\ No newline at end of file
diff --git a/Engine/Plugins/Experimental/Avalanche/Source/AvalancheMedia/Private/AvaMediaModule.cpp b/Engine/Plugins/Experimental/Avalanche/Source/AvalancheMedia/Private/AvaMediaModule.cpp
--- a/Engine/Plugins/Experimental/Avalanche/Source/AvalancheMedia/Private/AvaMediaModule.cpp
+++ b/Engine/Plugins/Experimental/Avalanche/Source/AvalancheMedia/Private/AvaMediaModule.cpp
@@ -10,6 +10,7 @@
 #include "Misc/CommandLine.h"
 #include "Misc/CoreDelegates.h"
 #include "Playback/AvaPlaybackClientDelegates.h"
+#include "Playback/AvaPlaybackClientDummy.h"
 #include "ShaderCore.h"
 
 #if WITH_EDITOR
@@ -170,88 +171,35 @@ void FAvaMediaModule::ShutdownModule()
 	AvaMediaSync.Reset();
 }
 
-void FAvaMediaModule::StartMediaPlaybackClient()
+void FAvaMediaModule::StartPlaybackClient()
 {
 	StartPlaybackClientCommand({});
 }
 
-void FAvaMediaModule::StopMediaPlaybackClient()
+void FAvaMediaModule::StopPlaybackClient()
 {
 	StopPlaybackClientCommand({});
 }
 
-void FAvaMediaModule::StartMediaPlaybackServer(const FString& InPlaybackServerName)
+void FAvaMediaModule::StartPlaybackServer(const FString& InPlaybackServerName)
 {
 	StartPlaybackServerCommand({InPlaybackServerName});
 }
 
-void FAvaMediaModule::StopMediaPlaybackServer()
+void FAvaMediaModule::StopPlaybackServer()
 {
 	StopPlaybackServerCommand({});
 }
 
-class FAvaPlaybackClientDummy : public IAvaPlaybackClient
-{
-	static const FString EmptyString;
-	static const TArray<FString> EmptyStringArray;
-public:
-	//IAvaPlaybackClient
-	virtual int32 GetNumConnectedServers() const override { return 0; }
-	virtual TArray<FString> GetServerNames() const override { return TArray<FString>();}
-	virtual FMessageAddress GetServerAddress(const FString& InServerName) const override
-	{
-		FMessageAddress InvalidAddress;
-		InvalidAddress.Invalidate();
-		return InvalidAddress;
-	}
-	virtual bool HasServerUserData(const FString& InServerName, const FString& InKey) const override { return false;}
-	virtual const FString& GetServerUserData(const FString& InServerName, const FString& InKey) const override { return EmptyString; }
-	virtual bool HasUserData(const FString& InKey) const override { return false;}
-	virtual const FString& GetUserData(const FString& InKey) const override { return EmptyString; }
-	virtual void SetUserData(const FString& InKey, const FString& InData) override {}
-	virtual void RemoveUserData(const FString& InKey) override {}
-	virtual void BroadcastStatCommand(const FString& InCommand, bool bInBroadcastLocalState) override {}
-	virtual void RequestPlaybackAssetStatus(const FSoftObjectPath& InAssetPath, const FString& InChannelOrServerName, bool bInForceRefresh) override {}
-	virtual void RequestPlayback(const FGuid& InInstanceId, const FSoftObjectPath& InAssetPath, const FString& InChannelName, EAvaPlaybackAction InAction, const FString& InArguments) override {}
-	virtual void RequestAnimPlayback(const FGuid& InInstanceId, const FSoftObjectPath& InAssetPath, const FString& InChannelName, const FAvaPlaybackAnimPlaySettings& InAnimSettings) override {};
-	virtual void RequestAnimAction(const FGuid& InInstanceId, const FSoftObjectPath& InAssetPath, const FString& InChannelName, const FString& InAnimationName, EAvaPlaybackAnimAction InAction) override {};
-	virtual void RequestRemoteControlUpdate(const FGuid& InInstanceId, const FSoftObjectPath& InAssetPath, const FString& InChannelName, const FAvaPlayableRemoteControlValues& InRemoteControlValues) override {};
-	virtual void RequestPlayableTransitionStart(const FGuid& InTransitionId, TArray<FGuid>&& InEnterInstanceIds, TArray<FGuid>&& InPlayingInstanceIds, TArray<FGuid>&& InExitInstanceIds, TArray<FAvaPlayableRemoteControlValues>&& InEnterValues, const FName& InChannelName, EAvaPlayableTransitionFlags InTransitionFlags) override {}
-	virtual void RequestPlayableTransitionStop(const FGuid& InTransitionId, const FName& InChannelName) override {}
-	virtual void RequestBroadcast(const FString& InProfile, const FName& InChannel, const TArray<UMediaOutput*>& InRemoteMediaOutputs, EAvaBroadcastAction InAction) override {}
-	virtual bool IsMediaOutputRemoteFallback(const UMediaOutput* InMediaOutput) override { return false;}
-	virtual EAvaBroadcastIssueSeverity GetMediaOutputIssueSeverity(const FString& InServerName, const FString& InChannelName, const FGuid& InOutputGuid) const override { return EAvaBroadcastIssueSeverity::None;}
-	virtual const TArray<FString>& GetMediaOutputIssueMessages(const FString& InServerName, const FString& InChannelName, const FGuid& InOutputGuid) const override { return EmptyStringArray; }
-	virtual EAvaBroadcastOutputState GetMediaOutputState(const FString& InServerName, const FString& InChannelName, const FGuid& InOutputGuid) const override { return EAvaBroadcastOutputState::Error; }
-	virtual bool HasAnyServerOnlineForChannel(const FName& InChannelName) const override { return false;}
-   	virtual TOptional<EAvaPlaybackStatus> GetRemotePlaybackStatus(const FGuid& InInstanceId, const FSoftObjectPath& InAssetPath, const FString& InChannelName, const FString& InServerName) const override
-	{
-		return EAvaPlaybackStatus::Unknown;
-	}
-	virtual const FString* GetRemotePlaybackUserData(const FGuid& InInstanceId, const FSoftObjectPath& InAssetPath, const FString& InChannelName, const FString& InServerName) const override
-	{
-		return nullptr;
-	}
-	virtual TOptional<EAvaPlaybackAssetStatus> GetRemotePlaybackAssetStatus(const FSoftObjectPath& InAssetPath, const FString& InServerName) const override
-	{
-		return EAvaPlaybackAssetStatus::Unknown;
-	}
-	//~IAvaPlaybackClient
-};
-const FString FAvaPlaybackClientDummy::EmptyString;
-const TArray<FString> FAvaPlaybackClientDummy::EmptyStringArray;
-
-IAvaPlaybackClient& FAvaMediaModule::GetMediaPlaybackClient()
+IAvaPlaybackClient& FAvaMediaModule::GetPlaybackClient()
 {
 	if (AvaPlaybackClient.IsValid())
 	{
 		return *AvaPlaybackClient;
 	}
-	else
-	{
-		static FAvaPlaybackClientDummy DummyClient;
-		return DummyClient;
-	}
+	
+	static FAvaPlaybackClientDummy DummyClient;
+	return DummyClient;
 }
 
 const IMediaIOCoreDeviceProvider* FAvaMediaModule::GetDeviceProvider(FName InProviderName, const FMediaIOOutputConfiguration* InMediaIOOutputConfiguration) const
@@ -315,7 +263,7 @@ bool FAvaMediaModule::IsLocalDevice(const FName& InDeviceProviderName, const FNa
 void FAvaMediaModule::LaunchGameModeLocalPlaybackServer()
 {
 	// In order to talk to the server, we need the client.
-	if (!IsMediaPlaybackClientStarted())
+	if (!IsPlaybackClientStarted())
 	{
 		// For now, we issue an error.
 		UE_LOG(LogAvaMedia, Error, TEXT("Playback Client must be started prior to starting a local playback server."));
@@ -642,7 +590,7 @@ void FAvaMediaModule::OnAvaPlaybackClientConnectionEvent(IAvaPlaybackClient& InP
 	switch (InArgs.Event)
 	{
 	case EConnectionEvent::ServerConnected:
-		if (!IsGameModeLocalPlaybackServerLaunched() && IsMediaPlaybackClientStarted())
+		if (!IsGameModeLocalPlaybackServerLaunched() && IsPlaybackClientStarted())
 		{
 			LocalPlaybackServerProcess = FAvaPlaybackServerProcess::Find(*AvaPlaybackClient);
 		}
diff --git a/Engine/Plugins/Experimental/Avalanche/Source/AvalancheMedia/Private/AvaMediaModule.h b/Engine/Plugins/Experimental/Avalanche/Source/AvalancheMedia/Private/AvaMediaModule.h
--- a/Engine/Plugins/Experimental/Avalanche/Source/AvalancheMedia/Private/AvaMediaModule.h
+++ b/Engine/Plugins/Experimental/Avalanche/Source/AvalancheMedia/Private/AvaMediaModule.h
@@ -23,15 +23,16 @@ public:
 	//~ Begin IAvaMediaModule
 	virtual void StartupModule() override;
 	virtual void ShutdownModule() override;
-	virtual bool IsMediaPlaybackClientStarted() const override { return AvaPlaybackClient.IsValid();}
-	virtual void StartMediaPlaybackClient() override;
-	virtual void StopMediaPlaybackClient() override;
-	virtual bool IsMediaPlaybackServerStarted() const override { return AvaPlaybackServer.IsValid();}
-	virtual void StartMediaPlaybackServer(const FString& InPlaybackServerName) override;
-	virtual void StopMediaPlaybackServer() override;
+	virtual bool IsPlaybackClientStarted() const override { return AvaPlaybackClient.IsValid();}
+	virtual void StartPlaybackClient() override;
+	virtual void StopPlaybackClient() override;
+	virtual bool IsPlaybackServerStarted() const override { return AvaPlaybackServer.IsValid();}
+	virtual void StartPlaybackServer(const FString& InPlaybackServerName) override;
+	virtual void StopPlaybackServer() override;
 	
-	virtual IAvaPlaybackClient& GetMediaPlaybackClient() override;
-	virtual TSharedPtr<FAvaPlaybackServer> GetMediaPlaybackServer() const override { return AvaPlaybackServer; }
+	virtual IAvaPlaybackClient& GetPlaybackClient() override;
+	virtual TSharedPtr<FAvaPlaybackServer> GetPlaybackServerInternal() const override { return AvaPlaybackServer; }
+	virtual IAvaPlaybackServer* GetPlaybackServer() const override { return AvaPlaybackServer.Get(); }
 	virtual const IMediaIOCoreDeviceProvider* GetDeviceProvider(FName InProviderName, const FMediaIOOutputConfiguration* InMediaIOOutputConfiguration) const override;
 	virtual TArray<const IMediaIOCoreDeviceProvider*> GetDeviceProvidersForServer(const FString& InServerName) const override;
 	virtual FString GetServerNameForDevice(const FName& InDeviceProviderName, const FName& InDeviceName) const override;
diff --git a/Engine/Plugins/Experimental/Avalanche/Source/AvalancheMedia/Private/Broadcast/AvaBroadcast.cpp b/Engine/Plugins/Experimental/Avalanche/Source/AvalancheMedia/Private/Broadcast/AvaBroadcast.cpp
--- a/Engine/Plugins/Experimental/Avalanche/Source/AvalancheMedia/Private/Broadcast/AvaBroadcast.cpp
+++ b/Engine/Plugins/Experimental/Avalanche/Source/AvalancheMedia/Private/Broadcast/AvaBroadcast.cpp
@@ -40,7 +40,7 @@ namespace UE::AvaBroadcast::Private
 		{
 			// When launching the server from the same project location, we want to avoid loading the same
 			// broadcast configuration as the client. The server needs a clean configuration.
-			const bool bIsServerRunning = (IAvaMediaModule::IsModuleLoaded() && IAvaMediaModule::Get().IsMediaPlaybackServerStarted());	
+			const bool bIsServerRunning = (IAvaMediaModule::IsModuleLoaded() && IAvaMediaModule::Get().IsPlaybackServerStarted());	
 			BroadcastConfigName = bIsServerRunning ? TEXT("MotionDesignServerBroadcastConfig.xml") : TEXT("MotionDesignBroadcastConfig.xml");
 		}
 		
diff --git a/Engine/Plugins/Experimental/Avalanche/Source/AvalancheMedia/Private/Broadcast/AvaBroadcastComponent.h b/Engine/Plugins/Experimental/Avalanche/Source/AvalancheMedia/Private/Broadcast/AvaBroadcastComponent.h
new file mode 100644
--- /dev/null
+++ b/Engine/Plugins/Experimental/Avalanche/Source/AvalancheMedia/Private/Broadcast/AvaBroadcastComponent.h
@@ -0,0 +1,46 @@
+// Copyright Epic Games, Inc. All Rights Reserved.
+
+#pragma once
+
+#include "Components/ActorComponent.h"
+#include "AvaBroadcastComponent.generated.h"
+
+/**
+ * Add this actor component to blueprint actor to expose the API to control the broadcasting
+ * of channels in game.
+ */
+UCLASS(BlueprintType, Blueprintable, ClassGroup = "Motion Design Broadcast",
+	meta = (Displayname = "Motion Design Broadcast Component", BlueprintSpawnableComponent))
+class UAvaBroadcastComponent : public UActorComponent
+{
+	GENERATED_BODY()
+
+public:
+	UAvaBroadcastComponent(const FObjectInitializer& ObjectInitializer);
+
+	/**
+	 *	Automatically stop broadcast when the world this component is part of is torn down.
+	 */
+	UPROPERTY(EditDefaultsOnly, Category = "Motion Design Media")
+	bool bStopBroadcastOnTearDown = true;
+
+	/**
+	 *	Start broadcasting all channels.
+	 *	Returns true on partial success, i.e. will be true even if some channels didn't start.
+	 *	Outputs error messages related to outputs that caused problems.
+	 */
+	UFUNCTION(BlueprintCallable, Category = "Motion Design Media")
+	bool StartBroadcasting(FString& OutErrorMessage);
+
+	/**
+	*	Stop broadcasting all channels.
+	*/
+	UFUNCTION(BlueprintCallable, Category = "Motion Design Media")
+	bool StopBroadcasting();
+	
+	virtual void InitializeComponent() override;
+	virtual void UninitializeComponent() override;
+
+protected:
+	void OnWorldBeginTearDown(UWorld* InWorld);
+};
diff --git a/Engine/Plugins/Experimental/Avalanche/Source/AvalancheMedia/Private/Broadcast/Channel/AvaBroadcastOutputChannel.cpp b/Engine/Plugins/Experimental/Avalanche/Source/AvalancheMedia/Private/Broadcast/Channel/AvaBroadcastOutputChannel.cpp
--- a/Engine/Plugins/Experimental/Avalanche/Source/AvalancheMedia/Private/Broadcast/Channel/AvaBroadcastOutputChannel.cpp
+++ b/Engine/Plugins/Experimental/Avalanche/Source/AvalancheMedia/Private/Broadcast/Channel/AvaBroadcastOutputChannel.cpp
@@ -303,10 +303,10 @@ EAvaBroadcastIssueSeverity FAvaBroadcastOutputChannel::GetMediaOutputIssueSeveri
 	if (InOutputState == EAvaBroadcastOutputState::Live || InOutputState == EAvaBroadcastOutputState::Preparing)
 	{
 		// If the output is broadcasting remote, fetch the status from the playback client (which is proxying that output's status).
-		if (IAvaMediaModule::Get().IsMediaPlaybackClientStarted() && IsMediaOutputRemote(InMediaOutput))
+		if (IAvaMediaModule::Get().IsPlaybackClientStarted() && IsMediaOutputRemote(InMediaOutput))
 		{
 			const FAvaBroadcastMediaOutputInfo& OutputInfo = GetMediaOutputInfo(InMediaOutput);
-			const IAvaPlaybackClient& PlaybackClient = IAvaMediaModule::Get().GetMediaPlaybackClient();
+			const IAvaPlaybackClient& PlaybackClient = IAvaMediaModule::Get().GetPlaybackClient();
 			return PlaybackClient.GetMediaOutputIssueSeverity(GetMediaOutputServerName(InMediaOutput), GetChannelName().ToString(), OutputInfo.Guid);
 		}
 
@@ -330,10 +330,10 @@ EAvaBroadcastIssueSeverity FAvaBroadcastOutputChannel::GetMediaOutputIssueSeveri
 
 const TArray<FString>& FAvaBroadcastOutputChannel::GetMediaOutputIssueMessages(const UMediaOutput* InMediaOutput) const
 {
-	if (IAvaMediaModule::Get().IsMediaPlaybackClientStarted() && IsMediaOutputRemote(InMediaOutput))
+	if (IAvaMediaModule::Get().IsPlaybackClientStarted() && IsMediaOutputRemote(InMediaOutput))
 	{
 		const FAvaBroadcastMediaOutputInfo& OutputInfo = GetMediaOutputInfo(InMediaOutput);
-		const IAvaPlaybackClient& PlaybackClient = IAvaMediaModule::Get().GetMediaPlaybackClient();
+		const IAvaPlaybackClient& PlaybackClient = IAvaMediaModule::Get().GetPlaybackClient();
 		return PlaybackClient.GetMediaOutputIssueMessages(GetMediaOutputServerName(InMediaOutput), GetChannelName().ToString(), OutputInfo.Guid);
 	}
 	else
@@ -351,7 +351,7 @@ EAvaBroadcastOutputState FAvaBroadcastOutputChannel::GetMediaOutputState(const U
 {
 	if (IsMediaOutputRemote(InMediaOutput))
 	{
-		if (IAvaMediaModule::Get().IsMediaPlaybackClientStarted())
+		if (IAvaMediaModule::Get().IsPlaybackClientStarted())
 		{
 			const FAvaBroadcastMediaOutputInfo& OutputInfo = GetMediaOutputInfo(InMediaOutput);
 			if (!OutputInfo.IsValid())
@@ -361,7 +361,7 @@ EAvaBroadcastOutputState FAvaBroadcastOutputChannel::GetMediaOutputState(const U
 			}
 			
 			// If the output is broadcasting remote, fetch the status from the playback client (which is proxying that output's status).
-			const IAvaPlaybackClient& PlaybackClient = IAvaMediaModule::Get().GetMediaPlaybackClient();
+			const IAvaPlaybackClient& PlaybackClient = IAvaMediaModule::Get().GetPlaybackClient();
 			return PlaybackClient.GetMediaOutputState(GetMediaOutputServerName(InMediaOutput),GetChannelName().ToString(), OutputInfo.Guid);
 		}
 		else
@@ -628,13 +628,13 @@ bool FAvaBroadcastOutputChannel::StartChannelBroadcast()
 	bInternalStateBroadcasting = true;
 	
 	IAvaMediaModule& Module = IAvaMediaModule::Get();
-	if (Module.IsMediaPlaybackClientStarted())
+	if (Module.IsPlaybackClientStarted())
 	{
 		const TArray<UMediaOutput*> RemoteMediaOutputs = GetRemoteMediaOutputs();
 		if (!RemoteMediaOutputs.IsEmpty())
 		{
 			const FString ProfileName = GetProfileName().ToString();
-			Module.GetMediaPlaybackClient().RequestBroadcast(ProfileName, GetChannelName(), RemoteMediaOutputs, EAvaBroadcastAction::Start);
+			Module.GetPlaybackClient().RequestBroadcast(ProfileName, GetChannelName(), RemoteMediaOutputs, EAvaBroadcastAction::Start);
 		}
 	}
 
@@ -721,13 +721,13 @@ void FAvaBroadcastOutputChannel::StopChannelBroadcast()
 	StopPlaceholderTick();
 	
 	IAvaMediaModule& Module = IAvaMediaModule::Get();
-	if (Module.IsMediaPlaybackClientStarted())
+	if (Module.IsPlaybackClientStarted())
 	{
 		const TArray<UMediaOutput*> RemoteMediaOutputs = GetRemoteMediaOutputs();
 		if (!RemoteMediaOutputs.IsEmpty())
 		{
 			const FString ProfileName = GetProfileName().ToString();
-			Module.GetMediaPlaybackClient().RequestBroadcast(ProfileName, GetChannelName(), RemoteMediaOutputs, EAvaBroadcastAction::Stop);
+			Module.GetPlaybackClient().RequestBroadcast(ProfileName, GetChannelName(), RemoteMediaOutputs, EAvaBroadcastAction::Stop);
 		}
 	}
 
@@ -842,10 +842,10 @@ void FAvaBroadcastOutputChannel::OnMediaOutputModified(UMediaOutput* InMediaOutp
 	}
 	
 	IAvaMediaModule& AvaMediaModule = IAvaMediaModule::Get();
-	if (IsMediaOutputRemote(InMediaOutput) && AvaMediaModule.IsMediaPlaybackClientStarted())
+	if (IsMediaOutputRemote(InMediaOutput) && AvaMediaModule.IsPlaybackClientStarted())
 	{
 		const FString ProfileName = GetProfileName().ToString();
-		IAvaPlaybackClient& PlaybackClient = AvaMediaModule.GetMediaPlaybackClient();
+		IAvaPlaybackClient& PlaybackClient = AvaMediaModule.GetPlaybackClient();
 		PlaybackClient.RequestBroadcast(ProfileName, GetChannelName(), GetRemoteMediaOutputs(), EAvaBroadcastAction::UpdateConfig);
 	}
 	else
@@ -993,9 +993,9 @@ bool FAvaBroadcastOutputChannel::IsMediaOutputRemote(const UMediaOutput* InMedia
 	// if the playback client is started, we can try the legacy method
 	// of checking the remote device providers. This is much less reliable.
 	IAvaMediaModule& AvaMediaModule = IAvaMediaModule::Get();
-	if (AvaMediaModule.IsMediaPlaybackClientStarted())
+	if (AvaMediaModule.IsPlaybackClientStarted())
 	{
-		return AvaMediaModule.GetMediaPlaybackClient().IsMediaOutputRemoteFallback(InMediaOutput);
+		return AvaMediaModule.GetPlaybackClient().IsMediaOutputRemoteFallback(InMediaOutput);
 	}
 
 	return false;
diff --git a/Engine/Plugins/Experimental/Avalanche/Source/AvalancheMedia/Private/Playable/AvaPlayableGroup.cpp b/Engine/Plugins/Experimental/Avalanche/Source/AvalancheMedia/Private/Playable/AvaPlayableGroup.cpp
--- a/Engine/Plugins/Experimental/Avalanche/Source/AvalancheMedia/Private/Playable/AvaPlayableGroup.cpp
+++ b/Engine/Plugins/Experimental/Avalanche/Source/AvalancheMedia/Private/Playable/AvaPlayableGroup.cpp
@@ -361,6 +361,31 @@ bool UAvaPlayableGroup::UpdateCameraSetup()
 	return false;
 }
 
+bool UAvaPlayableGroup::IsWorldPlaying() const
+{
+	return GameInstance ? GameInstance->IsWorldPlaying() : false;
+}
+
+bool UAvaPlayableGroup::IsRenderTargetReady() const
+{
+	return GameInstance ? GameInstance->IsRenderTargetReady() : false;
+}
+
+UTextureRenderTarget2D* UAvaPlayableGroup::GetRenderTarget() const
+{
+	return GameInstance ? GameInstance->GetRenderTarget() : RenderTarget.Get();
+}
+
+UGameInstance* UAvaPlayableGroup::GetGameInstance() const
+{
+	return GameInstance;
+}
+	
+UWorld* UAvaPlayableGroup::GetPlayWorld() const
+{
+	return GameInstance ? GameInstance->GetPlayWorld() : nullptr;
+}
+
 bool UAvaPlayableGroup::ConditionalRequestUnloadWorld(bool bForceImmediate)
 {
 	if (!GameInstance)
diff --git a/Engine/Plugins/Experimental/Avalanche/Source/AvalancheMedia/Private/Playable/AvaPlayableGroupManager.h b/Engine/Plugins/Experimental/Avalanche/Source/AvalancheMedia/Private/Playable/AvaPlayableGroupManager.h
new file mode 100644
--- /dev/null
+++ b/Engine/Plugins/Experimental/Avalanche/Source/AvalancheMedia/Private/Playable/AvaPlayableGroupManager.h
@@ -0,0 +1,113 @@
+// Copyright Epic Games, Inc. All Rights Reserved.
+
+#pragma once
+
+#include "UObject/Object.h"
+#include "UObject/ObjectPtr.h"
+#include "UObject/SoftObjectPtr.h"
+#include "AvaPlayableGroupManager.generated.h"
+
+class UAvaGameInstance;
+class UAvaPlayableGroup;
+class UAvaPlayableGroupManager;
+
+/**
+ * Manager for the shared playable groups per channel.
+ */
+UCLASS()
+class UAvaPlayableGroupChannelManager : public UObject
+{
+	GENERATED_BODY()
+
+public:
+	UAvaPlayableGroup* GetOrCreateSharedLevelGroup(bool bInIsRemoteProxy);
+
+	UAvaPlayableGroupManager* GetPlayableGroupManager() const;
+	
+protected:
+	//~ Begin UObject
+	virtual void BeginDestroy() override;
+	//~ End UObject
+
+	void Shutdown();
+	
+protected:
+	FName ChannelName;
+
+	/**
+	 * This is the shared playable group for this channel.
+	 * 
+	 * For now all the levels go in the same group for a given channel.
+	 * This a weak ptr because it is only a cache. The ownership
+	 * of the group is with the playable objects.
+	 */
+	TWeakObjectPtr<UAvaPlayableGroup> SharedLevelGroupWeak;
+
+	/**
+	 * For remote proxy playables, we need a specialized group
+	 * that doesn't implement the local game instance but still
+	 * provides the correct logic to emulate functionality.
+	 */
+	TWeakObjectPtr<UAvaPlayableGroup> SharedRemoteProxyLevelGroupWeak;
+
+	friend class UAvaPlayableGroupManager;
+};
+
+/**
+ * Manager for the shared playable groups.
+ * The scope of this manager is either global (in the global playback manager)
+ * or for a given playback manager.
+ */
+UCLASS()
+class UAvaPlayableGroupManager : public UObject
+{
+	GENERATED_BODY()
+	
+public:
+	void Init();
+
+	void Shutdown();
+
+	void Tick(double InDeltaSeconds);
+
+	UAvaPlayableGroupChannelManager* FindChannelManager(const FName& InChannelName) const
+	{
+		const TObjectPtr<UAvaPlayableGroupChannelManager>* ChannelManager = ChannelManagers.Find(InChannelName);
+		return ChannelManager ? *ChannelManager : nullptr;
+	}
+
+	UAvaPlayableGroupChannelManager* FindOrAddChannelManager(const FName& InChannelName);
+
+	UAvaPlayableGroup* GetOrCreateSharedLevelGroup(const FName& InChannelName, bool bInIsRemoteProxy)
+	{
+		UAvaPlayableGroupChannelManager* ChannelManager = FindOrAddChannelManager(InChannelName);
+		return ChannelManager ? ChannelManager->GetOrCreateSharedLevelGroup(bInIsRemoteProxy) : nullptr;
+	}
+
+	void RegisterForLevelStreamingUpdate(UAvaPlayableGroup* InPlayableGroup);
+	void UnregisterFromLevelStreamingUpdate(UAvaPlayableGroup* InPlayableGroup);
+
+	void RegisterForTransitionTicking(UAvaPlayableGroup* InPlayableGroup);
+	void UnregisterFromTransitionTicking(UAvaPlayableGroup* InPlayableGroup);
+
+protected:
+	//~ Begin UObject
+	virtual void BeginDestroy() override;
+	//~ End UObject
+
+	void OnGameInstanceEndPlay(UAvaGameInstance* InGameInstance, FName InChannelName);
+
+	void UpdateLevelStreaming();
+	
+	void TickTransitions(double InDeltaSeconds);
+	
+protected:
+	UPROPERTY(Transient)
+	TMap<FName, TObjectPtr<UAvaPlayableGroupChannelManager>> ChannelManagers;
+
+	bool bIsUpdatingStreaming = false;
+	TSet<TWeakObjectPtr<UAvaPlayableGroup>> GroupsToUpdateStreaming;
+
+	bool bIsTickingTransitions = false;
+	TSet<TWeakObjectPtr<UAvaPlayableGroup>> GroupsToTickTransitions;
+};
\ No newline at end of file
diff --git a/Engine/Plugins/Experimental/Avalanche/Source/AvalancheMedia/Private/Playable/Playables/AvaPlayableLevelStreaming.cpp b/Engine/Plugins/Experimental/Avalanche/Source/AvalancheMedia/Private/Playable/Playables/AvaPlayableLevelStreaming.cpp
--- a/Engine/Plugins/Experimental/Avalanche/Source/AvalancheMedia/Private/Playable/Playables/AvaPlayableLevelStreaming.cpp
+++ b/Engine/Plugins/Experimental/Avalanche/Source/AvalancheMedia/Private/Playable/Playables/AvaPlayableLevelStreaming.cpp
@@ -9,6 +9,7 @@
 #include "Engine/Level.h"
 #include "Engine/LevelStreamingDynamic.h"
 #include "Engine/LocalPlayer.h"
+#include "Framework/AvaGameInstance.h"
 #include "Framework/AvaInstanceSettings.h"
 #include "Framework/AvaSoftAssetPtr.h"
 #include "IAvaMediaModule.h"
diff --git a/Engine/Plugins/Experimental/Avalanche/Source/AvalancheMedia/Private/Playable/Playables/AvaPlayableRemoteProxy.cpp b/Engine/Plugins/Experimental/Avalanche/Source/AvalancheMedia/Private/Playable/Playables/AvaPlayableRemoteProxy.cpp
--- a/Engine/Plugins/Experimental/Avalanche/Source/AvalancheMedia/Private/Playable/Playables/AvaPlayableRemoteProxy.cpp
+++ b/Engine/Plugins/Experimental/Avalanche/Source/AvalancheMedia/Private/Playable/Playables/AvaPlayableRemoteProxy.cpp
@@ -20,7 +20,7 @@ bool UAvaPlayableRemoteProxy::LoadAsset(const FAvaSoftAssetPtr& InSourceAsset, b
 	
 	SourceAssetPath = InSourceAsset.ToSoftObjectPath();
 
-	IAvaPlaybackClient& PlaybackClient = IAvaMediaModule::Get().GetMediaPlaybackClient();
+	IAvaPlaybackClient& PlaybackClient = IAvaMediaModule::Get().GetPlaybackClient();
 	
 	if (PlaybackClient.HasAnyServerOnlineForChannel(PlayingChannelFName))
 	{
@@ -47,7 +47,7 @@ bool UAvaPlayableRemoteProxy::LoadAsset(const FAvaSoftAssetPtr& InSourceAsset, b
 
 bool UAvaPlayableRemoteProxy::UnloadAsset()
 {
-	IAvaPlaybackClient& PlaybackClient = IAvaMediaModule::Get().GetMediaPlaybackClient();
+	IAvaPlaybackClient& PlaybackClient = IAvaMediaModule::Get().GetPlaybackClient();
 	if (PlaybackClient.HasAnyServerOnlineForChannel(PlayingChannelFName))
 	{
 		PlaybackClient.RequestPlayback(InstanceId, SourceAssetPath, PlayingChannelName, EAvaPlaybackAction::Unload);
@@ -88,7 +88,7 @@ namespace UE::AvaMedIaRemoteProxyPlayable::Private
 
 EAvaPlayableStatus UAvaPlayableRemoteProxy::GetPlayableStatus() const
 {
-	IAvaPlaybackClient& PlaybackClient = IAvaMediaModule::Get().GetMediaPlaybackClient();
+	IAvaPlaybackClient& PlaybackClient = IAvaMediaModule::Get().GetPlaybackClient();
 	const TArray<FString> OnlineServers = GetOnlineServerForChannel(PlayingChannelFName);
 	
 	for (const FString& Server : OnlineServers)
@@ -115,7 +115,7 @@ IAvaSceneInterface* UAvaPlayableRemoteProxy::GetSceneInterface() const
 
 EAvaPlayableCommandResult UAvaPlayableRemoteProxy::ExecuteAnimationCommand(EAvaPlaybackAnimAction InAnimAction, const FAvaPlaybackAnimPlaySettings& InAnimPlaySettings)
 {
-	IAvaPlaybackClient& PlaybackClient = IAvaMediaModule::Get().GetMediaPlaybackClient();
+	IAvaPlaybackClient& PlaybackClient = IAvaMediaModule::Get().GetPlaybackClient();
 	
 	// If an animation event was locally scheduled on a remote playable,
 	// we need to propagate the event.
@@ -149,7 +149,7 @@ EAvaPlayableCommandResult UAvaPlayableRemoteProxy::ExecuteAnimationCommand(EAvaP
 
 EAvaPlayableCommandResult UAvaPlayableRemoteProxy::UpdateRemoteControlCommand(const TSharedRef<FAvaPlayableRemoteControlValues>& InRemoteControlValues)
 {
-	IAvaPlaybackClient& PlaybackClient = IAvaMediaModule::Get().GetMediaPlaybackClient();
+	IAvaPlaybackClient& PlaybackClient = IAvaMediaModule::Get().GetPlaybackClient();
 	if (PlaybackClient.HasAnyServerOnlineForChannel(PlayingChannelFName))
 	{
 		PlaybackClient.RequestRemoteControlUpdate(InstanceId, SourceAssetPath, PlayingChannelName, *InRemoteControlValues);
@@ -166,7 +166,7 @@ void UAvaPlayableRemoteProxy::SetUserData(const FString& InUserData)
 {
 	if (UserData != InUserData)
 	{
-		IAvaPlaybackClient& PlaybackClient = IAvaMediaModule::Get().GetMediaPlaybackClient();
+		IAvaPlaybackClient& PlaybackClient = IAvaMediaModule::Get().GetPlaybackClient();
 		if (PlaybackClient.HasAnyServerOnlineForChannel(PlayingChannelFName))
 		{
 			PlaybackClient.RequestPlayback(InstanceId, GetSourceAssetPath(), PlayingChannelName, EAvaPlaybackAction::SetUserData, InUserData);
@@ -206,12 +206,12 @@ bool UAvaPlayableRemoteProxy::InitPlayable(const FPlayableCreationInfo& InPlayab
 void UAvaPlayableRemoteProxy::OnPlay()
 {
 	IAvaMediaModule& AvaMediaModule = IAvaMediaModule::Get();
-	if (!AvaMediaModule.IsMediaPlaybackClientStarted())
+	if (!AvaMediaModule.IsPlaybackClientStarted())
 	{
 		return;
 	}
 	
-	IAvaPlaybackClient& Client = AvaMediaModule.GetMediaPlaybackClient();
+	IAvaPlaybackClient& Client = AvaMediaModule.GetPlaybackClient();
 	
 	if (Client.HasAnyServerOnlineForChannel(PlayingChannelFName))
 	{
@@ -231,7 +231,7 @@ void UAvaPlayableRemoteProxy::OnPlay()
 
 void UAvaPlayableRemoteProxy::OnEndPlay()
 {
-	IAvaPlaybackClient& PlaybackClient = IAvaMediaModule::Get().GetMediaPlaybackClient();
+	IAvaPlaybackClient& PlaybackClient = IAvaMediaModule::Get().GetPlaybackClient();
 	if (PlaybackClient.HasAnyServerOnlineForChannel(PlayingChannelFName))
 	{
 		PlaybackClient.RequestPlayback(InstanceId, SourceAssetPath, PlayingChannelName, EAvaPlaybackAction::Stop);
diff --git a/Engine/Plugins/Experimental/Avalanche/Source/AvalancheMedia/Private/Playable/Transition/AvaPlayableRemoteTransition.cpp b/Engine/Plugins/Experimental/Avalanche/Source/AvalancheMedia/Private/Playable/Transition/AvaPlayableRemoteTransition.cpp
--- a/Engine/Plugins/Experimental/Avalanche/Source/AvalancheMedia/Private/Playable/Transition/AvaPlayableRemoteTransition.cpp
+++ b/Engine/Plugins/Experimental/Avalanche/Source/AvalancheMedia/Private/Playable/Transition/AvaPlayableRemoteTransition.cpp
@@ -18,7 +18,7 @@ bool UAvaPlayableRemoteTransition::Start()
 {
 	IAvaMediaModule& AvaMediaModule = IAvaMediaModule::Get();
 	
-	if (!AvaMediaModule.IsMediaPlaybackClientStarted())
+	if (!AvaMediaModule.IsPlaybackClientStarted())
 	{
 		TransitionId.Invalidate();
 		return false;
@@ -42,7 +42,7 @@ bool UAvaPlayableRemoteTransition::Start()
 		EnterValues.Add(Values.IsValid() ? *Values : FAvaPlayableRemoteControlValues::GetDefaultEmpty());
 	}
 	
-	IAvaPlaybackClient& PlaybackClient = AvaMediaModule.GetMediaPlaybackClient();
+	IAvaPlaybackClient& PlaybackClient = AvaMediaModule.GetPlaybackClient();
 	PlaybackClient.RequestPlayableTransitionStart(TransitionId, MoveTemp(EnterInstanceIds), MoveTemp(PlayingInstanceIds), MoveTemp(ExitInstanceIds), MoveTemp(EnterValues), ChannelName, TransitionFlags);
 	RegisterToPlaybackClientDelegates();	// to get the transition events from the server side.
 
@@ -57,9 +57,9 @@ void UAvaPlayableRemoteTransition::Stop()
 	{
 		IAvaMediaModule& AvaMediaModule = IAvaMediaModule::Get();
 	
-		if (AvaMediaModule.IsMediaPlaybackClientStarted())
+		if (AvaMediaModule.IsPlaybackClientStarted())
 		{
-			AvaMediaModule.GetMediaPlaybackClient().RequestPlayableTransitionStop(TransitionId, ChannelName);
+			AvaMediaModule.GetPlaybackClient().RequestPlayableTransitionStop(TransitionId, ChannelName);
 		}
 	}
 
diff --git a/Engine/Plugins/Experimental/Avalanche/Source/AvalancheMedia/Private/Playable/Transition/AvaPlayableTransitionScene.h b/Engine/Plugins/Experimental/Avalanche/Source/AvalancheMedia/Private/Playable/Transition/AvaPlayableTransitionScene.h
new file mode 100644
--- /dev/null
+++ b/Engine/Plugins/Experimental/Avalanche/Source/AvalancheMedia/Private/Playable/Transition/AvaPlayableTransitionScene.h
@@ -0,0 +1,33 @@
+// Copyright Epic Games, Inc. All Rights Reserved.
+
+#pragma once
+
+#include "AvaTagHandle.h"
+#include "AvaTransitionScene.h"
+#include "AvaPlayableTransitionScene.generated.h"
+
+class UAvaPlayable;
+class UAvaPlayableTransition;
+
+USTRUCT()
+struct FAvaPlayableTransitionScene : public FAvaTransitionScene
+{
+	GENERATED_BODY()
+
+	FAvaPlayableTransitionScene() = default;
+
+	explicit FAvaPlayableTransitionScene(UAvaPlayable* InPlayable, UAvaPlayableTransition* InPlayableTransition);
+
+	explicit FAvaPlayableTransitionScene(const FAvaTagHandle& InTransitionLayer, UAvaPlayableTransition* InPlayableTransition);
+
+	//~ Begin FAvaTransitionScene
+	virtual EAvaTransitionComparisonResult Compare(const FAvaTransitionScene& InOther) const override;
+	virtual ULevel* GetLevel() const override;
+	virtual void GetOverrideTransitionLayer(FAvaTagHandle& OutTransitionLayer) const override;
+	virtual void OnFlagsChanged() override;
+	//~ End FAvaTransitionScene
+
+	TWeakObjectPtr<UAvaPlayableTransition> PlayableTransitionWeak;
+
+	TOptional<FAvaTagHandle> OverrideTransitionLayer;
+};
diff --git a/Engine/Plugins/Experimental/Avalanche/Source/AvalancheMedia/Private/Playback/AvaPlaybackClient.h b/Engine/Plugins/Experimental/Avalanche/Source/AvalancheMedia/Private/Playback/AvaPlaybackClient.h
new file mode 100644
--- /dev/null
+++ b/Engine/Plugins/Experimental/Avalanche/Source/AvalancheMedia/Private/Playback/AvaPlaybackClient.h
@@ -0,0 +1,490 @@
+// Copyright Epic Games, Inc. All Rights Reserved.
+
+#pragma once
+
+#include "Broadcast/OutputDevices/AvaBroadcastDeviceProviderData.h"
+#include "Containers/Ticker.h"
+#include "IMessageContext.h"
+#include "MessageEndpoint.h"
+#include "Playback/AvaPlaybackMessages.h"
+#include "Playback/IAvaPlaybackClient.h"
+
+class FAvaMediaModule;
+struct FAvaBroadcastOutputChannel;
+
+// Listens to command to play and broadcast
+class FAvaPlaybackClient : public IAvaPlaybackClient
+{
+public:
+	FAvaPlaybackClient(FAvaMediaModule* InParentModule);
+	virtual ~FAvaPlaybackClient() override;
+
+	void Init();
+
+	//IAvaPlaybackClient
+	virtual int32 GetNumConnectedServers() const override;
+	virtual TArray<FString> GetServerNames() const override;
+	virtual FMessageAddress GetServerAddress(const FString& InServerName) const override;
+	virtual bool HasServerUserData(const FString& InServerName, const FString& InKey) const override;
+	virtual const FString& GetServerUserData(const FString& InServerName, const FString& InKey) const override;
+	virtual bool HasUserData(const FString& InKey) const override;
+	virtual const FString& GetUserData(const FString& InKey) const override;
+	virtual void SetUserData(const FString& InKey, const FString& InData) override;
+	virtual void RemoveUserData(const FString& InKey) override;
+	virtual void BroadcastStatCommand(const FString& InCommand, bool bInBroadcastLocalState) override;
+	virtual void RequestPlaybackAssetStatus(const FSoftObjectPath& InAssetPath, const FString& InChannelOrServerName, bool bInForceRefresh) override;
+	virtual void RequestPlayback(const FGuid& InInstanceId, const FSoftObjectPath& InAssetPath, const FString& InChannelName, EAvaPlaybackAction InAction, const FString& InArguments = FString()) override;
+	virtual void RequestAnimPlayback(const FGuid& InInstanceId, const FSoftObjectPath& InAssetPath, const FString& InChannelName, const FAvaPlaybackAnimPlaySettings& InAnimSettings) override;
+	virtual void RequestAnimAction(const FGuid& InInstanceId, const FSoftObjectPath& InAssetPath, const FString& InChannelName, const FString& InAnimationName, EAvaPlaybackAnimAction InAction) override;
+	virtual void RequestRemoteControlUpdate(const FGuid& InInstanceId, const FSoftObjectPath& InAssetPath, const FString& InChannelName, const FAvaPlayableRemoteControlValues& InRemoteControlValues) override;
+	virtual void RequestPlayableTransitionStart(const FGuid& InTransitionId, TArray<FGuid>&& InEnterInstanceIds, TArray<FGuid>&& InPlayingInstanceIds, TArray<FGuid>&& InExitInstanceIds, TArray<FAvaPlayableRemoteControlValues>&& InEnterValues, const FName& InChannelName, EAvaPlayableTransitionFlags InTransitionFlags) override;
+	virtual void RequestPlayableTransitionStop(const FGuid& InTransitionId, const FName& InChannelName) override;
+	virtual void RequestBroadcast(const FString& InProfile, const FName& InChannel, const TArray<UMediaOutput*>& InRemoteMediaOutputs, EAvaBroadcastAction InAction) override;
+	virtual bool IsMediaOutputRemoteFallback(const UMediaOutput* InMediaOutput) override;
+	virtual EAvaBroadcastIssueSeverity GetMediaOutputIssueSeverity(const FString& InServerName, const FString& InChannelName, const FGuid& InOutputGuid) const override;
+	virtual const TArray<FString>& GetMediaOutputIssueMessages(const FString& InServerName, const FString& InChannelName, const FGuid& InOutputGuid) const override;
+	virtual EAvaBroadcastOutputState GetMediaOutputState(const FString& InServerName, const FString& InChannelName, const FGuid& InOutputGuid) const override;
+	virtual bool HasAnyServerOnlineForChannel(const FName& InChannelName) const override;
+	virtual TOptional<EAvaPlaybackStatus> GetRemotePlaybackStatus(const FGuid& InInstanceId, const FSoftObjectPath& InAssetPath, const FString& InChannelName, const FString& InServerName) const override;
+	virtual const FString* GetRemotePlaybackUserData(const FGuid& InInstanceId, const FSoftObjectPath& InAssetPath, const FString& InChannelName, const FString& InServerName) const override;
+	virtual TOptional<EAvaPlaybackAssetStatus> GetRemotePlaybackAssetStatus(const FSoftObjectPath& InAssetPath, const FString& InServerName) const override;
+	//~IAvaPlaybackClient
+
+	FString GetServerProjectContentPath(const FString& InServerName) const;
+	uint32 GetServerProcessId(const FString& InServerName) const;
+
+	FString GetProjectContentPath() const { return ProjectContentPath; }
+
+protected:
+	TOptional<EAvaPlaybackStatus> GetPlaybackStatusForServer(const FString& InServerName, const FString& InChannelName, const FGuid& InInstanceId, const FSoftObjectPath& InAssetPath) const;
+	const FString* GetPlaybackUserDataForServer(const FString& InServerName, const FString& InChannelName, const FGuid& InInstanceId, const FSoftObjectPath& InAssetPath) const;
+	void RequestPlaybackAssetStatusForServer(const FSoftObjectPath& InAssetPath, const FString& InServerName, bool bInForceRefresh);
+
+	void Tick();
+
+	bool HandlePingTicker(float InDeltaTime);
+
+	// Message handlers
+	void HandlePlaybackPongMessage(const FAvaPlaybackPong& InMessage, const TSharedRef<IMessageContext, ESPMode::ThreadSafe>& InContext);
+	void HandlePlaybackLogMessage(const FAvaPlaybackLog& InMessage, const TSharedRef<IMessageContext, ESPMode::ThreadSafe>& InContext);
+	void HandleUpdateServerUserData(const FAvaPlaybackUpdateServerUserData& InMessage, const TSharedRef<IMessageContext, ESPMode::ThreadSafe>& InContext);
+	void HandleStatStatus(const FAvaPlaybackStatStatus& InMessage, const TSharedRef<IMessageContext, ESPMode::ThreadSafe>& InContext);
+	void HandleDeviceProviderDataListMessage(const FAvaBroadcastDeviceProviderDataList& InMessage, const TSharedRef<IMessageContext, ESPMode::ThreadSafe>& InContext);
+	void HandleBroadcastStatusMessage(const FAvaBroadcastStatus& InMessage, const TSharedRef<IMessageContext, ESPMode::ThreadSafe>& InContext);
+	void HandlePlaybackAssetStatusMessage(const FAvaPlaybackAssetStatus& InMessage, const TSharedRef<IMessageContext, ESPMode::ThreadSafe>& InContext);
+	void HandlePlaybackStatusMessage(const FAvaPlaybackStatus& InMessage, const TSharedRef<IMessageContext, ESPMode::ThreadSafe>& InContext);
+	void HandlePlaybackStatusesMessage(const FAvaPlaybackStatuses& InMessage, const TSharedRef<IMessageContext, ESPMode::ThreadSafe>& InContext);
+	void HandlePlaybackSequenceEventMessage(const FAvaPlaybackSequenceEvent& InMessage, const TSharedRef<IMessageContext, ESPMode::ThreadSafe>& InContext);
+	void HandlePlaybackTransitionEventMessage(const FAvaPlaybackTransitionEvent& InMessage, const TSharedRef<IMessageContext, ESPMode::ThreadSafe>& InContext);
+	
+	void RegisterCommands();
+	
+	// Command handlers
+	void PingServersCommand(const TArray<FString>& InArgs) { PublishPlaybackPing(FDateTime::UtcNow(), false);}
+	void RequestPlaybackAssetStatusCommand(const TArray<FString>& InArgs);
+	void RequestPlaybackCommand(const TArray<FString>& InArgs);
+	void RequestBroadcastCommand(const TArray<FString>& InArgs);
+	void SetUserDataCommand(const TArray<FString>& InArgs);
+	void ShowStatusCommand(const TArray<FString>& InArgs);
+
+	// Event handlers
+	void OnBroadcastChanged(EAvaBroadcastChange InChange);
+	void OnChannelChanged(const FAvaBroadcastOutputChannel& InChannel, EAvaBroadcastChannelChange InChange);
+	void OnAvaMediaSettingsChanged(UObject*, struct FPropertyChangedEvent&);
+	void OnPreSavePackage(UPackage* InPackage, FObjectPreSaveContext InObjectSaveContext);
+	void OnPackageSaved(const FString& InPackageFileName, UPackage* InPackage, FObjectPostSaveContext InObjectSaveContext);
+	void OnAssetRemoved(const FAssetData& InAssetData);
+	
+	void ApplyAvaMediaSettings();
+
+protected:
+	void FillClientInfo(FAvaPlaybackClientMessageBase& InMessage)
+	{
+		// Remark: for now, client name and computer name are the same.
+		// We don't support having many clients per computer at this point.
+		InMessage.ClientName =  ComputerName;
+	}
+
+	template<typename MessageType>
+	void PublishRequest(MessageType* InMessage)
+	{
+		FillClientInfo(*InMessage);
+		MessageEndpoint->Publish(InMessage);
+	}
+
+	template<typename MessageType>
+	void SendRequest(MessageType* InMessage, const FMessageAddress& InRecipient, EMessageFlags flags = EMessageFlags::None)
+	{
+		FillClientInfo(*InMessage);
+		MessageEndpoint->Send(InMessage, MessageType::StaticStruct(), flags, nullptr,
+			TArrayBuilder<FMessageAddress>().Add(InRecipient), FTimespan::Zero(), FDateTime::MaxValue());
+	}
+	
+	template<typename MessageType>
+	void SendRequest(MessageType* InMessage, const TArray<FMessageAddress>& InRecipients, EMessageFlags flags = EMessageFlags::None)
+	{
+		FillClientInfo(*InMessage);
+		MessageEndpoint->Send(InMessage, MessageType::StaticStruct(), flags, nullptr,
+			InRecipients, FTimespan::Zero(), FDateTime::MaxValue());
+	}
+
+	void PublishPlaybackPing(const FDateTime& InCurrentTime, bool bInAutoPing);
+	void SendUserDataUpdate(const TArray<FMessageAddress>& InRecipients);
+	void SendBroadcastSettingsUpdate(const TArray<FMessageAddress>& InRecipients);
+	void SendAvaInstanceSettingsUpdate(const TArray<FMessageAddress>& InRecipients);
+	void SendBroadcastChannelSettingsUpdate(const TArray<FMessageAddress>& InRecipients, const FAvaBroadcastOutputChannel& InChannel);
+	void SendPackageEvent(const TArray<FMessageAddress>& InRecipients, const FName& InPackageName, EAvaPlaybackPackageEvent InEvent);
+	void SendStatCommand(const FString& InCommand, bool bInBroadcastLocalState, const TArray<FMessageAddress>& InRecipients);
+	void SendClientInfo(const FMessageAddress& InRecipient);	
+	void RemoveDeadServers(const FDateTime& InCurrentTime);
+	void UpdateServerAddresses();
+	
+	/**
+	 * Get the server addresses for the given channel in the current broadcast profile.
+	 * This is based on the channel's MediaOutputInfos. It will properly work even
+	 * if the channel is idle (but not offline).
+	 */
+	TArray<FMessageAddress> GetServerAddressesForChannel(const FName& InChannelName) const;
+	TArray<FMessageAddress> GetServerAddressesForChannel(const FString& InChannelName) const
+	{
+		return GetServerAddressesForChannel(FName(InChannelName));
+	}
+
+	/**
+	 * Get the server names for the given channel in the current broadcast profile.
+	 * This is based on the channel's MediaOutputInfos. It will properly work
+	 * regardless of the channel's status (idle or even offline).
+	 */
+	TArray<FString> GetServerNamesForChannel(const FName& InChannelName) const;
+
+	FString GetServerNameForMediaOutputFallback(const UMediaOutput* InMediaOutput) const;
+
+private:
+	FAvaMediaModule* ParentModule;
+	
+	TArray<IConsoleObject*> ConsoleCommands;
+	
+	/** Holds the messaging endpoint. */
+	TSharedPtr<FMessageEndpoint, ESPMode::ThreadSafe> MessageEndpoint;
+	
+	/** Holds a delegate to be invoked when the widget ticks. */
+	FTickerDelegate PingTickDelegate;
+
+	/** Handle to the registered TickDelegate. */
+	FTSTicker::FDelegateHandle PingTickDelegateHandle;
+	
+	FDelegateHandle BroadcastChangedHandle;
+	
+	FString ComputerName;
+	FString ProjectContentPath;
+	uint32 ProcessId = 0;
+	TMap<FString, FString> UserDataEntries;
+	
+	/** Structure holding information on pending playback asset status requests.*/
+	struct FPendingPlaybackAssetStatusRequest
+	{
+		FDateTime ExpirationTime;
+		bool bForceRefresh = false;
+	};
+	
+	/** Remote Broadcast Channel Info */
+	struct FBroadcastChannelInfo
+	{
+		EAvaBroadcastChannelState ChannelState;
+		EAvaBroadcastIssueSeverity ChannelIssueSeverity;
+		TMap<FGuid, FAvaBroadcastOutputStatus> MediaOutputStatuses;
+	};
+
+	/** Remote Playback Instance Info */
+	struct FPlaybackInstanceInfo
+	{
+		EAvaPlaybackStatus  Status;
+		TOptional<FString> UserData;
+	};
+	
+	/** Remote playback for multiple instances of a given asset. */
+	struct FPlaybackAssetInfo
+	{
+		TMap<FGuid, FPlaybackInstanceInfo> InstanceByIds;
+
+		/*
+		 * Return a relevance factor for status reporting, based on what we would like to see in the page status.
+		 */
+		static int GetStatusRelevance(EAvaPlaybackStatus InPlaybackStatus)
+		{
+			switch (InPlaybackStatus)
+			{
+			case EAvaPlaybackStatus::Unknown: return 0;
+			case EAvaPlaybackStatus::Missing: return 1;
+			case EAvaPlaybackStatus::Syncing: return 2;
+			case EAvaPlaybackStatus::Available: return 3;
+			case EAvaPlaybackStatus::Loading: return 4;
+			case EAvaPlaybackStatus::Loaded: return 5;
+			case EAvaPlaybackStatus::Starting: return 6;
+			case EAvaPlaybackStatus::Started: return 7;
+			case EAvaPlaybackStatus::Unloading: return 3;
+			case EAvaPlaybackStatus::Error: return 1;
+			default: return 0;
+			}
+		}
+		
+		TOptional<EAvaPlaybackStatus> GetInstanceStatus(const FGuid& InInstanceId) const
+		{
+			// If the request is for a specific instance, check if we have it.
+			if (InInstanceId.IsValid())
+			{
+				const FPlaybackInstanceInfo* Instance = InstanceByIds.Find(InInstanceId);
+				return Instance ? TOptional(Instance->Status) : TOptional<EAvaPlaybackStatus>();
+			}
+			
+			// If the instanceId is not specified, return the most relevant status we have already received.
+			// Todo: investigate why this happens and if there is better way to handle this (like better request handling).
+			TOptional<EAvaPlaybackStatus> MostRelevantStatus;
+			for (const TPair<FGuid, FPlaybackInstanceInfo>& PlaybackInfo :  InstanceByIds)
+			{
+				if (!MostRelevantStatus.IsSet() || GetStatusRelevance(PlaybackInfo.Value.Status) > GetStatusRelevance(MostRelevantStatus.GetValue()))
+				{
+					MostRelevantStatus = PlaybackInfo.Value.Status;
+				}
+			}
+			return MostRelevantStatus;
+		}
+
+		const FPlaybackInstanceInfo* GetInstanceInfo(const FGuid& InInstanceId) const
+		{
+			return InInstanceId.IsValid() ? InstanceByIds.Find(InInstanceId) : nullptr;
+		}
+
+		FPlaybackInstanceInfo* GetOrAddInstanceInfo(const FGuid& InInstanceId)
+		{
+			if (InInstanceId.IsValid())
+			{
+				FPlaybackInstanceInfo* Instance = InstanceByIds.Find(InInstanceId);
+				return Instance ? Instance : &InstanceByIds.Emplace(InInstanceId);
+			}
+			return nullptr;
+		}
+		
+		void SetInstanceStatus(const FGuid& InInstanceId, EAvaPlaybackStatus InStatus)
+		{
+			if (FPlaybackInstanceInfo* Instance = GetOrAddInstanceInfo(InInstanceId))
+			{
+				Instance->Status = InStatus;
+			}
+		}
+
+		const FString* GetInstanceUserData(const FGuid& InInstanceId) const
+		{
+			// If the request is for a specific instance, check if we have it.
+			if (const FPlaybackInstanceInfo* Instance = GetInstanceInfo(InInstanceId))
+			{
+				return Instance->UserData.IsSet() ? &Instance->UserData.GetValue() : nullptr;
+			}
+			return nullptr;
+		}
+
+		void SetInstanceUserData(const FGuid& InInstanceId, const FString& InUserData)
+		{
+			if (FPlaybackInstanceInfo* Instance = GetOrAddInstanceInfo(InInstanceId))
+			{
+				Instance->UserData = InUserData;
+			}
+		}
+	};
+
+	/** Playback Asset Info for a given channel.  */
+	struct FPlaybackChannelInfo
+	{
+		TMap<FSoftObjectPath, TUniquePtr<FPlaybackAssetInfo>> AssetInfoByPaths;
+
+		FPlaybackAssetInfo* GetAssetInfo(const FSoftObjectPath& InAssetPath) const
+		{
+			const TUniquePtr<FPlaybackAssetInfo>* AssetInfo = AssetInfoByPaths.Find(InAssetPath);
+			return AssetInfo ? (*AssetInfo).Get() : nullptr;
+		}
+
+		FPlaybackAssetInfo* GetOrAddAssetInfo(const FSoftObjectPath& InAssetPath)
+		{
+			if (const TUniquePtr<FPlaybackAssetInfo>* AssetInfo = AssetInfoByPaths.Find(InAssetPath))
+			{
+				return (*AssetInfo).Get();
+			}
+			return AssetInfoByPaths.Add(InAssetPath, MakeUnique<FPlaybackAssetInfo>()).Get();
+		}
+		
+		TOptional<EAvaPlaybackStatus> GetInstanceStatus(const FGuid& InInstanceId, const FSoftObjectPath& InAssetPath) const
+		{
+			const FPlaybackAssetInfo* AssetInfo = GetAssetInfo(InAssetPath);
+			return AssetInfo ? AssetInfo->GetInstanceStatus(InInstanceId) : TOptional<EAvaPlaybackStatus>();
+		}
+		
+		void SetInstanceStatus(const FGuid& InInstanceId, const FSoftObjectPath& InAssetPath, EAvaPlaybackStatus InStatus)
+		{
+			if (FPlaybackAssetInfo* AssetInfo = GetOrAddAssetInfo(InAssetPath))
+			{
+				AssetInfo->SetInstanceStatus(InInstanceId, InStatus);
+			}
+		}
+
+		const FString* GetInstanceUserData(const FGuid& InInstanceId, const FSoftObjectPath& InAssetPath) const
+		{
+			const FPlaybackAssetInfo* AssetInfo = GetAssetInfo(InAssetPath);
+			return AssetInfo ? AssetInfo->GetInstanceUserData(InInstanceId) : nullptr;
+		}
+		
+		void SetInstanceUserData(const FGuid& InInstanceId, const FSoftObjectPath& InAssetPath, const FString& InUserData)
+		{
+			if (FPlaybackAssetInfo* PlaybackInfos = GetOrAddAssetInfo(InAssetPath))
+			{
+				PlaybackInfos->SetInstanceUserData(InInstanceId, InUserData);
+			}
+		}
+	};
+	
+	/** Remote Server Info */
+	class FServerInfo
+	{
+	public:
+		FMessageAddress Address;
+		FString ServerName;
+		FString ProjectContentPath;
+		uint32 ProcessId = 0;
+		TMap<FString, FString> UserDataEntries;
+		TOptional<FDateTime> PingTimeout;
+		bool bDeviceProviderInstalled = false;
+
+		// Key ChannelName.
+		TMap<FString, TUniquePtr<FBroadcastChannelInfo>> BroadcastChannelInfosByName;
+		TMap<FString, TUniquePtr<FPlaybackChannelInfo>> PlaybackChannelInfosByName;
+
+		// Asset statuses (on disk) are not per-channel.
+		TMap<FSoftObjectPath, EAvaPlaybackAssetStatus> PlaybackAssetStatuses;
+
+		/** Map using "playback asset status request" key to keep track of pending playback asset status requests. */
+		TMap<FSoftObjectPath, FPendingPlaybackAssetStatusRequest> PendingPlaybackAssetStatusRequests;
+
+		/** Accumulate all the playback commands and send them all in one message on the next tick. */
+		TArray<FAvaPlaybackCommand> PendingPlaybackCommands;
+		
+		FBroadcastChannelInfo& GetOrCreateBroadcastChannelInfo(const FString& InChannelName)
+		{
+			const TUniquePtr<FBroadcastChannelInfo>* ChannelInfo = BroadcastChannelInfosByName.Find(InChannelName);
+			return ChannelInfo ? *(*ChannelInfo).Get() : *BroadcastChannelInfosByName.Add(InChannelName, MakeUnique<FBroadcastChannelInfo>()).Get();
+		}
+
+		const FBroadcastChannelInfo* GetBroadcastChannelInfo(const FString& InChannelName) const
+		{
+			const TUniquePtr<FBroadcastChannelInfo>* ChannelInfo = BroadcastChannelInfosByName.Find(InChannelName);
+			return ChannelInfo ? ChannelInfo->Get() : nullptr;
+		}
+
+		FPlaybackChannelInfo* GetOrCreatePlaybackChannelInfo(const FString& InChannelName)
+		{
+			const TUniquePtr<FPlaybackChannelInfo>* ChannelInfo = PlaybackChannelInfosByName.Find(InChannelName);
+			return ChannelInfo ? (*ChannelInfo).Get() : PlaybackChannelInfosByName.Add(InChannelName, MakeUnique<FPlaybackChannelInfo>()).Get();
+		}
+
+		const FPlaybackChannelInfo* GetPlaybackChannelInfo(const FString& InChannelName) const
+		{
+			const TUniquePtr<FPlaybackChannelInfo>* ChannelInfo = PlaybackChannelInfosByName.Find(InChannelName);
+			return ChannelInfo ? (*ChannelInfo).Get() : nullptr;
+		}
+		
+		TOptional<EAvaPlaybackStatus> GetInstanceStatus(
+			const FString& InChannelName, const FGuid& InInstanceId, const FSoftObjectPath& InAssetPath) const
+		{
+			const FPlaybackChannelInfo* ChannelInfo = GetPlaybackChannelInfo(InChannelName);
+			return ChannelInfo ? ChannelInfo->GetInstanceStatus(InInstanceId, InAssetPath) : TOptional<EAvaPlaybackStatus>();
+		}
+		
+		void SetInstanceStatus(
+			const FString& InChannelName, const FGuid& InInstanceId, const FSoftObjectPath& InAssetPath, EAvaPlaybackStatus InStatus)
+		{
+			if (FPlaybackChannelInfo* ChannelInfo = GetOrCreatePlaybackChannelInfo(InChannelName))
+			{
+				ChannelInfo->SetInstanceStatus(InInstanceId, InAssetPath, InStatus);
+			}
+		}
+
+		const FString* GetInstanceUserData(const FString& InChannelName, const FGuid& InInstanceId, const FSoftObjectPath& InAssetPath) const
+		{
+			const FPlaybackChannelInfo* ChannelInfo = GetPlaybackChannelInfo(InChannelName);
+			return ChannelInfo ? ChannelInfo->GetInstanceUserData(InInstanceId, InAssetPath) : nullptr;
+		}
+
+		void SetInstanceUserData(const FString& InChannelName, const FGuid& InInstanceId, const FSoftObjectPath& InAssetPath, const FString& InUserData)
+		{
+			if (FPlaybackChannelInfo* ChannelInfo = GetOrCreatePlaybackChannelInfo(InChannelName))
+			{
+				ChannelInfo->SetInstanceUserData(InInstanceId, InAssetPath, InUserData);
+			}
+		}
+		
+		TOptional<EAvaPlaybackAssetStatus> GetPlaybackAssetStatus(const FSoftObjectPath& InAssetPath) const
+		{
+			const EAvaPlaybackAssetStatus* FoundStatus = PlaybackAssetStatuses.Find(InAssetPath);
+			return FoundStatus ? TOptional(*FoundStatus) : TOptional<EAvaPlaybackAssetStatus>();
+		}
+		
+		void SetPlaybackAssetStatus(const FSoftObjectPath& InAssetPath, EAvaPlaybackAssetStatus InStatus)
+		{
+			PlaybackAssetStatuses.Add(InAssetPath, InStatus);
+		}
+		
+		/**
+		* Add a time by which a ping response should arrive
+		*/
+		void AddTimeout(const FDateTime& InNewTimeout)
+		{
+			if (!PingTimeout || InNewTimeout < PingTimeout.GetValue())
+			{
+				PingTimeout = InNewTimeout;
+			}
+		}
+
+		/**
+		 * Check if ping response arrived on time
+		*/
+		bool HasTimedOut(const FDateTime& InNow) const
+		{
+			return PingTimeout && InNow > PingTimeout.GetValue();
+		}
+
+		/**
+		 * Resets the ping timeout.
+		 * This is done when the ping response is received.
+		 */
+		void ResetPingTimeout()
+		{
+			PingTimeout.Reset();
+		}
+	};
+	TMap<FString, TSharedPtr<FServerInfo>> Servers;	// Key is ServerName
+	TArray<FMessageAddress> AllServerAddresses;
+
+	FServerInfo& GetOrCreateServerInfo(const FString& InServerName, const FMessageAddress& InSenderAddress, bool* bOutCreated = nullptr);
+	FServerInfo* GetServerInfo(const FString& InServerName)
+	{
+		const TSharedPtr<FServerInfo>* ServerInfo = Servers.Find(InServerName);
+		return ServerInfo ? (*ServerInfo).Get() : nullptr;
+	}
+	const FServerInfo* GetServerInfo(const FString& InServerName) const
+	{
+		const TSharedPtr<FServerInfo>* ServerInfo = Servers.Find(InServerName);
+		return ServerInfo ? (*ServerInfo).Get() : nullptr;
+	}
+	
+	void ForAllServers(TFunctionRef<void(FServerInfo& /*InServerInfo*/)> InFunction);
+	void ForAllServers(TFunctionRef<void(const FServerInfo& /*InServerInfo*/)> InFunction) const;
+	void OnServerAdded(const FServerInfo& InServerInfo);
+	void OnServerRemoved(const FServerInfo& InRemovedServer);
+	const FBroadcastChannelInfo* GetChannelInfo(const FString& InServerName, const FString& InChannelName) const;
+	
+	/** Map using "playback status request" key to keep track of pending playback status requests. */
+	TMap<FString, FDateTime> PendingPlaybackStatusRequests;
+	
+	static FString GetPlaybackStatusRequestKey(const FGuid& InInstanceId, const FSoftObjectPath& InAssetPath, const FString& InChannelName)
+	{
+		return InChannelName + InInstanceId.ToString() + InAssetPath.ToString();
+	}
+};
diff --git a/Engine/Plugins/Experimental/Avalanche/Source/AvalancheMedia/Private/Playback/AvaPlaybackClientDummy.cpp b/Engine/Plugins/Experimental/Avalanche/Source/AvalancheMedia/Private/Playback/AvaPlaybackClientDummy.cpp
new file mode 100644
--- /dev/null
+++ b/Engine/Plugins/Experimental/Avalanche/Source/AvalancheMedia/Private/Playback/AvaPlaybackClientDummy.cpp
@@ -0,0 +1,6 @@
+// Copyright Epic Games, Inc. All Rights Reserved.
+
+#include "AvaPlaybackClientDummy.h"
+
+const FString FAvaPlaybackClientDummy::EmptyString;
+const TArray<FString> FAvaPlaybackClientDummy::EmptyStringArray;
diff --git a/Engine/Plugins/Experimental/Avalanche/Source/AvalancheMedia/Private/Playback/AvaPlaybackClientDummy.h b/Engine/Plugins/Experimental/Avalanche/Source/AvalancheMedia/Private/Playback/AvaPlaybackClientDummy.h
new file mode 100644
--- /dev/null
+++ b/Engine/Plugins/Experimental/Avalanche/Source/AvalancheMedia/Private/Playback/AvaPlaybackClientDummy.h
@@ -0,0 +1,55 @@
+// Copyright Epic Games, Inc. All Rights Reserved.
+
+#pragma once
+
+#include "Playback/IAvaPlaybackClient.h"
+
+class FAvaPlaybackClientDummy : public IAvaPlaybackClient
+{
+	static const FString EmptyString;
+	static const TArray<FString> EmptyStringArray;
+public:
+	
+	//~ Begin IAvaPlaybackClient Interface
+	virtual int32 GetNumConnectedServers() const override { return 0; }
+	virtual TArray<FString> GetServerNames() const override { return TArray<FString>();}
+	virtual FMessageAddress GetServerAddress(const FString& InServerName) const override
+	{
+		FMessageAddress InvalidAddress;
+		InvalidAddress.Invalidate();
+		return InvalidAddress;
+	}
+	virtual bool HasServerUserData(const FString& InServerName, const FString& InKey) const override { return false;}
+	virtual const FString& GetServerUserData(const FString& InServerName, const FString& InKey) const override { return EmptyString; }
+	virtual bool HasUserData(const FString& InKey) const override { return false;}
+	virtual const FString& GetUserData(const FString& InKey) const override { return EmptyString; }
+	virtual void SetUserData(const FString& InKey, const FString& InData) override {}
+	virtual void RemoveUserData(const FString& InKey) override {}
+	virtual void BroadcastStatCommand(const FString& InCommand, bool bInBroadcastLocalState) override {}
+	virtual void RequestPlaybackAssetStatus(const FSoftObjectPath& InAssetPath, const FString& InChannelOrServerName, bool bInForceRefresh) override {}
+	virtual void RequestPlayback(const FGuid& InInstanceId, const FSoftObjectPath& InAssetPath, const FString& InChannelName, EAvaPlaybackAction InAction, const FString& InArguments) override {}
+	virtual void RequestAnimPlayback(const FGuid& InInstanceId, const FSoftObjectPath& InAssetPath, const FString& InChannelName, const FAvaPlaybackAnimPlaySettings& InAnimSettings) override {};
+	virtual void RequestAnimAction(const FGuid& InInstanceId, const FSoftObjectPath& InAssetPath, const FString& InChannelName, const FString& InAnimationName, EAvaPlaybackAnimAction InAction) override {};
+	virtual void RequestRemoteControlUpdate(const FGuid& InInstanceId, const FSoftObjectPath& InAssetPath, const FString& InChannelName, const FAvaPlayableRemoteControlValues& InRemoteControlValues) override {};
+	virtual void RequestPlayableTransitionStart(const FGuid& InTransitionId, TArray<FGuid>&& InEnterInstanceIds, TArray<FGuid>&& InPlayingInstanceIds, TArray<FGuid>&& InExitInstanceIds, TArray<FAvaPlayableRemoteControlValues>&& InEnterValues, const FName& InChannelName, EAvaPlayableTransitionFlags InTransitionFlags) override {}
+	virtual void RequestPlayableTransitionStop(const FGuid& InTransitionId, const FName& InChannelName) override {}
+	virtual void RequestBroadcast(const FString& InProfile, const FName& InChannel, const TArray<UMediaOutput*>& InRemoteMediaOutputs, EAvaBroadcastAction InAction) override {}
+	virtual bool IsMediaOutputRemoteFallback(const UMediaOutput* InMediaOutput) override { return false;}
+	virtual EAvaBroadcastIssueSeverity GetMediaOutputIssueSeverity(const FString& InServerName, const FString& InChannelName, const FGuid& InOutputGuid) const override { return EAvaBroadcastIssueSeverity::None;}
+	virtual const TArray<FString>& GetMediaOutputIssueMessages(const FString& InServerName, const FString& InChannelName, const FGuid& InOutputGuid) const override { return EmptyStringArray; }
+	virtual EAvaBroadcastOutputState GetMediaOutputState(const FString& InServerName, const FString& InChannelName, const FGuid& InOutputGuid) const override { return EAvaBroadcastOutputState::Error; }
+	virtual bool HasAnyServerOnlineForChannel(const FName& InChannelName) const override { return false;}
+   	virtual TOptional<EAvaPlaybackStatus> GetRemotePlaybackStatus(const FGuid& InInstanceId, const FSoftObjectPath& InAssetPath, const FString& InChannelName, const FString& InServerName) const override
+	{
+		return EAvaPlaybackStatus::Unknown;
+	}
+	virtual const FString* GetRemotePlaybackUserData(const FGuid& InInstanceId, const FSoftObjectPath& InAssetPath, const FString& InChannelName, const FString& InServerName) const override
+	{
+		return nullptr;
+	}
+	virtual TOptional<EAvaPlaybackAssetStatus> GetRemotePlaybackAssetStatus(const FSoftObjectPath& InAssetPath, const FString& InServerName) const override
+	{
+		return EAvaPlaybackAssetStatus::Unknown;
+	}
+	//~ End IAvaPlaybackClient Interface
+};
diff --git a/Engine/Plugins/Experimental/Avalanche/Source/AvalancheMedia/Private/Playback/AvaPlaybackMessages.h b/Engine/Plugins/Experimental/Avalanche/Source/AvalancheMedia/Private/Playback/AvaPlaybackMessages.h
new file mode 100644
--- /dev/null
+++ b/Engine/Plugins/Experimental/Avalanche/Source/AvalancheMedia/Private/Playback/AvaPlaybackMessages.h
@@ -0,0 +1,704 @@
+// Copyright Epic Games, Inc. All Rights Reserved.
+
+#pragma once
+
+#include "AvaMediaDefines.h"
+#include "Broadcast/AvaBroadcastDefines.h"
+#include "Broadcast/Channel/AvaBroadcastMediaOutputInfo.h"
+#include "Framework/AvaInstanceSettings.h"
+#include "PixelFormat.h"
+#include "Playable/AvaPlayableRemoteControlValues.h"
+#include "Playback/Nodes/Events/Actions/AvaPlaybackAnimations.h"
+#include "Viewport/AvaViewportQualitySettings.h"
+#include "AvaPlaybackMessages.generated.h"
+
+USTRUCT()
+struct FAvaPlaybackClientMessageBase
+{
+	GENERATED_BODY()
+
+	UPROPERTY()
+	FString ClientName;
+};
+
+/**
+ * Notes:
+ * - ServerName vs HostName
+ * The Host name is the physical computer device the playback server runs on.
+ * A Server name is what the playback server is called and it is not necessarily the Host name.
+ * If there are multiple servers on the same host, then each server must have a different name.
+ * The Server name is a unique key to identify a server, while the host name is not.
+ */
+USTRUCT()
+struct FAvaPlaybackServerMessageBase
+{
+	GENERATED_BODY()
+
+	UPROPERTY()
+	FString ServerName;
+};
+
+/**
+ * Request published by client to discover servers.
+ */
+USTRUCT()
+struct FAvaPlaybackPing : public FAvaPlaybackClientMessageBase
+{
+	GENERATED_BODY()
+
+	/** Indicate if the ping was generated on a timer (auto) or manually sent through a user command. */
+	UPROPERTY()
+	bool bAutoPing = true;
+
+	/** Defines the interval in seconds between each client's pings. */
+	UPROPERTY()
+	float PingIntervalSeconds = 0.0f;
+};
+
+/**
+ *	Response sent by server to client to be discovered.
+ */
+USTRUCT()
+struct FAvaPlaybackPong : public FAvaPlaybackServerMessageBase
+{
+	GENERATED_BODY()
+
+	/** Indicate if this is a reply from an auto ping. */
+	UPROPERTY()
+	bool bAutoPong = true;
+
+	/** Server may request client info in case of a reconnection event. */
+	UPROPERTY()
+	bool bRequestClientInfo = false;
+
+	/** Server's project content path. */
+	UPROPERTY()
+	FString ProjectContentPath;
+
+	/** Server's process id on the current host. */
+	UPROPERTY()
+	uint32 ProcessId = 0;
+};
+
+/**
+ *	Replicate server's log messages.
+ */
+USTRUCT()
+struct FAvaPlaybackLog : public FAvaPlaybackServerMessageBase
+{
+	GENERATED_BODY()
+
+	UPROPERTY()
+	FString Text;
+
+	UPROPERTY()
+	uint8 Verbosity = 0;
+
+	UPROPERTY()
+	FName Category;
+
+	UPROPERTY()
+	double Time = 0.0;
+};
+
+/**
+ * Request sent by client to replicate it's information on the destination server.
+ */
+USTRUCT()
+struct FAvaPlaybackUpdateClientInfo : public FAvaPlaybackClientMessageBase
+{
+	GENERATED_BODY()
+
+	UPROPERTY()
+	FString ComputerName;
+
+	UPROPERTY()
+	FString ProjectContentPath;
+
+	UPROPERTY()
+	uint32 ProcessId = 0;
+};
+
+/**
+ * Request sent by client to replicate it's user data on the destination server.
+ */
+USTRUCT()
+struct FAvaPlaybackUpdateClientUserData : public FAvaPlaybackClientMessageBase
+{
+	GENERATED_BODY()
+
+	UPROPERTY()
+	TMap<FString, FString> UserDataEntries;
+};
+
+/**
+ * Request sent by server to replicate it's user data to the client.
+ */
+USTRUCT()
+struct FAvaPlaybackUpdateServerUserData : public FAvaPlaybackServerMessageBase
+{
+	GENERATED_BODY()
+
+	UPROPERTY()
+	TMap<FString, FString> UserDataEntries;
+};
+
+USTRUCT()
+struct FAvaPlaybackStatCommand : public FAvaPlaybackClientMessageBase
+{
+	GENERATED_BODY()
+
+	UPROPERTY()
+	FString Command;
+
+	/**
+	 * If the local client couldn't get a viewport client to run the stat command, the
+	 * local state will not be reliable. In this case, the server state will be used.
+	 */
+	UPROPERTY()
+	bool bClientStateReliable = false;
+	
+	/**
+	 * Because the commands are toggles, there may be desync of states between client and server(s).
+	 * In order to correct that we also send a list of enabled states on the client so that servers can
+	 * ensure they have the same states enabled.
+	 */
+	UPROPERTY()
+	TArray<FString> ClientEnabledRuntimeStats;
+};
+
+/**
+ * Server's response when a stat command is received.
+ */
+USTRUCT()
+struct FAvaPlaybackStatStatus : public FAvaPlaybackServerMessageBase
+{
+	GENERATED_BODY()
+
+	/**
+	 * Replication of the same field from the command.
+	 * This is used by the client to replicate the server state in case
+	 * the command succeeded on the server and it had an unreliable state.
+	 */
+	UPROPERTY()
+	bool bClientStateReliable = false;
+	
+	/** Indicate if the server could execute the command. */
+	UPROPERTY()
+	bool bCommandSucceeded = false;
+
+	/** Resulting enabled stats from the command. */
+	UPROPERTY()
+	TArray<FString> EnabledRuntimeStats;
+};
+
+/**
+ * Request sent by client to obtain the device provider data from destination server.
+ * The device provider data contains all available devices and their configuration
+ * installed on the server.
+ */
+USTRUCT()
+struct FAvaPlaybackDeviceProviderDataRequest : public FAvaPlaybackClientMessageBase
+{
+	GENERATED_BODY()
+};
+
+/**
+ *	Request for the client to replicate it's Motion Design instance settings to the server.
+ **/
+USTRUCT()
+struct FAvaPlaybackInstanceSettingsUpdate : public FAvaPlaybackClientMessageBase
+{
+	GENERATED_BODY()
+
+	UPROPERTY()
+	FAvaInstanceSettings InstanceSettings;
+};
+
+UENUM()
+enum class EAvaPlaybackPackageEvent
+{
+	None,
+	PreSave,
+	PostSave,
+	AssetDeleted
+};
+
+/**
+ *	Message sent by the client to inform servers that a local package has been modified by an event.
+ **/
+USTRUCT()
+struct FAvaPlaybackPackageEvent : public FAvaPlaybackClientMessageBase
+{
+	GENERATED_BODY()
+
+	UPROPERTY()
+	FName PackageName;
+
+	UPROPERTY()
+	EAvaPlaybackPackageEvent Event = EAvaPlaybackPackageEvent::None;
+};
+
+/**
+ * Request by a client to obtain the status of an asset with the given path.
+ * Note: since this is an asset on disk, there is no channel name.
+ */
+USTRUCT()
+struct FAvaPlaybackAssetStatusRequest : public FAvaPlaybackClientMessageBase
+{
+	GENERATED_BODY()
+	
+	UPROPERTY()
+	FSoftObjectPath AssetPath;
+	
+	/**
+	 * If this is true, the server will perform a full asset status compare with the client.
+	 * Otherwise, the server may use locally cached status for the given asset.
+	 */
+	UPROPERTY()
+	bool bForceRefresh = false;
+};
+
+/** Response from the server to a playback asset status request. */
+USTRUCT()
+struct FAvaPlaybackAssetStatus : public FAvaPlaybackServerMessageBase
+{
+	GENERATED_BODY()
+	
+	UPROPERTY()
+	FSoftObjectPath AssetPath;
+	
+	UPROPERTY()
+	EAvaPlaybackAssetStatus Status = EAvaPlaybackAssetStatus::Unknown;
+};
+
+/**
+ * Playback command for a given asset.
+ */
+USTRUCT()
+struct FAvaPlaybackCommand
+{
+	GENERATED_BODY()
+
+	/** Playback's instance Id. */
+	UPROPERTY()
+	FGuid InstanceId;
+
+	/** Reference to the asset to playback. */
+	UPROPERTY()
+	FSoftObjectPath AssetPath;
+
+	/** Broadcast channel the action will apply to. */
+	UPROPERTY()
+	FString ChannelName;
+	
+	/** Action to do: Load, Play, Stop, Unload, etc. */
+	UPROPERTY()
+	EAvaPlaybackAction Action = EAvaPlaybackAction::None;
+
+	/** Command additional arguments. */
+	UPROPERTY()
+	FString Arguments;
+};
+
+/**
+ * Request by a client to execute a batch playback of commands on the connected servers.
+ */
+USTRUCT()
+struct FAvaPlaybackRequest : public FAvaPlaybackClientMessageBase
+{
+	GENERATED_BODY()
+	
+	UPROPERTY()
+	TArray<FAvaPlaybackCommand> Commands;
+};
+
+/** Response from the server to a playback request. */
+USTRUCT()
+struct FAvaPlaybackStatus : public FAvaPlaybackServerMessageBase
+{
+	GENERATED_BODY()
+
+	/** Playback's instance Id. */
+	UPROPERTY()
+	FGuid InstanceId;
+	
+	UPROPERTY()
+	FSoftObjectPath AssetPath;
+
+	UPROPERTY()
+	FString ChannelName;
+	
+	UPROPERTY()
+	EAvaPlaybackStatus Status = EAvaPlaybackStatus::Unknown;
+	
+	UPROPERTY()
+	bool bValidUserData = false;
+	
+	UPROPERTY()
+	FString UserData;
+};
+
+/**
+ * Message batching the status of many assets per channel/host.
+ * Used to reduce message overhead.
+ */
+USTRUCT()
+struct FAvaPlaybackStatuses : public FAvaPlaybackServerMessageBase
+{
+	GENERATED_BODY()
+	
+	UPROPERTY()
+	FString ChannelName;
+
+	/** Playback's instance Id. */
+	UPROPERTY()
+	TArray<FGuid> InstanceIds;
+	
+	UPROPERTY()
+	TArray<FSoftObjectPath> AssetPaths;
+
+	/** Batches all assets with the same status. */
+	UPROPERTY()
+	EAvaPlaybackStatus Status = EAvaPlaybackStatus::Unknown;
+};
+
+USTRUCT()
+struct FAvaPlaybackAnimActionInfo
+{
+	GENERATED_BODY()
+    
+	UPROPERTY()
+	FString AnimationName;
+
+	UPROPERTY()
+	EAvaPlaybackAnimAction AnimationAction = EAvaPlaybackAnimAction::None;
+};
+
+/**
+ * Request by a client to execute an animation on a playback asset
+ * on a connected server.
+ */
+USTRUCT()
+struct FAvaPlaybackAnimPlaybackRequest : public FAvaPlaybackClientMessageBase
+{
+	GENERATED_BODY()
+
+	/** Playback's instance Id. */
+	UPROPERTY()
+	FGuid InstanceId;
+	
+	UPROPERTY()
+	FSoftObjectPath AssetPath;
+	
+	UPROPERTY()
+	FString ChannelName;
+	
+	UPROPERTY()
+	TArray<FAvaPlaybackAnimPlaySettings> AnimPlaySettings;
+
+	UPROPERTY()
+	TArray<FAvaPlaybackAnimActionInfo> AnimActionInfos;
+};
+
+/** Server replication of playback sequence events. */
+USTRUCT()
+struct FAvaPlaybackSequenceEvent : public FAvaPlaybackServerMessageBase
+{
+	GENERATED_BODY()
+
+	/** Playback's instance Id. */
+	UPROPERTY()
+	FGuid InstanceId;
+	
+	UPROPERTY()
+	FSoftObjectPath AssetPath;
+
+	UPROPERTY()
+	FString ChannelName;
+	
+	UPROPERTY()
+	FString SequenceName;
+
+	UPROPERTY()
+	EAvaPlayableSequenceEventType EventType = EAvaPlayableSequenceEventType::None;
+};
+
+USTRUCT()
+struct FAvaPlaybackRemoteControlUpdateRequest : public FAvaPlaybackClientMessageBase
+{
+	GENERATED_BODY()
+
+	/** Playback's instance Id. */
+	UPROPERTY()
+	FGuid InstanceId;
+	
+	UPROPERTY()
+	FSoftObjectPath AssetPath;
+
+	UPROPERTY()
+	FString ChannelName;
+	
+	UPROPERTY()
+	FAvaPlayableRemoteControlValues RemoteControlValues;
+};
+
+USTRUCT()
+struct FAvaPlaybackTransitionStartRequest : public FAvaPlaybackClientMessageBase
+{
+	GENERATED_BODY()
+
+	UPROPERTY()
+	FString ChannelName;
+
+	UPROPERTY()
+	FGuid TransitionId;
+	
+	UPROPERTY()
+	TArray<FGuid> EnterInstanceIds;
+
+	UPROPERTY()
+	TArray<FGuid> PlayingInstanceIds;
+
+	UPROPERTY()
+	TArray<FGuid> ExitInstanceIds;
+
+	UPROPERTY()
+	TArray<FAvaPlayableRemoteControlValues> EnterValues;
+
+	UPROPERTY()
+	bool bUnloadDiscardedInstances = false;
+
+	/** See EAvaPlayableTransitionFlags. */
+	UPROPERTY()
+	uint8 TransitionFlags = 0;
+
+	EAvaPlayableTransitionFlags GetTransitionFlags() const { return static_cast<EAvaPlayableTransitionFlags>(TransitionFlags);}
+	void SetTransitionFlags(EAvaPlayableTransitionFlags InTransitionFlags) { TransitionFlags = static_cast<uint8>(InTransitionFlags);}
+
+};
+
+USTRUCT()
+struct FAvaPlaybackTransitionStopRequest : public FAvaPlaybackClientMessageBase
+{
+	GENERATED_BODY()
+
+	UPROPERTY()
+	FString ChannelName;
+
+	UPROPERTY()
+	FGuid TransitionId;
+	
+	UPROPERTY()
+	bool bUnloadDiscardedInstances = false;
+};
+
+USTRUCT()
+struct FAvaPlaybackTransitionEvent : public FAvaPlaybackServerMessageBase
+{
+	GENERATED_BODY()
+
+	UPROPERTY()
+	FString ChannelName;
+
+	UPROPERTY()
+	FGuid TransitionId;
+	
+	UPROPERTY()
+	FGuid InstanceId;
+
+	UPROPERTY()
+	uint8 EventFlags = static_cast<uint8>(EAvaPlayableTransitionEventFlags::None);
+
+	// Because the enum is used as flags, we need to convert to uint8 manually. (Can't use TEnumAsByte for this apparently.)
+	EAvaPlayableTransitionEventFlags GetEventFlags() const { return static_cast<EAvaPlayableTransitionEventFlags>(EventFlags);}
+	void SetEventFlags(EAvaPlayableTransitionEventFlags InFlags) { EventFlags = static_cast<uint8>(InFlags);}
+};
+
+class UMediaOutput;
+
+/**
+ * Encapsulates a UMediaOutput object in binary form to be able to serialize it
+ * nested within the broadcast request.
+ */
+USTRUCT()
+struct FAvaBroadcastOutputData
+{
+	GENERATED_BODY()
+
+	UPROPERTY()
+	FAvaBroadcastMediaOutputInfo OutputInfo;
+
+	UPROPERTY()
+	TSoftClassPtr<UMediaOutput> MediaOutputClass;
+	
+	UPROPERTY()
+	TArray<uint8> SerializedData;
+	
+	UPROPERTY()
+	uint64 ObjectFlags = 0;
+
+	EObjectFlags GetObjectFlags() const { return static_cast<EObjectFlags>(ObjectFlags); }
+};
+
+/** Request by a client to execute an action on broadcast channel(s). */
+USTRUCT()
+struct FAvaBroadcastRequest : public FAvaPlaybackClientMessageBase
+{
+	GENERATED_BODY()
+	
+	/** Profile to select for broadcast. */
+	UPROPERTY()
+	FString Profile;
+
+	/** Channel name. If empty, command is for all channels. */
+	UPROPERTY()
+	FString Channel;
+
+	UPROPERTY()
+	TArray<FAvaBroadcastOutputData> MediaOutputs;
+
+	/** Action to do on channel: Start, Stop, Disable, etc. */
+	UPROPERTY()
+	EAvaBroadcastAction Action = EAvaBroadcastAction::None;
+};
+
+/** Request by a client to update a broadcast channel's settings. */
+USTRUCT()
+struct FAvaBroadcastChannelSettingsUpdate : public FAvaPlaybackClientMessageBase
+{
+	GENERATED_BODY()
+	
+	/** Profile to select for broadcast. */
+	UPROPERTY()
+	FString Profile;
+
+	/** Channel name. If empty, command is for all channels. */
+	UPROPERTY()
+	FString Channel;
+
+	UPROPERTY()
+	FAvaViewportQualitySettings QualitySettings;
+};
+
+/**
+ * Client sends this request to have a status of broadcast.
+ *
+ * This is used by the client when it connects to a server which
+ * might be already broadcasting. This happens if the client
+ * disconnects and reconnects. It needs to query the full
+ * status of the server upon connection.
+ */
+USTRUCT()
+struct FAvaBroadcastStatusRequest : public FAvaPlaybackClientMessageBase
+{
+	GENERATED_BODY()
+
+	UPROPERTY()
+	bool bIncludeMediaOutputData = false;
+};
+
+USTRUCT()
+struct FAvaBroadcastSettings
+{
+	GENERATED_BODY()
+
+	/** Specifies the background clear color for the channel. */
+	UPROPERTY()
+	FLinearColor ChannelClearColor = FLinearColor::Black;
+
+	/** Pixel format used if no media output has specific format requirement. */
+	UPROPERTY()
+	TEnumAsByte<EPixelFormat> ChannelDefaultPixelFormat = EPixelFormat::PF_B8G8R8A8;
+
+	/** Resolution used if no media output has specific resolution requirement. */
+	UPROPERTY()
+	FIntPoint ChannelDefaultResolution = FIntPoint(1920, 1080);
+
+	/**
+	 * Enables drawing the placeholder widget when there is no Motion Design asset playing.
+	 * If false, the channel is cleared to the background color.
+	 */
+	UPROPERTY()
+	bool bDrawPlaceholderWidget = false;
+	
+	/** Specify a place holder widget to render when no Motion Design asset is playing. */
+	UPROPERTY()
+	FSoftObjectPath PlaceholderWidgetClass;
+};
+
+/**
+ *	Request for the client to replicate it's broadcast settings to the server.
+ **/
+USTRUCT()
+struct FAvaBroadcastSettingsUpdate : public FAvaPlaybackClientMessageBase
+{
+	GENERATED_BODY()
+
+	UPROPERTY()
+	FAvaBroadcastSettings BroadcastSettings;
+};
+
+/**
+ * Encapsulate the complete status of a media output.
+ */
+USTRUCT()
+struct FAvaBroadcastOutputStatus
+{
+	GENERATED_BODY()
+
+	UPROPERTY()
+	EAvaBroadcastOutputState MediaOutputState = EAvaBroadcastOutputState::Invalid;
+	
+	UPROPERTY()
+	EAvaBroadcastIssueSeverity MediaIssueSeverity = EAvaBroadcastIssueSeverity::None;
+
+	UPROPERTY()
+	TArray<FString> MediaIssueMessages;
+};
+
+/**
+ * Server's response when channel's broadcast status changes
+ */
+USTRUCT()
+struct FAvaBroadcastStatus : public FAvaPlaybackServerMessageBase
+{
+	GENERATED_BODY()
+	
+	/** Profile to select for broadcast. */
+	UPROPERTY()
+	FString Profile;
+	
+	/** Channel name. If empty, command is for all channels of the profile. */
+	UPROPERTY()
+	FString ChannelName;
+
+	/** Index of the channel in the profile. */
+	UPROPERTY()
+	int32 ChannelIndex = 0;
+
+	/** Number of channels in the profile. */
+	UPROPERTY()
+	int32 NumChannels = 0;
+
+	UPROPERTY()
+	EAvaBroadcastChannelState ChannelState = EAvaBroadcastChannelState::Idle;
+
+	UPROPERTY()
+	EAvaBroadcastIssueSeverity ChannelIssueSeverity = EAvaBroadcastIssueSeverity::None;
+
+	UPROPERTY()
+	TMap<FGuid, FAvaBroadcastOutputStatus> MediaOutputStatuses;
+	
+	/**
+	 * This is required to know if the MediaOutputs is empty because the data
+	 * was not included or if it is because the data was included and there is
+	 * actually no outputs defined on that channel.
+	 */
+	UPROPERTY()
+	bool bIncludeMediaOutputData = false;
+	
+	/** Remark: this is only included if bIncludeMediaOutputData == true in the request. */
+	UPROPERTY()
+	TArray<FAvaBroadcastOutputData> MediaOutputs;
+};
diff --git a/Engine/Plugins/Experimental/Avalanche/Source/AvalancheMedia/Private/Playback/AvaPlaybackServer.h b/Engine/Plugins/Experimental/Avalanche/Source/AvalancheMedia/Private/Playback/AvaPlaybackServer.h
new file mode 100644
--- /dev/null
+++ b/Engine/Plugins/Experimental/Avalanche/Source/AvalancheMedia/Private/Playback/AvaPlaybackServer.h
@@ -0,0 +1,364 @@
+// Copyright Epic Games, Inc. All Rights Reserved.
+
+#pragma once
+
+#include "Broadcast/IAvaBroadcastSettings.h"
+#include "Framework/AvaInstanceSettings.h"
+#include "Playback/IAvaPlaybackServer.h"
+#include "IMessageContext.h"
+#include "MessageEndpoint.h"
+#include "Playback/AvaPlaybackManager.h"
+#include "Playback/AvaPlaybackMessages.h"
+#include "Templates/SharedPointer.h"
+
+class FAvaPlaybackInstance;
+class FAvaPlaybackSyncManager;
+class UAvaPlaybackServerTransition;
+struct FAvaPlaybackAssetSyncStatusReceivedParams;
+
+DECLARE_LOG_CATEGORY_EXTERN(LogAvaPlaybackServer, Log, All);
+
+/**
+ * Playback (and Broadcast) Server
+ *
+ * The playback server implements the commands for broadcast (channels and outputs) and playback.
+ * Playback assets are either "playables" or "playback graphs", however the playback server will
+ * create playback graphs for everything. If the asset to play is a "playable" (e.g. a level),
+ * it will create a transient playback graph for it.
+ *
+ * A lot of the playback commands are geared toward running a playback graph with a single playable node,
+ * mostly because this system is only used with rundowns or playback graphs run on the client side.
+ * So, it will emulate a client-side playable with a local transient playback graph with one player node.
+ * The use case of running a more complex playback graph asset on the server side has not occured yet.
+ */
+class FAvaPlaybackServer : public TSharedFromThis<FAvaPlaybackServer>, public IAvaPlaybackServer
+{
+public:
+	FAvaPlaybackServer();
+	virtual ~FAvaPlaybackServer() override;
+
+	void Init(const FString& InAssignedServerName);
+	
+	/**
+	 * Returns a list of channel names from all the playing playback instances.
+	 * Optionally filter for a given asset.
+	 */
+	TArray<FString> GetAllChannelsFromPlayingPlaybacks(const FSoftObjectPath& InAssetPath = FSoftObjectPath()) const;
+	
+	/**
+	 *	Indicate the manager is in a shutdown sequence and will force game instances to destroy worlds right away.
+	*/
+	void StartShuttingDown();
+	
+	//~ Begin IAvaPlaybackServer Interface
+	
+	virtual TArray<FPlaybackInstanceReference> StopPlaybacks(const FString& InChannelName = FString(), const FSoftObjectPath& InAssetPath = FSoftObjectPath(), bool bInUnload = true) override;
+	virtual TArray<FPlaybackInstanceReference> StartPlaybacks() override;
+	
+	virtual void StartBroadcast() override;
+	virtual void StopBroadcast() override;
+	
+	virtual const FString& GetName() const override { return ServerName;}
+
+	virtual bool HasUserData(const FString& InKey) const override { return UserDataEntries.Contains(InKey);}
+	
+	virtual const FString& GetUserData(const FString& InKey) const override;
+
+	virtual void SetUserData(const FString& InKey, const FString& InData) override;
+
+	virtual void RemoveUserData(const FString& InKey) override;
+
+	virtual TArray<FString> GetClientNames() const override;
+
+	virtual FMessageAddress GetClientAddress(const FString& InClientName) const override;
+
+	virtual bool HasClientUserData(const FString& InClientName, const FString& InKey) const override;
+
+	virtual const FString& GetClientUserData(const FString& InClientName, const FString& InKey) const override;
+
+	virtual const IAvaBroadcastSettings* GetBroadcastSettings() const override;
+
+	virtual const FAvaInstanceSettings* GetAvaInstanceSettings() const override;
+	
+	virtual const FAvaPlaybackManager& GetPlaybackManager() const override { check(Manager); return *Manager; }
+	virtual FAvaPlaybackManager& GetPlaybackManager() override { check(Manager); return *Manager; }
+
+	//~ End IAvaPlaybackServer Interface
+	
+	TSharedPtr<FAvaPlaybackInstance> FindActivePlaybackInstance(const FGuid& InInstanceId) const
+	{
+		const TSharedPtr<FAvaPlaybackInstance>* FoundInstance = ActivePlaybackInstances.Find(InInstanceId);
+		return FoundInstance ? *FoundInstance : TSharedPtr<FAvaPlaybackInstance>();
+	}
+
+	bool RemoveActivePlaybackInstance(const FGuid& InInstanceId)
+	{
+		return ActivePlaybackInstances.Remove(InInstanceId) > 0;	
+	}
+
+	bool RemovePlaybackInstanceTransition(const FGuid& InTransitionId);
+	
+	void SendPlayableTransitionEvent(
+		const FGuid& InTransitionId, const FGuid& InInstanceId, EAvaPlayableTransitionEventFlags InFlags,
+		const FName& InChannelName, const FString& InClientName);
+	
+public:
+	// Message handlers
+	void HandlePlaybackPing(const FAvaPlaybackPing& InMessage, const TSharedRef<IMessageContext, ESPMode::ThreadSafe>& InContext);
+	void HandleUpdateClientUserData(const FAvaPlaybackUpdateClientUserData& InMessage, const TSharedRef<IMessageContext, ESPMode::ThreadSafe>& InContext);
+	void HandleStatCommand(const FAvaPlaybackStatCommand& InMessage, const TSharedRef<IMessageContext, ESPMode::ThreadSafe>& InContext);
+	void HandleDeviceProviderDataRequest(const FAvaPlaybackDeviceProviderDataRequest& InMessage, const TSharedRef<IMessageContext, ESPMode::ThreadSafe>& InContext);
+	void HandleUpdateClientInfo(const FAvaPlaybackUpdateClientInfo& InMessage, const TSharedRef<IMessageContext, ESPMode::ThreadSafe>& InContext);
+	void HandleAvaInstanceSettingsUpdate(const FAvaPlaybackInstanceSettingsUpdate& InMessage, const TSharedRef<IMessageContext, ESPMode::ThreadSafe>& InContext);
+	void HandlePackageEvent(const FAvaPlaybackPackageEvent& InMessage, const TSharedRef<IMessageContext, ESPMode::ThreadSafe>& InContext);
+	void HandlePlaybackAssetStatusRequest(const FAvaPlaybackAssetStatusRequest& InMessage, const TSharedRef<IMessageContext, ESPMode::ThreadSafe>& InContext);
+	void HandlePlaybackRequest(const FAvaPlaybackRequest& InMessage, const TSharedRef<IMessageContext, ESPMode::ThreadSafe>& InContext);
+	void HandleAnimPlaybackRequest(const FAvaPlaybackAnimPlaybackRequest& InMessage, const TSharedRef<IMessageContext, ESPMode::ThreadSafe>& InContext);
+	void HandleRemoteControlUpdateRequest(const FAvaPlaybackRemoteControlUpdateRequest& InMessage, const TSharedRef<IMessageContext, ESPMode::ThreadSafe>& InContext);
+	void HandlePlayableTransitionStartRequest(const FAvaPlaybackTransitionStartRequest& InMessage, const TSharedRef<IMessageContext, ESPMode::ThreadSafe>& InContext);
+	void HandlePlayableTransitionStopRequest(const FAvaPlaybackTransitionStopRequest& InMessage, const TSharedRef<IMessageContext, ESPMode::ThreadSafe>& InContext);
+	void HandleBroadcastSettingsUpdate(const FAvaBroadcastSettingsUpdate& InMessage, const TSharedRef<IMessageContext, ESPMode::ThreadSafe>& InContext);
+	void HandleBroadcastRequest(const FAvaBroadcastRequest& InMessage, const TSharedRef<IMessageContext, ESPMode::ThreadSafe>& InContext);
+	void HandleBroadcastChannelSettingsUpdate(const FAvaBroadcastChannelSettingsUpdate& InMessage, const TSharedRef<IMessageContext, ESPMode::ThreadSafe>& InContext);
+	void HandleBroadcastStatusRequest(const FAvaBroadcastStatusRequest& InMessage, const TSharedRef<IMessageContext, ESPMode::ThreadSafe>& InContext);
+
+	/** Returns underlying message bus endpoint address id  */
+	FString GetMessageEndpointAddressId() const;
+
+protected:
+	void Tick();
+
+	void RegisterCommands();
+	
+	// Command handlers
+	void StartPlaybackCommand(const TArray<FString>& InArgs);
+	void StopPlaybackCommand(const TArray<FString>& InArgs);
+	void StartBroadcastCommand(const TArray<FString>& InArgs);
+	void StopBroadcastCommand(const TArray<FString>& InArgs);
+	void SetUserDataCommand(const TArray<FString>& InArgs);
+	void ShowStatusCommand(const TArray<FString>& InArgs);
+	
+	// Event handlers
+	void OnAvaMediaSettingsChanged(UObject*, struct FPropertyChangedEvent&);
+	void OnChannelChanged(const FAvaBroadcastOutputChannel& InChannel, EAvaBroadcastChannelChange InChange);
+	void OnMediaOutputStateChanged(const FAvaBroadcastOutputChannel& InChannel, const UMediaOutput* InMediaOutput);
+	void OnAvaAssetSyncStatusReceived(const FAvaPlaybackAssetSyncStatusReceivedParams& InParams);
+	void OnPlaybackInstanceInvalidated(const FAvaPlaybackInstance& InPlaybackInstance);
+	void OnPlaybackInstanceStatusChanged(const FAvaPlaybackInstance& InPlaybackInstance);
+	void OnPlaybackAssetRemoved(const FSoftObjectPath& InAssetPath);
+	void OnPlayableSequenceEvent(UAvaPlayable* InPlayable, const FName& SequenceName, EAvaPlayableSequenceEventType InEventType);
+	
+	void ApplyAvaMediaSettings();
+
+	template<typename MessageType>
+	void FillServerInfo(MessageType* InMessage)
+	{
+		InMessage->ServerName = ServerName;
+	}
+	
+	template<typename MessageType>
+	void SendResponse(MessageType* InMessage, const FMessageAddress& InRecipient, EMessageFlags flags = EMessageFlags::None)
+	{
+		FillServerInfo(InMessage);
+		MessageEndpoint->Send(InMessage, MessageType::StaticStruct(), flags, nullptr,
+			TArrayBuilder<FMessageAddress>().Add(InRecipient), FTimespan::Zero(), FDateTime::MaxValue());
+	}
+
+	template<typename MessageType>
+	void SendResponse(MessageType* InMessage, const TArray<FMessageAddress>& InRecipients, EMessageFlags flags = EMessageFlags::None)
+	{
+		FillServerInfo(InMessage);
+		MessageEndpoint->Send(InMessage, MessageType::StaticStruct(), flags, nullptr,
+			InRecipients, FTimespan::Zero(), FDateTime::MaxValue());
+	}
+
+	void SendUserDataUpdate(const TArray<FMessageAddress>& InRecipients);
+
+	void SendChannelStatusUpdate(const FString& InChannelName, const FAvaBroadcastOutputChannel& InChannel, const FMessageAddress& InSender, bool bInIncludeOutputData = false);
+	void SendAllChannelStatusUpdate(const FMessageAddress& InSender, bool bInIncludeOutputData = false);
+	void SendLogMessage(const TCHAR* V, ELogVerbosity::Type Verbosity, const FName& Category, double Time);
+
+	// Playback Commands
+	void ExecutePendingPlaybackCommands();
+
+	TSharedPtr<FAvaPlaybackInstance> GetOrLoadPlaybackInstance(const FGuid& InInstanceId, const FString& InChannelName, const FSoftObjectPath& InAssetPath);
+	void LoadPlayback(const FMessageAddress& InReplyToAddress, const FGuid& InInstanceId, const FString& InChannelName, const FSoftObjectPath& InAssetPath);
+	void StartPlayback(const FMessageAddress& InReplyToAddress, const FGuid& InInstanceId, const FString& InChannelName, const FSoftObjectPath& InAssetPath);
+	void StopPlayback(const FMessageAddress& InReplyToAddress, const FGuid& InInstanceId, const FString& InChannelName, const FSoftObjectPath& InAssetPath);	
+	void UnloadPlayback(const FMessageAddress& InReplyToAddress, const FGuid& InInstanceId, const FString& InChannelName, const FSoftObjectPath& InAssetPath);
+	void SetPlaybackUserData(const FMessageAddress& InReplyToAddress, const FGuid& InInstanceId, const FString& InUserData);
+	void SendPlaybackUserData(const FMessageAddress& InReplyToAddress, const FGuid& InInstanceId);
+	void SendPlaybackStatus(const FMessageAddress& InReplyToAddress, const FGuid& InInstanceId, const FString& InChannelName, const FSoftObjectPath& InAssetPath);
+	
+	void SendPlaybackStatus(const FMessageAddress& InSendTo, const FGuid& InInstanceId, const FString& InChannelName, const FSoftObjectPath& InAssetPath, EAvaPlaybackStatus InStatus);
+	void SendPlaybackStatus(const TArray<FMessageAddress>& InRecipients, const FGuid& InInstanceId, const FString& InChannelName, const FSoftObjectPath& InAssetPath, EAvaPlaybackStatus InStatus);
+	void SendPlaybackStatuses(const FMessageAddress& InSendTo, const FString& InChannelName, const TArray<FPlaybackInstanceReference>& InInstances, EAvaPlaybackStatus InStatus);
+	void SendAllPlaybackStatusesForChannelAndAssetPath(const FMessageAddress& InSendTo, const FString& InChannelName, const FSoftObjectPath& InAssetPath);
+	void SendPlaybackAssetStatus(const FMessageAddress& InSendTo, const FSoftObjectPath& InAssetPath, EAvaPlaybackAssetStatus InStatus);
+
+	EAvaPlaybackStatus GetUnloadedPlaybackStatus(const FSoftObjectPath& InAssetPath)
+	{
+		return Manager->GetUnloadedPlaybackStatus(InAssetPath);
+	}
+
+	bool UpdateChannelOutputConfig(FAvaBroadcastOutputChannel& InChannel, const TArray<FAvaBroadcastOutputData>& InMediaOutputs, bool bInRefreshState);
+
+private:
+	FString ComputerName;
+	FString ServerName;
+	FString ProjectContentPath;
+	uint32 ProcessId = 0;
+	TMap<FString, FString> UserDataEntries;
+	
+	/** Holds the messaging endpoint. */
+	TSharedPtr<FMessageEndpoint, ESPMode::ThreadSafe> MessageEndpoint;
+
+	TArray<IConsoleObject*> ConsoleCommands;
+
+	/** The playback server has its own playback manager to not interfere with the local one. */
+	TSharedPtr<FAvaPlaybackManager> Manager;
+
+	// This is used to block sending status update from the channel event while the channels
+	// are refreshing state on media output state changes. This avoid sending spurious channel states
+	// while the update is not completed for all outputs.
+	bool bBlockChannelStatusUpdate = false;
+
+	struct FPendingPlaybackCommand
+	{
+		FMessageAddress ReplyTo;
+		FAvaPlaybackCommand Command;
+	};
+
+	/** Accumulate all the playback commands and execute them all in one batch on the next tick. */
+	TArray<FPendingPlaybackCommand> PendingPlaybackCommands;
+
+	/** Keep an map of active instances per id for fast lookup. */
+	TMap<FGuid, TSharedPtr<FAvaPlaybackInstance>> ActivePlaybackInstances;
+	
+	class FServerPlaybackInstanceTransitionCollection;
+	TUniquePtr<FServerPlaybackInstanceTransitionCollection> PlaybackInstanceTransitions;
+	
+	/**
+	 * Routes the log messages to the server for replication over the message bus.
+	 */
+	class FReplicationOutputDevice : public FOutputDevice
+	{
+	public:
+		FReplicationOutputDevice(FAvaPlaybackServer* InServer);
+		virtual ~FReplicationOutputDevice() override;
+
+		/** Set the minimum verbosity that will be replicated to the client. */
+		void SetVerbosityThreshold(ELogVerbosity::Type InVerbosityThreshold);
+
+	protected:
+		//~ FOutputDevice interface.
+		virtual void Serialize(const TCHAR* InText, ELogVerbosity::Type InVerbosity, const FName& InCategory) override;
+		virtual void Serialize(const TCHAR* InText, ELogVerbosity::Type InVerbosity, const FName& InCategory, double InTime) override;
+
+	private:
+		FAvaPlaybackServer* Server;
+		ELogVerbosity::Type VerbosityThreshold = ELogVerbosity::Type::NoLogging; 
+	};
+	
+	TUniquePtr<FReplicationOutputDevice> ReplicationOutputDevice;
+
+	TOptional<ELogVerbosity::Type> LogReplicationVerbosityFromCommandLine;
+
+	class FClientBroadcastSettings : public IAvaBroadcastSettings
+	{
+	public:
+		FAvaBroadcastSettings Settings;
+		virtual ~FClientBroadcastSettings() override = default;
+
+		// IAvaBroadcastSettings
+		virtual const FLinearColor& GetChannelClearColor() const override { return Settings.ChannelClearColor; }
+		virtual EPixelFormat GetDefaultPixelFormat() const override { return Settings.ChannelDefaultPixelFormat; }
+		virtual const FIntPoint& GetDefaultResolution() const override { return Settings.ChannelDefaultResolution; }
+		virtual bool IsDrawPlaceholderWidget() const override { return Settings.bDrawPlaceholderWidget; }
+		virtual const FSoftObjectPath& GetPlaceholderWidgetClass() const override { return Settings.PlaceholderWidgetClass; }
+		// ~IAvaBroadcastSettings
+	};
+	
+	/** Remote Client Info */
+	class FClientInfo
+	{
+	public:
+		FMessageAddress Address;
+		FString ClientName;
+		FString ComputerName;
+		FString ProjectContentPath;
+		uint32 ProcessId = 0;
+		TMap<FString, FString> UserDataEntries;
+		FClientBroadcastSettings BroadcastSettings;
+		FAvaInstanceSettings AvaInstanceSettings;
+
+		bool bClientInfoReceived = false;
+	
+		/**
+		 * The status of the server's asset compared to the remote is tracked
+		 * on the server side. For now, the remote client is considered to have the
+		 * reference asset. But in the future, a hub may hold the reference asset instead.
+		 */
+		TSharedPtr<FAvaPlaybackSyncManager> MediaSyncManager;
+
+		TOptional<FDateTime> PingTimeout;
+		
+		FClientInfo(const FMessageAddress& InClientAddress, const FString& InClientName);
+		~FClientInfo();
+
+		
+		/**
+		* Add a time by which a new ping should arrive
+		*/
+		void AddTimeout(const FDateTime& InNewTimeout)
+		{
+			if (!PingTimeout || InNewTimeout < PingTimeout.GetValue())
+			{
+				PingTimeout = InNewTimeout;
+			}
+		}
+
+		/**
+		 * Check if new ping arrived on time
+		*/
+		bool HasTimedOut(const FDateTime& InNow) const
+		{
+			return PingTimeout && InNow > PingTimeout.GetValue();
+		}
+
+		/**
+		 * Resets the ping timeout.
+		 * This is done when the ping response is received.
+		 */
+		void ResetPingTimeout()
+		{
+			PingTimeout.Reset();
+		}
+	};
+	TMap<FString, TSharedPtr<FClientInfo>> Clients;	// Key is ClientName
+
+	FClientInfo& GetOrCreateClientInfo(const FString& InClientName, const FMessageAddress& InClientAddress);
+	FClientInfo* GetClientInfo(const FMessageAddress& InClientAddress) const;
+	FClientInfo* GetClientInfo(const FString& InClientName) const
+	{
+		const TSharedPtr<FClientInfo>* ClientInfoPtr = Clients.Find(InClientName);
+		return ClientInfoPtr ? ClientInfoPtr->Get() : nullptr;
+	}
+	const FString& GetClientNameSafe(const FMessageAddress& InClientAddress) const;
+	const FMessageAddress& GetClientAddressSafe(const FString& InClientName) const;
+	TArray<FMessageAddress> GetAllClientAddresses(bool bInExcludeClientOnLocalProcess = false) const;
+	void RemoveDeadClients(const FDateTime& InCurrentTime);
+	void OnClientAdded(const FClientInfo& InClientInfo);
+	void OnClientRemoved(const FClientInfo& InRemovedClient);
+	/** A local client will run from the same computer and project folder as the current server. */
+	bool IsLocalClient(const FClientInfo& ClientInfo) const;
+	bool IsLocalClient(const FMessageAddress& InClientAddress) const
+	{
+		const FClientInfo* ClientInfo = GetClientInfo(InClientAddress);
+		return ClientInfo && IsLocalClient(*ClientInfo);
+	}
+	bool IsClientOnLocalProcess(const FClientInfo& ClientInfo) const
+	{
+		return ClientInfo.ComputerName == ComputerName && ClientInfo.ProcessId == ProcessId;
+	}
+};
diff --git a/Engine/Plugins/Experimental/Avalanche/Source/AvalancheMedia/Private/Playback/Http/AvaPlaybackHttpServer.cpp b/Engine/Plugins/Experimental/Avalanche/Source/AvalancheMedia/Private/Playback/Http/AvaPlaybackHttpServer.cpp
--- a/Engine/Plugins/Experimental/Avalanche/Source/AvalancheMedia/Private/Playback/Http/AvaPlaybackHttpServer.cpp
+++ b/Engine/Plugins/Experimental/Avalanche/Source/AvalancheMedia/Private/Playback/Http/AvaPlaybackHttpServer.cpp
@@ -51,7 +51,7 @@ void FAvaPlaybackHttpServer::Start(int32 InPortToUse)
 
 void FAvaPlaybackHttpServer::RegisterRoutes()
 {
-	const TSharedPtr<FAvaPlaybackServer> MediaPlaybackServer = FModuleManager::GetModulePtr<FAvaMediaModule>(UE_MODULE_NAME)->GetMediaPlaybackServer();
+	const TSharedPtr<FAvaPlaybackServer> MediaPlaybackServer = FModuleManager::GetModulePtr<FAvaMediaModule>(UE_MODULE_NAME)->GetPlaybackServerInternal();
 	check(MediaPlaybackServer.IsValid());
 	
 	// Map as per FAvaPlaybackServer::Init(const FString& AssignedServerName)
diff --git a/Engine/Plugins/Experimental/Avalanche/Source/AvalancheMedia/Private/Playback/Nodes/AvaPlaybackNodeCombiner.h b/Engine/Plugins/Experimental/Avalanche/Source/AvalancheMedia/Private/Playback/Nodes/AvaPlaybackNodeCombiner.h
new file mode 100644
--- /dev/null
+++ b/Engine/Plugins/Experimental/Avalanche/Source/AvalancheMedia/Private/Playback/Nodes/AvaPlaybackNodeCombiner.h
@@ -0,0 +1,36 @@
+// Copyright Epic Games, Inc. All Rights Reserved.
+
+#pragma once
+
+#include "Playback/Nodes/AvaPlaybackNode.h"
+#include "AvaPlaybackNodeCombiner.generated.h"
+
+class FText;
+struct FAvaPlaybackChannelParameters;
+
+#if WITH_EDITOR
+class FName;
+#endif
+
+UCLASS()
+class UAvaPlaybackNodeCombiner : public UAvaPlaybackNode
+{
+	GENERATED_BODY()
+	
+	virtual FText GetNodeDisplayNameText() const override;
+	
+	virtual int32 GetMinChildNodes() const override { return 1; }
+	virtual int32 GetMaxChildNodes() const override { return MaxAllowedChildNodes; }
+	
+	virtual void CreateStartingConnectors() override;
+	
+	virtual void Tick(float DeltaTime, FAvaPlaybackChannelParameters& ChannelParameters) override;
+
+#if WITH_EDITOR
+	virtual FName GetInputPinName(int32 InputPinIndex) const override;
+#endif
+	
+protected:
+	UPROPERTY(EditAnywhere, Category = "Motion Design")
+	TArray<bool> EnabledIndices;	
+};
diff --git a/Engine/Plugins/Experimental/Avalanche/Source/AvalancheMedia/Private/Playback/Nodes/AvaPlaybackNodeSwitcher.h b/Engine/Plugins/Experimental/Avalanche/Source/AvalancheMedia/Private/Playback/Nodes/AvaPlaybackNodeSwitcher.h
new file mode 100644
--- /dev/null
+++ b/Engine/Plugins/Experimental/Avalanche/Source/AvalancheMedia/Private/Playback/Nodes/AvaPlaybackNodeSwitcher.h
@@ -0,0 +1,36 @@
+// Copyright Epic Games, Inc. All Rights Reserved.
+
+#pragma once
+
+#include "Playback/Nodes/AvaPlaybackNode.h"
+#include "AvaPlaybackNodeSwitcher.generated.h"
+
+class FName;
+class FText;
+struct FAvaPlaybackChannelParameters;
+
+UCLASS()
+class UAvaPlaybackNodeSwitcher : public UAvaPlaybackNode
+{
+	GENERATED_BODY()
+	
+	virtual FText GetNodeDisplayNameText() const override;
+	
+	virtual int32 GetMinChildNodes() const override { return 1; }
+	virtual int32 GetMaxChildNodes() const override { return MaxAllowedChildNodes; }
+	
+	virtual void CreateStartingConnectors() override;
+	
+	virtual void Tick(float DeltaTime, FAvaPlaybackChannelParameters& ChannelParameters) override;
+
+#if WITH_EDITOR
+	virtual FName GetInputPinName(int32 InputPinIndex) const override;
+#endif
+	
+protected:
+	UPROPERTY(EditAnywhere, Category = "Motion Design")
+	bool bEnabled = true;
+	
+	UPROPERTY(EditAnywhere, Category = "Motion Design", meta=(ClampMin = "0"))
+	int32 SelectedIndex = 0;	
+};
diff --git a/Engine/Plugins/Experimental/Avalanche/Source/AvalancheMedia/Private/Playback/Nodes/Events/Actions/AvaPlaybackNode_PlayAnim.h b/Engine/Plugins/Experimental/Avalanche/Source/AvalancheMedia/Private/Playback/Nodes/Events/Actions/AvaPlaybackNode_PlayAnim.h
--- a/Engine/Plugins/Experimental/Avalanche/Source/AvalancheMedia/Private/Playback/Nodes/Events/Actions/AvaPlaybackNode_PlayAnim.h
+++ b/Engine/Plugins/Experimental/Avalanche/Source/AvalancheMedia/Private/Playback/Nodes/Events/Actions/AvaPlaybackNode_PlayAnim.h
@@ -16,7 +16,7 @@ struct FAvaPlaybackEventParameters;
 struct FSoftObjectPath;
 
 UCLASS()
-class AVALANCHEMEDIA_API UAvaPlaybackNode_PlayAnim : public UAvaPlaybackNodeAction
+class UAvaPlaybackNode_PlayAnim : public UAvaPlaybackNodeAction
 {
 	GENERATED_BODY()
 
diff --git a/Engine/Plugins/Experimental/Avalanche/Source/AvalancheMedia/Private/Playback/Nodes/Events/Actions/AvaPlaybackNode_PreloadPlayer.h b/Engine/Plugins/Experimental/Avalanche/Source/AvalancheMedia/Private/Playback/Nodes/Events/Actions/AvaPlaybackNode_PreloadPlayer.h
--- a/Engine/Plugins/Experimental/Avalanche/Source/AvalancheMedia/Private/Playback/Nodes/Events/Actions/AvaPlaybackNode_PreloadPlayer.h
+++ b/Engine/Plugins/Experimental/Avalanche/Source/AvalancheMedia/Private/Playback/Nodes/Events/Actions/AvaPlaybackNode_PreloadPlayer.h
@@ -9,7 +9,7 @@ class FText;
 struct FAvaPlaybackEventParameters;
 
 UCLASS()
-class AVALANCHEMEDIA_API UAvaPlaybackNode_PreloadPlayer : public UAvaPlaybackNodeAction
+class UAvaPlaybackNode_PreloadPlayer : public UAvaPlaybackNodeAction
 {
 	GENERATED_BODY()
 
diff --git a/Engine/Plugins/Experimental/Avalanche/Source/AvalancheMedia/Private/Playback/Nodes/Events/FlowControl/AvaPlaybackNode_Hub.h b/Engine/Plugins/Experimental/Avalanche/Source/AvalancheMedia/Private/Playback/Nodes/Events/FlowControl/AvaPlaybackNode_Hub.h
--- a/Engine/Plugins/Experimental/Avalanche/Source/AvalancheMedia/Private/Playback/Nodes/Events/FlowControl/AvaPlaybackNode_Hub.h
+++ b/Engine/Plugins/Experimental/Avalanche/Source/AvalancheMedia/Private/Playback/Nodes/Events/FlowControl/AvaPlaybackNode_Hub.h
@@ -9,7 +9,7 @@ class FText;
 struct FAvaPlaybackEventParameters;
 
 UCLASS()
-class AVALANCHEMEDIA_API UAvaPlaybackNode_Hub : public UAvaPlaybackNodeFlow
+class UAvaPlaybackNode_Hub : public UAvaPlaybackNodeFlow
 {
 	GENERATED_BODY()
 
diff --git a/Engine/Plugins/Experimental/Avalanche/Source/AvalancheMedia/Private/Playback/Nodes/Events/Triggers/AvaPlaybackNode_BeginPlay.h b/Engine/Plugins/Experimental/Avalanche/Source/AvalancheMedia/Private/Playback/Nodes/Events/Triggers/AvaPlaybackNode_BeginPlay.h
--- a/Engine/Plugins/Experimental/Avalanche/Source/AvalancheMedia/Private/Playback/Nodes/Events/Triggers/AvaPlaybackNode_BeginPlay.h
+++ b/Engine/Plugins/Experimental/Avalanche/Source/AvalancheMedia/Private/Playback/Nodes/Events/Triggers/AvaPlaybackNode_BeginPlay.h
@@ -8,7 +8,7 @@
 class FText;
 
 UCLASS()
-class AVALANCHEMEDIA_API UAvaPlaybackNode_BeginPlay : public UAvaPlaybackNodeTrigger
+class UAvaPlaybackNode_BeginPlay : public UAvaPlaybackNodeTrigger
 {
 	GENERATED_BODY()
 
diff --git a/Engine/Plugins/Experimental/Avalanche/Source/AvalancheMedia/Private/Playback/Nodes/Events/Triggers/AvaPlaybackNode_KeyInput.h b/Engine/Plugins/Experimental/Avalanche/Source/AvalancheMedia/Private/Playback/Nodes/Events/Triggers/AvaPlaybackNode_KeyInput.h
--- a/Engine/Plugins/Experimental/Avalanche/Source/AvalancheMedia/Private/Playback/Nodes/Events/Triggers/AvaPlaybackNode_KeyInput.h
+++ b/Engine/Plugins/Experimental/Avalanche/Source/AvalancheMedia/Private/Playback/Nodes/Events/Triggers/AvaPlaybackNode_KeyInput.h
@@ -14,7 +14,7 @@ class FText;
 struct FKeyEvent;
 
 UCLASS()
-class AVALANCHEMEDIA_API UAvaPlaybackNode_KeyInput : public UAvaPlaybackNodeTrigger
+class UAvaPlaybackNode_KeyInput : public UAvaPlaybackNodeTrigger
 {
 	GENERATED_BODY()
 	
diff --git a/Engine/Plugins/Experimental/Avalanche/Source/AvalancheMedia/Private/Playback/Transition/AvaPlaybackServerTransition.cpp b/Engine/Plugins/Experimental/Avalanche/Source/AvalancheMedia/Private/Playback/Transition/AvaPlaybackServerTransition.cpp
--- a/Engine/Plugins/Experimental/Avalanche/Source/AvalancheMedia/Private/Playback/Transition/AvaPlaybackServerTransition.cpp
+++ b/Engine/Plugins/Experimental/Avalanche/Source/AvalancheMedia/Private/Playback/Transition/AvaPlaybackServerTransition.cpp
@@ -234,7 +234,7 @@ void UAvaPlaybackServerTransition::Stop()
 	UnregisterFromPlayableTransitionEvent();
 	
 	// Remove transition from server.
-	if (const TSharedPtr<FAvaPlaybackServer> PlaybackServer = IAvaMediaModule::Get().GetMediaPlaybackServer())
+	if (const TSharedPtr<FAvaPlaybackServer> PlaybackServer = IAvaMediaModule::Get().GetPlaybackServerInternal())
 	{
 		if (!PlaybackServer->RemovePlaybackInstanceTransition(TransitionId))
 		{
@@ -315,7 +315,7 @@ void UAvaPlaybackServerTransition::OnTransitionEvent(UAvaPlayable* InPlayable, U
 		return;
 	}
 
-	const TSharedPtr<FAvaPlaybackServer> PlaybackServer = IAvaMediaModule::Get().GetMediaPlaybackServer();
+	const TSharedPtr<FAvaPlaybackServer> PlaybackServer = IAvaMediaModule::Get().GetPlaybackServerInternal();
 
 	// Find the page player for this playable
 	if (const TSharedPtr<FAvaPlaybackInstance> Instance = FindInstanceForPlayable(InPlayable))
diff --git a/Engine/Plugins/Experimental/Avalanche/Source/AvalancheMedia/Private/Playback/Transition/AvaPlaybackTransition.h b/Engine/Plugins/Experimental/Avalanche/Source/AvalancheMedia/Private/Playback/Transition/AvaPlaybackTransition.h
new file mode 100644
--- /dev/null
+++ b/Engine/Plugins/Experimental/Avalanche/Source/AvalancheMedia/Private/Playback/Transition/AvaPlaybackTransition.h
@@ -0,0 +1,35 @@
+// Copyright Epic Games, Inc. All Rights Reserved.
+
+#pragma once
+
+#include "Playable/IAvaPlayableVisibilityConstraint.h"
+#include "UObject/Object.h"
+#include "AvaPlaybackTransition.generated.h"
+
+/**
+ * Abstract base class for playback transitions that can be queued in
+ * the playback manager's commands.
+ */
+UCLASS()
+class UAvaPlaybackTransition : public UObject, public IAvaPlayableVisibilityConstraint
+{
+	GENERATED_BODY()
+	
+public:
+	/**
+	 * @brief Evaluate the status of loading playables to determine if the transition can start.
+	 * @param bOutShouldDiscard indicate if the pending start transition command should be discarded.
+	 * @return true if the transition can start, false otherwise.
+	 */
+	virtual bool CanStart(bool& bOutShouldDiscard) const { bOutShouldDiscard = true; return false; }
+
+	virtual void Start() {}
+
+	virtual void Stop() {}
+
+	virtual bool IsRunning() const { return false; }
+	
+	//~ Begin IAvaPlayableVisibilityConstraint
+	virtual bool IsVisibilityConstrained(const UAvaPlayable* InPlayable) const override { return false; }
+	//~ End IAvaPlayableVisibilityConstraint
+};
\ No newline at end of file
diff --git a/Engine/Plugins/Experimental/Avalanche/Source/AvalancheMedia/Private/Rundown/AvaRundown.cpp b/Engine/Plugins/Experimental/Avalanche/Source/AvalancheMedia/Private/Rundown/AvaRundown.cpp
--- a/Engine/Plugins/Experimental/Avalanche/Source/AvalancheMedia/Private/Rundown/AvaRundown.cpp
+++ b/Engine/Plugins/Experimental/Avalanche/Source/AvalancheMedia/Private/Rundown/AvaRundown.cpp
@@ -10,6 +10,7 @@
 #include "Playback/AvaPlaybackManager.h"
 #include "Rundown/AvaRundownManagedInstanceCache.h"
 #include "Rundown/AvaRundownPageLoadingManager.h"
+#include "Rundown/AvaRundownPagePlayer.h"
 #include "Rundown/AvaRundownPlaybackClientWatcher.h"
 #include "Rundown/Transition/AvaRundownPageTransition.h"
 #include "Rundown/Transition/AvaRundownPageTransitionBuilder.h"
@@ -2044,6 +2045,24 @@ void UAvaRundown::StopPageTransitionsForPage(const FAvaRundownPage& InPage, bool
 	}
 }
 
+UAvaRundownPagePlayer* UAvaRundown::FindPlayerForProgramPage(int32 InPageId) const
+{
+	const TObjectPtr<UAvaRundownPagePlayer>* FoundPlayer = PagePlayers.FindByPredicate([InPageId](const UAvaRundownPagePlayer* InPagePlayer)
+	{
+		return InPagePlayer->PageId == InPageId && !InPagePlayer->bIsPreview;
+	});
+	return FoundPlayer ? *FoundPlayer : nullptr;
+}
+
+UAvaRundownPagePlayer* UAvaRundown::FindPlayerForPreviewPage(int32 InPageId, const FName& InPreviewChannelFName) const
+{
+	const TObjectPtr<UAvaRundownPagePlayer>* FoundPlayer = PagePlayers.FindByPredicate([InPageId, InPreviewChannelFName](const UAvaRundownPagePlayer* InPagePlayer)
+	{
+		return InPagePlayer->PageId == InPageId && InPagePlayer->bIsPreview && InPagePlayer->ChannelFName == InPreviewChannelFName;
+	});
+	return FoundPlayer ? *FoundPlayer : nullptr;
+}
+
 void UAvaRundown::RemoveStoppedPagePlayers()
 {
 	for (UAvaRundownPagePlayer* PagePlayer : PagePlayers)
diff --git a/Engine/Plugins/Experimental/Avalanche/Source/AvalancheMedia/Private/Rundown/AvaRundownComponent.h b/Engine/Plugins/Experimental/Avalanche/Source/AvalancheMedia/Private/Rundown/AvaRundownComponent.h
new file mode 100644
--- /dev/null
+++ b/Engine/Plugins/Experimental/Avalanche/Source/AvalancheMedia/Private/Rundown/AvaRundownComponent.h
@@ -0,0 +1,43 @@
+// Copyright Epic Games, Inc. All Rights Reserved.
+
+#pragma once
+
+#include "Components/ActorComponent.h"
+#include "AvaRundownComponent.generated.h"
+
+class UAvaRundown;
+
+/**
+ * Add this actor component to blueprint actor to expose the API to control an
+ * Motion Design Rundown in game.
+ */
+UCLASS(BlueprintType, Blueprintable, ClassGroup = "Motion Design Rundown", 
+	meta = (DisplayName = "Motion Design Rundown Component", BlueprintSpawnableComponent))
+class UAvaRundownComponent : public UActorComponent
+{
+	GENERATED_BODY()
+
+public:
+	UAvaRundownComponent(const FObjectInitializer& ObjectInitializer);
+	
+	UPROPERTY(EditAnywhere, Category = "Motion Design Media", meta = (DisplayName = "Motion Design Rundown"))
+	TObjectPtr<UAvaRundown> Rundown = nullptr;
+	
+	UFUNCTION(BlueprintCallable, Category = "Motion Design Media")
+	bool PlayPage(int32 InPageId);
+
+	UFUNCTION(BlueprintCallable, Category = "Motion Design Media")
+	bool StopPage(int32 InPageId);
+
+	UFUNCTION(BlueprintCallable, Category = "Motion Design Media")
+	int32 GetNumberOfPages() const;
+	
+	UFUNCTION(BlueprintCallable, Category = "Motion Design Media")
+	int32 GetPageIdForIndex(int32 InPageIndex) const;
+
+	virtual void InitializeComponent() override;
+	virtual void UninitializeComponent() override;
+
+protected:
+	void OnWorldBeginTearDown(UWorld* InWorld);
+};
diff --git a/Engine/Plugins/Experimental/Avalanche/Source/AvalancheMedia/Private/Rundown/AvaRundownPage.cpp b/Engine/Plugins/Experimental/Avalanche/Source/AvalancheMedia/Private/Rundown/AvaRundownPage.cpp
--- a/Engine/Plugins/Experimental/Avalanche/Source/AvalancheMedia/Private/Rundown/AvaRundownPage.cpp
+++ b/Engine/Plugins/Experimental/Avalanche/Source/AvalancheMedia/Private/Rundown/AvaRundownPage.cpp
@@ -14,6 +14,7 @@
 #include "Rundown/AvaRundown.h"
 #include "Rundown/AvaRundownManagedInstanceCache.h"
 #include "Rundown/AvaRundownPageAssetUtils.h"
+#include "Rundown/AvaRundownPagePlayer.h"
 
 #define LOCTEXT_NAMESPACE "AvaRundownPage"
 
@@ -298,7 +299,7 @@ int32 FAvaRundownPage::AppendPageProgramStatuses(const UAvaRundown* InParentRund
 	
 	const FAvaBroadcastOutputChannel& Channel = UAvaBroadcast::Get().GetCurrentProfile().GetChannel(GetChannelName());
 	const TArray<UMediaOutput*>& Outputs = Channel.GetMediaOutputs();
-	IAvaPlaybackClient& PlaybackClient = AvaMediaModule.GetMediaPlaybackClient();
+	IAvaPlaybackClient& PlaybackClient = AvaMediaModule.GetPlaybackClient();
 
 	TSet<FString> AddedServers;
 
diff --git a/Engine/Plugins/Experimental/Avalanche/Source/AvalancheMedia/Private/Rundown/AvaRundownPageLoadingManager.h b/Engine/Plugins/Experimental/Avalanche/Source/AvalancheMedia/Private/Rundown/AvaRundownPageLoadingManager.h
--- a/Engine/Plugins/Experimental/Avalanche/Source/AvalancheMedia/Private/Rundown/AvaRundownPageLoadingManager.h
+++ b/Engine/Plugins/Experimental/Avalanche/Source/AvalancheMedia/Private/Rundown/AvaRundownPageLoadingManager.h
@@ -5,6 +5,8 @@
 #include "Containers/Queue.h"
 #include "Rundown/AvaRundown.h"
 
+class FAvaPlaybackInstance;
+
 /**
  * Manager for page loading requests.
  *
diff --git a/Engine/Plugins/Experimental/Avalanche/Source/AvalancheMedia/Private/Rundown/AvaRundownPlaybackClientWatcher.cpp b/Engine/Plugins/Experimental/Avalanche/Source/AvalancheMedia/Private/Rundown/AvaRundownPlaybackClientWatcher.cpp
--- a/Engine/Plugins/Experimental/Avalanche/Source/AvalancheMedia/Private/Rundown/AvaRundownPlaybackClientWatcher.cpp
+++ b/Engine/Plugins/Experimental/Avalanche/Source/AvalancheMedia/Private/Rundown/AvaRundownPlaybackClientWatcher.cpp
@@ -5,6 +5,7 @@
 #include "Broadcast/AvaBroadcast.h"
 #include "Playback/IAvaPlaybackClient.h"
 #include "Rundown/AvaRundown.h"
+#include "Rundown/AvaRundownPagePlayer.h"
 
 FAvaRundownPlaybackClientWatcher::FAvaRundownPlaybackClientWatcher(UAvaRundown* InRundown)
 	: Rundown(InRundown)
diff --git a/Engine/Plugins/Experimental/Avalanche/Source/AvalancheMedia/Private/Rundown/Transition/AvaRundownPageTransition.h b/Engine/Plugins/Experimental/Avalanche/Source/AvalancheMedia/Private/Rundown/Transition/AvaRundownPageTransition.h
new file mode 100644
--- /dev/null
+++ b/Engine/Plugins/Experimental/Avalanche/Source/AvalancheMedia/Private/Rundown/Transition/AvaRundownPageTransition.h
@@ -0,0 +1,93 @@
+// Copyright Epic Games, Inc. All Rights Reserved.
+
+#pragma once
+
+#include "AvaTag.h"
+#include "Playable/AvaPlayable.h"
+#include "Playback/Transition/AvaPlaybackTransition.h"
+#include "UObject/Object.h"
+#include "AvaRundownPageTransition.generated.h"
+
+class FAvaPlayableTransitionBuilder;
+class UAvaPlayableTransition;
+class UAvaPlaybackGraph;
+class UAvaRundown;
+class UAvaRundownPagePlayer;
+class UAvaRundownPlaybackInstancePlayer;
+enum class EAvaPlayableTransitionEntryRole : uint8;
+
+/**
+ * Class for creating and tracking page transitions in the rundown.
+ * It is responsible for creating the playable transition object when requested from
+ * the playback graphs.
+ */
+UCLASS()
+class UAvaRundownPageTransition : public UAvaPlaybackTransition
+{
+	GENERATED_BODY()
+	
+public:
+	bool AddEnterPage(UAvaRundownPagePlayer* InPagePlayer);
+	bool AddPlayingPage(UAvaRundownPagePlayer* InPagePlayer);
+	bool AddExitPage(UAvaRundownPagePlayer* InPagePlayer);
+	
+	//~ Begin IAvaPlayableVisibilityConstraint
+	virtual bool IsVisibilityConstrained(const UAvaPlayable* InPlayable) const override;
+	//~ End IAvaPlayableVisibilityConstraint
+	
+	//~ Begin UAvaPlaybackTransition
+	virtual bool CanStart(bool& bOutShouldDiscard) const override;
+	virtual void Start() override;
+	virtual void Stop() override;
+	virtual bool IsRunning() const override;
+	//~ End UAvaPlaybackTransition
+
+	/** Returns the channel this transition is happening in. A transition can only have pages within the same channel. */
+	FName GetChannelName() const { return ChannelName; }
+
+	bool HasEnterPages() const { return !EnterPlayersWeak.IsEmpty(); }
+
+	bool HasEnterPagesWithNoTransitionLogic() const;
+
+	bool HasPagePlayer(const UAvaRundownPagePlayer* InPagePlayer) const;
+
+	bool ContainsTransitionLayer(const FAvaTagId& InTagId) const;
+
+	UAvaRundown* GetRundown() const;
+
+protected:
+	UAvaRundownPlaybackInstancePlayer* FindInstancePlayerForPlayable(const UAvaPlayable* InPlayable) const;
+	UAvaRundownPagePlayer* FindPagePlayerForPlayable(const UAvaPlayable* InPlayable) const;
+
+	void OnTransitionEvent(UAvaPlayable* InPlayable, UAvaPlayableTransition* InTransition, EAvaPlayableTransitionEventFlags InTransitionFlags);
+	void OnPlayableCreated(UAvaPlaybackGraph* InPlayback, UAvaPlayable* InPlayable);
+
+	void AddPlayersToBuilder(FAvaPlayableTransitionBuilder& InOutBuilder, const TArray<TWeakObjectPtr<UAvaRundownPagePlayer>>& InPlayersWeak, const TCHAR* InCategory, EAvaPlayableTransitionEntryRole InEntryRole) const;
+	void AddPlayablesToBuilder(FAvaPlayableTransitionBuilder& InOutBuilder, const UAvaRundownPagePlayer* InPlayer, const TCHAR* InCategory, EAvaPlayableTransitionEntryRole InEntryRole) const;
+
+	void MakePlayableTransition();
+
+	void LogDetailedTransitionInfo() const;
+	FString GetBriefTransitionDescription() const;
+
+	void RegisterToPlayableTransitionEvent();
+	void UnregisterFromPlayableTransitionEvent() const;
+
+	void RegisterEnterPagePlayerEvents(UAvaRundownPagePlayer* InPagePlayer);
+	void UnregisterEnterPagePlayerEvents(UAvaRundownPagePlayer* InPagePlayer);
+
+	bool AddPagePlayer(UAvaRundownPagePlayer* InPagePlayer, TArray<TWeakObjectPtr<UAvaRundownPagePlayer>>& OutPagePlayersWeak);	
+	void UpdateChannelName(const UAvaRundownPagePlayer* InPagePlayer);
+
+protected:
+	FName ChannelName;
+	
+	TArray<TWeakObjectPtr<UAvaRundownPagePlayer>> EnterPlayersWeak;
+	TArray<TWeakObjectPtr<UAvaRundownPagePlayer>> PlayingPlayersWeak;
+	TArray<TWeakObjectPtr<UAvaRundownPagePlayer>> ExitPlayersWeak;
+
+	TSet<FAvaTagId> CachedTransitionLayers;
+	
+	UPROPERTY(Transient)
+	TObjectPtr<UAvaPlayableTransition> PlayableTransition;
+};
\ No newline at end of file
diff --git a/Engine/Plugins/Experimental/Avalanche/Source/AvalancheMedia/Private/Rundown/Transition/AvaRundownPageTransitionBuilder.cpp b/Engine/Plugins/Experimental/Avalanche/Source/AvalancheMedia/Private/Rundown/Transition/AvaRundownPageTransitionBuilder.cpp
--- a/Engine/Plugins/Experimental/Avalanche/Source/AvalancheMedia/Private/Rundown/Transition/AvaRundownPageTransitionBuilder.cpp
+++ b/Engine/Plugins/Experimental/Avalanche/Source/AvalancheMedia/Private/Rundown/Transition/AvaRundownPageTransitionBuilder.cpp
@@ -3,6 +3,7 @@
 #include "AvaRundownPageTransitionBuilder.h"
 
 #include "Rundown/AvaRundown.h"
+#include "Rundown/AvaRundownPagePlayer.h"
 #include "Rundown/Transition/AvaRundownPageTransition.h"
 
 FAvaRundownPageTransitionBuilder::~FAvaRundownPageTransitionBuilder()
diff --git a/Engine/Plugins/Experimental/Avalanche/Source/AvalancheMedia/Public/Broadcast/AvaBroadcastComponent.h b/Engine/Plugins/Experimental/Avalanche/Source/AvalancheMedia/Public/Broadcast/AvaBroadcastComponent.h
deleted file mode 100644
diff --git a/Engine/Plugins/Experimental/Avalanche/Source/AvalancheMedia/Public/Broadcast/AvaBroadcastDefines.h b/Engine/Plugins/Experimental/Avalanche/Source/AvalancheMedia/Public/Broadcast/AvaBroadcastDefines.h
new file mode 100644
--- /dev/null
+++ b/Engine/Plugins/Experimental/Avalanche/Source/AvalancheMedia/Public/Broadcast/AvaBroadcastDefines.h
@@ -0,0 +1,15 @@
+// Copyright Epic Games, Inc. All Rights Reserved.
+
+#pragma once
+
+#include "AvaBroadcastDefines.generated.h"
+
+UENUM()
+enum class EAvaBroadcastAction
+{
+	None,
+	Start,
+	Stop,
+	UpdateConfig,
+	DeleteChannel
+};
diff --git a/Engine/Plugins/Experimental/Avalanche/Source/AvalancheMedia/Public/IAvaMediaModule.h b/Engine/Plugins/Experimental/Avalanche/Source/AvalancheMedia/Public/IAvaMediaModule.h
--- a/Engine/Plugins/Experimental/Avalanche/Source/AvalancheMedia/Public/IAvaMediaModule.h
+++ b/Engine/Plugins/Experimental/Avalanche/Source/AvalancheMedia/Public/IAvaMediaModule.h
@@ -17,6 +17,7 @@ class IAvaBroadcastDeviceProviderProxyManager;
 class IAvaBroadcastSettings;
 class IAvaMediaSyncProvider;
 class IAvaPlaybackClient;
+class IAvaPlaybackServer;
 class IMediaIOCoreDeviceProvider;
 class UWorld;
 struct FAvaInstanceSettings;
@@ -50,36 +51,37 @@ public:
 	/**
 	 * @brief Returns true if the playback client is started.
 	 */
-	virtual bool IsMediaPlaybackClientStarted() const = 0;
+	virtual bool IsPlaybackClientStarted() const = 0;
 
 	/**
 	 * @brief Starts the playback client (if not already started).
 	 */
-	virtual void StartMediaPlaybackClient() = 0;
+	virtual void StartPlaybackClient() = 0;
 
 	/**
 	 * @brief Stops the playback client.
 	 */
-	virtual void StopMediaPlaybackClient() = 0;
+	virtual void StopPlaybackClient() = 0;
 
 	/**
 	 * @brief Returns true if the playback server is started.
 	 */
-	virtual bool IsMediaPlaybackServerStarted() const = 0;
+	virtual bool IsPlaybackServerStarted() const = 0;
 
 	/**
 	 * @brief Starts the playback server (if not already started).
 	 * @param InPlaybackServerName Optional server name. If empty, the host (computer) name will be used.
 	 */
-	virtual void StartMediaPlaybackServer(const FString& InPlaybackServerName) = 0;
+	virtual void StartPlaybackServer(const FString& InPlaybackServerName) = 0;
 
 	/**
 	 * @brief Stops the playback server.
 	 */
-	virtual void StopMediaPlaybackServer() = 0;
+	virtual void StopPlaybackServer() = 0;
 
-	virtual IAvaPlaybackClient& GetMediaPlaybackClient() = 0;
-	virtual TSharedPtr<FAvaPlaybackServer> GetMediaPlaybackServer() const = 0;
+	virtual IAvaPlaybackClient& GetPlaybackClient() = 0;
+	virtual TSharedPtr<FAvaPlaybackServer> GetPlaybackServerInternal() const = 0;
+	virtual IAvaPlaybackServer* GetPlaybackServer() const = 0;
 	virtual const IMediaIOCoreDeviceProvider* GetDeviceProvider(FName InProviderName, const FMediaIOOutputConfiguration* InMediaIOOutputConfiguration) const = 0;
 	virtual TArray<const IMediaIOCoreDeviceProvider*> GetDeviceProvidersForServer(const FString& InServerName) const = 0;
 	virtual FString GetServerNameForDevice(const FName& InDeviceProviderName, const FName& InDeviceName) const = 0;
diff --git a/Engine/Plugins/Experimental/Avalanche/Source/AvalancheMedia/Public/Playable/AvaPlayableGroup.h b/Engine/Plugins/Experimental/Avalanche/Source/AvalancheMedia/Public/Playable/AvaPlayableGroup.h
--- a/Engine/Plugins/Experimental/Avalanche/Source/AvalancheMedia/Public/Playable/AvaPlayableGroup.h
+++ b/Engine/Plugins/Experimental/Avalanche/Source/AvalancheMedia/Public/Playable/AvaPlayableGroup.h
@@ -2,7 +2,6 @@
 
 #pragma once
 
-#include "Framework/AvaGameInstance.h"
 #include "Playable/IAvaPlayableVisibilityConstraint.h"
 #include "UObject/Object.h"
 #include "UObject/ObjectKey.h"
@@ -11,10 +10,13 @@
 #include "UObject/WeakInterfacePtr.h"
 #include "AvaPlayableGroup.generated.h"
 
+class UAvaGameInstance;
 class UAvaPlayable;
 class UAvaPlayableGroupManager;
 class UAvaPlayableTransition;
+class UGameInstance;
 class UTextureRenderTarget2D;
+struct FAvaInstancePlaySettings;
 
 /**
  * This defines a game instance playable group.
@@ -98,19 +100,19 @@ public:
 	/** Search for the first remaining playing playable and use it's camera. */
 	bool UpdateCameraSetup();
 
-	bool IsWorldPlaying() const { return GameInstance ? GameInstance->IsWorldPlaying() : false;}
+	bool IsWorldPlaying() const;
 
-	bool IsRenderTargetReady() const { return GameInstance ? GameInstance->IsRenderTargetReady() : false;}
+	bool IsRenderTargetReady() const;
 
 	/**
 	 * Current logic for the render target, use the game instance's if there, fallback to internal one if not.
 	 */
-	UTextureRenderTarget2D* GetRenderTarget() const { return GameInstance ? GameInstance->GetRenderTarget() : RenderTarget.Get(); }
+	UTextureRenderTarget2D* GetRenderTarget() const;
 
 	/** Return a vanilla game instance, we want to eventually support any game instance. */
-	UGameInstance* GetGameInstance() const { return GameInstance;}
+	UGameInstance* GetGameInstance() const;
 	
-	UWorld* GetPlayWorld() const { return GameInstance ? GameInstance->GetPlayWorld() : nullptr; }
+	UWorld* GetPlayWorld() const;
 	
 	/**
 	 * Unloads the game instance's world if no more playables are loaded.
diff --git a/Engine/Plugins/Experimental/Avalanche/Source/AvalancheMedia/Public/Playable/AvaPlayableGroupManager.h b/Engine/Plugins/Experimental/Avalanche/Source/AvalancheMedia/Public/Playable/AvaPlayableGroupManager.h
deleted file mode 100644
diff --git a/Engine/Plugins/Experimental/Avalanche/Source/AvalancheMedia/Public/Playable/AvaPlayableRemoteControl.h b/Engine/Plugins/Experimental/Avalanche/Source/AvalancheMedia/Public/Playable/AvaPlayableRemoteControl.h
deleted file mode 100644
diff --git a/Engine/Plugins/Experimental/Avalanche/Source/AvalancheMedia/Public/Playable/Transition/AvaPlayableTransitionScene.h b/Engine/Plugins/Experimental/Avalanche/Source/AvalancheMedia/Public/Playable/Transition/AvaPlayableTransitionScene.h
deleted file mode 100644
diff --git a/Engine/Plugins/Experimental/Avalanche/Source/AvalancheMedia/Public/Playback/AvaPlaybackClient.h b/Engine/Plugins/Experimental/Avalanche/Source/AvalancheMedia/Public/Playback/AvaPlaybackClient.h
deleted file mode 100644
diff --git a/Engine/Plugins/Experimental/Avalanche/Source/AvalancheMedia/Public/Playback/AvaPlaybackDefines.h b/Engine/Plugins/Experimental/Avalanche/Source/AvalancheMedia/Public/Playback/AvaPlaybackDefines.h
new file mode 100644
--- /dev/null
+++ b/Engine/Plugins/Experimental/Avalanche/Source/AvalancheMedia/Public/Playback/AvaPlaybackDefines.h
@@ -0,0 +1,37 @@
+// Copyright Epic Games, Inc. All Rights Reserved.
+
+#pragma once
+
+#include "AvaPlaybackDefines.generated.h"
+
+UENUM()
+enum class EAvaPlaybackAction
+{
+	/** No op */
+	None,
+	/** Load the given asset. Used for pre-loading assets. */
+	Load,
+	/** Start (i.e. start ticking world and rendering) the given asset, loading it if not pre-loaded. */
+	Start,
+	/** Stop ticking and rendering the world. */
+	Stop,
+	/** Unload the given asset, i.e. destroy the world, etc. */
+	Unload,
+	/** Request the status of the asset. No action is actually performed on the asset. */
+	Status,
+	/** Request to set the user data of the playback instance.*/
+	SetUserData,
+	/** Request the user data of the playback instance. No action is actually performed on the asset. */
+	GetUserData,
+};
+
+UENUM()
+enum class EAvaPlaybackAnimAction
+{
+	None,
+	Play,
+	Continue,
+	Stop,
+	PreviewFrame,
+	CameraCut
+};
\ No newline at end of file
diff --git a/Engine/Plugins/Experimental/Avalanche/Source/AvalancheMedia/Public/Playback/AvaPlaybackMessages.h b/Engine/Plugins/Experimental/Avalanche/Source/AvalancheMedia/Public/Playback/AvaPlaybackMessages.h
deleted file mode 100644
diff --git a/Engine/Plugins/Experimental/Avalanche/Source/AvalancheMedia/Public/Playback/AvaPlaybackServer.h b/Engine/Plugins/Experimental/Avalanche/Source/AvalancheMedia/Public/Playback/AvaPlaybackServer.h
deleted file mode 100644
diff --git a/Engine/Plugins/Experimental/Avalanche/Source/AvalancheMedia/Public/Playback/AvaPlaybackUtils.h b/Engine/Plugins/Experimental/Avalanche/Source/AvalancheMedia/Public/Playback/AvaPlaybackUtils.h
deleted file mode 100644
diff --git a/Engine/Plugins/Experimental/Avalanche/Source/AvalancheMedia/Public/Playback/IAvaPlaybackClient.h b/Engine/Plugins/Experimental/Avalanche/Source/AvalancheMedia/Public/Playback/IAvaPlaybackClient.h
--- a/Engine/Plugins/Experimental/Avalanche/Source/AvalancheMedia/Public/Playback/IAvaPlaybackClient.h
+++ b/Engine/Plugins/Experimental/Avalanche/Source/AvalancheMedia/Public/Playback/IAvaPlaybackClient.h
@@ -2,10 +2,13 @@
 
 #pragma once
 
+#include "AvaMediaDefines.h"
+#include "AvaPlaybackDefines.h"
+#include "Broadcast/AvaBroadcastDefines.h"
 #include "IMessageContext.h"
-#include "Playback/AvaPlaybackMessages.h"
 
 class UMediaOutput;
+struct FAvaPlayableRemoteControlValues;
 struct FAvaPlaybackAnimPlaySettings;
 
 class AVALANCHEMEDIA_API IAvaPlaybackClient
diff --git a/Engine/Plugins/Experimental/Avalanche/Source/AvalancheMedia/Public/Playback/IAvaPlaybackServer.h b/Engine/Plugins/Experimental/Avalanche/Source/AvalancheMedia/Public/Playback/IAvaPlaybackServer.h
new file mode 100644
--- /dev/null
+++ b/Engine/Plugins/Experimental/Avalanche/Source/AvalancheMedia/Public/Playback/IAvaPlaybackServer.h
@@ -0,0 +1,71 @@
+// Copyright Epic Games, Inc. All Rights Reserved.
+
+#pragma once
+
+#include "IMessageContext.h"
+
+class FAvaPlaybackManager;
+
+class IAvaPlaybackServer
+{
+public:
+	
+	struct FPlaybackInstanceReference
+	{
+		FGuid Id;
+		FSoftObjectPath Path;
+	};
+	
+	virtual TArray<FPlaybackInstanceReference> StopPlaybacks(const FString& InChannelName = FString(), const FSoftObjectPath& InAssetPath = FSoftObjectPath(), bool bInUnload = true) = 0;
+	virtual TArray<FPlaybackInstanceReference> StartPlaybacks() = 0;
+	
+	virtual void StartBroadcast() = 0;
+	virtual void StopBroadcast() = 0;
+	
+	/** Returns the server's name. */
+	virtual const FString& GetName() const = 0;
+
+	virtual bool HasUserData(const FString& InKey) const = 0;
+	
+	virtual const FString& GetUserData(const FString& InKey) const = 0;
+
+	/** Add user data to this server. This is replicated and accessible to the client. */
+	virtual void SetUserData(const FString& InKey, const FString& InData) = 0;
+
+	/** Remove the server's user data entry from the given key. */
+	virtual void RemoveUserData(const FString& InKey) = 0;
+
+	/** Returns the list of connected clients. */
+	virtual TArray<FString> GetClientNames() const = 0;
+
+	/** Returns the client address. */
+	virtual FMessageAddress GetClientAddress(const FString& InClientName) const = 0;
+
+	/** Returns true if the corresponding client user data for the given client name is found. */
+	virtual bool HasClientUserData(const FString& InClientName, const FString& InKey) const = 0;
+
+	/**
+	 * Returns the corresponding client user data for the given client name and key.
+	 * Returns empty string if not found.
+	 */
+	virtual const FString& GetClientUserData(const FString& InClientName, const FString& InKey) const = 0;
+
+	/**
+	 * Access broadcast settings replicated from connected client.
+	 * Will return nullptr if no clients are connected.
+	 */
+	virtual const IAvaBroadcastSettings* GetBroadcastSettings() const = 0;
+
+	/**
+	 * Access Motion Design Instance settings replicated from connected client(s).
+	 * Will return nullptr if no clients are connected.
+	 */
+	virtual const FAvaInstanceSettings* GetAvaInstanceSettings() const = 0;
+	
+	/** Access the server's playback manager. */
+	virtual const FAvaPlaybackManager& GetPlaybackManager() const = 0;
+	virtual FAvaPlaybackManager& GetPlaybackManager() = 0;
+
+protected:
+	virtual ~IAvaPlaybackServer() = default;
+};
diff --git a/Engine/Plugins/Experimental/Avalanche/Source/AvalancheMedia/Public/Playback/Nodes/AvaPlaybackNodeCombiner.h b/Engine/Plugins/Experimental/Avalanche/Source/AvalancheMedia/Public/Playback/Nodes/AvaPlaybackNodeCombiner.h
deleted file mode 100644
diff --git a/Engine/Plugins/Experimental/Avalanche/Source/AvalancheMedia/Public/Playback/Nodes/AvaPlaybackNodeLevelPlayer.h b/Engine/Plugins/Experimental/Avalanche/Source/AvalancheMedia/Public/Playback/Nodes/AvaPlaybackNodeLevelPlayer.h
deleted file mode 100644
diff --git a/Engine/Plugins/Experimental/Avalanche/Source/AvalancheMedia/Public/Playback/Nodes/AvaPlaybackNodeSwitcher.h b/Engine/Plugins/Experimental/Avalanche/Source/AvalancheMedia/Public/Playback/Nodes/AvaPlaybackNodeSwitcher.h
deleted file mode 100644
diff --git a/Engine/Plugins/Experimental/Avalanche/Source/AvalancheMedia/Public/Playback/Nodes/Events/Actions/AvaPlaybackAnimations.h b/Engine/Plugins/Experimental/Avalanche/Source/AvalancheMedia/Public/Playback/Nodes/Events/Actions/AvaPlaybackAnimations.h
--- a/Engine/Plugins/Experimental/Avalanche/Source/AvalancheMedia/Public/Playback/Nodes/Events/Actions/AvaPlaybackAnimations.h
+++ b/Engine/Plugins/Experimental/Avalanche/Source/AvalancheMedia/Public/Playback/Nodes/Events/Actions/AvaPlaybackAnimations.h
@@ -4,20 +4,10 @@
 
 #include "AvaSequenceShared.h"
 #include "Containers/Set.h"
+#include "Playback/AvaPlaybackDefines.h"
 #include "UObject/NameTypes.h"
 #include "AvaPlaybackAnimations.generated.h"
 
-UENUM()
-enum class EAvaPlaybackAnimAction
-{
-	None,
-	Play,
-	Continue,
-	Stop,
-	PreviewFrame,
-	CameraCut
-};
-
 USTRUCT()
 struct FAvaPlaybackAnimPlaySettings
 {
diff --git a/Engine/Plugins/Experimental/Avalanche/Source/AvalancheMedia/Public/Playback/Transition/AvaPlaybackTransition.h b/Engine/Plugins/Experimental/Avalanche/Source/AvalancheMedia/Public/Playback/Transition/AvaPlaybackTransition.h
deleted file mode 100644
diff --git a/Engine/Plugins/Experimental/Avalanche/Source/AvalancheMedia/Public/Rundown/AvaRundown.h b/Engine/Plugins/Experimental/Avalanche/Source/AvalancheMedia/Public/Rundown/AvaRundown.h
--- a/Engine/Plugins/Experimental/Avalanche/Source/AvalancheMedia/Public/Rundown/AvaRundown.h
+++ b/Engine/Plugins/Experimental/Avalanche/Source/AvalancheMedia/Public/Rundown/AvaRundown.h
@@ -3,8 +3,8 @@
 #pragma once
 
 #include "AvaMediaDefines.h"
+#include "AvaRundownDefines.h"
 #include "AvaRundownPage.h"
-#include "AvaRundownPagePlayer.h"
 #include "UObject/Object.h"
 #include "AvaRundown.generated.h"
 
@@ -14,6 +14,7 @@ class FAvaPlaybackManager;
 class FAvaRundownPageTransitionBuilder;
 class FAvaRundownPlaybackClientWatcher;
 class UAvaPlaybackGraph;
+class UAvaRundownPagePlayer;
 class UAvaRundownPageTransition;
 class UTextureRenderTarget2D;
 
@@ -685,23 +686,9 @@ public:
 
 	const TArray<TObjectPtr<UAvaRundownPagePlayer>>& GetPagePlayers() const { return PagePlayers; }
 	
-	UAvaRundownPagePlayer* FindPlayerForProgramPage(int32 InPageId) const
-	{
-		const TObjectPtr<UAvaRundownPagePlayer>* FoundPlayer = PagePlayers.FindByPredicate([InPageId](const UAvaRundownPagePlayer* InPagePlayer)
-		{
-			return InPagePlayer->PageId == InPageId && !InPagePlayer->bIsPreview;
-		});
-		return FoundPlayer ? *FoundPlayer : nullptr;
-	}
+	UAvaRundownPagePlayer* FindPlayerForProgramPage(int32 InPageId) const;
 
-	UAvaRundownPagePlayer* FindPlayerForPreviewPage(int32 InPageId, const FName& InPreviewChannelFName) const
-	{
-		const TObjectPtr<UAvaRundownPagePlayer>* FoundPlayer = PagePlayers.FindByPredicate([InPageId, InPreviewChannelFName](const UAvaRundownPagePlayer* InPagePlayer)
-		{
-			return InPagePlayer->PageId == InPageId && InPagePlayer->bIsPreview && InPagePlayer->ChannelFName == InPreviewChannelFName;
-		});
-		return FoundPlayer ? *FoundPlayer : nullptr;
-	}
+	UAvaRundownPagePlayer* FindPlayerForPreviewPage(int32 InPageId, const FName& InPreviewChannelFName) const;
 	
 	UAvaRundownPagePlayer* FindPlayerForPage(int32 InPageId, bool bInIsPreview, const FName& InPreviewChannelName) const
 	{
diff --git a/Engine/Plugins/Experimental/Avalanche/Source/AvalancheMedia/Public/Rundown/AvaRundownComponent.h b/Engine/Plugins/Experimental/Avalanche/Source/AvalancheMedia/Public/Rundown/AvaRundownComponent.h
deleted file mode 100644
diff --git a/Engine/Plugins/Experimental/Avalanche/Source/AvalancheMedia/Public/Rundown/AvaRundownDefines.h b/Engine/Plugins/Experimental/Avalanche/Source/AvalancheMedia/Public/Rundown/AvaRundownDefines.h
new file mode 100644
--- /dev/null
+++ b/Engine/Plugins/Experimental/Avalanche/Source/AvalancheMedia/Public/Rundown/AvaRundownDefines.h
@@ -0,0 +1,13 @@
+// Copyright Epic Games, Inc. All Rights Reserved.
+
+#pragma once
+
+#include "AvaRundownDefines.generated.h"
+
+UENUM(BlueprintType, DisplayName = "Motion Design Rundown Page Play Type")
+enum class EAvaRundownPagePlayType : uint8
+{
+	PlayFromStart,
+	PreviewFromStart,
+	PreviewFromFrame
+};
diff --git a/Engine/Plugins/Experimental/Avalanche/Source/AvalancheMedia/Public/Rundown/AvaRundownManagedInstance.h b/Engine/Plugins/Experimental/Avalanche/Source/AvalancheMedia/Public/Rundown/AvaRundownManagedInstance.h
deleted file mode 100644
diff --git a/Engine/Plugins/Experimental/Avalanche/Source/AvalancheMedia/Public/Rundown/AvaRundownManagedInstanceCache.h b/Engine/Plugins/Experimental/Avalanche/Source/AvalancheMedia/Public/Rundown/AvaRundownManagedInstanceCache.h
deleted file mode 100644
diff --git a/Engine/Plugins/Experimental/Avalanche/Source/AvalancheMedia/Public/Rundown/AvaRundownPagePlayer.h b/Engine/Plugins/Experimental/Avalanche/Source/AvalancheMedia/Public/Rundown/AvaRundownPagePlayer.h
--- a/Engine/Plugins/Experimental/Avalanche/Source/AvalancheMedia/Public/Rundown/AvaRundownPagePlayer.h
+++ b/Engine/Plugins/Experimental/Avalanche/Source/AvalancheMedia/Public/Rundown/AvaRundownPagePlayer.h
@@ -5,6 +5,7 @@
 #include "AvaTagHandle.h"
 #include "Playable/AvaPlayable.h"
 #include "Playback/AvaPlaybackManager.h"
+#include "Rundown/AvaRundownDefines.h"
 #include "UObject/Object.h"
 #include "AvaRundownPagePlayer.generated.h"
 
@@ -14,14 +15,6 @@ class UAvaRundown;
 class UAvaRundownPagePlayer;
 struct FAvaRundownPage;
 
-UENUM(BlueprintType, DisplayName = "Motion Design Rundown Page Play Type")
-enum class EAvaRundownPagePlayType : uint8
-{
-	PlayFromStart,
-	PreviewFromStart,
-	PreviewFromFrame
-};
-
 /**
  *	To support combo templates, a page is now potentially composed of a number of
  *	playback instances that each have their own player.
diff --git a/Engine/Plugins/Experimental/Avalanche/Source/AvalancheMedia/Public/Rundown/Transition/AvaRundownPageTransition.h b/Engine/Plugins/Experimental/Avalanche/Source/AvalancheMedia/Public/Rundown/Transition/AvaRundownPageTransition.h
deleted file mode 100644
diff --git a/Engine/Plugins/Experimental/Avalanche/Source/AvalancheMediaEditor/Private/AvaMediaEditorModule.cpp b/Engine/Plugins/Experimental/Avalanche/Source/AvalancheMediaEditor/Private/AvaMediaEditorModule.cpp
--- a/Engine/Plugins/Experimental/Avalanche/Source/AvalancheMediaEditor/Private/AvaMediaEditorModule.cpp
+++ b/Engine/Plugins/Experimental/Avalanche/Source/AvalancheMediaEditor/Private/AvaMediaEditorModule.cpp
@@ -196,11 +196,11 @@ FSlateIcon FAvaMediaEditorModule::GetToolbarBroadcastButtonIcon() const
 {
 	const IAvaMediaModule& MediaModule = IAvaMediaModule::Get();
 
-	if (MediaModule.IsMediaPlaybackClientStarted())
+	if (MediaModule.IsPlaybackClientStarted())
 	{
 		return FSlateIcon(FAvaMediaEditorStyle::Get().GetStyleSetName(), "AvaMediaEditor.BroadcastClient", "AvaMediaEditor.BroadcastClient.Small");
 	}
-	else if (MediaModule.IsMediaPlaybackServerStarted())
+	else if (MediaModule.IsPlaybackServerStarted())
 	{
 		return FSlateIcon(FAvaMediaEditorStyle::Get().GetStyleSetName(), "AvaMediaEditor.BroadcastServer", "AvaMediaEditor.BroadcastServer.Small");
 	}
diff --git a/Engine/Plugins/Experimental/Avalanche/Source/AvalancheMediaEditor/Private/Broadcast/AvaBroadcastEditor.cpp b/Engine/Plugins/Experimental/Avalanche/Source/AvalancheMediaEditor/Private/Broadcast/AvaBroadcastEditor.cpp
--- a/Engine/Plugins/Experimental/Avalanche/Source/AvalancheMediaEditor/Private/Broadcast/AvaBroadcastEditor.cpp
+++ b/Engine/Plugins/Experimental/Avalanche/Source/AvalancheMediaEditor/Private/Broadcast/AvaBroadcastEditor.cpp
@@ -154,11 +154,11 @@ FText FAvaBroadcastEditor::GetToolkitName() const
 
 FText FAvaBroadcastEditor::GetToolkitToolTipText() const
 {
-	if (IAvaMediaModule::Get().IsMediaPlaybackServerStarted())
+	if (IAvaMediaModule::Get().IsPlaybackServerStarted())
 	{
 		return LOCTEXT("BroadcastServerAppToolTip", "Motion Design Broadcast: Playback Server running");
 	}
-	else if (IAvaMediaModule::Get().IsMediaPlaybackClientStarted())
+	else if (IAvaMediaModule::Get().IsPlaybackClientStarted())
 	{
 		return LOCTEXT("BroadcastClientAppToolTip", "Motion Design Broadcast: Playback Client running");
 	}
@@ -267,12 +267,12 @@ void FAvaBroadcastEditor::FillPlayToolBar(FToolBarBuilder& ToolBarBuilder)
 				})
 				, FCanExecuteAction::CreateLambda([]
 				{
-					return !IAvaMediaModule::Get().IsMediaPlaybackClientStarted();
+					return !IAvaMediaModule::Get().IsPlaybackClientStarted();
 				})
 				, FGetActionCheckState()
 				, FIsActionButtonVisible::CreateLambda([]
 				{
-					return !IAvaMediaModule::Get().IsMediaPlaybackClientStarted();
+					return !IAvaMediaModule::Get().IsPlaybackClientStarted();
 				})
 			)
 			, NAME_None
@@ -284,16 +284,16 @@ void FAvaBroadcastEditor::FillPlayToolBar(FToolBarBuilder& ToolBarBuilder)
 		// Add Stop button for the playback client service.
 		ToolBarBuilder.AddToolBarButton(FUIAction(FExecuteAction::CreateLambda([]
 				{
-					IAvaMediaModule::Get().StopMediaPlaybackClient();
+					IAvaMediaModule::Get().StopPlaybackClient();
 				})
 				, FCanExecuteAction::CreateLambda([]
 				{
-					return IAvaMediaModule::Get().IsMediaPlaybackClientStarted();
+					return IAvaMediaModule::Get().IsPlaybackClientStarted();
 				})
 				, FGetActionCheckState()
 				, FIsActionButtonVisible::CreateLambda([]
 				{
-					return IAvaMediaModule::Get().IsMediaPlaybackClientStarted();
+					return IAvaMediaModule::Get().IsPlaybackClientStarted();
 				})
 			)
 			, NAME_None
@@ -309,7 +309,7 @@ void FAvaBroadcastEditor::FillPlayToolBar(FToolBarBuilder& ToolBarBuilder)
     			, FCanExecuteAction::CreateLambda([]
     			{
     				const IAvaMediaModule& AvaMediaModule = IAvaMediaModule::Get();
-    				return AvaMediaModule.IsMediaPlaybackClientStarted() && !AvaMediaModule.IsGameModeLocalPlaybackServerLaunched();
+    				return AvaMediaModule.IsPlaybackClientStarted() && !AvaMediaModule.IsGameModeLocalPlaybackServerLaunched();
     			})
     			, FGetActionCheckState()
 				, FIsActionButtonVisible::CreateLambda([]
@@ -349,17 +349,16 @@ void FAvaBroadcastEditor::FillPlayToolBar(FToolBarBuilder& ToolBarBuilder)
 		// to reset all the servers.
 		ToolBarBuilder.AddToolBarButton(FUIAction(FExecuteAction::CreateLambda([]
 				{
+					IAvaPlaybackClient& PlaybackClient = IAvaMediaModule::Get().GetPlaybackClient(); 
 					// Stop and unload all playbacks.
-					IAvaMediaModule::Get().GetMediaPlaybackClient().RequestPlayback(
-						FGuid(), FSoftObjectPath(), FString(), EAvaPlaybackAction::Unload);
+					PlaybackClient.RequestPlayback(FGuid(), FSoftObjectPath(), FString(), EAvaPlaybackAction::Unload);
 					// Stop all broadcast channels.
-					IAvaMediaModule::Get().GetMediaPlaybackClient().RequestBroadcast(
-						TEXT(""), FName(), TArray<UMediaOutput*>(), EAvaBroadcastAction::Stop);
+					PlaybackClient.RequestBroadcast(TEXT(""), FName(), TArray<UMediaOutput*>(), EAvaBroadcastAction::Stop);
 				})
 				, FCanExecuteAction::CreateLambda([]
 				{
 					IAvaMediaModule& MediaModule = IAvaMediaModule::Get();
-					return MediaModule.IsMediaPlaybackClientStarted() && MediaModule.GetMediaPlaybackClient().GetNumConnectedServers() > 0;
+					return MediaModule.IsPlaybackClientStarted() && MediaModule.GetPlaybackClient().GetNumConnectedServers() > 0;
 				})
 				, FGetActionCheckState()
 				, FIsActionButtonVisible()
@@ -517,7 +516,7 @@ void FAvaBroadcastEditor::StartPlaybackClientAction()
 	// We can't run both playback client and server at the same time in the same process.
 	// If the user ask for the client to start, we need to stop the server, but we'll
 	// ask user confirmation.
-	if (AvaMediaModule.IsMediaPlaybackServerStarted())
+	if (AvaMediaModule.IsPlaybackServerStarted())
 	{
 		const EAppReturnType::Type Reply = FMessageDialog::Open(EAppMsgType::YesNo
 			, LOCTEXT("PlaybackServerStop_Message", "This action will stop the playback server, continue anyway?"));
@@ -527,7 +526,7 @@ void FAvaBroadcastEditor::StartPlaybackClientAction()
 			return;
 		}
 
-		AvaMediaModule.StopMediaPlaybackServer();
+		AvaMediaModule.StopPlaybackServer();
 
 		// Stopping the playback server requires a reload of the broadcast client config.
 		if (BroadcastEditor && BroadcastEditor->GetBroadcastObject())
@@ -538,12 +537,12 @@ void FAvaBroadcastEditor::StartPlaybackClientAction()
 		}
 	}
 		
-	AvaMediaModule.StartMediaPlaybackClient();
+	AvaMediaModule.StartPlaybackClient();
 }
 
 FText FAvaBroadcastEditor::GetLaunchLocalServerTooltip()
 {
-	if (IAvaMediaModule::Get().IsMediaPlaybackClientStarted())
+	if (IAvaMediaModule::Get().IsPlaybackClientStarted())
 	{
 		return LOCTEXT("LaunchLocalServerAuto_ToolTip", "Launches Game Mode Local Server Process");
 	}
diff --git a/Engine/Plugins/Experimental/Avalanche/Source/AvalancheMediaEditor/Private/Rundown/AvaRundownEditor.cpp b/Engine/Plugins/Experimental/Avalanche/Source/AvalancheMediaEditor/Private/Rundown/AvaRundownEditor.cpp
--- a/Engine/Plugins/Experimental/Avalanche/Source/AvalancheMediaEditor/Private/Rundown/AvaRundownEditor.cpp
+++ b/Engine/Plugins/Experimental/Avalanche/Source/AvalancheMediaEditor/Private/Rundown/AvaRundownEditor.cpp
@@ -22,6 +22,7 @@
 #include "Rundown/AvaRundownEditorUtils.h"
 #include "Rundown/AvaRundownMacroCollection.h"
 #include "Rundown/AvaRundownManagedInstanceCache.h"
+#include "Rundown/AvaRundownPagePlayer.h"
 #include "Rundown/AvaRundownPlaybackUtils.h"
 #include "Rundown/Factories/Filters/AvaRundownFactoriesUtils.h"
 #include "Rundown/Filters/AvaRundownPageTextFilter.h"
diff --git a/Engine/Plugins/Experimental/Avalanche/Source/AvalancheMediaEditor/Private/Rundown/AvaRundownEditorUtils.h b/Engine/Plugins/Experimental/Avalanche/Source/AvalancheMediaEditor/Private/Rundown/AvaRundownEditorUtils.h
--- a/Engine/Plugins/Experimental/Avalanche/Source/AvalancheMediaEditor/Private/Rundown/AvaRundownEditorUtils.h
+++ b/Engine/Plugins/Experimental/Avalanche/Source/AvalancheMediaEditor/Private/Rundown/AvaRundownEditorUtils.h
@@ -6,6 +6,7 @@
 #include "Rundown/AvaRundown.h"
 #include "Rundown/AvaRundownPage.h"
 #include "Templates/SharedPointer.h"
+#include "UObject/StrongObjectPtr.h"
 #include "XmlSerializationDefines.h"
 
 class FAvaRundownManagedInstance;
diff --git a/Engine/Plugins/Experimental/Avalanche/Source/AvalancheMediaEditor/Private/Rundown/AvaRundownServer.cpp b/Engine/Plugins/Experimental/Avalanche/Source/AvalancheMediaEditor/Private/Rundown/AvaRundownServer.cpp
--- a/Engine/Plugins/Experimental/Avalanche/Source/AvalancheMediaEditor/Private/Rundown/AvaRundownServer.cpp
+++ b/Engine/Plugins/Experimental/Avalanche/Source/AvalancheMediaEditor/Private/Rundown/AvaRundownServer.cpp
@@ -16,8 +16,8 @@
 #include "Broadcast/OutputDevices/AvaBroadcastOutputTreeItem.h"
 #include "Broadcast/OutputDevices/AvaBroadcastRenderTargetMediaUtils.h"
 #include "IAvaMediaModule.h"
-#include "ImageUtils.h"
 #include "IRemoteControlModule.h"
+#include "ImageUtils.h"
 #include "MediaOutput.h"
 #include "MediaOutputEditorUtils/AvaRundownOutputEditorUtils.h"
 #include "MessageEndpointBuilder.h"
@@ -27,10 +27,9 @@
 #include "RenderingThread.h"
 #include "Rundown/AvaRundown.h"
 #include "Rundown/AvaRundownManagedInstanceCache.h"
+#include "Rundown/AvaRundownPagePlayer.h"
 #include "Rundown/AvaRundownPlaybackUtils.h"
 #include "ScopedTransaction.h"
-#include "Serialization/JsonReader.h"
-#include "Serialization/JsonSerializer.h"
 #include "TextureResource.h"
 
 DEFINE_LOG_CATEGORY_STATIC(LogAvaRundownServer, Log, All);
diff --git a/Engine/Plugins/Experimental/Avalanche/Source/AvalancheMediaEditor/Private/Rundown/ChannelLayerStatus/SAvaRundownChannelLayerStatusList.cpp b/Engine/Plugins/Experimental/Avalanche/Source/AvalancheMediaEditor/Private/Rundown/ChannelLayerStatus/SAvaRundownChannelLayerStatusList.cpp
--- a/Engine/Plugins/Experimental/Avalanche/Source/AvalancheMediaEditor/Private/Rundown/ChannelLayerStatus/SAvaRundownChannelLayerStatusList.cpp
+++ b/Engine/Plugins/Experimental/Avalanche/Source/AvalancheMediaEditor/Private/Rundown/ChannelLayerStatus/SAvaRundownChannelLayerStatusList.cpp
@@ -5,6 +5,7 @@
 #include "Misc/CoreDelegates.h"
 #include "Rundown/AvaRundownDefines.h"
 #include "Rundown/AvaRundownEditor.h"
+#include "Rundown/AvaRundownPagePlayer.h"
 #include "Widgets/Colors/SColorBlock.h"
 #include "Widgets/Images/SImage.h"
 #include "Widgets/Layout/SGridPanel.h"
diff --git a/Engine/Plugins/Experimental/StormSyncAvalancheBridge/Source/StormSyncAvaBridge/Private/ModularFeature/StormSyncAvaSyncProvider.cpp b/Engine/Plugins/Experimental/StormSyncAvalancheBridge/Source/StormSyncAvaBridge/Private/ModularFeature/StormSyncAvaSyncProvider.cpp
--- a/Engine/Plugins/Experimental/StormSyncAvalancheBridge/Source/StormSyncAvaBridge/Private/ModularFeature/StormSyncAvaSyncProvider.cpp
+++ b/Engine/Plugins/Experimental/StormSyncAvalancheBridge/Source/StormSyncAvaBridge/Private/ModularFeature/StormSyncAvaSyncProvider.cpp
@@ -4,7 +4,7 @@
 #include "IAvaMediaModule.h"
 #include "IStormSyncTransportClientModule.h"
 #include "ModularFeature/AvaMediaSyncProviderFeatureTypes.h"
-#include "Playback/AvaPlaybackServer.h"
+#include "Playback/IAvaPlaybackServer.h"
 #include "Playback/IAvaPlaybackClient.h"
 #include "StormSyncAvaBridgeCommon.h"
 #include "StormSyncAvaBridgeLog.h"
@@ -218,7 +218,7 @@ bool FStormSyncAvaSyncProvider::GetAddressFromUserData(const FString& InRemoteNa
 
 bool FStormSyncAvaSyncProvider::GetAddressFromServerUserData(const FString& InServerName, const FString& InUserDataKey, FMessageAddress& OutAddress, FText* OutErrorMessage)
 {
-	if (!IAvaMediaModule::Get().IsMediaPlaybackClientStarted())
+	if (!IAvaMediaModule::Get().IsPlaybackClientStarted())
 	{
 		if (OutErrorMessage)
 		{
@@ -227,7 +227,7 @@ bool FStormSyncAvaSyncProvider::GetAddressFromServerUserData(const FString& InSe
 		return false;
 	}
 
-	const IAvaPlaybackClient& PlaybackClient = IAvaMediaModule::Get().GetMediaPlaybackClient();
+	const IAvaPlaybackClient& PlaybackClient = IAvaMediaModule::Get().GetPlaybackClient();
 	const FString AddressId = PlaybackClient.GetServerUserData(InServerName, InUserDataKey);
 	if (AddressId.IsEmpty())
 	{
@@ -247,7 +247,7 @@ bool FStormSyncAvaSyncProvider::GetAddressFromServerUserData(const FString& InSe
 
 bool FStormSyncAvaSyncProvider::GetAddressFromClientUserData(const FString& InClientName, const FString& InUserDataKey, FMessageAddress& OutAddress, FText* OutErrorMessage)
 {
-	if (!IAvaMediaModule::Get().IsMediaPlaybackServerStarted())
+	if (!IAvaMediaModule::Get().IsPlaybackServerStarted())
 	{
 		if (OutErrorMessage)
 		{
@@ -256,8 +256,9 @@ bool FStormSyncAvaSyncProvider::GetAddressFromClientUserData(const FString& InCl
 		return false;
 	}
 
-	const TSharedPtr<FAvaPlaybackServer> PlaybackServer = IAvaMediaModule::Get().GetMediaPlaybackServer();
-	if (!PlaybackServer.IsValid())
+	const IAvaPlaybackServer* PlaybackServer = IAvaMediaModule::Get().GetPlaybackServer();
+	
+	if (!PlaybackServer)
 	{
 		if (OutErrorMessage)
 		{
@@ -285,12 +286,12 @@ bool FStormSyncAvaSyncProvider::GetAddressFromClientUserData(const FString& InCl
 
 bool FStormSyncAvaSyncProvider::IsPlaybackServer()
 {
-	return IAvaMediaModule::Get().IsMediaPlaybackServerStarted();
+	return IAvaMediaModule::Get().IsPlaybackServerStarted();
 }
 
 bool FStormSyncAvaSyncProvider::IsPlaybackClient()
 {
-	return IAvaMediaModule::Get().IsMediaPlaybackClientStarted();
+	return IAvaMediaModule::Get().IsPlaybackClientStarted();
 }
 
 TSharedPtr<FAvaMediaSyncCompareResponse> FStormSyncAvaSyncProvider::CreateErrorResponse(const FText& InText)
diff --git a/Engine/Plugins/Experimental/StormSyncAvalancheBridge/Source/StormSyncAvaBridge/Private/StormSyncAvaBridge.cpp b/Engine/Plugins/Experimental/StormSyncAvalancheBridge/Source/StormSyncAvaBridge/Private/StormSyncAvaBridge.cpp
--- a/Engine/Plugins/Experimental/StormSyncAvalancheBridge/Source/StormSyncAvaBridge/Private/StormSyncAvaBridge.cpp
+++ b/Engine/Plugins/Experimental/StormSyncAvalancheBridge/Source/StormSyncAvaBridge/Private/StormSyncAvaBridge.cpp
@@ -8,11 +8,10 @@
 #include "IStormSyncTransportServerModule.h"
 #include "Misc/CoreDelegates.h"
 #include "ModularFeature/StormSyncAvaSyncProvider.h"
-#include "Playback/AvaPlaybackServer.h"
+#include "Playback/IAvaPlaybackServer.h"
 #include "Playback/IAvaPlaybackClient.h"
 #include "StormSyncAvaBridgeCommon.h"
 #include "StormSyncAvaBridgeLog.h"
-#include "StormSyncAvaBridgeLog.h"
 #include "StormSyncAvaBridgeUtils.h"
 #include "StormSyncCoreDelegates.h"
 
@@ -51,12 +50,12 @@ void FStormSyncAvaBridgeModule::OnPostEngineInit()
 {
 	if (IAvaMediaModule::IsModuleLoaded())
 	{
-		if (IAvaMediaModule::Get().IsMediaPlaybackServerStarted())
+		if (IAvaMediaModule::Get().IsPlaybackServerStarted())
 		{
 			RegisterUserDataForPlaybackServer();
 		}
 
-		if (IAvaMediaModule::Get().IsMediaPlaybackClientStarted())
+		if (IAvaMediaModule::Get().IsPlaybackClientStarted())
 		{
 			RegisterUserDataForPlaybackClient();
 		}
@@ -120,7 +119,7 @@ void FStormSyncAvaBridgeModule::ExecuteGetUserData(const TArray<FString>& Args)
 		return;
 	}
 
-	if (!IAvaMediaModule::Get().IsMediaPlaybackClientStarted())
+	if (!IAvaMediaModule::Get().IsPlaybackClientStarted())
 	{
 		UE_LOG(LogStormSyncAvaBridge, Error, TEXT("Ava Playback Client stopped"));
 		return;
@@ -136,7 +135,7 @@ void FStormSyncAvaBridgeModule::ExecuteGetUserData(const TArray<FString>& Args)
 		return;
 	}
 
-	const IAvaPlaybackClient& PlaybackClient = IAvaMediaModule::Get().GetMediaPlaybackClient();
+	const IAvaPlaybackClient& PlaybackClient = IAvaMediaModule::Get().GetPlaybackClient();
 	for (const FString& ServerName : ServerNames)
 	{
 		FString ServerValue = PlaybackClient.GetServerUserData(ServerName, Key);
@@ -164,7 +163,7 @@ void FStormSyncAvaBridgeModule::RegisterUserDataForPlaybackServer()
 		return;
 	}
 
-	if (const TSharedPtr<FAvaPlaybackServer> PlaybackServer = IAvaMediaModule::Get().GetMediaPlaybackServer())
+	if (IAvaPlaybackServer* PlaybackServer = IAvaMediaModule::Get().GetPlaybackServer())
 	{
 		const FString StormSyncServerAddress = IStormSyncTransportServerModule::Get().GetServerEndpointMessageAddressId();
 		const FString DiscoveryManagerAddress = IStormSyncTransportServerModule::Get().GetDiscoveryManagerMessageAddressId();
@@ -188,9 +187,9 @@ void FStormSyncAvaBridgeModule::RegisterUserDataForPlaybackClient()
 		return;
 	}
 
-	if (IAvaMediaModule::Get().IsMediaPlaybackClientStarted())
+	if (IAvaMediaModule::Get().IsPlaybackClientStarted())
 	{
-		IAvaPlaybackClient& PlaybackClient = IAvaMediaModule::Get().GetMediaPlaybackClient();
+		IAvaPlaybackClient& PlaybackClient = IAvaMediaModule::Get().GetPlaybackClient();
 
 		const FString StormSyncServerAddress = IStormSyncTransportServerModule::Get().GetServerEndpointMessageAddressId();
 		const FString DiscoveryManagerAddress = IStormSyncTransportServerModule::Get().GetDiscoveryManagerMessageAddressId();
diff --git a/Engine/Plugins/Experimental/StormSyncAvalancheBridge/Source/StormSyncAvaBridgeEditor/Private/StormSyncAvaRundownExtender.cpp b/Engine/Plugins/Experimental/StormSyncAvalancheBridge/Source/StormSyncAvaBridgeEditor/Private/StormSyncAvaRundownExtender.cpp
--- a/Engine/Plugins/Experimental/StormSyncAvalancheBridge/Source/StormSyncAvaBridgeEditor/Private/StormSyncAvaRundownExtender.cpp
+++ b/Engine/Plugins/Experimental/StormSyncAvalancheBridge/Source/StormSyncAvaBridgeEditor/Private/StormSyncAvaRundownExtender.cpp
@@ -276,9 +276,9 @@ void FStormSyncAvaRundownExtender::HandleInitializeAction(const UAvaRundown* InR
 			RemoteServerNames.Append(ServerNames);
 		}
 
-		if (IAvaMediaModule::Get().IsMediaPlaybackClientStarted())
+		if (IAvaMediaModule::Get().IsPlaybackClientStarted())
 		{
-			const IAvaPlaybackClient& PlaybackClient = IAvaMediaModule::Get().GetMediaPlaybackClient();
+			const IAvaPlaybackClient& PlaybackClient = IAvaMediaModule::Get().GetPlaybackClient();
 			for (const FString& ServerName : RemoteServerNames)
 			{
 				// Get address id for storm sync client on playback host
