[Motion Design] Font manager doesn't update font list when fonts are imported by [REDACTED_PERSON]

diff --git a/Engine/Plugins/Experimental/Avalanche/Source/AvalancheTextEditor/Private/Font/AvaFontManagerSubsystem.cpp b/Engine/Plugins/Experimental/Avalanche/Source/AvalancheTextEditor/Private/Font/AvaFontManagerSubsystem.cpp
--- a/Engine/Plugins/Experimental/Avalanche/Source/AvalancheTextEditor/Private/Font/AvaFontManagerSubsystem.cpp
+++ b/Engine/Plugins/Experimental/Avalanche/Source/AvalancheTextEditor/Private/Font/AvaFontManagerSubsystem.cpp
@@ -111,26 +111,16 @@ bool UAvaFontManagerSubsystem::IsSupportedFontFile(const FString& InFontFilePath
 	return Extension == TEXT("ttf") || Extension == TEXT("otf");
 }
 
-void UAvaFontManagerSubsystem::OnAssetCreated(UObject* InObject)
+void UAvaFontManagerSubsystem::OnAssetsAdded(const FAssetData& InAssetData)
 {
-	if (!IsValid(InObject))
+	if (UFont* CurrFont = Cast<UFont>(InAssetData.GetAsset()))
 	{
-		return;
-	}
-
-	if (InObject->GetClass()->IsChildOf<UFont>())
-	{
-		UFont* NewFont = Cast<UFont>(InObject);
-
-		if (IsValid(NewFont))
-		{
-			CreateProjectFont(NewFont);
-		}
+		CreateProjectFont(CurrFont);
 
 		RefreshAvalancheFontsMap();
 		RefreshAvalancheFontsOptions();
 
-		OnProjectFontCreatedDelegate.Broadcast(NewFont);
+		OnProjectFontCreatedDelegate.Broadcast(CurrFont);
 	}
 }
 
@@ -500,7 +490,7 @@ void UAvaFontManagerSubsystem::UnregisterAssetsCallbacks() const
 
 		if (AssetRegistryModule.IsValid())
 		{
-			AssetRegistryModule.Get().OnInMemoryAssetCreated().RemoveAll(this);
+			AssetRegistryModule.Get().OnAssetAdded().RemoveAll(this);
 			AssetRegistryModule.Get().OnInMemoryAssetDeleted().RemoveAll(this);
 		}
 	}
@@ -512,7 +502,7 @@ void UAvaFontManagerSubsystem::RegisterAssetsCallbacks()
 
 	if (AssetRegistryModule.IsValid())
 	{
-		AssetRegistryModule.Get().OnInMemoryAssetCreated().AddUObject(this, &UAvaFontManagerSubsystem::OnAssetCreated);
+		AssetRegistryModule.Get().OnAssetAdded().AddUObject(this, &UAvaFontManagerSubsystem::OnAssetsAdded);
 		AssetRegistryModule.Get().OnInMemoryAssetDeleted().AddUObject(this, &UAvaFontManagerSubsystem::OnAssetDeleted);
 	}
 }
diff --git a/Engine/Plugins/Experimental/Avalanche/Source/AvalancheTextEditor/Private/Font/AvaFontManagerSubsystem.h b/Engine/Plugins/Experimental/Avalanche/Source/AvalancheTextEditor/Private/Font/AvaFontManagerSubsystem.h
--- a/Engine/Plugins/Experimental/Avalanche/Source/AvalancheTextEditor/Private/Font/AvaFontManagerSubsystem.h
+++ b/Engine/Plugins/Experimental/Avalanche/Source/AvalancheTextEditor/Private/Font/AvaFontManagerSubsystem.h
@@ -167,7 +167,7 @@ private:
 
 	void ClearFontsData();
 	void LoadFavorites();
-	void OnAssetCreated(UObject* InObject);
+	void OnAssetsAdded(const FAssetData& InAssetData);
 	void OnAssetDeleted(UObject* InObject);
 	void LoadOSFonts();
 	void LoadProjectFonts();
