Motion Design:
- Add setter for Tag Container attribute

#changelist validated
#virtualized

diff --git a/Engine/Plugins/Experimental/Avalanche/Source/AvalancheAttribute/Private/Tags/AvaTagContainerAttribute.cpp b/Engine/Plugins/Experimental/Avalanche/Source/AvalancheAttribute/Private/Tags/AvaTagContainerAttribute.cpp
--- a/Engine/Plugins/Experimental/Avalanche/Source/AvalancheAttribute/Private/Tags/AvaTagContainerAttribute.cpp
+++ b/Engine/Plugins/Experimental/Avalanche/Source/AvalancheAttribute/Private/Tags/AvaTagContainerAttribute.cpp
@@ -14,4 +14,9 @@ bool UAvaTagContainerAttribute::ContainsTag(const FAvaTagHandle& InTagHandle) co
 	return TagContainer.ContainsTag(InTagHandle);
 }
 
+void UAvaTagContainerAttribute::SetTagContainer(const FAvaTagHandleContainer& InTagContainer)
+{
+	TagContainer = InTagContainer;
+}
+
 #undef LOCTEXT_NAMESPACE
diff --git a/Engine/Plugins/Experimental/Avalanche/Source/AvalancheAttribute/Public/Tags/AvaTagContainerAttribute.h b/Engine/Plugins/Experimental/Avalanche/Source/AvalancheAttribute/Public/Tags/AvaTagContainerAttribute.h
--- a/Engine/Plugins/Experimental/Avalanche/Source/AvalancheAttribute/Public/Tags/AvaTagContainerAttribute.h
+++ b/Engine/Plugins/Experimental/Avalanche/Source/AvalancheAttribute/Public/Tags/AvaTagContainerAttribute.h
@@ -21,6 +21,9 @@ public:
 	AVALANCHEATTRIBUTE_API virtual bool ContainsTag(const FAvaTagHandle& InTagHandle) const override;
 	//~ End UAvaTagAttributeBase
 
-	UPROPERTY(EditAnywhere, Category="Attributes")
+	UFUNCTION()
+	AVALANCHEATTRIBUTE_API void SetTagContainer(const FAvaTagHandleContainer& InTagContainer);
+
+	UPROPERTY(EditAnywhere, Setter, Category="Attributes")
 	FAvaTagHandleContainer TagContainer;
 };
