Motion Design Tag:
- Add "none" option for Single Tag Picker

#changelist validated
#virtualized

diff --git a/Engine/Plugins/Experimental/Avalanche/Source/AvalancheTagEditor/Private/Customization/SAvaTagPicker.cpp b/Engine/Plugins/Experimental/Avalanche/Source/AvalancheTagEditor/Private/Customization/SAvaTagPicker.cpp
--- a/Engine/Plugins/Experimental/Avalanche/Source/AvalancheTagEditor/Private/Customization/SAvaTagPicker.cpp
+++ b/Engine/Plugins/Experimental/Avalanche/Source/AvalancheTagEditor/Private/Customization/SAvaTagPicker.cpp
@@ -177,6 +177,12 @@ void SAvaTagPicker::RefreshTagOptions()
 	{
 		TArray<FAvaTagId> TagIds = TagCollection->GetTagIds();
 
+		// Add default (none) TagId as an option, if there's no multiple choice
+		if (!TagCustomizer->AllowMultipleTags())
+		{
+			TagOptions.Add(MakeShared<FAvaTagHandle>(TagCollection, FAvaTagId()));
+		}
+
 		for (const FAvaTagId& TagId : TagIds)
 		{
 			TSharedRef<FAvaTagHandle> TagOption = MakeShared<FAvaTagHandle>(TagCollection, TagId);
@@ -265,7 +271,7 @@ bool SAvaTagPicker::IsTagHandleSelected(FAvaTagHandle InTagHandle) const
 
 void SAvaTagPicker::OnTagHandleSelectionChanged(const FAvaTagHandle& InTagHandle, bool bInIsSelected)
 {
-	if (!InTagHandle.IsValid())
+	if (!InTagHandle.Source)
 	{
 		return;
 	}
