[Motion Design] AvaMediaEditor module public API cleanup.

Moving all Broadcast and Playback to private.

#rnx

diff --git a/Engine/Plugins/Experimental/Avalanche/Source/AvalancheMediaEditor/Private/Broadcast/AvaBroadcastEditor.h b/Engine/Plugins/Experimental/Avalanche/Source/AvalancheMediaEditor/Private/Broadcast/AvaBroadcastEditor.h
new file mode 100644
--- /dev/null
+++ b/Engine/Plugins/Experimental/Avalanche/Source/AvalancheMediaEditor/Private/Broadcast/AvaBroadcastEditor.h
@@ -0,0 +1,92 @@
+// Copyright Epic Games, Inc. All Rights Reserved.
+
+#pragma once
+
+#include "AvaMediaDefines.h"
+#include "WorkflowOrientedApp/WorkflowCentricApplication.h"
+#include "Templates/SharedPointer.h"
+#include "UObject/WeakObjectPtrTemplates.h"
+
+class FAvaBroadcastOutputTileItem;
+class FName;
+class FReply;
+class FText;
+class FToolBarBuilder;
+class FExtender;
+class UAvaBroadcast;
+class UMediaOutput;
+class UObject;
+struct FLinearColor;
+struct FPropertyChangedEvent;
+struct FSlateBrush;
+
+class FAvaBroadcastEditor : public FWorkflowCentricApplication
+{
+	struct FPrivateToken { explicit FPrivateToken() = default; };
+	
+public:
+	explicit FAvaBroadcastEditor(FPrivateToken) {}
+	
+	virtual ~FAvaBroadcastEditor() override;
+
+	static void OpenBroadcastEditor();
+
+	void SelectOutputTile(const TSharedPtr<FAvaBroadcastOutputTileItem>& InOutputTile);
+	TSharedPtr<FAvaBroadcastOutputTileItem> GetSelectedOutputTile() const;
+	
+	DECLARE_MULTICAST_DELEGATE_OneParam(FOnOutputItemSelectionChanged, const TSharedPtr<FAvaBroadcastOutputTileItem>&);
+	FOnOutputItemSelectionChanged OnOutputTileSelectionChanged;
+	
+protected:
+	void InitBroadcastEditor(UAvaBroadcast* InBroadcast);
+
+	void OnBroadcastChanged(EAvaBroadcastChange ChangedEvent);
+	void OnAvaMediaSettingsChanged(UObject*, FPropertyChangedEvent&);
+
+	virtual void SaveAsset_Execute() override;
+	virtual bool CanSaveAssetAs() const override { return false; }
+	virtual bool CanFindInContentBrowser() const override { return false; }
+
+	virtual void OnClose() override;
+
+	//~ Begin IToolkit Interface
+	virtual FName GetToolkitFName() const override;
+	virtual FText GetBaseToolkitName() const override;
+	virtual FText GetToolkitName() const override;
+	virtual FText GetToolkitToolTipText() const override;
+	virtual FString GetWorldCentricTabPrefix() const override;
+	virtual FLinearColor GetWorldCentricTabColorScale() const override;
+	//~ End IToolkit Interface
+
+	//~ Begin FAssetEditorToolkit Interface
+	virtual const FSlateBrush* GetDefaultTabIcon() const override;
+	//~ End FAssetEditorToolkit Interface
+
+public:
+	UAvaBroadcast* GetBroadcastObject() const;
+
+	void ExtendToolBar(TSharedPtr<FExtender> Extender);
+	void FillPlayToolBar(FToolBarBuilder& ToolBarBuilder);
+
+	FText GetCurrentProfileName() const;
+	void MakeProfilesToolbar(FToolBarBuilder& ToolBarBuilder);
+	TSharedRef<SWidget> MakeProfileComboButton();
+	
+	void CreateNewProfile();
+	FReply OnProfileSelected(FName InProfileName);
+	
+protected:
+	void RegisterApplicationModes();
+	void CreateDefaultCommands();
+
+	static FText GetStopPlaybackClientTooltip();
+	static FText GetLaunchLocalServerTooltip();
+	static void StartPlaybackClientAction();
+
+protected:
+	static TSharedPtr<FAvaBroadcastEditor> BroadcastEditor;
+	
+	TWeakObjectPtr<UAvaBroadcast> BroadcastWeak;
+	
+	TWeakPtr<FAvaBroadcastOutputTileItem> SelectedOutputTileWeak;	
+};
diff --git a/Engine/Plugins/Experimental/Avalanche/Source/AvalancheMediaEditor/Private/Broadcast/ChannelGrid/AvaBroadcastOutputTileItem.h b/Engine/Plugins/Experimental/Avalanche/Source/AvalancheMediaEditor/Private/Broadcast/ChannelGrid/AvaBroadcastOutputTileItem.h
new file mode 100644
--- /dev/null
+++ b/Engine/Plugins/Experimental/Avalanche/Source/AvalancheMediaEditor/Private/Broadcast/ChannelGrid/AvaBroadcastOutputTileItem.h
@@ -0,0 +1,70 @@
+// Copyright Epic Games, Inc. All Rights Reserved.
+
+#pragma once
+
+#include "Broadcast/Channel/AvaBroadcastOutputChannel.h"
+#include "Delegates/IDelegateInstance.h"
+#include "Internationalization/Text.h"
+#include "Templates/SharedPointer.h"
+#include "UObject/NameTypes.h"
+#include "UObject/WeakObjectPtrTemplates.h"
+
+class FAvaBroadcastOutputTileItem;
+class FReply;
+class SWidget;
+class UObject;
+class UMediaOutput;
+struct FAvaBroadcastOutputChannel;
+struct FGeometry;
+struct FPointerEvent;
+struct FPropertyChangedEvent;
+struct FSlateBrush;
+enum class EAvaBroadcastChange : uint8;
+
+using FAvaBroadcastOutputTileItemPtr = TSharedPtr<FAvaBroadcastOutputTileItem>;
+
+class FAvaBroadcastOutputTileItem : public TSharedFromThis<FAvaBroadcastOutputTileItem>
+{
+public:
+	FAvaBroadcastOutputTileItem(FName InChannelName, UMediaOutput* InMediaOutput);
+	~FAvaBroadcastOutputTileItem();
+	
+	const FAvaBroadcastOutputChannel& GetChannel() const;
+	FAvaBroadcastOutputChannel& GetChannel();
+	
+	UMediaOutput* GetMediaOutput() const { return MediaOutput.Get(); }
+
+	TSharedRef<SWidget> GenerateTile() const;	
+
+	FText GetDisplayText() const;
+	FText GetMediaOutputStatusText() const;
+	FText GetToolTipText() const;
+
+	const FSlateBrush* GetMediaOutputIcon() const;
+	const FSlateBrush* GetMediaOutputStatusBrush() const;
+
+	FReply OnDragDetected(const FGeometry& MyGeometry, const FPointerEvent& MouseEvent);
+	
+protected:
+	void OnMediaOutputPropertyChanged(UObject* InObject, FPropertyChangedEvent& PropertyChangedEvent);
+	void OnChannelChanged(const FAvaBroadcastOutputChannel& InChannel, EAvaBroadcastChannelChange InChange);
+	void OnMediaOutputStateChanged(const FAvaBroadcastOutputChannel& InChannel, const UMediaOutput* InMediaOutput);
+	void OnBroadcastChanged(EAvaBroadcastChange InChange);
+
+	void UpdateInfo();
+	
+	FText FindLatestDisplayText() const;
+
+protected:
+	FName ChannelName = NAME_None;
+
+	TWeakObjectPtr<UMediaOutput> MediaOutput;
+	
+	FText MediaOutputDisplayText;
+	FText MediaOutputStatusText;
+	FText MediaOutputToolTipText;
+	
+	const FSlateBrush* MediaOutputStatusBrush = nullptr;
+
+	FDelegateHandle BroadcastChangedHandle;
+};
diff --git a/Engine/Plugins/Experimental/Avalanche/Source/AvalancheMediaEditor/Private/Playback/AvaPlaybackCommands.h b/Engine/Plugins/Experimental/Avalanche/Source/AvalancheMediaEditor/Private/Playback/AvaPlaybackCommands.h
new file mode 100644
--- /dev/null
+++ b/Engine/Plugins/Experimental/Avalanche/Source/AvalancheMediaEditor/Private/Playback/AvaPlaybackCommands.h
@@ -0,0 +1,24 @@
+// Copyright Epic Games, Inc. All Rights Reserved.
+
+#pragma once
+
+#include "Framework/Commands/Commands.h"
+#include "Styling/AppStyle.h"
+
+class FAvaPlaybackCommands : public TCommands<FAvaPlaybackCommands>
+{
+public:
+	
+	FAvaPlaybackCommands() : TCommands<FAvaPlaybackCommands>(TEXT("AvaPlayback")
+		, NSLOCTEXT("MotionDesignPlaybackCommands", "MotionDesignPlaybackCommands", "Motion Design Playback")
+		, NAME_None
+		, FAppStyle::GetAppStyleSetName())
+	{
+	}
+
+	/** Initialize commands */
+	virtual void RegisterCommands() override;
+	
+	TSharedPtr<FUICommandInfo> AddInputPin;
+	TSharedPtr<FUICommandInfo> RemoveInputPin;
+};
diff --git a/Engine/Plugins/Experimental/Avalanche/Source/AvalancheMediaEditor/Private/Playback/AvaPlaybackGraphEditor.h b/Engine/Plugins/Experimental/Avalanche/Source/AvalancheMediaEditor/Private/Playback/AvaPlaybackGraphEditor.h
new file mode 100644
--- /dev/null
+++ b/Engine/Plugins/Experimental/Avalanche/Source/AvalancheMediaEditor/Private/Playback/AvaPlaybackGraphEditor.h
@@ -0,0 +1,139 @@
+// Copyright Epic Games, Inc. All Rights Reserved.
+
+#pragma once
+
+#include "Containers/Array.h"
+#include "Containers/Map.h"
+#include "Containers/Set.h"
+#include "Containers/StringFwd.h"
+#include "Math/MathFwd.h"
+#include "Playback/IAvaPlaybackGraphEditor.h"
+#include "Templates/SharedPointer.h"
+#include "Templates/SubclassOf.h"
+#include "Toolkits/IToolkit.h"
+#include "Toolkits/IToolkitHost.h"
+#include "Types/SlateEnums.h"
+#include "WorkflowOrientedApp/WorkflowCentricApplication.h"
+#include "UObject/WeakObjectPtrTemplates.h"
+
+class FExtender;
+class FName;
+class FSlateRect;
+class FText;
+class FToolBarBuilder;
+class FUICommandList;
+class SGraphEditor;
+class UAvaPlaybackEditorGraphNode;
+class UAvaPlaybackGraph;
+class UAvaPlaybackNode;
+class UEdGraph;
+class UEdGraphNode;
+class UObject;
+struct FLinearColor;
+
+class FAvaPlaybackGraphEditor
+	: public FWorkflowCentricApplication
+	, public IAvaPlaybackGraphEditor
+{
+public:
+	void InitPlaybackEditor(const EToolkitMode::Type InMode
+		, const TSharedPtr<IToolkitHost>& InitToolkitHost
+		, UAvaPlaybackGraph* InPlayback);
+
+	//~ Begin IToolkit
+	virtual FName GetToolkitFName() const override;
+	virtual FText GetBaseToolkitName() const override;
+	virtual FString GetWorldCentricTabPrefix() const override;
+	virtual FLinearColor GetWorldCentricTabColorScale() const override;
+	//~ End IToolkit
+
+	UAvaPlaybackGraph* GetPlaybackObject() const;
+	
+	void ExtendToolBar(TSharedPtr<FExtender> Extender);
+	void FillPlayToolBar(FToolBarBuilder& ToolBarBuilder);
+
+public:
+	static void CacheOverrideNodeClasses();
+	
+	TSharedRef<SGraphEditor> CreateGraphEditor();
+	
+	//~ Begin IAvaPlaybackGraphEditor
+	virtual UEdGraph* CreatePlaybackGraph(UAvaPlaybackGraph* InPlayback) override;
+	virtual void SetupPlaybackNode(UEdGraph* InGraph, UAvaPlaybackNode* InPlaybackNode, bool bSelectNewNode) override;
+	virtual void CompilePlaybackNodesFromGraphNodes(UAvaPlaybackGraph* InPlayback) override;
+	virtual void CreateInputPin(UEdGraphNode* InGraphNode) override;
+	virtual void RefreshNode(UEdGraphNode& InGraphNode) override;
+	virtual bool GetBoundsForSelectedNodes(FSlateRect& Rect, float Padding) override;
+	virtual TSet<UObject*> GetSelectedNodes() const override;
+	//~ End IAvaPlaybackGraphEditor
+	
+	/** Called when the selection changes in the GraphEditor */
+	void OnSelectedNodesChanged(const TSet<UObject*>& NewSelection);
+
+	/**
+	 * Called when a node's title is committed for a rename
+	 *
+	 * @param	NewText				New title text
+	 * @param	CommitInfo			How text was committed
+	 * @param	NodeBeingChanged	The node being changed
+	 */
+	void OnNodeTitleCommitted(const FText& NewText, ETextCommit::Type CommitInfo, UEdGraphNode* NodeBeingChanged);
+	
+	DECLARE_MULTICAST_DELEGATE_OneParam(FOnPlaybackSelectionChanged, const TArray<UObject*>&);
+	FOnPlaybackSelectionChanged OnPlaybackSelectionChanged;
+	
+protected:
+	void RegisterApplicationModes();
+	
+	void CreateDefaultCommands();
+	void CreateGraphCommands();
+
+public:
+	bool CanAddInputPin() const;
+	void AddInputPin();
+
+	bool CanRemoveInputPin() const;
+	void RemoveInputPin();
+	
+	void CreateComment();
+	
+	bool CanSelectAllNodes() const { return true; }
+	void SelectAllNodes();
+
+	bool CanDeleteSelectedNodes() const;
+	void DeleteSelectedNodes();
+	void DeleteSelectedDuplicatableNodes();
+	
+	bool CanCopySelectedNodes() const;
+	void CopySelectedNodes();
+	
+	bool CanCutSelectedNodes() const;
+	void CutSelectedNodes();
+	
+	virtual bool CanPasteNodes() const override;
+	void PasteNodes();
+	virtual void PasteNodesHere(const FVector2D& Location) override;
+	
+	bool CanDuplicateSelectedNodes() const;
+	void DuplicateSelectedNodes();
+
+	void OnAlignTop();
+	void OnAlignMiddle();
+	void OnAlignBottom();
+	void OnAlignLeft();
+	void OnAlignCenter();
+	void OnAlignRight();
+
+	void OnStraightenConnections();
+	void OnDistributeNodesH();
+	void OnDistributeNodesV();
+	
+protected:
+	TWeakObjectPtr<UAvaPlaybackGraph> PlaybackGraphWeak;
+	
+	TSharedPtr<FUICommandList> GraphEditorCommands;
+
+	TSharedPtr<SGraphEditor> GraphEditor;
+	
+	static TMap<TSubclassOf<UAvaPlaybackNode>, TSubclassOf<UAvaPlaybackEditorGraphNode>> OverrideNodeClasses;
+};
diff --git a/Engine/Plugins/Experimental/Avalanche/Source/AvalancheMediaEditor/Private/Playback/Graph/AvaPlaybackEditorGraph.h b/Engine/Plugins/Experimental/Avalanche/Source/AvalancheMediaEditor/Private/Playback/Graph/AvaPlaybackEditorGraph.h
new file mode 100644
--- /dev/null
+++ b/Engine/Plugins/Experimental/Avalanche/Source/AvalancheMediaEditor/Private/Playback/Graph/AvaPlaybackEditorGraph.h
@@ -0,0 +1,22 @@
+// Copyright Epic Games, Inc. All Rights Reserved.
+
+#pragma once
+
+#include "EdGraph/EdGraph.h"
+#include "Templates/SubclassOf.h"
+#include "AvaPlaybackEditorGraph.generated.h"
+
+class UAvaPlaybackGraph;
+class UAvaPlaybackEditorGraphNode;
+
+UCLASS()
+class UAvaPlaybackEditorGraph : public UEdGraph
+{
+	GENERATED_BODY()
+
+public:
+	UAvaPlaybackGraph* GetPlaybackGraph() const;
+	
+	UAvaPlaybackEditorGraphNode* CreatePlaybackEditorGraphNode(TSubclassOf<UAvaPlaybackEditorGraphNode> NewNodeClass
+		, bool bSelectNewNode = true);
+};
diff --git a/Engine/Plugins/Experimental/Avalanche/Source/AvalancheMediaEditor/Private/Playback/Graph/AvaPlaybackEditorGraphSchema.h b/Engine/Plugins/Experimental/Avalanche/Source/AvalancheMediaEditor/Private/Playback/Graph/AvaPlaybackEditorGraphSchema.h
new file mode 100644
--- /dev/null
+++ b/Engine/Plugins/Experimental/Avalanche/Source/AvalancheMediaEditor/Private/Playback/Graph/AvaPlaybackEditorGraphSchema.h
@@ -0,0 +1,75 @@
+// Copyright Epic Games, Inc. All Rights Reserved.
+
+#pragma once
+
+#include "Containers/Array.h"
+#include "Containers/StringFwd.h"
+#include "EdGraph/EdGraphSchema.h"
+#include "Math/MathFwd.h"
+#include "Templates/SubclassOf.h"
+#include "Templates/SharedPointer.h"
+#include "AvaPlaybackEditorGraphSchema.generated.h"
+
+class IAvaPlaybackGraphEditor;
+class FName;
+class UAvaPlaybackGraph;
+class UAvaPlaybackNode;
+class UEdGraph;
+class UEdGraphNode;
+class UEdGraphPin;
+class UGraphNodeContextMenuContext;
+class UToolMenu;
+struct FAssetData;
+struct FEdGraphPinType;
+struct FGraphContextMenuBuilder;
+struct FLinearColor;
+struct FPinConnectionResponse;
+
+UCLASS()
+class UAvaPlaybackEditorGraphSchema : public UEdGraphSchema
+{
+	GENERATED_BODY()
+
+public:
+	static const FLinearColor ActivePinColor;
+	static const FLinearColor InactivePinColor;
+	
+	// Allowable PinType.PinCategory values
+	static const FName PC_ChannelFeed;
+	static const FName PC_Event;
+
+	void CompilePlaybackNodesFromGraphNodes(UEdGraphNode& Node) const;
+	
+	/** Check whether connecting these pins would cause a loop */
+	bool ConnectionCausesLoop(const UEdGraphPin* InputPin, const UEdGraphPin* OutputPin) const;
+	
+	//~ Begin UEdGraphSchema
+	virtual void GetGraphContextActions(FGraphContextMenuBuilder& ContextMenuBuilder) const override;
+	virtual void GetContextMenuActions(UToolMenu* Menu, class UGraphNodeContextMenuContext* Context) const override;
+	virtual void CreateDefaultNodesForGraph(UEdGraph& Graph) const override;
+	virtual bool CanDuplicateGraph(UEdGraph* InSourceGraph) const override { return false; }
+	virtual bool ShouldAlwaysPurgeOnModification() const override { return true; }
+	virtual const FPinConnectionResponse CanCreateConnection(const UEdGraphPin* PinA, const UEdGraphPin* PinB) const override;
+	virtual bool TryCreateConnection(UEdGraphPin* PinA, UEdGraphPin* PinB) const override;
+	virtual FLinearColor GetPinTypeColor(const FEdGraphPinType& PinType) const override;
+	virtual bool ShouldHidePinDefaultValue(UEdGraphPin* Pin) const override;
+	virtual void BreakNodeLinks(UEdGraphNode& TargetNode) const override;
+	virtual void BreakPinLinks(UEdGraphPin& TargetPin, bool bSendsNodeNotifcation) const override;
+	virtual void GetAssetsGraphHoverMessage(const TArray<FAssetData>& Assets, const UEdGraph* HoverGraph, FString& OutTooltipText, bool& OutOkIcon) const override;
+	virtual void DroppedAssetsOnGraph(const TArray<struct FAssetData>& Assets, const FVector2D& GraphPosition, UEdGraph* Graph) const override;
+	virtual void DroppedAssetsOnNode(const TArray<FAssetData>& Assets, const FVector2D& GraphPosition, UEdGraphNode* Node) const override;
+	virtual int32 GetNodeSelectionCount(const UEdGraph* Graph) const override;
+	virtual TSharedPtr<FEdGraphSchemaAction> GetCreateCommentAction() const override;
+	//~ End UEdGraphSchema
+
+	static void CachePlaybackNodeClasses();
+	
+	/** Adds actions for creating every type of Playback Node */
+	void GetPlaybackNodeActions(FGraphActionMenuBuilder& ActionMenuBuilder, bool bShowSelectedActions) const;
+	void GetCommentAction(FGraphActionMenuBuilder& ActionMenuBuilder, const UEdGraph* CurrentGraph = nullptr) const;
+
+protected:
+	static TArray<TSubclassOf<UAvaPlaybackNode>> PlaybackNodeClasses;
+
+	TSharedPtr<IAvaPlaybackGraphEditor> GetPlaybackGraphEditor(const UEdGraph* Graph) const;
+};
diff --git a/Engine/Plugins/Experimental/Avalanche/Source/AvalancheMediaEditor/Private/Playback/Graph/Nodes/AvaPlaybackEditorGraphNode.h b/Engine/Plugins/Experimental/Avalanche/Source/AvalancheMediaEditor/Private/Playback/Graph/Nodes/AvaPlaybackEditorGraphNode.h
new file mode 100644
--- /dev/null
+++ b/Engine/Plugins/Experimental/Avalanche/Source/AvalancheMediaEditor/Private/Playback/Graph/Nodes/AvaPlaybackEditorGraphNode.h
@@ -0,0 +1,115 @@
+// Copyright Epic Games, Inc. All Rights Reserved.
+
+#pragma once
+
+#include "Containers/Array.h"
+#include "Containers/Set.h"
+#include "EdGraph/EdGraphNode.h"
+#include "Templates/SubclassOf.h"
+#include "Templates/SharedPointer.h"
+#include "UObject/ObjectPtr.h"
+#include "AvaPlaybackEditorGraphNode.generated.h"
+
+class FName;
+class SGraphNode;
+class UAvaPlaybackGraph;
+class UAvaPlaybackNode;
+class UEdGraphPin;
+class UEdGraphSchema;
+class UGraphNodeContextMenuContext;
+class UToolMenu;
+class UObject;
+struct FLinearColor;
+
+UCLASS()
+class UAvaPlaybackEditorGraphNode : public UEdGraphNode
+{
+	GENERATED_BODY()
+
+public:
+	UAvaPlaybackGraph* GetPlayback() const;
+	
+	void SetPlaybackNode(UAvaPlaybackNode* InPlaybackNode);
+
+	double GetLastTimeTicked();
+	double GetChildLastTimeTicked(int32 ChildIndex) const;
+	
+	bool CanAddInputPin() const;
+	void AddInputPin();
+	void CreateInputPin();
+
+	void RemoveInputPin(UEdGraphPin* InGraphPin);
+
+	//~ Begin UEdGraphNode interface
+	virtual void PrepareForCopying() override;
+	virtual void PostPasteNode() override;
+	//~ End UEdGraphNode interface
+
+	void PostCopyNode();
+
+	/** Make sure the PlaybackNode is owned by the Playback Object */
+	void ResetPlaybackNodeOwner();
+	
+	//Return which Playback Node this Graph Node should Represent
+	virtual TSubclassOf<UAvaPlaybackNode> GetPlaybackNodeClass() const;
+
+	UAvaPlaybackNode* GetPlaybackNode() const
+	{
+		return PlaybackNode;
+	}
+	
+protected:
+	/** Create all of the input pins required */
+	void CreateInputPins();
+	void CreateOutputPin();
+
+public:
+	/** Is this the non-deletable root node */
+	bool IsRootNode() const;
+	
+	UEdGraphPin* GetOutputPin() const;
+	TArray<UEdGraphPin*> GetInputPins() const;
+	
+	UEdGraphPin* GetInputPin(int32 InputIndex) const;
+	int32 GetInputPinIndex(UEdGraphPin* InInputPin) const;
+	int32 GetInputPinCount() const;
+
+	/**
+	 * Handles inserting the node between the FromPin and what the FromPin was original connected to
+	 *
+	 * @param FromPin			The pin this node is being spawned from
+	 * @param NewLinkPin		The new pin the FromPin will connect to
+	 * @param OutNodeList		Any nodes that are modified will get added to this list for notification purposes
+	 */
+	void InsertNewNode(UEdGraphPin* FromPin, UEdGraphPin* NewLinkPin, TSet<UEdGraphNode*>& OutNodeList);
+
+	//~ Begin UObject
+	virtual void PostLoad() override;
+	virtual void PostEditImport() override;
+	virtual void PostDuplicate(bool bDuplicateForPIE) override;
+	//~ End UObject
+	
+	//~ Begin UEdGraphNode
+	virtual FText GetNodeTitle(ENodeTitleType::Type TitleType) const override;
+	virtual FText GetTooltipText() const override;
+	virtual void GetNodeContextMenuActions(UToolMenu* Menu, UGraphNodeContextMenuContext* Context) const override;
+	virtual void AllocateDefaultPins() override;
+	virtual void ReconstructNode() override;
+	virtual void AutowireNewNode(UEdGraphPin* FromPin) override;
+	virtual bool CanCreateUnderSpecifiedSchema(const UEdGraphSchema* Schema) const override;
+	virtual FLinearColor GetNodeTitleColor() const override;
+	virtual TSharedPtr<SGraphNode> CreateVisualWidget() override;
+	//~ End UEdGraphNode
+	
+	virtual FName GetInputPinName(int32 InputPinIndex) const;
+	virtual FName GetInputPinCategory(int32 InputPinIndex) const;
+	virtual FName GetInputPinSubCategory(int32 InputPinIndex) const;
+	virtual UObject* GetInputPinSubCategoryObject(int32 InputPinIndex) const;
+
+	virtual FName GetOutputPinCategory() const;
+	
+protected:
+	// The Playback Node this represents
+	UPROPERTY(VisibleAnywhere, Instanced, Category=Playback)
+	TObjectPtr<UAvaPlaybackNode> PlaybackNode;
+};
diff --git a/Engine/Plugins/Experimental/Avalanche/Source/AvalancheMediaEditor/Private/Playback/Graph/Nodes/AvaPlaybackEditorGraphNode_EventAction.h b/Engine/Plugins/Experimental/Avalanche/Source/AvalancheMediaEditor/Private/Playback/Graph/Nodes/AvaPlaybackEditorGraphNode_EventAction.h
new file mode 100644
--- /dev/null
+++ b/Engine/Plugins/Experimental/Avalanche/Source/AvalancheMediaEditor/Private/Playback/Graph/Nodes/AvaPlaybackEditorGraphNode_EventAction.h
@@ -0,0 +1,23 @@
+// Copyright Epic Games, Inc. All Rights Reserved.
+
+#pragma once
+
+#include "AvaPlaybackEditorGraphNode.h"
+#include "AvaPlaybackEditorGraphNode_EventAction.generated.h"
+
+UCLASS()
+class UAvaPlaybackEditorGraphNode_EventAction : public UAvaPlaybackEditorGraphNode
+{
+	GENERATED_BODY()
+
+public:
+	
+	virtual TSubclassOf<UAvaPlaybackNode> GetPlaybackNodeClass() const override;
+
+	virtual FName GetInputPinCategory(int32 InputPinIndex) const override;
+	virtual FName GetOutputPinCategory() const override;
+
+	//UEdGraphNode interface
+	virtual FLinearColor GetNodeTitleColor() const override;
+	//~UEdGraphNode interface
+};
diff --git a/Engine/Plugins/Experimental/Avalanche/Source/AvalancheMediaEditor/Private/Playback/Graph/Nodes/AvaPlaybackEditorGraphNode_EventFlow.h b/Engine/Plugins/Experimental/Avalanche/Source/AvalancheMediaEditor/Private/Playback/Graph/Nodes/AvaPlaybackEditorGraphNode_EventFlow.h
new file mode 100644
--- /dev/null
+++ b/Engine/Plugins/Experimental/Avalanche/Source/AvalancheMediaEditor/Private/Playback/Graph/Nodes/AvaPlaybackEditorGraphNode_EventFlow.h
@@ -0,0 +1,19 @@
+// Copyright Epic Games, Inc. All Rights Reserved.
+
+#pragma once
+
+#include "AvaPlaybackEditorGraphNode.h"
+#include "AvaPlaybackEditorGraphNode_EventFlow.generated.h"
+
+UCLASS()
+class UAvaPlaybackEditorGraphNode_EventFlow : public UAvaPlaybackEditorGraphNode
+{
+	GENERATED_BODY()
+
+public:
+	
+	virtual TSubclassOf<UAvaPlaybackNode> GetPlaybackNodeClass() const override;
+
+	virtual FName GetInputPinCategory(int32 InputPinIndex) const override;
+	virtual FName GetOutputPinCategory() const override;
+};
diff --git a/Engine/Plugins/Experimental/Avalanche/Source/AvalancheMediaEditor/Private/Playback/Graph/Nodes/AvaPlaybackEditorGraphNode_EventTrigger.h b/Engine/Plugins/Experimental/Avalanche/Source/AvalancheMediaEditor/Private/Playback/Graph/Nodes/AvaPlaybackEditorGraphNode_EventTrigger.h
new file mode 100644
--- /dev/null
+++ b/Engine/Plugins/Experimental/Avalanche/Source/AvalancheMediaEditor/Private/Playback/Graph/Nodes/AvaPlaybackEditorGraphNode_EventTrigger.h
@@ -0,0 +1,21 @@
+// Copyright Epic Games, Inc. All Rights Reserved.
+
+#pragma once
+
+#include "AvaPlaybackEditorGraphNode.h"
+#include "AvaPlaybackEditorGraphNode_EventTrigger.generated.h"
+
+UCLASS()
+class UAvaPlaybackEditorGraphNode_EventTrigger : public UAvaPlaybackEditorGraphNode
+{
+	GENERATED_BODY()
+
+public:
+	
+	virtual TSubclassOf<UAvaPlaybackNode> GetPlaybackNodeClass() const override;
+	virtual FName GetOutputPinCategory() const override;
+
+	//UEdGraphNode interface
+	virtual FLinearColor GetNodeTitleColor() const override;
+	//~UEdGraphNode interface
+};
diff --git a/Engine/Plugins/Experimental/Avalanche/Source/AvalancheMediaEditor/Private/Playback/Graph/Nodes/AvaPlaybackEditorGraphNode_LevelPlayer.h b/Engine/Plugins/Experimental/Avalanche/Source/AvalancheMediaEditor/Private/Playback/Graph/Nodes/AvaPlaybackEditorGraphNode_LevelPlayer.h
new file mode 100644
--- /dev/null
+++ b/Engine/Plugins/Experimental/Avalanche/Source/AvalancheMediaEditor/Private/Playback/Graph/Nodes/AvaPlaybackEditorGraphNode_LevelPlayer.h
@@ -0,0 +1,20 @@
+// Copyright Epic Games, Inc. All Rights Reserved.
+
+#pragma once
+
+#include "AvaPlaybackEditorGraphNode.h"
+#include "AvaPlaybackEditorGraphNode_LevelPlayer.generated.h"
+
+UCLASS()
+class UAvaPlaybackEditorGraphNode_LevelPlayer : public UAvaPlaybackEditorGraphNode
+{
+	GENERATED_BODY()
+	
+	virtual TSubclassOf<UAvaPlaybackNode> GetPlaybackNodeClass() const override;
+	virtual FName GetInputPinCategory(int32 InputPinIndex) const override;
+	
+	//UEdGraphNode interface
+	virtual FLinearColor GetNodeTitleColor() const override;
+	virtual TSharedPtr<SGraphNode> CreateVisualWidget() override;
+	//~UEdGraphNode interface
+};
diff --git a/Engine/Plugins/Experimental/Avalanche/Source/AvalancheMediaEditor/Private/Playback/Graph/Nodes/AvaPlaybackEditorGraphNode_Root.h b/Engine/Plugins/Experimental/Avalanche/Source/AvalancheMediaEditor/Private/Playback/Graph/Nodes/AvaPlaybackEditorGraphNode_Root.h
new file mode 100644
--- /dev/null
+++ b/Engine/Plugins/Experimental/Avalanche/Source/AvalancheMediaEditor/Private/Playback/Graph/Nodes/AvaPlaybackEditorGraphNode_Root.h
@@ -0,0 +1,21 @@
+// Copyright Epic Games, Inc. All Rights Reserved.
+
+#pragma once
+
+#include "AvaPlaybackEditorGraphNode.h"
+#include "AvaPlaybackEditorGraphNode_Root.generated.h"
+
+UCLASS()
+class UAvaPlaybackEditorGraphNode_Root : public UAvaPlaybackEditorGraphNode
+{
+	GENERATED_BODY()
+	
+	virtual TSubclassOf<UAvaPlaybackNode> GetPlaybackNodeClass() const override;
+
+	//UEdGraphNode interface
+	virtual FLinearColor GetNodeTitleColor() const override;
+	virtual bool CanUserDeleteNode() const override;
+	virtual bool CanDuplicateNode() const override;
+	virtual TSharedPtr<SGraphNode> CreateVisualWidget() override;
+	//~UEdGraphNode interface
+};
diff --git a/Engine/Plugins/Experimental/Avalanche/Source/AvalancheMediaEditor/Private/Playback/Graph/Nodes/Slate/SAvaPlaybackEditorGraphNode.h b/Engine/Plugins/Experimental/Avalanche/Source/AvalancheMediaEditor/Private/Playback/Graph/Nodes/Slate/SAvaPlaybackEditorGraphNode.h
new file mode 100644
--- /dev/null
+++ b/Engine/Plugins/Experimental/Avalanche/Source/AvalancheMediaEditor/Private/Playback/Graph/Nodes/Slate/SAvaPlaybackEditorGraphNode.h
@@ -0,0 +1,36 @@
+// Copyright Epic Games, Inc. All Rights Reserved.
+
+#pragma once
+
+#include "SGraphNode.h"
+#include "Templates/SharedPointer.h"
+#include "UObject/WeakObjectPtrTemplates.h"
+
+class FReply;
+class SGraphPin;
+class SVerticalBox;
+class UAvaPlaybackEditorGraphNode;
+class UEdGraphPin;
+struct EVisibility;
+
+class SAvaPlaybackEditorGraphNode : public SGraphNode
+{
+public:
+	SLATE_BEGIN_ARGS(SAvaPlaybackEditorGraphNode) {}
+	SLATE_END_ARGS()
+
+	/** Constructs this widget with InArgs */
+	void Construct(const FArguments& InArgs, UAvaPlaybackEditorGraphNode* InGraphNode);
+	virtual void PostConstruct() {};
+	
+protected:
+	//~ Begin SGraphNode
+	virtual void CreateOutputSideAddButton(TSharedPtr<SVerticalBox> OutputBox) override;
+	virtual EVisibility IsAddPinButtonVisible() const override;
+	virtual FReply OnAddPin() override;
+	virtual TSharedPtr<SGraphPin> CreatePinWidget(UEdGraphPin* Pin) const override;
+	//~ End SGraphNode
+	
+protected:
+	TWeakObjectPtr<UAvaPlaybackEditorGraphNode> PlaybackGraphNode;
+};
diff --git a/Engine/Plugins/Experimental/Avalanche/Source/AvalancheMediaEditor/Private/Playback/Graph/SchemaActions/AvaPlaybackAction_NewComment.h b/Engine/Plugins/Experimental/Avalanche/Source/AvalancheMediaEditor/Private/Playback/Graph/SchemaActions/AvaPlaybackAction_NewComment.h
new file mode 100644
--- /dev/null
+++ b/Engine/Plugins/Experimental/Avalanche/Source/AvalancheMediaEditor/Private/Playback/Graph/SchemaActions/AvaPlaybackAction_NewComment.h
@@ -0,0 +1,26 @@
+// Copyright Epic Games, Inc. All Rights Reserved.
+
+#pragma once
+
+#include "EdGraph/EdGraphSchema.h"
+#include "AvaPlaybackAction_NewComment.generated.h"
+
+USTRUCT()
+struct FAvaPlaybackAction_NewComment : public FEdGraphSchemaAction
+{
+	GENERATED_BODY()
+
+	FAvaPlaybackAction_NewComment() = default;
+
+	FAvaPlaybackAction_NewComment(FText InNodeCategory, FText InMenuDesc, FText InToolTip, const int32 InGrouping)
+		: FEdGraphSchemaAction(MoveTemp(InNodeCategory), MoveTemp(InMenuDesc), MoveTemp(InToolTip), InGrouping)
+	{
+	}
+
+	//~ Begin FEdGraphSchemaAction Interface
+	virtual UEdGraphNode* PerformAction(class UEdGraph* ParentGraph
+		, UEdGraphPin* FromPin
+		, const FVector2D Location
+		, bool bSelectNewNode = true) override;
+	//~ End FEdGraphSchemaAction Interface
+};
diff --git a/Engine/Plugins/Experimental/Avalanche/Source/AvalancheMediaEditor/Private/Playback/Graph/SchemaActions/AvaPlaybackAction_NewNode.h b/Engine/Plugins/Experimental/Avalanche/Source/AvalancheMediaEditor/Private/Playback/Graph/SchemaActions/AvaPlaybackAction_NewNode.h
new file mode 100644
--- /dev/null
+++ b/Engine/Plugins/Experimental/Avalanche/Source/AvalancheMediaEditor/Private/Playback/Graph/SchemaActions/AvaPlaybackAction_NewNode.h
@@ -0,0 +1,42 @@
+// Copyright Epic Games, Inc. All Rights Reserved.
+
+#pragma once
+
+#include "EdGraph/EdGraphSchema.h"
+#include "Math/MathFwd.h"
+#include "Playback/Nodes/AvaPlaybackNode.h"
+#include "Templates/SubclassOf.h"
+#include "AvaPlaybackAction_NewNode.generated.h"
+
+class FText;
+class UAvaPlaybackNode;
+class UEdGraph;
+class UEdGraphNode;
+
+USTRUCT()
+struct FAvaPlaybackAction_NewNode : public FEdGraphSchemaAction
+{
+	GENERATED_BODY()
+
+public:
+	FAvaPlaybackAction_NewNode() = default;
+	
+	FAvaPlaybackAction_NewNode(FText InNodeCategory, FText InMenuDesc, FText InToolTip, const int32 InGrouping);
+
+	//~ Begin FEdGraphSchemaAction
+	virtual UEdGraphNode* PerformAction(class UEdGraph* ParentGraph
+		, UEdGraphPin* FromPin
+		, const FVector2D Location
+		, bool bSelectNewNode = true) override;
+	//~ End FEdGraphSchemaAction
+
+	void SetPlaybackNodeClass(TSubclassOf<UAvaPlaybackNode> InPlaybackNodeClass);
+
+protected:
+	/** Connects new node to output of selected nodes */
+	void ConnectToSelectedNodes(UAvaPlaybackNode* NewNode, UEdGraph* ParentGraph) const;
+
+	/** Class of node we want to create */
+	UPROPERTY()
+	TSubclassOf<UAvaPlaybackNode> PlaybackNodeClass;
+};
diff --git a/Engine/Plugins/Experimental/Avalanche/Source/AvalancheMediaEditor/Private/Playback/Graph/SchemaActions/AvaPlaybackAction_PasteNode.h b/Engine/Plugins/Experimental/Avalanche/Source/AvalancheMediaEditor/Private/Playback/Graph/SchemaActions/AvaPlaybackAction_PasteNode.h
new file mode 100644
--- /dev/null
+++ b/Engine/Plugins/Experimental/Avalanche/Source/AvalancheMediaEditor/Private/Playback/Graph/SchemaActions/AvaPlaybackAction_PasteNode.h
@@ -0,0 +1,26 @@
+// Copyright Epic Games, Inc. All Rights Reserved.
+
+#pragma once
+
+#include "EdGraph/EdGraphSchema.h"
+#include "AvaPlaybackAction_PasteNode.generated.h"
+
+USTRUCT()
+struct FAvaPlaybackAction_PasteNode : public FEdGraphSchemaAction
+{
+	GENERATED_BODY()
+
+	FAvaPlaybackAction_PasteNode() = default;
+
+	FAvaPlaybackAction_PasteNode(FText InNodeCategory, FText InMenuDesc, FText InToolTip, const int32 InGrouping)
+		: FEdGraphSchemaAction(MoveTemp(InNodeCategory), MoveTemp(InMenuDesc), MoveTemp(InToolTip), InGrouping)
+	{
+	}
+
+	//FEdGraphSchemaAction Interface
+	virtual UEdGraphNode* PerformAction(class UEdGraph* ParentGraph
+		, UEdGraphPin* FromPin
+		, const FVector2D Location
+		, bool bSelectNewNode = true) override;
+	//~FEdGraphSchemaAction Interface
+};
diff --git a/Engine/Plugins/Experimental/Avalanche/Source/AvalancheMediaEditor/Public/Broadcast/AvaBroadcastEditor.h b/Engine/Plugins/Experimental/Avalanche/Source/AvalancheMediaEditor/Public/Broadcast/AvaBroadcastEditor.h
deleted file mode 100644
diff --git a/Engine/Plugins/Experimental/Avalanche/Source/AvalancheMediaEditor/Public/Broadcast/ChannelGrid/AvaBroadcastOutputTileItem.h b/Engine/Plugins/Experimental/Avalanche/Source/AvalancheMediaEditor/Public/Broadcast/ChannelGrid/AvaBroadcastOutputTileItem.h
deleted file mode 100644
diff --git a/Engine/Plugins/Experimental/Avalanche/Source/AvalancheMediaEditor/Public/Playback/AvaPlaybackCommands.h b/Engine/Plugins/Experimental/Avalanche/Source/AvalancheMediaEditor/Public/Playback/AvaPlaybackCommands.h
deleted file mode 100644
diff --git a/Engine/Plugins/Experimental/Avalanche/Source/AvalancheMediaEditor/Public/Playback/AvaPlaybackGraphEditor.h b/Engine/Plugins/Experimental/Avalanche/Source/AvalancheMediaEditor/Public/Playback/AvaPlaybackGraphEditor.h
deleted file mode 100644
diff --git a/Engine/Plugins/Experimental/Avalanche/Source/AvalancheMediaEditor/Public/Playback/Graph/AvaPlaybackEditorGraph.h b/Engine/Plugins/Experimental/Avalanche/Source/AvalancheMediaEditor/Public/Playback/Graph/AvaPlaybackEditorGraph.h
deleted file mode 100644
diff --git a/Engine/Plugins/Experimental/Avalanche/Source/AvalancheMediaEditor/Public/Playback/Graph/AvaPlaybackEditorGraphSchema.h b/Engine/Plugins/Experimental/Avalanche/Source/AvalancheMediaEditor/Public/Playback/Graph/AvaPlaybackEditorGraphSchema.h
deleted file mode 100644
diff --git a/Engine/Plugins/Experimental/Avalanche/Source/AvalancheMediaEditor/Public/Playback/Graph/Nodes/AvaPlaybackEditorGraphNode.h b/Engine/Plugins/Experimental/Avalanche/Source/AvalancheMediaEditor/Public/Playback/Graph/Nodes/AvaPlaybackEditorGraphNode.h
deleted file mode 100644
diff --git a/Engine/Plugins/Experimental/Avalanche/Source/AvalancheMediaEditor/Public/Playback/Graph/Nodes/AvaPlaybackEditorGraphNode_EventAction.h b/Engine/Plugins/Experimental/Avalanche/Source/AvalancheMediaEditor/Public/Playback/Graph/Nodes/AvaPlaybackEditorGraphNode_EventAction.h
deleted file mode 100644
diff --git a/Engine/Plugins/Experimental/Avalanche/Source/AvalancheMediaEditor/Public/Playback/Graph/Nodes/AvaPlaybackEditorGraphNode_EventFlow.h b/Engine/Plugins/Experimental/Avalanche/Source/AvalancheMediaEditor/Public/Playback/Graph/Nodes/AvaPlaybackEditorGraphNode_EventFlow.h
deleted file mode 100644
diff --git a/Engine/Plugins/Experimental/Avalanche/Source/AvalancheMediaEditor/Public/Playback/Graph/Nodes/AvaPlaybackEditorGraphNode_EventTrigger.h b/Engine/Plugins/Experimental/Avalanche/Source/AvalancheMediaEditor/Public/Playback/Graph/Nodes/AvaPlaybackEditorGraphNode_EventTrigger.h
deleted file mode 100644
diff --git a/Engine/Plugins/Experimental/Avalanche/Source/AvalancheMediaEditor/Public/Playback/Graph/Nodes/AvaPlaybackEditorGraphNode_LevelPlayer.h b/Engine/Plugins/Experimental/Avalanche/Source/AvalancheMediaEditor/Public/Playback/Graph/Nodes/AvaPlaybackEditorGraphNode_LevelPlayer.h
deleted file mode 100644
diff --git a/Engine/Plugins/Experimental/Avalanche/Source/AvalancheMediaEditor/Public/Playback/Graph/Nodes/AvaPlaybackEditorGraphNode_Root.h b/Engine/Plugins/Experimental/Avalanche/Source/AvalancheMediaEditor/Public/Playback/Graph/Nodes/AvaPlaybackEditorGraphNode_Root.h
deleted file mode 100644
diff --git a/Engine/Plugins/Experimental/Avalanche/Source/AvalancheMediaEditor/Public/Playback/Graph/Nodes/Slate/SAvaPlaybackEditorGraphNode.h b/Engine/Plugins/Experimental/Avalanche/Source/AvalancheMediaEditor/Public/Playback/Graph/Nodes/Slate/SAvaPlaybackEditorGraphNode.h
deleted file mode 100644
diff --git a/Engine/Plugins/Experimental/Avalanche/Source/AvalancheMediaEditor/Public/Playback/Graph/SchemaActions/AvaPlaybackAction_NewComment.h b/Engine/Plugins/Experimental/Avalanche/Source/AvalancheMediaEditor/Public/Playback/Graph/SchemaActions/AvaPlaybackAction_NewComment.h
deleted file mode 100644
diff --git a/Engine/Plugins/Experimental/Avalanche/Source/AvalancheMediaEditor/Public/Playback/Graph/SchemaActions/AvaPlaybackAction_NewNode.h b/Engine/Plugins/Experimental/Avalanche/Source/AvalancheMediaEditor/Public/Playback/Graph/SchemaActions/AvaPlaybackAction_NewNode.h
deleted file mode 100644
diff --git a/Engine/Plugins/Experimental/Avalanche/Source/AvalancheMediaEditor/Public/Playback/Graph/SchemaActions/AvaPlaybackAction_PasteNode.h b/Engine/Plugins/Experimental/Avalanche/Source/AvalancheMediaEditor/Public/Playback/Graph/SchemaActions/AvaPlaybackAction_PasteNode.h
deleted file mode 100644
