[Motion Design] Mask effects settings are persisted

diff --git a/Engine/Plugins/Experimental/Avalanche/Source/AvalancheMask/Private/Mask2D/AvaMask2DBaseModifier.cpp b/Engine/Plugins/Experimental/Avalanche/Source/AvalancheMask/Private/Mask2D/AvaMask2DBaseModifier.cpp
--- a/Engine/Plugins/Experimental/Avalanche/Source/AvalancheMask/Private/Mask2D/AvaMask2DBaseModifier.cpp
+++ b/Engine/Plugins/Experimental/Avalanche/Source/AvalancheMask/Private/Mask2D/AvaMask2DBaseModifier.cpp
@@ -177,6 +177,8 @@ void UAvaMask2DBaseModifier::OnFeatherChanged()
 
 void UAvaMask2DBaseModifier::OnCanvasChanged()
 {
+	CanvasParamsToLocal();
+	
 	OnBlurChanged();
 	OnFeatherChanged();
 }
@@ -262,6 +264,10 @@ void UAvaMask2DBaseModifier::OnModifierAdded(EActorModifierCoreEnableReason InRe
 	{
 		LocalParamsToCanvas();
 	}
+	else
+	{
+		CanvasParamsToLocal();
+	}
 }
 
 void UAvaMask2DBaseModifier::OnModifierRemoved(EActorModifierCoreDisableReason InReason)
