[Motion Design] Merging AvaRundownDefines.h and AvaRundownEditorDefines.h. There is already AvaRundownDefines.h in AvaMedia/Rundown.

#rnx

diff --git a/Engine/Plugins/Experimental/Avalanche/Source/AvalancheMediaEditor/Private/Rundown/AvaRundownDefines.cpp b/Engine/Plugins/Experimental/Avalanche/Source/AvalancheMediaEditor/Private/Rundown/AvaRundownDefines.cpp
deleted file mode 100644
diff --git a/Engine/Plugins/Experimental/Avalanche/Source/AvalancheMediaEditor/Private/Rundown/AvaRundownEditorDefines.cpp b/Engine/Plugins/Experimental/Avalanche/Source/AvalancheMediaEditor/Private/Rundown/AvaRundownEditorDefines.cpp
new file mode 100644
--- /dev/null
+++ b/Engine/Plugins/Experimental/Avalanche/Source/AvalancheMediaEditor/Private/Rundown/AvaRundownEditorDefines.cpp
@@ -0,0 +1,7 @@
+// Copyright Epic Games, Inc. All Rights Reserved.
+
+#include "Rundown/AvaRundownEditorDefines.h"
+
+const FNumberFormattingOptions UE::AvaRundown::FEditorMetrics::PageIdFormattingOptions = FNumberFormattingOptions()
+	.SetMinimumIntegralDigits(4)
+	.SetUseGrouping(false);
diff --git a/Engine/Plugins/Experimental/Avalanche/Source/AvalancheMediaEditor/Private/Rundown/ChannelLayerStatus/SAvaRundownChannelLayerStatusList.cpp b/Engine/Plugins/Experimental/Avalanche/Source/AvalancheMediaEditor/Private/Rundown/ChannelLayerStatus/SAvaRundownChannelLayerStatusList.cpp
--- a/Engine/Plugins/Experimental/Avalanche/Source/AvalancheMediaEditor/Private/Rundown/ChannelLayerStatus/SAvaRundownChannelLayerStatusList.cpp
+++ b/Engine/Plugins/Experimental/Avalanche/Source/AvalancheMediaEditor/Private/Rundown/ChannelLayerStatus/SAvaRundownChannelLayerStatusList.cpp
@@ -3,8 +3,8 @@
 #include "SAvaRundownChannelLayerStatusList.h"
 
 #include "Misc/CoreDelegates.h"
-#include "Rundown/AvaRundownDefines.h"
 #include "Rundown/AvaRundownEditor.h"
+#include "Rundown/AvaRundownEditorDefines.h"
 #include "Rundown/AvaRundownPagePlayer.h"
 #include "Widgets/Colors/SColorBlock.h"
 #include "Widgets/Images/SImage.h"
diff --git a/Engine/Plugins/Experimental/Avalanche/Source/AvalancheMediaEditor/Private/Rundown/DetailsView/SAvaRundownPageDetails.h b/Engine/Plugins/Experimental/Avalanche/Source/AvalancheMediaEditor/Private/Rundown/DetailsView/SAvaRundownPageDetails.h
--- a/Engine/Plugins/Experimental/Avalanche/Source/AvalancheMediaEditor/Private/Rundown/DetailsView/SAvaRundownPageDetails.h
+++ b/Engine/Plugins/Experimental/Avalanche/Source/AvalancheMediaEditor/Private/Rundown/DetailsView/SAvaRundownPageDetails.h
@@ -3,7 +3,7 @@
 #pragma once
 
 #include "Containers/Array.h"
-#include "Rundown/AvaRundownDefines.h"
+#include "Rundown/AvaRundownEditorDefines.h"
 #include "Templates/SharedPointer.h"
 #include "Types/SlateEnums.h"
 #include "Widgets/SCompoundWidget.h"
diff --git a/Engine/Plugins/Experimental/Avalanche/Source/AvalancheMediaEditor/Private/Rundown/Filters/AvaRundownPageFilterExpressionContext.cpp b/Engine/Plugins/Experimental/Avalanche/Source/AvalancheMediaEditor/Private/Rundown/Filters/AvaRundownPageFilterExpressionContext.cpp
--- a/Engine/Plugins/Experimental/Avalanche/Source/AvalancheMediaEditor/Private/Rundown/Filters/AvaRundownPageFilterExpressionContext.cpp
+++ b/Engine/Plugins/Experimental/Avalanche/Source/AvalancheMediaEditor/Private/Rundown/Filters/AvaRundownPageFilterExpressionContext.cpp
@@ -4,7 +4,7 @@
 
 #include "IAvaMediaEditorModule.h"
 #include "Rundown/AvaRundown.h"
-#include "Rundown/AvaRundownDefines.h"
+#include "Rundown/AvaRundownEditorDefines.h"
 #include "Rundown/AvaRundownPage.h"
 #include "Rundown/Factories/Filters/IAvaRundownFilterExpressionFactory.h"
 
diff --git a/Engine/Plugins/Experimental/Avalanche/Source/AvalancheMediaEditor/Private/Rundown/Filters/AvaRundownPageFilterExpressionContext.h b/Engine/Plugins/Experimental/Avalanche/Source/AvalancheMediaEditor/Private/Rundown/Filters/AvaRundownPageFilterExpressionContext.h
--- a/Engine/Plugins/Experimental/Avalanche/Source/AvalancheMediaEditor/Private/Rundown/Filters/AvaRundownPageFilterExpressionContext.h
+++ b/Engine/Plugins/Experimental/Avalanche/Source/AvalancheMediaEditor/Private/Rundown/Filters/AvaRundownPageFilterExpressionContext.h
@@ -4,7 +4,7 @@
 
 #include "Containers/UnrealString.h"
 #include "Misc/TextFilterExpressionEvaluator.h"
-#include "Rundown/AvaRundownDefines.h"
+#include "Rundown/AvaRundownEditorDefines.h"
 #include "Rundown/Factories/Filters/AvaRundownFactoriesUtils.h"
 
 class FName;
diff --git a/Engine/Plugins/Experimental/Avalanche/Source/AvalancheMediaEditor/Private/Rundown/Filters/AvaRundownPageTextFilter.h b/Engine/Plugins/Experimental/Avalanche/Source/AvalancheMediaEditor/Private/Rundown/Filters/AvaRundownPageTextFilter.h
--- a/Engine/Plugins/Experimental/Avalanche/Source/AvalancheMediaEditor/Private/Rundown/Filters/AvaRundownPageTextFilter.h
+++ b/Engine/Plugins/Experimental/Avalanche/Source/AvalancheMediaEditor/Private/Rundown/Filters/AvaRundownPageTextFilter.h
@@ -4,7 +4,7 @@
 
 #include "Misc/IFilter.h"
 #include "Misc/TextFilterExpressionEvaluator.h"
-#include "Rundown/AvaRundownDefines.h"
+#include "Rundown/AvaRundownEditorDefines.h"
 #include "Templates/SharedPointer.h"
 
 class FAvaRundownPageFilterExpressionContext;
diff --git a/Engine/Plugins/Experimental/Avalanche/Source/AvalancheMediaEditor/Private/Rundown/Pages/Columns/AvaRundownPageEnabledColumn.h b/Engine/Plugins/Experimental/Avalanche/Source/AvalancheMediaEditor/Private/Rundown/Pages/Columns/AvaRundownPageEnabledColumn.h
--- a/Engine/Plugins/Experimental/Avalanche/Source/AvalancheMediaEditor/Private/Rundown/Pages/Columns/AvaRundownPageEnabledColumn.h
+++ b/Engine/Plugins/Experimental/Avalanche/Source/AvalancheMediaEditor/Private/Rundown/Pages/Columns/AvaRundownPageEnabledColumn.h
@@ -2,7 +2,7 @@
 
 #pragma once
 
-#include "Rundown/AvaRundownDefines.h"
+#include "Rundown/AvaRundownEditorDefines.h"
 #include "Rundown/Pages/Columns/IAvaRundownPageViewColumn.h"
 
 class FAvaRundownPageEnabledColumn : public IAvaRundownPageViewColumn
diff --git a/Engine/Plugins/Experimental/Avalanche/Source/AvalancheMediaEditor/Private/Rundown/Pages/Columns/Slate/SAvaRundownPageChannelSelector.h b/Engine/Plugins/Experimental/Avalanche/Source/AvalancheMediaEditor/Private/Rundown/Pages/Columns/Slate/SAvaRundownPageChannelSelector.h
--- a/Engine/Plugins/Experimental/Avalanche/Source/AvalancheMediaEditor/Private/Rundown/Pages/Columns/Slate/SAvaRundownPageChannelSelector.h
+++ b/Engine/Plugins/Experimental/Avalanche/Source/AvalancheMediaEditor/Private/Rundown/Pages/Columns/Slate/SAvaRundownPageChannelSelector.h
@@ -3,7 +3,7 @@
 #pragma once
 
 #include "Containers/Array.h"
-#include "Rundown/AvaRundownDefines.h"
+#include "Rundown/AvaRundownEditorDefines.h"
 #include "Widgets/Input/SComboBox.h"
 #include "Widgets/SCompoundWidget.h"
 #include "Templates/SharedPointer.h"
diff --git a/Engine/Plugins/Experimental/Avalanche/Source/AvalancheMediaEditor/Private/Rundown/Pages/Columns/Slate/SAvaRundownPageId.h b/Engine/Plugins/Experimental/Avalanche/Source/AvalancheMediaEditor/Private/Rundown/Pages/Columns/Slate/SAvaRundownPageId.h
--- a/Engine/Plugins/Experimental/Avalanche/Source/AvalancheMediaEditor/Private/Rundown/Pages/Columns/Slate/SAvaRundownPageId.h
+++ b/Engine/Plugins/Experimental/Avalanche/Source/AvalancheMediaEditor/Private/Rundown/Pages/Columns/Slate/SAvaRundownPageId.h
@@ -2,7 +2,7 @@
 
 #pragma once
 
-#include "Rundown/AvaRundownDefines.h"
+#include "Rundown/AvaRundownEditorDefines.h"
 #include "Rundown/Pages/PageViews/IAvaRundownPageView.h"
 #include "Templates/SharedPointer.h"
 #include "Types/SlateEnums.h"
diff --git a/Engine/Plugins/Experimental/Avalanche/Source/AvalancheMediaEditor/Private/Rundown/Pages/Columns/Slate/SAvaRundownPageName.h b/Engine/Plugins/Experimental/Avalanche/Source/AvalancheMediaEditor/Private/Rundown/Pages/Columns/Slate/SAvaRundownPageName.h
--- a/Engine/Plugins/Experimental/Avalanche/Source/AvalancheMediaEditor/Private/Rundown/Pages/Columns/Slate/SAvaRundownPageName.h
+++ b/Engine/Plugins/Experimental/Avalanche/Source/AvalancheMediaEditor/Private/Rundown/Pages/Columns/Slate/SAvaRundownPageName.h
@@ -2,7 +2,7 @@
 
 #pragma once
 
-#include "Rundown/AvaRundownDefines.h"
+#include "Rundown/AvaRundownEditorDefines.h"
 #include "Templates/SharedPointer.h"
 #include "Types/SlateEnums.h"
 #include "Widgets/SCompoundWidget.h"
diff --git a/Engine/Plugins/Experimental/Avalanche/Source/AvalancheMediaEditor/Private/Rundown/Pages/Columns/Slate/SAvaRundownPageThumbnail.h b/Engine/Plugins/Experimental/Avalanche/Source/AvalancheMediaEditor/Private/Rundown/Pages/Columns/Slate/SAvaRundownPageThumbnail.h
--- a/Engine/Plugins/Experimental/Avalanche/Source/AvalancheMediaEditor/Private/Rundown/Pages/Columns/Slate/SAvaRundownPageThumbnail.h
+++ b/Engine/Plugins/Experimental/Avalanche/Source/AvalancheMediaEditor/Private/Rundown/Pages/Columns/Slate/SAvaRundownPageThumbnail.h
@@ -2,7 +2,7 @@
 
 #pragma once
 
-#include "Rundown/AvaRundownDefines.h"
+#include "Rundown/AvaRundownEditorDefines.h"
 #include "Templates/SharedPointer.h"
 #include "Widgets/SCompoundWidget.h"
 
diff --git a/Engine/Plugins/Experimental/Avalanche/Source/AvalancheMediaEditor/Private/Rundown/Pages/PageViews/AvaRundownInstancedPageViewImpl.cpp b/Engine/Plugins/Experimental/Avalanche/Source/AvalancheMediaEditor/Private/Rundown/Pages/PageViews/AvaRundownInstancedPageViewImpl.cpp
--- a/Engine/Plugins/Experimental/Avalanche/Source/AvalancheMediaEditor/Private/Rundown/Pages/PageViews/AvaRundownInstancedPageViewImpl.cpp
+++ b/Engine/Plugins/Experimental/Avalanche/Source/AvalancheMediaEditor/Private/Rundown/Pages/PageViews/AvaRundownInstancedPageViewImpl.cpp
@@ -6,7 +6,7 @@
 #include "Framework/Application/SlateApplication.h"
 #include "Input/Reply.h"
 #include "Rundown/AvaRundown.h"
-#include "Rundown/AvaRundownDefines.h"
+#include "Rundown/AvaRundownEditorDefines.h"
 #include "Rundown/AvaRundownPage.h"
 
 #define LOCTEXT_NAMESPACE "AvaRundownInstancedPageViewImpl"
diff --git a/Engine/Plugins/Experimental/Avalanche/Source/AvalancheMediaEditor/Private/Rundown/Pages/PageViews/AvaRundownPageViewImpl.cpp b/Engine/Plugins/Experimental/Avalanche/Source/AvalancheMediaEditor/Private/Rundown/Pages/PageViews/AvaRundownPageViewImpl.cpp
--- a/Engine/Plugins/Experimental/Avalanche/Source/AvalancheMediaEditor/Private/Rundown/Pages/PageViews/AvaRundownPageViewImpl.cpp
+++ b/Engine/Plugins/Experimental/Avalanche/Source/AvalancheMediaEditor/Private/Rundown/Pages/PageViews/AvaRundownPageViewImpl.cpp
@@ -5,7 +5,7 @@
 #include "AssetRegistry/AssetData.h"
 #include "Framework/Application/SlateApplication.h"
 #include "Rundown/AvaRundown.h"
-#include "Rundown/AvaRundownDefines.h"
+#include "Rundown/AvaRundownEditorDefines.h"
 #include "Rundown/AvaRundownPage.h"
 #include "Rundown/Pages/Slate/SAvaRundownPageList.h"
 #include "ScopedTransaction.h"
diff --git a/Engine/Plugins/Experimental/Avalanche/Source/AvalancheMediaEditor/Private/Rundown/Pages/Slate/SAvaRundownInstancedPageList.h b/Engine/Plugins/Experimental/Avalanche/Source/AvalancheMediaEditor/Private/Rundown/Pages/Slate/SAvaRundownInstancedPageList.h
--- a/Engine/Plugins/Experimental/Avalanche/Source/AvalancheMediaEditor/Private/Rundown/Pages/Slate/SAvaRundownInstancedPageList.h
+++ b/Engine/Plugins/Experimental/Avalanche/Source/AvalancheMediaEditor/Private/Rundown/Pages/Slate/SAvaRundownInstancedPageList.h
@@ -2,7 +2,7 @@
 
 #pragma once
 
-#include "Rundown/AvaRundownDefines.h"
+#include "Rundown/AvaRundownEditorDefines.h"
 #include "SAvaRundownPageList.h"
 
 class SDockTab;
diff --git a/Engine/Plugins/Experimental/Avalanche/Source/AvalancheMediaEditor/Private/Rundown/Pages/Slate/SAvaRundownPageList.h b/Engine/Plugins/Experimental/Avalanche/Source/AvalancheMediaEditor/Private/Rundown/Pages/Slate/SAvaRundownPageList.h
--- a/Engine/Plugins/Experimental/Avalanche/Source/AvalancheMediaEditor/Private/Rundown/Pages/Slate/SAvaRundownPageList.h
+++ b/Engine/Plugins/Experimental/Avalanche/Source/AvalancheMediaEditor/Private/Rundown/Pages/Slate/SAvaRundownPageList.h
@@ -4,7 +4,6 @@
 
 #include "AvaMediaDefines.h"
 #include "Rundown/AvaRundown.h"
-#include "Rundown/AvaRundownDefines.h"
 #include "Rundown/AvaRundownEditorDefines.h"
 #include "Rundown/Pages/PageViews/IAvaRundownPageView.h"
 #include "Widgets/SCompoundWidget.h"
diff --git a/Engine/Plugins/Experimental/Avalanche/Source/AvalancheMediaEditor/Private/Rundown/Pages/Slate/SAvaRundownReadPage.h b/Engine/Plugins/Experimental/Avalanche/Source/AvalancheMediaEditor/Private/Rundown/Pages/Slate/SAvaRundownReadPage.h
--- a/Engine/Plugins/Experimental/Avalanche/Source/AvalancheMediaEditor/Private/Rundown/Pages/Slate/SAvaRundownReadPage.h
+++ b/Engine/Plugins/Experimental/Avalanche/Source/AvalancheMediaEditor/Private/Rundown/Pages/Slate/SAvaRundownReadPage.h
@@ -2,7 +2,7 @@
 
 #pragma once
 
-#include "Rundown/AvaRundownDefines.h"
+#include "Rundown/AvaRundownEditorDefines.h"
 #include "Widgets/SCompoundWidget.h"
 
 class FAvaRundownEditor;
diff --git a/Engine/Plugins/Experimental/Avalanche/Source/AvalancheMediaEditor/Private/Rundown/Pages/Slate/SAvaRundownTemplatePageList.h b/Engine/Plugins/Experimental/Avalanche/Source/AvalancheMediaEditor/Private/Rundown/Pages/Slate/SAvaRundownTemplatePageList.h
--- a/Engine/Plugins/Experimental/Avalanche/Source/AvalancheMediaEditor/Private/Rundown/Pages/Slate/SAvaRundownTemplatePageList.h
+++ b/Engine/Plugins/Experimental/Avalanche/Source/AvalancheMediaEditor/Private/Rundown/Pages/Slate/SAvaRundownTemplatePageList.h
@@ -3,7 +3,7 @@
 #pragma once
 
 #include "Rundown/AvaRundown.h"
-#include "Rundown/AvaRundownDefines.h"
+#include "Rundown/AvaRundownEditorDefines.h"
 #include "SAvaRundownPageList.h"
 
 class FAvaRundownEditor;
diff --git a/Engine/Plugins/Experimental/Avalanche/Source/AvalancheMediaEditor/Public/Rundown/AvaRundownDefines.h b/Engine/Plugins/Experimental/Avalanche/Source/AvalancheMediaEditor/Public/Rundown/AvaRundownDefines.h
deleted file mode 100644
diff --git a/Engine/Plugins/Experimental/Avalanche/Source/AvalancheMediaEditor/Public/Rundown/AvaRundownEditor.h b/Engine/Plugins/Experimental/Avalanche/Source/AvalancheMediaEditor/Public/Rundown/AvaRundownEditor.h
--- a/Engine/Plugins/Experimental/Avalanche/Source/AvalancheMediaEditor/Public/Rundown/AvaRundownEditor.h
+++ b/Engine/Plugins/Experimental/Avalanche/Source/AvalancheMediaEditor/Public/Rundown/AvaRundownEditor.h
@@ -2,7 +2,7 @@
 
 #pragma once
 
-#include "AvaRundownDefines.h"
+#include "AvaRundownEditorDefines.h"
 #include "Rundown/AvaRundown.h"
 #include "TickableEditorObject.h"
 #include "WorkflowOrientedApp/WorkflowCentricApplication.h"
diff --git a/Engine/Plugins/Experimental/Avalanche/Source/AvalancheMediaEditor/Public/Rundown/AvaRundownEditorDefines.h b/Engine/Plugins/Experimental/Avalanche/Source/AvalancheMediaEditor/Public/Rundown/AvaRundownEditorDefines.h
--- a/Engine/Plugins/Experimental/Avalanche/Source/AvalancheMediaEditor/Public/Rundown/AvaRundownEditorDefines.h
+++ b/Engine/Plugins/Experimental/Avalanche/Source/AvalancheMediaEditor/Public/Rundown/AvaRundownEditorDefines.h
@@ -2,8 +2,27 @@
 
 #pragma once
 
+#include "Internationalization/Text.h"
+#include "Templates/SharedPointer.h"
+
 #include "AvaRundownEditorDefines.generated.h"
 
+class IAvaRundownInstancedPageView;
+class IAvaRundownPageView;
+class IAvaRundownTemplatePageView;
+struct FAvaRundownPage;
+
+typedef TWeakPtr<IAvaRundownPageView> FAvaRundownPageViewWeak;
+typedef TSharedPtr<IAvaRundownPageView> FAvaRundownPageViewPtr;
+typedef TSharedRef<IAvaRundownPageView> FAvaRundownPageViewRef;
+typedef TWeakPtr<IAvaRundownTemplatePageView> FAvaRundownTemplatePageViewWeak;
+typedef TSharedPtr<IAvaRundownTemplatePageView> FAvaRundownTemplatePageViewPtr;
+typedef TSharedRef<IAvaRundownTemplatePageView> FAvaRundownTemplatePageViewRef;
+typedef TWeakPtr<IAvaRundownInstancedPageView> FAvaRundownInstancedPageViewWeak;
+typedef TSharedPtr<IAvaRundownInstancedPageView> FAvaRundownInstancedPageViewPtr;
+typedef TSharedRef<IAvaRundownInstancedPageView> FAvaRundownInstancedPageViewRef;
+
+
 /** Defines the page set that an action applies to. */
 UENUM()
 enum class EAvaRundownPageSet : uint8
@@ -29,3 +48,32 @@ enum class EAvaRundownPageSet : uint8
 	/** The action applies to playing pages only regardless of selection. */
 	Playing
 };
+
+namespace UE::AvaRundown
+{
+	enum {InvalidPageId	= -1};
+	
+	struct AVALANCHEMEDIAEDITOR_API FEditorMetrics
+	{
+		static constexpr float ColumnLeftOffset = 5.f;
+		static const FNumberFormattingOptions PageIdFormattingOptions;
+	};
+	
+	enum class EPageEvent : uint8
+	{
+		/** Called after the Selection has Changed */ 
+		SelectionChanged,
+
+		/** Called to sync the Selection of current Selected Items in the Caller */
+		SelectionRequest,
+
+		/** Called to Request a Rename of the Selected Items */
+		RenameRequest,
+
+		/** Called to Request Renumbering Page Id of Selected Items */
+		RenumberRequest,
+
+		/** Called to Request Reimporting Assets of Selected Items */
+		ReimportRequest,
+	};
+}
diff --git a/Engine/Plugins/Experimental/Avalanche/Source/AvalancheMediaEditor/Public/Rundown/Factories/Filters/IAvaRundownFilterSuggestionFactory.h b/Engine/Plugins/Experimental/Avalanche/Source/AvalancheMediaEditor/Public/Rundown/Factories/Filters/IAvaRundownFilterSuggestionFactory.h
--- a/Engine/Plugins/Experimental/Avalanche/Source/AvalancheMediaEditor/Public/Rundown/Factories/Filters/IAvaRundownFilterSuggestionFactory.h
+++ b/Engine/Plugins/Experimental/Avalanche/Source/AvalancheMediaEditor/Public/Rundown/Factories/Filters/IAvaRundownFilterSuggestionFactory.h
@@ -3,7 +3,7 @@
 #pragma once
 
 #include "HAL/Platform.h"
-#include "Rundown/AvaRundownDefines.h"
+#include "Rundown/AvaRundownEditorDefines.h"
 #include "Templates/SharedPointer.h"
 #include "Templates/UnrealTypeTraits.h"
 
diff --git a/Engine/Plugins/Experimental/Avalanche/Source/AvalancheMediaEditor/Public/Rundown/Pages/Columns/IAvaRundownPageViewColumn.h b/Engine/Plugins/Experimental/Avalanche/Source/AvalancheMediaEditor/Public/Rundown/Pages/Columns/IAvaRundownPageViewColumn.h
--- a/Engine/Plugins/Experimental/Avalanche/Source/AvalancheMediaEditor/Public/Rundown/Pages/Columns/IAvaRundownPageViewColumn.h
+++ b/Engine/Plugins/Experimental/Avalanche/Source/AvalancheMediaEditor/Public/Rundown/Pages/Columns/IAvaRundownPageViewColumn.h
@@ -3,7 +3,7 @@
 #pragma once
 
 #include "AvaType.h"
-#include "Rundown/AvaRundownDefines.h"
+#include "Rundown/AvaRundownEditorDefines.h"
 #include "Templates/SharedPointer.h"
 #include "Widgets/Views/SHeaderRow.h"
 
diff --git a/Engine/Plugins/Experimental/Avalanche/Source/AvalancheMediaEditor/Public/Rundown/Pages/Slate/SAvaRundownPageViewRow.h b/Engine/Plugins/Experimental/Avalanche/Source/AvalancheMediaEditor/Public/Rundown/Pages/Slate/SAvaRundownPageViewRow.h
--- a/Engine/Plugins/Experimental/Avalanche/Source/AvalancheMediaEditor/Public/Rundown/Pages/Slate/SAvaRundownPageViewRow.h
+++ b/Engine/Plugins/Experimental/Avalanche/Source/AvalancheMediaEditor/Public/Rundown/Pages/Slate/SAvaRundownPageViewRow.h
@@ -4,7 +4,7 @@
 
 #include "Containers/Array.h"
 #include "Input/DragAndDrop.h"
-#include "Rundown/AvaRundownDefines.h"
+#include "Rundown/AvaRundownEditorDefines.h"
 #include "Templates/SharedPointer.h"
 #include "Widgets/Views/SListView.h"
 #include "Widgets/Views/STableRow.h"
