[Motion Design] Rundown Editor public API clean up.

Keeping only the functions that are used in extensions.

#rnx#rb brad.monahan

diff --git a/Engine/Plugins/Experimental/Avalanche/Source/AvalancheMediaEditor/Private/AvaMediaEditorSettings.h b/Engine/Plugins/Experimental/Avalanche/Source/AvalancheMediaEditor/Private/AvaMediaEditorSettings.h
new file mode 100644
--- /dev/null
+++ b/Engine/Plugins/Experimental/Avalanche/Source/AvalancheMediaEditor/Private/AvaMediaEditorSettings.h
@@ -0,0 +1,51 @@
+// Copyright Epic Games, Inc. All Rights Reserved.
+
+#pragma once
+
+#include "Engine/DeveloperSettings.h"
+#include "AvaMediaEditorSettings.generated.h"
+
+UCLASS(config=EditorPerProjectUserSettings, meta=(DisplayName="Playback & Broadcast"))
+class UAvaMediaEditorSettings : public UDeveloperSettings
+{
+	GENERATED_BODY()
+
+public:
+	UAvaMediaEditorSettings();
+	
+	static const UAvaMediaEditorSettings& Get() {return *GetSingletonInstance();}
+	static UAvaMediaEditorSettings& GetMutable() {return *GetSingletonInstance();}
+
+	UPROPERTY(Config, EditAnywhere, Category = "Broadcast", meta = (InlineEditConditionToggle))
+	bool bBroadcastEnforceMaxChannelCount = true;
+	
+	UPROPERTY(Config, EditAnywhere, Category = "Broadcast", meta = (DisplayName = "Max Channel Count", EditCondition = "bBroadcastEnforceMaxChannelCount"))
+	int32 BroadcastMaxChannelCount = 9;
+
+	/**
+	 * By default, only the media output classes with a device provider are listed as output devices.
+	 * If this is set to true, then all the media output classes will be listed, thus allowing use of
+	 * outputs that don't correspond to a physical device, such as NDI, File output, etc.
+	 * Caution: enabling this mode may expose Media Output classes that are not usable.
+	 */
+	UPROPERTY(Config, EditAnywhere, Category = "Broadcast", meta = (DisplayName = "Show All Media Output Classes"))
+	bool bBroadcastShowAllMediaOutputClasses = false;
+	
+	UPROPERTY(Config, EditAnywhere, Category = "Playback", meta = (DisplayName = "Default Node Color"))
+	FLinearColor PlaybackDefaultNodeColor;
+
+	UPROPERTY(Config, EditAnywhere, Category = "Playback", meta = (DisplayName = "Channels Node Color"))
+	FLinearColor PlaybackChannelsNodeColor;
+	
+	UPROPERTY(Config, EditAnywhere, Category = "Playback", meta = (DisplayName = "Player Node Color"))
+	FLinearColor PlaybackPlayerNodeColor;
+	
+	UPROPERTY(Config, EditAnywhere, Category = "Playback", meta = (DisplayName = "Event Node Color"))
+	FLinearColor PlaybackEventNodeColor;
+
+	UPROPERTY(Config, EditAnywhere, Category = "Playback", meta = (DisplayName = "Action Node Color"))
+	FLinearColor PlaybackActionNodeColor;
+
+private:
+	static UAvaMediaEditorSettings* GetSingletonInstance();
+};
diff --git a/Engine/Plugins/Experimental/Avalanche/Source/AvalancheMediaEditor/Private/Rundown/AvaRundownEditorSettings.h b/Engine/Plugins/Experimental/Avalanche/Source/AvalancheMediaEditor/Private/Rundown/AvaRundownEditorSettings.h
new file mode 100644
--- /dev/null
+++ b/Engine/Plugins/Experimental/Avalanche/Source/AvalancheMediaEditor/Private/Rundown/AvaRundownEditorSettings.h
@@ -0,0 +1,59 @@
+// Copyright Epic Games, Inc. All Rights Reserved.
+
+#pragma once
+
+#include "Engine/DeveloperSettings.h"
+#include "Rundown/AvaRundownEditorDefines.h"
+#include "AvaRundownEditorSettings.generated.h"
+
+class UAvaRundownMacroCollection;
+
+UCLASS(config=EditorPerProjectUserSettings, meta=(DisplayName="Rundown"))
+class UAvaRundownEditorSettings : public UDeveloperSettings
+{
+	GENERATED_BODY()
+
+public:
+	UAvaRundownEditorSettings();
+
+	static const UAvaRundownEditorSettings* Get();
+	static UAvaRundownEditorSettings* GetMutable();
+
+	/**
+	 * Configuring the default page action when closing the editor.
+	 * By default it will stop the pages.
+	 */
+	UPROPERTY(Config, EditAnywhere, Category="Rundown")
+	bool bShouldStopPagesOnClose = true;
+	
+	/** Shows the checker board in preview */
+	UPROPERTY(Config, EditAnywhere, Category="Preview")
+	bool bPreviewCheckerBoard = false;
+
+	/** Current macro collection used by the rundown editor. */
+	UPROPERTY(Config, EditAnywhere, Category="Macros")
+	TSoftObjectPtr<UAvaRundownMacroCollection> MacroCollection;
+	
+	UPROPERTY(Config, EditAnywhere, Category="Page Actions")
+	EAvaRundownPageSet PreviewContinueActionPageSet = EAvaRundownPageSet::SelectedOrPlaying;
+
+	UPROPERTY(Config, EditAnywhere, Category="Page Actions")
+	EAvaRundownPageSet PreviewOutActionPageSet = EAvaRundownPageSet::SelectedOrPlaying;
+	
+	UPROPERTY(Config, EditAnywhere, Category="Page Actions")
+	EAvaRundownPageSet ContinueActionPageSet = EAvaRundownPageSet::SelectedOrPlaying;
+
+	UPROPERTY(Config, EditAnywhere, Category="Page Actions")
+	EAvaRundownPageSet TakeOutActionPageSet = EAvaRundownPageSet::SelectedOrPlaying;
+
+	UPROPERTY(Config, EditAnywhere, Category="Page Actions")
+	EAvaRundownPageSet UpdateValuesActionPageSet = EAvaRundownPageSet::SelectedOrPlaying;
+
+	/** Whether rundown server is started automatically. */
+	UPROPERTY(Config, EditAnywhere, Category="Server")
+	bool bAutoStartRundownServer = false;
+
+	/** Name given to the rundown server. If empty, the server name will be the host name. */
+	UPROPERTY(Config, EditAnywhere, Category="Server")
+	FString RundownServerName;
+};
diff --git a/Engine/Plugins/Experimental/Avalanche/Source/AvalancheMediaEditor/Private/Rundown/AvaRundownMessages.h b/Engine/Plugins/Experimental/Avalanche/Source/AvalancheMediaEditor/Private/Rundown/AvaRundownMessages.h
new file mode 100644
--- /dev/null
+++ b/Engine/Plugins/Experimental/Avalanche/Source/AvalancheMediaEditor/Private/Rundown/AvaRundownMessages.h
@@ -0,0 +1,688 @@
+// Copyright Epic Games, Inc. All Rights Reserved.
+
+#pragma once
+
+#include "Rundown/AvaRundownPage.h"
+#include "AvaRundownMessages.generated.h"
+
+namespace EAvaRundownApiVersion
+{
+	/**
+	 * Defines the protocol version of the Rundown Server API.
+	 *
+	 * API versioning is used to provide legacy support either on
+	 * the client side or server side for non compatible changes.
+	 * Clients can request a version of the API that they where implemented against,
+	 * if the server can still honor the request it will accept.
+	 */
+	enum Type
+	{
+		Initial = 1,
+
+		// -----<new versions can be added before this line>-------------------------------------------------
+		// - this needs to be the last line (see note below)
+		VersionPlusOne,
+		LatestVersion = VersionPlusOne - 1
+	};
+}
+
+
+USTRUCT()
+struct FAvaRundownMsgBase
+{
+	GENERATED_BODY()
+public:
+	UPROPERTY()
+	int32 RequestId = -1;
+};
+
+USTRUCT()
+struct FAvaRundownServerMsg : public FAvaRundownMsgBase
+{
+	GENERATED_BODY()
+public:
+	UPROPERTY()
+	FString Verbosity;
+
+	UPROPERTY()
+	FString Text;
+};
+
+/** Request published by client to discover servers. */
+USTRUCT()
+struct FAvaRundownPing : public FAvaRundownMsgBase
+{
+	GENERATED_BODY()
+public:
+	/** True if the request originates from an automatic timer. False if requests originates from user interaction. */
+	UPROPERTY()
+	bool bAuto = true;
+
+	/**
+	 * API Version the client has been implemented against.
+	 * If none (-1) the server will consider the initial version is requested.
+	 */
+	UPROPERTY()
+	int32 RequestedApiVersion = -1;
+};
+
+/** Response sent by server to client to be discovered. */
+USTRUCT()
+struct FAvaRundownPong : public FAvaRundownMsgBase
+{
+	GENERATED_BODY()
+
+public:
+	/** True if it is a reply to an auto ping. Mirrors the bAuto flag from Ping message. */
+	UPROPERTY()
+	bool bAuto = true;
+
+	/**
+	 * API Version the server will communicate with for this client.
+	 * The server may honor the requested version if possible.
+	 * Versions newer than server implementation will obviously not be honored either.
+	 * Clients should expect an older server to reply with an older version.
+	 */
+	UPROPERTY()
+	int32 ApiVersion = -1;
+
+	/** Minimum API Version the server implements. */
+	UPROPERTY()
+	int32 MinimumApiVersion = -1;
+
+	/** Latest API Version the server support. */
+	UPROPERTY()
+	int32 LatestApiVersion = -1;
+
+	UPROPERTY()
+	FString HostName;
+};
+
+/**
+ *	Request list of rundown that can be opened on the current server.
+ */
+USTRUCT()
+struct FAvaRundownGetRundowns : public FAvaRundownMsgBase
+{
+	GENERATED_BODY()
+};
+
+/**
+ *	List of all rundowns.
+ *	Expected Response from FAvaRundownGetRundowns.
+ */
+USTRUCT()
+struct FAvaRundownRundowns : public FAvaRundownMsgBase
+{
+	GENERATED_BODY()
+
+public:
+	UPROPERTY()
+	TArray<FString> Rundowns;
+};
+
+/**
+ *	Request that the given rundown be opened.
+ *	Only one rundown can be opened at a time. If another rundown
+ *	is opened, it will be closed and all currently playing pages stopped.
+ *	If the path is empty, nothing will be done and the server will reply with
+ *	a FAvaRundownLoadedRundown message indicating which rundown is currently loaded.
+ */
+USTRUCT()
+struct FAvaRundownLoadRundown : public FAvaRundownMsgBase
+{
+	GENERATED_BODY()
+
+public:
+	UPROPERTY()
+	FString Rundown;
+};
+
+/**
+ * Request the list of pages from the given rundown.
+ */
+USTRUCT()
+struct FAvaRundownGetPages : public FAvaRundownMsgBase
+{
+	GENERATED_BODY()
+public:
+	UPROPERTY()
+	FString Rundown;
+};
+
+USTRUCT()
+struct FAvaRundownCreatePage : public FAvaRundownMsgBase
+{
+	GENERATED_BODY()
+
+	UPROPERTY()
+	FString Rundown;
+
+	UPROPERTY()
+	int32 TemplateId = FAvaRundownPage::InvalidPageId;
+};
+
+USTRUCT()
+struct FAvaRundownDeletePage : public FAvaRundownMsgBase
+{
+	GENERATED_BODY()
+
+	UPROPERTY()
+	FString Rundown;
+
+	UPROPERTY()
+	int32 PageId = FAvaRundownPage::InvalidPageId;
+};
+
+USTRUCT()
+struct FAvaRundownCreateTemplate : public FAvaRundownMsgBase
+{
+	GENERATED_BODY()
+
+	UPROPERTY()
+	FString Rundown;
+};
+
+USTRUCT()
+struct FAvaRundownDeleteTemplate : public FAvaRundownMsgBase
+{
+	GENERATED_BODY()
+
+	UPROPERTY()
+	FString Rundown;
+
+	UPROPERTY()
+	int32 PageId = FAvaRundownPage::InvalidPageId;
+};
+
+USTRUCT()
+struct FAvaRundownChangeTemplateBP : public FAvaRundownMsgBase
+{
+	GENERATED_BODY()
+public:
+	UPROPERTY()
+	FString Rundown;
+
+	UPROPERTY()
+	int32 TemplateId = FAvaRundownPage::InvalidPageId;
+
+	UPROPERTY()
+	FString AssetPath;
+};
+
+USTRUCT()
+struct FAvaRundownPageInfo
+{
+	GENERATED_BODY()
+public:
+	UPROPERTY()
+	int32 PageId = FAvaRundownPage::InvalidPageId;
+
+	UPROPERTY()
+	FString PageName;
+
+	UPROPERTY()
+	bool IsTemplate = false;
+
+	UPROPERTY()
+	int32 TemplateId = FAvaRundownPage::InvalidPageId;
+
+	UPROPERTY()
+	TArray<int32> CombinedTemplateIds;
+
+	UPROPERTY()
+	FSoftObjectPath AssetPath;
+
+	UPROPERTY()
+	TArray<FAvaRundownChannelPageStatus> Statuses;
+
+	UPROPERTY()
+	FString TransitionLayerName;
+
+	UPROPERTY()
+	FString OutputChannel;
+
+	UPROPERTY()
+	bool bIsEnabled = false;
+
+	UPROPERTY()
+	bool bIsPlaying = false;
+};
+
+/*
+ * List of pages from the current rundown.
+ */
+USTRUCT()
+struct FAvaRundownPages : public FAvaRundownMsgBase
+{
+	GENERATED_BODY()
+public:
+	UPROPERTY()
+	TArray<FAvaRundownPageInfo> Pages;
+};
+
+/**
+ * Request the page details from the given rundown.
+ */
+USTRUCT()
+struct FAvaRundownGetPageDetails : public FAvaRundownMsgBase
+{
+	GENERATED_BODY()
+
+public:
+	UPROPERTY()
+	FString Rundown;
+
+	UPROPERTY()
+	int32 PageId = FAvaRundownPage::InvalidPageId;
+
+	/** This will request that a managed asset instance gets loaded to be
+	 * accessible through WebRC. */
+	UPROPERTY()
+	bool bLoadRemoteControlPreset = false;
+};
+
+/**
+ *	Server response to FAvaRundownGetPageDetails request.
+ */
+USTRUCT()
+struct FAvaRundownPageDetails : public FAvaRundownMsgBase
+{
+	GENERATED_BODY()
+
+public:
+	UPROPERTY()
+	FString Rundown;
+
+	UPROPERTY()
+	FAvaRundownPageInfo PageInfo;
+
+	UPROPERTY()
+	FAvaPlayableRemoteControlValues RemoteControlValues;
+
+	/** Name of the remote control preset to resolve through WebRC API. */
+	UPROPERTY()
+	FString RemoteControlPresetName;
+
+	UPROPERTY()
+	FString RemoteControlPresetId;
+};
+
+USTRUCT()
+struct FAvaRundownPagesStatuses : public FAvaRundownMsgBase
+{
+	GENERATED_BODY()
+
+public:
+	UPROPERTY()
+	FString Rundown;
+
+	UPROPERTY()
+	FAvaRundownPageInfo PageInfo;
+};
+
+USTRUCT()
+struct FAvaRundownPageListChanged : public FAvaRundownMsgBase
+{
+	GENERATED_BODY()
+
+	UPROPERTY()
+	FString Rundown;
+
+	/** See EAvaPageListChange flags. */
+	UPROPERTY()
+	uint8 ChangeType = 0;
+
+	UPROPERTY();
+	TArray<int32> AffectedPages;
+};
+
+USTRUCT()
+struct FAvaRundownPageBlueprintChanged : public FAvaRundownMsgBase
+{
+	GENERATED_BODY()
+
+	UPROPERTY()
+	FString Rundown;
+
+	UPROPERTY()
+	int32 PageId = FAvaRundownPage::InvalidPageId;
+
+	UPROPERTY()
+	FString BlueprintPath;
+};
+
+USTRUCT()
+struct FAvaRundownPageChannelChanged : public FAvaRundownMsgBase
+{
+	GENERATED_BODY()
+
+	UPROPERTY()
+	FString Rundown;
+
+	UPROPERTY()
+	int32 PageId = FAvaRundownPage::InvalidPageId;
+
+	UPROPERTY()
+	FString ChannelName;
+};
+
+USTRUCT()
+struct FAvaRundownPageAnimSettingsChanged : public FAvaRundownMsgBase
+{
+	GENERATED_BODY()
+
+	UPROPERTY()
+	FString Rundown;
+
+	UPROPERTY()
+	int32 PageId = FAvaRundownPage::InvalidPageId;
+};
+
+USTRUCT()
+struct FAvaRundownPageChangeChannel : public FAvaRundownMsgBase
+{
+	GENERATED_BODY()
+
+	UPROPERTY()
+	FString Rundown;
+
+	UPROPERTY()
+	int32 PageId = FAvaRundownPage::InvalidPageId;
+
+	UPROPERTY()
+	FString ChannelName;
+};
+
+/** This is a request to save the managed RCP back to the corresponding page. */
+USTRUCT()
+struct FAvaRundownUpdatePageFromRCP : public FAvaRundownMsgBase
+{
+	GENERATED_BODY()
+
+public:
+	/** Unregister the Remote Control Preset from the WebRC. */
+	UPROPERTY()
+	bool bUnregister = false;
+};
+
+/** Supported Page actions for playback. */
+UENUM()
+enum class EAvaRundownPageActions
+{
+	None,
+	Load,
+	Unload,
+	Play,
+	PlayNext,
+	Stop,
+	ForceStop,
+	Continue,
+	UpdateValues,
+	TakeToProgram
+};
+
+USTRUCT()
+struct FAvaRundownPageAction : public FAvaRundownMsgBase
+{
+	GENERATED_BODY()
+public:
+	UPROPERTY()
+	int32 PageId = FAvaRundownPage::InvalidPageId;
+
+	UPROPERTY()
+	EAvaRundownPageActions Action = EAvaRundownPageActions::None;
+};
+
+USTRUCT()
+struct FAvaRundownPagePreviewAction : public FAvaRundownPageAction
+{
+	GENERATED_BODY()
+public:
+	/** Specify which preview channel to use. If left empty, the rundown's default preview channel is used. */
+	UPROPERTY()
+	FString PreviewChannelName;
+};
+
+/**
+ * Command to execute an action on multiple pages at the same time.
+ * This is necessary for pages to be part of the same transition.
+ */
+USTRUCT()
+struct FAvaRundownPageActions : public FAvaRundownMsgBase
+{
+	GENERATED_BODY()
+public:
+	UPROPERTY()
+	TArray<int32> PageIds;
+
+	UPROPERTY()
+	EAvaRundownPageActions Action = EAvaRundownPageActions::None;
+};
+
+USTRUCT()
+struct FAvaRundownPagePreviewActions : public FAvaRundownPageActions
+{
+	GENERATED_BODY()
+public:
+	/** Specify which preview channel to use. If left empty, the rundown's default preview channel is used. */
+	UPROPERTY()
+	FString PreviewChannelName;
+};
+
+UENUM()
+enum class EAvaRundownPageEvents
+{
+	None,
+	AnimStarted,
+	AnimPaused,
+	AnimFinished
+};
+
+USTRUCT()
+struct FAvaRundownPageEvent : public FAvaRundownMsgBase
+{
+	GENERATED_BODY()
+public:
+	UPROPERTY()
+	int32 PageId = FAvaRundownPage::InvalidPageId;
+
+	UPROPERTY()
+	EAvaRundownPageEvents Event = EAvaRundownPageEvents::None;
+};
+
+USTRUCT()
+struct FAvaRundownOutputDeviceItem
+{
+	GENERATED_BODY()
+public:
+	UPROPERTY()
+	FString Name;
+
+	/**
+	 * Raw Json string representing a serialized UMediaOutput.
+	 */
+	UPROPERTY()
+	FString Data;
+};
+
+USTRUCT()
+struct FAvaRundownOutputClassItem
+{
+	GENERATED_BODY()
+public:
+	UPROPERTY()
+	FString Name;
+
+	UPROPERTY()
+	TArray<FAvaRundownOutputDeviceItem> Devices;
+};
+
+USTRUCT()
+struct FAvaRundownDevicesList : public FAvaRundownMsgBase
+{
+	GENERATED_BODY()
+public:
+	UPROPERTY()
+	TArray<FAvaRundownOutputClassItem> DeviceClasses;
+};
+
+USTRUCT()
+struct FAvaRundownGetChannel : public FAvaRundownMsgBase
+{
+	GENERATED_BODY()
+public:
+	UPROPERTY()
+	FString ChannelName;
+};
+
+USTRUCT()
+struct FAvaRundownGetChannels : public FAvaRundownMsgBase
+{
+	GENERATED_BODY()
+};
+
+USTRUCT()
+struct FAvaRundownChannel
+{
+	GENERATED_BODY()
+public:
+	UPROPERTY()
+	FString Name;
+
+	UPROPERTY()
+	TArray<FAvaRundownOutputDeviceItem> Devices;
+};
+
+USTRUCT()
+struct FAvaRundownChannelListChanged : public FAvaRundownMsgBase
+{
+	GENERATED_BODY()
+
+	UPROPERTY()
+	TArray<FAvaRundownChannel> Channels;
+};
+
+USTRUCT()
+struct FAvaRundownChannelResponse : public FAvaRundownMsgBase
+{
+	GENERATED_BODY()
+public:
+	UPROPERTY()
+	FAvaRundownChannel Channel;
+};
+
+USTRUCT()
+struct FAvaRundownChannels : public FAvaRundownMsgBase
+{
+	GENERATED_BODY()
+public:
+	UPROPERTY()
+	TArray<FAvaRundownChannel> Channels;
+};
+
+USTRUCT()
+struct FAvaRundownAssetsChanged : public FAvaRundownMsgBase
+{
+	GENERATED_BODY()
+
+	UPROPERTY()
+	FString AssetName;
+};
+
+// Channel actions
+UENUM()
+enum class EAvaRundownChannelActions
+{
+	None,
+	Start,
+	Stop
+};
+
+USTRUCT()
+struct FAvaRundownChannelAction : public FAvaRundownMsgBase
+{
+	GENERATED_BODY()
+public:
+	UPROPERTY()
+	FString ChannelName;
+
+	UPROPERTY()
+	EAvaRundownChannelActions Action = EAvaRundownChannelActions::None;
+};
+
+USTRUCT()
+struct FAvaRundownGetDevices : public FAvaRundownMsgBase
+{
+	GENERATED_BODY()
+};
+
+USTRUCT()
+struct FAvaRundownAddChannelDevice : public FAvaRundownMsgBase
+{
+	GENERATED_BODY()
+public:
+	UPROPERTY()
+	FString ChannelName;
+
+	UPROPERTY()
+	FString MediaOutputName;
+};
+
+USTRUCT()
+struct FAvaRundownEditChannelDevice : public FAvaRundownMsgBase
+{
+	GENERATED_BODY()
+public:
+	UPROPERTY()
+	FString ChannelName;
+
+	UPROPERTY()
+	FString MediaOutputName;
+
+	UPROPERTY()
+	FString Data;
+};
+
+USTRUCT()
+struct FAvaRundownRemoveChannelDevice : public FAvaRundownMsgBase
+{
+	GENERATED_BODY()
+public:
+	UPROPERTY()
+	FString ChannelName;
+
+	UPROPERTY()
+	FString MediaOutputName;
+};
+
+/* No difference from FAvaRundownOutputDeviceItem except this is meant to return a single device response. */
+USTRUCT()
+struct FAvaRundownOutputDeviceItemResponse : public FAvaRundownMsgBase
+{
+	GENERATED_BODY()
+public:
+	UPROPERTY()
+	FString Name;
+
+	UPROPERTY()
+	FString Data;
+};
+
+USTRUCT()
+struct FAvaRundownGetChannelImage : public FAvaRundownMsgBase
+{
+	GENERATED_BODY()
+public:
+	UPROPERTY()
+	FString ChannelName;
+};
+
+USTRUCT()
+struct FAvaRundownChannelImage : public FAvaRundownMsgBase
+{
+	GENERATED_BODY()
+public:
+	UPROPERTY()
+	TArray<uint8> ImageData;
+};
diff --git a/Engine/Plugins/Experimental/Avalanche/Source/AvalancheMediaEditor/Private/Rundown/AvaRundownServer.h b/Engine/Plugins/Experimental/Avalanche/Source/AvalancheMediaEditor/Private/Rundown/AvaRundownServer.h
new file mode 100644
--- /dev/null
+++ b/Engine/Plugins/Experimental/Avalanche/Source/AvalancheMediaEditor/Private/Rundown/AvaRundownServer.h
@@ -0,0 +1,224 @@
+// Copyright Epic Games, Inc. All Rights Reserved.
+
+#pragma once
+
+#include "AvaRundownMessages.h"
+#include "Broadcast/AvaBroadcastProfile.h"
+#include "MessageEndpoint.h"
+#include "Rundown/AvaRundown.h"
+#include "UObject/SoftObjectPath.h"
+#include "UObject/StrongObjectPtr.h"
+
+class FAvaRundownManagedInstance;
+class UMediaOutput;
+
+/**
+ * Implements a play list server that listens to commands on message bus.
+ * The intention is to run a web socket transport bridge so the messages can
+ * come from external applications.
+ */
+class FAvaRundownServer : public TSharedFromThis<FAvaRundownServer>
+{
+public:
+	FAvaRundownServer();
+	virtual ~FAvaRundownServer();
+
+	void Init(const FString& InAssignedHostName);
+
+	void SetupBroadcastDelegates(UAvaBroadcast* InBroadcast);
+	void SetupEditorDelegates();
+	void RemoveBroadcastDelegates(UAvaBroadcast* InBroadcast) const;
+	void RemoveEditorDelegates() const;
+	
+	void OnPagesChanged(const UAvaRundown* InRundown, const FAvaRundownPage& InPage, EAvaRundownPageChanges InChange) const;
+	void OnPageListChanged(const FAvaRundownPageListChangeParams& InParams) const;
+	void PageBlueprintChanged(const UAvaRundown* InRundown, const FAvaRundownPage& InPage, const FString& InBlueprintPath) const;
+	void PageStatusChanged(const UAvaRundown* InRundown, const FAvaRundownPage& InPage) const;
+	void PageChannelChanged(const UAvaRundown* InRundown, const FAvaRundownPage& InPage, const FString& InChannelName) const;
+	void PageAnimSettingsChanged(const UAvaRundown* InRundown, const FAvaRundownPage& InPage) const;
+	void OnBroadcastChannelListChanged(const FAvaBroadcastProfile& InProfile) const;
+	void OnAssetAddedOrRemoved(const FAssetData& InAssetData) const;
+
+	/** Returns the endpoint's message address. */
+	const FMessageAddress& GetMessageAddress() const;
+	
+	// Message handlers
+	void HandleRundownPing(const FAvaRundownPing& InMessage, const TSharedRef<IMessageContext, ESPMode::ThreadSafe>& InContext);
+	void HandleGetRundowns(const FAvaRundownGetRundowns& InMessage, const TSharedRef<IMessageContext, ESPMode::ThreadSafe>& InContext);
+	void HandleLoadRundown(const FAvaRundownLoadRundown& InMessage, const TSharedRef<IMessageContext, ESPMode::ThreadSafe>& InContext);
+	void HandleGetPages(const FAvaRundownGetPages& InMessage, const TSharedRef<IMessageContext, ESPMode::ThreadSafe>& InContext);
+	void HandleCreatePage(const FAvaRundownCreatePage& InMessage,const TSharedRef<IMessageContext, ESPMode::ThreadSafe>& InContext);
+	void HandleCreateTemplate(const FAvaRundownCreateTemplate& InMessage, const TSharedRef<IMessageContext, ESPMode::ThreadSafe>& InContext);
+	void HandleDeletePage(const FAvaRundownDeletePage& InMessage,const TSharedRef<IMessageContext, ESPMode::ThreadSafe>& InContext);
+	void HandleDeleteTemplate(const FAvaRundownDeleteTemplate& InMessage,const TSharedRef<IMessageContext, ESPMode::ThreadSafe>& InContext);
+	void HandleChangeTemplateBP(const FAvaRundownChangeTemplateBP& InMessage, const TSharedRef<IMessageContext, ESPMode::ThreadSafe>& InContext);
+	void HandleGetPageDetails(const FAvaRundownGetPageDetails& InMessage, const TSharedRef<IMessageContext, ESPMode::ThreadSafe>& InContext);
+	void HandleChangePageChannel(const FAvaRundownPageChangeChannel& InMessage, const TSharedRef<IMessageContext, ESPMode::ThreadSafe>& InContext);
+	void HandleUpdatePageFromRCP(const FAvaRundownUpdatePageFromRCP& InMessage, const TSharedRef<IMessageContext, ESPMode::ThreadSafe>& InContext);
+	void HandlePageAction(const FAvaRundownPageAction& InMessage, const TSharedRef<IMessageContext, ESPMode::ThreadSafe>& InContext);
+	void HandlePagePreviewAction(const FAvaRundownPagePreviewAction& InMessage, const TSharedRef<IMessageContext, ESPMode::ThreadSafe>& InContext);
+	void HandlePageActions(const FAvaRundownPageActions& InMessage, const TSharedRef<IMessageContext, ESPMode::ThreadSafe>& InContext);
+	void HandlePagePreviewActions(const FAvaRundownPagePreviewActions& InMessage, const TSharedRef<IMessageContext, ESPMode::ThreadSafe>& InContext);
+
+	void HandleGetChannel(const FAvaRundownGetChannel& InMessage, const TSharedRef<IMessageContext, ESPMode::ThreadSafe>& InContext);
+	void HandleGetChannels(const FAvaRundownGetChannels& InMessage, const TSharedRef<IMessageContext, ESPMode::ThreadSafe>& InContext);
+	void HandleChannelAction(const FAvaRundownChannelAction& InMessage, const TSharedRef<IMessageContext, ESPMode::ThreadSafe>& InContext);
+	void HandleAddChannelDevice(const FAvaRundownAddChannelDevice& InMessage, const TSharedRef<IMessageContext, ESPMode::ThreadSafe>& InContext);
+	void HandleEditChannelDevice(const FAvaRundownEditChannelDevice& InMessage, const TSharedRef<IMessageContext, ESPMode::ThreadSafe>& InContext);
+	void HandleRemoveChannelDevice(const FAvaRundownRemoveChannelDevice& InMessage, const TSharedRef<IMessageContext, ESPMode::ThreadSafe>& InContext);
+	void HandleGetChannelImage(const FAvaRundownGetChannelImage& InMessage, const TSharedRef<IMessageContext, ESPMode::ThreadSafe>& InContext);
+	
+	void HandleGetDevices(const FAvaRundownGetDevices& InMessage, const TSharedRef<IMessageContext, ESPMode::ThreadSafe>& InContext);
+	
+	void LogAndSendMessage(const FMessageAddress& InSender, int32 InRequestId, ELogVerbosity::Type InVerbosity, const TCHAR* InFormat, ...) const;
+	void SendMessage(const FMessageAddress& InSender, int32 InRequestId, ELogVerbosity::Type InVerbosity, const TCHAR* InFormat, ...) const;
+	void SendMessageImpl(const FMessageAddress& InSender, int32 InRequestId, ELogVerbosity::Type InVerbosity, const TCHAR* InMsg) const;
+
+	void RegisterConsoleCommands();
+	
+	void ShowStatusCommand(const TArray<FString>& InArgs);
+
+protected:
+	struct FRequestInfo
+	{
+		int32 RequestId;
+		FMessageAddress Sender;
+	};
+	
+	struct FChannelImage;
+	void FinishGetChannelImage(const FRequestInfo& InRequestInfo, const TSharedPtr<FChannelImage>& InChannelImage);
+
+	void HandlePageActions(const FRequestInfo& InRequestInfo, const TArray<int32>& InPageIds,
+		bool bInIsPreview, FName InPreviewChannelName, EAvaRundownPageActions InAction) const;
+	
+	/**
+	 * Helper function to retrieve the appropriate rundown for editing commands.
+	 */
+	UAvaRundown* GetOrLoadRundownForEdit(const FMessageAddress& InSender, int32 InRequestId, const FString& InRundownPath);
+
+	template<typename MessageType>
+	void SendResponse(MessageType* InMessage, const FMessageAddress& InRecipient, EMessageFlags InFlags = EMessageFlags::None) const
+	{
+		if (MessageEndpoint)
+		{
+			MessageEndpoint->Send(InMessage, InRecipient);
+		}
+	}
+	
+	template<typename MessageType>
+	void SendResponse(MessageType* InMessage, const TArray<FMessageAddress>& InRecipients, EMessageFlags InFlags = EMessageFlags::None) const
+	{
+		if (MessageEndpoint)
+		{
+			MessageEndpoint->Send(InMessage, MessageType::StaticStruct(), InFlags, nullptr,
+				InRecipients, FTimespan::Zero(), FDateTime::MaxValue());
+		}
+	}
+
+	void OnMessageBusNotification(const FMessageBusNotification& InNotification);
+
+private:
+	FString HostName;	
+	TSharedPtr<FMessageEndpoint, ESPMode::ThreadSafe> MessageEndpoint;
+	TArray<IConsoleObject*> ConsoleCommands;
+
+	/** Keep information on connected clients. */
+	struct FClientInfo
+	{
+		FMessageAddress Address;
+		/** Api version for communication with this client. */
+		int32 ApiVersion = -1;
+
+		explicit FClientInfo(const FMessageAddress& InAddress) : Address(InAddress) {}
+	};
+
+	/** Keep track of remote clients context information. */
+	TMap<FMessageAddress, TSharedPtr<FClientInfo>> Clients;
+
+	/** Array of just the client addresses for sending responses. */
+	TArray<FMessageAddress> ClientAddresses;
+	
+	FClientInfo* GetClientInfo(const FMessageAddress& InAddress) const
+	{
+		const TSharedPtr<FClientInfo>* ClientInfoPtr = Clients.Find(InAddress);
+		return ClientInfoPtr ? ClientInfoPtr->Get() : nullptr;
+	}
+	
+	FClientInfo& GetOrAddClientInfo(const FMessageAddress& InAddress)
+	{
+		if (const TSharedPtr<FClientInfo>* ExistingClientInfo = Clients.Find(InAddress))
+		{
+			return *ExistingClientInfo->Get();
+		}
+
+		const TSharedPtr<FClientInfo> NewClientInfo = MakeShared<FClientInfo>(InAddress);
+		Clients.Add(InAddress, NewClientInfo);
+		RefreshClientAddresses();
+		return *NewClientInfo;
+	}
+
+	void RefreshClientAddresses();
+	
+	/** Pool of images that can be recycled. */
+	TArray<TSharedPtr<FChannelImage>> AvailableChannelImages;
+	
+	/**
+	 * Keeps a current rundown cached for commands operating on rundown.
+	 * There is only one "current" rundown at a given time.
+	 */
+	struct FRundownCache
+	{
+		/** Currently loaded/cached rundown's path. */
+		FSoftObjectPath CurrentRundownPath;
+		
+		/** Currently loaded/cached rundown object. */
+		TStrongObjectPtr<UAvaRundown> CurrentRundown;
+
+		FDelegateHandle OnPlaybackInstanceStatusChangedDelegateHandle;
+		
+		typedef TFunctionRef<void(UAvaRundown*)> FRundownEventFunction;
+		
+		/**
+		 * Returns requested rundown specified by InRundownPath. Will load it if necessary or returned the cached one if it is the same.
+		 * If the new rundown fails to load, the return value is nullptr, and the previous rundown will remain loaded.
+		 * @param InRundownPath	Requested rundown path. 
+		 * @param InUnloadCurrentRundownFunction	Function called, with the previously cached rundown, when a new rundown is loaded
+		 *											and previous rundown needs to be unloaded.
+		 * @param InNewRundownLoadedFunction	Function called, with the new rundown, when a new rundown is loaded.
+		 */
+		UAvaRundown* GetOrLoadRundown(const FSoftObjectPath& InRundownPath, 
+			const FRundownEventFunction InUnloadCurrentRundownFunction,
+			const FRundownEventFunction InNewRundownLoadedFunction);
+
+		void SetupRundownDelegates(FAvaRundownServer* InRundownServer, UAvaRundown* InRundown);
+		void RemoveRundownDelegates(const FAvaRundownServer* InRundownServer, UAvaRundown* InRundown) const;
+	};
+	
+	/** Cached data for page editing commands (GetPages and GetPageDetails). */
+	struct FRundownEditCommandData : public FRundownCache
+	{
+		/** PageId of the current managed ava asset. */
+		int32 ManagedPageId = FAvaRundownPage::InvalidPageId;
+		TSharedPtr<FAvaRundownManagedInstance> ManagedInstance;
+
+		~FRundownEditCommandData();
+		
+		/**
+		 * Checks if previous RCP was registered.
+		 * If so, save modified values to corresponding page.
+ 		 * This may result in the Rundown to be modified.
+ 		 * Will also unregister RCP from RC Module if requested.
+		 */
+		void SaveCurrentRemoteControlPresetToPage(bool bInUnregister);
+	};
+	FRundownEditCommandData RundownEditCommandData;
+
+	/** Cached data for playback commands (LoadRundown and PageAction). */
+	struct FRundownPlaybackCommandData: public FRundownCache
+	{
+		~FRundownPlaybackCommandData();
+
+		void ClosePlaybackContext();
+	};
+	FRundownPlaybackCommandData RundownPlaybackCommandData;
+};
diff --git a/Engine/Plugins/Experimental/Avalanche/Source/AvalancheMediaEditor/Private/Rundown/Factories/Filters/AvaRundownFactoriesUtils.h b/Engine/Plugins/Experimental/Avalanche/Source/AvalancheMediaEditor/Private/Rundown/Factories/Filters/AvaRundownFactoriesUtils.h
new file mode 100644
--- /dev/null
+++ b/Engine/Plugins/Experimental/Avalanche/Source/AvalancheMediaEditor/Private/Rundown/Factories/Filters/AvaRundownFactoriesUtils.h
@@ -0,0 +1,15 @@
+// Copyright Epic Games, Inc. All Rights Reserved.
+
+#pragma once
+
+#include "HAL/Platform.h"
+#include "Misc/EnumClassFlags.h"
+
+/** Type of the Page list */
+enum class EAvaRundownSearchListType : uint8
+{
+	None      = 0x00, // None
+	Template  = 0x01, // Template pages type
+	Instanced = 0x10 // Instanced pages type
+};
+ENUM_CLASS_FLAGS(EAvaRundownSearchListType)
diff --git a/Engine/Plugins/Experimental/Avalanche/Source/AvalancheMediaEditor/Private/Rundown/Factories/Filters/IAvaRundownFilterExpressionFactory.h b/Engine/Plugins/Experimental/Avalanche/Source/AvalancheMediaEditor/Private/Rundown/Factories/Filters/IAvaRundownFilterExpressionFactory.h
new file mode 100644
--- /dev/null
+++ b/Engine/Plugins/Experimental/Avalanche/Source/AvalancheMediaEditor/Private/Rundown/Factories/Filters/IAvaRundownFilterExpressionFactory.h
@@ -0,0 +1,72 @@
+// Copyright Epic Games, Inc. All Rights Reserved.
+
+#pragma once
+
+#include "HAL/Platform.h"
+#include "Misc/TextFilterUtils.h"
+#include "Templates/SharedPointer.h"
+#include "Templates/UnrealTypeTraits.h"
+#include "UObject/NameTypes.h"
+
+class UAvaRundown;
+enum class EAvaRundownSearchListType : uint8;
+struct FAvaRundownPage;
+
+/**
+ * Holds the arguments required for the expression evaluation.
+ */
+struct FAvaRundownTextFilterArgs
+{
+	/** Item Rundown */
+	const UAvaRundown* ItemRundown;
+
+	/** Second Value */
+	FTextFilterString ValueToCheck;
+
+	/** Comparison Operation requested (Equal/Greater/etc...) still needed for certain Filter */
+	ETextFilterComparisonOperation ComparisonOperation;
+
+	/** Comparison Mode requested see ETextFilterTextComparisonMode for more information */
+	ETextFilterTextComparisonMode ComparisonMode;
+};
+
+class IAvaRundownFilterExpressionFactory : public TSharedFromThis<IAvaRundownFilterExpressionFactory>
+{
+public:
+	virtual ~IAvaRundownFilterExpressionFactory() {}
+
+	/** Get the filter identifier for this factory */
+	virtual FName GetFilterIdentifier() const = 0;
+
+	/**
+ 	 * Create and return an Instance of the AvaRundownFilterExpressionFactory Requested
+ 	 * @tparam InRundownFilterExpressionFactoryType The type of the factory to instantiate
+ 	 * @param InArgs Additional Args for constructor of the AvaRundownFilterExpressionFactory class if needed
+ 	 * @return The new rundown filter factory
+ 	 */
+	template <
+		typename InRundownFilterExpressionFactoryType,
+		typename... InArgsType
+		UE_REQUIRES(TIsDerivedFrom<InRundownFilterExpressionFactoryType, IAvaRundownFilterExpressionFactory>::Value)
+	>
+	static TSharedRef<InRundownFilterExpressionFactoryType> MakeInstance(InArgsType&&... InArgs)
+	{
+		return MakeShared<InRundownFilterExpressionFactoryType>(Forward<InArgsType>(InArgs)...);
+	}
+
+	/**
+	 * Evaluate the expression with the Value given
+	 * @param InItem Current Item being checked
+	 * @param InArgs The argument containing the data to evaluate the expression see FAvaTextFilterArgs for more information
+	 * @return True if the expression evaluate to True, False otherwise
+	 */
+	virtual bool FilterExpression(const FAvaRundownPage& InItem, const FAvaRundownTextFilterArgs& InArgs) const = 0;
+
+	/**
+	 * Whether the factory support a given Comparison Operation
+	 * @param InComparisonOperation Comparison Operation to check
+	 * @param InRundownSearchListType Type of the Search List either Template or Instanced
+	 * @return True if Factory support passed Comparison Operation, False otherwise
+	 */
+	virtual bool SupportsComparisonOperation(ETextFilterComparisonOperation InComparisonOperation, EAvaRundownSearchListType InRundownSearchListType) const = 0;
+};
diff --git a/Engine/Plugins/Experimental/Avalanche/Source/AvalancheMediaEditor/Private/Rundown/Factories/Filters/IAvaRundownFilterSuggestionFactory.h b/Engine/Plugins/Experimental/Avalanche/Source/AvalancheMediaEditor/Private/Rundown/Factories/Filters/IAvaRundownFilterSuggestionFactory.h
new file mode 100644
--- /dev/null
+++ b/Engine/Plugins/Experimental/Avalanche/Source/AvalancheMediaEditor/Private/Rundown/Factories/Filters/IAvaRundownFilterSuggestionFactory.h
@@ -0,0 +1,74 @@
+// Copyright Epic Games, Inc. All Rights Reserved.
+
+#pragma once
+
+#include "HAL/Platform.h"
+#include "Rundown/AvaRundownEditorDefines.h"
+#include "Templates/SharedPointer.h"
+#include "Templates/UnrealTypeTraits.h"
+
+class UAvaRundown;
+struct FAvaRundownPage;
+struct FAssetSearchBoxSuggestion;
+enum class EAvaRundownSearchListType : uint8;
+
+/**
+ * Holds the arguments required for the suggestion factory to get the suggestion.
+ */
+struct FAvaRundownFilterSuggestionPayload
+{
+	/** Holds all the suggestion that will be shown */
+	TArray<FAssetSearchBoxSuggestion>& PossibleSuggestions;
+
+	/** The current string value written in the searchbox */
+	const FString FilterValue;
+
+	/** Current item to check */
+	int32 ItemPageId = UE::AvaRundown::InvalidPageId;
+
+	/** Current item rundown */
+	const UAvaRundown* Rundown = nullptr;
+
+	/** Used to speed up the check on whether a suggestion is already added */
+	TSet<FString>& FilterCache;
+};
+
+class IAvaRundownFilterSuggestionFactory : public TSharedFromThis<IAvaRundownFilterSuggestionFactory>
+{
+public:
+	virtual ~IAvaRundownFilterSuggestionFactory() {}
+
+	/**
+	  * Create and return an Instance of the AvaRundownFilterSuggestionFactory Requested
+	  * @tparam InRundownSuggestionFactoryType The type of the factory to instantiate
+	  * @param InArgs Additional Args for constructor of the AvaRundownFilterSuggestionFactory class if needed
+	  * @return The new suggestion factory
+	  */
+	template <
+		typename InRundownSuggestionFactoryType,
+		typename... InArgsType
+		UE_REQUIRES(TIsDerivedFrom<InRundownSuggestionFactoryType, IAvaRundownFilterSuggestionFactory>::Value)
+	>
+	static TSharedRef<InRundownSuggestionFactoryType> MakeInstance(InArgsType&&... InArgs)
+	{
+		return MakeShared<InRundownSuggestionFactoryType>(Forward<InArgsType>(InArgs)...);
+	}
+
+	/** Get the suggestion identifier for this factory */
+	virtual FName GetSuggestionIdentifier() const = 0;
+
+	/** True if the suggestion doesn't need to use the item to add its suggestions */
+	virtual bool IsSimpleSuggestion() const = 0;
+	/**
+	 * Add suggestions entry to be shown
+	 * @param InPayload Arguments to be passed to the factory to get the suggestion to show
+	 */
+	virtual void AddSuggestion(const TSharedRef<FAvaRundownFilterSuggestionPayload>& InPayload) = 0;
+
+	/**
+	 * Whether the factory support a given type of suggestion
+	 * @param InSuggestionType Suggestion type to check
+	 * @return True if Factory support passed Suggestion Type, False otherwise
+	 */
+	virtual bool SupportSuggestionType(EAvaRundownSearchListType InSuggestionType) const = 0;
+};
diff --git a/Engine/Plugins/Experimental/Avalanche/Source/AvalancheMediaEditor/Private/Rundown/Pages/Columns/IAvaRundownPageViewColumn.h b/Engine/Plugins/Experimental/Avalanche/Source/AvalancheMediaEditor/Private/Rundown/Pages/Columns/IAvaRundownPageViewColumn.h
new file mode 100644
--- /dev/null
+++ b/Engine/Plugins/Experimental/Avalanche/Source/AvalancheMediaEditor/Private/Rundown/Pages/Columns/IAvaRundownPageViewColumn.h
@@ -0,0 +1,28 @@
+// Copyright Epic Games, Inc. All Rights Reserved.
+
+#pragma once
+
+#include "AvaType.h"
+#include "Rundown/AvaRundownEditorDefines.h"
+#include "Templates/SharedPointer.h"
+#include "Widgets/Views/SHeaderRow.h"
+
+class FName;
+class FText;
+class SAvaRundownPageViewRow;
+class SWidget;
+
+class IAvaRundownPageViewColumn : public IAvaTypeCastable, public TSharedFromThis<IAvaRundownPageViewColumn>
+{
+public:
+	UE_AVA_INHERITS(IAvaRundownPageViewColumn, IAvaTypeCastable);
+
+	FName GetColumnId() const { return GetTypeId().ToName(); }
+
+	virtual FText GetColumnDisplayNameText() const = 0;
+	virtual FText GetColumnToolTipText() const = 0;
+
+	virtual SHeaderRow::FColumn::FArguments ConstructHeaderRowColumn() = 0;
+	
+	virtual TSharedRef<SWidget> ConstructRowWidget(const FAvaRundownPageViewRef& InPageView, const TSharedPtr<SAvaRundownPageViewRow>& InRow) = 0;
+};
diff --git a/Engine/Plugins/Experimental/Avalanche/Source/AvalancheMediaEditor/Private/Rundown/Pages/PageViews/IAvaRundownInstancedPageView.h b/Engine/Plugins/Experimental/Avalanche/Source/AvalancheMediaEditor/Private/Rundown/Pages/PageViews/IAvaRundownInstancedPageView.h
new file mode 100644
--- /dev/null
+++ b/Engine/Plugins/Experimental/Avalanche/Source/AvalancheMediaEditor/Private/Rundown/Pages/PageViews/IAvaRundownInstancedPageView.h
@@ -0,0 +1,25 @@
+// Copyright Epic Games, Inc. All Rights Reserved.
+
+#pragma once
+
+#include "AvaType.h"
+
+enum class ECheckBoxState : uint8;
+struct FAvaRundownPage;
+
+class IAvaRundownInstancedPageView
+{
+public:
+	UE_AVA_TYPE(IAvaRundownInstancedPageView);
+
+	virtual ~IAvaRundownInstancedPageView() = default;
+
+	virtual ECheckBoxState IsEnabled() const = 0;
+	virtual void SetEnabled(ECheckBoxState InState) = 0;
+
+	virtual FName GetChannelName() const = 0;
+	virtual bool SetChannel(FName InChannel) = 0;
+
+	virtual const FAvaRundownPage& GetTemplate() const = 0;
+	virtual FText GetTemplateDescription() const = 0;
+};
diff --git a/Engine/Plugins/Experimental/Avalanche/Source/AvalancheMediaEditor/Private/Rundown/Pages/PageViews/IAvaRundownPageView.h b/Engine/Plugins/Experimental/Avalanche/Source/AvalancheMediaEditor/Private/Rundown/Pages/PageViews/IAvaRundownPageView.h
new file mode 100644
--- /dev/null
+++ b/Engine/Plugins/Experimental/Avalanche/Source/AvalancheMediaEditor/Private/Rundown/Pages/PageViews/IAvaRundownPageView.h
@@ -0,0 +1,70 @@
+// Copyright Epic Games, Inc. All Rights Reserved.
+
+#pragma once
+
+#include "AvaType.h"
+#include "Delegates/Delegate.h"
+#include "Templates/SharedPointer.h"
+
+class FReply;
+class FText;
+class UAvaRundown;
+struct FAssetData;
+struct FSoftObjectPath;
+
+enum class EAvaRundownPageActionState : uint8
+{
+	None,
+	Requested,
+	Completed,
+	Cancelled,
+};
+
+enum class EAvaRundownPageViewSelectionChangeType : uint8
+{
+	Deselect,
+	ReplaceSelection,
+	AddToSelection
+};
+
+class IAvaRundownPageView : public IAvaTypeCastable, public TSharedFromThis<IAvaRundownPageView>
+{
+public:
+	UE_AVA_INHERITS(IAvaRundownPageView, IAvaTypeCastable);
+
+	virtual UAvaRundown* GetRundown() const = 0;
+	
+	virtual int32 GetPageId() const = 0;
+	virtual FText GetPageIdText() const = 0;
+	virtual FText GetPageNameText() const = 0;
+	virtual FText GetPageTransitionLayerNameText() const = 0;
+
+	virtual FText GetPageSummary() const = 0;
+	virtual FText GetPageDescription() const = 0;
+
+	virtual bool IsTemplate() const = 0;
+
+	virtual bool HasObjectPath(const UAvaRundown* InRundown) const = 0;
+	virtual FSoftObjectPath GetObjectPath(const UAvaRundown* InRundown) const = 0;
+	virtual FText GetObjectName(const UAvaRundown* InRundown) const = 0;
+	virtual void OnObjectChanged(const FAssetData& InAssetData) = 0;
+
+	virtual bool Rename(const FText& InNewName) = 0;
+	virtual bool RenameFriendlyName(const FText& InNewName) = 0;
+	
+	DECLARE_MULTICAST_DELEGATE_OneParam(FOnPageAction, EAvaRundownPageActionState);
+	virtual FOnPageAction& GetOnRename() = 0;
+	virtual FOnPageAction& GetOnRenumber() = 0;
+
+	virtual FReply OnAssetStatusButtonClicked() = 0;
+	virtual bool CanChangeAssetStatus() const = 0;
+
+	virtual FReply OnPreviewButtonClicked() = 0;
+	virtual bool CanPreview() const = 0;
+
+	virtual FReply OnPlayButtonClicked() = 0;
+	virtual bool CanPlay() const = 0;
+
+	virtual bool IsPageSelected() const = 0;
+	virtual bool SetPageSelection(EAvaRundownPageViewSelectionChangeType InSelectionChangeType) = 0;
+};
diff --git a/Engine/Plugins/Experimental/Avalanche/Source/AvalancheMediaEditor/Private/Rundown/Pages/PageViews/IAvaRundownTemplatePageView.h b/Engine/Plugins/Experimental/Avalanche/Source/AvalancheMediaEditor/Private/Rundown/Pages/PageViews/IAvaRundownTemplatePageView.h
new file mode 100644
--- /dev/null
+++ b/Engine/Plugins/Experimental/Avalanche/Source/AvalancheMediaEditor/Private/Rundown/Pages/PageViews/IAvaRundownTemplatePageView.h
@@ -0,0 +1,13 @@
+// Copyright Epic Games, Inc. All Rights Reserved.
+
+#pragma once
+
+#include "AvaType.h"
+
+class IAvaRundownTemplatePageView
+{
+public:
+	UE_AVA_TYPE(IAvaRundownTemplatePageView);
+
+	virtual ~IAvaRundownTemplatePageView() = default;
+};
diff --git a/Engine/Plugins/Experimental/Avalanche/Source/AvalancheMediaEditor/Private/Rundown/Pages/Slate/SAvaRundownPageViewRow.h b/Engine/Plugins/Experimental/Avalanche/Source/AvalancheMediaEditor/Private/Rundown/Pages/Slate/SAvaRundownPageViewRow.h
new file mode 100644
--- /dev/null
+++ b/Engine/Plugins/Experimental/Avalanche/Source/AvalancheMediaEditor/Private/Rundown/Pages/Slate/SAvaRundownPageViewRow.h
@@ -0,0 +1,71 @@
+// Copyright Epic Games, Inc. All Rights Reserved.
+
+#pragma once
+
+#include "Containers/Array.h"
+#include "Input/DragAndDrop.h"
+#include "Rundown/AvaRundownEditorDefines.h"
+#include "Templates/SharedPointer.h"
+#include "Widgets/Views/SListView.h"
+#include "Widgets/Views/STableRow.h"
+#include "UObject/NameTypes.h"
+#include "Misc/Optional.h"
+
+class FDragDropEvent;
+class FReply;
+class IAvaRundownPageView;
+class SAvaRundownPageList;
+class SWidget;
+struct FGeometry;
+struct FPointerEvent;
+struct FSlateBrush;
+
+enum class EAvaRundownRowState
+{
+	Active,
+	Disabled
+};
+
+class FAvaRundownPageViewRowDragDropOp : public FDragDropOperation, public TSharedFromThis<FAvaRundownPageViewRowDragDropOp>
+{
+public:
+	DRAG_DROP_OPERATOR_TYPE(FAvaRundownPageViewRowDragDropOp, FDragDropOperation)
+
+	virtual TSharedPtr<SWidget> GetDefaultDecorator() const override;
+
+	static TSharedRef<FAvaRundownPageViewRowDragDropOp> New(const TSharedPtr<SAvaRundownPageList>& InPageList);
+
+	TSharedPtr<SAvaRundownPageList> GetPageList() const { return PageListWeak.Pin(); }
+	const TArray<int32>& GetDraggedIds() const { return DraggedIds; }
+
+	FName Color;
+
+protected:
+	TWeakPtr<SAvaRundownPageList> PageListWeak;
+	TArray<int32> DraggedIds;
+};
+
+class SAvaRundownPageViewRow : public SMultiColumnTableRow<FAvaRundownPageViewPtr>
+{
+public:	
+	SLATE_BEGIN_ARGS(SAvaRundownPageViewRow){}
+	SLATE_END_ARGS()
+
+	/** Constructs this widget with InArgs */
+	void Construct(const FArguments& InArgs, FAvaRundownPageViewPtr InPageView, const TSharedRef<SAvaRundownPageList>& InPageList);
+
+	virtual TSharedRef<SWidget> GenerateWidgetForColumn(const FName& InColumnName) override;
+
+	virtual const FSlateBrush* GetBorder() const override;
+
+protected:
+	FReply OnPageViewDragDetected(const FGeometry& InMyGeometry, const FPointerEvent& InPointerEvent);
+	TOptional<EItemDropZone> OnPageViewCanAcceptDrop(const FDragDropEvent& InDragDropEvent, EItemDropZone InDropZone, FAvaRundownPageViewPtr InItem);
+	FReply OnPageViewAcceptDrop(const FDragDropEvent& InDragDropEvent, EItemDropZone InDropZone, FAvaRundownPageViewPtr InItem);
+
+	TWeakPtr<IAvaRundownPageView> PageViewWeak;
+	TWeakPtr<SAvaRundownPageList> PageListWeak;
+
+private:
+	const FSlateBrush* GetRowColorBrush(EAvaRundownRowState InRowState) const;
+};
diff --git a/Engine/Plugins/Experimental/Avalanche/Source/AvalancheMediaEditor/Public/AvaMediaEditorSettings.h b/Engine/Plugins/Experimental/Avalanche/Source/AvalancheMediaEditor/Public/AvaMediaEditorSettings.h
deleted file mode 100644
diff --git a/Engine/Plugins/Experimental/Avalanche/Source/AvalancheMediaEditor/Public/Rundown/AvaRundownEditor.h b/Engine/Plugins/Experimental/Avalanche/Source/AvalancheMediaEditor/Public/Rundown/AvaRundownEditor.h
--- a/Engine/Plugins/Experimental/Avalanche/Source/AvalancheMediaEditor/Public/Rundown/AvaRundownEditor.h
+++ b/Engine/Plugins/Experimental/Avalanche/Source/AvalancheMediaEditor/Public/Rundown/AvaRundownEditor.h
@@ -15,7 +15,7 @@ class SAvaRundownPageList;
 class SAvaRundownReadPage;
 enum class EAvaRundownSearchListType : uint8;
 
-class AVALANCHEMEDIAEDITOR_API FAvaRundownEditor : public FWorkflowCentricApplication
+class FAvaRundownEditor : public FWorkflowCentricApplication
 {
 public:
 	FAvaRundownEditor();
@@ -41,7 +41,7 @@ public:
 	void FillPageToolBar(FToolBarBuilder& OutToolBarBuilder);
 
 	bool IsRundownValid() const { return AvaRundown.IsValid(); }
-	UAvaRundown* GetRundown() const;
+	AVALANCHEMEDIAEDITOR_API UAvaRundown* GetRundown() const;
 	void MarkAsModified();
 
 	DECLARE_MULTICAST_DELEGATE_TwoParams(FOnPageEvent, const TArray<int32>&, UE::AvaRundown::EPageEvent);
@@ -61,9 +61,9 @@ public:
 	/** Returns the currently focused Page List, including templates, instances or views. */
 	TSharedPtr<SAvaRundownPageList> GetFocusedListWidget() const;
 
-	int32 GetFirstSelectedPageOnActiveSubListWidget() const;
-	TConstArrayView<int32> GetSelectedPagesOnActiveSubListWidget() const;
-	TConstArrayView<int32> GetSelectedPagesOnFocusedWidget() const;
+	AVALANCHEMEDIAEDITOR_API int32 GetFirstSelectedPageOnActiveSubListWidget() const;
+	AVALANCHEMEDIAEDITOR_API TConstArrayView<int32> GetSelectedPagesOnActiveSubListWidget() const;
+	AVALANCHEMEDIAEDITOR_API TConstArrayView<int32> GetSelectedPagesOnFocusedWidget() const;
 
 	bool CanAddTemplate() const;
 	void AddTemplate();
diff --git a/Engine/Plugins/Experimental/Avalanche/Source/AvalancheMediaEditor/Public/Rundown/AvaRundownEditorSettings.h b/Engine/Plugins/Experimental/Avalanche/Source/AvalancheMediaEditor/Public/Rundown/AvaRundownEditorSettings.h
deleted file mode 100644
diff --git a/Engine/Plugins/Experimental/Avalanche/Source/AvalancheMediaEditor/Public/Rundown/AvaRundownMessages.h b/Engine/Plugins/Experimental/Avalanche/Source/AvalancheMediaEditor/Public/Rundown/AvaRundownMessages.h
deleted file mode 100644
diff --git a/Engine/Plugins/Experimental/Avalanche/Source/AvalancheMediaEditor/Public/Rundown/AvaRundownServer.h b/Engine/Plugins/Experimental/Avalanche/Source/AvalancheMediaEditor/Public/Rundown/AvaRundownServer.h
deleted file mode 100644
diff --git a/Engine/Plugins/Experimental/Avalanche/Source/AvalancheMediaEditor/Public/Rundown/Factories/Filters/AvaRundownFactoriesUtils.h b/Engine/Plugins/Experimental/Avalanche/Source/AvalancheMediaEditor/Public/Rundown/Factories/Filters/AvaRundownFactoriesUtils.h
deleted file mode 100644
diff --git a/Engine/Plugins/Experimental/Avalanche/Source/AvalancheMediaEditor/Public/Rundown/Factories/Filters/IAvaRundownFilterExpressionFactory.h b/Engine/Plugins/Experimental/Avalanche/Source/AvalancheMediaEditor/Public/Rundown/Factories/Filters/IAvaRundownFilterExpressionFactory.h
deleted file mode 100644
diff --git a/Engine/Plugins/Experimental/Avalanche/Source/AvalancheMediaEditor/Public/Rundown/Factories/Filters/IAvaRundownFilterSuggestionFactory.h b/Engine/Plugins/Experimental/Avalanche/Source/AvalancheMediaEditor/Public/Rundown/Factories/Filters/IAvaRundownFilterSuggestionFactory.h
deleted file mode 100644
diff --git a/Engine/Plugins/Experimental/Avalanche/Source/AvalancheMediaEditor/Public/Rundown/Pages/Columns/IAvaRundownPageViewColumn.h b/Engine/Plugins/Experimental/Avalanche/Source/AvalancheMediaEditor/Public/Rundown/Pages/Columns/IAvaRundownPageViewColumn.h
deleted file mode 100644
diff --git a/Engine/Plugins/Experimental/Avalanche/Source/AvalancheMediaEditor/Public/Rundown/Pages/PageViews/IAvaRundownInstancedPageView.h b/Engine/Plugins/Experimental/Avalanche/Source/AvalancheMediaEditor/Public/Rundown/Pages/PageViews/IAvaRundownInstancedPageView.h
deleted file mode 100644
diff --git a/Engine/Plugins/Experimental/Avalanche/Source/AvalancheMediaEditor/Public/Rundown/Pages/PageViews/IAvaRundownPageView.h b/Engine/Plugins/Experimental/Avalanche/Source/AvalancheMediaEditor/Public/Rundown/Pages/PageViews/IAvaRundownPageView.h
deleted file mode 100644
diff --git a/Engine/Plugins/Experimental/Avalanche/Source/AvalancheMediaEditor/Public/Rundown/Pages/PageViews/IAvaRundownTemplatePageView.h b/Engine/Plugins/Experimental/Avalanche/Source/AvalancheMediaEditor/Public/Rundown/Pages/PageViews/IAvaRundownTemplatePageView.h
deleted file mode 100644
diff --git a/Engine/Plugins/Experimental/Avalanche/Source/AvalancheMediaEditor/Public/Rundown/Pages/Slate/SAvaRundownPageViewRow.h b/Engine/Plugins/Experimental/Avalanche/Source/AvalancheMediaEditor/Public/Rundown/Pages/Slate/SAvaRundownPageViewRow.h
deleted file mode 100644
