[Motion Design]
 - Rename UE::Avalanche namespace to UE::Ava
 - Change scene asset tag from "AvalancheScene" to "MotionDesignScene"
 - Remove "Avalanche" from variable names
 - Change style names "AvalancheIcons.*" to "Icons.*"
 - Change "Avalanche" references to "Motion Design"
 - Replace "typedef" with "using"

diff --git a/Engine/Plugins/Experimental/AdvancedRenamer/Source/AdvancedRenamer/Private/Slate/SAdvancedRenamerPanel.h b/Engine/Plugins/Experimental/AdvancedRenamer/Source/AdvancedRenamer/Private/Slate/SAdvancedRenamerPanel.h
--- a/Engine/Plugins/Experimental/AdvancedRenamer/Source/AdvancedRenamer/Private/Slate/SAdvancedRenamerPanel.h
+++ b/Engine/Plugins/Experimental/AdvancedRenamer/Source/AdvancedRenamer/Private/Slate/SAdvancedRenamerPanel.h
@@ -40,8 +40,8 @@ struct FAdvancedRenamerPreviewListItem
 	mutable FString NewName;
 };
 
-typedef TSharedPtr<FAdvancedRenamerPreviewListItem, ESPMode::ThreadSafe> FObjectRenamePreviewListItemPtr;
-typedef TWeakPtr<FAdvancedRenamerPreviewListItem, ESPMode::ThreadSafe> FObjectRenamePreviewListItemWeakPtr;
+using FObjectRenamePreviewListItemPtr = TSharedPtr<FAdvancedRenamerPreviewListItem, ESPMode::ThreadSafe>;
+using FObjectRenamePreviewListItemWeakPtr = TWeakPtr<FAdvancedRenamerPreviewListItem, ESPMode::ThreadSafe>;
 
 class SAdvancedRenamerPreviewListRow : public SMultiColumnTableRow<FObjectRenamePreviewListItemPtr>
 {
diff --git a/Engine/Plugins/Experimental/Avalanche/Source/Avalanche/Internal/Tests/Framework/AvaTestUtils.h b/Engine/Plugins/Experimental/Avalanche/Source/Avalanche/Internal/Tests/Framework/AvaTestUtils.h
--- a/Engine/Plugins/Experimental/Avalanche/Source/Avalanche/Internal/Tests/Framework/AvaTestUtils.h
+++ b/Engine/Plugins/Experimental/Avalanche/Source/Avalanche/Internal/Tests/Framework/AvaTestUtils.h
@@ -1,4 +1,4 @@
-﻿// Copyright Epic Games, Inc. All Rights Reserved.
+// Copyright Epic Games, Inc. All Rights Reserved.
 
 #pragma once
 
@@ -6,7 +6,7 @@
 #include "Math/Transform.h"
 #include "UObject/WeakObjectPtrTemplates.h"
 
-// General Avalanche unit test log
+// General Motion Design unit test log
 DECLARE_LOG_CATEGORY_EXTERN(LogAvaTest, Log, All);
 
 class AActor;
diff --git a/Engine/Plugins/Experimental/Avalanche/Source/Avalanche/Private/AvaModule.h b/Engine/Plugins/Experimental/Avalanche/Source/Avalanche/Private/AvaModule.h
--- a/Engine/Plugins/Experimental/Avalanche/Source/Avalanche/Private/AvaModule.h
+++ b/Engine/Plugins/Experimental/Avalanche/Source/Avalanche/Private/AvaModule.h
@@ -25,7 +25,7 @@ public:
 
 private:
 	/**
-	 * Stats for all Avalanche runtime viewports is centralized here.
+	 * Stats for all Motion Design runtime viewports is centralized here.
 	 * The state (enabled) of the stats will persist across the lifetime of all
 	 * game instances so that an enabled stat remains even if playback is stopped.
 	 */
diff --git a/Engine/Plugins/Experimental/Avalanche/Source/Avalanche/Private/AvaScene.cpp b/Engine/Plugins/Experimental/Avalanche/Source/Avalanche/Private/AvaScene.cpp
--- a/Engine/Plugins/Experimental/Avalanche/Source/Avalanche/Private/AvaScene.cpp
+++ b/Engine/Plugins/Experimental/Avalanche/Source/Avalanche/Private/AvaScene.cpp
@@ -143,7 +143,7 @@ void AAvaScene::OnGetWorldTags(FAssetRegistryTagsContext Context) const
 	}
 
 	using namespace UE::Ava;
-	Context.AddTag(UObject::FAssetRegistryTag(AssetTags::AvalancheScene, AssetTags::Values::Enabled, UObject::FAssetRegistryTag::TT_Alphabetical));
+	Context.AddTag(UObject::FAssetRegistryTag(AssetTags::MotionDesignScene, AssetTags::Values::Enabled, UObject::FAssetRegistryTag::TT_Alphabetical));
 }
 #endif
 
diff --git a/Engine/Plugins/Experimental/Avalanche/Source/Avalanche/Private/AvaSceneSubsystem.cpp b/Engine/Plugins/Experimental/Avalanche/Source/Avalanche/Private/AvaSceneSubsystem.cpp
--- a/Engine/Plugins/Experimental/Avalanche/Source/Avalanche/Private/AvaSceneSubsystem.cpp
+++ b/Engine/Plugins/Experimental/Avalanche/Source/Avalanche/Private/AvaSceneSubsystem.cpp
@@ -7,7 +7,7 @@
 #include "Engine/World.h"
 #include "IAvaSceneInterface.h"
 
-namespace UE::Avalanche::Private
+namespace UE::Ava::Private
 {
 	static const TSet<EWorldType::Type> GUnsupportedWorlds
 		{
@@ -72,6 +72,6 @@ void UAvaSceneSubsystem::PostInitialize()
 
 bool UAvaSceneSubsystem::DoesSupportWorldType(const EWorldType::Type InWorldType) const
 {
-	const bool bDisallowedWorld = UE::Avalanche::Private::GUnsupportedWorlds.Contains(InWorldType);
+	const bool bDisallowedWorld = UE::Ava::Private::GUnsupportedWorlds.Contains(InWorldType);
 	return !bDisallowedWorld;
 }
diff --git a/Engine/Plugins/Experimental/Avalanche/Source/Avalanche/Private/Framework/AvaSoftAssetPtr.cpp b/Engine/Plugins/Experimental/Avalanche/Source/Avalanche/Private/Framework/AvaSoftAssetPtr.cpp
--- a/Engine/Plugins/Experimental/Avalanche/Source/Avalanche/Private/Framework/AvaSoftAssetPtr.cpp
+++ b/Engine/Plugins/Experimental/Avalanche/Source/Avalanche/Private/Framework/AvaSoftAssetPtr.cpp
@@ -3,13 +3,13 @@
 #include "Framework/AvaSoftAssetPtr.h"
 #include "Engine/World.h"
 
-EAvalancheAssetType FAvaSoftAssetPath::GetAssetTypeFromClass(const FSoftClassPath& InAssetClassPath, bool bInLoadIfUnknown)
+EMotionDesignAssetType FAvaSoftAssetPath::GetAssetTypeFromClass(const FSoftClassPath& InAssetClassPath, bool bInLoadIfUnknown)
 {
 	static const FSoftClassPath WorldClassPath(UWorld::StaticClass());
 
 	if (InAssetClassPath == WorldClassPath)
 	{
-		return EAvalancheAssetType::World;
+		return EMotionDesignAssetType::World;
 	}
 
 	// Try to load the asset class and compare possible derived types.
@@ -18,16 +18,16 @@ EAvalancheAssetType FAvaSoftAssetPath::GetAssetTypeFromClass(const FSoftClassPat
 		const UClass* AssetClass = InAssetClassPath.ResolveClass();
 		if (AssetClass->IsChildOf(UWorld::StaticClass()))
 		{
-			return EAvalancheAssetType::World;
+			return EMotionDesignAssetType::World;
 		}
 	}
-	return EAvalancheAssetType::Unknown;
+	return EMotionDesignAssetType::Unknown;
 }
 
-EAvalancheAssetType FAvaSoftAssetPtr::GetAssetType(bool bInLoadIfUnknown) const
+EMotionDesignAssetType FAvaSoftAssetPtr::GetAssetType(bool bInLoadIfUnknown) const
 {
-	const EAvalancheAssetType AssetTypeFromClass = FAvaSoftAssetPath::GetAssetTypeFromClass(AssetClassPath, bInLoadIfUnknown);
-	if (AssetTypeFromClass != EAvalancheAssetType::Unknown)
+	const EMotionDesignAssetType AssetTypeFromClass = FAvaSoftAssetPath::GetAssetTypeFromClass(AssetClassPath, bInLoadIfUnknown);
+	if (AssetTypeFromClass != EMotionDesignAssetType::Unknown)
 	{
 		return AssetTypeFromClass;
 	}
@@ -40,15 +40,15 @@ EAvalancheAssetType FAvaSoftAssetPtr::GetAssetType(bool bInLoadIfUnknown) const
 		UObject* LoadedSourceAsset;
 		{
 			TRACE_CPUPROFILER_EVENT_SCOPE(FAvaSoftAssetPtr::LoadSourceAsset);
-			const FSoftObjectPtr SourceAvalancheAsset(AssetPtr.ToSoftObjectPath());
-			LoadedSourceAsset = SourceAvalancheAsset.LoadSynchronous();
+			const FSoftObjectPtr SourceAsset(AssetPtr.ToSoftObjectPath());
+			LoadedSourceAsset = SourceAsset.LoadSynchronous();
 		}
 
 		if (Cast<UWorld>(LoadedSourceAsset))
 		{
-			return EAvalancheAssetType::World;
+			return EMotionDesignAssetType::World;
 		}
 	}
 
-	return EAvalancheAssetType::Unknown;
-}
\ No newline at end of file
+	return EMotionDesignAssetType::Unknown;
+}
diff --git a/Engine/Plugins/Experimental/Avalanche/Source/Avalanche/Private/Tests/Framework/AvaTestUtils.cpp b/Engine/Plugins/Experimental/Avalanche/Source/Avalanche/Private/Tests/Framework/AvaTestUtils.cpp
--- a/Engine/Plugins/Experimental/Avalanche/Source/Avalanche/Private/Tests/Framework/AvaTestUtils.cpp
+++ b/Engine/Plugins/Experimental/Avalanche/Source/Avalanche/Private/Tests/Framework/AvaTestUtils.cpp
@@ -1,4 +1,4 @@
-﻿// Copyright Epic Games, Inc. All Rights Reserved.
+// Copyright Epic Games, Inc. All Rights Reserved.
 
 #include "Tests/Framework/AvaTestUtils.h"
 
@@ -11,7 +11,7 @@
 #include "Tests/Framework/AvaTestStaticMeshActor.h"
 #include "UDynamicMesh.h"
 
-// General Avalanche unit test log
+// General Motion Design unit test log
 DEFINE_LOG_CATEGORY(LogAvaTest);
 
 void FAvaTestUtils::Init()
diff --git a/Engine/Plugins/Experimental/Avalanche/Source/Avalanche/Private/Viewport/AvaCameraManager.cpp b/Engine/Plugins/Experimental/Avalanche/Source/Avalanche/Private/Viewport/AvaCameraManager.cpp
--- a/Engine/Plugins/Experimental/Avalanche/Source/Avalanche/Private/Viewport/AvaCameraManager.cpp
+++ b/Engine/Plugins/Experimental/Avalanche/Source/Avalanche/Private/Viewport/AvaCameraManager.cpp
@@ -85,7 +85,7 @@ bool FAvaCameraManager::SetDefaultViewTarget(UWorld* InWorld, FName InDefaultVie
 {
 	if (!IsValid(InWorld))
 	{
-		UE_LOG(LogViewport, Error, TEXT("Error setting default view target for Avalanche BP: Null world."));
+		UE_LOG(LogViewport, Error, TEXT("Error setting default view target for Motion Design: Null world."));
 		return false;
 	}
 
@@ -103,7 +103,7 @@ bool FAvaCameraManager::SetDefaultViewTarget(UWorld* InWorld, FName InDefaultVie
 		}
 	}
 
-	UE_LOG(LogViewport, Error, TEXT("Error setting default view target for Avalanche BP: Startup camera was not found."));
+	UE_LOG(LogViewport, Error, TEXT("Error setting default view target for Motion Design: Startup camera was not found."));
 	return false;
 }
 #endif
diff --git a/Engine/Plugins/Experimental/Avalanche/Source/Avalanche/Private/Viewport/AvaViewportQualitySettings.cpp b/Engine/Plugins/Experimental/Avalanche/Source/Avalanche/Private/Viewport/AvaViewportQualitySettings.cpp
--- a/Engine/Plugins/Experimental/Avalanche/Source/Avalanche/Private/Viewport/AvaViewportQualitySettings.cpp
+++ b/Engine/Plugins/Experimental/Avalanche/Source/Avalanche/Private/Viewport/AvaViewportQualitySettings.cpp
@@ -48,7 +48,7 @@ FAvaViewportQualitySettings::FAvaViewportQualitySettings(const TArray<FAvaViewpo
 
 TArray<FAvaViewportQualitySettingsFeature> FAvaViewportQualitySettings::DefaultFeatures()
 {
-	// Add all the engine features we will support for editing by the designer and their Avalanche defaults.
+	// Add all the engine features we will support for editing by the designer and their Motion Design defaults.
 	TArray<FAvaViewportQualitySettingsFeature> AllDefaultFeatures =
 		{
 			{ "AntiAliasing", true },
diff --git a/Engine/Plugins/Experimental/Avalanche/Source/Avalanche/Public/AvaAssetTags.h b/Engine/Plugins/Experimental/Avalanche/Source/Avalanche/Public/AvaAssetTags.h
--- a/Engine/Plugins/Experimental/Avalanche/Source/Avalanche/Public/AvaAssetTags.h
+++ b/Engine/Plugins/Experimental/Avalanche/Source/Avalanche/Public/AvaAssetTags.h
@@ -6,7 +6,7 @@
 
 namespace UE::Ava::AssetTags
 {
-	constexpr const TCHAR* AvalancheScene = TEXT("AvalancheScene");
+	constexpr const TCHAR* MotionDesignScene = TEXT("MotionDesignScene");
 
 	namespace Values
 	{
diff --git a/Engine/Plugins/Experimental/Avalanche/Source/Avalanche/Public/AvaDefs.h b/Engine/Plugins/Experimental/Avalanche/Source/Avalanche/Public/AvaDefs.h
--- a/Engine/Plugins/Experimental/Avalanche/Source/Avalanche/Public/AvaDefs.h
+++ b/Engine/Plugins/Experimental/Avalanche/Source/Avalanche/Public/AvaDefs.h
@@ -1,4 +1,4 @@
-﻿// Copyright Epic Games, Inc. All Rights Reserved.
+// Copyright Epic Games, Inc. All Rights Reserved.
 
 #pragma once
 
@@ -8,7 +8,7 @@
 #include "AvaDefs.generated.h"
 
 // Single-axis exclusive bitflag map
-namespace UE::Avalanche::AnchorPoints
+namespace UE::Ava::AnchorPoints
 {
 	constexpr int32 None       = 0;
 	constexpr int32 Left       = 1 << 0; // Y- = 1
@@ -28,17 +28,17 @@ namespace UE::Avalanche::AnchorPoints
 UENUM(BlueprintType)
 enum class EAvaAnchors : uint8
 {
-	None = UE::Avalanche::AnchorPoints::None UMETA(Hidden),
-	TopLeft = UE::Avalanche::AnchorPoints::Top + UE::Avalanche::AnchorPoints::Left,
-	Top = UE::Avalanche::AnchorPoints::Top + UE::Avalanche::AnchorPoints::HMiddle,
-	TopRight = UE::Avalanche::AnchorPoints::Top + UE::Avalanche::AnchorPoints::Right,
-	Left = UE::Avalanche::AnchorPoints::VMiddle + UE::Avalanche::AnchorPoints::Left,
-	Center = UE::Avalanche::AnchorPoints::VMiddle + UE::Avalanche::AnchorPoints::HMiddle,
-	Right = UE::Avalanche::AnchorPoints::VMiddle + UE::Avalanche::AnchorPoints::Right,
-	BottomLeft = UE::Avalanche::AnchorPoints::Bottom + UE::Avalanche::AnchorPoints::Left,
-	Bottom = UE::Avalanche::AnchorPoints::Bottom + UE::Avalanche::AnchorPoints::HMiddle,
-	BottomRight = UE::Avalanche::AnchorPoints::Bottom + UE::Avalanche::AnchorPoints::Right,
-	Custom = UE::Avalanche::AnchorPoints::Custom
+	None = UE::Ava::AnchorPoints::None UMETA(Hidden),
+	TopLeft = UE::Ava::AnchorPoints::Top + UE::Ava::AnchorPoints::Left,
+	Top = UE::Ava::AnchorPoints::Top + UE::Ava::AnchorPoints::HMiddle,
+	TopRight = UE::Ava::AnchorPoints::Top + UE::Ava::AnchorPoints::Right,
+	Left = UE::Ava::AnchorPoints::VMiddle + UE::Ava::AnchorPoints::Left,
+	Center = UE::Ava::AnchorPoints::VMiddle + UE::Ava::AnchorPoints::HMiddle,
+	Right = UE::Ava::AnchorPoints::VMiddle + UE::Ava::AnchorPoints::Right,
+	BottomLeft = UE::Ava::AnchorPoints::Bottom + UE::Ava::AnchorPoints::Left,
+	Bottom = UE::Ava::AnchorPoints::Bottom + UE::Ava::AnchorPoints::HMiddle,
+	BottomRight = UE::Ava::AnchorPoints::Bottom + UE::Ava::AnchorPoints::Right,
+	Custom = UE::Ava::AnchorPoints::Custom
 };
 
 UENUM(BlueprintType)
@@ -106,7 +106,7 @@ constexpr auto ToUnderlyingType(Enum Val)
 	return static_cast<std::underlying_type_t<Enum>>(Val);
 }
 
-typedef int32 AvaAlignment; // xyz
+using AvaAlignment = int32; // xyz
 
 // alignment for X (Depth)
 template<typename AvaAlignment>
diff --git a/Engine/Plugins/Experimental/Avalanche/Source/Avalanche/Public/AvaOutlinerUtils.h b/Engine/Plugins/Experimental/Avalanche/Source/Avalanche/Public/AvaOutlinerUtils.h
--- a/Engine/Plugins/Experimental/Avalanche/Source/Avalanche/Public/AvaOutlinerUtils.h
+++ b/Engine/Plugins/Experimental/Avalanche/Source/Avalanche/Public/AvaOutlinerUtils.h
@@ -15,13 +15,13 @@ class IAvaOutliner;
 /** Editor Only Outliner Access Functions */
 struct AVALANCHE_API FAvaOutlinerUtils
 {
-	/** Returns the Avalanche outliner from the outliner module at runtime in the editor. */
+	/** Returns the Motion Design outliner from the outliner module at runtime in the editor. */
 	static TSharedPtr<IAvaOutliner> EditorGetOutliner(const UWorld* const InWorld);
 
 	/** Returns an array of direct child actors attached to a parent actor or all root level actors if the parent actor is null. */
 	static TArray<AActor*> EditorOutlinerChildActors(TSharedPtr<IAvaOutliner> InOutliner, AActor* const InParentActor = nullptr);
 
-	/** Returns true if the Avalanche BP editor is isolating actors and the array of actors currently being isolated. */
+	/** Returns true if the Motion Design editor is isolating actors and the array of actors currently being isolated. */
 	static bool EditorActorIsolationInfo(TSharedPtr<IAvaOutliner> InOutliner, TArray<TWeakObjectPtr<const AActor>>& OutIsolatedActors);
 
 	/**
diff --git a/Engine/Plugins/Experimental/Avalanche/Source/Avalanche/Public/Framework/AvaSoftAssetPtr.h b/Engine/Plugins/Experimental/Avalanche/Source/Avalanche/Public/Framework/AvaSoftAssetPtr.h
--- a/Engine/Plugins/Experimental/Avalanche/Source/Avalanche/Public/Framework/AvaSoftAssetPtr.h
+++ b/Engine/Plugins/Experimental/Avalanche/Source/Avalanche/Public/Framework/AvaSoftAssetPtr.h
@@ -4,14 +4,14 @@
 
 #include "UObject/SoftObjectPtr.h"
 
-enum class EAvalancheAssetType
+enum class EMotionDesignAssetType
 {
 	Unknown,
 	World
 };
 
 /**
- * Avalanche Extension to FSoftObjectPath.
+ * Motion Design Extension to FSoftObjectPath.
  * Adding the class information along to resolve about a Motion Design Asset without loading anything.
  */
 struct FAvaSoftAssetPath
@@ -21,13 +21,13 @@ struct FAvaSoftAssetPath
 
 	bool IsNull() const { return AssetPath.IsNull(); }
 	const FSoftObjectPath& ToSoftObjectPath() const { return AssetPath; }
-	EAvalancheAssetType GetAssetType(bool bInLoadIfUnknown = true) const { return GetAssetTypeFromClass(AssetClassPath, bInLoadIfUnknown);}
+	EMotionDesignAssetType GetAssetType(bool bInLoadIfUnknown = true) const { return GetAssetTypeFromClass(AssetClassPath, bInLoadIfUnknown);}
 
-	AVALANCHE_API static EAvalancheAssetType GetAssetTypeFromClass(const FSoftClassPath& InAssetClassPath, bool bInLoadIfUnknown = true);
+	AVALANCHE_API static EMotionDesignAssetType GetAssetTypeFromClass(const FSoftClassPath& InAssetClassPath, bool bInLoadIfUnknown = true);
 };
 
 /**
- * Avalanche Extension to FSoftObjectPtr.
+ * Motion Design Extension to FSoftObjectPtr.
  * Adding the class information along to resolve a Motion Design Asset without loading anything.
  */
 
@@ -47,9 +47,9 @@ struct FAvaSoftAssetPtr
 	/**
 	 * @brief Determines the asset type from the given class path or by loading the object.
 	 * @param bInLoadIfUnknown If the type can't be determined by the class path, the object is loaded to check it's type.
-	 * @return One of the recognised asset type for Avalanche playables.
+	 * @return One of the recognized asset types for Motion Design playables.
 	 */
-	AVALANCHE_API EAvalancheAssetType GetAssetType(bool bInLoadIfUnknown = true) const;
+	AVALANCHE_API EMotionDesignAssetType GetAssetType(bool bInLoadIfUnknown = true) const;
 
 	FORCEINLINE bool operator==(const FAvaSoftAssetPtr& Rhs) const
 	{
diff --git a/Engine/Plugins/Experimental/Avalanche/Source/Avalanche/Public/IAvaModule.h b/Engine/Plugins/Experimental/Avalanche/Source/Avalanche/Public/IAvaModule.h
--- a/Engine/Plugins/Experimental/Avalanche/Source/Avalanche/Public/IAvaModule.h
+++ b/Engine/Plugins/Experimental/Avalanche/Source/Avalanche/Public/IAvaModule.h
@@ -27,7 +27,7 @@ public:
 
 	/**
 	* Set a specific stat to either enabled or disabled (returns the number of remaining enabled stats)
-	* This applies to all the Avalanche instance viewports.
+	* This applies to all the Motion Design instance viewports.
 	*/
 	virtual int32 SetRuntimeStatEnabled(const TCHAR* InName, const bool bInEnabled) = 0;
 
diff --git a/Engine/Plugins/Experimental/Avalanche/Source/Avalanche/Public/IAvaSceneInterface.h b/Engine/Plugins/Experimental/Avalanche/Source/Avalanche/Public/IAvaSceneInterface.h
--- a/Engine/Plugins/Experimental/Avalanche/Source/Avalanche/Public/IAvaSceneInterface.h
+++ b/Engine/Plugins/Experimental/Avalanche/Source/Avalanche/Public/IAvaSceneInterface.h
@@ -14,7 +14,7 @@ class URemoteControlPreset;
 struct FAvaSceneTree;
 
 /**
- * Interface for the driver class of an Avalanche Scene
+ * Interface for the driver class of an Motion Design Scene
  * This interface's purpose is to provide access to scene data and to core interfaces in the Scene
  */
 UINTERFACE(MinimalAPI, meta = (CannotImplementInterfaceInBlueprint))
diff --git a/Engine/Plugins/Experimental/Avalanche/Source/Avalanche/Public/Viewport/AvaCineCameraActor.h b/Engine/Plugins/Experimental/Avalanche/Source/Avalanche/Public/Viewport/AvaCineCameraActor.h
--- a/Engine/Plugins/Experimental/Avalanche/Source/Avalanche/Public/Viewport/AvaCineCameraActor.h
+++ b/Engine/Plugins/Experimental/Avalanche/Source/Avalanche/Public/Viewport/AvaCineCameraActor.h
@@ -23,7 +23,7 @@ public:
 	AAvaCineCameraActor(const FObjectInitializer& ObjectInitializer);
 
 	/**
-	 * Initialize the camera with Avalanche scene default values: field of view, camera position, focus distance.
+	 * Initialize the camera with Motion Design scene default values: field of view, camera position, focus distance.
 	 * @param InCameraDistance: camera distance value, used to initialize camera position and manual focus distance
 	 */
 	AVALANCHE_API void Configure(float InCameraDistance);
diff --git a/Engine/Plugins/Experimental/Avalanche/Source/Avalanche/Public/Viewport/AvaViewportQualitySettings.h b/Engine/Plugins/Experimental/Avalanche/Source/Avalanche/Public/Viewport/AvaViewportQualitySettings.h
--- a/Engine/Plugins/Experimental/Avalanche/Source/Avalanche/Public/Viewport/AvaViewportQualitySettings.h
+++ b/Engine/Plugins/Experimental/Avalanche/Source/Avalanche/Public/Viewport/AvaViewportQualitySettings.h
@@ -34,10 +34,10 @@ struct FAvaViewportQualitySettingsFeature
 };
 
 /** 
- * Avalanche Viewport Quality Settings
+ * Motion Design Viewport Quality Settings
  * 
  * Advanced render and quality viewport settings to control performance for a given Viewport.
- * Human-readable and blueprintable structure that holds flags for the FShowEngineFlags structure.
+ * Human-readable and blueprint-able structure that holds flags for the FShowEngineFlags structure.
  * Can convert FShowEngineFlags to FAvaViewportQualitySettings and apply FAvaViewportQualitySettings to a FShowEngineFlags structure.
  */
 USTRUCT(BlueprintType)
diff --git a/Engine/Plugins/Experimental/Avalanche/Source/AvalancheComponentVisualizers/Private/AvaComponentVisualizersEdMode.cpp b/Engine/Plugins/Experimental/Avalanche/Source/AvalancheComponentVisualizers/Private/AvaComponentVisualizersEdMode.cpp
--- a/Engine/Plugins/Experimental/Avalanche/Source/AvalancheComponentVisualizers/Private/AvaComponentVisualizersEdMode.cpp
+++ b/Engine/Plugins/Experimental/Avalanche/Source/AvalancheComponentVisualizers/Private/AvaComponentVisualizersEdMode.cpp
@@ -9,7 +9,7 @@
 
 #define LOCTEXT_NAMESPACE "AvaComponentVisualizersEdMode"
 
-namespace UE::AvalancheComponentVisualizers::Private
+namespace UE::AvaComponentVisualizers::Private
 {
 	struct FAvaViewportClientAndVisualizer
 	{
@@ -81,7 +81,7 @@ UAvaComponentVisualizersEdMode::UAvaComponentVisualizersEdMode()
 		false);
 }
 
-namespace UE::AvalancheComponentVisualizers::Private
+namespace UE::AvaComponentVisualizers::Private
 {
 	const TArray<UE::Widget::EWidgetMode> WidgetModes = {
 			UE::Widget::WM_Translate,
@@ -94,7 +94,7 @@ void UAvaComponentVisualizersEdMode::ModeTick(float DeltaTime)
 {
 	Super::ModeTick(DeltaTime);
 
-	UE::AvalancheComponentVisualizers::Private::FAvaViewportClientAndVisualizer VisualizerAndClient(this);
+	UE::AvaComponentVisualizers::Private::FAvaViewportClientAndVisualizer VisualizerAndClient(this);
 
 	if (VisualizerAndClient.ActiveVisualizer.IsValid())
 	{
@@ -116,7 +116,7 @@ void UAvaComponentVisualizersEdMode::ModeTick(float DeltaTime)
 			}
 		}
 			
-		for (UE::Widget::EWidgetMode WidgetMode : UE::AvalancheComponentVisualizers::Private::WidgetModes)
+		for (UE::Widget::EWidgetMode WidgetMode : UE::AvaComponentVisualizers::Private::WidgetModes)
 		{
 			if (WidgetMode == CurrentModeStart)
 			{
@@ -146,7 +146,7 @@ void UAvaComponentVisualizersEdMode::Exit()
 
 EAxisList::Type UAvaComponentVisualizersEdMode::GetWidgetAxisToDraw(UE::Widget::EWidgetMode InWidgetMode) const
 {
-	UE::AvalancheComponentVisualizers::Private::FAvaViewportClientAndVisualizer VCnV(this);
+	UE::AvaComponentVisualizers::Private::FAvaViewportClientAndVisualizer VCnV(this);
 
 	EAxisList::Type AxisList = EAxisList::XYZ;
 
@@ -160,7 +160,7 @@ EAxisList::Type UAvaComponentVisualizersEdMode::GetWidgetAxisToDraw(UE::Widget::
 
 FVector UAvaComponentVisualizersEdMode::GetWidgetLocation() const
 {
-	UE::AvalancheComponentVisualizers::Private::FAvaViewportClientAndVisualizer VCnV(this);
+	UE::AvaComponentVisualizers::Private::FAvaViewportClientAndVisualizer VCnV(this);
 
 	FVector WidgetLocation = FVector::ZeroVector;
 
@@ -174,7 +174,7 @@ FVector UAvaComponentVisualizersEdMode::GetWidgetLocation() const
 
 bool UAvaComponentVisualizersEdMode::UsesTransformWidget(UE::Widget::EWidgetMode CheckMode) const
 {
-	UE::AvalancheComponentVisualizers::Private::FAvaViewportClientAndVisualizer VCnV(this);
+	UE::AvaComponentVisualizers::Private::FAvaViewportClientAndVisualizer VCnV(this);
 
 	if (VCnV.ActiveVisualizer)
 	{
@@ -192,7 +192,7 @@ bool UAvaComponentVisualizersEdMode::UsesTransformWidget(UE::Widget::EWidgetMode
 
 FVector UAvaComponentVisualizersEdMode::GetWidgetNormalFromCurrentAxis(void* InData)
 {
-	UE::AvalancheComponentVisualizers::Private::FAvaViewportClientAndVisualizer VCnV(this);
+	UE::AvaComponentVisualizers::Private::FAvaViewportClientAndVisualizer VCnV(this);
 
 	FVector WidgetNormal = FVector::ZeroVector;
 	UE::Widget::EWidgetMode WidgetMode;
@@ -246,7 +246,7 @@ void UAvaComponentVisualizersEdMode::SetCurrentWidgetAxis(EAxisList::Type InAxis
 
 EAxisList::Type UAvaComponentVisualizersEdMode::GetCurrentWidgetAxis() const
 {
-	UE::AvalancheComponentVisualizers::Private::FAvaViewportClientAndVisualizer VCnV(this);
+	UE::AvaComponentVisualizers::Private::FAvaViewportClientAndVisualizer VCnV(this);
 
 	EAxisList::Type AxisList = EAxisList::XYZ;
 
@@ -280,14 +280,14 @@ bool UAvaComponentVisualizersEdMode::UsesPropertyWidgets() const
 
 bool UAvaComponentVisualizersEdMode::GetCustomDrawingCoordinateSystem(FMatrix& OutMatrix, void* InData)
 {
-	UE::AvalancheComponentVisualizers::Private::FAvaViewportClientAndVisualizer VCnV(this);
+	UE::AvaComponentVisualizers::Private::FAvaViewportClientAndVisualizer VCnV(this);
 
 	return VCnV.ActiveVisualizer && VCnV.ActiveVisualizer->GetCustomInputCoordinateSystem(VCnV.EditorViewportClient, OutMatrix);
 }
 
 bool UAvaComponentVisualizersEdMode::GetCustomInputCoordinateSystem(FMatrix& OutMatrix, void* InData)
 {
-	UE::AvalancheComponentVisualizers::Private::FAvaViewportClientAndVisualizer VCnV(this);
+	UE::AvaComponentVisualizers::Private::FAvaViewportClientAndVisualizer VCnV(this);
 
 	return VCnV.ActiveVisualizer && VCnV.ActiveVisualizer->GetCustomInputCoordinateSystem(VCnV.EditorViewportClient, OutMatrix);
 }
diff --git a/Engine/Plugins/Experimental/Avalanche/Source/AvalancheComponentVisualizers/Private/AvaComponentVisualizersViewportOverlay.cpp b/Engine/Plugins/Experimental/Avalanche/Source/AvalancheComponentVisualizers/Private/AvaComponentVisualizersViewportOverlay.cpp
--- a/Engine/Plugins/Experimental/Avalanche/Source/AvalancheComponentVisualizers/Private/AvaComponentVisualizersViewportOverlay.cpp
+++ b/Engine/Plugins/Experimental/Avalanche/Source/AvalancheComponentVisualizers/Private/AvaComponentVisualizersViewportOverlay.cpp
@@ -22,7 +22,7 @@
 
 #define LOCTEXT_NAMESPACE "AvaComponentVisualizersViewportOverlay"
 
-namespace UE::AvalancheComponentVisualizers::Private
+namespace UE::AvaComponentVisualizers::Private
 {
 	TSharedRef<SWidget> CreateDraggableDetailsWidget(const FSlateBrush* InActorBrush, const FText& InActorName, const TSharedRef<SWidget>& InDetailsView)
 	{
@@ -85,7 +85,7 @@ namespace UE::AvalancheComponentVisualizers::Private
 			];
 	}
 
-	TSharedPtr<SWidget> CreateFullAvalancheDetailsWidget(const TArray<UObject*>& InSelectedObjects)
+	TSharedPtr<SWidget> CreateFullDetailsWidget(const TArray<UObject*>& InSelectedObjects)
 	{
 		FCustomDetailsViewArgs CustomDetailsViewArgs;
 		CustomDetailsViewArgs.bShowCategories = false;
@@ -320,7 +320,7 @@ void FAvaComponentVisualizersViewportOverlay::AddWidget(const TArray<TSharedPtr<
 		return;
 	}
 
-	using namespace UE::AvalancheComponentVisualizers::Private;
+	using namespace UE::AvaComponentVisualizers::Private;
 
 	switch (ViewportSettings->ShapeEditorOverlayType)
 	{
@@ -331,8 +331,8 @@ void FAvaComponentVisualizersViewportOverlay::AddWidget(const TArray<TSharedPtr<
 			OverlayWidget = CreateComponentVisualizerWidget(InSelectedObjects);
 			break;
 
-		case EAvaShapeEditorOverlayType::FullAvalancheDetails:
-			OverlayWidget = CreateFullAvalancheDetailsWidget(InSelectedObjects);
+		case EAvaShapeEditorOverlayType::FullDetails:
+			OverlayWidget = CreateFullDetailsWidget(InSelectedObjects);
 			break;
 	}
 
diff --git a/Engine/Plugins/Experimental/Avalanche/Source/AvalancheComponentVisualizers/Private/AvalancheComponentVisualizersModule.cpp b/Engine/Plugins/Experimental/Avalanche/Source/AvalancheComponentVisualizers/Private/AvalancheComponentVisualizersModule.cpp
--- a/Engine/Plugins/Experimental/Avalanche/Source/AvalancheComponentVisualizers/Private/AvalancheComponentVisualizersModule.cpp
+++ b/Engine/Plugins/Experimental/Avalanche/Source/AvalancheComponentVisualizers/Private/AvalancheComponentVisualizersModule.cpp
@@ -6,7 +6,7 @@
 #include "ComponentVisualizers.h"
 #include "Modules/ModuleManager.h"
 
-namespace UE::AvalancheComponentVisualizers::Private
+namespace UE::AvaComponentVisualizers::Private
 {
 	TSet<TWeakPtr<FComponentVisualizer>> AvalancheVisualizers;
 	static FAvaComponentVisualizersViewportOverlay ViewportOverlay;
@@ -22,17 +22,17 @@ void FAvalancheComponentVisualizersModule::RegisterComponentVisualizer(FName InC
 {
 	FComponentVisualizersModule& VisualisersModule = FModuleManager::LoadModuleChecked<FComponentVisualizersModule>("ComponentVisualizers");
 	VisualisersModule.RegisterComponentVisualizer(InComponentClassName, InVisualizer);
-	UE::AvalancheComponentVisualizers::Private::AvalancheVisualizers.Add(InVisualizer);
+	UE::AvaComponentVisualizers::Private::AvalancheVisualizers.Add(InVisualizer);
 }
 
 bool FAvalancheComponentVisualizersModule::IsAvalancheVisualizer(TSharedRef<FComponentVisualizer> InVisualizer) const
 {
-	return UE::AvalancheComponentVisualizers::Private::AvalancheVisualizers.Contains(InVisualizer);
+	return UE::AvaComponentVisualizers::Private::AvalancheVisualizers.Contains(InVisualizer);
 }
 
 IAvaComponentVisualizersViewportOverlay& FAvalancheComponentVisualizersModule::GetViewportOverlay() const
 {
-	return UE::AvalancheComponentVisualizers::Private::ViewportOverlay;
+	return UE::AvaComponentVisualizers::Private::ViewportOverlay;
 }
 
 IMPLEMENT_MODULE(FAvalancheComponentVisualizersModule, AvalancheComponentVisualizers)
diff --git a/Engine/Plugins/Experimental/Avalanche/Source/AvalancheComponentVisualizers/Public/IAvalancheComponentVisualizersModule.h b/Engine/Plugins/Experimental/Avalanche/Source/AvalancheComponentVisualizers/Public/IAvalancheComponentVisualizersModule.h
--- a/Engine/Plugins/Experimental/Avalanche/Source/AvalancheComponentVisualizers/Public/IAvalancheComponentVisualizersModule.h
+++ b/Engine/Plugins/Experimental/Avalanche/Source/AvalancheComponentVisualizers/Public/IAvalancheComponentVisualizersModule.h
@@ -11,7 +11,7 @@
 
 class IAvaComponentVisualizersViewportOverlay;
 
-namespace UE::AvalancheComponentVisualizers
+namespace UE::AvaComponentVisualizers
 {
 	const FName ModuleName = TEXT("AvalancheComponentVisualizers");
 }
@@ -21,12 +21,12 @@ class IAvalancheComponentVisualizersModule : public IModuleInterface
 public:
 	static IAvalancheComponentVisualizersModule& Get()
 	{
-		return FModuleManager::LoadModuleChecked<IAvalancheComponentVisualizersModule>(UE::AvalancheComponentVisualizers::ModuleName);
+		return FModuleManager::LoadModuleChecked<IAvalancheComponentVisualizersModule>(UE::AvaComponentVisualizers::ModuleName);
 	}
 
 	static IAvalancheComponentVisualizersModule* GetIfLoaded()
 	{
-		return FModuleManager::LoadModulePtr<IAvalancheComponentVisualizersModule>(UE::AvalancheComponentVisualizers::ModuleName);
+		return FModuleManager::LoadModulePtr<IAvalancheComponentVisualizersModule>(UE::AvaComponentVisualizers::ModuleName);
 	}
 
 	/** Could also use a TSet or custom storage class. */
diff --git a/Engine/Plugins/Experimental/Avalanche/Source/AvalancheEditor/Internal/Tests/Framework/AvaEditorTestUtils.h b/Engine/Plugins/Experimental/Avalanche/Source/AvalancheEditor/Internal/Tests/Framework/AvaEditorTestUtils.h
--- a/Engine/Plugins/Experimental/Avalanche/Source/AvalancheEditor/Internal/Tests/Framework/AvaEditorTestUtils.h
+++ b/Engine/Plugins/Experimental/Avalanche/Source/AvalancheEditor/Internal/Tests/Framework/AvaEditorTestUtils.h
@@ -6,8 +6,7 @@
 #include "Math/MathFwd.h"
 #include "Tests/Framework/AvaTestUtils.h"
 
-
-// General Avalanche Editor unit test log
+// General Motion Design Editor unit test log
 DECLARE_LOG_CATEGORY_EXTERN(LogAvaEditorTest, Log, All);
 
 class AActor;
diff --git a/Engine/Plugins/Experimental/Avalanche/Source/AvalancheEditor/Private/Alignment/SAvaAnchorAlignment.h b/Engine/Plugins/Experimental/Avalanche/Source/AvalancheEditor/Private/Alignment/SAvaAnchorAlignment.h
--- a/Engine/Plugins/Experimental/Avalanche/Source/AvalancheEditor/Private/Alignment/SAvaAnchorAlignment.h
+++ b/Engine/Plugins/Experimental/Avalanche/Source/AvalancheEditor/Private/Alignment/SAvaAnchorAlignment.h
@@ -9,7 +9,7 @@
 class SUniformGridPanel;
 
 /**
- * Avalanche Anchor Alignment
+ * Motion Design Anchor Alignment
  * 
  * Widget that holds holds axis alignment widgets for Left/Center/Right, Top/Center/Bottom, and Front/Center/Back.
  * One alignment button set per axis row (Horizontal, Vertical, Depth).
diff --git a/Engine/Plugins/Experimental/Avalanche/Source/AvalancheEditor/Private/Alignment/SAvaDepthAlignment.cpp b/Engine/Plugins/Experimental/Avalanche/Source/AvalancheEditor/Private/Alignment/SAvaDepthAlignment.cpp
--- a/Engine/Plugins/Experimental/Avalanche/Source/AvalancheEditor/Private/Alignment/SAvaDepthAlignment.cpp
+++ b/Engine/Plugins/Experimental/Avalanche/Source/AvalancheEditor/Private/Alignment/SAvaDepthAlignment.cpp
@@ -24,13 +24,13 @@ void SAvaDepthAlignment::Construct(const FArguments& InArgs)
 		.Value(this, &SAvaDepthAlignment::GetCurrentAlignment)
 		.OnValueChanged(this, &SAvaDepthAlignment::OnCurrentAlignmentChanged)
 		+ SSegmentedControl<EAvaDepthAlignment>::Slot(EAvaDepthAlignment::Front)
-			.Icon(FAvaEditorStyle::Get().GetBrush("AvalancheIcons.Alignment.Translation.Front"))
+			.Icon(FAvaEditorStyle::Get().GetBrush("Icons.Alignment.Translation.Front"))
 			.ToolTip(LOCTEXT("DAlignFront", "Front Align Depth"))
 		+ SSegmentedControl<EAvaDepthAlignment>::Slot(EAvaDepthAlignment::Center)
-			.Icon(FAvaEditorStyle::Get().GetBrush("AvalancheIcons.Alignment.Translation.Center_X"))
+			.Icon(FAvaEditorStyle::Get().GetBrush("Icons.Alignment.Translation.Center_X"))
 			.ToolTip(LOCTEXT("DAlignCenter", "Center Align Depth"))
 		+ SSegmentedControl<EAvaDepthAlignment>::Slot(EAvaDepthAlignment::Back)
-			.Icon(FAvaEditorStyle::Get().GetBrush("AvalancheIcons.Alignment.Translation.Back"))
+			.Icon(FAvaEditorStyle::Get().GetBrush("Icons.Alignment.Translation.Back"))
 			.ToolTip(LOCTEXT("DAlignBack", "Back Align Depth"))
 	];
 }
diff --git a/Engine/Plugins/Experimental/Avalanche/Source/AvalancheEditor/Private/Alignment/SAvaDepthAlignment.h b/Engine/Plugins/Experimental/Avalanche/Source/AvalancheEditor/Private/Alignment/SAvaDepthAlignment.h
--- a/Engine/Plugins/Experimental/Avalanche/Source/AvalancheEditor/Private/Alignment/SAvaDepthAlignment.h
+++ b/Engine/Plugins/Experimental/Avalanche/Source/AvalancheEditor/Private/Alignment/SAvaDepthAlignment.h
@@ -8,7 +8,7 @@
 enum class EAvaDepthAlignment : uint8;
 
 /**
- * Avalanche Depth Axis Alignment
+ * Motion Design Depth Axis Alignment
  * 
  * Widget that holds three buttons for either Front, Center, and Back.
  */
diff --git a/Engine/Plugins/Experimental/Avalanche/Source/AvalancheEditor/Private/Alignment/SAvaHorizontalAlignment.cpp b/Engine/Plugins/Experimental/Avalanche/Source/AvalancheEditor/Private/Alignment/SAvaHorizontalAlignment.cpp
--- a/Engine/Plugins/Experimental/Avalanche/Source/AvalancheEditor/Private/Alignment/SAvaHorizontalAlignment.cpp
+++ b/Engine/Plugins/Experimental/Avalanche/Source/AvalancheEditor/Private/Alignment/SAvaHorizontalAlignment.cpp
@@ -24,13 +24,13 @@ void SAvaHorizontalAlignment::Construct(const FArguments& InArgs)
 		.Value(this, &SAvaHorizontalAlignment::GetCurrentAlignment)
 		.OnValueChanged(this, &SAvaHorizontalAlignment::OnCurrentAlignmentChanged)
 		+ SSegmentedControl<EAvaHorizontalAlignment>::Slot(EAvaHorizontalAlignment::Left)
-			.Icon(FAvaEditorStyle::Get().GetBrush("AvalancheIcons.Alignment.Left"))
+			.Icon(FAvaEditorStyle::Get().GetBrush("Icons.Alignment.Left"))
 			.ToolTip(LOCTEXT("HAlignLeft", "Left Align Horizontally"))
 		+ SSegmentedControl<EAvaHorizontalAlignment>::Slot(EAvaHorizontalAlignment::Center)
-			.Icon(FAvaEditorStyle::Get().GetBrush("AvalancheIcons.Alignment.Center_Y"))
+			.Icon(FAvaEditorStyle::Get().GetBrush("Icons.Alignment.Center_Y"))
 			.ToolTip(LOCTEXT("HAlignCenter", "Center Align Horizontally"))
 		+ SSegmentedControl<EAvaHorizontalAlignment>::Slot(EAvaHorizontalAlignment::Right)
-			.Icon(FAvaEditorStyle::Get().GetBrush("AvalancheIcons.Alignment.Right"))
+			.Icon(FAvaEditorStyle::Get().GetBrush("Icons.Alignment.Right"))
 			.ToolTip(LOCTEXT("HAlignRight", "Right Align Horizontally"))
 	];
 }
diff --git a/Engine/Plugins/Experimental/Avalanche/Source/AvalancheEditor/Private/Alignment/SAvaVerticalAlignment.cpp b/Engine/Plugins/Experimental/Avalanche/Source/AvalancheEditor/Private/Alignment/SAvaVerticalAlignment.cpp
--- a/Engine/Plugins/Experimental/Avalanche/Source/AvalancheEditor/Private/Alignment/SAvaVerticalAlignment.cpp
+++ b/Engine/Plugins/Experimental/Avalanche/Source/AvalancheEditor/Private/Alignment/SAvaVerticalAlignment.cpp
@@ -24,13 +24,13 @@ void SAvaVerticalAlignment::Construct(const FArguments& InArgs)
 		.Value(this, &SAvaVerticalAlignment::GetCurrentAlignment)
 		.OnValueChanged(this, &SAvaVerticalAlignment::OnCurrentAlignmentChanged)
 		+ SSegmentedControl<EAvaVerticalAlignment>::Slot(EAvaVerticalAlignment::Top)
-			.Icon(FAvaEditorStyle::Get().GetBrush("AvalancheIcons.Alignment.Top"))
+			.Icon(FAvaEditorStyle::Get().GetBrush("Icons.Alignment.Top"))
 			.ToolTip(LOCTEXT("VAlignTop", "Top Align Vertically"))
 		+ SSegmentedControl<EAvaVerticalAlignment>::Slot(EAvaVerticalAlignment::Center)
-			.Icon(FAvaEditorStyle::Get().GetBrush("AvalancheIcons.Alignment.Center_Z"))
+			.Icon(FAvaEditorStyle::Get().GetBrush("Icons.Alignment.Center_Z"))
 			.ToolTip(LOCTEXT("VAlignCenter", "Center Align Vertically"))
 		+ SSegmentedControl<EAvaVerticalAlignment>::Slot(EAvaVerticalAlignment::Bottom)
-			.Icon(FAvaEditorStyle::Get().GetBrush("AvalancheIcons.Alignment.Bottom"))
+			.Icon(FAvaEditorStyle::Get().GetBrush("Icons.Alignment.Bottom"))
 			.ToolTip(LOCTEXT("VAlignBottom", "Bottom Align Vertically"))
 	];
 }
diff --git a/Engine/Plugins/Experimental/Avalanche/Source/AvalancheEditor/Private/Alignment/SAvaVerticalAlignment.h b/Engine/Plugins/Experimental/Avalanche/Source/AvalancheEditor/Private/Alignment/SAvaVerticalAlignment.h
--- a/Engine/Plugins/Experimental/Avalanche/Source/AvalancheEditor/Private/Alignment/SAvaVerticalAlignment.h
+++ b/Engine/Plugins/Experimental/Avalanche/Source/AvalancheEditor/Private/Alignment/SAvaVerticalAlignment.h
@@ -8,7 +8,7 @@
 enum class EAvaVerticalAlignment : uint8;
 
 /**
- * Avalanche Vertical Axis Alignment
+ * Motion Design Vertical Axis Alignment
  * 
  * Widget that holds three buttons for either Top, Center, and Bottom.
  */
diff --git a/Engine/Plugins/Experimental/Avalanche/Source/AvalancheEditor/Private/AvaEditorActorUtils.cpp b/Engine/Plugins/Experimental/Avalanche/Source/AvalancheEditor/Private/AvaEditorActorUtils.cpp
--- a/Engine/Plugins/Experimental/Avalanche/Source/AvalancheEditor/Private/AvaEditorActorUtils.cpp
+++ b/Engine/Plugins/Experimental/Avalanche/Source/AvalancheEditor/Private/AvaEditorActorUtils.cpp
@@ -99,7 +99,7 @@ void FAvaEditorActorUtils::GetAllAttachedActors(const USceneComponent* InParent,
 	}
 }
 
-namespace UE::AvalancheEditor::Private
+namespace UE::AvaEditor::Private
 {
 	bool IsComponentChildOf(const USceneComponent* InSceneComponent, const USceneComponent* InPossibleParent)
 	{
@@ -152,7 +152,7 @@ void FAvaEditorActorUtils::GetAllAttachedActors(const TSet<USceneComponent*>& In
 				continue;
 			}
 
-			if (UE::AvalancheEditor::Private::IsComponentChildOf(Component, ComponentInner))
+			if (UE::AvaEditor::Private::IsComponentChildOf(Component, ComponentInner))
 			{
 				bIsChild = true;
 				break;
diff --git a/Engine/Plugins/Experimental/Avalanche/Source/AvalancheEditor/Private/AvaEditorCommands.cpp b/Engine/Plugins/Experimental/Avalanche/Source/AvalancheEditor/Private/AvaEditorCommands.cpp
--- a/Engine/Plugins/Experimental/Avalanche/Source/AvalancheEditor/Private/AvaEditorCommands.cpp
+++ b/Engine/Plugins/Experimental/Avalanche/Source/AvalancheEditor/Private/AvaEditorCommands.cpp
@@ -34,7 +34,7 @@ void FAvaEditorCommands::RegisterViewportCommands()
 		, EUserInterfaceActionType::ToggleButton
 		, FInputChord(EKeys::C));
 
-	UI_COMMAND(SetAvalancheViewportType
+	UI_COMMAND(SetMotionDesignViewportType
 		, "Motion Design Viewport"
 		, "Viewport layout to view a fixed camera in an Motion Design Scene"
 		, EUserInterfaceActionType::RadioButton
diff --git a/Engine/Plugins/Experimental/Avalanche/Source/AvalancheEditor/Private/AvaEditorCommands.h b/Engine/Plugins/Experimental/Avalanche/Source/AvalancheEditor/Private/AvaEditorCommands.h
--- a/Engine/Plugins/Experimental/Avalanche/Source/AvalancheEditor/Private/AvaEditorCommands.h
+++ b/Engine/Plugins/Experimental/Avalanche/Source/AvalancheEditor/Private/AvaEditorCommands.h
@@ -19,7 +19,7 @@ public:
 
 	/** Switches between Camera "2D" viewport and 3D viewport. */
 	TSharedPtr<FUICommandInfo> SwitchViewports;
-	TSharedPtr<FUICommandInfo> SetAvalancheViewportType;
+	TSharedPtr<FUICommandInfo> SetMotionDesignViewportType;
 
 	/** Groups the Selected Actors via a Null Actor. */
 	TSharedPtr<FUICommandInfo> GroupActors;
diff --git a/Engine/Plugins/Experimental/Avalanche/Source/AvalancheEditor/Private/AvaEditorModule.cpp b/Engine/Plugins/Experimental/Avalanche/Source/AvalancheEditor/Private/AvaEditorModule.cpp
--- a/Engine/Plugins/Experimental/Avalanche/Source/AvalancheEditor/Private/AvaEditorModule.cpp
+++ b/Engine/Plugins/Experimental/Avalanche/Source/AvalancheEditor/Private/AvaEditorModule.cpp
@@ -74,7 +74,7 @@ struct FAvaViewportColorPickerLightAdapter : IAvaViewportColorPickerAdapter
 	//~ End IAvaViewportColorPickerActorAdapter
 };
 
-namespace UE::AvalancheEditor::Private
+namespace UE::AvaEditor::Private
 {
 	static FString LevelTemplatesPath = FString::Printf(TEXT("/%hs/%s"), UE_PLUGIN_NAME, TEXT("LevelTemplates"));
 	static FString DefaultLevelPath = FString::Printf(TEXT("/%hs/%s"), UE_PLUGIN_NAME, TEXT("DefaultMotionDesignLevel"));
@@ -93,7 +93,7 @@ void FAvaEditorModule::StartupModule()
 	RegisterCustomLayouts();
 
 	// Register Palettes
-	using namespace UE::AvalancheEditor::Private;
+	using namespace UE::AvaEditor::Private;
 
 	// Register Icon Customization
 	IAvaOutlinerModule::Get().RegisterOverriddenIcon<FAvaOutlinerActor, FAvaOutlinerObjectIconCustomization>(AAvaShapeActor::StaticClass())
@@ -304,22 +304,22 @@ void FAvaEditorModule::RegisterLevelTemplates()
 	{
 		// Register original template map,
 		// @todo: Deprecate post-beta 0.1.4
-		if (!GUnrealEd->IsTemplateMap(UE::AvalancheEditor::Private::DefaultLevelPath))
+		if (!GUnrealEd->IsTemplateMap(UE::AvaEditor::Private::DefaultLevelPath))
 		{
-			FTemplateMapInfo AvalancheMapTemplate;
-			AvalancheMapTemplate.Category = TEXT("Motion Design");
-			AvalancheMapTemplate.DisplayName = LOCTEXT("Map Template", "Motion Design");
-			AvalancheMapTemplate.ThumbnailTexture = UE::AvalancheEditor::Private::DefaultLevelThumbnailPath;			
-			AvalancheMapTemplate.Map = UE::AvalancheEditor::Private::DefaultLevelPath;
+			FTemplateMapInfo MapTemplate;
+			MapTemplate.Category = TEXT("Motion Design");
+			MapTemplate.DisplayName = LOCTEXT("Map Template", "Motion Design");
+			MapTemplate.ThumbnailTexture = UE::AvaEditor::Private::DefaultLevelThumbnailPath;
+			MapTemplate.Map = UE::AvaEditor::Private::DefaultLevelPath;
 				
-			GUnrealEd->AppendTemplateMaps({ AvalancheMapTemplate });	
+			GUnrealEd->AppendTemplateMaps({ MapTemplate });	
 		}
 
 		IAssetRegistry* AssetRegistry = FModuleManager::LoadModuleChecked<FAssetRegistryModule>(TEXT("AssetRegistry")).TryGet();
 		if (AssetRegistry)
 		{
 			TArray<FString> LevelTemplatePaths;
-			AssetRegistry->GetSubPaths(UE::AvalancheEditor::Private::LevelTemplatesPath, LevelTemplatePaths, false);
+			AssetRegistry->GetSubPaths(UE::AvaEditor::Private::LevelTemplatesPath, LevelTemplatePaths, false);
 
 			TArray<FTemplateMapInfo> TemplateMaps;
 			
diff --git a/Engine/Plugins/Experimental/Avalanche/Source/AvalancheEditor/Private/AvaEditorModule.h b/Engine/Plugins/Experimental/Avalanche/Source/AvalancheEditor/Private/AvaEditorModule.h
--- a/Engine/Plugins/Experimental/Avalanche/Source/AvalancheEditor/Private/AvaEditorModule.h
+++ b/Engine/Plugins/Experimental/Avalanche/Source/AvalancheEditor/Private/AvaEditorModule.h
@@ -14,7 +14,7 @@ class IAvaEditor;
 DECLARE_LOG_CATEGORY_EXTERN(AvaLog, Log, All);
 
 /**
- * Main Avalanche Editor Module
+ * Main Motion Design Editor Module
  */
 class FAvaEditorModule : public IModuleInterface
 {
@@ -27,7 +27,7 @@ public:
 	static FSlateIcon GetOutlinerShapeActorIcon(TSharedPtr<const FAvaOutlinerItem> InItem);
 	
 private:
-	/** Avalanche for the Level Editor*/
+	/** Motion Design Level Editor */
 	TSharedPtr<IAvaEditor> AvaLevelEditor;
 
 	void CreateAvaLevelEditor();
diff --git a/Engine/Plugins/Experimental/Avalanche/Source/AvalancheEditor/Private/AvaEditorStyle.cpp b/Engine/Plugins/Experimental/Avalanche/Source/AvalancheEditor/Private/AvaEditorStyle.cpp
--- a/Engine/Plugins/Experimental/Avalanche/Source/AvalancheEditor/Private/AvaEditorStyle.cpp
+++ b/Engine/Plugins/Experimental/Avalanche/Source/AvalancheEditor/Private/AvaEditorStyle.cpp
@@ -51,108 +51,108 @@ FAvaEditorStyle::FAvaEditorStyle()
 	Set("ClassIcon.AvaEffectorActor", new IMAGE_BRUSH_SVG("Icons/ToolboxIcons/effector", Icon16));
 
 	// Toolbox Icons
-	Set("AvalancheIcons.Chevron",          new IMAGE_BRUSH("Icons/ToolboxIcons/chevron",          Icon16));
-	Set("AvalancheIcons.RegularPolygon",   new IMAGE_BRUSH("Icons/ToolboxIcons/regularpolygon",   Icon16));
-	Set("AvalancheIcons.IrregularPolygon", new IMAGE_BRUSH("Icons/ToolboxIcons/irregularpolygon", Icon16));
-	Set("AvalancheIcons.Canvas",           new IMAGE_BRUSH("Icons/ToolboxIcons/Canvas",           Icon16));
-	Set("AvalancheIcons.LayoutGrid",       new IMAGE_BRUSH("Icons/ToolboxIcons/LayoutGrid",       Icon16));
-	Set("AvalancheIcons.Ellipse",          new IMAGE_BRUSH("Icons/ToolboxIcons/circle",           Icon16));
-	Set("AvalancheIcons.Ring",             new IMAGE_BRUSH("Icons/ToolboxIcons/ring",             Icon16));
-	Set("AvalancheIcons.Arrow",            new IMAGE_BRUSH("Icons/ToolboxIcons/arrow",            Icon16));
-	Set("AvalancheIcons.2DArrow",          new IMAGE_BRUSH("Icons/ToolboxIcons/arrow",            Icon16));
-	Set("AvalancheIcons.Sphere",           new IMAGE_BRUSH("Icons/ToolboxIcons/sphere",           Icon16));
-	Set("AvalancheIcons.Pyramid",          new IMAGE_BRUSH("Icons/ToolboxIcons/pyramid",          Icon16));
-	Set("AvalancheIcons.Cone",             new IMAGE_BRUSH("Icons/ToolboxIcons/cone",             Icon16));
-	Set("AvalancheIcons.Torus",            new IMAGE_BRUSH("Icons/ToolboxIcons/torus",            Icon16));
-	Set("AvalancheIcons.Cube",             new IMAGE_BRUSH("Icons/ToolboxIcons/cube",             Icon16));
-	Set("AvalancheIcons.Rectangle",        new IMAGE_BRUSH("Icons/ToolboxIcons/rectangle",        Icon16));
-	Set("AvalancheIcons.Plane",            new IMAGE_BRUSH("Icons/ToolboxIcons/plane",            Icon16));
-	Set("AvalancheIcons.Cylinder",         new IMAGE_BRUSH("Icons/ToolboxIcons/cylinder",         Icon16));
-	Set("AvalancheIcons.Line",             new IMAGE_BRUSH_SVG("Icons/ToolboxIcons/minus",             Icon16));
-	Set("AvalancheIcons.Star",             new IMAGE_BRUSH_SVG("Icons/ToolboxIcons/Favorite",          Icon16));
-	Set("AvalancheIcons.3DText",           new IMAGE_BRUSH_SVG("Icons/ToolboxIcons/Toolbar_Text",      Icon16));
-	Set("AvalancheIcons.RoundedBox",       new IMAGE_BRUSH_SVG("Icons/ToolboxIcons/Toolbar_Primitive", Icon16));
+	Set("Icons.Chevron",          new IMAGE_BRUSH("Icons/ToolboxIcons/chevron",          Icon16));
+	Set("Icons.RegularPolygon",   new IMAGE_BRUSH("Icons/ToolboxIcons/regularpolygon",   Icon16));
+	Set("Icons.IrregularPolygon", new IMAGE_BRUSH("Icons/ToolboxIcons/irregularpolygon", Icon16));
+	Set("Icons.Canvas",           new IMAGE_BRUSH("Icons/ToolboxIcons/Canvas",           Icon16));
+	Set("Icons.LayoutGrid",       new IMAGE_BRUSH("Icons/ToolboxIcons/LayoutGrid",       Icon16));
+	Set("Icons.Ellipse",          new IMAGE_BRUSH("Icons/ToolboxIcons/circle",           Icon16));
+	Set("Icons.Ring",             new IMAGE_BRUSH("Icons/ToolboxIcons/ring",             Icon16));
+	Set("Icons.Arrow",            new IMAGE_BRUSH("Icons/ToolboxIcons/arrow",            Icon16));
+	Set("Icons.2DArrow",          new IMAGE_BRUSH("Icons/ToolboxIcons/arrow",            Icon16));
+	Set("Icons.Sphere",           new IMAGE_BRUSH("Icons/ToolboxIcons/sphere",           Icon16));
+	Set("Icons.Pyramid",          new IMAGE_BRUSH("Icons/ToolboxIcons/pyramid",          Icon16));
+	Set("Icons.Cone",             new IMAGE_BRUSH("Icons/ToolboxIcons/cone",             Icon16));
+	Set("Icons.Torus",            new IMAGE_BRUSH("Icons/ToolboxIcons/torus",            Icon16));
+	Set("Icons.Cube",             new IMAGE_BRUSH("Icons/ToolboxIcons/cube",             Icon16));
+	Set("Icons.Rectangle",        new IMAGE_BRUSH("Icons/ToolboxIcons/rectangle",        Icon16));
+	Set("Icons.Plane",            new IMAGE_BRUSH("Icons/ToolboxIcons/plane",            Icon16));
+	Set("Icons.Cylinder",         new IMAGE_BRUSH("Icons/ToolboxIcons/cylinder",         Icon16));
+	Set("Icons.Line",             new IMAGE_BRUSH_SVG("Icons/ToolboxIcons/minus",             Icon16));
+	Set("Icons.Star",             new IMAGE_BRUSH_SVG("Icons/ToolboxIcons/Favorite",          Icon16));
+	Set("Icons.3DText",           new IMAGE_BRUSH_SVG("Icons/ToolboxIcons/Toolbar_Text",      Icon16));
+	Set("Icons.RoundedBox",       new IMAGE_BRUSH_SVG("Icons/ToolboxIcons/Toolbar_Primitive", Icon16));
 	
-	Set("AvalancheIcons.Alignment.Translation.TopLeft",     new IMAGE_BRUSH("Icons/DetailsPanelIcons/TopLeft",     Icon16));
-	Set("AvalancheIcons.Alignment.Translation.Top",         new IMAGE_BRUSH("Icons/DetailsPanelIcons/Top",         Icon16));
-	Set("AvalancheIcons.Alignment.Translation.TopRight",    new IMAGE_BRUSH("Icons/DetailsPanelIcons/TopRight",    Icon16));
-	Set("AvalancheIcons.Alignment.Translation.Left",        new IMAGE_BRUSH("Icons/DetailsPanelIcons/Left",        Icon16));
-	Set("AvalancheIcons.Alignment.Translation.Center",      new IMAGE_BRUSH("Icons/DetailsPanelIcons/Center",      Icon16));
-	Set("AvalancheIcons.Alignment.Translation.Right",       new IMAGE_BRUSH("Icons/DetailsPanelIcons/Right",       Icon16));
-	Set("AvalancheIcons.Alignment.Translation.BottomLeft",  new IMAGE_BRUSH("Icons/DetailsPanelIcons/BottomLeft",  Icon16));
-	Set("AvalancheIcons.Alignment.Translation.Bottom",      new IMAGE_BRUSH("Icons/DetailsPanelIcons/Bottom",      Icon16));
-	Set("AvalancheIcons.Alignment.Translation.BottomRight", new IMAGE_BRUSH("Icons/DetailsPanelIcons/BottomRight", Icon16));
+	Set("Icons.Alignment.Translation.TopLeft",     new IMAGE_BRUSH("Icons/DetailsPanelIcons/TopLeft",     Icon16));
+	Set("Icons.Alignment.Translation.Top",         new IMAGE_BRUSH("Icons/DetailsPanelIcons/Top",         Icon16));
+	Set("Icons.Alignment.Translation.TopRight",    new IMAGE_BRUSH("Icons/DetailsPanelIcons/TopRight",    Icon16));
+	Set("Icons.Alignment.Translation.Left",        new IMAGE_BRUSH("Icons/DetailsPanelIcons/Left",        Icon16));
+	Set("Icons.Alignment.Translation.Center",      new IMAGE_BRUSH("Icons/DetailsPanelIcons/Center",      Icon16));
+	Set("Icons.Alignment.Translation.Right",       new IMAGE_BRUSH("Icons/DetailsPanelIcons/Right",       Icon16));
+	Set("Icons.Alignment.Translation.BottomLeft",  new IMAGE_BRUSH("Icons/DetailsPanelIcons/BottomLeft",  Icon16));
+	Set("Icons.Alignment.Translation.Bottom",      new IMAGE_BRUSH("Icons/DetailsPanelIcons/Bottom",      Icon16));
+	Set("Icons.Alignment.Translation.BottomRight", new IMAGE_BRUSH("Icons/DetailsPanelIcons/BottomRight", Icon16));
 	
-	Set("AvalancheIcons.Alignment.Translation.Back",        new IMAGE_BRUSH_SVG("Icons/DetailsPanelIcons/StackBack",   Icon22));
-	Set("AvalancheIcons.Alignment.Translation.Center_X",    new IMAGE_BRUSH_SVG("Icons/DetailsPanelIcons/StackCenter", Icon22));
-	Set("AvalancheIcons.Alignment.Translation.Front",       new IMAGE_BRUSH_SVG("Icons/DetailsPanelIcons/StackFront",  Icon22));
-
-	Set("AvalancheIcons.Alignment.Center_YZ",   new IMAGE_BRUSH_SVG("Icons/PaletteIcons/AlignCenterVertAndHoriz", Icon22));
-	Set("AvalancheIcons.Alignment.Left",        new IMAGE_BRUSH_SVG("Icons/PaletteIcons/AlignLeft",               Icon22));
-	Set("AvalancheIcons.Alignment.Center_Y",    new IMAGE_BRUSH_SVG("Icons/PaletteIcons/AlignCenterHoriz",        Icon22));
-	Set("AvalancheIcons.Alignment.Right",       new IMAGE_BRUSH_SVG("Icons/PaletteIcons/AlignRight",              Icon22));
-	Set("AvalancheIcons.Alignment.Top",         new IMAGE_BRUSH_SVG("Icons/PaletteIcons/AlignTop",                Icon22));
-	Set("AvalancheIcons.Alignment.Center_Z",    new IMAGE_BRUSH_SVG("Icons/PaletteIcons/AlignCenterVert",         Icon22));
-	Set("AvalancheIcons.Alignment.Bottom",      new IMAGE_BRUSH_SVG("Icons/PaletteIcons/AlignBottom",             Icon22));
-	Set("AvalancheIcons.Alignment.DistributeX", new IMAGE_BRUSH_SVG("Icons/PaletteIcons/DistributeDepth",         Icon22));
-	Set("AvalancheIcons.Alignment.DistributeY", new IMAGE_BRUSH_SVG("Icons/PaletteIcons/DistributeHorizontal",    Icon22));
-	Set("AvalancheIcons.Alignment.DistributeZ", new IMAGE_BRUSH_SVG("Icons/PaletteIcons/DistributeVertical",      Icon22));
+	Set("Icons.Alignment.Translation.Back",        new IMAGE_BRUSH_SVG("Icons/DetailsPanelIcons/StackBack",   Icon22));
+	Set("Icons.Alignment.Translation.Center_X",    new IMAGE_BRUSH_SVG("Icons/DetailsPanelIcons/StackCenter", Icon22));
+	Set("Icons.Alignment.Translation.Front",       new IMAGE_BRUSH_SVG("Icons/DetailsPanelIcons/StackFront",  Icon22));
+
+	Set("Icons.Alignment.Center_YZ",   new IMAGE_BRUSH_SVG("Icons/PaletteIcons/AlignCenterVertAndHoriz", Icon22));
+	Set("Icons.Alignment.Left",        new IMAGE_BRUSH_SVG("Icons/PaletteIcons/AlignLeft",               Icon22));
+	Set("Icons.Alignment.Center_Y",    new IMAGE_BRUSH_SVG("Icons/PaletteIcons/AlignCenterHoriz",        Icon22));
+	Set("Icons.Alignment.Right",       new IMAGE_BRUSH_SVG("Icons/PaletteIcons/AlignRight",              Icon22));
+	Set("Icons.Alignment.Top",         new IMAGE_BRUSH_SVG("Icons/PaletteIcons/AlignTop",                Icon22));
+	Set("Icons.Alignment.Center_Z",    new IMAGE_BRUSH_SVG("Icons/PaletteIcons/AlignCenterVert",         Icon22));
+	Set("Icons.Alignment.Bottom",      new IMAGE_BRUSH_SVG("Icons/PaletteIcons/AlignBottom",             Icon22));
+	Set("Icons.Alignment.DistributeX", new IMAGE_BRUSH_SVG("Icons/PaletteIcons/DistributeDepth",         Icon22));
+	Set("Icons.Alignment.DistributeY", new IMAGE_BRUSH_SVG("Icons/PaletteIcons/DistributeHorizontal",    Icon22));
+	Set("Icons.Alignment.DistributeZ", new IMAGE_BRUSH_SVG("Icons/PaletteIcons/DistributeVertical",      Icon22));
 	
-	Set("AvalancheIcons.Alignment.Rotation.Actor.Roll",   new IMAGE_BRUSH_SVG("Icons/PaletteIcons/RotationAlignActorRoll",   Icon22));
-	Set("AvalancheIcons.Alignment.Rotation.Actor.Pitch",  new IMAGE_BRUSH_SVG("Icons/PaletteIcons/RotationAlignActorPitch",  Icon22));
-	Set("AvalancheIcons.Alignment.Rotation.Actor.Yaw",    new IMAGE_BRUSH_SVG("Icons/PaletteIcons/RotationAlignActorYaw",    Icon22));
-	Set("AvalancheIcons.Alignment.Rotation.Actor.All",    new IMAGE_BRUSH_SVG("Icons/PaletteIcons/RotationAlignActorAll",    Icon22));
-	Set("AvalancheIcons.Alignment.Rotation.Camera.Roll",  new IMAGE_BRUSH_SVG("Icons/PaletteIcons/RotationAlignCameraRoll",  Icon22));
-	Set("AvalancheIcons.Alignment.Rotation.Camera.Pitch", new IMAGE_BRUSH_SVG("Icons/PaletteIcons/RotationAlignCameraPitch", Icon22));
-	Set("AvalancheIcons.Alignment.Rotation.Camera.Yaw",   new IMAGE_BRUSH_SVG("Icons/PaletteIcons/RotationAlignCameraYaw",   Icon22));
-	Set("AvalancheIcons.Alignment.Rotation.Camera.All",   new IMAGE_BRUSH_SVG("Icons/PaletteIcons/RotationAlignCameraAll",   Icon22));
+	Set("Icons.Alignment.Rotation.Actor.Roll",   new IMAGE_BRUSH_SVG("Icons/PaletteIcons/RotationAlignActorRoll",   Icon22));
+	Set("Icons.Alignment.Rotation.Actor.Pitch",  new IMAGE_BRUSH_SVG("Icons/PaletteIcons/RotationAlignActorPitch",  Icon22));
+	Set("Icons.Alignment.Rotation.Actor.Yaw",    new IMAGE_BRUSH_SVG("Icons/PaletteIcons/RotationAlignActorYaw",    Icon22));
+	Set("Icons.Alignment.Rotation.Actor.All",    new IMAGE_BRUSH_SVG("Icons/PaletteIcons/RotationAlignActorAll",    Icon22));
+	Set("Icons.Alignment.Rotation.Camera.Roll",  new IMAGE_BRUSH_SVG("Icons/PaletteIcons/RotationAlignCameraRoll",  Icon22));
+	Set("Icons.Alignment.Rotation.Camera.Pitch", new IMAGE_BRUSH_SVG("Icons/PaletteIcons/RotationAlignCameraPitch", Icon22));
+	Set("Icons.Alignment.Rotation.Camera.Yaw",   new IMAGE_BRUSH_SVG("Icons/PaletteIcons/RotationAlignCameraYaw",   Icon22));
+	Set("Icons.Alignment.Rotation.Camera.All",   new IMAGE_BRUSH_SVG("Icons/PaletteIcons/RotationAlignCameraAll",   Icon22));
 
 	// Grid Icons
-	Set("AvalancheIcons.Grid.Toggle",           new IMAGE_BRUSH_SVG("Icons/PaletteIcons/GridToggle",              Icon22));
-	Set("AvalancheIcons.Grid.AlwaysShowToggle", new IMAGE_BRUSH_SVG("Icons/PaletteIcons/GridAlwaysVisibleToggle", Icon22));
+	Set("Icons.Grid.Toggle",           new IMAGE_BRUSH_SVG("Icons/PaletteIcons/GridToggle",              Icon22));
+	Set("Icons.Grid.AlwaysShowToggle", new IMAGE_BRUSH_SVG("Icons/PaletteIcons/GridAlwaysVisibleToggle", Icon22));
 
 	// Snap Icons
-	Set("AvalancheIcons.Snap.Toggle",       new IMAGE_BRUSH_SVG("Icons/PaletteIcons/SnapToggle",       Icon22));
-	Set("AvalancheIcons.Snap.ToggleGrid",   new IMAGE_BRUSH_SVG("Icons/PaletteIcons/SnapGridToggle",   Icon22));
-	Set("AvalancheIcons.Snap.ToggleActor",  new IMAGE_BRUSH_SVG("Icons/PaletteIcons/SnapActorToggle",  Icon22));
-	Set("AvalancheIcons.Snap.ToggleScreen", new IMAGE_BRUSH_SVG("Icons/PaletteIcons/SnapScreenToggle", Icon22));
+	Set("Icons.Snap.Toggle",       new IMAGE_BRUSH_SVG("Icons/PaletteIcons/SnapToggle",       Icon22));
+	Set("Icons.Snap.ToggleGrid",   new IMAGE_BRUSH_SVG("Icons/PaletteIcons/SnapGridToggle",   Icon22));
+	Set("Icons.Snap.ToggleActor",  new IMAGE_BRUSH_SVG("Icons/PaletteIcons/SnapActorToggle",  Icon22));
+	Set("Icons.Snap.ToggleScreen", new IMAGE_BRUSH_SVG("Icons/PaletteIcons/SnapScreenToggle", Icon22));
 
 	// Screen Icons
-	Set("AvalancheIcons.Screen.SizeToScreen",        new IMAGE_BRUSH_SVG("Icons/PaletteIcons/SizeToScreen",        Icon22));
-	Set("AvalancheIcons.Screen.SizeToScreenStretch", new IMAGE_BRUSH_SVG("Icons/PaletteIcons/SizeToScreenStretch", Icon22));
-	Set("AvalancheIcons.Screen.FitToScreen",         new IMAGE_BRUSH_SVG("Icons/PaletteIcons/FitToScreen", Icon22));
+	Set("Icons.Screen.SizeToScreen",        new IMAGE_BRUSH_SVG("Icons/PaletteIcons/SizeToScreen",        Icon22));
+	Set("Icons.Screen.SizeToScreenStretch", new IMAGE_BRUSH_SVG("Icons/PaletteIcons/SizeToScreenStretch", Icon22));
+	Set("Icons.Screen.FitToScreen",         new IMAGE_BRUSH_SVG("Icons/PaletteIcons/FitToScreen", Icon22));
 
 	// Logs Icons
-	Set("AvalancheIcons.Logs.Pin",   new IMAGE_BRUSH_SVG("Icons/LogsIcons/Pin",   Icon16));
-	Set("AvalancheIcons.Logs.Unpin", new IMAGE_BRUSH_SVG("Icons/LogsIcons/Unpin", Icon16));
+	Set("Icons.Logs.Pin",   new IMAGE_BRUSH_SVG("Icons/LogsIcons/Pin",   Icon16));
+	Set("Icons.Logs.Unpin", new IMAGE_BRUSH_SVG("Icons/LogsIcons/Unpin", Icon16));
 	
 	// Editor Icons
-	Set("AvalancheIcons.Editor.ButtonExpander",        new IMAGE_BRUSH("Icons/EditorIcons/ButtonExpander",          FVector2f(7.0f)));
-	Set("AvalancheIcons.Editor.ToolHightlight",        new IMAGE_BRUSH("Icons/EditorIcons/ToolIconHighlight",       Icon40));
-	Set("AvalancheIcons.Editor.GameView",              new IMAGE_BRUSH_SVG("Icons/EditorIcons/GameView",            Icon25));
-	Set("AvalancheIcons.Editor.Visualizers",           new IMAGE_BRUSH_SVG("Icons/EditorIcons/Visualizers",         Icon25));
-	Set("AvalancheIcons.Editor.Guides",                new IMAGE_BRUSH_SVG("Icons/EditorIcons/Guides",              Icon25));
-	Set("AvalancheIcons.Editor.Billboards",            new IMAGE_BRUSH("Icons/EditorIcons/Billboards",              Icon25));
-	Set("AvalancheIcons.Editor.SelectionLock",         new IMAGE_BRUSH("Icons/EditorIcons/SelectionLock",           Icon25));
-	Set("AvalancheIcons.Editor.IsolateActors",         new IMAGE_BRUSH("Icons/EditorIcons/IsolateActors",           Icon25));
-	Set("AvalancheIcons.Editor.Settings",              new IMAGE_BRUSH("Icons/EditorIcons/Settings",                Icon25));
-	Set("AvalancheIcons.Editor.Favorites",             new IMAGE_BRUSH("Icons/EditorIcons/Favorites",               FVector2f(9.0f, 14.0f)));
-	Set("AvalancheIcons.Editor.RemoteControlEditor",   new IMAGE_BRUSH("Icons/EditorIcons/RemoteControl",           Icon25));
-	Set("AvalancheIcons.Editor.BoundingBoxes",         new IMAGE_BRUSH_SVG("Icons/EditorIcons/Mode_Bounding-Box",   Icon25));
-	Set("AvalancheIcons.Editor.SafeFrames",            new IMAGE_BRUSH_SVG("Icons/EditorIcons/Mode_Broadcast-Safe", Icon25));
-	Set("AvalancheIcons.Editor.KeyPreview",            new IMAGE_BRUSH_SVG("Icons/EditorIcons/Mode_Key-Preview",    Icon25));
-	Set("AvalancheIcons.Editor.Stats",                 new IMAGE_BRUSH_SVG("Icons/EditorIcons/Mode_Performance",    Icon25));
-	Set("AvalancheIcons.Editor.Snapshot",              new IMAGE_BRUSH_SVG("Icons/EditorIcons/Mode_Snapshot",       Icon25));
-	Set("AvalancheIcons.Editor.WireframeMode",         new IMAGE_BRUSH_SVG("Icons/EditorIcons/Mode_Wireframe",      Icon25));
-	Set("AvalancheIcons.Editor.Snap.Toggle",           new IMAGE_BRUSH_SVG("Icons/EditorIcons/Snapping_Toggle",     Icon25));
-	Set("AvalancheIcons.Editor.Grid.Toggle",           new IMAGE_BRUSH_SVG("Icons/EditorIcons/Snapping_Grid",       Icon25));
-	Set("AvalancheIcons.Editor.Radio.BlackBackground", new IMAGE_BRUSH_SVG("Icons/EditorIcons/radio-background",    Icon16, FStyleColors::Foldout.GetSpecifiedColor()));
-	Set("AvalancheIcons.Editor.Mask.Toggle",           new IMAGE_BRUSH("Icons/EditorIcons/Mode_Mask",               Icon25));
-	Set("AvalancheIcons.Editor.PivotMode",             new IMAGE_BRUSH("Icons/EditorIcons/Pivot_Mode",              Icon25));
+	Set("Icons.Editor.ButtonExpander",        new IMAGE_BRUSH("Icons/EditorIcons/ButtonExpander",          FVector2f(7.0f)));
+	Set("Icons.Editor.ToolHightlight",        new IMAGE_BRUSH("Icons/EditorIcons/ToolIconHighlight",       Icon40));
+	Set("Icons.Editor.GameView",              new IMAGE_BRUSH_SVG("Icons/EditorIcons/GameView",            Icon25));
+	Set("Icons.Editor.Visualizers",           new IMAGE_BRUSH_SVG("Icons/EditorIcons/Visualizers",         Icon25));
+	Set("Icons.Editor.Guides",                new IMAGE_BRUSH_SVG("Icons/EditorIcons/Guides",              Icon25));
+	Set("Icons.Editor.Billboards",            new IMAGE_BRUSH("Icons/EditorIcons/Billboards",              Icon25));
+	Set("Icons.Editor.SelectionLock",         new IMAGE_BRUSH("Icons/EditorIcons/SelectionLock",           Icon25));
+	Set("Icons.Editor.IsolateActors",         new IMAGE_BRUSH("Icons/EditorIcons/IsolateActors",           Icon25));
+	Set("Icons.Editor.Settings",              new IMAGE_BRUSH("Icons/EditorIcons/Settings",                Icon25));
+	Set("Icons.Editor.Favorites",             new IMAGE_BRUSH("Icons/EditorIcons/Favorites",               FVector2f(9.0f, 14.0f)));
+	Set("Icons.Editor.RemoteControlEditor",   new IMAGE_BRUSH("Icons/EditorIcons/RemoteControl",           Icon25));
+	Set("Icons.Editor.BoundingBoxes",         new IMAGE_BRUSH_SVG("Icons/EditorIcons/Mode_Bounding-Box",   Icon25));
+	Set("Icons.Editor.SafeFrames",            new IMAGE_BRUSH_SVG("Icons/EditorIcons/Mode_Broadcast-Safe", Icon25));
+	Set("Icons.Editor.KeyPreview",            new IMAGE_BRUSH_SVG("Icons/EditorIcons/Mode_Key-Preview",    Icon25));
+	Set("Icons.Editor.Stats",                 new IMAGE_BRUSH_SVG("Icons/EditorIcons/Mode_Performance",    Icon25));
+	Set("Icons.Editor.Snapshot",              new IMAGE_BRUSH_SVG("Icons/EditorIcons/Mode_Snapshot",       Icon25));
+	Set("Icons.Editor.WireframeMode",         new IMAGE_BRUSH_SVG("Icons/EditorIcons/Mode_Wireframe",      Icon25));
+	Set("Icons.Editor.Snap.Toggle",           new IMAGE_BRUSH_SVG("Icons/EditorIcons/Snapping_Toggle",     Icon25));
+	Set("Icons.Editor.Grid.Toggle",           new IMAGE_BRUSH_SVG("Icons/EditorIcons/Snapping_Grid",       Icon25));
+	Set("Icons.Editor.Radio.BlackBackground", new IMAGE_BRUSH_SVG("Icons/EditorIcons/radio-background",    Icon16, FStyleColors::Foldout.GetSpecifiedColor()));
+	Set("Icons.Editor.Mask.Toggle",           new IMAGE_BRUSH("Icons/EditorIcons/Mode_Mask",               Icon25));
+	Set("Icons.Editor.PivotMode",             new IMAGE_BRUSH("Icons/EditorIcons/Pivot_Mode",              Icon25));
 	
-	Set("AvalancheIcons.Lock2d", new IMAGE_BRUSH("Icons/DetailsPanelIcons/Lock2d", Icon16));
-	Set("AvalancheIcons.Lock3d", new IMAGE_BRUSH("Icons/DetailsPanelIcons/Lock3d", Icon16));
-	Set("AvalancheIcons.Unlock", new IMAGE_BRUSH("Icons/DetailsPanelIcons/Unlock", Icon16));
+	Set("Icons.Lock2d", new IMAGE_BRUSH("Icons/DetailsPanelIcons/Lock2d", Icon16));
+	Set("Icons.Lock3d", new IMAGE_BRUSH("Icons/DetailsPanelIcons/Lock3d", Icon16));
+	Set("Icons.Unlock", new IMAGE_BRUSH("Icons/DetailsPanelIcons/Unlock", Icon16));
 
 	Set("AvaEditor.Thumbnail.Invalid", new FSlateRoundedBoxBrush(FStyleColors::Transparent, 4.0f, FLinearColor(1.0f, 0.2f, 0.2f, 1.0f), 1.0f));
 
@@ -244,7 +244,7 @@ FAvaEditorStyle::FAvaEditorStyle()
 	// Check Boxes
 	const FCheckBoxStyle& AppStyle_RadioButton = FAppStyle::GetWidgetStyle<FCheckBoxStyle>("RadioButton");
 	Set("AvaEditor.BlackRadioButton", FCheckBoxStyle(AppStyle_RadioButton)
-		.SetBackgroundImage(*GetBrush("AvalancheIcons.Editor.Radio.BlackBackground")));
+		.SetBackgroundImage(*GetBrush("Icons.Editor.Radio.BlackBackground")));
 
 	FSlateStyleRegistry::RegisterSlateStyle(*this);
 }
diff --git a/Engine/Plugins/Experimental/Avalanche/Source/AvalancheEditor/Private/DetailView/Customizations/AvaAnchorAlignmentPropertyTypeCustomization.h b/Engine/Plugins/Experimental/Avalanche/Source/AvalancheEditor/Private/DetailView/Customizations/AvaAnchorAlignmentPropertyTypeCustomization.h
--- a/Engine/Plugins/Experimental/Avalanche/Source/AvalancheEditor/Private/DetailView/Customizations/AvaAnchorAlignmentPropertyTypeCustomization.h
+++ b/Engine/Plugins/Experimental/Avalanche/Source/AvalancheEditor/Private/DetailView/Customizations/AvaAnchorAlignmentPropertyTypeCustomization.h
@@ -11,7 +11,7 @@ class IDetailChildrenBuilder;
 class IPropertyHandle;
 
 /**
- * Avalanche Anchor Alignment Customization
+ * Motion Design Anchor Alignment Customization
  * 
  * Shows a grid of buttons for setting anchor alignment in the following layout:
  * [HLeft ] [HCenter] [HRight ]
diff --git a/Engine/Plugins/Experimental/Avalanche/Source/AvalancheEditor/Private/DetailView/Customizations/AvaVectorPropertyTypeCustomization.cpp b/Engine/Plugins/Experimental/Avalanche/Source/AvalancheEditor/Private/DetailView/Customizations/AvaVectorPropertyTypeCustomization.cpp
--- a/Engine/Plugins/Experimental/Avalanche/Source/AvalancheEditor/Private/DetailView/Customizations/AvaVectorPropertyTypeCustomization.cpp
+++ b/Engine/Plugins/Experimental/Avalanche/Source/AvalancheEditor/Private/DetailView/Customizations/AvaVectorPropertyTypeCustomization.cpp
@@ -320,15 +320,15 @@ const FSlateBrush* FAvaVectorPropertyTypeCustomization::GetComboButtonBrush(cons
 {
 	if (Mode == ERatioMode::None)
 	{
-		return FAvaEditorStyle::Get().GetBrush("AvalancheIcons.Unlock");
+		return FAvaEditorStyle::Get().GetBrush("Icons.Unlock");
 	}
 	else if (Mode == ERatioMode::PreserveXY || Mode == ERatioMode::PreserveXZ || Mode == ERatioMode::PreserveYZ)
 	{
-		return FAvaEditorStyle::Get().GetBrush("AvalancheIcons.Lock2d");
+		return FAvaEditorStyle::Get().GetBrush("Icons.Lock2d");
 	}
 	else
 	{
-		return FAvaEditorStyle::Get().GetBrush("AvalancheIcons.Lock3d");
+		return FAvaEditorStyle::Get().GetBrush("Icons.Lock3d");
 	}
 }
 
diff --git a/Engine/Plugins/Experimental/Avalanche/Source/AvalancheEditor/Private/DetailView/Customizations/AvaVectorPropertyTypeCustomization.h b/Engine/Plugins/Experimental/Avalanche/Source/AvalancheEditor/Private/DetailView/Customizations/AvaVectorPropertyTypeCustomization.h
--- a/Engine/Plugins/Experimental/Avalanche/Source/AvalancheEditor/Private/DetailView/Customizations/AvaVectorPropertyTypeCustomization.h
+++ b/Engine/Plugins/Experimental/Avalanche/Source/AvalancheEditor/Private/DetailView/Customizations/AvaVectorPropertyTypeCustomization.h
@@ -54,8 +54,8 @@ ENUM_CLASS_FLAGS(ERatioMode)
 class FAvaVectorPropertyTypeCustomization : public IPropertyTypeCustomization
 {
 public:
-	typedef SNumericVectorInputBox<double, UE::Math::TVector2<double>, 2> SNumericVectorInputBox2D;
-	typedef SNumericVectorInputBox<FVector::FReal, UE::Math::TVector<FVector::FReal>, 3> SNumericVectorInputBox3D;
+	using SNumericVectorInputBox2D = SNumericVectorInputBox<double, UE::Math::TVector2<double>, 2>;
+	using SNumericVectorInputBox3D = SNumericVectorInputBox<FVector::FReal, UE::Math::TVector<FVector::FReal>, 3>;
 
 	static constexpr uint8 MULTI_OBJECT_DEBOUNCE  = 3;
 	static constexpr uint8 SINGLE_OBJECT_DEBOUNCE = 2;
diff --git a/Engine/Plugins/Experimental/Avalanche/Source/AvalancheEditor/Private/InteractiveTools/AvaInteractiveToolsExtension.cpp b/Engine/Plugins/Experimental/Avalanche/Source/AvalancheEditor/Private/InteractiveTools/AvaInteractiveToolsExtension.cpp
--- a/Engine/Plugins/Experimental/Avalanche/Source/AvalancheEditor/Private/InteractiveTools/AvaInteractiveToolsExtension.cpp
+++ b/Engine/Plugins/Experimental/Avalanche/Source/AvalancheEditor/Private/InteractiveTools/AvaInteractiveToolsExtension.cpp
@@ -21,7 +21,7 @@
 #include "MediaPlate.h"
 #include "Viewport/AvaPostProcessVolume.h"
 
-namespace UE::AvalancheEditor::Private
+namespace UE::AvaEditor::Private
 {
 	const FName StaticMeshCategory = FName(TEXT("StaticMeshes"));
 	const FName CamerasCategory = FName(TEXT("Cameras"));
@@ -77,7 +77,7 @@ void FAvaInteractiveToolsExtension::UnregisterToolDelegates()
 
 void FAvaInteractiveToolsExtension::RegisterCategories(IAvalancheInteractiveToolsModule* InModule)
 {
-	using namespace UE::AvalancheEditor::Private;
+	using namespace UE::AvaEditor::Private;
 
 	const FAvaEditorCommands& EditorCommands = FAvaEditorCommands::Get();
 
@@ -88,7 +88,7 @@ void FAvaInteractiveToolsExtension::RegisterCategories(IAvalancheInteractiveTool
 
 void FAvaInteractiveToolsExtension::RegisterTools(IAvalancheInteractiveToolsModule* InModule)
 {
-	using namespace UE::AvalancheEditor::Private;
+	using namespace UE::AvaEditor::Private;
 
 	const FAvaEditorCommands& EditorCommands = FAvaEditorCommands::Get();
 
diff --git a/Engine/Plugins/Experimental/Avalanche/Source/AvalancheEditor/Private/MaterialDesigner/AvaLevelMaterialDesignerExtension.cpp b/Engine/Plugins/Experimental/Avalanche/Source/AvalancheEditor/Private/MaterialDesigner/AvaLevelMaterialDesignerExtension.cpp
--- a/Engine/Plugins/Experimental/Avalanche/Source/AvalancheEditor/Private/MaterialDesigner/AvaLevelMaterialDesignerExtension.cpp
+++ b/Engine/Plugins/Experimental/Avalanche/Source/AvalancheEditor/Private/MaterialDesigner/AvaLevelMaterialDesignerExtension.cpp
@@ -13,7 +13,7 @@
 
 #define LOCTEXT_NAMESPACE "AvaMaterialDesignerExtension"
 
-namespace UE::AvalancheEditor::Private
+namespace UE::AvaEditor::Private
 {
 	int32 AvaLevelMaterialDesignerExtensionInstances = 0;
 	FDelegateHandle ContentBrowserExtenderDelegateHandle;
@@ -23,7 +23,7 @@ void FAvaLevelMaterialDesignerExtension::Activate()
 {
 	FAvaMaterialDesignerExtension::Activate();
 
-	using namespace UE::AvalancheEditor::Private;
+	using namespace UE::AvaEditor::Private;
 
 	++AvaLevelMaterialDesignerExtensionInstances;
 	bIsActive = true;
@@ -35,7 +35,7 @@ void FAvaLevelMaterialDesignerExtension::Deactivate()
 {
 	FAvaMaterialDesignerExtension::Deactivate();
 
-	using namespace UE::AvalancheEditor::Private;
+	using namespace UE::AvaEditor::Private;
 
 	if (bIsActive)
 	{
@@ -48,7 +48,7 @@ void FAvaLevelMaterialDesignerExtension::Deactivate()
 
 void FAvaLevelMaterialDesignerExtension::InitContentBrowserExtension()
 {
-	using namespace UE::AvalancheEditor::Private;
+	using namespace UE::AvaEditor::Private;
 
 	// Register only with the first instance.
 	if (AvaLevelMaterialDesignerExtensionInstances == 1)
@@ -62,7 +62,7 @@ void FAvaLevelMaterialDesignerExtension::InitContentBrowserExtension()
 
 void FAvaLevelMaterialDesignerExtension::DeinitContentBrowserExtension()
 {
-	using namespace UE::AvalancheEditor::Private;
+	using namespace UE::AvaEditor::Private;
 
 	// Unregister only with the last instance
 	if (AvaLevelMaterialDesignerExtensionInstances == 0 && ContentBrowserExtenderDelegateHandle.IsValid())
diff --git a/Engine/Plugins/Experimental/Avalanche/Source/AvalancheEditor/Private/Scene/AvaSceneDefaults.cpp b/Engine/Plugins/Experimental/Avalanche/Source/AvalancheEditor/Private/Scene/AvaSceneDefaults.cpp
--- a/Engine/Plugins/Experimental/Avalanche/Source/AvalancheEditor/Private/Scene/AvaSceneDefaults.cpp
+++ b/Engine/Plugins/Experimental/Avalanche/Source/AvalancheEditor/Private/Scene/AvaSceneDefaults.cpp
@@ -32,7 +32,7 @@ using FAvaSceneDefaultActorList = TMap<FName, TSharedRef<FAvaSceneDefaultActorRe
 
 DECLARE_DELEGATE_ThreeParams(FAvaSceneDefaultActorInitDelegate, TSharedRef<IAvaEditor>, const FAvaSceneDefaultActorList&, AActor*)
 
-namespace UE::AvalancheEditor::Private
+namespace UE::AvaEditor::Private
 {
 	static const FPreviewSceneProfile DefaultSceneProfile;
 
@@ -131,7 +131,7 @@ namespace UE::AvalancheEditor::Private
 			return;
 		}
 
-		using namespace UE::AvalancheEditor::Private;
+		using namespace UE::AvaEditor::Private;
 
 		ULightComponent* const LightComponent = DirectionalLight->GetLightComponent();
 
@@ -237,7 +237,7 @@ namespace UE::AvalancheEditor::Private
 
 		for (const TSharedPtr<IAvaViewportClient>& ViewportClient : ViewportExtension->GetViewportClients())
 		{
-			if (!ViewportClient->IsAvalancheViewport())
+			if (!ViewportClient->IsMotionDesignViewport())
 			{
 				continue;
 			}
@@ -431,7 +431,7 @@ void FAvaSceneDefaults::CreateDefaultScene(TSharedRef<IAvaEditor> InEditor, UWor
 		return;
 	}
 
-	using namespace UE::AvalancheEditor::Private;
+	using namespace UE::AvaEditor::Private;
 
 	TMap<FName, TSharedRef<FAvaSceneDefaultActorResponse>> ActorResponses = GeneratorInitialActorResponses(InWorld);
 
diff --git a/Engine/Plugins/Experimental/Avalanche/Source/AvalancheEditor/Private/Scene/SAvaSceneDefaultActorResponseRow.cpp b/Engine/Plugins/Experimental/Avalanche/Source/AvalancheEditor/Private/Scene/SAvaSceneDefaultActorResponseRow.cpp
--- a/Engine/Plugins/Experimental/Avalanche/Source/AvalancheEditor/Private/Scene/SAvaSceneDefaultActorResponseRow.cpp
+++ b/Engine/Plugins/Experimental/Avalanche/Source/AvalancheEditor/Private/Scene/SAvaSceneDefaultActorResponseRow.cpp
@@ -15,7 +15,7 @@
 
 #define LOCTEXT_NAMESPACE "SAvaSceneDefaultActorResponseRow"
 
-namespace UE::AvalancheEditor::Private
+namespace UE::AvaEditor::Private
 {
 	static const FMargin AvaSceneDefaultActorResponseRowPadding = FMargin(5.f, 5.f);
 }
@@ -71,7 +71,7 @@ TSharedRef<SWidget> SAvaSceneDefaultActorResponseRow::GenerateWidgetForColumn(co
 
 TSharedRef<SWidget> SAvaSceneDefaultActorResponseRow::GenerateRow_Description()
 {
-	using namespace UE::AvalancheEditor::Private;
+	using namespace UE::AvaEditor::Private;
 
 	FSlateIcon ActorIcon = FSlateIconFinder::FindIconForClass(AActor::StaticClass());
 
@@ -109,7 +109,7 @@ TSharedRef<SWidget> SAvaSceneDefaultActorResponseRow::GenerateRow_Description()
 
 TSharedRef<SWidget> SAvaSceneDefaultActorResponseRow::GenerateRow_Actor()
 {
-	using namespace UE::AvalancheEditor::Private;
+	using namespace UE::AvaEditor::Private;
 
 	TSharedPtr<SWidget> ActorContent;
 	TSharedPtr<SAvaSceneDefaultActorResponses> ActorResponses = ActorResponsesWeak.Pin();
@@ -166,7 +166,7 @@ TSharedRef<SWidget> SAvaSceneDefaultActorResponseRow::GenerateRow_Actor()
 
 TSharedRef<SWidget> SAvaSceneDefaultActorResponseRow::GenerateRow_CreateNew()
 {
-	using namespace UE::AvalancheEditor::Private;
+	using namespace UE::AvaEditor::Private;
 
 	return SNew(SBox)
 		.HAlign(HAlign_Center)
@@ -183,7 +183,7 @@ TSharedRef<SWidget> SAvaSceneDefaultActorResponseRow::GenerateRow_CreateNew()
 
 TSharedRef<SWidget> SAvaSceneDefaultActorResponseRow::GenerateRow_Replace()
 {
-	using namespace UE::AvalancheEditor::Private;
+	using namespace UE::AvaEditor::Private;
 
 	return SNew(SBox)
 		.HAlign(HAlign_Center)
@@ -200,7 +200,7 @@ TSharedRef<SWidget> SAvaSceneDefaultActorResponseRow::GenerateRow_Replace()
 
 TSharedRef<SWidget> SAvaSceneDefaultActorResponseRow::GenerateRow_Update()
 {
-	using namespace UE::AvalancheEditor::Private;
+	using namespace UE::AvaEditor::Private;
 
 	return SNew(SBox)
 		.HAlign(HAlign_Center)
@@ -217,7 +217,7 @@ TSharedRef<SWidget> SAvaSceneDefaultActorResponseRow::GenerateRow_Update()
 
 TSharedRef<SWidget> SAvaSceneDefaultActorResponseRow::GenerateRow_Enabled()
 {
-	using namespace UE::AvalancheEditor::Private;
+	using namespace UE::AvaEditor::Private;
 
 	return SNew(SBox)
 		.HAlign(HAlign_Center)
diff --git a/Engine/Plugins/Experimental/Avalanche/Source/AvalancheEditor/Private/Scene/SAvaSceneDefaultActorResponses.cpp b/Engine/Plugins/Experimental/Avalanche/Source/AvalancheEditor/Private/Scene/SAvaSceneDefaultActorResponses.cpp
--- a/Engine/Plugins/Experimental/Avalanche/Source/AvalancheEditor/Private/Scene/SAvaSceneDefaultActorResponses.cpp
+++ b/Engine/Plugins/Experimental/Avalanche/Source/AvalancheEditor/Private/Scene/SAvaSceneDefaultActorResponses.cpp
@@ -14,7 +14,7 @@ const FName SAvaSceneDefaultActorResponses::ReplaceRowName = "Replace";
 const FName SAvaSceneDefaultActorResponses::UpdateRowName = "Update";
 const FName SAvaSceneDefaultActorResponses::EnabledRowName = "Enabled";
 
-namespace UE::AvalancheEditor::Private
+namespace UE::AvaEditor::Private
 {
 	struct FAvaSceneDefaultActorResponseColumnData
 	{
@@ -73,7 +73,7 @@ namespace UE::AvalancheEditor::Private
 
 void SAvaSceneDefaultActorResponses::Construct(const FArguments& InArgs, UWorld* InWorld, const TArray<TSharedRef<FAvaSceneDefaultActorResponse>>& InResponses)
 {
-	using namespace UE::AvalancheEditor::Private;
+	using namespace UE::AvaEditor::Private;
 
 	World = InWorld;
 	Responses = InResponses;
diff --git a/Engine/Plugins/Experimental/Avalanche/Source/AvalancheEditor/Private/Tests/Framework/AvaEditorTestUtils.cpp b/Engine/Plugins/Experimental/Avalanche/Source/AvalancheEditor/Private/Tests/Framework/AvaEditorTestUtils.cpp
--- a/Engine/Plugins/Experimental/Avalanche/Source/AvalancheEditor/Private/Tests/Framework/AvaEditorTestUtils.cpp
+++ b/Engine/Plugins/Experimental/Avalanche/Source/AvalancheEditor/Private/Tests/Framework/AvaEditorTestUtils.cpp
@@ -9,7 +9,7 @@
 #include "Math/Rotator.h"
 #include "Misc/AutomationTest.h"
 
-// General Avalanche Editor unit test log
+// General Motion Design Editor unit test log
 DEFINE_LOG_CATEGORY(LogAvaEditorTest);
 
 TArray<FVector> FAvaEditorTestUtils::GetActorLocations(const TArray<AActor*>& InActors)
diff --git a/Engine/Plugins/Experimental/Avalanche/Source/AvalancheEditor/Private/Viewport/AvaLevelViewportExtension.cpp b/Engine/Plugins/Experimental/Avalanche/Source/AvalancheEditor/Private/Viewport/AvaLevelViewportExtension.cpp
--- a/Engine/Plugins/Experimental/Avalanche/Source/AvalancheEditor/Private/Viewport/AvaLevelViewportExtension.cpp
+++ b/Engine/Plugins/Experimental/Avalanche/Source/AvalancheEditor/Private/Viewport/AvaLevelViewportExtension.cpp
@@ -23,7 +23,7 @@
 
 #define LOCTEXT_NAMESPACE "AvaLevelViewportExtension"
 
-namespace UE::AvalancheEditor::Private
+namespace UE::AvaEditor::Private
 {
 	class SAvaLevelViewportExposeBinds : public SLevelViewport
 	{
@@ -65,20 +65,20 @@ namespace UE::AvalancheEditor::Private
 
 		TSharedPtr<FEditorViewportClient> ViewportClientToFocus = nullptr;
 
-		// Get the first valid Avalanche Viewport
+		// Get the first valid Motion Design viewport
 		for (const TSharedPtr<SLevelViewport>& Viewport : Viewports)
 		{
 			TSharedPtr<FEditorViewportClient> ViewportClient = Viewport->GetViewportClient();
 			TSharedPtr<IAvaViewportClient> AvaViewportClient = FAvaViewportUtils::GetAsAvaViewportClient(ViewportClient.Get());
 
-			if (AvaViewportClient && AvaViewportClient->IsAvalancheViewport())
+			if (AvaViewportClient && AvaViewportClient->IsMotionDesignViewport())
 			{
 				ViewportClientToFocus = ViewportClient;
 				break;
 			}
 		}
 
-		// if there isn't any ava viewport, use first viewport client
+		// if there isn't any Motion Design viewport, use first viewport client
 		if (!ViewportClientToFocus)
 		{
 			ViewportClientToFocus = Viewports[0]->GetViewportClient();
@@ -152,7 +152,7 @@ void FAvaLevelViewportExtension::Activate()
 		{
 			for (TSharedPtr<SLevelViewport> LevelViewport : LevelEditor->GetViewports())
 			{
-				StaticCastSharedPtr<UE::AvalancheEditor::Private::SAvaLevelViewportExposeBinds>(LevelViewport)->BindCommands();
+				StaticCastSharedPtr<UE::AvaEditor::Private::SAvaLevelViewportExposeBinds>(LevelViewport)->BindCommands();
 			}
 		}
 	}
@@ -167,7 +167,7 @@ void FAvaLevelViewportExtension::Deactivate()
 	FAvaViewportExtension::Deactivate();
 
 	// If the command is invalid, we're shutting down and this doesn't need to be run.
-	if (!FAvaEditorCommands::IsRegistered() || !FAvaEditorCommands::Get().SetAvalancheViewportType.IsValid())
+	if (!FAvaEditorCommands::IsRegistered() || !FAvaEditorCommands::Get().SetMotionDesignViewportType.IsValid())
 	{
 		return;
 	}
@@ -176,7 +176,7 @@ void FAvaLevelViewportExtension::Deactivate()
 	{
 		if (TSharedPtr<ILevelEditor> LevelEditor = LevelEditorModule->GetFirstLevelEditor())
 		{
-			TSharedPtr<FUICommandInfo> ToggleViewportCommandInfo = FAvaEditorCommands::Get().SetAvalancheViewportType;
+			TSharedPtr<FUICommandInfo> ToggleViewportCommandInfo = FAvaEditorCommands::Get().SetMotionDesignViewportType;
 
 			for (TSharedPtr<SLevelViewport> LevelViewport : LevelEditor->GetViewports())
 			{
@@ -290,21 +290,21 @@ void FAvaLevelViewportExtension::SetDefaultViewportType()
 	if (TSharedPtr<IAvaEditor> Editor = GetEditor())
 	{
 		Editor->GetCommandList()->ExecuteAction(FLevelViewportCommands::Get().SetDefaultViewportType.ToSharedRef());
-		UE::AvalancheEditor::Private::FixupInvalidFocusedLevelEditorViewport();
+		UE::AvaEditor::Private::FixupInvalidFocusedLevelEditorViewport();
 	} 
 }
 
-void FAvaLevelViewportExtension::SetAvalancheViewportType()
+void FAvaLevelViewportExtension::SetMotionDesignViewportType()
 {
 	if (TSharedPtr<IAvaEditor> Editor = GetEditor())
 	{
 		TSharedPtr<FUICommandList> CommandList = Editor->GetCommandList();
 		if (ensure(CommandList.IsValid()))
 		{
-			CommandList->ExecuteAction(FAvaEditorCommands::Get().SetAvalancheViewportType.ToSharedRef());
+			CommandList->ExecuteAction(FAvaEditorCommands::Get().SetMotionDesignViewportType.ToSharedRef());
 		}
 
-		UE::AvalancheEditor::Private::FixupInvalidFocusedLevelEditorViewport();
+		UE::AvaEditor::Private::FixupInvalidFocusedLevelEditorViewport();
 
 		if (AActor* LastCameraCutActor = LastCameraCutActorWeak.Get())
 		{
@@ -320,7 +320,7 @@ FViewportTypeDefinition FAvaLevelViewportExtension::MakeViewportTypeDefinition()
 		{
 			return MakeShared<FAvaLevelViewportLayoutEntity>(InArgs, InLevelEditor, ThisWeak.Pin());
 		}
-		, FAvaEditorCommands::Get().SetAvalancheViewportType);
+		, FAvaEditorCommands::Get().SetMotionDesignViewportType);
 }
 
 void FAvaLevelViewportExtension::OnSwitchViewports()
@@ -356,13 +356,13 @@ void FAvaLevelViewportExtension::OnSwitchViewports()
 	}
 
 	UObject* const SceneObject = GetSceneObject();
-	if (!SceneObject || AvaViewportClient->IsAvalancheViewport())
+	if (!SceneObject || AvaViewportClient->IsMotionDesignViewport())
 	{
 		SetDefaultViewportType();
 	}
 	else
 	{
-		SetAvalancheViewportType();
+		SetMotionDesignViewportType();
 	}
 }
 
@@ -407,14 +407,14 @@ void FAvaLevelViewportExtension::CheckValidViewportType()
 {
 	TSharedPtr<IAvaEditor> Editor = GetEditor();
 
-	// Set the active viewport to Avalanche if there is a valid scene object (not all the viewports)
+	// Set the active viewport to Motion Design if there is a valid scene object (not all the viewports)
 	if (Editor.IsValid() && Editor->IsActive() && GetSceneObject())
 	{
-		SetAvalancheViewportType();
+		SetMotionDesignViewportType();
 		return;
 	}
 
-	// Reset all avalanche viewports to default if there was an invalid scene object
+	// Reset all Motion Design viewports to default if there was an invalid scene object
 	FLevelEditorModule* LevelEditorModule = FAvaLevelEditorUtils::GetLevelEditorModule();
 	if (!LevelEditorModule)
 	{
@@ -427,17 +427,17 @@ void FAvaLevelViewportExtension::CheckValidViewportType()
 		return;
 	}
 
-	// Execute SetDefaultViewportType on all Avalanche Viewports, not only the Active one
+	// Execute SetDefaultViewportType on all Motion Design viewports, not only the Active one
 	for (TSharedPtr<SLevelViewport> LevelViewport : LevelEditor->GetViewports())
 	{
 		TSharedPtr<IAvaViewportClient> AvaViewportClient = FAvaViewportUtils::GetAsAvaViewportClient(LevelViewport->GetViewportClient().Get());
-		if (AvaViewportClient.IsValid() && AvaViewportClient->IsAvalancheViewport())
+		if (AvaViewportClient.IsValid() && AvaViewportClient->IsMotionDesignViewport())
 		{
 			LevelViewport->GetCommandList()->ExecuteAction(FLevelViewportCommands::Get().SetDefaultViewportType.ToSharedRef());
 		}
 	}
 
-	UE::AvalancheEditor::Private::FixupInvalidFocusedLevelEditorViewport();
+	UE::AvaEditor::Private::FixupInvalidFocusedLevelEditorViewport();
 }
 
 void FAvaLevelViewportExtension::BindCameraCutDelegate()
@@ -504,7 +504,7 @@ void FAvaLevelViewportExtension::SetActiveCamera(AActor* InActiveCameraActor, bo
 
 	for (TSharedPtr<IAvaViewportClient> ViewportClient : GetViewportClients())
 	{
-		if (ViewportClient->IsAvalancheViewport())
+		if (ViewportClient->IsMotionDesignViewport())
 		{
 			ViewportClient->OnCameraCut(InActiveCameraActor, bInJumpCut);
 		}
diff --git a/Engine/Plugins/Experimental/Avalanche/Source/AvalancheEditor/Private/Viewport/AvaLevelViewportExtension.h b/Engine/Plugins/Experimental/Avalanche/Source/AvalancheEditor/Private/Viewport/AvaLevelViewportExtension.h
--- a/Engine/Plugins/Experimental/Avalanche/Source/AvalancheEditor/Private/Viewport/AvaLevelViewportExtension.h
+++ b/Engine/Plugins/Experimental/Avalanche/Source/AvalancheEditor/Private/Viewport/AvaLevelViewportExtension.h
@@ -25,7 +25,7 @@ public:
 	static TArray<TSharedPtr<IAvaViewportClient>> GetLevelEditorViewportClients();
 
 	void SetDefaultViewportType();
-	void SetAvalancheViewportType();
+	void SetMotionDesignViewportType();
 
 	virtual ~FAvaLevelViewportExtension() override;
 
diff --git a/Engine/Plugins/Experimental/Avalanche/Source/AvalancheEditorCore/Private/AvaEdMode.cpp b/Engine/Plugins/Experimental/Avalanche/Source/AvalancheEditorCore/Private/AvaEdMode.cpp
--- a/Engine/Plugins/Experimental/Avalanche/Source/AvalancheEditorCore/Private/AvaEdMode.cpp
+++ b/Engine/Plugins/Experimental/Avalanche/Source/AvalancheEditorCore/Private/AvaEdMode.cpp
@@ -9,7 +9,7 @@
 
 #define LOCTEXT_NAMESPACE "AvaEdMode"
 
-const FEditorModeID UAvaEdMode::ModeID(TEXT("EM_Avalanche"));
+const FEditorModeID UAvaEdMode::ModeID(TEXT("EM_MotionDesign"));
 
 UAvaEdMode::UAvaEdMode()
 {
diff --git a/Engine/Plugins/Experimental/Avalanche/Source/AvalancheEditorCore/Private/AvaEdModeToolkit.cpp b/Engine/Plugins/Experimental/Avalanche/Source/AvalancheEditorCore/Private/AvaEdModeToolkit.cpp
--- a/Engine/Plugins/Experimental/Avalanche/Source/AvalancheEditorCore/Private/AvaEdModeToolkit.cpp
+++ b/Engine/Plugins/Experimental/Avalanche/Source/AvalancheEditorCore/Private/AvaEdModeToolkit.cpp
@@ -7,7 +7,7 @@
 
 #define LOCTEXT_NAMESPACE "AvaEdModeToolkit"
 
-namespace UE::AvalancheEditorCore::Private
+namespace UE::AvaEditorCore::Private
 {
 	FName ToolkitName = "AvaEdModeToolkit";
 	FText ToolkitBaseName = LOCTEXT("ToolkitName", "Motion Design Editor Toolkit");
@@ -20,12 +20,12 @@ FAvaEdModeToolkit::FAvaEdModeToolkit(UAvaEdMode* InEdMode)
 
 FName FAvaEdModeToolkit::GetToolkitFName() const
 {
-	return UE::AvalancheEditorCore::Private::ToolkitName;
+	return UE::AvaEditorCore::Private::ToolkitName;
 }
 
 FText FAvaEdModeToolkit::GetBaseToolkitName() const
 {
-	return UE::AvalancheEditorCore::Private::ToolkitBaseName;
+	return UE::AvaEditorCore::Private::ToolkitBaseName;
 }
 
 void FAvaEdModeToolkit::ExtendSecondaryModeToolbar(UToolMenu* InToolbarMenu)
diff --git a/Engine/Plugins/Experimental/Avalanche/Source/AvalancheEditorCore/Private/AvaEditor.cpp b/Engine/Plugins/Experimental/Avalanche/Source/AvalancheEditorCore/Private/AvaEditor.cpp
--- a/Engine/Plugins/Experimental/Avalanche/Source/AvalancheEditorCore/Private/AvaEditor.cpp
+++ b/Engine/Plugins/Experimental/Avalanche/Source/AvalancheEditorCore/Private/AvaEditor.cpp
@@ -183,7 +183,7 @@ void FAvaEditor::InvokeTabs()
 
 	TArray<FTabId> TabsToInvoke;
 
-	// Gather all the Avalanche Tabs that are supposed to be opened, but were possibly not due to not having a valid Toolkit Host at the time
+	// Gather all the Motion Design Tabs that are supposed to be opened, but were possibly not due to not having a valid Toolkit Host at the time
 	TArray<TSharedRef<FTabManager::FSplitter>> Splitters(StaticCastSharedRef<FLayout>(Layout)->GetAreas());
 
 	while (!Splitters.IsEmpty())
diff --git a/Engine/Plugins/Experimental/Avalanche/Source/AvalancheEditorCore/Private/AvaEditorWidgetUtils.cpp b/Engine/Plugins/Experimental/Avalanche/Source/AvalancheEditorCore/Private/AvaEditorWidgetUtils.cpp
--- a/Engine/Plugins/Experimental/Avalanche/Source/AvalancheEditorCore/Private/AvaEditorWidgetUtils.cpp
+++ b/Engine/Plugins/Experimental/Avalanche/Source/AvalancheEditorCore/Private/AvaEditorWidgetUtils.cpp
@@ -12,7 +12,7 @@
 
 #define LOCTEXT_NAMESPACE "AvaEditorWidgetUtils"
 
-namespace UE::AvalancheEditor::Private
+namespace UE::AvaEditor::Private
 {
 	void GetWidgetChildrenOfClass(const TSharedRef<SWidget>& Widget, const FName& WidgetClassName, TArray<TSharedRef<SWidget>>& Widgets)
 	{
@@ -37,7 +37,7 @@ TArray<TSharedRef<SWidget>> FAvaEditorWidgetUtils::GetWidgetChildrenOfClass(cons
 {
 	TArray<TSharedRef<SWidget>> Widgets;
 
-	UE::AvalancheEditor::Private::GetWidgetChildrenOfClass(InWidget.ToSharedRef(), InWidgetClass.GetWidgetType(), Widgets);
+	UE::AvaEditor::Private::GetWidgetChildrenOfClass(InWidget.ToSharedRef(), InWidgetClass.GetWidgetType(), Widgets);
 
 	return Widgets;
 }
diff --git a/Engine/Plugins/Experimental/Avalanche/Source/AvalancheEditorCore/Private/Bounds/AvaBoundsProviderSubsystem.cpp b/Engine/Plugins/Experimental/Avalanche/Source/AvalancheEditorCore/Private/Bounds/AvaBoundsProviderSubsystem.cpp
--- a/Engine/Plugins/Experimental/Avalanche/Source/AvalancheEditorCore/Private/Bounds/AvaBoundsProviderSubsystem.cpp
+++ b/Engine/Plugins/Experimental/Avalanche/Source/AvalancheEditorCore/Private/Bounds/AvaBoundsProviderSubsystem.cpp
@@ -9,7 +9,7 @@
 #include "Selection.h"
 #include "Selection/AvaSelectionProviderSubsystem.h"
 
-namespace UE::AvalancheEditorCore::Private
+namespace UE::AvaEditorCore::Private
 {
 	FBox InvalidBox = FBox(EForceInit::ForceInit);
 	FOrientedBox DefaultOrientedBox = FOrientedBox();
@@ -76,7 +76,7 @@ bool UAvaBoundsProviderSubsystem::CacheComponentWorldOrientedBounds(UPrimitiveCo
 		return true;
 	}
 
-	using namespace UE::AvalancheEditorCore::Private;
+	using namespace UE::AvaEditorCore::Private;
 
 	const FBox LocalBounds = GetComponentLocalBounds(InComponent);
 
@@ -110,7 +110,7 @@ bool UAvaBoundsProviderSubsystem::CacheComponentActorOrientedBounds(UPrimitiveCo
 		return true;
 	}
 
-	using namespace UE::AvalancheEditorCore::Private;
+	using namespace UE::AvaEditorCore::Private;
 
 	const FBox LocalBounds = GetComponentLocalBounds(InComponent);
 
@@ -200,7 +200,7 @@ bool UAvaBoundsProviderSubsystem::CacheActorOrientedBounds(AActor* InActor)
 		return true;
 	}
 
-	using namespace UE::AvalancheEditorCore::Private;
+	using namespace UE::AvaEditorCore::Private;
 
 	const FBox LocalBounds = GetActorLocalBounds(InActor);
 
@@ -249,7 +249,7 @@ bool UAvaBoundsProviderSubsystem::CacheActorAndChildrenLocalBounds(AActor* InAct
 	}
 	else
 	{
-		using namespace UE::AvalancheEditorCore::Private;
+		using namespace UE::AvaEditorCore::Private;
 
 		FTransform ActorTransform = InActor->GetActorTransform();
 		FTransform InverseActorTransformNoScale = FTransform(ActorTransform.GetRotation(), ActorTransform.GetLocation()).Inverse();
@@ -299,7 +299,7 @@ bool UAvaBoundsProviderSubsystem::CacheActorAndChildrenOrientedBounds(AActor* In
 		return true;
 	}
 
-	using namespace UE::AvalancheEditorCore::Private;
+	using namespace UE::AvaEditorCore::Private;
 
 	const FBox LocalBounds = GetActorAndChildrenLocalBounds(InActor);
 
@@ -369,7 +369,7 @@ FBox UAvaBoundsProviderSubsystem::GetComponentLocalBounds(UPrimitiveComponent* I
 		return CachedComponentLocalBounds[InComponent];
 	}
 
-	return UE::AvalancheEditorCore::Private::InvalidBox;
+	return UE::AvaEditorCore::Private::InvalidBox;
 }
 
 bool UAvaBoundsProviderSubsystem::GetComponentWorldOrientedBounds(UPrimitiveComponent* InComponent, FOrientedBox& OutOrientedBounds)
@@ -401,7 +401,7 @@ FBox UAvaBoundsProviderSubsystem::GetActorLocalBounds(AActor* InActor)
 		return CachedActorLocalBounds[InActor];
 	}
 
-	return UE::AvalancheEditorCore::Private::InvalidBox;
+	return UE::AvaEditorCore::Private::InvalidBox;
 }
 
 bool UAvaBoundsProviderSubsystem::GetActorOrientedBounds(AActor* InActor, FOrientedBox& OutOrientedBounds)
@@ -422,7 +422,7 @@ FBox UAvaBoundsProviderSubsystem::GetActorAndChildrenLocalBounds(AActor* InActor
 		return CachedActorAndChildrenLocalBounds[InActor];
 	}
 
-	return UE::AvalancheEditorCore::Private::InvalidBox;
+	return UE::AvaEditorCore::Private::InvalidBox;
 }
 
 bool UAvaBoundsProviderSubsystem::GetActorAndChildrenOrientedBounds(AActor* InActor, FOrientedBox& OutOrientedBounds)
@@ -619,7 +619,7 @@ bool UAvaBoundsProviderSubsystem::GetSelectionOrientedBounds(bool bInIncludeChil
 		return false;
 	}
 
-	OutOrientedBounds = UE::AvalancheEditorCore::Private::MakeOrientedBox(SelectionAxisBounds, FirstActorTransform);
+	OutOrientedBounds = UE::AvaEditorCore::Private::MakeOrientedBox(SelectionAxisBounds, FirstActorTransform);
 
 	return true;
 }
diff --git a/Engine/Plugins/Experimental/Avalanche/Source/AvalancheEditorCore/Private/LevelEditor/AvaLevelEditorToolbar.cpp b/Engine/Plugins/Experimental/Avalanche/Source/AvalancheEditorCore/Private/LevelEditor/AvaLevelEditorToolbar.cpp
--- a/Engine/Plugins/Experimental/Avalanche/Source/AvalancheEditorCore/Private/LevelEditor/AvaLevelEditorToolbar.cpp
+++ b/Engine/Plugins/Experimental/Avalanche/Source/AvalancheEditorCore/Private/LevelEditor/AvaLevelEditorToolbar.cpp
@@ -133,11 +133,9 @@ void FAvaLevelEditorToolbar::ExtendLevelEditorToolbar(const TSharedRef<IAvaEdito
 		const FAvaLevelEditorCommands& AvaLevelEditorCommands = FAvaLevelEditorCommands::Get();
 		FToolMenuSection& Section = LevelEditorToolbar->FindOrAddSection(Private::SceneSectionName);
 
-		const FText AvalancheLabel = LOCTEXT("MotionDesignScene", "Motion Design");
-
-		auto AddEntry = [this, &Section, AvalancheLabel](const TSharedPtr<FUICommandInfo>& InCommand)
+		auto AddEntry = [this, &Section](const TSharedPtr<FUICommandInfo>& InCommand)
 		{
-			FToolMenuEntry& Entry = Section.AddEntry(FToolMenuEntry::InitToolBarButton(InCommand, AvalancheLabel));
+			FToolMenuEntry& Entry = Section.AddEntry(FToolMenuEntry::InitToolBarButton(InCommand, LOCTEXT("MotionDesignScene", "Motion Design")));
 			Entry.SetCommandList(CommandList);
 			Entry.StyleNameOverride = "CalloutToolbar";
 		};
diff --git a/Engine/Plugins/Experimental/Avalanche/Source/AvalancheEditorCore/Private/LevelEditor/AvaLevelEditorToolbar.h b/Engine/Plugins/Experimental/Avalanche/Source/AvalancheEditorCore/Private/LevelEditor/AvaLevelEditorToolbar.h
--- a/Engine/Plugins/Experimental/Avalanche/Source/AvalancheEditorCore/Private/LevelEditor/AvaLevelEditorToolbar.h
+++ b/Engine/Plugins/Experimental/Avalanche/Source/AvalancheEditorCore/Private/LevelEditor/AvaLevelEditorToolbar.h
@@ -11,7 +11,7 @@ class SBorder;
 class SWidget;
 class UToolMenu;
 
-/** Responsible for manipulating the Level Editor toolbar for the Avalanche Editor */
+/** Responsible for manipulating the Level Editor toolbar for the Motion Design Editor */
 class FAvaLevelEditorToolbar : public TSharedFromThis<FAvaLevelEditorToolbar>
 {
 public:
diff --git a/Engine/Plugins/Experimental/Avalanche/Source/AvalancheEditorCore/Private/LevelEditor/AvaLevelEditorUtils.cpp b/Engine/Plugins/Experimental/Avalanche/Source/AvalancheEditorCore/Private/LevelEditor/AvaLevelEditorUtils.cpp
--- a/Engine/Plugins/Experimental/Avalanche/Source/AvalancheEditorCore/Private/LevelEditor/AvaLevelEditorUtils.cpp
+++ b/Engine/Plugins/Experimental/Avalanche/Source/AvalancheEditorCore/Private/LevelEditor/AvaLevelEditorUtils.cpp
@@ -5,19 +5,19 @@
 #include "LevelEditor.h"
 #include "Modules/ModuleManager.h"
 
-namespace UE::AvalancheEditorCore::Private
+namespace UE::AvaEditorCore::Private
 {
 	static constexpr const TCHAR* LevelEditorModuleName = TEXT("LevelEditor");
 }
 
 FLevelEditorModule* FAvaLevelEditorUtils::GetLevelEditorModule()
 {
-	return FModuleManager::Get().GetModulePtr<FLevelEditorModule>(UE::AvalancheEditorCore::Private::LevelEditorModuleName);
+	return FModuleManager::Get().GetModulePtr<FLevelEditorModule>(UE::AvaEditorCore::Private::LevelEditorModuleName);
 }
 
 FLevelEditorModule* FAvaLevelEditorUtils::LoadLevelEditorModule()
 {
-	return FModuleManager::Get().LoadModulePtr<FLevelEditorModule>(UE::AvalancheEditorCore::Private::LevelEditorModuleName);
+	return FModuleManager::Get().LoadModulePtr<FLevelEditorModule>(UE::AvaEditorCore::Private::LevelEditorModuleName);
 }
 
 TConstArrayView<FName> FAvaLevelEditorUtils::GetDetailsViewNames()
diff --git a/Engine/Plugins/Experimental/Avalanche/Source/AvalancheEditorCore/Public/IAvaEditor.h b/Engine/Plugins/Experimental/Avalanche/Source/AvalancheEditorCore/Public/IAvaEditor.h
--- a/Engine/Plugins/Experimental/Avalanche/Source/AvalancheEditorCore/Public/IAvaEditor.h
+++ b/Engine/Plugins/Experimental/Avalanche/Source/AvalancheEditorCore/Public/IAvaEditor.h
@@ -30,7 +30,7 @@ enum class EAvaEditorObjectQueryType : uint8
 	CreateIfNotFound,
 };
 
-/** Interface used to Extend an Editor (e.g. Level Editor) for the Avalanche Workflow that extends beyond the limits of Mode UI Layer*/
+/** Interface used to Extend an Editor (e.g. Level Editor) for the Motion Design Workflow that extends beyond the limits of Mode UI Layer*/
 class IAvaEditor : public IAvaTypeCastable, public TSharedFromThis<IAvaEditor>
 {
 public:
diff --git a/Engine/Plugins/Experimental/Avalanche/Source/AvalancheEditorCore/Public/LevelEditor/AvaLevelEditorUtils.h b/Engine/Plugins/Experimental/Avalanche/Source/AvalancheEditorCore/Public/LevelEditor/AvaLevelEditorUtils.h
--- a/Engine/Plugins/Experimental/Avalanche/Source/AvalancheEditorCore/Public/LevelEditor/AvaLevelEditorUtils.h
+++ b/Engine/Plugins/Experimental/Avalanche/Source/AvalancheEditorCore/Public/LevelEditor/AvaLevelEditorUtils.h
@@ -7,7 +7,7 @@
 class FLevelEditorModule;
 class FName;
 
-/** Utility class for the Level Editor and the Avalanche extension editor of it */
+/** Utility class for the Level Editor and the Motion Design extension editor of it */
 class FAvaLevelEditorUtils
 {
 public:
diff --git a/Engine/Plugins/Experimental/Avalanche/Source/AvalancheEffectorsEditor/Private/Cloner/AvaClonerActorVis.h b/Engine/Plugins/Experimental/Avalanche/Source/AvalancheEffectorsEditor/Private/Cloner/AvaClonerActorVis.h
--- a/Engine/Plugins/Experimental/Avalanche/Source/AvalancheEffectorsEditor/Private/Cloner/AvaClonerActorVis.h
+++ b/Engine/Plugins/Experimental/Avalanche/Source/AvalancheEffectorsEditor/Private/Cloner/AvaClonerActorVis.h
@@ -25,8 +25,8 @@ struct HAvaClonerActorSpacingHitProxy : HAvaHitProxy
 class FAvaClonerActorVisualizer : public FAvaVisualizerBase
 {
 public:
-	typedef FAvaVisualizerBase Super;
-	typedef UCEClonerComponent MeshType;
+	using Super = FAvaVisualizerBase;
+	using MeshType = UCEClonerComponent;
 
 	FAvaClonerActorVisualizer();
 
