[Motion Design]

diff --git a/Engine/Plugins/Experimental/Avalanche/Source/AvalancheEffectorsEditor/Private/Effector/AvaEffectorActorVis.h b/Engine/Plugins/Experimental/Avalanche/Source/AvalancheEffectorsEditor/Private/Effector/AvaEffectorActorVis.h
--- a/Engine/Plugins/Experimental/Avalanche/Source/AvalancheEffectorsEditor/Private/Effector/AvaEffectorActorVis.h
+++ b/Engine/Plugins/Experimental/Avalanche/Source/AvalancheEffectorsEditor/Private/Effector/AvaEffectorActorVis.h
@@ -23,7 +23,7 @@ struct HAvaEffectorActorZoneHitProxy : HAvaHitProxy
 class FAvaEffectorActorVisualizer : public FAvaVisualizerBase
 {
 public:
-	typedef FAvaVisualizerBase Super;
+	using Super = FAvaVisualizerBase;
 
 	FAvaEffectorActorVisualizer();
 
diff --git a/Engine/Plugins/Experimental/Avalanche/Source/AvalancheInteractiveTools/Private/AvalancheInteractiveToolsModule.cpp b/Engine/Plugins/Experimental/Avalanche/Source/AvalancheInteractiveTools/Private/AvalancheInteractiveToolsModule.cpp
--- a/Engine/Plugins/Experimental/Avalanche/Source/AvalancheInteractiveTools/Private/AvalancheInteractiveToolsModule.cpp
+++ b/Engine/Plugins/Experimental/Avalanche/Source/AvalancheInteractiveTools/Private/AvalancheInteractiveToolsModule.cpp
@@ -15,7 +15,7 @@
 
 DEFINE_LOG_CATEGORY(LogAvaInteractiveTools);
 
-namespace UE::AvalancheInteractiveTools::Private
+namespace UE::AvaInteractiveTools::Private
 {
 	bool bInitialRegistration = false;
 }
@@ -104,7 +104,7 @@ void FAvalancheInteractiveToolsModule::RegisterTool(FName InCategory, FAvaIntera
 
 	Tools[InCategory].Add(MoveTemp(InToolParams));
 
-	using namespace UE::AvalancheInteractiveTools::Private;
+	using namespace UE::AvaInteractiveTools::Private;
 
 	if (!bInitialRegistration)
 	{
@@ -144,7 +144,7 @@ void FAvalancheInteractiveToolsModule::OnToolDeactivated()
 
 void FAvalancheInteractiveToolsModule::OnPostEngineInit()
 {
-	using namespace UE::AvalancheInteractiveTools::Private;
+	using namespace UE::AvaInteractiveTools::Private;
 
 	IPlacementModeModule& PlacementMode = IPlacementModeModule::Get();
 
diff --git a/Engine/Plugins/Experimental/Avalanche/Source/AvalancheInteractiveTools/Private/EdMode/AvaInteractiveToolsEdMode.cpp b/Engine/Plugins/Experimental/Avalanche/Source/AvalancheInteractiveTools/Private/EdMode/AvaInteractiveToolsEdMode.cpp
--- a/Engine/Plugins/Experimental/Avalanche/Source/AvalancheInteractiveTools/Private/EdMode/AvaInteractiveToolsEdMode.cpp
+++ b/Engine/Plugins/Experimental/Avalanche/Source/AvalancheInteractiveTools/Private/EdMode/AvaInteractiveToolsEdMode.cpp
@@ -23,7 +23,7 @@ UAvaInteractiveToolsEdMode::UAvaInteractiveToolsEdMode()
 	LastActiveTool = "";
 }
 
-namespace UE::AvalancheInteractiveToolsEditorMode::Private
+namespace UE::AvaInteractiveToolsEditorMode::Private
 {
 	static TSet<FName> IncompatibleEdModes = {
 		FBuiltinEditorModes::EM_None,
@@ -39,7 +39,7 @@ namespace UE::AvalancheInteractiveToolsEditorMode::Private
 
 bool UAvaInteractiveToolsEdMode::IsCompatibleWith(FEditorModeID OtherModeID) const
 {
-	return !UE::AvalancheInteractiveToolsEditorMode::Private::IncompatibleEdModes.Contains(OtherModeID);
+	return !UE::AvaInteractiveToolsEditorMode::Private::IncompatibleEdModes.Contains(OtherModeID);
 }
 
 void UAvaInteractiveToolsEdMode::Enter()
diff --git a/Engine/Plugins/Experimental/Avalanche/Source/AvalancheInteractiveTools/Private/Tools/AvaInteractiveToolsToolBase.cpp b/Engine/Plugins/Experimental/Avalanche/Source/AvalancheInteractiveTools/Private/Tools/AvaInteractiveToolsToolBase.cpp
--- a/Engine/Plugins/Experimental/Avalanche/Source/AvalancheInteractiveTools/Private/Tools/AvaInteractiveToolsToolBase.cpp
+++ b/Engine/Plugins/Experimental/Avalanche/Source/AvalancheInteractiveTools/Private/Tools/AvaInteractiveToolsToolBase.cpp
@@ -27,7 +27,7 @@
 
 #define LOCTEXT_NAMESPACE "AvaInteractiveToolsToolBase"
 
-namespace UE::AvalancheInteractiveTools::Private
+namespace UE::AvaInteractiveTools::Private
 {
 	static const FName AvaITFToolPresetMenuName = "AvaITFToolPresetMenu";
 }
@@ -186,7 +186,7 @@ bool UAvaInteractiveToolsToolBase::ConditionalIdentityRotation() const
 				return false;
 			}
 
-			return !IsAvalancheViewport();
+			return !IsMotionDesignViewport();
 
 		default:
 		case EAvaInteractiveToolsDefaultActionAlignment::Axis:
@@ -626,7 +626,7 @@ void UAvaInteractiveToolsToolBase::RequestShutdown(EToolShutdownType InShutdownT
 	FAvalancheInteractiveToolsModule::Get().OnToolDeactivated();
 }
 
-bool UAvaInteractiveToolsToolBase::IsAvalancheViewport() const
+bool UAvaInteractiveToolsToolBase::IsMotionDesignViewport() const
 {
 	if (IToolsContextQueriesAPI* ContextAPI = GetToolManager()->GetContextQueriesAPI())
 	{
@@ -634,7 +634,7 @@ bool UAvaInteractiveToolsToolBase::IsAvalancheViewport() const
 		{
 			if (TSharedPtr<IAvaViewportClient> AvaViewportClient = FAvaViewportUtils::GetAvaViewportClient(Viewport))
 			{
-				return AvaViewportClient->IsAvalancheViewport();
+				return AvaViewportClient->IsMotionDesignViewport();
 			}
 		}
 	}
@@ -661,7 +661,7 @@ bool UAvaInteractiveToolsToolBase::ShouldUsePresetMenu() const
 
 void UAvaInteractiveToolsToolBase::ShowPresetMenu()
 {
-	using namespace UE::AvalancheInteractiveTools::Private;
+	using namespace UE::AvaInteractiveTools::Private;
 
 	RegisterPresetMenu();
 
@@ -709,7 +709,7 @@ void UAvaInteractiveToolsToolBase::OnPresetSelected(TStrongObjectPtr<UAvaInterac
 
 void UAvaInteractiveToolsToolBase::RegisterPresetMenu()
 {
-	using namespace UE::AvalancheInteractiveTools::Private;
+	using namespace UE::AvaInteractiveTools::Private;
 
 	UToolMenus* ToolMenus = UToolMenus::Get();
 
diff --git a/Engine/Plugins/Experimental/Avalanche/Source/AvalancheInteractiveTools/Public/IAvalancheInteractiveToolsModule.h b/Engine/Plugins/Experimental/Avalanche/Source/AvalancheInteractiveTools/Public/IAvalancheInteractiveToolsModule.h
--- a/Engine/Plugins/Experimental/Avalanche/Source/AvalancheInteractiveTools/Public/IAvalancheInteractiveToolsModule.h
+++ b/Engine/Plugins/Experimental/Avalanche/Source/AvalancheInteractiveTools/Public/IAvalancheInteractiveToolsModule.h
@@ -31,7 +31,7 @@ struct FAvaInteractiveToolsToolParameters
 	TMap<FName, TSharedRef<FAvaInteractiveToolsToolPresetBase>> Presets;
 };
 
-namespace UE::AvalancheInteractiveTools
+namespace UE::AvaInteractiveTools
 {
 	using FAvaInteractiveToolsCategoryToolMap = TMap<FName, TArray<FAvaInteractiveToolsToolParameters>>;
 }
diff --git a/Engine/Plugins/Experimental/Avalanche/Source/AvalancheInteractiveTools/Public/Tools/AvaInteractiveToolsToolBase.h b/Engine/Plugins/Experimental/Avalanche/Source/AvalancheInteractiveTools/Public/Tools/AvaInteractiveToolsToolBase.h
--- a/Engine/Plugins/Experimental/Avalanche/Source/AvalancheInteractiveTools/Public/Tools/AvaInteractiveToolsToolBase.h
+++ b/Engine/Plugins/Experimental/Avalanche/Source/AvalancheInteractiveTools/Public/Tools/AvaInteractiveToolsToolBase.h
@@ -184,7 +184,7 @@ protected:
 
 	void SetToolkitSettingsObject(UObject* InObject) const;
 
-	bool IsAvalancheViewport() const;
+	bool IsMotionDesignViewport() const;
 
 	bool ShouldForceDefaultAction() const;
 
diff --git a/Engine/Plugins/Experimental/Avalanche/Source/AvalancheLevelViewport/Private/AvaLevelViewportCommands.cpp b/Engine/Plugins/Experimental/Avalanche/Source/AvalancheLevelViewport/Private/AvaLevelViewportCommands.cpp
--- a/Engine/Plugins/Experimental/Avalanche/Source/AvalancheLevelViewport/Private/AvaLevelViewportCommands.cpp
+++ b/Engine/Plugins/Experimental/Avalanche/Source/AvalancheLevelViewport/Private/AvaLevelViewportCommands.cpp
@@ -153,13 +153,13 @@ void FAvaLevelViewportCommands::RegisterCameraCommands()
 
 	UI_COMMAND(CameraZoomReset
 		, "Reset Camera Zoom"
-		, "Resets the zoom and pan of the Avalanche Viewport camera."
+		, "Resets the zoom and pan of the Motion Design Viewport camera."
 		, EUserInterfaceActionType::Button
 		, FInputChord(EKeys::Home))
 
 	UI_COMMAND(CameraTransformReset
 		, "Reset Camera Transform"
-		, "Resets the location and rotation of the Avalanche Viewport camera."
+		, "Resets the location and rotation of the Motion Design Viewport camera."
 		, EUserInterfaceActionType::Button
 		, FInputChord(EKeys::Home, EModifierKey::Alt))
 
diff --git a/Engine/Plugins/Experimental/Avalanche/Source/AvalancheLevelViewport/Private/AvaLevelViewportLayoutEntity.cpp b/Engine/Plugins/Experimental/Avalanche/Source/AvalancheLevelViewport/Private/AvaLevelViewportLayoutEntity.cpp
--- a/Engine/Plugins/Experimental/Avalanche/Source/AvalancheLevelViewport/Private/AvaLevelViewportLayoutEntity.cpp
+++ b/Engine/Plugins/Experimental/Avalanche/Source/AvalancheLevelViewport/Private/AvaLevelViewportLayoutEntity.cpp
@@ -24,7 +24,7 @@ FAvaLevelViewportLayoutEntity::FAvaLevelViewportLayoutEntity(const FAssetEditorV
 
 FName FAvaLevelViewportLayoutEntity::GetStaticType()
 {
-	static const FName AvaLevelViewportName(TEXT("Avalanche"));
+	static const FName AvaLevelViewportName(TEXT("MotionDesign"));
 	return AvaLevelViewportName;
 }
 
diff --git a/Engine/Plugins/Experimental/Avalanche/Source/AvalancheLevelViewport/Private/AvaLevelViewportModule.cpp b/Engine/Plugins/Experimental/Avalanche/Source/AvalancheLevelViewport/Private/AvaLevelViewportModule.cpp
--- a/Engine/Plugins/Experimental/Avalanche/Source/AvalancheLevelViewport/Private/AvaLevelViewportModule.cpp
+++ b/Engine/Plugins/Experimental/Avalanche/Source/AvalancheLevelViewport/Private/AvaLevelViewportModule.cpp
@@ -11,7 +11,7 @@
 
 DEFINE_LOG_CATEGORY(AvaLevelViewportLog);
 
-namespace UE::AvalancheLevelViewport::Private
+namespace UE::AvaLevelViewport::Private
 {
 	TSharedPtr<IAvaViewportClient> GetAsAvaLevelViewportClient(FEditorViewportClient* InViewportClient)
 	{
@@ -33,7 +33,7 @@ void FAvaLevelViewportModule::StartupModule()
 	UToolMenus::RegisterStartupCallback(FSimpleMulticastDelegate::FDelegate::CreateRaw(this, &FAvaLevelViewportModule::RegisterMenus));
 
 	AvaLevelViewportClientCasterDelegateHandle = FAvaViewportUtils::RegisterViewportClientCaster(
-		&UE::AvalancheLevelViewport::Private::GetAsAvaLevelViewportClient
+		&UE::AvaLevelViewport::Private::GetAsAvaLevelViewportClient
 	);
 }
 
diff --git a/Engine/Plugins/Experimental/Avalanche/Source/AvalancheLevelViewport/Private/AvaLevelViewportStyle.cpp b/Engine/Plugins/Experimental/Avalanche/Source/AvalancheLevelViewport/Private/AvaLevelViewportStyle.cpp
--- a/Engine/Plugins/Experimental/Avalanche/Source/AvalancheLevelViewport/Private/AvaLevelViewportStyle.cpp
+++ b/Engine/Plugins/Experimental/Avalanche/Source/AvalancheLevelViewport/Private/AvaLevelViewportStyle.cpp
@@ -56,55 +56,55 @@ FAvaLevelViewportStyle::FAvaLevelViewportStyle()
 	Set("Button.PivotMode", new IMAGE_BRUSH("Icons/EditorIcons/Pivot_Mode", Icon25x25));
 
 	// Alignment
-	Set("AvalancheIcons.Alignment.Translation.TopLeft",     new IMAGE_BRUSH("Icons/DetailsPanelIcons/TopLeft",     Icon16x16));
-	Set("AvalancheIcons.Alignment.Translation.Top",         new IMAGE_BRUSH("Icons/DetailsPanelIcons/Top",         Icon16x16));
-	Set("AvalancheIcons.Alignment.Translation.TopRight",    new IMAGE_BRUSH("Icons/DetailsPanelIcons/TopRight",    Icon16x16));
-	Set("AvalancheIcons.Alignment.Translation.Left",        new IMAGE_BRUSH("Icons/DetailsPanelIcons/Left",        Icon16x16));
-	Set("AvalancheIcons.Alignment.Translation.Center",      new IMAGE_BRUSH("Icons/DetailsPanelIcons/Center",      Icon16x16));
-	Set("AvalancheIcons.Alignment.Translation.Right",       new IMAGE_BRUSH("Icons/DetailsPanelIcons/Right",       Icon16x16));
-	Set("AvalancheIcons.Alignment.Translation.BottomLeft",  new IMAGE_BRUSH("Icons/DetailsPanelIcons/BottomLeft",  Icon16x16));
-	Set("AvalancheIcons.Alignment.Translation.Bottom",      new IMAGE_BRUSH("Icons/DetailsPanelIcons/Bottom",      Icon16x16));
-	Set("AvalancheIcons.Alignment.Translation.BottomRight", new IMAGE_BRUSH("Icons/DetailsPanelIcons/BottomRight", Icon16x16));
+	Set("Icons.Alignment.Translation.TopLeft",     new IMAGE_BRUSH("Icons/DetailsPanelIcons/TopLeft",     Icon16x16));
+	Set("Icons.Alignment.Translation.Top",         new IMAGE_BRUSH("Icons/DetailsPanelIcons/Top",         Icon16x16));
+	Set("Icons.Alignment.Translation.TopRight",    new IMAGE_BRUSH("Icons/DetailsPanelIcons/TopRight",    Icon16x16));
+	Set("Icons.Alignment.Translation.Left",        new IMAGE_BRUSH("Icons/DetailsPanelIcons/Left",        Icon16x16));
+	Set("Icons.Alignment.Translation.Center",      new IMAGE_BRUSH("Icons/DetailsPanelIcons/Center",      Icon16x16));
+	Set("Icons.Alignment.Translation.Right",       new IMAGE_BRUSH("Icons/DetailsPanelIcons/Right",       Icon16x16));
+	Set("Icons.Alignment.Translation.BottomLeft",  new IMAGE_BRUSH("Icons/DetailsPanelIcons/BottomLeft",  Icon16x16));
+	Set("Icons.Alignment.Translation.Bottom",      new IMAGE_BRUSH("Icons/DetailsPanelIcons/Bottom",      Icon16x16));
+	Set("Icons.Alignment.Translation.BottomRight", new IMAGE_BRUSH("Icons/DetailsPanelIcons/BottomRight", Icon16x16));
 	
-	Set("AvalancheIcons.Alignment.Translation.Back",        new IMAGE_BRUSH_SVG("Icons/DetailsPanelIcons/StackBack",   Icon22x22));
-	Set("AvalancheIcons.Alignment.Translation.Center_X",    new IMAGE_BRUSH_SVG("Icons/DetailsPanelIcons/StackCenter", Icon22x22));
-	Set("AvalancheIcons.Alignment.Translation.Front",       new IMAGE_BRUSH_SVG("Icons/DetailsPanelIcons/StackFront",  Icon22x22));
-
-	Set("AvalancheIcons.Alignment.Center_YZ",   new IMAGE_BRUSH_SVG("Icons/PaletteIcons/AlignCenterVertAndHoriz", Icon22x22));
-	Set("AvalancheIcons.Alignment.Left",        new IMAGE_BRUSH_SVG("Icons/PaletteIcons/AlignLeft",               Icon22x22));
-	Set("AvalancheIcons.Alignment.Center_Y",    new IMAGE_BRUSH_SVG("Icons/PaletteIcons/AlignCenterHoriz",        Icon22x22));
-	Set("AvalancheIcons.Alignment.Right",       new IMAGE_BRUSH_SVG("Icons/PaletteIcons/AlignRight",              Icon22x22));
-	Set("AvalancheIcons.Alignment.Top",         new IMAGE_BRUSH_SVG("Icons/PaletteIcons/AlignTop",                Icon22x22));
-	Set("AvalancheIcons.Alignment.Center_Z",    new IMAGE_BRUSH_SVG("Icons/PaletteIcons/AlignCenterVert",         Icon22x22));
-	Set("AvalancheIcons.Alignment.Bottom",      new IMAGE_BRUSH_SVG("Icons/PaletteIcons/AlignBottom",             Icon22x22));
-	Set("AvalancheIcons.Alignment.DistributeX", new IMAGE_BRUSH_SVG("Icons/PaletteIcons/DistributeDepth",         Icon22x22));
-	Set("AvalancheIcons.Alignment.DistributeY", new IMAGE_BRUSH_SVG("Icons/PaletteIcons/DistributeHorizontal",    Icon22x22));
-	Set("AvalancheIcons.Alignment.DistributeZ", new IMAGE_BRUSH_SVG("Icons/PaletteIcons/DistributeVertical",      Icon22x22));
+	Set("Icons.Alignment.Translation.Back",        new IMAGE_BRUSH_SVG("Icons/DetailsPanelIcons/StackBack",   Icon22x22));
+	Set("Icons.Alignment.Translation.Center_X",    new IMAGE_BRUSH_SVG("Icons/DetailsPanelIcons/StackCenter", Icon22x22));
+	Set("Icons.Alignment.Translation.Front",       new IMAGE_BRUSH_SVG("Icons/DetailsPanelIcons/StackFront",  Icon22x22));
+
+	Set("Icons.Alignment.Center_YZ",   new IMAGE_BRUSH_SVG("Icons/PaletteIcons/AlignCenterVertAndHoriz", Icon22x22));
+	Set("Icons.Alignment.Left",        new IMAGE_BRUSH_SVG("Icons/PaletteIcons/AlignLeft",               Icon22x22));
+	Set("Icons.Alignment.Center_Y",    new IMAGE_BRUSH_SVG("Icons/PaletteIcons/AlignCenterHoriz",        Icon22x22));
+	Set("Icons.Alignment.Right",       new IMAGE_BRUSH_SVG("Icons/PaletteIcons/AlignRight",              Icon22x22));
+	Set("Icons.Alignment.Top",         new IMAGE_BRUSH_SVG("Icons/PaletteIcons/AlignTop",                Icon22x22));
+	Set("Icons.Alignment.Center_Z",    new IMAGE_BRUSH_SVG("Icons/PaletteIcons/AlignCenterVert",         Icon22x22));
+	Set("Icons.Alignment.Bottom",      new IMAGE_BRUSH_SVG("Icons/PaletteIcons/AlignBottom",             Icon22x22));
+	Set("Icons.Alignment.DistributeX", new IMAGE_BRUSH_SVG("Icons/PaletteIcons/DistributeDepth",         Icon22x22));
+	Set("Icons.Alignment.DistributeY", new IMAGE_BRUSH_SVG("Icons/PaletteIcons/DistributeHorizontal",    Icon22x22));
+	Set("Icons.Alignment.DistributeZ", new IMAGE_BRUSH_SVG("Icons/PaletteIcons/DistributeVertical",      Icon22x22));
 	
-	Set("AvalancheIcons.Alignment.Rotation.Actor.Roll",   new IMAGE_BRUSH_SVG("Icons/PaletteIcons/RotationAlignActorRoll",   Icon22x22));
-	Set("AvalancheIcons.Alignment.Rotation.Actor.Pitch",  new IMAGE_BRUSH_SVG("Icons/PaletteIcons/RotationAlignActorPitch",  Icon22x22));
-	Set("AvalancheIcons.Alignment.Rotation.Actor.Yaw",    new IMAGE_BRUSH_SVG("Icons/PaletteIcons/RotationAlignActorYaw",    Icon22x22));
-	Set("AvalancheIcons.Alignment.Rotation.Actor.All",    new IMAGE_BRUSH_SVG("Icons/PaletteIcons/RotationAlignActorAll",    Icon22x22));
-	Set("AvalancheIcons.Alignment.Rotation.Camera.Roll",  new IMAGE_BRUSH_SVG("Icons/PaletteIcons/RotationAlignCameraRoll",  Icon22x22));
-	Set("AvalancheIcons.Alignment.Rotation.Camera.Pitch", new IMAGE_BRUSH_SVG("Icons/PaletteIcons/RotationAlignCameraPitch", Icon22x22));
-	Set("AvalancheIcons.Alignment.Rotation.Camera.Yaw",   new IMAGE_BRUSH_SVG("Icons/PaletteIcons/RotationAlignCameraYaw",   Icon22x22));
-	Set("AvalancheIcons.Alignment.Rotation.Camera.All",   new IMAGE_BRUSH_SVG("Icons/PaletteIcons/RotationAlignCameraAll",   Icon22x22));
+	Set("Icons.Alignment.Rotation.Actor.Roll",   new IMAGE_BRUSH_SVG("Icons/PaletteIcons/RotationAlignActorRoll",   Icon22x22));
+	Set("Icons.Alignment.Rotation.Actor.Pitch",  new IMAGE_BRUSH_SVG("Icons/PaletteIcons/RotationAlignActorPitch",  Icon22x22));
+	Set("Icons.Alignment.Rotation.Actor.Yaw",    new IMAGE_BRUSH_SVG("Icons/PaletteIcons/RotationAlignActorYaw",    Icon22x22));
+	Set("Icons.Alignment.Rotation.Actor.All",    new IMAGE_BRUSH_SVG("Icons/PaletteIcons/RotationAlignActorAll",    Icon22x22));
+	Set("Icons.Alignment.Rotation.Camera.Roll",  new IMAGE_BRUSH_SVG("Icons/PaletteIcons/RotationAlignCameraRoll",  Icon22x22));
+	Set("Icons.Alignment.Rotation.Camera.Pitch", new IMAGE_BRUSH_SVG("Icons/PaletteIcons/RotationAlignCameraPitch", Icon22x22));
+	Set("Icons.Alignment.Rotation.Camera.Yaw",   new IMAGE_BRUSH_SVG("Icons/PaletteIcons/RotationAlignCameraYaw",   Icon22x22));
+	Set("Icons.Alignment.Rotation.Camera.All",   new IMAGE_BRUSH_SVG("Icons/PaletteIcons/RotationAlignCameraAll",   Icon22x22));
 
 	// Screen Icons
-	Set("AvalancheIcons.Screen.SizeToScreen",        new IMAGE_BRUSH_SVG("Icons/PaletteIcons/SizeToScreen",        Icon22x22));
-	Set("AvalancheIcons.Screen.SizeToScreenStretch", new IMAGE_BRUSH_SVG("Icons/PaletteIcons/SizeToScreenStretch", Icon22x22));
-	Set("AvalancheIcons.Screen.FitToScreen",         new IMAGE_BRUSH_SVG("Icons/PaletteIcons/FitToScreen",         Icon22x22));
+	Set("Icons.Screen.SizeToScreen",        new IMAGE_BRUSH_SVG("Icons/PaletteIcons/SizeToScreen",        Icon22x22));
+	Set("Icons.Screen.SizeToScreenStretch", new IMAGE_BRUSH_SVG("Icons/PaletteIcons/SizeToScreenStretch", Icon22x22));
+	Set("Icons.Screen.FitToScreen",         new IMAGE_BRUSH_SVG("Icons/PaletteIcons/FitToScreen",         Icon22x22));
 
 	// Color picker icons
-	Set("AvalancheIcons.ColorPicker.SolidColors",    new IMAGE_BRUSH("Icons/EditorIcons/SolidColors",    Icon20x20));
-	Set("AvalancheIcons.ColorPicker.LinearGradient", new IMAGE_BRUSH("Icons/EditorIcons/LinearGradient", Icon20x20));
+	Set("Icons.ColorPicker.SolidColors",    new IMAGE_BRUSH("Icons/EditorIcons/SolidColors",    Icon20x20));
+	Set("Icons.ColorPicker.LinearGradient", new IMAGE_BRUSH("Icons/EditorIcons/LinearGradient", Icon20x20));
 
 	// Post process icons
-	Set("AvalancheIcons.PostProcess.RGB",   new IMAGE_BRUSH("Icons/ViewportIcons/RGBSquare", Icon20x20));
-	Set("AvalancheIcons.PostProcess.Red",   new IMAGE_BRUSH("Icons/ViewportIcons/RedSquare",   Icon20x20));
-	Set("AvalancheIcons.PostProcess.Green", new IMAGE_BRUSH("Icons/ViewportIcons/GreenSquare", Icon20x20));
-	Set("AvalancheIcons.PostProcess.Blue",  new IMAGE_BRUSH("Icons/ViewportIcons/BlueSquare",  Icon20x20));
-	Set("AvalancheIcons.PostProcess.Alpha", new IMAGE_BRUSH("Icons/ViewportIcons/AlphaSquare", Icon20x20));
+	Set("Icons.PostProcess.RGB",   new IMAGE_BRUSH("Icons/ViewportIcons/RGBSquare", Icon20x20));
+	Set("Icons.PostProcess.Red",   new IMAGE_BRUSH("Icons/ViewportIcons/RedSquare",   Icon20x20));
+	Set("Icons.PostProcess.Green", new IMAGE_BRUSH("Icons/ViewportIcons/GreenSquare", Icon20x20));
+	Set("Icons.PostProcess.Blue",  new IMAGE_BRUSH("Icons/ViewportIcons/BlueSquare",  Icon20x20));
+	Set("Icons.PostProcess.Alpha", new IMAGE_BRUSH("Icons/ViewportIcons/AlphaSquare", Icon20x20));
 
 	// StatusBar ToolMenu
 	{
diff --git a/Engine/Plugins/Experimental/Avalanche/Source/AvalancheLevelViewport/Private/AvaViewportCameraHistory.cpp b/Engine/Plugins/Experimental/Avalanche/Source/AvalancheLevelViewport/Private/AvaViewportCameraHistory.cpp
--- a/Engine/Plugins/Experimental/Avalanche/Source/AvalancheLevelViewport/Private/AvaViewportCameraHistory.cpp
+++ b/Engine/Plugins/Experimental/Avalanche/Source/AvalancheLevelViewport/Private/AvaViewportCameraHistory.cpp
@@ -12,7 +12,7 @@
 
 #define LOCTEXT_NAMESPACE "AvaViewportCameraHistory"
 
-namespace UE::AvalancheEditor::Private
+namespace UE::AvaEditor::Private
 {
 	static FText UndoRedoMessage = LOCTEXT("UndoRedoMessage", "Camera Movement");
 
@@ -93,7 +93,7 @@ void FAvaViewportCameraHistory::PostEngineInit()
 void FAvaViewportCameraHistory::Reset()
 {
 	CameraTransformHistory.Reset();
-	CameraTransformHistory.SetNumUninitialized(UE::AvalancheEditor::Internal::CameraUndoHistoryCapacity, EAllowShrinking::No);
+	CameraTransformHistory.SetNumUninitialized(UE::AvaEditor::Internal::CameraUndoHistoryCapacity, EAllowShrinking::No);
 	CameraTransformHistoryIndex = INDEX_NONE;
 	CameraTransformHistoryHeadIndex = 1;
 }
@@ -147,12 +147,12 @@ void FAvaViewportCameraHistory::OnEndCameraTransform(UObject& InCameraObject)
 	if (AActor* CameraActor = Cast<AActor>(&InCameraObject))
 	{
 		// New item replaces next
-		int32 CameraTransformIndexToAdd = UE::AvalancheEditor::Private::WrapIndex(CameraTransformHistoryIndex + 1, CameraTransformHistory.Max());
+		int32 CameraTransformIndexToAdd = UE::AvaEditor::Private::WrapIndex(CameraTransformHistoryIndex + 1, CameraTransformHistory.Max());
 
 		CameraTransformHistory[CameraTransformIndexToAdd] = { CameraActor, CameraActor->GetActorTransform() };
 		
 		CameraTransformHistoryIndex = CameraTransformIndexToAdd;
-		CameraTransformHistoryHeadIndex = UE::AvalancheEditor::Private::WrapIndex(CameraTransformHistoryIndex + 1, CameraTransformHistory.Max());
+		CameraTransformHistoryHeadIndex = UE::AvaEditor::Private::WrapIndex(CameraTransformHistoryIndex + 1, CameraTransformHistory.Max());
 
 		UE_LOG(LogAva, Verbose, TEXT("Saved camera transform at %i, head %i, %s"), CameraTransformHistoryIndex, CameraTransformHistoryHeadIndex, *CameraActor->GetActorTransform().Rotator().ToString());
 	}
@@ -160,7 +160,7 @@ void FAvaViewportCameraHistory::OnEndCameraTransform(UObject& InCameraObject)
 
 void FAvaViewportCameraHistory::ExecuteCameraTransformUndo()
 {
-	const int32 CameraTransformIndexToRestore = UE::AvalancheEditor::Private::WrapIndex(CameraTransformHistoryIndex - 1, CameraTransformHistory.Max());
+	const int32 CameraTransformIndexToRestore = UE::AvaEditor::Private::WrapIndex(CameraTransformHistoryIndex - 1, CameraTransformHistory.Max());
 	if (CameraTransformIndexToRestore == CameraTransformHistoryHeadIndex)
 	{
 		return;
@@ -179,7 +179,7 @@ void FAvaViewportCameraHistory::ExecuteCameraTransformUndo()
 
 void FAvaViewportCameraHistory::ExecuteCameraTransformRedo()
 {
-	const int32 CameraTransformIndexToRestore = UE::AvalancheEditor::Private::WrapIndex(CameraTransformHistoryIndex + 1, CameraTransformHistory.Max());
+	const int32 CameraTransformIndexToRestore = UE::AvaEditor::Private::WrapIndex(CameraTransformHistoryIndex + 1, CameraTransformHistory.Max());
 	if (CameraTransformIndexToRestore == CameraTransformHistoryHeadIndex)
 	{
 		return;
@@ -199,13 +199,13 @@ void FAvaViewportCameraHistory::ExecuteCameraTransformRedo()
 void FAvaViewportCameraHistory::NotifyUndo()
 {
 	const FText UndoMessage = NSLOCTEXT("UnrealEd", "UndoMessageFormat", "Undo: {0}");
-	Notify(FText::Format(UndoMessage, UE::AvalancheEditor::Private::UndoRedoMessage));
+	Notify(FText::Format(UndoMessage, UE::AvaEditor::Private::UndoRedoMessage));
 }
 
 void FAvaViewportCameraHistory::NotifyRedo()
 {
 	const FText RedoMessage = NSLOCTEXT("UnrealEd", "RedoMessageFormat", "Redo: {0}");
-	Notify(FText::Format(RedoMessage, UE::AvalancheEditor::Private::UndoRedoMessage));
+	Notify(FText::Format(RedoMessage, UE::AvaEditor::Private::UndoRedoMessage));
 }
 
 // @see: UEditorEngine::ShowUndoRedoNotification
diff --git a/Engine/Plugins/Experimental/Avalanche/Source/AvalancheLevelViewport/Private/AvaViewportCameraHistory.h b/Engine/Plugins/Experimental/Avalanche/Source/AvalancheLevelViewport/Private/AvaViewportCameraHistory.h
--- a/Engine/Plugins/Experimental/Avalanche/Source/AvalancheLevelViewport/Private/AvaViewportCameraHistory.h
+++ b/Engine/Plugins/Experimental/Avalanche/Source/AvalancheLevelViewport/Private/AvaViewportCameraHistory.h
@@ -9,7 +9,7 @@
 #include "Templates/SharedPointer.h"
 #include "UObject/WeakObjectPtr.h"
 
-namespace UE::AvalancheEditor::Internal
+namespace UE::AvaEditor::Internal
 {
 	static constexpr int32 CameraUndoHistoryCapacity = 64;
 }
@@ -56,7 +56,7 @@ private:
 	TSharedPtr<SNotificationItem> UndoRedoNotificationItem;
 
 	// Stores previous camera transforms (as camera, transform pair) for undo/redo, if any
-	TArray<TPair<TWeakObjectPtr<AActor>, FTransform>, TFixedAllocator<UE::AvalancheEditor::Internal::CameraUndoHistoryCapacity>> CameraTransformHistory;
+	TArray<TPair<TWeakObjectPtr<AActor>, FTransform>, TFixedAllocator<UE::AvaEditor::Internal::CameraUndoHistoryCapacity>> CameraTransformHistory;
 
 	// Stores current position in the camera transform history, to allow redo of actions ahead of this position
 	int32 CameraTransformHistoryIndex = INDEX_NONE;
diff --git a/Engine/Plugins/Experimental/Avalanche/Source/AvalancheLevelViewport/Private/ColorPicker/AvaViewportColorPickerActorClassRegistry.cpp b/Engine/Plugins/Experimental/Avalanche/Source/AvalancheLevelViewport/Private/ColorPicker/AvaViewportColorPickerActorClassRegistry.cpp
--- a/Engine/Plugins/Experimental/Avalanche/Source/AvalancheLevelViewport/Private/ColorPicker/AvaViewportColorPickerActorClassRegistry.cpp
+++ b/Engine/Plugins/Experimental/Avalanche/Source/AvalancheLevelViewport/Private/ColorPicker/AvaViewportColorPickerActorClassRegistry.cpp
@@ -7,7 +7,7 @@
 #include "GameFramework/Actor.h"
 #include "Templates/SharedPointer.h"
 
-namespace UE::AvalancheLevelViewport::Private
+namespace UE::AvaLevelViewport::Private
 {
 	TMap<UClass*, TSharedRef<IAvaViewportColorPickerAdapter>> Adapters;
 
@@ -33,13 +33,13 @@ namespace UE::AvalancheLevelViewport::Private
 
 void FAvaViewportColorPickerActorClassRegistry::RegisterClassAdapter(UClass* InClass, const TSharedRef<IAvaViewportColorPickerAdapter>& InAdapter)
 {
-	using namespace UE::AvalancheLevelViewport;
+	using namespace UE::AvaLevelViewport;
 	Private::Adapters.Add(InClass, InAdapter);
 }
 
 bool FAvaViewportColorPickerActorClassRegistry::GetColorDataFromActor(const AActor* InActor, FAvaColorChangeData& OutColorData)
 {
-	using namespace UE::AvalancheLevelViewport;
+	using namespace UE::AvaLevelViewport;
 
 	if (TSharedPtr<IAvaViewportColorPickerAdapter> Adapter = Private::FindAdapterForClass(InActor))
 	{
@@ -51,7 +51,7 @@ bool FAvaViewportColorPickerActorClassRegistry::GetColorDataFromActor(const AAct
 
 bool FAvaViewportColorPickerActorClassRegistry::ApplyColorDataToActor(AActor* InActor, const FAvaColorChangeData& InColorData)
 {
-	using namespace UE::AvalancheLevelViewport;
+	using namespace UE::AvaLevelViewport;
 
 	if (TSharedPtr<IAvaViewportColorPickerAdapter> Adapter = Private::FindAdapterForClass(InActor))
 	{
diff --git a/Engine/Plugins/Experimental/Avalanche/Source/AvalancheLevelViewport/Private/SAvaLevelViewport.cpp b/Engine/Plugins/Experimental/Avalanche/Source/AvalancheLevelViewport/Private/SAvaLevelViewport.cpp
--- a/Engine/Plugins/Experimental/Avalanche/Source/AvalancheLevelViewport/Private/SAvaLevelViewport.cpp
+++ b/Engine/Plugins/Experimental/Avalanche/Source/AvalancheLevelViewport/Private/SAvaLevelViewport.cpp
@@ -34,7 +34,7 @@
 
 #define LOCTEXT_NAMESPACE "SAvaLevelViewport"
 
-namespace UE::AvalancheLevelViewport::Private
+namespace UE::AvaLevelViewport::Private
 {
 	static const FText GuideTransactionTitle = LOCTEXT("UpdateGuides", "Update Guides");
 }
@@ -540,7 +540,7 @@ void SAvaLevelViewport::SaveGuides()
 		return;
 	}
 
-	FScopedTransaction SaveTransaction(UE::AvalancheLevelViewport::Private::GuideTransactionTitle);
+	FScopedTransaction SaveTransaction(UE::AvaLevelViewport::Private::GuideTransactionTitle);
 	DataSubsystem->ModifyDataSource();
 
 	TArray<FAvaViewportGuideInfo> GuideInfos;
@@ -567,7 +567,7 @@ void SAvaLevelViewport::OnViewportDataProxyChanged()
 bool SAvaLevelViewport::MatchesContext(const FTransactionContext& InContext,
 	const TArray<TPair<UObject*, FTransactionObjectEvent>>& InTransactionObjectContexts) const
 {
-	return InContext.Title.EqualTo(UE::AvalancheLevelViewport::Private::GuideTransactionTitle);
+	return InContext.Title.EqualTo(UE::AvaLevelViewport::Private::GuideTransactionTitle);
 }
 
 void SAvaLevelViewport::PostUndo(bool bInSuccess)
diff --git a/Engine/Plugins/Experimental/Avalanche/Source/AvalancheLevelViewport/Private/SAvaLevelViewportFrame.cpp b/Engine/Plugins/Experimental/Avalanche/Source/AvalancheLevelViewport/Private/SAvaLevelViewportFrame.cpp
--- a/Engine/Plugins/Experimental/Avalanche/Source/AvalancheLevelViewport/Private/SAvaLevelViewportFrame.cpp
+++ b/Engine/Plugins/Experimental/Avalanche/Source/AvalancheLevelViewport/Private/SAvaLevelViewportFrame.cpp
@@ -11,7 +11,7 @@
 #include "Widgets/SAvaLevelViewportRuler.h"
 #include "Widgets/SAvaLevelViewportStatusBar.h"
 
-namespace UE::AvalancheLevelViewport::Private
+namespace UE::AvaLevelViewport::Private
 {
 	static constexpr float RulerSize = 20.f;
 	static constexpr float StatusBarSize = 26.f;
@@ -57,7 +57,7 @@ FAvaLevelViewportGuideFrameClientAndWidget::FAvaLevelViewportGuideFrameClientAnd
 void SAvaLevelViewportFrame::Construct(const FArguments& InArgs, const FAssetEditorViewportConstructionArgs& InViewportArgs,
 	TSharedPtr<ILevelEditor> InLevelEditor)
 {
-	using namespace UE::AvalancheLevelViewport::Private;
+	using namespace UE::AvaLevelViewport::Private;
 
 	ViewportClient = MakeShared<FAvaLevelViewportClient>();
 	ViewportClient->Init();
@@ -177,7 +177,7 @@ void SAvaLevelViewportFrame::Tick(const FGeometry& AllottedGeometry, const doubl
 
 FMargin SAvaLevelViewportFrame::GetHorizontalRulerOffset() const
 {
-	using namespace UE::AvalancheLevelViewport::Private;
+	using namespace UE::AvaLevelViewport::Private;
 
 	FMargin Offset = {RulerSize + 1, 0, 0, RulerSize};
 
@@ -195,7 +195,7 @@ FMargin SAvaLevelViewportFrame::GetHorizontalRulerOffset() const
 
 FMargin SAvaLevelViewportFrame::GetVerticalRulerOffset() const
 {
-	using namespace UE::AvalancheLevelViewport::Private;
+	using namespace UE::AvaLevelViewport::Private;
 
 	FMargin Offset = {0, RulerSize + 1, RulerSize, StatusBarSize + 1};
 
diff --git a/Engine/Plugins/Experimental/Avalanche/Source/AvalancheLevelViewport/Private/SAvaLevelViewport_Menus.cpp b/Engine/Plugins/Experimental/Avalanche/Source/AvalancheLevelViewport/Private/SAvaLevelViewport_Menus.cpp
--- a/Engine/Plugins/Experimental/Avalanche/Source/AvalancheLevelViewport/Private/SAvaLevelViewport_Menus.cpp
+++ b/Engine/Plugins/Experimental/Avalanche/Source/AvalancheLevelViewport/Private/SAvaLevelViewport_Menus.cpp
@@ -32,7 +32,7 @@
 
 #define LOCTEXT_NAMESPACE "SAvaLevelViewport"
 
-namespace UE::AvalancheLevelViewport::Private
+namespace UE::AvaLevelViewport::Private
 {
 	static constexpr uint32 MaxCameraEntries = 10;
 	static constexpr uint32 MaxViewportEntries = 4;
@@ -74,7 +74,7 @@ TSharedPtr<SWidget> SAvaLevelViewport::MakeViewportToolbar()
 	{
 		// Replace the Camera Menu with a custom one stripping off the View Entries (Perspective, Front, Left, etc...)
 		static const FName CameraMenuName = TEXT("EditorViewportToolBar.CameraMenu");
-		if (FSlotBase* const FoundSlot = UE::AvalancheLevelViewport::Private::FindChildSlotWithTag(CameraMenuName, ToolbarFrame.ToSharedRef()))
+		if (FSlotBase* const FoundSlot = UE::AvaLevelViewport::Private::FindChildSlotWithTag(CameraMenuName, ToolbarFrame.ToSharedRef()))
 		{
 			TSharedPtr<SViewportToolBar> ParentToolBar;
 			{
@@ -179,7 +179,7 @@ void SAvaLevelViewport::FillCameraMenu(UToolMenu* InMenu)
 
 	static const FText CamerasLabel = LOCTEXT("CameraActors", "Placed Cameras");
 
-	if (CameraActors.Num() > UE::AvalancheLevelViewport::Private::MaxCameraEntries)
+	if (CameraActors.Num() > UE::AvaLevelViewport::Private::MaxCameraEntries)
 	{
 		FToolMenuSection& Section = InMenu->AddSection("CameraActors");
 		Section.AddSubMenu("CameraActors"
@@ -244,7 +244,7 @@ void SAvaLevelViewport::FillCameraMenu(UToolMenu* InMenu)
 
 	static const FText ViewportTypesLabel = LOCTEXT("ViewportTypes", "Viewport Type");
 
-	if (ViewportTypeCount > UE::AvalancheLevelViewport::Private::MaxViewportEntries)
+	if (ViewportTypeCount > UE::AvaLevelViewport::Private::MaxViewportEntries)
 	{
 		FToolMenuSection& Section = InMenu->AddSection("ViewportTypes");
 		Section.AddSubMenu("ViewportTypes"
@@ -382,7 +382,7 @@ void SAvaLevelViewport::AddVirtualSizeDefaultEntries(FToolMenuSection& InSection
 
 void SAvaLevelViewport::AddVirtualSizeSizeSettings(FToolMenuSection& InSection)
 {
-	using namespace UE::AvalancheLevelViewport::Private;
+	using namespace UE::AvaLevelViewport::Private;
 
 	const FAvaLevelViewportCommands& LevelViewportCommands = FAvaLevelViewportCommands::Get();
 
@@ -1039,7 +1039,7 @@ void SAvaLevelViewport::OnBackgroundOpacityCommitted(float InValue, ETextCommit:
 
 void SAvaLevelViewport::OnVirtualSizeComponentSliderCommitted(int32 InNewDimension, ETextCommit::Type InCommitType, EAxis::Type InAxis)
 {
-	using namespace UE::AvalancheLevelViewport::Private;
+	using namespace UE::AvaLevelViewport::Private;
 
 	InNewDimension = FMath::Clamp(InNewDimension, VirtualSizeComponentMin, VirtualSizeComponentMax);
 
@@ -1115,7 +1115,7 @@ void SAvaLevelViewport::OnVirtualSizeAspectRatioCommitted(float InNewAspectRatio
 		return;
 	}
 
-	using namespace UE::AvalancheLevelViewport::Private;
+	using namespace UE::AvaLevelViewport::Private;
 
 	InNewAspectRatio = FMath::Clamp(InNewAspectRatio, VirtualSizeAspectRatioMin, VirtualSizeAspectRatioMax);
 
diff --git a/Engine/Plugins/Experimental/Avalanche/Source/AvalancheLevelViewport/Private/ViewportClient/AvaLevelViewportClient.cpp b/Engine/Plugins/Experimental/Avalanche/Source/AvalancheLevelViewport/Private/ViewportClient/AvaLevelViewportClient.cpp
--- a/Engine/Plugins/Experimental/Avalanche/Source/AvalancheLevelViewport/Private/ViewportClient/AvaLevelViewportClient.cpp
+++ b/Engine/Plugins/Experimental/Avalanche/Source/AvalancheLevelViewport/Private/ViewportClient/AvaLevelViewportClient.cpp
@@ -30,14 +30,14 @@
 
 #define LOCTEXT_NAMESPACE "AvaLevelViewportClient"
 
-namespace UE::AvalancheLevelViewport::Private
+namespace UE::AvaLevelViewport::Private
 {
 	static TSet<FEditorViewportClient*> ViewportClientRegistry;
 }
 
 bool FAvaLevelViewportClient::IsAvaLevelViewportClient(const FEditorViewportClient* InViewportClient)
 {
-	return InViewportClient && UE::AvalancheLevelViewport::Private::ViewportClientRegistry.Contains(InViewportClient);
+	return InViewportClient && UE::AvaLevelViewport::Private::ViewportClientRegistry.Contains(InViewportClient);
 }
 
 FAvaLevelViewportClient::FAvaLevelViewportClient()
@@ -53,12 +53,12 @@ FAvaLevelViewportClient::FAvaLevelViewportClient()
 	BoundingBoxVisualizer = FAvaViewportBoundingBoxVisualizerProvider::CreateVisualizer();
 	bUseControllingActorViewInfo = true;
 
-	UE::AvalancheLevelViewport::Private::ViewportClientRegistry.Add(this);
+	UE::AvaLevelViewport::Private::ViewportClientRegistry.Add(this);
 }
 
 FAvaLevelViewportClient::~FAvaLevelViewportClient()
 {
-	UE::AvalancheLevelViewport::Private::ViewportClientRegistry.Remove(this);
+	UE::AvaLevelViewport::Private::ViewportClientRegistry.Remove(this);
 }
 
 void FAvaLevelViewportClient::Init()
diff --git a/Engine/Plugins/Experimental/Avalanche/Source/AvalancheLevelViewport/Private/ViewportClient/AvaLevelViewportClient.h b/Engine/Plugins/Experimental/Avalanche/Source/AvalancheLevelViewport/Private/ViewportClient/AvaLevelViewportClient.h
--- a/Engine/Plugins/Experimental/Avalanche/Source/AvalancheLevelViewport/Private/ViewportClient/AvaLevelViewportClient.h
+++ b/Engine/Plugins/Experimental/Avalanche/Source/AvalancheLevelViewport/Private/ViewportClient/AvaLevelViewportClient.h
@@ -82,7 +82,7 @@ public:
 	//~ End IAvaViewportWorldCoordinateConverter
 
 	//~ Begin IAvaViewportClient
-	virtual bool IsAvalancheViewport() const override { return true; }
+	virtual bool IsMotionDesignViewport() const override { return true; }
 	virtual const FEditorViewportClient* AsEditorViewportClient() const override { return this; }
 	virtual bool SupportsZoom() const override { return true; }
 	virtual float GetZoomedFOV() const override;
diff --git a/Engine/Plugins/Experimental/Avalanche/Source/AvalancheLevelViewport/Private/Widgets/Engine/SRuler.cpp b/Engine/Plugins/Experimental/Avalanche/Source/AvalancheLevelViewport/Private/Widgets/Engine/SRuler.cpp
--- a/Engine/Plugins/Experimental/Avalanche/Source/AvalancheLevelViewport/Private/Widgets/Engine/SRuler.cpp
+++ b/Engine/Plugins/Experimental/Avalanche/Source/AvalancheLevelViewport/Private/Widgets/Engine/SRuler.cpp
@@ -26,7 +26,7 @@ struct FSlateBrush;
 
 #define LOCTEXT_NAMESPACE "EngineSRuler"
 
-namespace UE::AvalancheLevelViewport::Private::ScrubConstants
+namespace UE::AvaLevelViewport::Private::ScrubConstants
 {
 	/** The minimum amount of pixels between each major ticks on the widget */
 	const int32 MinPixelsPerDisplayTick = 5;
@@ -36,7 +36,7 @@ namespace UE::AvalancheLevelViewport::Private::ScrubConstants
 }
 
 /** Utility struct for converting between scrub range space and local/absolute screen space */
-struct UE::AvalancheLevelViewport::Private::SRuler::FScrubRangeToScreen
+struct UE::AvaLevelViewport::Private::SRuler::FScrubRangeToScreen
 {
 	float RulerLengthSlateUnits;
 
@@ -71,7 +71,7 @@ struct UE::AvalancheLevelViewport::Private::SRuler::FScrubRangeToScreen
 	}
 };
 
-struct UE::AvalancheLevelViewport::Private::SRuler::FDrawTickArgs
+struct UE::AvaLevelViewport::Private::SRuler::FDrawTickArgs
 {
 	/** Geometry of the area */
 	FGeometry AllottedGeometry;
@@ -115,7 +115,7 @@ static float GetNextSpacing( uint32 CurrentStep )
 /////////////////////////////////////////////////////
 // SRuler
 
-void UE::AvalancheLevelViewport::Private::SRuler::Construct(const UE::AvalancheLevelViewport::Private::SRuler::FArguments& InArgs)
+void UE::AvaLevelViewport::Private::SRuler::Construct(const UE::AvaLevelViewport::Private::SRuler::FArguments& InArgs)
 {
 	Orientation = InArgs._Orientation;
 	AbsoluteOrigin = FVector2D(0, 0);
@@ -124,7 +124,7 @@ void UE::AvalancheLevelViewport::Private::SRuler::Construct(const UE::AvalancheL
 	MouseButtonDownHandler = InArgs._OnMouseButtonDown;
 }
 
-float UE::AvalancheLevelViewport::Private::SRuler::DetermineOptimalSpacing(float InPixelsPerInput, uint32 MinTick, float MinTickSpacing) const
+float UE::AvaLevelViewport::Private::SRuler::DetermineOptimalSpacing(float InPixelsPerInput, uint32 MinTick, float MinTickSpacing) const
 {
 	if (InPixelsPerInput == 0.0f)
 		return MinTickSpacing;
@@ -143,18 +143,18 @@ float UE::AvalancheLevelViewport::Private::SRuler::DetermineOptimalSpacing(float
 	return Spacing;
 }
 
-void UE::AvalancheLevelViewport::Private::SRuler::SetRuling(FVector2D InAbsoluteOrigin, float InSlateToUnitScale)
+void UE::AvaLevelViewport::Private::SRuler::SetRuling(FVector2D InAbsoluteOrigin, float InSlateToUnitScale)
 {
 	AbsoluteOrigin = InAbsoluteOrigin;
 	SlateToUnitScale = InSlateToUnitScale;
 }
 
-void UE::AvalancheLevelViewport::Private::SRuler::SetCursor(TOptional<FVector2D> InAbsoluteCursor)
+void UE::AvaLevelViewport::Private::SRuler::SetCursor(TOptional<FVector2D> InAbsoluteCursor)
 {
 	AbsoluteCursor = InAbsoluteCursor;
 }
 
-int32 UE::AvalancheLevelViewport::Private::SRuler::DrawTicks( FSlateWindowElementList& OutDrawElements, const struct FScrubRangeToScreen& RangeToScreen, FDrawTickArgs& InArgs ) const
+int32 UE::AvaLevelViewport::Private::SRuler::DrawTicks( FSlateWindowElementList& OutDrawElements, const struct FScrubRangeToScreen& RangeToScreen, FDrawTickArgs& InArgs ) const
 {
 	const float Spacing = DetermineOptimalSpacing( RangeToScreen.PixelsPerInput, ScrubConstants::MinPixelsPerDisplayTick, ScrubConstants::MinDisplayTickSpacing );
 
@@ -300,7 +300,7 @@ int32 UE::AvalancheLevelViewport::Private::SRuler::DrawTicks( FSlateWindowElemen
 	return InArgs.StartLayer;
 }
 
-int32 UE::AvalancheLevelViewport::Private::SRuler::OnPaint(const FPaintArgs& Args, const FGeometry& AllottedGeometry, const FSlateRect& MyCullingRect, FSlateWindowElementList& OutDrawElements, int32 LayerId, const FWidgetStyle& InWidgetStyle, bool bParentEnabled) const
+int32 UE::AvaLevelViewport::Private::SRuler::OnPaint(const FPaintArgs& Args, const FGeometry& AllottedGeometry, const FSlateRect& MyCullingRect, FSlateWindowElementList& OutDrawElements, int32 LayerId, const FWidgetStyle& InWidgetStyle, bool bParentEnabled) const
 {
 	const bool bEnabled = bParentEnabled;
 	const ESlateDrawEffect DrawEffects = bEnabled ? ESlateDrawEffect::None : ESlateDrawEffect::DisabledEffect;
@@ -347,7 +347,7 @@ int32 UE::AvalancheLevelViewport::Private::SRuler::OnPaint(const FPaintArgs& Arg
 	return LayerId;
 }
 
-FReply UE::AvalancheLevelViewport::Private::SRuler::OnMouseButtonDown( const FGeometry& MyGeometry, const FPointerEvent& MouseEvent )
+FReply UE::AvaLevelViewport::Private::SRuler::OnMouseButtonDown( const FGeometry& MyGeometry, const FPointerEvent& MouseEvent )
 {
 	if ( MouseButtonDownHandler.IsBound() )
 	{
@@ -361,22 +361,22 @@ FReply UE::AvalancheLevelViewport::Private::SRuler::OnMouseButtonDown( const FGe
 	}
 }
 
-FReply UE::AvalancheLevelViewport::Private::SRuler::OnMouseButtonUp(const FGeometry& MyGeometry, const FPointerEvent& MouseEvent)
+FReply UE::AvaLevelViewport::Private::SRuler::OnMouseButtonUp(const FGeometry& MyGeometry, const FPointerEvent& MouseEvent)
 {
 	return FReply::Unhandled();
 }
 
-FReply UE::AvalancheLevelViewport::Private::SRuler::OnMouseMove(const FGeometry& MyGeometry, const FPointerEvent& MouseEvent)
+FReply UE::AvaLevelViewport::Private::SRuler::OnMouseMove(const FGeometry& MyGeometry, const FPointerEvent& MouseEvent)
 {
 	return FReply::Unhandled();
 }
 
-FVector2D UE::AvalancheLevelViewport::Private::SRuler::ComputeDesiredSize( float ) const
+FVector2D UE::AvaLevelViewport::Private::SRuler::ComputeDesiredSize( float ) const
 {
 	return Orientation == Orient_Horizontal ? FVector2D(100, 18) : FVector2D(18, 100);
 }
 
-FReply UE::AvalancheLevelViewport::Private::SRuler::OnMouseWheel(const FGeometry& MyGeometry, const FPointerEvent& MouseEvent)
+FReply UE::AvaLevelViewport::Private::SRuler::OnMouseWheel(const FGeometry& MyGeometry, const FPointerEvent& MouseEvent)
 {
 	return FReply::Unhandled();
 }
diff --git a/Engine/Plugins/Experimental/Avalanche/Source/AvalancheLevelViewport/Private/Widgets/Engine/SRuler.h b/Engine/Plugins/Experimental/Avalanche/Source/AvalancheLevelViewport/Private/Widgets/Engine/SRuler.h
--- a/Engine/Plugins/Experimental/Avalanche/Source/AvalancheLevelViewport/Private/Widgets/Engine/SRuler.h
+++ b/Engine/Plugins/Experimental/Avalanche/Source/AvalancheLevelViewport/Private/Widgets/Engine/SRuler.h
@@ -19,7 +19,7 @@ class FWidgetStyle;
 struct FGeometry;
 struct FPointerEvent;
 
-namespace UE::AvalancheLevelViewport::Private
+namespace UE::AvaLevelViewport::Private
 {
 
 class SRuler : public SCompoundWidget
diff --git a/Engine/Plugins/Experimental/Avalanche/Source/AvalancheLevelViewport/Private/Widgets/SAvaLevelViewportActorAlignmentMenu.cpp b/Engine/Plugins/Experimental/Avalanche/Source/AvalancheLevelViewport/Private/Widgets/SAvaLevelViewportActorAlignmentMenu.cpp
--- a/Engine/Plugins/Experimental/Avalanche/Source/AvalancheLevelViewport/Private/Widgets/SAvaLevelViewportActorAlignmentMenu.cpp
+++ b/Engine/Plugins/Experimental/Avalanche/Source/AvalancheLevelViewport/Private/Widgets/SAvaLevelViewportActorAlignmentMenu.cpp
@@ -26,7 +26,7 @@
 
 #define LOCTEXT_NAMESPACE "SAvaLevelViewportActorAlignmentMenu"
 
-namespace UE::AvalancheLevelViewport::Private
+namespace UE::AvaLevelViewport::Private
 {
 	// Transient setting. And that's fine.
 	EAvaAlignmentContext ActorAlignContextType = EAvaAlignmentContext::SelectedActors;
@@ -243,7 +243,7 @@ void SAvaLevelViewportActorAlignmentMenu::Construct(const FArguments& Args, cons
 	int32 Row = 0;
 
 	// Context
-	using namespace UE::AvalancheLevelViewport::Private;
+	using namespace UE::AvaLevelViewport::Private;
 
 	Grid->AddSlot(Button1, Row, ForegroundLayer)
 		.ColumnSpan(4)
@@ -290,7 +290,7 @@ void SAvaLevelViewportActorAlignmentMenu::Construct(const FArguments& Args, cons
 		[
 			CreateButtonCell(
 				GetLocationAlignmentButton(EAvaHorizontalAlignment::Left, static_cast<EAvaVerticalAlignment>(NO_VALUE), 
-					static_cast<EAvaDepthAlignment>(NO_VALUE), "AvalancheIcons.Alignment.Left", 
+					static_cast<EAvaDepthAlignment>(NO_VALUE), "Icons.Alignment.Left", 
 					LOCTEXT("AlignSelectedActorsLeft", "Align Selected Actors Left (to Screen or Actors)")),
 				ECellPosition::First
 			)
@@ -300,7 +300,7 @@ void SAvaLevelViewportActorAlignmentMenu::Construct(const FArguments& Args, cons
 		[
 			CreateButtonCell(
 				GetLocationAlignmentButton(EAvaHorizontalAlignment::Center, static_cast<EAvaVerticalAlignment>(NO_VALUE),
-					static_cast<EAvaDepthAlignment>(NO_VALUE), "AvalancheIcons.Alignment.Center_Y",
+					static_cast<EAvaDepthAlignment>(NO_VALUE), "Icons.Alignment.Center_Y",
 					LOCTEXT("AlignSelectedActorsHCenter", "Align Selected Actors Center (Horizontally) (to Screen or Actors)"))
 			)
 		];
@@ -309,7 +309,7 @@ void SAvaLevelViewportActorAlignmentMenu::Construct(const FArguments& Args, cons
 		[
 			CreateButtonCell(
 				GetLocationAlignmentButton(EAvaHorizontalAlignment::Right, static_cast<EAvaVerticalAlignment>(NO_VALUE),
-					static_cast<EAvaDepthAlignment>(NO_VALUE), "AvalancheIcons.Alignment.Right",
+					static_cast<EAvaDepthAlignment>(NO_VALUE), "Icons.Alignment.Right",
 					LOCTEXT("AlignSelectedActorsRight", "Align Selected Actors Right (to Screen or Actors)"))
 			)
 		];
@@ -319,7 +319,7 @@ void SAvaLevelViewportActorAlignmentMenu::Construct(const FArguments& Args, cons
 		[
 			CreateButtonCell(
 				GetLocationAlignmentButton(EAvaHorizontalAlignment::Center, EAvaVerticalAlignment::Center,
-					static_cast<EAvaDepthAlignment>(NO_VALUE), "AvalancheIcons.Alignment.Center_YZ",
+					static_cast<EAvaDepthAlignment>(NO_VALUE), "Icons.Alignment.Center_YZ",
 					LOCTEXT("AlignSelectedActorsHVCenter", "Align Selected Actors Center (Horizontally and Vertically) (to Screen or Actors)")),
 				ECellPosition::Last
 			)
@@ -332,7 +332,7 @@ void SAvaLevelViewportActorAlignmentMenu::Construct(const FArguments& Args, cons
 		[
 			CreateButtonCell(
 				GetLocationAlignmentButton(static_cast<EAvaHorizontalAlignment>(NO_VALUE), EAvaVerticalAlignment::Top, 
-					static_cast<EAvaDepthAlignment>(NO_VALUE), "AvalancheIcons.Alignment.Top", 
+					static_cast<EAvaDepthAlignment>(NO_VALUE), "Icons.Alignment.Top", 
 					LOCTEXT("AlignSelectedActorsTop", "Align Selected Actors Top (to Screen or Actors)")),
 				ECellPosition::First
 			)
@@ -342,7 +342,7 @@ void SAvaLevelViewportActorAlignmentMenu::Construct(const FArguments& Args, cons
 		[
 			CreateButtonCell(
 				GetLocationAlignmentButton(static_cast<EAvaHorizontalAlignment>(NO_VALUE), EAvaVerticalAlignment::Center,
-					static_cast<EAvaDepthAlignment>(NO_VALUE), "AvalancheIcons.Alignment.Center_Z",
+					static_cast<EAvaDepthAlignment>(NO_VALUE), "Icons.Alignment.Center_Z",
 					LOCTEXT("AlignSelectedActorsVCenter", "Align Selected Actors Center (Vertically) (to Screen or Actors)"))
 			)
 		];
@@ -351,7 +351,7 @@ void SAvaLevelViewportActorAlignmentMenu::Construct(const FArguments& Args, cons
 		[
 			CreateButtonCell(
 				GetLocationAlignmentButton(static_cast<EAvaHorizontalAlignment>(NO_VALUE), EAvaVerticalAlignment::Bottom,
-					static_cast<EAvaDepthAlignment>(NO_VALUE), "AvalancheIcons.Alignment.Bottom",
+					static_cast<EAvaDepthAlignment>(NO_VALUE), "Icons.Alignment.Bottom",
 					LOCTEXT("AlignSelectedActorsBottom", "Align Selected Actors Bottom (to Screen or Actors)"))
 			)
 		];
@@ -363,7 +363,7 @@ void SAvaLevelViewportActorAlignmentMenu::Construct(const FArguments& Args, cons
 		[
 			CreateButtonCell(
 				GetLocationAlignmentButton(static_cast<EAvaHorizontalAlignment>(NO_VALUE), static_cast<EAvaVerticalAlignment>(NO_VALUE), 
-					EAvaDepthAlignment::Front, "AvalancheIcons.Alignment.Translation.Front", 
+					EAvaDepthAlignment::Front, "Icons.Alignment.Translation.Front", 
 					LOCTEXT("AlignSelectedActorsFront", "Align Selected Actors Front (to Actors Only)")),
 				ECellPosition::First
 			)
@@ -373,7 +373,7 @@ void SAvaLevelViewportActorAlignmentMenu::Construct(const FArguments& Args, cons
 		[
 			CreateButtonCell(
 				GetLocationAlignmentButton(static_cast<EAvaHorizontalAlignment>(NO_VALUE), static_cast<EAvaVerticalAlignment>(NO_VALUE), 
-					EAvaDepthAlignment::Center, "AvalancheIcons.Alignment.Translation.Center_X", 
+					EAvaDepthAlignment::Center, "Icons.Alignment.Translation.Center_X", 
 					LOCTEXT("AlignSelectedActorsDCenter", "Align Selected Actors Center (Depth) (to Actors Only)"))
 			)
 		];
@@ -382,7 +382,7 @@ void SAvaLevelViewportActorAlignmentMenu::Construct(const FArguments& Args, cons
 		[
 			CreateButtonCell(
 				GetLocationAlignmentButton(static_cast<EAvaHorizontalAlignment>(NO_VALUE), static_cast<EAvaVerticalAlignment>(NO_VALUE),
-					EAvaDepthAlignment::Back, "AvalancheIcons.Alignment.Translation.Back",
+					EAvaDepthAlignment::Back, "Icons.Alignment.Translation.Back",
 					LOCTEXT("AlignSelectedActorsBack", "Align Selected Actors Back (to Actors Only)"))
 			)
 		];
@@ -393,7 +393,7 @@ void SAvaLevelViewportActorAlignmentMenu::Construct(const FArguments& Args, cons
 	Grid->AddSlot(Button1, Row, ForegroundLayer)
 		[
 			CreateButtonCell(
-				GetDistributionAlignmentButton(EAvaScreenAxis::Horizontal, "AvalancheIcons.Alignment.DistributeY",
+				GetDistributionAlignmentButton(EAvaScreenAxis::Horizontal, "Icons.Alignment.DistributeY",
 					LOCTEXT("AlignSelectedActorsDistributeY", "Distribute Selected Actors Left to Right (based on Screen or Actors)")),
 				ECellPosition::First
 			)
@@ -402,7 +402,7 @@ void SAvaLevelViewportActorAlignmentMenu::Construct(const FArguments& Args, cons
 	Grid->AddSlot(Button2, Row, ForegroundLayer)
 		[
 			CreateButtonCell(
-				GetDistributionAlignmentButton(EAvaScreenAxis::Vertical, "AvalancheIcons.Alignment.DistributeZ",
+				GetDistributionAlignmentButton(EAvaScreenAxis::Vertical, "Icons.Alignment.DistributeZ",
 					LOCTEXT("AlignSelectedActorsDistributeZ", "Distribute Selected Actors Top to Bottom (based on Screen or Actors)"))
 			)
 		];
@@ -410,7 +410,7 @@ void SAvaLevelViewportActorAlignmentMenu::Construct(const FArguments& Args, cons
 	Grid->AddSlot(Button3, Row, ForegroundLayer)
 		[
 			CreateButtonCell(
-				GetDistributionAlignmentButton(EAvaScreenAxis::Depth, "AvalancheIcons.Alignment.DistributeX",
+				GetDistributionAlignmentButton(EAvaScreenAxis::Depth, "Icons.Alignment.DistributeX",
 					LOCTEXT("AlignSelectedActorsDistributeX", "Distribute Selected Actors Front to Back (based on Screen or Actors)"))
 			)
 		];
@@ -421,7 +421,7 @@ void SAvaLevelViewportActorAlignmentMenu::Construct(const FArguments& Args, cons
 	Grid->AddSlot(Button1, Row, ForegroundLayer)
 		[
 			CreateButtonCell(
-				GetRotationAlignmentButton(EAvaRotationAxis::Roll, "AvalancheIcons.Alignment.Rotation.Actor.Roll",
+				GetRotationAlignmentButton(EAvaRotationAxis::Roll, "Icons.Alignment.Rotation.Actor.Roll",
 					LOCTEXT("AlignSelectedActorsToRoll", "Align Roll of Selected Actors")),
 				ECellPosition::First
 			)
@@ -430,7 +430,7 @@ void SAvaLevelViewportActorAlignmentMenu::Construct(const FArguments& Args, cons
 	Grid->AddSlot(Button2, Row, ForegroundLayer)
 		[
 			CreateButtonCell(
-				GetRotationAlignmentButton(EAvaRotationAxis::Pitch, "AvalancheIcons.Alignment.Rotation.Actor.Pitch",
+				GetRotationAlignmentButton(EAvaRotationAxis::Pitch, "Icons.Alignment.Rotation.Actor.Pitch",
 					LOCTEXT("AlignSelectedActorsToPitch", "Align Pitch of Selected Actors"))
 			)
 		];
@@ -438,7 +438,7 @@ void SAvaLevelViewportActorAlignmentMenu::Construct(const FArguments& Args, cons
 	Grid->AddSlot(Button3, Row, ForegroundLayer)
 		[
 			CreateButtonCell(
-				GetRotationAlignmentButton(EAvaRotationAxis::Yaw, "AvalancheIcons.Alignment.Rotation.Actor.Yaw",
+				GetRotationAlignmentButton(EAvaRotationAxis::Yaw, "Icons.Alignment.Rotation.Actor.Yaw",
 					LOCTEXT("AlignSelectedActorsToYaw", "Align Yaw of Selected Actors"))
 			)
 		];
@@ -446,7 +446,7 @@ void SAvaLevelViewportActorAlignmentMenu::Construct(const FArguments& Args, cons
 	Grid->AddSlot(Button4, Row, ForegroundLayer)
 		[
 			CreateButtonCell(
-				GetRotationAlignmentButton(EAvaRotationAxis::All, "AvalancheIcons.Alignment.Rotation.Actor.All",
+				GetRotationAlignmentButton(EAvaRotationAxis::All, "Icons.Alignment.Rotation.Actor.All",
 					LOCTEXT("AlignSelectedActorsToRollYawPitch", "Align Roll, Pitch and Yaw  of Selected Actors")),
 				ECellPosition::Last
 			)
@@ -458,7 +458,7 @@ void SAvaLevelViewportActorAlignmentMenu::Construct(const FArguments& Args, cons
 	Grid->AddSlot(Button1, Row, ForegroundLayer)
 		[
 			CreateButtonCell(
-				GetCameraRotationAlignmentButton(EAvaRotationAxis::Roll, "AvalancheIcons.Alignment.Rotation.Camera.Roll",
+				GetCameraRotationAlignmentButton(EAvaRotationAxis::Roll, "Icons.Alignment.Rotation.Camera.Roll",
 					LOCTEXT("AlignSelectedActorsToRollCamera", "Align Roll of Selected Actors to face the Camera")),
 				ECellPosition::First
 			)
@@ -467,7 +467,7 @@ void SAvaLevelViewportActorAlignmentMenu::Construct(const FArguments& Args, cons
 	Grid->AddSlot(Button2, Row, ForegroundLayer)
 		[
 			CreateButtonCell(
-				GetCameraRotationAlignmentButton(EAvaRotationAxis::Pitch, "AvalancheIcons.Alignment.Rotation.Camera.Pitch",
+				GetCameraRotationAlignmentButton(EAvaRotationAxis::Pitch, "Icons.Alignment.Rotation.Camera.Pitch",
 					LOCTEXT("AlignSelectedActorsToPitchCamera", "Align Pitch of Selected Actors to face the Camera"))
 			)
 		];
@@ -475,7 +475,7 @@ void SAvaLevelViewportActorAlignmentMenu::Construct(const FArguments& Args, cons
 	Grid->AddSlot(Button3, Row, ForegroundLayer)
 		[
 			CreateButtonCell(
-				GetCameraRotationAlignmentButton(EAvaRotationAxis::Yaw, "AvalancheIcons.Alignment.Rotation.Camera.Yaw",
+				GetCameraRotationAlignmentButton(EAvaRotationAxis::Yaw, "Icons.Alignment.Rotation.Camera.Yaw",
 					LOCTEXT("AlignSelectedActorsToYawCamera", "Align Yaw of Selected Actors to face the Camera"))
 			)
 		];
@@ -483,7 +483,7 @@ void SAvaLevelViewportActorAlignmentMenu::Construct(const FArguments& Args, cons
 	Grid->AddSlot(Button4, Row, ForegroundLayer)
 		[
 			CreateButtonCell(
-				GetCameraRotationAlignmentButton(EAvaRotationAxis::All, "AvalancheIcons.Alignment.Rotation.Camera.All",
+				GetCameraRotationAlignmentButton(EAvaRotationAxis::All, "Icons.Alignment.Rotation.Camera.All",
 					LOCTEXT("AlignSelectedActorsToRollYawPitchCamera", "Align Roll, Pitch and Yaw of Selected Actors to face the Camera")),
 				ECellPosition::Last
 			)
@@ -502,7 +502,7 @@ void SAvaLevelViewportActorAlignmentMenu::Construct(const FArguments& Args, cons
 				.ToolTipText(LOCTEXT("SizeToScreen", "Size the Currently Selected Actors to the Screen."))
 				[
 					SNew(SImage)
-					.Image(FAvaLevelViewportStyle::Get().GetBrush("AvalancheIcons.Screen.SizeToScreen"))
+					.Image(FAvaLevelViewportStyle::Get().GetBrush("Icons.Screen.SizeToScreen"))
 					.DesiredSizeOverride(FVector2D(16.f, 16.f))
 				],
 				ECellPosition::First
@@ -519,7 +519,7 @@ void SAvaLevelViewportActorAlignmentMenu::Construct(const FArguments& Args, cons
 				.ToolTipText(LOCTEXT("SizeToScreenStretch", "Stretch the Currently Selected Actors to the Screen."))
 				[
 					SNew(SImage)
-					.Image(FAvaLevelViewportStyle::Get().GetBrush("AvalancheIcons.Screen.SizeToScreenStretch"))
+					.Image(FAvaLevelViewportStyle::Get().GetBrush("Icons.Screen.SizeToScreenStretch"))
 					.DesiredSizeOverride(FVector2D(16.f, 16.f))
 				]
 			)
@@ -535,7 +535,7 @@ void SAvaLevelViewportActorAlignmentMenu::Construct(const FArguments& Args, cons
 				.ToolTipText(LOCTEXT("FitToScreen", "Stretches the Currently Selected Actors to the Screen and centers it.\n\n- Shift: If an object is rotated beyond 45 degrees with respect to the camera, align to nearest axis."))
 				[
 					SNew(SImage)
-					.Image(FAvaLevelViewportStyle::Get().GetBrush("AvalancheIcons.Screen.FitToScreen"))
+					.Image(FAvaLevelViewportStyle::Get().GetBrush("Icons.Screen.FitToScreen"))
 					.DesiredSizeOverride(FVector2D(16.f, 16.f))
 				]
 			)
@@ -607,7 +607,7 @@ TSharedRef<SButton> SAvaLevelViewportActorAlignmentMenu::GetCameraRotationAlignm
 
 FReply SAvaLevelViewportActorAlignmentMenu::OnSizeToScreenClicked(bool bInStretchToFit)
 {
-	using namespace UE::AvalancheLevelViewport::Private;
+	using namespace UE::AvaLevelViewport::Private;
 
 	FDismissMenus DismissMenus;
 
@@ -659,7 +659,7 @@ FReply SAvaLevelViewportActorAlignmentMenu::OnSizeToScreenClicked(bool bInStretc
 
 FReply SAvaLevelViewportActorAlignmentMenu::OnFitToScreenClicked()
 {
-	using namespace UE::AvalancheLevelViewport::Private;
+	using namespace UE::AvaLevelViewport::Private;
 
 	FDismissMenus DismissMenus;
 
@@ -714,7 +714,7 @@ FReply SAvaLevelViewportActorAlignmentMenu::OnFitToScreenClicked()
 
 FReply SAvaLevelViewportActorAlignmentMenu::OnRotationAlignmentButtonClicked(EAvaRotationAxis InAxisList)
 {
-	using namespace UE::AvalancheLevelViewport::Private;
+	using namespace UE::AvaLevelViewport::Private;
 
 	FDismissMenus DismissMenus;
 
@@ -779,7 +779,7 @@ FReply SAvaLevelViewportActorAlignmentMenu::OnRotationAlignmentButtonClicked(EAv
 
 FReply SAvaLevelViewportActorAlignmentMenu::OnCameraRotationAlignmentButtonClicked(EAvaRotationAxis InAxisList)
 {
-	using namespace UE::AvalancheLevelViewport::Private;
+	using namespace UE::AvaLevelViewport::Private;
 
 	FDismissMenus DismissMenus;
 
@@ -844,7 +844,7 @@ FReply SAvaLevelViewportActorAlignmentMenu::OnCameraRotationAlignmentButtonClick
 
 FReply SAvaLevelViewportActorAlignmentMenu::OnDistributeAlignmentButtonClicked(EAvaScreenAxis InScreenAxis)
 {
-	using namespace UE::AvalancheLevelViewport::Private;
+	using namespace UE::AvaLevelViewport::Private;
 
 	FDismissMenus DismissMenus;
 
@@ -917,7 +917,7 @@ FReply SAvaLevelViewportActorAlignmentMenu::OnDistributeAlignmentButtonClicked(E
 FReply SAvaLevelViewportActorAlignmentMenu::OnLocationAlignmentButtonClicked(EAvaHorizontalAlignment InHoriz, EAvaVerticalAlignment InVert,
 	EAvaDepthAlignment InDepth)
 {
-	using namespace UE::AvalancheLevelViewport::Private;
+	using namespace UE::AvaLevelViewport::Private;
 
 	FDismissMenus DismissMenus;
 
diff --git a/Engine/Plugins/Experimental/Avalanche/Source/AvalancheLevelViewport/Private/Widgets/SAvaLevelViewportActorColorMenu.cpp b/Engine/Plugins/Experimental/Avalanche/Source/AvalancheLevelViewport/Private/Widgets/SAvaLevelViewportActorColorMenu.cpp
--- a/Engine/Plugins/Experimental/Avalanche/Source/AvalancheLevelViewport/Private/Widgets/SAvaLevelViewportActorColorMenu.cpp
+++ b/Engine/Plugins/Experimental/Avalanche/Source/AvalancheLevelViewport/Private/Widgets/SAvaLevelViewportActorColorMenu.cpp
@@ -230,13 +230,13 @@ void SAvaLevelViewportActorColorMenu::Construct(const FArguments& InArgs, const
 				+ SWidgetSwitcher::Slot()
 				[
 					SNew(SImage)
-					.Image(AvaLevelViewportStyle.GetBrush("AvalancheIcons.ColorPicker.SolidColors"))
+					.Image(AvaLevelViewportStyle.GetBrush("Icons.ColorPicker.SolidColors"))
 					.OnMouseButtonDown(this, &SAvaLevelViewportActorColorMenu::OnColorSwitchedMouseDown)
 				]
 				+ SWidgetSwitcher::Slot()
 				[
 					SNew(SImage)
-					.Image(AvaLevelViewportStyle.GetBrush("AvalancheIcons.ColorPicker.LinearGradient"))
+					.Image(AvaLevelViewportStyle.GetBrush("Icons.ColorPicker.LinearGradient"))
 					.OnMouseButtonDown(this, &SAvaLevelViewportActorColorMenu::OnColorSwitchedMouseDown)
 				]
 			]
diff --git a/Engine/Plugins/Experimental/Avalanche/Source/AvalancheLevelViewport/Private/Widgets/SAvaLevelViewportGuide.cpp b/Engine/Plugins/Experimental/Avalanche/Source/AvalancheLevelViewport/Private/Widgets/SAvaLevelViewportGuide.cpp
--- a/Engine/Plugins/Experimental/Avalanche/Source/AvalancheLevelViewport/Private/Widgets/SAvaLevelViewportGuide.cpp
+++ b/Engine/Plugins/Experimental/Avalanche/Source/AvalancheLevelViewport/Private/Widgets/SAvaLevelViewportGuide.cpp
@@ -17,7 +17,7 @@
 
 #define LOCTEXT_NAMESPACE "AvaLevelViewportGuide"
 
-namespace UE::AvalancheViewport::Private
+namespace UE::AvaViewport::Private
 {
 	constexpr float Padding = 2.f;
 	const FTimespan DoubleClickMaxDelay = FTimespan::FromMilliseconds(850);
@@ -272,7 +272,7 @@ void SAvaLevelViewportGuide::Construct(const FArguments& InArgs, TSharedPtr<SAva
 	LastClickTime = FDateTime(0);
 	LastClickLocation = FVector2f::ZeroVector;
 
-	using namespace UE::AvalancheViewport::Private;
+	using namespace UE::AvaViewport::Private;
 
 	if (Info.Orientation == Orient_Horizontal)
 	{
@@ -368,7 +368,7 @@ FReply SAvaLevelViewportGuide::OnMouseButtonUp(const FGeometry& Geometry, const
 
 	const FDateTime CurrentTime = FDateTime::Now();
 
-	using namespace UE::AvalancheViewport::Private;
+	using namespace UE::AvaViewport::Private;
 
 	if ((PointerEvent.GetScreenSpacePosition() - LastMouseDownLocation).Size() < ClickMaxDistance)
 	{
@@ -510,7 +510,7 @@ bool SAvaLevelViewportGuide::DragUpdate()
 		SnapOperation->SnapScreenLocation(MousePosition, Info.Orientation == EOrientation::Orient_Vertical, Info.Orientation == EOrientation::Orient_Horizontal);
 	}
 
-	using namespace UE::AvalancheViewport::Private;
+	using namespace UE::AvaViewport::Private;
 
 	switch (Info.Orientation)
 	{
@@ -607,7 +607,7 @@ bool SAvaLevelViewportGuide::SetOffset(float Offset)
 	const FVector2f ViewportSize = FrameAndClient.ViewportClient->GetCachedViewportSize();
 	float NewOffsetFraction;
 
-	using namespace UE::AvalancheViewport::Private;
+	using namespace UE::AvaViewport::Private;
 
 	switch (Info.Orientation)
 	{
@@ -657,7 +657,7 @@ FVector2D SAvaLevelViewportGuide::GetSize() const
 
 	const FVector2f ViewportSize = FrameAndClient.ViewportClient->GetCachedViewportSize();
 
-	using namespace UE::AvalancheViewport::Private;
+	using namespace UE::AvaViewport::Private;
 
 	switch (Info.Orientation)
 	{
@@ -686,7 +686,7 @@ FVector2D SAvaLevelViewportGuide::GetPosition() const
 	const FVector2f Offset = FrameAndClient.ViewportClient->GetCachedViewportOffset();
 	FVector2f GuidePosition;
 
-	using namespace UE::AvalancheViewport::Private;
+	using namespace UE::AvaViewport::Private;
 
 	if (Info.Orientation == Orient_Horizontal)
 	{
@@ -818,7 +818,7 @@ void SAvaLevelViewportGuide::UpdateGuideData() const
 
 void SAvaLevelViewportGuide::OpenRightClickMenu()
 {
-	using namespace UE::AvalancheViewport::Private;
+	using namespace UE::AvaViewport::Private;
 
 	RegisterGuideMenu();
 
diff --git a/Engine/Plugins/Experimental/Avalanche/Source/AvalancheLevelViewport/Private/Widgets/SAvaLevelViewportRuler.h b/Engine/Plugins/Experimental/Avalanche/Source/AvalancheLevelViewport/Private/Widgets/SAvaLevelViewportRuler.h
--- a/Engine/Plugins/Experimental/Avalanche/Source/AvalancheLevelViewport/Private/Widgets/SAvaLevelViewportRuler.h
+++ b/Engine/Plugins/Experimental/Avalanche/Source/AvalancheLevelViewport/Private/Widgets/SAvaLevelViewportRuler.h
@@ -6,7 +6,7 @@
 
 class SAvaLevelViewportFrame;
 
-class SAvaLevelViewportRuler : public UE::AvalancheLevelViewport::Private::SRuler
+class SAvaLevelViewportRuler : public UE::AvaLevelViewport::Private::SRuler
 {
 public:
 	SLATE_BEGIN_ARGS(SAvaLevelViewportRuler)
diff --git a/Engine/Plugins/Experimental/Avalanche/Source/AvalancheLevelViewport/Private/Widgets/SAvaLevelViewportStatusBar.h b/Engine/Plugins/Experimental/Avalanche/Source/AvalancheLevelViewport/Private/Widgets/SAvaLevelViewportStatusBar.h
--- a/Engine/Plugins/Experimental/Avalanche/Source/AvalancheLevelViewport/Private/Widgets/SAvaLevelViewportStatusBar.h
+++ b/Engine/Plugins/Experimental/Avalanche/Source/AvalancheLevelViewport/Private/Widgets/SAvaLevelViewportStatusBar.h
@@ -18,7 +18,7 @@
 class FUICommandInfo;
 class SAvaLevelViewportFrame;
 
-namespace UE::AvalancheLevelViewport::Private
+namespace UE::AvaLevelViewport::Private
 {
 	namespace ViewportStatusBarButton
 	{
diff --git a/Engine/Plugins/Experimental/Avalanche/Source/AvalancheLevelViewport/Private/Widgets/SAvaLevelViewportStatusBarButtons.cpp b/Engine/Plugins/Experimental/Avalanche/Source/AvalancheLevelViewport/Private/Widgets/SAvaLevelViewportStatusBarButtons.cpp
--- a/Engine/Plugins/Experimental/Avalanche/Source/AvalancheLevelViewport/Private/Widgets/SAvaLevelViewportStatusBarButtons.cpp
+++ b/Engine/Plugins/Experimental/Avalanche/Source/AvalancheLevelViewport/Private/Widgets/SAvaLevelViewportStatusBarButtons.cpp
@@ -37,12 +37,12 @@ namespace UE::Ava::LevelViewportStatusBarButtons::Private
 {
 	static const FName AvaLevelViewportStyleName = FAvaLevelViewportStyle::Get().GetStyleSetName();
 	static const FName AppStyleSetName = FAppStyle::Get().GetStyleSetName();
-	static const FSlateIcon RGBChannelIcon = FSlateIcon(AvaLevelViewportStyleName, "AvalancheIcons.PostProcess.RGB");
+	static const FSlateIcon RGBChannelIcon = FSlateIcon(AvaLevelViewportStyleName, "Icons.PostProcess.RGB");
 	static const FSlateIcon BackgroundIcon = FSlateIcon(AppStyleSetName, "Icons.Role");
-	static const FSlateIcon RedChannelIcon = FSlateIcon(AvaLevelViewportStyleName, "AvalancheIcons.PostProcess.Red");
-	static const FSlateIcon GreenChannelIcon = FSlateIcon(AvaLevelViewportStyleName, "AvalancheIcons.PostProcess.Green");
-	static const FSlateIcon BlueChannelIcon = FSlateIcon(AvaLevelViewportStyleName, "AvalancheIcons.PostProcess.Blue");
-	static const FSlateIcon AlphaChannelIcon = FSlateIcon(AvaLevelViewportStyleName, "AvalancheIcons.PostProcess.Alpha");
+	static const FSlateIcon RedChannelIcon = FSlateIcon(AvaLevelViewportStyleName, "Icons.PostProcess.Red");
+	static const FSlateIcon GreenChannelIcon = FSlateIcon(AvaLevelViewportStyleName, "Icons.PostProcess.Green");
+	static const FSlateIcon BlueChannelIcon = FSlateIcon(AvaLevelViewportStyleName, "Icons.PostProcess.Blue");
+	static const FSlateIcon AlphaChannelIcon = FSlateIcon(AvaLevelViewportStyleName, "Icons.PostProcess.Alpha");
 	static const FSlateIcon CheckerboardIcon = FSlateIcon(AppStyleSetName, "Checker");
 
 	static bool IsViewportPostProcessManagerEnabled(const TWeakPtr<SAvaLevelViewportFrame>& InViewportFrameWeak)
@@ -226,7 +226,7 @@ void SAvaLevelViewportStatusBarButtons::CreateContextMenuWigets()
 
 void SAvaLevelViewportStatusBarButtons::PopulateActorButtons(TSharedPtr<SHorizontalBox> InContainer)
 {
-	using namespace UE::AvalancheLevelViewport::Private;
+	using namespace UE::AvaLevelViewport::Private;
 
 	const FAvaLevelViewportCommands& CommandsRef = FAvaLevelViewportCommands::Get();
 
@@ -264,7 +264,7 @@ void SAvaLevelViewportStatusBarButtons::PopulateActorButtons(TSharedPtr<SHorizon
 
 void SAvaLevelViewportStatusBarButtons::PopulateViewportButtons(TSharedPtr<SHorizontalBox> InContainer)
 {
-	using namespace UE::AvalancheLevelViewport::Private;
+	using namespace UE::AvaLevelViewport::Private;
 
 	const FAvaLevelViewportCommands& CommandsRef = FAvaLevelViewportCommands::Get();
 	const FLevelViewportCommands& ViewportActionsRef = FLevelViewportCommands::Get();
@@ -456,7 +456,7 @@ void SAvaLevelViewportStatusBarButtons::PopulateViewportButtons(TSharedPtr<SHori
 
 FSlateColor SAvaLevelViewportStatusBarButtons::GetPostProcessColor() const
 {
-	using namespace UE::AvalancheLevelViewport::Private;
+	using namespace UE::AvaLevelViewport::Private;
 
 	const FAvaLevelViewportGuideFrameAndClient FrameAndClient(ViewportFrameWeak);
 
@@ -822,7 +822,7 @@ TSharedRef<SWidget> SAvaLevelViewportStatusBarButtons::GetActorColorMenuContent(
 
 FSlateColor SAvaLevelViewportStatusBarButtons::GetToggleSnapColor() const
 {
-	using namespace UE::AvalancheLevelViewport::Private;
+	using namespace UE::AvaLevelViewport::Private;
 
 	const FAvaLevelViewportGuideFrameAndWidget FrameAndWidget(ViewportFrameWeak);
 
@@ -918,7 +918,7 @@ TSharedRef<SWidget> SAvaLevelViewportStatusBarButtons::GetSnappingMenuContent()
 
 FSlateColor SAvaLevelViewportStatusBarButtons::GetToggleShapeEditorOverlayColor() const
 {
-	using namespace UE::AvalancheLevelViewport::Private;
+	using namespace UE::AvaLevelViewport::Private;
 
 	if (const UAvaViewportSettings* AvaViewportSettings = GetDefault<UAvaViewportSettings>())
 	{
@@ -940,7 +940,7 @@ bool SAvaLevelViewportStatusBarButtons::GetToggleShapeEditorOverlayEnabled() con
 
 FReply SAvaLevelViewportStatusBarButtons::ToggleShapeEditorOverlay()
 {
-	using namespace UE::AvalancheLevelViewport::Private;
+	using namespace UE::AvaLevelViewport::Private;
 
 	if (UAvaViewportSettings* AvaViewportSettings = GetMutableDefault<UAvaViewportSettings>())
 	{
@@ -956,7 +956,7 @@ FReply SAvaLevelViewportStatusBarButtons::ToggleShapeEditorOverlay()
 
 FSlateColor SAvaLevelViewportStatusBarButtons::GetToggleGuidesColor() const
 {
-	using namespace UE::AvalancheLevelViewport::Private;
+	using namespace UE::AvaLevelViewport::Private;
 
 	const FAvaLevelViewportGuideFrameAndWidget FrameAndWidget(ViewportFrameWeak);
 
@@ -998,7 +998,7 @@ FReply SAvaLevelViewportStatusBarButtons::ToggleGuides()
 
 FSlateColor SAvaLevelViewportStatusBarButtons::GetToggleGridColor() const
 {
-	using namespace UE::AvalancheLevelViewport::Private;
+	using namespace UE::AvaLevelViewport::Private;
 
 	const FAvaLevelViewportGuideFrameAndWidget FrameAndWidget(ViewportFrameWeak);
 
@@ -1108,7 +1108,7 @@ void SAvaLevelViewportStatusBarButtons::OnGridSizeCommitted(int32 InNewValue, ET
 
 bool SAvaLevelViewportStatusBarButtons::GetViewportInfoEnabled() const
 {
-	using namespace UE::AvalancheLevelViewport::Private;
+	using namespace UE::AvaLevelViewport::Private;
 
 	const FAvaLevelViewportGuideFrameAndClient FrameAndClient(ViewportFrameWeak);
 
@@ -1125,7 +1125,7 @@ bool SAvaLevelViewportStatusBarButtons::GetViewportInfoEnabled() const
 
 FSlateColor SAvaLevelViewportStatusBarButtons::GetViewportInfoColor() const
 {
-	using namespace UE::AvalancheLevelViewport::Private;
+	using namespace UE::AvaLevelViewport::Private;
 
 	const FAvaLevelViewportGuideFrameAndClient FrameAndClient(ViewportFrameWeak);
 
@@ -1145,7 +1145,7 @@ FSlateColor SAvaLevelViewportStatusBarButtons::GetViewportInfoColor() const
 
 TSharedRef<SWidget> SAvaLevelViewportStatusBarButtons::GetViewportInfoWidget() const
 {
-	using namespace UE::AvalancheLevelViewport::Private;
+	using namespace UE::AvaLevelViewport::Private;
 
 	const FAvaLevelViewportGuideFrameAndClient FrameAndClient(ViewportFrameWeak);
 
@@ -1177,7 +1177,7 @@ bool SAvaLevelViewportStatusBarButtons::CanChangeGridSize() const
 
 FSlateColor SAvaLevelViewportStatusBarButtons::GetToggleOverlayColor() const
 {
-	using namespace UE::AvalancheLevelViewport::Private;
+	using namespace UE::AvaLevelViewport::Private;
 
 	const FAvaLevelViewportGuideFrameAndWidget FrameAndWidget(ViewportFrameWeak);
 
@@ -1219,7 +1219,7 @@ FReply SAvaLevelViewportStatusBarButtons::ToggleOverlay()
 
 FSlateColor SAvaLevelViewportStatusBarButtons::GetToggleBoundingBoxesColor() const
 {
-	using namespace UE::AvalancheLevelViewport::Private;
+	using namespace UE::AvaLevelViewport::Private;
 
 	const FAvaLevelViewportGuideFrameAndWidget FrameAndWidget(ViewportFrameWeak);
 
@@ -1261,7 +1261,7 @@ FReply SAvaLevelViewportStatusBarButtons::ToggleBoundingBoxes()
 
 FSlateColor SAvaLevelViewportStatusBarButtons::GetToggleIsolateActorsColor() const
 {
-	using namespace UE::AvalancheLevelViewport::Private;
+	using namespace UE::AvaLevelViewport::Private;
 
 	const FAvaLevelViewportGuideFrameAndClient FrameAndClient(ViewportFrameWeak);
 
@@ -1275,7 +1275,7 @@ FSlateColor SAvaLevelViewportStatusBarButtons::GetToggleIsolateActorsColor() con
 
 bool SAvaLevelViewportStatusBarButtons::GetToggleIsolateActorsEnabled() const
 {
-	using namespace UE::AvalancheLevelViewport::Private;
+	using namespace UE::AvaLevelViewport::Private;
 
 	const FAvaLevelViewportGuideFrameAndClient FrameAndClient(ViewportFrameWeak);
 
@@ -1289,7 +1289,7 @@ bool SAvaLevelViewportStatusBarButtons::GetToggleIsolateActorsEnabled() const
 
 FReply SAvaLevelViewportStatusBarButtons::ToggleIsolateActors()
 {
-	using namespace UE::AvalancheLevelViewport::Private;
+	using namespace UE::AvaLevelViewport::Private;
 
 	const FAvaLevelViewportGuideFrameAndClient FrameAndClient(ViewportFrameWeak);
 
@@ -1304,7 +1304,7 @@ FReply SAvaLevelViewportStatusBarButtons::ToggleIsolateActors()
 
 FSlateColor SAvaLevelViewportStatusBarButtons::GetToggleSafeFramesColor() const
 {
-	using namespace UE::AvalancheLevelViewport::Private;
+	using namespace UE::AvaLevelViewport::Private;
 
 	const FAvaLevelViewportGuideFrameAndWidget FrameAndWidget(ViewportFrameWeak);
 
@@ -1351,7 +1351,7 @@ FReply SAvaLevelViewportStatusBarButtons::ToggleSafeFrames()
 
 FSlateColor SAvaLevelViewportStatusBarButtons::GetToggleGameViewColor() const
 {
-	using namespace UE::AvalancheLevelViewport::Private;
+	using namespace UE::AvaLevelViewport::Private;
 
 	const FAvaLevelViewportGuideFrameAndClient FrameAndClient(ViewportFrameWeak);
 
diff --git a/Engine/Plugins/Experimental/Avalanche/Source/AvalancheLevelViewport/Private/Widgets/SAvaLevelViewportStatusBarTransformSettings.cpp b/Engine/Plugins/Experimental/Avalanche/Source/AvalancheLevelViewport/Private/Widgets/SAvaLevelViewportStatusBarTransformSettings.cpp
--- a/Engine/Plugins/Experimental/Avalanche/Source/AvalancheLevelViewport/Private/Widgets/SAvaLevelViewportStatusBarTransformSettings.cpp
+++ b/Engine/Plugins/Experimental/Avalanche/Source/AvalancheLevelViewport/Private/Widgets/SAvaLevelViewportStatusBarTransformSettings.cpp
@@ -16,7 +16,7 @@
 
 #define LOCTEXT_NAMESPACE "SAvaLevelViewportStatusBarTransformSettings"
 
-namespace UE::AvalancheLevelViewport::Private
+namespace UE::AvaLevelViewport::Private
 {
 	TSharedPtr<IDetailKeyframeHandler> GetKeyframeHandler()
 	{
@@ -49,7 +49,7 @@ void SAvaLevelViewportStatusBarTransformSettings::Construct(const FArguments& In
 	TransformDetails = MakeShared<FAvaLevelViewportComponentTransformDetails>(InViewportFrame.IsValid() ? InViewportFrame->GetViewportClient() : nullptr);
 	SetCanTick(true);
 
-	using namespace UE::AvalancheLevelViewport::Private;
+	using namespace UE::AvaLevelViewport::Private;
 
 	auto CreateKeyframeButton = [this](int32 InIndex)
 		{
@@ -136,7 +136,7 @@ void SAvaLevelViewportStatusBarTransformSettings::Construct(const FArguments& In
 
 TSharedRef<SHorizontalBox> SAvaLevelViewportStatusBarTransformSettings::CreateLocationWidgets()
 {
-	using namespace UE::AvalancheLevelViewport::Private;
+	using namespace UE::AvaLevelViewport::Private;
 
 	const FAvaLevelViewportCommands& CommandsRef = FAvaLevelViewportCommands::Get();
 
@@ -176,7 +176,7 @@ TSharedRef<SHorizontalBox> SAvaLevelViewportStatusBarTransformSettings::CreateLo
 
 TSharedRef<SHorizontalBox> SAvaLevelViewportStatusBarTransformSettings::CreateRotationWidgets()
 {
-	using namespace UE::AvalancheLevelViewport::Private;
+	using namespace UE::AvaLevelViewport::Private;
 
 	const FAvaLevelViewportCommands& CommandsRef = FAvaLevelViewportCommands::Get();
 
@@ -215,7 +215,7 @@ TSharedRef<SHorizontalBox> SAvaLevelViewportStatusBarTransformSettings::CreateRo
 
 TSharedRef<SHorizontalBox> SAvaLevelViewportStatusBarTransformSettings::CreateScaleWidgets()
 {
-	using namespace UE::AvalancheLevelViewport::Private;
+	using namespace UE::AvaLevelViewport::Private;
 
 	const FAvaLevelViewportCommands& CommandsRef = FAvaLevelViewportCommands::Get();
 
@@ -263,7 +263,7 @@ void SAvaLevelViewportStatusBarTransformSettings::CreateKeyframeButtons()
 		return;
 	}
 
-	using namespace UE::AvalancheLevelViewport::Private;
+	using namespace UE::AvaLevelViewport::Private;
 
 	TConstArrayView<TSharedPtr<IPropertyHandle>> PropertyHandles = TransformDetails->GetPropertyHandles();
 
@@ -321,7 +321,7 @@ bool SAvaLevelViewportStatusBarTransformSettings::IsResetLocationEnabled() const
 
 FSlateColor SAvaLevelViewportStatusBarTransformSettings::GetResetLocationColor() const
 {
-	using namespace UE::AvalancheLevelViewport::Private;
+	using namespace UE::AvaLevelViewport::Private;
 	return IsResetLocationEnabled() ? ViewportStatusBarButton::EnabledColor : ViewportStatusBarButton::DisabledColor;
 }
 
@@ -343,7 +343,7 @@ bool SAvaLevelViewportStatusBarTransformSettings::IsResetRotationEnabled() const
 
 FSlateColor SAvaLevelViewportStatusBarTransformSettings::GetResetRotationColor() const
 {
-	using namespace UE::AvalancheLevelViewport::Private;
+	using namespace UE::AvaLevelViewport::Private;
 	return IsResetRotationEnabled() ? ViewportStatusBarButton::EnabledColor : ViewportStatusBarButton::DisabledColor;
 }
 
@@ -365,7 +365,7 @@ bool SAvaLevelViewportStatusBarTransformSettings::IsResetScaleEnabled() const
 
 FSlateColor SAvaLevelViewportStatusBarTransformSettings::GetResetScaleColor() const
 {
-	using namespace UE::AvalancheLevelViewport::Private;
+	using namespace UE::AvaLevelViewport::Private;
 	return IsResetScaleEnabled() ? ViewportStatusBarButton::EnabledColor : ViewportStatusBarButton::DisabledColor;
 }
 
diff --git a/Engine/Plugins/Experimental/Avalanche/Source/AvalancheMaskEditor/Private/AvaMaskEditorLog.cpp b/Engine/Plugins/Experimental/Avalanche/Source/AvalancheMaskEditor/Private/AvaMaskEditorLog.cpp
--- a/Engine/Plugins/Experimental/Avalanche/Source/AvalancheMaskEditor/Private/AvaMaskEditorLog.cpp
+++ b/Engine/Plugins/Experimental/Avalanche/Source/AvalancheMaskEditor/Private/AvaMaskEditorLog.cpp
@@ -2,4 +2,4 @@
 
 #include "AvaMaskEditorLog.h"
 
-DEFINE_LOG_CATEGORY(LogAvalancheMaskEditor);
+DEFINE_LOG_CATEGORY(LogAvaMaskEditor);
diff --git a/Engine/Plugins/Experimental/Avalanche/Source/AvalancheMaskEditor/Private/AvaMaskEditorLog.h b/Engine/Plugins/Experimental/Avalanche/Source/AvalancheMaskEditor/Private/AvaMaskEditorLog.h
--- a/Engine/Plugins/Experimental/Avalanche/Source/AvalancheMaskEditor/Private/AvaMaskEditorLog.h
+++ b/Engine/Plugins/Experimental/Avalanche/Source/AvalancheMaskEditor/Private/AvaMaskEditorLog.h
@@ -4,4 +4,4 @@
 
 #include "Logging/LogMacros.h"
 
-DECLARE_LOG_CATEGORY_EXTERN(LogAvalancheMaskEditor, VeryVerbose, All);
+DECLARE_LOG_CATEGORY_EXTERN(LogAvaMaskEditor, VeryVerbose, All);
diff --git a/Engine/Plugins/Experimental/Avalanche/Source/AvalancheMaskEditor/Private/AvaMaskEditorMode.cpp b/Engine/Plugins/Experimental/Avalanche/Source/AvalancheMaskEditor/Private/AvaMaskEditorMode.cpp
--- a/Engine/Plugins/Experimental/Avalanche/Source/AvalancheMaskEditor/Private/AvaMaskEditorMode.cpp
+++ b/Engine/Plugins/Experimental/Avalanche/Source/AvalancheMaskEditor/Private/AvaMaskEditorMode.cpp
@@ -31,7 +31,7 @@
 
 #define LOCTEXT_NAMESPACE "AvalancheMaskEditorMode"
 
-const FEditorModeID UAvaMaskEditorMode::EM_MotionDesignMaskEditorModeId(UE::AvalancheMaskEditor::MotionDesignMaskEditorModeName);
+const FEditorModeID UAvaMaskEditorMode::EM_MotionDesignMaskEditorModeId(UE::AvaMaskEditor::MotionDesignMaskEditorModeName);
 
 UAvaMaskEditorMode::UAvaMaskEditorMode()
 {
@@ -77,7 +77,7 @@ void UAvaMaskEditorMode::Enter()
     
     	TSharedPtr<SWidget> ToolWidget = nullptr;
     	{
-    		static FName ToolkitOverlayMenuName = UE::AvalancheMaskEditor::Internal::ToolkitOverlayMenuName;
+    		static FName ToolkitOverlayMenuName = UE::AvaMaskEditor::Internal::ToolkitOverlayMenuName;
     		UToolMenu* Menu = UToolMenus::Get()->ExtendMenu(ToolkitOverlayMenuName);
     		Menu->SetStyleSet(&FAvaMaskEditorStyle::Get());
     		Menu->StyleName = TEXT("AvaMaskEditor.ViewportOverlayToolbar");
@@ -422,11 +422,11 @@ UAvaMask2DBaseModifier* UAvaMaskEditorMode::FindOrAddMaskModifier(AActor* InActo
 		{
 			if (InsertOp.FailReason)
 			{
-				UE_LOG(LogAvalancheMaskEditor, Error, TEXT("Error inserting Mask modifier: %s"), *InsertOp.FailReason->ToString());
+				UE_LOG(LogAvaMaskEditor, Error, TEXT("Error inserting Mask modifier: %s"), *InsertOp.FailReason->ToString());
 			}
 			else
 			{
-				UE_LOG(LogAvalancheMaskEditor, Error, TEXT("Error inserting Mask modifier."));
+				UE_LOG(LogAvaMaskEditor, Error, TEXT("Error inserting Mask modifier."));
 			}
 		}
 		else
diff --git a/Engine/Plugins/Experimental/Avalanche/Source/AvalancheMaskEditor/Private/AvaMaskEditorModule.cpp b/Engine/Plugins/Experimental/Avalanche/Source/AvalancheMaskEditor/Private/AvaMaskEditorModule.cpp
--- a/Engine/Plugins/Experimental/Avalanche/Source/AvalancheMaskEditor/Private/AvaMaskEditorModule.cpp
+++ b/Engine/Plugins/Experimental/Avalanche/Source/AvalancheMaskEditor/Private/AvaMaskEditorModule.cpp
@@ -116,7 +116,7 @@ void FAvalancheMaskEditorModule::RegisterMenus()
 
 	// Bottom center viewport overlay when mode is active
 	{
-		static FName ToolkitOverlayMenuName = UE::AvalancheMaskEditor::Internal::ToolkitOverlayMenuName;
+		static FName ToolkitOverlayMenuName = UE::AvaMaskEditor::Internal::ToolkitOverlayMenuName;
 
 		FToolMenuContext MenuContext(CommandList);
 		
diff --git a/Engine/Plugins/Experimental/Avalanche/Source/AvalancheMaskEditor/Private/AvaMaskEditorSVE.cpp b/Engine/Plugins/Experimental/Avalanche/Source/AvalancheMaskEditor/Private/AvaMaskEditorSVE.cpp
--- a/Engine/Plugins/Experimental/Avalanche/Source/AvalancheMaskEditor/Private/AvaMaskEditorSVE.cpp
+++ b/Engine/Plugins/Experimental/Avalanche/Source/AvalancheMaskEditor/Private/AvaMaskEditorSVE.cpp
@@ -33,7 +33,7 @@ void FAvaMaskSceneViewExtension::SetupView(
 	// Don't apply without this setting
 	if (GetDefault<URendererSettings>()->CustomDepthStencil != ECustomDepthStencil::EnabledWithStencil)
 	{
-		UE_LOG(LogAvalancheMaskEditor, Warning, TEXT("The Mask overlay requires CustomDepth to be \"EnabledWithStencil\" in Project Settings."))
+		UE_LOG(LogAvaMaskEditor, Warning, TEXT("The Mask overlay requires CustomDepth to be \"EnabledWithStencil\" in Project Settings."))
 	}
 	else
 	{
diff --git a/Engine/Plugins/Experimental/Avalanche/Source/AvalancheMaskEditor/Private/AvaMaskEditorSubsystem.h b/Engine/Plugins/Experimental/Avalanche/Source/AvalancheMaskEditor/Private/AvaMaskEditorSubsystem.h
--- a/Engine/Plugins/Experimental/Avalanche/Source/AvalancheMaskEditor/Private/AvaMaskEditorSubsystem.h
+++ b/Engine/Plugins/Experimental/Avalanche/Source/AvalancheMaskEditor/Private/AvaMaskEditorSubsystem.h
@@ -9,7 +9,7 @@
 
 class FAvaMaskSceneViewExtension;
 
-namespace UE::AvalancheMaskEditor::Internal
+namespace UE::AvaMaskEditor::Internal
 {
 	static FName ToolkitOverlayMenuName = TEXT("AvalancheMask.Editor.ModeOverlayToolbar");
 }
diff --git a/Engine/Plugins/Experimental/Avalanche/Source/AvalancheMaskEditor/Public/IAvaMaskEditor.h b/Engine/Plugins/Experimental/Avalanche/Source/AvalancheMaskEditor/Public/IAvaMaskEditor.h
--- a/Engine/Plugins/Experimental/Avalanche/Source/AvalancheMaskEditor/Public/IAvaMaskEditor.h
+++ b/Engine/Plugins/Experimental/Avalanche/Source/AvalancheMaskEditor/Public/IAvaMaskEditor.h
@@ -4,7 +4,7 @@
 
 #include "UObject/NameTypes.h"
 
-namespace UE::AvalancheMaskEditor
+namespace UE::AvaMaskEditor
 {
 	const static FName MotionDesignMaskEditorModeName("EditMode.MotionDesignMask");
 }
diff --git a/Engine/Plugins/Experimental/Avalanche/Source/AvalancheMedia/Private/AvaMediaStyle.cpp b/Engine/Plugins/Experimental/Avalanche/Source/AvalancheMedia/Private/AvaMediaStyle.cpp
--- a/Engine/Plugins/Experimental/Avalanche/Source/AvalancheMedia/Private/AvaMediaStyle.cpp
+++ b/Engine/Plugins/Experimental/Avalanche/Source/AvalancheMedia/Private/AvaMediaStyle.cpp
@@ -18,7 +18,7 @@ FAvaMediaStyle::FAvaMediaStyle()
 	SetContentRoot(Plugin->GetBaseDir() / TEXT("Resources"));
 	
 	Set("AvaMedia.UnrealIcon", new IMAGE_BRUSH("Icons/MediaIcons/ue_logo", Icon200x200));
-	Set("ClassIcon.AvalancheBroadcast", new IMAGE_BRUSH("Icons/MediaIcons/MediaOutput", Icon16x16));
+	Set("ClassIcon.AvaBroadcast", new IMAGE_BRUSH("Icons/MediaIcons/MediaOutput", Icon16x16));
 
 	FSlateStyleRegistry::RegisterSlateStyle(*this);
 }
diff --git a/Engine/Plugins/Experimental/Avalanche/Source/AvalancheMedia/Private/Playable/AvaPlayable.cpp b/Engine/Plugins/Experimental/Avalanche/Source/AvalancheMedia/Private/Playable/AvaPlayable.cpp
--- a/Engine/Plugins/Experimental/Avalanche/Source/AvalancheMedia/Private/Playable/AvaPlayable.cpp
+++ b/Engine/Plugins/Experimental/Avalanche/Source/AvalancheMedia/Private/Playable/AvaPlayable.cpp
@@ -367,7 +367,7 @@ UAvaPlayable* UAvaPlayable::CreateLocalPlayable(UObject* InOuter, const FPlayabl
 {
 	switch (InPlayableInfo.SourceAsset.GetAssetType())
 	{
-	case EAvalancheAssetType::World:
+	case EMotionDesignAssetType::World:
 		return NewObject<UAvaPlayableLevelStreaming>(InOuter ? InOuter : GEngine);
 	default:
 		UE_LOG(LogAvaPlayable, Error, TEXT("Asset \"%s\" is an unsupported type."), *InPlayableInfo.SourceAsset.ToSoftObjectPath().ToString());
diff --git a/Engine/Plugins/Experimental/Avalanche/Source/AvalancheMedia/Private/Playable/Playables/AvaPlayableLevelStreaming.cpp b/Engine/Plugins/Experimental/Avalanche/Source/AvalancheMedia/Private/Playable/Playables/AvaPlayableLevelStreaming.cpp
--- a/Engine/Plugins/Experimental/Avalanche/Source/AvalancheMedia/Private/Playable/Playables/AvaPlayableLevelStreaming.cpp
+++ b/Engine/Plugins/Experimental/Avalanche/Source/AvalancheMedia/Private/Playable/Playables/AvaPlayableLevelStreaming.cpp
@@ -208,7 +208,7 @@ bool UAvaPlayableLevelStreaming::LoadAsset(const FAvaSoftAssetPtr& InSourceAsset
 	const FAvaInstanceSettings& PlaybackInstanceSettings = IAvaMediaModule::Get().GetAvaInstanceSettings();
 	bLoadSubPlayables = PlaybackInstanceSettings.bEnableLoadSubPlayables;
 	
-	check(InSourceAsset.GetAssetType() == EAvalancheAssetType::World);
+	check(InSourceAsset.GetAssetType() == EMotionDesignAssetType::World);
 	const bool bAssetLoading = LoadLevel(TSoftObjectPtr<UWorld>(InSourceAsset.ToSoftObjectPath()), bInInitiallyVisible); 
 	if (bAssetLoading)
 	{
diff --git a/Engine/Plugins/Experimental/Avalanche/Source/AvalancheMedia/Private/Playable/Playables/AvaPlayableRemoteProxy.cpp b/Engine/Plugins/Experimental/Avalanche/Source/AvalancheMedia/Private/Playable/Playables/AvaPlayableRemoteProxy.cpp
--- a/Engine/Plugins/Experimental/Avalanche/Source/AvalancheMedia/Private/Playable/Playables/AvaPlayableRemoteProxy.cpp
+++ b/Engine/Plugins/Experimental/Avalanche/Source/AvalancheMedia/Private/Playable/Playables/AvaPlayableRemoteProxy.cpp
@@ -187,7 +187,7 @@ bool UAvaPlayableRemoteProxy::InitPlayable(const FPlayableCreationInfo& InPlayab
 	// Remote playables have proxy playable groups imitating the same setup as local ones.
 	switch (InPlayableInfo.SourceAsset.GetAssetType())
 	{
-	case EAvalancheAssetType::World:
+	case EMotionDesignAssetType::World:
 		PlayableGroup = InPlayableInfo.PlayableGroupManager->GetOrCreateSharedLevelGroup(InPlayableInfo.ChannelName, bIsRemoteProxy);
 		break;
 	default:
diff --git a/Engine/Plugins/Experimental/Avalanche/Source/AvalancheMedia/Private/Playback/AvaPlaybackUtils.cpp b/Engine/Plugins/Experimental/Avalanche/Source/AvalancheMedia/Private/Playback/AvaPlaybackUtils.cpp
--- a/Engine/Plugins/Experimental/Avalanche/Source/AvalancheMedia/Private/Playback/AvaPlaybackUtils.cpp
+++ b/Engine/Plugins/Experimental/Avalanche/Source/AvalancheMedia/Private/Playback/AvaPlaybackUtils.cpp
@@ -164,15 +164,15 @@ bool FAvaPlaybackUtils::IsMapAsset(const FString& InPackageName)
 
 bool FAvaPlaybackUtils::IsPlayableAsset(const FAssetData& InAssetData)
 {
-	const EAvalancheAssetType AssetType = FAvaSoftAssetPath::GetAssetTypeFromClass(InAssetData.GetClass(), true);
-	if (AssetType == EAvalancheAssetType::Unknown)
+	const EMotionDesignAssetType AssetType = FAvaSoftAssetPath::GetAssetTypeFromClass(InAssetData.GetClass(), true);
+	if (AssetType == EMotionDesignAssetType::Unknown)
 	{
 		return false;
 	}
 	// For world type, we need to check the tags.
-	if (AssetType == EAvalancheAssetType::World)
+	if (AssetType == EMotionDesignAssetType::World)
 	{
-		const FAssetTagValueRef SceneTag = InAssetData.TagsAndValues.FindTag(UE::Ava::AssetTags::AvalancheScene);
+		const FAssetTagValueRef SceneTag = InAssetData.TagsAndValues.FindTag(UE::Ava::AssetTags::MotionDesignScene);
 		if (!SceneTag.IsSet() || !SceneTag.Equals(UE::Ava::AssetTags::Values::Enabled))
 		{
 			return false;
diff --git a/Engine/Plugins/Experimental/Avalanche/Source/AvalancheMedia/Private/Rundown/AvaRundown.cpp b/Engine/Plugins/Experimental/Avalanche/Source/AvalancheMedia/Private/Rundown/AvaRundown.cpp
--- a/Engine/Plugins/Experimental/Avalanche/Source/AvalancheMedia/Private/Rundown/AvaRundown.cpp
+++ b/Engine/Plugins/Experimental/Avalanche/Source/AvalancheMedia/Private/Rundown/AvaRundown.cpp
@@ -1495,7 +1495,7 @@ void UAvaRundown::NotifyPageRemoteControlValueChanged(int32 InPageId, EAvaPlayab
 	if (EnumHasAnyFlags(InRemoteControlChanges, EAvaPlayableRemoteControlChanges::EntityValues))
 	{
 		// For now, potentially pushing all values multiple time (per frame) is mitigated by the
-		// optimization in FAvalancheRemoteControlUtils::SetValueOfEntity that
+		// optimization in FAvaRemoteControlUtils::SetValueOfEntity that
 		// will only set the value of the entity if it changed.
 		PushRuntimeRemoteControlValues(InPageId, true);
 	}
diff --git a/Engine/Plugins/Experimental/Avalanche/Source/AvalancheMediaEditor/Private/Broadcast/OutputDevices/AvaBroadcastOutputTreeItem.h b/Engine/Plugins/Experimental/Avalanche/Source/AvalancheMediaEditor/Private/Broadcast/OutputDevices/AvaBroadcastOutputTreeItem.h
--- a/Engine/Plugins/Experimental/Avalanche/Source/AvalancheMediaEditor/Private/Broadcast/OutputDevices/AvaBroadcastOutputTreeItem.h
+++ b/Engine/Plugins/Experimental/Avalanche/Source/AvalancheMediaEditor/Private/Broadcast/OutputDevices/AvaBroadcastOutputTreeItem.h
@@ -14,7 +14,7 @@ struct FGeometry;
 struct FPointerEvent;
 struct FSlateBrush;
 
-typedef TSharedPtr<IAvaBroadcastOutputTreeItem> FAvaOutputTreeItemPtr;
+using FAvaOutputTreeItemPtr = TSharedPtr<IAvaBroadcastOutputTreeItem>;
 
 class IAvaBroadcastOutputTreeItem : public IAvaTypeCastable, public TSharedFromThis<IAvaBroadcastOutputTreeItem>
 {
diff --git a/Engine/Plugins/Experimental/Avalanche/Source/AvalancheMediaEditor/Private/Playback/Graph/AvaPlaybackConnectionDrawingPolicy.h b/Engine/Plugins/Experimental/Avalanche/Source/AvalancheMediaEditor/Private/Playback/Graph/AvaPlaybackConnectionDrawingPolicy.h
--- a/Engine/Plugins/Experimental/Avalanche/Source/AvalancheMediaEditor/Private/Playback/Graph/AvaPlaybackConnectionDrawingPolicy.h
+++ b/Engine/Plugins/Experimental/Avalanche/Source/AvalancheMediaEditor/Private/Playback/Graph/AvaPlaybackConnectionDrawingPolicy.h
@@ -42,9 +42,9 @@ protected:
 	};
 
 	// Map of pairings
-	typedef TMap<UEdGraphNode*, FTimePair> FExecPairingMap;
+	using FExecPairingMap = TMap<UEdGraphNode*, FTimePair>;
 
-	// Map of nodes that preceeded before a given node in the execution sequence (one entry for each pairing)
+	// Map of nodes that preceded before a given node in the execution sequence (one entry for each pairing)
 	TMap<UEdGraphNode*, FExecPairingMap> PredecessorNodes;
 
 	UEdGraph* GraphObject;
diff --git a/Engine/Plugins/Experimental/Avalanche/Source/AvalancheMediaEditor/Private/Rundown/AvaRundownServer.h b/Engine/Plugins/Experimental/Avalanche/Source/AvalancheMediaEditor/Private/Rundown/AvaRundownServer.h
--- a/Engine/Plugins/Experimental/Avalanche/Source/AvalancheMediaEditor/Private/Rundown/AvaRundownServer.h
+++ b/Engine/Plugins/Experimental/Avalanche/Source/AvalancheMediaEditor/Private/Rundown/AvaRundownServer.h
@@ -176,7 +176,7 @@ private:
 
 		FDelegateHandle OnPlaybackInstanceStatusChangedDelegateHandle;
 		
-		typedef TFunctionRef<void(UAvaRundown*)> FRundownEventFunction;
+		using FRundownEventFunction = TFunctionRef<void(UAvaRundown*)>;
 		
 		/**
 		 * Returns requested rundown specified by InRundownPath. Will load it if necessary or returned the cached one if it is the same.
diff --git a/Engine/Plugins/Experimental/Avalanche/Source/AvalancheMediaEditor/Private/Rundown/Pages/Slate/SAvaRundownPageList.cpp b/Engine/Plugins/Experimental/Avalanche/Source/AvalancheMediaEditor/Private/Rundown/Pages/Slate/SAvaRundownPageList.cpp
--- a/Engine/Plugins/Experimental/Avalanche/Source/AvalancheMediaEditor/Private/Rundown/Pages/Slate/SAvaRundownPageList.cpp
+++ b/Engine/Plugins/Experimental/Avalanche/Source/AvalancheMediaEditor/Private/Rundown/Pages/Slate/SAvaRundownPageList.cpp
@@ -845,7 +845,7 @@ namespace UE::AvaPageList::Private
 {
 	static bool IsAvaAsset(const FAssetData& InAssetData, bool bInLogInfo)
 	{
-		static const FName& SceneAssetTag = UE::Ava::AssetTags::AvalancheScene;
+		static const FName& SceneAssetTagName = UE::Ava::AssetTags::MotionDesignScene;
 		static const FString& AssetTagValueEnabled = UE::Ava::AssetTags::Values::Enabled;
 
 		const UClass* AssetClass = InAssetData.GetClass(EResolveClass::Yes);
@@ -854,18 +854,18 @@ namespace UE::AvaPageList::Private
 			return false;
 		}
 		
-		const EAvalancheAssetType AssetType = FAvaSoftAssetPath::GetAssetTypeFromClass(AssetClass, true);
-		if (AssetType == EAvalancheAssetType::Unknown)
+		const EMotionDesignAssetType AssetType = FAvaSoftAssetPath::GetAssetTypeFromClass(AssetClass, true);
+		if (AssetType == EMotionDesignAssetType::Unknown)
 		{
 			return false;
 		}
 		
 		// If the asset is a level, we check if it has the Motion Design scene tag.
-		if (AssetType == EAvalancheAssetType::World)
+		if (AssetType == EMotionDesignAssetType::World)
 		{
-			const FAssetTagValueRef AvalancheSceneTag = InAssetData.TagsAndValues.FindTag(SceneAssetTag);
+			const FAssetTagValueRef SceneTag = InAssetData.TagsAndValues.FindTag(SceneAssetTagName);
 			
-			if (!AvalancheSceneTag.IsSet())
+			if (!SceneTag.IsSet())
 			{
 				if (bInLogInfo)
 				{
@@ -876,7 +876,7 @@ namespace UE::AvaPageList::Private
 				return false;
 			}
 			
-			if (!AvalancheSceneTag.Equals(AssetTagValueEnabled))
+			if (!SceneTag.Equals(AssetTagValueEnabled))
 			{
 				if (bInLogInfo)
 				{
diff --git a/Engine/Plugins/Experimental/Avalanche/Source/AvalancheModifiers/Internal/Tests/Framework/AvaModifiersTestUtils.h b/Engine/Plugins/Experimental/Avalanche/Source/AvalancheModifiers/Internal/Tests/Framework/AvaModifiersTestUtils.h
--- a/Engine/Plugins/Experimental/Avalanche/Source/AvalancheModifiers/Internal/Tests/Framework/AvaModifiersTestUtils.h
+++ b/Engine/Plugins/Experimental/Avalanche/Source/AvalancheModifiers/Internal/Tests/Framework/AvaModifiersTestUtils.h
@@ -8,7 +8,7 @@
 #include "UObject/WeakObjectPtr.h"
 #include "UObject/WeakObjectPtrTemplates.h"
 
-// General Avalanche unit test log
+// General Motion Design unit test log
 DECLARE_LOG_CATEGORY_EXTERN(LogAvaModifiersTest, Log, All);
 
 class AAvaShapeActor;
diff --git a/Engine/Plugins/Experimental/Avalanche/Source/AvalancheModifiers/Private/Modifiers/AvaExtrudeModifier.cpp b/Engine/Plugins/Experimental/Avalanche/Source/AvalancheModifiers/Private/Modifiers/AvaExtrudeModifier.cpp
--- a/Engine/Plugins/Experimental/Avalanche/Source/AvalancheModifiers/Private/Modifiers/AvaExtrudeModifier.cpp
+++ b/Engine/Plugins/Experimental/Avalanche/Source/AvalancheModifiers/Private/Modifiers/AvaExtrudeModifier.cpp
@@ -148,7 +148,7 @@ void UAvaExtrudeModifier::Apply()
 				{
 					const FVector UVSize = GetMeshBounds().GetSize() / EditMesh.GetBounds().MaxDim();
 					FAvaShapeMaterialUVParameters& UVParams = *ShapeComponent->GetInUseMaterialUVParams(UAvaShapeDynamicMeshBase::MESH_INDEX_PRIMARY);
-					UE::AvalancheShapes::TransformMeshUVs(EditMesh, ExtrudeUVIds.Array(), UVParams, FVector2D(UVSize.Y, UVSize.Z), FVector2D(0.5, 0.5), 90.f);
+					UE::AvaShapes::TransformMeshUVs(EditMesh, ExtrudeUVIds.Array(), UVParams, FVector2D(UVSize.Y, UVSize.Z), FVector2D(0.5, 0.5), 90.f);
 				}
 			}
 		}
diff --git a/Engine/Plugins/Experimental/Avalanche/Source/AvalancheModifiers/Private/Modifiers/AvaOutlineModifier.cpp b/Engine/Plugins/Experimental/Avalanche/Source/AvalancheModifiers/Private/Modifiers/AvaOutlineModifier.cpp
--- a/Engine/Plugins/Experimental/Avalanche/Source/AvalancheModifiers/Private/Modifiers/AvaOutlineModifier.cpp
+++ b/Engine/Plugins/Experimental/Avalanche/Source/AvalancheModifiers/Private/Modifiers/AvaOutlineModifier.cpp
@@ -148,7 +148,7 @@ void UAvaOutlineModifier::Apply()
 			if (UAvaShapeDynamicMeshBase* ShapeGenerator = ShapeActor->GetDynamicMesh())
 			{
 				FAvaShapeMaterialUVParameters& UVParams = *ShapeGenerator->GetInUseMaterialUVParams(UAvaShapeDynamicMeshBase::MESH_INDEX_PRIMARY);
-				UE::AvalancheShapes::TransformMeshUVs(EditMesh, Result.NewUVElements, UVParams, FVector2D(Size3D.Y,Size3D.Z), FVector2D(0.5, 0.5), 0.f);
+				UE::AvaShapes::TransformMeshUVs(EditMesh, Result.NewUVElements, UVParams, FVector2D(Size3D.Y,Size3D.Z), FVector2D(0.5, 0.5), 0.f);
 			}
 		}
 	});
diff --git a/Engine/Plugins/Experimental/Avalanche/Source/AvalancheModifiers/Private/Tests/Framework/AvaModifiersTestUtils.cpp b/Engine/Plugins/Experimental/Avalanche/Source/AvalancheModifiers/Private/Tests/Framework/AvaModifiersTestUtils.cpp
--- a/Engine/Plugins/Experimental/Avalanche/Source/AvalancheModifiers/Private/Tests/Framework/AvaModifiersTestUtils.cpp
+++ b/Engine/Plugins/Experimental/Avalanche/Source/AvalancheModifiers/Private/Tests/Framework/AvaModifiersTestUtils.cpp
@@ -18,7 +18,7 @@
 #include "Tests/Framework/AvaTestDynamicMeshActor.h"
 #include "Tests/Framework/AvaTestStaticMeshActor.h"
 
-// General Avalanche Modifiers unit test log
+// General Motion Design Modifiers unit test log
 DEFINE_LOG_CATEGORY(LogAvaModifiersTest);
 
 FAvaModifierTestUtils::FAvaModifierTestUtils(TSharedPtr<FAvaTestUtils> InTestUtils)
diff --git a/Engine/Plugins/Experimental/Avalanche/Source/AvalancheOutliner/Private/Columns/AvaOutlinerColumnExtender.h b/Engine/Plugins/Experimental/Avalanche/Source/AvalancheOutliner/Private/Columns/AvaOutlinerColumnExtender.h
--- a/Engine/Plugins/Experimental/Avalanche/Source/AvalancheOutliner/Private/Columns/AvaOutlinerColumnExtender.h
+++ b/Engine/Plugins/Experimental/Avalanche/Source/AvalancheOutliner/Private/Columns/AvaOutlinerColumnExtender.h
@@ -20,7 +20,7 @@ class IAvaOutlinerColumn;
  */
 class FAvaOutlinerColumnExtender
 {
-	typedef TSharedPtr<IAvaOutlinerColumn> FAvaOutlinerColumnPtr;
+	using FAvaOutlinerColumnPtr = TSharedPtr<IAvaOutlinerColumn>;
 
 public:
 	template<typename InColumnType, EAvaOutlinerExtensionPosition InExtensionPosition, typename InRefColumnType
diff --git a/Engine/Plugins/Experimental/Avalanche/Source/AvalancheOutliner/Private/Data/AvaOutlinerSaveState.h b/Engine/Plugins/Experimental/Avalanche/Source/AvalancheOutliner/Private/Data/AvaOutlinerSaveState.h
--- a/Engine/Plugins/Experimental/Avalanche/Source/AvalancheOutliner/Private/Data/AvaOutlinerSaveState.h
+++ b/Engine/Plugins/Experimental/Avalanche/Source/AvalancheOutliner/Private/Data/AvaOutlinerSaveState.h
@@ -72,7 +72,7 @@ private:
 	void UpdateItemIdContexts(FStringView InOldContext, FStringView InNewContext);
 
 	/**
-	 * Deprecated in favor of using Avalanche Scene Tree
+	 * Deprecated in favor of using Motion Design Scene Tree
 	 * @see FAvaSceneTree
 	 */
 	TMap<FString, int32> ItemSorting_DEPRECATED;
