[Motion Design]

diff --git a/Engine/Plugins/Experimental/Avalanche/Source/AvalancheOutliner/Private/Data/AvaOutlinerVersion.h b/Engine/Plugins/Experimental/Avalanche/Source/AvalancheOutliner/Private/Data/AvaOutlinerVersion.h
--- a/Engine/Plugins/Experimental/Avalanche/Source/AvalancheOutliner/Private/Data/AvaOutlinerVersion.h
+++ b/Engine/Plugins/Experimental/Avalanche/Source/AvalancheOutliner/Private/Data/AvaOutlinerVersion.h
@@ -4,7 +4,7 @@
 
 #include "Misc/Guid.h"
 
-/** Version used for serializing Avalanche Outliner Data */
+/** Version used for serializing Motion Design Outliner Data */
 struct FAvaOutlinerVersion
 {
 private:
@@ -25,7 +25,7 @@ public:
 		/** Outliner View: Serialization of Outliner Column Visibility */
 		ColumnVisibility,
 
-		/** Outliner: Moved Item Ordering Serialization to using Avalanche Scene Tree */
+		/** Outliner: Moved Item Ordering Serialization to using Motion Design Scene Tree */
 		SceneTree,
 
 		/** Outliner: Changed Item Id to be Full Object Path String rather than just the Sub String */
diff --git a/Engine/Plugins/Experimental/Avalanche/Source/AvalancheOutliner/Private/Tests/AvaOutlinerTestUtils.cpp b/Engine/Plugins/Experimental/Avalanche/Source/AvalancheOutliner/Private/Tests/AvaOutlinerTestUtils.cpp
--- a/Engine/Plugins/Experimental/Avalanche/Source/AvalancheOutliner/Private/Tests/AvaOutlinerTestUtils.cpp
+++ b/Engine/Plugins/Experimental/Avalanche/Source/AvalancheOutliner/Private/Tests/AvaOutlinerTestUtils.cpp
@@ -86,7 +86,7 @@ UE::AvaOutliner::Private::FAvaOutlinerProviderTest::FAvaOutlinerProviderTest()
 	, ModeTools(MakeShared<FAvaOutlinerEditorModeTools>())
 {
 	constexpr EWorldType::Type WorldType = EWorldType::Editor;
-	World = UWorld::CreateWorld(WorldType, false, FName(TEXT("AvalancheOutlinerTestWorld")));
+	World = UWorld::CreateWorld(WorldType, false, FName(TEXT("MotionDesignOutlinerTestWorld")));
 
 	UPackage* const Package = World->GetPackage();
 	Package->SetFlags(RF_Transient | RF_Public);
diff --git a/Engine/Plugins/Experimental/Avalanche/Source/AvalancheOutliner/Public/AvaOutlinerDefines.h b/Engine/Plugins/Experimental/Avalanche/Source/AvalancheOutliner/Public/AvaOutlinerDefines.h
--- a/Engine/Plugins/Experimental/Avalanche/Source/AvalancheOutliner/Public/AvaOutlinerDefines.h
+++ b/Engine/Plugins/Experimental/Avalanche/Source/AvalancheOutliner/Public/AvaOutlinerDefines.h
@@ -11,10 +11,10 @@ class AActor;
 class IAvaOutlinerItem;
 class USceneComponent;
 
-typedef TPair<FName, FLinearColor>   FAvaOutlinerColorPair;
-typedef const IAvaOutlinerItem&      FAvaOutlinerFilterType;
-typedef TSharedPtr<IAvaOutlinerItem> FAvaOutlinerItemPtr;
-typedef TWeakPtr<IAvaOutlinerItem>   FAvaOutlinerItemWeakPtr;
+using FAvaOutlinerColorPair = TPair<FName, FLinearColor>;
+using FAvaOutlinerFilterType = const IAvaOutlinerItem&;
+using FAvaOutlinerItemPtr = TSharedPtr<IAvaOutlinerItem>;
+using FAvaOutlinerItemWeakPtr = TWeakPtr<IAvaOutlinerItem>;
 
 enum class EAvaOutlinerItemFlags : uint8
 {
diff --git a/Engine/Plugins/Experimental/Avalanche/Source/AvalancheOutliner/Public/Item/IAvaOutlinerItem.h b/Engine/Plugins/Experimental/Avalanche/Source/AvalancheOutliner/Public/Item/IAvaOutlinerItem.h
--- a/Engine/Plugins/Experimental/Avalanche/Source/AvalancheOutliner/Public/Item/IAvaOutlinerItem.h
+++ b/Engine/Plugins/Experimental/Avalanche/Source/AvalancheOutliner/Public/Item/IAvaOutlinerItem.h
@@ -31,7 +31,7 @@ struct FSlateIcon;
 class IAvaOutlinerItem : public IAvaTypeCastable, public TSharedFromThis<IAvaOutlinerItem>
 {
 protected:
-	typedef TArray<IAvaOutlinerItem>::SizeType IndexType;
+	using IndexType = TArray<IAvaOutlinerItem>::SizeType;
 
 public:
 	UE_AVA_INHERITS(IAvaOutlinerItem, IAvaTypeCastable);
@@ -251,7 +251,7 @@ public:
 	 * const TSharedPtr<FAvaOutlinerView>& - The Outliner View where this Expansion happened
 	 * bool - Whether the Item is expanded
 	 */
-	typedef TMulticastDelegate<void(const TSharedPtr<FAvaOutlinerView>&, bool)> FOnExpansionChanged;
+	using FOnExpansionChanged = TMulticastDelegate<void(const TSharedPtr<FAvaOutlinerView>&, bool)>;
 	
 	/** Called when Expansion state (Expanded/Collapsed) has been changed */
 	virtual FOnExpansionChanged& OnExpansionChanged() = 0;
@@ -310,7 +310,7 @@ public:
 	 * EAvaOutlinerRenameAction the type of action being relayed (e.g. request a rename, or notify rename complete, etc)
 	 * const TSharedPtr<FAvaOutlinerView>& - The Outliner View where the rename action is taking place
 	 */
-	typedef TMulticastDelegate<void(EAvaOutlinerRenameAction, const TSharedPtr<FAvaOutlinerView>&)> FOnRenameAction;
+	using FOnRenameAction = TMulticastDelegate<void(EAvaOutlinerRenameAction, const TSharedPtr<FAvaOutlinerView>&)>;
 	
 	/** Broadcasts whenever a rename action takes place from a given view (e.g. when pressing "F2" to rename, or committing the rename text) */
 	virtual FOnRenameAction& OnRenameAction() = 0;
diff --git a/Engine/Plugins/Experimental/Avalanche/Source/AvalancheRemoteControl/Private/AvaRemoteControlRebind.cpp b/Engine/Plugins/Experimental/Avalanche/Source/AvalancheRemoteControl/Private/AvaRemoteControlRebind.cpp
--- a/Engine/Plugins/Experimental/Avalanche/Source/AvalancheRemoteControl/Private/AvaRemoteControlRebind.cpp
+++ b/Engine/Plugins/Experimental/Avalanche/Source/AvalancheRemoteControl/Private/AvaRemoteControlRebind.cpp
@@ -53,7 +53,7 @@ namespace UE::AvaRemoteControl::Private
 }
 
 /**
- *	Since the runtime Avalanche Scene Object is not in the world where the
+ *	Since the runtime Motion Design Scene Object is not in the world where the
  *	remote control bindings were done, it needs to be rebound.
  *	
  *	The original FRemoteControlPresetRebindingManager::Rebind_NewAlgo
@@ -69,7 +69,7 @@ public:
 
 	/**
 	 * Version of the function returning only actors from the specified level.
-	 * This is required for support Avalanche graphics instancing as sub-levels.
+	 * This is required for support Motion Design graphics instancing as sub-levels.
 	 * There is one RCP per sub-level and must resolve only to that level.
 	 */
 	static bool GetActorsOfClass(const ULevel* InLevel, const UClass* InTargetClass, TArray<AActor*>& OutActors)
@@ -91,7 +91,7 @@ public:
 	/**
 	 *	The original function RCPresetRebindingManager::GetActorsOfClass
 	 *	prevents rebinding with transient actors. All the runtime RCP from embedded
-	 *	Avalanche graphics are transient. This is the reason this has to be rewritten.
+	 *	Motion Design graphics are transient. This is the reason this has to be rewritten.
 	 */
 	static bool GetActorsOfClass(UWorld* InPresetWorld, UClass* InTargetClass, TArray<AActor*>& OutActors)
 	{
diff --git a/Engine/Plugins/Experimental/Avalanche/Source/AvalancheRemoteControl/Private/AvaRemoteControlUtils.cpp b/Engine/Plugins/Experimental/Avalanche/Source/AvalancheRemoteControl/Private/AvaRemoteControlUtils.cpp
--- a/Engine/Plugins/Experimental/Avalanche/Source/AvalancheRemoteControl/Private/AvaRemoteControlUtils.cpp
+++ b/Engine/Plugins/Experimental/Avalanche/Source/AvalancheRemoteControl/Private/AvaRemoteControlUtils.cpp
@@ -43,7 +43,7 @@ bool FAvaRemoteControlUtils::RegisterRemoteControlPreset(URemoteControlPreset* I
 			}
 		
 			// Workaround to enable multiple instances of the same RCP to be registered.
-			// This is a normal occurence with Avalanche graphics that are meant to be templates that are used in
+			// This is a normal occurence with Motion Design graphics that are meant to be templates that are used in
 			// multiple pages. Because of page transitions, we must support having multiple instances of the same template
 			// running together at the same time.
 			// We need to assign a new unique ID otherwise all the RCP instances will have the
@@ -67,7 +67,7 @@ bool FAvaRemoteControlUtils::RegisterRemoteControlPreset(URemoteControlPreset* I
 		RemoteControlComponentsSubsystem->RegisterPreset(InRemoteControlPreset);
 	}
 
-	// Avalanche's RCPs, from an ava asset (e.g. level) are considered "embedded" because
+	// Motion Design's RCPs, from an ava asset (e.g. level) are considered "embedded" because
 	// they are not an asset.
 	constexpr bool bReplaceExisting = true;
 	return RemoteControlModule.RegisterEmbeddedPreset(InRemoteControlPreset, bReplaceExisting);
diff --git a/Engine/Plugins/Experimental/Avalanche/Source/AvalancheSequence/Private/AvaSequenceController.cpp b/Engine/Plugins/Experimental/Avalanche/Source/AvalancheSequence/Private/AvaSequenceController.cpp
--- a/Engine/Plugins/Experimental/Avalanche/Source/AvalancheSequence/Private/AvaSequenceController.cpp
+++ b/Engine/Plugins/Experimental/Avalanche/Source/AvalancheSequence/Private/AvaSequenceController.cpp
@@ -13,7 +13,7 @@
 
 DEFINE_LOG_CATEGORY_STATIC(LogAvaSequenceController, Log, All);
 
-namespace UE::Avalanche::Private
+namespace UE::Ava::Private
 {
 	double GetFrameTimeAsDisplayDecimal(UMovieScene& InMovieScene, const FFrameTime& InTickTime)
 	{
@@ -22,7 +22,7 @@ namespace UE::Avalanche::Private
 
 #if !NO_LOGGING
 	bool bAvaLogMarkedFrames = false;
-	FAutoConsoleVariableRef CVarLogAvaSequenceController(TEXT("Log.Avalanche.SequenceController")
+	FAutoConsoleVariableRef CVarLogAvaSequenceController(TEXT("Log.MotionDesign.SequenceController")
 		, bAvaLogMarkedFrames
 		, TEXT("Log FAvaSequenceController"));
 #endif
@@ -49,7 +49,7 @@ public:
 	static bool IsLoggingEnabled()
 	{
 #if !NO_LOGGING
-		return UE::Avalanche::Private::bAvaLogMarkedFrames;
+		return UE::Ava::Private::bAvaLogMarkedFrames;
 #else
 		return false;
 #endif
@@ -135,8 +135,8 @@ void FAvaSequenceController::Tick(const FAvaSequencePlayerVariant& InPlayerVaria
 
 	UE_CLOG(FAvaSequenceControllerTickScope::IsLoggingEnabled(), LogAvaSequenceController, Log
 		, TEXT("Previous Frame: %f ---  Current Frame: %f")
-		, UE::Avalanche::Private::GetFrameTimeAsDisplayDecimal(MovieScene, PreviousFrame)
-		, UE::Avalanche::Private::GetFrameTimeAsDisplayDecimal(MovieScene, CurrentFrame));
+		, UE::Ava::Private::GetFrameTimeAsDisplayDecimal(MovieScene, PreviousFrame)
+		, UE::Ava::Private::GetFrameTimeAsDisplayDecimal(MovieScene, CurrentFrame));
 
 	SortMarks(MovieScene);
 
@@ -224,7 +224,7 @@ FString FAvaSequenceController::GetMarkRoleAsString(const UAvaSequence& InSequen
 
 FString FAvaSequenceController::GetMarkAsString(UMovieScene& InMovieScene, const FMovieSceneMarkedFrame& InMark) const
 {
-	const double DisplayFrame = UE::Avalanche::Private::GetFrameTimeAsDisplayDecimal(InMovieScene, InMark.FrameNumber);
+	const double DisplayFrame = UE::Ava::Private::GetFrameTimeAsDisplayDecimal(InMovieScene, InMark.FrameNumber);
 
 	return InMark.Label
 		+ TEXT(" @ ")
@@ -261,7 +261,7 @@ void FAvaSequenceController::UpdatePlaybackVariables()
 {
 	PlayDirection = PlaybackContext->IsPlayingForwards() ? EPlayDirection::Forwards : EPlayDirection::Backwards;
 
-	// NOTE: Marks in both Sequencer + Avalanche Player are processed prior to these Players processing Tick this frame
+	// NOTE: Marks in both Sequencer + Motion Design Player are processed prior to these Players processing Tick this frame
 	// So Global Time is currently the result of the Last Frame, rather than the result of this frame
 	const FFrameTime LastFrameTime = PlaybackContext->GetGlobalTime();
 
@@ -336,8 +336,8 @@ TArray<const FMovieSceneMarkedFrame*> FAvaSequenceController::FindIntersectedMar
 		if (IntersectedMarks.IsEmpty())
 		{
 			UE_LOG(LogAvaSequenceController, Log, TEXT("No Marks between %f and %f")
-				, UE::Avalanche::Private::GetFrameTimeAsDisplayDecimal(InMovieScene, LowerBound)
-				, UE::Avalanche::Private::GetFrameTimeAsDisplayDecimal(InMovieScene, UpperBound));	
+				, UE::Ava::Private::GetFrameTimeAsDisplayDecimal(InMovieScene, LowerBound)
+				, UE::Ava::Private::GetFrameTimeAsDisplayDecimal(InMovieScene, UpperBound));	
 		}
 		else
 		{
@@ -348,8 +348,8 @@ TArray<const FMovieSceneMarkedFrame*> FAvaSequenceController::FindIntersectedMar
 			}
 			UE_LOG(LogAvaSequenceController, Log, TEXT("%d Marks between %f and %f --- [ %s ]")
 				, IntersectedMarks.Num()
-				, UE::Avalanche::Private::GetFrameTimeAsDisplayDecimal(InMovieScene, LowerBound)
-				, UE::Avalanche::Private::GetFrameTimeAsDisplayDecimal(InMovieScene, UpperBound)
+				, UE::Ava::Private::GetFrameTimeAsDisplayDecimal(InMovieScene, LowerBound)
+				, UE::Ava::Private::GetFrameTimeAsDisplayDecimal(InMovieScene, UpperBound)
 				, *MarkList);	
 		}
 	}
diff --git a/Engine/Plugins/Experimental/Avalanche/Source/AvalancheSequence/Private/AvaSequencePlaybackContext.cpp b/Engine/Plugins/Experimental/Avalanche/Source/AvalancheSequence/Private/AvaSequencePlaybackContext.cpp
--- a/Engine/Plugins/Experimental/Avalanche/Source/AvalancheSequence/Private/AvaSequencePlaybackContext.cpp
+++ b/Engine/Plugins/Experimental/Avalanche/Source/AvalancheSequence/Private/AvaSequencePlaybackContext.cpp
@@ -30,11 +30,11 @@ void FAvaSequencePlaybackContext::Resolve(const FAvaSequencePlayerVariant& InPla
 	LastJumpedFrame = JumpedFrame;
 	JumpedFrame.Reset();
 
-	// Find if the Player is an Avalanche Player 
-	if (UAvaSequencePlayer* const AvalanchePlayer = InPlayerVariant.TryGet<UAvaSequencePlayer>())
+	// Find if the Player is an Motion Design Player 
+	if (UAvaSequencePlayer* const MotionDesignPlayer = InPlayerVariant.TryGet<UAvaSequencePlayer>())
 	{
-		SequencePlayer   = AvalanchePlayer;
-		bPlayingForwards = !AvalanchePlayer->IsReversed();
+		SequencePlayer   = MotionDesignPlayer;
+		bPlayingForwards = !MotionDesignPlayer->IsReversed();
 	}
 #if WITH_EDITOR
 	// Else it could be the Sequencer Player
diff --git a/Engine/Plugins/Experimental/Avalanche/Source/AvalancheSequence/Public/AvaSequence.h b/Engine/Plugins/Experimental/Avalanche/Source/AvalancheSequence/Public/AvaSequence.h
--- a/Engine/Plugins/Experimental/Avalanche/Source/AvalancheSequence/Public/AvaSequence.h
+++ b/Engine/Plugins/Experimental/Avalanche/Source/AvalancheSequence/Public/AvaSequence.h
@@ -28,7 +28,7 @@ class UWorld;
 struct FMovieSceneMarkedFrame;
 struct FMovieScenePossessable;
 
-UCLASS(BlueprintType, Config=Avalanche)
+UCLASS(BlueprintType, Config=MotionDesign)
 class AVALANCHESEQUENCE_API UAvaSequence : public ULevelSequence
 {
 	GENERATED_BODY()
diff --git a/Engine/Plugins/Experimental/Avalanche/Source/AvalancheSequencer/Private/AvaSequenceItemShared.h b/Engine/Plugins/Experimental/Avalanche/Source/AvalancheSequencer/Private/AvaSequenceItemShared.h
--- a/Engine/Plugins/Experimental/Avalanche/Source/AvalancheSequencer/Private/AvaSequenceItemShared.h
+++ b/Engine/Plugins/Experimental/Avalanche/Source/AvalancheSequencer/Private/AvaSequenceItemShared.h
@@ -6,7 +6,8 @@
 #include "Templates/SharedPointer.h"
 
 class IAvaSequenceItem;
-typedef TSharedPtr<IAvaSequenceItem> FAvaSequenceItemPtr;
+
+using FAvaSequenceItemPtr = TSharedPtr<IAvaSequenceItem>;
 
 /*
  * Delegate Wrapper so that things other than IAvaSequenceItem can only Add or Remove Listeners but not call Broadcast
diff --git a/Engine/Plugins/Experimental/Avalanche/Source/AvalancheSequencer/Private/AvaSequencer.cpp b/Engine/Plugins/Experimental/Avalanche/Source/AvalancheSequencer/Private/AvaSequencer.cpp
--- a/Engine/Plugins/Experimental/Avalanche/Source/AvalancheSequencer/Private/AvaSequencer.cpp
+++ b/Engine/Plugins/Experimental/Avalanche/Source/AvalancheSequencer/Private/AvaSequencer.cpp
@@ -1703,8 +1703,8 @@ void FAvaSequencer::ExecuteSequencerDuplication(FExecuteAction InExecuteAction)
 		// HACK: Sequencer Duplicates Actors via UUnrealEdEngine::edactDuplicateSelected
 		// Sequencer then expects that after this function is called, the GSelectedActors are the newly duplicated actors.
 		// However, ULevelFactory::FactoryCreateText only changes selection when the World in question is the GWorld,
-		// which is not true for Avalanche.
-		// So for this we temporarily set GWorld to our Avalanche World so that selections happen correctly.
+		// which is not true for Motion Design.
+		// So for this we temporarily set GWorld to our Motion Design World so that selections happen correctly.
 		UWorld* const OldGWorld = GWorld;
 		GWorld = InWorld;
 		InExecuteAction.ExecuteIfBound();
diff --git a/Engine/Plugins/Experimental/Avalanche/Source/AvalancheSequencer/Private/Director/AvaSequenceDirectorCompiler.cpp b/Engine/Plugins/Experimental/Avalanche/Source/AvalancheSequencer/Private/Director/AvaSequenceDirectorCompiler.cpp
--- a/Engine/Plugins/Experimental/Avalanche/Source/AvalancheSequencer/Private/Director/AvaSequenceDirectorCompiler.cpp
+++ b/Engine/Plugins/Experimental/Avalanche/Source/AvalancheSequencer/Private/Director/AvaSequenceDirectorCompiler.cpp
@@ -160,8 +160,8 @@ void FAvaSequenceDirectorCompilerContext::FinishCompilingClass(UClass* InClass)
 			const TSoftObjectPtr<UAvaSequence>* FoundSequence = SequencePropertyMap.FindKey(SequenceProperty);
 			if (!FoundSequence)
 			{
-				const FText RequiredSequenceNotBound = LOCTEXT("RequiredAvaSequenceNotBound"
-					, "A required Avalanche Sequence binding @@ was not found.");
+				const FText RequiredSequenceNotBound = LOCTEXT("RequiredSequenceNotBound"
+					, "A required Motion Design Sequence binding @@ was not found.");
 
 				if (Blueprint->bIsNewlyCreated)
 				{
diff --git a/Engine/Plugins/Experimental/Avalanche/Source/AvalancheSequencer/Private/Outliner/AvaOutlinerSequenceProxy.cpp b/Engine/Plugins/Experimental/Avalanche/Source/AvalancheSequencer/Private/Outliner/AvaOutlinerSequenceProxy.cpp
--- a/Engine/Plugins/Experimental/Avalanche/Source/AvalancheSequencer/Private/Outliner/AvaOutlinerSequenceProxy.cpp
+++ b/Engine/Plugins/Experimental/Avalanche/Source/AvalancheSequencer/Private/Outliner/AvaOutlinerSequenceProxy.cpp
@@ -100,13 +100,13 @@ void FAvaOutlinerSequenceProxy::GetProxiedItems(const TSharedRef<IAvaOutlinerIte
 {
 	if (const FAvaOutlinerActor* const ActorItem = InParent->CastTo<FAvaOutlinerActor>())
 	{
-		TSharedPtr<IAvaSequencer> AvalancheSequencer = AvaSequencerWeak.Pin();
+		TSharedPtr<IAvaSequencer> AvaSequencer = AvaSequencerWeak.Pin();
 
 		AActor* const Actor = ActorItem->GetActor();
 		
-		if (AvalancheSequencer.IsValid() && IsValid(Actor))
+		if (AvaSequencer.IsValid() && IsValid(Actor))
 		{
-			const TArray<UAvaSequence*> Sequences = AvalancheSequencer->GetSequencesForObject(Actor);
+			const TArray<UAvaSequence*> Sequences = AvaSequencer->GetSequencesForObject(Actor);
 			
 			for (UAvaSequence* const Sequence : Sequences)
 			{
diff --git a/Engine/Plugins/Experimental/Avalanche/Source/AvalancheShapes/Private/AvaShapesDefs.cpp b/Engine/Plugins/Experimental/Avalanche/Source/AvalancheShapes/Private/AvaShapesDefs.cpp
--- a/Engine/Plugins/Experimental/Avalanche/Source/AvalancheShapes/Private/AvaShapesDefs.cpp
+++ b/Engine/Plugins/Experimental/Avalanche/Source/AvalancheShapes/Private/AvaShapesDefs.cpp
@@ -4,7 +4,7 @@
 #include "DynamicMesh/DynamicMeshAttributeSet.h"
 #include "DynamicMeshes/AvaShapeDynMeshBase.h"
 
-namespace UE::AvalancheShapes
+namespace UE::AvaShapes
 {
 	FVector2D FindClosestPointOnLine(const FVector2D& LineStart, const FVector2D& LineEnd, const FVector2D& TestPoint)
 	{
diff --git a/Engine/Plugins/Experimental/Avalanche/Source/AvalancheShapes/Private/DynamicMeshes/AvaShapeDynMeshBase.cpp b/Engine/Plugins/Experimental/Avalanche/Source/AvalancheShapes/Private/DynamicMeshes/AvaShapeDynMeshBase.cpp
--- a/Engine/Plugins/Experimental/Avalanche/Source/AvalancheShapes/Private/DynamicMeshes/AvaShapeDynMeshBase.cpp
+++ b/Engine/Plugins/Experimental/Avalanche/Source/AvalancheShapes/Private/DynamicMeshes/AvaShapeDynMeshBase.cpp
@@ -33,7 +33,7 @@
 #include "PropertyEditorModule.h"
 #endif
 
-DEFINE_LOG_CATEGORY_STATIC(LogAvalancheDynamicMesh, Log, All);
+DEFINE_LOG_CATEGORY_STATIC(LogAvaDynamicMesh, Log, All);
 
 using namespace UE::Geometry;
 
@@ -1393,38 +1393,38 @@ FVector2D UAvaShapeDynamicMeshBase::GetNumericsFromAnchor(EAvaAnchors AnchorEnum
 {
 	FVector2D Anchor = FVector2D::ZeroVector;
 
-	switch (ToUnderlyingType<EAvaAnchors>(AnchorEnum) & UE::Avalanche::AnchorPoints::Horizontal)
+	switch (ToUnderlyingType<EAvaAnchors>(AnchorEnum) & UE::Ava::AnchorPoints::Horizontal)
 	{
 	default:
 		//falls through
 
-	case UE::Avalanche::AnchorPoints::Left:
+	case UE::Ava::AnchorPoints::Left:
 		// nothing to do
 		break;
 
-	case UE::Avalanche::AnchorPoints::HMiddle:
+	case UE::Ava::AnchorPoints::HMiddle:
 		Anchor.X = 0.5f;
 		break;
 
-	case UE::Avalanche::AnchorPoints::Right:
+	case UE::Ava::AnchorPoints::Right:
 		Anchor.X = 1.f;
 		break;
 	}
 
-	switch (ToUnderlyingType<EAvaAnchors>(AnchorEnum) & UE::Avalanche::AnchorPoints::Vertical)
+	switch (ToUnderlyingType<EAvaAnchors>(AnchorEnum) & UE::Ava::AnchorPoints::Vertical)
 	{
 	default:
 		//falls through
 
-	case UE::Avalanche::AnchorPoints::Top:
+	case UE::Ava::AnchorPoints::Top:
 		// nothing to do
 		break;
 
-	case UE::Avalanche::AnchorPoints::VMiddle:
+	case UE::Ava::AnchorPoints::VMiddle:
 		Anchor.Y = 0.5f;
 		break;
 
-	case UE::Avalanche::AnchorPoints::Bottom:
+	case UE::Ava::AnchorPoints::Bottom:
 		Anchor.Y = 1.f;
 		break;
 	}
@@ -1463,21 +1463,21 @@ bool UAvaShapeDynamicMeshBase::CreateDynamicMesh(FAvaShapeMesh& InMesh)
 	// vertices and triangles should not exists if we create them
 	if (InMesh.VerticeIds.Num() != 0 || InMesh.TriangleIds.Num() != 0)
 	{
-		UE_LOG(LogAvalancheDynamicMesh, Warning, TEXT("CreateDynamicMesh %s %i : VerticeIds %i, TriangleIds %i should be empty to create mesh"), *GetMeshName(), InMesh.GetMeshIndex(), InMesh.VerticeIds.Num(), InMesh.TriangleIds.Num());
+		UE_LOG(LogAvaDynamicMesh, Warning, TEXT("CreateDynamicMesh %s %i : VerticeIds %i, TriangleIds %i should be empty to create mesh"), *GetMeshName(), InMesh.GetMeshIndex(), InMesh.VerticeIds.Num(), InMesh.TriangleIds.Num());
 		return false;
 	}
 
 	// should have same number for vertices
 	if (!(InMesh.Vertices.Num() == InMesh.Normals.Num()))
 	{
-		UE_LOG(LogAvalancheDynamicMesh, Warning, TEXT("CreateDynamicMesh %s %i : Arrays Vertices %i, Normals %i should have same length, invalid array given"), *GetMeshName(), InMesh.GetMeshIndex(), InMesh.Vertices.Num(), InMesh.Normals.Num());
+		UE_LOG(LogAvaDynamicMesh, Warning, TEXT("CreateDynamicMesh %s %i : Arrays Vertices %i, Normals %i should have same length, invalid array given"), *GetMeshName(), InMesh.GetMeshIndex(), InMesh.Vertices.Num(), InMesh.Normals.Num());
 		return false;
 	}
 
 	// should be multiple of 3
 	if ((InMesh.Triangles.Num() % 3) != 0)
 	{
-		UE_LOG(LogAvalancheDynamicMesh, Warning, TEXT("CreateDynamicMesh %s %i : Triangles array should be multiple of 3, %i invalid array given"), *GetMeshName(), InMesh.GetMeshIndex(), InMesh.Triangles.Num());
+		UE_LOG(LogAvaDynamicMesh, Warning, TEXT("CreateDynamicMesh %s %i : Triangles array should be multiple of 3, %i invalid array given"), *GetMeshName(), InMesh.GetMeshIndex(), InMesh.Triangles.Num());
 		return false;
 	}
 
@@ -1514,7 +1514,7 @@ bool UAvaShapeDynamicMeshBase::CreateDynamicMesh(FAvaShapeMesh& InMesh)
 				!InMesh.VerticeIds.IsValidIndex(Idx2) ||
 				!InMesh.VerticeIds.IsValidIndex(Idx3))
 			{
-				UE_LOG(LogAvalancheDynamicMesh, Warning, TEXT("CreateDynamicMesh %s %i : Invalid Vertice idx for triangle %i %i %i"), *GetMeshName(), InMesh.GetMeshIndex(), Idx1, Idx2, Idx3);
+				UE_LOG(LogAvaDynamicMesh, Warning, TEXT("CreateDynamicMesh %s %i : Invalid Vertice idx for triangle %i %i %i"), *GetMeshName(), InMesh.GetMeshIndex(), Idx1, Idx2, Idx3);
 				continue;
 			}
 
@@ -1526,7 +1526,7 @@ bool UAvaShapeDynamicMeshBase::CreateDynamicMesh(FAvaShapeMesh& InMesh)
 			int32 TId = Mesh.AppendTriangle(VId1, VId2, VId3, InMesh.GetMeshIndex());
 			if (TId < 0)
 			{
-				UE_LOG(LogAvalancheDynamicMesh, Warning, TEXT("CreateDynamicMesh %s %i : Invalid Triangle ID for mesh"), *GetMeshName(), InMesh.GetMeshIndex());
+				UE_LOG(LogAvaDynamicMesh, Warning, TEXT("CreateDynamicMesh %s %i : Invalid Triangle ID for mesh"), *GetMeshName(), InMesh.GetMeshIndex());
 				continue;
 			}
 			InMesh.TriangleIds.Add(TId);
@@ -1569,7 +1569,7 @@ bool UAvaShapeDynamicMeshBase::UpdateDynamicMesh(FAvaShapeMesh& InMesh)
 	// should have same number for vertices to update
 	if (!(InMesh.Vertices.Num() == InMesh.Normals.Num()))
 	{
-		UE_LOG(LogAvalancheDynamicMesh, Warning, TEXT("UpdateDynamicMesh %s %i: Arrays Vertices %i, Normals %i should have same length, invalid array given"), *GetMeshName(), InMesh.GetMeshIndex(), InMesh.Vertices.Num(), InMesh.Normals.Num());
+		UE_LOG(LogAvaDynamicMesh, Warning, TEXT("UpdateDynamicMesh %s %i: Arrays Vertices %i, Normals %i should have same length, invalid array given"), *GetMeshName(), InMesh.GetMeshIndex(), InMesh.Vertices.Num(), InMesh.Normals.Num());
 		return false;
 	}
 
@@ -1878,7 +1878,7 @@ bool UAvaShapeDynamicMeshBase::ApplyUVsManually(FAvaShapeMesh& InMesh)
 				const int32* FoundElementID = BaseToOverlayVIDMap.Find(TriVtx[j]);
 				if (FoundElementID == nullptr)
 				{
-					UE_LOG(LogAvalancheDynamicMesh, Warning, TEXT("ApplyUVsManually %s %i: vertice id %i for uv is invalid, skipping triangle"), *GetMeshName(), InMesh.GetMeshIndex(), TriVtx[j]);
+					UE_LOG(LogAvaDynamicMesh, Warning, TEXT("ApplyUVsManually %s %i: vertice id %i for uv is invalid, skipping triangle"), *GetMeshName(), InMesh.GetMeshIndex(), TriVtx[j]);
 					break;
 				}
 				else
@@ -1975,7 +1975,7 @@ bool UAvaShapeDynamicMeshBase::ApplyUVsTransform(FAvaShapeMesh& InMesh, FAvaShap
 
 	bool bResult = true;
 	DynMesh->GetDynamicMesh()->EditMesh([InMesh, InParams, ShapeSize, UVOffset, UVFixRotation, &bResult](FDynamicMesh3& EditMesh) {
-		bResult = UE::AvalancheShapes::TransformMeshUVs(EditMesh, InMesh.UVIds, InParams, ShapeSize, UVOffset, UVFixRotation);
+		bResult = UE::AvaShapes::TransformMeshUVs(EditMesh, InMesh.UVIds, InParams, ShapeSize, UVOffset, UVFixRotation);
 	}, EDynamicMeshChangeType::AttributeEdit, EDynamicMeshAttributeChangeFlags::UVs);
 
 	return bResult;
@@ -2035,7 +2035,7 @@ bool UAvaShapeDynamicMeshBase::RegisterMesh(FAvaShapeMeshData& NewMeshData)
 {
 	if (MeshDatas.Contains(NewMeshData.GetMeshIndex()))
 	{
-		UE_LOG(LogAvalancheDynamicMesh, Warning, TEXT("RegisterMesh %s %i : Already contains mesh cannot add again"), *GetMeshName(), NewMeshData.GetMeshIndex());
+		UE_LOG(LogAvaDynamicMesh, Warning, TEXT("RegisterMesh %s %i : Already contains mesh cannot add again"), *GetMeshName(), NewMeshData.GetMeshIndex());
 		return false;
 	}
 
diff --git a/Engine/Plugins/Experimental/Avalanche/Source/AvalancheShapes/Private/DynamicMeshes/AvaShapeIrregularPolygonDynMesh.cpp b/Engine/Plugins/Experimental/Avalanche/Source/AvalancheShapes/Private/DynamicMeshes/AvaShapeIrregularPolygonDynMesh.cpp
--- a/Engine/Plugins/Experimental/Avalanche/Source/AvalancheShapes/Private/DynamicMeshes/AvaShapeIrregularPolygonDynMesh.cpp
+++ b/Engine/Plugins/Experimental/Avalanche/Source/AvalancheShapes/Private/DynamicMeshes/AvaShapeIrregularPolygonDynMesh.cpp
@@ -210,7 +210,7 @@ bool UAvaShapeIrregularPolygonDynamicMesh::IsPointTooCloseToALine(const FVector2
 			continue;
 		}
 
-		FVector2D ClosestPoint = UE::AvalancheShapes::FindClosestPointOnLine(
+		FVector2D ClosestPoint = UE::AvaShapes::FindClosestPointOnLine(
 			Points[PreviousPointIdx].Location,
 			Points[PointIdx].Location,
 			InPoint
@@ -234,7 +234,7 @@ bool UAvaShapeIrregularPolygonDynamicMesh::IsLineTooCloseToAPoint(const FVector2
 			continue;
 		}
 
-		FVector2D ClosestPoint = UE::AvalancheShapes::FindClosestPointOnLine(Start, End, Points[PointIdx].Location);
+		FVector2D ClosestPoint = UE::AvaShapes::FindClosestPointOnLine(Start, End, Points[PointIdx].Location);
 
 		if ((ClosestPoint - Points[PointIdx].Location).Size() < UAvaShapeIrregularPolygonDynamicMesh::MinPointDistance)
 		{
diff --git a/Engine/Plugins/Experimental/Avalanche/Source/AvalancheShapes/Private/Sequencer/MovieScene/MovieSceneAvaShapeRectCornerSection.cpp b/Engine/Plugins/Experimental/Avalanche/Source/AvalancheShapes/Private/Sequencer/MovieScene/MovieSceneAvaShapeRectCornerSection.cpp
--- a/Engine/Plugins/Experimental/Avalanche/Source/AvalancheShapes/Private/Sequencer/MovieScene/MovieSceneAvaShapeRectCornerSection.cpp
+++ b/Engine/Plugins/Experimental/Avalanche/Source/AvalancheShapes/Private/Sequencer/MovieScene/MovieSceneAvaShapeRectCornerSection.cpp
@@ -10,7 +10,7 @@
 #define LOCTEXT_NAMESPACE "MovieSceneAvaShapeRectCornerSection"
 
 #if WITH_EDITOR
-namespace UE::Avalanche::Private
+namespace UE::Ava::Private
 {
 	struct FShapeRectCornerChannelMetaData
 	{
@@ -54,7 +54,7 @@ UMovieSceneAvaShapeRectCornerSection::UMovieSceneAvaShapeRectCornerSection(const
 	FMovieSceneChannelProxyData Channels;
 
 #if WITH_EDITOR
-	using namespace UE::Avalanche::Private;
+	using namespace UE::Ava::Private;
 	static const FShapeRectCornerChannelMetaData MetaData;
 	Channels.Add(BevelSize,         MetaData.BevelSizeMetaData,          TMovieSceneExternalValue<float>());
 	Channels.Add(BevelSubdivisions, MetaData.BevelSubdivisionsMetaData,  TMovieSceneExternalValue<uint8>());
diff --git a/Engine/Plugins/Experimental/Avalanche/Source/AvalancheShapes/Private/Sequencer/MovieScene/MovieSceneAvaShapeRectCornerSectionTemplate.h b/Engine/Plugins/Experimental/Avalanche/Source/AvalancheShapes/Private/Sequencer/MovieScene/MovieSceneAvaShapeRectCornerSectionTemplate.h
--- a/Engine/Plugins/Experimental/Avalanche/Source/AvalancheShapes/Private/Sequencer/MovieScene/MovieSceneAvaShapeRectCornerSectionTemplate.h
+++ b/Engine/Plugins/Experimental/Avalanche/Source/AvalancheShapes/Private/Sequencer/MovieScene/MovieSceneAvaShapeRectCornerSectionTemplate.h
@@ -121,7 +121,7 @@ template<> FMovieSceneAnimTypeID GetBlendingDataType<FAvaShapeRectangleCornerSet
 
 template<> struct TBlendableTokenTraits<FAvaShapeRectangleCornerSettings>
 {
-	typedef UE::MovieScene::FBlendableShapeRectCorner WorkingDataType;
+	using WorkingDataType = UE::MovieScene::FBlendableShapeRectCorner;
 };
 
 USTRUCT()
diff --git a/Engine/Plugins/Experimental/Avalanche/Source/AvalancheShapes/Public/AvaShapesDefs.h b/Engine/Plugins/Experimental/Avalanche/Source/AvalancheShapes/Public/AvaShapesDefs.h
--- a/Engine/Plugins/Experimental/Avalanche/Source/AvalancheShapes/Public/AvaShapesDefs.h
+++ b/Engine/Plugins/Experimental/Avalanche/Source/AvalancheShapes/Public/AvaShapesDefs.h
@@ -13,7 +13,7 @@ class UMaterialInstanceDynamic;
 class UMaterialInterface;
 struct FAvaShapeMaterialUVParameters;
 
-namespace UE::AvalancheShapes
+namespace UE::AvaShapes
 {
 	FVector2D AVALANCHESHAPES_API FindClosestPointOnLine(const FVector2D& LineStart, const FVector2D& LineEnd, const FVector2D& TestPoint);
 
diff --git a/Engine/Plugins/Experimental/Avalanche/Source/AvalancheShapesEditor/Private/Tools/AvaShapesEditorShapeToolBase.cpp b/Engine/Plugins/Experimental/Avalanche/Source/AvalancheShapesEditor/Private/Tools/AvaShapesEditorShapeToolBase.cpp
--- a/Engine/Plugins/Experimental/Avalanche/Source/AvalancheShapesEditor/Private/Tools/AvaShapesEditorShapeToolBase.cpp
+++ b/Engine/Plugins/Experimental/Avalanche/Source/AvalancheShapesEditor/Private/Tools/AvaShapesEditorShapeToolBase.cpp
@@ -104,7 +104,7 @@ bool UAvaShapesEditorShapeToolBase::UseIdentityLocation() const
 		return false;
 	}
 
-	return !IsAvalancheViewport();
+	return !IsMotionDesignViewport();
 }
 
 bool UAvaShapesEditorShapeToolBase::UseIdentityRotation() const
diff --git a/Engine/Plugins/Experimental/Avalanche/Source/AvalancheText/Private/Font/AvaFont.cpp b/Engine/Plugins/Experimental/Avalanche/Source/AvalancheText/Private/Font/AvaFont.cpp
--- a/Engine/Plugins/Experimental/Avalanche/Source/AvalancheText/Private/Font/AvaFont.cpp
+++ b/Engine/Plugins/Experimental/Avalanche/Source/AvalancheText/Private/Font/AvaFont.cpp
@@ -50,7 +50,7 @@ UFont* FAvaFont::GetDefaultFont()
 			}
 		}
 
-		if (ensureMsgf(DefaultFont, TEXT("AvalancheFont: cannot load any font to be used as default.")))
+		if (ensureMsgf(DefaultFont, TEXT("MotionDesignFont: cannot load any font to be used as default.")))
 		{
 			DefaultFontObjects.AvaDefaultFont = DefaultFont;
 		}
@@ -276,7 +276,7 @@ void FAvaFont::InitFromFont(UFont* InFont)
 	if (InFont)
 	{
 		FString Name;
-		UE::Avalanche::FontUtilities::Public::GetFontName(InFont, Name);
+		UE::Ava::FontUtilities::Public::GetFontName(InFont, Name);
 
 		MotionDesignFontObject = NewObject<UAvaFontObject>();
 		MotionDesignFontObject->InitProjectFont(InFont, Name);
@@ -397,7 +397,7 @@ UFont* FAvaFont::GetFontByName(const FString& InFontName)
 		if (UFont* const CurrFont = Cast<UFont>(AssetData.GetAsset()))
 		{
 			FString CurrFontName;
-			UE::Avalanche::FontUtilities::Public::GetFontName(CurrFont, CurrFontName);
+			UE::Ava::FontUtilities::Public::GetFontName(CurrFont, CurrFontName);
 
 			if (CurrFontName == InFontName)
 			{
@@ -410,7 +410,7 @@ UFont* FAvaFont::GetFontByName(const FString& InFontName)
 }
 
 // note: this function used to be in the AvalancheEditor module, since it was not needed at Runtime before 
-void UE::Avalanche::FontUtilities::Public::GetFontName(const UFont* InFont, FString& OutFontName)
+void UE::Ava::FontUtilities::Public::GetFontName(const UFont* InFont, FString& OutFontName)
 {
 	if (IsValid(InFont))
 	{
diff --git a/Engine/Plugins/Experimental/Avalanche/Source/AvalancheText/Public/Font/AvaFont.h b/Engine/Plugins/Experimental/Avalanche/Source/AvalancheText/Public/Font/AvaFont.h
--- a/Engine/Plugins/Experimental/Avalanche/Source/AvalancheText/Public/Font/AvaFont.h
+++ b/Engine/Plugins/Experimental/Avalanche/Source/AvalancheText/Public/Font/AvaFont.h
@@ -238,7 +238,7 @@ template<> struct TStructOpsTypeTraits<FAvaFont> : public TStructOpsTypeTraitsBa
 	};
 };
 
-namespace UE::Avalanche::FontUtilities::Public
+namespace UE::Ava::FontUtilities::Public
 {
 	AVALANCHETEXT_API void GetFontName(const UFont* InFont, FString& OutFontName);
 }
diff --git a/Engine/Plugins/Experimental/Avalanche/Source/AvalancheTextEditor/Private/Font/AvaFontDetailsCustomization.cpp b/Engine/Plugins/Experimental/Avalanche/Source/AvalancheTextEditor/Private/Font/AvaFontDetailsCustomization.cpp
--- a/Engine/Plugins/Experimental/Avalanche/Source/AvalancheTextEditor/Private/Font/AvaFontDetailsCustomization.cpp
+++ b/Engine/Plugins/Experimental/Avalanche/Source/AvalancheTextEditor/Private/Font/AvaFontDetailsCustomization.cpp
@@ -70,7 +70,7 @@ void FAvaFontDetailsCustomization::CustomizeHeader(TSharedRef<IPropertyHandle> F
 		.ShowItalicFontsOnly(bIsShowItalicActive)
 		.OptionsSource(&Options)
 		.InitiallySelectedItem(SelectedOption)
-		.AvalancheFontPropertyHandle(AvaFontPropertyHandle)
+		.FontPropertyHandle(AvaFontPropertyHandle)
 		.OnGenerateWidget(this, &FAvaFontDetailsCustomization::HandleGenerateWidget)
 		.OnSelectionChanged(this, &FAvaFontDetailsCustomization::HandleSelectionChanged)
 		[
@@ -277,8 +277,8 @@ FReply FAvaFontDetailsCustomization::ImportButtonClicked()
 TSharedRef<SWidget> FAvaFontDetailsCustomization::HandleGenerateWidget(const TSharedPtr<FAvaFontView>& InItem)
 {
 	return SNew(SAvaFontField)
-			.OnAvaFontFieldModified(this, &FAvaFontDetailsCustomization::OnFontFieldUpdated)
-			.AvalancheFont(InItem);
+			.OnFontFieldModified(this, &FAvaFontDetailsCustomization::OnFontFieldUpdated)
+			.FontView(InItem);
 }
 
 TSharedRef<SWidget> FAvaFontDetailsCustomization::GenerateMultipleSelectionField() const
diff --git a/Engine/Plugins/Experimental/Avalanche/Source/AvalancheTextEditor/Private/Font/AvaFontManagerSubsystem.cpp b/Engine/Plugins/Experimental/Avalanche/Source/AvalancheTextEditor/Private/Font/AvaFontManagerSubsystem.cpp
--- a/Engine/Plugins/Experimental/Avalanche/Source/AvalancheTextEditor/Private/Font/AvaFontManagerSubsystem.cpp
+++ b/Engine/Plugins/Experimental/Avalanche/Source/AvalancheTextEditor/Private/Font/AvaFontManagerSubsystem.cpp
@@ -69,8 +69,8 @@ void UAvaFontManagerSubsystem::ClearFontsData()
 {
 	OSFontsMap.Empty();
 	ProjectFontsMap.Empty();
-	AvalancheFontsMap.Empty();
-	AvalancheFontsOptions.Empty();
+	FontsMap.Empty();
+	FontsOptions.Empty();
 	FontsToImportOnSave.Empty();
 }
 
@@ -95,8 +95,8 @@ void UAvaFontManagerSubsystem::Initialize()
 	// init OS fonts only, project fonts can be loaded lately with InitializeProjectFonts()
 	InitializeOSFonts();
 
-	RefreshAvalancheFontsMap();
-	RefreshAvalancheFontsOptions();
+	RefreshFontsMap();
+	RefreshFontsOptions();
 
 	RegisterAssetsCallbacks();
 
@@ -117,8 +117,8 @@ void UAvaFontManagerSubsystem::OnAssetsAdded(const FAssetData& InAssetData)
 	{
 		CreateProjectFont(CurrFont);
 
-		RefreshAvalancheFontsMap();
-		RefreshAvalancheFontsOptions();
+		RefreshFontsMap();
+		RefreshFontsOptions();
 
 		OnProjectFontCreatedDelegate.Broadcast(CurrFont);
 	}
@@ -144,20 +144,20 @@ void UAvaFontManagerSubsystem::OnAssetDeleted(UObject* InObject)
 		// this should allow the assets to still point to the same font, in case the deleted font is a system one
 		LoadOSFonts();
 
-		RefreshAvalancheFontsMap();
-		RefreshAvalancheFontsOptions();
+		RefreshFontsMap();
+		RefreshFontsOptions();
 	}
 }
 
 void UAvaFontManagerSubsystem::LoadFavorites()
 {
-	if (!IsValid(AvaFontManagerConfig))
+	if (!IsValid(FontManagerConfig))
 	{
 		static const FString PackageName = GetFontConfigPackageName();
 		UPackage* NewPackage = CreatePackage(*PackageName);
 		NewPackage->SetFlags(RF_Transient);
 		NewPackage->AddToRoot();
-		AvaFontManagerConfig = NewObject<UAvaFontConfig>(NewPackage, "AvalancheFontConfig", RF_Transient | RF_Transactional | RF_Standalone);
+		FontManagerConfig = NewObject<UAvaFontConfig>(NewPackage, "MotionDesignFontConfig", RF_Transient | RF_Transactional | RF_Standalone);
 	}
 }
 
@@ -246,8 +246,8 @@ bool UAvaFontManagerSubsystem::ImportAsAsset(UAvaFontObject* InFontToImport)
 					CreateProjectFont(NewFont);
 				}
 
-				RefreshAvalancheFontsMap();
-				RefreshAvalancheFontsOptions();
+				RefreshFontsMap();
+				RefreshFontsOptions();
 				OnSystemFontsUpdatedDelegate.Broadcast();
 
 				return true;
@@ -459,7 +459,7 @@ bool UAvaFontManagerSubsystem::Exec(class UWorld* InWorld, const TCHAR* InCmd, F
 	{
 		if (FParse::Command(&InCmd, TEXT("list")))
 		{
-			UE::Avalanche::Private::Fonts::ListAvailableFontFiles();
+			UE::Ava::Private::Fonts::ListAvailableFontFiles();
 			return true;
 		}
 		else if (FParse::Command(&InCmd, TEXT("refresh")))
@@ -477,9 +477,9 @@ bool UAvaFontManagerSubsystem::Exec(class UWorld* InWorld, const TCHAR* InCmd, F
 
 void UAvaFontManagerSubsystem::CombineAllFonts()
 {
-	AvalancheFontsMap.Empty();
-	AvalancheFontsMap.Append(ProjectFontsMap);
-	AvalancheFontsMap.Append(OSFontsMap);
+	FontsMap.Empty();
+	FontsMap.Append(ProjectFontsMap);
+	FontsMap.Append(OSFontsMap);
 }
 
 void UAvaFontManagerSubsystem::UnregisterAssetsCallbacks() const
@@ -557,9 +557,9 @@ void UAvaFontManagerSubsystem::LoadProjectFonts()
 
 		if (IsValid(CurrFont))
 		{
-			if (const UAvaFontObject* const NewAvalancheFontObj = CreateProjectFont(CurrFont, ProjectFontsPackage))
+			if (const UAvaFontObject* const NewAvaFontObj = CreateProjectFont(CurrFont, ProjectFontsPackage))
 			{
-				CurrFontAssetsNamesArray.Add(NewAvalancheFontObj->GetFontName());
+				CurrFontAssetsNamesArray.Add(NewAvaFontObj->GetFontName());
 				CurrFont->GetPackage()->FullyLoad();
 			}
 		}
@@ -664,15 +664,15 @@ UAvaFontObject* UAvaFontManagerSubsystem::CreateProjectFont(UFont* InSourceFont,
 					ProjectFontsPackage->AddToRoot();
 				}
 
-				UAvaFontObject* NewAvalancheFont = NewObject<UAvaFontObject>(ProjectFontsPackage, *SanitizedFontName, RF_Public | RF_Standalone | RF_Transient);
-				NewAvalancheFont->AddToRoot();
+				UAvaFontObject* NewAvaFont = NewObject<UAvaFontObject>(ProjectFontsPackage, *SanitizedFontName, RF_Public | RF_Standalone | RF_Transient);
+				NewAvaFont->AddToRoot();
 
-				NewAvalancheFont->InitProjectFont(InSourceFont, FontName);
-				UAvaFontManagerSubsystem::SetupMetrics(NewAvalancheFont);
+				NewAvaFont->InitProjectFont(InSourceFont, FontName);
+				UAvaFontManagerSubsystem::SetupMetrics(NewAvaFont);
 
-				ProjectFontsMap.Add(SanitizedFontName, NewAvalancheFont);
+				ProjectFontsMap.Add(SanitizedFontName, NewAvaFont);
 
-				return NewAvalancheFont;
+				return NewAvaFont;
 			}
 			else
 			{
@@ -686,7 +686,7 @@ UAvaFontObject* UAvaFontManagerSubsystem::CreateProjectFont(UFont* InSourceFont,
 
 void UAvaFontManagerSubsystem::LoadOSFonts()
 {
-	UE::Avalanche::Private::Fonts::GetSystemFontInfo(FontsInfoMap);
+	UE::Ava::Private::Fonts::GetSystemFontInfo(FontsInfoMap);
 
 	if (FontsInfoMap.IsEmpty())
 	{
@@ -720,8 +720,8 @@ void UAvaFontManagerSubsystem::RemoveProjectFontsFromOSFonts()
 		}
 	}
 
-	RefreshAvalancheFontsMap();
-	RefreshAvalancheFontsOptions();
+	RefreshFontsMap();
+	RefreshFontsOptions();
 	OnSystemFontsUpdatedDelegate.Broadcast();
 }
 
@@ -766,27 +766,27 @@ void UAvaFontManagerSubsystem::ImportSystemFontFamily(const FSystemFontsRetrieve
 	NewFont->FontCacheType = EFontCacheType::Runtime;
 
 	SetupFontFamilyTypefaces(NewFont, InFontParams);
-	UAvaFontObject* NewAvalancheFont = NewObject<UAvaFontObject>(ImportPackage,  *SanitizedFontFamilyName, TempFontAssetFlags);
-	NewAvalancheFont->AddToRoot();
+	UAvaFontObject* NewAvaFont = NewObject<UAvaFontObject>(ImportPackage,  *SanitizedFontFamilyName, TempFontAssetFlags);
+	NewAvaFont->AddToRoot();
 
 	FTypefaceEntry& FallbackTypefaceEntry = NewFont->CompositeFont.FallbackTypeface.Typeface.Fonts[NewFont->CompositeFont.FallbackTypeface.Typeface.Fonts.AddDefaulted()];
 	FallbackTypefaceEntry.Font = FFontData(GetFallbackFontFace());
 	FallbackTypefaceEntry.Name = FName("Regular");
 
-	NewAvalancheFont->InitSystemFont(InFontParams, NewFont);
-	UAvaFontManagerSubsystem::SetupMetrics(NewAvalancheFont);
-	OSFontsMap.Add(SanitizedFontFamilyName, NewAvalancheFont);
+	NewAvaFont->InitSystemFont(InFontParams, NewFont);
+	UAvaFontManagerSubsystem::SetupMetrics(NewAvaFont);
+	OSFontsMap.Add(SanitizedFontFamilyName, NewAvaFont);
 }
 
 void UAvaFontManagerSubsystem::SortFontsMap()
 {
-	AvalancheFontsMap.KeySort([](const FString& A, const FString& B)
+	FontsMap.KeySort([](const FString& A, const FString& B)
 	{
 		return A < B;
 	});
 }
 
-void UAvaFontManagerSubsystem::RefreshAvalancheFontsMap()
+void UAvaFontManagerSubsystem::RefreshFontsMap()
 {
 	CombineAllFonts();
 	SortFontsMap();
@@ -794,12 +794,12 @@ void UAvaFontManagerSubsystem::RefreshAvalancheFontsMap()
 
 TConstArrayView<const TSharedPtr<FAvaFontView>> UAvaFontManagerSubsystem::GetFontOptions()
 {
-	return AvalancheFontsOptions;
+	return FontsOptions;
 }
 
 TSharedPtr<FAvaFontView> UAvaFontManagerSubsystem::GetFontViewFromName(const FString& InName)
 {
-	TSharedPtr<FAvaFontView>* AvalancheFontView = AvalancheFontsOptions.FindByPredicate([&InName](const TSharedPtr<FAvaFontView>& AvaFontValue)
+	TSharedPtr<FAvaFontView>* FontView = FontsOptions.FindByPredicate([&InName](const TSharedPtr<FAvaFontView>& AvaFontValue)
 	{
 		if (AvaFontValue.IsValid() && AvaFontValue->HasValidFont())
 		{
@@ -809,9 +809,9 @@ TSharedPtr<FAvaFontView> UAvaFontManagerSubsystem::GetFontViewFromName(const FSt
 		return false;
 	});
 
-	if (AvalancheFontView)
+	if (FontView)
 	{
-		return *AvalancheFontView;
+		return *FontView;
 	}
 
 	return nullptr;
@@ -976,7 +976,7 @@ void UAvaFontManagerSubsystem::GetMultipleSelectionInformation(const TSharedPtr<
 void UAvaFontManagerSubsystem::GetFontName(const UFont* InFont, FString& OutFontName)
 {
 	// this function used to be defined here, but it has been moved to the Runtime module for easier maintenance
-	UE::Avalanche::FontUtilities::Public::GetFontName(InFont, OutFontName);
+	UE::Ava::FontUtilities::Public::GetFontName(InFont, OutFontName);
 }
 
 void UAvaFontManagerSubsystem::GetSanitizedFontName(const UFont* InFont, FString& OutFontName)
@@ -988,13 +988,13 @@ void UAvaFontManagerSubsystem::GetSanitizedFontName(const UFont* InFont, FString
 	}
 }
 
-void UAvaFontManagerSubsystem::RefreshAvalancheFontsOptions()
+void UAvaFontManagerSubsystem::RefreshFontsOptions()
 {
-	AvalancheFontsOptions.Empty();
+	FontsOptions.Empty();
 
 	int32 Count = 0;
 
-	for (const TPair<FString, TObjectPtr<UAvaFontObject>>& Elem : AvalancheFontsMap)
+	for (const TPair<FString, TObjectPtr<UAvaFontObject>>& Elem : FontsMap)
 	{
 		FString ElemFontName = Elem.Key;
 		const TObjectPtr<UAvaFontObject>& FontObject = Elem.Value;
@@ -1005,7 +1005,7 @@ void UAvaFontManagerSubsystem::RefreshAvalancheFontsOptions()
 			FAvaFont FontOption = FAvaFont(FontObject.Get());
 			const bool bFontIsFavorite = IsFavoriteFont(FontOption.GetFontNameAsString());
 			FontOption.SetFavorite(bFontIsFavorite);
-			AvalancheFontsOptions.Add(FAvaFontView::Make(FontOption));
+			FontsOptions.Add(FAvaFontView::Make(FontOption));
 
 			if (FontObject->GetFont() == FAvaFont::GetDefaultFont())
 			{
@@ -1022,9 +1022,9 @@ void UAvaFontManagerSubsystem::AddFavorite(const FString& InFontName)
 	// does something only if needed
 	LoadFavorites();
 
-	if (IsValid(AvaFontManagerConfig))
+	if (IsValid(FontManagerConfig))
 	{
-		AvaFontManagerConfig->AddFavoriteFont(InFontName);
+		FontManagerConfig->AddFavoriteFont(InFontName);
 	}
 }
 
@@ -1033,9 +1033,9 @@ void UAvaFontManagerSubsystem::RemoveFavorite(const FString& InFontName)
 	// does something only if needed
 	LoadFavorites();
 
-	if (IsValid(AvaFontManagerConfig))
+	if (IsValid(FontManagerConfig))
 	{
-		AvaFontManagerConfig->RemoveFavoriteFont(InFontName);
+		FontManagerConfig->RemoveFavoriteFont(InFontName);
 	}
 }
 
@@ -1044,9 +1044,9 @@ bool UAvaFontManagerSubsystem::IsFavoriteFont(const FString& InFontName)
 	// does something only if needed
 	LoadFavorites();
 
-	if (IsValid(AvaFontManagerConfig))
+	if (IsValid(FontManagerConfig))
 	{
-		return AvaFontManagerConfig->IsFavoriteFont(InFontName);
+		return FontManagerConfig->IsFavoriteFont(InFontName);
 	}
 
 	return false;
@@ -1127,7 +1127,7 @@ void UAvaFontManagerSubsystem::FontsAutoImportPreSave(UWorld* InWorld, FObjectPr
 
 UAvaFontConfig* UAvaFontManagerSubsystem::GetFontManagerConfig()
 {
-	return AvaFontManagerConfig;
+	return FontManagerConfig;
 }
 
 bool UAvaFontManagerSubsystem::IsFontMonospaced(const UFont* const InFont)
@@ -1224,7 +1224,7 @@ void UAvaFontManagerSubsystem::CleanFreeTypeLibrary()
 	}
 }
 
-//////// UAvalancheFontConfig ////////
+//////// UAvaFontConfig ////////
 
 bool UAvaFontConfig::IsFavoriteFont(const FString& InFontName) const
 {
diff --git a/Engine/Plugins/Experimental/Avalanche/Source/AvalancheTextEditor/Private/Font/AvaFontManagerSubsystem.h b/Engine/Plugins/Experimental/Avalanche/Source/AvalancheTextEditor/Private/Font/AvaFontManagerSubsystem.h
--- a/Engine/Plugins/Experimental/Avalanche/Source/AvalancheTextEditor/Private/Font/AvaFontManagerSubsystem.h
+++ b/Engine/Plugins/Experimental/Avalanche/Source/AvalancheTextEditor/Private/Font/AvaFontManagerSubsystem.h
@@ -119,7 +119,7 @@ public:
 
 	int32 GetDefaultFontIndex() const { return DefaultFontIndex; }
 
-	/** Given a Property Handle, tries to retrieve and return a font view from AvalancheFontsOptions */
+	/** Given a Property Handle, tries to retrieve and return a font view from FontsOptions */
 	TSharedPtr<FAvaFontView> GetFontViewFromPropertyHandle(const TSharedPtr<IPropertyHandle>& InFontPropertyHandle, FPropertyAccess::Result& OutAccessResult);
 
 	FOnProjectFontsChange& OnProjectFontCreated() { return OnProjectFontCreatedDelegate; }
@@ -140,7 +140,7 @@ private:
 	static void SetupFontFamilyTypefaces(UFont* InFont, const FSystemFontsRetrieveParams& InFontParams);
 
 	/**
-	 * Initializes OS fonts and Avalanche fonts data structures
+	 * Initializes OS fonts and Motion Design fonts data structures
 	 * Since this is called before project assets are available, it will NOT initialize project fonts.
 	 * That is done later by InitializeProjectFonts function.
 	 */
@@ -159,10 +159,10 @@ private:
 
 	UAvaFontObject* CreateProjectFont(UFont* InSourceFont, UPackage* InProjectFontsPackage = nullptr);
 
-	/** Given a font name, tries to retrieve and return a font from AvalancheFontsOptions */
+	/** Given a font name, tries to retrieve and return a font from FontsOptions */
 	TSharedPtr<FAvaFontView> GetFontViewFromName(const FString& InName);
 
-	/** Given a Property Handle, tries to retrieve and return a font from AvalancheFontsOptions as a raw pointer */
+	/** Given a Property Handle, tries to retrieve and return a font from FontsOptions as a raw pointer */
 	TSharedPtr<FAvaFontView> GetFontViewFromPropertyHandle(const TSharedPtr<IPropertyHandle>& InFontPropertyHandle);
 
 	void ClearFontsData();
@@ -172,8 +172,8 @@ private:
 	void LoadOSFonts();
 	void LoadProjectFonts();
 	void RemoveProjectFontsFromOSFonts();
-	void RefreshAvalancheFontsMap();
-	void RefreshAvalancheFontsOptions();
+	void RefreshFontsMap();
+	void RefreshFontsOptions();
 	void SortFontsMap();
 	void ImportSystemFontFamily(const FSystemFontsRetrieveParams& InFontParams, UPackage* InImportPackage);
 	void UnmarkFontFromAutoImport(const TSharedPtr<IPropertyHandle>& InFontToUnregister);
@@ -220,7 +220,7 @@ private:
 
 	/** used for custom config - e.g. favorite fonts */
 	UPROPERTY(Transient)
-	TObjectPtr<UAvaFontConfig> AvaFontManagerConfig;
+	TObjectPtr<UAvaFontConfig> FontManagerConfig;
 
 	/** Fonts available from the OS */
 	UPROPERTY(Transient)
@@ -232,12 +232,12 @@ private:
 
 	/** OS + current project fonts combined */
 	UPROPERTY(Transient)
-	TMap<FString, TObjectPtr<UAvaFontObject>> AvalancheFontsMap;
+	TMap<FString, TObjectPtr<UAvaFontObject>> FontsMap;
 
 	/** Caching font options here, so that details panel can get them right away */
-	TArray<TSharedPtr<FAvaFontView>> AvalancheFontsOptions;
+	TArray<TSharedPtr<FAvaFontView>> FontsOptions;
 
-	/** Used to allow auto font import on ava bp and level editor save  */
+	/** Used to allow auto font import on Motion Design asset save */
 	TMap<TSharedPtr<FAvaFontView>, TSharedPtr<IPropertyHandle>> FontsToImportOnSave;
 
 	TMap<FString, FSystemFontsRetrieveParams> FontsInfoMap;
diff --git a/Engine/Plugins/Experimental/Avalanche/Source/AvalancheTextEditor/Private/Platforms/Linux/LinuxSystemFontLoading.cpp b/Engine/Plugins/Experimental/Avalanche/Source/AvalancheTextEditor/Private/Platforms/Linux/LinuxSystemFontLoading.cpp
--- a/Engine/Plugins/Experimental/Avalanche/Source/AvalancheTextEditor/Private/Platforms/Linux/LinuxSystemFontLoading.cpp
+++ b/Engine/Plugins/Experimental/Avalanche/Source/AvalancheTextEditor/Private/Platforms/Linux/LinuxSystemFontLoading.cpp
@@ -2,7 +2,7 @@
 
 #include "LinuxSystemFontLoading.h"
 
-namespace UE::Avalanche::Private::Fonts
+namespace UE::Ava::Private::Fonts
 {
 	void GetSystemFontInfo(TMap<FString, FSystemFontsRetrieveParams>& OutFontsInfo)
 	{
diff --git a/Engine/Plugins/Experimental/Avalanche/Source/AvalancheTextEditor/Private/Platforms/Linux/LinuxSystemFontLoading.h b/Engine/Plugins/Experimental/Avalanche/Source/AvalancheTextEditor/Private/Platforms/Linux/LinuxSystemFontLoading.h
--- a/Engine/Plugins/Experimental/Avalanche/Source/AvalancheTextEditor/Private/Platforms/Linux/LinuxSystemFontLoading.h
+++ b/Engine/Plugins/Experimental/Avalanche/Source/AvalancheTextEditor/Private/Platforms/Linux/LinuxSystemFontLoading.h
@@ -8,7 +8,7 @@
 struct FSystemFontsRetrieveParams;
 
 //todo: yet to be implemented
-namespace UE::Avalanche::Private::Fonts
+namespace UE::Ava::Private::Fonts
 {
 	void GetSystemFontInfo(TMap<FString, FSystemFontsRetrieveParams>& OutFontsInfo);
 
diff --git a/Engine/Plugins/Experimental/Avalanche/Source/AvalancheTextEditor/Private/Slate/SAvaFontField.cpp b/Engine/Plugins/Experimental/Avalanche/Source/AvalancheTextEditor/Private/Slate/SAvaFontField.cpp
--- a/Engine/Plugins/Experimental/Avalanche/Source/AvalancheTextEditor/Private/Slate/SAvaFontField.cpp
+++ b/Engine/Plugins/Experimental/Avalanche/Source/AvalancheTextEditor/Private/Slate/SAvaFontField.cpp
@@ -12,7 +12,7 @@
 #include "Widgets/Layout/SScaleBox.h"
 #include "Widgets/Text/STextBlock.h"
 
-namespace UE::AvalancheTextEditor::Private
+namespace UE::AvaTextEditor::Private
 {
 	static constexpr float DefaultFontSize = 12.0f;
 }
@@ -23,8 +23,8 @@ BEGIN_SLATE_FUNCTION_BUILD_OPTIMIZATION
 
 void SAvaFontField::Construct(const FArguments& InArgs)
 {
-	AvalancheFont = InArgs._AvalancheFont;
-	OnAvaFontFieldModified = InArgs._OnAvaFontFieldModified;
+	FontView = InArgs._FontView;
+	OnFontFieldModified = InArgs._OnFontFieldModified;
 
 	const FVector2D Icon16x16(16.0f, 16.0f);
 
@@ -93,20 +93,21 @@ void SAvaFontField::Construct(const FArguments& InArgs)
 		]
 	];
 
-	UpdateFont(AvalancheFont.Get());
+	UpdateFont(FontView.Get());
 }
+
 END_SLATE_FUNCTION_BUILD_OPTIMIZATION
 
-void SAvaFontField::UpdateFont(const TSharedPtr<FAvaFontView>& InAvaFont)
+void SAvaFontField::UpdateFont(const TSharedPtr<FAvaFontView>& InFontView)
 {
-	if (!InAvaFont.IsValid())
+	if (!InFontView.IsValid())
 	{
 		return;
 	}
 
-	AvalancheFont.Set(InAvaFont);
+	FontView.Set(InFontView);
 
-	if (const TSharedPtr<FAvaFontView>& FontViewPtr = AvalancheFont.Get())
+	if (const TSharedPtr<FAvaFontView>& FontViewPtr = FontView.Get())
 	{
 		const UAvaFontObject* const CurrFontObj = FontViewPtr->GetFontObject();
 		const UFont* const CurrFont = FontViewPtr->GetFont();
@@ -131,7 +132,7 @@ void SAvaFontField::UpdateFont(const TSharedPtr<FAvaFontView>& InAvaFont)
 			{
 				if (CurrFont->GetCompositeFont())
 				{
-					LeftFontNameText->SetFont(FSlateFontInfo(CurrFont, UE::AvalancheTextEditor::Private::DefaultFontSize));
+					LeftFontNameText->SetFont(FSlateFontInfo(CurrFont, UE::AvaTextEditor::Private::DefaultFontSize));
 				}
 			}
 		}
@@ -150,7 +151,7 @@ void SAvaFontField::UpdateFont(const TSharedPtr<FAvaFontView>& InAvaFont)
 
 					if (CurrFont->GetCompositeFont())
 					{
-						RightFontNameText->SetFont(FSlateFontInfo(CurrFont, UE::AvalancheTextEditor::Private::DefaultFontSize));
+						RightFontNameText->SetFont(FSlateFontInfo(CurrFont, UE::AvaTextEditor::Private::DefaultFontSize));
 					}
 				}
 			}
@@ -164,7 +165,7 @@ void SAvaFontField::UpdateFont(const TSharedPtr<FAvaFontView>& InAvaFont)
 
 ECheckBoxState SAvaFontField::GetFavoriteState() const
 {
-	if (const TSharedPtr<FAvaFontView>& FontViewPtr = AvalancheFont.Get())
+	if (const TSharedPtr<FAvaFontView>& FontViewPtr = FontView.Get())
 	{
 		if (FontViewPtr->IsFavorite())
 		{
@@ -177,7 +178,7 @@ ECheckBoxState SAvaFontField::GetFavoriteState() const
 
 FSlateColor SAvaFontField::GetToggleFavoriteColor() const
 {
-	if (const TSharedPtr<FAvaFontView>& FontViewPtr = AvalancheFont.Get())
+	if (const TSharedPtr<FAvaFontView>& FontViewPtr = FontView.Get())
 	{
 		if (FontViewPtr->IsFavorite())
 		{
@@ -196,7 +197,7 @@ FReply SAvaFontField::OnToggleFavoriteClicked()
 		FReply::Handled();
 	}
 
-	if (const TSharedPtr<FAvaFontView>& FontViewPtr = AvalancheFont.Get())
+	if (const TSharedPtr<FAvaFontView>& FontViewPtr = FontView.Get())
 	{
 		const bool bIsNowFavorite = !FontViewPtr->IsFavorite();
 
@@ -213,7 +214,7 @@ FReply SAvaFontField::OnToggleFavoriteClicked()
 			FontManagerSubsystem->RemoveFavorite(*FontViewPtr->GetFontNameAsString());
 		}
 
-		OnAvaFontFieldModified.ExecuteIfBound(FontViewPtr);
+		OnFontFieldModified.ExecuteIfBound(FontViewPtr);
 	}
 
 	return FReply::Handled();
@@ -229,7 +230,7 @@ EVisibility SAvaFontField::GetLocallyAvailableIconVisibility() const
 	bool bIsAvailable = false;
 	if (const UAvaFontManagerSubsystem* FontManagerSubsystem = UAvaFontManagerSubsystem::Get())
 	{
-		if (const TSharedPtr<FAvaFontView>& FontViewPtr = AvalancheFont.Get())
+		if (const TSharedPtr<FAvaFontView>& FontViewPtr = FontView.Get())
 		{
 			bIsAvailable = FontManagerSubsystem->IsFontAvailableOnLocalOS(FontViewPtr->GetFont());
 		}
diff --git a/Engine/Plugins/Experimental/Avalanche/Source/AvalancheTextEditor/Private/Slate/SAvaFontField.h b/Engine/Plugins/Experimental/Avalanche/Source/AvalancheTextEditor/Private/Slate/SAvaFontField.h
--- a/Engine/Plugins/Experimental/Avalanche/Source/AvalancheTextEditor/Private/Slate/SAvaFontField.h
+++ b/Engine/Plugins/Experimental/Avalanche/Source/AvalancheTextEditor/Private/Slate/SAvaFontField.h
@@ -17,19 +17,13 @@ DECLARE_DELEGATE_OneParam(FOnAvaFontFieldModified, const TSharedPtr<FAvaFontView
 class SAvaFontField : public SCompoundWidget
 {
 public:
-	SLATE_BEGIN_ARGS(SAvaFontField)
-		:_AvalancheFont(),
-		_OnAvaFontFieldModified()
-		{
-		}
-
-		SLATE_ATTRIBUTE(TSharedPtr<FAvaFontView>, AvalancheFont)
-		SLATE_EVENT(FOnAvaFontFieldModified, OnAvaFontFieldModified)
-
+	SLATE_BEGIN_ARGS(SAvaFontField) {}
+		SLATE_ATTRIBUTE(TSharedPtr<FAvaFontView>, FontView)
+		SLATE_EVENT(FOnAvaFontFieldModified, OnFontFieldModified)
 	SLATE_END_ARGS()
 
 	void Construct(const FArguments& InArgs);
-	void UpdateFont(const TSharedPtr<FAvaFontView>& InAvaFont);
+	void UpdateFont(const TSharedPtr<FAvaFontView>& InFontView);
 	void Select() { bIsSelected = true; }
 	void Deselect() { bIsSelected = false; }
 
@@ -42,14 +36,14 @@ protected:
 
 	bool IsSelected() const { return bIsSelected; }
 
-	FOnAvaFontFieldModified OnAvaFontFieldModified;
+	FOnAvaFontFieldModified OnFontFieldModified;
 
 	FText FontName;
 
 	TSharedPtr<STextBlock> LeftFontNameText;
 	TSharedPtr<STextBlock> RightFontNameText;
 
-	TAttribute<TSharedPtr<FAvaFontView>> AvalancheFont;
+	TAttribute<TSharedPtr<FAvaFontView>> FontView;
 
 	TSharedPtr<FSlateBrush> CheckBoxBg;
 
diff --git a/Engine/Plugins/Experimental/Avalanche/Source/AvalancheTextEditor/Private/Slate/SAvaFontSelector.cpp b/Engine/Plugins/Experimental/Avalanche/Source/AvalancheTextEditor/Private/Slate/SAvaFontSelector.cpp
--- a/Engine/Plugins/Experimental/Avalanche/Source/AvalancheTextEditor/Private/Slate/SAvaFontSelector.cpp
+++ b/Engine/Plugins/Experimental/Avalanche/Source/AvalancheTextEditor/Private/Slate/SAvaFontSelector.cpp
@@ -133,7 +133,7 @@ void SAvaFontSelector::Construct(const FArguments& InArgs)
 	bShowBoldOnly = InArgs._ShowBoldFontsOnly;
 	bShowItalicOnly = InArgs._ShowItalicFontsOnly;
 
-	AvalancheFontPropertyHandlePtr = InArgs._AvalancheFontPropertyHandle.Get();
+	FontPropertyHandlePtr = InArgs._FontPropertyHandle.Get();
 
 	FontsOptionsSource = InArgs._OptionsSource;
 	CustomScrollbar = InArgs._CustomScrollbar;
@@ -739,7 +739,7 @@ void SAvaFontSelector::RefreshUIFacingFont()
 		return;
 	}
 
-	if (!AvalancheFontPropertyHandlePtr)
+	if (!FontPropertyHandlePtr)
 	{
 		return;
 	}
@@ -751,7 +751,7 @@ void SAvaFontSelector::RefreshUIFacingFont()
 	}
 
 	FPropertyAccess::Result AccessResult;
-	const TSharedPtr<FAvaFontView> ActualFontPtr = FontManagerSubsystem->GetFontViewFromPropertyHandle(AvalancheFontPropertyHandlePtr, AccessResult);
+	const TSharedPtr<FAvaFontView> ActualFontPtr = FontManagerSubsystem->GetFontViewFromPropertyHandle(FontPropertyHandlePtr, AccessResult);
 	if (!ActualFontPtr)
 	{
 		return;
diff --git a/Engine/Plugins/Experimental/Avalanche/Source/AvalancheTextEditor/Private/Slate/SAvaFontSelector.h b/Engine/Plugins/Experimental/Avalanche/Source/AvalancheTextEditor/Private/Slate/SAvaFontSelector.h
--- a/Engine/Plugins/Experimental/Avalanche/Source/AvalancheTextEditor/Private/Slate/SAvaFontSelector.h
+++ b/Engine/Plugins/Experimental/Avalanche/Source/AvalancheTextEditor/Private/Slate/SAvaFontSelector.h
@@ -27,18 +27,18 @@ DECLARE_DELEGATE(FOnAvaComboBoxOpening)
 class SAvaFontSelector : public SComboButton
 {
 public:
-	typedef TListTypeTraits<TSharedPtr<FAvaFontView>> ListTypeTraits;
-	typedef TListTypeTraits<TSharedPtr<FAvaFontView>>::NullableType NullableOptionType;
+	using ListTypeTraits = TListTypeTraits<TSharedPtr<FAvaFontView>>;
+	using NullableOptionType = TListTypeTraits<TSharedPtr<FAvaFontView>>::NullableType;
 
 	/** Type of list used for showing menu options. */
-	typedef SListView<TSharedPtr<FAvaFontView>> SComboListType;
+	using SComboListType = SListView<TSharedPtr<FAvaFontView>>;
 
 	/** Type of text filter used for FAvaFontView */
-	typedef TTextFilter<TSharedPtr<FAvaFontView>> FontViewTextFilter;
+	using FontViewTextFilter = TTextFilter<TSharedPtr<FAvaFontView>>;
 
 	/** Delegate types used to generate widgets that represent Options */
-	typedef TSlateDelegates<const TSharedPtr<FAvaFontView>&>::FOnGenerateWidget FOnGenerateWidget;
-	typedef TSlateDelegates<NullableOptionType>::FOnSelectionChanged FOnSelectionChanged;
+	using FOnGenerateWidget = TSlateDelegates<const TSharedPtr<FAvaFontView>&>::FOnGenerateWidget;
+	using FOnSelectionChanged = TSlateDelegates<NullableOptionType>::FOnSelectionChanged;
 
 	SLATE_BEGIN_ARGS(SAvaFontSelector)
 		: _Content()
@@ -47,7 +47,7 @@ public:
 		, _ItemStyle(&FAppStyle::Get().GetWidgetStyle< FTableRowStyle >("ComboBox.Row"))
 		, _ContentPadding(_ComboBoxStyle->ContentPadding)
 		, _ForegroundColor(FSlateColor::UseStyle())
-		, _AvalancheFontPropertyHandle()
+		, _FontPropertyHandle()
 		, _OptionsSource()
 		, _OnSelectionChanged()
 		, _OnGenerateWidget()
@@ -76,7 +76,7 @@ public:
 		SLATE_ATTRIBUTE(FSlateColor, ForegroundColor)
 
 		/** The Property Handle for the font being customized */
-		SLATE_ATTRIBUTE(TSharedPtr<IPropertyHandle>, AvalancheFontPropertyHandle)
+		SLATE_ATTRIBUTE(TSharedPtr<IPropertyHandle>, FontPropertyHandle)
 
 		/** Options source for the fonts list available to this font selector widget */
 		SLATE_ARGUMENT(const TArray<TSharedPtr<FAvaFontView>>*, OptionsSource)
@@ -284,5 +284,5 @@ private:
 
 	TArray<TSharedPtr<FAvaFontView>> FilteredFonts;
 
-	TSharedPtr<IPropertyHandle> AvalancheFontPropertyHandlePtr;
+	TSharedPtr<IPropertyHandle> FontPropertyHandlePtr;
 };
diff --git a/Engine/Plugins/Experimental/Avalanche/Source/AvalancheTextEditor/Private/Slate/SAvaTextAlignmentWidget.cpp b/Engine/Plugins/Experimental/Avalanche/Source/AvalancheTextEditor/Private/Slate/SAvaTextAlignmentWidget.cpp
--- a/Engine/Plugins/Experimental/Avalanche/Source/AvalancheTextEditor/Private/Slate/SAvaTextAlignmentWidget.cpp
+++ b/Engine/Plugins/Experimental/Avalanche/Source/AvalancheTextEditor/Private/Slate/SAvaTextAlignmentWidget.cpp
@@ -51,15 +51,15 @@ void SAvaTextAlignmentWidget::Construct(const FArguments& InArgs)
 			]
 			+ SGridPanel::Slot(1, 0, ForegroundLayer)
 			[
-				GetHorizontalAlignmentButton(HorizontalLeftButton, EText3DHorizontalTextAlignment::Left, "AvalancheIcons.Alignment.Left", LOCTEXT("AlignSelectedTextLeft", "Align Text to the Left"))
+				GetHorizontalAlignmentButton(HorizontalLeftButton, EText3DHorizontalTextAlignment::Left, "Icons.Alignment.Left", LOCTEXT("AlignSelectedTextLeft", "Align Text to the Left"))
 			]
 			+ SGridPanel::Slot(2, 0, ForegroundLayer)
 			[
-				GetHorizontalAlignmentButton(HorizontalCenterButton, EText3DHorizontalTextAlignment::Center, "AvalancheIcons.Alignment.Center_Y", LOCTEXT("AlignSelectedTextCenter", "Align Text to Center"))
+				GetHorizontalAlignmentButton(HorizontalCenterButton, EText3DHorizontalTextAlignment::Center, "Icons.Alignment.Center_Y", LOCTEXT("AlignSelectedTextCenter", "Align Text to Center"))
 			]
 			+ SGridPanel::Slot(3, 0, ForegroundLayer)
 			[
-				GetHorizontalAlignmentButton(HorizontalRightButton, EText3DHorizontalTextAlignment::Right, "AvalancheIcons.Alignment.Right", LOCTEXT("AlignSelectedTextRight", "Align Text to the Right"))
+				GetHorizontalAlignmentButton(HorizontalRightButton, EText3DHorizontalTextAlignment::Right, "Icons.Alignment.Right", LOCTEXT("AlignSelectedTextRight", "Align Text to the Right"))
 			]
 			+ SGridPanel::Slot(0, 0, BackgroundLayer) // color 1st empty slot
 			[
@@ -68,19 +68,19 @@ void SAvaTextAlignmentWidget::Construct(const FArguments& InArgs)
 			]
 			+ SGridPanel::Slot(0, 1, ForegroundLayer)
 			[
-				GetVerticalAlignmentButton(VerticalFirstLineButton, EText3DVerticalTextAlignment::FirstLine, "AvalancheIcons.Alignment.Top", LOCTEXT("AlignSelectedTextFirstLine", "Align Text to First Line"))
+				GetVerticalAlignmentButton(VerticalFirstLineButton, EText3DVerticalTextAlignment::FirstLine, "Icons.Alignment.Top", LOCTEXT("AlignSelectedTextFirstLine", "Align Text to First Line"))
 			]
 			+ SGridPanel::Slot(1, 1, ForegroundLayer)
 			[
-				GetVerticalAlignmentButton(VerticalTopButton, EText3DVerticalTextAlignment::Top, "AvalancheIcons.Alignment.Top", LOCTEXT("AlignSelectedTextTop", "Align Text to Top"))
+				GetVerticalAlignmentButton(VerticalTopButton, EText3DVerticalTextAlignment::Top, "Icons.Alignment.Top", LOCTEXT("AlignSelectedTextTop", "Align Text to Top"))
 			]
 			+ SGridPanel::Slot(2, 1, ForegroundLayer)
 			[
-				GetVerticalAlignmentButton(VerticalCenterButton, EText3DVerticalTextAlignment::Center, "AvalancheIcons.Alignment.Center_Z", LOCTEXT("AlignSelectedTextCenter", "Align Text to Center"))
+				GetVerticalAlignmentButton(VerticalCenterButton, EText3DVerticalTextAlignment::Center, "Icons.Alignment.Center_Z", LOCTEXT("AlignSelectedTextCenter", "Align Text to Center"))
 			]
 			+ SGridPanel::Slot(3, 1, ForegroundLayer)
 			[
-				GetVerticalAlignmentButton(VerticalBottomButton, EText3DVerticalTextAlignment::Bottom, "AvalancheIcons.Alignment.Bottom", LOCTEXT("AlignSelectedTextBottom", "Align Text to Bottom"))
+				GetVerticalAlignmentButton(VerticalBottomButton, EText3DVerticalTextAlignment::Bottom, "Icons.Alignment.Bottom", LOCTEXT("AlignSelectedTextBottom", "Align Text to Bottom"))
 			]
 		]
 	];
@@ -92,7 +92,7 @@ END_SLATE_FUNCTION_BUILD_OPTIMIZATION
 
 TSharedRef<SButton> SAvaTextAlignmentWidget::GetHorizontalAlignmentButton(TSharedPtr<SButton>& OutButton, EText3DHorizontalTextAlignment InHorizontalAlignment, FName Image, FText Tooltip)
 {
-	const ISlateStyle* AvalancheEditorStyle = FSlateStyleRegistry::FindSlateStyle("AvaEditor");
+	const ISlateStyle* MotionDesignEditorStyle = FSlateStyleRegistry::FindSlateStyle("AvaEditor");
 
 	OutButton = SNew(SButton)
 		.ButtonStyle(&FAppStyle::Get().GetWidgetStyle<FButtonStyle>("SimpleButton"))
@@ -101,7 +101,7 @@ TSharedRef<SButton> SAvaTextAlignmentWidget::GetHorizontalAlignmentButton(TShare
 		.ToolTipText(Tooltip)
 		[
 			SNew(SImage)
-			.Image(AvalancheEditorStyle ? AvalancheEditorStyle->GetBrush(Image) : nullptr)
+			.Image(MotionDesignEditorStyle ? MotionDesignEditorStyle->GetBrush(Image) : nullptr)
 			.DesiredSizeOverride(FVector2D(16.f, 16.f))
 			.ColorAndOpacity(this, &SAvaTextAlignmentWidget::ApplyHorizontalAlignmentButtonColorAndOpacity, InHorizontalAlignment)
 		];
@@ -111,7 +111,7 @@ TSharedRef<SButton> SAvaTextAlignmentWidget::GetHorizontalAlignmentButton(TShare
 
 TSharedRef<SButton> SAvaTextAlignmentWidget::GetVerticalAlignmentButton(TSharedPtr<SButton>& OutButton, EText3DVerticalTextAlignment InVerticalAlignment, FName Image, FText Tooltip)
 {
-	const ISlateStyle* AvalancheEditorStyle = FSlateStyleRegistry::FindSlateStyle("AvaEditor");
+	const ISlateStyle* MotionDesignEditorStyle = FSlateStyleRegistry::FindSlateStyle("AvaEditor");
 
 	OutButton = SNew(SButton)
 		.ButtonStyle(&FAppStyle::Get().GetWidgetStyle<FButtonStyle>("SimpleButton"))
@@ -120,7 +120,7 @@ TSharedRef<SButton> SAvaTextAlignmentWidget::GetVerticalAlignmentButton(TSharedP
 		.ToolTipText(Tooltip)
 		[
 			SNew(SImage)
-			.Image(AvalancheEditorStyle ? AvalancheEditorStyle->GetBrush(Image) : nullptr)
+			.Image(MotionDesignEditorStyle ? MotionDesignEditorStyle->GetBrush(Image) : nullptr)
 			.DesiredSizeOverride(FVector2D(16.f, 16.f))
 			.ColorAndOpacity(this, &SAvaTextAlignmentWidget::ApplyVerticalAlignmentButtonColorAndOpacity, InVerticalAlignment)
 		];
diff --git a/Engine/Plugins/Experimental/Avalanche/Source/AvalancheTextEditor/Private/Visualizer/AvaTextVisualizer.cpp b/Engine/Plugins/Experimental/Avalanche/Source/AvalancheTextEditor/Private/Visualizer/AvaTextVisualizer.cpp
--- a/Engine/Plugins/Experimental/Avalanche/Source/AvalancheTextEditor/Private/Visualizer/AvaTextVisualizer.cpp
+++ b/Engine/Plugins/Experimental/Avalanche/Source/AvalancheTextEditor/Private/Visualizer/AvaTextVisualizer.cpp
@@ -26,7 +26,7 @@ IMPLEMENT_HIT_PROXY(HAvaTextGradientLineEndHandleProxy, HAvaHitProxy)
 IMPLEMENT_HIT_PROXY(HAvaTextGradientCenterHandleProxy, HAvaHitProxy)
 IMPLEMENT_HIT_PROXY(HAvaTextGradientSmoothnessHandleProxy, HAvaHitProxy)
 
-namespace UE::Avalanche::Private
+namespace UE::Ava::Private
 {
 	struct FTextVisualizerStatics
 	{
@@ -287,7 +287,7 @@ bool FAvaTextVisualizer::VisProxyHandleClick(FEditorViewportClient* InViewportCl
 
 FVector FAvaTextVisualizer::GetGradientSmoothnessHandleLocation(const UAvaText3DComponent* InAvaTextComponent) const
 {
-	using namespace UE::Avalanche::Private;
+	using namespace UE::Ava::Private;
 	
 	return GetGradientCenterHandleLocation(InAvaTextComponent) + FVector(0, FTextVisualizerStatics::GradientSmoothnessHandleOffset, 0);
 }
@@ -422,7 +422,7 @@ bool FAvaTextVisualizer::GetWidgetAxisListDragOverride(const FEditorViewportClie
 bool FAvaTextVisualizer::HandleInputDeltaInternal(FEditorViewportClient* InViewportClient, FViewport* InViewport, const FVector& InAccumulatedTranslation,
 	const FRotator& InAccumulatedRotation, const FVector& InAccumulatedScale)
 {
-	using namespace UE::Avalanche::Private;
+	using namespace UE::Ava::Private;
 	
 	if (bEditingWidth)
 	{
@@ -843,7 +843,7 @@ FVector FAvaTextVisualizer::GetHeightHandleLocation(const UText3DComponent* InTe
 
 void FAvaTextVisualizer::GetTextActorGradientControlsLocations(const UAvaText3DComponent* InAvaTextComponent, FVector& OutGradientCenterLocation, FVector& OutGradientStartLocation, FVector& OutGradientEndLocation) const
 {
-	using namespace UE::Avalanche::Private;
+	using namespace UE::Ava::Private;
 	if (IsValid(InAvaTextComponent))
 	{
 		if (const AActor* const Text3DActor = InAvaTextComponent->GetOwner())
@@ -1099,7 +1099,7 @@ bool FAvaTextVisualizer::DrawGradientSmoothnessHandle(const UText3DComponent* In
 
 void FAvaTextVisualizer::DrawGradientHandles(const UAvaText3DComponent* InAvaTextComponent, const FSceneView* InView, FPrimitiveDrawInterface* InPDI) const
 {
-	using namespace UE::Avalanche::Private;
+	using namespace UE::Ava::Private;
 	if (!InAvaTextComponent)
 	{
 		return;
diff --git a/Engine/Plugins/Experimental/Avalanche/Source/AvalancheTextEditor/Private/Visualizer/AvaTextVisualizer.h b/Engine/Plugins/Experimental/Avalanche/Source/AvalancheTextEditor/Private/Visualizer/AvaTextVisualizer.h
--- a/Engine/Plugins/Experimental/Avalanche/Source/AvalancheTextEditor/Private/Visualizer/AvaTextVisualizer.h
+++ b/Engine/Plugins/Experimental/Avalanche/Source/AvalancheTextEditor/Private/Visualizer/AvaTextVisualizer.h
@@ -100,9 +100,8 @@ struct HAvaTextGradientSmoothnessHandleProxy : HAvaHitProxy
 
 class FAvaTextVisualizer : public FAvaVisualizerBase
 {
-
 public:
-	typedef FAvaVisualizerBase Super;
+	using Super = FAvaVisualizerBase;
 
 	void ResetEditingFlags();
 
diff --git a/Engine/Plugins/Experimental/Avalanche/Source/AvalancheTransitionEditor/Private/TabFactories/AvaTransitionTabFactory.h b/Engine/Plugins/Experimental/Avalanche/Source/AvalancheTransitionEditor/Private/TabFactories/AvaTransitionTabFactory.h
--- a/Engine/Plugins/Experimental/Avalanche/Source/AvalancheTransitionEditor/Private/TabFactories/AvaTransitionTabFactory.h
+++ b/Engine/Plugins/Experimental/Avalanche/Source/AvalancheTransitionEditor/Private/TabFactories/AvaTransitionTabFactory.h
@@ -5,7 +5,7 @@
 #include "AvaTransitionEditor.h"
 #include "WorkflowOrientedApp/WorkflowTabFactory.h"
 
-/** Base class for all Tab Factories in Avalanche Transition Editor */
+/** Base class for all Tab Factories in Motion Design Transition Editor */
 class FAvaTransitionTabFactory : public FWorkflowTabFactory
 {
 public:
diff --git a/Engine/Plugins/Experimental/Avalanche/Source/AvalancheViewport/Private/AvaScreenAlignmentActorInfo.cpp b/Engine/Plugins/Experimental/Avalanche/Source/AvalancheViewport/Private/AvaScreenAlignmentActorInfo.cpp
--- a/Engine/Plugins/Experimental/Avalanche/Source/AvalancheViewport/Private/AvaScreenAlignmentActorInfo.cpp
+++ b/Engine/Plugins/Experimental/Avalanche/Source/AvalancheViewport/Private/AvaScreenAlignmentActorInfo.cpp
@@ -10,7 +10,7 @@
 #include "Math/Transform.h"
 #include "Math/Vector.h"
 
-namespace UE::AvalancheEditor::Private
+namespace UE::AvaEditor::Private
 {
 	[[nodiscard]] FBox GetActorLocalBounds(AActor& InActor, EAvaAlignmentSizeMode InActorSizeMode)
 	{
@@ -111,7 +111,7 @@ FAvaScreenAlignmentActorInfo FAvaScreenAlignmentActorInfo::Create(const TSharedR
 		return CreateInfo();
 	}
 
-	const FBox ActorLocalBoundingBox = UE::AvalancheEditor::Private::GetActorLocalBounds(InActor, InActorSizeMode);
+	const FBox ActorLocalBoundingBox = UE::AvaEditor::Private::GetActorLocalBounds(InActor, InActorSizeMode);
 
 	if (!ActorLocalBoundingBox.IsValid)
 	{
diff --git a/Engine/Plugins/Experimental/Avalanche/Source/AvalancheViewport/Private/AvaScreenAlignmentUtils.cpp b/Engine/Plugins/Experimental/Avalanche/Source/AvalancheViewport/Private/AvaScreenAlignmentUtils.cpp
--- a/Engine/Plugins/Experimental/Avalanche/Source/AvalancheViewport/Private/AvaScreenAlignmentUtils.cpp
+++ b/Engine/Plugins/Experimental/Avalanche/Source/AvalancheViewport/Private/AvaScreenAlignmentUtils.cpp
@@ -13,7 +13,7 @@
 #include "GameFramework/Actor.h"
 #include "ViewportClient/IAvaViewportWorldCoordinateConverter.h"
 
-namespace UE::AvalancheViewport::Private
+namespace UE::AvaViewport::Private
 {
 	static constexpr float ValueTolerance = 0.01f;
 	static constexpr float ExactValueTolerance = 0.001f;
@@ -769,7 +769,7 @@ void FAvaScreenAlignmentUtils::SizeActorToScreen(const TSharedRef<IAvaViewportWo
 	const FAvaScreenAlignmentActorInfo ActorInfo = FAvaScreenAlignmentActorInfo::Create(InCoordinateConverter,
 		InActor, EAvaAlignmentSizeMode::Self);
 
-	UE::AvalancheViewport::Private::SizeActorToScreen(InCoordinateConverter, ActorInfo, bInStretchToFit);
+	UE::AvaViewport::Private::SizeActorToScreen(InCoordinateConverter, ActorInfo, bInStretchToFit);
 }
 
 void FAvaScreenAlignmentUtils::FitActorToScreen(const TSharedRef<IAvaViewportWorldCoordinateConverter>& InCoordinateConverter,
@@ -793,7 +793,7 @@ void FAvaScreenAlignmentUtils::FitActorToScreen(const TSharedRef<IAvaViewportWor
 			const FAvaScreenAlignmentActorInfo ActorInfo = FAvaScreenAlignmentActorInfo::Create(InCoordinateConverter,
 				InActor, EAvaAlignmentSizeMode::Self);
 
-			UE::AvalancheViewport::Private::SizeActorToScreen(InCoordinateConverter, ActorInfo, bInStretchToFit);
+			UE::AvaViewport::Private::SizeActorToScreen(InCoordinateConverter, ActorInfo, bInStretchToFit);
 
 			const FTransform ActorTransform = InActor.GetActorTransform();
 			const FBox ActorLocalBounds = BoundsSubsystem->GetActorLocalBounds(&InActor);
@@ -820,7 +820,7 @@ void FAvaScreenAlignmentUtils::AlignActorRotationAxis(const TSharedRef<IAvaViewp
 		);
 	}
 
-	using namespace UE::AvalancheViewport::Private;
+	using namespace UE::AvaViewport::Private;
 
 	bool bMadeChange = false;
 	UE::Math::TMatrix<double> ActorRotationMatrix = UE::Math::TRotationMatrix<double>(InActor.GetActorRotation());
@@ -918,7 +918,7 @@ void FAvaScreenAlignmentUtils::AlignActorsHorizontal(const TSharedRef<IAvaViewpo
 	const TArray<AActor*>& InActors, EAvaHorizontalAlignment InHorizontalAlignment, EAvaAlignmentSizeMode InActorSizeMode, 
 	EAvaAlignmentContext InContextType)
 {
-	UE::AvalancheViewport::Private::AlignActors<EAvaHorizontalAlignment>(
+	UE::AvaViewport::Private::AlignActors<EAvaHorizontalAlignment>(
 		InCoordinateConverter,
 		InActors,
 		InHorizontalAlignment,
@@ -931,7 +931,7 @@ void FAvaScreenAlignmentUtils::AlignActorsVertical(const TSharedRef<IAvaViewport
 	const TArray<AActor*>& InActors, EAvaVerticalAlignment InVerticalAlignment, EAvaAlignmentSizeMode InActorSizeMode, 
 	EAvaAlignmentContext InContextType)
 {
-	UE::AvalancheViewport::Private::AlignActors<EAvaVerticalAlignment>(
+	UE::AvaViewport::Private::AlignActors<EAvaVerticalAlignment>(
 		InCoordinateConverter,
 		InActors,
 		InVerticalAlignment,
@@ -943,7 +943,7 @@ void FAvaScreenAlignmentUtils::AlignActorsVertical(const TSharedRef<IAvaViewport
 void FAvaScreenAlignmentUtils::AlignActorsDepth(const TSharedRef<IAvaViewportWorldCoordinateConverter>& InCoordinateConverter, 
 	const TArray<AActor*>& InActors, EAvaDepthAlignment InDepthAlignment, EAvaAlignmentSizeMode InActorSizeMode)
 {
-	UE::AvalancheViewport::Private::AlignActors<EAvaDepthAlignment>(
+	UE::AvaViewport::Private::AlignActors<EAvaDepthAlignment>(
 		InCoordinateConverter,
 		InActors,
 		InDepthAlignment,
@@ -956,7 +956,7 @@ void FAvaScreenAlignmentUtils::DistributeActorsHorizontal(const TSharedRef<IAvaV
 	const TArray<AActor*>& InActors, EAvaAlignmentSizeMode InActorSizeMode, EAvaActorDistributionMode InDistributionMode, 
 	EAvaAlignmentContext InContextType)
 {
-	UE::AvalancheViewport::Private::DistributeActors<EAvaHorizontalAlignment>(
+	UE::AvaViewport::Private::DistributeActors<EAvaHorizontalAlignment>(
 		InCoordinateConverter,
 		InActors,
 		InActorSizeMode,
@@ -969,7 +969,7 @@ void FAvaScreenAlignmentUtils::DistributeActorsVertical(const TSharedRef<IAvaVie
 	const TArray<AActor*>& InActors, EAvaAlignmentSizeMode InActorSizeMode, EAvaActorDistributionMode InDistributionMode, 
 	EAvaAlignmentContext InContextType)
 {
-	UE::AvalancheViewport::Private::DistributeActors<EAvaVerticalAlignment>(
+	UE::AvaViewport::Private::DistributeActors<EAvaVerticalAlignment>(
 		InCoordinateConverter,
 		InActors,
 		InActorSizeMode,
@@ -981,7 +981,7 @@ void FAvaScreenAlignmentUtils::DistributeActorsVertical(const TSharedRef<IAvaVie
 void FAvaScreenAlignmentUtils::DistributeActorsDepth(const TSharedRef<IAvaViewportWorldCoordinateConverter>& InCoordinateConverter,
 	const TArray<AActor*>& InActors, EAvaAlignmentSizeMode InActorSizeMode, EAvaActorDistributionMode InDistributionMode)
 {
-	UE::AvalancheViewport::Private::DistributeActors<EAvaDepthAlignment>(
+	UE::AvaViewport::Private::DistributeActors<EAvaDepthAlignment>(
 		InCoordinateConverter,
 		InActors,
 		InActorSizeMode,
diff --git a/Engine/Plugins/Experimental/Avalanche/Source/AvalancheViewport/Private/AvaViewportGuideInfo.cpp b/Engine/Plugins/Experimental/Avalanche/Source/AvalancheViewport/Private/AvaViewportGuideInfo.cpp
--- a/Engine/Plugins/Experimental/Avalanche/Source/AvalancheViewport/Private/AvaViewportGuideInfo.cpp
+++ b/Engine/Plugins/Experimental/Avalanche/Source/AvalancheViewport/Private/AvaViewportGuideInfo.cpp
@@ -4,7 +4,7 @@
 #include "Dom/JsonObject.h"
 #include "Dom/JsonValue.h"
 
-namespace UE::AvalancheViewport::Private
+namespace UE::AvaViewport::Private
 {
 	static const FString IsHorizontalFieldName = TEXT("IsHorizontal");
 	static const FString OffsetFieldName = TEXT("Offset");
@@ -14,7 +14,7 @@ namespace UE::AvalancheViewport::Private
 
 bool FAvaViewportGuideInfo::DeserializeJson(const TSharedRef<FJsonObject>& InJsonObject, FAvaViewportGuideInfo& OutGuideInfo, const FVector2f& InViewportSize)
 {
-	using namespace UE::AvalancheViewport::Private;
+	using namespace UE::AvaViewport::Private;
 
 	bool bValid = false;
 
@@ -58,7 +58,7 @@ bool FAvaViewportGuideInfo::DeserializeJson(const TSharedRef<FJsonObject>& InJso
 
 bool FAvaViewportGuideInfo::SerializeJson(const TSharedRef<FJsonObject>& InJsonObject, const FVector2f& InViewportSize) const
 {
-	using namespace UE::AvalancheViewport::Private;
+	using namespace UE::AvaViewport::Private;
 
 	const bool bIsHorziontal = Orientation == EOrientation::Orient_Horizontal;
 
diff --git a/Engine/Plugins/Experimental/Avalanche/Source/AvalancheViewport/Private/AvaViewportGuidePresetProvider.cpp b/Engine/Plugins/Experimental/Avalanche/Source/AvalancheViewport/Private/AvaViewportGuidePresetProvider.cpp
--- a/Engine/Plugins/Experimental/Avalanche/Source/AvalancheViewport/Private/AvaViewportGuidePresetProvider.cpp
+++ b/Engine/Plugins/Experimental/Avalanche/Source/AvalancheViewport/Private/AvaViewportGuidePresetProvider.cpp
@@ -14,7 +14,7 @@
 #include "Serialization/JsonSerializable.h"
 #include "Serialization/JsonWriter.h"
 
-namespace UE::AvalancheViewport::Private
+namespace UE::AvaViewport::Private
 {
 	static const FString GuideListFieldName = TEXT("Guides");
 
@@ -162,7 +162,7 @@ namespace UE::AvalancheViewport::Private
 bool FAvaViewportGuidePresetProvider::SaveGuidePreset(const FString& InPresetName, const TArray<FAvaViewportGuideInfo>& InGuides,
 	const FVector2f InViewportSize)
 {
-	using namespace UE::AvalancheViewport::Private;
+	using namespace UE::AvaViewport::Private;
 
 	const FString Path = GetGuideConfigPath();
 
@@ -187,7 +187,7 @@ bool FAvaViewportGuidePresetProvider::SaveGuidePreset(const FString& InPresetNam
 bool FAvaViewportGuidePresetProvider::LoadGuidePreset(const FString& InPresetName, TArray<FAvaViewportGuideInfo>& OutGuides, 
 	const FVector2f InViewportSize)
 {
-	using namespace UE::AvalancheViewport::Private;
+	using namespace UE::AvaViewport::Private;
 
 	const FString Path = GetGuideConfigPath();
 
@@ -214,7 +214,7 @@ bool FAvaViewportGuidePresetProvider::LoadGuidePreset(const FString& InPresetNam
 
 bool FAvaViewportGuidePresetProvider::RemoveGuidePreset(const FString& InPresetName)
 {
-	using namespace UE::AvalancheViewport::Private;
+	using namespace UE::AvaViewport::Private;
 
 	const FString Path = GetGuideConfigPath(/* Allow Create */ false);
 
@@ -235,7 +235,7 @@ bool FAvaViewportGuidePresetProvider::RemoveGuidePreset(const FString& InPresetN
 
 TArray<FString> FAvaViewportGuidePresetProvider::GetGuidePresetNames()
 {
-	using namespace UE::AvalancheViewport::Private;
+	using namespace UE::AvaViewport::Private;
 
 	const FString Path = GetGuideConfigPath();
 
diff --git a/Engine/Plugins/Experimental/Avalanche/Source/AvalancheViewport/Private/AvaViewportUtils.cpp b/Engine/Plugins/Experimental/Avalanche/Source/AvalancheViewport/Private/AvaViewportUtils.cpp
--- a/Engine/Plugins/Experimental/Avalanche/Source/AvalancheViewport/Private/AvaViewportUtils.cpp
+++ b/Engine/Plugins/Experimental/Avalanche/Source/AvalancheViewport/Private/AvaViewportUtils.cpp
@@ -10,7 +10,7 @@
 #include "ViewportClient/EditorViewportClientUtilityWrapper.h"
 #include "ViewportClient/IAvaViewportClient.h"
 
-namespace UE::AvalancheViewport::Private
+namespace UE::AvaViewport::Private
 {
 	TArray<FAvaViewportClientCastDelegate> ViewportClientCasterRegistry;
 	TArray<FAvaViewportCastDelegate> ViewportCasterRegistry;
@@ -18,7 +18,7 @@ namespace UE::AvalancheViewport::Private
 
 FDelegateHandle FAvaViewportUtils::RegisterViewportClientCaster(FAvaViewportClientCastDelegate::TFuncType InFunction)
 {
-	using namespace UE::AvalancheViewport::Private;
+	using namespace UE::AvaViewport::Private;
 
 	FAvaViewportClientCastDelegate& NewDelegate = ViewportClientCasterRegistry.AddDefaulted_GetRef();
 	NewDelegate.BindStatic(InFunction);
@@ -28,7 +28,7 @@ FDelegateHandle FAvaViewportUtils::RegisterViewportClientCaster(FAvaViewportClie
 
 void FAvaViewportUtils::UnregisterViewportClientCaster(FDelegateHandle InDelegateHandle)
 {
-	using namespace UE::AvalancheViewport::Private;
+	using namespace UE::AvaViewport::Private;
 
 	ViewportClientCasterRegistry.RemoveAll(
 		[InDelegateHandle](const FAvaViewportClientCastDelegate& InElement)
@@ -39,7 +39,7 @@ void FAvaViewportUtils::UnregisterViewportClientCaster(FDelegateHandle InDelegat
 
 FDelegateHandle FAvaViewportUtils::RegisterViewportCaster(FAvaViewportCastDelegate::TFuncType InFunction)
 {
-	using namespace UE::AvalancheViewport::Private;
+	using namespace UE::AvaViewport::Private;
 
 	FAvaViewportCastDelegate& NewDelegate = ViewportCasterRegistry.AddDefaulted_GetRef();
 	NewDelegate.BindStatic(InFunction);
@@ -49,7 +49,7 @@ FDelegateHandle FAvaViewportUtils::RegisterViewportCaster(FAvaViewportCastDelega
 
 void FAvaViewportUtils::UnregisterViewportCaster(FDelegateHandle InDelegateHandle)
 {
-	using namespace UE::AvalancheViewport::Private;
+	using namespace UE::AvaViewport::Private;
 
 	ViewportCasterRegistry.RemoveAll(
 		[InDelegateHandle](const FAvaViewportCastDelegate& InElement)
@@ -65,7 +65,7 @@ TSharedPtr<IAvaViewportClient> FAvaViewportUtils::GetAsAvaViewportClient(FEditor
 		return nullptr;
 	}
 
-	using namespace UE::AvalancheViewport::Private;
+	using namespace UE::AvaViewport::Private;
 
 	for (const FAvaViewportClientCastDelegate& CastDelegate : ViewportClientCasterRegistry)
 	{
@@ -99,7 +99,7 @@ FEditorViewportClient* FAvaViewportUtils::GetAsEditorViewportClient(FViewport* I
 		return nullptr;
 	}
 
-	using namespace UE::AvalancheViewport::Private;
+	using namespace UE::AvaViewport::Private;
 
 	for (const FAvaViewportCastDelegate& CastDelegate : ViewportCasterRegistry)
 	{
diff --git a/Engine/Plugins/Experimental/Avalanche/Source/AvalancheViewport/Private/Interaction/AvaCameraZoomController.cpp b/Engine/Plugins/Experimental/Avalanche/Source/AvalancheViewport/Private/Interaction/AvaCameraZoomController.cpp
--- a/Engine/Plugins/Experimental/Avalanche/Source/AvalancheViewport/Private/Interaction/AvaCameraZoomController.cpp
+++ b/Engine/Plugins/Experimental/Avalanche/Source/AvalancheViewport/Private/Interaction/AvaCameraZoomController.cpp
@@ -9,7 +9,7 @@
 #include "Framework/Application/SlateApplication.h"
 #include "ViewportClient/IAvaViewportClient.h"
 
-namespace UE::AvalancheLevelViewport::Private
+namespace UE::AvaLevelViewport::Private
 {
 	constexpr uint8 ZOOM_LEVEL_MIN = 0;
 	constexpr uint8 ZOOM_LEVEL_MAX = 9;
@@ -32,7 +32,7 @@ FAvaCameraZoomController::FAvaCameraZoomController(TSharedRef<IAvaViewportClient
 
 void FAvaCameraZoomController::SetZoomLevel(uint8 InZoomLevel)
 {
-	ZoomLevel = FMath::Clamp(InZoomLevel, UE::AvalancheLevelViewport::Private::ZOOM_LEVEL_MIN, UE::AvalancheLevelViewport::Private::ZOOM_LEVEL_MAX);
+	ZoomLevel = FMath::Clamp(InZoomLevel, UE::AvaLevelViewport::Private::ZOOM_LEVEL_MIN, UE::AvaLevelViewport::Private::ZOOM_LEVEL_MAX);
 
 	UpdateVisibleAreas();
 	InvalidateViewport();
@@ -40,12 +40,12 @@ void FAvaCameraZoomController::SetZoomLevel(uint8 InZoomLevel)
 
 bool FAvaCameraZoomController::IsZoomed() const
 {
-	return ZoomLevel > UE::AvalancheLevelViewport::Private::ZOOM_LEVEL_MIN;
+	return ZoomLevel > UE::AvaLevelViewport::Private::ZOOM_LEVEL_MIN;
 }
 
 float FAvaCameraZoomController::GetFOVPerStep() const
 {
-	return GetDefaultFOV() / static_cast<float>(UE::AvalancheLevelViewport::Private::ZOOM_LEVEL_MAX - UE::AvalancheLevelViewport::Private::ZOOM_LEVEL_MIN + 1);
+	return GetDefaultFOV() / static_cast<float>(UE::AvaLevelViewport::Private::ZOOM_LEVEL_MAX - UE::AvaLevelViewport::Private::ZOOM_LEVEL_MIN + 1);
 }
 
 void FAvaCameraZoomController::ZoomIn()
@@ -327,7 +327,7 @@ void FAvaCameraZoomController::EndPanning()
 
 void FAvaCameraZoomController::ZoomIn_Internal()
 {
-	if (ZoomLevel >= UE::AvalancheLevelViewport::Private::ZOOM_LEVEL_MAX)
+	if (ZoomLevel >= UE::AvaLevelViewport::Private::ZOOM_LEVEL_MAX)
 	{
 		return;
 	}
@@ -337,7 +337,7 @@ void FAvaCameraZoomController::ZoomIn_Internal()
 
 void FAvaCameraZoomController::ZoomOut_Internal()
 {
-	if (ZoomLevel <= UE::AvalancheLevelViewport::Private::ZOOM_LEVEL_MIN)
+	if (ZoomLevel <= UE::AvaLevelViewport::Private::ZOOM_LEVEL_MIN)
 	{
 		return;
 	}
diff --git a/Engine/Plugins/Experimental/Avalanche/Source/AvalancheViewport/Private/Interaction/AvaIsolateActorsOperation.cpp b/Engine/Plugins/Experimental/Avalanche/Source/AvalancheViewport/Private/Interaction/AvaIsolateActorsOperation.cpp
--- a/Engine/Plugins/Experimental/Avalanche/Source/AvalancheViewport/Private/Interaction/AvaIsolateActorsOperation.cpp
+++ b/Engine/Plugins/Experimental/Avalanche/Source/AvalancheViewport/Private/Interaction/AvaIsolateActorsOperation.cpp
@@ -13,7 +13,7 @@
 #include "Selection/AvaSelectionProviderSubsystem.h"
 #include "ViewportClient/IAvaViewportClient.h"
 
-namespace UE::AvalancheViewport::Private
+namespace UE::AvaViewport::Private
 {
 	const TArray<UClass*> PersistentActorClasses = {
 		AInfo::StaticClass(),
@@ -135,7 +135,7 @@ void FAvaIsolateActorsOperation::IsolateActors()
 	IsolatedActors.Empty();
 	IsolatedActors.Reserve(SelectedActors.Num());
 
-	using namespace UE::AvalancheViewport::Private;
+	using namespace UE::AvaViewport::Private;
 
 	for (const TWeakObjectPtr<AActor>& ActorWeak : SelectedActors)
 	{
diff --git a/Engine/Plugins/Experimental/Avalanche/Source/AvalancheViewport/Private/Interaction/AvaSnapOperation.cpp b/Engine/Plugins/Experimental/Avalanche/Source/AvalancheViewport/Private/Interaction/AvaSnapOperation.cpp
--- a/Engine/Plugins/Experimental/Avalanche/Source/AvalancheViewport/Private/Interaction/AvaSnapOperation.cpp
+++ b/Engine/Plugins/Experimental/Avalanche/Source/AvalancheViewport/Private/Interaction/AvaSnapOperation.cpp
@@ -27,7 +27,7 @@
 #include "Viewport/Interaction/IAvaViewportDataProvider.h"
 #include "ViewportClient/IAvaViewportClient.h"
 
-namespace UE::AvalancheViewport::Private
+namespace UE::AvaViewport::Private
 {
 	EAvaViewportSnapState GetSnapState()
 	{
@@ -259,7 +259,7 @@ FAvaSnapOperation::~FAvaSnapOperation()
 
 void FAvaSnapOperation::GenerateLocalSnapPoints(AActor* InActor, TArray<FAvaSnapPoint>& InOutSnapPoints)
 {
-	using namespace UE::AvalancheViewport::Private;
+	using namespace UE::AvaViewport::Private;
 
 	if (!InActor || !IsActorEligibleForSnapping(InActor))
 	{
@@ -415,7 +415,7 @@ void FAvaSnapOperation::GenerateBoundsSnapPoints(const FVector& InOrigin, const
 
 void FAvaSnapOperation::GenerateScreenSnapPoints()
 {
-	using namespace UE::AvalancheViewport::Private;
+	using namespace UE::AvaViewport::Private;
 
 	if (!EnumHasAnyFlags(GetSnapState(), EAvaViewportSnapState::Screen))
 	{
@@ -476,7 +476,7 @@ void FAvaSnapOperation::GenerateScreenSnapPoints()
 
 void FAvaSnapOperation::GenerateComponentSnapPoints(const UActorComponent* InComponent)
 {
-	using namespace UE::AvalancheViewport::Private;
+	using namespace UE::AvaViewport::Private;
 
 	if (!EnumHasAnyFlags(GetSnapState(), EAvaViewportSnapState::Actor))
 	{
@@ -513,7 +513,7 @@ void FAvaSnapOperation::GenerateComponentSnapPoints(const UActorComponent* InCom
 void FAvaSnapOperation::GenerateActorSnapPoints(const TConstArrayView<TWeakObjectPtr<AActor>>& InSelectedActors,
 	const TConstArrayView<TWeakObjectPtr<AActor>>& InExcludedActors)
 {
-	using namespace UE::AvalancheViewport::Private;
+	using namespace UE::AvaViewport::Private;
 
 	TSharedPtr<IAvaViewportClient> AvaViewportClient = FAvaViewportUtils::GetAsAvaViewportClient(EditorViewportClient);
 
@@ -719,7 +719,7 @@ bool FAvaSnapOperation::SnapY(float& InOutComponent, float InViewportSize, int32
 bool FAvaSnapOperation::SnapVectorComponent(float& InOutComponent, const TArray<FAvaSnapPointLink>& InSnapPointLinks,
 	float InViewportSize, int32 InScreenSnapPointCount, int32& OutClosestSnapPointIdx, float& OutDistanceFromSnap)
 {
-	using namespace UE::AvalancheViewport::Private;
+	using namespace UE::AvaViewport::Private;
 
 	OutDistanceFromSnap = -1;	
 
@@ -831,7 +831,7 @@ bool FAvaSnapOperation::SnapVectorComponent(float& InOutComponent, const TArray<
 
 void FAvaSnapOperation::SnapScreenLocation(FVector2f& InOutScreenLocation, bool bInSnapX, bool bInSnapY)
 {
-	using namespace UE::AvalancheViewport::Private;
+	using namespace UE::AvaViewport::Private;
 
 	bSnappedToLinkX = false;
 	bSnappedToLinkY = false;
@@ -895,7 +895,7 @@ void FAvaSnapOperation::SnapScreenLocationY(float& InOutScreenLocation)
 
 bool FAvaSnapOperation::SnapLocation(FVector& Location)
 {
-	using namespace UE::AvalancheViewport::Private;
+	using namespace UE::AvaViewport::Private;
 
 	bSnappedToLinkX = false;
 	bSnappedToLinkY = false;
@@ -995,7 +995,7 @@ bool FAvaSnapOperation::SnapDragLocation(const TArray<FAvaSnapPoint>& InDraggedA
 	 * - Move actor
 	 */
 
-	using namespace UE::AvalancheViewport::Private;
+	using namespace UE::AvaViewport::Private;
 
 	bSnappedToLinkX = false;
 	bSnappedToLinkY = false;
diff --git a/Engine/Plugins/Experimental/Avalanche/Source/AvalancheViewport/Private/ViewportClient/EditorViewportClientUtilityWrapper.h b/Engine/Plugins/Experimental/Avalanche/Source/AvalancheViewport/Private/ViewportClient/EditorViewportClientUtilityWrapper.h
--- a/Engine/Plugins/Experimental/Avalanche/Source/AvalancheViewport/Private/ViewportClient/EditorViewportClientUtilityWrapper.h
+++ b/Engine/Plugins/Experimental/Avalanche/Source/AvalancheViewport/Private/ViewportClient/EditorViewportClientUtilityWrapper.h
@@ -26,7 +26,7 @@ public:
 	//~ End IAvaViewportWorldCoordinateConverter
 
 	//~ Begin IAvaViewportClient
-	virtual bool IsAvalancheViewport() const override { return false; }
+	virtual bool IsMotionDesignViewport() const override { return false; }
 	virtual const FEditorViewportClient* AsEditorViewportClient() const override;
 	virtual bool SupportsZoom() const override { return false; } // No support
 	virtual float GetUnZoomedFOV() const override;
diff --git a/Engine/Plugins/Experimental/Avalanche/Source/AvalancheViewport/Private/Visualizers/AvaViewportBackgroundVisualizer.cpp b/Engine/Plugins/Experimental/Avalanche/Source/AvalancheViewport/Private/Visualizers/AvaViewportBackgroundVisualizer.cpp
--- a/Engine/Plugins/Experimental/Avalanche/Source/AvalancheViewport/Private/Visualizers/AvaViewportBackgroundVisualizer.cpp
+++ b/Engine/Plugins/Experimental/Avalanche/Source/AvalancheViewport/Private/Visualizers/AvaViewportBackgroundVisualizer.cpp
@@ -16,7 +16,7 @@
 
 #define LOCTEXT_NAMESPACE "AvaViewportBackgroundVisualizer"
 
-namespace UE::AvalancheViewport::Private
+namespace UE::AvaViewport::Private
 {
 	const FString BackgroundReferencerName = FString(TEXT("AvaViewportBackgroundVisualizer"));
 	const FName TextureObjectName = FName(TEXT("TextureObject"));
@@ -80,7 +80,7 @@ void FAvaViewportBackgroundVisualizer::AddReferencedObjects(FReferenceCollector&
 
 FString FAvaViewportBackgroundVisualizer::GetReferencerName() const
 {
-	return UE::AvalancheViewport::Private::BackgroundReferencerName;
+	return UE::AvaViewport::Private::BackgroundReferencerName;
 }
 
 void FAvaViewportBackgroundVisualizer::UpdateForViewport(const FAvaVisibleArea& InVisibleArea, const FVector2f& InWidgetSize, 
@@ -108,7 +108,7 @@ void FAvaViewportBackgroundVisualizer::UpdateForViewport(const FAvaVisibleArea&
 		return;
 	}
 
-	using namespace UE::AvalancheViewport::Private;
+	using namespace UE::AvaViewport::Private;
 
 	const float ImageAspectRatio = ImageSize.X / ImageSize.Y;
 	const float WidgetAspectRatio = InWidgetSize.X / InWidgetSize.Y;
@@ -176,7 +176,7 @@ void FAvaViewportBackgroundVisualizer::UpdatePostProcessMaterial()
 
 	Super::UpdatePostProcessMaterial();
 
-	using namespace UE::AvalancheViewport::Private;
+	using namespace UE::AvaViewport::Private;
 
 	PostProcessMaterial->SetTextureParameterValue(TextureObjectName, Texture);
 }
diff --git a/Engine/Plugins/Experimental/Avalanche/Source/AvalancheViewport/Private/Visualizers/AvaViewportBoundingBoxVisualizer.cpp b/Engine/Plugins/Experimental/Avalanche/Source/AvalancheViewport/Private/Visualizers/AvaViewportBoundingBoxVisualizer.cpp
--- a/Engine/Plugins/Experimental/Avalanche/Source/AvalancheViewport/Private/Visualizers/AvaViewportBoundingBoxVisualizer.cpp
+++ b/Engine/Plugins/Experimental/Avalanche/Source/AvalancheViewport/Private/Visualizers/AvaViewportBoundingBoxVisualizer.cpp
@@ -12,7 +12,7 @@
 #include "Selection.h"
 #include "Selection/AvaSelectionProviderSubsystem.h"
 
-namespace UE::AvalancheViewport::Private
+namespace UE::AvaViewport::Private
 {
 	constexpr float BoundingBoxLineThickness = 0.4f;
 }
@@ -22,7 +22,7 @@ TSharedRef<IAvaViewportBoundingBoxVisualizerInterface> FAvaViewportBoundingBoxVi
 	return MakeShared<FAvaLevelViewportBoundingBoxVisualizer>();
 }
 
-namespace UE::AvalancheLevelViewport::Private
+namespace UE::AvaLevelViewport::Private
 {
 	struct FAvaTaskTimer
 	{
@@ -78,7 +78,7 @@ namespace UE::AvalancheLevelViewport::Private
 			return;
 		}
 
-		using namespace UE::AvalancheViewport::Private;
+		using namespace UE::AvaViewport::Private;
 
 		FVector Vertices[8];
 		InOrientedBox.CalcVertices(Vertices);
@@ -140,7 +140,7 @@ void FAvaLevelViewportBoundingBoxVisualizer::Draw(UAvaSelectionProviderSubsystem
 		return;
 	}
 
-	using namespace UE::AvalancheLevelViewport::Private;
+	using namespace UE::AvaLevelViewport::Private;
 
 	constexpr double MaxTime = 0.003;
 	constexpr double HalfMaxTime = MaxTime * 0.5;
diff --git a/Engine/Plugins/Experimental/Avalanche/Source/AvalancheViewport/Private/Visualizers/AvaViewportChannelVisualizer.cpp b/Engine/Plugins/Experimental/Avalanche/Source/AvalancheViewport/Private/Visualizers/AvaViewportChannelVisualizer.cpp
--- a/Engine/Plugins/Experimental/Avalanche/Source/AvalancheViewport/Private/Visualizers/AvaViewportChannelVisualizer.cpp
+++ b/Engine/Plugins/Experimental/Avalanche/Source/AvalancheViewport/Private/Visualizers/AvaViewportChannelVisualizer.cpp
@@ -9,7 +9,7 @@
 #include "UObject/Package.h"
 #include "Viewport/Interaction/AvaViewportPostProcessInfo.h"
 
-namespace UE::AvalancheViewport::Private
+namespace UE::AvaViewport::Private
 {
 	const FString ChannelReferencerName = FString(TEXT("AvaViewportChannelVisualizer"));
 
@@ -26,7 +26,7 @@ FAvaViewportChannelVisualizer::FAvaViewportChannelVisualizer(TSharedRef<IAvaView
 {
 	bRequiresTonemapperSetting = (InChannel == EAvaViewportPostProcessType::AlphaChannel);
 
-	using namespace UE::AvalancheViewport::Private;
+	using namespace UE::AvaViewport::Private;
 
 	if (const TSoftObjectPtr<UMaterial>* BaseMaterialPtr = ChannelAssets.Find(InChannel))
 	{
@@ -40,5 +40,5 @@ FAvaViewportChannelVisualizer::FAvaViewportChannelVisualizer(TSharedRef<IAvaView
 
 FString FAvaViewportChannelVisualizer::GetReferencerName() const
 {
-	return UE::AvalancheViewport::Private::ChannelReferencerName;
+	return UE::AvaViewport::Private::ChannelReferencerName;
 }
diff --git a/Engine/Plugins/Experimental/Avalanche/Source/AvalancheViewport/Private/Visualizers/AvaViewportCheckerboardVisualizer.cpp b/Engine/Plugins/Experimental/Avalanche/Source/AvalancheViewport/Private/Visualizers/AvaViewportCheckerboardVisualizer.cpp
--- a/Engine/Plugins/Experimental/Avalanche/Source/AvalancheViewport/Private/Visualizers/AvaViewportCheckerboardVisualizer.cpp
+++ b/Engine/Plugins/Experimental/Avalanche/Source/AvalancheViewport/Private/Visualizers/AvaViewportCheckerboardVisualizer.cpp
@@ -9,7 +9,7 @@
 
 #define LOCTEXT_NAMESPACE "AvaViewportBackgroundVisualizer"
 
-namespace UE::AvalancheViewport::Private
+namespace UE::AvaViewport::Private
 {
 	const FString CheckerboardReferencerName = FString(TEXT("AvaViewportCheckerboardVisualizer"));
 	const FName CheckerboardColor0Name = FName(TEXT("CheckerboardColor0"));
@@ -51,7 +51,7 @@ void FAvaViewportCheckerboardVisualizer::UpdatePostProcessMaterial()
 	
 	if (PostProcessMaterial)
 	{
-		using namespace UE::AvalancheViewport::Private;
+		using namespace UE::AvaViewport::Private;
 		PostProcessMaterial->SetVectorParameterValue(CheckerboardColor0Name, CheckerboardColor0);
 		PostProcessMaterial->SetVectorParameterValue(CheckerboardColor1Name, CheckerboardColor1);
 		PostProcessMaterial->SetVectorParameterValue(CheckerboardSizeName, CheckerboardSize);
@@ -60,7 +60,7 @@ void FAvaViewportCheckerboardVisualizer::UpdatePostProcessMaterial()
 
 FString FAvaViewportCheckerboardVisualizer::GetReferencerName() const
 {
-	return UE::AvalancheViewport::Private::CheckerboardReferencerName;
+	return UE::AvaViewport::Private::CheckerboardReferencerName;
 }
 
 void FAvaViewportCheckerboardVisualizer::OnSettingChanged(UObject* InSettings, FPropertyChangedEvent& InPropertyChangeEvent)
diff --git a/Engine/Plugins/Experimental/Avalanche/Source/AvalancheViewport/Private/Visualizers/AvaViewportPostProcessVisualizer.cpp b/Engine/Plugins/Experimental/Avalanche/Source/AvalancheViewport/Private/Visualizers/AvaViewportPostProcessVisualizer.cpp
--- a/Engine/Plugins/Experimental/Avalanche/Source/AvalancheViewport/Private/Visualizers/AvaViewportPostProcessVisualizer.cpp
+++ b/Engine/Plugins/Experimental/Avalanche/Source/AvalancheViewport/Private/Visualizers/AvaViewportPostProcessVisualizer.cpp
@@ -19,7 +19,7 @@
 
 #define LOCTEXT_NAMESPACE "AvaViewportPostProcessVisualizer"
 
-namespace UE::AvalancheViewport::Private
+namespace UE::AvaViewport::Private
 {
 	const FName OpacityName = FName(TEXT("Opacity"));
 }
@@ -215,7 +215,7 @@ void FAvaViewportPostProcessVisualizer::UpdatePostProcessMaterial()
 		return;
 	}
 
-	using namespace UE::AvalancheViewport::Private;
+	using namespace UE::AvaViewport::Private;
 
 	PostProcessMaterial->SetScalarParameterValue(OpacityName, PostProcessOpacity);
 }
