[Motion Design]
 - Remove "Avalanche" from resource names

#virtualized

diff --git a/Engine/Plugins/Experimental/Avalanche/Source/AvalancheEditor/Private/DetailView/Customizations/AvaMeshesDetailCustomization.cpp b/Engine/Plugins/Experimental/Avalanche/Source/AvalancheEditor/Private/DetailView/Customizations/AvaMeshesDetailCustomization.cpp
--- a/Engine/Plugins/Experimental/Avalanche/Source/AvalancheEditor/Private/DetailView/Customizations/AvaMeshesDetailCustomization.cpp
+++ b/Engine/Plugins/Experimental/Avalanche/Source/AvalancheEditor/Private/DetailView/Customizations/AvaMeshesDetailCustomization.cpp
@@ -257,7 +257,7 @@ FReply FAvaMeshesDetailCustomization::OnConvertToStaticMeshClicked()
 	}
 
 	// generate name for asset
-	FString NewNameSuggestion = TEXT("SM_Avalanche") + DynMeshBase->GetMeshName();
+	FString NewNameSuggestion = TEXT("SM_MotionDesign") + DynMeshBase->GetMeshName();
 	FString PackageName = FString(TEXT("/Game/Meshes/")) + NewNameSuggestion;
 	FString AssetName;
 
diff --git a/Engine/Plugins/Experimental/Avalanche/Source/AvalancheModifiers/Private/Modifiers/AvaDynamicMeshConverterModifier.cpp b/Engine/Plugins/Experimental/Avalanche/Source/AvalancheModifiers/Private/Modifiers/AvaDynamicMeshConverterModifier.cpp
--- a/Engine/Plugins/Experimental/Avalanche/Source/AvalancheModifiers/Private/Modifiers/AvaDynamicMeshConverterModifier.cpp
+++ b/Engine/Plugins/Experimental/Avalanche/Source/AvalancheModifiers/Private/Modifiers/AvaDynamicMeshConverterModifier.cpp
@@ -247,7 +247,7 @@ void UAvaDynamicMeshConverterModifier::ConvertToStaticMeshAsset()
 	}
 	
 	// generate name for asset
-	const FString NewNameSuggestion = TEXT("SM_Avalanche_") + OwningActor->GetActorNameOrLabel();
+	const FString NewNameSuggestion = TEXT("SM_MotionDesign_") + OwningActor->GetActorNameOrLabel();
 	FString PackageName = FString(TEXT("/Game/Meshes/")) + NewNameSuggestion;
 	FString AssetName;
 
diff --git a/Engine/Plugins/Experimental/Avalanche/Source/AvalancheOutliner/Private/Clipboard/AvaOutlinerImporter.cpp b/Engine/Plugins/Experimental/Avalanche/Source/AvalancheOutliner/Private/Clipboard/AvaOutlinerImporter.cpp
--- a/Engine/Plugins/Experimental/Avalanche/Source/AvalancheOutliner/Private/Clipboard/AvaOutlinerImporter.cpp
+++ b/Engine/Plugins/Experimental/Avalanche/Source/AvalancheOutliner/Private/Clipboard/AvaOutlinerImporter.cpp
@@ -96,7 +96,7 @@ TArray<FName> FAvaOutlinerImporter::ImportOutlinerData(const FString& InBuffer)
 		TArray<FName> ActorNames;
 	};
 
-	UPackage* const TempPackage = NewObject<UPackage>(nullptr, TEXT("/Temp/AvalancheOutliner/Import"), RF_Transient);
+	UPackage* const TempPackage = NewObject<UPackage>(nullptr, TEXT("/Temp/MotionDesignOutliner/Import"), RF_Transient);
 	TempPackage->AddToRoot();
 
 	// Turn the text buffer into objects
diff --git a/Engine/Plugins/Experimental/Avalanche/Source/AvalancheTextEditor/Private/Font/AvaFontManagerSubsystem.h b/Engine/Plugins/Experimental/Avalanche/Source/AvalancheTextEditor/Private/Font/AvaFontManagerSubsystem.h
--- a/Engine/Plugins/Experimental/Avalanche/Source/AvalancheTextEditor/Private/Font/AvaFontManagerSubsystem.h
+++ b/Engine/Plugins/Experimental/Avalanche/Source/AvalancheTextEditor/Private/Font/AvaFontManagerSubsystem.h
@@ -192,7 +192,7 @@ private:
 	static FString GetImportFontPackageNameRoot() { return TEXT("/Game/SystemFonts"); }
 	static FString GetTempFontPackageNameRoot() { return TEXT("/Temp/SystemFonts"); }
 	static FString GetFontObjPackageNameRoot() { return TEXT("/Temp/SystemFonts/FontObjs"); }
-	static FString GetFontConfigPackageName() { return TEXT("/Temp/AvalancheEditor/AvalancheFontConfig"); }
+	static FString GetFontConfigPackageName() { return TEXT("/Temp/MotionDesignEditor/FontConfig"); }
 
 	static void InitFreeTypeLibrary();
 	static void CleanFreeTypeLibrary();
diff --git a/Engine/Plugins/Experimental/Avalanche/Source/AvalancheViewport/Private/AvaViewportSettings.cpp b/Engine/Plugins/Experimental/Avalanche/Source/AvalancheViewport/Private/AvaViewportSettings.cpp
--- a/Engine/Plugins/Experimental/Avalanche/Source/AvalancheViewport/Private/AvaViewportSettings.cpp
+++ b/Engine/Plugins/Experimental/Avalanche/Source/AvalancheViewport/Private/AvaViewportSettings.cpp
@@ -9,8 +9,8 @@ UAvaViewportSettings::UAvaViewportSettings()
 
 	bEnableViewportOverlay = true;
 	bEnableBoundingBoxes = true;
-	ViewportBackgroundMaterial = TSoftObjectPtr<UMaterial>(FSoftObjectPath(TEXT("/Script/Engine.Material'/Avalanche/EditorResources/M_AvalancheBackground.M_AvalancheBackground'")));
-	ViewportCheckerboardMaterial = TSoftObjectPtr<UMaterial>(FSoftObjectPath(TEXT("/Script/Engine.Material'/Avalanche/EditorResources/M_AvalancheCheckerboard.M_AvalancheCheckerboard'")));
+	ViewportBackgroundMaterial = TSoftObjectPtr<UMaterial>(FSoftObjectPath(TEXT("/Script/Engine.Material'/Avalanche/EditorResources/M_Background.M_Background'")));
+	ViewportCheckerboardMaterial = TSoftObjectPtr<UMaterial>(FSoftObjectPath(TEXT("/Script/Engine.Material'/Avalanche/EditorResources/M_Checkerboard.M_Checkerboard'")));
 	ViewportCheckerboardColor0 = FLinearColor(FVector3f(0.048172f));
 	ViewportCheckerboardColor1 = FLinearColor(FVector3f(0.177888f));
 	ViewportCheckerboardSize = 8.0f;
diff --git a/Engine/Plugins/Experimental/Avalanche/Source/AvalancheViewport/Private/Visualizers/AvaViewportChannelVisualizer.cpp b/Engine/Plugins/Experimental/Avalanche/Source/AvalancheViewport/Private/Visualizers/AvaViewportChannelVisualizer.cpp
--- a/Engine/Plugins/Experimental/Avalanche/Source/AvalancheViewport/Private/Visualizers/AvaViewportChannelVisualizer.cpp
+++ b/Engine/Plugins/Experimental/Avalanche/Source/AvalancheViewport/Private/Visualizers/AvaViewportChannelVisualizer.cpp
@@ -14,10 +14,10 @@ namespace UE::AvaViewport::Private
 	const FString ChannelReferencerName = FString(TEXT("AvaViewportChannelVisualizer"));
 
 	const TMap<EAvaViewportPostProcessType, TSoftObjectPtr<UMaterial>> ChannelAssets = {
-		{ EAvaViewportPostProcessType::RedChannel,   TSoftObjectPtr<UMaterial>(FSoftObjectPath("/Script/Engine.Material'/Avalanche/EditorResources/M_AvalancheChannelVisualizer_Red.M_AvalancheChannelVisualizer_Red'")) },
-		{ EAvaViewportPostProcessType::GreenChannel, TSoftObjectPtr<UMaterial>(FSoftObjectPath("/Script/Engine.Material'/Avalanche/EditorResources/M_AvalancheChannelVisualizer_Green.M_AvalancheChannelVisualizer_Green'")) },
-		{ EAvaViewportPostProcessType::BlueChannel,  TSoftObjectPtr<UMaterial>(FSoftObjectPath("/Script/Engine.Material'/Avalanche/EditorResources/M_AvalancheChannelVisualizer_Blue.M_AvalancheChannelVisualizer_Blue'")) },
-		{ EAvaViewportPostProcessType::AlphaChannel, TSoftObjectPtr<UMaterial>(FSoftObjectPath("/Script/Engine.Material'/Avalanche/EditorResources/M_AvalancheChannelVisualizer_Alpha.M_AvalancheChannelVisualizer_Alpha'")) }
+		{ EAvaViewportPostProcessType::RedChannel,   TSoftObjectPtr<UMaterial>(FSoftObjectPath("/Script/Engine.Material'/Avalanche/EditorResources/M_ChannelVisualizer_Red.M_ChannelVisualizer_Red'")) },
+		{ EAvaViewportPostProcessType::GreenChannel, TSoftObjectPtr<UMaterial>(FSoftObjectPath("/Script/Engine.Material'/Avalanche/EditorResources/M_ChannelVisualizer_Green.M_ChannelVisualizer_Green'")) },
+		{ EAvaViewportPostProcessType::BlueChannel,  TSoftObjectPtr<UMaterial>(FSoftObjectPath("/Script/Engine.Material'/Avalanche/EditorResources/M_ChannelVisualizer_Blue.M_ChannelVisualizer_Blue'")) },
+		{ EAvaViewportPostProcessType::AlphaChannel, TSoftObjectPtr<UMaterial>(FSoftObjectPath("/Script/Engine.Material'/Avalanche/EditorResources/M_ChannelVisualizer_Alpha.M_ChannelVisualizer_Alpha'")) }
 	};
 }
 
