MotionDesign : Modifiers
- Cleaning public facing API

#changelist validated
#virtualized

diff --git a/Engine/Plugins/Experimental/Avalanche/Source/AvalancheModifiers/Internal/Tests/Framework/AvaModifiersTestUtils.h b/Engine/Plugins/Experimental/Avalanche/Source/AvalancheModifiers/Internal/Tests/Framework/AvaModifiersTestUtils.h
--- a/Engine/Plugins/Experimental/Avalanche/Source/AvalancheModifiers/Internal/Tests/Framework/AvaModifiersTestUtils.h
+++ b/Engine/Plugins/Experimental/Avalanche/Source/AvalancheModifiers/Internal/Tests/Framework/AvaModifiersTestUtils.h
@@ -17,7 +17,7 @@ class AAvaTestDynamicMeshActor;
 class AAvaTestStaticMeshActor;
 class UWorld;
 
-struct AVALANCHEMODIFIERS_API FAvaModifierTestUtils
+struct FAvaModifierTestUtils
 {
 public:
 	FAvaModifierTestUtils(TSharedPtr<FAvaTestUtils> InTestUtils);
diff --git a/Engine/Plugins/Experimental/Avalanche/Source/AvalancheModifiers/Internal/Tests/Framework/AvaTestData.h b/Engine/Plugins/Experimental/Avalanche/Source/AvalancheModifiers/Internal/Tests/Framework/AvaTestData.h
--- a/Engine/Plugins/Experimental/Avalanche/Source/AvalancheModifiers/Internal/Tests/Framework/AvaTestData.h
+++ b/Engine/Plugins/Experimental/Avalanche/Source/AvalancheModifiers/Internal/Tests/Framework/AvaTestData.h
@@ -5,7 +5,7 @@
 #include "Containers/Set.h"
 #include "UObject/NameTypes.h"
 
-class AVALANCHEMODIFIERS_API FAvaTestData
+class FAvaTestData
 {
 public:
 	FAvaTestData();
diff --git a/Engine/Plugins/Experimental/Avalanche/Source/AvalancheModifiers/Private/Modifiers/AvaAlignBetweenModifier.cpp b/Engine/Plugins/Experimental/Avalanche/Source/AvalancheModifiers/Private/Modifiers/AvaAlignBetweenModifier.cpp
--- a/Engine/Plugins/Experimental/Avalanche/Source/AvalancheModifiers/Private/Modifiers/AvaAlignBetweenModifier.cpp
+++ b/Engine/Plugins/Experimental/Avalanche/Source/AvalancheModifiers/Private/Modifiers/AvaAlignBetweenModifier.cpp
@@ -123,7 +123,7 @@ TSet<AActor*> UAvaAlignBetweenModifier::GetActors(const bool bEnabledOnly) const
 	return OutActors;
 }
 
-void UAvaAlignBetweenModifier::SetReferenceActors(TSet<FAvaAlignBetweenWeightedActor> NewReferenceActors)
+void UAvaAlignBetweenModifier::SetReferenceActors(const TSet<FAvaAlignBetweenWeightedActor>& NewReferenceActors)
 {
 	ReferenceActors = NewReferenceActors;
 	OnReferenceActorsChanged();
diff --git a/Engine/Plugins/Experimental/Avalanche/Source/AvalancheModifiers/Private/Modifiers/AvaHideEmptyModifier.cpp b/Engine/Plugins/Experimental/Avalanche/Source/AvalancheModifiers/Private/Modifiers/AvaHideEmptyModifier.cpp
--- a/Engine/Plugins/Experimental/Avalanche/Source/AvalancheModifiers/Private/Modifiers/AvaHideEmptyModifier.cpp
+++ b/Engine/Plugins/Experimental/Avalanche/Source/AvalancheModifiers/Private/Modifiers/AvaHideEmptyModifier.cpp
@@ -34,7 +34,7 @@ void UAvaHideEmptyModifier::OnModifierAdded(EActorModifierCoreEnableReason InRea
 			ReferenceActor.ReferenceContainer = EAvaReferenceContainer::Other;
 			ReferenceActor.ReferenceActorWeak = ActorModified;
 			ReferenceActor.bSkipHiddenActors = false;
-			
+
 			TextComponent = ActorModified->FindComponentByClass<UText3DComponent>();
 		}
 	}
@@ -72,7 +72,7 @@ void UAvaHideEmptyModifier::Apply()
 	}
 
 	UAvaVisibilityModifierShared* VisibilityShared = GetShared<UAvaVisibilityModifierShared>(true);
-	
+
 	if (Text3DComponent->GetText().IsEmpty())
 	{
 		const bool bNewVisibility = bInvertVisibility;
@@ -93,14 +93,14 @@ void UAvaHideEmptyModifier::Apply()
 
 	VisibilityShared->RestoreActorsState(this, ChildrenActorsWeak.Difference(NewChildrenActorsWeak));
 	ChildrenActorsWeak = NewChildrenActorsWeak;
-	
+
 	Next();
 }
 
 void UAvaHideEmptyModifier::OnModifierDisabled(EActorModifierCoreDisableReason InReason)
 {
 	Super::OnModifierDisabled(InReason);
-	
+
 	if (UText3DComponent* Text3DComponent = TextComponent.Get())
 	{
 		Text3DComponent->OnTextGenerated().RemoveAll(this);
@@ -118,7 +118,7 @@ void UAvaHideEmptyModifier::PostEditChangeProperty(FPropertyChangedEvent& Proper
 	Super::PostEditChangeProperty(PropertyChangedEvent);
 
 	const FName MemberName = PropertyChangedEvent.GetMemberPropertyName();
-	
+
 	static const FName ContainerActorName = GET_MEMBER_NAME_CHECKED(UAvaHideEmptyModifier, ContainerActorWeak);
 	static const FName InvertVisibilityName = GET_MEMBER_NAME_CHECKED(UAvaHideEmptyModifier, bInvertVisibility);
 
@@ -133,6 +133,17 @@ void UAvaHideEmptyModifier::PostEditChangeProperty(FPropertyChangedEvent& Proper
 }
 #endif
 
+void UAvaHideEmptyModifier::SetContainerActorWeak(TWeakObjectPtr<AActor> InContainer)
+{
+	if (ContainerActorWeak == InContainer)
+	{
+		return;
+	}
+
+	ContainerActorWeak = InContainer;
+	OnContainerActorChanged();
+}
+
 void UAvaHideEmptyModifier::SetInvertVisibility(bool bInInvert)
 {
 	if (bInvertVisibility == bInInvert)
@@ -154,7 +165,7 @@ void UAvaHideEmptyModifier::OnContainerActorChanged()
 	ReferenceActor.ReferenceContainer = EAvaReferenceContainer::Other;
 	ReferenceActor.ReferenceActorWeak = ContainerActorWeak;
 	ReferenceActor.bSkipHiddenActors = false;
-	
+
 	if (const FAvaSceneTreeUpdateModifierExtension* SceneExtension = GetExtension<FAvaSceneTreeUpdateModifierExtension>())
 	{
 		SceneExtension->CheckTrackedActorUpdate(0);
@@ -169,7 +180,7 @@ void UAvaHideEmptyModifier::OnInvertVisibilityChanged()
 void UAvaHideEmptyModifier::OnSceneTreeTrackedActorChanged(int32 InIdx, AActor* InPreviousActor, AActor* InNewActor)
 {
 	Super::OnSceneTreeTrackedActorChanged(InIdx, InPreviousActor, InNewActor);
-	
+
 	if (UAvaVisibilityModifierShared* VisibilityShared = GetShared<UAvaVisibilityModifierShared>(false))
 	{
 		VisibilityShared->RestoreActorState(this, InPreviousActor);
diff --git a/Engine/Plugins/Experimental/Avalanche/Source/AvalancheModifiers/Public/Extensions/AvaRenderStateUpdateModifierExtension.h b/Engine/Plugins/Experimental/Avalanche/Source/AvalancheModifiers/Public/Extensions/AvaRenderStateUpdateModifierExtension.h
--- a/Engine/Plugins/Experimental/Avalanche/Source/AvalancheModifiers/Public/Extensions/AvaRenderStateUpdateModifierExtension.h
+++ b/Engine/Plugins/Experimental/Avalanche/Source/AvalancheModifiers/Public/Extensions/AvaRenderStateUpdateModifierExtension.h
@@ -32,7 +32,7 @@ public:
  * This extension tracks specific actors for render state updates,
  * when an update happens it will invoke IAvaRenderStateUpdateExtension function
  */
-class AVALANCHEMODIFIERS_API FAvaRenderStateUpdateModifierExtension : public FActorModifierCoreExtension
+class FAvaRenderStateUpdateModifierExtension : public FActorModifierCoreExtension
 {
 
 public:
diff --git a/Engine/Plugins/Experimental/Avalanche/Source/AvalancheModifiers/Public/Extensions/AvaSceneTreeUpdateModifierExtension.h b/Engine/Plugins/Experimental/Avalanche/Source/AvalancheModifiers/Public/Extensions/AvaSceneTreeUpdateModifierExtension.h
--- a/Engine/Plugins/Experimental/Avalanche/Source/AvalancheModifiers/Public/Extensions/AvaSceneTreeUpdateModifierExtension.h
+++ b/Engine/Plugins/Experimental/Avalanche/Source/AvalancheModifiers/Public/Extensions/AvaSceneTreeUpdateModifierExtension.h
@@ -15,16 +15,16 @@ USTRUCT(BlueprintType)
 struct FAvaSceneTreeActor
 {
 	friend class FAvaSceneTreeUpdateModifierExtension;
-	
+
 	GENERATED_BODY()
 
 	FAvaSceneTreeActor() = default;
-	
+
 	explicit FAvaSceneTreeActor(AActor* InActor)
 		: ReferenceContainer(EAvaReferenceContainer::Other)
 		, ReferenceActorWeak(InActor)
 	{}
-	
+
 	/** The method for finding a reference actor based on it's position in the parent's hierarchy */
 	UPROPERTY(EditAnywhere, BlueprintReadWrite, Category="Outliner")
 	EAvaReferenceContainer ReferenceContainer = EAvaReferenceContainer::Other;
@@ -52,7 +52,7 @@ protected:
 	/** Direct children of reference actor where order counts */
 	UPROPERTY(Transient)
 	TArray<TWeakObjectPtr<AActor>> ReferenceActorDirectChildrenWeak;
-	
+
 	/** Tracked references actors, if we skip hidden actors, we still need to track those for visibility changes, can be rebuilt */
 	UPROPERTY(Transient)
 	TArray<TWeakObjectPtr<AActor>> ReferenceActorsWeak;
@@ -76,7 +76,7 @@ class UAvaSceneTreeUpdateHandler : public UInterface
 class IAvaSceneTreeUpdateHandler
 {
 	GENERATED_BODY()
-	
+
 public:
 	virtual void OnSceneTreeTrackedActorChanged(int32 InIdx, AActor* InPreviousActor, AActor* InNewActor) = 0;
 
@@ -93,9 +93,9 @@ public:
  * This extension tracks specific actors for render state updates,
  * when an update happens it will dirty the modifier it is attached on if filter passes
  */
-class AVALANCHEMODIFIERS_API FAvaSceneTreeUpdateModifierExtension : public FActorModifierCoreExtension
+class FAvaSceneTreeUpdateModifierExtension : public FActorModifierCoreExtension
 {
-	
+
 public:
 	explicit FAvaSceneTreeUpdateModifierExtension(IAvaSceneTreeUpdateHandler* InExtensionHandler);
 
@@ -106,11 +106,11 @@ public:
 
 	void CheckTrackedActorsUpdate() const;
 	void CheckTrackedActorUpdate(int32 InIdx) const;
-	
+
 	TSet<TWeakObjectPtr<AActor>> GetChildrenActorsRecursive(const AActor* InActor) const;
 	TArray<TWeakObjectPtr<AActor>> GetDirectChildrenActor(AActor* InActor) const;
 	TArray<TWeakObjectPtr<AActor>> GetParentActors(const AActor* InActor) const;
-	
+
 protected:
 	//~ Begin FActorModifierCoreExtension
 	virtual void OnExtensionEnabled(EActorModifierCoreEnableReason InReason) override;
@@ -122,17 +122,17 @@ private:
 	virtual void OnOutlinerLoaded();
 	virtual void OnActorHierarchyChanged(AActor* InActor, const AActor* InParentActor, EAvaOutlinerHierarchyChangeType InChangeType);
 #endif
-	
+
 	void OnRenderStateDirty(UActorComponent& InComponent);
 	void OnWorldActorDestroyed(AActor* InActor);
 
 	bool IsSameActorArray(const TArray<TWeakObjectPtr<AActor>>& InPreviousActorWeak, const TArray<TWeakObjectPtr<AActor>>& InNewActorWeak) const;
-	
+
 	TArray<TWeakObjectPtr<AActor>> GetReferenceActors(const FAvaSceneTreeActor* InTrackedActor) const;
-	
+
 	TWeakInterfacePtr<IAvaSceneTreeUpdateHandler> ExtensionHandlerWeak;
 
 	TMap<int32, FAvaSceneTreeActor*> TrackedActors;
-	
+
 	FDelegateHandle WorldActorDestroyedDelegate;
 };
\ No newline at end of file
diff --git a/Engine/Plugins/Experimental/Avalanche/Source/AvalancheModifiers/Public/Extensions/AvaTransformUpdateModifierExtension.h b/Engine/Plugins/Experimental/Avalanche/Source/AvalancheModifiers/Public/Extensions/AvaTransformUpdateModifierExtension.h
--- a/Engine/Plugins/Experimental/Avalanche/Source/AvalancheModifiers/Public/Extensions/AvaTransformUpdateModifierExtension.h
+++ b/Engine/Plugins/Experimental/Avalanche/Source/AvalancheModifiers/Public/Extensions/AvaTransformUpdateModifierExtension.h
@@ -26,7 +26,7 @@ public:
  * This extension tracks specific actors for transform updates,
  * when an update happens it will invoke the IAvaTransformUpdateExtension function
  */
-class AVALANCHEMODIFIERS_API FAvaTransformUpdateModifierExtension : public FActorModifierCoreExtension
+class FAvaTransformUpdateModifierExtension : public FActorModifierCoreExtension
 {
 
 public:
diff --git a/Engine/Plugins/Experimental/Avalanche/Source/AvalancheModifiers/Public/Modifiers/AvaAlignBetweenModifier.h b/Engine/Plugins/Experimental/Avalanche/Source/AvalancheModifiers/Public/Modifiers/AvaAlignBetweenModifier.h
--- a/Engine/Plugins/Experimental/Avalanche/Source/AvalancheModifiers/Public/Modifiers/AvaAlignBetweenModifier.h
+++ b/Engine/Plugins/Experimental/Avalanche/Source/AvalancheModifiers/Public/Modifiers/AvaAlignBetweenModifier.h
@@ -10,23 +10,23 @@ class AActor;
 
 /** Represents an actor with a weight and an enabled state. */
 USTRUCT(BlueprintType)
-struct AVALANCHEMODIFIERS_API FAvaAlignBetweenWeightedActor
+struct FAvaAlignBetweenWeightedActor
 {
 	GENERATED_BODY()
 
 	FAvaAlignBetweenWeightedActor()
 	{}
-	
+
 	explicit FAvaAlignBetweenWeightedActor(AActor* InActor)
 		: ActorWeak(InActor)
 	{}
-	
+
 	explicit FAvaAlignBetweenWeightedActor(AActor* InActor, float InWeight, bool bInEnabled)
 		: ActorWeak(InActor)
 		, Weight(InWeight)
 		, bEnabled(bInEnabled)
 	{}
-	
+
 	/** Returns true if the actor is valid and the state is enabled. */
 	bool IsValid() const
 	{
@@ -42,7 +42,7 @@ struct AVALANCHEMODIFIERS_API FAvaAlignBetweenWeightedActor
 	{
 		return ActorWeak == Other.ActorWeak;
 	}
-	
+
 	/** An actor that will effect the placement location. */
 	UPROPERTY(EditAnywhere, BlueprintReadWrite, Category = "Motion Design")
 	TWeakObjectPtr<AActor> ActorWeak;
@@ -59,34 +59,30 @@ struct AVALANCHEMODIFIERS_API FAvaAlignBetweenWeightedActor
 /**
  * Moves the modifying actor to the averaged location between an array of specified actors.
  */
-UCLASS(BlueprintType)
-class AVALANCHEMODIFIERS_API UAvaAlignBetweenModifier : public UAvaBaseModifier
+UCLASS(MinimalAPI, BlueprintType)
+class UAvaAlignBetweenModifier : public UAvaBaseModifier
 	, public IAvaTransformUpdateHandler
 {
 	GENERATED_BODY()
-	
-public:
-	//~ Begin UObject
-#if WITH_EDITOR
-	virtual void PostEditChangeProperty(FPropertyChangedEvent& PropertyChangedEvent) override;
-	virtual void PostTransacted(const FTransactionObjectEvent& TransactionEvent) override;
-#endif
-	//~ End UObject
 
+public:
 	/** Gets all actors from their reference actor structs. */
 	UFUNCTION(BlueprintPure, Category = "Motion Design|Modifiers|AlignBetween")
 	TSet<AActor*> GetActors(const bool bEnabledOnly = false) const;
 
 	/** Returns all valid reference actors that enabled and have a weight greater than 0. */
 	TSet<FAvaAlignBetweenWeightedActor> GetEnabledReferenceActors() const;
-	
+
 	/** Gets all reference actors and their weights. */
 	UFUNCTION(BlueprintPure, Category = "Motion Design|Modifiers|AlignBetween")
-	TSet<FAvaAlignBetweenWeightedActor> GetReferenceActors() const { return ReferenceActors; }
+	TSet<FAvaAlignBetweenWeightedActor> GetReferenceActors() const
+	{
+		return ReferenceActors;
+	}
 
 	/** Sets all reference actors and their weights. */
 	UFUNCTION(BlueprintCallable, Category = "Motion Design|Modifiers|AlignBetween")
-	void SetReferenceActors(TSet<FAvaAlignBetweenWeightedActor> NewReferenceActors);
+	AVALANCHEMODIFIERS_API void SetReferenceActors(const TSet<FAvaAlignBetweenWeightedActor>& NewReferenceActors);
 
 	/** Adds an actor to the reference list. */
 	UFUNCTION(BlueprintCallable, Category = "Motion Design|Modifiers|AlignBetween")
@@ -101,6 +97,13 @@ public:
 	bool FindReferenceActor(AActor* InActor, FAvaAlignBetweenWeightedActor& OutReferenceActor) const;
 
 protected:
+	//~ Begin UObject
+#if WITH_EDITOR
+	virtual void PostEditChangeProperty(FPropertyChangedEvent& PropertyChangedEvent) override;
+	virtual void PostTransacted(const FTransactionObjectEvent& TransactionEvent) override;
+#endif
+	//~ End UObject
+
 	//~ Begin UActorModifierCoreBase
 	virtual void OnModifierCDOSetup(FActorModifierCoreMetadata& InMetadata) override;
 	virtual void OnModifierAdded(EActorModifierCoreEnableReason InReason) override;
@@ -109,15 +112,15 @@ protected:
 	virtual void Apply() override;
 	virtual void OnModifiedActorTransformed() override;
 	//~ End UActorModifierCoreBase
-	
-	void OnReferenceActorsChanged();
+
 
 	//~ Begin IAvaTransformUpdatedExtension
 	virtual void OnTransformUpdated(AActor* InActor, bool bInParentMoved) override;
 	//~ End IAvaTransformUpdatedExtension
-	
+
+	void OnReferenceActorsChanged();
 	void SetTransformExtensionReferenceActors();
-	
+
 	/** Editable set of reference actors and weights used to calculate the average location for this actor */
 	UPROPERTY(EditAnywhere, BlueprintReadWrite, Setter="SetReferenceActors", Getter="GetReferenceActors", Category="AlignBetween", meta=(AllowPrivateAccess="true"))
 	TSet<FAvaAlignBetweenWeightedActor> ReferenceActors;
diff --git a/Engine/Plugins/Experimental/Avalanche/Source/AvalancheModifiers/Public/Modifiers/AvaArrangeBaseModifier.h b/Engine/Plugins/Experimental/Avalanche/Source/AvalancheModifiers/Public/Modifiers/AvaArrangeBaseModifier.h
--- a/Engine/Plugins/Experimental/Avalanche/Source/AvalancheModifiers/Public/Modifiers/AvaArrangeBaseModifier.h
+++ b/Engine/Plugins/Experimental/Avalanche/Source/AvalancheModifiers/Public/Modifiers/AvaArrangeBaseModifier.h
@@ -12,23 +12,23 @@ class AActor;
 /**
  * Abstract base class for modifiers dealing with arrangement and attachment actors on self
  */
-UCLASS(Abstract)
-class AVALANCHEMODIFIERS_API UAvaArrangeBaseModifier : public UAvaAttachmentBaseModifier
+UCLASS(MinimalAPI, Abstract)
+class UAvaArrangeBaseModifier : public UAvaAttachmentBaseModifier
 	, public IAvaRenderStateUpdateHandler
 	, public IAvaTransformUpdateHandler
 {
 	GENERATED_BODY()
-	
+
 protected:
 	//~ Begin UActorModifierCoreBase
-	virtual void OnModifierAdded(EActorModifierCoreEnableReason InReason) override;
-	virtual void OnModifierDisabled(EActorModifierCoreDisableReason InReason) override;
-	virtual void OnModifiedActorTransformed() override;
+	AVALANCHEMODIFIERS_API virtual void OnModifierAdded(EActorModifierCoreEnableReason InReason) override;
+	AVALANCHEMODIFIERS_API virtual void OnModifierDisabled(EActorModifierCoreDisableReason InReason) override;
+	AVALANCHEMODIFIERS_API virtual void OnModifiedActorTransformed() override;
 	//~ End UActorModifierCoreBase
-	
+
 	//~ Begin IAvaSceneTreeUpdateModifierExtension
-	virtual void OnSceneTreeTrackedActorChildrenChanged(int32 InIdx, const TSet<TWeakObjectPtr<AActor>>& InPreviousChildrenActors, const TSet<TWeakObjectPtr<AActor>>& InNewChildrenActors) override;
-	virtual void OnSceneTreeTrackedActorDirectChildrenChanged(int32 InIdx, const TArray<TWeakObjectPtr<AActor>>& InPreviousChildrenActors, const TArray<TWeakObjectPtr<AActor>>& InNewChildrenActors) override;
+	AVALANCHEMODIFIERS_API virtual void OnSceneTreeTrackedActorChildrenChanged(int32 InIdx, const TSet<TWeakObjectPtr<AActor>>& InPreviousChildrenActors, const TSet<TWeakObjectPtr<AActor>>& InNewChildrenActors) override;
+	AVALANCHEMODIFIERS_API virtual void OnSceneTreeTrackedActorDirectChildrenChanged(int32 InIdx, const TArray<TWeakObjectPtr<AActor>>& InPreviousChildrenActors, const TArray<TWeakObjectPtr<AActor>>& InNewChildrenActors) override;
 	//~ End IAvaSceneTreeUpdateModifierExtension
 
 	//~ Begin IAvaRenderStateUpdateExtension
diff --git a/Engine/Plugins/Experimental/Avalanche/Source/AvalancheModifiers/Public/Modifiers/AvaAttachmentBaseModifier.h b/Engine/Plugins/Experimental/Avalanche/Source/AvalancheModifiers/Public/Modifiers/AvaAttachmentBaseModifier.h
--- a/Engine/Plugins/Experimental/Avalanche/Source/AvalancheModifiers/Public/Modifiers/AvaAttachmentBaseModifier.h
+++ b/Engine/Plugins/Experimental/Avalanche/Source/AvalancheModifiers/Public/Modifiers/AvaAttachmentBaseModifier.h
@@ -11,18 +11,18 @@ class AActor;
 /**
  * Abstract base class for all modifiers that deal with attachments
  */
-UCLASS(Abstract)
-class AVALANCHEMODIFIERS_API UAvaAttachmentBaseModifier : public UAvaBaseModifier
+UCLASS(MinimalAPI, Abstract)
+class UAvaAttachmentBaseModifier : public UAvaBaseModifier
 	, public IAvaSceneTreeUpdateHandler
 {
 	GENERATED_BODY()
 
 protected:
 	//~ Begin UActorModifierCoreBase
-	virtual void OnModifierCDOSetup(FActorModifierCoreMetadata& InMetadata) override;
+	AVALANCHEMODIFIERS_API virtual void OnModifierCDOSetup(FActorModifierCoreMetadata& InMetadata) override;
 	virtual void OnModifierAdded(EActorModifierCoreEnableReason InReason) override;
 	//~ End UActorModifierCoreBase
-	
+
 	//~ Begin IAvaSceneTreeUpdateModifierExtension
 	virtual void OnSceneTreeTrackedActorChanged(int32 InIdx, AActor* InPreviousActor, AActor* InNewActor) override {}
 	virtual void OnSceneTreeTrackedActorChildrenChanged(int32 InIdx, const TSet<TWeakObjectPtr<AActor>>& InPreviousChildrenActors, const TSet<TWeakObjectPtr<AActor>>& InNewChildrenActors) override {}
diff --git a/Engine/Plugins/Experimental/Avalanche/Source/AvalancheModifiers/Public/Modifiers/AvaAutoFollowModifier.h b/Engine/Plugins/Experimental/Avalanche/Source/AvalancheModifiers/Public/Modifiers/AvaAutoFollowModifier.h
--- a/Engine/Plugins/Experimental/Avalanche/Source/AvalancheModifiers/Public/Modifiers/AvaAutoFollowModifier.h
+++ b/Engine/Plugins/Experimental/Avalanche/Source/AvalancheModifiers/Public/Modifiers/AvaAutoFollowModifier.h
@@ -16,24 +16,16 @@ class UActorComponent;
 /**
  * Moves the modifying actor along with a specified actor relative to the specified actor's bounds.
  */
-UCLASS(BlueprintType)
-class AVALANCHEMODIFIERS_API UAvaAutoFollowModifier : public UAvaAttachmentBaseModifier
+UCLASS(MinimalAPI, BlueprintType)
+class UAvaAutoFollowModifier : public UAvaAttachmentBaseModifier
 	, public IAvaTransformUpdateHandler
 	, public IAvaRenderStateUpdateHandler
 {
 	GENERATED_BODY()
-	
-public:
-	//~ Begin UObject
-	virtual void PostLoad() override;
-#if WITH_EDITOR
-	virtual void PostEditChangeProperty(FPropertyChangedEvent& PropertyChangedEvent) override;
-	virtual void PostTransacted(const FTransactionObjectEvent& TransactionEvent) override;
-#endif
-	//~ End UObject
 
+public:
 	UFUNCTION(BlueprintCallable, Category = "Motion Design|Modifiers|AutoFollow")
-	void SetReferenceActor(const FAvaSceneTreeActor& InReferenceActor);
+	AVALANCHEMODIFIERS_API void SetReferenceActor(const FAvaSceneTreeActor& InReferenceActor);
 
 	UFUNCTION(BlueprintPure, Category = "Motion Design|Modifiers|AutoFollow")
 	const FAvaSceneTreeActor& GetReferenceActor() const
@@ -41,14 +33,18 @@ public:
 		return ReferenceActor;
 	}
 
+	UFUNCTION(BlueprintCallable, Category = "Motion Design|Modifiers|AutoFollow")
+	AVALANCHEMODIFIERS_API void SetFollowedAxis(int32 InFollowedAxis);
+
 	UFUNCTION(BlueprintPure, Category = "Motion Design|Modifiers|AutoFollow")
 	int32 GetFollowedAxis() const
 	{
 		return FollowedAxis;
 	}
-	
+
+	/** Sets the distance from this actor to the followed actor. */
 	UFUNCTION(BlueprintCallable, Category = "Motion Design|Modifiers|AutoFollow")
-	void SetFollowedAxis(int32 InFollowedAxis);
+	AVALANCHEMODIFIERS_API void SetDefaultDistance(const FVector& NewDefaultDistance);
 
 	/** Gets the distance from this actor to the followed actor. */
 	UFUNCTION(BlueprintPure, Category = "Motion Design|Modifiers|AutoFollow")
@@ -57,9 +53,9 @@ public:
 		return DefaultDistance;
 	}
 
-	/** Sets the distance from this actor to the followed actor. */
+	/** Sets the maximum distance from this actor to the followed actor. */
 	UFUNCTION(BlueprintCallable, Category = "Motion Design|Modifiers|AutoFollow")
-	void SetDefaultDistance(const FVector& NewDefaultDistance);
+	AVALANCHEMODIFIERS_API void SetMaxDistance(const FVector& NewMaxDistance);
 
 	/** Gets the maximum distance from this actor to the followed actor. */
 	UFUNCTION(BlueprintPure, Category = "Motion Design|Modifiers|AutoFollow")
@@ -68,9 +64,9 @@ public:
 		return MaxDistance;
 	}
 
-	/** Sets the maximum distance from this actor to the followed actor. */
+	/** Sets the percent % progress from the maximum distance to the default distance. */
 	UFUNCTION(BlueprintCallable, Category = "Motion Design|Modifiers|AutoFollow")
-	void SetMaxDistance(const FVector& NewMaxDistance);
+	AVALANCHEMODIFIERS_API void SetProgress(const FVector& NewProgress);
 
 	/** Gets the percent % progress from the maximum distance to the default distance. */
 	UFUNCTION(BlueprintPure, Category = "Motion Design|Modifiers|AutoFollow")
@@ -79,9 +75,9 @@ public:
 		return Progress;
 	}
 
-	/** Sets the percent % progress from the maximum distance to the default distance. */
+	/** Sets the alignment for the followed actor's center. */
 	UFUNCTION(BlueprintCallable, Category = "Motion Design|Modifiers|AutoFollow")
-	void SetProgress(const FVector& NewProgress);
+	AVALANCHEMODIFIERS_API void SetFollowedAlignment(const FAvaAnchorAlignment& NewFollowedAlignment);
 
 	/** Gets the alignment for the followed actor's center. */
 	UFUNCTION(BlueprintPure, Category = "Motion Design|Modifiers|AutoFollow")
@@ -90,9 +86,9 @@ public:
 		return FollowedAlignment;
 	}
 
-	/** Sets the alignment for the followed actor's center. */
+	/** Sets the alignment for this actor's center. */
 	UFUNCTION(BlueprintCallable, Category = "Motion Design|Modifiers|AutoFollow")
-	void SetFollowedAlignment(const FAvaAnchorAlignment& NewFollowedAlignment);
+	AVALANCHEMODIFIERS_API void SetLocalAlignment(const FAvaAnchorAlignment& NewLocalAlignment);
 
 	/** Gets the alignment for this actor's center. */
 	UFUNCTION(BlueprintPure, Category = "Motion Design|Modifiers|AutoFollow")
@@ -101,9 +97,9 @@ public:
 		return LocalAlignment;
 	}
 
-	/** Sets the alignment for this actor's center. */
+	/** Sets the axis direction to offset this actor from the followed actor's bounds. */
 	UFUNCTION(BlueprintCallable, Category = "Motion Design|Modifiers|AutoFollow")
-	void SetLocalAlignment(const FAvaAnchorAlignment& NewLocalAlignment);
+	AVALANCHEMODIFIERS_API void SetOffsetAxis(const FVector& NewOffsetAxis);
 
 	/** Gets the axis direction to offset this actor from the followed actor's bounds. */
 	UFUNCTION(BlueprintPure, Category = "Motion Design|Modifiers|AutoFollow")
@@ -112,11 +108,15 @@ public:
 		return OffsetAxis;
 	}
 
-	/** Sets the axis direction to offset this actor from the followed actor's bounds. */
-	UFUNCTION(BlueprintCallable, Category = "Motion Design|Modifiers|AutoFollow")
-	void SetOffsetAxis(const FVector& NewOffsetAxis);
-	
 protected:
+	//~ Begin UObject
+	virtual void PostLoad() override;
+#if WITH_EDITOR
+	virtual void PostEditChangeProperty(FPropertyChangedEvent& PropertyChangedEvent) override;
+	virtual void PostTransacted(const FTransactionObjectEvent& TransactionEvent) override;
+#endif
+	//~ End UObject
+
 	//~ Begin UActorModifierCoreBase
 	virtual bool IsModifierDirtyable() const override;
 	virtual void OnModifierCDOSetup(FActorModifierCoreMetadata& InMetadata) override;
@@ -126,7 +126,7 @@ protected:
 	virtual void OnModifiedActorTransformed() override;
 	virtual void Apply() override;
 	//~ End UActorModifierCoreBase
-	
+
 	//~ Begin IAvaTransformUpdatedExtension
 	virtual void OnTransformUpdated(AActor* InActor, bool bInParentMoved) override;
 	//~ End IAvaTransformUpdatedExtension
@@ -155,7 +155,7 @@ protected:
 	/** The actor being followed by the modifier. This is user selectable if the Reference Container is set to "Other". */
 	UPROPERTY(meta=(DeprecatedProperty, DeprecationMessage="Use ReferenceActor instead"))
 	TWeakObjectPtr<AActor> ReferenceActorWeak_DEPRECATED = nullptr;
-	
+
 	/** If true, will search for the next visible actor based on the selected reference container. */
 	UPROPERTY(meta=(DeprecatedProperty, DeprecationMessage="Use ReferenceActor instead"))
 	bool bIgnoreHiddenActors_DEPRECATED = false;
@@ -166,15 +166,15 @@ protected:
 		EAvaModifiersAxis::Y |
 		EAvaModifiersAxis::Z
 	);
-	
+
 	/** Based on followed axis, the direction to offset this actor from the followed actor's bounds. */
 	UPROPERTY(EditAnywhere, BlueprintReadWrite, Setter="SetOffsetAxis", Getter="GetOffsetAxis", Interp, Category="AutoFollow", meta=(UIMin="-1.0", UIMax="1.0", AllowPrivateAccess="true"))
 	FVector OffsetAxis = FVector(0, 1, 0);
-	
+
 	/** The alignment for the followed actor's center. */
 	UPROPERTY(EditAnywhere, BlueprintReadWrite, Setter="SetFollowedAlignment", Getter="GetFollowedAlignment", Category="AutoFollow", meta=(AllowPrivateAccess="true"))
 	FAvaAnchorAlignment FollowedAlignment;
-	
+
 	/** The alignment for this actor's center. */
 	UPROPERTY(EditAnywhere, BlueprintReadWrite, Setter="SetLocalAlignment", Getter="GetLocalAlignment", Category="AutoFollow", meta=(AllowPrivateAccess="true"))
 	FAvaAnchorAlignment LocalAlignment;
@@ -197,7 +197,7 @@ private:
 
 	UPROPERTY(Transient, DuplicateTransient, TextExportTransient)
 	FBox CachedReferenceBounds = FBox(EForceInit::ForceInit);
-	
+
 	UPROPERTY(Transient, DuplicateTransient, TextExportTransient)
 	FBox CachedModifiedBounds = FBox(EForceInit::ForceInit);
 
diff --git a/Engine/Plugins/Experimental/Avalanche/Source/AvalancheModifiers/Public/Modifiers/AvaAutoSizeModifier.h b/Engine/Plugins/Experimental/Avalanche/Source/AvalancheModifiers/Public/Modifiers/AvaAutoSizeModifier.h
--- a/Engine/Plugins/Experimental/Avalanche/Source/AvalancheModifiers/Public/Modifiers/AvaAutoSizeModifier.h
+++ b/Engine/Plugins/Experimental/Avalanche/Source/AvalancheModifiers/Public/Modifiers/AvaAutoSizeModifier.h
@@ -25,8 +25,8 @@ enum class EAvaAutoSizeFitMode : uint8
 /**
  * Adapts the modified actor geometry size/scale and position so that it acts as a background for a specified actor
  */
-UCLASS(BlueprintType)
-class AVALANCHEMODIFIERS_API UAvaAutoSizeModifier : public UAvaGeometryBaseModifier
+UCLASS(MinimalAPI, BlueprintType)
+class UAvaAutoSizeModifier : public UAvaGeometryBaseModifier
 	, public IAvaTransformUpdateHandler
 	, public IAvaRenderStateUpdateHandler
 	, public IAvaSceneTreeUpdateHandler
@@ -34,15 +34,8 @@ class AVALANCHEMODIFIERS_API UAvaAutoSizeModifier : public UAvaGeometryBaseModif
 	GENERATED_BODY()
 
 public:
-	//~ Begin UObject
-	virtual void PostLoad() override;
-#if WITH_EDITOR
-	virtual void PostEditChangeProperty(FPropertyChangedEvent& PropertyChangedEvent) override;
-#endif
-	//~ End UObject
-
 	UFUNCTION(BlueprintCallable, Category = "Motion Design|Modifiers|AutoSize")
-	void SetReferenceActor(const FAvaSceneTreeActor& InReferenceActor);
+	AVALANCHEMODIFIERS_API void SetReferenceActor(const FAvaSceneTreeActor& InReferenceActor);
 
 	UFUNCTION(BlueprintPure, Category = "Motion Design|Modifiers|AutoSize")
 	const FAvaSceneTreeActor& GetReferenceActor() const
@@ -50,14 +43,18 @@ public:
 		return ReferenceActor;
 	}
 
+	UFUNCTION(BlueprintCallable, Category = "Motion Design|Modifiers|AutoSize")
+	AVALANCHEMODIFIERS_API void SetFollowedAxis(int32 InFollowedAxis);
+
 	UFUNCTION(BlueprintPure, Category = "Motion Design|Modifiers|AutoSize")
 	int32 GetFollowedAxis() const
 	{
 		return FollowedAxis;
 	}
 
+	/** Sets the actor affecting the modifier. This is user selectable if the Reference Container is set to "Other". */
 	UFUNCTION(BlueprintCallable, Category = "Motion Design|Modifiers|AutoSize")
-	void SetFollowedAxis(int32 InFollowedAxis);
+	AVALANCHEMODIFIERS_API void SetPadding(const FMargin& InPadding);
 
 	/** Gets the actor affecting the modifier. This is user selectable if the Reference Container is set to "Other". */
 	UFUNCTION(BlueprintPure, Category = "Motion Design|Modifiers|AutoSize")
@@ -66,9 +63,8 @@ public:
 		return Padding;
 	}
 
-	/** Sets the actor affecting the modifier. This is user selectable if the Reference Container is set to "Other". */
 	UFUNCTION(BlueprintCallable, Category = "Motion Design|Modifiers|AutoSize")
-	void SetPadding(const FMargin& InPadding);
+	AVALANCHEMODIFIERS_API void SetFitMode(const EAvaAutoSizeFitMode InFitMode);
 
 	UFUNCTION(BlueprintPure, Category = "Motion Design|Modifiers|AutoSize")
 	EAvaAutoSizeFitMode GetFitMode() const
@@ -77,10 +73,7 @@ public:
 	}
 
 	UFUNCTION(BlueprintCallable, Category = "Motion Design|Modifiers|AutoSize")
-	void SetFitMode(const EAvaAutoSizeFitMode InFitMode);
-
-	UFUNCTION(BlueprintCallable, Category = "Motion Design|Modifiers|AutoSize")
-	void SetIncludeChildren(bool bInIncludeChildren);
+	AVALANCHEMODIFIERS_API void SetIncludeChildren(bool bInIncludeChildren);
 
 	UFUNCTION(BlueprintPure, Category = "Motion Design|Modifiers|AutoSize")
 	bool GetIncludeChildren() const
@@ -89,6 +82,13 @@ public:
 	}
 
 protected:
+	//~ Begin UObject
+	virtual void PostLoad() override;
+#if WITH_EDITOR
+	virtual void PostEditChangeProperty(FPropertyChangedEvent& PropertyChangedEvent) override;
+#endif
+	//~ End UObject
+
 	//~ Begin UActorModifierCoreBase
 	virtual bool IsModifierDirtyable() const override;
 	virtual void OnModifierCDOSetup(FActorModifierCoreMetadata& InMetadata) override;
diff --git a/Engine/Plugins/Experimental/Avalanche/Source/AvalancheModifiers/Public/Modifiers/AvaBaseModifier.h b/Engine/Plugins/Experimental/Avalanche/Source/AvalancheModifiers/Public/Modifiers/AvaBaseModifier.h
--- a/Engine/Plugins/Experimental/Avalanche/Source/AvalancheModifiers/Public/Modifiers/AvaBaseModifier.h
+++ b/Engine/Plugins/Experimental/Avalanche/Source/AvalancheModifiers/Public/Modifiers/AvaBaseModifier.h
@@ -6,8 +6,8 @@
 #include "AvaBaseModifier.generated.h"
 
 /** This is the base class for all Motion Design modifiers */
-UCLASS(Abstract)
-class AVALANCHEMODIFIERS_API UAvaBaseModifier : public UActorModifierCoreBase
+UCLASS(MinimalAPI, Abstract)
+class UAvaBaseModifier : public UActorModifierCoreBase
 {
 	GENERATED_BODY()
 };
diff --git a/Engine/Plugins/Experimental/Avalanche/Source/AvalancheModifiers/Public/Modifiers/AvaBendModifier.h b/Engine/Plugins/Experimental/Avalanche/Source/AvalancheModifiers/Public/Modifiers/AvaBendModifier.h
--- a/Engine/Plugins/Experimental/Avalanche/Source/AvalancheModifiers/Public/Modifiers/AvaBendModifier.h
+++ b/Engine/Plugins/Experimental/Avalanche/Source/AvalancheModifiers/Public/Modifiers/AvaBendModifier.h
@@ -5,69 +5,81 @@
 #include "AvaGeometryBaseModifier.h"
 #include "AvaBendModifier.generated.h"
 
-UCLASS(BlueprintType)
-class AVALANCHEMODIFIERS_API UAvaBendModifier : public UAvaGeometryBaseModifier
+UCLASS(MinimalAPI, BlueprintType)
+class UAvaBendModifier : public UAvaGeometryBaseModifier
 {
 	GENERATED_BODY()
-	
-public:
-	//~ Begin UObject
-#if WITH_EDITOR
-	virtual void PostEditChangeProperty(FPropertyChangedEvent& PropertyChangedEvent) override;
-#endif
-	//~End UObject
 
-	UFUNCTION()
-	void SetAngle(float InAngle);
-	float GetAngle() const { return Angle; }
+public:
+	AVALANCHEMODIFIERS_API void SetAngle(float InAngle);
+	float GetAngle() const
+	{
+		return Angle;
+	}
 
-	UFUNCTION()
-	void SetExtent(float InExtent);
-	float GetExtent() const { return Extent; }
+	AVALANCHEMODIFIERS_API void SetExtent(float InExtent);
+	float GetExtent() const
+	{
+		return Extent;
+	}
 
-	UFUNCTION()
-	void SetBendPosition(const FVector& InBendPosition);
-	const FVector& GetBendPosition() const { return BendPosition; }
+	AVALANCHEMODIFIERS_API void SetBendPosition(const FVector& InBendPosition);
+	const FVector& GetBendPosition() const
+	{
+		return BendPosition;
+	}
 
-	UFUNCTION()
-	void SetBendRotation(const FRotator& InBendRotation);
-	const FRotator& GetBendRotation() const { return BendRotation; }
+	AVALANCHEMODIFIERS_API void SetBendRotation(const FRotator& InBendRotation);
+	const FRotator& GetBendRotation() const
+	{
+		return BendRotation;
+	}
 
-	UFUNCTION()
-	void SetSymmetricExtents(bool bInSymmetricExtents);
-	bool GetSymmetricExtents() const { return bSymmetricExtents; }
+	AVALANCHEMODIFIERS_API void SetSymmetricExtents(bool bInSymmetricExtents);
+	bool GetSymmetricExtents() const
+	{
+		return bSymmetricExtents;
+	}
 
-	UFUNCTION()
-	void SetBidirectional(bool bInBidirectional);
-	bool GetBidirectional() const { return bBidirectional; }
+	AVALANCHEMODIFIERS_API void SetBidirectional(bool bInBidirectional);
+	bool GetBidirectional() const
+	{
+		return bBidirectional;
+	}
 
 protected:
+	//~ Begin UObject
+#if WITH_EDITOR
+	virtual void PostEditChangeProperty(FPropertyChangedEvent& PropertyChangedEvent) override;
+#endif
+	//~End UObject
+
 	//~ Begin UActorModifierCoreBase
 	virtual void OnModifierCDOSetup(FActorModifierCoreMetadata& InMetadata) override;
 	virtual void Apply() override;
 	//~ End UActorModifierCoreBase
-	
+
 	void OnBendTransformChanged();
 	void OnBendOptionChanged();
-	
+
 	UPROPERTY(EditInstanceOnly, BlueprintReadWrite, Setter="SetBendPosition", Getter="GetBendPosition", Category="Bend", meta=(AllowPrivateAccess="true"))
 	FVector BendPosition = FVector::ZeroVector;
-	
+
 	UPROPERTY(EditInstanceOnly, BlueprintReadWrite, Setter="SetBendRotation", Getter="GetBendRotation", Category="Bend", meta=(AllowPrivateAccess="true"))
 	FRotator BendRotation = FRotator::ZeroRotator;
-	
+
 	UPROPERTY(EditInstanceOnly, BlueprintReadWrite, Setter="SetAngle", Getter="GetAngle", Category="Bend", meta=(ClampMin="-360.0", ClampMax="360.0", AllowPrivateAccess="true"))
 	float Angle = 25;
-	
+
 	UPROPERTY(EditInstanceOnly, BlueprintReadWrite, Setter="SetExtent", Getter="GetExtent", Category="Bend", meta=(ClampMin="0", ClampMax="1.0", AllowPrivateAccess="true"))
 	float Extent = 1.0;
-	
+
 	UPROPERTY(EditInstanceOnly, BlueprintReadWrite, Setter="SetSymmetricExtents", Getter="GetSymmetricExtents", Category="Bend", meta=(AllowPrivateAccess="true"))
 	bool bSymmetricExtents = true;
-	
+
 	UPROPERTY(EditInstanceOnly, BlueprintReadWrite, Setter="SetBidirectional", Getter="GetBidirectional", Category="Bend", meta=(AllowPrivateAccess="true"))
 	bool bBidirectional = false;
-	
+
 	UPROPERTY()
 	float ModifiedMeshMaxExtent = 50;
 };
diff --git a/Engine/Plugins/Experimental/Avalanche/Source/AvalancheModifiers/Public/Modifiers/AvaBevelModifier.h b/Engine/Plugins/Experimental/Avalanche/Source/AvalancheModifiers/Public/Modifiers/AvaBevelModifier.h
--- a/Engine/Plugins/Experimental/Avalanche/Source/AvalancheModifiers/Public/Modifiers/AvaBevelModifier.h
+++ b/Engine/Plugins/Experimental/Avalanche/Source/AvalancheModifiers/Public/Modifiers/AvaBevelModifier.h
@@ -5,8 +5,8 @@
 #include "AvaGeometryBaseModifier.h"
 #include "AvaBevelModifier.generated.h"
 
-UCLASS(BlueprintType)
-class AVALANCHEMODIFIERS_API UAvaBevelModifier : public UAvaGeometryBaseModifier
+UCLASS(MinimalAPI, BlueprintType)
+class UAvaBevelModifier : public UAvaGeometryBaseModifier
 {
 	GENERATED_BODY()
 
@@ -16,31 +16,35 @@ public:
 	static constexpr int32 MinIterations = 1;
 	static constexpr int32 MaxIterations = 3;
 
+	AVALANCHEMODIFIERS_API void SetInset(float InBevel);
+	float GetInset() const
+	{
+		return Inset;
+	}
+
+	AVALANCHEMODIFIERS_API void SetIterations(int32 InIterations);
+	int32 GetIterations() const
+	{
+		return Iterations;
+	}
+
+protected:
 	//~ Begin UObject
 #if WITH_EDITOR
 	virtual void PostEditChangeProperty(FPropertyChangedEvent& PropertyChangedEvent) override;
 #endif
 	//~ End UObject
-	
-	UFUNCTION()
-	void SetInset(float InBevel);
-	float GetInset() const { return Inset; }
-
-	UFUNCTION()
-	void SetIterations(int32 InIterations);
-	int32 GetIterations() const { return Iterations; }
 
-protected:
 	//~ Begin UActorModifierCoreBase
 	virtual void OnModifierCDOSetup(FActorModifierCoreMetadata& InMetadata) override;
 	virtual void Apply() override;
 	//~ End UActorModifierCoreBase
-	
+
 	void OnInsetChanged();
 	void OnIterationsChanged();
 
 	float GetMaxBevel() const;
-	
+
 	UPROPERTY(EditInstanceOnly, BlueprintReadWrite, Setter="SetInset", Getter="GetInset", Category="Bevel", meta=(ClampMin="0.0", AllowPrivateAccess="true"))
 	float Inset = 1.0f;
 
diff --git a/Engine/Plugins/Experimental/Avalanche/Source/AvalancheModifiers/Public/Modifiers/AvaBooleanModifier.h b/Engine/Plugins/Experimental/Avalanche/Source/AvalancheModifiers/Public/Modifiers/AvaBooleanModifier.h
--- a/Engine/Plugins/Experimental/Avalanche/Source/AvalancheModifiers/Public/Modifiers/AvaBooleanModifier.h
+++ b/Engine/Plugins/Experimental/Avalanche/Source/AvalancheModifiers/Public/Modifiers/AvaBooleanModifier.h
@@ -22,8 +22,8 @@ enum class EAvaBooleanMode : uint8
 };
 
 /** This modifier allows you to apply a mask on a certain shape, this will affect every shape it collides with that matches options */
-UCLASS(BlueprintType)
-class AVALANCHEMODIFIERS_API UAvaBooleanModifier : public UAvaGeometryBaseModifier
+UCLASS(MinimalAPI, BlueprintType)
+class UAvaBooleanModifier : public UAvaGeometryBaseModifier
 	, public IAvaTransformUpdateHandler
 {
 	GENERATED_BODY()
@@ -34,19 +34,17 @@ public:
 	/** This is the min depth threshold needed on the mask for it to work properly */
 	inline static constexpr float MinDepth = UE_KINDA_SMALL_NUMBER * 2;
 
-	//~ Begin UObject
-#if WITH_EDITOR
-	virtual void PostEditChangeProperty(FPropertyChangedEvent& InPropertyChangedEvent) override;
-#endif
-	//~ End UObject
-
-	UFUNCTION()
-	void SetMode(EAvaBooleanMode InMode);
-	EAvaBooleanMode GetMode() const { return Mode; }
+	AVALANCHEMODIFIERS_API void SetMode(EAvaBooleanMode InMode);
+	EAvaBooleanMode GetMode() const
+	{
+		return Mode;
+	}
 
-	UFUNCTION()
-	void SetChannel(uint8 InChannel);
-	uint8 GetChannel() const { return Channel; }
+	AVALANCHEMODIFIERS_API void SetChannel(uint8 InChannel);
+	uint8 GetChannel() const
+	{
+		return Channel;
+	}
 
 	const FAvaBooleanModifierSharedChannelInfo& GetChannelInfo() const
 	{
@@ -54,6 +52,12 @@ public:
 	}
 
 protected:
+	//~ Begin UObject
+#if WITH_EDITOR
+	virtual void PostEditChangeProperty(FPropertyChangedEvent& InPropertyChangedEvent) override;
+#endif
+	//~ End UObject
+
 	//~ Begin UActorModifierCoreBase
 	virtual void OnModifierCDOSetup(FActorModifierCoreMetadata& InMetadata) override;
 	virtual void OnModifierAdded(EActorModifierCoreEnableReason InReason) override;
diff --git a/Engine/Plugins/Experimental/Avalanche/Source/AvalancheModifiers/Public/Modifiers/AvaDynamicMeshConverterModifier.h b/Engine/Plugins/Experimental/Avalanche/Source/AvalancheModifiers/Public/Modifiers/AvaDynamicMeshConverterModifier.h
--- a/Engine/Plugins/Experimental/Avalanche/Source/AvalancheModifiers/Public/Modifiers/AvaDynamicMeshConverterModifier.h
+++ b/Engine/Plugins/Experimental/Avalanche/Source/AvalancheModifiers/Public/Modifiers/AvaDynamicMeshConverterModifier.h
@@ -42,7 +42,7 @@ struct FAvaDynamicMeshConverterModifierComponentState
 
 	explicit FAvaDynamicMeshConverterModifierComponentState() {}
 	explicit FAvaDynamicMeshConverterModifierComponentState(UPrimitiveComponent* InPrimitiveComponent);
-	
+
 	/** The component we are converting to dynamic mesh */
 	UPROPERTY()
 	TWeakObjectPtr<UPrimitiveComponent> Component = nullptr;
@@ -66,12 +66,49 @@ struct FAvaDynamicMeshConverterModifierComponentState
 	UE::Geometry::FDynamicMesh3 Mesh;
 };
 
-UCLASS(BlueprintType)
-class AVALANCHEMODIFIERS_API UAvaDynamicMeshConverterModifier : public UAvaGeometryBaseModifier
+UCLASS(MinimalAPI, BlueprintType)
+class UAvaDynamicMeshConverterModifier : public UAvaGeometryBaseModifier
 {
 	GENERATED_BODY()
 
 public:
+	AVALANCHEMODIFIERS_API void SetSourceActorWeak(const TWeakObjectPtr<AActor>& InActor);
+	TWeakObjectPtr<AActor> GetSourceActorWeak() const
+	{
+		return SourceActorWeak;
+	}
+
+	AVALANCHEMODIFIERS_API void SetComponentType(int32 InComponentType);
+	int32 GetComponentType() const
+	{
+		return ComponentType;
+	}
+
+	AVALANCHEMODIFIERS_API void SetFilterActorMode(EAvaDynamicMeshConverterModifierFilter InFilter);
+	EAvaDynamicMeshConverterModifierFilter GetFilterActorMode() const
+	{
+		return FilterActorMode;
+	}
+
+	AVALANCHEMODIFIERS_API void SetFilterActorClasses(const TSet<TSubclassOf<AActor>>& InClasses);
+	const TSet<TSubclassOf<AActor>>& GetFilterActorClasses() const
+	{
+		return FilterActorClasses;
+	}
+
+	AVALANCHEMODIFIERS_API void SetIncludeAttachedActors(bool bInInclude);
+	bool GetIncludeAttachedActors() const
+	{
+		return bIncludeAttachedActors;
+	}
+
+	AVALANCHEMODIFIERS_API void SetHideConvertedMesh(bool bInHide);
+	bool GetHideConvertedMesh() const
+	{
+		return bHideConvertedMesh;
+	}
+
+protected:
 	//~ Begin UObject
 #if WITH_EDITOR
 	virtual void PostEditChangeProperty(FPropertyChangedEvent& PropertyChangedEvent) override;
@@ -81,25 +118,6 @@ public:
 #endif
 	//~ End UObject
 
-	void SetSourceActorWeak(const TWeakObjectPtr<AActor>& InActor);
-	TWeakObjectPtr<AActor> GetSourceActorWeak() const { return SourceActorWeak; }
-
-	void SetComponentType(int32 InComponentType);
-	int32 GetComponentType() const { return ComponentType; }
-
-	void SetFilterActorMode(EAvaDynamicMeshConverterModifierFilter InFilter);
-	EAvaDynamicMeshConverterModifierFilter GetFilterActorMode() const { return FilterActorMode; }
-
-	void SetFilterActorClasses(const TSet<TSubclassOf<AActor>>& InClasses);
-	const TSet<TSubclassOf<AActor>>& GetFilterActorClasses() const { return FilterActorClasses; }
-
-	void SetIncludeAttachedActors(bool bInInclude);
-	bool GetIncludeAttachedActors() const { return bIncludeAttachedActors; }
-
-	void SetHideConvertedMesh(bool bInHide);
-	bool GetHideConvertedMesh() const { return bHideConvertedMesh; }
-	
-protected:
 	//~ Begin UActorModifierCoreBase
 	virtual void OnModifierCDOSetup(FActorModifierCoreMetadata& InMetadata) override;
 	virtual void OnModifierAdded(EActorModifierCoreEnableReason InReason) override;
@@ -108,17 +126,17 @@ protected:
 	virtual void Apply() override;
 	virtual void OnModifierRemoved(EActorModifierCoreDisableReason InReason) override;
 	//~ End UActorModifierCoreBase
-	
+
 	void OnSourceActorChanged();
 
 	void ConvertComponents(TArray<FAvaDynamicMeshConverterModifierComponentState>& OutResults) const;
 	bool HasFlag(EAvaDynamicMeshConverterModifierType InFlag) const;
-	
+
 	void AddDynamicMeshComponent();
 	void RemoveDynamicMeshComponent();
 
 	void GetFilteredActors(TArray<AActor*>& OutActors) const;
-	
+
 	void GetStaticMeshComponents(const TArray<AActor*>& InActors, TArray<UStaticMeshComponent*>& OutComponents) const;
 	void GetDynamicMeshComponents(const TArray<AActor*>& InActors, TArray<UDynamicMeshComponent*>& OutComponents) const;
 	void GetSkeletalMeshComponents(const TArray<AActor*>& InActors, TArray<USkeletalMeshComponent*>& OutComponents) const;
@@ -132,7 +150,7 @@ protected:
 	/** What actor should we copy from, by default is self */
 	UPROPERTY(EditInstanceOnly, BlueprintReadWrite, Setter="SetSourceActorWeak", Getter="GetSourceActorWeak", Category="DynamicMeshConverter", meta=(DisplayName="SourceActor", AllowPrivateAccess="true"))
 	TWeakObjectPtr<AActor> SourceActorWeak = GetModifiedActor();
-	
+
 	/** Which components should we take into account for the conversion */
 	UPROPERTY(EditInstanceOnly, BlueprintReadWrite, Setter="SetComponentType", Getter="GetComponentType", Category="DynamicMeshConverter", meta=(Bitmask, BitmaskEnum="/Script/AvalancheModifiers.EAvaDynamicMeshConverterModifierType", AllowPrivateAccess="true"))
 	int32 ComponentType = static_cast<int32>(
@@ -165,7 +183,7 @@ protected:
 	/** Components converted to dynamic mesh */
 	UPROPERTY()
 	TArray<FAvaDynamicMeshConverterModifierComponentState> ConvertedComponents;
-	
+
 	/** Cached mesh to set state without reconverting again */
 	TOptional<UE::Geometry::FDynamicMesh3> ConvertedMesh;
 
diff --git a/Engine/Plugins/Experimental/Avalanche/Source/AvalancheModifiers/Public/Modifiers/AvaExtrudeModifier.h b/Engine/Plugins/Experimental/Avalanche/Source/AvalancheModifiers/Public/Modifiers/AvaExtrudeModifier.h
--- a/Engine/Plugins/Experimental/Avalanche/Source/AvalancheModifiers/Public/Modifiers/AvaExtrudeModifier.h
+++ b/Engine/Plugins/Experimental/Avalanche/Source/AvalancheModifiers/Public/Modifiers/AvaExtrudeModifier.h
@@ -14,53 +14,59 @@ enum class EAvaExtrudeMode : uint8
 };
 
 /** This modifier extrude triangles from a 2D shape with a specific depth and optionally closes the back */
-UCLASS(BlueprintType)
-class AVALANCHEMODIFIERS_API UAvaExtrudeModifier : public UAvaGeometryBaseModifier
+UCLASS(MinimalAPI, BlueprintType)
+class UAvaExtrudeModifier : public UAvaGeometryBaseModifier
 {
 	GENERATED_BODY()
 
 public:
 	static inline const FName ExtrudePolygroupLayerName = TEXT("ExtrudeSide");
 	static inline const FName BackPolygroupLayerName = TEXT("ExtrudeBack");
-	
+
+	AVALANCHEMODIFIERS_API void SetDepth(float InDepth);
+	float GetDepth() const
+	{
+		return Depth;
+	}
+
+	AVALANCHEMODIFIERS_API void SetCloseBack(bool bInCloseBack);
+	bool GetCloseBack() const
+	{
+		return bCloseBack;
+	}
+
+	AVALANCHEMODIFIERS_API void SetExtrudeMode(EAvaExtrudeMode InExtrudeMode);
+	EAvaExtrudeMode GetExtrudeMode() const
+	{
+		return ExtrudeMode;
+	}
+
+protected:
 	//~ Begin UObject
 	virtual void Serialize(FArchive& InArchive) override;
 #if WITH_EDITOR
 	virtual void PostEditChangeProperty(FPropertyChangedEvent& PropertyChangedEvent) override;
 #endif
 	//~ End UObject
-	
-	UFUNCTION()
-	void SetDepth(float InDepth);
-	float GetDepth() const { return Depth; }
-	
-	UFUNCTION()
-	void SetCloseBack(bool bInCloseBack);
-	bool GetCloseBack() const { return bCloseBack; }
-
-	UFUNCTION()
-	void SetExtrudeMode(EAvaExtrudeMode InExtrudeMode);
-	EAvaExtrudeMode GetExtrudeMode() const { return ExtrudeMode; }
 
-protected:
 	//~ Begin UActorModifierCoreBase
 	virtual void OnModifierCDOSetup(FActorModifierCoreMetadata& InMetadata) override;
 	virtual void Apply() override;
 	//~ End UActorModifierCoreBase
-	
+
 	void OnDepthChanged();
 	void OnCloseBackChanged();
 	void OnExtrudeModeChanged();
 	FVector GetExtrudeDirection() const;
-	
+
 	/** Handles mesh depth to extrude primary section */
 	UPROPERTY(EditInstanceOnly, BlueprintReadWrite, Setter="SetDepth", Getter="GetDepth", Category="Extrude", meta=(ClampMin="0", DisplayName="Extrude depth", Units="Centimeters", AllowPrivateAccess="true"))
 	float Depth = 30.f;
-	
+
 	/** Closes the back of the extrude for a 2D shape for example */
 	UPROPERTY(EditInstanceOnly, BlueprintReadWrite, Setter="SetCloseBack", Getter="GetCloseBack", Category="Extrude", meta=(DisplayName="Close back", AllowPrivateAccess="true"))
 	bool bCloseBack = true;
-	
+
 	/** Moves the mesh in the opposite extrude direction by the depth distance */
 	UPROPERTY(meta=(DeprecatedProperty, DeprecationMessage="Use ExtrudeMode instead"))
 	bool bMoveMeshOppositeDirection_DEPRECATED = true;
diff --git a/Engine/Plugins/Experimental/Avalanche/Source/AvalancheModifiers/Public/Modifiers/AvaGeometryBaseModifier.h b/Engine/Plugins/Experimental/Avalanche/Source/AvalancheModifiers/Public/Modifiers/AvaGeometryBaseModifier.h
--- a/Engine/Plugins/Experimental/Avalanche/Source/AvalancheModifiers/Public/Modifiers/AvaGeometryBaseModifier.h
+++ b/Engine/Plugins/Experimental/Avalanche/Source/AvalancheModifiers/Public/Modifiers/AvaGeometryBaseModifier.h
@@ -14,16 +14,16 @@ class UDynamicMeshComponent;
 
 /** Base class used for geometry modifier that uses dynamic mesh components */
 UCLASS(Abstract)
-class AVALANCHEMODIFIERS_API UAvaGeometryBaseModifier : public UAvaBaseModifier
+class UAvaGeometryBaseModifier : public UAvaBaseModifier
 {
 	GENERATED_BODY()
 
 	friend class FAvaGeometryModifierProfiler;
-	
+
 public:
 	/** Get polygroup layer by name */
     static UE::Geometry::FDynamicMeshPolygroupAttribute* FindOrCreatePolygroupLayer(UE::Geometry::FDynamicMesh3& EditMesh, const FName& InLayerName, TArray<int32>* GroupTriangles = nullptr);
-	
+
 protected:
 	//~ Begin UActorModifierCoreBase
 	virtual void OnModifierCDOSetup(FActorModifierCoreMetadata& InMetadata) override;
@@ -31,23 +31,23 @@ protected:
 	virtual void SavePreState() override;
 	virtual void RestorePreState() override;
 	//~ End UActorModifierCoreBase
-	
+
 	/** Gets the cached dynamic mesh component. If null, find the first attached and caches it */
 	UDynamicMeshComponent* GetMeshComponent() const;
 
 	/** Gets the dynamic mesh object from the cached mesh component */
 	UDynamicMesh* GetMeshObject() const;
-    
+
 	/** Checks if the dynamic mesh component is valid */
 	bool IsMeshValid() const;
 
 	/** Get dynamic mesh bounds */
 	FBox GetMeshBounds() const;
-	
+
 	/** Cached Mesh to restore to the Pre-Modifier State */
 	TOptional<UE::Geometry::FDynamicMesh3> PreModifierCachedMesh;
 	TOptional<FBox> PreModifierCachedBounds;
-	
+
 private:
 	/** The actor dynamic mesh component */
 	UPROPERTY(DuplicateTransient)
diff --git a/Engine/Plugins/Experimental/Avalanche/Source/AvalancheModifiers/Public/Modifiers/AvaGlobalOpacityModifier.h b/Engine/Plugins/Experimental/Avalanche/Source/AvalancheModifiers/Public/Modifiers/AvaGlobalOpacityModifier.h
--- a/Engine/Plugins/Experimental/Avalanche/Source/AvalancheModifiers/Public/Modifiers/AvaGlobalOpacityModifier.h
+++ b/Engine/Plugins/Experimental/Avalanche/Source/AvalancheModifiers/Public/Modifiers/AvaGlobalOpacityModifier.h
@@ -9,27 +9,29 @@ class UDMMaterialComponent;
 enum class EDMUpdateType : uint8;
 
 /** This modifier sets global opacity parameters on an actor with Material Designer Instances generated with the Material Designer */
-UCLASS(BlueprintType)
-class AVALANCHEMODIFIERS_API UAvaGlobalOpacityModifier : public UAvaMaterialParameterModifier
+UCLASS(MinimalAPI, BlueprintType)
+class UAvaGlobalOpacityModifier : public UAvaMaterialParameterModifier
 {
 	GENERATED_BODY()
 
 	static inline const FName MaterialDesignerGlobalOpacityValueName = FName("VALUE_GlobalOpacity");
-	
+
 public:
 	UAvaGlobalOpacityModifier();
 
+	AVALANCHEMODIFIERS_API void SetGlobalOpacity(float InOpacity);
+	float GetGlobalOpacity() const
+	{
+		return GlobalOpacity;
+	}
+
+protected:
 	//~ Begin UObject
 #if WITH_EDITOR
 	virtual void PostEditChangeProperty(FPropertyChangedEvent& PropertyChangedEvent) override;
 #endif
 	//~ End UObject
 
-	UFUNCTION()
-	void SetGlobalOpacity(float InOpacity);
-	float GetGlobalOpacity() const { return GlobalOpacity; }
-	
-protected:
 	//~ Begin UActorModifierCoreBase
 	virtual void OnModifierCDOSetup(FActorModifierCoreMetadata& InMetadata) override;
 	//~ End UActorModifierCoreBase
@@ -39,7 +41,7 @@ protected:
 	virtual void OnActorMaterialAdded(UMaterialInstanceDynamic* InAdded) override;
 	virtual void OnActorMaterialRemoved(UMaterialInstanceDynamic* InRemoved) override;
 	void OnDynamicMaterialValueChanged(UDMMaterialComponent* InComponent, EDMUpdateType InUpdateType);
-	
+
 	/** Global opacity to set on all Material Designer Instances */
 	UPROPERTY(EditAnywhere, BlueprintReadWrite, Setter="SetGlobalOpacity", Getter="GetGlobalOpacity", Category="Material Parameter", meta=(ClampMin="0", ClampMax="1", AllowPrivateAccess="true"))
 	float GlobalOpacity = 1.f;
diff --git a/Engine/Plugins/Experimental/Avalanche/Source/AvalancheModifiers/Public/Modifiers/AvaGridArrangeModifier.h b/Engine/Plugins/Experimental/Avalanche/Source/AvalancheModifiers/Public/Modifiers/AvaGridArrangeModifier.h
--- a/Engine/Plugins/Experimental/Avalanche/Source/AvalancheModifiers/Public/Modifiers/AvaGridArrangeModifier.h
+++ b/Engine/Plugins/Experimental/Avalanche/Source/AvalancheModifiers/Public/Modifiers/AvaGridArrangeModifier.h
@@ -18,46 +18,52 @@ enum class EAvaGridArrangeDirection : uint8
 /**
  * Arranges child actors in a 2D grid format.
  */
-UCLASS(BlueprintType)
-class AVALANCHEMODIFIERS_API UAvaGridArrangeModifier : public UAvaArrangeBaseModifier
+UCLASS(MinimalAPI, BlueprintType)
+class UAvaGridArrangeModifier : public UAvaArrangeBaseModifier
 {
 	GENERATED_BODY()
-	
+
 public:
-	//~ Begin UObject
-#if WITH_EDITOR
-	virtual void PostEditChangeProperty(FPropertyChangedEvent& PropertyChangedEvent) override;
-#endif
-	//~ End UObject
+	/** Sets the width and height of the grid. */
+	UFUNCTION(BlueprintCallable, Category = "Motion Design|Modifiers|GridArrange")
+	AVALANCHEMODIFIERS_API void SetCount(const FIntPoint& NewCount);
 
 	/** Returns the width and height of the grid. */
 	UFUNCTION(BlueprintPure, Category = "Motion Design|Modifiers|GridArrange")
-	FIntPoint GetCount() const { return Count; }
+	FIntPoint GetCount() const
+	{
+		return Count;
+	}
 
-	/** Sets the width and height of the grid. */
+	/** Sets the distance between each horizontal and vertical child. */
 	UFUNCTION(BlueprintCallable, Category = "Motion Design|Modifiers|GridArrange")
-	void SetCount(const FIntPoint& NewCount);
+	AVALANCHEMODIFIERS_API void SetSpread(const FVector2D& NewSpread);
 
 	/** Returns the distance between each horizontal and vertical child. */
 	UFUNCTION(BlueprintPure, Category = "Motion Design|Modifiers|GridArrange")
-	FVector2D GetSpread() const { return Spread; }
+	FVector2D GetSpread() const
+	{
+		return Spread;
+	}
 
-	/** Sets the distance between each horizontal and vertical child. */
 	UFUNCTION(BlueprintCallable, Category = "Motion Design|Modifiers|GridArrange")
-	void SetSpread(const FVector2D& NewSpread);
+	AVALANCHEMODIFIERS_API void SetStartCorner(EAvaCorner2D InCorner);
 
-	UFUNCTION(BlueprintCallable, Category = "Motion Design|Modifiers|GridArrange")
-	void SetStartCorner(EAvaCorner2D InCorner);
-	
 	UFUNCTION(BlueprintPure, Category = "Motion Design|Modifiers|GridArrange")
-	EAvaCorner2D GetStartCorner() const { return StartCorner; }
+	EAvaCorner2D GetStartCorner() const
+	{
+		return StartCorner;
+	}
 
 	UFUNCTION(BlueprintCallable, Category = "Motion Design|Modifiers|GridArrange")
-	void SetStartDirection(EAvaGridArrangeDirection InDirection);
-	
+	AVALANCHEMODIFIERS_API void SetStartDirection(EAvaGridArrangeDirection InDirection);
+
 	UFUNCTION(BlueprintPure, Category = "Motion Design|Modifiers|GridArrange")
-	EAvaGridArrangeDirection GetStartDirection() const { return StartDirection; }
-	
+	EAvaGridArrangeDirection GetStartDirection() const
+	{
+		return StartDirection;
+	}
+
 	/** Returns the 2D corner and direction from which to start the arrangement. */
 	UFUNCTION(BlueprintPure, Category = "Motion Design|Modifiers|GridArrange")
 	void GetArrangementLayout(UPARAM(DisplayName = "StartCorner") EAvaCorner2D& OutStartCorner, UPARAM(DisplayName = "Direction") EAvaGridArrangeDirection& OutStartDirection) const;
@@ -67,11 +73,17 @@ public:
 	void SetArrangementLayout(const EAvaCorner2D NewStartCorner, const EAvaGridArrangeDirection NewStartDirection);
 
 protected:
+	//~ Begin UObject
+#if WITH_EDITOR
+	virtual void PostEditChangeProperty(FPropertyChangedEvent& PropertyChangedEvent) override;
+#endif
+	//~ End UObject
+
 	//~ Begin UActorModifierCoreBase
 	virtual void OnModifierCDOSetup(FActorModifierCoreMetadata& InMetadata) override;
 	virtual void Apply() override;
 	//~ End UActorModifierCoreBase
-	
+
 	/** The width and height of the grid. */
 	UPROPERTY(EditAnywhere, BlueprintReadWrite, Setter="SetCount", Getter="GetCount", Category="GridArrange", Meta=(ClampMin="1", UIMin="1", AllowPrivateAccess="true"))
 	FIntPoint Count = FIntPoint(1, 1);
diff --git a/Engine/Plugins/Experimental/Avalanche/Source/AvalancheModifiers/Public/Modifiers/AvaHideEmptyModifier.h b/Engine/Plugins/Experimental/Avalanche/Source/AvalancheModifiers/Public/Modifiers/AvaHideEmptyModifier.h
--- a/Engine/Plugins/Experimental/Avalanche/Source/AvalancheModifiers/Public/Modifiers/AvaHideEmptyModifier.h
+++ b/Engine/Plugins/Experimental/Avalanche/Source/AvalancheModifiers/Public/Modifiers/AvaHideEmptyModifier.h
@@ -7,22 +7,31 @@
 
 class UText3DComponent;
 
-UCLASS(BlueprintType)
-class AVALANCHEMODIFIERS_API UAvaHideEmptyModifier : public UAvaArrangeBaseModifier
+UCLASS(MinimalAPI, BlueprintType)
+class UAvaHideEmptyModifier : public UAvaArrangeBaseModifier
 {
 	GENERATED_BODY()
-	
+
 public:
+	AVALANCHEMODIFIERS_API void SetContainerActorWeak(TWeakObjectPtr<AActor> InContainer);
+	TWeakObjectPtr<AActor> GetContainerActorWeak() const
+	{
+		return ContainerActorWeak;
+	}
+
+	AVALANCHEMODIFIERS_API void SetInvertVisibility(bool bInInvert);
+	bool GetInvertVisibility() const
+	{
+		return bInvertVisibility;
+	}
+
+protected:
 	//~ Begin UObject
 #if WITH_EDITOR
 	virtual void PostEditChangeProperty(FPropertyChangedEvent& PropertyChangedEvent) override;
 #endif
 	//~ End UObject
 
-	void SetInvertVisibility(bool bInInvert);
-	bool GetInvertVisibility() const { return bInvertVisibility; }
-	
-protected:
 	//~ Begin UActorModifierCoreBase
 	virtual void OnModifierCDOSetup(FActorModifierCoreMetadata& InMetadata) override;
 	virtual void OnModifierAdded(EActorModifierCoreEnableReason InReason) override;
@@ -31,7 +40,7 @@ protected:
 	virtual void OnModifierDisabled(EActorModifierCoreDisableReason InReason) override;
 	virtual void OnModifiedActorTransformed() override;
 	//~ End UActorModifierCoreBase
-	
+
 	void OnTextChanged();
 	void OnContainerActorChanged();
 	void OnInvertVisibilityChanged();
@@ -40,11 +49,11 @@ protected:
 	virtual void OnSceneTreeTrackedActorChanged(int32 InIdx, AActor* InPreviousActor, AActor* InNewActor) override;
 	virtual void OnSceneTreeTrackedActorDirectChildrenChanged(int32 InIdx, const TArray<TWeakObjectPtr<AActor>>& InPreviousChildrenActors, const TArray<TWeakObjectPtr<AActor>>& InNewChildrenActors) override;
 	//~ End IAvaSceneTreeUpdateModifierExtension
-	
+
 	/** The container to hide when text is empty, by default self */
-	UPROPERTY(EditAnywhere, BlueprintReadWrite, Category="HideEmpty", meta=(DisplayName="ContainerActor", AllowPrivateAccess="true"))
+	UPROPERTY(EditAnywhere, BlueprintReadWrite, Setter, Getter, Category="HideEmpty", meta=(DisplayName="ContainerActor", AllowPrivateAccess="true"))
 	TWeakObjectPtr<AActor> ContainerActorWeak;
-	
+
 	/** Invert the behaviour and visibility of the container if text is empty */
 	UPROPERTY(EditAnywhere, BlueprintReadWrite, Setter="SetInvertVisibility", Getter="GetInvertVisibility", Category="HideEmpty", meta=(AllowPrivateAccess="true"))
 	bool bInvertVisibility = false;
diff --git a/Engine/Plugins/Experimental/Avalanche/Source/AvalancheModifiers/Public/Modifiers/AvaJustifyModifier.h b/Engine/Plugins/Experimental/Avalanche/Source/AvalancheModifiers/Public/Modifiers/AvaJustifyModifier.h
--- a/Engine/Plugins/Experimental/Avalanche/Source/AvalancheModifiers/Public/Modifiers/AvaJustifyModifier.h
+++ b/Engine/Plugins/Experimental/Avalanche/Source/AvalancheModifiers/Public/Modifiers/AvaJustifyModifier.h
@@ -38,75 +38,93 @@ enum class EAvaJustifyDepth : uint8
  * 
  * Aligns child actors, based on their bounding box, according to the specified justification
  */
-UCLASS(BlueprintType)
-class AVALANCHEMODIFIERS_API UAvaJustifyModifier : public UAvaArrangeBaseModifier
+UCLASS(MinimalAPI, BlueprintType)
+class UAvaJustifyModifier : public UAvaArrangeBaseModifier
 {
 	GENERATED_BODY()
 
 public:
-	//~ Begin UObject
-#if WITH_EDITOR
-	virtual void PostEditChangeProperty(FPropertyChangedEvent& InPropertyChangedEvent) override;
-#endif
-	//~ End UObject
+	UFUNCTION(BlueprintCallable, Category = "Justify")
+	AVALANCHEMODIFIERS_API void SetHorizontalAlignment(EAvaJustifyHorizontal InHorizontalAlignment);
 
 	UFUNCTION(BlueprintPure, Category = "Justify")
-	EAvaJustifyHorizontal GetHorizontalAlignment() const { return HorizontalAlignment; }
-	
+	EAvaJustifyHorizontal GetHorizontalAlignment() const
+	{
+		return HorizontalAlignment;
+	}
+
 	UFUNCTION(BlueprintCallable, Category = "Justify")
-	void SetHorizontalAlignment(EAvaJustifyHorizontal InHorizontalAlignment);
+	AVALANCHEMODIFIERS_API void SetVerticalAlignment(EAvaJustifyVertical InVerticalAlignment);
 
 	UFUNCTION(BlueprintPure, Category = "Justify")
-	EAvaJustifyVertical GetVerticalAlignment() const { return VerticalAlignment; }
-	
+	EAvaJustifyVertical GetVerticalAlignment() const
+	{
+		return VerticalAlignment;
+	}
+
 	UFUNCTION(BlueprintCallable, Category = "Justify")
-	void SetVerticalAlignment(EAvaJustifyVertical InVerticalAlignment);
+	AVALANCHEMODIFIERS_API void SetDepthAlignment(EAvaJustifyDepth InDepthAlignment);
 
 	UFUNCTION(BlueprintPure, Category = "Justify")
-	EAvaJustifyDepth GetDepthAlignment() const { return DepthAlignment; }
-	
+	EAvaJustifyDepth GetDepthAlignment() const
+	{
+		return DepthAlignment;
+	}
+
 	UFUNCTION(BlueprintCallable, Category = "Justify")
-	void SetDepthAlignment(EAvaJustifyDepth InDepthAlignment);
+	AVALANCHEMODIFIERS_API void SetHorizontalAnchor(float InHorizontalAnchor);
 
 	UFUNCTION(BlueprintPure, Category = "Justify")
-	float GetHorizontalAnchor() const { return HorizontalAnchor; }
-	
+	float GetHorizontalAnchor() const
+	{
+		return HorizontalAnchor;
+	}
+
 	UFUNCTION(BlueprintCallable, Category = "Justify")
-	void SetHorizontalAnchor(float InHorizontalAnchor);
+	AVALANCHEMODIFIERS_API void SetVerticalAnchor(float InVerticalAnchor);
 
 	UFUNCTION(BlueprintPure, Category = "Justify")
-	float GetVerticalAnchor() const { return VerticalAnchor; }
-	
+	float GetVerticalAnchor() const
+	{
+		return VerticalAnchor;
+	}
+
 	UFUNCTION(BlueprintCallable, Category = "Justify")
-	void SetVerticalAnchor(float InVerticalAnchor);
+	AVALANCHEMODIFIERS_API void SetDepthAnchor(float InDepthAnchor);
 
 	UFUNCTION(BlueprintPure, Category = "Justify")
-	float GetDepthAnchor() const { return DepthAnchor; }
-	
-	UFUNCTION(BlueprintCallable, Category = "Justify")
-	void SetDepthAnchor(float InDepthAnchor);
+	float GetDepthAnchor() const
+	{
+		return DepthAnchor;
+	}
 
 protected:
+	//~ Begin UObject
+#if WITH_EDITOR
+	virtual void PostEditChangeProperty(FPropertyChangedEvent& InPropertyChangedEvent) override;
+#endif
+	//~ End UObject
+
 	//~ Begin UActorModifierCoreBase
 	virtual void OnModifierCDOSetup(FActorModifierCoreMetadata& InMetadata) override;
 	virtual bool IsModifierDirtyable() const override;
 	virtual void OnModifiedActorTransformed() override;
 	virtual void Apply() override;
 	//~ End UActorModifierCoreBase
-	
+
 	//~ Begin IAvaSceneTreeUpdateModifierExtension
 	virtual void OnSceneTreeTrackedActorDirectChildrenChanged(int32 InIdx, const TArray<TWeakObjectPtr<AActor>>& InPreviousChildrenActors, const TArray<TWeakObjectPtr<AActor>>& InNewChildrenActors) override;
 	//~ End IAvaSceneTreeUpdateModifierExtension
-	
+
 	//~ Begin IAvaRenderStateUpdateExtension
 	virtual void OnRenderStateUpdated(AActor* InActor, UActorComponent* InComponent) override;
 	virtual void OnActorVisibilityChanged(AActor* InActor) override;
 	//~ End IAvaRenderStateUpdateExtension
-	
+
 	//~ Begin IAvaTransformUpdateExtension
 	virtual void OnTransformUpdated(AActor* InActor, bool bInParentMoved) override;
 	//~ End IAvaTransformUpdateExtension
-	
+
 	// Get tracked actors from children actors
 	void GetChildrenActors(TSet<TWeakObjectPtr<AActor>>& OutChildren) const;
 	void GetTrackedActors(const TSet<TWeakObjectPtr<AActor>>& InChildrenActors, TArray<TWeakObjectPtr<const AActor>>& OutTrackedActors) const;
@@ -119,13 +137,13 @@ protected:
 	FVector MakeConstrainedAxisVector() const;
 	FVector GetAnchorOffset() const;
 	FVector GetAlignmentOffset(const FVector& InExtent) const;
-	
+
 	UPROPERTY(EditAnywhere, BlueprintReadWrite, Setter="SetHorizontalAlignment", Getter="GetHorizontalAlignment", Category="Justify", meta=(AllowPrivateAccess="true"))
 	EAvaJustifyHorizontal HorizontalAlignment;
 
 	UPROPERTY(EditAnywhere, BlueprintReadWrite, Setter="SetVerticalAlignment", Getter="GetVerticalAlignment", Category="Justify", meta=(AllowPrivateAccess="true"))
 	EAvaJustifyVertical VerticalAlignment;
-	
+
 	UPROPERTY(EditAnywhere, BlueprintReadWrite, Setter="SetDepthAlignment", Getter="GetDepthAlignment", Category="Justify", meta=(AllowPrivateAccess="true"))
 	EAvaJustifyDepth DepthAlignment;
 
diff --git a/Engine/Plugins/Experimental/Avalanche/Source/AvalancheModifiers/Public/Modifiers/AvaLookAtModifier.h b/Engine/Plugins/Experimental/Avalanche/Source/AvalancheModifiers/Public/Modifiers/AvaLookAtModifier.h
--- a/Engine/Plugins/Experimental/Avalanche/Source/AvalancheModifiers/Public/Modifiers/AvaLookAtModifier.h
+++ b/Engine/Plugins/Experimental/Avalanche/Source/AvalancheModifiers/Public/Modifiers/AvaLookAtModifier.h
@@ -12,30 +12,26 @@ enum class EAvaAxis : uint8;
 /**
  * Rotates the modifying actor to point it's specified axis at another actor.
  */
-UCLASS(BlueprintType)
-class AVALANCHEMODIFIERS_API UAvaLookAtModifier : public UAvaAttachmentBaseModifier
+UCLASS(MinimalAPI, BlueprintType)
+class UAvaLookAtModifier : public UAvaAttachmentBaseModifier
 	, public IAvaTransformUpdateHandler
 {
 	GENERATED_BODY()
-	
-public:
-	//~ Begin UObject
-	virtual void PostLoad() override;
-#if WITH_EDITOR
-	virtual void PostEditChangeProperty(FPropertyChangedEvent& PropertyChangedEvent) override;
-	virtual void PostTransacted(const FTransactionObjectEvent& TransactionEvent) override;
-#endif
-	//~ End UObject
 
+public:
 	UFUNCTION(BlueprintCallable, Category = "Motion Design|Modifiers|LookAt")
-	void SetReferenceActor(const FAvaSceneTreeActor& InReferenceActor);
-	
+	AVALANCHEMODIFIERS_API void SetReferenceActor(const FAvaSceneTreeActor& InReferenceActor);
+
 	UFUNCTION(BlueprintPure, Category = "Motion Design|Modifiers|LookAt")
 	const FAvaSceneTreeActor& GetReferenceActor() const
 	{
 		return ReferenceActor;
 	}
-	
+
+	/** Sets the axis that will point towards the reference actor. */
+	UFUNCTION(BlueprintCallable, Category = "Motion Design|Modifiers|LookAt")
+	AVALANCHEMODIFIERS_API void SetAxis(const EAvaAxis NewAxis);
+
 	/** Returns the axis that will point towards t he reference actor. */
 	UFUNCTION(BlueprintPure, Category = "Motion Design|Modifiers|LookAt")
 	EAvaAxis GetAxis() const
@@ -43,9 +39,9 @@ public:
 		return Axis;
 	}
 
-	/** Sets the axis that will point towards the reference actor. */
+	/** Sets the look-at direction to be flipped. */
 	UFUNCTION(BlueprintCallable, Category = "Motion Design|Modifiers|LookAt")
-	void SetAxis(const EAvaAxis NewAxis);
+	AVALANCHEMODIFIERS_API void SetFlipAxis(const bool bNewFlipAxis);
 
 	/** Returns true if flipping the look-at rotation axis. */
 	UFUNCTION(BlueprintPure, Category = "Motion Design|Modifiers|LookAt")
@@ -54,11 +50,15 @@ public:
 		return bFlipAxis;
 	}
 
-	/** Sets the look-at direction to be flipped. */
-	UFUNCTION(BlueprintCallable, Category = "Motion Design|Modifiers|LookAt")
-	void SetFlipAxis(const bool bNewFlipAxis);
-
 protected:
+	//~ Begin UObject
+	virtual void PostLoad() override;
+#if WITH_EDITOR
+	virtual void PostEditChangeProperty(FPropertyChangedEvent& PropertyChangedEvent) override;
+	virtual void PostTransacted(const FTransactionObjectEvent& TransactionEvent) override;
+#endif
+	//~ End UObject
+
 	//~ Begin UActorModifierCoreBase
 	virtual void OnModifierCDOSetup(FActorModifierCoreMetadata& InMetadata) override;
 	virtual void OnModifierAdded(EActorModifierCoreEnableReason InReason) override;
@@ -67,17 +67,17 @@ protected:
 	virtual void Apply() override;
 	virtual void OnModifiedActorTransformed() override;
 	//~ End UActorModifierCoreBase
-	
+
 	//~ Begin IAvaTransformUpdateExtension
 	virtual void OnTransformUpdated(AActor* InActor, bool bInParentMoved) override;
 	//~ End IAvaTransformUpdateExtension
-	
+
 	//~ Begin IAvaSceneTreeUpdateModifierExtension
 	virtual void OnSceneTreeTrackedActorChanged(int32 InIdx, AActor* InPreviousActor, AActor* InNewActor) override;
 	//~ End IAvaSceneTreeUpdateModifierExtension
-	
+
 	void OnReferenceActorChanged();
-	
+
 	UPROPERTY(EditAnywhere, BlueprintReadWrite, Setter="SetReferenceActor", Getter="GetReferenceActor", Category="LookAt", meta=(ShowOnlyInnerProperties, AllowPrivateAccess="true"))
 	FAvaSceneTreeActor ReferenceActor;
 
diff --git a/Engine/Plugins/Experimental/Avalanche/Source/AvalancheModifiers/Public/Modifiers/AvaMaterialParameterModifier.h b/Engine/Plugins/Experimental/Avalanche/Source/AvalancheModifiers/Public/Modifiers/AvaMaterialParameterModifier.h
--- a/Engine/Plugins/Experimental/Avalanche/Source/AvalancheModifiers/Public/Modifiers/AvaMaterialParameterModifier.h
+++ b/Engine/Plugins/Experimental/Avalanche/Source/AvalancheModifiers/Public/Modifiers/AvaMaterialParameterModifier.h
@@ -34,27 +34,33 @@ struct FAvaMaterialParameterMap
 };
 
 /** This modifier sets specified dynamic materials parameters on an actor and its children */
-UCLASS(BlueprintType)
-class AVALANCHEMODIFIERS_API UAvaMaterialParameterModifier : public UAvaArrangeBaseModifier
+UCLASS(MinimalAPI, BlueprintType)
+class UAvaMaterialParameterModifier : public UAvaArrangeBaseModifier
 {
 	GENERATED_BODY()
-	
+
 public:
 	UAvaMaterialParameterModifier();
 
+	AVALANCHEMODIFIERS_API void SetMaterialParameters(const FAvaMaterialParameterMap& InParameterMap);
+	const FAvaMaterialParameterMap& GetMaterialParameters() const
+	{
+		return MaterialParameters;
+	}
+
+	AVALANCHEMODIFIERS_API void SetUpdateChildren(bool bInUpdateChildren);
+	bool GetUpdateChildren() const
+	{
+		return bUpdateChildren;
+	}
+
+protected:
 	//~ Begin UObject
 #if WITH_EDITOR
 	virtual void PostEditChangeProperty(FPropertyChangedEvent& PropertyChangedEvent) override;
 #endif
 	//~ End UObject
 
-	void SetMaterialParameters(const FAvaMaterialParameterMap& InParameterMap);
-	const FAvaMaterialParameterMap& GetMaterialParameters() const { return MaterialParameters; }
-	
-	void SetUpdateChildren(bool bInUpdateChildren);
-	bool GetUpdateChildren() const { return bUpdateChildren; }
-
-protected:
 	//~ Begin UActorModifierCoreBase
 	virtual void OnModifierCDOSetup(FActorModifierCoreMetadata& InMetadata) override;
 	virtual void OnModifierEnabled(EActorModifierCoreEnableReason InReason) override;
@@ -63,7 +69,7 @@ protected:
 	virtual void RestorePreState() override;
 	virtual void Apply() override;
 	//~ End UActorModifierCoreBase
-	
+
 	//~ Begin IAvaSceneTreeUpdateModifierExtension
 	virtual void OnSceneTreeTrackedActorChildrenChanged(int32 InIdx, const TSet<TWeakObjectPtr<AActor>>& InPreviousChildrenActors, const TSet<TWeakObjectPtr<AActor>>& InNewChildrenActors) override;
 	virtual void OnSceneTreeTrackedActorDirectChildrenChanged(int32 InIdx, const TArray<TWeakObjectPtr<AActor>>& InPreviousChildrenActors, const TArray<TWeakObjectPtr<AActor>>& InNewChildrenActors) override;
@@ -78,7 +84,7 @@ protected:
 	/** Called when a property changes, used to detect material changes */
 	virtual void OnActorPropertyChanged(UObject* InObject, FPropertyChangedEvent& InChangeEvent);
 #endif
-	
+
 	void ScanActorMaterials();
 
 	void OnMaterialParametersChanged();
@@ -86,13 +92,13 @@ protected:
 
 	virtual void OnActorMaterialAdded(UMaterialInstanceDynamic* InAdded) {}
 	virtual void OnActorMaterialRemoved(UMaterialInstanceDynamic* InRemoved) {}
-	
+
 	/** Checks if this actor has a Material Designer Instance or that we already track one */
 	bool IsActorSupported(const AActor* InActor) const;
 
 	/** Retrieves all Material Designer Instance from a primitive component */
 	TSet<UMaterialInstanceDynamic*> GetComponentDynamicMaterials(const UPrimitiveComponent* InComponent) const;
-	
+
 	/** Which parameters should we set on the Material Designer Instance,
 	 * use EditCondition="bShowMaterialParameters && bShowMaterialParameters" otherwise edit inline boolean appear in details when it should not */
 	UPROPERTY(EditAnywhere, BlueprintReadWrite, Setter="SetMaterialParameters", Getter="GetMaterialParameters", Category="MaterialParameter", meta=(EditCondition="bShowMaterialParameters && bShowMaterialParameters", EditConditionHides, AllowPrivateAccess="true"))
@@ -105,7 +111,7 @@ protected:
 	/** Filter material type for child modifiers */
 	UPROPERTY(Transient)
 	TSubclassOf<UMaterialInstanceDynamic> MaterialClass;
-	
+
 	/** Will also look into attached children actors */
 	UPROPERTY(EditAnywhere, BlueprintReadWrite, Setter="SetUpdateChildren", Getter="GetUpdateChildren", Category="MaterialParameter", meta=(AllowPrivateAccess="true"))
 	bool bUpdateChildren = true;
diff --git a/Engine/Plugins/Experimental/Avalanche/Source/AvalancheModifiers/Public/Modifiers/AvaMirrorModifier.h b/Engine/Plugins/Experimental/Avalanche/Source/AvalancheModifiers/Public/Modifiers/AvaMirrorModifier.h
--- a/Engine/Plugins/Experimental/Avalanche/Source/AvalancheModifiers/Public/Modifiers/AvaMirrorModifier.h
+++ b/Engine/Plugins/Experimental/Avalanche/Source/AvalancheModifiers/Public/Modifiers/AvaMirrorModifier.h
@@ -8,39 +8,49 @@
 
 class UStaticMesh;
 
-UCLASS(BlueprintType)
-class AVALANCHEMODIFIERS_API UAvaMirrorModifier : public UAvaGeometryBaseModifier
+UCLASS(MinimalAPI, BlueprintType)
+class UAvaMirrorModifier : public UAvaGeometryBaseModifier
 {
 	GENERATED_BODY()
-	
+
 public:
+	AVALANCHEMODIFIERS_API void SetMirrorFramePosition(const FVector& InMirrorFramePosition);
+	const FVector& GetMirrorFramePosition() const
+	{
+		return MirrorFramePosition;
+	}
+
+	AVALANCHEMODIFIERS_API void SetMirrorFrameRotation(const FRotator& InMirrorFrameRotation);
+	const FRotator& GetMirrorFrameRotation() const
+	{
+		return MirrorFrameRotation;
+	}
+
+	AVALANCHEMODIFIERS_API void SetApplyPlaneCut(bool bInApplyPlaneCut);
+	bool GetApplyPlaneCut() const
+	{
+		return bApplyPlaneCut;
+	}
+
+	AVALANCHEMODIFIERS_API void SetFlipCutSide(bool bInFlipCutSide);
+	bool GetFlipCutSide() const
+	{
+		return bFlipCutSide;
+	}
+
+	AVALANCHEMODIFIERS_API void SetWeldAlongPlane(bool bInWeldAlongPlane);
+	bool GetWeldAlongPlane() const
+	{
+		return bWeldAlongPlane;
+	}
+
+protected:
 	//~ Begin UObject
 #if WITH_EDITOR
 	virtual void PostEditChangeProperty(FPropertyChangedEvent& PropertyChangedEvent) override;
 #endif
 	//~ End UObject
-	
-	UFUNCTION()
-	void SetMirrorFramePosition(const FVector& InMirrorFramePosition);
-	const FVector& GetMirrorFramePosition() const { return MirrorFramePosition; }
-
-	UFUNCTION()
-	void SetMirrorFrameRotation(const FRotator& InMirrorFrameRotation);
-	const FRotator& GetMirrorFrameRotation() const { return MirrorFrameRotation; }
-	
-	UFUNCTION()
-	void SetApplyPlaneCut(bool bInApplyPlaneCut);
-	bool GetApplyPlaneCut() const { return bApplyPlaneCut; }
-
-	UFUNCTION()
-	void SetFlipCutSide(bool bInFlipCutSide);
-	bool GetFlipCutSide() const { return bFlipCutSide; }
-
-	UFUNCTION()
-	void SetWeldAlongPlane(bool bInWeldAlongPlane);
-	bool GetWeldAlongPlane() const { return bWeldAlongPlane; }
-	
-protected:
+
 	//~ Begin UActorModifierCoreBase
 	virtual void OnModifierCDOSetup(FActorModifierCoreMetadata& InMetadata) override;
 	virtual void OnModifierAdded(EActorModifierCoreEnableReason InReason) override;
@@ -48,7 +58,7 @@ protected:
 	virtual void OnModifierDisabled(EActorModifierCoreDisableReason InReason) override;
 	virtual void OnModifierRemoved(EActorModifierCoreDisableReason InReason) override;
 	//~ End UActorModifierCoreBase
-	
+
 	void OnMirrorFrameChanged();
 	void OnMirrorOptionChanged();
 
@@ -56,10 +66,10 @@ protected:
 	void CreatePreviewComponent();
 	void DestroyPreviewComponent();
 	void UpdatePreviewComponent();
-	
+
 	void OnShowMirrorFrameChanged();
 #endif
-	
+
     UPROPERTY(EditInstanceOnly, BlueprintReadWrite, Setter="SetMirrorFramePosition", Getter="GetMirrorFramePosition", Category="Mirror", meta=(AllowPrivateAccess="true"))
 	FVector MirrorFramePosition = FVector::ZeroVector;
 
diff --git a/Engine/Plugins/Experimental/Avalanche/Source/AvalancheModifiers/Public/Modifiers/AvaNormalModifier.h b/Engine/Plugins/Experimental/Avalanche/Source/AvalancheModifiers/Public/Modifiers/AvaNormalModifier.h
--- a/Engine/Plugins/Experimental/Avalanche/Source/AvalancheModifiers/Public/Modifiers/AvaNormalModifier.h
+++ b/Engine/Plugins/Experimental/Avalanche/Source/AvalancheModifiers/Public/Modifiers/AvaNormalModifier.h
@@ -20,53 +20,64 @@ enum class EAvaNormalModifierSplitMethod : uint8
 	Threshold
 };
 
-UCLASS(BlueprintType)
-class AVALANCHEMODIFIERS_API UAvaNormalModifier : public UAvaGeometryBaseModifier
+UCLASS(MinimalAPI, BlueprintType)
+class UAvaNormalModifier : public UAvaGeometryBaseModifier
 {
 	GENERATED_BODY()
 
 public:
+	AVALANCHEMODIFIERS_API void SetAngleWeighted(bool bInAngleWeighted);
+	bool GetAngleWeighted() const
+	{
+		return bAngleWeighted;
+	}
+
+	AVALANCHEMODIFIERS_API void SetAreaWeighted(bool bInAreaWeighted);
+	bool GetAreaWeighted() const
+	{
+		return bAreaWeighted;
+	}
+
+	AVALANCHEMODIFIERS_API void SetInvert(bool bInInvert);
+	bool GetInvert() const
+	{
+		return bInvert;
+	}
+
+	AVALANCHEMODIFIERS_API void SetSplitMethod(EAvaNormalModifierSplitMethod InSplitMethod);
+	EAvaNormalModifierSplitMethod GetSplitMethod() const
+	{
+		return SplitMethod;
+	}
+
+	AVALANCHEMODIFIERS_API void SetAngleThreshold(float InAngleThreshold);
+	float GetAngleThreshold() const
+	{
+		return AngleThreshold;
+	}
+
+	AVALANCHEMODIFIERS_API void SetPolyGroupLayerIdx(int32 InPolyGroupLayer);
+	AVALANCHEMODIFIERS_API int32 GetPolyGroupLayerIdx() const;
+
+	AVALANCHEMODIFIERS_API void SetPolyGroupLayer(FString& InString);
+	FString GetPolyGroupLayer() const
+	{
+		return PolyGroupLayer;
+	}
+
+protected:
 	//~ Begin UObject
 #if WITH_EDITOR
 	virtual void PostEditChangeProperty(FPropertyChangedEvent& PropertyChangedEvent) override;
 #endif
 	//~ End UObject
 
-	UFUNCTION()
-	void SetAngleWeighted(bool bInAngleWeighted);
-	bool GetAngleWeighted() const { return bAngleWeighted; }
-	
-	UFUNCTION()
-	void SetAreaWeighted(bool bInAreaWeighted);
-	bool GetAreaWeighted() const { return bAreaWeighted; }
-	
-	UFUNCTION()
-	void SetInvert(bool bInInvert);
-	bool GetInvert() const { return bInvert; }
-	
-	UFUNCTION()
-	void SetSplitMethod(EAvaNormalModifierSplitMethod InSplitMethod);
-	EAvaNormalModifierSplitMethod GetSplitMethod() const { return SplitMethod; }
-	
-	UFUNCTION()
-	void SetAngleThreshold(float InAngleThreshold);
-	float GetAngleThreshold() const { return AngleThreshold; }
-
-	UFUNCTION()
-	void SetPolyGroupLayerIdx(int32 InPolyGroupLayer);
-	int32 GetPolyGroupLayerIdx() const;
-
-	UFUNCTION()
-	void SetPolyGroupLayer(FString& InString);
-	FString GetPolyGroupLayer() const { return PolyGroupLayer; }
-
-protected:
 	//~ Begin UActorModifierCoreBase
 	virtual void OnModifierCDOSetup(FActorModifierCoreMetadata& InMetadata) override;
 	virtual void OnModifierAdded(EActorModifierCoreEnableReason InReason) override;
 	virtual void Apply() override;
 	//~ End UActorModifierCoreBase
-	
+
 	void OnAngleWeightedChanged();
 	void OnAreaWeightedChanged();
 	void OnInvertChanged();
diff --git a/Engine/Plugins/Experimental/Avalanche/Source/AvalancheModifiers/Public/Modifiers/AvaOutlineModifier.h b/Engine/Plugins/Experimental/Avalanche/Source/AvalancheModifiers/Public/Modifiers/AvaOutlineModifier.h
--- a/Engine/Plugins/Experimental/Avalanche/Source/AvalancheModifiers/Public/Modifiers/AvaOutlineModifier.h
+++ b/Engine/Plugins/Experimental/Avalanche/Source/AvalancheModifiers/Public/Modifiers/AvaOutlineModifier.h
@@ -13,46 +13,52 @@ enum class EAvaOutlineMode : uint8
 };
 
 /** This modifier adds an outline around a 2D shape with a specific distance */
-UCLASS(BlueprintType)
-class AVALANCHEMODIFIERS_API UAvaOutlineModifier : public UAvaGeometryBaseModifier
+UCLASS(MinimalAPI, BlueprintType)
+class UAvaOutlineModifier : public UAvaGeometryBaseModifier
 {
 	GENERATED_BODY()
 
 public:
+	AVALANCHEMODIFIERS_API void SetMode(EAvaOutlineMode InMode);
+	EAvaOutlineMode GetMode() const
+	{
+		return Mode;
+	}
+
+	AVALANCHEMODIFIERS_API void SetDistance(float InDistance);
+	float GetDistance() const
+	{
+		return Distance;
+	}
+
+	AVALANCHEMODIFIERS_API void SetRemoveInside(bool bInRemoveInside);
+	bool GetRemoveInside() const
+	{
+		return bRemoveInside;
+	}
+
+protected:
 	//~ Begin UObject
 #if WITH_EDITOR
 	virtual void PostEditChangeProperty(FPropertyChangedEvent& PropertyChangedEvent) override;
 #endif
 	//~ End UObject
 
-	UFUNCTION()
-	void SetMode(EAvaOutlineMode InMode);
-	EAvaOutlineMode GetMode() const { return Mode; }
-	
-	UFUNCTION()
-	void SetDistance(float InDistance);
-	float GetDistance() const { return Distance; }
-
-	UFUNCTION()
-	void SetRemoveInside(bool bInRemoveInside);
-	bool GetRemoveInside() const { return bRemoveInside; }
-
-	float GetMaxInsetDistance() const;
-
-protected:
 	//~ Begin UActorModifierCoreBase
 	virtual void OnModifierCDOSetup(FActorModifierCoreMetadata& InMetadata) override;
 	virtual void Apply() override;
 	//~ End UActorModifierCoreBase
-	
+
 	void OnModeChanged();
 	void OnDistanceChanged();
 	void OnRemoveInsideChanged();
 
+	float GetMaxInsetDistance() const;
+
 	/** Set the mode like inset or outset */
 	UPROPERTY(EditInstanceOnly, BlueprintReadWrite, Setter="SetMode", Getter="GetMode", Category="Outline", meta=(AllowPrivateAccess="true"))
 	EAvaOutlineMode Mode = EAvaOutlineMode::Outset;
-	
+
 	/** Set the distance for the outline */
 	UPROPERTY(EditInstanceOnly, BlueprintReadWrite, Setter="SetDistance", Getter="GetDistance", Category="Outline", meta=(ClampMin="0", AllowPrivateAccess="true"))
 	float Distance = 10.f;
diff --git a/Engine/Plugins/Experimental/Avalanche/Source/AvalancheModifiers/Public/Modifiers/AvaPatternModifier.h b/Engine/Plugins/Experimental/Avalanche/Source/AvalancheModifiers/Public/Modifiers/AvaPatternModifier.h
--- a/Engine/Plugins/Experimental/Avalanche/Source/AvalancheModifiers/Public/Modifiers/AvaPatternModifier.h
+++ b/Engine/Plugins/Experimental/Avalanche/Source/AvalancheModifiers/Public/Modifiers/AvaPatternModifier.h
@@ -154,31 +154,43 @@ struct FAvaPatternModifierCircleLayoutOptions
 };
 
 /** This modifier clones a shape following various layouts and options */
-UCLASS(BlueprintType, AutoExpandCategories=(Pattern))
-class AVALANCHEMODIFIERS_API UAvaPatternModifier : public UAvaGeometryBaseModifier
+UCLASS(MinimalAPI, BlueprintType, AutoExpandCategories=(Pattern))
+class UAvaPatternModifier : public UAvaGeometryBaseModifier
 {
 	GENERATED_BODY()
 
 public:
+	AVALANCHEMODIFIERS_API void SetLayout(EAvaPatternModifierLayout InLayout);
+	EAvaPatternModifierLayout GetLayout() const
+	{
+		return Layout;
+	}
+
+	AVALANCHEMODIFIERS_API void SetLineLayoutOptions(const FAvaPatternModifierLineLayoutOptions& InOptions);
+	const FAvaPatternModifierLineLayoutOptions& GetLineLayoutOptions() const
+	{
+		return LineLayoutOptions;
+	}
+
+	AVALANCHEMODIFIERS_API void SetGridLayoutOptions(const FAvaPatternModifierGridLayoutOptions& InOptions);
+	const FAvaPatternModifierGridLayoutOptions& GetGridLayoutOptions() const
+	{
+		return GridLayoutOptions;
+	}
+
+	AVALANCHEMODIFIERS_API void SetCircleLayoutOptions(const FAvaPatternModifierCircleLayoutOptions& InOptions);
+	const FAvaPatternModifierCircleLayoutOptions& GetCircleLayoutOptions() const
+	{
+		return CircleLayoutOptions;
+	}
+
+protected:
 	//~ Begin UObject
 #if WITH_EDITOR
 	virtual void PostEditChangeProperty(FPropertyChangedEvent& PropertyChangedEvent) override;
 #endif
 	//~ End UObject
 
-	void SetLayout(EAvaPatternModifierLayout InLayout);
-	EAvaPatternModifierLayout GetLayout() const { return Layout; }
-
-	void SetLineLayoutOptions(const FAvaPatternModifierLineLayoutOptions& InOptions);
-	const FAvaPatternModifierLineLayoutOptions& GetLineLayoutOptions() const { return LineLayoutOptions; }
-
-	void SetGridLayoutOptions(const FAvaPatternModifierGridLayoutOptions& InOptions);
-	const FAvaPatternModifierGridLayoutOptions& GetGridLayoutOptions() const { return GridLayoutOptions; }
-
-	void SetCircleLayoutOptions(const FAvaPatternModifierCircleLayoutOptions& InOptions);
-	const FAvaPatternModifierCircleLayoutOptions& GetCircleLayoutOptions() const { return CircleLayoutOptions; }
-
-protected:
 	//~ Begin UActorModifierCoreBase
 	virtual void OnModifierCDOSetup(FActorModifierCoreMetadata& InMetadata) override;
 	virtual void Apply() override;
diff --git a/Engine/Plugins/Experimental/Avalanche/Source/AvalancheModifiers/Public/Modifiers/AvaPlaneCutModifier.h b/Engine/Plugins/Experimental/Avalanche/Source/AvalancheModifiers/Public/Modifiers/AvaPlaneCutModifier.h
--- a/Engine/Plugins/Experimental/Avalanche/Source/AvalancheModifiers/Public/Modifiers/AvaPlaneCutModifier.h
+++ b/Engine/Plugins/Experimental/Avalanche/Source/AvalancheModifiers/Public/Modifiers/AvaPlaneCutModifier.h
@@ -7,35 +7,43 @@
 #include "AvaPlaneCutModifier.generated.h"
 
 /** This modifier cuts a shape based on a 2D plane */
-UCLASS(BlueprintType)
-class AVALANCHEMODIFIERS_API UAvaPlaneCutModifier : public UAvaGeometryBaseModifier
+UCLASS(MinimalAPI, BlueprintType)
+class UAvaPlaneCutModifier : public UAvaGeometryBaseModifier
 {
 	GENERATED_BODY()
 
 public:
+	AVALANCHEMODIFIERS_API void SetPlaneOrigin(float InOrigin);
+	float GetPlaneOrigin() const
+	{
+		return PlaneOrigin;
+	}
+
+	AVALANCHEMODIFIERS_API void SetPlaneRotation(const FRotator& InRotation);
+	const FRotator& GetPlaneRotation() const
+	{
+		return PlaneRotation;
+	}
+
+	AVALANCHEMODIFIERS_API void SetInvertCut(bool bInInvertCut);
+	bool GetInvertCut() const
+	{
+		return bInvertCut;
+	}
+
+	AVALANCHEMODIFIERS_API void SetFillHoles(bool bInFillHoles);
+	bool GetFillHoles() const
+	{
+		return bFillHoles;
+	}
+
+protected:
 	//~ Begin UObject
 #if WITH_EDITOR
 	virtual void PostEditChangeProperty(FPropertyChangedEvent& InPropertyChangedEvent) override;
 #endif
 	//~ End UObject
 
-	UFUNCTION()
-	void SetPlaneOrigin(float InOrigin);
-	float GetPlaneOrigin() const { return PlaneOrigin; }
-
-	UFUNCTION()
-	void SetPlaneRotation(const FRotator& InRotation);
-	const FRotator& GetPlaneRotation() const { return PlaneRotation; }
-
-	UFUNCTION()
-	void SetInvertCut(bool bInInvertCut);
-	bool GetInvertCut() const { return bInvertCut; }
-
-	UFUNCTION()
-	void SetFillHoles(bool bInFillHoles);
-	bool GetFillHoles() const { return bFillHoles; }
-
-protected:
 	//~ Begin UActorModifierCoreBase
 	virtual void OnModifierCDOSetup(FActorModifierCoreMetadata& InMetadata) override;
 	virtual void OnModifierAdded(EActorModifierCoreEnableReason InReason) override;
@@ -43,10 +51,10 @@ protected:
 	virtual void OnModifierDisabled(EActorModifierCoreDisableReason InReason) override;
 	virtual void OnModifierRemoved(EActorModifierCoreDisableReason InReason) override;
 	//~ End UActorModifierCoreBase
-	
+
 	/** Returns actual location of plane bounds restricted */
 	FVector GetPlaneLocation() const;
-	
+
 	void OnPlaneRotationChanged();
 	void OnFillHolesChanged();
 	void OnInvertCutChanged();
@@ -67,7 +75,7 @@ protected:
 
 	UPROPERTY(EditInstanceOnly, BlueprintReadWrite, Setter="SetInvertCut", Getter="GetInvertCut", Category="PlaneCut", meta=(AllowPrivateAccess="true"))
 	bool bInvertCut = false;
-	
+
 	UPROPERTY(EditInstanceOnly, BlueprintReadWrite, Setter="SetFillHoles", Getter="GetFillHoles", Category="PlaneCut", meta=(AllowPrivateAccess="true"))
 	bool bFillHoles = true;
 
diff --git a/Engine/Plugins/Experimental/Avalanche/Source/AvalancheModifiers/Public/Modifiers/AvaRadialArrangeModifier.h b/Engine/Plugins/Experimental/Avalanche/Source/AvalancheModifiers/Public/Modifiers/AvaRadialArrangeModifier.h
--- a/Engine/Plugins/Experimental/Avalanche/Source/AvalancheModifiers/Public/Modifiers/AvaRadialArrangeModifier.h
+++ b/Engine/Plugins/Experimental/Avalanche/Source/AvalancheModifiers/Public/Modifiers/AvaRadialArrangeModifier.h
@@ -24,112 +24,145 @@ enum class EAvaRadialArrangement : uint8
 /**
  * Arranges child actors in a circular rings around its center
  */
-UCLASS(BlueprintType)
-class AVALANCHEMODIFIERS_API UAvaRadialArrangeModifier : public UAvaArrangeBaseModifier
+UCLASS(MinimalAPI, BlueprintType)
+class UAvaRadialArrangeModifier : public UAvaArrangeBaseModifier
 {
 	GENERATED_BODY()
-	
+
 public:
-	//~ Begin UObject
-#if WITH_EDITOR
-	virtual void PostEditChangeProperty(FPropertyChangedEvent& PropertyChangedEvent) override;
-#endif
-	//~ End UObject
+	/** Sets the number of child elements to use in the arrangement. Children whose index is greater than or equal to this value will be hidden. */
+	UFUNCTION(BlueprintCallable, Category = "Motion Design|Modifiers|RadialArrange")
+	AVALANCHEMODIFIERS_API void SetCount(const int32 NewCount);
 
 	/** Gets the number of child elements to use in the arrangement. Children whose index is greater than or equal to this value will be hidden. */
 	UFUNCTION(BlueprintPure, Category = "Motion Design|Modifiers|RadialArrange")
-	int32 GetCount() const { return Count; }
+	int32 GetCount() const
+	{
+		return Count;
+	}
 
-	/** Sets the number of child elements to use in the arrangement. Children whose index is greater than or equal to this value will be hidden. */
+	/** Sets the number of rings. */
 	UFUNCTION(BlueprintCallable, Category = "Motion Design|Modifiers|RadialArrange")
-	void SetCount(const int32 NewCount);
+	AVALANCHEMODIFIERS_API void SetRings(const int32 NewRings);
 
 	/** Gets the number of rings. */
 	UFUNCTION(BlueprintPure, Category = "Motion Design|Modifiers|RadialArrange")
-	int32 GetRings() const { return Rings; }
+	int32 GetRings() const
+	{
+		return Rings;
+	}
 
-	/** Sets the number of rings. */
+	/** Sets the radius from the center to the first inner ring. */
 	UFUNCTION(BlueprintCallable, Category = "Motion Design|Modifiers|RadialArrange")
-	void SetRings(const int32 NewRings);
+	AVALANCHEMODIFIERS_API void SetInnerRadius(const float NewInnerRadius);
 
 	/** Gets the radius from the center to the first inner ring. */
 	UFUNCTION(BlueprintPure, Category = "Motion Design|Modifiers|RadialArrange")
-	float GetInnerRadius() const { return InnerRadius; }
+	float GetInnerRadius() const
+	{
+		return InnerRadius;
+	}
 
-	/** Sets the radius from the center to the first inner ring. */
+	/** Sets the radius from the center to the last outer ring. */
 	UFUNCTION(BlueprintCallable, Category = "Motion Design|Modifiers|RadialArrange")
-	void SetInnerRadius(const float NewInnerRadius);
+	AVALANCHEMODIFIERS_API void SetOuterRadius(const float NewOuterRadius);
 
 	/** Gets the radius from the center to the last outer ring. */
 	UFUNCTION(BlueprintPure, Category = "Motion Design|Modifiers|RadialArrange")
-	float GetOuterRadius() const { return OuterRadius; }
+	float GetOuterRadius() const
+	{
+		return OuterRadius;
+	}
 
-	/** Sets the radius from the center to the last outer ring. */
+	/** Sets the start angle for the arrangement space and moving clockwise. 0 = Up, -90 = Left, 90 = Right */
 	UFUNCTION(BlueprintCallable, Category = "Motion Design|Modifiers|RadialArrange")
-	void SetOuterRadius(const float NewOuterRadius);
+	AVALANCHEMODIFIERS_API void SetStartAngle(const float NewStartAngle);
 
 	/** Gets the start angle for the arrangement space and moving clockwise. 0 = Up, -90 = Left, 90 = Right */
 	UFUNCTION(BlueprintPure, Category = "Motion Design|Modifiers|RadialArrange")
-	float GetStartAngle() const { return StartAngle; }
+	float GetStartAngle() const
+	{
+		return StartAngle;
+	}
 
-	/** Sets the start angle for the arrangement space and moving clockwise. 0 = Up, -90 = Left, 90 = Right */
+	/** Sets the end angle for the arrangement space. 0 = Up, -90 = Left, 90 = Right */
 	UFUNCTION(BlueprintCallable, Category = "Motion Design|Modifiers|RadialArrange")
-	void SetStartAngle(const float NewStartAngle);
+	AVALANCHEMODIFIERS_API void SetEndAngle(const float NewEndAngle);
 
 	/** Gets the end angle for the arrangement space. 0 = Up, -90 = Left, 90 = Right */
 	UFUNCTION(BlueprintPure, Category = "Motion Design|Modifiers|RadialArrange")
-	float GetEndAngle() const { return EndAngle; }
+	float GetEndAngle() const
+	{
+		return EndAngle;
+	}
 
-	/** Sets the end angle for the arrangement space. 0 = Up, -90 = Left, 90 = Right */
+	/** Defines how to arrange the child elements around the center. */
 	UFUNCTION(BlueprintCallable, Category = "Motion Design|Modifiers|RadialArrange")
-	void SetEndAngle(const float NewEndAngle);
+	AVALANCHEMODIFIERS_API void SetArrangement(const EAvaRadialArrangement NewArrangement);
 
 	/** Defines how to arrange the child elements around the center. */
 	UFUNCTION(BlueprintPure, Category = "Motion Design|Modifiers|RadialArrange")
-	EAvaRadialArrangement GetArrangement() const { return Arrangement; }
+	EAvaRadialArrangement GetArrangement() const
+	{
+		return Arrangement;
+	}
 
-	/** Defines how to arrange the child elements around the center. */
+	/** If true, will arrange the child elements starting from the outer radius and moving to the inner radius. Has no effect if only using one ring. */
 	UFUNCTION(BlueprintCallable, Category = "Motion Design|Modifiers|RadialArrange")
-	void SetArrangement(const EAvaRadialArrangement NewArrangement);
+	AVALANCHEMODIFIERS_API void SetStartFromOuterRadius(const bool bMewStartFromOuterRadius);
 
 	/** If true, will arrange the child elements starting from the outer radius and moving to the inner radius. Has no effect if only using one ring. */
 	UFUNCTION(BlueprintPure, Category = "Motion Design|Modifiers|RadialArrange")
-	bool GetStartFromOuterRadius() const { return bStartFromOuterRadius; }
+	bool GetStartFromOuterRadius() const
+	{
+		return bStartFromOuterRadius;
+	}
 
-	/** If true, will arrange the child elements starting from the outer radius and moving to the inner radius. Has no effect if only using one ring. */
+	/** If true, will orient the selected axis torwards the center. */
 	UFUNCTION(BlueprintCallable, Category = "Motion Design|Modifiers|RadialArrange")
-	void SetStartFromOuterRadius(const bool bMewStartFromOuterRadius);
+	AVALANCHEMODIFIERS_API void SetOrient(const bool bNewOrient);
 
 	/** If true, will orient the selected axis torwards the center. */
 	UFUNCTION(BlueprintPure, Category = "Motion Design|Modifiers|RadialArrange")
-	bool GetOrient() const { return bOrient; }
+	bool GetOrient() const
+	{
+		return bOrient;
+	}
 
-	/** If true, will orient the selected axis torwards the center. */
+	/** Sets the axis to look at the center. */
 	UFUNCTION(BlueprintCallable, Category = "Motion Design|Modifiers|RadialArrange")
-	void SetOrient(const bool bNewOrient);
+	AVALANCHEMODIFIERS_API void SetOrientAxis(const EAvaAxis NewOrientAxis);
 
 	/** Gets the axis to look at the center. */
 	UFUNCTION(BlueprintPure, Category = "Motion Design|Modifiers|RadialArrange")
-	EAvaAxis GetOrientAxis() const { return OrientAxis; }
+	EAvaAxis GetOrientAxis() const
+	{
+		return OrientAxis;
+	}
 
-	/** Sets the axis to look at the center. */
+	/** If true, will flip the center orientation to face outwards. */
 	UFUNCTION(BlueprintCallable, Category = "Motion Design|Modifiers|RadialArrange")
-	void SetOrientAxis(const EAvaAxis NewOrientAxis);
+	AVALANCHEMODIFIERS_API void SetFlipOrient(const bool bNewFlipOrient);
 
 	/** If true, will flip the center orientation to face outwards. */
 	UFUNCTION(BlueprintPure, Category = "Motion Design|Modifiers|RadialArrange")
-	bool GetFlipOrient() const { return bFlipOrient; }
-
-	/** If true, will flip the center orientation to face outwards. */
-	UFUNCTION(BlueprintCallable, Category = "Motion Design|Modifiers|RadialArrange")
-	void SetFlipOrient(const bool bNewFlipOrient);
+	bool GetFlipOrient() const
+	{
+		return bFlipOrient;
+	}
 
 protected:
+	//~ Begin UObject
+#if WITH_EDITOR
+	virtual void PostEditChangeProperty(FPropertyChangedEvent& PropertyChangedEvent) override;
+#endif
+	//~ End UObject
+
 	//~ Begin UActorModifierCoreBase
 	virtual void OnModifierCDOSetup(FActorModifierCoreMetadata& InMetadata) override;
 	virtual void Apply() override;
 	//~ End UActorModifierCoreBase
-	
+
 	/** The number of child elements to limit in the arrangement, or -1 if unlimited. Children whose index is greater than or equal to this value will be hidden. */
 	UPROPERTY(EditAnywhere, BlueprintReadWrite, Setter="SetCount", Getter="GetCount", Category="RadialArrange", meta=(ClampMin="-1", UIMin="-1", AllowPrivateAccess="true"))
 	int32 Count = -1;
@@ -145,7 +178,7 @@ protected:
 	/** The radius from the center to the last outer ring. */
 	UPROPERTY(EditAnywhere, BlueprintReadWrite, Setter="SetOuterRadius", Getter="GetOuterRadius", Interp, Category="RadialArrange", meta=(EditCondition="Rings > 1", EditConditionHides, ClampMin="0.0", UIMin="0.0", AllowPrivateAccess="true"))
 	float OuterRadius = 200.0f;
-	
+
 	/** The start angle for the arrangement space and moving clockwise. 0 = Up, -90 = Left, 90 = Right */
 	UPROPERTY(EditAnywhere, BlueprintReadWrite, Setter="SetStartAngle", Getter="GetStartAngle", Interp, Category="RadialArrange", meta=(ClampMin="-180.0", UIMin="-180.0", ClampMax="180.0", UIMax="180.0", AllowPrivateAccess="true"))
 	float StartAngle = -180.0f;
diff --git a/Engine/Plugins/Experimental/Avalanche/Source/AvalancheModifiers/Public/Modifiers/AvaSplineSweepModifier.h b/Engine/Plugins/Experimental/Avalanche/Source/AvalancheModifiers/Public/Modifiers/AvaSplineSweepModifier.h
--- a/Engine/Plugins/Experimental/Avalanche/Source/AvalancheModifiers/Public/Modifiers/AvaSplineSweepModifier.h
+++ b/Engine/Plugins/Experimental/Avalanche/Source/AvalancheModifiers/Public/Modifiers/AvaSplineSweepModifier.h
@@ -18,8 +18,8 @@ enum class EAvaSplineSweepSampleMode : uint8
 	CustomDistance
 };
 
-UCLASS(BlueprintType)
-class AVALANCHEMODIFIERS_API UAvaSplineSweepModifier : public UAvaGeometryBaseModifier
+UCLASS(MinimalAPI, BlueprintType)
+class UAvaSplineSweepModifier : public UAvaGeometryBaseModifier
 	, public IAvaRenderStateUpdateHandler
 {
 	GENERATED_BODY()
@@ -27,20 +27,14 @@ class AVALANCHEMODIFIERS_API UAvaSplineSweepModifier : public UAvaGeometryBaseMo
 public:
 	static inline constexpr int32 MaxSampleCount = 100;
 
-	//~ Begin UObject
-#if WITH_EDITOR
-	virtual void PostEditChangeProperty(FPropertyChangedEvent& InPropertyChangedEvent) override;
-#endif
-	//~ End UObject
-
-	void SetSplineActorWeak(TWeakObjectPtr<AActor> InSplineActorWeak);
+	AVALANCHEMODIFIERS_API void SetSplineActorWeak(TWeakObjectPtr<AActor> InSplineActorWeak);
 	TWeakObjectPtr<AActor> GetSplineActorWeak() const
 	{
 		return SplineActorWeak;
 	}
 
 	UFUNCTION(BlueprintCallable, Category = "Motion Design|Modifiers|SplineSweep")
-	void SetSampleMode(EAvaSplineSweepSampleMode InMode);
+	AVALANCHEMODIFIERS_API void SetSampleMode(EAvaSplineSweepSampleMode InMode);
 
 	UFUNCTION(BlueprintPure, Category = "Motion Design|Modifiers|SplineSweep")
 	EAvaSplineSweepSampleMode GetSampleMode() const
@@ -49,7 +43,7 @@ public:
 	}
 
 	UFUNCTION(BlueprintCallable, Category = "Motion Design|Modifiers|SplineSweep")
-	void SetSampleDistance(float InDistance);
+	AVALANCHEMODIFIERS_API void SetSampleDistance(float InDistance);
 
 	UFUNCTION(BlueprintPure, Category = "Motion Design|Modifiers|SplineSweep")
 	float GetSampleDistance() const
@@ -58,7 +52,7 @@ public:
 	}
 
 	UFUNCTION(BlueprintCallable, Category = "Motion Design|Modifiers|SplineSweep")
-	void SetSteps(int32 InSteps);
+	AVALANCHEMODIFIERS_API void SetSteps(int32 InSteps);
 
 	UFUNCTION(BlueprintPure, Category = "Motion Design|Modifiers|SplineSweep")
 	int32 GetSteps() const
@@ -67,7 +61,7 @@ public:
 	}
 
 	UFUNCTION(BlueprintCallable, Category = "Motion Design|Modifiers|SplineSweep")
-	void SetProgressOffset(float InOffset);
+	AVALANCHEMODIFIERS_API void SetProgressOffset(float InOffset);
 
 	UFUNCTION(BlueprintPure, Category = "Motion Design|Modifiers|SplineSweep")
 	float GetProgressOffset() const
@@ -76,7 +70,7 @@ public:
 	}
 
 	UFUNCTION(BlueprintCallable, Category = "Motion Design|Modifiers|SplineSweep")
-	void SetProgressStart(float InStart);
+	AVALANCHEMODIFIERS_API void SetProgressStart(float InStart);
 
 	UFUNCTION(BlueprintPure, Category = "Motion Design|Modifiers|SplineSweep")
 	float GetProgressStart() const
@@ -85,7 +79,7 @@ public:
 	}
 
 	UFUNCTION(BlueprintCallable, Category = "Motion Design|Modifiers|SplineSweep")
-	void SetProgressEnd(float InEnd);
+	AVALANCHEMODIFIERS_API void SetProgressEnd(float InEnd);
 
 	UFUNCTION(BlueprintPure, Category = "Motion Design|Modifiers|SplineSweep")
 	float GetProgressEnd() const
@@ -94,7 +88,7 @@ public:
 	}
 
 	UFUNCTION(BlueprintCallable, Category = "Motion Design|Modifiers|SplineSweep")
-	void SetScaleStart(float InScaleStart);
+	AVALANCHEMODIFIERS_API void SetScaleStart(float InScaleStart);
 
 	UFUNCTION(BlueprintPure, Category = "Motion Design|Modifiers|SplineSweep")
 	float GetScaleStart() const
@@ -103,7 +97,7 @@ public:
 	}
 
 	UFUNCTION(BlueprintCallable, Category = "Motion Design|Modifiers|SplineSweep")
-	void SetScaleEnd(float InScaleEnd);
+	AVALANCHEMODIFIERS_API void SetScaleEnd(float InScaleEnd);
 
 	UFUNCTION(BlueprintPure, Category = "Motion Design|Modifiers|SplineSweep")
 	float GetScaleEnd() const
@@ -112,7 +106,7 @@ public:
 	}
 
 	UFUNCTION(BlueprintCallable, Category = "Motion Design|Modifiers|SplineSweep")
-	void SetCapped(bool bInCapped);
+	AVALANCHEMODIFIERS_API void SetCapped(bool bInCapped);
 
 	UFUNCTION(BlueprintPure, Category = "Motion Design|Modifiers|SplineSweep")
 	bool GetCapped() const
@@ -121,7 +115,7 @@ public:
 	}
 
 	UFUNCTION(BlueprintCallable, Category = "Motion Design|Modifiers|SplineSweep")
-	void SetLooped(bool bInLooped);
+	AVALANCHEMODIFIERS_API void SetLooped(bool bInLooped);
 
 	UFUNCTION(BlueprintPure, Category = "Motion Design|Modifiers|SplineSweep")
 	bool GetLooped() const
@@ -130,6 +124,12 @@ public:
 	}
 
 protected:
+	//~ Begin UObject
+#if WITH_EDITOR
+	virtual void PostEditChangeProperty(FPropertyChangedEvent& InPropertyChangedEvent) override;
+#endif
+	//~ End UObject
+
 	//~ Begin IAvaRenderStateUpdateHandler
 	virtual void OnRenderStateUpdated(AActor* InActor, UActorComponent* InComponent) override;
 	virtual void OnActorVisibilityChanged(AActor* InActor) override {}
diff --git a/Engine/Plugins/Experimental/Avalanche/Source/AvalancheModifiers/Public/Modifiers/AvaSubdivideModifier.h b/Engine/Plugins/Experimental/Avalanche/Source/AvalancheModifiers/Public/Modifiers/AvaSubdivideModifier.h
--- a/Engine/Plugins/Experimental/Avalanche/Source/AvalancheModifiers/Public/Modifiers/AvaSubdivideModifier.h
+++ b/Engine/Plugins/Experimental/Avalanche/Source/AvalancheModifiers/Public/Modifiers/AvaSubdivideModifier.h
@@ -13,8 +13,8 @@ enum class EAvaSubdivisionType : uint8
 	PN
 };
 
-UCLASS(BlueprintType)
-class AVALANCHEMODIFIERS_API UAvaSubdivideModifier : public UAvaGeometryBaseModifier
+UCLASS(MinimalAPI, BlueprintType)
+class UAvaSubdivideModifier : public UAvaGeometryBaseModifier
 {
 	GENERATED_BODY()
 
@@ -22,38 +22,44 @@ public:
 	static constexpr int MinSubdivideCuts = 1;
     static constexpr int MaxSubdivideCuts = 15;
 
+	AVALANCHEMODIFIERS_API void SetCuts(int32 InCuts);
+	int32 GetCuts() const
+	{
+		return Cuts;
+	}
+
+	AVALANCHEMODIFIERS_API void SetRecomputeNormals(bool bInRecomputeNormals);
+	bool GetRecomputeNormals() const
+	{
+		return bRecomputeNormals;
+	}
+
+	AVALANCHEMODIFIERS_API void SetType(EAvaSubdivisionType InType);
+	EAvaSubdivisionType GetType() const
+	{
+		return Type;
+	}
+
+protected:
 	//~ Begin UObject
 #if WITH_EDITOR
 	virtual void PostEditChangeProperty(FPropertyChangedEvent& PropertyChangedEvent) override;
 #endif
 	//~ End UObject
-	
-	UFUNCTION()
-	void SetCuts(int32 InCuts);
-	int32 GetCuts() const { return Cuts; }
-	
-	UFUNCTION()
-	void SetRecomputeNormals(bool bInRecomputeNormals);
-	bool GetRecomputeNormals() const { return bRecomputeNormals; }
-
-	UFUNCTION()
-	void SetType(EAvaSubdivisionType InType);
-	EAvaSubdivisionType GetType() const { return Type; }
 
-protected:
 	//~ Begin UActorModifierCoreBase
 	virtual void OnModifierCDOSetup(FActorModifierCoreMetadata& InMetadata) override;
 	virtual void Apply() override;
 	//~ End UActorModifierCoreBase
-	
+
 	void OnOptionsChanged();
-	
+
 	UPROPERTY(EditInstanceOnly, BlueprintReadWrite, Setter="SetCuts", Getter="GetCuts", Category="Subdivide", meta=(ClampMin="1", ClampMax="15", AllowPrivateAccess="true"))
 	int32 Cuts = 2;
-	
+
 	UPROPERTY(EditInstanceOnly, BlueprintReadWrite, Setter="SetRecomputeNormals", Getter="GetRecomputeNormals", Category="Subdivide", meta=(EditCondition="Type == EAvaSubdivisionType::PN", EditConditionHides, AllowPrivateAccess="true"))
 	bool bRecomputeNormals = true;
-	
+
 	UPROPERTY(EditInstanceOnly, BlueprintReadWrite, Setter="SetType", Getter="GetType", Category="Subdivide", meta=(AllowPrivateAccess="true"))
 	EAvaSubdivisionType Type = EAvaSubdivisionType::Uniform;
 };
diff --git a/Engine/Plugins/Experimental/Avalanche/Source/AvalancheModifiers/Public/Modifiers/AvaTaperModifier.h b/Engine/Plugins/Experimental/Avalanche/Source/AvalancheModifiers/Public/Modifiers/AvaTaperModifier.h
--- a/Engine/Plugins/Experimental/Avalanche/Source/AvalancheModifiers/Public/Modifiers/AvaTaperModifier.h
+++ b/Engine/Plugins/Experimental/Avalanche/Source/AvalancheModifiers/Public/Modifiers/AvaTaperModifier.h
@@ -20,8 +20,8 @@ enum class EAvaTaperExtent : uint8
 	Custom
 };
 
-UCLASS(BlueprintType)
-class AVALANCHEMODIFIERS_API UAvaTaperModifier : public UAvaGeometryBaseModifier
+UCLASS(MinimalAPI, BlueprintType)
+class UAvaTaperModifier : public UAvaGeometryBaseModifier
 {
 	GENERATED_BODY()
 
@@ -29,65 +29,81 @@ public:
 	static constexpr int MinTaperLatticeResolution = 1;
 	static constexpr int MaxTaperLatticeResolution = 20;
 
+	AVALANCHEMODIFIERS_API void SetAmount(float InAmount);
+	float GetAmount() const
+	{
+		return Amount;
+	}
+
+	AVALANCHEMODIFIERS_API void SetUpperExtent(float InUpperExtent);
+	float GetUpperExtent() const
+	{
+		return UpperExtent;
+	}
+
+	AVALANCHEMODIFIERS_API void SetLowerExtent(float InLowerExtent);
+	float GetLowerExtent() const
+	{
+		return LowerExtent;
+	}
+
+	AVALANCHEMODIFIERS_API void SetExtent(EAvaTaperExtent InExtent);
+	EAvaTaperExtent GetExtent() const
+	{
+		return Extent;
+	}
+
+	AVALANCHEMODIFIERS_API void SetInterpolationType(EAvaTaperInterpolationType InInterpolationType);
+	EAvaTaperInterpolationType GetInterpolationType() const
+	{
+		return InterpolationType;
+	}
+
+	AVALANCHEMODIFIERS_API void SetReferenceFrame(EAvaTaperReferenceFrame InReferenceFrame);
+	EAvaTaperReferenceFrame GetReferenceFrame() const
+	{
+		return ReferenceFrame;
+	}
+
+	AVALANCHEMODIFIERS_API void SetResolution(int32 InResolution);
+	int32 GetResolution() const
+	{
+		return Resolution;
+	}
+
+	AVALANCHEMODIFIERS_API void SetOffset(FVector2D InOffset);
+	FVector2D GetOffset() const
+	{
+		return Offset;
+	}
+
+protected:
 	//~ Begin UObject
 #if WITH_EDITOR
 	virtual void PostEditChangeProperty(FPropertyChangedEvent& PropertyChangedEvent) override;
 #endif
 	//~ End UObject
-	
-	UFUNCTION()
-	void SetAmount(float InAmount);
-	float GetAmount() const { return Amount; }
-
-	UFUNCTION()
-	void SetUpperExtent(float InUpperExtent);
-	float GetUpperExtent() const { return UpperExtent; }
-
-	UFUNCTION()
-	void SetLowerExtent(float InLowerExtent);
-	float GetLowerExtent() const { return LowerExtent; }
-
-	UFUNCTION()
-	void SetExtent(EAvaTaperExtent InExtent);
-	EAvaTaperExtent GetExtent() const { return Extent; }
-	
-	UFUNCTION()
-	void SetInterpolationType(EAvaTaperInterpolationType InInterpolationType);
-	EAvaTaperInterpolationType GetInterpolationType() const { return InterpolationType; }
-
-	UFUNCTION()
-	void SetReferenceFrame(EAvaTaperReferenceFrame InReferenceFrame);
-	EAvaTaperReferenceFrame GetReferenceFrame() const { return ReferenceFrame; }
-
-	UFUNCTION()
-	void SetResolution(int32 InResolution);
-	int32 GetResolution() const { return Resolution; }
-
-	UFUNCTION()
-	void SetOffset(FVector2D InOffset);
-	FVector2D GetOffset() const { return Offset; }
-	
-protected:
+
 	//~ Begin UActorModifierCoreBase
 	virtual void OnModifierCDOSetup(FActorModifierCoreMetadata& InMetadata) override;
 	virtual void Apply() override;
 	//~ End UActorModifierCoreBase
-	
+
 	void CreateTaperTool();
-	
+
 	void OnParameterChanged();
-	
+
 	FVector2D GetRequiredOffset() const;
 	FVector2D GetRequiredExtent() const;
 
 	int32 GetSubdividersCuts() const;
-	
+
 	UPROPERTY(EditInstanceOnly, BlueprintReadWrite, Setter="SetAmount", Getter="GetAmount", Category="Taper", meta=(ClampMin="0.0", ClampMax="1.0", AllowPrivateAccess="true"))
 	float Amount = 0.0;
 
 	UPROPERTY(EditInstanceOnly, BlueprintReadWrite, Setter="SetExtent", Getter="GetExtent", Category="Taper", meta=(AllowPrivateAccess="true"))
 	EAvaTaperExtent Extent = EAvaTaperExtent::WholeShape;
-	
+
 	UPROPERTY(EditInstanceOnly, BlueprintReadWrite, Setter="SetUpperExtent", Getter="GetUpperExtent", Category="Taper", meta=(EditCondition="Extent == EAvaTaperExtent::Custom", EditConditionHides, ClampMin="0", ClampMax="100", Units="Percent", ToolTip="100%: shape top.\n0%: shape bottom.", AllowPrivateAccess="true"))
 	float UpperExtent = 100;
 
@@ -96,13 +112,13 @@ protected:
 
 	UPROPERTY(EditInstanceOnly, BlueprintReadWrite, Setter="SetInterpolationType", Getter="GetInterpolationType", Category="Taper", meta=(AllowPrivateAccess="true"))
 	EAvaTaperInterpolationType InterpolationType = EAvaTaperInterpolationType::Linear;
-	
+
 	UPROPERTY(EditInstanceOnly, BlueprintReadWrite, Setter="SetResolution", Getter="GetResolution", Category="Taper", meta=(ClampMin="1", ClampMax="20", ToolTip="The number of vertical control points used to apply the taper. If the modifier is in a stack with Subdivide modifiers, taper will use the max value between Resolution and the total subdivision cuts.", AllowPrivateAccess="true"))
 	int32 Resolution = 5;
 
 	UPROPERTY(EditInstanceOnly, BlueprintReadWrite, Setter="SetReferenceFrame", Getter="GetReferenceFrame", Category="Taper", meta=(AllowPrivateAccess="true"))
 	EAvaTaperReferenceFrame ReferenceFrame = EAvaTaperReferenceFrame::MeshCenter;
-	
+
 	UPROPERTY(EditInstanceOnly, BlueprintReadWrite, Setter="SetOffset", Getter="GetOffset", Category="Taper", meta=(EditCondition="ReferenceFrame == EAvaTaperReferenceFrame::Custom", EditConditionHides, AllowPrivateAccess="true"))
 	FVector2D Offset = FVector2D::ZeroVector;
 
diff --git a/Engine/Plugins/Experimental/Avalanche/Source/AvalancheModifiers/Public/Modifiers/AvaTranslucentPriorityModifier.h b/Engine/Plugins/Experimental/Avalanche/Source/AvalancheModifiers/Public/Modifiers/AvaTranslucentPriorityModifier.h
--- a/Engine/Plugins/Experimental/Avalanche/Source/AvalancheModifiers/Public/Modifiers/AvaTranslucentPriorityModifier.h
+++ b/Engine/Plugins/Experimental/Avalanche/Source/AvalancheModifiers/Public/Modifiers/AvaTranslucentPriorityModifier.h
@@ -21,58 +21,58 @@ enum class EAvaTranslucentPriorityModifierMode : uint8
 	Manual
 };
 
-UCLASS(BlueprintType)
-class AVALANCHEMODIFIERS_API UAvaTranslucentPriorityModifier : public UAvaArrangeBaseModifier
+UCLASS(MinimalAPI, BlueprintType)
+class UAvaTranslucentPriorityModifier : public UAvaArrangeBaseModifier
 {
 	GENERATED_BODY()
 
 	friend class UAvaTranslucentPriorityModifierShared;
 
 public:
-	//~ Begin UObject
-	virtual void PostLoad() override;
-#if WITH_EDITOR
-	virtual void PostEditChangeProperty(FPropertyChangedEvent& PropertyChangedEvent) override;
-#endif
-	//~ End UObject
-
-    void SetMode(EAvaTranslucentPriorityModifierMode InMode);
+    AVALANCHEMODIFIERS_API void SetMode(EAvaTranslucentPriorityModifierMode InMode);
     EAvaTranslucentPriorityModifierMode GetMode() const
     {
 	    return Mode;
     }
 
-    void SetCameraActorWeak(const TWeakObjectPtr<ACameraActor>& InCameraActor);
+    AVALANCHEMODIFIERS_API void SetCameraActorWeak(const TWeakObjectPtr<ACameraActor>& InCameraActor);
 	TWeakObjectPtr<ACameraActor> GetCameraActorWeak() const
 	{
 		return CameraActorWeak;
 	}
 
-    void SetSortPriority(int32 InSortPriority);
+    AVALANCHEMODIFIERS_API void SetSortPriority(int32 InSortPriority);
     int32 GetSortPriority() const
     {
 	    return SortPriority;
     }
 
-	void SetSortPriorityOffset(int32 InOffset);
+	AVALANCHEMODIFIERS_API void SetSortPriorityOffset(int32 InOffset);
 	int32 GetSortPriorityOffset() const
 	{
 		return SortPriorityOffset;
 	}
 
-	void SetSortPriorityStep(int32 InStep);
+	AVALANCHEMODIFIERS_API void SetSortPriorityStep(int32 InStep);
 	int32 GetSortPriorityStep() const
 	{
 		return SortPriorityStep;
 	}
 
-	void SetIncludeChildren(bool bInIncludeChildren);
+	AVALANCHEMODIFIERS_API void SetIncludeChildren(bool bInIncludeChildren);
 	bool GetIncludeChildren() const
 	{
 		return bIncludeChildren;
 	}
 
 protected:
+	//~ Begin UObject
+	virtual void PostLoad() override;
+#if WITH_EDITOR
+	virtual void PostEditChangeProperty(FPropertyChangedEvent& PropertyChangedEvent) override;
+#endif
+	//~ End UObject
+
 	//~ Begin UActorModifierCoreBase
 	virtual void OnModifierCDOSetup(FActorModifierCoreMetadata& InMetadata) override;
 	virtual void OnModifierAdded(EActorModifierCoreEnableReason InReason) override;
diff --git a/Engine/Plugins/Experimental/Avalanche/Source/AvalancheModifiers/Public/Modifiers/AvaVisibilityModifier.h b/Engine/Plugins/Experimental/Avalanche/Source/AvalancheModifiers/Public/Modifiers/AvaVisibilityModifier.h
--- a/Engine/Plugins/Experimental/Avalanche/Source/AvalancheModifiers/Public/Modifiers/AvaVisibilityModifier.h
+++ b/Engine/Plugins/Experimental/Avalanche/Source/AvalancheModifiers/Public/Modifiers/AvaVisibilityModifier.h
@@ -10,44 +10,53 @@ class AActor;
 /**
  * Controls the visibility of a range of child actors by index.
  */
-UCLASS(BlueprintType)
-class AVALANCHEMODIFIERS_API UAvaVisibilityModifier : public UAvaArrangeBaseModifier
+UCLASS(MinimalAPI, BlueprintType)
+class UAvaVisibilityModifier : public UAvaArrangeBaseModifier
 {
 	GENERATED_BODY()
-	
+
 public:
-	//~ Begin UObject
-#if WITH_EDITOR
-	virtual void PostEditChangeProperty(FPropertyChangedEvent& PropertyChangedEvent) override;
-#endif
-	//~ End UObject
+	/** Sets the child index range to hide instead of showing. */
+	UFUNCTION(BlueprintCallable, Category = "Motion Design|Modifiers|Visibility")
+	AVALANCHEMODIFIERS_API void SetInvertVisibility(const bool bNewInvertVisibility);
 
 	/** Returns true if hiding the child index range instead of showing. */
 	UFUNCTION(BlueprintPure, Category = "Motion Design|Modifiers|Visibility")
-	bool GetInvertVisibility() const { return bInvertVisibility; }
+	bool GetInvertVisibility() const
+	{
+		return bInvertVisibility;
+	}
 
-	/** Sets the child index range to hide instead of showing. */
 	UFUNCTION(BlueprintCallable, Category = "Motion Design|Modifiers|Visibility")
-	void SetInvertVisibility(const bool bNewInvertVisibility);
+	AVALANCHEMODIFIERS_API void SetIndex(int32 InIndex);
 
 	UFUNCTION(BlueprintPure, Category = "Motion Design|Modifiers|Visibility")
-	int32 GetIndex() const { return Index; }
-	
+	int32 GetIndex() const
+	{
+		return Index;
+	}
+
 	UFUNCTION(BlueprintCallable, Category = "Motion Design|Modifiers|Visibility")
-	void SetIndex(int32 InIndex);
-	
+	AVALANCHEMODIFIERS_API void SetTreatAsRange(const bool bInTreatAsRange);
+
 	UFUNCTION(BlueprintPure, Category = "Motion Design|Modifiers|Visibility")
-	bool GetTreatAsRange() const { return bTreatAsRange; }
-	
-	UFUNCTION(BlueprintCallable, Category = "Motion Design|Modifiers|Visibility")
-	void SetTreatAsRange(const bool bInTreatAsRange);
+	bool GetTreatAsRange() const
+	{
+		return bTreatAsRange;
+	}
 
 protected:
+	//~ Begin UObject
+#if WITH_EDITOR
+	virtual void PostEditChangeProperty(FPropertyChangedEvent& PropertyChangedEvent) override;
+#endif
+	//~ End UObject
+
 	//~ Begin UActorModifierCoreBase
 	virtual void OnModifierCDOSetup(FActorModifierCoreMetadata& InMetadata) override;
 	virtual void Apply() override;
 	//~ End UActorModifierCoreBase
-	
+
 	//~ Begin IAvaRenderStateUpdateExtension
 	virtual void OnRenderStateUpdated(AActor* InActor, UActorComponent* InComponent) override;
 	//~ End IAvaRenderStateUpdateExtension
@@ -69,7 +78,7 @@ protected:
 	/** Treat index as a range from 0 to index */
 	UPROPERTY(EditAnywhere, BlueprintReadWrite, Setter="SetTreatAsRange", Getter="GetTreatAsRange", Category="Visibility", meta=(AllowPrivateAccess="true"))
 	bool bTreatAsRange = false;
-	
+
 	/** If true, will hide the child index range instead of showing. */
 	UPROPERTY(EditAnywhere, BlueprintReadWrite, Setter="SetInvertVisibility", Getter="GetInvertVisibility", Category="Visibility", meta=(AllowPrivateAccess="true"))
 	bool bInvertVisibility = false;
diff --git a/Engine/Plugins/Experimental/Avalanche/Source/AvalancheModifiers/Public/Shared/AvaBooleanModifierShared.h b/Engine/Plugins/Experimental/Avalanche/Source/AvalancheModifiers/Public/Shared/AvaBooleanModifierShared.h
--- a/Engine/Plugins/Experimental/Avalanche/Source/AvalancheModifiers/Public/Shared/AvaBooleanModifierShared.h
+++ b/Engine/Plugins/Experimental/Avalanche/Source/AvalancheModifiers/Public/Shared/AvaBooleanModifierShared.h
@@ -7,7 +7,7 @@
 #include "AvaBooleanModifierShared.generated.h"
 
 class UAvaBooleanModifier;
-enum class EAvaBooleanMode : uint8; 
+enum class EAvaBooleanMode : uint8;
 
 USTRUCT()
 struct FAvaBooleanModifierSharedChannel
@@ -26,7 +26,7 @@ struct FAvaBooleanModifierSharedChannelInfo
 	/** The number of channel currently active */
 	UPROPERTY(VisibleAnywhere, BlueprintReadOnly, Category = "Channel")
 	int32 ChannelCount = 0;
-	
+
 	/** The number of modifier on that channel, the more there are, the more impact on performance */
 	UPROPERTY(VisibleAnywhere, BlueprintReadOnly, Category = "Channel")
 	int32 ChannelModifierCount = 0;
@@ -38,7 +38,7 @@ struct FAvaBooleanModifierSharedChannelInfo
 	/** The number of modifier on that channel that are masked by other mask modifier */
 	UPROPERTY(VisibleAnywhere, BlueprintReadOnly, Category = "Channel")
 	int32 ChannelTargetCount = 0;
-	
+
 	/** The number of modifier intersecting with this modifier */
 	UPROPERTY(VisibleAnywhere, BlueprintReadOnly, Category = "Channel")
 	int32 ChannelIntersectCount = 0;
@@ -49,10 +49,10 @@ struct FAvaBooleanModifierSharedChannelInfo
  * Transient because we can rebuild it, no need to save the asset
  */
 UCLASS(Transient)
-class AVALANCHEMODIFIERS_API UAvaBooleanModifierShared : public UActorModifierCoreSharedObject
+class UAvaBooleanModifierShared : public UActorModifierCoreSharedObject
 {
 	GENERATED_BODY()
-	
+
 public:
 	/** Track a modifier on a channel */
 	void TrackModifierChannel(UAvaBooleanModifier* InModifier);
@@ -65,20 +65,20 @@ public:
 
 	/** Get number of active channel */
 	uint8 GetChannelCount() const;
-	
+
 	/** Get number of modifier active on that channel */
 	int32 GetChannelModifierCount(uint8 InChannel) const;
 
 	/** Get number of modifier active on that channel with that mode set */
 	int32 GetChannelModifierModeCount(uint8 InChannel, EAvaBooleanMode InMode) const;
-	
+
 	/** Get all intersecting modifiers with the input target modifier */
 	TSet<TWeakObjectPtr<UAvaBooleanModifier>> GetIntersectingModifiers(const UAvaBooleanModifier* InTargetModifier, FAvaBooleanModifierSharedChannelInfo* OutDesc = nullptr);
 
 private:
 	/** Tests if two meshes are intersecting, always pass the tool/mask mesh first */
 	bool TestIntersection(const UE::Geometry::FDynamicMesh3& InToolMesh, const FTransform& InToolTransform, const UE::Geometry::FDynamicMesh3& InTargetMesh, const FTransform& InTargetTransform) const;
-	
+
 	UPROPERTY(Transient)
 	TMap<uint8, FAvaBooleanModifierSharedChannel> Channels;
 };
diff --git a/Engine/Plugins/Experimental/Avalanche/Source/AvalancheModifiers/Public/Shared/AvaTransformModifierShared.h b/Engine/Plugins/Experimental/Avalanche/Source/AvalancheModifiers/Public/Shared/AvaTransformModifierShared.h
--- a/Engine/Plugins/Experimental/Avalanche/Source/AvalancheModifiers/Public/Shared/AvaTransformModifierShared.h
+++ b/Engine/Plugins/Experimental/Avalanche/Source/AvalancheModifiers/Public/Shared/AvaTransformModifierShared.h
@@ -25,7 +25,7 @@ enum class EAvaTransformSharedModifier : uint8
 };
 
 USTRUCT()
-struct AVALANCHEMODIFIERS_API FAvaTransformSharedModifierState
+struct FAvaTransformSharedModifierState
 {
 	GENERATED_BODY()
 
@@ -39,7 +39,7 @@ struct AVALANCHEMODIFIERS_API FAvaTransformSharedModifierState
 
 	/** Restore this modifier state if valid */
 	void Restore(AActor* InActor, EAvaTransformSharedModifier InRestoreState) const;
-	
+
 	friend uint32 GetTypeHash(const FAvaTransformSharedModifierState& InItem)
 	{
 		return GetTypeHash(InItem.ModifierWeak);
@@ -52,13 +52,13 @@ struct AVALANCHEMODIFIERS_API FAvaTransformSharedModifierState
 
 	UPROPERTY()
 	TWeakObjectPtr<UAvaBaseModifier> ModifierWeak;
-	
+
 	UPROPERTY()
 	FTransform ActorTransform;
 };
 
 USTRUCT()
-struct AVALANCHEMODIFIERS_API FAvaTransformSharedActorState
+struct FAvaTransformSharedActorState
 {
 	GENERATED_BODY()
 
@@ -103,10 +103,10 @@ struct AVALANCHEMODIFIERS_API FAvaTransformSharedActorState
  * and restore it when no other modifier is watching it
  */
 UCLASS()
-class AVALANCHEMODIFIERS_API UAvaTransformModifierShared : public UActorModifierCoreSharedObject
+class UAvaTransformModifierShared : public UActorModifierCoreSharedObject
 {
 	GENERATED_BODY()
-	
+
 public:
 	/** Save actor state, adds it if it is not tracked */
 	void SaveActorState(UAvaBaseModifier* InModifierContext, AActor* InActor);
@@ -131,10 +131,12 @@ public:
 
 	/** Returns true, if this modifier is tracking any actor */
 	bool IsActorsStateSaved(UAvaBaseModifier* InModifierContext);
-	
+
 protected:
+	//~ Begin UObject
 	virtual void PostLoad() override;
-	
+	//~ End UObject
+
 	/** Actor state before any modifier applied to it */
 	UPROPERTY()
 	TSet<FAvaTransformSharedActorState> ActorStates;
diff --git a/Engine/Plugins/Experimental/Avalanche/Source/AvalancheModifiers/Public/Shared/AvaTranslucentPriorityModifierShared.h b/Engine/Plugins/Experimental/Avalanche/Source/AvalancheModifiers/Public/Shared/AvaTranslucentPriorityModifierShared.h
--- a/Engine/Plugins/Experimental/Avalanche/Source/AvalancheModifiers/Public/Shared/AvaTranslucentPriorityModifierShared.h
+++ b/Engine/Plugins/Experimental/Avalanche/Source/AvalancheModifiers/Public/Shared/AvaTranslucentPriorityModifierShared.h
@@ -62,7 +62,7 @@ struct FAvaTranslucentPriorityModifierComponentState
  * Singleton class for translucent priority modifiers to share data about component state
  */
 UCLASS()
-class AVALANCHEMODIFIERS_API UAvaTranslucentPriorityModifierShared : public UActorModifierCoreSharedObject
+class UAvaTranslucentPriorityModifierShared : public UActorModifierCoreSharedObject
 {
 	GENERATED_BODY()
 
@@ -91,20 +91,20 @@ public:
 	/** Get sorted components state based on modifier context */
 	TArray<const FAvaTranslucentPriorityModifierComponentState*> GetSortedComponentStates(UAvaTranslucentPriorityModifier* InModifierContext) const;
 
+	void SetSortPriorityOffset(int32 InOffset);
+
 	int32 GetSortPriorityOffset() const
 	{
 		return SortPriorityOffset;
 	}
 
-	void SetSortPriorityOffset(int32 InOffset);
+	void SetSortPriorityStep(int32 InStep);
 
 	int32 GetSortPriorityStep() const
 	{
 		return SortPriorityStep;
 	}
 
-	void SetSortPriorityStep(int32 InStep);
-
 private:
 	//~ Begin UObject
 	virtual void PostLoad() override;
diff --git a/Engine/Plugins/Experimental/Avalanche/Source/AvalancheModifiers/Public/Shared/AvaVisibilityModifierShared.h b/Engine/Plugins/Experimental/Avalanche/Source/AvalancheModifiers/Public/Shared/AvaVisibilityModifierShared.h
--- a/Engine/Plugins/Experimental/Avalanche/Source/AvalancheModifiers/Public/Shared/AvaVisibilityModifierShared.h
+++ b/Engine/Plugins/Experimental/Avalanche/Source/AvalancheModifiers/Public/Shared/AvaVisibilityModifierShared.h
@@ -22,7 +22,7 @@ enum class EAvaVisibilityActor : uint8
 ENUM_CLASS_FLAGS(EAvaVisibilityActor);
 
 USTRUCT()
-struct AVALANCHEMODIFIERS_API FAvaVisibilitySharedModifierState
+struct FAvaVisibilitySharedModifierState
 {
 	GENERATED_BODY()
 
@@ -36,7 +36,7 @@ struct AVALANCHEMODIFIERS_API FAvaVisibilitySharedModifierState
 
 	/** Restore this modifier state if valid */
 	void Restore(AActor* InActor) const;
-	
+
 	friend uint32 GetTypeHash(const FAvaVisibilitySharedModifierState& InItem)
 	{
 		return GetTypeHash(InItem.ModifierWeak);
@@ -62,7 +62,7 @@ struct AVALANCHEMODIFIERS_API FAvaVisibilitySharedModifierState
 };
 
 USTRUCT()
-struct AVALANCHEMODIFIERS_API FAvaVisibilitySharedActorState
+struct FAvaVisibilitySharedActorState
 {
 	GENERATED_BODY()
 
@@ -70,13 +70,13 @@ struct AVALANCHEMODIFIERS_API FAvaVisibilitySharedActorState
 	FAvaVisibilitySharedActorState(AActor* InActor)
 		: ActorWeak(InActor)
 	{}
-	
+
 	/** Save this actor state if valid */
 	void Save();
 
 	/** Restore this actor state if valid */
 	void Restore() const;
-	
+
 	friend uint32 GetTypeHash(const FAvaVisibilitySharedActorState& InItem)
 	{
 		return GetTypeHash(InItem.ActorWeak);
@@ -94,7 +94,7 @@ struct AVALANCHEMODIFIERS_API FAvaVisibilitySharedActorState
 	/** Actor that this state is describing */
 	UPROPERTY()
 	TWeakObjectPtr<AActor> ActorWeak;
-	
+
 #if WITH_EDITORONLY_DATA
 	/** Pre state editor visibility saved */
 	UPROPERTY()
@@ -113,10 +113,10 @@ struct AVALANCHEMODIFIERS_API FAvaVisibilitySharedActorState
  * and restore it when no other modifier is watching it
  */
 UCLASS()
-class AVALANCHEMODIFIERS_API UAvaVisibilityModifierShared : public UActorModifierCoreSharedObject
+class UAvaVisibilityModifierShared : public UActorModifierCoreSharedObject
 {
 	GENERATED_BODY()
-	
+
 public:
 	/** Watch actor state, adds it if it is not tracked */
 	void SaveActorState(UAvaBaseModifier* InModifierContext, AActor* InActor);
@@ -137,15 +137,18 @@ public:
 	void RestoreActorsState(UAvaBaseModifier* InModifierContext, const TSet<AActor*>* InActors = nullptr);
 
 	void RestoreActorsState(UAvaBaseModifier* InModifierContext, const TSet<TWeakObjectPtr<AActor>>& InActors);
-	
+
 	/** Returns true, if this modifier is tracking this actor */
 	bool IsActorStateSaved(UAvaBaseModifier* InModifierContext, AActor* InActor);
 
 	/** Returns true, if this modifier is tracking any actor */
 	bool IsActorsStateSaved(UAvaBaseModifier* InModifierContext);
+
 protected:
+	//~ Begin UObject
 	virtual void PostLoad() override;
-	
+	//~ End UObject
+
 	/** Actor state before any modifier applied to it */
 	UPROPERTY()
 	TSet<FAvaVisibilitySharedActorState> ActorStates;
