Motion Design:
- Viewport API cleanup

#changelist validated
#virtualized

diff --git a/Engine/Plugins/Experimental/Avalanche/Source/AvalancheLevelViewport/Private/SAvaLevelViewport_StatusBar.cpp b/Engine/Plugins/Experimental/Avalanche/Source/AvalancheLevelViewport/Private/SAvaLevelViewport_StatusBar.cpp
--- a/Engine/Plugins/Experimental/Avalanche/Source/AvalancheLevelViewport/Private/SAvaLevelViewport_StatusBar.cpp
+++ b/Engine/Plugins/Experimental/Avalanche/Source/AvalancheLevelViewport/Private/SAvaLevelViewport_StatusBar.cpp
@@ -4,8 +4,8 @@
 #include "AvaViewportPostProcessManager.h"
 #include "AvaViewportSettings.h"
 #include "ViewportClient/AvaLevelViewportClient.h"
-#include "Visualizers/AvaViewportPostProcessVisualizer.h"
-#include "Visualizers/IAvaViewportBoundingBoxVisualizerInterface.h"
+#include "Visualizers/IAvaViewportBoundingBoxVisualizer.h"
+#include "Visualizers/IAvaViewportPostProcessVisualizer.h"
 
 void SAvaLevelViewport::ExecuteToggleChildActorLock()
 {
@@ -100,7 +100,7 @@ bool SAvaLevelViewport::CanTogglePostProcessBackground() const
 		return true;
 	}
 
-	if (TSharedPtr<FAvaViewportPostProcessVisualizer> BackgroundVisualizer = ViewportClient->GetPostProcessManager()->GetVisualizer(EAvaViewportPostProcessType::Background))
+	if (TSharedPtr<IAvaViewportPostProcessVisualizer> BackgroundVisualizer = ViewportClient->GetPostProcessManager()->GetVisualizer(EAvaViewportPostProcessType::Background))
 	{
 		return BackgroundVisualizer->CanActivate(/* bInSilent */ true);
 	}
@@ -167,7 +167,7 @@ bool SAvaLevelViewport::CanTogglePostProcessChannelRed() const
 		return true;
 	}
 
-	if (TSharedPtr<FAvaViewportPostProcessVisualizer> ChannelRedVisualizer = ViewportClient->GetPostProcessManager()->GetVisualizer(EAvaViewportPostProcessType::RedChannel))
+	if (TSharedPtr<IAvaViewportPostProcessVisualizer> ChannelRedVisualizer = ViewportClient->GetPostProcessManager()->GetVisualizer(EAvaViewportPostProcessType::RedChannel))
 	{
 		return ChannelRedVisualizer->CanActivate(/* bInSilent */ true);
 	}
@@ -234,7 +234,7 @@ bool SAvaLevelViewport::CanTogglePostProcessChannelGreen() const
 		return true;
 	}
 
-	if (TSharedPtr<FAvaViewportPostProcessVisualizer> ChannelGreenVisualizer = ViewportClient->GetPostProcessManager()->GetVisualizer(EAvaViewportPostProcessType::GreenChannel))
+	if (TSharedPtr<IAvaViewportPostProcessVisualizer> ChannelGreenVisualizer = ViewportClient->GetPostProcessManager()->GetVisualizer(EAvaViewportPostProcessType::GreenChannel))
 	{
 		return ChannelGreenVisualizer->CanActivate(/* bInSilent */ true);
 	}
@@ -301,7 +301,7 @@ bool SAvaLevelViewport::CanTogglePostProcessChannelBlue() const
 		return true;
 	}
 
-	if (TSharedPtr<FAvaViewportPostProcessVisualizer> ChannelBlueVisualizer = ViewportClient->GetPostProcessManager()->GetVisualizer(EAvaViewportPostProcessType::BlueChannel))
+	if (TSharedPtr<IAvaViewportPostProcessVisualizer> ChannelBlueVisualizer = ViewportClient->GetPostProcessManager()->GetVisualizer(EAvaViewportPostProcessType::BlueChannel))
 	{
 		return ChannelBlueVisualizer->CanActivate(/* bInSilent */ true);
 	}
@@ -368,7 +368,7 @@ bool SAvaLevelViewport::CanTogglePostProcessChannelAlpha() const
 		return true;
 	}
 
-	if (TSharedPtr<FAvaViewportPostProcessVisualizer> ChannelAlphaVisualizer = ViewportClient->GetPostProcessManager()->GetVisualizer(EAvaViewportPostProcessType::AlphaChannel))
+	if (TSharedPtr<IAvaViewportPostProcessVisualizer> ChannelAlphaVisualizer = ViewportClient->GetPostProcessManager()->GetVisualizer(EAvaViewportPostProcessType::AlphaChannel))
 	{
 		return ChannelAlphaVisualizer->CanActivate(/* bInSilent */ true);
 	}
@@ -425,7 +425,7 @@ bool SAvaLevelViewport::CanTogglePostProcessCheckerboard() const
 		return true;
 	}
 
-	if (const TSharedPtr<FAvaViewportPostProcessVisualizer> CheckerboardVisualizer = ViewportClient->GetPostProcessManager()->GetVisualizer(EAvaViewportPostProcessType::Checkerboard))
+	if (const TSharedPtr<IAvaViewportPostProcessVisualizer> CheckerboardVisualizer = ViewportClient->GetPostProcessManager()->GetVisualizer(EAvaViewportPostProcessType::Checkerboard))
 	{
 		return CheckerboardVisualizer->CanActivate(/* bInSilent */ true);
 	}
diff --git a/Engine/Plugins/Experimental/Avalanche/Source/AvalancheLevelViewport/Private/ViewportClient/AvaLevelViewportClient.cpp b/Engine/Plugins/Experimental/Avalanche/Source/AvalancheLevelViewport/Private/ViewportClient/AvaLevelViewportClient.cpp
--- a/Engine/Plugins/Experimental/Avalanche/Source/AvalancheLevelViewport/Private/ViewportClient/AvaLevelViewportClient.cpp
+++ b/Engine/Plugins/Experimental/Avalanche/Source/AvalancheLevelViewport/Private/ViewportClient/AvaLevelViewportClient.cpp
@@ -26,7 +26,7 @@
 #include "UnrealEdGlobals.h"
 #include "Viewport/Interaction/IAvaViewportDataProvider.h"
 #include "Viewport/Interaction/IAvaViewportDataProxy.h"
-#include "Visualizers/IAvaViewportBoundingBoxVisualizerInterface.h"
+#include "Visualizers/IAvaViewportBoundingBoxVisualizer.h"
 
 #define LOCTEXT_NAMESPACE "AvaLevelViewportClient"
 
diff --git a/Engine/Plugins/Experimental/Avalanche/Source/AvalancheLevelViewport/Private/ViewportClient/AvaLevelViewportClient.h b/Engine/Plugins/Experimental/Avalanche/Source/AvalancheLevelViewport/Private/ViewportClient/AvaLevelViewportClient.h
--- a/Engine/Plugins/Experimental/Avalanche/Source/AvalancheLevelViewport/Private/ViewportClient/AvaLevelViewportClient.h
+++ b/Engine/Plugins/Experimental/Avalanche/Source/AvalancheLevelViewport/Private/ViewportClient/AvaLevelViewportClient.h
@@ -18,7 +18,7 @@ class FAvaSnapOperation;
 class FAvaViewportPostProcessManager;
 class FEditorViewportClient;
 class IAvaBoundsProviderInterface;
-class IAvaViewportBoundingBoxVisualizerInterface;
+class IAvaViewportBoundingBoxVisualizer;
 class IAvaViewportDataProvider;
 class IAvaViewportDataProxy;
 class SAvaLevelViewport;
@@ -57,7 +57,7 @@ public:
 	bool AreChildActorsLocked() const { return bLockChildActorsOnDrag; }
 	void SetChildActorsLocked(bool bInLockChildActors) { bLockChildActorsOnDrag = bInLockChildActors; }
 
-	TSharedRef<IAvaViewportBoundingBoxVisualizerInterface> GetBoundingBoxVisualizer() const { return BoundingBoxVisualizer.ToSharedRef(); }
+	TSharedRef<IAvaViewportBoundingBoxVisualizer> GetBoundingBoxVisualizer() const { return BoundingBoxVisualizer.ToSharedRef(); }
 
 	TSharedRef<FAvaIsolateActorsOperation> GetIsolateActorsOperation() const { return IsolateActorsOperation.ToSharedRef(); }
 
@@ -121,7 +121,7 @@ protected:
 
 	TSharedPtr<FAvaCameraZoomController> ZoomController;
 
-	TSharedPtr<IAvaViewportBoundingBoxVisualizerInterface> BoundingBoxVisualizer;
+	TSharedPtr<IAvaViewportBoundingBoxVisualizer> BoundingBoxVisualizer;
 
 	TWeakObjectPtr<UCameraComponent> ActiveCameraComponentWeak;
 
diff --git a/Engine/Plugins/Experimental/Avalanche/Source/AvalancheViewport/Private/AvaViewportPostProcessManager.cpp b/Engine/Plugins/Experimental/Avalanche/Source/AvalancheViewport/Private/AvaViewportPostProcessManager.cpp
--- a/Engine/Plugins/Experimental/Avalanche/Source/AvalancheViewport/Private/AvaViewportPostProcessManager.cpp
+++ b/Engine/Plugins/Experimental/Avalanche/Source/AvalancheViewport/Private/AvaViewportPostProcessManager.cpp
@@ -2,6 +2,7 @@
 
 #include "AvaViewportPostProcessManager.h"
 
+#include "AvaTypeSharedPointer.h"
 #include "AvaViewportDataSubsystem.h"
 #include "AvaVisibleArea.h"
 #include "Interaction/AvaCameraZoomController.h"
@@ -17,7 +18,7 @@ FAvaViewportPostProcessManager::FAvaViewportPostProcessManager(TSharedRef<IAvaVi
 {
 	AvaViewportClientWeak = InAvaViewportClient;
 
-	Visualizers.Emplace(EAvaViewportPostProcessType::Background,    MakeShared<FAvaViewportBackgroundVisualizer>(InAvaViewportClient));
+	Visualizers.Emplace(EAvaViewportPostProcessType::Background,   MakeShared<FAvaViewportBackgroundVisualizer>(InAvaViewportClient));
 	Visualizers.Emplace(EAvaViewportPostProcessType::RedChannel,   MakeShared<FAvaViewportChannelVisualizer>(InAvaViewportClient, EAvaViewportPostProcessType::RedChannel));
 	Visualizers.Emplace(EAvaViewportPostProcessType::GreenChannel, MakeShared<FAvaViewportChannelVisualizer>(InAvaViewportClient, EAvaViewportPostProcessType::GreenChannel));
 	Visualizers.Emplace(EAvaViewportPostProcessType::BlueChannel,  MakeShared<FAvaViewportChannelVisualizer>(InAvaViewportClient, EAvaViewportPostProcessType::BlueChannel));
@@ -58,7 +59,7 @@ void FAvaViewportPostProcessManager::UpdateSceneView(FSceneView* InSceneView)
 		return;
 	}
 
-	TSharedPtr<FAvaViewportPostProcessVisualizer> Visualizer = GetActiveVisualizer();
+	TSharedPtr<FAvaViewportPostProcessVisualizer> Visualizer = UE::AvaCore::CastSharedPtr<FAvaViewportPostProcessVisualizer>(GetActiveVisualizer());
 
 	if (!Visualizer.IsValid())
 	{
@@ -83,7 +84,7 @@ void FAvaViewportPostProcessManager::UpdateSceneView(FSceneView* InSceneView)
 
 void FAvaViewportPostProcessManager::LoadPostProcessInfo()
 {
-	TSharedPtr<FAvaViewportPostProcessVisualizer> Visualizer = GetActiveVisualizer();
+	TSharedPtr<FAvaViewportPostProcessVisualizer> Visualizer = UE::AvaCore::CastSharedPtr<FAvaViewportPostProcessVisualizer>(GetActiveVisualizer());
 
 	if (!Visualizer.IsValid())
 	{
@@ -117,14 +118,14 @@ void FAvaViewportPostProcessManager::SetType(EAvaViewportPostProcessType InType)
 		return;
 	}
 
-	TSharedPtr<FAvaViewportPostProcessVisualizer> NewVisualizer = GetVisualizer(InType);
+	TSharedPtr<FAvaViewportPostProcessVisualizer> NewVisualizer = UE::AvaCore::CastSharedPtr<FAvaViewportPostProcessVisualizer>(GetVisualizer(InType));
 
 	if (NewVisualizer.IsValid() && !NewVisualizer->CanActivate(/* bInSilent */ false))
 	{
 		return;
 	}
 
-	if (TSharedPtr<FAvaViewportPostProcessVisualizer> CurrentVisualizer = GetActiveVisualizer())
+	if (TSharedPtr<FAvaViewportPostProcessVisualizer> CurrentVisualizer = UE::AvaCore::CastSharedPtr<FAvaViewportPostProcessVisualizer>(GetActiveVisualizer()))
 	{
 		CurrentVisualizer->OnDeactivate();
 	}
@@ -156,9 +157,9 @@ void FAvaViewportPostProcessManager::SetOpacity(float InOpacity)
 	}
 }
 
-TSharedPtr<FAvaViewportPostProcessVisualizer> FAvaViewportPostProcessManager::GetVisualizer(EAvaViewportPostProcessType InType) const
+TSharedPtr<IAvaViewportPostProcessVisualizer> FAvaViewportPostProcessManager::GetVisualizer(EAvaViewportPostProcessType InType) const
 {
-	if (const TSharedPtr<FAvaViewportPostProcessVisualizer>* VisualizerPtr = Visualizers.Find(InType))
+	if (const TSharedPtr<IAvaViewportPostProcessVisualizer>* VisualizerPtr = Visualizers.Find(InType))
 	{
 		return *VisualizerPtr;
 	}
@@ -166,7 +167,7 @@ TSharedPtr<FAvaViewportPostProcessVisualizer> FAvaViewportPostProcessManager::Ge
 	return nullptr;
 }
 
-TSharedPtr<FAvaViewportPostProcessVisualizer> FAvaViewportPostProcessManager::GetActiveVisualizer() const
+TSharedPtr<IAvaViewportPostProcessVisualizer> FAvaViewportPostProcessManager::GetActiveVisualizer() const
 {
 	if (FAvaViewportPostProcessInfo* PostProcessInfo = GetPostProcessInfo())
 	{
diff --git a/Engine/Plugins/Experimental/Avalanche/Source/AvalancheViewport/Private/Interaction/AvaSnapOperation.cpp b/Engine/Plugins/Experimental/Avalanche/Source/AvalancheViewport/Private/Interaction/AvaSnapOperation.cpp
--- a/Engine/Plugins/Experimental/Avalanche/Source/AvalancheViewport/Private/Interaction/AvaSnapOperation.cpp
+++ b/Engine/Plugins/Experimental/Avalanche/Source/AvalancheViewport/Private/Interaction/AvaSnapOperation.cpp
@@ -1407,28 +1407,6 @@ const FAvaSnapPoint* const FAvaSnapOperation::GetSnapPointByLink(int32 InSnapPoi
 	return GetActorSnapPoint(InSnapPointLinkIdx - ScreenSnapPoints.Num());
 }
 
-void FAvaSnapOperation::SetClosestSnapPointLinkIdxX(int32 InSnapPointLinkIdx)
-{
-	if (InSnapPointLinkIdx < INDEX_NONE
-		&& InSnapPointLinkIdx >= (ScreenSnapPoints.Num() + ActorSnapPoints.Num()))
-	{
-		return;
-	}
-
-	ClosestSnapPointLinkIdxX = InSnapPointLinkIdx;
-}
-
-void FAvaSnapOperation::SetClosestSnapPointLinkIdxY(int32 InSnapPointLinkIdx)
-{
-	if (InSnapPointLinkIdx < INDEX_NONE
-		&& InSnapPointLinkIdx >= (ScreenSnapPoints.Num() + ActorSnapPoints.Num()))
-	{
-		return;
-	}
-
-	ClosestSnapPointLinkIdxX = InSnapPointLinkIdx;
-}
-
 void FAvaSnapOperation::SetSnappedToX(bool bInSnappedTo)
 {
 	bSnappedToLinkX = bInSnappedTo;
diff --git a/Engine/Plugins/Experimental/Avalanche/Source/AvalancheViewport/Private/Test/AvaViewportWorldCoordinateConverterProvider.cpp b/Engine/Plugins/Experimental/Avalanche/Source/AvalancheViewport/Private/Test/AvaViewportWorldCoordinateConverterProvider.cpp
--- a/Engine/Plugins/Experimental/Avalanche/Source/AvalancheViewport/Private/Test/AvaViewportWorldCoordinateConverterProvider.cpp
+++ b/Engine/Plugins/Experimental/Avalanche/Source/AvalancheViewport/Private/Test/AvaViewportWorldCoordinateConverterProvider.cpp
@@ -12,11 +12,6 @@ FAvaViewportWorldCoordinateConverterProvider::FAvaViewportWorldCoordinateConvert
 {
 }
 
-FVector2f FAvaViewportWorldCoordinateConverterProvider::GetViewportSize() const
-{
-	return ViewportSize;
-}
-
 FAvaViewportWorldCoordinateConverterProviderPerspective::FAvaViewportWorldCoordinateConverterProviderPerspective(FVector InLocation,
 	FRotator InRotation, FVector2f InViewportSize, float InFieldOfView)
 	: FAvaViewportWorldCoordinateConverterProvider(FTransform(InRotation, InLocation, FVector::OneVector), InViewportSize)
diff --git a/Engine/Plugins/Experimental/Avalanche/Source/AvalancheViewport/Private/Visualizers/AvaViewportBackgroundVisualizer.h b/Engine/Plugins/Experimental/Avalanche/Source/AvalancheViewport/Private/Visualizers/AvaViewportBackgroundVisualizer.h
new file mode 100644
--- /dev/null
+++ b/Engine/Plugins/Experimental/Avalanche/Source/AvalancheViewport/Private/Visualizers/AvaViewportBackgroundVisualizer.h
@@ -0,0 +1,50 @@
+// Copyright Epic Games, Inc. All Rights Reserved.
+
+#pragma once
+
+#include "Visualizers/AvaViewportPostProcessVisualizer.h"
+#include "Math/Vector.h"
+
+class UTexture;
+struct FAvaViewportPostProcessInfo;
+
+class FAvaViewportBackgroundVisualizer : public FAvaViewportPostProcessVisualizer
+{
+public:
+	UE_AVA_INHERITS_WITH_SUPER(FAvaViewportBackgroundVisualizer, FAvaViewportPostProcessVisualizer)
+
+	FAvaViewportBackgroundVisualizer(TSharedRef<IAvaViewportClient> InAvaViewportClient);
+	virtual ~FAvaViewportBackgroundVisualizer() override = default;
+
+	UTexture* GetTexture() const;
+	void SetTexture(UTexture* InTexture);
+
+	//~ Begin FAvaViewportPostProcessVisualizer
+	virtual void UpdateForViewport(const FAvaVisibleArea& InVisibleArea, const FVector2f& InWidgetSize, const FVector2f& InCameraOffset) override;
+	//~ End FAvaViewportPostProcessVisualizer
+
+	//~ Begin FGCObject
+	virtual void AddReferencedObjects(FReferenceCollector& InCollector) override;
+	virtual FString GetReferencerName() const override;
+	//~ End FGCObject
+
+protected:
+	TObjectPtr<UTexture> Texture;
+
+	/** Offset from the top-left of the viewport where the texture should start in pixel space. */
+	FVector TextureOffset;
+
+	/** Scale of the texture. Larger values mean a larger texture. */
+	FVector TextureScale;
+
+	using FAvaViewportPostProcessVisualizer::UpdatePostProcessInfo;
+
+	void SetTextureInternal(UTexture* InTexture);
+
+	//~ Begin FAvaViewportPostProcessVisualizer
+	virtual void LoadPostProcessInfo(const FAvaViewportPostProcessInfo& InPostProcessInfo) override;
+	virtual void UpdatePostProcessInfo(FAvaViewportPostProcessInfo& InPostProcessInfo) const override;
+	virtual void UpdatePostProcessMaterial() override;
+	virtual bool SetupPostProcessSettings(FPostProcessSettings& InPostProcessSettings) const override;
+	//~ End FAvaViewportPostProcessVisualizer
+};
\ No newline at end of file
diff --git a/Engine/Plugins/Experimental/Avalanche/Source/AvalancheViewport/Private/Visualizers/AvaViewportBoundingBoxVisualizer.cpp b/Engine/Plugins/Experimental/Avalanche/Source/AvalancheViewport/Private/Visualizers/AvaViewportBoundingBoxVisualizer.cpp
--- a/Engine/Plugins/Experimental/Avalanche/Source/AvalancheViewport/Private/Visualizers/AvaViewportBoundingBoxVisualizer.cpp
+++ b/Engine/Plugins/Experimental/Avalanche/Source/AvalancheViewport/Private/Visualizers/AvaViewportBoundingBoxVisualizer.cpp
@@ -17,7 +17,7 @@ namespace UE::AvaViewport::Private
 	constexpr float BoundingBoxLineThickness = 0.4f;
 }
 
-TSharedRef<IAvaViewportBoundingBoxVisualizerInterface> FAvaViewportBoundingBoxVisualizerProvider::CreateVisualizer()
+TSharedRef<IAvaViewportBoundingBoxVisualizer> FAvaViewportBoundingBoxVisualizerProvider::CreateVisualizer()
 {
 	return MakeShared<FAvaLevelViewportBoundingBoxVisualizer>();
 }
diff --git a/Engine/Plugins/Experimental/Avalanche/Source/AvalancheViewport/Private/Visualizers/AvaViewportBoundingBoxVisualizer.h b/Engine/Plugins/Experimental/Avalanche/Source/AvalancheViewport/Private/Visualizers/AvaViewportBoundingBoxVisualizer.h
--- a/Engine/Plugins/Experimental/Avalanche/Source/AvalancheViewport/Private/Visualizers/AvaViewportBoundingBoxVisualizer.h
+++ b/Engine/Plugins/Experimental/Avalanche/Source/AvalancheViewport/Private/Visualizers/AvaViewportBoundingBoxVisualizer.h
@@ -3,23 +3,23 @@
 #pragma once
 
 #include "AvaType.h"
-#include "Visualizers/IAvaViewportBoundingBoxVisualizerInterface.h"
+#include "Visualizers/IAvaViewportBoundingBoxVisualizer.h"
 
 class AActor;
 
-class FAvaLevelViewportBoundingBoxVisualizer : public IAvaViewportBoundingBoxVisualizerInterface
+class FAvaLevelViewportBoundingBoxVisualizer : public IAvaViewportBoundingBoxVisualizer
 {
 public:
-	UE_AVA_INHERITS(FAvaLevelViewportBoundingBoxVisualizer, IAvaViewportBoundingBoxVisualizerInterface)
+	UE_AVA_INHERITS(FAvaLevelViewportBoundingBoxVisualizer, IAvaViewportBoundingBoxVisualizer)
 
 	FAvaLevelViewportBoundingBoxVisualizer();
 	virtual ~FAvaLevelViewportBoundingBoxVisualizer() = default;
 
-	//~ Begin IAvaViewportBoundingBoxVisualizerInterface
+	//~ Begin IAvaViewportBoundingBoxVisualizer
 	virtual void Draw(UAvaSelectionProviderSubsystem& InSelectionProvider, UAvaBoundsProviderSubsystem& InBoundsProvider, FPrimitiveDrawInterface& InPDI) override;
 	virtual EAvaViewportBoundingBoxOptimizationState GetOptimizationState() const override { return OptimizationState; }
 	virtual void ResetOptimizationState() override;
-	//~ End IAvaViewportBoundingBoxVisualizerInterface
+	//~ End IAvaViewportBoundingBoxVisualizer
 
 protected:
 	EAvaViewportBoundingBoxOptimizationState OptimizationState;
diff --git a/Engine/Plugins/Experimental/Avalanche/Source/AvalancheViewport/Private/Visualizers/AvaViewportChannelVisualizer.h b/Engine/Plugins/Experimental/Avalanche/Source/AvalancheViewport/Private/Visualizers/AvaViewportChannelVisualizer.h
new file mode 100644
--- /dev/null
+++ b/Engine/Plugins/Experimental/Avalanche/Source/AvalancheViewport/Private/Visualizers/AvaViewportChannelVisualizer.h
@@ -0,0 +1,21 @@
+// Copyright Epic Games, Inc. All Rights Reserved.
+
+#pragma once
+
+#include "Visualizers/AvaViewportPostProcessVisualizer.h"
+#include "Math/Vector.h"
+
+enum class EAvaViewportPostProcessType : uint8;
+
+class FAvaViewportChannelVisualizer : public FAvaViewportPostProcessVisualizer
+{
+public:
+	UE_AVA_INHERITS_WITH_SUPER(FAvaViewportChannelVisualizer, FAvaViewportPostProcessVisualizer)
+
+	FAvaViewportChannelVisualizer(TSharedRef<IAvaViewportClient> InAvaViewportClient, EAvaViewportPostProcessType InChannel);
+	virtual ~FAvaViewportChannelVisualizer() override = default;
+
+	//~ Begin FGCObject
+	virtual FString GetReferencerName() const override;
+	//~ End FGCObject
+};
diff --git a/Engine/Plugins/Experimental/Avalanche/Source/AvalancheViewport/Private/Visualizers/AvaViewportCheckerboardVisualizer.h b/Engine/Plugins/Experimental/Avalanche/Source/AvalancheViewport/Private/Visualizers/AvaViewportCheckerboardVisualizer.h
new file mode 100644
--- /dev/null
+++ b/Engine/Plugins/Experimental/Avalanche/Source/AvalancheViewport/Private/Visualizers/AvaViewportCheckerboardVisualizer.h
@@ -0,0 +1,40 @@
+// Copyright Epic Games, Inc. All Rights Reserved.
+
+#pragma once
+
+#include "Visualizers/AvaViewportPostProcessVisualizer.h"
+#include "Math/Vector.h"
+
+struct FAvaViewportPostProcessInfo;
+
+class FAvaViewportCheckerboardVisualizer : public FAvaViewportPostProcessVisualizer
+{
+public:
+	UE_AVA_INHERITS_WITH_SUPER(FAvaViewportCheckerboardVisualizer, FAvaViewportPostProcessVisualizer)
+
+	FAvaViewportCheckerboardVisualizer(TSharedRef<IAvaViewportClient> InAvaViewportClient);
+	virtual ~FAvaViewportCheckerboardVisualizer() override;
+	
+	//~ Begin FAvaViewportPostProcessVisualizer
+	virtual void UpdatePostProcessMaterial() override;
+	//~ End FAvaViewportPostProcessVisualizer
+
+	//~ Begin FGCObject
+	virtual FString GetReferencerName() const override;
+	//~ End FGCObject
+
+	void OnSettingChanged(UObject* InSettings, struct FPropertyChangedEvent& InPropertyChangeEvent);
+
+	bool InitPostProcessMaterial();
+	bool UpdateFromViewportSettings();
+	
+protected:
+	/** First color of the checkerboard. */
+	FVector CheckerboardColor0;
+
+	/** Second color of the checkerboard. */
+	FVector CheckerboardColor1;
+
+	/** Size of the checkerboard (in pixels). */
+	FVector CheckerboardSize;
+};
\ No newline at end of file
diff --git a/Engine/Plugins/Experimental/Avalanche/Source/AvalancheViewport/Private/Visualizers/AvaViewportPostProcessVisualizer.cpp b/Engine/Plugins/Experimental/Avalanche/Source/AvalancheViewport/Private/Visualizers/AvaViewportPostProcessVisualizer.cpp
--- a/Engine/Plugins/Experimental/Avalanche/Source/AvalancheViewport/Private/Visualizers/AvaViewportPostProcessVisualizer.cpp
+++ b/Engine/Plugins/Experimental/Avalanche/Source/AvalancheViewport/Private/Visualizers/AvaViewportPostProcessVisualizer.cpp
@@ -108,6 +108,14 @@ void FAvaViewportPostProcessVisualizer::OnActivate()
 	LoadPostProcessInfo();
 }
 
+void FAvaViewportPostProcessVisualizer::OnDeactivate()
+{
+}
+
+void FAvaViewportPostProcessVisualizer::UpdateForViewport(const FAvaVisibleArea& InVisibleArea, const FVector2f& InWidgetSize, const FVector2f& InCameraOffset)
+{
+}
+
 void FAvaViewportPostProcessVisualizer::ApplyToSceneView(FSceneView* InSceneView) const
 {
 	if (!InSceneView || FMath::IsNearlyZero(PostProcessOpacity) || !PostProcessMaterial)
diff --git a/Engine/Plugins/Experimental/Avalanche/Source/AvalancheViewport/Private/Visualizers/AvaViewportPostProcessVisualizer.h b/Engine/Plugins/Experimental/Avalanche/Source/AvalancheViewport/Private/Visualizers/AvaViewportPostProcessVisualizer.h
new file mode 100644
--- /dev/null
+++ b/Engine/Plugins/Experimental/Avalanche/Source/AvalancheViewport/Private/Visualizers/AvaViewportPostProcessVisualizer.h
@@ -0,0 +1,79 @@
+// Copyright Epic Games, Inc. All Rights Reserved.
+
+#pragma once
+
+#include "AvaType.h"
+#include "EditorUndoClient.h"
+#include "Math/MathFwd.h"
+#include "Templates/SharedPointer.h"
+#include "UObject/GCObject.h"
+#include "Visualizers/IAvaViewportPostProcessVisualizer.h"
+
+class FSceneView;
+class IAvaViewportClient;
+class UMaterial;
+class UMaterialInstanceDynamic;
+struct FAvaViewportPostProcessInfo;
+struct FAvaVisibleArea;
+struct FPostProcessSettings;
+
+class FAvaViewportPostProcessVisualizer : public IAvaViewportPostProcessVisualizer, public FGCObject, public FEditorUndoClient
+{
+public:
+	UE_AVA_INHERITS(FAvaViewportPostProcessVisualizer, IAvaViewportPostProcessVisualizer)
+
+	FAvaViewportPostProcessVisualizer(TSharedRef<IAvaViewportClient> InAvaViewportClient);
+
+	virtual ~FAvaViewportPostProcessVisualizer() override;
+
+	TSharedPtr<IAvaViewportClient> GetAvaViewportClient() const;
+
+	float GetPostProcessOpacity() const { return PostProcessOpacity; }
+	void SetPostProcessOpacity(float InOpacity);
+
+	FAvaViewportPostProcessInfo* GetPostProcessInfo() const;
+
+	void LoadPostProcessInfo();
+
+	//~ Begin IAvaViewportPostProcessVisualizer
+	virtual bool CanActivate(bool bInSilent) const override;
+	//~ End IAvaViewportPostProcessVisualizer
+
+	virtual void OnActivate();
+	virtual void OnDeactivate();
+
+	virtual void UpdateForViewport(const FAvaVisibleArea& InVisibleArea, const FVector2f& InWidgetSize, const FVector2f& InCameraOffset);
+
+	void ApplyToSceneView(FSceneView* InSceneView) const;
+
+protected:
+	//~ Begin FGCObject
+	virtual void AddReferencedObjects(FReferenceCollector& InCollector) override;
+	//~ End FGCObject
+
+	// FEditorUndoClient Interface
+	virtual void PostUndo(bool bSuccess) override;
+	virtual void PostRedo(bool bSuccess) override;
+	// ~FEditorUndoClient Interface
+
+	TWeakPtr<IAvaViewportClient> AvaViewportClientWeak;
+
+	TObjectPtr<UMaterial> PostProcessBaseMaterial;
+
+	TObjectPtr<UMaterialInstanceDynamic> PostProcessMaterial;
+
+	float PostProcessOpacity;
+
+	bool bRequiresTonemapperSetting;
+
+	void SetPostProcessOpacityInternal(float InOpacity);
+
+	void UpdatePostProcessInfo();
+
+	virtual void LoadPostProcessInfo(const FAvaViewportPostProcessInfo& InPostProcessInfo);
+	virtual void UpdatePostProcessInfo(FAvaViewportPostProcessInfo& InPostProcessInfo) const;
+
+	virtual void UpdatePostProcessMaterial();
+
+	virtual bool SetupPostProcessSettings(FPostProcessSettings& InPostProcessSettings) const;
+};
diff --git a/Engine/Plugins/Experimental/Avalanche/Source/AvalancheViewport/Public/AvaViewportPostProcessManager.h b/Engine/Plugins/Experimental/Avalanche/Source/AvalancheViewport/Public/AvaViewportPostProcessManager.h
--- a/Engine/Plugins/Experimental/Avalanche/Source/AvalancheViewport/Public/AvaViewportPostProcessManager.h
+++ b/Engine/Plugins/Experimental/Avalanche/Source/AvalancheViewport/Public/AvaViewportPostProcessManager.h
@@ -8,9 +8,9 @@
 
 class FAvaViewportBackgroundVisualizer;
 class FAvaViewportChannelVisualizer;
-class FAvaViewportPostProcessVisualizer;
 class FSceneView;
 class IAvaViewportClient;
+class IAvaViewportPostProcessVisualizer;
 class UTexture;
 
 UENUM()
@@ -50,9 +50,9 @@ public:
 
 	TSharedPtr<IAvaViewportClient> GetAvaViewportClient() const { return AvaViewportClientWeak.Pin(); }
 
-	TSharedPtr<FAvaViewportPostProcessVisualizer> GetVisualizer(EAvaViewportPostProcessType InType) const;
+	TSharedPtr<IAvaViewportPostProcessVisualizer> GetVisualizer(EAvaViewportPostProcessType InType) const;
 
-	TSharedPtr<FAvaViewportPostProcessVisualizer> GetActiveVisualizer() const;
+	TSharedPtr<IAvaViewportPostProcessVisualizer> GetActiveVisualizer() const;
 
 	FAvaViewportPostProcessInfo* GetPostProcessInfo() const;
 
@@ -69,5 +69,5 @@ public:
 protected:
 	TWeakPtr<IAvaViewportClient> AvaViewportClientWeak;
 
-	TMap<EAvaViewportPostProcessType, TSharedPtr<FAvaViewportPostProcessVisualizer>> Visualizers;
+	TMap<EAvaViewportPostProcessType, TSharedPtr<IAvaViewportPostProcessVisualizer>> Visualizers;
 };
diff --git a/Engine/Plugins/Experimental/Avalanche/Source/AvalancheViewport/Public/Interaction/AvaCameraZoomController.h b/Engine/Plugins/Experimental/Avalanche/Source/AvalancheViewport/Public/Interaction/AvaCameraZoomController.h
--- a/Engine/Plugins/Experimental/Avalanche/Source/AvalancheViewport/Public/Interaction/AvaCameraZoomController.h
+++ b/Engine/Plugins/Experimental/Avalanche/Source/AvalancheViewport/Public/Interaction/AvaCameraZoomController.h
@@ -9,23 +9,30 @@
 class IAvaViewportClient;
 struct FAvaVisibleArea;
 
-class AVALANCHEVIEWPORT_API FAvaCameraZoomController
+class FAvaCameraZoomController
 {
 public:
-	static bool IsCameraZoomPossible();
+	AVALANCHEVIEWPORT_API static bool IsCameraZoomPossible();
 
-	FAvaCameraZoomController(TSharedRef<IAvaViewportClient> InAvaViewportClient, float InFallbackFOV);
+	AVALANCHEVIEWPORT_API FAvaCameraZoomController(TSharedRef<IAvaViewportClient> InAvaViewportClient, float InFallbackFOV);
 
-	TSharedPtr<IAvaViewportClient> GetViewportClient() const { return AvaViewportClientWeak.Pin(); }
+	TSharedPtr<IAvaViewportClient> GetViewportClient() const
+	{
+		return AvaViewportClientWeak.Pin();
+	}
+
+	uint8 GetZoomLevel() const
+	{
+		return ZoomLevel;
+	}
 
-	uint8 GetZoomLevel() const { return ZoomLevel; }
 	void SetZoomLevel(uint8 InZoomLevel);
 	bool IsZoomed() const;
 
 	float GetFOVPerStep() const;
 
-	void ZoomIn();
-	void ZoomInCursor();
+	AVALANCHEVIEWPORT_API void ZoomIn();
+	AVALANCHEVIEWPORT_API void ZoomInCursor();
 
 	/* Zooms in maintaining the current PanOffsetFraction. Uses the absolute screen position. */
 	void ZoomInAroundPoint(const FVector2f& InScreenPosition);
@@ -33,8 +40,8 @@ public:
 	/* Zooms in maintaining the position of the given viewport position. Uses the absolute screen position. */
 	void ZoomInRelativePoint(const FVector2f& InViewportPosition);
 
-	void ZoomOut();
-	void ZoomOutCursor();
+	AVALANCHEVIEWPORT_API void ZoomOut();
+	AVALANCHEVIEWPORT_API void ZoomOutCursor();
 
 	/* Zooms out maintaining the current PanOffsetFraction. Uses the absolute screen position. */
 	void ZoomOutAroundPoint(const FVector2f& OutScreenPosition);
@@ -42,38 +49,50 @@ public:
 	/* Zooms out maintaining the position of the given viewport position. Uses the absolute screen position. */
 	void ZoomOutRelativePoint(const FVector2f& OutViewportPosition);
 
-	void PanLeft();
-	void PanRight();
-	void PanUp();
-	void PanDown();
+	AVALANCHEVIEWPORT_API void PanLeft();
+	AVALANCHEVIEWPORT_API void PanRight();
+	AVALANCHEVIEWPORT_API void PanUp();
+	AVALANCHEVIEWPORT_API void PanDown();
+
+	AVALANCHEVIEWPORT_API void FrameActor();
 
-	void FrameActor();
+	AVALANCHEVIEWPORT_API void Reset();
 
-	void Reset();
+	const FVector2f& GetPanOffsetFraction() const
+	{
+		return PanOffsetFraction;
+	}
 
-	const FVector2f& GetPanOffsetFraction() const { return PanOffsetFraction; }
 	void SetPanOffsetFraction(const FVector2f& InOffsetFraction);
 	void PanAdjust(const FVector2f& InDirection);
 
 	/** Adjusts zoom pan based on current zoom settings. */
-	void PanAdjustZoomed(const FVector2f& InZoomedDirection);
+	AVALANCHEVIEWPORT_API void PanAdjustZoomed(const FVector2f& InZoomedDirection);
 
 	void CenterOnPoint(const FVector2f& InPoint);
 	void CenterOnBox(const FBox& InBoundingBox, const FTransform& InBoxTransform);
 
-	bool IsPanning() const { return bIsPanning; }
-	void StartPanning();
-	void EndPanning();
+	bool IsPanning() const
+	{
+		return bIsPanning;
+	}
+
+	AVALANCHEVIEWPORT_API void StartPanning();
+	AVALANCHEVIEWPORT_API void EndPanning();
+
+	float GetFallbackFOV() const
+	{
+		return FallbackFOV;
+	}
 
-	float GetFallbackFOV() const { return FallbackFOV; }
-	float GetDefaultFOV() const;
-	float GetFOV() const;
-	FVector2f GetCameraProjectionOffset() const;
+	AVALANCHEVIEWPORT_API float GetDefaultFOV() const;
+	AVALANCHEVIEWPORT_API float GetFOV() const;
+	AVALANCHEVIEWPORT_API FVector2f GetCameraProjectionOffset() const;
 
-	const FAvaVisibleArea& GetCachedVisibleArea() const;
-	const FAvaVisibleArea& GetCachedZoomedVisibleArea() const;
+	AVALANCHEVIEWPORT_API const FAvaVisibleArea& GetCachedVisibleArea() const;
+	AVALANCHEVIEWPORT_API const FAvaVisibleArea& GetCachedZoomedVisibleArea() const;
 
-	void UpdateVisibleAreas();
+	AVALANCHEVIEWPORT_API void UpdateVisibleAreas();
 
 protected:
 	TWeakPtr<IAvaViewportClient> AvaViewportClientWeak;
@@ -89,6 +108,7 @@ protected:
 	FAvaVisibleArea CachedVisibleArea;
 	FAvaVisibleArea CachedZoomedVisibleArea;
 
+private:
 	void ZoomIn_Internal();
 	void ZoomOut_Internal();
 
diff --git a/Engine/Plugins/Experimental/Avalanche/Source/AvalancheViewport/Public/Interaction/AvaChildTransformLockOperation.h b/Engine/Plugins/Experimental/Avalanche/Source/AvalancheViewport/Public/Interaction/AvaChildTransformLockOperation.h
--- a/Engine/Plugins/Experimental/Avalanche/Source/AvalancheViewport/Public/Interaction/AvaChildTransformLockOperation.h
+++ b/Engine/Plugins/Experimental/Avalanche/Source/AvalancheViewport/Public/Interaction/AvaChildTransformLockOperation.h
@@ -18,14 +18,14 @@ class AActor;
  * transform of the grandchild. To be honest, I'm not sure how UE would handle
  * this by default anyway...
  */
-
-class AVALANCHEVIEWPORT_API FAvaChildTransformLockOperation
+class FAvaChildTransformLockOperation
 {
 public:
-	FAvaChildTransformLockOperation(TSharedRef<IAvaViewportClient> InAvaViewportClient);
+	AVALANCHEVIEWPORT_API FAvaChildTransformLockOperation(TSharedRef<IAvaViewportClient> InAvaViewportClient);
+
+	AVALANCHEVIEWPORT_API void Save();
 
-	void Save();
-	void Restore();
+	AVALANCHEVIEWPORT_API void Restore();
 
 protected:
 	TWeakPtr<IAvaViewportClient> AvaViewportClientWeak;
diff --git a/Engine/Plugins/Experimental/Avalanche/Source/AvalancheViewport/Public/Interaction/AvaDragOperation.h b/Engine/Plugins/Experimental/Avalanche/Source/AvalancheViewport/Public/Interaction/AvaDragOperation.h
--- a/Engine/Plugins/Experimental/Avalanche/Source/AvalancheViewport/Public/Interaction/AvaDragOperation.h
+++ b/Engine/Plugins/Experimental/Avalanche/Source/AvalancheViewport/Public/Interaction/AvaDragOperation.h
@@ -12,19 +12,22 @@
 class FAvaSnapOperation;
 class IAvaViewportClient;
 
-class AVALANCHEVIEWPORT_API FAvaDragOperation
+class FAvaDragOperation
 {
 public:
-	FAvaDragOperation(TSharedRef<IAvaViewportClient> InAvaViewportClient, bool bInAllowSnapToChildren);
-	virtual ~FAvaDragOperation();
+	AVALANCHEVIEWPORT_API FAvaDragOperation(TSharedRef<IAvaViewportClient> InAvaViewportClient, bool bInAllowSnapToChildren);
+	AVALANCHEVIEWPORT_API virtual ~FAvaDragOperation();
 
 	/** Restore position of actors so they can be accurately dragged. */
-	void PreMouseUpdate();
+	AVALANCHEVIEWPORT_API void PreMouseUpdate();
 
 	/** Snap actors in their post-drag position. */
-	void PostMouseUpdate();
+	AVALANCHEVIEWPORT_API void PostMouseUpdate();
 
-	bool IsValid() const { return bValid; }
+	bool IsValid() const
+	{
+		return bValid;
+	}
 
 protected:
 	TWeakPtr<IAvaViewportClient> AvaViewportClientWeak;
diff --git a/Engine/Plugins/Experimental/Avalanche/Source/AvalancheViewport/Public/Interaction/AvaSnapOperation.h b/Engine/Plugins/Experimental/Avalanche/Source/AvalancheViewport/Public/Interaction/AvaSnapOperation.h
--- a/Engine/Plugins/Experimental/Avalanche/Source/AvalancheViewport/Public/Interaction/AvaSnapOperation.h
+++ b/Engine/Plugins/Experimental/Avalanche/Source/AvalancheViewport/Public/Interaction/AvaSnapOperation.h
@@ -40,7 +40,7 @@ struct FAvaSnapPointLink
 	}
 };
 
-class AVALANCHEVIEWPORT_API FAvaSnapOperation : public TSharedFromThis<FAvaSnapOperation>
+class FAvaSnapOperation : public TSharedFromThis<FAvaSnapOperation>
 {
 public:
 	static inline const float MaximumSnapDistance = 10.f;
@@ -62,28 +62,39 @@ public:
 	// Generates snap points based on a box extent with no actor reference
 	static void GenerateBoundsSnapPoints(const FVector& InOrigin, const FVector& InBoxExtent, TArray<FAvaSnapPoint>& InOutSnapPoints);
 
-	FAvaSnapOperation(FEditorViewportClient* InEditorViewportClient);
-	virtual ~FAvaSnapOperation();
+	AVALANCHEVIEWPORT_API FAvaSnapOperation(FEditorViewportClient* InEditorViewportClient);
+	AVALANCHEVIEWPORT_API virtual ~FAvaSnapOperation();
 
-	const TArray<FAvaSnapPoint>& GetScreenSnapPoints() const { return ScreenSnapPoints; }
-	void SetScreenSnapPoints(const TArray<FAvaSnapPoint>& InSnapPoints) { ScreenSnapPoints = InSnapPoints; }
-	const TArray<FAvaSnapPoint>& GetActorSnapPoints() const { return ActorSnapPoints; }
+	const TArray<FAvaSnapPoint>& GetScreenSnapPoints() const
+	{
+		return ScreenSnapPoints;
+	}
+
+	void SetScreenSnapPoints(const TArray<FAvaSnapPoint>& InSnapPoints)
+	{
+		ScreenSnapPoints = InSnapPoints;
+	}
+
+	const TArray<FAvaSnapPoint>& GetActorSnapPoints() const
+	{
+		return ActorSnapPoints;
+	}
 
 	// Adds screen- and guide-based snap points
 	void GenerateScreenSnapPoints();
 
 	// Generates snap points, excluding the actor for the given Dynamic Mesh
-	void GenerateComponentSnapPoints(const UActorComponent* InComponent);
+	AVALANCHEVIEWPORT_API void GenerateComponentSnapPoints(const UActorComponent* InComponent);
 
 	// Generates snap points, excluding the list of actors
-	void GenerateActorSnapPoints(const TConstArrayView<TWeakObjectPtr<AActor>>& InSelectedActors,
+	AVALANCHEVIEWPORT_API void GenerateActorSnapPoints(const TConstArrayView<TWeakObjectPtr<AActor>>& InSelectedActors,
 		const TConstArrayView<TWeakObjectPtr<AActor>>& InExcludedActors);
 
 	// Adds a snap point for a shape to the list of shape snap points
-	void AddActorSnapPoint(const FAvaSnapPoint& InSnapPoint);
+	AVALANCHEVIEWPORT_API void AddActorSnapPoint(const FAvaSnapPoint& InSnapPoint);
 
 	// Sorts snap point links
-	void FinaliseSnapPoints();
+	AVALANCHEVIEWPORT_API void FinaliseSnapPoints();
 
 	FVector GetSelectionLocation() const;
 	FVector GetDragOffset() const;
@@ -101,27 +112,38 @@ public:
 
 	void SnapScreenLocationX(float& InOutScreenLocation);
 	void SnapScreenLocationY(float& InOutScreenLocation);
-	void SnapScreenLocation(FVector2f& InOutScreenLocation);
-	void SnapScreenLocation(FVector2f& InOutScreenLocation, bool bInSnapX, bool bInSnapY);
 	bool SnapDragLocation(const TArray<FAvaSnapPoint>& InDraggedActorSnapPoints, FVector& OutSnapOffset);
-	bool SnapLocation(FVector& Location);
 
-	int32 GetClosestSnapPointLinkIdxX() const { return ClosestSnapPointLinkIdxX; }
-	void SetClosestSnapPointLinkIdxX(int32 InSnapPointLinkIdx);
-	
-	int32 GetClosestSnapPointLinkIdxY() const { return ClosestSnapPointLinkIdxY; }
-	void SetClosestSnapPointLinkIdxY(int32 InSnapPointLinkIdx);
-	
-	bool WasSnappedToX() const { return bSnappedToLinkX; }
-	void SetSnappedToX(bool bInSnappedTo);
-	
-	bool WasSnappedToY() const { return bSnappedToLinkY; }
-	void SetSnappedToY(bool bInSnappedTo);
+	AVALANCHEVIEWPORT_API void SnapScreenLocation(FVector2f& InOutScreenLocation);
+	AVALANCHEVIEWPORT_API void SnapScreenLocation(FVector2f& InOutScreenLocation, bool bInSnapX, bool bInSnapY);
+	AVALANCHEVIEWPORT_API bool SnapLocation(FVector& Location);
 
-	bool WasSnappedTo() const { return bSnappedToLinkX || bSnappedToLinkY; }
-	const FVector2f& GetSnappedToLocation() const { return SnappedToLocation; }
+	bool WasSnappedToX() const
+	{
+		return bSnappedToLinkX;
+	}
+	AVALANCHEVIEWPORT_API void SetSnappedToX(bool bInSnappedTo);
 	
-	const FVector2f& GetSnapDistances() const { return SnapDistances; }
+	bool WasSnappedToY() const
+	{
+		return bSnappedToLinkY;
+	}
+	AVALANCHEVIEWPORT_API void SetSnappedToY(bool bInSnappedTo);
+
+	bool WasSnappedTo() const
+	{
+		return bSnappedToLinkX || bSnappedToLinkY;
+	}
+
+	const FVector2f& GetSnappedToLocation() const
+	{
+		return SnappedToLocation;
+	}
+
+	const FVector2f& GetSnapDistances() const
+	{
+		return SnapDistances;
+	}
 
 protected:
 	/**
diff --git a/Engine/Plugins/Experimental/Avalanche/Source/AvalancheViewport/Public/Test/AvaViewportWorldCoordinateConverterProvider.h b/Engine/Plugins/Experimental/Avalanche/Source/AvalancheViewport/Public/Test/AvaViewportWorldCoordinateConverterProvider.h
--- a/Engine/Plugins/Experimental/Avalanche/Source/AvalancheViewport/Public/Test/AvaViewportWorldCoordinateConverterProvider.h
+++ b/Engine/Plugins/Experimental/Avalanche/Source/AvalancheViewport/Public/Test/AvaViewportWorldCoordinateConverterProvider.h
@@ -12,11 +12,10 @@
  * Instantiate with MakeShared<FAvaViewportWorldCoordinateConverterProviderPerspective>(...)
  *   or MakeShared<FAvaViewportWorldCoordinateConverterProviderOrthographic>(...)
  */
-
-struct AVALANCHEVIEWPORT_API FAvaViewportWorldCoordinateConverterProvider : public IAvaViewportWorldCoordinateConverter
+struct FAvaViewportWorldCoordinateConverterProvider : public IAvaViewportWorldCoordinateConverter
 {
 	//~ Begin IAvaViewportWorldCoordinateConverter
-	virtual FVector2f GetViewportSize() const override;
+	virtual FVector2f GetViewportSize() const override { return ViewportSize; }
 	virtual FTransform GetViewportViewTransform() const override { return ViewTransform; }
 	//~End IAvaViewportWorldCoordinateConverter
 
@@ -27,32 +26,32 @@ protected:
 	FAvaViewportWorldCoordinateConverterProvider(const FTransform& InViewTransform, const FVector2f& InViewportSize);
 };
 
-struct AVALANCHEVIEWPORT_API FAvaViewportWorldCoordinateConverterProviderPerspective : public FAvaViewportWorldCoordinateConverterProvider
+struct FAvaViewportWorldCoordinateConverterProviderPerspective : FAvaViewportWorldCoordinateConverterProvider
 {
-	FAvaViewportWorldCoordinateConverterProviderPerspective(FVector InLocation, FRotator InRotation, FVector2f InViewportSize, 
+	AVALANCHEVIEWPORT_API FAvaViewportWorldCoordinateConverterProviderPerspective(FVector InLocation, FRotator InRotation, FVector2f InViewportSize, 
 		float InFieldOfView);
 
+protected:
 	//~ Begin IAvaViewportWorldCoordinateConverter
 	virtual FVector2D GetFrustumSizeAtDistance(double InDistance) const override;
 	virtual FVector ViewportPositionToWorldPosition(const FVector2f& InViewportPosition, double InDistance) const override;
 	virtual void WorldPositionToViewportPosition(const FVector& InWorldPosition, FVector2f& OutViewportPosition, double& OutDistance) const override;
 	//~End IAvaViewportWorldCoordinateConverter
 
-protected:
 	float FieldOfView;
 };
 
-struct AVALANCHEVIEWPORT_API FAvaViewportWorldCoordinateConverterProviderOrthographic : public FAvaViewportWorldCoordinateConverterProvider
+struct FAvaViewportWorldCoordinateConverterProviderOrthographic : public FAvaViewportWorldCoordinateConverterProvider
 {
-	FAvaViewportWorldCoordinateConverterProviderOrthographic(FVector InLocation, FRotator InRotation, FVector2f InViewportSize, 
+	AVALANCHEVIEWPORT_API FAvaViewportWorldCoordinateConverterProviderOrthographic(FVector InLocation, FRotator InRotation, FVector2f InViewportSize, 
 		float InOrthographicWidth);
 
+protected:
 	//~ Begin IAvaViewportWorldCoordinateConverter
 	virtual FVector2D GetFrustumSizeAtDistance(double InDistance) const override;
 	virtual FVector ViewportPositionToWorldPosition(const FVector2f& InViewportPosition, double InDistance) const override;
 	virtual void WorldPositionToViewportPosition(const FVector& InWorldPosition, FVector2f& OutViewportPosition, double& OutDistance) const override;
 	//~End IAvaViewportWorldCoordinateConverter
 
-protected:
 	float OrthographicWidth;
 };
diff --git a/Engine/Plugins/Experimental/Avalanche/Source/AvalancheViewport/Public/Visualizers/AvaViewportBackgroundVisualizer.h b/Engine/Plugins/Experimental/Avalanche/Source/AvalancheViewport/Public/Visualizers/AvaViewportBackgroundVisualizer.h
deleted file mode 100644
diff --git a/Engine/Plugins/Experimental/Avalanche/Source/AvalancheViewport/Public/Visualizers/AvaViewportChannelVisualizer.h b/Engine/Plugins/Experimental/Avalanche/Source/AvalancheViewport/Public/Visualizers/AvaViewportChannelVisualizer.h
deleted file mode 100644
diff --git a/Engine/Plugins/Experimental/Avalanche/Source/AvalancheViewport/Public/Visualizers/AvaViewportCheckerboardVisualizer.h b/Engine/Plugins/Experimental/Avalanche/Source/AvalancheViewport/Public/Visualizers/AvaViewportCheckerboardVisualizer.h
deleted file mode 100644
diff --git a/Engine/Plugins/Experimental/Avalanche/Source/AvalancheViewport/Public/Visualizers/AvaViewportPostProcessVisualizer.h b/Engine/Plugins/Experimental/Avalanche/Source/AvalancheViewport/Public/Visualizers/AvaViewportPostProcessVisualizer.h
deleted file mode 100644
diff --git a/Engine/Plugins/Experimental/Avalanche/Source/AvalancheViewport/Public/Visualizers/IAvaViewportBoundingBoxVisualizer.h b/Engine/Plugins/Experimental/Avalanche/Source/AvalancheViewport/Public/Visualizers/IAvaViewportBoundingBoxVisualizer.h
new file mode 100644
--- /dev/null
+++ b/Engine/Plugins/Experimental/Avalanche/Source/AvalancheViewport/Public/Visualizers/IAvaViewportBoundingBoxVisualizer.h
@@ -0,0 +1,34 @@
+// Copyright Epic Games, Inc. All Rights Reserved.
+
+#pragma once
+
+#include "AvaType.h"
+#include "Templates/SharedPointer.h"
+
+class AActor;
+class FPrimitiveDrawInterface;
+class UAvaBoundsProviderSubsystem;
+class UAvaSelectionProviderSubsystem;
+
+enum class EAvaViewportBoundingBoxOptimizationState : uint8
+{
+	RenderNothing,
+	RenderSelectedActors,
+	RenderSelectedActorsAndChildren,
+	RenderSelectionBounds
+};
+
+class IAvaViewportBoundingBoxVisualizer : public IAvaTypeCastable
+{
+public:
+	UE_AVA_INHERITS(IAvaViewportBoundingBoxVisualizer, IAvaTypeCastable)
+
+	virtual void Draw(UAvaSelectionProviderSubsystem& InSelectionProvider, UAvaBoundsProviderSubsystem& InBoundsProvider, FPrimitiveDrawInterface& InPDI) = 0;
+	virtual EAvaViewportBoundingBoxOptimizationState GetOptimizationState() const = 0;
+	virtual void ResetOptimizationState() = 0;
+};
+
+struct FAvaViewportBoundingBoxVisualizerProvider
+{
+	AVALANCHEVIEWPORT_API static TSharedRef<IAvaViewportBoundingBoxVisualizer> CreateVisualizer();
+};
diff --git a/Engine/Plugins/Experimental/Avalanche/Source/AvalancheViewport/Public/Visualizers/IAvaViewportBoundingBoxVisualizerInterface.h b/Engine/Plugins/Experimental/Avalanche/Source/AvalancheViewport/Public/Visualizers/IAvaViewportBoundingBoxVisualizerInterface.h
deleted file mode 100644
diff --git a/Engine/Plugins/Experimental/Avalanche/Source/AvalancheViewport/Public/Visualizers/IAvaViewportPostProcessVisualizer.h b/Engine/Plugins/Experimental/Avalanche/Source/AvalancheViewport/Public/Visualizers/IAvaViewportPostProcessVisualizer.h
new file mode 100644
--- /dev/null
+++ b/Engine/Plugins/Experimental/Avalanche/Source/AvalancheViewport/Public/Visualizers/IAvaViewportPostProcessVisualizer.h
@@ -0,0 +1,13 @@
+// Copyright Epic Games, Inc. All Rights Reserved.
+
+#pragma once
+
+#include "AvaType.h"
+
+class IAvaViewportPostProcessVisualizer : public IAvaTypeCastable
+{
+public:
+	UE_AVA_INHERITS(IAvaViewportPostProcessVisualizer, IAvaTypeCastable)
+
+	virtual bool CanActivate(bool bInSilent) const = 0;
+};
diff --git a/Engine/Plugins/Experimental/Avalanche/Source/AvalancheViewport/Public/Widgets/SAvaDraggableBoxOverlay.h b/Engine/Plugins/Experimental/Avalanche/Source/AvalancheViewport/Public/Widgets/SAvaDraggableBoxOverlay.h
--- a/Engine/Plugins/Experimental/Avalanche/Source/AvalancheViewport/Public/Widgets/SAvaDraggableBoxOverlay.h
+++ b/Engine/Plugins/Experimental/Avalanche/Source/AvalancheViewport/Public/Widgets/SAvaDraggableBoxOverlay.h
@@ -17,16 +17,14 @@ class SBox;
  * Class that can be used to place a draggable box into a viewport or some other large widget as an
  * overlay. Just place the widget that you want to be draggable as the contents of SAvaDraggableBoxOverlay.
  */
-class AVALANCHEVIEWPORT_API SAvaDraggableBoxOverlay : public SCompoundWidget
+class SAvaDraggableBoxOverlay : public SCompoundWidget
 {
 public:
-
-	SLATE_BEGIN_ARGS(SAvaDraggableBoxOverlay) 
-		{}
+	SLATE_BEGIN_ARGS(SAvaDraggableBoxOverlay)  {}
 	SLATE_DEFAULT_SLOT(FArguments, Content)
 	SLATE_END_ARGS()
 
-	void Construct(const FArguments& InArgs);
+	AVALANCHEVIEWPORT_API void Construct(const FArguments& InArgs);
 
 	/* Returns the current position of the box on the viewport. */
 	FVector2f GetBoxAlignmentOffset() const;
@@ -47,8 +45,8 @@ public:
 protected:
 	TSharedPtr<SBox> Container;
 	TSharedPtr<SWidget> DraggableBox;
-	EHorizontalAlignment HorizontalAlignment;
-	EVerticalAlignment VerticalAlignment;
+	EHorizontalAlignment HorizontalAlignment = HAlign_Fill;
+	EVerticalAlignment VerticalAlignment = VAlign_Fill;
 	FMargin Padding;
 
 	FMargin GetPadding() const;
