Motion Design Text:
- Further API export cleanup

#changelist validated
#virtualized

diff --git a/Engine/Plugins/Experimental/Avalanche/Source/AvalancheText/Private/AvaText3DComponent.cpp b/Engine/Plugins/Experimental/Avalanche/Source/AvalancheText/Private/AvaText3DComponent.cpp
--- a/Engine/Plugins/Experimental/Avalanche/Source/AvalancheText/Private/AvaText3DComponent.cpp
+++ b/Engine/Plugins/Experimental/Avalanche/Source/AvalancheText/Private/AvaText3DComponent.cpp
@@ -149,9 +149,6 @@ void UAvaText3DComponent::PostEditChangeProperty(FPropertyChangedEvent& Property
 void UAvaText3DComponent::PostEditComponentMove(bool bFinished)
 {
 	Super::PostEditComponentMove(bFinished);
-
-	OnEditMoveDelegate.Broadcast();
-	
 	ScheduleTextRefreshOnTicker();
 }
 
diff --git a/Engine/Plugins/Experimental/Avalanche/Source/AvalancheText/Private/Font/AvaFont.cpp b/Engine/Plugins/Experimental/Avalanche/Source/AvalancheText/Private/Font/AvaFont.cpp
--- a/Engine/Plugins/Experimental/Avalanche/Source/AvalancheText/Private/Font/AvaFont.cpp
+++ b/Engine/Plugins/Experimental/Avalanche/Source/AvalancheText/Private/Font/AvaFont.cpp
@@ -148,6 +148,24 @@ EAvaFontSource FAvaFont::GetFontSource() const
 	return EAvaFontSource::Invalid;
 }
 
+bool FAvaFont::operator==(const FAvaFont& Other) const
+{
+	if (CurrentFont_DEPRECATED)
+	{
+		return (CurrentFont_DEPRECATED == Other.CurrentFont_DEPRECATED) && (GetFontName() == Other.GetFontName());
+	}
+
+	const UAvaFontObject* MyFontObject = GetFontObject();
+	const UAvaFontObject* OtherFontObject = Other.GetFontObject();
+
+	if (MyFontObject && OtherFontObject && MyFontObject->GetFont() != OtherFontObject->GetFont())
+	{
+		return false;
+	}
+
+	return (GetFontName() == Other.GetFontName());
+}
+
 FName FAvaFont::GetFontName() const
 {
 	return FName(GetFontNameAsString());
@@ -163,11 +181,6 @@ FString FAvaFont::GetFontNameAsString() const
 	return FontName;
 }
 
-FText FAvaFont::GetFontNameAsText() const
-{
-	return FText::FromName(GetFontName());
-}
-
 bool FAvaFont::IsFavorite() const
 {
 	return bIsFavorite;
@@ -276,7 +289,7 @@ void FAvaFont::InitFromFont(UFont* InFont)
 	if (InFont)
 	{
 		FString Name;
-		UE::Ava::FontUtilities::Public::GetFontName(InFont, Name);
+		UE::Ava::FontUtilities::GetFontName(InFont, Name);
 
 		MotionDesignFontObject = NewObject<UAvaFontObject>();
 		MotionDesignFontObject->InitProjectFont(InFont, Name);
@@ -397,7 +410,7 @@ UFont* FAvaFont::GetFontByName(const FString& InFontName)
 		if (UFont* const CurrFont = Cast<UFont>(AssetData.GetAsset()))
 		{
 			FString CurrFontName;
-			UE::Ava::FontUtilities::Public::GetFontName(CurrFont, CurrFontName);
+			UE::Ava::FontUtilities::GetFontName(CurrFont, CurrFontName);
 
 			if (CurrFontName == InFontName)
 			{
@@ -410,7 +423,7 @@ UFont* FAvaFont::GetFontByName(const FString& InFontName)
 }
 
 // note: this function used to be in the AvalancheEditor module, since it was not needed at Runtime before 
-void UE::Ava::FontUtilities::Public::GetFontName(const UFont* InFont, FString& OutFontName)
+void UE::Ava::FontUtilities::GetFontName(const UFont* InFont, FString& OutFontName)
 {
 	if (IsValid(InFont))
 	{
diff --git a/Engine/Plugins/Experimental/Avalanche/Source/AvalancheText/Private/Font/AvalancheFontObject.cpp b/Engine/Plugins/Experimental/Avalanche/Source/AvalancheText/Private/Font/AvalancheFontObject.cpp
--- a/Engine/Plugins/Experimental/Avalanche/Source/AvalancheText/Private/Font/AvalancheFontObject.cpp
+++ b/Engine/Plugins/Experimental/Avalanche/Source/AvalancheText/Private/Font/AvalancheFontObject.cpp
@@ -29,11 +29,6 @@ static TMap<int32, FString> CharsetMap
 	{186, TEXT("Baltic")},        // BALTIC_CHARSET
 };
 
-void FAvaSystemFontMetrics::AddCharset(int32 InCharset)
-{
-	Charsets.AddUnique(InCharset);
-}
-
 void UAvaFontObject::InitProjectFont(UFont* InFont, const FString& InFontName)
 {
 	Font = InFont;
diff --git a/Engine/Plugins/Experimental/Avalanche/Source/AvalancheText/Public/AvaText3DComponent.h b/Engine/Plugins/Experimental/Avalanche/Source/AvalancheText/Public/AvaText3DComponent.h
--- a/Engine/Plugins/Experimental/Avalanche/Source/AvalancheText/Public/AvaText3DComponent.h
+++ b/Engine/Plugins/Experimental/Avalanche/Source/AvalancheText/Public/AvaText3DComponent.h
@@ -22,8 +22,8 @@ enum class EAvaTextRefreshReason : uint8
 };
 ENUM_CLASS_FLAGS(EAvaTextRefreshReason)
 
-UCLASS(ClassGroup="Text3D", PrioritizeCategories=("Text", "Layout", "Geometry", "Style", "Materials"), meta=(BlueprintSpawnableComponent))
-class AVALANCHETEXT_API UAvaText3DComponent : public UText3DComponent
+UCLASS(MinimalAPI, ClassGroup="Text3D", PrioritizeCategories=("Text", "Layout", "Geometry", "Style", "Materials"), meta=(BlueprintSpawnableComponent))
+class UAvaText3DComponent : public UText3DComponent
 {
 	friend class FAvaTextVisualizer;
 	friend class FAvaToolboxTextVisualizer;
@@ -35,80 +35,91 @@ public:
 	UAvaText3DComponent();
 
 	UFUNCTION(BlueprintSetter)
-	void SetEnforceUpperCase(bool bInEnforceUpperCase);
+	AVALANCHETEXT_API void SetEnforceUpperCase(bool bInEnforceUpperCase);
 	bool IsEnforcingUpperCase() const { return bEnforceUpperCase; }
-	
+
 	UFUNCTION(BlueprintSetter)
-	void SetMotionDesignFont(const FAvaFont& InFont);
+	AVALANCHETEXT_API void SetMotionDesignFont(const FAvaFont& InFont);
 	const FAvaFont& GetMotionDesignFont() const { return MotionDesignFont; }
 
 	UFUNCTION(BlueprintSetter)
-	void SetAlignment(FAvaTextAlignment InAlignment);
+	AVALANCHETEXT_API void SetAlignment(FAvaTextAlignment InAlignment);
 	const FAvaTextAlignment& GetAlignment() const { return Alignment; }
 
 	UFUNCTION(BlueprintSetter)
-	void SetColoringStyle(const EAvaTextColoringStyle& InColoringStyle);
+	AVALANCHETEXT_API void SetColoringStyle(const EAvaTextColoringStyle& InColoringStyle);
 	EAvaTextColoringStyle GetColoringStyle() const { return ColoringStyle; }
 
 	UFUNCTION(BlueprintSetter)
-	void SetColor(const FLinearColor& InColor);
+	AVALANCHETEXT_API void SetColor(const FLinearColor& InColor);
 	const FLinearColor& GetColor() const { return Color; }
 
 	UFUNCTION(BlueprintSetter)
-	void SetExtrudeColor(const FLinearColor& InColor);
+	AVALANCHETEXT_API void SetExtrudeColor(const FLinearColor& InColor);
 	const FLinearColor& GetExtrudeColor() const { return ExtrudeColor; }
 
 	UFUNCTION(BlueprintSetter)
-	void SetBevelColor(const FLinearColor& InColor);
+	AVALANCHETEXT_API void SetBevelColor(const FLinearColor& InColor);
 	const FLinearColor& GetBevelColor() const { return BevelColor; }
 
 	UFUNCTION(BlueprintSetter)
-	void SetGradientSettings(const FAvaLinearGradientSettings& InGradientSettings);
+	AVALANCHETEXT_API void SetGradientSettings(const FAvaLinearGradientSettings& InGradientSettings);
 	const FAvaLinearGradientSettings& GetGradientSettings() const { return GradientSettings; }
 
 	UFUNCTION(BlueprintSetter)
-	void SetGradientColors(const FLinearColor& InColorA, const FLinearColor& InColorB);
+	AVALANCHETEXT_API void SetGradientColors(const FLinearColor& InColorA, const FLinearColor& InColorB);
 
 	UFUNCTION(BlueprintSetter)
-	void SetMainTexture(UTexture2D* InMainTexture);
+	AVALANCHETEXT_API void SetMainTexture(UTexture2D* InMainTexture);
 	UTexture2D* GetMainTexture() const { return MainTexture; }
 
 	UFUNCTION(BlueprintSetter)
-	void SetTiling(const FVector2D InTiling);
+	AVALANCHETEXT_API void SetTiling(const FVector2D InTiling);
 	const FVector2D& GetTiling() const { return Tiling; }
 
 	UFUNCTION(BlueprintSetter)
-	void SetCustomMaterial(UMaterialInterface* InCustomMaterial);
+	AVALANCHETEXT_API void SetCustomMaterial(UMaterialInterface* InCustomMaterial);
 	UMaterialInterface* GetCustomMaterial() const { return CustomMaterial; }
 
 	UFUNCTION(BlueprintSetter)
-	void SetTranslucencyStyle(EAvaTextTranslucency InTranslucencyStyle);
+	AVALANCHETEXT_API void SetTranslucencyStyle(EAvaTextTranslucency InTranslucencyStyle);
 	EAvaTextTranslucency GetTranslucencyStyle() const { return TranslucencyStyle; }
 
 	UFUNCTION(BlueprintSetter)
-	void SetOpacity(float InOpacity);
+	AVALANCHETEXT_API void SetOpacity(float InOpacity);
 	float GetOpacity() const { return Opacity; }
-	
+
 	UFUNCTION(BlueprintSetter)
-	void SetMaskOrientation(EAvaMaterialMaskOrientation InMaskOrientation);
+	AVALANCHETEXT_API void SetMaskOrientation(EAvaMaterialMaskOrientation InMaskOrientation);
 	EAvaMaterialMaskOrientation GetMaskOrientation() const { return MaskOrientation; }
 
 	UFUNCTION(BlueprintSetter)
-	void SetMaskSmoothness(float InMaskSmoothness);
+	AVALANCHETEXT_API void SetMaskSmoothness(float InMaskSmoothness);
 	float GetMaskSmoothness() const { return MaskSmoothness; }
 
 	UFUNCTION(BlueprintSetter)
-	void SetMaskOffset(float InMaskOffset);
+	AVALANCHETEXT_API void SetMaskOffset(float InMaskOffset);
 	float GetMaskOffset() const { return MaskOffset; }
 
 	UFUNCTION(BlueprintSetter)
-	void SetMaskRotation(float InMaskRotation);
+	AVALANCHETEXT_API void SetMaskRotation(float InMaskRotation);
 	float GetMaskRotation() const { return MaskRotation; }
 
 	UFUNCTION(BlueprintSetter)
-	void SetIsUnlit(bool bInIsUnlit);
+	AVALANCHETEXT_API void SetIsUnlit(bool bInIsUnlit);
 	bool GetIsUnlit() const { return bIsUnlit; }
 
+	using FOnGetMaterialWithSettings = TDelegate<UMaterialInterface*(const UMaterialInterface* InPreviousMaterial, const FAvaTextMaterialSettings& InSettings)>;
+
+	AVALANCHETEXT_API FOnGetMaterialWithSettings& GetMaterialProviderDelegate();
+
+	AVALANCHETEXT_API void ForEachMID(TUniqueFunction<void(UMaterialInstanceDynamic* InMID)>&& InFunc) const;
+
+	AVALANCHETEXT_API FVector GetGradientDirection() const;
+
+	AVALANCHETEXT_API void RefreshMaterialInstances();
+
+private:
 #if WITH_EDITOR
 	static void RegisterOnPropertyChangeFunctions();
 #endif
@@ -125,11 +136,10 @@ public:
 	void RefreshVerticalAlignment();
 	void RefreshUnlit();
 	void RefreshMasked();
-	void RefreshMaterialInstances();
 	void RefreshGeometry();
 	void RefreshLayout();
 	void RefreshGradientValues();
-	void RefreshGradient();
+	AVALANCHETEXT_API void RefreshGradient();
 	void RefreshMaskValues();
 	void RefreshMask();
 	void RefreshMaterialGeometryParameterValues();
@@ -138,29 +148,26 @@ public:
 
 	void SetMaterialGeometryParameterValues(UMaterialInstanceDynamic* InMaterial, const FVector& TextScaleFactor) const;
 	void SetMaskedMaterialValues(UMaterialInstanceDynamic* InMaterial);
-	using FOnGetMaterialWithSettings = TDelegate<UMaterialInterface*(const UMaterialInterface* InPreviousMaterial, const FAvaTextMaterialSettings& InSettings)>;
-	FOnGetMaterialWithSettings& GetMaterialProviderDelegate();
-	void ForEachMID(TUniqueFunction<void(UMaterialInstanceDynamic* InMID)>&& InFunc) const;
-
-	FVector GetGradientDirection() const;
 
+	//~ Begin UObject
 #if WITH_EDITOR
-	DECLARE_MULTICAST_DELEGATE(FOnEditMoveDelegate);
-	FOnEditMoveDelegate& OnEditMove() { return OnEditMoveDelegate; } // used by visualizer
-
 	virtual void PostEditChangeProperty(FPropertyChangedEvent& PropertyChangedEvent) override;
-	virtual void PostEditComponentMove(bool bFinished) override;
+	//~ End UObject
 #endif
 
 	//~ Begin USceneComponent
+#if WITH_EDITOR
+	virtual void PostEditComponentMove(bool bFinished) override;
+#endif
 	virtual void OnUpdateTransform(EUpdateTransformFlags UpdateTransformFlags, ETeleportType Teleport) override;
 	virtual void OnAttachmentChanged() override;
 	virtual void OnRegister() override;
 	//~ End USceneComponent
 
+	//~ Begin UText3DComponent
 	virtual void FormatText(FText& InOutText) const override;
+	//~ End UText3DComponent
 
-private:
 	UPROPERTY(EditAnywhere, Getter="IsEnforcingUpperCase", Setter="SetEnforceUpperCase", Category="Text", meta=(DisplayAfter="Text", AllowPrivateAccess="true"))
 	bool bEnforceUpperCase = false;
 	
@@ -331,10 +338,6 @@ protected:
 
 	EAvaTextMaterialFeatures GetMaterialFeaturesFromProperties() const;
 
-#if WITH_EDITOR
-	FOnEditMoveDelegate OnEditMoveDelegate;
-#endif
-
 	/** Used to filter which refreshes need be performed when calling RefreshText3D */
 	EAvaTextRefreshReason TextRefreshReason;
 
diff --git a/Engine/Plugins/Experimental/Avalanche/Source/AvalancheText/Public/AvaTextActor.h b/Engine/Plugins/Experimental/Avalanche/Source/AvalancheText/Public/AvaTextActor.h
--- a/Engine/Plugins/Experimental/Avalanche/Source/AvalancheText/Public/AvaTextActor.h
+++ b/Engine/Plugins/Experimental/Avalanche/Source/AvalancheText/Public/AvaTextActor.h
@@ -9,19 +9,22 @@
 class UAvaTextCharacterTransform;
 struct FAvaColorChangeData;
 
-UCLASS(ClassGroup = (Text3D), DisplayName = "Motion Design Text", meta = (ComponentWrapperClass))
-class AVALANCHETEXT_API AAvaTextActor : public AActor
+UCLASS(MinimalAPI, ClassGroup = (Text3D), DisplayName = "Motion Design Text", meta = (ComponentWrapperClass))
+class AAvaTextActor : public AActor
 {
 	GENERATED_BODY()
 
 public:
 	AAvaTextActor();
 
-	/** Returns Text3D SubObject **/
-	UText3DComponent* GetText3DComponent() const { return Text3DComponent; }
+	UText3DComponent* GetText3DComponent() const
+	{
+		return Text3DComponent;
+	}
 
-	FAvaColorChangeData GetColorData() const;
-	void SetColorData(const FAvaColorChangeData& InColorData);
+	AVALANCHETEXT_API FAvaColorChangeData GetColorData() const;
+
+	AVALANCHETEXT_API void SetColorData(const FAvaColorChangeData& InColorData);
 
 private:
 	UPROPERTY(Category = "Text", VisibleAnywhere, BlueprintReadOnly, meta = (AllowPrivateAccess = "true"))
diff --git a/Engine/Plugins/Experimental/Avalanche/Source/AvalancheText/Public/AvaTextCharacterTransform.h b/Engine/Plugins/Experimental/Avalanche/Source/AvalancheText/Public/AvaTextCharacterTransform.h
--- a/Engine/Plugins/Experimental/Avalanche/Source/AvalancheText/Public/AvaTextCharacterTransform.h
+++ b/Engine/Plugins/Experimental/Avalanche/Source/AvalancheText/Public/AvaTextCharacterTransform.h
@@ -9,8 +9,8 @@
 /**
  * This Component is kept for compatibility with old AAvaTextActor assets
  */
-UCLASS(ClassGroup=(Text3D), HideCategories = (Collision, Tags, Activation, Cooking, Rendering, Physics, Mobility, LOD, AssetUserData, Navigation, Transform), meta = (BlueprintSpawnableComponent))
-class AVALANCHETEXT_API UAvaTextCharacterTransform : public UText3DCharacterTransform
+UCLASS(MinimalAPI, ClassGroup=(Text3D), HideCategories = (Collision, Tags, Activation, Cooking, Rendering, Physics, Mobility, LOD, AssetUserData, Navigation, Transform), meta = (BlueprintSpawnableComponent))
+class UAvaTextCharacterTransform : public UText3DCharacterTransform
 {
 	GENERATED_BODY()
 };
diff --git a/Engine/Plugins/Experimental/Avalanche/Source/AvalancheText/Public/Font/AvaFont.h b/Engine/Plugins/Experimental/Avalanche/Source/AvalancheText/Public/Font/AvaFont.h
--- a/Engine/Plugins/Experimental/Avalanche/Source/AvalancheText/Public/Font/AvaFont.h
+++ b/Engine/Plugins/Experimental/Avalanche/Source/AvalancheText/Public/Font/AvaFont.h
@@ -9,7 +9,7 @@
 #include "AvaFont.generated.h"
 
 USTRUCT(BlueprintType, meta=(DisplayName="Motion Design Font"))
-struct AVALANCHETEXT_API FAvaFont
+struct FAvaFont
 {
 	friend class UAvaText3DComponent;
 
@@ -19,7 +19,7 @@ struct AVALANCHETEXT_API FAvaFont
 	 * Returns a default font (Roboto). The font will be loaded from asset if needed.
 	 * In the unlikely case in which Roboto is not available, the first available font will be returned
 	 */
-	static UFont* GetDefaultFont();
+	AVALANCHETEXT_API static UFont* GetDefaultFont();
 
 	/**
 	 * Returns a formatted string representing a FAvaFont for the specified FontObjectPathName and FontName
@@ -37,7 +37,7 @@ struct AVALANCHETEXT_API FAvaFont
 	 * @param OutFormattedString {will contain the formatted string}
 	 * @return {true if successful}
 	 */
-	static bool GenerateFontFormattedString(const UAvaFontObject* InFontObject, FString& OutFormattedString);
+	AVALANCHETEXT_API static bool GenerateFontFormattedString(const UAvaFontObject* InFontObject, FString& OutFormattedString);
 
 	/**
 	 * Checks if the specified Ava Fonts are the same
@@ -48,10 +48,10 @@ struct AVALANCHETEXT_API FAvaFont
 	static bool AreSameFont(const FAvaFont* InFontA, const FAvaFont* InFontB);
 
 	/** Default constructor: calling GetFont() right after the constructor will return the default font */
-	FAvaFont();
+	AVALANCHETEXT_API FAvaFont();
 
 	/** this FAvaFont will be initialized using the provided UAvaFontObject */
-	FAvaFont(UAvaFontObject* InFontObject);
+	AVALANCHETEXT_API FAvaFont(UAvaFontObject* InFontObject);
 
 	/** Returns the currently stored font, or default font if none available. Preferred getter to just get the UFont and use it */
 	UFont* GetFont();
@@ -60,13 +60,10 @@ struct AVALANCHETEXT_API FAvaFont
 	FName GetFontName() const;
 
 	/** Returns current font name as FString */
-	FString GetFontNameAsString() const;
-
-	/** Returns current font name as FText */
-	FText GetFontNameAsText() const;
+	AVALANCHETEXT_API FString GetFontNameAsString() const;
 
 	/** Returns true if this FAvaFont is marked as favorite by the user */
-	bool IsFavorite() const;
+	AVALANCHETEXT_API bool IsFavorite() const;
 
 	/**
 	 * Checks if this FAvaFont is referencing the default font. This will be happening also in fallback state.
@@ -81,7 +78,7 @@ struct AVALANCHETEXT_API FAvaFont
 	 * In fallback state, a FAvaFont will return the default font when GetFont() is called.
 	 * @return true if this FAvaFont is in fallback state
 	 */
-	bool IsFallbackFont() const;
+	AVALANCHETEXT_API bool IsFallbackFont() const;
 
 	/** Returns true if font is monospaced */
 	bool IsMonospaced() const;
@@ -96,7 +93,7 @@ struct AVALANCHETEXT_API FAvaFont
 	 * Mark/unmark this font as favorite
 	 * @param bFavorite {the flag to mark the font as favorite or not. True means favorite.}
 	 */
-	void SetFavorite(const bool bFavorite);
+	AVALANCHETEXT_API void SetFavorite(const bool bFavorite);
 
 	/**
 	 * Calling this function will null the now deprecated UFont CurrentFont reference, and use that UFont to initialize the currently used UAvaFontObject
@@ -133,25 +130,9 @@ struct AVALANCHETEXT_API FAvaFont
 	EAvaFontSource GetFontSource() const;
 
 	/** Fonts are considered equal if they share the same font and font name */
-	FORCEINLINE bool operator==(const FAvaFont& Other) const
-	{
-		if (CurrentFont_DEPRECATED)
-		{
-			return (CurrentFont_DEPRECATED == Other.CurrentFont_DEPRECATED) && (GetFontName() == Other.GetFontName());
-		}
-
-		const UAvaFontObject* MyFontObject = GetFontObject();
-		const UAvaFontObject* OtherFontObject = Other.GetFontObject();
-
-		if (MyFontObject && OtherFontObject && MyFontObject->GetFont() != OtherFontObject->GetFont())
-		{
-			return false;
-		}
-
-		return (GetFontName() == Other.GetFontName());
-	}
+	AVALANCHETEXT_API bool operator==(const FAvaFont& Other) const;
 
-	FORCEINLINE bool operator!=(const FAvaFont& Other) const
+	bool operator!=(const FAvaFont& Other) const
 	{
 		return !(*this == Other);
 	}
@@ -160,7 +141,7 @@ struct AVALANCHETEXT_API FAvaFont
 	 * Called after serialization of FAvaFont.
 	 * In case of missing font assets, the FAvaFont state will be marked as "FallbackFont"
 	 */
-	void PostSerialize(const FArchive& Ar);
+	AVALANCHETEXT_API void PostSerialize(const FArchive& Ar);
 
 private:
 	/** A struct to hold default values which we want to prevent from being GC'd */
@@ -238,7 +219,7 @@ template<> struct TStructOpsTypeTraits<FAvaFont> : public TStructOpsTypeTraitsBa
 	};
 };
 
-namespace UE::Ava::FontUtilities::Public
+namespace UE::Ava::FontUtilities
 {
 	AVALANCHETEXT_API void GetFontName(const UFont* InFont, FString& OutFontName);
 }
diff --git a/Engine/Plugins/Experimental/Avalanche/Source/AvalancheText/Public/Font/AvaFontObject.h b/Engine/Plugins/Experimental/Avalanche/Source/AvalancheText/Public/Font/AvaFontObject.h
--- a/Engine/Plugins/Experimental/Avalanche/Source/AvalancheText/Public/Font/AvaFontObject.h
+++ b/Engine/Plugins/Experimental/Avalanche/Source/AvalancheText/Public/Font/AvaFontObject.h
@@ -18,7 +18,7 @@ enum class EAvaFontSource : uint8
 	Invalid = 2
 };
 
-struct AVALANCHETEXT_API FAvaSystemFontMetrics
+struct FAvaSystemFontMetrics
 {
 	TArray<int32> Charsets; // this will become proper enum
 
@@ -29,11 +29,6 @@ struct AVALANCHETEXT_API FAvaSystemFontMetrics
 	bool bIsItalic;
 
 	FString FamilyName;
-
-	// todo: add other metrics, also check what is already available in UFont
-
-	void AddCharset(int32 InCharset);
-	TArray<int32> GetCharsets() const {return Charsets;}
 };
 
 struct FSystemFontsRetrieveParams
diff --git a/Engine/Plugins/Experimental/Avalanche/Source/AvalancheText/Public/MaterialHub/AvaTextMaterialHub.h b/Engine/Plugins/Experimental/Avalanche/Source/AvalancheText/Public/MaterialHub/AvaTextMaterialHub.h
--- a/Engine/Plugins/Experimental/Avalanche/Source/AvalancheText/Public/MaterialHub/AvaTextMaterialHub.h
+++ b/Engine/Plugins/Experimental/Avalanche/Source/AvalancheText/Public/MaterialHub/AvaTextMaterialHub.h
@@ -115,8 +115,8 @@ private:
  * - Retrieve a Text3D Material with the specified FAvaTextMaterialSettings
  * - Retrieve FAvaTextMaterialSettings for a given material
  */
-UCLASS(Config=Editor)
-class AVALANCHETEXT_API UAvaTextMaterialHub : public UObject
+UCLASS(MinimalAPI, Config=Editor)
+class UAvaTextMaterialHub : public UObject
 {
 	GENERATED_BODY()
 	
@@ -128,7 +128,7 @@ public:
 	 *
 	 * @return A material matching the desired settings
 	 */
-	static UMaterialInterface* GetMaterial(FAvaTextMaterialSettings InMaterialSettings);
+	AVALANCHETEXT_API static UMaterialInterface* GetMaterial(FAvaTextMaterialSettings InMaterialSettings);
 
 	/**
 	 * Can be used to verify if the provided material has the same type as one of the materials handled by FAvaTextMaterials.
diff --git a/Engine/Plugins/Experimental/Avalanche/Source/AvalancheTextEditor/Private/Font/AvaFontManagerSubsystem.cpp b/Engine/Plugins/Experimental/Avalanche/Source/AvalancheTextEditor/Private/Font/AvaFontManagerSubsystem.cpp
--- a/Engine/Plugins/Experimental/Avalanche/Source/AvalancheTextEditor/Private/Font/AvaFontManagerSubsystem.cpp
+++ b/Engine/Plugins/Experimental/Avalanche/Source/AvalancheTextEditor/Private/Font/AvaFontManagerSubsystem.cpp
@@ -976,7 +976,7 @@ void UAvaFontManagerSubsystem::GetMultipleSelectionInformation(const TSharedPtr<
 void UAvaFontManagerSubsystem::GetFontName(const UFont* InFont, FString& OutFontName)
 {
 	// this function used to be defined here, but it has been moved to the Runtime module for easier maintenance
-	UE::Ava::FontUtilities::Public::GetFontName(InFont, OutFontName);
+	UE::Ava::FontUtilities::GetFontName(InFont, OutFontName);
 }
 
 void UAvaFontManagerSubsystem::GetSanitizedFontName(const UFont* InFont, FString& OutFontName)
