Motion Design:
- Shapes Editor API Cleanup

#changelist validated
#virtualized

diff --git a/Engine/Plugins/Experimental/Avalanche/Source/AvalancheEffectorsEditor/AvalancheEffectorsEditor.Build.cs b/Engine/Plugins/Experimental/Avalanche/Source/AvalancheEffectorsEditor/AvalancheEffectorsEditor.Build.cs
--- a/Engine/Plugins/Experimental/Avalanche/Source/AvalancheEffectorsEditor/AvalancheEffectorsEditor.Build.cs
+++ b/Engine/Plugins/Experimental/Avalanche/Source/AvalancheEffectorsEditor/AvalancheEffectorsEditor.Build.cs
@@ -8,6 +8,13 @@ public class AvalancheEffectorsEditor : ModuleRules
 	{
 		PCHUsage = ModuleRules.PCHUsageMode.UseExplicitOrSharedPCHs;
 
+		PrivateIncludePathModuleNames.AddRange(
+			new string[]
+			{
+				"AvalancheShapesEditor",
+			}
+		);
+
 		PrivateDependencyModuleNames.AddRange(
 			new string[]
 			{
@@ -16,7 +23,6 @@ public class AvalancheEffectorsEditor : ModuleRules
 				"AvalancheEffectors",
 				"AvalancheInteractiveTools",
 				"AvalancheShapes",
-				"AvalancheShapesEditor",
 				"ClonerEffector",
 				"ClonerEffectorEditor",
 				"ComponentVisualizers",
diff --git a/Engine/Plugins/Experimental/Avalanche/Source/AvalancheEffectorsEditor/Private/AvalancheEffectorsEditorModule.cpp b/Engine/Plugins/Experimental/Avalanche/Source/AvalancheEffectorsEditor/Private/AvalancheEffectorsEditorModule.cpp
--- a/Engine/Plugins/Experimental/Avalanche/Source/AvalancheEffectorsEditor/Private/AvalancheEffectorsEditorModule.cpp
+++ b/Engine/Plugins/Experimental/Avalanche/Source/AvalancheEffectorsEditor/Private/AvalancheEffectorsEditorModule.cpp
@@ -4,7 +4,6 @@
 #include "AvaEffectorsEditorCommands.h"
 #include "AvaEffectorsEditorStyle.h"
 #include "AvaInteractiveToolsDelegates.h"
-#include "AvalancheShapesEditorModule.h"
 #include "Cloner/AvaClonerActorTool.h"
 #include "Cloner/AvaClonerActorVis.h"
 #include "Cloner/CEClonerComponent.h"
diff --git a/Engine/Plugins/Experimental/Avalanche/Source/AvalancheEffectorsEditor/Private/Cloner/AvaClonerActorVis.cpp b/Engine/Plugins/Experimental/Avalanche/Source/AvalancheEffectorsEditor/Private/Cloner/AvaClonerActorVis.cpp
--- a/Engine/Plugins/Experimental/Avalanche/Source/AvalancheEffectorsEditor/Private/Cloner/AvaClonerActorVis.cpp
+++ b/Engine/Plugins/Experimental/Avalanche/Source/AvalancheEffectorsEditor/Private/Cloner/AvaClonerActorVis.cpp
@@ -1,9 +1,8 @@
 // Copyright Epic Games, Inc. All Rights Reserved.
 
 #include "Cloner/AvaClonerActorVis.h"
-
 #include "AvaField.h"
-#include "AvalancheShapesEditorModule.h"
+#include "AvaShapeSprites.h"
 #include "Cloner/CEClonerActor.h"
 #include "Cloner/CEClonerComponent.h"
 #include "Cloner/Layouts/CEClonerCircleLayout.h"
@@ -677,7 +676,7 @@ FVector FAvaClonerActorVisualizer::GetHandleSpacingLocation(const ACEClonerActor
 
 void FAvaClonerActorVisualizer::DrawSpacingButton(const ACEClonerActor* InClonerActor, const FSceneView* InView, FPrimitiveDrawInterface* InPDI, int32 InIconIndex, ECEClonerAxis InAxis, FLinearColor InColor) const
 {
-	UTexture2D* SpacingSprite = IAvalancheComponentVisualizersModule::Get().GetSettings()->GetVisualizerSprite(FAvalancheShapesEditorModule::BevelSprite);
+	UTexture2D* SpacingSprite = IAvalancheComponentVisualizersModule::Get().GetSettings()->GetVisualizerSprite(UE::AvaShapes::BevelSprite);
 
 	if (!SpacingSprite || !SpacingSprite->GetResource() || !InClonerActor->GetEnabled() || InClonerActor->GetMeshCount() == 0)
 	{
diff --git a/Engine/Plugins/Experimental/Avalanche/Source/AvalancheEffectorsEditor/Private/Effector/AvaEffectorActorVis.cpp b/Engine/Plugins/Experimental/Avalanche/Source/AvalancheEffectorsEditor/Private/Effector/AvaEffectorActorVis.cpp
--- a/Engine/Plugins/Experimental/Avalanche/Source/AvalancheEffectorsEditor/Private/Effector/AvaEffectorActorVis.cpp
+++ b/Engine/Plugins/Experimental/Avalanche/Source/AvalancheEffectorsEditor/Private/Effector/AvaEffectorActorVis.cpp
@@ -1,9 +1,8 @@
 // Copyright Epic Games, Inc. All Rights Reserved.
 
 #include "Effector/AvaEffectorActorVis.h"
-
 #include "AvaField.h"
-#include "AvalancheShapesEditorModule.h"
+#include "AvaShapeSprites.h"
 #include "AvaVisBase.h"
 #include "EditorViewportClient.h"
 #include "Effector/CEEffectorComponent.h"
@@ -241,7 +240,7 @@ FVector FAvaEffectorActorVisualizer::GetHandleZoneLocation(const ACEEffectorActo
 
 void FAvaEffectorActorVisualizer::DrawZoneButton(const ACEEffectorActor* InEffectorActor, const FSceneView* InView, FPrimitiveDrawInterface* InPDI, int32 InIconIndex, bool bInInnerZone, FLinearColor InColor) const
 {
-	UTexture2D* ZoneSprite = IAvalancheComponentVisualizersModule::Get().GetSettings()->GetVisualizerSprite(FAvalancheShapesEditorModule::BevelSprite);
+	UTexture2D* ZoneSprite = IAvalancheComponentVisualizersModule::Get().GetSettings()->GetVisualizerSprite(UE::AvaShapes::BevelSprite);
 
 	if (!ZoneSprite || !ZoneSprite->GetResource())
 	{
diff --git a/Engine/Plugins/Experimental/Avalanche/Source/AvalancheShapesEditor/AvalancheShapesEditor.Build.cs b/Engine/Plugins/Experimental/Avalanche/Source/AvalancheShapesEditor/AvalancheShapesEditor.Build.cs
--- a/Engine/Plugins/Experimental/Avalanche/Source/AvalancheShapesEditor/AvalancheShapesEditor.Build.cs
+++ b/Engine/Plugins/Experimental/Avalanche/Source/AvalancheShapesEditor/AvalancheShapesEditor.Build.cs
@@ -12,11 +12,6 @@ public class AvalancheShapesEditor : ModuleRules
 			new string[]
 			{
 				"Core",
-				"CoreUObject",
-				"DeveloperSettings",
-				"Engine",
-				"Slate",
-				"SlateCore"
 			}
 		);
 
@@ -31,7 +26,10 @@ public class AvalancheShapesEditor : ModuleRules
 				"AvalancheShapes",
 				"AvalancheViewport",
 				"ComponentVisualizers",
+				"CoreUObject",
+				"DeveloperSettings",
 				"EditorFramework",
+				"Engine",
 				"GeometryCore",
 				"GeometryFramework",
 				"InputCore",
@@ -42,6 +40,8 @@ public class AvalancheShapesEditor : ModuleRules
 				"MovieSceneTracks",
 				"Projects",
 				"Sequencer",
+				"Slate",
+				"SlateCore",
 				"UnrealEd",
 				"WidgetRegistration"
 			}
diff --git a/Engine/Plugins/Experimental/Avalanche/Source/AvalancheShapesEditor/Private/AvaShapeFactory.h b/Engine/Plugins/Experimental/Avalanche/Source/AvalancheShapesEditor/Private/AvaShapeFactory.h
new file mode 100644
--- /dev/null
+++ b/Engine/Plugins/Experimental/Avalanche/Source/AvalancheShapesEditor/Private/AvaShapeFactory.h
@@ -0,0 +1,30 @@
+// Copyright Epic Games, Inc. All Rights Reserved.
+
+#pragma once
+
+#include "ActorFactories/ActorFactory.h"
+#include "DynamicMeshes/AvaShapeDynMeshBase.h"
+#include "AvaShapeFactory.generated.h"
+
+class UAvaShapeDynamicMeshBase;
+
+UCLASS()
+class UAvaShapeFactory : public UActorFactory
+{
+	GENERATED_BODY()
+
+public:
+	UAvaShapeFactory();
+
+	void SetMeshClass(TSubclassOf<UAvaShapeDynamicMeshBase> InMeshClass);
+
+protected:
+	TSubclassOf<UAvaShapeDynamicMeshBase> MeshClass;
+
+	//~ Begin UActorFactory
+	virtual bool CanCreateActorFrom(const FAssetData& AssetData, FText& OutErrorMsg) override;
+	virtual AActor* GetDefaultActor(const FAssetData& AssetData) override;
+	virtual AActor* SpawnActor(UObject* InAsset, ULevel* InLevel, const FTransform& InTransform, const FActorSpawnParameters& InSpawnParams) override;
+	virtual void PostPlaceAsset(TArrayView<const FTypedElementHandle> InHandle, const FAssetPlacementInfo& InPlacementInfo, const FPlacementOptions& InPlacementOptions) override;
+	//~ End UActorFactory
+};
diff --git a/Engine/Plugins/Experimental/Avalanche/Source/AvalancheShapesEditor/Private/AvalancheShapesEditorModule.cpp b/Engine/Plugins/Experimental/Avalanche/Source/AvalancheShapesEditor/Private/AvalancheShapesEditorModule.cpp
--- a/Engine/Plugins/Experimental/Avalanche/Source/AvalancheShapesEditor/Private/AvalancheShapesEditorModule.cpp
+++ b/Engine/Plugins/Experimental/Avalanche/Source/AvalancheShapesEditor/Private/AvalancheShapesEditorModule.cpp
@@ -3,6 +3,7 @@
 #include "AvalancheShapesEditorModule.h"
 #include "AvaInteractiveToolsDelegates.h"
 #include "AvaShapeActor.h"
+#include "AvaShapeSprites.h"
 #include "AvaShapesEditorCommands.h"
 #include "ComponentVisualizers.h"
 #include "Engine/Texture2D.h"
@@ -60,28 +61,6 @@
 
 #define LOCTEXT_NAMESPACE "AvalancheShapesEditor"
 
-const FName FAvalancheShapesEditorModule::BevelSprite = TEXT("Bevel");
-const FName FAvalancheShapesEditorModule::BreakSideSprite = TEXT("BreakSide");
-const FName FAvalancheShapesEditorModule::ColorSelectionSprite = TEXT("ColorSelection");
-const FName FAvalancheShapesEditorModule::CornerSprite = TEXT("Corner");
-const FName FAvalancheShapesEditorModule::DepthSprite = TEXT("Depth");
-const FName FAvalancheShapesEditorModule::InnerSizeSprite = TEXT("InnerSize");
-const FName FAvalancheShapesEditorModule::LinearGradientSprite = TEXT("LinearGradient");
-const FName FAvalancheShapesEditorModule::NumPointsSprite = TEXT("NumPoints");
-const FName FAvalancheShapesEditorModule::NumSidesSprite = TEXT("NumSides");
-const FName FAvalancheShapesEditorModule::TextMaxHeightSprite = TEXT("TextMaxHeight");
-const FName FAvalancheShapesEditorModule::TextMaxWidthSprite = TEXT("TextMaxWidth");
-const FName FAvalancheShapesEditorModule::TextScaleProportionallySprite = TEXT("TextScaleProportionally");
-const FName FAvalancheShapesEditorModule::SizeSprite = TEXT("Size");
-const FName FAvalancheShapesEditorModule::SlantSprite = TEXT("Slant");
-const FName FAvalancheShapesEditorModule::UVSprite = TEXT("UV");
-
-FAvalancheShapesEditorModule& FAvalancheShapesEditorModule::Get()
-{
-	static const FName ModuleName = TEXT("AvalancheShapesEditor");
-	return FModuleManager::LoadModuleChecked<FAvalancheShapesEditorModule>(ModuleName);
-}
-
 void FAvalancheShapesEditorModule::StartupModule()
 {
 	FAvaShapesEditorCommands::Register();
@@ -96,21 +75,21 @@ void FAvalancheShapesEditorModule::StartupModule()
 		AvaVisSettings->SetDefaultVisualizerSprite(InName, Sprite);
 	};
 
-	RegisterDefaultSprite(BevelSprite, TEXT("Texture2D'/Avalanche/EditorResources/NewBevelHandle.NewBevelHandle'"));
-	RegisterDefaultSprite(BreakSideSprite, TEXT("Texture2D'/Engine/EditorResources/S_Emitter.S_Emitter'"));
-	RegisterDefaultSprite(ColorSelectionSprite, TEXT("Texture2D'/Engine/EditorResources/S_ReflActorIcon.S_ReflActorIcon'"));
-	RegisterDefaultSprite(CornerSprite, TEXT("Texture2D'/Avalanche/EditorResources/Bevel.Bevel'"));
-	RegisterDefaultSprite(DepthSprite, TEXT("Texture2D'/Engine/EditorResources/S_Terrain.S_Terrain'"));
-	RegisterDefaultSprite(InnerSizeSprite, TEXT("Texture2D'/Engine/EditorResources/S_RadForce.S_RadForce'"));
-	RegisterDefaultSprite(LinearGradientSprite, TEXT("Texture2D'/Avalanche/EditorResources/LinearGradient.LinearGradient'"));
-	RegisterDefaultSprite(NumPointsSprite, TEXT("Texture2D'/Engine/EditorResources/S_Emitter.S_Emitter'"));
-	RegisterDefaultSprite(NumSidesSprite, TEXT("Texture2D'/Engine/EditorResources/S_Emitter.S_Emitter'"));
-	RegisterDefaultSprite(SizeSprite, TEXT("Texture2D'/Avalanche/EditorResources/NewSizeHandle.NewSizeHandle'"));
-	RegisterDefaultSprite(SlantSprite, TEXT("Texture2D'/Avalanche/EditorResources/Slant.Slant'"));
-	RegisterDefaultSprite(TextMaxHeightSprite, TEXT("Texture2D'/Engine/EngineResources/Cursors/SplitterVert.SplitterVert'"));
-	RegisterDefaultSprite(TextMaxWidthSprite, TEXT("Texture2D'/Engine/EngineResources/Cursors/SplitterHorz.SplitterHorz'"));
-	RegisterDefaultSprite(TextScaleProportionallySprite, TEXT("Texture2D'/Engine/EditorResources/S_TextRenderActorIcon.S_TextRenderActorIcon'"));
-	RegisterDefaultSprite(UVSprite, TEXT("Texture2D'/Engine/EditorResources/MatInstActSprite.MatInstActSprite'"));
+	RegisterDefaultSprite(UE::AvaShapes::BevelSprite, TEXT("Texture2D'/Avalanche/EditorResources/NewBevelHandle.NewBevelHandle'"));
+	RegisterDefaultSprite(UE::AvaShapes::BreakSideSprite, TEXT("Texture2D'/Engine/EditorResources/S_Emitter.S_Emitter'"));
+	RegisterDefaultSprite(UE::AvaShapes::ColorSelectionSprite, TEXT("Texture2D'/Engine/EditorResources/S_ReflActorIcon.S_ReflActorIcon'"));
+	RegisterDefaultSprite(UE::AvaShapes::CornerSprite, TEXT("Texture2D'/Avalanche/EditorResources/Bevel.Bevel'"));
+	RegisterDefaultSprite(UE::AvaShapes::DepthSprite, TEXT("Texture2D'/Engine/EditorResources/S_Terrain.S_Terrain'"));
+	RegisterDefaultSprite(UE::AvaShapes::InnerSizeSprite, TEXT("Texture2D'/Engine/EditorResources/S_RadForce.S_RadForce'"));
+	RegisterDefaultSprite(UE::AvaShapes::LinearGradientSprite, TEXT("Texture2D'/Avalanche/EditorResources/LinearGradient.LinearGradient'"));
+	RegisterDefaultSprite(UE::AvaShapes::NumPointsSprite, TEXT("Texture2D'/Engine/EditorResources/S_Emitter.S_Emitter'"));
+	RegisterDefaultSprite(UE::AvaShapes::NumSidesSprite, TEXT("Texture2D'/Engine/EditorResources/S_Emitter.S_Emitter'"));
+	RegisterDefaultSprite(UE::AvaShapes::SizeSprite, TEXT("Texture2D'/Avalanche/EditorResources/NewSizeHandle.NewSizeHandle'"));
+	RegisterDefaultSprite(UE::AvaShapes::SlantSprite, TEXT("Texture2D'/Avalanche/EditorResources/Slant.Slant'"));
+	RegisterDefaultSprite(UE::AvaShapes::TextMaxHeightSprite, TEXT("Texture2D'/Engine/EngineResources/Cursors/SplitterVert.SplitterVert'"));
+	RegisterDefaultSprite(UE::AvaShapes::TextMaxWidthSprite, TEXT("Texture2D'/Engine/EngineResources/Cursors/SplitterHorz.SplitterHorz'"));
+	RegisterDefaultSprite(UE::AvaShapes::TextScaleProportionallySprite, TEXT("Texture2D'/Engine/EditorResources/S_TextRenderActorIcon.S_TextRenderActorIcon'"));
+	RegisterDefaultSprite(UE::AvaShapes::UVSprite, TEXT("Texture2D'/Engine/EditorResources/MatInstActSprite.MatInstActSprite'"));
 
 	ISequencerModule& SequencerModule = FModuleManager::LoadModuleChecked<ISequencerModule>("Sequencer");
 	TrackEditorHandle = SequencerModule.RegisterTrackEditor(FOnCreateTrackEditor::CreateStatic(&FAvaShapeRectCornerTrackEditor::CreateTrackEditor));
diff --git a/Engine/Plugins/Experimental/Avalanche/Source/AvalancheShapesEditor/Private/AvalancheShapesEditorModule.h b/Engine/Plugins/Experimental/Avalanche/Source/AvalancheShapesEditor/Private/AvalancheShapesEditorModule.h
new file mode 100644
--- /dev/null
+++ b/Engine/Plugins/Experimental/Avalanche/Source/AvalancheShapesEditor/Private/AvalancheShapesEditorModule.h
@@ -0,0 +1,26 @@
+// Copyright Epic Games, Inc. All Rights Reserved.
+
+#pragma once
+
+#include "Modules/ModuleInterface.h"
+#include "Containers/Array.h"
+#include "Delegates/IDelegateInstance.h"
+#include "Templates/SharedPointer.h"
+
+class FComponentVisualizer;
+class IAvalancheInteractiveToolsModule;
+
+class FAvalancheShapesEditorModule : public IModuleInterface
+{
+	//~ Begin IModuleInterface interface
+	virtual void StartupModule() override;
+	virtual void ShutdownModule() override;
+	//~ End IModuleInterface interface
+
+	TArray<TSharedPtr<FComponentVisualizer>> Visualizers;
+
+	void RegisterShapeTools(IAvalancheInteractiveToolsModule* InModule);
+	void RegisterVisualizers();
+
+	FDelegateHandle TrackEditorHandle;
+};
diff --git a/Engine/Plugins/Experimental/Avalanche/Source/AvalancheShapesEditor/Private/Tools/AvaShapesEditorShapeAreaToolBase.h b/Engine/Plugins/Experimental/Avalanche/Source/AvalancheShapesEditor/Private/Tools/AvaShapesEditorShapeAreaToolBase.h
new file mode 100644
--- /dev/null
+++ b/Engine/Plugins/Experimental/Avalanche/Source/AvalancheShapesEditor/Private/Tools/AvaShapesEditorShapeAreaToolBase.h
@@ -0,0 +1,35 @@
+// Copyright Epic Games, Inc. All Rights Reserved.
+
+#pragma once
+
+#include "Tools/AvaShapesEditorShapeToolBase.h"
+#include "AvaShapesEditorShapeAreaToolBase.generated.h"
+
+class AActor;
+class AAvaShapeActor;
+
+UCLASS(Abstract)
+class UAvaShapesEditorShapeAreaToolBase : public UAvaShapesEditorShapeToolBase
+{
+	GENERATED_BODY()
+
+	using UAvaShapesEditorShapeToolBase::SpawnActor;
+
+public:
+	UAvaShapesEditorShapeAreaToolBase();
+
+	//~ Begin UAvaInteractiveToolsToolBase
+	virtual void OnViewportPlannerUpdate() override;
+	virtual void OnViewportPlannerComplete() override;
+	virtual void DefaultAction() override;
+	//~ End UAvaInteractiveToolsToolBase
+
+protected:
+	void UpdateShapeSize(AAvaShapeActor* InShapeActor) const;
+	void SelectActor(AActor* InActor) const;
+
+	//~ Begin UAvaInteractiveToolsToolBase
+	virtual AActor* SpawnActor(TSubclassOf<AActor> InActorClass, EAvaViewportStatus InViewportStatus, 
+		const FVector2f& InViewportPosition, bool bInPreview, FString* InActorLabelOverride = nullptr) const override;
+	//~ End UAvaInteractiveToolsToolBase
+};
diff --git a/Engine/Plugins/Experimental/Avalanche/Source/AvalancheShapesEditor/Private/Tools/AvaShapesEditorShapeTool2DArrow.cpp b/Engine/Plugins/Experimental/Avalanche/Source/AvalancheShapesEditor/Private/Tools/AvaShapesEditorShapeTool2DArrow.cpp
--- a/Engine/Plugins/Experimental/Avalanche/Source/AvalancheShapesEditor/Private/Tools/AvaShapesEditorShapeTool2DArrow.cpp
+++ b/Engine/Plugins/Experimental/Avalanche/Source/AvalancheShapesEditor/Private/Tools/AvaShapesEditorShapeTool2DArrow.cpp
@@ -2,7 +2,6 @@
 
 #include "Tools/AvaShapesEditorShapeTool2DArrow.h"
 #include "AvaShapesEditorCommands.h"
-#include "AvalancheShapesEditorModule.h"
 #include "Builders/AvaInteractiveToolsToolBuilder.h"
 #include "DynamicMeshes/AvaShape2DArrowDynMesh.h"
 #include "UObject/ConstructorHelpers.h"
diff --git a/Engine/Plugins/Experimental/Avalanche/Source/AvalancheShapesEditor/Private/Tools/AvaShapesEditorShapeToolBase.h b/Engine/Plugins/Experimental/Avalanche/Source/AvalancheShapesEditor/Private/Tools/AvaShapesEditorShapeToolBase.h
new file mode 100644
--- /dev/null
+++ b/Engine/Plugins/Experimental/Avalanche/Source/AvalancheShapesEditor/Private/Tools/AvaShapesEditorShapeToolBase.h
@@ -0,0 +1,58 @@
+// Copyright Epic Games, Inc. All Rights Reserved.
+
+#pragma once
+
+#include "Tools/AvaInteractiveToolsActorToolBase.h"
+#include "AvaShapeFactory.h"
+#include "AvaShapeParametricMaterial.h"
+#include "Templates/SubclassOf.h"
+#include "AvaShapesEditorShapeToolBase.generated.h"
+
+class AAvaShapeActor;
+class UAvaShapeDynamicMeshBase;
+
+UCLASS(Abstract)
+class UAvaShapesEditorShapeToolBase : public UAvaInteractiveToolsActorToolBase
+{
+	GENERATED_BODY()
+
+public:
+	using UAvaInteractiveToolsActorToolBase::SpawnActor;
+
+	static constexpr float DefaultDim = 100;
+
+	UAvaShapesEditorShapeToolBase();
+
+	//~ Begin UAvaInteractiveToolsToolBase
+	virtual bool OnBegin() override;
+	//~ End UAvaInteractiveToolsToolBase
+
+protected:
+	template<typename InMeshClass
+		UE_REQUIRES(std::derived_from<InMeshClass, UAvaShapeDynamicMeshBase>)>
+	static UAvaShapeFactory* CreateFactory()
+	{
+		UAvaShapeFactory* Factory = CreateActorFactory<UAvaShapeFactory>();
+		Factory->SetMeshClass(InMeshClass::StaticClass());
+		return Factory;
+	}
+
+	UPROPERTY()
+	TSubclassOf<UAvaShapeDynamicMeshBase> ShapeClass = nullptr;
+
+	FAvaShapeParametricMaterial DefaultMaterialParams;
+
+	virtual void InitShape(UAvaShapeDynamicMeshBase* InShape) const;
+
+	virtual void SetShapeSize(AAvaShapeActor* InShapeActor, const FVector2D& InShapeSize) const;
+
+	//~ Begin UAvaInteractiveToolsToolBase
+	virtual bool UseIdentityLocation() const override;
+	virtual bool UseIdentityRotation() const override;
+
+	virtual AActor* SpawnActor(TSubclassOf<AActor> InActorClass, EAvaViewportStatus InViewportStatus, 
+		const FVector2f& InViewportPosition, bool bInPreview, FString* InActorLabelOverride = nullptr) const override;
+	//~ End UAvaInteractiveToolsToolBase
+
+	virtual FString GetActorNameOverride() const;
+};
diff --git a/Engine/Plugins/Experimental/Avalanche/Source/AvalancheShapesEditor/Private/Tools/AvaShapesEditorShapeToolChevron.cpp b/Engine/Plugins/Experimental/Avalanche/Source/AvalancheShapesEditor/Private/Tools/AvaShapesEditorShapeToolChevron.cpp
--- a/Engine/Plugins/Experimental/Avalanche/Source/AvalancheShapesEditor/Private/Tools/AvaShapesEditorShapeToolChevron.cpp
+++ b/Engine/Plugins/Experimental/Avalanche/Source/AvalancheShapesEditor/Private/Tools/AvaShapesEditorShapeToolChevron.cpp
@@ -2,7 +2,6 @@
 
 #include "Tools/AvaShapesEditorShapeToolChevron.h"
 #include "AvaShapesEditorCommands.h"
-#include "AvalancheShapesEditorModule.h"
 #include "Builders/AvaInteractiveToolsToolBuilder.h"
 #include "DynamicMeshes/AvaShapeChevronDynMesh.h"
 #include "UObject/ConstructorHelpers.h"
diff --git a/Engine/Plugins/Experimental/Avalanche/Source/AvalancheShapesEditor/Private/Tools/AvaShapesEditorShapeToolCone.cpp b/Engine/Plugins/Experimental/Avalanche/Source/AvalancheShapesEditor/Private/Tools/AvaShapesEditorShapeToolCone.cpp
--- a/Engine/Plugins/Experimental/Avalanche/Source/AvalancheShapesEditor/Private/Tools/AvaShapesEditorShapeToolCone.cpp
+++ b/Engine/Plugins/Experimental/Avalanche/Source/AvalancheShapesEditor/Private/Tools/AvaShapesEditorShapeToolCone.cpp
@@ -2,7 +2,6 @@
 
 #include "Tools/AvaShapesEditorShapeToolCone.h"
 #include "AvaShapesEditorCommands.h"
-#include "AvalancheShapesEditorModule.h"
 #include "Builders/AvaInteractiveToolsToolBuilder.h"
 #include "DynamicMeshes/AvaShapeConeDynMesh.h"
 #include "UObject/ConstructorHelpers.h"
diff --git a/Engine/Plugins/Experimental/Avalanche/Source/AvalancheShapesEditor/Private/Tools/AvaShapesEditorShapeToolCube.cpp b/Engine/Plugins/Experimental/Avalanche/Source/AvalancheShapesEditor/Private/Tools/AvaShapesEditorShapeToolCube.cpp
--- a/Engine/Plugins/Experimental/Avalanche/Source/AvalancheShapesEditor/Private/Tools/AvaShapesEditorShapeToolCube.cpp
+++ b/Engine/Plugins/Experimental/Avalanche/Source/AvalancheShapesEditor/Private/Tools/AvaShapesEditorShapeToolCube.cpp
@@ -2,7 +2,6 @@
 
 #include "Tools/AvaShapesEditorShapeToolCube.h"
 #include "AvaShapesEditorCommands.h"
-#include "AvalancheShapesEditorModule.h"
 #include "Builders/AvaInteractiveToolsToolBuilder.h"
 #include "DynamicMeshes/AvaShapeCubeDynMesh.h"
 #include "UObject/ConstructorHelpers.h"
diff --git a/Engine/Plugins/Experimental/Avalanche/Source/AvalancheShapesEditor/Private/Tools/AvaShapesEditorShapeToolEllipse.cpp b/Engine/Plugins/Experimental/Avalanche/Source/AvalancheShapesEditor/Private/Tools/AvaShapesEditorShapeToolEllipse.cpp
--- a/Engine/Plugins/Experimental/Avalanche/Source/AvalancheShapesEditor/Private/Tools/AvaShapesEditorShapeToolEllipse.cpp
+++ b/Engine/Plugins/Experimental/Avalanche/Source/AvalancheShapesEditor/Private/Tools/AvaShapesEditorShapeToolEllipse.cpp
@@ -2,7 +2,6 @@
 
 #include "Tools/AvaShapesEditorShapeToolEllipse.h"
 #include "AvaShapesEditorCommands.h"
-#include "AvalancheShapesEditorModule.h"
 #include "Builders/AvaInteractiveToolsToolBuilder.h"
 #include "DynamicMeshes/AvaShapeEllipseDynMesh.h"
 #include "UObject/ConstructorHelpers.h"
diff --git a/Engine/Plugins/Experimental/Avalanche/Source/AvalancheShapesEditor/Private/Tools/AvaShapesEditorShapeToolIrregularPoly.cpp b/Engine/Plugins/Experimental/Avalanche/Source/AvalancheShapesEditor/Private/Tools/AvaShapesEditorShapeToolIrregularPoly.cpp
--- a/Engine/Plugins/Experimental/Avalanche/Source/AvalancheShapesEditor/Private/Tools/AvaShapesEditorShapeToolIrregularPoly.cpp
+++ b/Engine/Plugins/Experimental/Avalanche/Source/AvalancheShapesEditor/Private/Tools/AvaShapesEditorShapeToolIrregularPoly.cpp
@@ -1,7 +1,6 @@
 // Copyright Epic Games, Inc. All Rights Reserved.
 
 #include "Tools/AvaShapesEditorShapeToolIrregularPoly.h"
-#include "AvalancheShapesEditorModule.h"
 #include "AvaShapeActor.h"
 #include "AvaShapesEditorCommands.h"
 #include "AvaViewportUtils.h"
diff --git a/Engine/Plugins/Experimental/Avalanche/Source/AvalancheShapesEditor/Private/Tools/AvaShapesEditorShapeToolLine.cpp b/Engine/Plugins/Experimental/Avalanche/Source/AvalancheShapesEditor/Private/Tools/AvaShapesEditorShapeToolLine.cpp
--- a/Engine/Plugins/Experimental/Avalanche/Source/AvalancheShapesEditor/Private/Tools/AvaShapesEditorShapeToolLine.cpp
+++ b/Engine/Plugins/Experimental/Avalanche/Source/AvalancheShapesEditor/Private/Tools/AvaShapesEditorShapeToolLine.cpp
@@ -1,7 +1,6 @@
 // Copyright Epic Games, Inc. All Rights Reserved.
 
 #include "Tools/AvaShapesEditorShapeToolLine.h"
-#include "AvalancheShapesEditorModule.h"
 #include "AvaShapeActor.h"
 #include "AvaShapesEditorCommands.h"
 #include "Builders/AvaInteractiveToolsToolBuilder.h"
diff --git a/Engine/Plugins/Experimental/Avalanche/Source/AvalancheShapesEditor/Private/Tools/AvaShapesEditorShapeToolNGon.cpp b/Engine/Plugins/Experimental/Avalanche/Source/AvalancheShapesEditor/Private/Tools/AvaShapesEditorShapeToolNGon.cpp
--- a/Engine/Plugins/Experimental/Avalanche/Source/AvalancheShapesEditor/Private/Tools/AvaShapesEditorShapeToolNGon.cpp
+++ b/Engine/Plugins/Experimental/Avalanche/Source/AvalancheShapesEditor/Private/Tools/AvaShapesEditorShapeToolNGon.cpp
@@ -2,7 +2,6 @@
 
 #include "Tools/AvaShapesEditorShapeToolNGon.h"
 #include "AvaShapesEditorCommands.h"
-#include "AvalancheShapesEditorModule.h"
 #include "Builders/AvaInteractiveToolsToolBuilder.h"
 #include "DynamicMeshes/AvaShapeNGonDynMesh.h"
 #include "UObject/ConstructorHelpers.h"
diff --git a/Engine/Plugins/Experimental/Avalanche/Source/AvalancheShapesEditor/Private/Tools/AvaShapesEditorShapeToolRing.cpp b/Engine/Plugins/Experimental/Avalanche/Source/AvalancheShapesEditor/Private/Tools/AvaShapesEditorShapeToolRing.cpp
--- a/Engine/Plugins/Experimental/Avalanche/Source/AvalancheShapesEditor/Private/Tools/AvaShapesEditorShapeToolRing.cpp
+++ b/Engine/Plugins/Experimental/Avalanche/Source/AvalancheShapesEditor/Private/Tools/AvaShapesEditorShapeToolRing.cpp
@@ -2,7 +2,6 @@
 
 #include "Tools/AvaShapesEditorShapeToolRing.h"
 #include "AvaShapesEditorCommands.h"
-#include "AvalancheShapesEditorModule.h"
 #include "Builders/AvaInteractiveToolsToolBuilder.h"
 #include "DynamicMeshes/AvaShapeRingDynMesh.h"
 #include "UObject/ConstructorHelpers.h"
diff --git a/Engine/Plugins/Experimental/Avalanche/Source/AvalancheShapesEditor/Private/Tools/AvaShapesEditorShapeToolSphere.cpp b/Engine/Plugins/Experimental/Avalanche/Source/AvalancheShapesEditor/Private/Tools/AvaShapesEditorShapeToolSphere.cpp
--- a/Engine/Plugins/Experimental/Avalanche/Source/AvalancheShapesEditor/Private/Tools/AvaShapesEditorShapeToolSphere.cpp
+++ b/Engine/Plugins/Experimental/Avalanche/Source/AvalancheShapesEditor/Private/Tools/AvaShapesEditorShapeToolSphere.cpp
@@ -2,7 +2,6 @@
 
 #include "Tools/AvaShapesEditorShapeToolSphere.h"
 #include "AvaShapesEditorCommands.h"
-#include "AvalancheShapesEditorModule.h"
 #include "Builders/AvaInteractiveToolsToolBuilder.h"
 #include "DynamicMeshes/AvaShapeSphereDynMesh.h"
 #include "UObject/ConstructorHelpers.h"
diff --git a/Engine/Plugins/Experimental/Avalanche/Source/AvalancheShapesEditor/Private/Tools/AvaShapesEditorShapeToolStar.cpp b/Engine/Plugins/Experimental/Avalanche/Source/AvalancheShapesEditor/Private/Tools/AvaShapesEditorShapeToolStar.cpp
--- a/Engine/Plugins/Experimental/Avalanche/Source/AvalancheShapesEditor/Private/Tools/AvaShapesEditorShapeToolStar.cpp
+++ b/Engine/Plugins/Experimental/Avalanche/Source/AvalancheShapesEditor/Private/Tools/AvaShapesEditorShapeToolStar.cpp
@@ -2,7 +2,6 @@
 
 #include "Tools/AvaShapesEditorShapeToolStar.h"
 #include "AvaShapesEditorCommands.h"
-#include "AvalancheShapesEditorModule.h"
 #include "Builders/AvaInteractiveToolsToolBuilder.h"
 #include "DynamicMeshes/AvaShapeStarDynMesh.h"
 #include "UObject/ConstructorHelpers.h"
diff --git a/Engine/Plugins/Experimental/Avalanche/Source/AvalancheShapesEditor/Private/Tools/AvaShapesEditorShapeToolTorus.cpp b/Engine/Plugins/Experimental/Avalanche/Source/AvalancheShapesEditor/Private/Tools/AvaShapesEditorShapeToolTorus.cpp
--- a/Engine/Plugins/Experimental/Avalanche/Source/AvalancheShapesEditor/Private/Tools/AvaShapesEditorShapeToolTorus.cpp
+++ b/Engine/Plugins/Experimental/Avalanche/Source/AvalancheShapesEditor/Private/Tools/AvaShapesEditorShapeToolTorus.cpp
@@ -2,7 +2,6 @@
 
 #include "Tools/AvaShapesEditorShapeToolTorus.h"
 #include "AvaShapesEditorCommands.h"
-#include "AvalancheShapesEditorModule.h"
 #include "Builders/AvaInteractiveToolsToolBuilder.h"
 #include "DynamicMeshes/AvaShapeTorusDynMesh.h"
 #include "UObject/ConstructorHelpers.h"
diff --git a/Engine/Plugins/Experimental/Avalanche/Source/AvalancheShapesEditor/Private/Visualizers/AvaShape2DDynMeshVis.h b/Engine/Plugins/Experimental/Avalanche/Source/AvalancheShapesEditor/Private/Visualizers/AvaShape2DDynMeshVis.h
new file mode 100644
--- /dev/null
+++ b/Engine/Plugins/Experimental/Avalanche/Source/AvalancheShapesEditor/Private/Visualizers/AvaShape2DDynMeshVis.h
@@ -0,0 +1,70 @@
+// Copyright Epic Games, Inc. All Rights Reserved.
+
+#pragma once
+
+#include "AvaShapeDynMeshVis.h"
+#include "DynamicMeshes/AvaShape2DDynMeshBase.h"
+
+class UAvaShape2DDynMeshBase;
+
+class FAvaShape2DDynamicMeshVisualizer : public FAvaShapeDynamicMeshVisualizer
+{
+public:
+	using Super = FAvaShapeDynamicMeshVisualizer;
+	using FMeshType = UAvaShape2DDynMeshBase;
+
+	FAvaShape2DDynamicMeshVisualizer();
+
+	//~ Begin FAvaVisualizerBase
+	virtual TMap<UObject*, TArray<FProperty*>> GatherEditableProperties(UObject* InObject) const override;
+	virtual bool VisProxyHandleClick(FEditorViewportClient* InViewportClient, HComponentVisProxy* InVisProxy,
+		const FViewportClick& InClick) override;
+	virtual bool GetWidgetLocation(const FEditorViewportClient* InViewportClient, FVector& OutLocation) const override;
+	virtual bool GetWidgetMode(const FEditorViewportClient* InViewportClient, UE::Widget::EWidgetMode& OutMode) const override;
+	virtual bool GetWidgetAxisList(const FEditorViewportClient* InViewportClient, UE::Widget::EWidgetMode InWidgetMode,
+		EAxisList::Type& OutAxisList) const override;
+	virtual bool HandleModifiedClick(FEditorViewportClient* InViewportClient, HHitProxy* InHitProxy, const FViewportClick& InClick) override;
+	virtual bool ResetValue(FEditorViewportClient* InViewportClient, HHitProxy* InHitProxy) override;
+	virtual bool IsEditing() const override;
+	virtual void EndEditing() override;
+	//~ End FAvaVisualizerBase
+
+protected:
+	FProperty* SizeProperty;
+
+	AvaAlignment SizeDragAnchor; // XYZ
+	FVector2D InitialSize;
+	
+	// Supported alignments
+	// Center alignments removed on request.
+	static inline const TArray<AvaAlignment> SupportedAlignments {
+		MakeAlignment(EAvaDepthAlignment::Back, EAvaHorizontalAlignment::Left, EAvaVerticalAlignment::Top), // EAvaAnchors::TopLeft
+		MakeAlignment(EAvaDepthAlignment::Back, EAvaHorizontalAlignment::Right, EAvaVerticalAlignment::Top), //  EAvaAnchors::TopRight
+		MakeAlignment(EAvaDepthAlignment::Back, EAvaHorizontalAlignment::Left, EAvaVerticalAlignment::Bottom), // EAvaAnchors::BottomLeft
+		MakeAlignment(EAvaDepthAlignment::Back, EAvaHorizontalAlignment::Right, EAvaVerticalAlignment::Bottom), // EAvaAnchors::BottomRight
+	};
+
+	int32 UVSectionIdx;
+	bool bEditingUVAnchor;
+	FVector2D InitialPrimaryUVOffset;
+	FVector2D InitialPrimaryUVScale;
+	float InitialPrimaryUVRotation;
+	FVector2D InitialPrimaryUVAnchor;
+	FVector2D PrevShapeSize;
+
+	virtual void DrawVisualizationNotEditing(const UActorComponent* InComponent, const FSceneView* InView,
+		FPrimitiveDrawInterface* InPDI, int32& InOutIconIndex) override;
+	virtual void DrawVisualizationEditing(const UActorComponent* Component, const FSceneView* View,
+		FPrimitiveDrawInterface* InPDI, int32& InOutIconIndex) override;
+
+	virtual void DrawSizeButtons(const FMeshType* InDynMesh, const FSceneView* InView, FPrimitiveDrawInterface* InPDI);
+
+	virtual void StoreInitialValues() override;
+
+	virtual bool HandleInputDeltaInternal(FEditorViewportClient* InViewportClient, FViewport* InViewport, const FVector& InAccumulatedTranslation,
+		const FRotator& InAccumulatedRotation, const FVector& InAccumulatedScale) override;
+
+	void SnapLocation2D(FEditorViewportClient* InViewportClient, const FTransform& InActorTransform, FVector2D& OutLocation2D) const;
+
+	bool IsExtendBothSidesEnabled(FMeshType* InDynMesh) const;
+};
diff --git a/Engine/Plugins/Experimental/Avalanche/Source/AvalancheShapesEditor/Private/Visualizers/AvaShape3DDynMeshVis.h b/Engine/Plugins/Experimental/Avalanche/Source/AvalancheShapesEditor/Private/Visualizers/AvaShape3DDynMeshVis.h
new file mode 100644
--- /dev/null
+++ b/Engine/Plugins/Experimental/Avalanche/Source/AvalancheShapesEditor/Private/Visualizers/AvaShape3DDynMeshVis.h
@@ -0,0 +1,64 @@
+// Copyright Epic Games, Inc. All Rights Reserved.
+
+#pragma once
+
+#include "AvaShapeDynMeshVis.h"
+
+class UAvaShape3DDynMeshBase;
+
+class FAvaShape3DDynamicMeshVisualizer : public FAvaShapeDynamicMeshVisualizer
+{
+public:
+	using Super = FAvaShapeDynamicMeshVisualizer;
+	using FMeshType = UAvaShape3DDynMeshBase;
+	
+	FAvaShape3DDynamicMeshVisualizer();
+
+	//~ Begin FAvaVisualizerBase
+	virtual TMap<UObject*, TArray<FProperty*>> GatherEditableProperties(UObject* InObject) const override;
+	virtual bool VisProxyHandleClick(FEditorViewportClient* InViewportClient, HComponentVisProxy* InVisProxy,
+		const FViewportClick& InClick) override;
+	virtual bool GetWidgetLocation(const FEditorViewportClient* InViewportClient, FVector& OutLocation) const override;
+	virtual bool GetWidgetMode(const FEditorViewportClient* InViewportClient, UE::Widget::EWidgetMode& OutMode) const override;
+	virtual bool GetWidgetAxisList(const FEditorViewportClient* InViewportClient, UE::Widget::EWidgetMode InWidgetMode,
+		EAxisList::Type& OutAxisList) const override;
+	virtual bool HandleModifiedClick(FEditorViewportClient* InViewportClient, HHitProxy* InHitProxy, const FViewportClick& InClick) override;
+	virtual bool ResetValue(FEditorViewportClient* InViewportClient, HHitProxy* InHitProxy) override;
+	virtual bool IsEditing() const override;
+	virtual void EndEditing() override;
+	//~ End FAvaVisualizerBase
+
+protected:
+
+	// direct access to size property to notify about update
+	FProperty* SizeProperty;
+	// this gives information about current handle XYZ we are updating
+	AvaAlignment SizeDragAnchor;
+	// size before any update with visualizer
+	FVector InitialSize;
+	
+	// Supported alignments
+	// Center alignments removed on request.
+	static inline const TArray<AvaAlignment> SupportedAlignments {
+		/* back */
+		MakeAlignment(EAvaDepthAlignment::Back, EAvaHorizontalAlignment::Left, EAvaVerticalAlignment::Top),
+		MakeAlignment(EAvaDepthAlignment::Back, EAvaHorizontalAlignment::Right, EAvaVerticalAlignment::Top),
+		MakeAlignment(EAvaDepthAlignment::Back, EAvaHorizontalAlignment::Left, EAvaVerticalAlignment::Bottom),
+		MakeAlignment(EAvaDepthAlignment::Back, EAvaHorizontalAlignment::Right, EAvaVerticalAlignment::Bottom),
+		/* front */
+		MakeAlignment(EAvaDepthAlignment::Front, EAvaHorizontalAlignment::Left, EAvaVerticalAlignment::Top),
+		MakeAlignment(EAvaDepthAlignment::Front, EAvaHorizontalAlignment::Right, EAvaVerticalAlignment::Top),
+		MakeAlignment(EAvaDepthAlignment::Front, EAvaHorizontalAlignment::Left, EAvaVerticalAlignment::Bottom),
+		MakeAlignment(EAvaDepthAlignment::Front, EAvaHorizontalAlignment::Right, EAvaVerticalAlignment::Bottom),
+	};
+	
+	virtual void DrawVisualizationNotEditing(const UActorComponent* InComponent, const FSceneView* InView,
+		FPrimitiveDrawInterface* InPDI, int32& InOutIconIndex) override;
+	virtual void DrawVisualizationEditing(const UActorComponent* InComponent, const FSceneView* InView,
+		FPrimitiveDrawInterface* InPDI, int32& InOutIconIndex) override;
+	virtual void StoreInitialValues() override;
+	virtual bool HandleInputDeltaInternal(FEditorViewportClient* InViewportClient, FViewport* InViewport, const FVector& InAccumulatedTranslation,
+		const FRotator& InAccumulatedRotation, const FVector& InAccumulatedScale) override;
+	
+	void SnapLocation3D(FEditorViewportClient* InViewportClient, const FTransform& InActorTransform, FVector& OutLocation) const;
+};
diff --git a/Engine/Plugins/Experimental/Avalanche/Source/AvalancheShapesEditor/Private/Visualizers/AvaShapeConeDynMeshVis.cpp b/Engine/Plugins/Experimental/Avalanche/Source/AvalancheShapesEditor/Private/Visualizers/AvaShapeConeDynMeshVis.cpp
--- a/Engine/Plugins/Experimental/Avalanche/Source/AvalancheShapesEditor/Private/Visualizers/AvaShapeConeDynMeshVis.cpp
+++ b/Engine/Plugins/Experimental/Avalanche/Source/AvalancheShapesEditor/Private/Visualizers/AvaShapeConeDynMeshVis.cpp
@@ -3,7 +3,7 @@
 #include "AvaShapeConeDynMeshVis.h"
 #include "AvaField.h"
 #include "AvaShapeActor.h"
-#include "AvalancheShapesEditorModule.h"
+#include "AvaShapeSprites.h"
 #include "DynamicMeshes/AvaShapeConeDynMesh.h"
 #include "EditorViewportClient.h"
 #include "Engine/Texture2D.h"
@@ -406,7 +406,7 @@ void FAvaShapeConeDynamicMeshVisualizer::StoreInitialValues()
 void FAvaShapeConeDynamicMeshVisualizer::DrawTopRadiusButton(const FMeshType* InDynMesh, const FSceneView* InView,
 	FPrimitiveDrawInterface* InPDI, int32 InIconIndex, const FLinearColor& InColor) const
 {
-	UTexture2D* CornerSprite = IAvalancheComponentVisualizersModule::Get().GetSettings()->GetVisualizerSprite(FAvalancheShapesEditorModule::CornerSprite);
+	UTexture2D* CornerSprite = IAvalancheComponentVisualizersModule::Get().GetSettings()->GetVisualizerSprite(UE::AvaShapes::CornerSprite);
 
 	if (!CornerSprite || !CornerSprite->GetResource())
 	{
@@ -426,7 +426,7 @@ void FAvaShapeConeDynamicMeshVisualizer::DrawTopRadiusButton(const FMeshType* In
 void FAvaShapeConeDynamicMeshVisualizer::DrawNumSidesButton(const FMeshType* InDynMesh, const FSceneView* InView,
 	FPrimitiveDrawInterface* InPDI, int32 InIconIndex, const FLinearColor& InColor) const
 {
-	UTexture2D* NumSidesSprite = IAvalancheComponentVisualizersModule::Get().GetSettings()->GetVisualizerSprite(FAvalancheShapesEditorModule::NumSidesSprite);
+	UTexture2D* NumSidesSprite = IAvalancheComponentVisualizersModule::Get().GetSettings()->GetVisualizerSprite(UE::AvaShapes::NumSidesSprite);
 
 	if (!NumSidesSprite || !NumSidesSprite->GetResource())
 	{
@@ -446,7 +446,7 @@ void FAvaShapeConeDynamicMeshVisualizer::DrawNumSidesButton(const FMeshType* InD
 void FAvaShapeConeDynamicMeshVisualizer::DrawAngleDegreeButton(const FMeshType* InDynMesh, const FSceneView* InView,
 	FPrimitiveDrawInterface* InPDI, int32 InIconIndex, const FLinearColor& InColor) const
 {
-	UTexture2D* InnerSizeSprite = IAvalancheComponentVisualizersModule::Get().GetSettings()->GetVisualizerSprite(FAvalancheShapesEditorModule::InnerSizeSprite);
+	UTexture2D* InnerSizeSprite = IAvalancheComponentVisualizersModule::Get().GetSettings()->GetVisualizerSprite(UE::AvaShapes::InnerSizeSprite);
 
 	if (!InnerSizeSprite || !InnerSizeSprite->GetResource())
 	{
diff --git a/Engine/Plugins/Experimental/Avalanche/Source/AvalancheShapesEditor/Private/Visualizers/AvaShapeCubeDynMeshVis.cpp b/Engine/Plugins/Experimental/Avalanche/Source/AvalancheShapesEditor/Private/Visualizers/AvaShapeCubeDynMeshVis.cpp
--- a/Engine/Plugins/Experimental/Avalanche/Source/AvalancheShapesEditor/Private/Visualizers/AvaShapeCubeDynMeshVis.cpp
+++ b/Engine/Plugins/Experimental/Avalanche/Source/AvalancheShapesEditor/Private/Visualizers/AvaShapeCubeDynMeshVis.cpp
@@ -3,7 +3,7 @@
 #include "AvaShapeCubeDynMeshVis.h"
 #include "AvaField.h"
 #include "AvaShapeActor.h"
-#include "AvalancheShapesEditorModule.h"
+#include "AvaShapeSprites.h"
 #include "DynamicMeshes/AvaShapeCubeDynMesh.h"
 #include "EditorViewportClient.h"
 #include "Engine/Texture2D.h"
@@ -328,7 +328,7 @@ void FAvaShapeCubeDynamicMeshVisualizer::StoreInitialValues()
 void FAvaShapeCubeDynamicMeshVisualizer::DrawBevelSizeRatioButton(const FMeshType* InDynMesh, const FSceneView* InView,
 	FPrimitiveDrawInterface* InPDI, int32 InIconIndex, const FLinearColor& InColor) const
 {
-	UTexture2D* CornerSprite = IAvalancheComponentVisualizersModule::Get().GetSettings()->GetVisualizerSprite(FAvalancheShapesEditorModule::CornerSprite);
+	UTexture2D* CornerSprite = IAvalancheComponentVisualizersModule::Get().GetSettings()->GetVisualizerSprite(UE::AvaShapes::CornerSprite);
 
 	if (!CornerSprite || !CornerSprite->GetResource())
 	{
@@ -348,7 +348,7 @@ void FAvaShapeCubeDynamicMeshVisualizer::DrawBevelSizeRatioButton(const FMeshTyp
 void FAvaShapeCubeDynamicMeshVisualizer::DrawBevelNumButton(const FMeshType* InDynMesh, const FSceneView* InView,
 	FPrimitiveDrawInterface* InPDI, int32 InIconIndex, const FLinearColor& InColor) const
 {
-	UTexture2D* NumSidesSprite = IAvalancheComponentVisualizersModule::Get().GetSettings()->GetVisualizerSprite(FAvalancheShapesEditorModule::NumSidesSprite);
+	UTexture2D* NumSidesSprite = IAvalancheComponentVisualizersModule::Get().GetSettings()->GetVisualizerSprite(UE::AvaShapes::NumSidesSprite);
 
 	if (!NumSidesSprite || !NumSidesSprite->GetResource())
 	{
diff --git a/Engine/Plugins/Experimental/Avalanche/Source/AvalancheShapesEditor/Private/Visualizers/AvaShapeDynMeshVis.cpp b/Engine/Plugins/Experimental/Avalanche/Source/AvalancheShapesEditor/Private/Visualizers/AvaShapeDynMeshVis.cpp
--- a/Engine/Plugins/Experimental/Avalanche/Source/AvalancheShapesEditor/Private/Visualizers/AvaShapeDynMeshVis.cpp
+++ b/Engine/Plugins/Experimental/Avalanche/Source/AvalancheShapesEditor/Private/Visualizers/AvaShapeDynMeshVis.cpp
@@ -4,7 +4,7 @@
 #include "AvaActorUtils.h"
 #include "AvaField.h"
 #include "AvaShapeActor.h"
-#include "AvalancheShapesEditorModule.h"
+#include "AvaShapeSprites.h"
 #include "Components/DynamicMeshComponent.h"
 #include "DynamicMeshes/AvaShapeDynMeshBase.h"
 #include "Engine/Texture2D.h"
@@ -135,7 +135,7 @@ void FAvaShapeDynamicMeshVisualizer::DrawSizeButton(const UAvaShapeDynamicMeshBa
 {
 	static const float BaseSize = 1.f;
 	
-	UTexture2D* SizeSprite = IAvalancheComponentVisualizersModule::Get().GetSettings()->GetVisualizerSprite(FAvalancheShapesEditorModule::SizeSprite);
+	UTexture2D* SizeSprite = IAvalancheComponentVisualizersModule::Get().GetSettings()->GetVisualizerSprite(UE::AvaShapes::SizeSprite);
 
 	if (!SizeSprite || !SizeSprite->GetResource())
 	{
@@ -161,7 +161,7 @@ void FAvaShapeDynamicMeshVisualizer::DrawSizeButton(const UAvaShapeDynamicMeshBa
 void FAvaShapeDynamicMeshVisualizer::DrawUVButton(const UAvaShapeDynamicMeshBase* InDynMesh, const FSceneView* InView,
 	FPrimitiveDrawInterface* InPDI, int32 InIconIndex, int32 InSectionIdx, const FLinearColor& InColor) const
 {
-	UTexture2D* UVSprite = IAvalancheComponentVisualizersModule::Get().GetSettings()->GetVisualizerSprite(FAvalancheShapesEditorModule::UVSprite);
+	UTexture2D* UVSprite = IAvalancheComponentVisualizersModule::Get().GetSettings()->GetVisualizerSprite(UE::AvaShapes::UVSprite);
 
 	if (!UVSprite || !UVSprite->GetResource())
 	{
diff --git a/Engine/Plugins/Experimental/Avalanche/Source/AvalancheShapesEditor/Private/Visualizers/AvaShapeDynMeshVis.h b/Engine/Plugins/Experimental/Avalanche/Source/AvalancheShapesEditor/Private/Visualizers/AvaShapeDynMeshVis.h
new file mode 100644
--- /dev/null
+++ b/Engine/Plugins/Experimental/Avalanche/Source/AvalancheShapesEditor/Private/Visualizers/AvaShapeDynMeshVis.h
@@ -0,0 +1,134 @@
+// Copyright Epic Games, Inc. All Rights Reserved.
+
+#pragma once
+
+#include "AvaDefs.h"
+#include "AvaVisBase.h"
+
+class UAvaShapeDynamicMeshBase;
+
+struct HAvaShapeSizeHitProxy : HAvaHitProxy
+{
+	DECLARE_HIT_PROXY();
+
+	HAvaShapeSizeHitProxy(
+		const UActorComponent* InComponent, 
+		AvaAlignment InAlignment)
+		: HAvaHitProxy(InComponent)
+		, DragAnchor(InAlignment) {}
+
+	AvaAlignment DragAnchor; // XYZ
+};
+
+struct HAvaShapeUVHitProxy : HAvaHitProxy
+{
+	DECLARE_HIT_PROXY();
+
+	HAvaShapeUVHitProxy(const UActorComponent* InComponent, int32 InSectionIdx)
+		: HAvaHitProxy(InComponent)
+		, SectionIdx(InSectionIdx)
+	{}
+
+	int32 SectionIdx;
+};
+
+struct HAvaShapeNumSidesHitProxy : HAvaHitProxy
+{
+	DECLARE_HIT_PROXY();
+
+	HAvaShapeNumSidesHitProxy(const UActorComponent* InComponent)
+		: HAvaHitProxy(InComponent)
+	{}
+};
+
+struct HAvaShapeNumPointsHitProxy : HAvaHitProxy
+{
+	DECLARE_HIT_PROXY();
+
+	HAvaShapeNumPointsHitProxy(const UActorComponent* InComponent)
+		: HAvaHitProxy(InComponent)
+	{}
+};
+
+struct HAvaShapeInnerSizeHitProxy : HAvaHitProxy
+{
+	DECLARE_HIT_PROXY();
+
+	HAvaShapeInnerSizeHitProxy(const UActorComponent* InComponent)
+		: HAvaHitProxy(InComponent)
+	{}
+};
+
+struct HAvaShapeCornersHitProxy : HAvaHitProxy
+{
+	DECLARE_HIT_PROXY();
+
+	HAvaShapeCornersHitProxy(const UActorComponent* InComponent)
+		: HAvaHitProxy(InComponent)
+	{}
+};
+
+struct HAvaShapeAngleDegreeHitProxy : HAvaHitProxy
+{
+	DECLARE_HIT_PROXY();
+
+	HAvaShapeAngleDegreeHitProxy(
+		const UActorComponent* InComponent,
+		AvaAlignment InAlignment)
+		: HAvaHitProxy(InComponent)
+		, DragAnchor(InAlignment)
+	{}
+
+	AvaAlignment DragAnchor; // XYZ
+};
+
+class FAvaShapeDynamicMeshVisualizer : public FAvaVisualizerBase
+{
+public:
+	using Super = FAvaVisualizerBase;
+
+	UAvaShapeDynamicMeshBase* GetDynamicMesh() const;
+
+	//~ Begin FAvaVisualizerBase
+	virtual UActorComponent* GetEditedComponent() const override;
+	virtual TMap<UObject*, TArray<FProperty*>> GatherEditableProperties(UObject* InObject) const override;
+	virtual const USceneComponent* GetEditedSceneComponent() const override;
+	virtual const USceneComponent* GetEditedSceneComponent(const UActorComponent* InComponent) const override;
+	virtual void StartEditing(FEditorViewportClient* InViewportClient, UActorComponent* InEditedComponent) override;
+	virtual void EndEditing() override;
+	//~ End FAvaVisualizerBase
+
+protected:
+	TWeakObjectPtr<UAvaShapeDynamicMeshBase> DynamicMeshComponent;
+	FTransform InitialTransform;
+
+	FProperty* MeshRegenWorldLocationProperty;
+	FProperty* MeshDataProperty;
+	FProperty* UVParamsProperty;
+	FProperty* UVOffsetProperty;
+	FProperty* UVScaleProperty;
+	FProperty* UVRotationProperty;
+	FProperty* UVAnchorProperty;
+	FProperty* UVModeProperty;
+	FProperty* UVHorizFlipProperty;
+	FProperty* UVVertFlipProperty;
+
+	FAvaShapeDynamicMeshVisualizer();
+
+	virtual void StoreInitialValues() override;
+	virtual FBox GetComponentBounds(const UActorComponent* InComponent) const override;
+	virtual FTransform GetComponentTransform(const UActorComponent* InComponent) const override;
+
+	// draw size button hit proxy with icon at specific alignment location
+	void DrawSizeButton(const UAvaShapeDynamicMeshBase* InDynMesh, const FSceneView* InView, FPrimitiveDrawInterface* InPDI,
+		AvaAlignment InSizeDragAnchor) const;
+
+	// draw uv button hit proxy with icon for section index
+	void DrawUVButton(const UAvaShapeDynamicMeshBase* InDynMesh, const FSceneView* InView, FPrimitiveDrawInterface* InPDI, 
+		int32 InIconIndex, int32 InSectionIdx, const FLinearColor& InColor) const;
+
+	// checks whether an anchor location is on the correct size of the view and can be reachable
+	bool IsAnchorReachable(const UAvaShapeDynamicMeshBase* InDynMesh, const FSceneView* InView, const FVector& InAnchorLocation) const;
+
+	FVector GetFinalAnchorLocation(const UAvaShapeDynamicMeshBase* InDynMesh, AvaAlignment InSizeDragAnchor) const;
+};
diff --git a/Engine/Plugins/Experimental/Avalanche/Source/AvalancheShapesEditor/Private/Visualizers/AvaShapeEllipseDynMeshVis.cpp b/Engine/Plugins/Experimental/Avalanche/Source/AvalancheShapesEditor/Private/Visualizers/AvaShapeEllipseDynMeshVis.cpp
--- a/Engine/Plugins/Experimental/Avalanche/Source/AvalancheShapesEditor/Private/Visualizers/AvaShapeEllipseDynMeshVis.cpp
+++ b/Engine/Plugins/Experimental/Avalanche/Source/AvalancheShapesEditor/Private/Visualizers/AvaShapeEllipseDynMeshVis.cpp
@@ -3,7 +3,7 @@
 #include "AvaShapeEllipseDynMeshVis.h"
 #include "AvaField.h"
 #include "AvaShapeActor.h"
-#include "AvalancheShapesEditorModule.h"
+#include "AvaShapeSprites.h"
 #include "DynamicMeshes/AvaShapeEllipseDynMesh.h"
 #include "EditorViewportClient.h"
 #include "Engine/Texture2D.h"
@@ -96,7 +96,7 @@ void FAvaShapeEllipseDynamicMeshVisualizer::DrawVisualizationEditing(const UActo
 void FAvaShapeEllipseDynamicMeshVisualizer::DrawNumSidesButton(const FMeshType* InDynMesh, const FSceneView* InView,
 	FPrimitiveDrawInterface* InPDI, int32 InIconIndex, const FLinearColor& InColor) const
 {
-	UTexture2D* NumSidesSprite = IAvalancheComponentVisualizersModule::Get().GetSettings()->GetVisualizerSprite(FAvalancheShapesEditorModule::NumSidesSprite);
+	UTexture2D* NumSidesSprite = IAvalancheComponentVisualizersModule::Get().GetSettings()->GetVisualizerSprite(UE::AvaShapes::NumSidesSprite);
 
 	if (!NumSidesSprite || !NumSidesSprite->GetResource())
 	{
diff --git a/Engine/Plugins/Experimental/Avalanche/Source/AvalancheShapesEditor/Private/Visualizers/AvaShapeIrregularPolygonDynMeshVis.cpp b/Engine/Plugins/Experimental/Avalanche/Source/AvalancheShapesEditor/Private/Visualizers/AvaShapeIrregularPolygonDynMeshVis.cpp
--- a/Engine/Plugins/Experimental/Avalanche/Source/AvalancheShapesEditor/Private/Visualizers/AvaShapeIrregularPolygonDynMeshVis.cpp
+++ b/Engine/Plugins/Experimental/Avalanche/Source/AvalancheShapesEditor/Private/Visualizers/AvaShapeIrregularPolygonDynMeshVis.cpp
@@ -2,8 +2,8 @@
 
 #include "AvaShapeIrregularPolygonDynMeshVis.h"
 #include "AvaField.h"
-#include "AvalancheShapesEditorModule.h"
 #include "AvaShapeActor.h"
+#include "AvaShapeSprites.h"
 #include "DynamicMeshes/AvaShapeIrregularPolygonDynMesh.h"
 #include "EditorViewportClient.h"
 #include "Engine/Texture2D.h"
@@ -269,7 +269,7 @@ void FAvaShapeIrregularPolygonDynamicMeshVisualizer::GenerateContextSensitiveSna
 void FAvaShapeIrregularPolygonDynamicMeshVisualizer::DrawGlobalBevelButton(const FMeshType* InDynMesh, const FSceneView* InView,
 	FPrimitiveDrawInterface* InPDI, int32 InIconIndex, const FLinearColor& InColor) const
 {
-	UTexture2D* UVSprite = IAvalancheComponentVisualizersModule::Get().GetSettings()->GetVisualizerSprite(FAvalancheShapesEditorModule::CornerSprite);
+	UTexture2D* UVSprite = IAvalancheComponentVisualizersModule::Get().GetSettings()->GetVisualizerSprite(UE::AvaShapes::CornerSprite);
 
 	if (!UVSprite || !UVSprite->GetResource())
 	{
@@ -296,7 +296,7 @@ void FAvaShapeIrregularPolygonDynamicMeshVisualizer::DrawPointButton(const FMesh
 
 	static float BaseSize = 1.f;
 
-	UTexture2D* SizeSprite = IAvalancheComponentVisualizersModule::Get().GetSettings()->GetVisualizerSprite(FAvalancheShapesEditorModule::SizeSprite);
+	UTexture2D* SizeSprite = IAvalancheComponentVisualizersModule::Get().GetSettings()->GetVisualizerSprite(UE::AvaShapes::SizeSprite);
 
 	if (!SizeSprite || !SizeSprite->GetResource())
 	{
@@ -324,7 +324,7 @@ void FAvaShapeIrregularPolygonDynamicMeshVisualizer::DrawBevelButton(const FMesh
 
 	static float BaseSize = 1.f;
 
-	UTexture2D* SizeSprite = IAvalancheComponentVisualizersModule::Get().GetSettings()->GetVisualizerSprite(FAvalancheShapesEditorModule::BevelSprite);
+	UTexture2D* SizeSprite = IAvalancheComponentVisualizersModule::Get().GetSettings()->GetVisualizerSprite(UE::AvaShapes::BevelSprite);
 
 	if (!SizeSprite || !SizeSprite->GetResource())
 	{
@@ -351,7 +351,7 @@ void FAvaShapeIrregularPolygonDynamicMeshVisualizer::DrawBreakButton(const FMesh
 
 	static float BaseSize = 1.f;
 
-	UTexture2D* BreakSprite = IAvalancheComponentVisualizersModule::Get().GetSettings()->GetVisualizerSprite(FAvalancheShapesEditorModule::BreakSideSprite);
+	UTexture2D* BreakSprite = IAvalancheComponentVisualizersModule::Get().GetSettings()->GetVisualizerSprite(UE::AvaShapes::BreakSideSprite);
 
 	if (!BreakSprite || !BreakSprite->GetResource())
 	{
diff --git a/Engine/Plugins/Experimental/Avalanche/Source/AvalancheShapesEditor/Private/Visualizers/AvaShapeLineDynMeshVis.cpp b/Engine/Plugins/Experimental/Avalanche/Source/AvalancheShapesEditor/Private/Visualizers/AvaShapeLineDynMeshVis.cpp
--- a/Engine/Plugins/Experimental/Avalanche/Source/AvalancheShapesEditor/Private/Visualizers/AvaShapeLineDynMeshVis.cpp
+++ b/Engine/Plugins/Experimental/Avalanche/Source/AvalancheShapesEditor/Private/Visualizers/AvaShapeLineDynMeshVis.cpp
@@ -2,8 +2,8 @@
 
 #include "AvaShapeLineDynMeshVis.h"
 #include "AvaField.h"
-#include "AvalancheShapesEditorModule.h"
 #include "AvaShapeActor.h"
+#include "AvaShapeSprites.h"
 #include "DynamicMeshes/AvaShapeLineDynMesh.h"
 #include "EditorViewportClient.h"
 #include "Engine/Texture2D.h"
@@ -163,7 +163,7 @@ void FAvaShapeLineDynamicMeshVisualizer::DrawStartButton(const FMeshType* InDynM
 {
 	static float BaseSize = 1.f;
 
-	UTexture2D* SizeSprite = IAvalancheComponentVisualizersModule::Get().GetSettings()->GetVisualizerSprite(FAvalancheShapesEditorModule::SizeSprite);
+	UTexture2D* SizeSprite = IAvalancheComponentVisualizersModule::Get().GetSettings()->GetVisualizerSprite(UE::AvaShapes::SizeSprite);
 
 	if (!SizeSprite || !SizeSprite->GetResource())
 	{
@@ -185,7 +185,7 @@ void FAvaShapeLineDynamicMeshVisualizer::DrawEndButton(const FMeshType* InDynMes
 {
 	static float BaseSize = 1.f;
 
-	UTexture2D* SizeSprite = IAvalancheComponentVisualizersModule::Get().GetSettings()->GetVisualizerSprite(FAvalancheShapesEditorModule::SizeSprite);
+	UTexture2D* SizeSprite = IAvalancheComponentVisualizersModule::Get().GetSettings()->GetVisualizerSprite(UE::AvaShapes::SizeSprite);
 
 	if (!SizeSprite || !SizeSprite->GetResource())
 	{
diff --git a/Engine/Plugins/Experimental/Avalanche/Source/AvalancheShapesEditor/Private/Visualizers/AvaShapeNGonDynMeshVis.cpp b/Engine/Plugins/Experimental/Avalanche/Source/AvalancheShapesEditor/Private/Visualizers/AvaShapeNGonDynMeshVis.cpp
--- a/Engine/Plugins/Experimental/Avalanche/Source/AvalancheShapesEditor/Private/Visualizers/AvaShapeNGonDynMeshVis.cpp
+++ b/Engine/Plugins/Experimental/Avalanche/Source/AvalancheShapesEditor/Private/Visualizers/AvaShapeNGonDynMeshVis.cpp
@@ -3,7 +3,7 @@
 #include "AvaShapeNGonDynMeshVis.h"
 #include "AvaField.h"
 #include "AvaShapeActor.h"
-#include "AvalancheShapesEditorModule.h"
+#include "AvaShapeSprites.h"
 #include "DynamicMeshes/AvaShapeNGonDynMesh.h"
 #include "EditorViewportClient.h"
 #include "Engine/Texture2D.h"
@@ -96,7 +96,7 @@ void FAvaShapeNGonDynamicMeshVisualizer::DrawVisualizationEditing(const UActorCo
 void FAvaShapeNGonDynamicMeshVisualizer::DrawNumSidesButton(const FMeshType* InDynMesh,
 	const FSceneView* InView, FPrimitiveDrawInterface* InPDI, int32 InIconIndex, const FLinearColor& InColor) const
 {
-	UTexture2D* NumSidesSprite = IAvalancheComponentVisualizersModule::Get().GetSettings()->GetVisualizerSprite(FAvalancheShapesEditorModule::NumSidesSprite);
+	UTexture2D* NumSidesSprite = IAvalancheComponentVisualizersModule::Get().GetSettings()->GetVisualizerSprite(UE::AvaShapes::NumSidesSprite);
 
 	if (!NumSidesSprite || !NumSidesSprite->GetResource())
 	{
diff --git a/Engine/Plugins/Experimental/Avalanche/Source/AvalancheShapesEditor/Private/Visualizers/AvaShapeRectangleDynMeshVis.cpp b/Engine/Plugins/Experimental/Avalanche/Source/AvalancheShapesEditor/Private/Visualizers/AvaShapeRectangleDynMeshVis.cpp
--- a/Engine/Plugins/Experimental/Avalanche/Source/AvalancheShapesEditor/Private/Visualizers/AvaShapeRectangleDynMeshVis.cpp
+++ b/Engine/Plugins/Experimental/Avalanche/Source/AvalancheShapesEditor/Private/Visualizers/AvaShapeRectangleDynMeshVis.cpp
@@ -3,7 +3,7 @@
 #include "AvaShapeRectangleDynMeshVis.h"
 #include "AvaField.h"
 #include "AvaShapeActor.h"
-#include "AvalancheShapesEditorModule.h"
+#include "AvaShapeSprites.h"
 #include "DynamicMeshes/AvaShapeRectangleDynMesh.h"
 #include "EditorViewportClient.h"
 #include "Engine/Texture2D.h"
@@ -326,7 +326,7 @@ FVector FAvaShapeRectangleDynamicMeshVisualizer::GetCornerDragLocation(const FMe
 void FAvaShapeRectangleDynamicMeshVisualizer::DrawGlobalBevelButton(const FMeshType* InDynMesh, const FSceneView* InView,
 	FPrimitiveDrawInterface* InPDI, int32 InIconIndex, const FLinearColor& InColor) const
 {
-	UTexture2D* UVSprite = IAvalancheComponentVisualizersModule::Get().GetSettings()->GetVisualizerSprite(FAvalancheShapesEditorModule::CornerSprite);
+	UTexture2D* UVSprite = IAvalancheComponentVisualizersModule::Get().GetSettings()->GetVisualizerSprite(UE::AvaShapes::CornerSprite);
 
 	if (!UVSprite || !UVSprite->GetResource())
 	{
@@ -348,7 +348,7 @@ void FAvaShapeRectangleDynamicMeshVisualizer::DrawBevelButton(const FMeshType* I
 {
 	static const float BaseSize = 1.f;
 
-	UTexture2D* NumSidesSprite = IAvalancheComponentVisualizersModule::Get().GetSettings()->GetVisualizerSprite(FAvalancheShapesEditorModule::BevelSprite);
+	UTexture2D* NumSidesSprite = IAvalancheComponentVisualizersModule::Get().GetSettings()->GetVisualizerSprite(UE::AvaShapes::BevelSprite);
 
 	if (!NumSidesSprite || !NumSidesSprite->GetResource())
 	{
@@ -369,7 +369,7 @@ void FAvaShapeRectangleDynamicMeshVisualizer::DrawSlantButton(const FMeshType* I
 {
 	static const float BaseSize = 1.f;
 
-	UTexture2D* NumSidesSprite = IAvalancheComponentVisualizersModule::Get().GetSettings()->GetVisualizerSprite(FAvalancheShapesEditorModule::SlantSprite);
+	UTexture2D* NumSidesSprite = IAvalancheComponentVisualizersModule::Get().GetSettings()->GetVisualizerSprite(UE::AvaShapes::SlantSprite);
 
 	if (!NumSidesSprite || !NumSidesSprite->GetResource())
 	{
diff --git a/Engine/Plugins/Experimental/Avalanche/Source/AvalancheShapesEditor/Private/Visualizers/AvaShapeRingDynMeshVis.cpp b/Engine/Plugins/Experimental/Avalanche/Source/AvalancheShapesEditor/Private/Visualizers/AvaShapeRingDynMeshVis.cpp
--- a/Engine/Plugins/Experimental/Avalanche/Source/AvalancheShapesEditor/Private/Visualizers/AvaShapeRingDynMeshVis.cpp
+++ b/Engine/Plugins/Experimental/Avalanche/Source/AvalancheShapesEditor/Private/Visualizers/AvaShapeRingDynMeshVis.cpp
@@ -3,7 +3,7 @@
 #include "AvaShapeRingDynMeshVis.h"
 #include "AvaField.h"
 #include "AvaShapeActor.h"
-#include "AvalancheShapesEditorModule.h"
+#include "AvaShapeSprites.h"
 #include "DynamicMeshes/AvaShapeRingDynMesh.h"
 #include "EditorViewportClient.h"
 #include "Engine/Texture2D.h"
@@ -114,7 +114,7 @@ void FAvaShapeRingDynamicMeshVisualizer::DrawVisualizationEditing(const UActorCo
 void FAvaShapeRingDynamicMeshVisualizer::DrawNumSidesButton(const FMeshType* InDynMesh, const FSceneView* InView,
 	FPrimitiveDrawInterface* InPDI, int32 InIconIndex, const FLinearColor& InColor) const
 {
-	UTexture2D* NumSidesSprite = IAvalancheComponentVisualizersModule::Get().GetSettings()->GetVisualizerSprite(FAvalancheShapesEditorModule::NumSidesSprite);
+	UTexture2D* NumSidesSprite = IAvalancheComponentVisualizersModule::Get().GetSettings()->GetVisualizerSprite(UE::AvaShapes::NumSidesSprite);
 
 	if (!NumSidesSprite || !NumSidesSprite->GetResource())
 	{
@@ -134,7 +134,7 @@ void FAvaShapeRingDynamicMeshVisualizer::DrawNumSidesButton(const FMeshType* InD
 void FAvaShapeRingDynamicMeshVisualizer::DrawInnerSizeButton(const FMeshType* InDynMesh, const FSceneView* InView,
 	FPrimitiveDrawInterface* InPDI, int32 InIconIndex, const FLinearColor& InColor) const
 {
-	UTexture2D* NumSidesSprite = IAvalancheComponentVisualizersModule::Get().GetSettings()->GetVisualizerSprite(FAvalancheShapesEditorModule::InnerSizeSprite);
+	UTexture2D* NumSidesSprite = IAvalancheComponentVisualizersModule::Get().GetSettings()->GetVisualizerSprite(UE::AvaShapes::InnerSizeSprite);
 
 	if (!NumSidesSprite || !NumSidesSprite->GetResource())
 	{
diff --git a/Engine/Plugins/Experimental/Avalanche/Source/AvalancheShapesEditor/Private/Visualizers/AvaShapeRoundedPolygonDynMeshVis.cpp b/Engine/Plugins/Experimental/Avalanche/Source/AvalancheShapesEditor/Private/Visualizers/AvaShapeRoundedPolygonDynMeshVis.cpp
--- a/Engine/Plugins/Experimental/Avalanche/Source/AvalancheShapesEditor/Private/Visualizers/AvaShapeRoundedPolygonDynMeshVis.cpp
+++ b/Engine/Plugins/Experimental/Avalanche/Source/AvalancheShapesEditor/Private/Visualizers/AvaShapeRoundedPolygonDynMeshVis.cpp
@@ -2,7 +2,7 @@
 
 #include "AvaShapeRoundedPolygonDynMeshVis.h"
 #include "AvaField.h"
-#include "AvalancheShapesEditorModule.h"
+#include "AvaShapeSprites.h"
 #include "DynamicMeshes/AvaShapeRoundedPolygonDynMesh.h"
 #include "EditorViewportClient.h"
 #include "Engine/Texture2D.h"
@@ -102,7 +102,7 @@ void FAvaShapeRoundedPolygonDynamicMeshVisualizer::DrawVisualizationEditing(cons
 void FAvaShapeRoundedPolygonDynamicMeshVisualizer::DrawBevelButton(const FMeshType* InDynMesh, const FSceneView* InView,
 	FPrimitiveDrawInterface* InPDI, int32 InIconIndex, const FLinearColor& InColor) const
 {
-	UTexture2D* UVSprite = IAvalancheComponentVisualizersModule::Get().GetSettings()->GetVisualizerSprite(FAvalancheShapesEditorModule::CornerSprite);
+	UTexture2D* UVSprite = IAvalancheComponentVisualizersModule::Get().GetSettings()->GetVisualizerSprite(UE::AvaShapes::CornerSprite);
 
 	if (!UVSprite || !UVSprite->GetResource())
 	{
diff --git a/Engine/Plugins/Experimental/Avalanche/Source/AvalancheShapesEditor/Private/Visualizers/AvaShapeSphereDynMeshVis.cpp b/Engine/Plugins/Experimental/Avalanche/Source/AvalancheShapesEditor/Private/Visualizers/AvaShapeSphereDynMeshVis.cpp
--- a/Engine/Plugins/Experimental/Avalanche/Source/AvalancheShapesEditor/Private/Visualizers/AvaShapeSphereDynMeshVis.cpp
+++ b/Engine/Plugins/Experimental/Avalanche/Source/AvalancheShapesEditor/Private/Visualizers/AvaShapeSphereDynMeshVis.cpp
@@ -3,7 +3,7 @@
 #include "AvaShapeSphereDynMeshVis.h"
 #include "AvaField.h"
 #include "AvaShapeActor.h"
-#include "AvalancheShapesEditorModule.h"
+#include "AvaShapeSprites.h"
 #include "DynamicMeshes/AvaShapeSphereDynMesh.h"
 #include "EditorViewportClient.h"
 #include "Engine/Texture2D.h"
@@ -455,7 +455,7 @@ void FAvaShapeSphereDynamicMeshVisualizer::StoreInitialValues()
 void FAvaShapeSphereDynamicMeshVisualizer::DrawNumSidesButton(const FMeshType* InDynMesh, const FSceneView* InView, FPrimitiveDrawInterface* InPDI,
 		int32 InIconIndex, const FLinearColor& InColor) const
 {
-	UTexture2D* NumSidesSprite = IAvalancheComponentVisualizersModule::Get().GetSettings()->GetVisualizerSprite(FAvalancheShapesEditorModule::NumSidesSprite);
+	UTexture2D* NumSidesSprite = IAvalancheComponentVisualizersModule::Get().GetSettings()->GetVisualizerSprite(UE::AvaShapes::NumSidesSprite);
 
 	if (!NumSidesSprite || !NumSidesSprite->GetResource())
 	{
@@ -476,7 +476,7 @@ void FAvaShapeSphereDynamicMeshVisualizer::DrawNumSidesButton(const FMeshType* I
 void FAvaShapeSphereDynamicMeshVisualizer::DrawAngleDegreeButton(const FMeshType* InDynMesh, const FSceneView* InView, FPrimitiveDrawInterface* InPDI,
 		int32 InIconIndex, const FLinearColor& InColor, AvaAlignment InAngleAnchor) const
 {
-	UTexture2D* InnerSizeSprite = IAvalancheComponentVisualizersModule::Get().GetSettings()->GetVisualizerSprite(FAvalancheShapesEditorModule::InnerSizeSprite);
+	UTexture2D* InnerSizeSprite = IAvalancheComponentVisualizersModule::Get().GetSettings()->GetVisualizerSprite(UE::AvaShapes::InnerSizeSprite);
 
 	if (!InnerSizeSprite || !InnerSizeSprite->GetResource())
 	{
diff --git a/Engine/Plugins/Experimental/Avalanche/Source/AvalancheShapesEditor/Private/Visualizers/AvaShapeStarDynMeshVis.cpp b/Engine/Plugins/Experimental/Avalanche/Source/AvalancheShapesEditor/Private/Visualizers/AvaShapeStarDynMeshVis.cpp
--- a/Engine/Plugins/Experimental/Avalanche/Source/AvalancheShapesEditor/Private/Visualizers/AvaShapeStarDynMeshVis.cpp
+++ b/Engine/Plugins/Experimental/Avalanche/Source/AvalancheShapesEditor/Private/Visualizers/AvaShapeStarDynMeshVis.cpp
@@ -3,7 +3,7 @@
 #include "AvaShapeStarDynMeshVis.h"
 #include "AvaField.h"
 #include "AvaShapeActor.h"
-#include "AvalancheShapesEditorModule.h"
+#include "AvaShapeSprites.h"
 #include "DynamicMeshes/AvaShapeStarDynMesh.h"
 #include "EditorViewportClient.h"
 #include "Engine/Texture2D.h"
@@ -114,7 +114,7 @@ void FAvaShapeStarDynamicMeshVisualizer::DrawVisualizationEditing(const UActorCo
 void FAvaShapeStarDynamicMeshVisualizer::DrawNumPointsButton(const FMeshType* InDynMesh, const FSceneView* InView,
 	FPrimitiveDrawInterface* InPDI, int32 InIconIndex, const FLinearColor& InColor) const
 {
-	UTexture2D* NumPointsSprite = IAvalancheComponentVisualizersModule::Get().GetSettings()->GetVisualizerSprite(FAvalancheShapesEditorModule::NumPointsSprite);
+	UTexture2D* NumPointsSprite = IAvalancheComponentVisualizersModule::Get().GetSettings()->GetVisualizerSprite(UE::AvaShapes::NumPointsSprite);
 
 	if (!NumPointsSprite || !NumPointsSprite->GetResource())
 	{
@@ -134,7 +134,7 @@ void FAvaShapeStarDynamicMeshVisualizer::DrawNumPointsButton(const FMeshType* In
 void FAvaShapeStarDynamicMeshVisualizer::DrawInnerSizeButton(const FMeshType* InDynMesh, const FSceneView* InView,
 	FPrimitiveDrawInterface* InPDI, int32 InIconIndex, const FLinearColor& InColor) const
 {
-	UTexture2D* NumPointsSprite = IAvalancheComponentVisualizersModule::Get().GetSettings()->GetVisualizerSprite(FAvalancheShapesEditorModule::InnerSizeSprite);
+	UTexture2D* NumPointsSprite = IAvalancheComponentVisualizersModule::Get().GetSettings()->GetVisualizerSprite(UE::AvaShapes::InnerSizeSprite);
 
 	if (!NumPointsSprite || !NumPointsSprite->GetResource())
 	{
diff --git a/Engine/Plugins/Experimental/Avalanche/Source/AvalancheShapesEditor/Private/Visualizers/AvaShapeTorusDynMeshVis.cpp b/Engine/Plugins/Experimental/Avalanche/Source/AvalancheShapesEditor/Private/Visualizers/AvaShapeTorusDynMeshVis.cpp
--- a/Engine/Plugins/Experimental/Avalanche/Source/AvalancheShapesEditor/Private/Visualizers/AvaShapeTorusDynMeshVis.cpp
+++ b/Engine/Plugins/Experimental/Avalanche/Source/AvalancheShapesEditor/Private/Visualizers/AvaShapeTorusDynMeshVis.cpp
@@ -3,7 +3,7 @@
 #include "AvaShapeTorusDynMeshVis.h"
 #include "AvaField.h"
 #include "AvaShapeActor.h"
-#include "AvalancheShapesEditorModule.h"
+#include "AvaShapeSprites.h"
 #include "DynamicMeshes/AvaShapeTorusDynMesh.h"
 #include "EditorViewportClient.h"
 #include "Engine/Texture2D.h"
@@ -462,7 +462,7 @@ void FAvaShapeTorusDynamicMeshVisualizer::StoreInitialValues()
 void FAvaShapeTorusDynamicMeshVisualizer::DrawNumSlicesButton(const FMeshType* InDynMesh, const FSceneView* InView,
 	FPrimitiveDrawInterface* InPDI, int32 InIconIndex, const FLinearColor& InColor) const
 {
-	UTexture2D* NumSidesSprite = IAvalancheComponentVisualizersModule::Get().GetSettings()->GetVisualizerSprite(FAvalancheShapesEditorModule::NumSidesSprite);
+	UTexture2D* NumSidesSprite = IAvalancheComponentVisualizersModule::Get().GetSettings()->GetVisualizerSprite(UE::AvaShapes::NumSidesSprite);
 
 	if (!NumSidesSprite || !NumSidesSprite->GetResource())
 	{
@@ -482,7 +482,7 @@ void FAvaShapeTorusDynamicMeshVisualizer::DrawNumSlicesButton(const FMeshType* I
 void FAvaShapeTorusDynamicMeshVisualizer::DrawNumSidesButton(const FMeshType* InDynMesh, const FSceneView* InView,
 	FPrimitiveDrawInterface* InPDI, int32 InIconIndex, const FLinearColor& InColor) const
 {
-	UTexture2D* NumPointsSprite = IAvalancheComponentVisualizersModule::Get().GetSettings()->GetVisualizerSprite(FAvalancheShapesEditorModule::NumPointsSprite);
+	UTexture2D* NumPointsSprite = IAvalancheComponentVisualizersModule::Get().GetSettings()->GetVisualizerSprite(UE::AvaShapes::NumPointsSprite);
 
 	if (!NumPointsSprite || !NumPointsSprite->GetResource())
 	{
@@ -502,7 +502,7 @@ void FAvaShapeTorusDynamicMeshVisualizer::DrawNumSidesButton(const FMeshType* In
 void FAvaShapeTorusDynamicMeshVisualizer::DrawInnerSizeButton(const FMeshType* InDynMesh, const FSceneView* InView,
 	FPrimitiveDrawInterface* InPDI, int32 InIconIndex, const FLinearColor& InColor) const
 {
-	UTexture2D* CornerSprite = IAvalancheComponentVisualizersModule::Get().GetSettings()->GetVisualizerSprite(FAvalancheShapesEditorModule::CornerSprite);
+	UTexture2D* CornerSprite = IAvalancheComponentVisualizersModule::Get().GetSettings()->GetVisualizerSprite(UE::AvaShapes::CornerSprite);
 
 	if (!CornerSprite || !CornerSprite->GetResource())
 	{
@@ -522,7 +522,7 @@ void FAvaShapeTorusDynamicMeshVisualizer::DrawInnerSizeButton(const FMeshType* I
 void FAvaShapeTorusDynamicMeshVisualizer::DrawAngleDegreeButton(const FMeshType* InDynMesh, const FSceneView* InView,
 	FPrimitiveDrawInterface* InPDI, int32 InIconIndex, const FLinearColor& InColor) const
 {
-	UTexture2D* InnerSizeSprite = IAvalancheComponentVisualizersModule::Get().GetSettings()->GetVisualizerSprite(FAvalancheShapesEditorModule::InnerSizeSprite);
+	UTexture2D* InnerSizeSprite = IAvalancheComponentVisualizersModule::Get().GetSettings()->GetVisualizerSprite(UE::AvaShapes::InnerSizeSprite);
 
 	if (!InnerSizeSprite || !InnerSizeSprite->GetResource())
 	{
diff --git a/Engine/Plugins/Experimental/Avalanche/Source/AvalancheShapesEditor/Public/AvaShapeFactory.h b/Engine/Plugins/Experimental/Avalanche/Source/AvalancheShapesEditor/Public/AvaShapeFactory.h
deleted file mode 100644
diff --git a/Engine/Plugins/Experimental/Avalanche/Source/AvalancheShapesEditor/Public/AvaShapeSprites.h b/Engine/Plugins/Experimental/Avalanche/Source/AvalancheShapesEditor/Public/AvaShapeSprites.h
new file mode 100644
--- /dev/null
+++ b/Engine/Plugins/Experimental/Avalanche/Source/AvalancheShapesEditor/Public/AvaShapeSprites.h
@@ -0,0 +1,24 @@
+// Copyright Epic Games, Inc. All Rights Reserved.
+
+#pragma once
+
+#include "HAL/Platform.h"
+
+namespace UE::AvaShapes
+{
+	constexpr const TCHAR* BevelSprite = TEXT("Bevel");
+	constexpr const TCHAR* BreakSideSprite = TEXT("BreakSide");
+	constexpr const TCHAR* ColorSelectionSprite = TEXT("ColorSelection");
+	constexpr const TCHAR* CornerSprite = TEXT("Corner");
+	constexpr const TCHAR* DepthSprite = TEXT("Depth");
+	constexpr const TCHAR* InnerSizeSprite = TEXT("InnerSize");
+	constexpr const TCHAR* LinearGradientSprite = TEXT("LinearGradient");
+	constexpr const TCHAR* NumPointsSprite = TEXT("NumPoints");
+	constexpr const TCHAR* NumSidesSprite = TEXT("NumSides");
+	constexpr const TCHAR* TextMaxHeightSprite = TEXT("TextMaxHeight");
+	constexpr const TCHAR* TextMaxWidthSprite = TEXT("TextMaxWidth");
+	constexpr const TCHAR* TextScaleProportionallySprite = TEXT("TextScaleProportionally");
+	constexpr const TCHAR* SizeSprite = TEXT("Size");
+	constexpr const TCHAR* SlantSprite = TEXT("Slant");
+	constexpr const TCHAR* UVSprite = TEXT("UV");
+}
diff --git a/Engine/Plugins/Experimental/Avalanche/Source/AvalancheShapesEditor/Public/AvalancheShapesEditorModule.h b/Engine/Plugins/Experimental/Avalanche/Source/AvalancheShapesEditor/Public/AvalancheShapesEditorModule.h
deleted file mode 100644
diff --git a/Engine/Plugins/Experimental/Avalanche/Source/AvalancheShapesEditor/Public/Tools/AvaShapesEditorShapeAreaToolBase.h b/Engine/Plugins/Experimental/Avalanche/Source/AvalancheShapesEditor/Public/Tools/AvaShapesEditorShapeAreaToolBase.h
deleted file mode 100644
diff --git a/Engine/Plugins/Experimental/Avalanche/Source/AvalancheShapesEditor/Public/Tools/AvaShapesEditorShapeToolBase.h b/Engine/Plugins/Experimental/Avalanche/Source/AvalancheShapesEditor/Public/Tools/AvaShapesEditorShapeToolBase.h
deleted file mode 100644
diff --git a/Engine/Plugins/Experimental/Avalanche/Source/AvalancheShapesEditor/Public/Visualizers/AvaShape2DDynMeshVis.h b/Engine/Plugins/Experimental/Avalanche/Source/AvalancheShapesEditor/Public/Visualizers/AvaShape2DDynMeshVis.h
deleted file mode 100644
diff --git a/Engine/Plugins/Experimental/Avalanche/Source/AvalancheShapesEditor/Public/Visualizers/AvaShape3DDynMeshVis.h b/Engine/Plugins/Experimental/Avalanche/Source/AvalancheShapesEditor/Public/Visualizers/AvaShape3DDynMeshVis.h
deleted file mode 100644
diff --git a/Engine/Plugins/Experimental/Avalanche/Source/AvalancheShapesEditor/Public/Visualizers/AvaShapeDynMeshVis.h b/Engine/Plugins/Experimental/Avalanche/Source/AvalancheShapesEditor/Public/Visualizers/AvaShapeDynMeshVis.h
deleted file mode 100644
diff --git a/Engine/Plugins/Experimental/Avalanche/Source/AvalancheTextEditor/AvalancheTextEditor.Build.cs b/Engine/Plugins/Experimental/Avalanche/Source/AvalancheTextEditor/AvalancheTextEditor.Build.cs
--- a/Engine/Plugins/Experimental/Avalanche/Source/AvalancheTextEditor/AvalancheTextEditor.Build.cs
+++ b/Engine/Plugins/Experimental/Avalanche/Source/AvalancheTextEditor/AvalancheTextEditor.Build.cs
@@ -9,6 +9,13 @@ public class AvalancheTextEditor : ModuleRules
 	{
 		PCHUsage = ModuleRules.PCHUsageMode.UseExplicitOrSharedPCHs;
 
+		PrivateIncludePathModuleNames.AddRange(
+			new string[]
+			{
+				"AvalancheShapesEditor",
+			}
+		);
+
 		PrivateDependencyModuleNames.AddRange(
 			new string[]
 			{
@@ -16,7 +23,6 @@ public class AvalancheTextEditor : ModuleRules
 				"AvalancheComponentVisualizers",
 				"AvalancheInteractiveTools",
 				"AvalancheLevelViewport",
-				"AvalancheShapesEditor",
 				"AvalancheText",
 				"ComponentVisualizers",
 				"Core",
diff --git a/Engine/Plugins/Experimental/Avalanche/Source/AvalancheTextEditor/Private/Visualizer/AvaTextVisualizer.cpp b/Engine/Plugins/Experimental/Avalanche/Source/AvalancheTextEditor/Private/Visualizer/AvaTextVisualizer.cpp
--- a/Engine/Plugins/Experimental/Avalanche/Source/AvalancheTextEditor/Private/Visualizer/AvaTextVisualizer.cpp
+++ b/Engine/Plugins/Experimental/Avalanche/Source/AvalancheTextEditor/Private/Visualizer/AvaTextVisualizer.cpp
@@ -2,8 +2,8 @@
 
 #include "AvaTextVisualizer.h"
 #include "AvaField.h"
+#include "AvaShapeSprites.h"
 #include "AvaText3DComponent.h"
-#include "AvalancheShapesEditorModule.h"
 #include "EditorViewportClient.h"
 #include "Engine/Texture2D.h"
 #include "IAvalancheComponentVisualizersModule.h"
@@ -1059,7 +1059,7 @@ bool FAvaTextVisualizer::DrawGradientCenterHandle(const UText3DComponent* InText
 {
 	static constexpr float BaseSize = 1.f;
 
-	UTexture2D* UVSprite = IAvalancheComponentVisualizersModule::Get().GetSettings()->GetVisualizerSprite(FAvalancheShapesEditorModule::SizeSprite);
+	UTexture2D* UVSprite = IAvalancheComponentVisualizersModule::Get().GetSettings()->GetVisualizerSprite(UE::AvaShapes::SizeSprite);
 
 	if (!UVSprite || !UVSprite->GetResource())
 	{
@@ -1080,7 +1080,7 @@ bool FAvaTextVisualizer::DrawGradientSmoothnessHandle(const UText3DComponent* In
 {
 	static constexpr float BaseSize = 1.5f;
 		
-	const UTexture2D* const UVSprite = IAvalancheComponentVisualizersModule::Get().GetSettings()->GetVisualizerSprite(FAvalancheShapesEditorModule::UVSprite);
+	const UTexture2D* const UVSprite = IAvalancheComponentVisualizersModule::Get().GetSettings()->GetVisualizerSprite(UE::AvaShapes::UVSprite);
 
 	if (!UVSprite || !UVSprite->GetResource())
 	{
@@ -1185,7 +1185,7 @@ void FAvaTextVisualizer::DrawMaxTextWidthButton(const UText3DComponent* InTextCo
 		return;
 	}
 
-	UTexture2D* UVSprite = IAvalancheComponentVisualizersModule::Get().GetSettings()->GetVisualizerSprite(FAvalancheShapesEditorModule::TextMaxWidthSprite);
+	UTexture2D* UVSprite = IAvalancheComponentVisualizersModule::Get().GetSettings()->GetVisualizerSprite(UE::AvaShapes::TextMaxWidthSprite);
 
 	if (!UVSprite || !UVSprite->GetResource())
 	{
@@ -1210,7 +1210,7 @@ void FAvaTextVisualizer::DrawMaxTextHeightButton(const UText3DComponent* InTextC
 		return;
 	}
 
-	const UTexture2D* const UVSprite = IAvalancheComponentVisualizersModule::Get().GetSettings()->GetVisualizerSprite(FAvalancheShapesEditorModule::TextMaxHeightSprite);
+	const UTexture2D* const UVSprite = IAvalancheComponentVisualizersModule::Get().GetSettings()->GetVisualizerSprite(UE::AvaShapes::TextMaxHeightSprite);
 
 	if (!UVSprite || !UVSprite->GetResource())
 	{
@@ -1235,7 +1235,7 @@ void FAvaTextVisualizer::DrawScaleProportionallyButton(const UText3DComponent* I
 		return;
 	}
 
-	const UTexture2D* const UVSprite = IAvalancheComponentVisualizersModule::Get().GetSettings()->GetVisualizerSprite(FAvalancheShapesEditorModule::TextScaleProportionallySprite);
+	const UTexture2D* const UVSprite = IAvalancheComponentVisualizersModule::Get().GetSettings()->GetVisualizerSprite(UE::AvaShapes::TextScaleProportionallySprite);
 
 	if (!UVSprite || !UVSprite->GetResource())
 	{
@@ -1265,7 +1265,7 @@ void FAvaTextVisualizer::DrawEditGradientButton(
 	}
 
 	// todo: proper sprite
-	const UTexture2D* const UVSprite = IAvalancheComponentVisualizersModule::Get().GetSettings()->GetVisualizerSprite(FAvalancheShapesEditorModule::LinearGradientSprite);
+	const UTexture2D* const UVSprite = IAvalancheComponentVisualizersModule::Get().GetSettings()->GetVisualizerSprite(UE::AvaShapes::LinearGradientSprite);
 
 	if (!UVSprite || !UVSprite->GetResource())
 	{
@@ -1294,7 +1294,7 @@ void FAvaTextVisualizer::DrawMaxTextWidthHandle(
 		return;
 	}
 
-	const UTexture2D* const UVSprite = IAvalancheComponentVisualizersModule::Get().GetSettings()->GetVisualizerSprite(FAvalancheShapesEditorModule::SizeSprite);
+	const UTexture2D* const UVSprite = IAvalancheComponentVisualizersModule::Get().GetSettings()->GetVisualizerSprite(UE::AvaShapes::SizeSprite);
 
 	if (!UVSprite || !UVSprite->GetResource())
 	{
@@ -1320,7 +1320,7 @@ void FAvaTextVisualizer::DrawMaxTextHeightHandle(const UText3DComponent* InTextC
 		return;
 	}
 
-	const UTexture2D* const UVSprite = IAvalancheComponentVisualizersModule::Get().GetSettings()->GetVisualizerSprite(FAvalancheShapesEditorModule::SizeSprite);
+	const UTexture2D* const UVSprite = IAvalancheComponentVisualizersModule::Get().GetSettings()->GetVisualizerSprite(UE::AvaShapes::SizeSprite);
 
 	if (!UVSprite || !UVSprite->GetResource())
 	{
@@ -1346,7 +1346,7 @@ void FAvaTextVisualizer::DrawGradientLineStartHandle(const UAvaText3DComponent*
 		return;
 	}
 
-	const UTexture2D* const UVSprite = IAvalancheComponentVisualizersModule::Get().GetSettings()->GetVisualizerSprite(FAvalancheShapesEditorModule::ColorSelectionSprite);
+	const UTexture2D* const UVSprite = IAvalancheComponentVisualizersModule::Get().GetSettings()->GetVisualizerSprite(UE::AvaShapes::ColorSelectionSprite);
 
 	if (!UVSprite || !UVSprite->GetResource())
 	{
@@ -1377,7 +1377,7 @@ void FAvaTextVisualizer::DrawGradientLineEndHandle(const UAvaText3DComponent* In
 		return;
 	}
 
-	const UTexture2D* const UVSprite = IAvalancheComponentVisualizersModule::Get().GetSettings()->GetVisualizerSprite(FAvalancheShapesEditorModule::ColorSelectionSprite);
+	const UTexture2D* const UVSprite = IAvalancheComponentVisualizersModule::Get().GetSettings()->GetVisualizerSprite(UE::AvaShapes::ColorSelectionSprite);
 
 	if (!UVSprite || !UVSprite->GetResource())
 	{
