[Motion Design] Text 3D Component ordering and grouping cleanup

diff --git a/Engine/Plugins/Experimental/Avalanche/Source/AvalancheText/Private/AvaText3DComponent.cpp b/Engine/Plugins/Experimental/Avalanche/Source/AvalancheText/Private/AvaText3DComponent.cpp
--- a/Engine/Plugins/Experimental/Avalanche/Source/AvalancheText/Private/AvaText3DComponent.cpp
+++ b/Engine/Plugins/Experimental/Avalanche/Source/AvalancheText/Private/AvaText3DComponent.cpp
@@ -25,7 +25,7 @@
 		{
 			return;
 		}
-		
+
 		if (UMaterialInstance* CurrentMaterialInstance = Cast<UMaterialInstance>(Invoke(InGetter, InComponent)))
 		{
 			InNewMaterial->CopyInterpParameters(CurrentMaterialInstance);
@@ -39,30 +39,30 @@
 UAvaText3DComponent::UAvaText3DComponent()
 {
 	PrimaryComponentTick.bCanEverTick = true;
-	
+
 	Extrude = 0.0f;
 	Bevel = 0.0f;
 	BevelType = EText3DBevelType::Convex;
 	BevelSegments = 8;
-	
+
 	HorizontalAlignment = EText3DHorizontalTextAlignment::Left;
 	VerticalAlignment = EText3DVerticalTextAlignment::FirstLine;
 
 	bHasMaxWidth = false;
 	bHasMaxHeight = false;
 	bScaleProportionally = false;
-	
+
 	Kerning = 0.0f;
 	WordSpacing = 0.0f;
 	LineSpacing = 0.0f;
-	
+
 	MaxWidth = 100.0f;
 	MaxHeight = 100.0f;
 
 	MotionDesignFont = {};
 
 	Tiling = FVector2D(1.0, 1.0);
-	
+
 	Color = FLinearColor::White;
 	ExtrudeColor = FLinearColor::Gray;
 	BevelColor = FLinearColor::Gray;
@@ -82,459 +82,213 @@
 #endif
 }
 
-void UAvaText3DComponent::PostInitProperties()
+void UAvaText3DComponent::SetEnforceUpperCase(bool bInEnforceUpperCase)
 {
-	Super::PostInitProperties();
-	
-	SetupParentRootAndCallbacks();
-
-	MarkForFullRefresh();
-}
-
-void UAvaText3DComponent::RetrieveDataFromBaseComponent()
-{
-	RetrieveCurrentColoringStyle();	
-	RetrieveCurrentMaterials();	
-	RetrieveCurrentTextAndFont();
-	RetrieveCurrentAlignment();
-}
-
-void UAvaText3DComponent::PostLoad()
-{
-	Super::PostLoad();
-
-	SetupParentRootAndCallbacks();
-	RetrieveDataFromBaseComponent();
-}
-
-void UAvaText3DComponent::PostDuplicateActions()
-{
-	SetupParentRootAndCallbacks();
-	
-	RetrieveCurrentTextAndFont();
-	
-	MarkForFullRefresh();
-	RefreshText3D();
-}
-
-void UAvaText3DComponent::PostDuplicate(EDuplicateMode::Type DuplicateMode)
-{	
-	Super::PostDuplicate(DuplicateMode);
-
-	PostDuplicateActions();
-}
-
-// overriding PostEditImport to handle Owner duplication
-void UAvaText3DComponent::PostEditImport()
-{
-	Super::PostEditImport();
-	
-	PostDuplicateActions();
-}
-
-#if WITH_EDITOR
-void UAvaText3DComponent::PostEditChangeProperty(FPropertyChangedEvent& PropertyChangedEvent)
-{
-	const FName ChangedPropertyName = PropertyChangedEvent.GetMemberPropertyName();
-
-	if (PropertyChangeFunctions.Contains(ChangedPropertyName))
-	{
-		PropertyChangeFunctions[ChangedPropertyName](this);
-	}
-
-	// Call the Super::PostEditChangeProperty after the PropertyChangeFunctions execute the right Setter/Refresh
-	Super::PostEditChangeProperty(PropertyChangedEvent);
-}
-
-void UAvaText3DComponent::PostEditComponentMove(bool bFinished)
-{
-	Super::PostEditComponentMove(bFinished);
-	ScheduleTextRefreshOnTicker();
-}
-
-#endif
-
-void UAvaText3DComponent::OnUpdateTransform(EUpdateTransformFlags UpdateTransformFlags, ETeleportType Teleport)
-{
-	Super::OnUpdateTransform(UpdateTransformFlags, Teleport);
-
-	ScheduleTextRefreshOnTicker();
-}
-
-void UAvaText3DComponent::SetupParentRootAndCallbacks()
-{
-	RefreshParentRootComponent();
-	RegisterCallbacks();
-}
-
-void UAvaText3DComponent::OnAttachmentChanged()
-{
-	Super::OnAttachmentChanged();
-
-	SetupParentRootAndCallbacks();
-}
-
-void UAvaText3DComponent::OnRegister()
-{
-	Super::OnRegister();
-
-	SetComponentTickEnabled(false);
-	
-	SetupParentRootAndCallbacks();
-	RefreshText3D();
-}
-
-void UAvaText3DComponent::FormatText(FText& InOutText) const
-{
-	if (bEnforceUpperCase)
-	{
-		InOutText = InOutText.ToUpper();
-	}
-}
-
-void UAvaText3DComponent::OnRootComponentTransformUpdated(USceneComponent* InSceneComponent, EUpdateTransformFlags InUpdateTransformFlags, ETeleportType InTeleportType)
-{
-	ScheduleTextRefreshOnTicker();
-}
-
-void UAvaText3DComponent::SetupGradientMaterial()
-{
-	TexturedTextMaterialInstance = nullptr;
-	
-	if (!IsValid(GradientTextMaterialInstance) || GradientTextMaterialInstance->GetOuter() != this || IsMarkedForMaterialRefresh())
-	{
-		const EAvaTextMaterialFeatures MaterialFeatures = GetMaterialFeaturesFromProperties();
-		const FAvaTextMaterialSettings MaterialSettings(EAvaTextColoringStyle::Gradient, MaterialFeatures);
-		
-		if (UMaterialInstanceDynamic* const GradientMaterial = GetMIDWithSettings(GradientTextMaterialInstance, MaterialSettings))
-		{
-			UMaterialInstanceDynamic* PreviousMaterial = GradientTextMaterialInstance;
-			GradientTextMaterialInstance = GradientMaterial;
-			if (PreviousMaterial)
-			{
-				GradientTextMaterialInstance->CopyInterpParameters(PreviousMaterial);	
-			}
-		
-			RemoveRefreshReason(EAvaTextRefreshReason::MaterialInstancesChange);
-		}
-	}
-
-	if (!GradientTextMaterialInstance)
+	if (bEnforceUpperCase == bInEnforceUpperCase)
 	{
 		return;
 	}
 
-	UE::AvaText::Private::PatchMaterial(
-		this
-		, &UText3DComponent::GetFrontMaterial
-		, &UText3DComponent::SetFrontMaterial
-		, GradientTextMaterialInstance);
-
-	UE::AvaText::Private::PatchMaterial(
-		this
-		, &UText3DComponent::GetExtrudeMaterial
-		, &UText3DComponent::SetExtrudeMaterial
-		, GradientTextMaterialInstance);
-
-	UE::AvaText::Private::PatchMaterial(
-		this
-		, &UText3DComponent::GetBevelMaterial
-		, &UText3DComponent::SetBevelMaterial
-		, GradientTextMaterialInstance);
-
-	UE::AvaText::Private::PatchMaterial(
-		this
-		, &UText3DComponent::GetBackMaterial
-		, &UText3DComponent::SetBackMaterial
-		, GradientTextMaterialInstance);
-
-	if (IsMarkedForMaterialParametersRefresh())
-	{
-		RefreshGradientValues();
-		RefreshTranslucencyValues();
-	}
+	bEnforceUpperCase = bInEnforceUpperCase;
+	RefreshFormatting();
 }
 
-void UAvaText3DComponent::SetupCustomMaterial()
+void UAvaText3DComponent::SetMotionDesignFont(const FAvaFont& InFont)
 {
-	TexturedTextMaterialInstance = nullptr;
-	GradientTextMaterialInstance = nullptr;
-	
-	if (CustomMaterial && IsMarkedForMaterialRefresh())
+	if (MotionDesignFont == InFont)
 	{
-		RemoveRefreshReason(EAvaTextRefreshReason::MaterialInstancesChange);
-
-		UE::AvaText::Private::PatchMaterial(
-			this
-			, &UText3DComponent::GetFrontMaterial
-			, &UText3DComponent::SetFrontMaterial
-			, GradientTextMaterialInstance);
-
-		if (UMaterialInstanceDynamic* CustomMaterialAsInstance = Cast<UMaterialInstanceDynamic>(CustomMaterial))
-		{
-			UE::AvaText::Private::PatchMaterial(
-				this
-				, &UText3DComponent::GetFrontMaterial
-				, &UText3DComponent::SetFrontMaterial
-				, CustomMaterialAsInstance);
-
-			UE::AvaText::Private::PatchMaterial(
-				this
-				, &UText3DComponent::GetExtrudeMaterial
-				, &UText3DComponent::SetExtrudeMaterial
-				, CustomMaterialAsInstance);
-
-			UE::AvaText::Private::PatchMaterial(
-				this
-				, &UText3DComponent::GetBevelMaterial
-				, &UText3DComponent::SetBevelMaterial
-				, CustomMaterialAsInstance);
-
-			UE::AvaText::Private::PatchMaterial(
-				this
-				, &UText3DComponent::GetBackMaterial
-				, &UText3DComponent::SetBackMaterial
-				, CustomMaterialAsInstance);
-		}
-		else
-		{
-			SetFrontMaterial(CustomMaterial);
-			SetExtrudeMaterial(CustomMaterial);
-			SetBevelMaterial(CustomMaterial);
-			SetBackMaterial(CustomMaterial);
-		}
+		return;
 	}
+	MotionDesignFont = InFont;
+	RefreshGeometry();
 }
 
-void UAvaText3DComponent::RefreshTranslucencyValues()
+void UAvaText3DComponent::SetAlignment(FAvaTextAlignment InAlignment)
 {
-	switch (TranslucencyStyle)
-	{
-		case EAvaTextTranslucency::Translucent:
-			RefreshOpacity();
-			break;
-				
-		case EAvaTextTranslucency::GradientMask:
-			RefreshMaskValues();
-			break;
-			
-		default: ;
-	}
-}
-
-void UAvaText3DComponent::SetupTexturedMaterial()
-{
-	GradientTextMaterialInstance = nullptr;
-	
-	if (!IsValid(TexturedTextMaterialInstance) || TexturedTextMaterialInstance->GetOuter() != this || IsMarkedForMaterialRefresh())
-	{
-		const EAvaTextMaterialFeatures MaterialFeatures = GetMaterialFeaturesFromProperties();
-		const FAvaTextMaterialSettings MaterialSettings(EAvaTextColoringStyle::FromTexture, MaterialFeatures);
-
-		if (UMaterialInstanceDynamic* const TexturedMaterial = GetMIDWithSettings(TexturedTextMaterialInstance, MaterialSettings))
-		{
-			UMaterialInstanceDynamic* PreviousMaterial = TexturedTextMaterialInstance;
-			TexturedTextMaterialInstance = TexturedMaterial;
-			if (PreviousMaterial)
-			{
-				TexturedTextMaterialInstance->CopyInterpParameters(PreviousMaterial);	
-			}
-		}
-		
-		RemoveRefreshReason(EAvaTextRefreshReason::MaterialInstancesChange);
-	}
-
-	if (!TexturedTextMaterialInstance)
+	if (Alignment == InAlignment)
 	{
 		return;
 	}
 
-	{
-		UE::AvaText::Private::PatchMaterial(
-			this
-			, &UText3DComponent::GetFrontMaterial
-			, &UText3DComponent::SetFrontMaterial
-			, TexturedTextMaterialInstance);
-
-		UE::AvaText::Private::PatchMaterial(
-			this
-			, &UText3DComponent::GetExtrudeMaterial
-			, &UText3DComponent::SetExtrudeMaterial
-			, TexturedTextMaterialInstance);
-
-		UE::AvaText::Private::PatchMaterial(
-			this
-			, &UText3DComponent::GetBevelMaterial
-			, &UText3DComponent::SetBevelMaterial
-			, TexturedTextMaterialInstance);
-
-		UE::AvaText::Private::PatchMaterial(
-			this
-			, &UText3DComponent::GetBackMaterial
-			, &UText3DComponent::SetBackMaterial
-			, TexturedTextMaterialInstance);
-	}
-
-	UTexture* CurrentTexture;
-	const FMemoryImageMaterialParameterInfo MaterialInfo(UAvaTextMaterialHub::FMaterialStatics::MainTexture_MatParam);
-	TexturedTextMaterialInstance->GetTextureParameterValue(MaterialInfo, CurrentTexture);
-
-	if (MainTexture != CurrentTexture)
-	{
-		TexturedTextMaterialInstance->SetTextureParameterValue(UAvaTextMaterialHub::FMaterialStatics::MainTexture_MatParam, MainTexture);
-	}
-
-	TexturedTextMaterialInstance->SetScalarParameterValue(UAvaTextMaterialHub::FMaterialStatics::TexturedUTiling_MatParam, Tiling.X);
-	TexturedTextMaterialInstance->SetScalarParameterValue(UAvaTextMaterialHub::FMaterialStatics::TexturedVTiling_MatParam, Tiling.Y);
-
-	RefreshTranslucencyValues();
+	Alignment = InAlignment;
+	RefreshAlignment();
 }
 
-void UAvaText3DComponent::SetupSolidMaterial()
+void UAvaText3DComponent::SetColoringStyle(const EAvaTextColoringStyle& InColoringStyle)
 {
-	TexturedTextMaterialInstance = nullptr;
-	GradientTextMaterialInstance = nullptr;
+	if (ColoringStyle != InColoringStyle)
+	{
+		ColoringStyle = InColoringStyle;
+		RefreshMaterialInstances();
+	}
+}
 
-	if (!IsValid(TextMaterialInstance) || TextMaterialInstance->GetOuter() != this || IsMarkedForMaterialRefresh())
+void UAvaText3DComponent::SetColor(const FLinearColor& InColor)
+{
+	if (Color != InColor)
 	{
-		const EAvaTextMaterialFeatures MaterialFeatures = GetMaterialFeaturesFromProperties();
-		const FAvaTextMaterialSettings MaterialSettings(EAvaTextColoringStyle::Solid, MaterialFeatures);
-
-		if (UMaterialInstanceDynamic* const SolidMaterial = GetMIDWithSettings(TextMaterialInstance, MaterialSettings))
-		{
-			UMaterialInstanceDynamic* PreviousMaterial = TextMaterialInstance;
-			TextMaterialInstance = SolidMaterial;
-			if (PreviousMaterial)
-			{
-				TextMaterialInstance->CopyInterpParameters(PreviousMaterial);	
-			}
-
-			PreviousMaterial = ExtrudeTextMaterialInstance;
-			ExtrudeTextMaterialInstance = SolidMaterial;
-			if (PreviousMaterial)
-			{
-				ExtrudeTextMaterialInstance->CopyInterpParameters(PreviousMaterial);
-			}
-
-			PreviousMaterial = BevelTextMaterialInstance;
-			BevelTextMaterialInstance = SolidMaterial;
-			if (PreviousMaterial)
-			{
-				BevelTextMaterialInstance->CopyInterpParameters(PreviousMaterial);
-			}
-		}
-	
-		RemoveRefreshReason(EAvaTextRefreshReason::MaterialInstancesChange);
-	}
- 
-	{
-		UE::AvaText::Private::PatchMaterial(
-			this
-			, &UText3DComponent::GetFrontMaterial
-			, &UText3DComponent::SetFrontMaterial
-			, TextMaterialInstance);
-
-		UE::AvaText::Private::PatchMaterial(
-			this
-			, &UText3DComponent::GetExtrudeMaterial
-			, &UText3DComponent::SetExtrudeMaterial
-			, ExtrudeTextMaterialInstance);
-
-		UE::AvaText::Private::PatchMaterial(
-			this
-			, &UText3DComponent::GetBevelMaterial
-			, &UText3DComponent::SetBevelMaterial
-			, BevelTextMaterialInstance);
-
-		UE::AvaText::Private::PatchMaterial(
-			this
-			, &UText3DComponent::GetBackMaterial
-			, &UText3DComponent::SetBackMaterial
-			, TextMaterialInstance);
-	}
-
-	if (IsMarkedForMaterialParametersRefresh())
-	{
+		Color = InColor;
 		RefreshColor();
+	}
+}
+
+void UAvaText3DComponent::SetExtrudeColor(const FLinearColor& InColor)
+{
+	if (ExtrudeColor != InColor)
+	{
+		ExtrudeColor = InColor;
 		RefreshExtrudeColor();
+	}
+}
+
+void UAvaText3DComponent::SetBevelColor(const FLinearColor& InColor)
+{
+	if (BevelColor != InColor)
+	{
+		BevelColor = InColor;
 		RefreshBevelColor();
-		
-		RefreshTranslucencyValues();
 	}
 }
 
-void UAvaText3DComponent::SetupTextMaterials(const EAvaTextColoringStyle& InColoringStyle)
+void UAvaText3DComponent::SetGradientSettings(const FAvaLinearGradientSettings& InGradientSettings)
 {
-	if (IsMarkedForMaterialRefresh())
+	if (GradientSettings == InGradientSettings)
 	{
-		switch (InColoringStyle)
-		{
-		case EAvaTextColoringStyle::Invalid:
-			break;
-			
-		case EAvaTextColoringStyle::Solid:
-			SetupSolidMaterial();
-			break;
+		return;
+	}
 
-		case EAvaTextColoringStyle::Gradient:
-			SetupGradientMaterial();
-			break;
-					
-		case EAvaTextColoringStyle::FromTexture:
-			SetupTexturedMaterial();
-			break;
-					
-		case EAvaTextColoringStyle::CustomMaterial:
-			SetupCustomMaterial();
-			break;
-		
-		default:
-			break;
-		}
+	GradientSettings = InGradientSettings;
+	RefreshGradient();
+}
+
+void UAvaText3DComponent::SetGradientColors(const FLinearColor& InColorA, const FLinearColor& InColorB)
+{
+	if (GradientSettings.ColorA == InColorA && GradientSettings.ColorB == InColorB)
+	{
+		return;
+	}
+
+	GradientSettings.ColorA = InColorA;
+	GradientSettings.ColorB = InColorB;
+	RefreshGradient();
+}
+
+void UAvaText3DComponent::SetMainTexture(UTexture2D* InMainTexture)
+{
+	if (MainTexture != InMainTexture)
+	{
+		MainTexture = InMainTexture;
+		RefreshMaterialInstances();
 	}
 }
 
-void UAvaText3DComponent::SetupTextLayout()
+void UAvaText3DComponent::SetTiling(const FVector2D InTiling)
 {
-	if (IsMarkedForLayoutRefresh())
+	if (Tiling != InTiling)
 	{
-		TriggerInternalRebuild(EText3DModifyFlags::Layout);
-		RemoveRefreshReason(EAvaTextRefreshReason::LayoutChange);
+		Tiling = InTiling;
+		RefreshMaterialInstances();
 	}
 }
 
-void UAvaText3DComponent::SetupTextGeometry()
+void UAvaText3DComponent::SetCustomMaterial(UMaterialInterface* InCustomMaterial)
 {
-	if (IsMarkedForGeometryRefresh())
+	if (CustomMaterial != InCustomMaterial)
 	{
-		// we are about to setup multiple properties together, let's freeze rebuild
-		SetFreeze(true);
-	
-		SetText(Text);
-
-		if (UFont* UnderlyingFont = MotionDesignFont.GetFont())
-		{
-			SetFont(UnderlyingFont);	
-		}
-		
-		SetExtrude(Extrude);
-		
-		SetBevel(Bevel);
-		SetBevelType(BevelType);
-		SetBevelSegments(BevelSegments);
-	
-		//... and finally unfreeze, to trigger pending rebuild
-		SetFreeze(false);
-
-		RemoveRefreshReason(EAvaTextRefreshReason::GeometryChange);
+		CustomMaterial = InCustomMaterial;
+		RefreshMaterialInstances();
 	}
 }
 
+void UAvaText3DComponent::SetTranslucencyStyle(EAvaTextTranslucency InTranslucencyStyle)
+{
+	if (TranslucencyStyle != InTranslucencyStyle)
+	{
+		TranslucencyStyle = InTranslucencyStyle;
+		RefreshMaterialInstances();
+	}
+}
+
+void UAvaText3DComponent::SetOpacity(float InOpacity)
+{
+	if (Opacity != InOpacity)
+	{
+		Opacity = InOpacity;
+		RefreshOpacity();
+	}
+}
+
+void UAvaText3DComponent::SetMaskOrientation(EAvaMaterialMaskOrientation InMaskOrientation)
+{
+	if (MaskOrientation != InMaskOrientation)
+	{
+		MaskOrientation = InMaskOrientation;
+		RefreshMask();
+	}
+}
+
+void UAvaText3DComponent::SetMaskSmoothness(float InMaskSmoothness)
+{
+	if (MaskSmoothness != InMaskSmoothness)
+	{
+		MaskSmoothness = InMaskSmoothness;
+		RefreshMask();
+	}
+}
+
+void UAvaText3DComponent::SetMaskOffset(float InMaskOffset)
+{
+	if (MaskOffset != InMaskOffset)
+	{
+		MaskOffset = InMaskOffset;
+		RefreshMask();
+	}
+}
+
+void UAvaText3DComponent::SetMaskRotation(float InMaskRotation)
+{
+	if (MaskRotation != InMaskRotation)
+	{
+		MaskRotation = InMaskRotation;
+		RefreshMask();
+	}
+}
+
+void UAvaText3DComponent::SetIsUnlit(bool bInIsUnlit)
+{
+	if (bIsUnlit != bInIsUnlit)
+	{
+		bIsUnlit = bInIsUnlit;
+		RefreshUnlit();
+	}
+}
+
+UAvaText3DComponent::FOnGetMaterialWithSettings& UAvaText3DComponent::GetMaterialProviderDelegate()
+{
+	return OnGetMaterialWithSettingsDelegate;
+}
+
+void UAvaText3DComponent::ForEachMID(TUniqueFunction<void(UMaterialInstanceDynamic* InMID)>&& InFunc) const
+{
+	auto ExecuteIfMID = [&InFunc](UMaterialInterface* InMaterial)
+	{
+		if (UMaterialInstanceDynamic* MID = Cast<UMaterialInstanceDynamic>(InMaterial))
+		{
+			InFunc(MID);
+		}
+	};
+
+	ExecuteIfMID(GetFrontMaterial());
+	ExecuteIfMID(GetBevelMaterial());
+	ExecuteIfMID(GetExtrudeMaterial());
+	ExecuteIfMID(GetBackMaterial());
+}
+
 FVector UAvaText3DComponent::GetGradientDirection() const
 {
 	// Gradient Rot is a [0;1] ratio, representing to [0;360] degrees
 	float GradientRot;
-	
+
 	if (GradientSettings.Direction == EAvaGradientDirection::Vertical)
 	{
 		GradientRot = UAvaTextMaterialHub::FMaterialStatics::MaterialMaskTopBottomRotationValue;
@@ -547,9 +301,9 @@
 	{
 		GradientRot = GradientSettings.Rotation;
 	}
-	
+
 	FVector GradientDir = GetUpVector().RotateAngleAxis(-GradientRot * 360, GetForwardVector());
-	
+
 	/**
 	 * In order to properly map gradient offset along the text surface, text bounds are not normalized (anymore) in the material function creating the gradient.
 	 * Therefore, we need to remap the gradient direction, taking into account the current text actor bounds.
@@ -566,12 +320,18 @@
 	return GradientDir;
 }
 
+void UAvaText3DComponent::RefreshMaterialInstances()
+{
+	MarkForMaterialsFullRefresh();
+	RefreshText3D();
+}
+
 void UAvaText3DComponent::RefreshText3D()
 {
 	SetupTextGeometry();
 	SetupTextLayout();
 	SetupTextMaterials(ColoringStyle);
-	
+
 	RefreshGradientValues();
 	RefreshMaskValues();
 	RefreshMaterialGeometryParameterValues();
@@ -586,13 +346,69 @@
 	MarkRenderInstancesDirty();
 }
 
+void UAvaText3DComponent::MarkForGeometryRefresh()
+{
+	EnumAddFlags(TextRefreshReason, EAvaTextRefreshReason::GeometryChange);
+}
+
+bool UAvaText3DComponent::IsMarkedForGeometryRefresh() const
+{
+	return EnumHasAnyFlags(TextRefreshReason, EAvaTextRefreshReason::GeometryChange);
+}
+
+void UAvaText3DComponent::MarkForLayoutRefresh()
+{
+	EnumAddFlags(TextRefreshReason, EAvaTextRefreshReason::LayoutChange);
+}
+
+bool UAvaText3DComponent::IsMarkedForLayoutRefresh() const
+{
+	return EnumHasAnyFlags(TextRefreshReason, EAvaTextRefreshReason::LayoutChange);
+}
+
+void UAvaText3DComponent::MarkForMaterialsFullRefresh()
+{
+	EnumAddFlags(TextRefreshReason, EAvaTextRefreshReason::MaterialInstancesChange);
+	EnumAddFlags(TextRefreshReason, EAvaTextRefreshReason::MaterialParametersChange);
+}
+
+bool UAvaText3DComponent::IsMarkedForMaterialRefresh() const
+{
+	return EnumHasAnyFlags(TextRefreshReason, EAvaTextRefreshReason::MaterialInstancesChange);
+}
+
+void UAvaText3DComponent::MarkForMaterialParametersRefresh()
+{
+	EnumAddFlags(TextRefreshReason, EAvaTextRefreshReason::MaterialParametersChange);
+}
+
+bool UAvaText3DComponent::IsMarkedForMaterialParametersRefresh() const
+{
+	return EnumHasAnyFlags(TextRefreshReason, EAvaTextRefreshReason::MaterialParametersChange);
+}
+
+void UAvaText3DComponent::RemoveRefreshReason(EAvaTextRefreshReason InRefreshReasonToRemove)
+{
+	EnumRemoveFlags(TextRefreshReason, InRefreshReasonToRemove);
+}
+
+void UAvaText3DComponent::MarkForFullRefresh()
+{
+	EnumAddFlags(TextRefreshReason, EAvaTextRefreshReason::FullRefresh);
+}
+
+void UAvaText3DComponent::ClearRefreshReasons()
+{
+	EnumAddFlags(TextRefreshReason, EAvaTextRefreshReason::NoRefresh);
+}
+
 void UAvaText3DComponent::RegisterCallbacks()
 {
 	if (!OnTextGeometryGeneratedDelegate.IsValid())
 	{
 		OnTextGeometryGeneratedDelegate = OnTextGenerated().AddUObject(this, &UAvaText3DComponent::OnTextGeometryGenerated);
 	}
-	
+
 	if (ParentRootComponent.IsValid() && !OnTransformUpdatedDelegate.IsValid())
 	{
 		OnTransformUpdatedDelegate = ParentRootComponent->TransformUpdated.AddUObject(this, &UAvaText3DComponent::OnRootComponentTransformUpdated);
@@ -607,285 +423,404 @@
 	}
 }
 
-void UAvaText3DComponent::SetIsUnlit(bool bInIsUnlit)
+void UAvaText3DComponent::SetupParentRootAndCallbacks()
 {
-	if (bIsUnlit != bInIsUnlit)
-	{
-		bIsUnlit = bInIsUnlit;
-		RefreshUnlit();
-	}
+	RefreshParentRootComponent();
+	RegisterCallbacks();
 }
 
-void UAvaText3DComponent::SetEnforceUpperCase(bool bInEnforceUpperCase)
+void UAvaText3DComponent::SetupTextMaterials(const EAvaTextColoringStyle& InColoringStyle)
 {
-	if (bEnforceUpperCase == bInEnforceUpperCase)
+	if (IsMarkedForMaterialRefresh())
 	{
-		return;
-	}
-	
-	bEnforceUpperCase = bInEnforceUpperCase;
-	RefreshFormatting();
-}
+		switch (InColoringStyle)
+		{
+		case EAvaTextColoringStyle::Invalid:
+			break;
 
-void UAvaText3DComponent::SetMotionDesignFont(const FAvaFont& InFont)
-{
-	if (MotionDesignFont == InFont)
-	{
-		return;
-	}
-	MotionDesignFont = InFont;
-	RefreshGeometry();
-}
+		case EAvaTextColoringStyle::Solid:
+			SetupSolidMaterial();
+			break;
 
-void UAvaText3DComponent::SetMainTexture(UTexture2D* InMainTexture)
-{
-	if (MainTexture != InMainTexture)
-	{
-		MainTexture = InMainTexture;
-		RefreshMaterialInstances();
-	}
-}
+		case EAvaTextColoringStyle::Gradient:
+			SetupGradientMaterial();
+			break;
 
-void UAvaText3DComponent::SetTiling(const FVector2D InTiling)
-{
-	if (Tiling != InTiling)
-	{
-		Tiling = InTiling;
-		RefreshMaterialInstances();
-	}
-}
-
-void UAvaText3DComponent::SetCustomMaterial(UMaterialInterface* InCustomMaterial)
-{
-	if (CustomMaterial != InCustomMaterial)
-	{
-		CustomMaterial = InCustomMaterial;
-		RefreshMaterialInstances();
-	}
-}
-
-void UAvaText3DComponent::SetTranslucencyStyle(EAvaTextTranslucency InTranslucencyStyle)
-{
-	if (TranslucencyStyle != InTranslucencyStyle)
-	{
-		TranslucencyStyle = InTranslucencyStyle;
-		RefreshMaterialInstances();
-	}
-}
-
-void UAvaText3DComponent::SetOpacity(float InOpacity)
-{
-	if (Opacity != InOpacity)
-	{
-		Opacity = InOpacity;
-		RefreshOpacity();
-	}
-}
-
-void UAvaText3DComponent::SetColoringStyle(const EAvaTextColoringStyle& InColoringStyle)
-{
-	if (ColoringStyle != InColoringStyle)
-	{
-		ColoringStyle = InColoringStyle;	
-		RefreshMaterialInstances();
-	}
-}
-
-void UAvaText3DComponent::SetColor(const FLinearColor& InColor)
-{
-	if (Color != InColor)
-	{
-		Color = InColor;
-		RefreshColor();
-	}
-}
-
-void UAvaText3DComponent::SetExtrudeColor(const FLinearColor& InColor)
-{
-	if (ExtrudeColor != InColor)
-	{
-		ExtrudeColor = InColor;
-		RefreshExtrudeColor();	
-	}
-}
-
-void UAvaText3DComponent::SetBevelColor(const FLinearColor& InColor)
-{
-	if (BevelColor != InColor)
-	{
-		BevelColor = InColor;
-		RefreshBevelColor();	
-	}
-}
-
-void UAvaText3DComponent::SetGradientSettings(const FAvaLinearGradientSettings& InGradientSettings)
-{
-	if (GradientSettings == InGradientSettings)
-	{
-		return;
-	}
-
-	GradientSettings = InGradientSettings;
-	RefreshGradient();	
-}
-
-void UAvaText3DComponent::SetGradientColors(const FLinearColor& InColorA, const FLinearColor& InColorB)
-{
-	if (GradientSettings.ColorA == InColorA && GradientSettings.ColorB == InColorB)
-	{
-		return;
-	}
-	
-	GradientSettings.ColorA = InColorA;
-	GradientSettings.ColorB = InColorB;
-	RefreshGradient();
-}
-
-void UAvaText3DComponent::SetMaskOrientation(EAvaMaterialMaskOrientation InMaskOrientation)
-{
-	if (MaskOrientation != InMaskOrientation)
-	{
-		MaskOrientation = InMaskOrientation;
-		RefreshMask();
-	}
-}
-
-void UAvaText3DComponent::SetMaskSmoothness(float InMaskSmoothness)
-{
-	if (MaskSmoothness != InMaskSmoothness)
-	{
-		MaskSmoothness = InMaskSmoothness;
-		RefreshMask();	
-	}
-}
-
-void UAvaText3DComponent::SetMaskOffset(float InMaskOffset)
-{
-	if (MaskOffset != InMaskOffset)
-	{
-		MaskOffset = InMaskOffset;
-		RefreshMask();	
-	}
-}
-
-void UAvaText3DComponent::SetMaskRotation(float InMaskRotation)
-{
-	if (MaskRotation != InMaskRotation)
-	{
-		MaskRotation = InMaskRotation;
-		RefreshMask();	
-	}
-}
-
-void UAvaText3DComponent::SetAlignment(FAvaTextAlignment InAlignment)
-{
-	if (Alignment == InAlignment)
-	{
-		return;
-	}
-
-	Alignment = InAlignment;
-	RefreshAlignment();	
-}
-
-void UAvaText3DComponent::RefreshOpacity() const
-{
-	if (TranslucencyStyle == EAvaTextTranslucency::Translucent)
-	{
-		// Set for any/all materials assigned
-		ForEachMID([this](UMaterialInstanceDynamic* InMID)
-		{
-			InMID->SetScalarParameterValue(UAvaTextMaterialHub::FMaterialStatics::Opacity_MatParam, Opacity);
-		});
-
-		// ... and set for the default materials
-		switch (ColoringStyle)
-		{
-		case EAvaTextColoringStyle::Solid:
-			if (TextMaterialInstance)
-			{
-				TextMaterialInstance->SetScalarParameterValue(UAvaTextMaterialHub::FMaterialStatics::Opacity_MatParam, Opacity);
-			}
-		
-			if (ExtrudeTextMaterialInstance)
-			{
-				ExtrudeTextMaterialInstance->SetScalarParameterValue(UAvaTextMaterialHub::FMaterialStatics::Opacity_MatParam, Opacity);
-			}
-
-			if (BevelTextMaterialInstance)
-			{
-				BevelTextMaterialInstance->SetScalarParameterValue(UAvaTextMaterialHub::FMaterialStatics::Opacity_MatParam, Opacity);
-			}
-			break;
-				
-		case EAvaTextColoringStyle::Gradient:
-			if (GradientTextMaterialInstance)
-			{
-				GradientTextMaterialInstance->SetScalarParameterValue(UAvaTextMaterialHub::FMaterialStatics::Opacity_MatParam, Opacity);
-			}
-			break;
-				
 		case EAvaTextColoringStyle::FromTexture:
-			if (TexturedTextMaterialInstance)
-			{
-				TexturedTextMaterialInstance->SetScalarParameterValue(UAvaTextMaterialHub::FMaterialStatics::Opacity_MatParam, Opacity);
-			}
+			SetupTexturedMaterial();
 			break;
-				
+
 		case EAvaTextColoringStyle::CustomMaterial:
+			SetupCustomMaterial();
 			break;
-			
+
+		default:
+			break;
+		}
+	}
+}
+
+void UAvaText3DComponent::SetupTextGeometry()
+{
+	if (IsMarkedForGeometryRefresh())
+	{
+		// we are about to setup multiple properties together, let's freeze rebuild
+		SetFreeze(true);
+
+		SetText(Text);
+
+		if (UFont* UnderlyingFont = MotionDesignFont.GetFont())
+		{
+			SetFont(UnderlyingFont);
+		}
+
+		SetExtrude(Extrude);
+
+		SetBevel(Bevel);
+		SetBevelType(BevelType);
+		SetBevelSegments(BevelSegments);
+
+		//... and finally unfreeze, to trigger pending rebuild
+		SetFreeze(false);
+
+		RemoveRefreshReason(EAvaTextRefreshReason::GeometryChange);
+	}
+}
+
+void UAvaText3DComponent::SetupTextLayout()
+{
+	if (IsMarkedForLayoutRefresh())
+	{
+		TriggerInternalRebuild(EText3DModifyFlags::Layout);
+		RemoveRefreshReason(EAvaTextRefreshReason::LayoutChange);
+	}
+}
+
+void UAvaText3DComponent::SetupSolidMaterial()
+{
+	TexturedTextMaterialInstance = nullptr;
+	GradientTextMaterialInstance = nullptr;
+
+	if (!IsValid(TextMaterialInstance) || TextMaterialInstance->GetOuter() != this || IsMarkedForMaterialRefresh())
+	{
+		const EAvaTextMaterialFeatures MaterialFeatures = GetMaterialFeaturesFromProperties();
+		const FAvaTextMaterialSettings MaterialSettings(EAvaTextColoringStyle::Solid, MaterialFeatures);
+
+		if (UMaterialInstanceDynamic* const SolidMaterial = GetMIDWithSettings(TextMaterialInstance, MaterialSettings))
+		{
+			UMaterialInstanceDynamic* PreviousMaterial = TextMaterialInstance;
+			TextMaterialInstance = SolidMaterial;
+			if (PreviousMaterial)
+			{
+				TextMaterialInstance->CopyInterpParameters(PreviousMaterial);
+			}
+
+			PreviousMaterial = ExtrudeTextMaterialInstance;
+			ExtrudeTextMaterialInstance = SolidMaterial;
+			if (PreviousMaterial)
+			{
+				ExtrudeTextMaterialInstance->CopyInterpParameters(PreviousMaterial);
+			}
+
+			PreviousMaterial = BevelTextMaterialInstance;
+			BevelTextMaterialInstance = SolidMaterial;
+			if (PreviousMaterial)
+			{
+				BevelTextMaterialInstance->CopyInterpParameters(PreviousMaterial);
+			}
+		}
+
+		RemoveRefreshReason(EAvaTextRefreshReason::MaterialInstancesChange);
+	}
+ 
+	{
+		UE::AvaText::Private::PatchMaterial(
+			this
+			, &UText3DComponent::GetFrontMaterial
+			, &UText3DComponent::SetFrontMaterial
+			, TextMaterialInstance);
+
+		UE::AvaText::Private::PatchMaterial(
+			this
+			, &UText3DComponent::GetExtrudeMaterial
+			, &UText3DComponent::SetExtrudeMaterial
+			, ExtrudeTextMaterialInstance);
+
+		UE::AvaText::Private::PatchMaterial(
+			this
+			, &UText3DComponent::GetBevelMaterial
+			, &UText3DComponent::SetBevelMaterial
+			, BevelTextMaterialInstance);
+
+		UE::AvaText::Private::PatchMaterial(
+			this
+			, &UText3DComponent::GetBackMaterial
+			, &UText3DComponent::SetBackMaterial
+			, TextMaterialInstance);
+	}
+
+	if (IsMarkedForMaterialParametersRefresh())
+	{
+		RefreshColor();
+		RefreshExtrudeColor();
+		RefreshBevelColor();
+
+		RefreshTranslucencyValues();
+	}
+}
+
+void UAvaText3DComponent::SetupGradientMaterial()
+{
+	TexturedTextMaterialInstance = nullptr;
+
+	if (!IsValid(GradientTextMaterialInstance) || GradientTextMaterialInstance->GetOuter() != this || IsMarkedForMaterialRefresh())
+	{
+		const EAvaTextMaterialFeatures MaterialFeatures = GetMaterialFeaturesFromProperties();
+		const FAvaTextMaterialSettings MaterialSettings(EAvaTextColoringStyle::Gradient, MaterialFeatures);
+
+		if (UMaterialInstanceDynamic* const GradientMaterial = GetMIDWithSettings(GradientTextMaterialInstance, MaterialSettings))
+		{
+			UMaterialInstanceDynamic* PreviousMaterial = GradientTextMaterialInstance;
+			GradientTextMaterialInstance = GradientMaterial;
+			if (PreviousMaterial)
+			{
+				GradientTextMaterialInstance->CopyInterpParameters(PreviousMaterial);
+			}
+
+			RemoveRefreshReason(EAvaTextRefreshReason::MaterialInstancesChange);
+		}
+	}
+
+	if (!GradientTextMaterialInstance)
+	{
+		return;
+	}
+
+	UE::AvaText::Private::PatchMaterial(
+		this
+		, &UText3DComponent::GetFrontMaterial
+		, &UText3DComponent::SetFrontMaterial
+		, GradientTextMaterialInstance);
+
+	UE::AvaText::Private::PatchMaterial(
+		this
+		, &UText3DComponent::GetExtrudeMaterial
+		, &UText3DComponent::SetExtrudeMaterial
+		, GradientTextMaterialInstance);
+
+	UE::AvaText::Private::PatchMaterial(
+		this
+		, &UText3DComponent::GetBevelMaterial
+		, &UText3DComponent::SetBevelMaterial
+		, GradientTextMaterialInstance);
+
+	UE::AvaText::Private::PatchMaterial(
+		this
+		, &UText3DComponent::GetBackMaterial
+		, &UText3DComponent::SetBackMaterial
+		, GradientTextMaterialInstance);
+
+	if (IsMarkedForMaterialParametersRefresh())
+	{
+		RefreshGradientValues();
+		RefreshTranslucencyValues();
+	}
+}
+
+void UAvaText3DComponent::SetupCustomMaterial()
+{
+	TexturedTextMaterialInstance = nullptr;
+	GradientTextMaterialInstance = nullptr;
+
+	if (CustomMaterial && IsMarkedForMaterialRefresh())
+	{
+		RemoveRefreshReason(EAvaTextRefreshReason::MaterialInstancesChange);
+
+		UE::AvaText::Private::PatchMaterial(
+			this
+			, &UText3DComponent::GetFrontMaterial
+			, &UText3DComponent::SetFrontMaterial
+			, GradientTextMaterialInstance);
+
+		if (UMaterialInstanceDynamic* CustomMaterialAsInstance = Cast<UMaterialInstanceDynamic>(CustomMaterial))
+		{
+			UE::AvaText::Private::PatchMaterial(
+				this
+				, &UText3DComponent::GetFrontMaterial
+				, &UText3DComponent::SetFrontMaterial
+				, CustomMaterialAsInstance);
+
+			UE::AvaText::Private::PatchMaterial(
+				this
+				, &UText3DComponent::GetExtrudeMaterial
+				, &UText3DComponent::SetExtrudeMaterial
+				, CustomMaterialAsInstance);
+
+			UE::AvaText::Private::PatchMaterial(
+				this
+				, &UText3DComponent::GetBevelMaterial
+				, &UText3DComponent::SetBevelMaterial
+				, CustomMaterialAsInstance);
+
+			UE::AvaText::Private::PatchMaterial(
+				this
+				, &UText3DComponent::GetBackMaterial
+				, &UText3DComponent::SetBackMaterial
+				, CustomMaterialAsInstance);
+		}
+		else
+		{
+			SetFrontMaterial(CustomMaterial);
+			SetExtrudeMaterial(CustomMaterial);
+			SetBevelMaterial(CustomMaterial);
+			SetBackMaterial(CustomMaterial);
+		}
+	}
+}
+
+void UAvaText3DComponent::RefreshTranslucencyValues()
+{
+	switch (TranslucencyStyle)
+	{
+		case EAvaTextTranslucency::Translucent:
+			RefreshOpacity();
+			break;
+
+		case EAvaTextTranslucency::GradientMask:
+			RefreshMaskValues();
+			break;
+
 		default: ;
-		}
 	}
 }
 
-void UAvaText3DComponent::RefreshColor() const
+void UAvaText3DComponent::SetupTexturedMaterial()
 {
-	if (UMaterialInstanceDynamic* MID = Cast<UMaterialInstanceDynamic>(GetFrontMaterial()))
+	GradientTextMaterialInstance = nullptr;
+
+	if (!IsValid(TexturedTextMaterialInstance) || TexturedTextMaterialInstance->GetOuter() != this || IsMarkedForMaterialRefresh())
 	{
-		MID->SetVectorParameterValue(UAvaTextMaterialHub::FMaterialStatics::SolidColor_MatParam, Color);
+		const EAvaTextMaterialFeatures MaterialFeatures = GetMaterialFeaturesFromProperties();
+		const FAvaTextMaterialSettings MaterialSettings(EAvaTextColoringStyle::FromTexture, MaterialFeatures);
+
+		if (UMaterialInstanceDynamic* const TexturedMaterial = GetMIDWithSettings(TexturedTextMaterialInstance, MaterialSettings))
+		{
+			UMaterialInstanceDynamic* PreviousMaterial = TexturedTextMaterialInstance;
+			TexturedTextMaterialInstance = TexturedMaterial;
+			if (PreviousMaterial)
+			{
+				TexturedTextMaterialInstance->CopyInterpParameters(PreviousMaterial);
+			}
+		}
+
+		RemoveRefreshReason(EAvaTextRefreshReason::MaterialInstancesChange);
 	}
 
-	if (TextMaterialInstance)
+	if (!TexturedTextMaterialInstance)
 	{
-		TextMaterialInstance->SetVectorParameterValue(UAvaTextMaterialHub::FMaterialStatics::SolidColor_MatParam, Color);
+		return;
+	}
+
+	{
+		UE::AvaText::Private::PatchMaterial(
+			this
+			, &UText3DComponent::GetFrontMaterial
+			, &UText3DComponent::SetFrontMaterial
+			, TexturedTextMaterialInstance);
+
+		UE::AvaText::Private::PatchMaterial(
+			this
+			, &UText3DComponent::GetExtrudeMaterial
+			, &UText3DComponent::SetExtrudeMaterial
+			, TexturedTextMaterialInstance);
+
+		UE::AvaText::Private::PatchMaterial(
+			this
+			, &UText3DComponent::GetBevelMaterial
+			, &UText3DComponent::SetBevelMaterial
+			, TexturedTextMaterialInstance);
+
+		UE::AvaText::Private::PatchMaterial(
+			this
+			, &UText3DComponent::GetBackMaterial
+			, &UText3DComponent::SetBackMaterial
+			, TexturedTextMaterialInstance);
+	}
+
+	UTexture* CurrentTexture;
+	const FMemoryImageMaterialParameterInfo MaterialInfo(UAvaTextMaterialHub::FMaterialStatics::MainTexture_MatParam);
+	TexturedTextMaterialInstance->GetTextureParameterValue(MaterialInfo, CurrentTexture);
+
+	if (MainTexture != CurrentTexture)
+	{
+		TexturedTextMaterialInstance->SetTextureParameterValue(UAvaTextMaterialHub::FMaterialStatics::MainTexture_MatParam, MainTexture);
+	}
+
+	TexturedTextMaterialInstance->SetScalarParameterValue(UAvaTextMaterialHub::FMaterialStatics::TexturedUTiling_MatParam, Tiling.X);
+	TexturedTextMaterialInstance->SetScalarParameterValue(UAvaTextMaterialHub::FMaterialStatics::TexturedVTiling_MatParam, Tiling.Y);
+
+	RefreshTranslucencyValues();
+}
+
+void UAvaText3DComponent::RefreshStoredBoundsValues()
+{
+	if (const AActor* ParentActor = GetOwner())
+	{
+		// local bounds (used for gradient or masked materials)
+		const FBox LocalBounds = FAvaActorUtils::GetActorLocalBoundingBox(ParentActor);
+		LocalBoundsLowerLeftCorner = LocalBounds.Min;
+		LocalBoundsExtent = LocalBounds.GetExtent();
 	}
 }
 
-void UAvaText3DComponent::RefreshExtrudeColor() const
+void UAvaText3DComponent::PostLoad()
 {
-	if (UMaterialInstanceDynamic* MID = Cast<UMaterialInstanceDynamic>(GetExtrudeMaterial()))
-	{
-		MID->SetVectorParameterValue(UAvaTextMaterialHub::FMaterialStatics::SolidColor_MatParam, ExtrudeColor);
-	}
+	Super::PostLoad();
 
-	if (ExtrudeTextMaterialInstance)
-	{
-		ExtrudeTextMaterialInstance->SetVectorParameterValue(UAvaTextMaterialHub::FMaterialStatics::SolidColor_MatParam, Color);
-	}
+	SetupParentRootAndCallbacks();
+	RetrieveDataFromBaseComponent();
 }
 
-void UAvaText3DComponent::RefreshBevelColor() const
+void UAvaText3DComponent::PostDuplicate(EDuplicateMode::Type DuplicateMode)
 {
-	if (UMaterialInstanceDynamic* MID = Cast<UMaterialInstanceDynamic>(GetBevelMaterial()))
-	{
-		MID->SetVectorParameterValue(UAvaTextMaterialHub::FMaterialStatics::SolidColor_MatParam, BevelColor);
-	}
+	Super::PostDuplicate(DuplicateMode);
 
-	if (BevelTextMaterialInstance)
-	{
-		BevelTextMaterialInstance->SetVectorParameterValue(UAvaTextMaterialHub::FMaterialStatics::SolidColor_MatParam, Color);
-	}
+	PostDuplicateActions();
 }
 
+void UAvaText3DComponent::PostInitProperties()
+{
+	Super::PostInitProperties();
+
+	SetupParentRootAndCallbacks();
+	MarkForFullRefresh();
+}
+
+// overriding PostEditImport to handle Owner duplication
+void UAvaText3DComponent::PostEditImport()
+{
+	Super::PostEditImport();
+
+	PostDuplicateActions();
+}
+
+void UAvaText3DComponent::PostDuplicateActions()
+{
+	SetupParentRootAndCallbacks();
+
+	RetrieveCurrentTextAndFont();
+
+	MarkForFullRefresh();
+	RefreshText3D();
+}
+
 void UAvaText3DComponent::ScheduleTextRefreshOnTicker()
 {
 	if (!TickerHandle.IsValid())
 	{
 		TWeakObjectPtr<UAvaText3DComponent> WeakThis(this);
 		FTSTicker& Ticker = FTSTicker::GetCoreTicker();
-		
+
 		TickerHandle = Ticker.AddTicker(FTickerDelegate::CreateLambda(
 			[WeakThis](float InDeltaTime)->bool
 			{
@@ -907,4 +842,339 @@
 	}
 }
 
+void UAvaText3DComponent::OnTextGeneratedActions()
+{
+	RefreshStoredBoundsValues();
+
+	// bounding box has been updated, so we need to also update gradient and mask values
+	RefreshGradient();
+	RefreshMask();
+}
+
+void UAvaText3DComponent::OnTextGeometryGenerated()
+{
+	OnTextGeneratedActions();
+
 #if WITH_EDITOR
+	if (GEditor)
+	{
+		GEditor->BroadcastOnComponentTransformChanged(this, ETeleportType::None);
+	}
+#endif
+
+	TransformUpdated.Broadcast(this, EUpdateTransformFlags::None, ETeleportType::None);
+}
+
+void UAvaText3DComponent::OnRootComponentTransformUpdated(USceneComponent* InSceneComponent, EUpdateTransformFlags InUpdateTransformFlags, ETeleportType InTeleportType)
+{
+	ScheduleTextRefreshOnTicker();
+}
+
+void UAvaText3DComponent::RetrieveDataFromBaseComponent()
+{
+	RetrieveCurrentColoringStyle();
+	RetrieveCurrentMaterials();
+	RetrieveCurrentTextAndFont();
+	RetrieveCurrentAlignment();
+}
+
+void UAvaText3DComponent::RetrieveCurrentTextAndFont()
+{
+	// if default had to be loaded (e.g. font asset missing), then we regenerate geometry
+	if (MotionDesignFont.IsDefaultFont())
+	{
+		if (IsValid(Font) && Font != MotionDesignFont.GetFont())
+		{
+			MotionDesignFont.InitFromFont(Font);
+		}
+
+		MarkForGeometryRefresh();
+	}
+
+	if (!IsValid(Font))
+	{
+		MarkForGeometryRefresh();
+	}
+}
+
+void UAvaText3DComponent::RetrieveCurrentAlignment()
+{
+	Alignment.HorizontalAlignment = HorizontalAlignment;
+	Alignment.VerticalAlignment = VerticalAlignment;
+}
+
+void UAvaText3DComponent::RetrieveCurrentMaterials()
+{
+	switch (ColoringStyle)
+	{
+		case EAvaTextColoringStyle::Invalid:
+			break;
+
+		case EAvaTextColoringStyle::Solid:
+			RetrieveSolidMaterialInstances();
+			break;
+
+		case EAvaTextColoringStyle::Gradient:
+			RetrieveGradientMaterialInstance();
+			break;
+
+		case EAvaTextColoringStyle::FromTexture: 
+			RetrieveTexturedMaterialInstance();
+			break;
+
+		case EAvaTextColoringStyle::CustomMaterial: 
+			RetrieveCustomMaterialReference();
+			break;
+	}
+}
+
+void UAvaText3DComponent::RetrieveCurrentColoringStyle()
+{
+	if (!FrontMaterial)
+	{
+		return;
+	}
+
+	if (UMaterialInterface* CurrentFrontMaterial = FrontMaterial->GetMaterial())
+	{
+		if (const FAvaTextMaterialSettings* MaterialSettings = UAvaTextMaterialHub::GetSettingsFromMaterial(CurrentFrontMaterial))
+		{
+			ColoringStyle = MaterialSettings->ColoringStyle;
+
+			if (EnumHasAnyFlags(MaterialSettings->MaterialFeatures, EAvaTextMaterialFeatures::GradientMask))
+			{
+				TranslucencyStyle = EAvaTextTranslucency::GradientMask;
+			}
+
+			if (EnumHasAnyFlags(MaterialSettings->MaterialFeatures, EAvaTextMaterialFeatures::Translucent))
+			{
+				TranslucencyStyle = EAvaTextTranslucency::Translucent;
+			}
+
+			bIsUnlit = EnumHasAnyFlags(MaterialSettings->MaterialFeatures, EAvaTextMaterialFeatures::Unlit);
+		}
+		else // if there is no match, best guess is custom material
+		{
+			ColoringStyle = EAvaTextColoringStyle::CustomMaterial;
+			CustomMaterial = CurrentFrontMaterial;
+			TranslucencyStyle = EAvaTextTranslucency::None;
+		}
+	}
+}
+
+void UAvaText3DComponent::RetrieveSolidMaterialInstances()
+{
+	if (FrontMaterial)
+	{
+		TextMaterialInstance = Cast<UMaterialInstanceDynamic>(FrontMaterial);
+		if (TextMaterialInstance)
+		{
+			TextMaterialInstance->GetVectorParameterValue(UAvaTextMaterialHub::FMaterialStatics::SolidColor_MatParam, Color);
+
+			RetrieveTranslucentMaterialParameters();
+		}
+	}
+
+	if (ExtrudeMaterial)
+	{
+		ExtrudeTextMaterialInstance = Cast<UMaterialInstanceDynamic>(ExtrudeMaterial);
+		if (ExtrudeTextMaterialInstance)
+		{
+			ExtrudeTextMaterialInstance->GetVectorParameterValue(UAvaTextMaterialHub::FMaterialStatics::SolidColor_MatParam, ExtrudeColor);
+		}
+	}
+
+	if (BevelMaterial)
+	{
+		BevelTextMaterialInstance = Cast<UMaterialInstanceDynamic>(BevelMaterial);
+		if (BevelTextMaterialInstance)
+		{
+			BevelTextMaterialInstance->GetVectorParameterValue(UAvaTextMaterialHub::FMaterialStatics::SolidColor_MatParam, BevelColor);
+		}
+	}
+}
+
+void UAvaText3DComponent::RetrieveGradientMaterialInstance()
+{
+	if (!FrontMaterial)
+	{
+		return;
+	}
+
+	GradientTextMaterialInstance = Cast<UMaterialInstanceDynamic>(FrontMaterial);
+
+	if (GradientTextMaterialInstance)
+	{
+		GradientTextMaterialInstance->GetVectorParameterValue(UAvaTextMaterialHub::FMaterialStatics::GradientColorA_MatParam, GradientSettings.ColorA);
+		GradientTextMaterialInstance->GetVectorParameterValue(UAvaTextMaterialHub::FMaterialStatics::GradientColorB_MatParam, GradientSettings.ColorB);
+
+		GradientTextMaterialInstance->GetScalarParameterValue(UAvaTextMaterialHub::FMaterialStatics::GradientOffset_MatParam, GradientSettings.Offset);
+		GradientTextMaterialInstance->GetScalarParameterValue(UAvaTextMaterialHub::FMaterialStatics::GradientSmoothness_MatParam, GradientSettings.Smoothness);
+
+		const bool bRotationRetrieved = GradientTextMaterialInstance->GetScalarParameterValue(UAvaTextMaterialHub::FMaterialStatics::GradientRotation_MatParam, GradientSettings.Rotation);
+
+		if (bRotationRetrieved)
+		{
+			if (GradientSettings.Rotation == UAvaTextMaterialHub::FMaterialStatics::MaterialMaskTopBottomRotationValue)
+			{
+				GradientSettings.Direction = EAvaGradientDirection::Vertical;
+			}
+			else if (GradientSettings.Rotation == UAvaTextMaterialHub::FMaterialStatics::MaterialMaskLeftToRightRotationValue)
+			{
+				GradientSettings.Direction = EAvaGradientDirection::Horizontal;
+			}
+			else
+			{
+				GradientSettings.Direction = EAvaGradientDirection::Custom;
+			}
+		}
+
+		RetrieveTranslucentMaterialParameters();
+	}
+}
+
+void UAvaText3DComponent::RetrieveTexturedMaterialInstance()
+{
+	if (!FrontMaterial)
+	{
+		return;
+	}
+
+	TexturedTextMaterialInstance = Cast<UMaterialInstanceDynamic>(FrontMaterial);
+
+	if (TexturedTextMaterialInstance)
+	{
+		float Tiling_X;
+		float Tiling_Y;
+		TexturedTextMaterialInstance->GetScalarParameterValue(UAvaTextMaterialHub::FMaterialStatics::TexturedUTiling_MatParam, Tiling_X);
+		TexturedTextMaterialInstance->GetScalarParameterValue(UAvaTextMaterialHub::FMaterialStatics::TexturedVTiling_MatParam, Tiling_Y);
+
+		Tiling = FVector2D(Tiling_X, Tiling_Y);
+
+		UTexture* CurrentTexture;
+		TexturedTextMaterialInstance->GetTextureParameterValue(UAvaTextMaterialHub::FMaterialStatics::MainTexture_MatParam, CurrentTexture);
+
+		if (CurrentTexture)
+		{
+			MainTexture = Cast<UTexture2D>(CurrentTexture);
+		}
+
+		RetrieveTranslucentMaterialParameters();
+	}
+}
+
+void UAvaText3DComponent::RetrieveCustomMaterialReference()
+{
+	if (!FrontMaterial)
+	{
+		return;
+	}
+
+	CustomMaterial = FrontMaterial;
+}
+
+void UAvaText3DComponent::RetrieveTranslucentMaterialParameters()
+{
+	switch (TranslucencyStyle)
+	{
+	case EAvaTextTranslucency::Translucent:
+		RetrieveSimpleTranslucentMaterialParameters(TextMaterialInstance);
+		break;
+
+	case EAvaTextTranslucency::GradientMask:
+		RetrieveMaskedMaterialParameters(TextMaterialInstance);
+		break;
+
+	default: ;
+	}
+}
+
+void UAvaText3DComponent::RetrieveMaskedMaterialParameters(const UMaterialInstanceDynamic* InMaterialInstance)
+{
+	if (!InMaterialInstance)
+	{
+		return;
+	}
+
+	InMaterialInstance->GetScalarParameterValue(UAvaTextMaterialHub::FMaterialStatics::MaskSmoothness_MatParam, MaskSmoothness);
+	InMaterialInstance->GetScalarParameterValue(UAvaTextMaterialHub::FMaterialStatics::MaskOffset_MatParam, MaskOffset);
+	InMaterialInstance->GetScalarParameterValue(UAvaTextMaterialHub::FMaterialStatics::MaskRotation_MatParam, MaskRotation);
+
+	// in case the material was saved with the deprecated value, let's fix the rotation value
+	if (MaskRotation == UAvaTextMaterialHub::FMaterialStatics::MaterialMaskRightToLeftRotationValue_DEPRECATED)
+	{
+		MaskRotation = UAvaTextMaterialHub::FMaterialStatics::MaterialMaskRightToLeftRotationValue;
+	}
+
+	if (MaskRotation == UAvaTextMaterialHub::FMaterialStatics::MaterialMaskLeftToRightRotationValue)
+	{
+		MaskOrientation = EAvaMaterialMaskOrientation::LeftRight;
+	}
+	else if (MaskRotation == UAvaTextMaterialHub::FMaterialStatics::MaterialMaskRightToLeftRotationValue)
+	{
+		MaskOrientation = EAvaMaterialMaskOrientation::RightLeft;
+	}
+	else
+	{
+		MaskOrientation = EAvaMaterialMaskOrientation::Custom;
+	}
+}
+
+void UAvaText3DComponent::RetrieveSimpleTranslucentMaterialParameters(const UMaterialInstanceDynamic* InMaterialInstance)
+{
+	if (!InMaterialInstance)
+	{
+		return;
+	}
+
+	InMaterialInstance->GetScalarParameterValue(UAvaTextMaterialHub::FMaterialStatics::Opacity_MatParam, Opacity);
+}
+
+UMaterialInstanceDynamic* UAvaText3DComponent::GetMIDWithSettings(
+	const UMaterialInterface* InPreviousMaterial
+	, const FAvaTextMaterialSettings& InSettings)
+{
+	UMaterialInterface* Result = nullptr;
+	if (OnGetMaterialWithSettingsDelegate.IsBound())
+	{
+		Result = OnGetMaterialWithSettingsDelegate.Execute(InPreviousMaterial, InSettings);
+	}
+
+	if (!Result)
+	{
+		Result = UAvaTextMaterialHub::GetMaterial(InSettings);
+	}
+
+	if (!::IsValid(Result))
+	{
+		UE_LOG(LogAva, Warning, TEXT("Material could not be found or was invalid for the provided settings."));
+		return nullptr;
+	}
+
+	if (UMaterialInstanceDynamic* ResultAsMID = Cast<UMaterialInstanceDynamic>(Result))
+	{
+		return ResultAsMID;
+	}
+
+	return UMaterialInstanceDynamic::Create(Result, this);
+}
+
+EAvaTextMaterialFeatures UAvaText3DComponent::GetMaterialFeaturesFromProperties() const
+{
+	EAvaTextMaterialFeatures MaterialFeatures = EAvaTextMaterialFeatures::None;
+
+	if (TranslucencyStyle == EAvaTextTranslucency::GradientMask)
+	{
+		MaterialFeatures |= EAvaTextMaterialFeatures::GradientMask;
+	}
+	else if (TranslucencyStyle == EAvaTextTranslucency::Translucent)
+	{
+		MaterialFeatures |= EAvaTextMaterialFeatures::Translucent;
+	}
+
+	MaterialFeatures |= bIsUnlit ? EAvaTextMaterialFeatures::Unlit : EAvaTextMaterialFeatures::None;
+
+	return MaterialFeatures;
+}
+
+#if WITH_EDITOR
@@ -931,7 +1201,7 @@
 		PropertyChangeFunctions.Add(GET_MEMBER_NAME_CHECKED(UAvaText3DComponent, MainTexture), &UAvaText3DComponent::RefreshMaterialInstances);
 		PropertyChangeFunctions.Add(GET_MEMBER_NAME_CHECKED(UAvaText3DComponent, Tiling), &UAvaText3DComponent::RefreshMaterialInstances);
 		PropertyChangeFunctions.Add(GET_MEMBER_NAME_CHECKED(UAvaText3DComponent, TranslucencyStyle), &UAvaText3DComponent::RefreshMaterialInstances);
-		
+
 		PropertyChangeFunctions.Add(GET_MEMBER_NAME_CHECKED(UAvaText3DComponent, MaskOrientation), &UAvaText3DComponent::RefreshMask);
 		PropertyChangeFunctions.Add(GET_MEMBER_NAME_CHECKED(UAvaText3DComponent, MaskSmoothness), &UAvaText3DComponent::RefreshMask);
 		PropertyChangeFunctions.Add(GET_MEMBER_NAME_CHECKED(UAvaText3DComponent, MaskOffset), &UAvaText3DComponent::RefreshMask);
@@ -952,6 +1222,97 @@
 }
 #endif
 
+void UAvaText3DComponent::RefreshOpacity() const
+{
+	if (TranslucencyStyle == EAvaTextTranslucency::Translucent)
+	{
+		// Set for any/all materials assigned
+		ForEachMID([this](UMaterialInstanceDynamic* InMID)
+		{
+			InMID->SetScalarParameterValue(UAvaTextMaterialHub::FMaterialStatics::Opacity_MatParam, Opacity);
+		});
+
+		// ... and set for the default materials
+		switch (ColoringStyle)
+		{
+		case EAvaTextColoringStyle::Solid:
+			if (TextMaterialInstance)
+			{
+				TextMaterialInstance->SetScalarParameterValue(UAvaTextMaterialHub::FMaterialStatics::Opacity_MatParam, Opacity);
+			}
+
+			if (ExtrudeTextMaterialInstance)
+			{
+				ExtrudeTextMaterialInstance->SetScalarParameterValue(UAvaTextMaterialHub::FMaterialStatics::Opacity_MatParam, Opacity);
+			}
+
+			if (BevelTextMaterialInstance)
+			{
+				BevelTextMaterialInstance->SetScalarParameterValue(UAvaTextMaterialHub::FMaterialStatics::Opacity_MatParam, Opacity);
+			}
+			break;
+
+		case EAvaTextColoringStyle::Gradient:
+			if (GradientTextMaterialInstance)
+			{
+				GradientTextMaterialInstance->SetScalarParameterValue(UAvaTextMaterialHub::FMaterialStatics::Opacity_MatParam, Opacity);
+			}
+			break;
+
+		case EAvaTextColoringStyle::FromTexture:
+			if (TexturedTextMaterialInstance)
+			{
+				TexturedTextMaterialInstance->SetScalarParameterValue(UAvaTextMaterialHub::FMaterialStatics::Opacity_MatParam, Opacity);
+			}
+			break;
+
+		case EAvaTextColoringStyle::CustomMaterial:
+			break;
+
+		default: ;
+		}
+	}
+}
+
+void UAvaText3DComponent::RefreshColor() const
+{
+	if (UMaterialInstanceDynamic* MID = Cast<UMaterialInstanceDynamic>(GetFrontMaterial()))
+	{
+		MID->SetVectorParameterValue(UAvaTextMaterialHub::FMaterialStatics::SolidColor_MatParam, Color);
+	}
+
+	if (TextMaterialInstance)
+	{
+		TextMaterialInstance->SetVectorParameterValue(UAvaTextMaterialHub::FMaterialStatics::SolidColor_MatParam, Color);
+	}
+}
+
+void UAvaText3DComponent::RefreshExtrudeColor() const
+{
+	if (UMaterialInstanceDynamic* MID = Cast<UMaterialInstanceDynamic>(GetExtrudeMaterial()))
+	{
+		MID->SetVectorParameterValue(UAvaTextMaterialHub::FMaterialStatics::SolidColor_MatParam, ExtrudeColor);
+	}
+
+	if (ExtrudeTextMaterialInstance)
+	{
+		ExtrudeTextMaterialInstance->SetVectorParameterValue(UAvaTextMaterialHub::FMaterialStatics::SolidColor_MatParam, Color);
+	}
+}
+
+void UAvaText3DComponent::RefreshBevelColor() const
+{
+	if (UMaterialInstanceDynamic* MID = Cast<UMaterialInstanceDynamic>(GetBevelMaterial()))
+	{
+		MID->SetVectorParameterValue(UAvaTextMaterialHub::FMaterialStatics::SolidColor_MatParam, BevelColor);
+	}
+
+	if (BevelTextMaterialInstance)
+	{
+		BevelTextMaterialInstance->SetVectorParameterValue(UAvaTextMaterialHub::FMaterialStatics::SolidColor_MatParam, Color);
+	}
+}
+
 void UAvaText3DComponent::RefreshAlignment()
 {
 	HorizontalAlignment = Alignment.HorizontalAlignment;
@@ -985,7 +1346,7 @@
 
 void UAvaText3DComponent::RefreshUnlit()
 {
-	MarkForMaterialsFullRefresh();	
+	MarkForMaterialsFullRefresh();
 	RefreshText3D();
 }
 
@@ -995,12 +1356,6 @@
 	RefreshText3D();
 }
 
-void UAvaText3DComponent::RefreshMaterialInstances()
-{
-	MarkForMaterialsFullRefresh();
-	RefreshText3D();
-}
-
 void UAvaText3DComponent::RefreshGeometry()
 {
 	MarkForGeometryRefresh();
@@ -1025,15 +1380,15 @@
 			GradientTextMaterialInstance->SetVectorParameterValue(UAvaTextMaterialHub::FMaterialStatics::GradientColorB_MatParam, GradientSettings.ColorB);
 
 			float GradientRotation = UAvaTextMaterialHub::FMaterialStatics::MaterialMaskTopBottomRotationValue;
-			
+
 			switch (GradientSettings.Direction)
 			{
 			case EAvaGradientDirection::None:
 				break;
 
-			case EAvaGradientDirection::Vertical:				
+			case EAvaGradientDirection::Vertical:
 					GradientRotation = UAvaTextMaterialHub::FMaterialStatics::MaterialMaskTopBottomRotationValue;
-					break;				
+					break;
 
 			case EAvaGradientDirection::Horizontal:
 					GradientRotation = UAvaTextMaterialHub::FMaterialStatics::MaterialMaskLeftToRightRotationValue;
@@ -1076,25 +1431,25 @@
 		{
 		case EAvaTextColoringStyle::Invalid:
 			break;
-	
+
 		case EAvaTextColoringStyle::Solid:
 			SetMaskedMaterialValues(TextMaterialInstance);
 			SetMaskedMaterialValues(ExtrudeTextMaterialInstance);
 			SetMaskedMaterialValues(BevelTextMaterialInstance);
 			break;
-		
+
 		case EAvaTextColoringStyle::Gradient:
 			SetMaskedMaterialValues(GradientTextMaterialInstance);
 			break;
-		
+
 		case EAvaTextColoringStyle::FromTexture:
 			SetMaskedMaterialValues(TexturedTextMaterialInstance);
 			break;
-		
+
 		case EAvaTextColoringStyle::CustomMaterial:
 			break;
 		}
-		
+
 		RemoveRefreshReason(EAvaTextRefreshReason::MaterialParametersChange);
 	}
 }
@@ -1105,55 +1460,30 @@
 	RefreshMaskValues();
 }
 
-void UAvaText3DComponent::SetMaterialGeometryParameterValues(UMaterialInstanceDynamic* InMaterial, const FVector& TextScaleFactor) const
-{
-	if (!InMaterial)
-	{
-		return;
-	}
-	
-	// we use the bounds minimum point, because the material needs the bbox origin in the the lower left corner
-	const FVector BoundsLowerLeftCornerLocalScaled = LocalBoundsLowerLeftCorner/TextScaleFactor;
-	const FVector BoundsSizeScaled = LocalBoundsExtent/TextScaleFactor;
-
-	InMaterial->SetVectorParameterValue(UAvaTextMaterialHub::FMaterialStatics::BoundsOrigin_MatParam, BoundsLowerLeftCornerLocalScaled);
-	InMaterial->SetVectorParameterValue(UAvaTextMaterialHub::FMaterialStatics::BoundsSize_MatParam, BoundsSizeScaled);
-	
-	FVector TextPositionParameter = GetComponentLocation();
-	
-	// in case this component is not the root, let's take into account also relative location
-	if (ParentRootComponent != this)
-	{
-		TextPositionParameter -= GetRelativeLocation();
-	}
-	
-	InMaterial->SetVectorParameterValue(UAvaTextMaterialHub::FMaterialStatics::TextPosition_MatParam, TextPositionParameter);
-}
-
 void UAvaText3DComponent::RefreshMaterialGeometryParameterValues()
 {
 	// scaling factor required to "normalize" origin and local bounds to properly update the material
 	const FVector TextScaleFactor = GetTextScale() * GetComponentScale();
-	
+
 	switch (ColoringStyle)
 	{
 	case EAvaTextColoringStyle::Invalid:
 		break;
-	
+
 	case EAvaTextColoringStyle::Solid:
 		SetMaterialGeometryParameterValues(TextMaterialInstance, TextScaleFactor);
 		SetMaterialGeometryParameterValues(ExtrudeTextMaterialInstance, TextScaleFactor);
 		SetMaterialGeometryParameterValues(BevelTextMaterialInstance, TextScaleFactor);
 		break;
-	
+
 	case EAvaTextColoringStyle::Gradient:
 		SetMaterialGeometryParameterValues(GradientTextMaterialInstance, TextScaleFactor);
 		break;
-	
+
 	case EAvaTextColoringStyle::FromTexture:
 		SetMaterialGeometryParameterValues(TexturedTextMaterialInstance, TextScaleFactor);
 		break;
-	
+
 	case EAvaTextColoringStyle::CustomMaterial:
 		break;
 	}
@@ -1164,6 +1494,31 @@
 	TriggerInternalRebuild(EText3DModifyFlags::Geometry);
 }
 
+void UAvaText3DComponent::SetMaterialGeometryParameterValues(UMaterialInstanceDynamic* InMaterial, const FVector& TextScaleFactor) const
+{
+	if (!InMaterial)
+	{
+		return;
+	}
+
+	// we use the bounds minimum point, because the material needs the bbox origin in the the lower left corner
+	const FVector BoundsLowerLeftCornerLocalScaled = LocalBoundsLowerLeftCorner/TextScaleFactor;
+	const FVector BoundsSizeScaled = LocalBoundsExtent/TextScaleFactor;
+
+	InMaterial->SetVectorParameterValue(UAvaTextMaterialHub::FMaterialStatics::BoundsOrigin_MatParam, BoundsLowerLeftCornerLocalScaled);
+	InMaterial->SetVectorParameterValue(UAvaTextMaterialHub::FMaterialStatics::BoundsSize_MatParam, BoundsSizeScaled);
+
+	FVector TextPositionParameter = GetComponentLocation();
+
+	// in case this component is not the root, let's take into account also relative location
+	if (ParentRootComponent != this)
+	{
+		TextPositionParameter -= GetRelativeLocation();
+	}
+
+	InMaterial->SetVectorParameterValue(UAvaTextMaterialHub::FMaterialStatics::TextPosition_MatParam, TextPositionParameter);
+}
+
 void UAvaText3DComponent::SetMaskedMaterialValues(UMaterialInstanceDynamic* InMaterial)
 {
 	if (InMaterial)
@@ -1171,7 +1526,7 @@
 		InMaterial->SetScalarParameterValue(UAvaTextMaterialHub::FMaterialStatics::MaskOffset_MatParam, MaskOffset);
 
 		float GradientRotation = UAvaTextMaterialHub::FMaterialStatics::MaterialMaskLeftToRightRotationValue;
-		
+
 		switch (MaskOrientation)
 		{
 		case EAvaMaterialMaskOrientation::LeftRight:
@@ -1194,412 +1549,55 @@
 	}
 }
 
-UAvaText3DComponent::FOnGetMaterialWithSettings& UAvaText3DComponent::GetMaterialProviderDelegate()
-{
-	return OnGetMaterialWithSettingsDelegate;
-}
-
-void UAvaText3DComponent::ForEachMID(TUniqueFunction<void(UMaterialInstanceDynamic* InMID)>&& InFunc) const
-{
-	auto ExecuteIfMID = [&InFunc](UMaterialInterface* InMaterial)
-	{
-		if (UMaterialInstanceDynamic* MID = Cast<UMaterialInstanceDynamic>(InMaterial))
-		{
-			InFunc(MID);
-		}
-	};
-	
-	ExecuteIfMID(GetFrontMaterial());
-	ExecuteIfMID(GetBevelMaterial());
-	ExecuteIfMID(GetExtrudeMaterial());
-	ExecuteIfMID(GetBackMaterial());
-}
-
-void UAvaText3DComponent::RefreshStoredBoundsValues()
-{
-	if (const AActor* ParentActor = GetOwner())
-	{
-		// local bounds (used for gradient or masked materials)
-		const FBox LocalBounds = FAvaActorUtils::GetActorLocalBoundingBox(ParentActor);
-		LocalBoundsLowerLeftCorner = LocalBounds.Min;
-		LocalBoundsExtent = LocalBounds.GetExtent();
-	}
-}
-
-UMaterialInstanceDynamic* UAvaText3DComponent::GetMIDWithSettings(
-	const UMaterialInterface* InPreviousMaterial
-	, const FAvaTextMaterialSettings& InSettings)
-{
-	UMaterialInterface* Result = nullptr;
-	if (OnGetMaterialWithSettingsDelegate.IsBound())
-	{
-		Result = OnGetMaterialWithSettingsDelegate.Execute(InPreviousMaterial, InSettings);
-	}
-
-	if (!Result)
-	{
-		Result = UAvaTextMaterialHub::GetMaterial(InSettings);
-	}
-
-	if (!::IsValid(Result))
-	{
-		UE_LOG(LogAva, Warning, TEXT("Material could not be found or was invalid for the provided settings."));
-		return nullptr;
-	}
-
-	if (UMaterialInstanceDynamic* ResultAsMID = Cast<UMaterialInstanceDynamic>(Result))
-	{
-		return ResultAsMID;
-	}
-	
-	return UMaterialInstanceDynamic::Create(Result, this);
-}
-
-void UAvaText3DComponent::OnTextGeneratedActions()
-{
-	RefreshStoredBoundsValues();
-
-	// bounding box has been updated, so we need to also update gradient and mask values
-	RefreshGradient();
-	RefreshMask();
-}
-
-void UAvaText3DComponent::OnTextGeometryGenerated()
-{
-	OnTextGeneratedActions();
-
 #if WITH_EDITOR
-	if (GEditor)
+void UAvaText3DComponent::PostEditChangeProperty(FPropertyChangedEvent& PropertyChangedEvent)
+{
+	const FName ChangedPropertyName = PropertyChangedEvent.GetMemberPropertyName();
+
+	if (PropertyChangeFunctions.Contains(ChangedPropertyName))
 	{
-		GEditor->BroadcastOnComponentTransformChanged(this, ETeleportType::None);
+		PropertyChangeFunctions[ChangedPropertyName](this);
 	}
+
+	// Call the Super::PostEditChangeProperty after the PropertyChangeFunctions execute the right Setter/Refresh
+	Super::PostEditChangeProperty(PropertyChangedEvent);
+}
+
+void UAvaText3DComponent::PostEditComponentMove(bool bFinished)
+{
+	Super::PostEditComponentMove(bFinished);
+	ScheduleTextRefreshOnTicker();
+}
 #endif
 
-	TransformUpdated.Broadcast(this, EUpdateTransformFlags::None, ETeleportType::None);
+void UAvaText3DComponent::OnUpdateTransform(EUpdateTransformFlags UpdateTransformFlags, ETeleportType Teleport)
+{
+	Super::OnUpdateTransform(UpdateTransformFlags, Teleport);
+
+	ScheduleTextRefreshOnTicker();
 }
 
-void UAvaText3DComponent::MarkForGeometryRefresh()
+void UAvaText3DComponent::OnAttachmentChanged()
 {
-	EnumAddFlags(TextRefreshReason, EAvaTextRefreshReason::GeometryChange);
+	Super::OnAttachmentChanged();
+
+	SetupParentRootAndCallbacks();
 }
 
-bool UAvaText3DComponent::IsMarkedForGeometryRefresh() const
+void UAvaText3DComponent::OnRegister()
 {
-	return EnumHasAnyFlags(TextRefreshReason, EAvaTextRefreshReason::GeometryChange);
+	Super::OnRegister();
+
+	SetComponentTickEnabled(false);
+
+	SetupParentRootAndCallbacks();
+	RefreshText3D();
 }
 
-void UAvaText3DComponent::RemoveRefreshReason(EAvaTextRefreshReason InRefreshReasonToRemove)
+void UAvaText3DComponent::FormatText(FText& InOutText) const
 {
-	EnumRemoveFlags(TextRefreshReason, InRefreshReasonToRemove);
-}
-
-void UAvaText3DComponent::MarkForLayoutRefresh()
-{
-	EnumAddFlags(TextRefreshReason, EAvaTextRefreshReason::LayoutChange);
-}
-
-bool UAvaText3DComponent::IsMarkedForLayoutRefresh() const
-{
-	return EnumHasAnyFlags(TextRefreshReason, EAvaTextRefreshReason::LayoutChange);
-}
-
-void UAvaText3DComponent::MarkForMaterialsFullRefresh()
-{
-	EnumAddFlags(TextRefreshReason, EAvaTextRefreshReason::MaterialInstancesChange);
-	EnumAddFlags(TextRefreshReason, EAvaTextRefreshReason::MaterialParametersChange);
-}
-
-bool UAvaText3DComponent::IsMarkedForMaterialRefresh() const
-{
-	return EnumHasAnyFlags(TextRefreshReason, EAvaTextRefreshReason::MaterialInstancesChange);
-}
-
-void UAvaText3DComponent::MarkForMaterialParametersRefresh()
-{
-	EnumAddFlags(TextRefreshReason, EAvaTextRefreshReason::MaterialParametersChange);
-}
-
-bool UAvaText3DComponent::IsMarkedForMaterialParametersRefresh() const
-{
-	return EnumHasAnyFlags(TextRefreshReason, EAvaTextRefreshReason::MaterialParametersChange);
-}
-
-void UAvaText3DComponent::MarkForFullRefresh()
-{
-	EnumAddFlags(TextRefreshReason, EAvaTextRefreshReason::FullRefresh);
-}
-
-void UAvaText3DComponent::ClearRefreshReasons()
-{
-	EnumAddFlags(TextRefreshReason, EAvaTextRefreshReason::NoRefresh);
-}
-
-void UAvaText3DComponent::RetrieveTranslucentMaterialParameters()
-{
-	switch (TranslucencyStyle)
+	if (bEnforceUpperCase)
 	{
-	case EAvaTextTranslucency::Translucent:
-		RetrieveSimpleTranslucentMaterialParameters(TextMaterialInstance);
-		break;
-				
-	case EAvaTextTranslucency::GradientMask:
-		RetrieveMaskedMaterialParameters(TextMaterialInstance);
-		break;
-				
-	default: ;
+		InOutText = InOutText.ToUpper();
 	}
 }
-
-void UAvaText3DComponent::RetrieveSolidMaterialInstances()
-{
-	if (FrontMaterial)
-	{
-		TextMaterialInstance = Cast<UMaterialInstanceDynamic>(FrontMaterial);
-		if (TextMaterialInstance)
-		{
-			TextMaterialInstance->GetVectorParameterValue(UAvaTextMaterialHub::FMaterialStatics::SolidColor_MatParam, Color);
-
-			RetrieveTranslucentMaterialParameters();
-		}
-	}
-
-	if (ExtrudeMaterial)
-	{
-		ExtrudeTextMaterialInstance = Cast<UMaterialInstanceDynamic>(ExtrudeMaterial);
-		if (ExtrudeTextMaterialInstance)
-		{
-			ExtrudeTextMaterialInstance->GetVectorParameterValue(UAvaTextMaterialHub::FMaterialStatics::SolidColor_MatParam, ExtrudeColor);
-		}
-	}
-
-	if (BevelMaterial)
-	{
-		BevelTextMaterialInstance = Cast<UMaterialInstanceDynamic>(BevelMaterial);
-		if (BevelTextMaterialInstance)
-		{
-			BevelTextMaterialInstance->GetVectorParameterValue(UAvaTextMaterialHub::FMaterialStatics::SolidColor_MatParam, BevelColor);
-		}
-	}
-}
-
-void UAvaText3DComponent::RetrieveGradientMaterialInstance()
-{
-	if (!FrontMaterial)
-	{
-		return;
-	}
-	
-	GradientTextMaterialInstance = Cast<UMaterialInstanceDynamic>(FrontMaterial);
-
-	if (GradientTextMaterialInstance)
-	{
-		GradientTextMaterialInstance->GetVectorParameterValue(UAvaTextMaterialHub::FMaterialStatics::GradientColorA_MatParam, GradientSettings.ColorA);
-		GradientTextMaterialInstance->GetVectorParameterValue(UAvaTextMaterialHub::FMaterialStatics::GradientColorB_MatParam, GradientSettings.ColorB);
-
-		GradientTextMaterialInstance->GetScalarParameterValue(UAvaTextMaterialHub::FMaterialStatics::GradientOffset_MatParam, GradientSettings.Offset);
-		GradientTextMaterialInstance->GetScalarParameterValue(UAvaTextMaterialHub::FMaterialStatics::GradientSmoothness_MatParam, GradientSettings.Smoothness);
-
-		const bool bRotationRetrieved = GradientTextMaterialInstance->GetScalarParameterValue(UAvaTextMaterialHub::FMaterialStatics::GradientRotation_MatParam, GradientSettings.Rotation);
-		
-		if (bRotationRetrieved)
-		{
-			if (GradientSettings.Rotation == UAvaTextMaterialHub::FMaterialStatics::MaterialMaskTopBottomRotationValue)
-			{
-				GradientSettings.Direction = EAvaGradientDirection::Vertical;
-			}
-			else if (GradientSettings.Rotation == UAvaTextMaterialHub::FMaterialStatics::MaterialMaskLeftToRightRotationValue)
-			{
-				GradientSettings.Direction = EAvaGradientDirection::Horizontal;
-			}
-			else
-			{
-				GradientSettings.Direction = EAvaGradientDirection::Custom;
-			}
-		}
-
-		RetrieveTranslucentMaterialParameters();
-	}
-}
-
-void UAvaText3DComponent::RetrieveTexturedMaterialInstance()
-{
-	if (!FrontMaterial)
-	{
-		return;
-	}
-	
-	TexturedTextMaterialInstance = Cast<UMaterialInstanceDynamic>(FrontMaterial);
-
-	if (TexturedTextMaterialInstance)
-	{
-		float Tiling_X;
-		float Tiling_Y;
-		TexturedTextMaterialInstance->GetScalarParameterValue(UAvaTextMaterialHub::FMaterialStatics::TexturedUTiling_MatParam, Tiling_X);
-		TexturedTextMaterialInstance->GetScalarParameterValue(UAvaTextMaterialHub::FMaterialStatics::TexturedVTiling_MatParam, Tiling_Y);
-		
-		Tiling = FVector2D(Tiling_X, Tiling_Y);
-		
-		UTexture* CurrentTexture;
-		TexturedTextMaterialInstance->GetTextureParameterValue(UAvaTextMaterialHub::FMaterialStatics::MainTexture_MatParam, CurrentTexture);
-
-		if (CurrentTexture)
-		{
-			MainTexture = Cast<UTexture2D>(CurrentTexture);
-		}
-		
-		RetrieveTranslucentMaterialParameters();
-	}
-}
-
-void UAvaText3DComponent::RetrieveMaskedMaterialParameters(const UMaterialInstanceDynamic* InMaterialInstance)
-{
-	if (!InMaterialInstance)
-	{
-		return;
-	}
-	
-	InMaterialInstance->GetScalarParameterValue(UAvaTextMaterialHub::FMaterialStatics::MaskSmoothness_MatParam, MaskSmoothness);
-	InMaterialInstance->GetScalarParameterValue(UAvaTextMaterialHub::FMaterialStatics::MaskOffset_MatParam, MaskOffset);
-	InMaterialInstance->GetScalarParameterValue(UAvaTextMaterialHub::FMaterialStatics::MaskRotation_MatParam, MaskRotation);
-
-	// in case the material was saved with the deprecated value, let's fix the rotation value
-	if (MaskRotation == UAvaTextMaterialHub::FMaterialStatics::MaterialMaskRightToLeftRotationValue_DEPRECATED)
-	{
-		MaskRotation = UAvaTextMaterialHub::FMaterialStatics::MaterialMaskRightToLeftRotationValue;
-	}
-
-	if (MaskRotation == UAvaTextMaterialHub::FMaterialStatics::MaterialMaskLeftToRightRotationValue)
-	{
-		MaskOrientation = EAvaMaterialMaskOrientation::LeftRight;
-	}
-	else if (MaskRotation == UAvaTextMaterialHub::FMaterialStatics::MaterialMaskRightToLeftRotationValue)
-	{
-		MaskOrientation = EAvaMaterialMaskOrientation::RightLeft;
-	}
-	else
-	{
-		MaskOrientation = EAvaMaterialMaskOrientation::Custom;
-	}
-}
-
-void UAvaText3DComponent::RetrieveSimpleTranslucentMaterialParameters(const UMaterialInstanceDynamic* InMaterialInstance)
-{
-	if (!InMaterialInstance)
-	{
-		return;
-	}
-	
-	InMaterialInstance->GetScalarParameterValue(UAvaTextMaterialHub::FMaterialStatics::Opacity_MatParam, Opacity);
-}
-
-void UAvaText3DComponent::RetrieveCustomMaterialReference()
-{
-	if (!FrontMaterial)
-	{
-		return;
-	}
-	
-	CustomMaterial = FrontMaterial;
-}
-
-void UAvaText3DComponent::RetrieveCurrentMaterials()
-{
-	switch (ColoringStyle)
-	{
-		case EAvaTextColoringStyle::Invalid:
-			break;
-
-		case EAvaTextColoringStyle::Solid:
-			RetrieveSolidMaterialInstances();
-			break;
-
-		case EAvaTextColoringStyle::Gradient:
-			RetrieveGradientMaterialInstance();
-			break;
-
-		case EAvaTextColoringStyle::FromTexture: 
-			RetrieveTexturedMaterialInstance();
-			break;
-
-		case EAvaTextColoringStyle::CustomMaterial: 
-			RetrieveCustomMaterialReference();
-			break;
-	}
-}
-
-void UAvaText3DComponent::RetrieveCurrentColoringStyle()
-{
-	if (!FrontMaterial)
-	{
-		return;
-	}
-	
-	if (UMaterialInterface* CurrentFrontMaterial = FrontMaterial->GetMaterial())
-	{
-		if (const FAvaTextMaterialSettings* MaterialSettings = UAvaTextMaterialHub::GetSettingsFromMaterial(CurrentFrontMaterial))
-		{
-			ColoringStyle = MaterialSettings->ColoringStyle;
-			
-			if (EnumHasAnyFlags(MaterialSettings->MaterialFeatures, EAvaTextMaterialFeatures::GradientMask))
-			{
-				TranslucencyStyle = EAvaTextTranslucency::GradientMask;
-			}
-
-			if (EnumHasAnyFlags(MaterialSettings->MaterialFeatures, EAvaTextMaterialFeatures::Translucent))
-			{
-				TranslucencyStyle = EAvaTextTranslucency::Translucent;
-			}
-			
-			bIsUnlit = EnumHasAnyFlags(MaterialSettings->MaterialFeatures, EAvaTextMaterialFeatures::Unlit);
-		}		
-		else // if there is no match, best guess is custom material
-		{
-			ColoringStyle = EAvaTextColoringStyle::CustomMaterial;
-			CustomMaterial = CurrentFrontMaterial;
-			TranslucencyStyle = EAvaTextTranslucency::None;
-		}
-	}
-}
-
-void UAvaText3DComponent::RetrieveCurrentTextAndFont()
-{
-	// if default had to be loaded (e.g. font asset missing), then we regenerate geometry
-	if (MotionDesignFont.IsDefaultFont())
-	{
-		if (IsValid(Font) && Font != MotionDesignFont.GetFont())
-		{
-			MotionDesignFont.InitFromFont(Font);
-		}
-		
-		MarkForGeometryRefresh();
-	}
-	
-	if (!IsValid(Font))
-	{
-		MarkForGeometryRefresh();
-	}
-}
-
-void UAvaText3DComponent::RetrieveCurrentAlignment()
-{
-	Alignment.HorizontalAlignment = HorizontalAlignment;
-	Alignment.VerticalAlignment = VerticalAlignment;
-}
-
-EAvaTextMaterialFeatures UAvaText3DComponent::GetMaterialFeaturesFromProperties() const
-{
-	EAvaTextMaterialFeatures MaterialFeatures = EAvaTextMaterialFeatures::None;
-	
-	if (TranslucencyStyle == EAvaTextTranslucency::GradientMask)
-	{
-		MaterialFeatures |= EAvaTextMaterialFeatures::GradientMask;
-	}
-	else if (TranslucencyStyle == EAvaTextTranslucency::Translucent)
-	{
-		MaterialFeatures |= EAvaTextMaterialFeatures::Translucent;
-	}
-
-	MaterialFeatures |= bIsUnlit ? EAvaTextMaterialFeatures::Unlit : EAvaTextMaterialFeatures::None;
-
-	return MaterialFeatures;
-}
diff --git a/Engine/Plugins/Experimental/Avalanche/Source/AvalancheText/Public/AvaText3DComponent.h b/Engine/Plugins/Experimental/Avalanche/Source/AvalancheText/Public/AvaText3DComponent.h
--- a/Engine/Plugins/Experimental/Avalanche/Source/AvalancheText/Public/AvaText3DComponent.h
+++ b/Engine/Plugins/Experimental/Avalanche/Source/AvalancheText/Public/AvaText3DComponent.h
@@ -119,13 +119,140 @@
 
 	AVALANCHETEXT_API void RefreshMaterialInstances();
 
+protected:
+	void RefreshText3D();
+	void MarkRenderDirty();
+
+	void MarkForGeometryRefresh();
+	bool IsMarkedForGeometryRefresh() const;
+
+	void MarkForLayoutRefresh();
+	bool IsMarkedForLayoutRefresh() const;
+
+	/** Mark this Text3D component so that at the next RefreshText3D call it will refresh both materials and their parameters*/
+	void MarkForMaterialsFullRefresh();
+	bool IsMarkedForMaterialRefresh() const;
+
+	void MarkForMaterialParametersRefresh();
+	bool IsMarkedForMaterialParametersRefresh() const;
+
+	void RemoveRefreshReason(EAvaTextRefreshReason InRefreshReasonToRemove);
+
+	void MarkForFullRefresh();
+	void ClearRefreshReasons();
+
+	void RegisterCallbacks();
+	void RefreshParentRootComponent();
+	void SetupParentRootAndCallbacks();
+
+	void SetupTextMaterials(const EAvaTextColoringStyle& InColoringStyle);
+	void SetupTextGeometry();
+	void SetupTextLayout();
+
+	void SetupSolidMaterial();
+	void SetupGradientMaterial();
+	void SetupCustomMaterial();
+	void RefreshTranslucencyValues();
+	void SetupTexturedMaterial();
+
+	void RefreshStoredBoundsValues();
+
+	//~ Begin UObject
+	virtual void PostLoad() override;
+	virtual void PostDuplicate(EDuplicateMode::Type DuplicateMode) override;
+	virtual void PostInitProperties() override;
+	virtual void PostEditImport() override;
+	//~ Begin UObject
+
+	void PostDuplicateActions();
+
+	/** Will schedule a Text Refresh to be executed by Core Ticker */
+	void ScheduleTextRefreshOnTicker();
+
+	void OnTextGeneratedActions();
+	void OnTextGeometryGenerated();
+
+	void OnRootComponentTransformUpdated(USceneComponent* InSceneComponent, EUpdateTransformFlags InUpdateTransformFlags, ETeleportType InTeleportType);
+
+	/**
+	 * This function can be used to obtain implicit data from the base UText3DComponent.
+	 * 
+	 * Previously, the logic from UAvaText3DComponent was handled by AAvaTextActor, which is now mostly a wrapper for UAvaText3DComponent.
+	 * Therefore, previously saved AAvaTextActor assets will have UText3DComponent materials compatible with UAvaText3DComponent.
+	 * This function extracts implicit information from those materials, and updates UAvaText3DComponent properties accordingly.
+	 * Furthermore, materials themselves will be retrieved, so that this component won't have to store them as well.
+	 * 
+	 * Additionally, also text, font and layout information will be retrieved, to avoid potential data loss from older assets.
+	 */
+	void RetrieveDataFromBaseComponent();
+
+	/** Retrieves font and text from the values stored in the component */
+	void RetrieveCurrentTextAndFont();
+
+	/** Retrieves alignment information from the values stored in the component */
+	void RetrieveCurrentAlignment();
+
+	/** Retrieves materials from the component, according to the current settings*/
+	void RetrieveCurrentMaterials();
+
+	void RetrieveCurrentColoringStyle();
+
+	void RetrieveSolidMaterialInstances();
+	void RetrieveGradientMaterialInstance();
+	void RetrieveTexturedMaterialInstance();
+	void RetrieveCustomMaterialReference();
+
+	void RetrieveTranslucentMaterialParameters();
+	void RetrieveMaskedMaterialParameters(const UMaterialInstanceDynamic* InMaterialInstance);
+	void RetrieveSimpleTranslucentMaterialParameters(const UMaterialInstanceDynamic* InMaterialInstance);
+
+	UMaterialInstanceDynamic* GetMIDWithSettings(const UMaterialInterface* InPreviousMaterial, const FAvaTextMaterialSettings& InSettings);
+
+	EAvaTextMaterialFeatures GetMaterialFeaturesFromProperties() const;
+
+	/**
+	 * UMaterialInstanceDynamic references are transient.
+	 * We need to retrieve them from material when loading assets using this component, in order to avoid having to create the material again.
+	 * Refer to UAvaText3DComponent::RetrieveCurrentMaterials().
+	 *
+	 * Additionally, this helps when loading old versions of the AAvaTextActor, which was hosting all the custom properties now handled by this component (UAvaText3DComponent)
+	 */
+
+	UPROPERTY(Transient)
+	TObjectPtr<UMaterialInstanceDynamic> TextMaterialInstance;
+
+	UPROPERTY(Transient)
+	TObjectPtr<UMaterialInstanceDynamic> ExtrudeTextMaterialInstance;
+
+	UPROPERTY(Transient)
+	TObjectPtr<UMaterialInstanceDynamic> BevelTextMaterialInstance;
+
+	UPROPERTY(Transient)
+	TObjectPtr<UMaterialInstanceDynamic> TexturedTextMaterialInstance;
+
+	UPROPERTY(Transient)
+	TObjectPtr<UMaterialInstanceDynamic> GradientTextMaterialInstance;
+
+	TWeakObjectPtr<USceneComponent> ParentRootComponent;
+
+	FOnGetMaterialWithSettings OnGetMaterialWithSettingsDelegate;
+
+	/** Used to filter which refreshes need be performed when calling RefreshText3D */
+	EAvaTextRefreshReason TextRefreshReason;
+
+	FVector LocalBoundsLowerLeftCorner;
+	FVector LocalBoundsExtent;
+
+	FDelegateHandle OnTextGeometryGeneratedDelegate;
+	FDelegateHandle OnTransformUpdatedDelegate;
+	FTSTicker::FDelegateHandle TickerHandle;
+
 private:
 #if WITH_EDITOR
 	static void RegisterOnPropertyChangeFunctions();
 #endif
 
 	void RefreshOpacity() const;
-	
 	void RefreshColor() const;
 	void RefreshExtrudeColor() const;
 	void RefreshBevelColor() const;
@@ -170,7 +297,7 @@
 
 	UPROPERTY(EditAnywhere, Getter="IsEnforcingUpperCase", Setter="SetEnforceUpperCase", Category="Text", meta=(DisplayAfter="Text", AllowPrivateAccess="true"))
 	bool bEnforceUpperCase = false;
-	
+
 	UPROPERTY(EditAnywhere, Getter="GetMotionDesignFont", Setter="SetMotionDesignFont", Category="Text", meta=(DisplayAfter="Text", AllowPrivateAccess="true"))
 	FAvaFont MotionDesignFont;
 
@@ -200,7 +327,7 @@
 
 	UPROPERTY(EditAnywhere, Getter="GetCustomMaterial", Setter="SetCustomMaterial", Category="Style", meta=(EditCondition="ColoringStyle == EAvaTextColoringStyle::CustomMaterial", EditConditionHides, DisplayAfter="Text", AllowPrivateAccess="true"))
 	TObjectPtr<UMaterialInterface> CustomMaterial = nullptr;
-	
+
 	UPROPERTY(EditAnywhere, Getter="GetTranslucencyStyle", Setter="SetTranslucencyStyle", Category="Style", meta=(EditCondition="ColoringStyle != EAvaTextColoringStyle::CustomMaterial", EditConditionHides,DisplayAfter="Text", AllowPrivateAccess="true"))
 	EAvaTextTranslucency TranslucencyStyle = EAvaTextTranslucency::None;
 
@@ -222,133 +349,6 @@
 	UPROPERTY(EditAnywhere, Getter="GetIsUnlit", Setter="SetIsUnlit", Category="Style", meta=(EditCondition="ColoringStyle != EAvaTextColoringStyle::Custom", EditConditionHides, DisplayAfter="Text", AllowPrivateAccess="true"))
 	bool bIsUnlit = true;
 
-protected:
-	void RefreshText3D();
-	void MarkRenderDirty();
-
-	void MarkForGeometryRefresh();
-	bool IsMarkedForGeometryRefresh() const;
-
-	void MarkForLayoutRefresh();
-	bool IsMarkedForLayoutRefresh() const;
-
-	/** Mark this Text3D component so that at the next RefreshText3D call it will refresh both materials and their parameters*/
-	void MarkForMaterialsFullRefresh();
-	bool IsMarkedForMaterialRefresh() const;
-
-	void MarkForMaterialParametersRefresh();
-	bool IsMarkedForMaterialParametersRefresh() const;
-
-	void RemoveRefreshReason(EAvaTextRefreshReason InRefreshReasonToRemove);
-
-	void MarkForFullRefresh();
-	void ClearRefreshReasons();
-
-	void RegisterCallbacks();
-	void RefreshParentRootComponent();
-	void SetupParentRootAndCallbacks();
-
-	void SetupTextMaterials(const EAvaTextColoringStyle& InColoringStyle);
-	void SetupTextGeometry();
-	void SetupTextLayout();
-
-	void SetupSolidMaterial();
-	void SetupGradientMaterial();
-	void SetupCustomMaterial();
-	void RefreshTranslucencyValues();
-	void SetupTexturedMaterial();
-
-	void RefreshStoredBoundsValues();
-
-	//~ Begin UObject
-	virtual void PostLoad() override;
-	virtual void PostDuplicate(EDuplicateMode::Type DuplicateMode) override;
-	virtual void PostInitProperties() override;
-	virtual void PostEditImport() override;
-	//~ Begin UObject
-
-	void PostDuplicateActions();
-
-	/** Will schedule a Text Refresh to be executed by Core Ticker */
-	void ScheduleTextRefreshOnTicker();
-
-	TWeakObjectPtr<USceneComponent> ParentRootComponent;
-
-	/**
-	 * UMaterialInstanceDynamic references are transient.
-	 * We need to retrieve them from material when loading assets using this component, in order to avoid having to create the material again.
-	 * Refer to UAvaText3DComponent::RetrieveCurrentMaterials().
-	 *
-	 * Additionally, this helps when loading old versions of the AAvaTextActor, which was hosting all the custom properties now handled by this component (UAvaText3DComponent)
-	 */
-
-	UPROPERTY(Transient)
-	TObjectPtr<UMaterialInstanceDynamic> TextMaterialInstance;
-
-	UPROPERTY(Transient)
-	TObjectPtr<UMaterialInstanceDynamic> ExtrudeTextMaterialInstance;
-
-	UPROPERTY(Transient)
-	TObjectPtr<UMaterialInstanceDynamic> BevelTextMaterialInstance;
-
-	UPROPERTY(Transient)
-	TObjectPtr<UMaterialInstanceDynamic> TexturedTextMaterialInstance;
-
-	UPROPERTY(Transient)
-	TObjectPtr<UMaterialInstanceDynamic> GradientTextMaterialInstance;
-
-	FOnGetMaterialWithSettings OnGetMaterialWithSettingsDelegate;
-	UMaterialInstanceDynamic* GetMIDWithSettings(const UMaterialInterface* InPreviousMaterial, const FAvaTextMaterialSettings& InSettings);
-
-	void OnTextGeneratedActions();
-	void OnTextGeometryGenerated();
-
-	void OnRootComponentTransformUpdated(USceneComponent* InSceneComponent, EUpdateTransformFlags InUpdateTransformFlags, ETeleportType InTeleportType);
-
-	/**
-	 * This function can be used to obtain implicit data from the base UText3DComponent.
-	 * 
-	 * Previously, the logic from UAvaText3DComponent was handled by AAvaTextActor, which is now mostly a wrapper for UAvaText3DComponent.
-	 * Therefore, previously saved AAvaTextActor assets will have UText3DComponent materials compatible with UAvaText3DComponent.
-	 * This function extracts implicit information from those materials, and updates UAvaText3DComponent properties accordingly.
-	 * Furthermore, materials themselves will be retrieved, so that this component won't have to store them as well.
-	 * 
-	 * Additionally, also text, font and layout information will be retrieved, to avoid potential data loss from older assets.
-	 */
-	void RetrieveDataFromBaseComponent();
-
-	/** Retrieves font and text from the values stored in the component */
-	void RetrieveCurrentTextAndFont();
-
-	/** Retrieves alignment information from the values stored in the component */
-	void RetrieveCurrentAlignment();
-
-	/** Retrieves materials from the component, according to the current settings*/
-	void RetrieveCurrentMaterials();
-
-	void RetrieveCurrentColoringStyle();
-	void RetrieveSolidMaterialInstances();
-	void RetrieveGradientMaterialInstance();
-	void RetrieveTexturedMaterialInstance();
-	void RetrieveCustomMaterialReference();
-
-	void RetrieveTranslucentMaterialParameters();
-	void RetrieveMaskedMaterialParameters(const UMaterialInstanceDynamic* InMaterialInstance);
-	void RetrieveSimpleTranslucentMaterialParameters(const UMaterialInstanceDynamic* InMaterialInstance);
-
-	EAvaTextMaterialFeatures GetMaterialFeaturesFromProperties() const;
-
-	/** Used to filter which refreshes need be performed when calling RefreshText3D */
-	EAvaTextRefreshReason TextRefreshReason;
-
-	FVector LocalBoundsLowerLeftCorner;
-	FVector LocalBoundsExtent;
-
-	FDelegateHandle OnTextGeometryGeneratedDelegate;
-	FDelegateHandle OnTransformUpdatedDelegate;
-	FTSTicker::FDelegateHandle TickerHandle;
-
-private:
 #if WITH_EDITOR
 	static inline TMap<FName, TFunction<void(UAvaText3DComponent*)>> PropertyChangeFunctions;
 #endif
