Motion Design
- The opacity slider on post process effects now updates in real time.

diff --git a/Engine/Plugins/Experimental/Avalanche/Source/AvalancheLevelViewport/Private/Widgets/SAvaLevelViewportStatusBarButtons.cpp b/Engine/Plugins/Experimental/Avalanche/Source/AvalancheLevelViewport/Private/Widgets/SAvaLevelViewportStatusBarButtons.cpp
--- a/Engine/Plugins/Experimental/Avalanche/Source/AvalancheLevelViewport/Private/Widgets/SAvaLevelViewportStatusBarButtons.cpp
+++ b/Engine/Plugins/Experimental/Avalanche/Source/AvalancheLevelViewport/Private/Widgets/SAvaLevelViewportStatusBarButtons.cpp
@@ -182,6 +182,7 @@ void SAvaLevelViewportStatusBarButtons::CreateContextMenuWigets()
 			.OnBeginSliderMovement(LevelViewport, &SAvaLevelViewport::OnBackgroundOpacitySliderBegin)
 			.OnEndSliderMovement(LevelViewport, &SAvaLevelViewport::OnBackgroundOpacitySliderEnd)
 			.OnValueCommitted(LevelViewport, &SAvaLevelViewport::OnBackgroundOpacityCommitted)
+			.OnValueChanged(LevelViewport, &SAvaLevelViewport::OnBackgroundOpacityCommitted, ETextCommit::Default)
 			.Value(LevelViewport, &SAvaLevelViewport::GetBackgroundOpacity)
 			.MinValue(0.f)
 			.MinSliderValue(0.f)
