Motion Design
- Corrected default asset location for background visualizer material

diff --git a/Engine/Plugins/Experimental/Avalanche/Source/AvalancheViewport/Private/AvaViewportSettings.cpp b/Engine/Plugins/Experimental/Avalanche/Source/AvalancheViewport/Private/AvaViewportSettings.cpp
--- a/Engine/Plugins/Experimental/Avalanche/Source/AvalancheViewport/Private/AvaViewportSettings.cpp
+++ b/Engine/Plugins/Experimental/Avalanche/Source/AvalancheViewport/Private/AvaViewportSettings.cpp
@@ -9,7 +9,7 @@ UAvaViewportSettings::UAvaViewportSettings()
 
 	bEnableViewportOverlay = true;
 	bEnableBoundingBoxes = true;
-	ViewportBackgroundMaterial = TSoftObjectPtr<UMaterial>(FSoftObjectPath(TEXT("/Script/Engine.Material'/Avalanche/EditorResources/M_Background.M_Background'")));
+	ViewportBackgroundMaterial = TSoftObjectPtr<UMaterial>(FSoftObjectPath(TEXT("/Script/Engine.Material'/Avalanche/EditorResources/M_Backplate.M_Backplate'")));
 	ViewportCheckerboardMaterial = TSoftObjectPtr<UMaterial>(FSoftObjectPath(TEXT("/Script/Engine.Material'/Avalanche/EditorResources/M_Checkerboard.M_Checkerboard'")));
 	ViewportCheckerboardColor0 = FLinearColor(FVector3f(0.048172f));
 	ViewportCheckerboardColor1 = FLinearColor(FVector3f(0.177888f));
