Motion Design
- Asset version (?) fix.

#virtualized

diff --git a/Engine/Plugins/Experimental/Avalanche/Config/DefaultAvalanche.ini b/Engine/Plugins/Experimental/Avalanche/Config/DefaultAvalanche.ini
--- a/Engine/Plugins/Experimental/Avalanche/Config/DefaultAvalanche.ini
+++ b/Engine/Plugins/Experimental/Avalanche/Config/DefaultAvalanche.ini
@@ -267,5 +267,12 @@
 +PropertyRedirects=(OldName="/Script/AvalancheMediaEditor.AvaRundownPageChangeChannel.Playlist",NewName="/Script/AvalancheMediaEditor.AvaRundownPageChangeChannel.Rundown")
 +EnumRedirects=(OldName="/Script/AvalancheMedia.EAvaRemoteControlResult",NewName="/Script/AvalancheMedia.EAvaPlayableRemoteControlResult")
 
++PackageRedirects=(OldName="/Avalanche/EditorResources/M_AvalancheBackplate.M_AvalancheBackplate",NewName="/Avalanche/EditorResources/M_Backplate.M_Backplate")
++PackageRedirects=(OldName="/Avalanche/EditorResources/M_AvalancheBackplate.M_AvalancheCheckerboard",NewName="/Avalanche/EditorResources/M_Checkerboard.M_Checkerboard")
++PackageRedirects=(OldName="/Avalanche/EditorResources/M_AvalancheChannelVisualizer_Red.M_AvalancheChannelVisualizer_Red",NewName="/Avalanche/EditorResources/M_ChannelVisualizer_Red.M_ChannelVisualizer_Red")
++PackageRedirects=(OldName="/Avalanche/EditorResources/M_AvalancheChannelVisualizer_Green.M_AvalancheChannelVisualizer_Green",NewName="/Avalanche/EditorResources/M_ChannelVisualizer_Green.M_ChannelVisualizer_Green")
++PackageRedirects=(OldName="/Avalanche/EditorResources/M_AvalancheChannelVisualizer_Blue.M_AvalancheChannelVisualizer_Blue",NewName="/Avalanche/EditorResources/M_ChannelVisualizer_Blue.M_ChannelVisualizer_Blue")
++PackageRedirects=(OldName="/Avalanche/EditorResources/M_AvalancheChannelVisualizer_Alpha.M_AvalancheChannelVisualizer_Alpha",NewName="/Avalanche/EditorResources/M_ChannelVisualizer_Alpha.M_ChannelVisualizer_Alpha")
+
 [/Script/Avalanche.AvaAnimation]
 DefaultCompletionMode=KeepState
