Motion Design:
- Shapes API Cleanup

#changelist validated
#virtualized

diff --git a/Engine/Plugins/Experimental/Avalanche/Source/AvalancheShapes/Private/AvaShapeParametricMaterial.cpp b/Engine/Plugins/Experimental/Avalanche/Source/AvalancheShapes/Private/AvaShapeParametricMaterial.cpp
--- a/Engine/Plugins/Experimental/Avalanche/Source/AvalancheShapes/Private/AvaShapeParametricMaterial.cpp
+++ b/Engine/Plugins/Experimental/Avalanche/Source/AvalancheShapes/Private/AvaShapeParametricMaterial.cpp
@@ -317,28 +317,6 @@ FAvaShapeParametricMaterial::FOnMaterialChanged& FAvaShapeParametricMaterial::On
 	return OnMaterialChangedDelegate;
 }
 
-UObject* FAvaShapeParametricMaterial::GetMaterialsOuter() const
-{
-	UObject* Outer = nullptr;
-
-	for (const TObjectPtr<UMaterialInstanceDynamic>& Material : InstanceMaterials)
-	{
-		if (!IsValid(Material))
-		{
-			continue;
-		}
-
-		if (Outer && Outer != Material->GetOuter())
-		{
-			return nullptr;
-		}
-
-		Outer = Material->GetOuter();
-	}
-
-	return Outer;
-}
-
 void FAvaShapeParametricMaterial::SetMaterial(UMaterialInstanceDynamic* InMaterial)
 {
 	if (!IsValid(InMaterial))
diff --git a/Engine/Plugins/Experimental/Avalanche/Source/AvalancheShapes/Private/AvaShapesDefs.cpp b/Engine/Plugins/Experimental/Avalanche/Source/AvalancheShapes/Private/AvaShapesDefs.cpp
--- a/Engine/Plugins/Experimental/Avalanche/Source/AvalancheShapes/Private/AvaShapesDefs.cpp
+++ b/Engine/Plugins/Experimental/Avalanche/Source/AvalancheShapes/Private/AvaShapesDefs.cpp
@@ -18,30 +18,6 @@ namespace UE::AvaShapes
 		return LineStart + T * LineVector;
 	}
 
-	float GetSmoothnessFromSubdivisions(float BevelSize, uint8 Subdivisions)
-	{
-		if (Subdivisions == 0)
-		{
-			return 0;
-		}
-
-		const float Smoothness = FMath::Log2(29 * static_cast<float>(Subdivisions) / 90.f + 1.f) / FMath::Log2(30.f);
-
-		return FMath::Clamp(Smoothness, 0, 1);
-	}
-
-	uint8 GetSubdivisionsFromSmoothness(float BevelSize, float Smoothness)
-	{
-		if (Smoothness == 0)
-		{
-			return 0;
-		}
-
-		const float Subdivisions = 90.f * (FMath::Pow(30.f, Smoothness) - 1.f) / 29.f;
-
-		return static_cast<uint8>(FMath::Clamp(Subdivisions, 0, 128));
-	}
-
 	bool TransformMeshUVs(UE::Geometry::FDynamicMesh3& InEditMesh, const TArray<int32>& UVIds
 		, const FAvaShapeMaterialUVParameters& InParams, const FVector2D& InShapeSize, const FVector2D& InUVOffset
 		, const float InUVFixRotation)
diff --git a/Engine/Plugins/Experimental/Avalanche/Source/AvalancheShapes/Public/AvaShapeMesh.h b/Engine/Plugins/Experimental/Avalanche/Source/AvalancheShapes/Public/AvaShapeMesh.h
--- a/Engine/Plugins/Experimental/Avalanche/Source/AvalancheShapes/Public/AvaShapeMesh.h
+++ b/Engine/Plugins/Experimental/Avalanche/Source/AvalancheShapes/Public/AvaShapeMesh.h
@@ -80,10 +80,9 @@ protected:
 
 /** Represents a mesh section with its material, uv data */
 USTRUCT(BlueprintType)
-struct AVALANCHESHAPES_API FAvaShapeMeshData
+struct FAvaShapeMeshData
 {
 	friend class UAvaShapeDynamicMeshBase;
-
 	friend class FAvaMeshesDetailCustomization;
 	friend class FAvaShapeDynamicMeshVisualizer;
 
diff --git a/Engine/Plugins/Experimental/Avalanche/Source/AvalancheShapes/Public/AvaShapeParametricMaterial.h b/Engine/Plugins/Experimental/Avalanche/Source/AvalancheShapes/Public/AvaShapeParametricMaterial.h
--- a/Engine/Plugins/Experimental/Avalanche/Source/AvalancheShapes/Public/AvaShapeParametricMaterial.h
+++ b/Engine/Plugins/Experimental/Avalanche/Source/AvalancheShapes/Public/AvaShapeParametricMaterial.h
@@ -16,7 +16,7 @@ class UObject;
 struct FAvaShapeParametricMaterial;
 
 USTRUCT(BlueprintType)
-struct AVALANCHESHAPES_API FAvaShapeParametricMaterial
+struct FAvaShapeParametricMaterial
 {
 	GENERATED_BODY()
 
@@ -28,33 +28,30 @@ struct AVALANCHESHAPES_API FAvaShapeParametricMaterial
 	DECLARE_MULTICAST_DELEGATE_OneParam(FOnMaterialChanged, FAvaShapeParametricMaterial&)
 	DECLARE_MULTICAST_DELEGATE(FOnMaterialParameterChanged)
 
-	FAvaShapeParametricMaterial();
+	AVALANCHESHAPES_API FAvaShapeParametricMaterial();
 
-	FAvaShapeParametricMaterial(const FAvaShapeParametricMaterial& Other);
-	FAvaShapeParametricMaterial& operator=(const FAvaShapeParametricMaterial& Other);
+	AVALANCHESHAPES_API FAvaShapeParametricMaterial(const FAvaShapeParametricMaterial& Other);
+	AVALANCHESHAPES_API FAvaShapeParametricMaterial& operator=(const FAvaShapeParametricMaterial& Other);
 
 	/** Called when active material parameters are updated */
-	FOnMaterialParameterChanged& OnMaterialParameterChanged();
+	AVALANCHESHAPES_API FOnMaterialParameterChanged& OnMaterialParameterChanged();
 
 	/** Called when active material index changed */
 	static FOnMaterialChanged& OnMaterialChanged();
 
-	bool CopyFromMaterialParameters(UMaterialInstance* InMaterial);
+	AVALANCHESHAPES_API bool CopyFromMaterialParameters(UMaterialInstance* InMaterial);
 
 	/** Check if input material is a parametric material */
 	bool IsParametricMaterial(UMaterialInterface* InMaterial, const bool bCheckIfDefault = false) const;
 
 	/** Get default parent material currently active */
-	UMaterialInterface* GetDefaultMaterial() const;
-
-	/** Get outer from all instanced materials or null if they differ */
-	UObject* GetMaterialsOuter() const;
+	AVALANCHESHAPES_API UMaterialInterface* GetDefaultMaterial() const;
 
 	/** Get active up to date material instance */
-	UMaterialInstanceDynamic* GetMaterial() const;
+	AVALANCHESHAPES_API UMaterialInstanceDynamic* GetMaterial() const;
 
 	/** Get active up to date material instance or creates it */
-	UMaterialInstanceDynamic* GetOrCreateMaterial(UObject* InOuter);
+	AVALANCHESHAPES_API UMaterialInstanceDynamic* GetOrCreateMaterial(UObject* InOuter);
 
 	/** Set instanced material currently active */
 	void SetMaterial(UMaterialInstanceDynamic* InMaterial);
@@ -64,52 +61,52 @@ struct AVALANCHESHAPES_API FAvaShapeParametricMaterial
 		return bUseTranslucentMaterial;
 	}
 
-	void SetUseTranslucentMaterial(bool bInUse);
+	AVALANCHESHAPES_API void SetUseTranslucentMaterial(bool bInUse);
 
 	bool GetUseAutoTranslucency() const
 	{
 		return bUseAutoTranslucency;
 	}
 
-	void SetUseAutoTranslucency(bool bInUse);
+	AVALANCHESHAPES_API void SetUseAutoTranslucency(bool bInUse);
 
 	EAvaShapeParametricMaterialStyle GetStyle() const
 	{
 		return Style;
 	}
 
-	void SetStyle(EAvaShapeParametricMaterialStyle InStyle);
+	AVALANCHESHAPES_API void SetStyle(EAvaShapeParametricMaterialStyle InStyle);
 
 	const FLinearColor& GetPrimaryColor() const
 	{
 		return ColorA;
 	}
 
-	void SetPrimaryColor(const FLinearColor& InColor);
+	AVALANCHESHAPES_API void SetPrimaryColor(const FLinearColor& InColor);
 
 	const FLinearColor& GetSecondaryColor() const
 	{
 		return ColorB;
 	}
 
-	void SetSecondaryColor(const FLinearColor& InColor);
+	AVALANCHESHAPES_API void SetSecondaryColor(const FLinearColor& InColor);
 
 	float GetGradientOffset() const
 	{
 		return GradientOffset;
 	}
 
-	void SetGradientOffset(float InOffset);
+	AVALANCHESHAPES_API void SetGradientOffset(float InOffset);
 
 	bool GetUseUnlitMaterial() const
 	{
 		return bUseUnlitMaterial;
 	}
 
-	void SetUseUnlitMaterial(bool bInUse);
+	AVALANCHESHAPES_API void SetUseUnlitMaterial(bool bInUse);
 
 	/** Set parameter values on a material instance */
-	void SetMaterialParameterValues(UMaterialInstanceDynamic* InMaterialInstance, bool bInNotifyUpdate = false) const;
+	AVALANCHESHAPES_API void SetMaterialParameterValues(UMaterialInstanceDynamic* InMaterialInstance, bool bInNotifyUpdate = false) const;
 
 protected:
 	/** Parent material
@@ -159,6 +156,13 @@ protected:
 	UPROPERTY(EditInstanceOnly, BlueprintReadWrite, Category="Material", meta=(DisplayName="Use Unlit Material", AllowPrivateAccess = "true"))
 	bool bUseUnlitMaterial;
 
+private:
+	static FOnMaterialChanged OnMaterialChangedDelegate;
+
+	FOnMaterialParameterChanged OnMaterialParameterChangedDelegate;
+
+	int32 ActiveInstanceIndex = INDEX_NONE;
+
 	/** Load parents materials to create instance materials */
 	void LoadDefaultMaterials();
 
@@ -170,11 +174,4 @@ protected:
 
 	/** Get the active material index */
 	int32 GetActiveInstanceIndex() const;
-
-private:
-	static FOnMaterialChanged OnMaterialChangedDelegate;
-
-	FOnMaterialParameterChanged OnMaterialParameterChangedDelegate;
-
-	int32 ActiveInstanceIndex = INDEX_NONE;
 };
diff --git a/Engine/Plugins/Experimental/Avalanche/Source/AvalancheShapes/Public/AvaShapePrimitiveFunctions.h b/Engine/Plugins/Experimental/Avalanche/Source/AvalancheShapes/Public/AvaShapePrimitiveFunctions.h
--- a/Engine/Plugins/Experimental/Avalanche/Source/AvalancheShapes/Public/AvaShapePrimitiveFunctions.h
+++ b/Engine/Plugins/Experimental/Avalanche/Source/AvalancheShapes/Public/AvaShapePrimitiveFunctions.h
@@ -10,17 +10,15 @@
 class AAvaShapeActor;
 
 /** FunctionLibrary to Create Ava Shape Meshes and apply them to a Shape Actor. */
-UCLASS()
-class AVALANCHESHAPES_API UAvaShapeMeshFunctions
-	: public UBlueprintFunctionLibrary
+UCLASS(MinimalAPI)
+class UAvaShapeMeshFunctions : public UBlueprintFunctionLibrary
 {
 	GENERATED_BODY()
 
 public:
 	/** Sets the Shape Actor mesh to Rectangle. */
 	UFUNCTION(BlueprintCallable, Category = "Shapes", meta = (ScriptMethod))
-	static UPARAM(DisplayName = "Mesh") UAvaShapeRectangleDynamicMesh*
-	SetRectangle(AAvaShapeActor* ShapeActor
+	static AVALANCHESHAPES_API UPARAM(DisplayName = "Mesh") UAvaShapeRectangleDynamicMesh* SetRectangle(AAvaShapeActor* ShapeActor
 		, const FVector2D& Size
 		, const FTransform& Transform);
 
diff --git a/Engine/Plugins/Experimental/Avalanche/Source/AvalancheShapes/Public/AvaShapesDefs.h b/Engine/Plugins/Experimental/Avalanche/Source/AvalancheShapes/Public/AvaShapesDefs.h
--- a/Engine/Plugins/Experimental/Avalanche/Source/AvalancheShapes/Public/AvaShapesDefs.h
+++ b/Engine/Plugins/Experimental/Avalanche/Source/AvalancheShapes/Public/AvaShapesDefs.h
@@ -15,11 +15,7 @@ struct FAvaShapeMaterialUVParameters;
 
 namespace UE::AvaShapes
 {
-	FVector2D AVALANCHESHAPES_API FindClosestPointOnLine(const FVector2D& LineStart, const FVector2D& LineEnd, const FVector2D& TestPoint);
-
-	float AVALANCHESHAPES_API GetSmoothnessFromSubdivisions(float BevelSize, uint8 Subdivisions);
-
-	uint8 AVALANCHESHAPES_API GetSubdivisionsFromSmoothness(float BevelSize, float Smoothness);
+	FVector2D FindClosestPointOnLine(const FVector2D& LineStart, const FVector2D& LineEnd, const FVector2D& TestPoint);
 
 	bool AVALANCHESHAPES_API TransformMeshUVs(UE::Geometry::FDynamicMesh3& InEditMesh
 		, const TArray<int32>& UVIds
