Motion Design:
- Outliner API cleanup

#changelist validated
#virtualized

diff --git a/Engine/Plugins/Experimental/Avalanche/Source/Avalanche/Private/AvaActorUtils.cpp b/Engine/Plugins/Experimental/Avalanche/Source/Avalanche/Private/AvaActorUtils.cpp
--- a/Engine/Plugins/Experimental/Avalanche/Source/Avalanche/Private/AvaActorUtils.cpp
+++ b/Engine/Plugins/Experimental/Avalanche/Source/Avalanche/Private/AvaActorUtils.cpp
@@ -20,7 +20,6 @@
 #include "AvaOutlinerSubsystem.h"
 #include "AvaOutlinerUtils.h"
 #include "IAvaOutliner.h"
-#include "Item/AvaOutlinerTreeRoot.h"
 #endif
 
 FOrientedBox FAvaActorUtils::MakeOrientedBox(const FBox& InLocalBox, const FTransform& InWorldTransform)
diff --git a/Engine/Plugins/Experimental/Avalanche/Source/Avalanche/Private/AvaOutlinerUtils.cpp b/Engine/Plugins/Experimental/Avalanche/Source/Avalanche/Private/AvaOutlinerUtils.cpp
--- a/Engine/Plugins/Experimental/Avalanche/Source/Avalanche/Private/AvaOutlinerUtils.cpp
+++ b/Engine/Plugins/Experimental/Avalanche/Source/Avalanche/Private/AvaOutlinerUtils.cpp
@@ -9,7 +9,6 @@
 #include "GameFramework/Actor.h"
 #include "IAvaOutliner.h"
 #include "Item/AvaOutlinerActor.h"
-#include "Item/AvaOutlinerTreeRoot.h"
 
 TSharedPtr<IAvaOutliner> FAvaOutlinerUtils::EditorGetOutliner(const UWorld* const InWorld)
 {
diff --git a/Engine/Plugins/Experimental/Avalanche/Source/AvalancheOutliner/Private/AvaOutliner.cpp b/Engine/Plugins/Experimental/Avalanche/Source/AvalancheOutliner/Private/AvaOutliner.cpp
--- a/Engine/Plugins/Experimental/Avalanche/Source/AvalancheOutliner/Private/AvaOutliner.cpp
+++ b/Engine/Plugins/Experimental/Avalanche/Source/AvalancheOutliner/Private/AvaOutliner.cpp
@@ -391,6 +391,11 @@ void FAvaOutliner::Refresh()
 	SaveState->SaveSceneTree(*this, /*bInResetTree*/true);
 }
 
+TSharedRef<FAvaOutlinerItem> FAvaOutliner::GetTreeRoot() const
+{
+	return RootItem;
+}
+
 FAvaOutlinerItemPtr FAvaOutliner::FindItem(const FAvaOutlinerItemId& InItemId) const
 {
 	if (bIteratingItemMap)
diff --git a/Engine/Plugins/Experimental/Avalanche/Source/AvalancheOutliner/Private/AvaOutliner.h b/Engine/Plugins/Experimental/Avalanche/Source/AvalancheOutliner/Private/AvaOutliner.h
--- a/Engine/Plugins/Experimental/Avalanche/Source/AvalancheOutliner/Private/AvaOutliner.h
+++ b/Engine/Plugins/Experimental/Avalanche/Source/AvalancheOutliner/Private/AvaOutliner.h
@@ -91,7 +91,7 @@ public:
 	virtual void UnregisterItem(const FAvaOutlinerItemId& InItemId) override;
 	virtual void RequestRefresh() override;
 	virtual void Refresh() override;
-	virtual TSharedRef<FAvaOutlinerTreeRoot> GetTreeRoot() const override { return RootItem; }
+	virtual TSharedRef<FAvaOutlinerItem> GetTreeRoot() const override;
 	virtual FAvaOutlinerItemPtr FindItem(const FAvaOutlinerItemId& InItemId) const override;
 	virtual void SetIgnoreNotify(EAvaOutlinerIgnoreNotifyFlags InFlag, bool bIgnore) override;
 	virtual void OnActorsCopied(FString& InOutCopiedData, TConstArrayView<AActor*> InCopiedActors) override;
diff --git a/Engine/Plugins/Experimental/Avalanche/Source/AvalancheOutliner/Private/AvaOutlinerView.cpp b/Engine/Plugins/Experimental/Avalanche/Source/AvalancheOutliner/Private/AvaOutlinerView.cpp
--- a/Engine/Plugins/Experimental/Avalanche/Source/AvalancheOutliner/Private/AvaOutlinerView.cpp
+++ b/Engine/Plugins/Experimental/Avalanche/Source/AvalancheOutliner/Private/AvaOutlinerView.cpp
@@ -1241,7 +1241,7 @@ void FAvaOutlinerView::OnDragEnter(const FDragDropEvent& InDragDropEvent, FAvaOu
 {
 	if (!InTargetItem.IsValid() && OutlinerWeak.IsValid())
 	{
-		TSharedRef<FAvaOutlinerTreeRoot> TreeRoot = OutlinerWeak.Pin()->GetTreeRoot();
+		TSharedRef<FAvaOutlinerItem> TreeRoot = OutlinerWeak.Pin()->GetTreeRoot();
 		const bool bCanAcceptDrop = TreeRoot->CanAcceptDrop(InDragDropEvent, EItemDropZone::OntoItem).IsSet();
 		SetDragIntoTreeRoot(bCanAcceptDrop);
 	}
@@ -1302,7 +1302,7 @@ FReply FAvaOutlinerView::OnDrop(const FDragDropEvent& InDragDropEvent, EItemDrop
 		return InTargetItem->AcceptDrop(InDragDropEvent, InDropZone);
 	}
 
-	TSharedPtr<FAvaOutlinerTreeRoot> TreeRoot;
+	TSharedPtr<FAvaOutlinerItem> TreeRoot;
 
 	if (OutlinerWeak.IsValid())
 	{
diff --git a/Engine/Plugins/Experimental/Avalanche/Source/AvalancheOutliner/Private/Data/AvaOutlinerSaveState.cpp b/Engine/Plugins/Experimental/Avalanche/Source/AvalancheOutliner/Private/Data/AvaOutlinerSaveState.cpp
--- a/Engine/Plugins/Experimental/Avalanche/Source/AvalancheOutliner/Private/Data/AvaOutlinerSaveState.cpp
+++ b/Engine/Plugins/Experimental/Avalanche/Source/AvalancheOutliner/Private/Data/AvaOutlinerSaveState.cpp
@@ -177,7 +177,7 @@ void FAvaOutlinerSaveState::SaveOutlinerViewItems(const FAvaOutliner& InOutliner
 	, const FAvaOutlinerView& InOutlinerView
 	, FAvaOutlinerViewSaveState& OutSaveState)
 {
-	const TSharedRef<FAvaOutlinerTreeRoot> TreeRoot = InOutliner.GetTreeRoot();
+	const TSharedRef<FAvaOutlinerItem> TreeRoot = InOutliner.GetTreeRoot();
 
 	TArray<FAvaOutlinerItemPtr> ItemsToSave = TreeRoot->GetChildren();
 
@@ -213,7 +213,7 @@ void FAvaOutlinerSaveState::LoadOutlinerViewItems(const FAvaOutliner& InOutliner
 	, FAvaOutlinerView& InOutOutlinerView
 	, const FAvaOutlinerViewSaveState& InSaveState)
 {
-	const TSharedRef<FAvaOutlinerTreeRoot> TreeRoot = InOutliner.GetTreeRoot();
+	const TSharedRef<FAvaOutlinerItem> TreeRoot = InOutliner.GetTreeRoot();
 
 	TArray<FAvaOutlinerItemPtr> ItemsToLoad = TreeRoot->GetChildren();
 
diff --git a/Engine/Plugins/Experimental/Avalanche/Source/AvalancheOutliner/Private/DragDropOps/Handlers/AvaOutlinerActorDropHandler.h b/Engine/Plugins/Experimental/Avalanche/Source/AvalancheOutliner/Private/DragDropOps/Handlers/AvaOutlinerActorDropHandler.h
new file mode 100644
--- /dev/null
+++ b/Engine/Plugins/Experimental/Avalanche/Source/AvalancheOutliner/Private/DragDropOps/Handlers/AvaOutlinerActorDropHandler.h
@@ -0,0 +1,21 @@
+﻿// Copyright Epic Games, Inc. All Rights Reserved.
+
+#pragma once
+
+#include "DragDropOps/Handlers/AvaOutlinerItemDropHandler.h"
+
+/** Class that handles Dropping Actor Items into a Target Item */
+class FAvaOutlinerActorDropHandler : public FAvaOutlinerItemDropHandler
+{
+public:
+	UE_AVA_INHERITS(FAvaOutlinerActorDropHandler, FAvaOutlinerItemDropHandler);
+
+private:
+	//~ Begin FAvaOutlinerItemDropHandler
+	virtual bool IsDraggedItemSupported(const FAvaOutlinerItemPtr& InDraggedItem) const override;
+	virtual TOptional<EItemDropZone> CanDrop(EItemDropZone InDropZone, FAvaOutlinerItemPtr InTargetItem) const override;
+	virtual bool Drop(EItemDropZone InDropZone, FAvaOutlinerItemPtr InTargetItem) override;
+	//~ End FAvaOutlinerItemDropHandler
+
+	void MoveItems(EItemDropZone InDropZone, FAvaOutlinerItemPtr InTargetItem);
+};
diff --git a/Engine/Plugins/Experimental/Avalanche/Source/AvalancheOutliner/Private/Filters/AvaOutlinerItemTypeFilter.cpp b/Engine/Plugins/Experimental/Avalanche/Source/AvalancheOutliner/Private/Filters/AvaOutlinerItemTypeFilter.cpp
--- a/Engine/Plugins/Experimental/Avalanche/Source/AvalancheOutliner/Private/Filters/AvaOutlinerItemTypeFilter.cpp
+++ b/Engine/Plugins/Experimental/Avalanche/Source/AvalancheOutliner/Private/Filters/AvaOutlinerItemTypeFilter.cpp
@@ -77,6 +77,11 @@ const FSlateBrush* FAvaOutlinerItemTypeFilterData::GetIcon() const
 	return nullptr;
 }
 
+EAvaOutlinerTypeFilterMode FAvaOutlinerItemTypeFilterData::GetFilterMode() const
+{
+	return FilterMode;
+}
+
 void FAvaOutlinerItemTypeFilterData::SetOverrideIconColor(FSlateColor InNewIconColor)
 {
 	OverrideIcon.TintColor = InNewIconColor;
diff --git a/Engine/Plugins/Experimental/Avalanche/Source/AvalancheOutliner/Private/Icon/AvaOutlinerObjectIconCustomization.cpp b/Engine/Plugins/Experimental/Avalanche/Source/AvalancheOutliner/Private/Icon/AvaOutlinerObjectIconCustomization.cpp
--- a/Engine/Plugins/Experimental/Avalanche/Source/AvalancheOutliner/Private/Icon/AvaOutlinerObjectIconCustomization.cpp
+++ b/Engine/Plugins/Experimental/Avalanche/Source/AvalancheOutliner/Private/Icon/AvaOutlinerObjectIconCustomization.cpp
@@ -15,6 +15,16 @@ FAvaOutlinerObjectIconCustomization::FAvaOutlinerObjectIconCustomization(const U
 	}
 }
 
+void FAvaOutlinerObjectIconCustomization::SetOverriddenIcon(const FOnGetOverriddenObjectIcon& InOverriddenIcon)
+{
+	OnGetOverriddenIcon = InOverriddenIcon;
+}
+
+FName FAvaOutlinerObjectIconCustomization::GetOutlinerItemIdentifier() const
+{
+	return SupportedClassName;
+}
+
 bool FAvaOutlinerObjectIconCustomization::HasOverrideIcon(TSharedPtr<const FAvaOutlinerItem> InOutlinerItem) const
 {
 	if (const FAvaOutlinerObject* ObjectItem = InOutlinerItem->CastTo<FAvaOutlinerObject>())
@@ -40,8 +50,3 @@ FSlateIcon FAvaOutlinerObjectIconCustomization::GetOverrideIcon(TSharedPtr<const
 
 	return FSlateIcon();
 }
-
-void FAvaOutlinerObjectIconCustomization::SetOverriddenIcon(const FOnGetOverriddenObjectIcon& InOverriddenIcon)
-{
-	OnGetOverriddenIcon = InOverriddenIcon;
-}
diff --git a/Engine/Plugins/Experimental/Avalanche/Source/AvalancheOutliner/Private/Item/AvaOutlinerTreeRoot.h b/Engine/Plugins/Experimental/Avalanche/Source/AvalancheOutliner/Private/Item/AvaOutlinerTreeRoot.h
new file mode 100644
--- /dev/null
+++ b/Engine/Plugins/Experimental/Avalanche/Source/AvalancheOutliner/Private/Item/AvaOutlinerTreeRoot.h
@@ -0,0 +1,43 @@
+// Copyright Epic Games, Inc. All Rights Reserved.
+
+#pragma once
+
+#include "Item/AvaOutlinerItem.h"
+
+/*
+ * The Item that ensures that every item (except self) has a parent in the hierarchy to make it easier
+ * to handle. This is not really a visual item, so it can't appear in the Outliner View
+ */
+class FAvaOutlinerTreeRoot final : public FAvaOutlinerItem
+{
+public:
+	UE_AVA_INHERITS_WITH_SUPER(FAvaOutlinerTreeRoot, FAvaOutlinerItem)
+
+	FAvaOutlinerTreeRoot(IAvaOutliner& InOutliner)
+		: FAvaOutlinerItem(InOutliner)
+	{
+	}
+
+	virtual ~FAvaOutlinerTreeRoot() override = default;
+
+	//~ Begin IAvaOutlinerItem
+	virtual void FindChildren(TArray<FAvaOutlinerItemPtr>& OutChildren, bool bRecursive) override;
+	virtual bool CanAddChild(const FAvaOutlinerItemPtr& InChild) const override;
+	virtual bool AddChild(const FAvaOutlinerAddItemParams& InAddItemParams) override;
+	virtual bool IsAllowedInOutliner() const override;
+	virtual FText GetDisplayName() const override;
+	virtual FText GetClassName() const override;
+	virtual FText GetIconTooltipText() const override;
+	virtual FSlateIcon GetIcon() const override;
+	virtual TSharedRef<SWidget> GenerateLabelWidget(const TSharedRef<SAvaOutlinerTreeRow>& InRow) override;
+	virtual bool CanRename() const override;
+	virtual bool Rename(const FString& InName) override;
+	virtual TOptional<EItemDropZone> CanAcceptDrop(const FDragDropEvent& InDragDropEvent, EItemDropZone InDropZone) override;	
+	virtual FReply AcceptDrop(const FDragDropEvent& InDragDropEvent, EItemDropZone InDropZone) override;
+	//~ End IAvaOutlinerItem
+	
+protected:
+	//~ Begin FAvaOutlinerItem
+	virtual FAvaOutlinerItemId CalculateItemId() const override;
+	//~ End FAvaOutlinerItem
+};
diff --git a/Engine/Plugins/Experimental/Avalanche/Source/AvalancheOutliner/Public/DragDropOps/AvaOutlinerItemDragDropOp.h b/Engine/Plugins/Experimental/Avalanche/Source/AvalancheOutliner/Public/DragDropOps/AvaOutlinerItemDragDropOp.h
--- a/Engine/Plugins/Experimental/Avalanche/Source/AvalancheOutliner/Public/DragDropOps/AvaOutlinerItemDragDropOp.h
+++ b/Engine/Plugins/Experimental/Avalanche/Source/AvalancheOutliner/Public/DragDropOps/AvaOutlinerItemDragDropOp.h
@@ -12,7 +12,7 @@ class FAvaOutlinerItemDropHandler;
 enum class EItemDropZone;
 
 /** Drag Drop Operation for Ava Outliner Items. Customized behavior can be added in via the AddDropHandler function */
-class AVALANCHEOUTLINER_API FAvaOutlinerItemDragDropOp : public FDecoratedDragDropOp
+class FAvaOutlinerItemDragDropOp : public FDecoratedDragDropOp
 {
 public:
 	DRAG_DROP_OPERATOR_TYPE(FAvaOutlinerItemDragDropOp, FDecoratedDragDropOp)
@@ -41,7 +41,7 @@ public:
 
 	/** Called when the FAvaOutlinerItemDragDropOp has been created and Initialized in FAvaOutlinerItemDragDropOp::Init */
 	DECLARE_MULTICAST_DELEGATE_OneParam(FOnItemDragDropOpInitialized, FAvaOutlinerItemDragDropOp&)
-	static FOnItemDragDropOpInitialized& OnItemDragDropOpInitialized();
+	AVALANCHEOUTLINER_API static FOnItemDragDropOpInitialized& OnItemDragDropOpInitialized();
 
 	template<typename InDropHandlerType
 		, typename = typename TEnableIf<TIsDerivedFrom<InDropHandlerType, FAvaOutlinerItemDropHandler>::Value>::Type
diff --git a/Engine/Plugins/Experimental/Avalanche/Source/AvalancheOutliner/Public/DragDropOps/Handlers/AvaOutlinerActorDropHandler.h b/Engine/Plugins/Experimental/Avalanche/Source/AvalancheOutliner/Public/DragDropOps/Handlers/AvaOutlinerActorDropHandler.h
deleted file mode 100644
diff --git a/Engine/Plugins/Experimental/Avalanche/Source/AvalancheOutliner/Public/DragDropOps/Handlers/AvaOutlinerItemDropHandler.h b/Engine/Plugins/Experimental/Avalanche/Source/AvalancheOutliner/Public/DragDropOps/Handlers/AvaOutlinerItemDropHandler.h
--- a/Engine/Plugins/Experimental/Avalanche/Source/AvalancheOutliner/Public/DragDropOps/Handlers/AvaOutlinerItemDropHandler.h
+++ b/Engine/Plugins/Experimental/Avalanche/Source/AvalancheOutliner/Public/DragDropOps/Handlers/AvaOutlinerItemDropHandler.h
@@ -17,20 +17,25 @@ template<typename OptionalType> struct TOptional;
  * Base Class to Handle Dropping Outliner Items into a Target Outliner Item
  * @see built-in example FAvaOutlinerActorDropHandler
  */
-class AVALANCHEOUTLINER_API FAvaOutlinerItemDropHandler : public IAvaTypeCastable, public TSharedFromThis<FAvaOutlinerItemDropHandler>
+class FAvaOutlinerItemDropHandler : public IAvaTypeCastable, public TSharedFromThis<FAvaOutlinerItemDropHandler>
 {
 	friend FAvaOutlinerItemDragDropOp;
 
-	void Initialize(const FAvaOutlinerItemDragDropOp& InDragDropOp);
+	AVALANCHEOUTLINER_API void Initialize(const FAvaOutlinerItemDragDropOp& InDragDropOp);
 
 public:
 	UE_AVA_INHERITS(FAvaOutlinerItemDropHandler, IAvaTypeCastable);
 
-	TConstArrayView<FAvaOutlinerItemPtr> GetItems() const { return Items; }	
+	TConstArrayView<FAvaOutlinerItemPtr> GetItems() const
+	{
+		return Items;
+	}	
 
 protected:
 	virtual bool IsDraggedItemSupported(const FAvaOutlinerItemPtr& InDraggedItem) const = 0;
+
 	virtual TOptional<EItemDropZone> CanDrop(EItemDropZone InDropZone, FAvaOutlinerItemPtr InTargetItem) const = 0;
+
 	virtual bool Drop(EItemDropZone InDropZone, FAvaOutlinerItemPtr InTargetItem) = 0;
 
 	enum class EIterationResult
diff --git a/Engine/Plugins/Experimental/Avalanche/Source/AvalancheOutliner/Public/Filters/AvaOutlinerItemTypeFilter.h b/Engine/Plugins/Experimental/Avalanche/Source/AvalancheOutliner/Public/Filters/AvaOutlinerItemTypeFilter.h
--- a/Engine/Plugins/Experimental/Avalanche/Source/AvalancheOutliner/Public/Filters/AvaOutlinerItemTypeFilter.h
+++ b/Engine/Plugins/Experimental/Avalanche/Source/AvalancheOutliner/Public/Filters/AvaOutlinerItemTypeFilter.h
@@ -25,13 +25,13 @@ enum class EAvaOutlinerTypeFilterMode : uint8
 ENUM_CLASS_FLAGS(EAvaOutlinerTypeFilterMode);
 
 USTRUCT()
-struct AVALANCHEOUTLINER_API FAvaOutlinerItemTypeFilterData
+struct FAvaOutlinerItemTypeFilterData
 {
 	GENERATED_BODY()
 
 	FAvaOutlinerItemTypeFilterData() = default;
 
-	FAvaOutlinerItemTypeFilterData(const TArray<TSubclassOf<UObject>>& InFilterClasses
+	AVALANCHEOUTLINER_API FAvaOutlinerItemTypeFilterData(const TArray<TSubclassOf<UObject>>& InFilterClasses
 		, EAvaOutlinerTypeFilterMode InMode  = EAvaOutlinerTypeFilterMode::MatchesType
 		, const FSlateBrush* InIconBrush     = nullptr
 		, const FText& InTooltipText         = FText::GetEmpty()
@@ -44,10 +44,7 @@ struct AVALANCHEOUTLINER_API FAvaOutlinerItemTypeFilterData
 
 	const FSlateBrush* GetIcon() const;
 
-	EAvaOutlinerTypeFilterMode GetFilterMode() const
-	{
-		return FilterMode;
-	}
+	EAvaOutlinerTypeFilterMode GetFilterMode() const;
 
 	void SetOverrideIconColor(FSlateColor InNewIconColor);
 
@@ -86,7 +83,7 @@ private:
 class AVALANCHEOUTLINER_API FAvaOutlinerItemTypeFilter : public IAvaOutlinerItemFilter
 {
 public:
-	FAvaOutlinerItemTypeFilter(FName InFilterId
+	explicit FAvaOutlinerItemTypeFilter(FName InFilterId
 			, const TArray<TSubclassOf<UObject>>& InFilterClasses
 			, EAvaOutlinerTypeFilterMode InMode  = EAvaOutlinerTypeFilterMode::MatchesType
 			, const FSlateBrush* InIconBrush     = nullptr
@@ -98,7 +95,7 @@ public:
 	{
 	}
 
-	FAvaOutlinerItemTypeFilter(FName InFilterId, const FAvaOutlinerItemTypeFilterData& InFilterData)
+	explicit FAvaOutlinerItemTypeFilter(FName InFilterId, const FAvaOutlinerItemTypeFilterData& InFilterData)
 		: FilterId(InFilterId)
 		, FilterData(InFilterData)
 	{
diff --git a/Engine/Plugins/Experimental/Avalanche/Source/AvalancheOutliner/Public/IAvaOutliner.h b/Engine/Plugins/Experimental/Avalanche/Source/AvalancheOutliner/Public/IAvaOutliner.h
--- a/Engine/Plugins/Experimental/Avalanche/Source/AvalancheOutliner/Public/IAvaOutliner.h
+++ b/Engine/Plugins/Experimental/Avalanche/Source/AvalancheOutliner/Public/IAvaOutliner.h
@@ -141,7 +141,7 @@ public:
 	virtual void Refresh() = 0;
 
 	/** Gets the Tree Root Item of the Outliner */
-	virtual TSharedRef<FAvaOutlinerTreeRoot> GetTreeRoot() const = 0;
+	virtual TSharedRef<FAvaOutlinerItem> GetTreeRoot() const = 0;
 
 	/**
 	 * Finds the Registered Item that has the given Id
diff --git a/Engine/Plugins/Experimental/Avalanche/Source/AvalancheOutliner/Public/Icon/AvaOutlinerObjectIconCustomization.h b/Engine/Plugins/Experimental/Avalanche/Source/AvalancheOutliner/Public/Icon/AvaOutlinerObjectIconCustomization.h
--- a/Engine/Plugins/Experimental/Avalanche/Source/AvalancheOutliner/Public/Icon/AvaOutlinerObjectIconCustomization.h
+++ b/Engine/Plugins/Experimental/Avalanche/Source/AvalancheOutliner/Public/Icon/AvaOutlinerObjectIconCustomization.h
@@ -11,21 +11,20 @@ class UClass;
 
 DECLARE_DELEGATE_RetVal_OneParam(FSlateIcon /** OutIconForItem */, FOnGetOverriddenObjectIcon, TSharedPtr<const FAvaOutlinerItem> /** InItemToCustomize */)
 
-class AVALANCHEOUTLINER_API FAvaOutlinerObjectIconCustomization : public IAvaOutlinerIconCustomization
+class FAvaOutlinerObjectIconCustomization : public IAvaOutlinerIconCustomization
 {
 public:
-	FAvaOutlinerObjectIconCustomization(const UClass* InSupportedClass);
+	AVALANCHEOUTLINER_API FAvaOutlinerObjectIconCustomization(const UClass* InSupportedClass);
 
-	void SetOverriddenIcon(const FOnGetOverriddenObjectIcon& InOverriddenIcon);
+	AVALANCHEOUTLINER_API void SetOverriddenIcon(const FOnGetOverriddenObjectIcon& InOverriddenIcon);
 
-protected:
+private:
 	//~ Begin IAvaOutlinerIconCustomization
-	virtual FName GetOutlinerItemIdentifier() const override { return SupportedClassName; }
+	virtual FName GetOutlinerItemIdentifier() const override;
 	virtual bool HasOverrideIcon(TSharedPtr<const FAvaOutlinerItem> InOutlinerItem) const override;
 	virtual FSlateIcon GetOverrideIcon(TSharedPtr<const FAvaOutlinerItem> InOutlinerItem) const override;
 	//~ End IAvaOutlinerIconCustomization
 
-private:
 	FName SupportedClassName;
 
 	FOnGetOverriddenObjectIcon OnGetOverriddenIcon;
diff --git a/Engine/Plugins/Experimental/Avalanche/Source/AvalancheOutliner/Public/Item/AvaOutlinerTreeRoot.h b/Engine/Plugins/Experimental/Avalanche/Source/AvalancheOutliner/Public/Item/AvaOutlinerTreeRoot.h
deleted file mode 100644
