Motion Design:
- Trivial removal of API

#changelist validated
#virtualized

diff --git a/Engine/Plugins/Experimental/Avalanche/Source/AvalancheMedia/Public/Broadcast/IAvaBroadcastSettings.h b/Engine/Plugins/Experimental/Avalanche/Source/AvalancheMedia/Public/Broadcast/IAvaBroadcastSettings.h
--- a/Engine/Plugins/Experimental/Avalanche/Source/AvalancheMedia/Public/Broadcast/IAvaBroadcastSettings.h
+++ b/Engine/Plugins/Experimental/Avalanche/Source/AvalancheMedia/Public/Broadcast/IAvaBroadcastSettings.h
@@ -12,7 +12,7 @@ struct FSoftObjectPath;
  * Broadcast settings can either be local to the client, or if running
  * on a playback server, they may be replicated from a connected remote client.
  */
-class AVALANCHEMEDIA_API IAvaBroadcastSettings
+class IAvaBroadcastSettings
 {
 public:
 	virtual const FLinearColor& GetChannelClearColor() const = 0;
@@ -20,7 +20,7 @@ public:
 	virtual const FIntPoint& GetDefaultResolution() const = 0;
 	virtual bool IsDrawPlaceholderWidget() const = 0;
 	virtual const FSoftObjectPath& GetPlaceholderWidgetClass() const = 0;
-	
+
 protected:
 	virtual ~IAvaBroadcastSettings() = default;
 };
diff --git a/Engine/Plugins/Experimental/Avalanche/Source/AvalancheMedia/Public/IAvaMediaModule.h b/Engine/Plugins/Experimental/Avalanche/Source/AvalancheMedia/Public/IAvaMediaModule.h
--- a/Engine/Plugins/Experimental/Avalanche/Source/AvalancheMedia/Public/IAvaMediaModule.h
+++ b/Engine/Plugins/Experimental/Avalanche/Source/AvalancheMedia/Public/IAvaMediaModule.h
@@ -6,7 +6,7 @@
 #include "Modules/ModuleInterface.h"
 #include "Modules/ModuleManager.h"
 
-AVALANCHEMEDIA_API DECLARE_LOG_CATEGORY_EXTERN(LogAvaMedia, Log, All);
+DECLARE_LOG_CATEGORY_EXTERN(LogAvaMedia, Log, All);
 
 class FAvaPlaybackManager;
 class FAvaPlaybackServer;
@@ -33,7 +33,7 @@ enum class EAvaMediaMapChangeType : uint8
 	TearDownWorld
 };
 
-class AVALANCHEMEDIA_API IAvaMediaModule : public IModuleInterface
+class IAvaMediaModule : public IModuleInterface
 {
 public:
 	static bool IsModuleLoaded()
diff --git a/Engine/Plugins/Experimental/Avalanche/Source/AvalancheMedia/Public/ModularFeature/AvaMediaSyncProviderFeatureTypes.h b/Engine/Plugins/Experimental/Avalanche/Source/AvalancheMedia/Public/ModularFeature/AvaMediaSyncProviderFeatureTypes.h
--- a/Engine/Plugins/Experimental/Avalanche/Source/AvalancheMedia/Public/ModularFeature/AvaMediaSyncProviderFeatureTypes.h
+++ b/Engine/Plugins/Experimental/Avalanche/Source/AvalancheMedia/Public/ModularFeature/AvaMediaSyncProviderFeatureTypes.h
@@ -71,7 +71,7 @@ struct AVALANCHEMEDIA_API FAvaMediaSyncConnectionInfo
 
 /** Base sync response with Status and Error information */
 USTRUCT()
-struct AVALANCHEMEDIA_API FAvaMediaSyncResponse
+struct FAvaMediaSyncResponse
 {
 	GENERATED_BODY()
 
@@ -117,7 +117,7 @@ struct AVALANCHEMEDIA_API FAvaMediaSyncResponse
 
 /** Payload for a comparison request */
 USTRUCT()
-struct AVALANCHEMEDIA_API FAvaMediaSyncCompareResponse : public FAvaMediaSyncResponse
+struct FAvaMediaSyncCompareResponse : public FAvaMediaSyncResponse
 {
 	GENERATED_BODY()
 
diff --git a/Engine/Plugins/Experimental/Avalanche/Source/AvalancheMedia/Public/Playback/AvaPlaybackGraph.h b/Engine/Plugins/Experimental/Avalanche/Source/AvalancheMedia/Public/Playback/AvaPlaybackGraph.h
--- a/Engine/Plugins/Experimental/Avalanche/Source/AvalancheMedia/Public/Playback/AvaPlaybackGraph.h
+++ b/Engine/Plugins/Experimental/Avalanche/Source/AvalancheMedia/Public/Playback/AvaPlaybackGraph.h
@@ -12,7 +12,7 @@
 #include "UObject/SoftObjectPtr.h"
 #include "AvaPlaybackGraph.generated.h"
 
-AVALANCHEMEDIA_API DECLARE_LOG_CATEGORY_EXTERN(LogAvaPlayback, Log, All);
+DECLARE_LOG_CATEGORY_EXTERN(LogAvaPlayback, Log, All);
 
 class FAvaPlaybackManager;
 class IAvaPlaybackGraphEditor;
@@ -73,7 +73,7 @@ ENUM_CLASS_FLAGS(EAvaPlaybackUnloadOptions);
  * Used to group playables per channels, but can be extended to any conceptual grouping.
  */
 USTRUCT()
-struct AVALANCHEMEDIA_API FAvaPlaybackPlayableGroup
+struct FAvaPlaybackPlayableGroup
 {
 	GENERATED_BODY()
 
diff --git a/Engine/Plugins/Experimental/Avalanche/Source/AvalancheMedia/Public/Rundown/AvaRundownPage.h b/Engine/Plugins/Experimental/Avalanche/Source/AvalancheMedia/Public/Rundown/AvaRundownPage.h
--- a/Engine/Plugins/Experimental/Avalanche/Source/AvalancheMedia/Public/Rundown/AvaRundownPage.h
+++ b/Engine/Plugins/Experimental/Avalanche/Source/AvalancheMedia/Public/Rundown/AvaRundownPage.h
@@ -37,11 +37,10 @@ enum class EAvaRundownPageStatus : uint8
 };
 
 USTRUCT()
-struct AVALANCHEMEDIA_API FAvaRundownChannelPageStatus
+struct FAvaRundownChannelPageStatus
 {
 	GENERATED_BODY()
-	
-public:
+
 	UPROPERTY()
 	EAvaBroadcastChannelType Type = EAvaBroadcastChannelType::Program;
 	
@@ -59,7 +58,6 @@ struct AVALANCHEMEDIA_API FAvaRundownPage
 	GENERATED_BODY()
 
 public:
-	
 	static FAvaRundownPage NullPage;
 	static const int32 InvalidPageId;
 
