[Motion Design] Fixing redirect for M_AvalancheCheckerboard. (Trivial)

#rnx

diff --git a/Engine/Plugins/Experimental/Avalanche/Config/DefaultAvalanche.ini b/Engine/Plugins/Experimental/Avalanche/Config/DefaultAvalanche.ini
--- a/Engine/Plugins/Experimental/Avalanche/Config/DefaultAvalanche.ini
+++ b/Engine/Plugins/Experimental/Avalanche/Config/DefaultAvalanche.ini
@@ -268,7 +268,7 @@
 +EnumRedirects=(OldName="/Script/AvalancheMedia.EAvaRemoteControlResult",NewName="/Script/AvalancheMedia.EAvaPlayableRemoteControlResult")
 
 +PackageRedirects=(OldName="/Avalanche/EditorResources/M_AvalancheBackplate.M_AvalancheBackplate",NewName="/Avalanche/EditorResources/M_Backplate.M_Backplate")
-+PackageRedirects=(OldName="/Avalanche/EditorResources/M_AvalancheBackplate.M_AvalancheCheckerboard",NewName="/Avalanche/EditorResources/M_Checkerboard.M_Checkerboard")
++PackageRedirects=(OldName="/Avalanche/EditorResources/M_AvalancheCheckerboard.M_AvalancheCheckerboard",NewName="/Avalanche/EditorResources/M_Checkerboard.M_Checkerboard")
 +PackageRedirects=(OldName="/Avalanche/EditorResources/M_AvalancheChannelVisualizer_Red.M_AvalancheChannelVisualizer_Red",NewName="/Avalanche/EditorResources/M_ChannelVisualizer_Red.M_ChannelVisualizer_Red")
 +PackageRedirects=(OldName="/Avalanche/EditorResources/M_AvalancheChannelVisualizer_Green.M_AvalancheChannelVisualizer_Green",NewName="/Avalanche/EditorResources/M_ChannelVisualizer_Green.M_ChannelVisualizer_Green")
 +PackageRedirects=(OldName="/Avalanche/EditorResources/M_AvalancheChannelVisualizer_Blue.M_AvalancheChannelVisualizer_Blue",NewName="/Avalanche/EditorResources/M_ChannelVisualizer_Blue.M_ChannelVisualizer_Blue")
