Motion Design:
- Fixed default thumbnail path

#changelist validated
#virtualized

diff --git a/Engine/Plugins/Experimental/Avalanche/Source/AvalancheEditor/Private/AvaEditorModule.cpp b/Engine/Plugins/Experimental/Avalanche/Source/AvalancheEditor/Private/AvaEditorModule.cpp
--- a/Engine/Plugins/Experimental/Avalanche/Source/AvalancheEditor/Private/AvaEditorModule.cpp
+++ b/Engine/Plugins/Experimental/Avalanche/Source/AvalancheEditor/Private/AvaEditorModule.cpp
@@ -78,7 +78,7 @@ namespace UE::AvaEditor::Private
 {
 	static FString LevelTemplatesPath = FString::Printf(TEXT("/%hs/%s"), UE_PLUGIN_NAME, TEXT("LevelTemplates"));
 	static FString DefaultLevelPath = FString::Printf(TEXT("/%hs/%s"), UE_PLUGIN_NAME, TEXT("DefaultMotionDesignLevel"));
-	static FString DefaultLevelThumbnailPath = FString::Printf(TEXT("/%hs/%s"), UE_PLUGIN_NAME, TEXT("DefaultMotionDesignLevelThumbnail"));
+	static FString DefaultLevelThumbnailPath = FString::Printf(TEXT("/%hs/%s"), UE_PLUGIN_NAME, TEXT("DefaultMotionDesignLevelThumbnail.DefaultMotionDesignLevelThumbnail"));
 }
 
 void FAvaEditorModule::StartupModule()
