Motion Design:
- Updated Toolbox Ed Mode Icon

#changelist validated
#virtualized

diff --git a/Engine/Plugins/Experimental/Avalanche/Source/AvalancheInteractiveTools/Private/AvaInteractiveToolsStyle.cpp b/Engine/Plugins/Experimental/Avalanche/Source/AvalancheInteractiveTools/Private/AvaInteractiveToolsStyle.cpp
--- a/Engine/Plugins/Experimental/Avalanche/Source/AvalancheInteractiveTools/Private/AvaInteractiveToolsStyle.cpp
+++ b/Engine/Plugins/Experimental/Avalanche/Source/AvalancheInteractiveTools/Private/AvaInteractiveToolsStyle.cpp
@@ -21,6 +21,8 @@ FAvaInteractiveToolsStyle::FAvaInteractiveToolsStyle()
 	SetContentRoot(Plugin->GetBaseDir() / TEXT("Resources"));
 	SetCoreContentRoot(FPaths::EngineContentDir() / TEXT("Editor/Slate"));
 
+	Set("Icons.Toolbox", new IMAGE_BRUSH(TEXT("Icons/ToolboxIcons/toolbox"), Icon16x16));
+
 	// Categories
 	Set("AvaInteractiveTools.Category_2D",     new IMAGE_BRUSH("Icons/ToolboxIcons/rectangle", Icon20x20));
 	Set("AvaInteractiveTools.Category_3D",     new IMAGE_BRUSH("Icons/ToolboxIcons/cube", Icon20x20));
diff --git a/Engine/Plugins/Experimental/Avalanche/Source/AvalancheInteractiveTools/Private/EdMode/AvaInteractiveToolsEdMode.cpp b/Engine/Plugins/Experimental/Avalanche/Source/AvalancheInteractiveTools/Private/EdMode/AvaInteractiveToolsEdMode.cpp
--- a/Engine/Plugins/Experimental/Avalanche/Source/AvalancheInteractiveTools/Private/EdMode/AvaInteractiveToolsEdMode.cpp
+++ b/Engine/Plugins/Experimental/Avalanche/Source/AvalancheInteractiveTools/Private/EdMode/AvaInteractiveToolsEdMode.cpp
@@ -2,10 +2,11 @@
 
 #include "AvaInteractiveToolsEdMode.h"
 #include "AvaInteractiveToolsEdModeToolkit.h"
+#include "AvaInteractiveToolsStyle.h"
 #include "ContextObjectStore.h"
+#include "EdModeInteractiveToolsContext.h"
 #include "EditorModeManager.h"
 #include "EditorModes.h"
-#include "EdModeInteractiveToolsContext.h"
 #include "GameFramework/InputSettings.h"
 #include "IAvalancheInteractiveToolsModule.h"
 #include "Selection.h"
@@ -17,7 +18,7 @@ UAvaInteractiveToolsEdMode::UAvaInteractiveToolsEdMode()
 {
 	Info = FEditorModeInfo(IAvalancheInteractiveToolsModule::EM_AvaInteractiveToolsEdModeId,
 		LOCTEXT("AvaInteractiveToolsEdModeName", "Motion Design"),
-		FSlateIcon(FAppStyle::GetAppStyleSetName(), TEXT("LevelEditor.Tabs.EditorModes")),
+		FSlateIcon(FAvaInteractiveToolsStyle::Get().GetStyleSetName(), TEXT("Icons.Toolbox")),
 		true);
 
 	LastActiveTool = "";
