Motion Design
- Background post process effect now works correctly at 100% dpi scaling.

diff --git a/Engine/Plugins/Experimental/Avalanche/Source/AvalancheViewport/Private/Visualizers/AvaViewportBackgroundVisualizer.cpp b/Engine/Plugins/Experimental/Avalanche/Source/AvalancheViewport/Private/Visualizers/AvaViewportBackgroundVisualizer.cpp
--- a/Engine/Plugins/Experimental/Avalanche/Source/AvalancheViewport/Private/Visualizers/AvaViewportBackgroundVisualizer.cpp
+++ b/Engine/Plugins/Experimental/Avalanche/Source/AvalancheViewport/Private/Visualizers/AvaViewportBackgroundVisualizer.cpp
@@ -114,8 +114,8 @@ void FAvaViewportBackgroundVisualizer::UpdateForViewport(const FAvaVisibleArea&
 	const float WidgetAspectRatio = InWidgetSize.X / InWidgetSize.Y;
 	const float ViewportAspectRatio = InVisibleArea.AbsoluteSize.X / InVisibleArea.AbsoluteSize.Y;
 	const FVector2f WidgetBasedScale = InVisibleArea.AbsoluteSize / InWidgetSize;
-
-	const FVector2f Scale = WidgetBasedScale / InVisibleArea.GetVisibleAreaFraction();
+	const float VisibleAreaFraction = InVisibleArea.GetVisibleAreaFraction();
+	const FVector2f Scale = WidgetBasedScale / VisibleAreaFraction;
 
 	if (!FMath::IsNearlyEqual(TextureScale.X, Scale.X)
 		|| !FMath::IsNearlyEqual(TextureScale.Y, Scale.Y))
@@ -127,21 +127,30 @@ void FAvaViewportBackgroundVisualizer::UpdateForViewport(const FAvaVisibleArea&
 	}
 
 	FVector2f Offset = ((InWidgetSize - InVisibleArea.AbsoluteSize) * 0.5f)
-		+ InCameraOffset / InVisibleArea.GetVisibleAreaFraction();
+		+ InCameraOffset * Scale;
+
+	if (InVisibleArea.IsZoomedView())
+	{
+		Offset -= InVisibleArea.GetInvisibleSize() / InVisibleArea.VisibleSize * InVisibleArea.AbsoluteSize * 0.5f;
+	}
 
 	if (!FMath::IsNearlyEqual(WidgetAspectRatio, ViewportAspectRatio))
 	{
 		if (WidgetAspectRatio > ViewportAspectRatio)
 		{
-			Offset.X -= (Scale.Y - 1) * InVisibleArea.AbsoluteSize.X * 0.5f;
-			Offset.Y -= (Scale.Y - 1) * InVisibleArea.AbsoluteSize.Y * 0.5f;
+			const float Distance = InWidgetSize.X - InVisibleArea.AbsoluteSize.X;
+			const float Scalar = InVisibleArea.VisibleSize.X;
+			const float OffsetX = Distance * InCameraOffset.X / Scalar;
+			Offset.X += OffsetX * Scale.X * InVisibleArea.GetVisibleAreaFraction();
 		}
 		else
 		{
-			Offset.X -= (Scale.X - 1) * InVisibleArea.AbsoluteSize.X * 0.5f;
-			Offset.Y -= (Scale.X - 1) * InVisibleArea.AbsoluteSize.Y * 0.5f;
+			const float Distance = InWidgetSize.Y - InVisibleArea.AbsoluteSize.Y;
+			const float Scalar = InVisibleArea.VisibleSize.Y;
+			const float OffsetY = Distance * InCameraOffset.Y / Scalar;
+			Offset.Y += OffsetY * Scale.Y * InVisibleArea.GetVisibleAreaFraction();
 		}
-	}		
+	}
 
 	if (!FMath::IsNearlyEqual(TextureOffset.X, Offset.X)
 		|| !FMath::IsNearlyEqual(TextureOffset.Y, Offset.Y))
