[AvalancheEditor]
* Fixed clang editor linker error by adding exports to some virtual functions

diff --git a/Engine/Plugins/Experimental/Avalanche/Source/AvalancheOutliner/Public/Icon/AvaOutlinerObjectIconCustomization.h b/Engine/Plugins/Experimental/Avalanche/Source/AvalancheOutliner/Public/Icon/AvaOutlinerObjectIconCustomization.h
--- a/Engine/Plugins/Experimental/Avalanche/Source/AvalancheOutliner/Public/Icon/AvaOutlinerObjectIconCustomization.h
+++ b/Engine/Plugins/Experimental/Avalanche/Source/AvalancheOutliner/Public/Icon/AvaOutlinerObjectIconCustomization.h
@@ -20,9 +20,9 @@ public:
 
 private:
 	//~ Begin IAvaOutlinerIconCustomization
-	virtual FName GetOutlinerItemIdentifier() const override;
-	virtual bool HasOverrideIcon(TSharedPtr<const FAvaOutlinerItem> InOutlinerItem) const override;
-	virtual FSlateIcon GetOverrideIcon(TSharedPtr<const FAvaOutlinerItem> InOutlinerItem) const override;
+	AVALANCHEOUTLINER_API virtual FName GetOutlinerItemIdentifier() const override;
+	AVALANCHEOUTLINER_API virtual bool HasOverrideIcon(TSharedPtr<const FAvaOutlinerItem> InOutlinerItem) const override;
+	AVALANCHEOUTLINER_API virtual FSlateIcon GetOverrideIcon(TSharedPtr<const FAvaOutlinerItem> InOutlinerItem) const override;
 	//~ End IAvaOutlinerIconCustomization
 
 	FName SupportedClassName;
