Motion Design:
- Fix linux server build by adding server to target deny list in ava runtime modules (currently not supported for server)

diff --git a/Engine/Plugins/Experimental/Avalanche/Avalanche.uplugin b/Engine/Plugins/Experimental/Avalanche/Avalanche.uplugin
--- a/Engine/Plugins/Experimental/Avalanche/Avalanche.uplugin
+++ b/Engine/Plugins/Experimental/Avalanche/Avalanche.uplugin
@@ -24,12 +24,14 @@
 		{
 			"Name": "Avalanche",
 			"Type": "RuntimeAndProgram",
-			"LoadingPhase": "Default"
+			"LoadingPhase": "Default",
+			"TargetDenyList": [ "Server" ]
 		},
 		{
 			"Name": "AvalancheCore",
 			"Type": "Runtime",
-			"LoadingPhase": "Default"
+			"LoadingPhase": "Default",
+			"TargetDenyList": [ "Server" ]
 		},
 		{
 			"Name": "AvalancheEditor",
@@ -49,7 +51,8 @@
 		{
 			"Name": "AvalancheEffectors",
 			"Type": "RuntimeAndProgram",
-			"LoadingPhase": "Default"
+			"LoadingPhase": "Default",
+			"TargetDenyList": [ "Server" ]
 		},
 		{
 			"Name": "AvalancheEffectorsEditor",
@@ -64,7 +67,8 @@
 		{
 			"Name": "AvalancheInteractiveToolsRuntime",
 			"Type": "Runtime",
-			"LoadingPhase": "Default"
+			"LoadingPhase": "Default",
+			"TargetDenyList": [ "Server" ]
 		},
 		{
 			"Name": "AvalancheLevelViewport",
@@ -79,7 +83,8 @@
 		{
 			"Name": "AvalancheMask",
 			"Type": "Runtime",
-			"LoadingPhase": "Default"
+			"LoadingPhase": "Default",
+			"TargetDenyList": [ "Server" ]
 		},
 		{
 			"Name": "AvalancheMaskEditor",
@@ -89,7 +94,8 @@
 		{
 			"Name": "AvalancheMedia",
 			"Type": "Runtime",
-			"LoadingPhase": "PostConfigInit"
+			"LoadingPhase": "PostConfigInit",
+			"TargetDenyList": [ "Server" ]
 		},
 		{
 			"Name": "AvalancheMediaEditor",
@@ -99,7 +105,8 @@
 		{
 			"Name": "AvalancheModifiers",
 			"Type": "Runtime",
-			"LoadingPhase": "Default"
+			"LoadingPhase": "Default",
+			"TargetDenyList": [ "Server" ]
 		},
 		{
 			"Name": "AvalancheModifiersEditor",
@@ -109,7 +116,8 @@
 		{
 			"Name": "AvalanchePropertyAnimator",
 			"Type": "Runtime",
-			"LoadingPhase": "Default"
+			"LoadingPhase": "Default",
+			"TargetDenyList": [ "Server" ]
 		},
 		{
 			"Name": "AvalanchePropertyAnimatorEditor",
@@ -119,7 +127,8 @@
 		{
 			"Name": "AvalancheRemoteControl",
 			"Type": "Runtime",
-			"LoadingPhase": "Default"
+			"LoadingPhase": "Default",
+			"TargetDenyList": [ "Server" ]
 		},
 		{
 			"Name": "AvalancheRemoteControlEditor",
@@ -129,7 +138,8 @@
 		{
 			"Name": "AvalancheSequence",
 			"Type": "Runtime",
-			"LoadingPhase": "Default"
+			"LoadingPhase": "Default",
+			"TargetDenyList": [ "Server" ]
 		},
 		{
 			"Name": "AvalancheSequencer",
@@ -139,7 +149,8 @@
 		{
 			"Name": "AvalancheShapes",
 			"Type": "RuntimeAndProgram",
-			"LoadingPhase": "Default"
+			"LoadingPhase": "Default",
+			"TargetDenyList": [ "Server" ]
 		},
 		{
 			"Name": "AvalancheShapesEditor",
@@ -149,7 +160,8 @@
 		{
 			"Name": "AvalancheText",
 			"Type": "RuntimeAndProgram",
-			"LoadingPhase": "Default"
+			"LoadingPhase": "Default",
+			"TargetDenyList": [ "Server" ]
 		},
 		{
 			"Name": "AvalancheTextEditor",
@@ -164,22 +176,26 @@
 		{
 			"Name": "AvalancheSceneTree",
 			"Type": "Runtime",
-			"LoadingPhase": "Default"
+			"LoadingPhase": "Default",
+			"TargetDenyList": [ "Server" ]
 		},
 		{
 			"Name": "AvalancheTransition",
 			"Type": "Runtime",
-			"LoadingPhase": "Default"
+			"LoadingPhase": "Default",
+			"TargetDenyList": [ "Server" ]
 		},
 		{
 			"Name": "AvalancheTransitionEditor",
 			"Type": "UncookedOnly",
-			"LoadingPhase": "Default"
+			"LoadingPhase": "Default",
+			"TargetDenyList": [ "Server" ]
 		},
 		{
 			"Name": "AvalancheMRQ",
 			"Type": "Runtime",
-			"LoadingPhase": "Default"
+			"LoadingPhase": "Default",
+			"TargetDenyList": [ "Server" ]
 		},
 		{
 			"Name": "AvalancheMRQEditor",
@@ -194,7 +210,8 @@
 		{
 			"Name": "AvalancheTag",
 			"Type": "Runtime",
-			"LoadingPhase": "Default"
+			"LoadingPhase": "Default",
+			"TargetDenyList": [ "Server" ]
 		},
 		{
 			"Name": "AvalancheTagEditor",
@@ -204,7 +221,8 @@
 		{
 			"Name": "AvalancheAttribute",
 			"Type": "Runtime",
-			"LoadingPhase": "Default"
+			"LoadingPhase": "Default",
+			"TargetDenyList": [ "Server" ]
 		},
 		{
 			"Name": "AvalancheAttributeEditor",
diff --git a/Engine/Plugins/Experimental/PropertyAnimator/PropertyAnimator.uplugin b/Engine/Plugins/Experimental/PropertyAnimator/PropertyAnimator.uplugin
--- a/Engine/Plugins/Experimental/PropertyAnimator/PropertyAnimator.uplugin
+++ b/Engine/Plugins/Experimental/PropertyAnimator/PropertyAnimator.uplugin
@@ -17,7 +17,8 @@
 		{
 			"Name": "PropertyAnimator",
 			"Type": "Runtime",
-			"LoadingPhase": "Default"
+			"LoadingPhase": "Default",
+         	"TargetDenyList": [ "Server" ]
 		},
 		{
 			"Name": "PropertyAnimatorEditor",
diff --git a/Engine/Plugins/Experimental/StormSyncAvalancheBridge/StormSyncAvalancheBridge.uplugin b/Engine/Plugins/Experimental/StormSyncAvalancheBridge/StormSyncAvalancheBridge.uplugin
--- a/Engine/Plugins/Experimental/StormSyncAvalancheBridge/StormSyncAvalancheBridge.uplugin
+++ b/Engine/Plugins/Experimental/StormSyncAvalancheBridge/StormSyncAvalancheBridge.uplugin
@@ -23,7 +23,8 @@
 		{
 			"Name": "StormSyncAvaBridge",
 			"Type": "Runtime",
-			"LoadingPhase": "Default"
+			"LoadingPhase": "Default",
+			"TargetDenyList": [ "Server" ]
 		}
 	],
 	"Plugins": [
