Motion Design:
- Fix crash happening due to entity being incorrectly cast without prior type checking

diff --git a/Engine/Plugins/Experimental/Avalanche/Source/AvalancheMedia/Private/Playable/AvaPlayableRemoteControl.cpp b/Engine/Plugins/Experimental/Avalanche/Source/AvalancheMedia/Private/Playable/AvaPlayableRemoteControl.cpp
--- a/Engine/Plugins/Experimental/Avalanche/Source/AvalancheMedia/Private/Playable/AvaPlayableRemoteControl.cpp
+++ b/Engine/Plugins/Experimental/Avalanche/Source/AvalancheMedia/Private/Playable/AvaPlayableRemoteControl.cpp
@@ -13,11 +13,28 @@
 #include "IRemoteControlModule.h"
 #include "RCVirtualProperty.h"
 #include "RemoteControlPropertyIdRegistry.h"
+#include "Templates/CopyQualifiersFromTo.h"
 
 DEFINE_LOG_CATEGORY(LogAvaPlayableRemoteControl);
 
 namespace UE::AvaPlayableRemoteControl::Private
 {
+	/**
+	 * Safely casts a Remote Control Entity from one Type to the other
+	 * @param InEntity the entity to cast
+	 * @return the entity casted to the desired type if it was derived from it. nullptr otherwise
+	 */
+	template<typename To, typename From>
+	TSharedPtr<typename TCopyQualifiersFromTo<From, To>::Type> CastEntity(const TSharedPtr<From>& InEntity)
+	{
+		const UScriptStruct* EntityStruct = InEntity.IsValid() ? InEntity->GetStruct() : nullptr;
+		if (EntityStruct && EntityStruct->IsChildOf(To::StaticStruct()))
+		{
+			return StaticCastSharedPtr<typename TCopyQualifiersFromTo<From, To>::Type>(InEntity);
+		}
+		return nullptr;
+	}
+
 	bool ResolveObjectPropertyForReadOnly(UObject* Object, FRCFieldPathInfo PropertyPath, FRCObjectReference& OutObjectRef, FString& OutErrorText)
 	{
 		if (!Object)
@@ -94,7 +111,7 @@ EAvaPlayableRemoteControlResult UE::AvaPlayableRemoteControl::GetValueOfEntity(c
 	using namespace UE::AvaPlayableRemoteControl::Private;
 	OutValue.Reset();
 
-	TSharedPtr<const FRemoteControlProperty> Field = StaticCastSharedPtr<const FRemoteControlProperty>(RemoteControlEntity);
+	TSharedPtr<const FRemoteControlProperty> Field = CastEntity<FRemoteControlProperty>(RemoteControlEntity);
 	if (!Field.IsValid())
 	{
 		return EAvaPlayableRemoteControlResult::InvalidParameter;
@@ -122,7 +139,7 @@ EAvaPlayableRemoteControlResult UE::AvaPlayableRemoteControl::SetValueOfEntity(c
 {
 	using namespace UE::AvaPlayableRemoteControl::Private;
 
-	TSharedPtr<FRemoteControlProperty> Field = StaticCastSharedPtr<FRemoteControlProperty>(RemoteControlEntity);
+	TSharedPtr<FRemoteControlProperty> Field = CastEntity<FRemoteControlProperty>(RemoteControlEntity);
 	if (!Field.IsValid())
 	{
 		return EAvaPlayableRemoteControlResult::InvalidParameter;
