MotionDesign : Translucent priority modifier
- Fix translucent priority modifier from updating too often and affecting performance in rundown

diff --git a/Engine/Plugins/Experimental/Avalanche/Source/AvalancheModifiers/Private/Extensions/AvaRenderStateUpdateModifierExtension.cpp b/Engine/Plugins/Experimental/Avalanche/Source/AvalancheModifiers/Private/Extensions/AvaRenderStateUpdateModifierExtension.cpp
--- a/Engine/Plugins/Experimental/Avalanche/Source/AvalancheModifiers/Private/Extensions/AvaRenderStateUpdateModifierExtension.cpp
+++ b/Engine/Plugins/Experimental/Avalanche/Source/AvalancheModifiers/Private/Extensions/AvaRenderStateUpdateModifierExtension.cpp
@@ -94,8 +94,10 @@ void FAvaRenderStateUpdateModifierExtension::OnExtensionDisabled(EActorModifierC
 
 void FAvaRenderStateUpdateModifierExtension::OnRenderStateDirty(UActorComponent& InComponent)
 {
+	const AActor* ModifierActor = GetModifierActor();
 	AActor* ActorDirty = InComponent.GetOwner();
-	if (!ActorDirty)
+
+	if (!ActorDirty || !ModifierActor)
 	{
 		return;
 	}
@@ -106,7 +108,7 @@ void FAvaRenderStateUpdateModifierExtension::OnRenderStateDirty(UActorComponent&
 		return;
 	}
 
-	if (ActorDirty->GetWorld() != GetModifierWorld())
+	if (ActorDirty->GetLevel() != ModifierActor->GetLevel())
 	{
 		return;
 	}
diff --git a/Engine/Plugins/Experimental/Avalanche/Source/AvalancheModifiers/Private/Modifiers/AvaTranslucentPriorityModifier.cpp b/Engine/Plugins/Experimental/Avalanche/Source/AvalancheModifiers/Private/Modifiers/AvaTranslucentPriorityModifier.cpp
--- a/Engine/Plugins/Experimental/Avalanche/Source/AvalancheModifiers/Private/Modifiers/AvaTranslucentPriorityModifier.cpp
+++ b/Engine/Plugins/Experimental/Avalanche/Source/AvalancheModifiers/Private/Modifiers/AvaTranslucentPriorityModifier.cpp
@@ -160,6 +160,8 @@ void UAvaTranslucentPriorityModifier::Apply()
 	const int32 GlobalOffset = SharedObject->GetSortPriorityOffset();
 	const int32 GlobalStep = SharedObject->GetSortPriorityStep();
 
+	LastSortPriorities.Empty(CachedSortedComponentStates.Num());
+
 	if (Mode == EAvaTranslucentPriorityModifierMode::Manual)
 	{
 		int32 TranslucentSortPriority = GlobalOffset + SortPriority;
@@ -171,6 +173,7 @@ void UAvaTranslucentPriorityModifier::Apply()
 			{
 				if (UPrimitiveComponent* Component = SortedComponentState->PrimitiveComponentWeak.Get())
 				{
+					LastSortPriorities.Add(Component, TranslucentSortPriority);
 					Component->SetTranslucentSortPriority(TranslucentSortPriority);
 
 					TranslucentSortPriority += GlobalStep;
@@ -198,15 +201,16 @@ void UAvaTranslucentPriorityModifier::Apply()
 				continue;
 			}
 
-			// This modifier handles this component
-			if (ComponentModifier == this)
-			{
-				Component->SetTranslucentSortPriority(TranslucentSortPriority);
-			}
-			// Another modifier handles this component
-			else
+			if (Component->TranslucencySortPriority != TranslucentSortPriority)
 			{
-				if (Component->TranslucencySortPriority != TranslucentSortPriority)
+				// This modifier handles this component
+				if (ComponentModifier == this)
+				{
+					LastSortPriorities.Add(Component, TranslucentSortPriority);
+					Component->SetTranslucentSortPriority(TranslucentSortPriority);
+				}
+				// Another modifier handles this component
+				else
 				{
 					// Cache to avoid doing the same query for the same result
 					ComponentModifier->CachedSortedComponentStates = CachedSortedComponentStates;
@@ -408,8 +412,11 @@ void UAvaTranslucentPriorityModifier::OnRenderStateUpdated(AActor* InActor, UAct
 	Super::OnRenderStateUpdated(InActor, InComponent);
 
 	const AActor* const ActorModified = GetModifiedActor();
+	UPrimitiveComponent* PrimitiveComponent = Cast<UPrimitiveComponent>(InComponent);
 
-	if (!IsValid(ActorModified))
+	if (!IsValid(ActorModified)
+		|| !IsValid(InActor)
+		|| !IsValid(PrimitiveComponent))
 	{
 		return;
 	}
@@ -419,6 +426,11 @@ void UAvaTranslucentPriorityModifier::OnRenderStateUpdated(AActor* InActor, UAct
 		return;
 	}
 
+	if (bIncludeChildren && InActor != ActorModified && !InActor->IsAttachedTo(ActorModified))
+	{
+		return;
+	}
+
 	if (bIncludeChildren && !ChildrenActorsWeak.Contains(InActor))
 	{
 		const UAvaTranslucentPriorityModifierShared* SharedObject = GetShared<UAvaTranslucentPriorityModifierShared>(false);
@@ -427,15 +439,18 @@ void UAvaTranslucentPriorityModifier::OnRenderStateUpdated(AActor* InActor, UAct
 			return;
 		}
 
-		UPrimitiveComponent* PrimitiveComponent = Cast<UPrimitiveComponent>(InComponent);
-		if (!PrimitiveComponent)
+		// If the component is already handled by another modifier, return early
+		const UAvaTranslucentPriorityModifier* OtherModifier = SharedObject->FindModifierContext(PrimitiveComponent);
+		if (OtherModifier && OtherModifier != this)
 		{
 			return;
 		}
+	}
 
-		// If the component is already handled by another modifier, return early
-		const UAvaTranslucentPriorityModifier* OtherModifier = SharedObject->FindModifierContext(PrimitiveComponent);
-		if (OtherModifier && OtherModifier != this)
+	// In case the sort priority hasn't changed do not update
+	if (const int32* LastSortPriority = LastSortPriorities.Find(PrimitiveComponent))
+	{
+		if (PrimitiveComponent->TranslucencySortPriority == *LastSortPriority)
 		{
 			return;
 		}
@@ -450,7 +465,9 @@ void UAvaTranslucentPriorityModifier::OnTransformUpdated(AActor* InActor, bool b
 
 	const AActor* const ActorModified = GetModifiedActor();
 
+	// Only AutoCameraDistance mode relies on actor location to compute sort priority
 	if (!IsValid(ActorModified)
+		|| Mode != EAvaTranslucentPriorityModifierMode::AutoCameraDistance
 		|| (!bIncludeChildren && InActor != ActorModified)
 		|| (bIncludeChildren && !ChildrenActorsWeak.Contains(InActor)))
 	{
diff --git a/Engine/Plugins/Experimental/Avalanche/Source/AvalancheModifiers/Public/Modifiers/AvaTranslucentPriorityModifier.h b/Engine/Plugins/Experimental/Avalanche/Source/AvalancheModifiers/Public/Modifiers/AvaTranslucentPriorityModifier.h
--- a/Engine/Plugins/Experimental/Avalanche/Source/AvalancheModifiers/Public/Modifiers/AvaTranslucentPriorityModifier.h
+++ b/Engine/Plugins/Experimental/Avalanche/Source/AvalancheModifiers/Public/Modifiers/AvaTranslucentPriorityModifier.h
@@ -141,6 +141,9 @@ private:
 	UPROPERTY()
 	TMap<TWeakObjectPtr<UPrimitiveComponent>, int32> PreviousSortPriorities;
 
+	/** Last primitive components assigned sort priority, used for comparison on change */
+	TMap<FObjectKey, int32> LastSortPriorities;
+
 	/** Used to avoid querying again the full list of component states */
 	TArray<const FAvaTranslucentPriorityModifierComponentState*> CachedSortedComponentStates;
 };
