Avalanche
- Return of the texture overlay

diff --git a/Engine/Plugins/Experimental/Avalanche/Source/AvalancheLevelViewport/Private/AvaLevelViewportCommands.cpp b/Engine/Plugins/Experimental/Avalanche/Source/AvalancheLevelViewport/Private/AvaLevelViewportCommands.cpp
--- a/Engine/Plugins/Experimental/Avalanche/Source/AvalancheLevelViewport/Private/AvaLevelViewportCommands.cpp
+++ b/Engine/Plugins/Experimental/Avalanche/Source/AvalancheLevelViewport/Private/AvaLevelViewportCommands.cpp
@@ -37,13 +37,13 @@ void FAvaLevelViewportCommands::RegisterViewportCommands()
 	UI_COMMAND(ToggleBoundingBoxes
 		, "Toggle Bounding Boxes"
 		, "Toggles the display of selected actor bounding boxes."
-		, EUserInterfaceActionType::Check
+		, EUserInterfaceActionType::ToggleButton
 		, FInputChord(EKeys::B));
 
 	UI_COMMAND(ToggleIsolateActors
 		, "Isolate Selected Actors"
 		, "Changes the viewport and outliner to only show the selected actors. Also removes non-selected actors from snapping consideration. Changing the actor selection will not change which actors are isolated. Camera Preview Viewport Cameras are always visible (and their associated Canvas.)"
-		, EUserInterfaceActionType::Button
+		, EUserInterfaceActionType::ToggleButton
 		, FInputChord(EKeys::Q, EModifierKey::Alt))
 
 	UI_COMMAND(ToggleSafeFrames
@@ -105,6 +105,12 @@ void FAvaLevelViewportCommands::RegisterViewportCommands()
 		, "Switch to the checkerboard post process filter."
 		, EUserInterfaceActionType::Check
 		, FInputChord())
+
+	UI_COMMAND(ToggleTextureOverlay
+		, "Toggle Texture Overlay"
+		, "Turns on and off the texture overlay."
+		, EUserInterfaceActionType::ToggleButton
+		, FInputChord())
 }
 
 void FAvaLevelViewportCommands::RegisterCameraCommands()
diff --git a/Engine/Plugins/Experimental/Avalanche/Source/AvalancheLevelViewport/Private/SAvaLevelViewport.cpp b/Engine/Plugins/Experimental/Avalanche/Source/AvalancheLevelViewport/Private/SAvaLevelViewport.cpp
--- a/Engine/Plugins/Experimental/Avalanche/Source/AvalancheLevelViewport/Private/SAvaLevelViewport.cpp
+++ b/Engine/Plugins/Experimental/Avalanche/Source/AvalancheLevelViewport/Private/SAvaLevelViewport.cpp
@@ -26,6 +26,7 @@
 #include "Widgets/Layout/SBox.h"
 #include "Widgets/SAvaLevelViewportCameraBounds.h"
 #include "Widgets/SAvaLevelViewportGuide.h"
+#include "Widgets/SAvaLevelViewportTextureOverlay.h"
 #include "Widgets/SAvaLevelViewportPixelGrid.h"
 #include "Widgets/SAvaLevelViewportSafeFrames.h"
 #include "Widgets/SAvaLevelViewportScreenGrid.h"
@@ -157,6 +158,11 @@ void SAvaLevelViewport::BindCommands()
 		FCanExecuteAction::CreateSP(this, &SAvaLevelViewport::CanTogglePostProcessCheckerboard),
 		FIsActionChecked::CreateSP(this, &SAvaLevelViewport::IsPostProcessCheckerboardEnabled)
 	);
+	CommandListRef.MapAction(
+		AvaLevelViewportCommands.ToggleTextureOverlay,
+		FExecuteAction::CreateSP(this, &SAvaLevelViewport::ExecuteToggleTextureOverlay),
+		FCanExecuteAction::CreateSP(this, &SAvaLevelViewport::CanToggleTextureOverlay)
+	);
 
 	// Grid
 	CommandListRef.MapAction(
@@ -357,6 +363,11 @@ void SAvaLevelViewport::PopulateViewportOverlays(TSharedRef<SOverlay> Overlay)
 			SAssignNew(PixelGrid, SAvaLevelViewportPixelGrid, ViewportClient.ToSharedRef())
 		];
 
+		Overlay->AddSlot(-2)
+		[
+			SAssignNew(TextureOverlay, SAvaLevelViewportTextureOverlay, ViewportClient.ToSharedRef())
+		];
+
 		Overlay->AddSlot(-2)
 		[
 			SAssignNew(CameraBounds, SAvaLevelViewportCameraBounds, ViewportClient.ToSharedRef())
@@ -649,6 +660,11 @@ void SAvaLevelViewport::ApplySettings(const UAvaViewportSettings* InSettings)
 		SnapIndicators->SetVisibility(InSettings->bEnableViewportOverlay && InSettings->bSnapIndicatorsEnabled ? EVisibility::HitTestInvisible : EVisibility::Collapsed);
 	}
 
+	if (TextureOverlay.IsValid())
+	{
+		TextureOverlay->SetVisibility(InSettings->bEnableTextureOverlay ? EVisibility::HitTestInvisible : EVisibility::Collapsed);
+	}
+
 	if (CameraBounds.IsValid())
 	{
 		CameraBounds->SetVisibility(InSettings->bEnableViewportOverlay ? EVisibility::HitTestInvisible : EVisibility::Collapsed);
@@ -948,6 +964,28 @@ void SAvaLevelViewport::UnregisterPanelExtension()
 	PanelExtensionSubsystem->UnregisterPanelFactory(PanelExtensionFactory.Identifier, TEXT("LevelViewportToolBar.RightExtension"));
 }
 
+ECheckBoxState SAvaLevelViewport::GetTextureOverlayStretchEnabledCheckBoxState() const
+{
+	if (const UAvaViewportSettings* AvaViewportSettings = GetDefault<UAvaViewportSettings>())
+	{
+		return AvaViewportSettings->bTextureOverlayStretch
+			? ECheckBoxState::Checked
+			: ECheckBoxState::Unchecked;
+	}
+
+	return ECheckBoxState::Undetermined;
+}
+
+void SAvaLevelViewport::OnTextureOverlayStretchEnabledCheckBoxChanged(ECheckBoxState InState)
+{
+	if (UAvaViewportSettings* AvaViewportSettings = GetMutableDefault<UAvaViewportSettings>())
+	{
+		AvaViewportSettings->bTextureOverlayStretch = (InState == ECheckBoxState::Checked);
+		AvaViewportSettings->BroadcastSettingChanged(GET_MEMBER_NAME_CHECKED(UAvaViewportSettings, bTextureOverlayStretch));
+		AvaViewportSettings->SaveConfig();
+	}
+}
+
 int32 SAvaLevelViewport::GetVirtualSizeX() const
 {
 	return GetVirtualSize().X;
diff --git a/Engine/Plugins/Experimental/Avalanche/Source/AvalancheLevelViewport/Private/SAvaLevelViewport.h b/Engine/Plugins/Experimental/Avalanche/Source/AvalancheLevelViewport/Private/SAvaLevelViewport.h
--- a/Engine/Plugins/Experimental/Avalanche/Source/AvalancheLevelViewport/Private/SAvaLevelViewport.h
+++ b/Engine/Plugins/Experimental/Avalanche/Source/AvalancheLevelViewport/Private/SAvaLevelViewport.h
@@ -9,6 +9,7 @@
 #include "Subsystems/PanelExtensionSubsystem.h"
 #include "Templates/SharedPointer.h"
 #include "UObject/WeakObjectPtrFwd.h"
+#include "Widgets/SAvaLevelViewportTextureOverlay.h"
 
 class ACameraActor;
 class FAvaLevelViewportClient;
@@ -22,6 +23,7 @@ class SAvaLevelViewportPixelGrid;
 class SAvaLevelViewportSafeFrames;
 class SAvaLevelViewportScreenGrid;
 class SAvaLevelViewportSnapIndicators;
+class SAvaLevelViewportTextureOverlay;
 class SCanvas;
 class SOverlay;
 class SWidget;
@@ -153,6 +155,9 @@ public:
 	bool IsPostProcessCheckerboardEnabled() const;
 	bool CanTogglePostProcessCheckerboard() const;
 	void ExecuteTogglePostProcessCheckerboard();
+
+	bool CanToggleTextureOverlay() const;
+	void ExecuteToggleTextureOverlay();
 	
 	FIntPoint GetVirtualSize() const;
 	bool IsVirtualSizeActive(FIntPoint InVirtualSize) const;
@@ -183,6 +188,7 @@ public:
 private:
 	TWeakPtr<SAvaLevelViewportFrame> ViewportFrameWeak;
 	TSharedPtr<SAvaLevelViewportCameraBounds> CameraBounds;
+	TSharedPtr<SAvaLevelViewportTextureOverlay> TextureOverlay;
 	TSharedPtr<SAvaLevelViewportPixelGrid> PixelGrid;
 	TSharedPtr<SAvaLevelViewportScreenGrid> ScreenGrid;
 	TSharedPtr<SAvaLevelViewportSafeFrames> SafeFrames;
@@ -246,6 +252,17 @@ private:
 	void OnBackgroundOpacityChanged(float InValue);
 	void OnBackgroundOpacityCommitted(float InValue, ETextCommit::Type InCommitType);
 
+	FString GetTextureOverlayTextureObjectPath() const;
+	void OnTextureOverlayTextureChanged(const FAssetData& InAssetData);
+
+	float GetTextureOverlayOpacity() const;
+	void OnTextureOverlayOpacitySliderEnd(float InValue);
+	void OnTextureOverlayOpacityChanged(float InValue);
+	void OnTextureOverlayOpacityCommitted(float InValue, ETextCommit::Type InCommitType);
+
+	ECheckBoxState GetTextureOverlayStretchEnabledCheckBoxState() const;
+	void OnTextureOverlayStretchEnabledCheckBoxChanged(ECheckBoxState InState);
+
 	int32 GetVirtualSizeX() const;
 	int32 GetVirtualSizeY() const;
 	void OnVirtualSizeComponentSliderCommitted(int32 InNewDimension, ETextCommit::Type InCommitType, EAxis::Type InAxis);
diff --git a/Engine/Plugins/Experimental/Avalanche/Source/AvalancheLevelViewport/Private/SAvaLevelViewport_Menus.cpp b/Engine/Plugins/Experimental/Avalanche/Source/AvalancheLevelViewport/Private/SAvaLevelViewport_Menus.cpp
--- a/Engine/Plugins/Experimental/Avalanche/Source/AvalancheLevelViewport/Private/SAvaLevelViewport_Menus.cpp
+++ b/Engine/Plugins/Experimental/Avalanche/Source/AvalancheLevelViewport/Private/SAvaLevelViewport_Menus.cpp
@@ -15,13 +15,11 @@
 #include "LevelEditor.h"
 #include "LevelEditor/AvaLevelEditorUtils.h"
 #include "Misc/MessageDialog.h"
-#include "PropertyCustomizationHelpers.h"
 #include "ScopedTransaction.h"
 #include "SEditorViewportToolBarButton.h"
 #include "SEditorViewportToolBarMenu.h"
 #include "SortHelper.h"
 #include "SViewportToolBar.h"
-#include "ThumbnailRendering/ThumbnailManager.h"
 #include "Viewport/Interaction/IAvaViewportDataProvider.h"
 #include "ViewportClient/AvaLevelViewportClient.h"
 #include "Widgets/Images/SImage.h"
@@ -867,176 +865,6 @@ FReply SAvaLevelViewport::OnChildActorLockButtonClicked()
 	return FReply::Handled();
 }
 
-FString SAvaLevelViewport::GetBackgroundTextureObjectPath() const
-{
-	TSharedPtr<FAvaLevelViewportClient> ViewportClient = GetAvaLevelViewportClient();
-
-	if (!ViewportClient.IsValid())
-	{
-		return "";
-	}
-
-	if (!ViewportClient->GetPostProcessManager().IsValid())
-	{
-		return "";
-	}
-
-	FAvaViewportPostProcessInfo* PostProcessInfo = ViewportClient->GetPostProcessManager()->GetPostProcessInfo();
-
-	if (!PostProcessInfo)
-	{
-		return "";
-	}
-
-	return PostProcessInfo->Texture.ToString();
-}
-
-void SAvaLevelViewport::OnBackgroundTextureChanged(const FAssetData& InAssetData)
-{
-	TSharedPtr<FAvaLevelViewportClient> ViewportClient = GetAvaLevelViewportClient();
-
-	if (!ViewportClient.IsValid())
-	{
-		return;
-	}
-
-	if (!ViewportClient->GetPostProcessManager().IsValid())
-	{
-		return;
-	}
-
-	FAvaViewportPostProcessInfo* PostProcessInfo = ViewportClient->GetPostProcessManager()->GetPostProcessInfo();
-
-	if (!PostProcessInfo)
-	{
-		return;
-	}
-
-	BeginPostProcessInfoTransaction();
-
-	PostProcessInfo->Texture = Cast<UTexture>(InAssetData.GetAsset());
-	ViewportClient->GetPostProcessManager()->LoadPostProcessInfo();
-	ViewportClient->Invalidate();
-
-	EndPostProcessInfoTransaction();
-}
-
-float SAvaLevelViewport::GetBackgroundOpacity() const
-{
-	TSharedPtr<FAvaLevelViewportClient> ViewportClient = GetAvaLevelViewportClient();
-
-	if (!ViewportClient.IsValid())
-	{
-		return 1.f;
-	}
-
-	if (!ViewportClient->GetPostProcessManager().IsValid())
-	{
-		return 1.f;
-	}
-
-	FAvaViewportPostProcessInfo* PostProcessInfo = ViewportClient->GetPostProcessManager()->GetPostProcessInfo();
-
-	return ViewportClient->GetPostProcessManager()->GetOpacity();
-}
-
-void SAvaLevelViewport::BeginPostProcessInfoTransaction()
-{
-	TSharedPtr<FAvaLevelViewportClient> ViewportClient = GetAvaLevelViewportClient();
-
-	if (!ViewportClient.IsValid())
-	{
-		return;
-	}
-
-	IAvaViewportDataProvider* DataProvider = ViewportClient->GetViewportDataProvider();
-
-	if (!DataProvider)
-	{
-		return;
-	}
-
-	UObject* DataObject = DataProvider->ToUObject();
-
-	if (!DataObject)
-	{
-		return;
-	}
-
-	if (!PostProcessInfoTransaction.IsValid())
-	{
-		PostProcessInfoTransaction = MakeShared<FScopedTransaction>(LOCTEXT("PostProcessSettingsChange", "Post Process Settings Change"));
-	}
-
-	DataObject->Modify();
-}
-
-void SAvaLevelViewport::EndPostProcessInfoTransaction()
-{
-	PostProcessInfoTransaction.Reset();
-}
-
-void SAvaLevelViewport::OnBackgroundOpacitySliderBegin()
-{
-	BeginPostProcessInfoTransaction();
-}
-
-void SAvaLevelViewport::OnBackgroundOpacitySliderEnd(float InValue)
-{
-	EndPostProcessInfoTransaction();
-}
-
-void SAvaLevelViewport::OnBackgroundOpacityChanged(float InValue)
-{
-	TSharedPtr<FAvaLevelViewportClient> ViewportClient = GetAvaLevelViewportClient();
-
-	if (!ViewportClient.IsValid())
-	{
-		return;
-	}
-
-	if (!ViewportClient->GetPostProcessManager().IsValid())
-	{
-		return;
-	}
-
-	ViewportClient->GetPostProcessManager()->SetOpacity(InValue);
-	ViewportClient->Invalidate();
-}
-
-void SAvaLevelViewport::OnBackgroundOpacityCommitted(float InValue, ETextCommit::Type InCommitType)
-{
-	if (InCommitType == ETextCommit::OnCleared)
-	{
-		return;
-	}
-
-	TSharedPtr<FAvaLevelViewportClient> ViewportClient = GetAvaLevelViewportClient();
-
-	if (!ViewportClient.IsValid())
-	{
-		return;
-	}
-
-	if (!ViewportClient->GetPostProcessManager().IsValid())
-	{
-		return;
-	}
-
-	if (InCommitType == ETextCommit::OnEnter)
-	{
-		BeginPostProcessInfoTransaction();
-	}
-
-	ViewportClient->GetPostProcessManager()->SetOpacity(InValue);
-	ViewportClient->Invalidate();
-
-	if (InCommitType == ETextCommit::OnEnter)
-	{
-		EndPostProcessInfoTransaction();
-	}
-}
-
 void SAvaLevelViewport::OnVirtualSizeComponentSliderCommitted(int32 InNewDimension, ETextCommit::Type InCommitType, EAxis::Type InAxis)
 {
 	using namespace UE::AvaLevelViewport::Private;
diff --git a/Engine/Plugins/Experimental/Avalanche/Source/AvalancheLevelViewport/Private/SAvaLevelViewport_StatusBar.cpp b/Engine/Plugins/Experimental/Avalanche/Source/AvalancheLevelViewport/Private/SAvaLevelViewport_StatusBar.cpp
--- a/Engine/Plugins/Experimental/Avalanche/Source/AvalancheLevelViewport/Private/SAvaLevelViewport_StatusBar.cpp
+++ b/Engine/Plugins/Experimental/Avalanche/Source/AvalancheLevelViewport/Private/SAvaLevelViewport_StatusBar.cpp
@@ -3,10 +3,15 @@
 #include "SAvaLevelViewport.h"
 #include "AvaViewportPostProcessManager.h"
 #include "AvaViewportSettings.h"
+#include "Engine/Texture.h"
+#include "ScopedTransaction.h"
+#include "Viewport/Interaction/IAvaViewportDataProvider.h"
 #include "ViewportClient/AvaLevelViewportClient.h"
 #include "Visualizers/IAvaViewportBoundingBoxVisualizer.h"
 #include "Visualizers/IAvaViewportPostProcessVisualizer.h"
 
+#define LOCTEXT_NAMESPACE "SAvaLevelViewport"
+
 void SAvaLevelViewport::ExecuteToggleChildActorLock()
 {
 	TSharedPtr<FAvaLevelViewportClient> ViewportClient = GetAvaLevelViewportClient();
@@ -765,3 +770,254 @@ void SAvaLevelViewport::ExecuteAddVerticalGuide()
 {
 	AddGuide(EOrientation::Orient_Vertical, 0.5f);
 }
+
+FString SAvaLevelViewport::GetBackgroundTextureObjectPath() const
+{
+	TSharedPtr<FAvaLevelViewportClient> ViewportClient = GetAvaLevelViewportClient();
+
+	if (!ViewportClient.IsValid())
+	{
+		return "";
+	}
+
+	if (!ViewportClient->GetPostProcessManager().IsValid())
+	{
+		return "";
+	}
+
+	FAvaViewportPostProcessInfo* PostProcessInfo = ViewportClient->GetPostProcessManager()->GetPostProcessInfo();
+
+	if (!PostProcessInfo)
+	{
+		return "";
+	}
+
+	return PostProcessInfo->Texture.ToString();
+}
+
+void SAvaLevelViewport::OnBackgroundTextureChanged(const FAssetData& InAssetData)
+{
+	TSharedPtr<FAvaLevelViewportClient> ViewportClient = GetAvaLevelViewportClient();
+
+	if (!ViewportClient.IsValid())
+	{
+		return;
+	}
+
+	if (!ViewportClient->GetPostProcessManager().IsValid())
+	{
+		return;
+	}
+
+	FAvaViewportPostProcessInfo* PostProcessInfo = ViewportClient->GetPostProcessManager()->GetPostProcessInfo();
+
+	if (!PostProcessInfo)
+	{
+		return;
+	}
+
+	BeginPostProcessInfoTransaction();
+
+	PostProcessInfo->Texture = Cast<UTexture>(InAssetData.GetAsset());
+	ViewportClient->GetPostProcessManager()->LoadPostProcessInfo();
+	ViewportClient->Invalidate();
+
+	EndPostProcessInfoTransaction();
+}
+
+float SAvaLevelViewport::GetBackgroundOpacity() const
+{
+	TSharedPtr<FAvaLevelViewportClient> ViewportClient = GetAvaLevelViewportClient();
+
+	if (!ViewportClient.IsValid())
+	{
+		return 1.f;
+	}
+
+	if (!ViewportClient->GetPostProcessManager().IsValid())
+	{
+		return 1.f;
+	}
+
+	FAvaViewportPostProcessInfo* PostProcessInfo = ViewportClient->GetPostProcessManager()->GetPostProcessInfo();
+
+	return ViewportClient->GetPostProcessManager()->GetOpacity();
+}
+
+void SAvaLevelViewport::BeginPostProcessInfoTransaction()
+{
+	TSharedPtr<FAvaLevelViewportClient> ViewportClient = GetAvaLevelViewportClient();
+
+	if (!ViewportClient.IsValid())
+	{
+		return;
+	}
+
+	IAvaViewportDataProvider* DataProvider = ViewportClient->GetViewportDataProvider();
+
+	if (!DataProvider)
+	{
+		return;
+	}
+
+	UObject* DataObject = DataProvider->ToUObject();
+
+	if (!DataObject)
+	{
+		return;
+	}
+
+	if (!PostProcessInfoTransaction.IsValid())
+	{
+		PostProcessInfoTransaction = MakeShared<FScopedTransaction>(LOCTEXT("PostProcessSettingsChange", "Post Process Settings Change"));
+	}
+
+	DataObject->Modify();
+}
+
+void SAvaLevelViewport::EndPostProcessInfoTransaction()
+{
+	PostProcessInfoTransaction.Reset();
+}
+
+void SAvaLevelViewport::OnBackgroundOpacitySliderBegin()
+{
+	BeginPostProcessInfoTransaction();
+}
+
+void SAvaLevelViewport::OnBackgroundOpacitySliderEnd(float InValue)
+{
+	EndPostProcessInfoTransaction();
+}
+
+void SAvaLevelViewport::OnBackgroundOpacityChanged(float InValue)
+{
+	TSharedPtr<FAvaLevelViewportClient> ViewportClient = GetAvaLevelViewportClient();
+
+	if (!ViewportClient.IsValid())
+	{
+		return;
+	}
+
+	if (!ViewportClient->GetPostProcessManager().IsValid())
+	{
+		return;
+	}
+
+	ViewportClient->GetPostProcessManager()->SetOpacity(InValue);
+	ViewportClient->Invalidate();
+}
+
+void SAvaLevelViewport::OnBackgroundOpacityCommitted(float InValue, ETextCommit::Type InCommitType)
+{
+	if (InCommitType == ETextCommit::OnCleared)
+	{
+		return;
+	}
+
+	TSharedPtr<FAvaLevelViewportClient> ViewportClient = GetAvaLevelViewportClient();
+
+	if (!ViewportClient.IsValid())
+	{
+		return;
+	}
+
+	if (!ViewportClient->GetPostProcessManager().IsValid())
+	{
+		return;
+	}
+
+	if (InCommitType == ETextCommit::OnEnter)
+	{
+		BeginPostProcessInfoTransaction();
+	}
+
+	ViewportClient->GetPostProcessManager()->SetOpacity(InValue);
+	ViewportClient->Invalidate();
+
+	if (InCommitType == ETextCommit::OnEnter)
+	{
+		EndPostProcessInfoTransaction();
+	}
+}
+
+FString SAvaLevelViewport::GetTextureOverlayTextureObjectPath() const
+{
+	if (const UAvaViewportSettings* AvaViewportSettings = GetDefault<UAvaViewportSettings>())
+	{
+		return AvaViewportSettings->TextureOverlayTexture.ToString();
+	}
+
+	return "";
+}
+
+void SAvaLevelViewport::OnTextureOverlayTextureChanged(const FAssetData& InAssetData)
+{
+	if (UAvaViewportSettings* AvaViewportSettings = GetMutableDefault<UAvaViewportSettings>())
+	{
+		AvaViewportSettings->TextureOverlayTexture = Cast<UTexture>(InAssetData.GetAsset());
+		AvaViewportSettings->BroadcastSettingChanged(GET_MEMBER_NAME_CHECKED(UAvaViewportSettings, TextureOverlayTexture));
+		AvaViewportSettings->SaveConfig();
+	}
+}
+
+float SAvaLevelViewport::GetTextureOverlayOpacity() const
+{
+	if (const UAvaViewportSettings* AvaViewportSettings = GetDefault<UAvaViewportSettings>())
+	{
+		return AvaViewportSettings->TextureOverlayOpacity;
+	}
+
+	return 0.f;
+}
+
+void SAvaLevelViewport::OnTextureOverlayOpacitySliderEnd(float InValue)
+{
+	if (UAvaViewportSettings* AvaViewportSettings = GetMutableDefault<UAvaViewportSettings>())
+	{
+		AvaViewportSettings->TextureOverlayOpacity = InValue;
+		AvaViewportSettings->BroadcastSettingChanged(GET_MEMBER_NAME_CHECKED(UAvaViewportSettings, TextureOverlayOpacity));
+		AvaViewportSettings->SaveConfig();
+	}
+}
+
+void SAvaLevelViewport::OnTextureOverlayOpacityChanged(float InValue)
+{
+	if (UAvaViewportSettings* AvaViewportSettings = GetMutableDefault<UAvaViewportSettings>())
+	{
+		AvaViewportSettings->TextureOverlayOpacity = InValue;
+		AvaViewportSettings->BroadcastSettingChanged(GET_MEMBER_NAME_CHECKED(UAvaViewportSettings, TextureOverlayOpacity));
+	}
+}
+
+void SAvaLevelViewport::OnTextureOverlayOpacityCommitted(float InValue, ETextCommit::Type InCommitType)
+{
+	if (UAvaViewportSettings* AvaViewportSettings = GetMutableDefault<UAvaViewportSettings>())
+	{
+		AvaViewportSettings->TextureOverlayOpacity = InValue;
+		AvaViewportSettings->BroadcastSettingChanged(GET_MEMBER_NAME_CHECKED(UAvaViewportSettings, TextureOverlayOpacity));
+		AvaViewportSettings->SaveConfig();
+	}
+}
+
+bool SAvaLevelViewport::CanToggleTextureOverlay() const
+{
+	if (const UAvaViewportSettings* AvaViewportSettings = GetDefault<UAvaViewportSettings>())
+	{
+		return AvaViewportSettings->bEnableViewportOverlay;
+	}
+
+	return false;
+}
+
+void SAvaLevelViewport::ExecuteToggleTextureOverlay()
+{
+	if (UAvaViewportSettings* AvaViewportSettings = GetMutableDefault<UAvaViewportSettings>())
+	{
+		AvaViewportSettings->bEnableTextureOverlay = !AvaViewportSettings->bEnableTextureOverlay;
+		AvaViewportSettings->SaveConfig();
+		ApplySettings(AvaViewportSettings);
+	}
+}
+
+#undef LOCTEXT_NAMESPACE
diff --git a/Engine/Plugins/Experimental/Avalanche/Source/AvalancheLevelViewport/Private/Widgets/SAvaLevelViewportStatusBarButtons.cpp b/Engine/Plugins/Experimental/Avalanche/Source/AvalancheLevelViewport/Private/Widgets/SAvaLevelViewportStatusBarButtons.cpp
--- a/Engine/Plugins/Experimental/Avalanche/Source/AvalancheLevelViewport/Private/Widgets/SAvaLevelViewportStatusBarButtons.cpp
+++ b/Engine/Plugins/Experimental/Avalanche/Source/AvalancheLevelViewport/Private/Widgets/SAvaLevelViewportStatusBarButtons.cpp
@@ -22,6 +22,7 @@
 #include "ToolMenu.h"
 #include "ToolMenus.h"
 #include "ViewportClient/AvaLevelViewportClient.h"
+#include "Widgets/Input/SCheckBox.h"
 #include "Widgets/Input/SSpinBox.h"
 #include "Widgets/Layout/SBox.h"
 #include "Widgets/Layout/SSeparator.h"
@@ -223,6 +224,45 @@ void SAvaLevelViewportStatusBarButtons::CreateContextMenuWigets()
 				.OnValueCommitted(this, &SAvaLevelViewportStatusBarButtons::OnGridSizeCommitted)
 			];
 	}
+
+	if (!TextureOverlayOpacitySlider.IsValid())
+	{
+		TextureOverlayOpacitySlider = SNew(SSpinBox<float>)
+			.ClearKeyboardFocusOnCommit(true)
+			.MaxFractionalDigits(3)
+			.MinDesiredWidth(50.f)
+			.OnEndSliderMovement(LevelViewport, &SAvaLevelViewport::OnTextureOverlayOpacitySliderEnd)
+			.OnValueCommitted(LevelViewport, &SAvaLevelViewport::OnTextureOverlayOpacityCommitted)
+			.OnValueChanged(LevelViewport, &SAvaLevelViewport::OnTextureOverlayOpacityChanged)
+			.Value(LevelViewport, &SAvaLevelViewport::GetTextureOverlayOpacity)
+			.MinValue(0.f)
+			.MinSliderValue(0.f)
+			.MaxValue(1.f)
+			.MaxSliderValue(1.f);
+	}
+
+	if (!TextureOverlayTextureSelector.IsValid())
+	{
+		TextureOverlayTextureSelector = SNew(SObjectPropertyEntryBox)
+			.AllowClear(true)
+			.AllowedClass(UTexture::StaticClass())
+			.DisplayBrowse(true)
+			.DisplayThumbnail(true)
+			.DisplayCompactSize(true)
+			.DisplayUseSelected(true)
+			.ThumbnailPool(UThumbnailManager::Get().GetSharedThumbnailPool())
+			.EnableContentPicker(true)
+			.ObjectPath(LevelViewport, &SAvaLevelViewport::GetTextureOverlayTextureObjectPath)
+			.OnObjectChanged(LevelViewport, &SAvaLevelViewport::OnTextureOverlayTextureChanged)
+			.OnShouldSetAsset(FOnShouldSetAsset::CreateLambda([](const FAssetData& InAssetData) { return false; }));
+	}
+
+	if (!TextureOverlayStretchCheckBox.IsValid())
+	{
+		TextureOverlayStretchCheckBox = SNew(SCheckBox)
+			.IsChecked(LevelViewport, &SAvaLevelViewport::GetTextureOverlayStretchEnabledCheckBoxState)
+			.OnCheckStateChanged(LevelViewport, &SAvaLevelViewport::OnTextureOverlayStretchEnabledCheckBoxChanged);
+	}
 }
 
 void SAvaLevelViewportStatusBarButtons::PopulateActorButtons(TSharedPtr<SHorizontalBox> InContainer)
@@ -284,22 +324,6 @@ void SAvaLevelViewportStatusBarButtons::PopulateViewportButtons(TSharedPtr<SHori
 			)
 		];
 
-	TSharedRef<SComboButton> PostProcessButton = ViewportStatusBarButton::MakeMenuButton(
-		LOCTEXT("PostProcessEffects", "Post Process Effects"),
-		FOnGetContent::CreateSP(this, &SAvaLevelViewportStatusBarButtons::GetPostProcessMenuContent),
-		TAttribute<const FSlateBrush*>::CreateSP(this, &SAvaLevelViewportStatusBarButtons::GetPostProcessIcon),
-		TAttribute<FSlateColor>::CreateSP(this, &SAvaLevelViewportStatusBarButtons::GetPostProcessColor)
-	);
-
-	PostProcessButton->SetEnabled(TAttribute<bool>::CreateSP(this, &SAvaLevelViewportStatusBarButtons::GetPostProcessEnabled));
-
-	InContainer->AddSlot()
-		.AutoWidth()
-		.Padding(ViewportStatusBarButton::Padding)
-		[
-			PostProcessButton
-		];
-
 	InContainer->AddSlot()
 		.AutoWidth()
 		.Padding(ViewportStatusBarButton::Padding)
@@ -384,6 +408,39 @@ void SAvaLevelViewportStatusBarButtons::PopulateViewportButtons(TSharedPtr<SHori
 			)
 		];
 
+	TSharedRef<SComboButton> PostProcessButton = ViewportStatusBarButton::MakeMenuButton(
+		LOCTEXT("PostProcessEffects", "Post Process Effects"),
+		FOnGetContent::CreateSP(this, &SAvaLevelViewportStatusBarButtons::GetPostProcessMenuContent),
+		TAttribute<const FSlateBrush*>::CreateSP(this, &SAvaLevelViewportStatusBarButtons::GetPostProcessIcon),
+		TAttribute<FSlateColor>::CreateSP(this, &SAvaLevelViewportStatusBarButtons::GetPostProcessColor)
+	);
+
+	PostProcessButton->SetEnabled(TAttribute<bool>::CreateSP(this, &SAvaLevelViewportStatusBarButtons::GetPostProcessEnabled));
+
+	InContainer->AddSlot()
+		.AutoWidth()
+		.Padding(ViewportStatusBarButton::Padding)
+		[
+			PostProcessButton
+		];
+
+	TSharedRef<SAvaMultiComboButton> TextureOverlayButton = ViewportStatusBarButton::MakeMultiMenuButton(
+		CommandsRef.ToggleTextureOverlay->GetDescription(),
+		FOnGetContent::CreateSP(this, &SAvaLevelViewportStatusBarButtons::GetTextureOverlayMenuContent),
+		FAppStyle::Get().GetBrush(TEXT("GenericCommands.Paste")),
+		TAttribute<FSlateColor>::CreateSP(this, &SAvaLevelViewportStatusBarButtons::GetTextureOverlayColor),
+		FOnClicked::CreateSP(this, &SAvaLevelViewportStatusBarButtons::ToggleTextureOverlay)
+	);
+
+	TextureOverlayButton->SetEnabled(TAttribute<bool>::CreateSP(this, &SAvaLevelViewportStatusBarButtons::GetPostProcessEnabled));
+
+	InContainer->AddSlot()
+		.AutoWidth()
+		.Padding(ViewportStatusBarButton::Padding)
+		[
+			TextureOverlayButton
+		];
+
 	InContainer->AddSlot()
 		.AutoWidth()
 		.Padding(ViewportStatusBarButton::Padding)
@@ -947,7 +1004,7 @@ FReply SAvaLevelViewportStatusBarButtons::ToggleShapeEditorOverlay()
 	{
 		AvaViewportSettings->bEnableShapesEditorOverlay = !AvaViewportSettings->bEnableShapesEditorOverlay;
 		AvaViewportSettings->SaveConfig();
-		AvaViewportSettings->OnChange.Broadcast(AvaViewportSettings, GET_MEMBER_NAME_CHECKED(UAvaViewportSettings, bEnableShapesEditorOverlay));
+		AvaViewportSettings->BroadcastSettingChanged(GET_MEMBER_NAME_CHECKED(UAvaViewportSettings, bEnableShapesEditorOverlay));
 
 		return FReply::Handled();
 	}
@@ -1164,6 +1221,105 @@ TSharedRef<SWidget> SAvaLevelViewportStatusBarButtons::GetViewportInfoWidget() c
 	return SNullWidget::NullWidget;
 }
 
+FSlateColor SAvaLevelViewportStatusBarButtons::GetTextureOverlayColor() const
+{
+	using namespace UE::AvaLevelViewport::Private;
+
+	const FAvaLevelViewportGuideFrameAndWidget FrameAndWidget(ViewportFrameWeak);
+
+	if (FrameAndWidget.IsValid() && FrameAndWidget.ViewportWidget->CanToggleTextureOverlay())
+	{
+		if (const UAvaViewportSettings* AvaViewportSettings = GetDefault<UAvaViewportSettings>())
+		{
+			return AvaViewportSettings->bEnableTextureOverlay ? ViewportStatusBarButton::ActiveColor : ViewportStatusBarButton::EnabledColor;
+		}
+	}
+
+	return ViewportStatusBarButton::DisabledColor;
+}
+
+bool SAvaLevelViewportStatusBarButtons::GetTextureOverlayEnabled() const
+{
+	const FAvaLevelViewportGuideFrameAndWidget FrameAndWidget(ViewportFrameWeak);
+
+	if (FrameAndWidget.IsValid())
+	{
+		return FrameAndWidget.ViewportWidget->CanToggleTextureOverlay();
+	}
+
+	return false;
+}
+
+TSharedRef<SWidget> SAvaLevelViewportStatusBarButtons::GetTextureOverlayMenuContent()
+{
+	UToolMenus* Menus = UToolMenus::Get();
+
+	check(Menus);
+
+	static const FName TextureOverlayMenuName = TEXT("AvaLevelViewport.StatusBar.TextureOverlay");
+
+	UToolMenu* ContextMenu = Menus->FindMenu(TextureOverlayMenuName);
+
+	if (!ContextMenu)
+	{
+		using namespace UE::Ava::LevelViewportStatusBarButtons::Private;
+
+		ContextMenu = Menus->RegisterMenu(TextureOverlayMenuName, NAME_None, EMultiBoxType::Menu);
+
+		FToolMenuSection& OptionsSection = ContextMenu->AddSection("TextureOverlay", LOCTEXT("TextureOverlay", "Texture Overlay"));
+
+		if (TextureOverlayOpacitySlider.IsValid())
+		{
+			OptionsSection.AddEntry(FToolMenuEntry::InitWidget(
+				"TextureOverlayOpacity",
+				TextureOverlayOpacitySlider.ToSharedRef(),
+				LOCTEXT("TextureOverlayOpacity", "Opacity"),
+				true
+			));
+		}
+
+		if (TextureOverlayTextureSelector.IsValid())
+		{
+			OptionsSection.AddEntry(FToolMenuEntry::InitWidget(
+				"TextureOverlayTexture",
+				TextureOverlayTextureSelector.ToSharedRef(),
+				LOCTEXT("TextureOverlayTexture", "Texture"),
+				true
+			));
+		}
+
+		if (TextureOverlayStretchCheckBox.IsValid())
+		{
+			OptionsSection.AddEntry(FToolMenuEntry::InitWidget(
+				"TextureOverlayStretch",
+				TextureOverlayStretchCheckBox.ToSharedRef(),
+				LOCTEXT("TextureOverlayStretch", "Stretch Texture"),
+				true
+			));
+		}
+	}
+
+	if (!ContextMenu)
+	{
+		return SNullWidget::NullWidget;
+	}
+
+	return Menus->GenerateWidget(ContextMenu);
+}
+
+FReply SAvaLevelViewportStatusBarButtons::ToggleTextureOverlay()
+{
+	const FAvaLevelViewportGuideFrameAndWidget FrameAndWidget(ViewportFrameWeak);
+
+	if (FrameAndWidget.IsValid() && FrameAndWidget.ViewportWidget->CanToggleSnapping())
+	{
+		FrameAndWidget.ViewportWidget->ExecuteToggleTextureOverlay();
+		return FReply::Handled();
+	}
+
+	return FReply::Unhandled();
+}
+
 bool SAvaLevelViewportStatusBarButtons::CanChangeGridSize() const
 {
 	const FAvaLevelViewportGuideFrameAndWidget FrameAndWidget(ViewportFrameWeak);
diff --git a/Engine/Plugins/Experimental/Avalanche/Source/AvalancheLevelViewport/Private/Widgets/SAvaLevelViewportStatusBarButtons.h b/Engine/Plugins/Experimental/Avalanche/Source/AvalancheLevelViewport/Private/Widgets/SAvaLevelViewportStatusBarButtons.h
--- a/Engine/Plugins/Experimental/Avalanche/Source/AvalancheLevelViewport/Private/Widgets/SAvaLevelViewportStatusBarButtons.h
+++ b/Engine/Plugins/Experimental/Avalanche/Source/AvalancheLevelViewport/Private/Widgets/SAvaLevelViewportStatusBarButtons.h
@@ -30,6 +30,12 @@ protected:
 
 	TSharedPtr<SWidget> GridSizeSlider;
 
+	TSharedPtr<SWidget> TextureOverlayTextureSelector;
+
+	TSharedPtr<SWidget> TextureOverlayOpacitySlider;
+
+	TSharedPtr<SWidget> TextureOverlayStretchCheckBox;
+
 	void CreateContextMenuWigets();
 
 	void PopulateActorButtons(TSharedPtr<SHorizontalBox> InContainer);
@@ -104,4 +110,9 @@ protected:
 	bool GetViewportInfoEnabled() const;
 	FSlateColor GetViewportInfoColor() const;
 	TSharedRef<SWidget> GetViewportInfoWidget() const;
+
+	FSlateColor GetTextureOverlayColor() const;
+	bool GetTextureOverlayEnabled() const;
+	TSharedRef<SWidget> GetTextureOverlayMenuContent();
+	FReply ToggleTextureOverlay();
 };
\ No newline at end of file
diff --git a/Engine/Plugins/Experimental/Avalanche/Source/AvalancheLevelViewport/Private/Widgets/SAvaLevelViewportTextureOverlay.cpp b/Engine/Plugins/Experimental/Avalanche/Source/AvalancheLevelViewport/Private/Widgets/SAvaLevelViewportTextureOverlay.cpp
new file mode 100644
--- /dev/null
+++ b/Engine/Plugins/Experimental/Avalanche/Source/AvalancheLevelViewport/Private/Widgets/SAvaLevelViewportTextureOverlay.cpp
@@ -0,0 +1,160 @@
+// Copyright Epic Games, Inc. All Rights Reserved.
+
+#include "Widgets/SAvaLevelViewportTextureOverlay.h"
+#include "AvaViewportSettings.h"
+#include "AvaViewportUtils.h"
+#include "AvaVisibleArea.h"
+#include "Engine/Texture.h"
+#include "TextureCompiler.h"
+#include "ViewportClient/AvaLevelViewportClient.h"
+
+SAvaLevelViewportTextureOverlay::SAvaLevelViewportTextureOverlay()
+	: TextureBrush((UObject*)nullptr, FVector2f::ZeroVector)
+{
+}
+
+void SAvaLevelViewportTextureOverlay::Construct(const FArguments& InArgs, TSharedRef<FAvaLevelViewportClient> InAvaLevelViewportClient)
+{
+	AvaLevelViewportClientWeak = InAvaLevelViewportClient;
+	SetVisibility(EVisibility::HitTestInvisible);
+	TextureStrong.Reset(nullptr);
+
+	if (UAvaViewportSettings* AvaViewportSettings = GetMutableDefault<UAvaViewportSettings>())
+	{
+		AvaViewportSettings->OnChange.AddSP(this, &SAvaLevelViewportTextureOverlay::OnAvaViewportSettingsChanged);
+		SetTexture(AvaViewportSettings->TextureOverlayTexture.LoadSynchronous());
+		SetOpacity(AvaViewportSettings->TextureOverlayOpacity);
+	}
+}
+
+SAvaLevelViewportTextureOverlay::~SAvaLevelViewportTextureOverlay()
+{
+	if (UObjectInitialized())
+	{
+		if (UAvaViewportSettings* AvaViewportSettings = GetMutableDefault<UAvaViewportSettings>())
+		{
+			AvaViewportSettings->OnChange.RemoveAll(this);
+		}		
+	}
+}
+
+int32 SAvaLevelViewportTextureOverlay::OnPaint(const FPaintArgs& InPaintArgs, const FGeometry& InAllottedGeometry,
+	const FSlateRect& InMyCullingRect, FSlateWindowElementList& OutDrawElements, int32 InLayerId, const FWidgetStyle& InWidgetStyle,
+	bool bInParentEnabled) const
+{
+	InLayerId = SCompoundWidget::OnPaint(InPaintArgs, InAllottedGeometry, InMyCullingRect, OutDrawElements, 
+		InLayerId, InWidgetStyle, bInParentEnabled);
+
+	if (UAvaViewportSettings const* AvaViewportSettings = GetDefault<UAvaViewportSettings>())
+	{
+		if (AvaViewportSettings->bEnableTextureOverlay)
+		{
+			UTexture* Texture = TextureStrong.Get();
+
+			if (IsValid(Texture))
+			{
+				if (TSharedPtr<FAvaLevelViewportClient> AvaLevelViewportClient = AvaLevelViewportClientWeak.Pin())
+				{
+					DrawOverlay(InPaintArgs, InAllottedGeometry, InMyCullingRect, OutDrawElements, InLayerId,
+						Texture, AvaViewportSettings->TextureOverlayOpacity, AvaViewportSettings->bTextureOverlayStretch);
+				}
+			}
+		}
+	}
+
+	return InLayerId;
+}
+
+void SAvaLevelViewportTextureOverlay::DrawOverlay(const FPaintArgs& InPaintArgs, const FGeometry& InAllottedGeometry,
+	const FSlateRect& InMyCullingRect, FSlateWindowElementList& OutDrawElements, int32& InOutLayerId, 
+	UTexture* InTexture, float InOpacity, bool bInStretch) const
+{
+	TSharedPtr<FAvaLevelViewportClient> AvaLevelViewportClient = AvaLevelViewportClientWeak.Pin();
+
+	if (!AvaLevelViewportClient.IsValid())
+	{
+		return;
+	}
+
+	const FAvaVisibleArea& VisibleArea = AvaLevelViewportClient->GetZoomedVisibleArea();
+
+	if (!VisibleArea.IsValid())
+	{
+		return;
+	}
+
+	++InOutLayerId;
+
+	const FVector2f Offset = AvaLevelViewportClient->GetCachedViewportOffset();
+	FVector2f DrawSize = VisibleArea.AbsoluteSize;
+
+	if (!bInStretch)
+	{
+		const FVector2f TextureSize = TextureBrush.GetImageSize();
+
+		if (!TextureSize.Equals(DrawSize))
+		{
+			const float VisibleAreaAspectRatio = DrawSize.X / DrawSize.Y;
+			const float TextureAspectRatio = TextureSize.X / TextureSize.Y;
+
+			if (VisibleAreaAspectRatio > TextureAspectRatio)
+			{
+				DrawSize.X *= TextureAspectRatio / VisibleAreaAspectRatio;
+			}
+			else if (VisibleAreaAspectRatio < TextureAspectRatio)
+			{
+				DrawSize.Y *= VisibleAreaAspectRatio / TextureAspectRatio;
+			}
+		}
+	}
+
+	const FVector2f TextureOffset = (VisibleArea.AbsoluteSize - DrawSize) * 0.5f;
+	const FVector2f TopLeft = VisibleArea.GetVisiblePosition(TextureOffset);
+	const FVector2f BottomRight = VisibleArea.GetVisiblePosition(VisibleArea.AbsoluteSize - TextureOffset);
+
+	const FPaintGeometry TextureGeometry = InAllottedGeometry.ToPaintGeometry(
+		BottomRight - TopLeft,
+		FSlateLayoutTransform(TopLeft + Offset)
+	);
+
+	FSlateDrawElement::MakeBox(
+		OutDrawElements,
+		InOutLayerId,
+		TextureGeometry,
+		&TextureBrush,
+		ESlateDrawEffect::NoGamma,
+		TextureBrush.TintColor.GetSpecifiedColor()
+	);
+}
+
+void SAvaLevelViewportTextureOverlay::OnAvaViewportSettingsChanged(const UAvaViewportSettings* InSettings, FName InSetting)
+{
+	static const FName OverlayTextureName = GET_MEMBER_NAME_CHECKED(UAvaViewportSettings, TextureOverlayTexture);
+	static const FName OverlayOpacity = GET_MEMBER_NAME_CHECKED(UAvaViewportSettings, TextureOverlayOpacity);
+
+	if (InSetting == OverlayTextureName)
+	{
+		SetTexture(InSettings->TextureOverlayTexture.LoadSynchronous());		
+	}
+	else if (InSetting == OverlayOpacity)
+	{
+		SetOpacity(InSettings->TextureOverlayOpacity);
+	}
+}
+
+void SAvaLevelViewportTextureOverlay::SetTexture(UTexture* InTexture)
+{
+	TextureStrong.Reset(InTexture);
+	TextureBrush.SetResourceObject(InTexture);
+
+	if (InTexture)
+	{
+		FTextureCompilingManager::Get().FinishCompilation({InTexture});
+		TextureBrush.SetImageSize(FVector2f((float)InTexture->GetSurfaceWidth(), (float)InTexture->GetSurfaceHeight()));
+	}
+}
+
+void SAvaLevelViewportTextureOverlay::SetOpacity(float InOpacity)
+{
+	TextureBrush.TintColor = FSlateColor(FLinearColor(1.f, 1.f, 1.f, FMath::Clamp(InOpacity, 0.f, 1.f)));
+}
diff --git a/Engine/Plugins/Experimental/Avalanche/Source/AvalancheLevelViewport/Private/Widgets/SAvaLevelViewportTextureOverlay.h b/Engine/Plugins/Experimental/Avalanche/Source/AvalancheLevelViewport/Private/Widgets/SAvaLevelViewportTextureOverlay.h
new file mode 100644
--- /dev/null
+++ b/Engine/Plugins/Experimental/Avalanche/Source/AvalancheLevelViewport/Private/Widgets/SAvaLevelViewportTextureOverlay.h
@@ -0,0 +1,44 @@
+// Copyright Epic Games, Inc. All Rights Reserved.
+
+#pragma once
+
+#include "Brushes/SlateImageBrush.h"
+#include "Templates/SharedPointer.h"
+#include "UObject/StrongObjectPtr.h"
+#include "Widgets/SCompoundWidget.h"
+
+class FAvaLevelViewportClient;
+class UAvaViewportSettings;
+class UTexture;
+
+class SAvaLevelViewportTextureOverlay: public SCompoundWidget
+{
+public:
+	SLATE_BEGIN_ARGS(SAvaLevelViewportTextureOverlay)
+		{}
+	SLATE_END_ARGS()
+
+	SAvaLevelViewportTextureOverlay();
+	virtual ~SAvaLevelViewportTextureOverlay() override;
+
+	void Construct(const FArguments& InArgs, TSharedRef<FAvaLevelViewportClient> InAvaLevelViewportClient);
+
+	//~ Begin SWidget
+	virtual int32 OnPaint(const FPaintArgs& InPaintArgs, const FGeometry& InAllottedGeometry, const FSlateRect& InMyCullingRect,
+		FSlateWindowElementList& OutDrawElements, int32 InLayerId, const FWidgetStyle& InWidgetStyle, bool bInParentEnabled) const override;
+	//~ End SWidget
+
+protected:
+	TWeakPtr<FAvaLevelViewportClient> AvaLevelViewportClientWeak;
+	TStrongObjectPtr<UTexture> TextureStrong;
+	FSlateImageBrush TextureBrush;
+
+	void DrawOverlay(const FPaintArgs& InPaintArgs, const FGeometry& InAllottedGeometry, const FSlateRect& InMyCullingRect,
+		FSlateWindowElementList& OutDrawElements, int32& InOutLayerId, UTexture* InTexture, float InOpacity, bool bInStretch) const;
+
+	void OnAvaViewportSettingsChanged(const UAvaViewportSettings* InSettings, FName InSetting);
+
+	void SetTexture(UTexture* InTexture);
+
+	void SetOpacity(float InOpacity);
+};
diff --git a/Engine/Plugins/Experimental/Avalanche/Source/AvalancheLevelViewport/Public/AvaLevelViewportCommands.h b/Engine/Plugins/Experimental/Avalanche/Source/AvalancheLevelViewport/Public/AvaLevelViewportCommands.h
--- a/Engine/Plugins/Experimental/Avalanche/Source/AvalancheLevelViewport/Public/AvaLevelViewportCommands.h
+++ b/Engine/Plugins/Experimental/Avalanche/Source/AvalancheLevelViewport/Public/AvaLevelViewportCommands.h
@@ -35,6 +35,7 @@ public:
 	TSharedPtr<FUICommandInfo> TogglePostProcessChannelBlue;
 	TSharedPtr<FUICommandInfo> TogglePostProcessChannelAlpha;
 	TSharedPtr<FUICommandInfo> TogglePostProcessCheckerboard;
+	TSharedPtr<FUICommandInfo> ToggleTextureOverlay;
 
 	// Camera
 	TSharedPtr<FUICommandInfo> CameraZoomInCenter;
diff --git a/Engine/Plugins/Experimental/Avalanche/Source/AvalancheViewport/Private/AvaViewportSettings.cpp b/Engine/Plugins/Experimental/Avalanche/Source/AvalancheViewport/Private/AvaViewportSettings.cpp
--- a/Engine/Plugins/Experimental/Avalanche/Source/AvalancheViewport/Private/AvaViewportSettings.cpp
+++ b/Engine/Plugins/Experimental/Avalanche/Source/AvalancheViewport/Private/AvaViewportSettings.cpp
@@ -73,6 +73,11 @@ void UAvaViewportSettings::SetSnapState(EAvaViewportSnapState InSnapState)
 	SnapState = static_cast<int32>(InSnapState);
 }
 
+void UAvaViewportSettings::BroadcastSettingChanged(FName InSettingName)
+{
+	OnChange.Broadcast(this, InSettingName);
+}
+
 void UAvaViewportSettings::PostEditChangeProperty(FPropertyChangedEvent& PropertyChangedEvent)
 {
 	Super::PostEditChangeProperty(PropertyChangedEvent);
@@ -81,6 +86,6 @@ void UAvaViewportSettings::PostEditChangeProperty(FPropertyChangedEvent& Propert
 
 	if (MemberName != NAME_None)
 	{
-		OnChange.Broadcast(this, MemberName);
+		BroadcastSettingChanged(MemberName);
 	}
 }
diff --git a/Engine/Plugins/Experimental/Avalanche/Source/AvalancheViewport/Public/AvaViewportSettings.h b/Engine/Plugins/Experimental/Avalanche/Source/AvalancheViewport/Public/AvaViewportSettings.h
--- a/Engine/Plugins/Experimental/Avalanche/Source/AvalancheViewport/Public/AvaViewportSettings.h
+++ b/Engine/Plugins/Experimental/Avalanche/Source/AvalancheViewport/Public/AvaViewportSettings.h
@@ -4,10 +4,12 @@
 
 #include "Engine/DeveloperSettings.h"
 #include "Interaction/AvaSnapDefs.h"
+#include "UObject/SoftObjectPtr.h"
 #include "Types/SlateEnums.h"
 #include "AvaViewportSettings.generated.h"
 
 class UMaterial;
+class UTexture;
 
 USTRUCT(BlueprintType)
 struct FAvaLevelViewportSafeFrame
@@ -168,6 +170,18 @@ public:
 	UPROPERTY(Config, EditAnywhere, BlueprintReadWrite, Category = "Camera Bounds")
 	FLinearColor CameraBoundsShadeColor;
 
+	UPROPERTY(Config, EditAnywhere, BlueprintReadWrite, Category = "Texture Overlay")
+	bool bEnableTextureOverlay = false;
+
+	UPROPERTY(Config, EditAnywhere, BlueprintReadWrite, Category = "Texture Overlay")
+	TSoftObjectPtr<UTexture> TextureOverlayTexture;
+
+	UPROPERTY(Config, EditAnywhere, BlueprintReadWrite, Category = "Texture Overlay", Meta = (UIMin = 0.01, UIMax = 1, ClampMin = 0, ClampMax = 1))
+	float TextureOverlayOpacity = 0.25;
+
+	UPROPERTY(Config, EditAnywhere, BlueprintReadWrite, Category = "Texture Overlay")
+	bool bTextureOverlayStretch = false;
+
 	UFUNCTION(BlueprintPure, Category = "Motion Design")
 	EAvaViewportSnapState GetSnapState() const;
 
@@ -180,5 +194,7 @@ public:
 	DECLARE_MULTICAST_DELEGATE_TwoParams(FOnAvaViewportSettingsChanged, const UAvaViewportSettings* /* This */, FName /* Setting name */)
 	FOnAvaViewportSettingsChanged OnChange;
 
+	void BroadcastSettingChanged(FName InSettingName);
+
 	virtual void PostEditChangeProperty(FPropertyChangedEvent& PropertyChangedEvent) override;
 };
