Motion Design
- Fixed a bug with the background post process visualizer so it renders in the correct place in all dpi scales.

diff --git a/Engine/Plugins/Experimental/Avalanche/Source/AvalancheViewport/Private/AvaViewportPostProcessManager.cpp b/Engine/Plugins/Experimental/Avalanche/Source/AvalancheViewport/Private/AvaViewportPostProcessManager.cpp
--- a/Engine/Plugins/Experimental/Avalanche/Source/AvalancheViewport/Private/AvaViewportPostProcessManager.cpp
+++ b/Engine/Plugins/Experimental/Avalanche/Source/AvalancheViewport/Private/AvaViewportPostProcessManager.cpp
@@ -75,6 +75,7 @@ void FAvaViewportPostProcessManager::UpdateSceneView(FSceneView* InSceneView)
 
 	Visualizer->UpdateForViewport(
 		AvaViewportClient->GetZoomedVisibleArea(),
+		AvaViewportClient->GetViewportOffset(),
 		AvaViewportClient->GetViewportWidgetSize(),
 		PanOffset
 	);
diff --git a/Engine/Plugins/Experimental/Avalanche/Source/AvalancheViewport/Private/Visualizers/AvaViewportBackgroundVisualizer.cpp b/Engine/Plugins/Experimental/Avalanche/Source/AvalancheViewport/Private/Visualizers/AvaViewportBackgroundVisualizer.cpp
--- a/Engine/Plugins/Experimental/Avalanche/Source/AvalancheViewport/Private/Visualizers/AvaViewportBackgroundVisualizer.cpp
+++ b/Engine/Plugins/Experimental/Avalanche/Source/AvalancheViewport/Private/Visualizers/AvaViewportBackgroundVisualizer.cpp
@@ -83,8 +83,8 @@ FString FAvaViewportBackgroundVisualizer::GetReferencerName() const
 	return UE::AvaViewport::Private::BackgroundReferencerName;
 }
 
-void FAvaViewportBackgroundVisualizer::UpdateForViewport(const FAvaVisibleArea& InVisibleArea, const FVector2f& InWidgetSize, 
-	const FVector2f& InCameraOffset)
+void FAvaViewportBackgroundVisualizer::UpdateForViewport(const FAvaVisibleArea& InVisibleArea, const FVector2f& InVisibleAreaOffset, 
+	const FVector2f& InWidgetSize, const FVector2f& InCameraOffset)
 {
 	if (FMath::IsNearlyZero(PostProcessOpacity) || !Texture || !PostProcessMaterial)
 	{
@@ -110,9 +110,6 @@ void FAvaViewportBackgroundVisualizer::UpdateForViewport(const FAvaVisibleArea&
 
 	using namespace UE::AvaViewport::Private;
 
-	const float ImageAspectRatio = ImageSize.X / ImageSize.Y;
-	const float WidgetAspectRatio = InWidgetSize.X / InWidgetSize.Y;
-	const float ViewportAspectRatio = InVisibleArea.AbsoluteSize.X / InVisibleArea.AbsoluteSize.Y;
 	const FVector2f WidgetBasedScale = InVisibleArea.AbsoluteSize / InWidgetSize;
 	const float VisibleAreaFraction = InVisibleArea.GetVisibleAreaFraction();
 	const FVector2f Scale = WidgetBasedScale / VisibleAreaFraction;
@@ -126,37 +123,13 @@ void FAvaViewportBackgroundVisualizer::UpdateForViewport(const FAvaVisibleArea&
 		PostProcessMaterial->SetVectorParameterValue(TextureScaleName, TextureScale);
 	}
 
-	FVector2f Offset = ((InWidgetSize - InVisibleArea.AbsoluteSize) * 0.5f)
-		+ InCameraOffset * Scale;
+	const FVector2f TopLeft = (InVisibleArea.GetVisiblePosition(FVector2f::ZeroVector) + InVisibleAreaOffset) * InVisibleArea.DPIScale;
 
-	if (InVisibleArea.IsZoomedView())
+	if (!FMath::IsNearlyEqual(TextureOffset.X, TopLeft.X)
+		|| !FMath::IsNearlyEqual(TextureOffset.Y, TopLeft.Y))
 	{
-		Offset -= InVisibleArea.GetInvisibleSize() / InVisibleArea.VisibleSize * InVisibleArea.AbsoluteSize * 0.5f;
-	}
-
-	if (!FMath::IsNearlyEqual(WidgetAspectRatio, ViewportAspectRatio))
-	{
-		if (WidgetAspectRatio > ViewportAspectRatio)
-		{
-			const float Distance = InWidgetSize.X - InVisibleArea.AbsoluteSize.X;
-			const float Scalar = InVisibleArea.VisibleSize.X;
-			const float OffsetX = Distance * InCameraOffset.X / Scalar;
-			Offset.X += OffsetX * Scale.X * InVisibleArea.GetVisibleAreaFraction();
-		}
-		else
-		{
-			const float Distance = InWidgetSize.Y - InVisibleArea.AbsoluteSize.Y;
-			const float Scalar = InVisibleArea.VisibleSize.Y;
-			const float OffsetY = Distance * InCameraOffset.Y / Scalar;
-			Offset.Y += OffsetY * Scale.Y * InVisibleArea.GetVisibleAreaFraction();
-		}
-	}
-
-	if (!FMath::IsNearlyEqual(TextureOffset.X, Offset.X)
-		|| !FMath::IsNearlyEqual(TextureOffset.Y, Offset.Y))
-	{
-		TextureOffset.X = Offset.X;
-		TextureOffset.Y = Offset.Y;
+		TextureOffset.X = TopLeft.X;
+		TextureOffset.Y = TopLeft.Y;
 		TextureOffset.Z = 0.f;
 		PostProcessMaterial->SetVectorParameterValue(TextureOffsetName, TextureOffset);
 	}
diff --git a/Engine/Plugins/Experimental/Avalanche/Source/AvalancheViewport/Private/Visualizers/AvaViewportBackgroundVisualizer.h b/Engine/Plugins/Experimental/Avalanche/Source/AvalancheViewport/Private/Visualizers/AvaViewportBackgroundVisualizer.h
--- a/Engine/Plugins/Experimental/Avalanche/Source/AvalancheViewport/Private/Visualizers/AvaViewportBackgroundVisualizer.h
+++ b/Engine/Plugins/Experimental/Avalanche/Source/AvalancheViewport/Private/Visualizers/AvaViewportBackgroundVisualizer.h
@@ -20,7 +20,8 @@ public:
 	void SetTexture(UTexture* InTexture);
 
 	//~ Begin FAvaViewportPostProcessVisualizer
-	virtual void UpdateForViewport(const FAvaVisibleArea& InVisibleArea, const FVector2f& InWidgetSize, const FVector2f& InCameraOffset) override;
+	virtual void UpdateForViewport(const FAvaVisibleArea& InVisibleArea, const FVector2f& InVisibleAreaOffset, 
+		const FVector2f& InWidgetSize, const FVector2f& InCameraOffset) override;
 	//~ End FAvaViewportPostProcessVisualizer
 
 	//~ Begin FGCObject
diff --git a/Engine/Plugins/Experimental/Avalanche/Source/AvalancheViewport/Private/Visualizers/AvaViewportPostProcessVisualizer.cpp b/Engine/Plugins/Experimental/Avalanche/Source/AvalancheViewport/Private/Visualizers/AvaViewportPostProcessVisualizer.cpp
--- a/Engine/Plugins/Experimental/Avalanche/Source/AvalancheViewport/Private/Visualizers/AvaViewportPostProcessVisualizer.cpp
+++ b/Engine/Plugins/Experimental/Avalanche/Source/AvalancheViewport/Private/Visualizers/AvaViewportPostProcessVisualizer.cpp
@@ -112,7 +112,8 @@ void FAvaViewportPostProcessVisualizer::OnDeactivate()
 {
 }
 
-void FAvaViewportPostProcessVisualizer::UpdateForViewport(const FAvaVisibleArea& InVisibleArea, const FVector2f& InWidgetSize, const FVector2f& InCameraOffset)
+void FAvaViewportPostProcessVisualizer::UpdateForViewport(const FAvaVisibleArea& InVisibleArea, const FVector2f& InVisibleAreaOffset, 
+	const FVector2f& InWidgetSize, const FVector2f& InCameraOffset)
 {
 }
 
diff --git a/Engine/Plugins/Experimental/Avalanche/Source/AvalancheViewport/Private/Visualizers/AvaViewportPostProcessVisualizer.h b/Engine/Plugins/Experimental/Avalanche/Source/AvalancheViewport/Private/Visualizers/AvaViewportPostProcessVisualizer.h
--- a/Engine/Plugins/Experimental/Avalanche/Source/AvalancheViewport/Private/Visualizers/AvaViewportPostProcessVisualizer.h
+++ b/Engine/Plugins/Experimental/Avalanche/Source/AvalancheViewport/Private/Visualizers/AvaViewportPostProcessVisualizer.h
@@ -42,7 +42,8 @@ public:
 	virtual void OnActivate();
 	virtual void OnDeactivate();
 
-	virtual void UpdateForViewport(const FAvaVisibleArea& InVisibleArea, const FVector2f& InWidgetSize, const FVector2f& InCameraOffset);
+	virtual void UpdateForViewport(const FAvaVisibleArea& InVisibleArea, const FVector2f& InVisibleAreaOffset, 
+		const FVector2f& InWidgetSize, const FVector2f& InCameraOffset);
 
 	void ApplyToSceneView(FSceneView* InSceneView) const;
 
