Motion Design
- Added NotPlaceable and Hidden to a few internal actors

diff --git a/Engine/Plugins/Experimental/Avalanche/Source/Avalanche/Internal/Tests/Framework/AvaTestStaticMeshActor.h b/Engine/Plugins/Experimental/Avalanche/Source/Avalanche/Internal/Tests/Framework/AvaTestStaticMeshActor.h
--- a/Engine/Plugins/Experimental/Avalanche/Source/Avalanche/Internal/Tests/Framework/AvaTestStaticMeshActor.h
+++ b/Engine/Plugins/Experimental/Avalanche/Source/Avalanche/Internal/Tests/Framework/AvaTestStaticMeshActor.h
@@ -7,7 +7,7 @@
 
 class UStaticMeshComponent;
 
-UCLASS(MinimalAPI, DisplayName = "Motion Design Test Static Mesh Actor")
+UCLASS(NotPlaceable, Hidden, DisplayName = "Motion Design Test Static Mesh Actor")
 class AAvaTestStaticMeshActor : public AActor
 {
 	GENERATED_BODY()
diff --git a/Engine/Plugins/Experimental/Avalanche/Source/AvalancheSequence/Private/AvaSequencePlaybackActor.h b/Engine/Plugins/Experimental/Avalanche/Source/AvalancheSequence/Private/AvaSequencePlaybackActor.h
--- a/Engine/Plugins/Experimental/Avalanche/Source/AvalancheSequence/Private/AvaSequencePlaybackActor.h
+++ b/Engine/Plugins/Experimental/Avalanche/Source/AvalancheSequence/Private/AvaSequencePlaybackActor.h
@@ -13,7 +13,7 @@ class UAvaSequencePlayer;
 class UMovieSceneSequenceTickManager;
 
 /** Base Actor for Ava Sequence Playback Management */
-UCLASS(DisplayName = "Motion Design Sequence Playback Actor")
+UCLASS(NotPlaceable, Hidden, DisplayName = "Motion Design Sequence Playback Actor")
 class AAvaSequencePlaybackActor : public AActor, public IAvaSequencePlaybackObject
 {
 	GENERATED_BODY()
diff --git a/Engine/Plugins/Experimental/Avalanche/Source/AvalancheTransition/Private/Behavior/AvaTransitionBehaviorActor.h b/Engine/Plugins/Experimental/Avalanche/Source/AvalancheTransition/Private/Behavior/AvaTransitionBehaviorActor.h
--- a/Engine/Plugins/Experimental/Avalanche/Source/AvalancheTransition/Private/Behavior/AvaTransitionBehaviorActor.h
+++ b/Engine/Plugins/Experimental/Avalanche/Source/AvalancheTransition/Private/Behavior/AvaTransitionBehaviorActor.h
@@ -14,7 +14,7 @@ class UAvaTransitionBehaviorInstance;
 class UAvaTransitionSubsystem;
 struct FStateTreeExecutionContext;
 
-UCLASS(DisplayName = "Motion Design Transition Behavior Actor")
+UCLASS(NotPlaceable, Hidden, DisplayName = "Motion Design Transition Behavior Actor")
 class AAvaTransitionBehaviorActor : public AActor, public IAvaTransitionBehavior
 {
 	GENERATED_BODY()
diff --git a/Engine/Plugins/Experimental/Avalanche/Source/AvalancheViewport/Private/AvaViewportDataActor.h b/Engine/Plugins/Experimental/Avalanche/Source/AvalancheViewport/Private/AvaViewportDataActor.h
--- a/Engine/Plugins/Experimental/Avalanche/Source/AvalancheViewport/Private/AvaViewportDataActor.h
+++ b/Engine/Plugins/Experimental/Avalanche/Source/AvalancheViewport/Private/AvaViewportDataActor.h
@@ -12,7 +12,7 @@
 #include "Viewport/Interaction/AvaViewportPostProcessInfo.h"
 #include "AvaViewportDataActor.generated.h"
 
-UCLASS(DisplayName = "Motion Design Viewport Data Actor")
+UCLASS(NotPlaceable, Hidden, DisplayName = "Motion Design Viewport Data Actor")
 class AAvaViewportDataActor : public AActor
 {
 	GENERATED_BODY()
