Motion Design
- Set exposure compensation override back to 1 on the motion design post process volume.

diff --git a/Engine/Plugins/Experimental/Avalanche/Source/Avalanche/Private/Viewport/AvaPostProcessVolume.cpp b/Engine/Plugins/Experimental/Avalanche/Source/Avalanche/Private/Viewport/AvaPostProcessVolume.cpp
--- a/Engine/Plugins/Experimental/Avalanche/Source/Avalanche/Private/Viewport/AvaPostProcessVolume.cpp
+++ b/Engine/Plugins/Experimental/Avalanche/Source/Avalanche/Private/Viewport/AvaPostProcessVolume.cpp
@@ -1,4 +1,4 @@
-﻿// Copyright Epic Games, Inc. All Rights Reserved.
+// Copyright Epic Games, Inc. All Rights Reserved.
 
 #include "Viewport/AvaPostProcessVolume.h"
 
@@ -8,7 +8,7 @@ AAvaPostProcessVolume::AAvaPostProcessVolume()
 	bUnbound = true;
 
 	Settings.bOverride_AutoExposureBias = 1;
-	Settings.AutoExposureBias = 0.f;
+	Settings.AutoExposureBias = 1.f;
 
 	Settings.bOverride_AutoExposureMinBrightness = 1;
 	Settings.AutoExposureMinBrightness = 1.f;
