MotionDesign : GridArrange
- Fix crash issue when count is set to 0,0 through setter or direct property access

diff --git a/Engine/Plugins/Experimental/Avalanche/Source/AvalancheModifiers/Private/Modifiers/AvaGridArrangeModifier.cpp b/Engine/Plugins/Experimental/Avalanche/Source/AvalancheModifiers/Private/Modifiers/AvaGridArrangeModifier.cpp
--- a/Engine/Plugins/Experimental/Avalanche/Source/AvalancheModifiers/Private/Modifiers/AvaGridArrangeModifier.cpp
+++ b/Engine/Plugins/Experimental/Avalanche/Source/AvalancheModifiers/Private/Modifiers/AvaGridArrangeModifier.cpp
@@ -21,7 +21,7 @@ void UAvaGridArrangeModifier::PostEditChangeProperty(FPropertyChangedEvent& Prop
 	static const FName SpreadPropertyName = GET_MEMBER_NAME_CHECKED(UAvaGridArrangeModifier, Spread);
 	static const FName StartCornerPropertyName = GET_MEMBER_NAME_CHECKED(UAvaGridArrangeModifier, StartCorner);
 	static const FName StartDirectionPropertyName = GET_MEMBER_NAME_CHECKED(UAvaGridArrangeModifier, StartDirection);
-	
+
 	if (MemberName == CountPropertyName
 		|| MemberName == SpreadPropertyName
 		|| MemberName == StartCornerPropertyName
@@ -32,13 +32,15 @@ void UAvaGridArrangeModifier::PostEditChangeProperty(FPropertyChangedEvent& Prop
 }
 #endif
 
-void UAvaGridArrangeModifier::SetCount(const FIntPoint& NewCount)
+void UAvaGridArrangeModifier::SetCount(const FIntPoint& InCount)
 {
+	const FIntPoint NewCount = InCount.ComponentMax(FIntPoint(1, 1));
+
 	if (Count == NewCount)
 	{
 		return;
 	}
-	
+
 	Count = NewCount;
 	MarkModifierDirty();
 }
@@ -49,7 +51,7 @@ void UAvaGridArrangeModifier::SetSpread(const FVector2D& NewSpread)
 	{
 		return;
 	}
-	
+
 	Spread = NewSpread;
 	MarkModifierDirty();
 }
@@ -60,7 +62,7 @@ void UAvaGridArrangeModifier::SetStartCorner(EAvaCorner2D InCorner)
 	{
 		return;
 	}
-	
+
 	StartCorner = InCorner;
 	MarkModifierDirty();
 }
@@ -118,7 +120,13 @@ void UAvaGridArrangeModifier::Apply()
 		Fail(LOCTEXT("InvalidSceneExtension", "Scene extension could not be found"));
 		return;
 	}
-	
+
+	if (Count.X < 1 || Count.Y < 1)
+	{
+		Fail(LOCTEXT("InvalidGridCount", "Count must be greater than 0"));
+		return;
+	}
+
 	const TArray<TWeakObjectPtr<AActor>> AttachedActors = SceneExtension->GetDirectChildrenActor(ModifyActor);
 	const int32 TotalSlotCount = Count.X * Count.Y;
 	const int32 AttachedActorCount = AttachedActors.Num();
@@ -151,7 +159,7 @@ void UAvaGridArrangeModifier::Apply()
 		{
 			continue;
 		}
-		
+
 		{
 			// Track all new children actors
 			TArray<AActor*> ChildrenActors { AttachedActor };
@@ -167,7 +175,7 @@ void UAvaGridArrangeModifier::Apply()
 		{
 			continue;
 		}
-		
+
 		// Track this actor visibility state
 		const bool bIsVisible = ChildIndex < TotalSlotCount;
 		VisibilityShared->SetActorVisibility(this, AttachedActor, !bIsVisible, true);
