Motion Design
- Fixed horde build error

diff --git a/Engine/Plugins/Experimental/Avalanche/Source/AvalancheMediaEditor/Private/Broadcast/ChannelGrid/Slate/SAvaBroadcastChannel.cpp b/Engine/Plugins/Experimental/Avalanche/Source/AvalancheMediaEditor/Private/Broadcast/ChannelGrid/Slate/SAvaBroadcastChannel.cpp
--- a/Engine/Plugins/Experimental/Avalanche/Source/AvalancheMediaEditor/Private/Broadcast/ChannelGrid/Slate/SAvaBroadcastChannel.cpp
+++ b/Engine/Plugins/Experimental/Avalanche/Source/AvalancheMediaEditor/Private/Broadcast/ChannelGrid/Slate/SAvaBroadcastChannel.cpp
@@ -19,6 +19,7 @@
 #include "Layout/Geometry.h"
 #include "Layout/Visibility.h"
 #include "MediaOutput.h"
+#include "Modules/ModuleManager.h"
 #include "PropertyEditorModule.h"
 #include "ScopedTransaction.h"
 #include "Styling/SlateBrush.h"
