Motion Design
- Horde error

diff --git a/Engine/Plugins/Experimental/Avalanche/Source/AvalancheMediaEditor/Private/Rundown/AssetDefinition_AvaRundown.cpp b/Engine/Plugins/Experimental/Avalanche/Source/AvalancheMediaEditor/Private/Rundown/AssetDefinition_AvaRundown.cpp
--- a/Engine/Plugins/Experimental/Avalanche/Source/AvalancheMediaEditor/Private/Rundown/AssetDefinition_AvaRundown.cpp
+++ b/Engine/Plugins/Experimental/Avalanche/Source/AvalancheMediaEditor/Private/Rundown/AssetDefinition_AvaRundown.cpp
@@ -4,11 +4,14 @@
 
 #include "AvaMediaEditorStyle.h"
 #include "ContentBrowserMenuContexts.h"
+#include "Framework/Commands/UIAction.h"
 #include "IAvaMediaEditorModule.h"
 #include "Misc/MessageDialog.h"
 #include "Rundown/AvaRundown.h"
 #include "Rundown/AvaRundownEditor.h"
 #include "Rundown/AvaRundownEditorUtils.h"
+#include "ToolMenu.h"
+#include "ToolMenus.h"
 
 #define LOCTEXT_NAMESPACE "UAssetDefinition_AvaRundown"
 
