Motion Design
- Added some logging to the post process visualizer system to track down errors.
- Consolidated post process menu methods down to a single set with a parameter.

diff --git a/Engine/Plugins/Experimental/Avalanche/Source/AvalancheLevelViewport/Private/SAvaLevelViewport.cpp b/Engine/Plugins/Experimental/Avalanche/Source/AvalancheLevelViewport/Private/SAvaLevelViewport.cpp
--- a/Engine/Plugins/Experimental/Avalanche/Source/AvalancheLevelViewport/Private/SAvaLevelViewport.cpp
+++ b/Engine/Plugins/Experimental/Avalanche/Source/AvalancheLevelViewport/Private/SAvaLevelViewport.cpp
@@ -71,12 +71,12 @@ int32 SAvaLevelViewport::OnPaint(const FPaintArgs& Args, const FGeometry& Allott
 {
 	const_cast<SAvaLevelViewport*>(this)->CheckVirtualSizeCameraUpdateSettings();
 
-	return SLevelViewport::OnPaint(Args, AllottedGeometry, MyCullingRect, OutDrawElements, LayerId, InWidgetStyle, bParentEnabled);
+	return Super::OnPaint(Args, AllottedGeometry, MyCullingRect, OutDrawElements, LayerId, InWidgetStyle, bParentEnabled);
 }
 
 void SAvaLevelViewport::BindCommands()
 {
-	SLevelViewport::BindCommands();
+	Super::BindCommands();
 
 	const FEditorViewportCommands& EditorViewportCommands = FEditorViewportCommands::Get();
 	FUICommandList& CommandListRef = *CommandList;
@@ -112,51 +112,51 @@ void SAvaLevelViewport::BindCommands()
 
 	CommandListRef.MapAction(
 		AvaLevelViewportCommands.TogglePostProcessNone,
-		FExecuteAction::CreateSP(this, &SAvaLevelViewport::ExecuteTogglePostProcessNone),
-		FCanExecuteAction::CreateSP(this, &SAvaLevelViewport::CanTogglePostProcessNone),
-		FIsActionChecked::CreateSP(this, &SAvaLevelViewport::IsPostProcessNoneEnabled)
+		FExecuteAction::CreateSP(this, &SAvaLevelViewport::ExecuteTogglePostProcessType, EAvaViewportPostProcessType::None),
+		FCanExecuteAction::CreateSP(this, &SAvaLevelViewport::CanTogglePostProcessType, EAvaViewportPostProcessType::None),
+		FIsActionChecked::CreateSP(this, &SAvaLevelViewport::IsPostProcessTypeEnabled, EAvaViewportPostProcessType::None)
 	);
 
 	CommandListRef.MapAction(
 		AvaLevelViewportCommands.TogglePostProcessBackground,
-		FExecuteAction::CreateSP(this, &SAvaLevelViewport::ExecuteTogglePostProcessBackground),
-		FCanExecuteAction::CreateSP(this, &SAvaLevelViewport::CanTogglePostProcessBackground),
-		FIsActionChecked::CreateSP(this, &SAvaLevelViewport::IsPostProcessBackgroundEnabled)
+		FExecuteAction::CreateSP(this, &SAvaLevelViewport::ExecuteTogglePostProcessType, EAvaViewportPostProcessType::Background),
+		FCanExecuteAction::CreateSP(this, &SAvaLevelViewport::CanTogglePostProcessType, EAvaViewportPostProcessType::Background),
+		FIsActionChecked::CreateSP(this, &SAvaLevelViewport::IsPostProcessTypeEnabled, EAvaViewportPostProcessType::Background)
 	);
 
 	CommandListRef.MapAction(
 		AvaLevelViewportCommands.TogglePostProcessChannelRed,
-		FExecuteAction::CreateSP(this, &SAvaLevelViewport::ExecuteTogglePostProcessChannelRed),
-		FCanExecuteAction::CreateSP(this, &SAvaLevelViewport::CanTogglePostProcessChannelRed),
-		FIsActionChecked::CreateSP(this, &SAvaLevelViewport::IsPostProcessChannelRedEnabled)
+		FExecuteAction::CreateSP(this, &SAvaLevelViewport::ExecuteTogglePostProcessType, EAvaViewportPostProcessType::RedChannel),
+		FCanExecuteAction::CreateSP(this, &SAvaLevelViewport::CanTogglePostProcessType, EAvaViewportPostProcessType::RedChannel),
+		FIsActionChecked::CreateSP(this, &SAvaLevelViewport::IsPostProcessTypeEnabled, EAvaViewportPostProcessType::RedChannel)
 	);
 
 	CommandListRef.MapAction(
 		AvaLevelViewportCommands.TogglePostProcessChannelGreen,
-		FExecuteAction::CreateSP(this, &SAvaLevelViewport::ExecuteTogglePostProcessChannelGreen),
-		FCanExecuteAction::CreateSP(this, &SAvaLevelViewport::CanTogglePostProcessChannelGreen),
-		FIsActionChecked::CreateSP(this, &SAvaLevelViewport::IsPostProcessChannelGreenEnabled)
+		FExecuteAction::CreateSP(this, &SAvaLevelViewport::ExecuteTogglePostProcessType, EAvaViewportPostProcessType::GreenChannel),
+		FCanExecuteAction::CreateSP(this, &SAvaLevelViewport::CanTogglePostProcessType, EAvaViewportPostProcessType::GreenChannel),
+		FIsActionChecked::CreateSP(this, &SAvaLevelViewport::IsPostProcessTypeEnabled, EAvaViewportPostProcessType::GreenChannel)
 	);
 
 	CommandListRef.MapAction(
 		AvaLevelViewportCommands.TogglePostProcessChannelBlue,
-		FExecuteAction::CreateSP(this, &SAvaLevelViewport::ExecuteTogglePostProcessChannelBlue),
-		FCanExecuteAction::CreateSP(this, &SAvaLevelViewport::CanTogglePostProcessChannelBlue),
-		FIsActionChecked::CreateSP(this, &SAvaLevelViewport::IsPostProcessChannelBlueEnabled)
+		FExecuteAction::CreateSP(this, &SAvaLevelViewport::ExecuteTogglePostProcessType, EAvaViewportPostProcessType::BlueChannel),
+		FCanExecuteAction::CreateSP(this, &SAvaLevelViewport::CanTogglePostProcessType, EAvaViewportPostProcessType::BlueChannel),
+		FIsActionChecked::CreateSP(this, &SAvaLevelViewport::IsPostProcessTypeEnabled, EAvaViewportPostProcessType::BlueChannel)
 	);
 
 	CommandListRef.MapAction(
 		AvaLevelViewportCommands.TogglePostProcessChannelAlpha,
-		FExecuteAction::CreateSP(this, &SAvaLevelViewport::ExecuteTogglePostProcessChannelAlpha),
-		FCanExecuteAction::CreateSP(this, &SAvaLevelViewport::CanTogglePostProcessChannelAlpha),
-		FIsActionChecked::CreateSP(this, &SAvaLevelViewport::IsPostProcessChannelAlphaEnabled)
+		FExecuteAction::CreateSP(this, &SAvaLevelViewport::ExecuteTogglePostProcessType, EAvaViewportPostProcessType::AlphaChannel),
+		FCanExecuteAction::CreateSP(this, &SAvaLevelViewport::CanTogglePostProcessType, EAvaViewportPostProcessType::AlphaChannel),
+		FIsActionChecked::CreateSP(this, &SAvaLevelViewport::IsPostProcessTypeEnabled, EAvaViewportPostProcessType::AlphaChannel)
 	);
 
 	CommandListRef.MapAction(
 		AvaLevelViewportCommands.TogglePostProcessCheckerboard,
-		FExecuteAction::CreateSP(this, &SAvaLevelViewport::ExecuteTogglePostProcessCheckerboard),
-		FCanExecuteAction::CreateSP(this, &SAvaLevelViewport::CanTogglePostProcessCheckerboard),
-		FIsActionChecked::CreateSP(this, &SAvaLevelViewport::IsPostProcessCheckerboardEnabled)
+		FExecuteAction::CreateSP(this, &SAvaLevelViewport::ExecuteTogglePostProcessType, EAvaViewportPostProcessType::Checkerboard),
+		FCanExecuteAction::CreateSP(this, &SAvaLevelViewport::CanTogglePostProcessType, EAvaViewportPostProcessType::Checkerboard),
+		FIsActionChecked::CreateSP(this, &SAvaLevelViewport::IsPostProcessTypeEnabled, EAvaViewportPostProcessType::Checkerboard)
 	);
 	CommandListRef.MapAction(
 		AvaLevelViewportCommands.ToggleTextureOverlay,
@@ -331,7 +331,7 @@ void SAvaLevelViewport::BindCommands()
 
 void SAvaLevelViewport::PopulateViewportOverlays(TSharedRef<SOverlay> Overlay)
 {
-	SLevelViewport::PopulateViewportOverlays(Overlay);
+	Super::PopulateViewportOverlays(Overlay);
 
 	if (TSharedPtr<FAvaLevelViewportClient> ViewportClient = GetAvaLevelViewportClient())
 	{
diff --git a/Engine/Plugins/Experimental/Avalanche/Source/AvalancheLevelViewport/Private/SAvaLevelViewport.h b/Engine/Plugins/Experimental/Avalanche/Source/AvalancheLevelViewport/Private/SAvaLevelViewport.h
--- a/Engine/Plugins/Experimental/Avalanche/Source/AvalancheLevelViewport/Private/SAvaLevelViewport.h
+++ b/Engine/Plugins/Experimental/Avalanche/Source/AvalancheLevelViewport/Private/SAvaLevelViewport.h
@@ -29,6 +29,7 @@ class SOverlay;
 class SWidget;
 class UAvaViewportSettings;
 class UToolMenu;
+enum class EAvaViewportPostProcessType : uint8;
 enum class EAvaViewportVirtualSizeAspectRatioState : uint8;
 enum class ECheckBoxState : uint8;
 struct FAssetEditorViewportConstructionArgs;
@@ -128,33 +129,9 @@ public:
 
 	void ExecuteToggleChildActorLock();
 
-	bool IsPostProcessNoneEnabled() const;
-	bool CanTogglePostProcessNone() const;
-	void ExecuteTogglePostProcessNone();
-
-	bool IsPostProcessBackgroundEnabled() const;
-	bool CanTogglePostProcessBackground() const;
-	void ExecuteTogglePostProcessBackground();
-
-	bool IsPostProcessChannelRedEnabled() const;
-	bool CanTogglePostProcessChannelRed() const;
-	void ExecuteTogglePostProcessChannelRed();
-
-	bool IsPostProcessChannelGreenEnabled() const;
-	bool CanTogglePostProcessChannelGreen() const;
-	void ExecuteTogglePostProcessChannelGreen();
-
-	bool IsPostProcessChannelBlueEnabled() const;
-	bool CanTogglePostProcessChannelBlue() const;
-	void ExecuteTogglePostProcessChannelBlue();
-
-	bool IsPostProcessChannelAlphaEnabled() const;
-	bool CanTogglePostProcessChannelAlpha() const;
-	void ExecuteTogglePostProcessChannelAlpha();
-
-	bool IsPostProcessCheckerboardEnabled() const;
-	bool CanTogglePostProcessCheckerboard() const;
-	void ExecuteTogglePostProcessCheckerboard();
+	bool IsPostProcessTypeEnabled(EAvaViewportPostProcessType InPostProcessType) const;
+	bool CanTogglePostProcessType(EAvaViewportPostProcessType InPostProcessType) const;
+	void ExecuteTogglePostProcessType(EAvaViewportPostProcessType InPostProcessType);
 
 	bool CanToggleTextureOverlay() const;
 	void ExecuteToggleTextureOverlay();
diff --git a/Engine/Plugins/Experimental/Avalanche/Source/AvalancheLevelViewport/Private/SAvaLevelViewport_StatusBar.cpp b/Engine/Plugins/Experimental/Avalanche/Source/AvalancheLevelViewport/Private/SAvaLevelViewport_StatusBar.cpp
--- a/Engine/Plugins/Experimental/Avalanche/Source/AvalancheLevelViewport/Private/SAvaLevelViewport_StatusBar.cpp
+++ b/Engine/Plugins/Experimental/Avalanche/Source/AvalancheLevelViewport/Private/SAvaLevelViewport_StatusBar.cpp
@@ -1,5 +1,6 @@
 // Copyright Epic Games, Inc. All Rights Reserved.
 
+#include "AvaLevelViewportModule.h"
 #include "SAvaLevelViewport.h"
 #include "AvaViewportPostProcessManager.h"
 #include "AvaViewportSettings.h"
@@ -24,270 +25,20 @@ void SAvaLevelViewport::ExecuteToggleChildActorLock()
 	ViewportClient->SetChildActorsLocked(!ViewportClient->AreChildActorsLocked());
 }
 
-bool SAvaLevelViewport::IsPostProcessNoneEnabled() const
+bool SAvaLevelViewport::IsPostProcessTypeEnabled(EAvaViewportPostProcessType InPostProcessType) const
 {
 	TSharedPtr<FAvaLevelViewportClient> ViewportClient = GetAvaLevelViewportClient();
 
-	if (!ViewportClient.IsValid())
-	{
-		// Make sure none is toggled on if there's an error.
-		return true;
-	}
-
-	if (!ViewportClient->GetPostProcessManager().IsValid())
-	{
-		return true;
-	}
-
-	return ViewportClient->GetPostProcessManager()->GetType() == EAvaViewportPostProcessType::None;
-}
-
-bool SAvaLevelViewport::CanTogglePostProcessNone() const
-{
-	return GetAvaLevelViewportClient().IsValid();
-}
-
-void SAvaLevelViewport::ExecuteTogglePostProcessNone()
-{
-	TSharedPtr<FAvaLevelViewportClient> ViewportClient = GetAvaLevelViewportClient();
-
-	if (!ViewportClient.IsValid())
-	{
-		return;
-	}
-
-	if (!ViewportClient->GetPostProcessManager().IsValid())
-	{
-		return;
-	}
-
-	BeginPostProcessInfoTransaction();
-
-	ViewportClient->GetPostProcessManager()->SetType(EAvaViewportPostProcessType::None);
-	ViewportClient->Invalidate();
-
-	EndPostProcessInfoTransaction();
-}
-
-bool SAvaLevelViewport::IsPostProcessBackgroundEnabled() const
-{
-	TSharedPtr<FAvaLevelViewportClient> ViewportClient = GetAvaLevelViewportClient();
-
-	if (!ViewportClient.IsValid())
-	{
-		return false;
-	}
-
-	if (!ViewportClient->GetPostProcessManager().IsValid())
-	{
-		return false;
-	}
-
-	return ViewportClient->GetPostProcessManager()->GetType() == EAvaViewportPostProcessType::Background;
-}
-
-bool SAvaLevelViewport::CanTogglePostProcessBackground() const
-{
-	TSharedPtr<FAvaLevelViewportClient> ViewportClient = GetAvaLevelViewportClient();
-
-	if (!ViewportClient.IsValid())
-	{
-		return false;
-	}
-
-	if (!ViewportClient->GetPostProcessManager().IsValid())
-	{
-		return false;
-	}
-
-	if (IsPostProcessBackgroundEnabled())
-	{
-		return true;
-	}
-
-	if (TSharedPtr<IAvaViewportPostProcessVisualizer> BackgroundVisualizer = ViewportClient->GetPostProcessManager()->GetVisualizer(EAvaViewportPostProcessType::Background))
-	{
-		return BackgroundVisualizer->CanActivate(/* bInSilent */ true);
-	}
-
-	return false;
-}
-
-void SAvaLevelViewport::ExecuteTogglePostProcessBackground()
-{
-	TSharedPtr<FAvaLevelViewportClient> ViewportClient = GetAvaLevelViewportClient();
-
-	if (!ViewportClient.IsValid())
-	{
-		return;
-	}
-
-	if (!ViewportClient->GetPostProcessManager().IsValid())
-	{
-		return;
-	}
-
-	BeginPostProcessInfoTransaction();
-
-	ViewportClient->GetPostProcessManager()->SetType(EAvaViewportPostProcessType::Background);
-	ViewportClient->Invalidate();
-
-	EndPostProcessInfoTransaction();
-}
-
-bool SAvaLevelViewport::IsPostProcessChannelRedEnabled() const
-{
-	TSharedPtr<FAvaLevelViewportClient> ViewportClient = GetAvaLevelViewportClient();
-
-	if (!ViewportClient.IsValid())
-	{
-		// Make sure none is toggled on if there's an error.
-		return true;
-	}
-
-	if (!ViewportClient->GetPostProcessManager().IsValid())
-	{
-		return true;
-	}
-
-	return ViewportClient->GetPostProcessManager()->GetType() == EAvaViewportPostProcessType::RedChannel;
-}
-
-bool SAvaLevelViewport::CanTogglePostProcessChannelRed() const
-{
-	TSharedPtr<FAvaLevelViewportClient> ViewportClient = GetAvaLevelViewportClient();
-
-	if (!ViewportClient.IsValid())
-	{
-		return false;
-	}
-
-	if (!ViewportClient->GetPostProcessManager().IsValid())
-	{
-		return false;
-	}
-
-	if (IsPostProcessChannelRedEnabled())
-	{
-		return true;
-	}
-
-	if (TSharedPtr<IAvaViewportPostProcessVisualizer> ChannelRedVisualizer = ViewportClient->GetPostProcessManager()->GetVisualizer(EAvaViewportPostProcessType::RedChannel))
-	{
-		return ChannelRedVisualizer->CanActivate(/* bInSilent */ true);
-	}
-
-	return false;
-}
-
-void SAvaLevelViewport::ExecuteTogglePostProcessChannelRed()
-{
-	TSharedPtr<FAvaLevelViewportClient> ViewportClient = GetAvaLevelViewportClient();
-
-	if (!ViewportClient.IsValid())
-	{
-		return;
-	}
-
-	if (!ViewportClient->GetPostProcessManager().IsValid())
-	{
-		return;
-	}
-
-	BeginPostProcessInfoTransaction();
-
-	ViewportClient->GetPostProcessManager()->SetType(EAvaViewportPostProcessType::RedChannel);
-	ViewportClient->Invalidate();
-
-	EndPostProcessInfoTransaction();
-}
-
-bool SAvaLevelViewport::IsPostProcessChannelGreenEnabled() const
-{
-	TSharedPtr<FAvaLevelViewportClient> ViewportClient = GetAvaLevelViewportClient();
-
-	if (!ViewportClient.IsValid())
-	{
-		// Make sure none is toggled on if there's an error.
-		return true;
-	}
-
-	if (!ViewportClient->GetPostProcessManager().IsValid())
-	{
-		return true;
-	}
-
-	return ViewportClient->GetPostProcessManager()->GetType() == EAvaViewportPostProcessType::GreenChannel;
-}
-
-bool SAvaLevelViewport::CanTogglePostProcessChannelGreen() const
-{
-	TSharedPtr<FAvaLevelViewportClient> ViewportClient = GetAvaLevelViewportClient();
-
-	if (!ViewportClient.IsValid())
-	{
-		return false;
-	}
-
-	if (!ViewportClient->GetPostProcessManager().IsValid())
-	{
-		return false;
-	}
-
-	if (IsPostProcessChannelGreenEnabled())
-	{
-		return true;
-	}
-
-	if (TSharedPtr<IAvaViewportPostProcessVisualizer> ChannelGreenVisualizer = ViewportClient->GetPostProcessManager()->GetVisualizer(EAvaViewportPostProcessType::GreenChannel))
-	{
-		return ChannelGreenVisualizer->CanActivate(/* bInSilent */ true);
-	}
-
-	return false;
-}
-
-void SAvaLevelViewport::ExecuteTogglePostProcessChannelGreen()
-{
-	TSharedPtr<FAvaLevelViewportClient> ViewportClient = GetAvaLevelViewportClient();
-
-	if (!ViewportClient.IsValid())
-	{
-		return;
-	}
-
-	if (!ViewportClient->GetPostProcessManager().IsValid())
-	{
-		return;
-	}
-
-	BeginPostProcessInfoTransaction();
-
-	ViewportClient->GetPostProcessManager()->SetType(EAvaViewportPostProcessType::GreenChannel);
-	ViewportClient->Invalidate();
-
-	EndPostProcessInfoTransaction();
-}
-
-bool SAvaLevelViewport::IsPostProcessChannelBlueEnabled() const
-{
-	TSharedPtr<FAvaLevelViewportClient> ViewportClient = GetAvaLevelViewportClient();
-
-	if (!ViewportClient.IsValid())
+	if (!ViewportClient.IsValid() || !ViewportClient->GetPostProcessManager().IsValid())
 	{
 		// Make sure none is toggled on if there's an error.
-		return true;
+		return (InPostProcessType == EAvaViewportPostProcessType::None);
 	}
 
-	if (!ViewportClient->GetPostProcessManager().IsValid())
-	{
-		return true;
-	}
-
-	return ViewportClient->GetPostProcessManager()->GetType() == EAvaViewportPostProcessType::BlueChannel;
+	return ViewportClient->GetPostProcessManager()->GetType() == InPostProcessType;
 }
 
-bool SAvaLevelViewport::CanTogglePostProcessChannelBlue() const
+bool SAvaLevelViewport::CanTogglePostProcessType(EAvaViewportPostProcessType InPostProcessType) const
 {
 	TSharedPtr<FAvaLevelViewportClient> ViewportClient = GetAvaLevelViewportClient();
 
@@ -301,158 +52,48 @@ bool SAvaLevelViewport::CanTogglePostProcessChannelBlue() const
 		return false;
 	}
 
-	if (IsPostProcessChannelBlueEnabled())
-	{
-		return true;
-	}
-
-	if (TSharedPtr<IAvaViewportPostProcessVisualizer> ChannelBlueVisualizer = ViewportClient->GetPostProcessManager()->GetVisualizer(EAvaViewportPostProcessType::BlueChannel))
-	{
-		return ChannelBlueVisualizer->CanActivate(/* bInSilent */ true);
-	}
-
-	return false;
-}
-
-void SAvaLevelViewport::ExecuteTogglePostProcessChannelBlue()
-{
-	TSharedPtr<FAvaLevelViewportClient> ViewportClient = GetAvaLevelViewportClient();
-
-	if (!ViewportClient.IsValid())
-	{
-		return;
-	}
-
-	if (!ViewportClient->GetPostProcessManager().IsValid())
-	{
-		return;
-	}
-
-	BeginPostProcessInfoTransaction();
-
-	ViewportClient->GetPostProcessManager()->SetType(EAvaViewportPostProcessType::BlueChannel);
-	ViewportClient->Invalidate();
-
-	EndPostProcessInfoTransaction();
-}
-
-bool SAvaLevelViewport::IsPostProcessChannelAlphaEnabled() const
-{
-	TSharedPtr<FAvaLevelViewportClient> ViewportClient = GetAvaLevelViewportClient();
-
-	if (!ViewportClient.IsValid())
-	{
-		// Make sure none is toggled on if there's an error.
-		return true;
-	}
-
-	if (!ViewportClient->GetPostProcessManager().IsValid())
+	if (InPostProcessType == EAvaViewportPostProcessType::None)
 	{
 		return true;
 	}
 
-	return ViewportClient->GetPostProcessManager()->GetType() == EAvaViewportPostProcessType::AlphaChannel;
-}
-
-bool SAvaLevelViewport::CanTogglePostProcessChannelAlpha() const
-{
-	TSharedPtr<FAvaLevelViewportClient> ViewportClient = GetAvaLevelViewportClient();
-
-	if (!ViewportClient.IsValid())
-	{
-		return false;
-	}
-
-	if (!ViewportClient->GetPostProcessManager().IsValid())
-	{
-		return false;
-	}
-
-	if (IsPostProcessChannelAlphaEnabled())
+	if (IsPostProcessTypeEnabled(InPostProcessType))
 	{
 		return true;
 	}
 
-	if (TSharedPtr<IAvaViewportPostProcessVisualizer> ChannelAlphaVisualizer = ViewportClient->GetPostProcessManager()->GetVisualizer(EAvaViewportPostProcessType::AlphaChannel))
+	if (TSharedPtr<IAvaViewportPostProcessVisualizer> Visualizer = ViewportClient->GetPostProcessManager()->GetVisualizer(InPostProcessType))
 	{
-		return ChannelAlphaVisualizer->CanActivate(/* bInSilent */ true);
+		return Visualizer->CanActivate(/* bInSilent */ true);
 	}
 
 	return false;
 }
 
-void SAvaLevelViewport::ExecuteTogglePostProcessChannelAlpha()
+void SAvaLevelViewport::ExecuteTogglePostProcessType(EAvaViewportPostProcessType InPostProcessType)
 {
 	TSharedPtr<FAvaLevelViewportClient> ViewportClient = GetAvaLevelViewportClient();
 
 	if (!ViewportClient.IsValid())
 	{
+		UE_LOG(AvaLevelViewportLog, Warning, TEXT("SAvaLevelViewport::ExecuteTogglePostProcessType: Invalid viewport client."));
 		return;
 	}
 
 	if (!ViewportClient->GetPostProcessManager().IsValid())
 	{
+		UE_LOG(AvaLevelViewportLog, Warning, TEXT("SAvaLevelViewport::ExecuteTogglePostProcessType: Missing post process manager."));
 		return;
 	}
 
 	BeginPostProcessInfoTransaction();
 
-	ViewportClient->GetPostProcessManager()->SetType(EAvaViewportPostProcessType::AlphaChannel);
+	ViewportClient->GetPostProcessManager()->SetType(InPostProcessType);
 	ViewportClient->Invalidate();
 
 	EndPostProcessInfoTransaction();
 }
 
-bool SAvaLevelViewport::IsPostProcessCheckerboardEnabled() const
-{
-	const TSharedPtr<FAvaLevelViewportClient> ViewportClient = GetAvaLevelViewportClient();
-
-	if (!ViewportClient.IsValid() || !ViewportClient->GetPostProcessManager().IsValid())
-	{
-		// Make sure none is toggled on if there's an error.
-		return true;
-	}
-
-	return ViewportClient->GetPostProcessManager()->GetType() == EAvaViewportPostProcessType::Checkerboard;
-}
-
-bool SAvaLevelViewport::CanTogglePostProcessCheckerboard() const
-{
-	const TSharedPtr<FAvaLevelViewportClient> ViewportClient = GetAvaLevelViewportClient();
-
-	if (!ViewportClient.IsValid() || !ViewportClient->GetPostProcessManager().IsValid())
-	{
-		return false;
-	}
-
-	if (IsPostProcessCheckerboardEnabled())
-	{
-		return true;
-	}
-
-	if (const TSharedPtr<IAvaViewportPostProcessVisualizer> CheckerboardVisualizer = ViewportClient->GetPostProcessManager()->GetVisualizer(EAvaViewportPostProcessType::Checkerboard))
-	{
-		return CheckerboardVisualizer->CanActivate(/* bInSilent */ true);
-	}
-
-	return false;
-}
-
-void SAvaLevelViewport::ExecuteTogglePostProcessCheckerboard()
-{
-	const TSharedPtr<FAvaLevelViewportClient> ViewportClient = GetAvaLevelViewportClient();
-
-	if (!ViewportClient.IsValid() || !ViewportClient->GetPostProcessManager().IsValid())
-	{
-		return;
-	}
-
-	BeginPostProcessInfoTransaction();
-	ViewportClient->GetPostProcessManager()->SetType(EAvaViewportPostProcessType::Checkerboard);
-	ViewportClient->Invalidate();
-	EndPostProcessInfoTransaction();
-}
-
 bool SAvaLevelViewport::CanToggleOverlay() const
 {
 	return true;
diff --git a/Engine/Plugins/Experimental/Avalanche/Source/AvalancheLevelViewport/Private/Widgets/SAvaLevelViewportStatusBarButtons.cpp b/Engine/Plugins/Experimental/Avalanche/Source/AvalancheLevelViewport/Private/Widgets/SAvaLevelViewportStatusBarButtons.cpp
--- a/Engine/Plugins/Experimental/Avalanche/Source/AvalancheLevelViewport/Private/Widgets/SAvaLevelViewportStatusBarButtons.cpp
+++ b/Engine/Plugins/Experimental/Avalanche/Source/AvalancheLevelViewport/Private/Widgets/SAvaLevelViewportStatusBarButtons.cpp
@@ -87,6 +87,14 @@ namespace UE::Ava::LevelViewportStatusBarButtons::Private
 					PostProcessManager->SetType(InPostProcessType);
 				}
 			}
+			else
+			{
+				UE_LOG(AvaLevelViewportLog, Warning, TEXT("TogglePostProcess: Unable to find post process manager."));
+			}
+		}
+		else
+		{
+			UE_LOG(AvaLevelViewportLog, Warning, TEXT("TogglePostProcess: Invalid viewport frame/client."));
 		}
 	}
 }
diff --git a/Engine/Plugins/Experimental/Avalanche/Source/AvalancheViewport/Private/AvaViewportPostProcessManager.cpp b/Engine/Plugins/Experimental/Avalanche/Source/AvalancheViewport/Private/AvaViewportPostProcessManager.cpp
--- a/Engine/Plugins/Experimental/Avalanche/Source/AvalancheViewport/Private/AvaViewportPostProcessManager.cpp
+++ b/Engine/Plugins/Experimental/Avalanche/Source/AvalancheViewport/Private/AvaViewportPostProcessManager.cpp
@@ -1,14 +1,10 @@
 // Copyright Epic Games, Inc. All Rights Reserved.
 
 #include "AvaViewportPostProcessManager.h"
-
+#include "AvalancheViewportModule.h"
 #include "AvaTypeSharedPointer.h"
 #include "AvaViewportDataSubsystem.h"
-#include "AvaVisibleArea.h"
 #include "Interaction/AvaCameraZoomController.h"
-#include "Viewport/Interaction/AvaViewportPostProcessInfo.h"
-#include "Viewport/Interaction/IAvaViewportDataProvider.h"
-#include "Viewport/Interaction/IAvaViewportDataProxy.h"
 #include "ViewportClient/IAvaViewportClient.h"
 #include "Visualizers/AvaViewportBackgroundVisualizer.h"
 #include "Visualizers/AvaViewportChannelVisualizer.h"
@@ -32,6 +28,7 @@ FAvaViewportPostProcessInfo* FAvaViewportPostProcessManager::GetPostProcessInfo(
 
 	if (!AvaViewportClient.IsValid())
 	{
+		UE_LOG(AvaViewportLog, Warning, TEXT("FAvaViewportPostProcessManager::GetPostProcessInfo: Invalid viewport client."));
 		return nullptr;
 	}
 
@@ -39,6 +36,7 @@ FAvaViewportPostProcessInfo* FAvaViewportPostProcessManager::GetPostProcessInfo(
 
 	if (!DataSubsystem)
 	{
+		UE_LOG(AvaViewportLog, Warning, TEXT("FAvaViewportPostProcessManager::GetPostProcessInfo: Failed to find data subsystem."));
 		return nullptr;
 	}
 
@@ -47,6 +45,7 @@ FAvaViewportPostProcessInfo* FAvaViewportPostProcessManager::GetPostProcessInfo(
 		return &Data->PostProcessInfo;
 	}
 
+	UE_LOG(AvaViewportLog, Warning, TEXT("FAvaViewportPostProcessManager::GetPostProcessInfo: Missing viewport data."));
 	return nullptr;
 }
 
@@ -89,6 +88,7 @@ void FAvaViewportPostProcessManager::LoadPostProcessInfo()
 
 	if (!Visualizer.IsValid())
 	{
+		UE_LOG(AvaViewportLog, Warning, TEXT("FAvaViewportPostProcessManager::LoadPostProcessInfo: Invalid visualizer."));
 		return;
 	}
 
@@ -111,6 +111,7 @@ void FAvaViewportPostProcessManager::SetType(EAvaViewportPostProcessType InType)
 
 	if (!PostProcessInfo)
 	{
+		UE_LOG(AvaViewportLog, Warning, TEXT("AvaViewportPostProcessManager::SetType: Invalid post process info."));
 		return;
 	}
 
@@ -123,6 +124,7 @@ void FAvaViewportPostProcessManager::SetType(EAvaViewportPostProcessType InType)
 
 	if (NewVisualizer.IsValid() && !NewVisualizer->CanActivate(/* bInSilent */ false))
 	{
+		UE_LOG(AvaViewportLog, Warning, TEXT("AvaViewportPostProcessManager::SetType: Cannot activate new visualizer."));
 		return;
 	}
 
@@ -146,6 +148,7 @@ float FAvaViewportPostProcessManager::GetOpacity()
 		return PostProcessInfo->Opacity;
 	}
 
+	UE_LOG(AvaViewportLog, Warning, TEXT("FAvaViewportPostProcessManager::GetOpacity: Missing post process info."));
 	return 1.f;
 }
 
@@ -156,6 +159,10 @@ void FAvaViewportPostProcessManager::SetOpacity(float InOpacity)
 		PostProcessInfo->Opacity = InOpacity;
 		LoadPostProcessInfo();
 	}
+	else
+	{
+		UE_LOG(AvaViewportLog, Warning, TEXT("FAvaViewportPostProcessManager::SetOpacity: Missing post process info."));
+	}
 }
 
 TSharedPtr<IAvaViewportPostProcessVisualizer> FAvaViewportPostProcessManager::GetVisualizer(EAvaViewportPostProcessType InType) const
@@ -165,6 +172,11 @@ TSharedPtr<IAvaViewportPostProcessVisualizer> FAvaViewportPostProcessManager::Ge
 		return *VisualizerPtr;
 	}
 
+	if (InType != EAvaViewportPostProcessType::None)
+	{
+		UE_LOG(AvaViewportLog, Warning, TEXT("FAvaViewportPostProcessManager::GetVisualizer: Missing visualizer."));
+	}
+
 	return nullptr;
 }
 
diff --git a/Engine/Plugins/Experimental/Avalanche/Source/AvalancheViewport/Private/Visualizers/AvaViewportBackgroundVisualizer.cpp b/Engine/Plugins/Experimental/Avalanche/Source/AvalancheViewport/Private/Visualizers/AvaViewportBackgroundVisualizer.cpp
--- a/Engine/Plugins/Experimental/Avalanche/Source/AvalancheViewport/Private/Visualizers/AvaViewportBackgroundVisualizer.cpp
+++ b/Engine/Plugins/Experimental/Avalanche/Source/AvalancheViewport/Private/Visualizers/AvaViewportBackgroundVisualizer.cpp
@@ -1,6 +1,7 @@
 // Copyright Epic Games, Inc. All Rights Reserved.
 
 #include "Visualizers/AvaViewportBackgroundVisualizer.h"
+#include "AvalancheViewportModule.h"
 #include "AvaViewportPostProcessManager.h"
 #include "AvaViewportSettings.h"
 #include "AvaViewportUtils.h"
@@ -9,9 +10,6 @@
 #include "Materials/Material.h"
 #include "Materials/MaterialInstanceDynamic.h"
 #include "UObject/Package.h"
-#include "Viewport/Interaction/AvaViewportPostProcessInfo.h"
-#include "Viewport/Interaction/IAvaViewportDataProvider.h"
-#include "Viewport/Interaction/IAvaViewportDataProxy.h"
 #include "ViewportClient/IAvaViewportClient.h"
 
 #define LOCTEXT_NAMESPACE "AvaViewportBackgroundVisualizer"
@@ -36,6 +34,7 @@ FAvaViewportBackgroundVisualizer::FAvaViewportBackgroundVisualizer(TSharedRef<IA
 
 	if (!ViewportSettings)
 	{
+		UE_LOG(AvaViewportLog, Warning, TEXT("FAvaViewportBackgroundVisualizer::FAvaViewportBackgroundVisualizer: Unable to find viewport settings."));
 		return;
 	}
 
@@ -43,6 +42,7 @@ FAvaViewportBackgroundVisualizer::FAvaViewportBackgroundVisualizer(TSharedRef<IA
 
 	if (!BackgroundMaterial)
 	{
+		UE_LOG(AvaViewportLog, Warning, TEXT("FAvaViewportBackgroundVisualizer::FAvaViewportBackgroundVisualizer: Unable to find background material."));
 		return;
 	}
 
@@ -153,6 +153,7 @@ void FAvaViewportBackgroundVisualizer::UpdatePostProcessMaterial()
 {
 	if (!PostProcessMaterial)
 	{
+		UE_LOG(AvaViewportLog, Warning, TEXT("FAvaViewportBackgroundVisualizer::UpdatePostProcessMaterial: Invalid post process material."));
 		return;
 	}
 
diff --git a/Engine/Plugins/Experimental/Avalanche/Source/AvalancheViewport/Private/Visualizers/AvaViewportCheckerboardVisualizer.cpp b/Engine/Plugins/Experimental/Avalanche/Source/AvalancheViewport/Private/Visualizers/AvaViewportCheckerboardVisualizer.cpp
--- a/Engine/Plugins/Experimental/Avalanche/Source/AvalancheViewport/Private/Visualizers/AvaViewportCheckerboardVisualizer.cpp
+++ b/Engine/Plugins/Experimental/Avalanche/Source/AvalancheViewport/Private/Visualizers/AvaViewportCheckerboardVisualizer.cpp
@@ -1,7 +1,7 @@
 // Copyright Epic Games, Inc. All Rights Reserved.
 
 #include "Visualizers/AvaViewportCheckerboardVisualizer.h"
-
+#include "AvalancheViewportModule.h"
 #include "AvaViewportSettings.h"
 #include "Materials/Material.h"
 #include "Materials/MaterialInstanceDynamic.h"
@@ -56,6 +56,10 @@ void FAvaViewportCheckerboardVisualizer::UpdatePostProcessMaterial()
 		PostProcessMaterial->SetVectorParameterValue(CheckerboardColor1Name, CheckerboardColor1);
 		PostProcessMaterial->SetVectorParameterValue(CheckerboardSizeName, CheckerboardSize);
 	}
+	else
+	{
+		UE_LOG(AvaViewportLog, Warning, TEXT("FAvaViewportCheckerboardVisualizer::UpdatePostProcessMaterial: Missing post process material."));
+	}
 }
 
 FString FAvaViewportCheckerboardVisualizer::GetReferencerName() const
@@ -100,6 +104,8 @@ bool FAvaViewportCheckerboardVisualizer::InitPostProcessMaterial()
 			return true;
 		}
 	}
+
+	UE_LOG(AvaViewportLog, Warning, TEXT("FAvaViewportCheckerboardVisualizer::InitPostProcessMaterial: Failed to init post process material."));
 	return false;
 }
 
@@ -112,6 +118,8 @@ bool FAvaViewportCheckerboardVisualizer::UpdateFromViewportSettings()
 		CheckerboardSize = FVector(ViewportSettings->ViewportCheckerboardSize);
 		return true;
 	}
+
+	UE_LOG(AvaViewportLog, Warning, TEXT("FAvaViewportCheckerboardVisualizer::UpdateFromViewportSettings: Failed to update from viewport settings."));
 	return false;
 }
 
diff --git a/Engine/Plugins/Experimental/Avalanche/Source/AvalancheViewport/Private/Visualizers/AvaViewportPostProcessVisualizer.cpp b/Engine/Plugins/Experimental/Avalanche/Source/AvalancheViewport/Private/Visualizers/AvaViewportPostProcessVisualizer.cpp
--- a/Engine/Plugins/Experimental/Avalanche/Source/AvalancheViewport/Private/Visualizers/AvaViewportPostProcessVisualizer.cpp
+++ b/Engine/Plugins/Experimental/Avalanche/Source/AvalancheViewport/Private/Visualizers/AvaViewportPostProcessVisualizer.cpp
@@ -1,10 +1,10 @@
 // Copyright Epic Games, Inc. All Rights Reserved.
 
 #include "Visualizers/AvaViewportPostProcessVisualizer.h"
+#include "AvalancheViewportModule.h"
 #include "AvaViewportDataSubsystem.h"
 #include "AvaViewportPostProcessManager.h"
 #include "Editor.h"
-#include "Editor/UnrealEdEngine.h"
 #include "Engine/RendererSettings.h"
 #include "FinalPostProcessSettings.h"
 #include "ISettingsEditorModule.h"
@@ -12,9 +12,6 @@
 #include "Misc/MessageDialog.h"
 #include "Modules/ModuleManager.h"
 #include "SceneView.h"
-#include "Viewport/Interaction/AvaViewportPostProcessInfo.h"
-#include "Viewport/Interaction/IAvaViewportDataProvider.h"
-#include "Viewport/Interaction/IAvaViewportDataProxy.h"
 #include "ViewportClient/IAvaViewportClient.h"
 
 #define LOCTEXT_NAMESPACE "AvaViewportPostProcessVisualizer"
@@ -167,6 +164,7 @@ FAvaViewportPostProcessInfo* FAvaViewportPostProcessVisualizer::GetPostProcessIn
 
 	if (!AvaViewportClient.IsValid())
 	{
+		UE_LOG(AvaViewportLog, Warning, TEXT("FAvaViewportPostProcessVisualizer::GetPostProcessInfo: Invalid viewport."));
 		return nullptr;
 	}
 
@@ -174,6 +172,7 @@ FAvaViewportPostProcessInfo* FAvaViewportPostProcessVisualizer::GetPostProcessIn
 
 	if (!DataSubsystem)
 	{
+		UE_LOG(AvaViewportLog, Warning, TEXT("FAvaViewportPostProcessVisualizer::GetPostProcessInfo: Missing data subsystem."));
 		return nullptr;
 	}
 
@@ -182,6 +181,7 @@ FAvaViewportPostProcessInfo* FAvaViewportPostProcessVisualizer::GetPostProcessIn
 		return &Data->PostProcessInfo;
 	}
 
+	UE_LOG(AvaViewportLog, Warning, TEXT("FAvaViewportPostProcessVisualizer::GetPostProcessInfo: Missing viewport data."));
 	return nullptr;
 }
 
@@ -200,6 +200,7 @@ void FAvaViewportPostProcessVisualizer::UpdatePostProcessInfo()
 
 	if (!AvaViewportClient.IsValid())
 	{
+		UE_LOG(AvaViewportLog, Warning, TEXT("FAvaViewportPostProcessVisualizer::UpdatePostProcessInfoL Invalid viewport."));
 		return;
 	}
 
@@ -207,6 +208,7 @@ void FAvaViewportPostProcessVisualizer::UpdatePostProcessInfo()
 
 	if (!DataSubsystem)
 	{
+		UE_LOG(AvaViewportLog, Warning, TEXT("FAvaViewportPostProcessVisualizer::UpdatePostProcessInfo: Missing data subsystem."));
 		return;
 	}
 
@@ -215,12 +217,15 @@ void FAvaViewportPostProcessVisualizer::UpdatePostProcessInfo()
 		DataSubsystem->ModifyDataSource();
 		return UpdatePostProcessInfo(Data->PostProcessInfo);
 	}
+
+	UE_LOG(AvaViewportLog, Warning, TEXT("FAvaViewportPostProcessVisualizer::UpdatePostProcessInfo: Missing viewport data."));
 }
 
 void FAvaViewportPostProcessVisualizer::UpdatePostProcessMaterial()
 {
 	if (!PostProcessMaterial)
 	{
+		UE_LOG(AvaViewportLog, Warning, TEXT("FAvaViewportPostProcessVisualizer::UpdatePostProcessMaterial: Missing post process material."));
 		return;
 	}
 
