[Motion Design] Fix Combo Page special logic: when pages are bypassed, they should still be part of the transition as "playing" to run transition state tree.

#rnx

diff --git a/Engine/Plugins/Experimental/Avalanche/Source/AvalancheMedia/Private/Rundown/AvaRundown.cpp b/Engine/Plugins/Experimental/Avalanche/Source/AvalancheMedia/Private/Rundown/AvaRundown.cpp
--- a/Engine/Plugins/Experimental/Avalanche/Source/AvalancheMedia/Private/Rundown/AvaRundown.cpp
+++ b/Engine/Plugins/Experimental/Avalanche/Source/AvalancheMedia/Private/Rundown/AvaRundown.cpp
@@ -1832,8 +1832,8 @@ bool UAvaRundown::PlayPageWithTransition(FAvaRundownPageTransitionBuilder& InBui
 				// Setup user instance data to be able to track this page.
 				UAvaRundownPagePlayer::SetInstanceUserDataFromPage(*InstancePlayer->PlaybackInstance, InPage);
 				
-				// This instance has to be excluded from the next playable transition.
-				NewPagePlayer->InstancesExcludedFromTransition.Add(InstancePlayer->GetPlaybackInstanceId());
+				// Mark this instance as "bypassing" the next playable transition.
+				NewPagePlayer->InstancesBypassingTransition.Add(InstancePlayer->GetPlaybackInstanceId());
 				bFoundExistingInstancePlayer = true;
 			}
 		}
diff --git a/Engine/Plugins/Experimental/Avalanche/Source/AvalancheMedia/Private/Rundown/Transition/AvaRundownPageTransition.cpp b/Engine/Plugins/Experimental/Avalanche/Source/AvalancheMedia/Private/Rundown/Transition/AvaRundownPageTransition.cpp
--- a/Engine/Plugins/Experimental/Avalanche/Source/AvalancheMedia/Private/Rundown/Transition/AvaRundownPageTransition.cpp
+++ b/Engine/Plugins/Experimental/Avalanche/Source/AvalancheMedia/Private/Rundown/Transition/AvaRundownPageTransition.cpp
@@ -257,7 +257,7 @@ void UAvaRundownPageTransition::Stop()
 	{
 		if (UAvaRundownPagePlayer* PagePlayer = PagePlayerWeak.Get())
 		{
-			PagePlayer->InstancesExcludedFromTransition.Reset();
+			PagePlayer->InstancesBypassingTransition.Reset();
 			UnregisterEnterPagePlayerEvents(PagePlayer);
 		}
 	}
@@ -440,9 +440,9 @@ void UAvaRundownPageTransition::AddPlayersToBuilder(
 		{
 			AddPlayablesToBuilder(InOutBuilder, Player, InCategory, InEntryRole);
 			
-			// Excluded instances are only for current transition.
+			// Bypassed instances are only for current transition.
 			// Todo: Ideally, it would be stored in the transition itself to avoid the clean up.
-			Player->InstancesExcludedFromTransition.Reset();
+			Player->InstancesBypassingTransition.Reset();
 		}
 	}
 }
@@ -453,15 +453,24 @@ void UAvaRundownPageTransition::AddPlayablesToBuilder(FAvaPlayableTransitionBuil
 
 	for (const TObjectPtr<UAvaRundownPlaybackInstancePlayer>& InstancePlayer : InPlayer->InstancePlayers)
 	{
-		if (InPlayer->InstancesExcludedFromTransition.Contains(InstancePlayer->GetPlaybackInstanceId()))
+		EAvaPlayableTransitionEntryRole EntryRole = InEntryRole;
+
+		// -- Special Transition Logic --
+		if (InPlayer->InstancesBypassingTransition.Contains(InstancePlayer->GetPlaybackInstanceId()))
 		{
-			continue;
+			if (EntryRole != EAvaPlayableTransitionEntryRole::Enter)
+			{
+				continue;	// Completely skipping if not an "enter" instance.
+			}
+
+			// Enter instances excluded should still be part as playing pages.
+			EntryRole = EAvaPlayableTransitionEntryRole::Playing;
 		}
 	
 		if (UAvaPlayable* Playable = GetPlayable(InstancePlayer))
 		{
-			const bool bPlayableAdded = InOutBuilder.AddPlayable(Playable, InEntryRole);
-			if (InEntryRole == EAvaPlayableTransitionEntryRole::Enter && bPlayableAdded)
+			const bool bPlayableAdded = InOutBuilder.AddPlayable(Playable, EntryRole);
+			if (EntryRole == EAvaPlayableTransitionEntryRole::Enter && bPlayableAdded)
 			{
 				InOutBuilder.AddEnterPlayableValues(GetRemoteControlValues(InPlayer));
 			}
diff --git a/Engine/Plugins/Experimental/Avalanche/Source/AvalancheMedia/Public/Rundown/AvaRundownPagePlayer.h b/Engine/Plugins/Experimental/Avalanche/Source/AvalancheMedia/Public/Rundown/AvaRundownPagePlayer.h
--- a/Engine/Plugins/Experimental/Avalanche/Source/AvalancheMedia/Public/Rundown/AvaRundownPagePlayer.h
+++ b/Engine/Plugins/Experimental/Avalanche/Source/AvalancheMedia/Public/Rundown/AvaRundownPagePlayer.h
@@ -179,9 +179,9 @@ public:
 	UPROPERTY()
 	TArray<TObjectPtr<UAvaRundownPlaybackInstancePlayer>> InstancePlayers;
 
-	/** Instances that are excluded from the next transition. */
+	/** Instances that should bypass the next transition. */
 	UPROPERTY(Transient)
-	TSet<FGuid> InstancesExcludedFromTransition;
+	TSet<FGuid> InstancesBypassingTransition;
 
 protected:
 	TWeakObjectPtr<UAvaRundown> RundownWeak;
