MotionDesign : TranslucentPriorityModifier
- Showing a warning when OIT is enabled in project settings

diff --git a/Engine/Plugins/Experimental/Avalanche/Source/AvalancheModifiers/Private/Modifiers/AvaTranslucentPriorityModifier.cpp b/Engine/Plugins/Experimental/Avalanche/Source/AvalancheModifiers/Private/Modifiers/AvaTranslucentPriorityModifier.cpp
--- a/Engine/Plugins/Experimental/Avalanche/Source/AvalancheModifiers/Private/Modifiers/AvaTranslucentPriorityModifier.cpp
+++ b/Engine/Plugins/Experimental/Avalanche/Source/AvalancheModifiers/Private/Modifiers/AvaTranslucentPriorityModifier.cpp
@@ -7,6 +7,7 @@
 #include "Camera/CameraComponent.h"
 #include "Components/PrimitiveComponent.h"
 #include "EngineUtils.h"
+#include "Engine/RendererSettings.h"
 #include "Framework/AvaGameInstance.h"
 #include "Modifiers/ActorModifierCoreBase.h"
 #include "Shared/AvaTranslucentPriorityModifierShared.h"
@@ -142,6 +143,14 @@ void UAvaTranslucentPriorityModifier::RestorePreState()
 
 void UAvaTranslucentPriorityModifier::Apply()
 {
+	const URendererSettings* RendererSettings = GetDefault<URendererSettings>();
+
+	if (RendererSettings && RendererSettings->bOrderedIndependentTransparencyEnable)
+	{
+		Fail(LOCTEXT("InvalidRendererSettingsOIT", "Incompatible with ordered independent transparency project setting"));
+		return;
+	}
+
 	const UAvaTranslucentPriorityModifierShared* SharedObject = GetShared<UAvaTranslucentPriorityModifierShared>(false);
 
 	if (!SharedObject)
