[Motion Design]
 - Move Broadcast "New Channel" button to main toolbar to provide maximum space for channel displays
 - Add quality setting preset buttons back to menu

diff --git a/Engine/Plugins/Experimental/Avalanche/Source/Avalanche/Private/Viewport/AvaViewportQualitySettings.cpp b/Engine/Plugins/Experimental/Avalanche/Source/Avalanche/Private/Viewport/AvaViewportQualitySettings.cpp
--- a/Engine/Plugins/Experimental/Avalanche/Source/Avalanche/Private/Viewport/AvaViewportQualitySettings.cpp
+++ b/Engine/Plugins/Experimental/Avalanche/Source/Avalanche/Private/Viewport/AvaViewportQualitySettings.cpp
@@ -31,7 +31,7 @@ FAvaViewportQualitySettings::FAvaViewportQualitySettings(const FEngineShowFlags&
 			continue;
 		}
 
-		Feature.Enabled = InShowFlags.GetSingleFlag(FeatureIndex);
+		Feature.bEnabled = InShowFlags.GetSingleFlag(FeatureIndex);
 	}
 
 	SortFeaturesByDisplayText();
@@ -46,6 +46,26 @@ FAvaViewportQualitySettings::FAvaViewportQualitySettings(const TArray<FAvaViewpo
 	SortFeaturesByDisplayText();
 }
 
+bool FAvaViewportQualitySettings::operator==(const FAvaViewportQualitySettings& InOther) const
+{
+	// All features must exist and match the same value in the other struct to be equal
+	for (const FAvaViewportQualitySettingsFeature& Feature : Features)
+	{
+		const FAvaViewportQualitySettingsFeature* const FoundFeaturePtr = FindFeatureByName(InOther.Features, Feature.Name);
+		if (!FoundFeaturePtr || FoundFeaturePtr->bEnabled != Feature.bEnabled)
+		{
+			return false;
+		}
+	}
+
+	return true;
+}
+
+bool FAvaViewportQualitySettings::operator!=(const FAvaViewportQualitySettings& InOther) const
+{
+	return !(*this == InOther);
+}
+
 TArray<FAvaViewportQualitySettingsFeature> FAvaViewportQualitySettings::DefaultFeatures()
 {
 	// Add all the engine features we will support for editing by the designer and their Motion Design defaults.
@@ -100,7 +120,7 @@ TArray<FAvaViewportQualitySettingsFeature> FAvaViewportQualitySettings::AllFeatu
 
 	for (FAvaViewportQualitySettingsFeature& Feature : AllDefaultFeatures)
 	{
-		Feature.Enabled = bUseAllFeatures;
+		Feature.bEnabled = bUseAllFeatures;
 	}
 
 	SortFeaturesByDisplayText(AllDefaultFeatures);
@@ -159,7 +179,7 @@ FAvaViewportQualitySettings FAvaViewportQualitySettings::All(const bool bUseAllF
 	FAvaViewportQualitySettings DefaultSettings = Default();
 	for (FAvaViewportQualitySettingsFeature& Feature : DefaultSettings.Features)
 	{
-		Feature.Enabled = bUseAllFeatures;
+		Feature.bEnabled = bUseAllFeatures;
 	}
 	return DefaultSettings;
 }
@@ -173,7 +193,7 @@ void FAvaViewportQualitySettings::Apply(FEngineShowFlags& InFlags)
 		const int32 FeatureIndex = InFlags.FindIndexByName(*Feature.Name);
 		if (FeatureIndex != INDEX_NONE)
 		{
-			InFlags.SetSingleFlag(FeatureIndex, Feature.Enabled);
+			InFlags.SetSingleFlag(FeatureIndex, Feature.bEnabled);
 		}
 	}
 }
@@ -233,12 +253,11 @@ void FAvaViewportQualitySettings::EnableFeaturesByName(const bool bInEnabled, co
 	for (const FString& FeatureName : InFeatureNames)
 	{
 		const FAvaViewportQualitySettingsFeature* DefaultFeature = FindFeatureByName(AllDefaultFeatures, FeatureName);
-		const bool bExistsInDefaults = DefaultFeature != nullptr;
-		if (bExistsInDefaults)
+		if (DefaultFeature)
 		{
 			if (FAvaViewportQualitySettingsFeature* Feature = FindFeatureByName(Features, FeatureName))
 			{
-				Feature->Enabled = bInEnabled;
+				Feature->bEnabled = bInEnabled;
 			}
 		}
 	}
@@ -246,7 +265,15 @@ void FAvaViewportQualitySettings::EnableFeaturesByName(const bool bInEnabled, co
 
 FAvaViewportQualitySettingsFeature* FAvaViewportQualitySettings::FindFeatureByName(TArray<FAvaViewportQualitySettingsFeature>& InFeatures, const FString& InFeatureName)
 {
-	return InFeatures.FindByPredicate([&InFeatureName](const FAvaViewportQualitySettingsFeature& InFeature)
+	return InFeatures.FindByPredicate([&InFeatureName](FAvaViewportQualitySettingsFeature& InFeature)
+		{
+			return InFeature.Name.Equals(InFeatureName);
+		});
+}
+
+const FAvaViewportQualitySettingsFeature* FAvaViewportQualitySettings::FindFeatureByName(const TArray<FAvaViewportQualitySettingsFeature>& InFeatures, const FString& InFeatureName)
+{
+	return InFeatures.FindByPredicate([&InFeatureName](FAvaViewportQualitySettingsFeature& InFeature)
 		{
 			return InFeature.Name.Equals(InFeatureName);
 		});
diff --git a/Engine/Plugins/Experimental/Avalanche/Source/Avalanche/Public/Viewport/AvaViewportQualitySettings.h b/Engine/Plugins/Experimental/Avalanche/Source/Avalanche/Public/Viewport/AvaViewportQualitySettings.h
--- a/Engine/Plugins/Experimental/Avalanche/Source/Avalanche/Public/Viewport/AvaViewportQualitySettings.h
+++ b/Engine/Plugins/Experimental/Avalanche/Source/Avalanche/Public/Viewport/AvaViewportQualitySettings.h
@@ -10,27 +10,27 @@ class FText;
 struct FEngineShowFlags;
 
 USTRUCT(BlueprintType)
-struct FAvaViewportQualitySettingsFeature
+struct AVALANCHE_API FAvaViewportQualitySettingsFeature
 {
 	GENERATED_BODY()
 
-	/** The name of the feature in the engine show flags. */
-	UPROPERTY(EditAnywhere, BlueprintReadWrite, Category = "Quality")
-	FString Name;
-
-	/** True if this engine feature show flag should be enabled. */
-	UPROPERTY(EditAnywhere, BlueprintReadWrite, Category = "Quality")
-	bool Enabled = false;
-
 	FAvaViewportQualitySettingsFeature() {}
-	FAvaViewportQualitySettingsFeature(const FString& InName, const bool InEnabled)
-		: Name(InName), Enabled(InEnabled)
+	FAvaViewportQualitySettingsFeature(const FString& InName, const bool bInEnabled)
+		: Name(InName), bEnabled(bInEnabled)
 	{}
 
 	bool operator==(const FAvaViewportQualitySettingsFeature& InOther) const
 	{
 		return Name.Equals(InOther.Name);
 	}
+
+	/** The name of the feature in the engine show flags. */
+	UPROPERTY(EditAnywhere, BlueprintReadWrite, Category = "Quality")
+	FString Name;
+
+	/** True if this engine feature show flag should be enabled. */
+	UPROPERTY(EditAnywhere, BlueprintReadWrite, Category = "Quality")
+	bool bEnabled = false;
 };
 
 /** 
@@ -46,15 +46,6 @@ struct AVALANCHE_API FAvaViewportQualitySettings
 	GENERATED_BODY()
 
 public:
-	/** Advanced viewport client engine features indexed by FEngineShowFlags names. */
-	UPROPERTY(EditAnywhere, BlueprintReadWrite, EditFixedSize, Category = "Quality", meta=(EditFixedOrder))
-	TArray<FAvaViewportQualitySettingsFeature> Features;
-
-	FAvaViewportQualitySettings();
-	FAvaViewportQualitySettings(const bool bInUseAllFeatures);
-	FAvaViewportQualitySettings(const FEngineShowFlags& InShowFlags);
-	FAvaViewportQualitySettings(const TArray<FAvaViewportQualitySettingsFeature>& InFeatures);
-
 	static TArray<FAvaViewportQualitySettingsFeature> DefaultFeatures();
 	static TArray<FAvaViewportQualitySettingsFeature> AllFeatures(const bool bUseAllFeatures);
 
@@ -64,16 +55,31 @@ public:
 
 	static void FeatureNameAndTooltipText(const FString& InFeatureName, FText& OutNameText, FText& OutTooltipText);
 
+	static FAvaViewportQualitySettingsFeature* FindFeatureByName(TArray<FAvaViewportQualitySettingsFeature>& InFeatures, const FString& InFeatureName);
+	static const FAvaViewportQualitySettingsFeature* FindFeatureByName(const TArray<FAvaViewportQualitySettingsFeature>& InFeatures, const FString& InFeatureName);
+
+	static void VerifyIntegrity(TArray<FAvaViewportQualitySettingsFeature>& InFeatures);
+
+	static void SortFeaturesByDisplayText(TArray<FAvaViewportQualitySettingsFeature>& InFeatures);
+
+	FAvaViewportQualitySettings();
+	FAvaViewportQualitySettings(const bool bInUseAllFeatures);
+	FAvaViewportQualitySettings(const FEngineShowFlags& InShowFlags);
+	FAvaViewportQualitySettings(const TArray<FAvaViewportQualitySettingsFeature>& InFeatures);
+
+	bool operator==(const FAvaViewportQualitySettings& InOther) const;
+	bool operator!=(const FAvaViewportQualitySettings& InOther) const;
+
 	/** Applies the settings to the FEngineShowFlags structure provided. */
 	void Apply(FEngineShowFlags& InFlags);
 
 	void EnableFeaturesByName(const bool bInEnabled, const TArray<FString>& InFeatureNames);
 
-	static FAvaViewportQualitySettingsFeature* FindFeatureByName(TArray<FAvaViewportQualitySettingsFeature>& InFeatures, const FString& InFeatureName);
-
-	static void VerifyIntegrity(TArray<FAvaViewportQualitySettingsFeature>& InFeatures);
 	void VerifyIntegrity();
 
-	static void SortFeaturesByDisplayText(TArray<FAvaViewportQualitySettingsFeature>& InFeatures);
 	void SortFeaturesByDisplayText();
+
+	/** Advanced viewport client engine features indexed by FEngineShowFlags names. */
+	UPROPERTY(EditAnywhere, BlueprintReadWrite, EditFixedSize, Category = "Quality", meta = (EditFixedOrder))
+	TArray<FAvaViewportQualitySettingsFeature> Features;
 };
diff --git a/Engine/Plugins/Experimental/Avalanche/Source/AvalancheEditor/Private/AvaEditorSettings.cpp b/Engine/Plugins/Experimental/Avalanche/Source/AvalancheEditor/Private/AvaEditorSettings.cpp
--- a/Engine/Plugins/Experimental/Avalanche/Source/AvalancheEditor/Private/AvaEditorSettings.cpp
+++ b/Engine/Plugins/Experimental/Avalanche/Source/AvalancheEditor/Private/AvaEditorSettings.cpp
@@ -59,11 +59,4 @@ void UAvaEditorSettings::PostInitProperties()
 	}
 }
 
-void UAvaEditorSettings::PostEditChangeProperty(FPropertyChangedEvent& PropertyChangedEvent)
-{
-	Super::PostEditChangeProperty(PropertyChangedEvent);
-
-	OnChanged.Broadcast(this, PropertyChangedEvent.GetPropertyName());
-}
-
 #undef LOCTEXT_NAMESPACE
diff --git a/Engine/Plugins/Experimental/Avalanche/Source/AvalancheEditor/Private/AvaEditorSettings.h b/Engine/Plugins/Experimental/Avalanche/Source/AvalancheEditor/Private/AvaEditorSettings.h
--- a/Engine/Plugins/Experimental/Avalanche/Source/AvalancheEditor/Private/AvaEditorSettings.h
+++ b/Engine/Plugins/Experimental/Avalanche/Source/AvalancheEditor/Private/AvaEditorSettings.h
@@ -12,31 +12,6 @@
 #include "Viewport/AvaViewportQualitySettings.h"
 #include "AvaEditorSettings.generated.h"
 
-USTRUCT()
-struct FAvaPaletteSettings
-{
-	GENERATED_BODY()
-
-	UPROPERTY(Config, VisibleAnywhere, Category = "Interface")
-	FName Name;
-
-	// Name is the TMap key
-	UPROPERTY(Config, EditAnywhere, Category = "Interface")
-	bool bEnablePalette = true;
-
-	UPROPERTY(Config)
-	bool bExpanded = false;
-};
-
-USTRUCT()
-struct FAvaPaletteTabSettings
-{
-	GENERATED_BODY()
-
-	UPROPERTY(Config, EditAnywhere, EditFixedSize, Category = "Interface", meta = (TitleProperty = "Name", FullyExpand))
-	TArray<FAvaPaletteSettings> Palettes;
-};
-
 /**
  * Motion Design Editor Settings
  */
@@ -67,14 +42,6 @@ public:
 	UPROPERTY(Config, EditAnywhere, BlueprintReadWrite, Category = "Behavior")
 	float CameraDistance = 500.0f;
 
-	/**  */
-	UPROPERTY(Config, EditAnywhere, BlueprintReadWrite, Category = "Toolbox")
-	TArray<FString> FavoriteTools;
-
-	/**  */
-	UPROPERTY(Config, EditAnywhere, EditFixedSize, Category = "Interface", meta = (EditFixedOrder, ShowOnlyInnerProperties))
-	TArray<FAvaPaletteTabSettings> PaletteTabs;
-
 	/** Default viewport quality settings for all newly created Motion Design blueprints. */
 	UPROPERTY(Config, EditAnywhere, BlueprintReadWrite, Category = "Quality")
 	FAvaViewportQualitySettings DefaultViewportQualitySettings = FAvaViewportQualitySettings(true);
@@ -83,12 +50,8 @@ public:
 	UPROPERTY(Config, EditAnywhere, BlueprintReadWrite, Category = "Quality")
 	TMap<FName, FAvaViewportQualitySettings> ViewportQualityPresets;
 
-	DECLARE_MULTICAST_DELEGATE_TwoParams(FOnSettingsChanged, const UAvaEditorSettings* InSettings, FName InSetting)
-	FOnSettingsChanged OnChanged;
-
 	//~ Begin UObject
 	virtual void PostInitProperties() override;
-	virtual void PostEditChangeProperty(FPropertyChangedEvent& PropertyChangedEvent) override;
 	//~ End UObject
 
 	void OpenEditorSettingsWindow() const;
diff --git a/Engine/Plugins/Experimental/Avalanche/Source/AvalancheEditor/Private/DetailView/Customizations/AvaViewportQualitySettingsPropertyTypeCustomization.cpp b/Engine/Plugins/Experimental/Avalanche/Source/AvalancheEditor/Private/DetailView/Customizations/AvaViewportQualitySettingsPropertyTypeCustomization.cpp
--- a/Engine/Plugins/Experimental/Avalanche/Source/AvalancheEditor/Private/DetailView/Customizations/AvaViewportQualitySettingsPropertyTypeCustomization.cpp
+++ b/Engine/Plugins/Experimental/Avalanche/Source/AvalancheEditor/Private/DetailView/Customizations/AvaViewportQualitySettingsPropertyTypeCustomization.cpp
@@ -1,12 +1,21 @@
 // Copyright Epic Games, Inc. All Rights Reserved.
 
 #include "AvaViewportQualitySettingsPropertyTypeCustomization.h"
+#include "AvaEditorSettings.h"
 #include "DetailLayoutBuilder.h"
 #include "DetailWidgetRow.h"
 #include "IDetailChildrenBuilder.h"
+#include "Internationalization/Text.h"
 #include "PropertyCustomizationHelpers.h"
 #include "PropertyHandle.h"
+#include "Styling/AppStyle.h"
+#include "Templates/SharedPointer.h"
 #include "Viewport/AvaViewportQualitySettings.h"
+#include "Widgets/Input/SButton.h"
+#include "Widgets/Layout/SBorder.h"
+#include "Widgets/Layout/SBox.h"
+#include "Widgets/Layout/SSeparator.h"
+#include "Widgets/Layout/SWrapBox.h"
 #include "Widgets/Text/STextBlock.h"
 
 #define LOCTEXT_NAMESPACE "AvaViewportQualitySettingsPropertyTypeCustomization"
@@ -16,49 +25,229 @@ TSharedRef<IPropertyTypeCustomization> FAvaViewportQualitySettingsPropertyTypeCu
 	return MakeShared<FAvaViewportQualitySettingsPropertyTypeCustomization>();
 }
 
-void FAvaViewportQualitySettingsPropertyTypeCustomization::CustomizeHeader(TSharedRef<IPropertyHandle> InStructPropertyHandle, FDetailWidgetRow& InHeaderRow, IPropertyTypeCustomizationUtils& StructCustomizationUtils)
+void FAvaViewportQualitySettingsPropertyTypeCustomization::CustomizeHeader(TSharedRef<IPropertyHandle> InStructPropertyHandle, FDetailWidgetRow& HeaderRow, IPropertyTypeCustomizationUtils& StructCustomizationUtils)
 {
-	InHeaderRow.NameContent()
-		[
-			InStructPropertyHandle->CreatePropertyNameWidget()
-		];
+	StructPropertyHandle = InStructPropertyHandle;
+
+	if (InStructPropertyHandle->GetBoolMetaData(TEXT("HideHeader")))
+	{
+		HeaderRow.Visibility(EVisibility::Collapsed);
+	}
+	else
+	{
+		HeaderRow.NameContent()
+			[
+				InStructPropertyHandle->CreatePropertyNameWidget()
+			];
+	}
 }
 
-void FAvaViewportQualitySettingsPropertyTypeCustomization::CustomizeChildren(TSharedRef<IPropertyHandle> InStructPropertyHandle, IDetailChildrenBuilder& InDetailBuilder, IPropertyTypeCustomizationUtils& InStructCustomizationUtils)
+void FAvaViewportQualitySettingsPropertyTypeCustomization::CustomizeChildren(TSharedRef<IPropertyHandle> InStructPropertyHandle
+	, IDetailChildrenBuilder& InOutDetailBuilder
+	, IPropertyTypeCustomizationUtils& InOutStructCustomizationUtils)
 {
-	TSharedPtr<IPropertyHandle> FeaturesProperty = InStructPropertyHandle->GetChildHandle(GET_MEMBER_NAME_CHECKED(FAvaViewportQualitySettings, Features));
+	if (InStructPropertyHandle->GetBoolMetaData(TEXT("ShowPresets")))
+	{
+		InOutDetailBuilder.AddCustomRow(LOCTEXT("Presets", "Presets"))
+			[
+				SNew(SBorder)
+				.BorderImage(FAppStyle::GetBrush(TEXT("ToolPanel.GroupBorder")))
+				[
+					SAssignNew(PresetsWrapBox, SWrapBox)
+					.UseAllottedSize(true)
+					.Orientation(EOrientation::Orient_Horizontal)
+				]
+			];
+
+		RefreshPresets();
+	}
 
+	const TSharedPtr<IPropertyHandle> FeaturesProperty = InStructPropertyHandle->GetChildHandle(GET_MEMBER_NAME_CHECKED(FAvaViewportQualitySettings, Features));
 	const TSharedRef<FDetailArrayBuilder> FeaturesArrayBuilder = MakeShared<FDetailArrayBuilder>(FeaturesProperty.ToSharedRef(), /*InGenerateHeader*/ false, /*InDisplayResetToDefault*/ true, /*InDisplayElementNum*/ false);
 
-	static const bool bDisplayDefaultPropertyButtons = false;
+	FeaturesArrayBuilder->OnGenerateArrayElementWidget(FOnGenerateArrayElementWidget::CreateLambda([](TSharedRef<IPropertyHandle> InElementPropertyHandle, const int32 InArrayIndex, IDetailChildrenBuilder& InOutChildrenBuilder)
+		{
+			TSharedPtr<IPropertyHandle> NameProperty = InElementPropertyHandle->GetChildHandle(GET_MEMBER_NAME_CHECKED(FAvaViewportQualitySettingsFeature, Name));
+			TSharedPtr<IPropertyHandle> ValueProperty = InElementPropertyHandle->GetChildHandle(GET_MEMBER_NAME_CHECKED(FAvaViewportQualitySettingsFeature, bEnabled));
 
-	FeaturesArrayBuilder->OnGenerateArrayElementWidget(FOnGenerateArrayElementWidget::CreateLambda([](TSharedRef<IPropertyHandle> InElementPropertyHandle, int32 InArrayIndex, IDetailChildrenBuilder& InChildrenBuilder)
-	{
-		TSharedPtr<IPropertyHandle> NameProperty  = InElementPropertyHandle->GetChildHandle(GET_MEMBER_NAME_CHECKED(FAvaViewportQualitySettingsFeature, Name));
-		TSharedPtr<IPropertyHandle> ValueProperty = InElementPropertyHandle->GetChildHandle(GET_MEMBER_NAME_CHECKED(FAvaViewportQualitySettingsFeature, Enabled));
+			FString FeatureName;
+			NameProperty->GetValue(FeatureName);
 
-		FString FeatureName;
-		NameProperty->GetValue(FeatureName);
+			FText NameText, TooltipText;
+			FAvaViewportQualitySettings::FeatureNameAndTooltipText(FeatureName, NameText, TooltipText);
 
-		FText NameText, TooltipText;
-		FAvaViewportQualitySettings::FeatureNameAndTooltipText(FeatureName, NameText, TooltipText);
+			InOutChildrenBuilder.AddProperty(InElementPropertyHandle)
+				.ToolTip(TooltipText)
+				.CustomWidget()
+				.NameContent()
+				[
+					SNew(STextBlock)
+					.Text(NameText)
+					.Font(IDetailLayoutBuilder::GetDetailFont())
+				]
+				.ValueContent()
+				[
+					ValueProperty->CreatePropertyValueWidget(/*bDisplayDefaultPropertyButtons=*/false)
+				];
+		}));
 
-		InChildrenBuilder.AddProperty(InElementPropertyHandle)
-			.ToolTip(TooltipText)
-			.CustomWidget()
-			.NameContent()
-			[
-				SNew(STextBlock)
-				.Text(NameText)
-				.Font(IDetailLayoutBuilder::GetDetailFont())
-			]
-			.ValueContent()
+	InOutDetailBuilder.AddCustomBuilder(FeaturesArrayBuilder);
+}
+
+void FAvaViewportQualitySettingsPropertyTypeCustomization::RefreshPresets()
+{
+	if (!PresetsWrapBox.IsValid())
+	{
+		return;
+	}
+
+	PresetsWrapBox->ClearChildren();
+
+	auto AddSlotToWrapBox = [this](const FName InName, FOnClicked InOnClicked = FOnClicked(), const TAttribute<bool>& InIsEnabled = true)
+	{
+		if (!InOnClicked.IsBound())
+		{
+			InOnClicked = FOnClicked::CreateSP(this, &FAvaViewportQualitySettingsPropertyTypeCustomization::HandlePresetButtonClick, InName);
+		}
+		PresetsWrapBox->AddSlot()
 			[
-				ValueProperty->CreatePropertyValueWidget(bDisplayDefaultPropertyButtons)
+				SNew(SBox)
+				.Padding(2.f)
+				[
+					SNew(SButton)
+					.ButtonStyle(FAppStyle::Get(), "SimpleButton")
+					.OnClicked(InOnClicked)
+					.IsEnabled(InIsEnabled)
+					[
+						SNew(STextBlock)
+						.TextStyle(FAppStyle::Get(), TEXT("SmallText"))
+						.Text(FText::FromName(InName))
+					]
+				]
 			];
-	}));
+	};
+
+	AddSlotToWrapBox(TEXT("Defaults")
+		, FOnClicked::CreateSP(this, &FAvaViewportQualitySettingsPropertyTypeCustomization::HandleDefaultsButtonClick)
+		, TAttribute<bool>::CreateSP(this, &FAvaViewportQualitySettingsPropertyTypeCustomization::IsDefaultsButtonEnabled));
+
+	AddSlotToWrapBox(TEXT("All")
+		, FOnClicked::CreateSP(this, &FAvaViewportQualitySettingsPropertyTypeCustomization::HandleEnableAllButtonClick)
+		, TAttribute<bool>::CreateSP(this, &FAvaViewportQualitySettingsPropertyTypeCustomization::IsAllButtonEnabled));
+
+	AddSlotToWrapBox(TEXT("None")
+		, FOnClicked::CreateSP(this, &FAvaViewportQualitySettingsPropertyTypeCustomization::HandleDisableAllButtonClick)
+		, TAttribute<bool>::CreateSP(this, &FAvaViewportQualitySettingsPropertyTypeCustomization::IsNoneButtonEnabled));
+
+	PresetsWrapBox->AddSlot()
+		.Padding(5.f, 0.f)
+		[
+			SNew(SSeparator)
+			.Orientation(EOrientation::Orient_Vertical)
+		];
+
+	for (const TPair<FName, FAvaViewportQualitySettings>& Preset : UAvaEditorSettings::Get()->ViewportQualityPresets)
+	{
+		AddSlotToWrapBox(Preset.Key
+			, FOnClicked()
+			, TAttribute<bool>::CreateSP(this, &FAvaViewportQualitySettingsPropertyTypeCustomization::IsPresetButtonEnabled, Preset.Key));
+	}
+}
+
+FAvaViewportQualitySettings& FAvaViewportQualitySettingsPropertyTypeCustomization::GetStructRef() const
+{
+	check(StructPropertyHandle.IsValid());
+
+	void* OutAddress = nullptr;
+	StructPropertyHandle->GetValueData(OutAddress);
+
+	return *reinterpret_cast<FAvaViewportQualitySettings*>(OutAddress);
+}
+
+FReply FAvaViewportQualitySettingsPropertyTypeCustomization::HandleDefaultsButtonClick()
+{
+	GetStructRef() = UAvaEditorSettings::Get()->DefaultViewportQualitySettings;
+
+	return FReply::Handled();
+}
+
+FReply FAvaViewportQualitySettingsPropertyTypeCustomization::HandleEnableAllButtonClick()
+{
+	for (FAvaViewportQualitySettingsFeature& Feature : GetStructRef().Features)
+	{
+		Feature.bEnabled = true;
+	}
+
+	return FReply::Handled();
+}
+
+FReply FAvaViewportQualitySettingsPropertyTypeCustomization::HandleDisableAllButtonClick()
+{
+	for (FAvaViewportQualitySettingsFeature& Feature : GetStructRef().Features)
+	{
+		Feature.bEnabled = false;
+	}
+
+	return FReply::Handled();
+}
+
+FReply FAvaViewportQualitySettingsPropertyTypeCustomization::HandlePresetButtonClick(const FName InNewPresetName)
+{
+	UAvaEditorSettings* const EditorSettings = UAvaEditorSettings::Get();
+	if (!EditorSettings->ViewportQualityPresets.Contains(InNewPresetName))
+	{
+		return FReply::Unhandled();
+	}
+
+	GetStructRef() = EditorSettings->ViewportQualityPresets[InNewPresetName];
+
+	return FReply::Handled();
+}
+
+bool FAvaViewportQualitySettingsPropertyTypeCustomization::IsDefaultsButtonEnabled() const
+{
+	return UAvaEditorSettings::Get()->DefaultViewportQualitySettings != GetStructRef();
+}
+
+bool FAvaViewportQualitySettingsPropertyTypeCustomization::IsAllButtonEnabled() const
+{
+	for (FAvaViewportQualitySettingsFeature& Feature : GetStructRef().Features)
+	{
+		if (!Feature.bEnabled)
+		{
+			return true;
+		}
+	}
+
+	return false;
+}
+
+bool FAvaViewportQualitySettingsPropertyTypeCustomization::IsNoneButtonEnabled() const
+{
+	for (FAvaViewportQualitySettingsFeature& Feature : GetStructRef().Features)
+	{
+		if (Feature.bEnabled)
+		{
+			return true;
+		}
+	}
+
+	return false;
+}
+
+bool FAvaViewportQualitySettingsPropertyTypeCustomization::IsPresetButtonEnabled(const FName InPresetName) const
+{
+	const FAvaViewportQualitySettings& StructRef = GetStructRef();
+
+	for (const TPair<FName, FAvaViewportQualitySettings>& Preset : UAvaEditorSettings::Get()->ViewportQualityPresets)
+	{
+		if (Preset.Key == InPresetName)
+		{
+			return Preset.Value != StructRef;
+		}
+	}
 
-	InDetailBuilder.AddCustomBuilder(FeaturesArrayBuilder);
+	return true;
 }
 
 #undef LOCTEXT_NAMESPACE
diff --git a/Engine/Plugins/Experimental/Avalanche/Source/AvalancheEditor/Private/DetailView/Customizations/AvaViewportQualitySettingsPropertyTypeCustomization.h b/Engine/Plugins/Experimental/Avalanche/Source/AvalancheEditor/Private/DetailView/Customizations/AvaViewportQualitySettingsPropertyTypeCustomization.h
--- a/Engine/Plugins/Experimental/Avalanche/Source/AvalancheEditor/Private/DetailView/Customizations/AvaViewportQualitySettingsPropertyTypeCustomization.h
+++ b/Engine/Plugins/Experimental/Avalanche/Source/AvalancheEditor/Private/DetailView/Customizations/AvaViewportQualitySettingsPropertyTypeCustomization.h
@@ -3,6 +3,14 @@
 #pragma once
 
 #include "IPropertyTypeCustomization.h"
+#include "Templates/SharedPointer.h"
+#include "Widgets/Layout/SWrapBox.h"
+
+class FName;
+class FReply;
+class FStructOnScope;
+struct FAvaViewportQualitySettings;
+struct FPropertyChangedEvent;
 
 class FAvaViewportQualitySettingsPropertyTypeCustomization : public IPropertyTypeCustomization
 {
@@ -10,7 +18,26 @@ public:
 	static TSharedRef<IPropertyTypeCustomization> MakeInstance();
 
 	//~ Begin IPropertyTypeCustomization
-	virtual void CustomizeHeader(TSharedRef<IPropertyHandle> InStructPropertyHandle, FDetailWidgetRow& InHeaderRow, IPropertyTypeCustomizationUtils& InStructCustomizationUtils) override;
-	virtual void CustomizeChildren(TSharedRef<IPropertyHandle> InStructPropertyHandle, IDetailChildrenBuilder& InDetailBuilder, IPropertyTypeCustomizationUtils& InStructCustomizationUtils) override;
+	virtual void CustomizeHeader(TSharedRef<IPropertyHandle> InStructPropertyHandle, FDetailWidgetRow& HeaderRow, IPropertyTypeCustomizationUtils& StructCustomizationUtils) override;
+	virtual void CustomizeChildren(TSharedRef<IPropertyHandle> InStructPropertyHandle, IDetailChildrenBuilder& DetailBuilder, IPropertyTypeCustomizationUtils& StructCustomizationUtils) override;
 	//~ End IPropertyTypeCustomization
+
+protected:
+	void RefreshPresets();
+
+	FAvaViewportQualitySettings& GetStructRef() const;
+
+	FReply HandleDefaultsButtonClick();
+	FReply HandleEnableAllButtonClick();
+	FReply HandleDisableAllButtonClick();
+	FReply HandlePresetButtonClick(const FName InPresetName);
+
+	bool IsDefaultsButtonEnabled() const;
+	bool IsAllButtonEnabled() const;
+	bool IsNoneButtonEnabled() const;
+	bool IsPresetButtonEnabled(const FName InPresetName) const;
+
+	TSharedPtr<IPropertyHandle> StructPropertyHandle;
+
+	TSharedPtr<SWrapBox> PresetsWrapBox;
 };
diff --git a/Engine/Plugins/Experimental/Avalanche/Source/AvalancheMedia/Private/Playback/AvaPlaybackMessages.h b/Engine/Plugins/Experimental/Avalanche/Source/AvalancheMedia/Private/Playback/AvaPlaybackMessages.h
--- a/Engine/Plugins/Experimental/Avalanche/Source/AvalancheMedia/Private/Playback/AvaPlaybackMessages.h
+++ b/Engine/Plugins/Experimental/Avalanche/Source/AvalancheMedia/Private/Playback/AvaPlaybackMessages.h
@@ -577,7 +577,7 @@ struct FAvaBroadcastChannelSettingsUpdate : public FAvaPlaybackClientMessageBase
 	UPROPERTY()
 	FString Channel;
 
-	UPROPERTY()
+	UPROPERTY(meta = (HideHeader, ShowPresets))
 	FAvaViewportQualitySettings QualitySettings;
 };
 
diff --git a/Engine/Plugins/Experimental/Avalanche/Source/AvalancheMedia/Public/Broadcast/Channel/AvaBroadcastOutputChannel.h b/Engine/Plugins/Experimental/Avalanche/Source/AvalancheMedia/Public/Broadcast/Channel/AvaBroadcastOutputChannel.h
--- a/Engine/Plugins/Experimental/Avalanche/Source/AvalancheMedia/Public/Broadcast/Channel/AvaBroadcastOutputChannel.h
+++ b/Engine/Plugins/Experimental/Avalanche/Source/AvalancheMedia/Public/Broadcast/Channel/AvaBroadcastOutputChannel.h
@@ -8,8 +8,8 @@
 #include "Containers/Map.h"
 #include "Containers/StringFwd.h"
 #include "Engine/TextureRenderTarget2D.h"
-#include "Viewport/AvaViewportQualitySettings.h"
 #include "UObject/ObjectPtr.h"
+#include "Viewport/AvaViewportQualitySettings.h"
 #include "AvaBroadcastOutputChannel.generated.h"
 
 class FAudioDeviceHandle;
@@ -223,7 +223,7 @@ protected:
 	UPROPERTY()
 	TArray<FAvaBroadcastMediaOutputInfo> MediaOutputInfos;
 
-	UPROPERTY(EditAnywhere, Category = "Motion Design")
+	UPROPERTY(EditAnywhere, Category = "Motion Design", meta = (HideHeader, ShowPresets))
 	FAvaViewportQualitySettings QualitySettings;
 
 	/*
diff --git a/Engine/Plugins/Experimental/Avalanche/Source/AvalancheMediaEditor/Private/Broadcast/AvaBroadcastEditor.cpp b/Engine/Plugins/Experimental/Avalanche/Source/AvalancheMediaEditor/Private/Broadcast/AvaBroadcastEditor.cpp
--- a/Engine/Plugins/Experimental/Avalanche/Source/AvalancheMediaEditor/Private/Broadcast/AvaBroadcastEditor.cpp
+++ b/Engine/Plugins/Experimental/Avalanche/Source/AvalancheMediaEditor/Private/Broadcast/AvaBroadcastEditor.cpp
@@ -2,6 +2,7 @@
 
 #include "Broadcast/AvaBroadcastEditor.h"
 #include "AppModes/AvaBroadcastDefaultMode.h"
+#include "AvaMediaEditorSettings.h"
 #include "AvaMediaSettings.h"
 #include "Broadcast/AvaBroadcast.h"
 #include "Broadcast/ChannelGrid/AvaBroadcastOutputTileItem.h"
@@ -200,6 +201,15 @@ void FAvaBroadcastEditor::FillPlayToolBar(FToolBarBuilder& ToolBarBuilder)
 	//TODO: Change Lambdas to their Own Command Action
 	ToolBarBuilder.BeginSection(TEXT("Player"));
 	{
+		ToolBarBuilder.AddToolBarButton(
+			FUIAction(FExecuteAction::CreateSP(this, &FAvaBroadcastEditor::AddChannel)
+				, FCanExecuteAction::CreateSP(this, &FAvaBroadcastEditor::CanAddChannel))
+			, NAME_None
+			, LOCTEXT("NewChannel_Label", "New Channel")
+			, LOCTEXT("NewChannel_ToolTip", "New Channel")
+			, FSlateIcon(FAppStyle::GetAppStyleSetName(), "Icons.Plus")
+		);
+
 		ToolBarBuilder.AddToolBarButton(FUIAction(FExecuteAction::CreateLambda([Broadcast]
 				{
 					Broadcast->StartBroadcast();
@@ -549,4 +559,24 @@ FText FAvaBroadcastEditor::GetLaunchLocalServerTooltip()
 	return LOCTEXT("LaunchLocalServerManual_ToolTip", "Launches Game Mode Local Server Process (disabled, start client first)");
 }
 
+bool FAvaBroadcastEditor::CanAddChannel()
+{
+	const int32 ChannelCount = UAvaBroadcast::Get().GetCurrentProfile().GetChannels().Num();
+
+	const UAvaMediaEditorSettings& MediaSettings = UAvaMediaEditorSettings::Get();
+
+	return !MediaSettings.bBroadcastEnforceMaxChannelCount
+		|| ChannelCount < MediaSettings.BroadcastMaxChannelCount;
+}
+
+void FAvaBroadcastEditor::AddChannel()
+{
+	if (UAvaBroadcast* const Broadcast = BroadcastWeak.Get())
+	{
+		FScopedTransaction Transaction(LOCTEXT("AddChannel", "Add Channel"));
+		Broadcast->Modify();
+		Broadcast->GetCurrentProfile().AddChannel();
+	}
+}
+
 #undef LOCTEXT_NAMESPACE
diff --git a/Engine/Plugins/Experimental/Avalanche/Source/AvalancheMediaEditor/Private/Broadcast/AvaBroadcastEditor.h b/Engine/Plugins/Experimental/Avalanche/Source/AvalancheMediaEditor/Private/Broadcast/AvaBroadcastEditor.h
--- a/Engine/Plugins/Experimental/Avalanche/Source/AvalancheMediaEditor/Private/Broadcast/AvaBroadcastEditor.h
+++ b/Engine/Plugins/Experimental/Avalanche/Source/AvalancheMediaEditor/Private/Broadcast/AvaBroadcastEditor.h
@@ -74,6 +74,9 @@ public:
 	
 	void CreateNewProfile();
 	FReply OnProfileSelected(FName InProfileName);
+
+	bool CanAddChannel();
+	void AddChannel();
 	
 protected:
 	void RegisterApplicationModes();
diff --git a/Engine/Plugins/Experimental/Avalanche/Source/AvalancheMediaEditor/Private/Broadcast/ChannelGrid/Slate/SAvaBroadcastChannel.cpp b/Engine/Plugins/Experimental/Avalanche/Source/AvalancheMediaEditor/Private/Broadcast/ChannelGrid/Slate/SAvaBroadcastChannel.cpp
--- a/Engine/Plugins/Experimental/Avalanche/Source/AvalancheMediaEditor/Private/Broadcast/ChannelGrid/Slate/SAvaBroadcastChannel.cpp
+++ b/Engine/Plugins/Experimental/Avalanche/Source/AvalancheMediaEditor/Private/Broadcast/ChannelGrid/Slate/SAvaBroadcastChannel.cpp
@@ -10,6 +10,7 @@
 #include "Broadcast/OutputDevices/DragDropOps/AvaBroadcastOutputTreeItemDragDropOp.h"
 #include "Broadcast/OutputDevices/Slate/SAvaBroadcastCaptureImage.h"
 #include "Brushes/SlateImageBrush.h"
+#include "Engine/TextureRenderTarget2D.h"
 #include "Framework/Commands/GenericCommands.h"
 #include "Framework/MultiBox/MultiBoxBuilder.h"
 #include "Input/DragAndDrop.h"
@@ -22,13 +23,18 @@
 #include "Modules/ModuleManager.h"
 #include "PropertyEditorModule.h"
 #include "ScopedTransaction.h"
+#include "Styling/AppStyle.h"
 #include "Styling/SlateBrush.h"
 #include "Styling/SlateTypes.h"
+#include "Widgets/Images/SImage.h"
 #include "Widgets/Input/SButton.h"
+#include "Widgets/Input/SComboButton.h"
 #include "Widgets/Input/SMenuAnchor.h"
 #include "Widgets/Layout/SScaleBox.h"
 #include "Widgets/Layout/SScrollBox.h"
 #include "Widgets/Layout/SSeparator.h"
+#include "Widgets/SBoxPanel.h"
+#include "Widgets/SOverlay.h"
 #include "Widgets/Text/SInlineEditableTextBlock.h"
 #include "Widgets/Views/ITableRow.h"
 #include "Widgets/Views/SListView.h"
@@ -48,6 +54,13 @@ void SAvaBroadcastChannel::Construct(const FArguments& InArgs, const TSharedPtr<
 	
 	FAvaBroadcastOutputChannel::GetOnChannelChanged().AddRaw(this, &SAvaBroadcastChannel::OnChannelChanged);
 	InBroadcastEditor->OnOutputTileSelectionChanged.AddRaw(this, &SAvaBroadcastChannel::OnOutputTileSelectionChanged);
+
+	FAvaBroadcastOutputChannel& Channel = UAvaBroadcast::Get().GetCurrentProfile().GetChannelMutable(ChannelName);
+	QualitySettingsStruct = MakeShared<FStructOnScope>(FAvaBroadcastOutputChannel::StaticStruct(), reinterpret_cast<uint8*>(&Channel));
+
+	CreateQualitySettingsDetails();
+
+	const FMargin ButtonPadding = FMargin(2.f);
 	
 	ChildSlot
 	[
@@ -78,17 +91,17 @@ void SAvaBroadcastChannel::Construct(const FArguments& InArgs, const TSharedPtr<
 			+ SHorizontalBox::Slot()
 			.AutoWidth()
 			[
-				SAssignNew(ChannelSettingsMenuAnchor, SMenuAnchor)
-				.Content()
+				SNew(SComboButton)
+				.ToolTipText(LOCTEXT("ChannelSettingToolTip", "Open the channel viewport quality settings"))
+				.ButtonStyle(FAppStyle::Get(), "SimpleButton")
+				.ContentPadding(ButtonPadding)
+				.HasDownArrow(false)
+				.MenuPlacement(MenuPlacement_BelowAnchor)
+				.OnGetMenuContent(this, &SAvaBroadcastChannel::GetQualitySettingsDetails)
+				.ButtonContent()
 				[
-					SNew(SButton)
-					.ButtonStyle(FAppStyle::Get(), "SimpleButton")
-					.OnClicked(this, &SAvaBroadcastChannel::OnChannelSettingsButtonClicked)
-					.ToolTipText(LOCTEXT("ChannelSettingToolTip", "Open the channel settings"))
-					[
-						SNew(SImage)
-						.Image(FAppStyle::Get().GetBrush("Icons.Settings"))
-					]
+					SNew(SImage)
+					.Image(FAppStyle::Get().GetBrush("Icons.Settings"))
 				]
 			]
 			+ SHorizontalBox::Slot()
@@ -96,6 +109,7 @@ void SAvaBroadcastChannel::Construct(const FArguments& InArgs, const TSharedPtr<
 			[
 				SNew(SButton)
 				.ButtonStyle(FAppStyle::Get(), "SimpleButton")
+				.ContentPadding(ButtonPadding)
 				.OnClicked(this, &SAvaBroadcastChannel::OnChannelPinButtonClicked)
 				.ToolTipText(LOCTEXT("ChannelPinToolTip", "Pin the channel across all profiles."))
 				[
@@ -108,6 +122,7 @@ void SAvaBroadcastChannel::Construct(const FArguments& InArgs, const TSharedPtr<
 			[
 				SNew(SButton)
 				.ButtonStyle(FAppStyle::Get(), "SimpleButton")
+				.ContentPadding(ButtonPadding)
 				.OnClicked(this, &SAvaBroadcastChannel::OnChannelMaximizeButtonClicked)
 				.ToolTipText(this, &SAvaBroadcastChannel::GetChannelMaximizeRestoreTooltipText)
 				[
@@ -120,6 +135,7 @@ void SAvaBroadcastChannel::Construct(const FArguments& InArgs, const TSharedPtr<
 			[
 				SNew(SButton)
 				.ButtonStyle(FAppStyle::Get(), "SimpleButton")
+				.ContentPadding(ButtonPadding)
 				.OnClicked(this, &SAvaBroadcastChannel::OnChannelRemoveButtonClicked)
 				.IsEnabled(this, &SAvaBroadcastChannel::CanEditChanges)
 				.ToolTipText(LOCTEXT("ChannelRemoveToolTip", "Remove this channel."))
@@ -141,6 +157,7 @@ void SAvaBroadcastChannel::Construct(const FArguments& InArgs, const TSharedPtr<
 			[
 				SNew(SButton)
 				.ButtonStyle(FAppStyle::Get(), "SimpleButton")
+				.ContentPadding(ButtonPadding)
 				.OnClicked(this, &SAvaBroadcastChannel::OnChannelStatusButtonClicked)
 				.ToolTipText(LOCTEXT("ChannelStatusChangeToolTip", "Change the status of this channel."))
 				[
@@ -170,6 +187,7 @@ void SAvaBroadcastChannel::Construct(const FArguments& InArgs, const TSharedPtr<
 			[
 				SNew(SButton)
 				.ButtonStyle(FAppStyle::Get(), "SimpleButton")
+				.ContentPadding(ButtonPadding)
 				.OnClicked(this, &SAvaBroadcastChannel::OnChannelTypeToggleButtonClicked)
 				.ToolTipText(LOCTEXT("ChannelTypeToggleToolTip", "Toggles between \"Program\" or \"Preview\" channel type."))
 				[
@@ -181,7 +199,7 @@ void SAvaBroadcastChannel::Construct(const FArguments& InArgs, const TSharedPtr<
 	];
 
 	OutputTileListView->SetStyle(nullptr);	
-	const FAvaBroadcastOutputChannel& Channel = UAvaBroadcast::Get().GetCurrentProfile().GetChannel(ChannelName);
+
 	OnChannelBroadcastStateChanged(Channel);
 	OnChannelMediaOutputsChanged(Channel);
 }
@@ -609,54 +627,6 @@ FReply SAvaBroadcastChannel::OnChannelStatusButtonClicked()
 	return FReply::Unhandled();
 }
 
-FReply SAvaBroadcastChannel::OnChannelSettingsButtonClicked()
-{
-	if (!ChannelSettings.IsValid())
-	{
-		FAvaBroadcastOutputChannel& Channel = UAvaBroadcast::Get().GetCurrentProfile().GetChannelMutable(ChannelName);
-		if (!Channel.IsValidChannel())
-		{
-			return FReply::Unhandled();
-		}
-
-		FPropertyEditorModule& PropertyEditorModule = FModuleManager::GetModuleChecked<FPropertyEditorModule>("PropertyEditor");
-
-		FDetailsViewArgs DetailsViewArgs;
-		DetailsViewArgs.NameAreaSettings = FDetailsViewArgs::HideNameArea;
-		DetailsViewArgs.bAllowSearch = false;
-		DetailsViewArgs.bShowOptions = false;
-
-		TSharedRef<FStructOnScope> ChannelStruct = MakeShared<FStructOnScope>(FAvaBroadcastOutputChannel::StaticStruct(), reinterpret_cast<uint8*>(&Channel));
-
-		ChannelSettings = PropertyEditorModule.CreateStructureDetailView(DetailsViewArgs, FStructureDetailsViewArgs(), ChannelStruct);
-
-		ChannelSettings->GetOnFinishedChangingPropertiesDelegate().AddLambda(
-			[ChannelStructWeak = TWeakPtr<FStructOnScope>(ChannelStruct)](const FPropertyChangedEvent&)
-			{
-				if (TSharedPtr<FStructOnScope> ChannelStruct = ChannelStructWeak.Pin())
-				{
-					const FAvaBroadcastOutputChannel& Channel = *reinterpret_cast<const FAvaBroadcastOutputChannel*>(ChannelStruct->GetStructMemory());
-					FAvaBroadcastOutputChannel::GetOnChannelChanged().Broadcast(Channel, EAvaBroadcastChannelChange::Settings);
-				}
-			});
-
-		ChannelSettingsMenuAnchor->SetMenuContent(SNew(SBox)
-				.MaxDesiredHeight(500.f)
-				.MinDesiredWidth(150.f)
-				[
-					SNew(SScrollBox)
-					+ SScrollBox::Slot()
-					[
-				    	ChannelSettings->GetWidget().ToSharedRef()
-					]
-				]
-			);
-	}
-
-	ChannelSettingsMenuAnchor->SetIsOpen(true);
-	return FReply::Handled();
-}
-
 FReply SAvaBroadcastChannel::OnChannelPinButtonClicked()
 {
 	FScopedTransaction Transaction(LOCTEXT("ToggleChannelPin", "Toggle Channel Pin"));
@@ -814,4 +784,40 @@ FText SAvaBroadcastChannel::GetChannelMaximizeRestoreTooltipText() const
 	return TooltipText;
 }
 
+void SAvaBroadcastChannel::CreateQualitySettingsDetails()
+{
+	FPropertyEditorModule& PropertyEditorModule = FModuleManager::GetModuleChecked<FPropertyEditorModule>(TEXT("PropertyEditor"));
+
+	FDetailsViewArgs DetailsViewArgs;
+	DetailsViewArgs.NameAreaSettings = FDetailsViewArgs::HideNameArea;
+	DetailsViewArgs.bAllowSearch = false;
+	DetailsViewArgs.bShowOptions = false;
+
+	QualitySettingsDetails = PropertyEditorModule.CreateStructureDetailView(DetailsViewArgs, FStructureDetailsViewArgs(), QualitySettingsStruct);
+	QualitySettingsDetails->GetOnFinishedChangingPropertiesDelegate().AddSP(this, &SAvaBroadcastChannel::OnChannelQualitySettingsChanged);
+}
+
+void SAvaBroadcastChannel::OnChannelQualitySettingsChanged(const FPropertyChangedEvent& InEvent)
+{
+	if (QualitySettingsStruct.IsValid())
+	{
+		const FAvaBroadcastOutputChannel& Channel = *reinterpret_cast<const FAvaBroadcastOutputChannel*>(QualitySettingsStruct->GetStructMemory());
+		FAvaBroadcastOutputChannel::GetOnChannelChanged().Broadcast(Channel, EAvaBroadcastChannelChange::Settings);
+	}
+}
+
+TSharedRef<SWidget> SAvaBroadcastChannel::GetQualitySettingsDetails() const
+{
+	return SNew(SBox)
+		.MaxDesiredHeight(500.f)
+		.MinDesiredWidth(120.f)
+		[
+			SNew(SScrollBox)
+			+ SScrollBox::Slot()
+			[
+				QualitySettingsDetails->GetWidget().ToSharedRef()
+			]
+		];
+}
+
 #undef LOCTEXT_NAMESPACE
diff --git a/Engine/Plugins/Experimental/Avalanche/Source/AvalancheMediaEditor/Private/Broadcast/ChannelGrid/Slate/SAvaBroadcastChannel.h b/Engine/Plugins/Experimental/Avalanche/Source/AvalancheMediaEditor/Private/Broadcast/ChannelGrid/Slate/SAvaBroadcastChannel.h
--- a/Engine/Plugins/Experimental/Avalanche/Source/AvalancheMediaEditor/Private/Broadcast/ChannelGrid/Slate/SAvaBroadcastChannel.h
+++ b/Engine/Plugins/Experimental/Avalanche/Source/AvalancheMediaEditor/Private/Broadcast/ChannelGrid/Slate/SAvaBroadcastChannel.h
@@ -13,11 +13,13 @@
 class FAvaBroadcastEditor;
 class FName;
 class FReply;
+class FStructOnScope;
 class FText;
 class FUICommandList;
 class IStructureDetailsView;
 class ITableRow;
 class SAvaBroadcastChannel;
+class SAvaViewportQualitySettings;
 class SInlineEditableTextBlock;
 class SlateBrush;
 class SMenuAnchor;
@@ -86,7 +88,6 @@ public:
 	void OnChannelNameTextCommitted(const FText& InText, ETextCommit::Type CommitInfo);
 
 	FReply OnChannelStatusButtonClicked();
-	FReply OnChannelSettingsButtonClicked();
 	FReply OnChannelPinButtonClicked();
 	FReply OnChannelTypeToggleButtonClicked();
 	FReply OnChannelMaximizeButtonClicked();
@@ -111,6 +112,12 @@ public:
 	
 protected:
 	FText GetChannelMaximizeRestoreTooltipText() const;
+
+	void CreateQualitySettingsDetails();
+
+	void OnChannelQualitySettingsChanged(const FPropertyChangedEvent& InEvent);
+
+	TSharedRef<SWidget> GetQualitySettingsDetails() const;
 	
 	FName ChannelName = NAME_None;
 
@@ -118,9 +125,9 @@ protected:
 	
 	TSharedPtr<SListView<FAvaBroadcastOutputTileItemPtr>> OutputTileListView;
 	
-	TSharedPtr<IStructureDetailsView> ChannelSettings;
+	TSharedPtr<FStructOnScope> QualitySettingsStruct;
 
-	TSharedPtr<SMenuAnchor> ChannelSettingsMenuAnchor;
+	TSharedPtr<IStructureDetailsView> QualitySettingsDetails;
 
 	TArray<FAvaBroadcastOutputTileItemPtr> OutputTileItems;
 	
diff --git a/Engine/Plugins/Experimental/Avalanche/Source/AvalancheMediaEditor/Private/Broadcast/ChannelGrid/Slate/SAvaBroadcastChannels.cpp b/Engine/Plugins/Experimental/Avalanche/Source/AvalancheMediaEditor/Private/Broadcast/ChannelGrid/Slate/SAvaBroadcastChannels.cpp
--- a/Engine/Plugins/Experimental/Avalanche/Source/AvalancheMediaEditor/Private/Broadcast/ChannelGrid/Slate/SAvaBroadcastChannels.cpp
+++ b/Engine/Plugins/Experimental/Avalanche/Source/AvalancheMediaEditor/Private/Broadcast/ChannelGrid/Slate/SAvaBroadcastChannels.cpp
@@ -88,20 +88,7 @@ void SAvaBroadcastChannels::Construct(const FArguments& InArgs, const TSharedPtr
 	.HAlign(EHorizontalAlignment::HAlign_Fill)
 	.VAlign(EVerticalAlignment::VAlign_Fill)
 	[
-		SNew(SVerticalBox)
-		+ SVerticalBox::Slot()
-		.AutoHeight()
-		.HAlign(EHorizontalAlignment::HAlign_Center)
-		[
-			MakeChannelsToolbar()
-		]
-		+ SVerticalBox::Slot()
-		.FillHeight(1.f)
-		.HAlign(EHorizontalAlignment::HAlign_Fill)
-		.VAlign(EVerticalAlignment::VAlign_Fill)
-		[
-			ChannelGrid.ToSharedRef()
-		]
+		ChannelGrid.ToSharedRef()
 	];
 
 	RefreshChannelGrid();
@@ -115,26 +102,6 @@ SAvaBroadcastChannels::~SAvaBroadcastChannels()
 	}
 }
 
-bool SAvaBroadcastChannels::CanAddChannel()
-{
-	const int32 ChannelCount = UAvaBroadcast::Get().GetCurrentProfile().GetChannels().Num();
-					
-	const UAvaMediaEditorSettings& MediaSettings = UAvaMediaEditorSettings::Get();
-	
-	return !MediaSettings.bBroadcastEnforceMaxChannelCount
-		|| ChannelCount < MediaSettings.BroadcastMaxChannelCount;
-}
-
-void SAvaBroadcastChannels::AddChannel()
-{
-	if (UAvaBroadcast* const Broadcast = BroadcastWeak.Get())
-	{
-		FScopedTransaction Transaction(LOCTEXT("AddChannel", "Add Channel"));
-		Broadcast->Modify();
-		Broadcast->GetCurrentProfile().AddChannel();
-	}
-}
-
 float SAvaBroadcastChannels::GetRowFill(int32 RowId) const
 {
 	return ChannelMaximizer.GetRowFill(RowId);
@@ -145,27 +112,6 @@ float SAvaBroadcastChannels::GetColumnFill(int32 ColumnId) const
 	return ChannelMaximizer.GetColumnFill(ColumnId);
 }
 
-TSharedRef<SWidget> SAvaBroadcastChannels::MakeChannelsToolbar()
-{
-	FSlimHorizontalToolBarBuilder ToolBarBuilder(nullptr, FMultiBoxCustomization::None);
-	
-	ToolBarBuilder.BeginSection(TEXT("Channels"));
-	{
-		const FUIAction AddChannelAction(FExecuteAction::CreateSP(this, &SAvaBroadcastChannels::AddChannel)
-			, FCanExecuteAction::CreateSP(this, &SAvaBroadcastChannels::CanAddChannel));
-		
-		ToolBarBuilder.AddToolBarButton(AddChannelAction
-			, NAME_None
-			, LOCTEXT("NewChannel_Label", "New Channel")
-			, LOCTEXT("NewChannel_ToolTip", "New Channel")
-			, FSlateIcon(FAppStyle::GetAppStyleSetName(), "Icons.Plus")
-		);		
-	}
-	ToolBarBuilder.EndSection();
-	
-	return ToolBarBuilder.MakeWidget();
-}
-
 bool SAvaBroadcastChannels::CanMaximizeChannel() const
 {
 	const int32 ChannelCount = Channels.Num();
diff --git a/Engine/Plugins/Experimental/Avalanche/Source/AvalancheMediaEditor/Private/Broadcast/ChannelGrid/Slate/SAvaBroadcastChannels.h b/Engine/Plugins/Experimental/Avalanche/Source/AvalancheMediaEditor/Private/Broadcast/ChannelGrid/Slate/SAvaBroadcastChannels.h
--- a/Engine/Plugins/Experimental/Avalanche/Source/AvalancheMediaEditor/Private/Broadcast/ChannelGrid/Slate/SAvaBroadcastChannels.h
+++ b/Engine/Plugins/Experimental/Avalanche/Source/AvalancheMediaEditor/Private/Broadcast/ChannelGrid/Slate/SAvaBroadcastChannels.h
@@ -40,14 +40,9 @@ public:
 	void Construct(const FArguments& InArgs, const TSharedPtr<FAvaBroadcastEditor>& InBroadcastEditor);
 	
 	virtual ~SAvaBroadcastChannels() override;
-
-	bool CanAddChannel();
-	void AddChannel();
 	
 	float GetRowFill(int32 RowId) const;
 	float GetColumnFill(int32 ColumnId) const;
-
-	TSharedRef<SWidget> MakeChannelsToolbar();
 	
 	bool CanMaximizeChannel() const;
 	void ToggleMaximizeChannel(const TSharedRef<SAvaBroadcastChannel>& InWidget);
