MotionDesign : PropertyAnimator
- Clone/Copy operation support in outliner, resolving properties against new owner if available otherwise properties are unlinked automatically (same during duplication)

diff --git a/Engine/Plugins/Experimental/Avalanche/Source/AvalanchePropertyAnimatorEditor/Private/AvaPropertyAnimatorEditorModule.cpp b/Engine/Plugins/Experimental/Avalanche/Source/AvalanchePropertyAnimatorEditor/Private/AvaPropertyAnimatorEditorModule.cpp
--- a/Engine/Plugins/Experimental/Avalanche/Source/AvalanchePropertyAnimatorEditor/Private/AvaPropertyAnimatorEditorModule.cpp
+++ b/Engine/Plugins/Experimental/Avalanche/Source/AvalanchePropertyAnimatorEditor/Private/AvaPropertyAnimatorEditorModule.cpp
@@ -2,11 +2,13 @@
 
 #include "AvaPropertyAnimatorEditorModule.h"
 
+#include "DragDropOps/AvaOutlinerItemDragDropOp.h"
 #include "IAvaOutliner.h"
 #include "IAvaOutlinerModule.h"
 #include "Item/AvaOutlinerActor.h"
 #include "ItemProxies/AvaOutlinerItemProxyRegistry.h"
 #include "Outliner/AvaPropertyAnimatorEditorOutlinerContextMenu.h"
+#include "Outliner/AvaPropertyAnimatorEditorOutlinerDropHandler.h"
 #include "Outliner/AvaPropertyAnimatorEditorOutlinerProxy.h"
 
 #define LOCTEXT_NAMESPACE "AvaPropertyAnimatorEditorModule"
@@ -41,6 +43,12 @@ void FAvaPropertyAnimatorEditorModule::RegisterOutlinerItems()
 
 	OutlinerContextDelegateHandle = IAvaOutlinerModule::Get().GetOnExtendOutlinerItemContextMenu()
 		.AddStatic(&FAvaPropertyAnimatorEditorOutlinerContextMenu::OnExtendOutlinerContextMenu);
+
+	OutlinerDropHandlerDelegateHandle = FAvaOutlinerItemDragDropOp::OnItemDragDropOpInitialized().AddStatic(
+		[](FAvaOutlinerItemDragDropOp& InDragDropOp)
+		{
+			InDragDropOp.AddDropHandler<FAvaPropertyAnimatorEditorOutlinerDropHandler>();
+		});
 }
 
 void FAvaPropertyAnimatorEditorModule::UnregisterOutlinerItems()
@@ -57,6 +65,9 @@ void FAvaPropertyAnimatorEditorModule::UnregisterOutlinerItems()
 
 		OutlinerModule.GetOnExtendOutlinerItemContextMenu().Remove(OutlinerContextDelegateHandle);
 		OutlinerContextDelegateHandle.Reset();
+
+		FAvaOutlinerItemDragDropOp::OnItemDragDropOpInitialized().Remove(OutlinerDropHandlerDelegateHandle);
+		OutlinerDropHandlerDelegateHandle.Reset();
 	}
 }
 
diff --git a/Engine/Plugins/Experimental/Avalanche/Source/AvalanchePropertyAnimatorEditor/Private/AvaPropertyAnimatorEditorModule.h b/Engine/Plugins/Experimental/Avalanche/Source/AvalanchePropertyAnimatorEditor/Private/AvaPropertyAnimatorEditorModule.h
--- a/Engine/Plugins/Experimental/Avalanche/Source/AvalanchePropertyAnimatorEditor/Private/AvaPropertyAnimatorEditorModule.h
+++ b/Engine/Plugins/Experimental/Avalanche/Source/AvalanchePropertyAnimatorEditor/Private/AvaPropertyAnimatorEditorModule.h
@@ -18,4 +18,5 @@ protected:
 
 	FDelegateHandle OutlinerProxiesExtensionDelegateHandle;
 	FDelegateHandle OutlinerContextDelegateHandle;
+	FDelegateHandle OutlinerDropHandlerDelegateHandle;
 };
diff --git a/Engine/Plugins/Experimental/Avalanche/Source/AvalanchePropertyAnimatorEditor/Private/Outliner/AvaPropertyAnimatorEditorOutlinerDropHandler.cpp b/Engine/Plugins/Experimental/Avalanche/Source/AvalanchePropertyAnimatorEditor/Private/Outliner/AvaPropertyAnimatorEditorOutlinerDropHandler.cpp
new file mode 100644
--- /dev/null
+++ b/Engine/Plugins/Experimental/Avalanche/Source/AvalanchePropertyAnimatorEditor/Private/Outliner/AvaPropertyAnimatorEditorOutlinerDropHandler.cpp
@@ -0,0 +1,215 @@
+// Copyright Epic Games, Inc. All Rights Reserved.
+
+#include "Outliner/AvaPropertyAnimatorEditorOutlinerDropHandler.h"
+
+#include "Components/PropertyAnimatorCoreComponent.h"
+#include "Framework/Notifications/NotificationManager.h"
+#include "Item/AvaOutlinerActor.h"
+#include "Outliner/AvaPropertyAnimatorEditorOutliner.h"
+#include "Outliner/AvaPropertyAnimatorEditorOutlinerProxy.h"
+#include "Subsystems/PropertyAnimatorCoreSubsystem.h"
+#include "Widgets/Notifications/SNotificationList.h"
+#include "Widgets/Views/STableRow.h"
+
+#define LOCTEXT_NAMESPACE "AvaPropertyAnimatorEditorOutlinerDropHandler"
+
+DEFINE_LOG_CATEGORY_STATIC(LogAvaPropertyAnimatorEditorOutlinerDropHandler, Log, All);
+
+bool FAvaPropertyAnimatorEditorOutlinerDropHandler::IsDraggedItemSupported(const FAvaOutlinerItemPtr& InDraggedItem) const
+{
+	return InDraggedItem->IsA<FAvaPropertyAnimatorEditorOutliner>() || InDraggedItem->IsA<FAvaPropertyAnimatorEditorOutlinerProxy>();
+}
+
+TOptional<EItemDropZone> FAvaPropertyAnimatorEditorOutlinerDropHandler::CanDrop(EItemDropZone InDropZone, FAvaOutlinerItemPtr InTargetItem) const
+{
+	// Dropping on actor directly
+	if (const FAvaOutlinerActor* const TargetActorItem = InTargetItem->CastTo<FAvaOutlinerActor>())
+	{
+		return CanDropOnActor(TargetActorItem->GetActor(), InDropZone);
+	}
+
+	// Dropping onto the animator proxy itself has the same effects as dropping on actor
+	if (const FAvaPropertyAnimatorEditorOutlinerProxy* const TargetAnimatorProxy = InTargetItem->CastTo<FAvaPropertyAnimatorEditorOutlinerProxy>())
+	{
+		if (const UPropertyAnimatorCoreComponent* const AnimatorComponent = TargetAnimatorProxy->GetPropertyAnimatorComponent())
+		{
+			return CanDropOnActor(AnimatorComponent->GetOwner(), InDropZone);
+		}
+
+		return TOptional<EItemDropZone>();
+	}
+
+	// If target item is none of the above, nor an animator item, then it's not a supported target
+	const FAvaPropertyAnimatorEditorOutliner* const TargetAnimatorItem = InTargetItem->CastTo<FAvaPropertyAnimatorEditorOutliner>();
+	if (!TargetAnimatorItem || !TargetAnimatorItem->GetPropertyAnimator())
+	{
+		return TOptional<EItemDropZone>();
+	}
+
+	// Only clone animator on target, move not supported
+	const UPropertyAnimatorCoreBase* TargetAnimator = TargetAnimatorItem->GetPropertyAnimator();
+	TSet<UPropertyAnimatorCoreBase*> DraggedAnimators = GetDraggedAnimators();
+
+	for (TSet<UPropertyAnimatorCoreBase*>::TIterator It(DraggedAnimators); It; ++It)
+	{
+		const UPropertyAnimatorCoreBase* Animator = *It;
+		if (!Animator || Animator->GetAnimatorActor() == TargetAnimator->GetAnimatorActor())
+		{
+			It.RemoveCurrent();
+		}
+	}
+
+	// If the animators are empty, return fail early
+	if (DraggedAnimators.IsEmpty())
+	{
+		return TOptional<EItemDropZone>();
+	}
+
+	return InDropZone;
+}
+
+bool FAvaPropertyAnimatorEditorOutlinerDropHandler::Drop(EItemDropZone InDropZone, FAvaOutlinerItemPtr InTargetItem)
+{
+	if (const FAvaOutlinerActor* const TargetActorItem = InTargetItem->CastTo<FAvaOutlinerActor>())
+	{
+		return DropAnimatorsOnActor(TargetActorItem->GetActor(), InDropZone);
+	}
+
+	// Dropping onto the animator proxy itself has the same effects as dropping on actor
+	if (const FAvaPropertyAnimatorEditorOutlinerProxy* const TargetAnimatorProxy = InTargetItem->CastTo<FAvaPropertyAnimatorEditorOutlinerProxy>())
+	{
+		if (const UPropertyAnimatorCoreComponent* const AnimatorComponent = TargetAnimatorProxy->GetPropertyAnimatorComponent())
+		{
+			return DropAnimatorsOnActor(AnimatorComponent->GetOwner(), InDropZone);
+		}
+
+		return false;
+	}
+
+	// If target item is none of the above, nor an animator item, then it's not a supported target
+	const FAvaPropertyAnimatorEditorOutliner* const TargetAnimatorItem = InTargetItem->CastTo<FAvaPropertyAnimatorEditorOutliner>();
+	if (!TargetAnimatorItem || !TargetAnimatorItem->GetPropertyAnimator())
+	{
+		return false;
+	}
+
+	UPropertyAnimatorCoreBase* TargetAnimator = TargetAnimatorItem->GetPropertyAnimator();
+	return DropAnimatorsOnAnimator(TargetAnimator, InDropZone);
+}
+
+TSet<UPropertyAnimatorCoreBase*> FAvaPropertyAnimatorEditorOutlinerDropHandler::GetDraggedAnimators() const
+{
+	TSet<UPropertyAnimatorCoreBase*> DraggedAnimators;
+
+	// Rather than iterating separately, keep order of how they were dragged
+	ForEachItem<IAvaOutlinerItem>([&DraggedAnimators](IAvaOutlinerItem& InItem)->EIterationResult
+	{
+		if (const FAvaPropertyAnimatorEditorOutlinerProxy* const AnimatorItemProxy = InItem.CastTo<FAvaPropertyAnimatorEditorOutlinerProxy>())
+		{
+			if (const UPropertyAnimatorCoreComponent* const AnimatorComponent = AnimatorItemProxy->GetPropertyAnimatorComponent())
+			{
+				AnimatorComponent->ForEachAnimator([&DraggedAnimators](UPropertyAnimatorCoreBase* InAnimator)
+				{
+					DraggedAnimators.Add(InAnimator);
+					return true;
+				});
+			}
+		}
+		else if (const FAvaPropertyAnimatorEditorOutliner* const AnimatorItem = InItem.CastTo<FAvaPropertyAnimatorEditorOutliner>())
+		{
+			if (UPropertyAnimatorCoreBase* const Animator = AnimatorItem->GetPropertyAnimator())
+			{
+				DraggedAnimators.Add(Animator);
+			}
+		}
+
+		return EIterationResult::Continue;
+	});
+
+	return DraggedAnimators;
+}
+
+TOptional<EItemDropZone> FAvaPropertyAnimatorEditorOutlinerDropHandler::CanDropOnActor(AActor* InActor, EItemDropZone InDropZone) const
+{
+	if (!IsValid(InActor))
+	{
+		return TOptional<EItemDropZone>();
+	}
+
+	const UPropertyAnimatorCoreSubsystem* const AnimatorSubsystem = UPropertyAnimatorCoreSubsystem::Get();
+	if (!AnimatorSubsystem)
+	{
+		return TOptional<EItemDropZone>();
+	}
+
+	TSet<UPropertyAnimatorCoreBase*> DraggedAnimators = GetDraggedAnimators();
+
+	for (TSet<UPropertyAnimatorCoreBase*>::TIterator It(DraggedAnimators); It; ++It)
+	{
+		const UPropertyAnimatorCoreBase* Animator = *It;
+		if (!Animator || Animator->GetAnimatorActor() == InActor)
+		{
+			It.RemoveCurrent();
+		}
+	}
+
+	// For actor items, drop zone can only be onto the actor
+	return !DraggedAnimators.IsEmpty()
+		? EItemDropZone::OntoItem
+		: TOptional<EItemDropZone>();
+}
+
+bool FAvaPropertyAnimatorEditorOutlinerDropHandler::DropAnimatorsOnActor(AActor* InActor, EItemDropZone InDropZone) const
+{
+	const UPropertyAnimatorCoreSubsystem* const AnimatorSubsystem = UPropertyAnimatorCoreSubsystem::Get();
+
+	if (!AnimatorSubsystem)
+	{
+		return false;
+	}
+
+	TSet<UPropertyAnimatorCoreBase*> CloneAnimators = GetDraggedAnimators();
+
+	// Remove animators that are already on the target actor
+	for (TSet<UPropertyAnimatorCoreBase*>::TIterator It(CloneAnimators); It; ++It)
+	{
+		const UPropertyAnimatorCoreBase* Animator = *It;
+		if (!Animator || Animator->GetAnimatorActor() == InActor)
+		{
+			It.RemoveCurrent();
+		}
+	}
+
+	if (CloneAnimators.IsEmpty())
+	{
+		return false;
+	}
+
+	UE_LOG(LogAvaPropertyAnimatorEditorOutlinerDropHandler, Log, TEXT("Dropping %i animator(s) on actor %s")
+		, CloneAnimators.Num()
+		, *InActor->GetActorNameOrLabel());
+
+	const TSet<UPropertyAnimatorCoreBase*> NewAnimators = AnimatorSubsystem->CloneAnimators(CloneAnimators, InActor, true);
+
+	if (NewAnimators.Num() != CloneAnimators.Num())
+	{
+		FNotificationInfo NotificationInfo(LOCTEXT("CloneAnimatorsFail", "An issue occured while cloning animators on an actor"));
+		NotificationInfo.ExpireDuration = 3.f;
+		NotificationInfo.bFireAndForget = true;
+		FSlateNotificationManager::Get().AddNotification(NotificationInfo);
+	}
+
+	return NewAnimators.Num() > 0;
+}
+
+bool FAvaPropertyAnimatorEditorOutlinerDropHandler::DropAnimatorsOnAnimator(UPropertyAnimatorCoreBase* InTargetAnimator, EItemDropZone InDropZone) const
+{
+	if (!IsValid(InTargetAnimator))
+	{
+		return false;
+	}
+
+	return DropAnimatorsOnActor(InTargetAnimator->GetAnimatorActor(), InDropZone);
+}
+
+#undef LOCTEXT_NAMESPACE
diff --git a/Engine/Plugins/Experimental/Avalanche/Source/AvalanchePropertyAnimatorEditor/Private/Outliner/AvaPropertyAnimatorEditorOutlinerDropHandler.h b/Engine/Plugins/Experimental/Avalanche/Source/AvalanchePropertyAnimatorEditor/Private/Outliner/AvaPropertyAnimatorEditorOutlinerDropHandler.h
new file mode 100644
--- /dev/null
+++ b/Engine/Plugins/Experimental/Avalanche/Source/AvalanchePropertyAnimatorEditor/Private/Outliner/AvaPropertyAnimatorEditorOutlinerDropHandler.h
@@ -0,0 +1,26 @@
+// Copyright Epic Games, Inc. All Rights Reserved.
+
+#pragma once
+
+#include "DragDropOps/Handlers/AvaOutlinerItemDropHandler.h"
+
+class UPropertyAnimatorCoreBase;
+
+/** Class that handles dropping animator item into a target item */
+class FAvaPropertyAnimatorEditorOutlinerDropHandler : public FAvaOutlinerItemDropHandler
+{
+public:
+	UE_AVA_INHERITS(FAvaPropertyAnimatorEditorOutlinerDropHandler, FAvaOutlinerItemDropHandler);
+
+protected:
+	//~ Begin FAvaOutlinerItemDropHandler
+	virtual bool IsDraggedItemSupported(const FAvaOutlinerItemPtr& InDraggedItem) const override;
+	virtual TOptional<EItemDropZone> CanDrop(EItemDropZone InDropZone, FAvaOutlinerItemPtr InTargetItem) const override;
+	virtual bool Drop(EItemDropZone InDropZone, FAvaOutlinerItemPtr InTargetItem) override;
+	//~ End FAvaOutlinerItemDropHandler
+
+	TSet<UPropertyAnimatorCoreBase*> GetDraggedAnimators() const;
+	TOptional<EItemDropZone> CanDropOnActor(AActor* InActor, EItemDropZone InDropZone) const;
+	bool DropAnimatorsOnActor(AActor* InActor, EItemDropZone InDropZone) const;
+	bool DropAnimatorsOnAnimator(UPropertyAnimatorCoreBase* InTargetAnimator, EItemDropZone InDropZone) const;
+};
diff --git a/Engine/Plugins/Experimental/PropertyAnimatorCore/Source/PropertyAnimatorCore/Private/Animators/PropertyAnimatorCoreBase.cpp b/Engine/Plugins/Experimental/PropertyAnimatorCore/Source/PropertyAnimatorCore/Private/Animators/PropertyAnimatorCoreBase.cpp
--- a/Engine/Plugins/Experimental/PropertyAnimatorCore/Source/PropertyAnimatorCore/Private/Animators/PropertyAnimatorCoreBase.cpp
+++ b/Engine/Plugins/Experimental/PropertyAnimatorCore/Source/PropertyAnimatorCore/Private/Animators/PropertyAnimatorCoreBase.cpp
@@ -130,14 +130,28 @@ void UPropertyAnimatorCoreBase::PostLoad()
 	OnAnimatorEnabledChanged();
 }
 
-void UPropertyAnimatorCoreBase::PreDuplicate(FObjectDuplicationParameters& InDupParams)
+void UPropertyAnimatorCoreBase::PostEditImport()
 {
-	Super::PreDuplicate(InDupParams);
+	Super::PostEditImport();
+
+	ResolveProperties();
+}
+
+void UPropertyAnimatorCoreBase::PreDuplicate(FObjectDuplicationParameters& InParams)
+{
+	Super::PreDuplicate(InParams);
 
 	constexpr bool bForceReset = true;
 	RestoreProperties(bForceReset);
 }
 
+void UPropertyAnimatorCoreBase::PostDuplicate(EDuplicateMode::Type InMode)
+{
+	Super::PostDuplicate(InMode);
+
+	ResolveProperties();
+}
+
 #if WITH_EDITOR
 void UPropertyAnimatorCoreBase::PreEditUndo()
 {
@@ -368,6 +382,31 @@ void UPropertyAnimatorCoreBase::OnTimeSourceNameChanged()
 	{
 		ActiveTimeSource->ActivateTimeSource();
 	}
+
+	OnTimeSourceChanged();
+}
+
+void UPropertyAnimatorCoreBase::ResolveProperties()
+{
+	// Resolve linked properties against current actor
+	TSet<FPropertyAnimatorCoreData> UnresolvedProperties;
+
+	ForEachLinkedProperty<UPropertyAnimatorCoreContext>(
+		[this, &UnresolvedProperties](UPropertyAnimatorCoreContext* InContext, const FPropertyAnimatorCoreData& InProperty)->bool
+		{
+			if (!InContext->ResolveProperty())
+			{
+				UnresolvedProperties.Add(InProperty);
+			}
+
+			return true;
+		}, false);
+
+	// Remove unresolved properties
+	for (const FPropertyAnimatorCoreData& UnresolvedProperty : UnresolvedProperties)
+	{
+		UnlinkProperty(UnresolvedProperty);
+	}
 }
 
 void UPropertyAnimatorCoreBase::EvaluateAnimator()
@@ -624,24 +663,11 @@ bool UPropertyAnimatorCoreBase::LinkProperty(const FPropertyAnimatorCoreData& In
 
 bool UPropertyAnimatorCoreBase::UnlinkProperty(const FPropertyAnimatorCoreData& InUnlinkProperty)
 {
-	if (!InUnlinkProperty.IsResolved())
-	{
-		return false;
-	}
-
 	if (!IsPropertyLinked(InUnlinkProperty))
 	{
 		return false;
 	}
 
-	const UObject* Owner = InUnlinkProperty.GetOwner();
-	const AActor* OwningActor = GetTypedOuter<AActor>();
-
-	if (Owner != OwningActor && !Owner->IsIn(OwningActor))
-	{
-		return false;
-	}
-
 	UPropertyAnimatorCoreContext* PropertyContext = GetLinkedPropertyContext(InUnlinkProperty);
 
 	PropertyContext->Restore();
diff --git a/Engine/Plugins/Experimental/PropertyAnimatorCore/Source/PropertyAnimatorCore/Private/Components/PropertyAnimatorCoreComponent.cpp b/Engine/Plugins/Experimental/PropertyAnimatorCore/Source/PropertyAnimatorCore/Private/Components/PropertyAnimatorCoreComponent.cpp
--- a/Engine/Plugins/Experimental/PropertyAnimatorCore/Source/PropertyAnimatorCore/Private/Components/PropertyAnimatorCoreComponent.cpp
+++ b/Engine/Plugins/Experimental/PropertyAnimatorCore/Source/PropertyAnimatorCore/Private/Components/PropertyAnimatorCoreComponent.cpp
@@ -26,6 +26,30 @@ UPropertyAnimatorCoreBase* UPropertyAnimatorCoreComponent::AddAnimator(const UCl
 	return NewAnimator;
 }
 
+UPropertyAnimatorCoreBase* UPropertyAnimatorCoreComponent::CloneAnimator(UPropertyAnimatorCoreBase* InAnimator)
+{
+	UPropertyAnimatorCoreBase* CloneAnimator = nullptr;
+
+	if (!InAnimator)
+	{
+		return CloneAnimator;
+	}
+
+	// Duplicate animator
+	FObjectDuplicationParameters Parameters = InitStaticDuplicateObjectParams(InAnimator, this);
+	CloneAnimator = Cast<UPropertyAnimatorCoreBase>(StaticDuplicateObjectEx(Parameters));
+
+	// Force current state
+	CloneAnimator->OnAnimatorEnabledChanged();
+
+	AnimatorsInternal = Animators;
+	Animators.Add(CloneAnimator);
+
+	OnAnimatorsChanged();
+
+	return CloneAnimator;
+}
+
 bool UPropertyAnimatorCoreComponent::RemoveAnimator(UPropertyAnimatorCoreBase* InAnimator)
 {
 	if (!Animators.Contains(InAnimator))
diff --git a/Engine/Plugins/Experimental/PropertyAnimatorCore/Source/PropertyAnimatorCore/Private/Properties/PropertyAnimatorCoreContext.cpp b/Engine/Plugins/Experimental/PropertyAnimatorCore/Source/PropertyAnimatorCore/Private/Properties/PropertyAnimatorCoreContext.cpp
--- a/Engine/Plugins/Experimental/PropertyAnimatorCore/Source/PropertyAnimatorCore/Private/Properties/PropertyAnimatorCoreContext.cpp
+++ b/Engine/Plugins/Experimental/PropertyAnimatorCore/Source/PropertyAnimatorCore/Private/Properties/PropertyAnimatorCoreContext.cpp
@@ -200,6 +200,46 @@ void UPropertyAnimatorCoreContext::OnGroupNameChanged()
 	}
 }
 
+bool UPropertyAnimatorCoreContext::ResolveProperty()
+{
+	UObject* NewOwner = GetTypedOuter<AActor>();
+
+	if (AnimatedProperty.GetOwningActor() == NewOwner)
+	{
+		return true;
+	}
+
+	bool bFound = IsValid(NewOwner);
+
+	for (const UObject* OtherOuter : AnimatedProperty.GetOuters(AnimatedProperty.GetOwningActor()))
+	{
+		bFound = false;
+		TArray<UObject*> ThisOwnedObjects;
+		GetObjectsWithOuter(NewOwner, ThisOwnedObjects, false);
+
+		for (UObject* ThisOuter : ThisOwnedObjects)
+		{
+			if (ThisOuter->GetClass() == OtherOuter->GetClass()
+				&& ThisOuter->GetFName() == OtherOuter->GetFName())
+			{
+				bFound = true;
+				NewOwner = ThisOuter;
+				break;
+			}
+		}
+	}
+
+	if (bFound
+		&& NewOwner->GetClass() == AnimatedProperty.GetOwner()->GetClass()
+		&& FindFProperty<FProperty>(NewOwner->GetClass(), AnimatedProperty.GetMemberPropertyName()))
+	{
+		SetAnimatedPropertyOwner(NewOwner);
+		return true;
+	}
+
+	return false;
+}
+
 void UPropertyAnimatorCoreContext::ConstructInternal(const FPropertyAnimatorCoreData& InProperty)
 {
 	AnimatedProperty = InProperty;
diff --git a/Engine/Plugins/Experimental/PropertyAnimatorCore/Source/PropertyAnimatorCore/Private/Properties/PropertyAnimatorCoreData.cpp b/Engine/Plugins/Experimental/PropertyAnimatorCore/Source/PropertyAnimatorCore/Private/Properties/PropertyAnimatorCoreData.cpp
--- a/Engine/Plugins/Experimental/PropertyAnimatorCore/Source/PropertyAnimatorCore/Private/Properties/PropertyAnimatorCoreData.cpp
+++ b/Engine/Plugins/Experimental/PropertyAnimatorCore/Source/PropertyAnimatorCore/Private/Properties/PropertyAnimatorCoreData.cpp
@@ -123,6 +123,28 @@ UActorComponent* FPropertyAnimatorCoreData::GetOwningComponent() const
 	return Owner->GetTypedOuter<UActorComponent>();
 }
 
+TArray<UObject*> FPropertyAnimatorCoreData::GetOuters(const UObject* InStopOuter) const
+{
+	TArray<UObject*> Owners;
+
+	UObject* Outer = OwnerWeak.Get();
+
+	if (!Outer)
+	{
+		return Owners;
+	}
+
+	while (Outer && Outer != InStopOuter)
+	{
+		Owners.Add(Outer);
+		Outer = Outer->GetOuter();
+	}
+
+	Algo::Reverse(Owners);
+
+	return Owners;
+}
+
 FName FPropertyAnimatorCoreData::GetMemberPropertyName() const
 {
 	const FProperty* MemberProperty = GetMemberProperty();
diff --git a/Engine/Plugins/Experimental/PropertyAnimatorCore/Source/PropertyAnimatorCore/Private/Subsystems/PropertyAnimatorCoreSubsystem.cpp b/Engine/Plugins/Experimental/PropertyAnimatorCore/Source/PropertyAnimatorCore/Private/Subsystems/PropertyAnimatorCoreSubsystem.cpp
--- a/Engine/Plugins/Experimental/PropertyAnimatorCore/Source/PropertyAnimatorCore/Private/Subsystems/PropertyAnimatorCoreSubsystem.cpp
+++ b/Engine/Plugins/Experimental/PropertyAnimatorCore/Source/PropertyAnimatorCore/Private/Subsystems/PropertyAnimatorCoreSubsystem.cpp
@@ -374,12 +374,56 @@ TSet<UPropertyAnimatorCoreBase*> UPropertyAnimatorCoreSubsystem::CreateAnimators
 	return NewAnimators;
 }
 
+TSet<UPropertyAnimatorCoreBase*> UPropertyAnimatorCoreSubsystem::CloneAnimators(const TSet<UPropertyAnimatorCoreBase*>& InAnimators, AActor* InTargetActor, bool bInTransact) const
+{
+	TSet<UPropertyAnimatorCoreBase*> CopyAnimators;
+
+	if (!IsValid(InTargetActor))
+	{
+		return CopyAnimators;
+	}
+
+#if WITH_EDITOR
+	const FText TransactionText = LOCTEXT("CloneAnimators", "Cloning {0} animator(s) on actor %s");
+	const FText AnimatorCount = FText::FromString(FString::FromInt(InAnimators.Num()));
+	const FText ActorName = FText::FromString(InTargetActor->GetActorNameOrLabel());
+
+	FScopedTransaction Transaction(FText::Format(TransactionText, AnimatorCount, ActorName), bInTransact);
+#endif
+
+	UPropertyAnimatorCoreComponent* Component = UPropertyAnimatorCoreComponent::FindOrAdd(InTargetActor);
+	if (!Component)
+	{
+		return CopyAnimators;
+	}
+
+#if WITH_EDITOR
+	Component->Modify();
+#endif
+
+	CopyAnimators.Reserve(InAnimators.Num());
+
+	for (UPropertyAnimatorCoreBase* Animator : InAnimators)
+	{
+		if (UPropertyAnimatorCoreBase* CopyAnimator = Component->CloneAnimator(Animator))
+		{
+#if WITH_EDITOR
+			CopyAnimator->Modify();
+#endif
+
+			CopyAnimators.Add(CopyAnimator);
+		}
+	}
+
+	return CopyAnimators;
+}
+
 bool UPropertyAnimatorCoreSubsystem::RemoveAnimator(UPropertyAnimatorCoreBase* InAnimator, bool bInTransact) const
 {
 	return RemoveAnimators({InAnimator}, bInTransact);
 }
 
-bool UPropertyAnimatorCoreSubsystem::RemoveAnimators(const TSet<UPropertyAnimatorCoreBase*> InAnimators, bool bInTransact) const
+bool UPropertyAnimatorCoreSubsystem::RemoveAnimators(const TSet<UPropertyAnimatorCoreBase*>& InAnimators, bool bInTransact) const
 {
 	if (InAnimators.IsEmpty())
 	{
diff --git a/Engine/Plugins/Experimental/PropertyAnimatorCore/Source/PropertyAnimatorCore/Public/Animators/PropertyAnimatorCoreBase.h b/Engine/Plugins/Experimental/PropertyAnimatorCore/Source/PropertyAnimatorCore/Public/Animators/PropertyAnimatorCoreBase.h
--- a/Engine/Plugins/Experimental/PropertyAnimatorCore/Source/PropertyAnimatorCore/Public/Animators/PropertyAnimatorCoreBase.h
+++ b/Engine/Plugins/Experimental/PropertyAnimatorCore/Source/PropertyAnimatorCore/Public/Animators/PropertyAnimatorCoreBase.h
@@ -149,7 +149,9 @@ protected:
 	//~ Begin UObject
 	PROPERTYANIMATORCORE_API virtual void BeginDestroy() override;
 	PROPERTYANIMATORCORE_API virtual void PostLoad() override;
+	PROPERTYANIMATORCORE_API virtual void PostEditImport() override;
 	PROPERTYANIMATORCORE_API virtual void PreDuplicate(FObjectDuplicationParameters& InDupParams) override;
+	PROPERTYANIMATORCORE_API virtual void PostDuplicate(EDuplicateMode::Type InDuplicateMode) override;
 #if WITH_EDITOR
 	PROPERTYANIMATORCORE_API virtual void PreEditUndo() override;
 	PROPERTYANIMATORCORE_API virtual void PostEditUndo() override;
@@ -242,6 +244,8 @@ protected:
 	PROPERTYANIMATORCORE_API virtual void OnAnimatorEnabled();
 	PROPERTYANIMATORCORE_API virtual void OnAnimatorDisabled();
 
+	virtual void OnTimeSourceChanged() {}
+
 	/** Returns the property context class to use */
 	PROPERTYANIMATORCORE_API virtual TSubclassOf<UPropertyAnimatorCoreContext> GetPropertyContextClass(const FPropertyAnimatorCoreData& InProperty);
 
@@ -274,6 +278,9 @@ private:
 
 	void OnTimeSourceNameChanged();
 
+	/** Called after an action that causes the owner to change */
+	void ResolveProperties();
+
 	UPropertyAnimatorCoreTimeSourceBase* FindOrAddTimeSource(FName InTimeSourceName);
 
 	UFUNCTION()
diff --git a/Engine/Plugins/Experimental/PropertyAnimatorCore/Source/PropertyAnimatorCore/Public/Components/PropertyAnimatorCoreComponent.h b/Engine/Plugins/Experimental/PropertyAnimatorCore/Source/PropertyAnimatorCore/Public/Components/PropertyAnimatorCoreComponent.h
--- a/Engine/Plugins/Experimental/PropertyAnimatorCore/Source/PropertyAnimatorCore/Public/Components/PropertyAnimatorCoreComponent.h
+++ b/Engine/Plugins/Experimental/PropertyAnimatorCore/Source/PropertyAnimatorCore/Public/Components/PropertyAnimatorCoreComponent.h
@@ -47,7 +47,7 @@ public:
 	}
 
 	/** Process a function for each controller, stops when false is returned otherwise continue until the end */
-	void ForEachAnimator(TFunctionRef<bool(UPropertyAnimatorCoreBase*)> InFunction) const;
+	PROPERTYANIMATORCORE_API void ForEachAnimator(TFunctionRef<bool(UPropertyAnimatorCoreBase*)> InFunction) const;
 
 protected:
 	static FName GetAnimatorName(const UPropertyAnimatorCoreBase* InAnimator);
@@ -75,6 +75,9 @@ protected:
 	/** Adds a new animator of that class */
 	UPropertyAnimatorCoreBase* AddAnimator(const UClass* InAnimatorClass);
 
+	/** Clones an existing animator */
+	UPropertyAnimatorCoreBase* CloneAnimator(UPropertyAnimatorCoreBase* InAnimator);
+
 	/** Removes an existing animator */
 	bool RemoveAnimator(UPropertyAnimatorCoreBase* InAnimator);
 
diff --git a/Engine/Plugins/Experimental/PropertyAnimatorCore/Source/PropertyAnimatorCore/Public/Properties/PropertyAnimatorCoreContext.h b/Engine/Plugins/Experimental/PropertyAnimatorCore/Source/PropertyAnimatorCore/Public/Properties/PropertyAnimatorCoreContext.h
--- a/Engine/Plugins/Experimental/PropertyAnimatorCore/Source/PropertyAnimatorCore/Public/Properties/PropertyAnimatorCoreContext.h
+++ b/Engine/Plugins/Experimental/PropertyAnimatorCore/Source/PropertyAnimatorCore/Public/Properties/PropertyAnimatorCoreContext.h
@@ -79,6 +79,9 @@ public:
 		return static_cast<InRuleType*>(GetConverterRulePtr(InRuleType::StaticStruct()));
 	}
 
+	/** Called when the owner has changed and we want to update the animated property */
+	bool ResolveProperty();
+
 protected:
 	//~ Begin UObject
 	PROPERTYANIMATORCORE_API virtual void PostLoad() override;
diff --git a/Engine/Plugins/Experimental/PropertyAnimatorCore/Source/PropertyAnimatorCore/Public/Properties/PropertyAnimatorCoreData.h b/Engine/Plugins/Experimental/PropertyAnimatorCore/Source/PropertyAnimatorCore/Public/Properties/PropertyAnimatorCoreData.h
--- a/Engine/Plugins/Experimental/PropertyAnimatorCore/Source/PropertyAnimatorCore/Public/Properties/PropertyAnimatorCoreData.h
+++ b/Engine/Plugins/Experimental/PropertyAnimatorCore/Source/PropertyAnimatorCore/Public/Properties/PropertyAnimatorCoreData.h
@@ -77,6 +77,9 @@ struct FPropertyAnimatorCoreData
 		return OwnerWeak;
 	}
 
+	/** Returns chain of owner until it reaches StopOuter */
+	TArray<UObject*> GetOuters(const UObject* InStopOuter) const;
+
 	/** The member property of the owner, top property inside the owner itself */
 	FProperty* GetMemberProperty() const
 	{
diff --git a/Engine/Plugins/Experimental/PropertyAnimatorCore/Source/PropertyAnimatorCore/Public/Subsystems/PropertyAnimatorCoreSubsystem.h b/Engine/Plugins/Experimental/PropertyAnimatorCore/Source/PropertyAnimatorCore/Public/Subsystems/PropertyAnimatorCoreSubsystem.h
--- a/Engine/Plugins/Experimental/PropertyAnimatorCore/Source/PropertyAnimatorCore/Public/Subsystems/PropertyAnimatorCoreSubsystem.h
+++ b/Engine/Plugins/Experimental/PropertyAnimatorCore/Source/PropertyAnimatorCore/Public/Subsystems/PropertyAnimatorCoreSubsystem.h
@@ -146,11 +146,14 @@ public:
 	/** Create animators of specific class for actors */
 	PROPERTYANIMATORCORE_API TSet<UPropertyAnimatorCoreBase*> CreateAnimators(const TSet<AActor*>& InActors, const UClass* InAnimatorClass, UPropertyAnimatorCorePresetBase* InPreset = nullptr, bool bInTransact = false) const;
 
+	/** Clone animators onto an actor */
+	PROPERTYANIMATORCORE_API TSet<UPropertyAnimatorCoreBase*> CloneAnimators(const TSet<UPropertyAnimatorCoreBase*>& InAnimators, AActor* InTargetActor, bool bInTransact = false) const;
+
 	/** Removes a animator bound to an owner */
 	PROPERTYANIMATORCORE_API bool RemoveAnimator(UPropertyAnimatorCoreBase* InAnimator, bool bInTransact = false) const;
 
 	/** Removes animators from their owner */
-	PROPERTYANIMATORCORE_API bool RemoveAnimators(const TSet<UPropertyAnimatorCoreBase*> InAnimators, bool bInTransact = false) const;
+	PROPERTYANIMATORCORE_API bool RemoveAnimators(const TSet<UPropertyAnimatorCoreBase*>& InAnimators, bool bInTransact = false) const;
 
 	/** Apply a preset on an existing animator */
 	PROPERTYANIMATORCORE_API bool ApplyAnimatorPreset(UPropertyAnimatorCoreBase* InAnimator, UPropertyAnimatorCorePresetBase* InPreset, bool bInTransact = false);
