Motion Design Outliner: Fix issue with saving a new level created from template resetting the template's outliner settings

diff --git a/Engine/Plugins/Experimental/Avalanche/Source/AvalancheOutliner/Private/Data/AvaOutlinerSaveState.cpp b/Engine/Plugins/Experimental/Avalanche/Source/AvalancheOutliner/Private/Data/AvaOutlinerSaveState.cpp
--- a/Engine/Plugins/Experimental/Avalanche/Source/AvalancheOutliner/Private/Data/AvaOutlinerSaveState.cpp
+++ b/Engine/Plugins/Experimental/Avalanche/Source/AvalancheOutliner/Private/Data/AvaOutlinerSaveState.cpp
@@ -74,7 +74,7 @@ void FAvaOutlinerSaveState::Serialize(FAvaOutliner& InOutliner, FArchive& Ar)
 			}
 		}
 
-		ContextPath = FSoftObjectPath(OutlinerWorld).ToString();
+		UpdateItemIdContexts(FSoftObjectPath(OutlinerWorld).ToString());
 	}
 
 	const int32 OutlinerVersion = Ar.CustomVer(FAvaOutlinerVersion::GUID);
@@ -95,7 +95,7 @@ void FAvaOutlinerSaveState::Serialize(FAvaOutliner& InOutliner, FArchive& Ar)
 
 	if (Ar.IsLoading())
 	{
-		UpdateItemIdContexts(ContextPath, FSoftObjectPath(OutlinerWorld).ToString());
+		UpdateItemIdContexts(FSoftObjectPath(OutlinerWorld).ToString());
 
 		// Use Scene Tree if Supported by the Outliner Version to Load the Item Ordering
 		FAvaSceneTree* const SceneTree = (OutlinerVersion >= FAvaOutlinerVersion::SceneTree)
@@ -340,21 +340,26 @@ void FAvaOutlinerSaveState::EnsureOutlinerViewCount(int32 InOutlinerViewId)
 	}
 }
 
-void FAvaOutlinerSaveState::UpdateItemIdContexts(FStringView InOldContext, FStringView InNewContext)
+void FAvaOutlinerSaveState::UpdateItemIdContexts(FStringView InNewContext)
 {
 	// Already updated
-	if (InOldContext == InNewContext)
+	if (ContextPath == InNewContext)
 	{
 		return;
 	}
 
-	auto FixItemIdMap = [InOldContext, InNewContext]<typename InValueType>(TMap<FString, InValueType>& InItemIdMap)
+	ON_SCOPE_EXIT
+	{
+		ContextPath = InNewContext;
+	};
+
+	auto FixItemIdMap = [&ContextPath = ContextPath, InNewContext]<typename InValueType>(TMap<FString, InValueType>& InItemIdMap)
 		{
 			TMap<FString, InValueType> ItemIdMapTemp = InItemIdMap;
 			for (TPair<FString, InValueType>& Pair : ItemIdMapTemp)
 			{
 				FSoftObjectPath ObjectPath;
-				ObjectPath.SetPath(InOldContext);
+				ObjectPath.SetPath(ContextPath);
 
 				FString AssetPath = ObjectPath.GetAssetPath().ToString();
 				if (AssetPath.IsEmpty() || Pair.Key.StartsWith(AssetPath))
diff --git a/Engine/Plugins/Experimental/Avalanche/Source/AvalancheOutliner/Private/Data/AvaOutlinerSaveState.h b/Engine/Plugins/Experimental/Avalanche/Source/AvalancheOutliner/Private/Data/AvaOutlinerSaveState.h
--- a/Engine/Plugins/Experimental/Avalanche/Source/AvalancheOutliner/Private/Data/AvaOutlinerSaveState.h
+++ b/Engine/Plugins/Experimental/Avalanche/Source/AvalancheOutliner/Private/Data/AvaOutlinerSaveState.h
@@ -69,7 +69,7 @@ private:
 
 	void EnsureOutlinerViewCount(int32 InOutlinerViewId);
 
-	void UpdateItemIdContexts(FStringView InOldContext, FStringView InNewContext);
+	void UpdateItemIdContexts(FStringView InNewContext);
 
 	/**
 	 * Deprecated in favor of using Motion Design Scene Tree
