Motion Design
- Made the tool actor factories more robust in their GC approach

diff --git a/Engine/Plugins/Experimental/Avalanche/Source/AvalancheInteractiveTools/Private/AvalancheInteractiveToolsModule.cpp b/Engine/Plugins/Experimental/Avalanche/Source/AvalancheInteractiveTools/Private/AvalancheInteractiveToolsModule.cpp
--- a/Engine/Plugins/Experimental/Avalanche/Source/AvalancheInteractiveTools/Private/AvalancheInteractiveToolsModule.cpp
+++ b/Engine/Plugins/Experimental/Avalanche/Source/AvalancheInteractiveTools/Private/AvalancheInteractiveToolsModule.cpp
@@ -132,6 +132,26 @@ bool FAvalancheInteractiveToolsModule::HasActiveTool() const
 	return bHasActiveTool;
 }
 
+void FAvalancheInteractiveToolsModule::AddReferencedObjects(FReferenceCollector& InCollector)
+{
+	for (TPair<FName, TArray<FAvaInteractiveToolsToolParameters>>& CategoryPair : Tools)
+	{
+		for (FAvaInteractiveToolsToolParameters& ToolParams : CategoryPair.Value)
+		{
+			if (ToolParams.Factory)
+			{
+				InCollector.AddReferencedObject<UActorFactory>(ToolParams.Factory);
+			}
+		}
+	}
+}
+
+FString FAvalancheInteractiveToolsModule::GetReferencerName() const
+{
+	static const FString ReferencerName = "AvaITFModule";
+	return ReferencerName;
+}
+
 void FAvalancheInteractiveToolsModule::OnToolActivated()
 {
 	bHasActiveTool = true;
@@ -239,7 +259,7 @@ void FAvalancheInteractiveToolsModule::OnPlacementCategoryRefreshed(FName InCate
 		if (Tool.Factory)
 		{
 			PlaceableItem = MakeShared<FPlaceableItem>(
-				Tool.Factory,
+				Tool.Factory.Get(),
 				FAssetData(Tool.Factory->NewActorClass->GetDefaultObject()),
 				Tool.Priority
 			);
diff --git a/Engine/Plugins/Experimental/Avalanche/Source/AvalancheInteractiveTools/Private/AvalancheInteractiveToolsModule.h b/Engine/Plugins/Experimental/Avalanche/Source/AvalancheInteractiveTools/Private/AvalancheInteractiveToolsModule.h
--- a/Engine/Plugins/Experimental/Avalanche/Source/AvalancheInteractiveTools/Private/AvalancheInteractiveToolsModule.h
+++ b/Engine/Plugins/Experimental/Avalanche/Source/AvalancheInteractiveTools/Private/AvalancheInteractiveToolsModule.h
@@ -8,7 +8,7 @@ class FUICommandInfo;
 
 DECLARE_LOG_CATEGORY_EXTERN(LogAvaInteractiveTools, Log, All);
 
-class FAvalancheInteractiveToolsModule : public IAvalancheInteractiveToolsModule
+class FAvalancheInteractiveToolsModule : public IAvalancheInteractiveToolsModule, public FGCObject
 {
 public:
 	static FAvalancheInteractiveToolsModule& Get()
@@ -37,6 +37,11 @@ public:
 	virtual bool HasActiveTool() const override;
 	//~ End IAvalancheInteractiveToolsModule
 
+	// ~Begin FGCObject Interface
+	virtual void AddReferencedObjects(FReferenceCollector& InCollector) override;
+	virtual FString GetReferencerName() const override;
+	// ~End FGCObject Interface
+
 	void OnToolActivated();
 	void OnToolDeactivated();
 
diff --git a/Engine/Plugins/Experimental/Avalanche/Source/AvalancheInteractiveTools/Public/IAvalancheInteractiveToolsModule.h b/Engine/Plugins/Experimental/Avalanche/Source/AvalancheInteractiveTools/Public/IAvalancheInteractiveToolsModule.h
--- a/Engine/Plugins/Experimental/Avalanche/Source/AvalancheInteractiveTools/Public/IAvalancheInteractiveToolsModule.h
+++ b/Engine/Plugins/Experimental/Avalanche/Source/AvalancheInteractiveTools/Public/IAvalancheInteractiveToolsModule.h
@@ -27,7 +27,7 @@ struct FAvaInteractiveToolsToolParameters
 	int32 Priority = 0;
 	FAvalancheInteractiveToolsCreateBuilder CreateBuilder;
 	TSubclassOf<AActor> FactoryClass = nullptr;
-	UActorFactory* Factory = nullptr;
+	TObjectPtr<UActorFactory> Factory = nullptr;
 	TMap<FName, TSharedRef<FAvaInteractiveToolsToolPresetBase>> Presets;
 };
 
