Motion Design
- Changed buttons colours for a few post process states.

diff --git a/Engine/Plugins/Experimental/Avalanche/Source/AvalancheLevelViewport/Private/Widgets/SAvaLevelViewportStatusBarButtons.cpp b/Engine/Plugins/Experimental/Avalanche/Source/AvalancheLevelViewport/Private/Widgets/SAvaLevelViewportStatusBarButtons.cpp
--- a/Engine/Plugins/Experimental/Avalanche/Source/AvalancheLevelViewport/Private/Widgets/SAvaLevelViewportStatusBarButtons.cpp
+++ b/Engine/Plugins/Experimental/Avalanche/Source/AvalancheLevelViewport/Private/Widgets/SAvaLevelViewportStatusBarButtons.cpp
@@ -46,29 +46,6 @@ namespace UE::Ava::LevelViewportStatusBarButtons::Private
 	static const FSlateIcon AlphaChannelIcon = FSlateIcon(AvaLevelViewportStyleName, "Icons.PostProcess.Alpha");
 	static const FSlateIcon CheckerboardIcon = FSlateIcon(AppStyleSetName, "Checker");
 
-	static bool IsViewportPostProcessManagerEnabled(const TWeakPtr<SAvaLevelViewportFrame>& InViewportFrameWeak)
-	{
-		const FAvaLevelViewportGuideFrameAndClient FrameAndClient(InViewportFrameWeak);
-
-		return FrameAndClient.IsValid() && FrameAndClient.ViewportClient->GetPostProcessManager().IsValid();
-	}
-
-	static FSlateColor GetPostProcessColor(const TWeakPtr<SAvaLevelViewportFrame>& InViewportFrameWeak, EAvaViewportPostProcessType InPostProcessType,
-		const FSlateColor& InActiveColor, const FSlateColor& InEnabledColor, const FSlateColor& InDisabledColor)
-	{
-		const FAvaLevelViewportGuideFrameAndClient FrameAndClient(InViewportFrameWeak);
-
-		if (FrameAndClient.IsValid())
-		{
-			if (const TSharedPtr<FAvaViewportPostProcessManager> PostProcessManager = FrameAndClient.ViewportClient->GetPostProcessManager())
-			{
-				return PostProcessManager->GetType() == InPostProcessType ? InActiveColor : InEnabledColor;
-			}
-		}
-
-		return InDisabledColor;
-	}
-
 	static void TogglePostProcess(const TWeakPtr<SAvaLevelViewportFrame>& InViewportFrameWeak, EAvaViewportPostProcessType InPostProcessType)
 	{
 		const FAvaLevelViewportGuideFrameAndClient FrameAndClient(InViewportFrameWeak);
@@ -530,7 +507,16 @@ FSlateColor SAvaLevelViewportStatusBarButtons::GetPostProcessColor() const
 	{
 		if (const TSharedPtr<FAvaViewportPostProcessManager> PostProcessManager = FrameAndClient.ViewportClient->GetPostProcessManager())
 		{
-			return ViewportStatusBarButton::EnabledColor;
+			switch (PostProcessManager->GetType())
+			{
+				// Alpha channel could be included here, but I feel it shouldn't.
+				case EAvaViewportPostProcessType::Background:
+				case EAvaViewportPostProcessType::Checkerboard:
+					return ViewportStatusBarButton::ActiveColor;
+
+				default:
+					return ViewportStatusBarButton::EnabledColor;
+			}
 		}
 	}
 
