Motion Design
- Changed zoom level in viewport info to show the inverse of what it was previously.
- Changed zoom level in the viewport to show more "whole" zoom increments based on the default fov.

diff --git a/Engine/Plugins/Experimental/Avalanche/Source/AvalancheLevelViewport/Private/Widgets/SAvaViewportInfo.cpp b/Engine/Plugins/Experimental/Avalanche/Source/AvalancheLevelViewport/Private/Widgets/SAvaViewportInfo.cpp
--- a/Engine/Plugins/Experimental/Avalanche/Source/AvalancheLevelViewport/Private/Widgets/SAvaViewportInfo.cpp
+++ b/Engine/Plugins/Experimental/Avalanche/Source/AvalancheLevelViewport/Private/Widgets/SAvaViewportInfo.cpp
@@ -391,7 +391,7 @@ FText SAvaViewportInfo::GetZoomLevel() const
 
 	if (VisibleArea.IsValid())
 	{
-		const int32 Percentage = FMath::RoundToInt(VisibleArea.GetVisibleAreaFraction() * 100.f);
+		const int32 Percentage = FMath::RoundToInt(100.f / VisibleArea.GetVisibleAreaFraction());
 
 		return FText::Format(
 			LOCTEXT("PercentFormat", "{0}%"),
diff --git a/Engine/Plugins/Experimental/Avalanche/Source/AvalancheViewport/Private/Interaction/AvaCameraZoomController.cpp b/Engine/Plugins/Experimental/Avalanche/Source/AvalancheViewport/Private/Interaction/AvaCameraZoomController.cpp
--- a/Engine/Plugins/Experimental/Avalanche/Source/AvalancheViewport/Private/Interaction/AvaCameraZoomController.cpp
+++ b/Engine/Plugins/Experimental/Avalanche/Source/AvalancheViewport/Private/Interaction/AvaCameraZoomController.cpp
@@ -9,10 +9,54 @@
 #include "Framework/Application/SlateApplication.h"
 #include "ViewportClient/IAvaViewportClient.h"
 
-namespace UE::AvaLevelViewport::Private
-{
-	constexpr uint8 ZOOM_LEVEL_MIN = 0;
-	constexpr uint8 ZOOM_LEVEL_MAX = 9;
+namespace UE::AvaViewport::Private
+{
+	const TArray<float> ZoomLevels = {
+		1.0f,
+		1.25f,
+		1.5f,
+		2.f,
+		2.5f,
+		3.f,
+		4.f,
+		5.f,
+		7.5f,
+		10.f,
+		12.5f,
+		15.f,
+		17.5f,
+		20.f
+	};
+
+	/** Default zoom level will change. It's 0 for now. */
+	constexpr uint8 DefaultZoomLevel = 0;
+
+	/** Returns the FOV angle at the given zoom level. Degrees. */
+	float GetFOVForZoomLevel(float InDefaultFOV, uint8 InZoomLevel)
+	{
+		return 2.f * FMath::RadiansToDegrees(FMath::Atan(
+			FMath::Tan(FMath::DegreesToRadians(InDefaultFOV * 0.5f)) / ZoomLevels[InZoomLevel]
+		));
+	}
+
+	/**
+	 * Returns the highest zoom level (lowest FOV) that is >= the given FOV.
+	 * Will return ZoomLevel 0 if there isn't a Zoom FOV high enough.
+	 */
+	uint8 GetZoomLevelForFOV(float InDefaultFOV, float InFOV)
+	{
+		for (uint8 ZoomLevel = 1; ZoomLevel < ZoomLevels.Num(); ++ZoomLevel)
+		{
+			const float ZoomFOV = GetFOVForZoomLevel(InDefaultFOV, ZoomLevel);
+
+			if (ZoomFOV < InFOV)
+			{
+				return ZoomLevel - 1;
+			}
+		}
+
+		return ZoomLevels.Num() - 1;
+	}
 }
 
 bool FAvaCameraZoomController::IsCameraZoomPossible()
@@ -32,7 +76,7 @@ FAvaCameraZoomController::FAvaCameraZoomController(TSharedRef<IAvaViewportClient
 
 void FAvaCameraZoomController::SetZoomLevel(uint8 InZoomLevel)
 {
-	ZoomLevel = FMath::Clamp(InZoomLevel, UE::AvaLevelViewport::Private::ZOOM_LEVEL_MIN, UE::AvaLevelViewport::Private::ZOOM_LEVEL_MAX);
+	ZoomLevel = FMath::Clamp(InZoomLevel, 0, UE::AvaViewport::Private::ZoomLevels.Num() - 1);
 
 	UpdateVisibleAreas();
 	InvalidateViewport();
@@ -40,12 +84,7 @@ void FAvaCameraZoomController::SetZoomLevel(uint8 InZoomLevel)
 
 bool FAvaCameraZoomController::IsZoomed() const
 {
-	return ZoomLevel > UE::AvaLevelViewport::Private::ZOOM_LEVEL_MIN;
-}
-
-float FAvaCameraZoomController::GetFOVPerStep() const
-{
-	return GetDefaultFOV() / static_cast<float>(UE::AvaLevelViewport::Private::ZOOM_LEVEL_MAX - UE::AvaLevelViewport::Private::ZOOM_LEVEL_MIN + 1);
+	return ZoomLevel != UE::AvaViewport::Private::DefaultZoomLevel;
 }
 
 void FAvaCameraZoomController::ZoomIn()
@@ -187,7 +226,7 @@ void FAvaCameraZoomController::FrameActor()
 
 void FAvaCameraZoomController::Reset()
 {
-	ZoomLevel = 0;
+	ZoomLevel = UE::AvaViewport::Private::DefaultZoomLevel;
 	PanOffsetFraction = FVector2f::ZeroVector;
 
 	InvalidateViewport();
@@ -238,7 +277,7 @@ void FAvaCameraZoomController::CenterOnPoint(const FVector2f& InPoint)
 
 void FAvaCameraZoomController::CenterOnBox(const FBox& InBoundingBox, const FTransform& InBoxTransform)
 {
-	// TODO @Update
+	// @TODO Update
 	return;
 
 	TSharedPtr<IAvaViewportClient> AvaViewportClient = AvaViewportClientWeak.Pin();
@@ -295,13 +334,10 @@ void FAvaCameraZoomController::CenterOnBox(const FBox& InBoundingBox, const FTra
 	// Convert vertical fov to horizontal fov
 	RequiredVerticalFOV = FMath::RadiansToDegrees(2.f * FMath::Atan(FMath::Tan(FMath::DegreesToRadians(RequiredVerticalFOV) / 2.f) * AspectRatio));
 
-	const float FOVAdjustment = GetFOVPerStep();
-	RequiredHorizontalFOV = FMath::Max(RequiredHorizontalFOV, RequiredVerticalFOV);
-	RequiredHorizontalFOV = FMath::Floor(RequiredHorizontalFOV / FOVAdjustment) * FOVAdjustment; // Round to nearest 10
-	RequiredHorizontalFOV = FMath::Clamp(RequiredHorizontalFOV, FOVAdjustment, DefaultFOV);
-
-	const float FOVStep = GetFOVPerStep();
-	const uint8 RequiredZoomLevel = static_cast<uint8>(FMath::Clamp((DefaultFOV - RequiredHorizontalFOV) / FOVStep, 0.f, 255.f));
+	const uint8 RequiredZoomLevel = UE::AvaViewport::Private::GetZoomLevelForFOV(
+		DefaultFOV,
+		FMath::Max(RequiredHorizontalFOV, RequiredVerticalFOV)
+	);
 
 	if (RequiredZoomLevel == 0)
 	{
@@ -327,7 +363,7 @@ void FAvaCameraZoomController::EndPanning()
 
 void FAvaCameraZoomController::ZoomIn_Internal()
 {
-	if (ZoomLevel >= UE::AvaLevelViewport::Private::ZOOM_LEVEL_MAX)
+	if (ZoomLevel >= (UE::AvaViewport::Private::ZoomLevels.Num() - 1))
 	{
 		return;
 	}
@@ -337,7 +373,7 @@ void FAvaCameraZoomController::ZoomIn_Internal()
 
 void FAvaCameraZoomController::ZoomOut_Internal()
 {
-	if (ZoomLevel <= UE::AvaLevelViewport::Private::ZOOM_LEVEL_MIN)
+	if (ZoomLevel <= 0)
 	{
 		return;
 	}
@@ -391,8 +427,12 @@ float FAvaCameraZoomController::GetFOV() const
 		return DefaultFOV;
 	}
 
-	const float FOVStep = GetFOVPerStep();
-	return FMath::Max(DefaultFOV - static_cast<float>(ZoomLevel) * FOVStep, FOVStep);
+	if (ZoomLevel == 0)
+	{
+		return DefaultFOV;
+	}
+
+	return UE::AvaViewport::Private::GetFOVForZoomLevel(DefaultFOV, ZoomLevel);
 }
 
 FVector2f FAvaCameraZoomController::GetCameraProjectionOffset() const
diff --git a/Engine/Plugins/Experimental/Avalanche/Source/AvalancheViewport/Public/Interaction/AvaCameraZoomController.h b/Engine/Plugins/Experimental/Avalanche/Source/AvalancheViewport/Public/Interaction/AvaCameraZoomController.h
--- a/Engine/Plugins/Experimental/Avalanche/Source/AvalancheViewport/Public/Interaction/AvaCameraZoomController.h
+++ b/Engine/Plugins/Experimental/Avalanche/Source/AvalancheViewport/Public/Interaction/AvaCameraZoomController.h
@@ -29,8 +29,6 @@ public:
 	void SetZoomLevel(uint8 InZoomLevel);
 	bool IsZoomed() const;
 
-	float GetFOVPerStep() const;
-
 	AVALANCHEVIEWPORT_API void ZoomIn();
 	AVALANCHEVIEWPORT_API void ZoomInCursor();
 
