Motion Design
- Changed the icon on teh status bar for the post process selector.
- Moved the post process selector to the left a bit.
- Fixed a bug where the visualizer viewport overlay button was not showing as disabled when the viewport overlay was disabled.
- Moved grid button into the group with the rest of the viewport overlay buttons.

diff --git a/Engine/Plugins/Experimental/Avalanche/Source/AvalancheLevelViewport/Private/Widgets/SAvaLevelViewportStatusBarButtons.cpp b/Engine/Plugins/Experimental/Avalanche/Source/AvalancheLevelViewport/Private/Widgets/SAvaLevelViewportStatusBarButtons.cpp
--- a/Engine/Plugins/Experimental/Avalanche/Source/AvalancheLevelViewport/Private/Widgets/SAvaLevelViewportStatusBarButtons.cpp
+++ b/Engine/Plugins/Experimental/Avalanche/Source/AvalancheLevelViewport/Private/Widgets/SAvaLevelViewportStatusBarButtons.cpp
@@ -323,6 +323,22 @@ void SAvaLevelViewportStatusBarButtons::PopulateViewportButtons(TSharedPtr<SHori
 			)
 		];
 
+	TSharedRef<SComboButton> PostProcessButton = ViewportStatusBarButton::MakeMenuButton(
+		LOCTEXT("PostProcessEffects", "Post Process Effects"),
+		FOnGetContent::CreateSP(this, &SAvaLevelViewportStatusBarButtons::GetPostProcessMenuContent),
+		FSlateIcon(FAppStyle::GetAppStyleSetName(), "Icons.Advanced").GetIcon(),
+		TAttribute<FSlateColor>::CreateSP(this, &SAvaLevelViewportStatusBarButtons::GetPostProcessColor)
+	);
+
+	PostProcessButton->SetEnabled(TAttribute<bool>::CreateSP(this, &SAvaLevelViewportStatusBarButtons::GetPostProcessEnabled));
+
+	InContainer->AddSlot()
+		.AutoWidth()
+		.Padding(ViewportStatusBarButton::Padding)
+		[
+			PostProcessButton
+		];
+
 	InContainer->AddSlot()
 		.AutoWidth()
 		.Padding(ViewportStatusBarButton::Padding)
@@ -393,37 +409,41 @@ void SAvaLevelViewportStatusBarButtons::PopulateViewportButtons(TSharedPtr<SHori
 			)
 		];
 
-	TSharedRef<SComboButton> PostProcessButton = ViewportStatusBarButton::MakeMenuButton(
-		LOCTEXT("PostProcessEffects", "Post Process Effects"),
-		FOnGetContent::CreateSP(this, &SAvaLevelViewportStatusBarButtons::GetPostProcessMenuContent),
-		TAttribute<const FSlateBrush*>::CreateSP(this, &SAvaLevelViewportStatusBarButtons::GetPostProcessIcon),
-		TAttribute<FSlateColor>::CreateSP(this, &SAvaLevelViewportStatusBarButtons::GetPostProcessColor)
+	TSharedRef<SAvaMultiComboButton> TextureOverlayButton = ViewportStatusBarButton::MakeMultiMenuButton(
+		CommandsRef.ToggleTextureOverlay->GetDescription(),
+		FOnGetContent::CreateSP(this, &SAvaLevelViewportStatusBarButtons::GetTextureOverlayMenuContent),
+		FAppStyle::Get().GetBrush(TEXT("GenericCommands.Paste")),
+		TAttribute<FSlateColor>::CreateSP(this, &SAvaLevelViewportStatusBarButtons::GetTextureOverlayColor),
+		FOnClicked::CreateSP(this, &SAvaLevelViewportStatusBarButtons::ToggleTextureOverlay)
 	);
 
-	PostProcessButton->SetEnabled(TAttribute<bool>::CreateSP(this, &SAvaLevelViewportStatusBarButtons::GetPostProcessEnabled));
+	TextureOverlayButton->SetEnabled(TAttribute<bool>::CreateSP(this, &SAvaLevelViewportStatusBarButtons::GetTextureOverlayEnabled));
 
 	InContainer->AddSlot()
 		.AutoWidth()
 		.Padding(ViewportStatusBarButton::Padding)
 		[
-			PostProcessButton
+			TextureOverlayButton
 		];
 
-	TSharedRef<SAvaMultiComboButton> TextureOverlayButton = ViewportStatusBarButton::MakeMultiMenuButton(
-		CommandsRef.ToggleTextureOverlay->GetDescription(),
-		FOnGetContent::CreateSP(this, &SAvaLevelViewportStatusBarButtons::GetTextureOverlayMenuContent),
-		FAppStyle::Get().GetBrush(TEXT("GenericCommands.Paste")),
-		TAttribute<FSlateColor>::CreateSP(this, &SAvaLevelViewportStatusBarButtons::GetTextureOverlayColor),
-		FOnClicked::CreateSP(this, &SAvaLevelViewportStatusBarButtons::ToggleTextureOverlay)
+	TSharedRef<SAvaMultiComboButton> GridButton = ViewportStatusBarButton::MakeMultiMenuButton(
+		CommandsRef.ToggleGrid->GetDescription(),
+		FOnGetContent::CreateSP(this, &SAvaLevelViewportStatusBarButtons::GetGridMenuContent),
+		FAvaLevelViewportStyle::Get().GetBrush(TEXT("Button.ToggleGrid")),
+		TAttribute<FSlateColor>::CreateSP(this, &SAvaLevelViewportStatusBarButtons::GetToggleGridColor),
+		FOnClicked::CreateSP(this, &SAvaLevelViewportStatusBarButtons::ToggleGrid)
 	);
 
-	TextureOverlayButton->SetEnabled(TAttribute<bool>::CreateSP(this, &SAvaLevelViewportStatusBarButtons::GetPostProcessEnabled));
+	GridButton->SetEnabled(TAttribute<bool>::CreateSP(
+		this,
+		&SAvaLevelViewportStatusBarButtons::GetToggleGridEnabled
+	));
 
 	InContainer->AddSlot()
 		.AutoWidth()
 		.Padding(ViewportStatusBarButton::Padding)
 		[
-			TextureOverlayButton
+			GridButton
 		];
 
 	InContainer->AddSlot()
@@ -460,26 +480,6 @@ void SAvaLevelViewportStatusBarButtons::PopulateViewportButtons(TSharedPtr<SHori
 			SnapButton
 		];
 
-	TSharedRef<SAvaMultiComboButton> GridButton = ViewportStatusBarButton::MakeMultiMenuButton(
-		CommandsRef.ToggleGrid->GetDescription(),
-		FOnGetContent::CreateSP(this, &SAvaLevelViewportStatusBarButtons::GetGridMenuContent),
-		FAvaLevelViewportStyle::Get().GetBrush(TEXT("Button.ToggleGrid")),
-		TAttribute<FSlateColor>::CreateSP(this, &SAvaLevelViewportStatusBarButtons::GetToggleGridColor),
-		FOnClicked::CreateSP(this, &SAvaLevelViewportStatusBarButtons::ToggleGrid)
-	);
-
-	GridButton->SetEnabled(TAttribute<bool>::CreateSP(
-		this,
-		&SAvaLevelViewportStatusBarButtons::GetToggleGridEnabled
-	));
-
-	InContainer->AddSlot()
-		.AutoWidth()
-		.Padding(ViewportStatusBarButton::Padding)
-		[
-			GridButton
-		];
-
 	TSharedRef<SComboButton> ViewportInfoButton = ViewportStatusBarButton::MakeMenuButton(
 		LOCTEXT("ViewportInfomation", "Viewport Information"),
 		FOnGetContent::CreateSP(this, &SAvaLevelViewportStatusBarButtons::GetViewportInfoWidget),
@@ -507,69 +507,15 @@ FSlateColor SAvaLevelViewportStatusBarButtons::GetPostProcessColor() const
 	{
 		if (const TSharedPtr<FAvaViewportPostProcessManager> PostProcessManager = FrameAndClient.ViewportClient->GetPostProcessManager())
 		{
-			switch (PostProcessManager->GetType())
-			{
-				// Alpha channel could be included here, but I feel it shouldn't.
-				case EAvaViewportPostProcessType::Background:
-				case EAvaViewportPostProcessType::Checkerboard:
-					return ViewportStatusBarButton::ActiveColor;
-
-				default:
-					return ViewportStatusBarButton::EnabledColor;
-			}
+			return PostProcessManager->GetType() != EAvaViewportPostProcessType::None
+				? ViewportStatusBarButton::ActiveColor
+				: ViewportStatusBarButton::EnabledColor;
 		}
 	}
 
 	return ViewportStatusBarButton::DisabledColor;
 }
 
-const FSlateBrush* SAvaLevelViewportStatusBarButtons::GetPostProcessIcon() const
-{
-	using namespace UE::Ava::LevelViewportStatusBarButtons::Private;
-
-	static const FSlateBrush* RGBBrush = RGBChannelIcon.GetIcon();
-	static const FSlateBrush* BackgroundBrush = BackgroundIcon.GetIcon();
-	static const FSlateBrush* RedChannelBrush = RedChannelIcon.GetIcon();
-	static const FSlateBrush* GreenChannelBrush = GreenChannelIcon.GetIcon();
-	static const FSlateBrush* BlueChannelBrush = BlueChannelIcon.GetIcon();
-	static const FSlateBrush* AlphaChannelBrush = AlphaChannelIcon.GetIcon();
-	static const FSlateBrush* CheckerBrush = CheckerboardIcon.GetIcon();
-
-	const FAvaLevelViewportGuideFrameAndClient FrameAndClient(ViewportFrameWeak);
-
-	if (FrameAndClient.IsValid())
-	{
-		if (const TSharedPtr<FAvaViewportPostProcessManager> PostProcessManager = FrameAndClient.ViewportClient->GetPostProcessManager())
-		{
-			switch (PostProcessManager->GetType())
-			{
-				case EAvaViewportPostProcessType::None:
-					return RGBBrush;
-
-				case EAvaViewportPostProcessType::Background:
-					return BackgroundBrush;
-
-				case EAvaViewportPostProcessType::RedChannel:
-					return RedChannelBrush;
-
-				case EAvaViewportPostProcessType::GreenChannel:
-					return GreenChannelBrush;
-
-				case EAvaViewportPostProcessType::BlueChannel:
-					return BlueChannelBrush;
-
-				case EAvaViewportPostProcessType::AlphaChannel:
-					return AlphaChannelBrush;
-
-				case EAvaViewportPostProcessType::Checkerboard:
-					return CheckerBrush;
-			}
-		}
-	}
-
-	return RGBBrush;
-}
-
 bool SAvaLevelViewportStatusBarButtons::GetPostProcessEnabled() const
 {
 	const FAvaLevelViewportGuideFrameAndClient FrameAndClient(ViewportFrameWeak);
@@ -974,7 +920,10 @@ FSlateColor SAvaLevelViewportStatusBarButtons::GetToggleShapeEditorOverlayColor(
 
 	if (const UAvaViewportSettings* AvaViewportSettings = GetDefault<UAvaViewportSettings>())
 	{
-		return AvaViewportSettings->bEnableShapesEditorOverlay ? ViewportStatusBarButton::ActiveColor : ViewportStatusBarButton::EnabledColor;
+		if (AvaViewportSettings->bEnableViewportOverlay)
+		{
+			return AvaViewportSettings->bEnableShapesEditorOverlay ? ViewportStatusBarButton::ActiveColor : ViewportStatusBarButton::EnabledColor;
+		}
 	}
 
 	return ViewportStatusBarButton::DisabledColor;
diff --git a/Engine/Plugins/Experimental/Avalanche/Source/AvalancheLevelViewport/Private/Widgets/SAvaLevelViewportStatusBarButtons.h b/Engine/Plugins/Experimental/Avalanche/Source/AvalancheLevelViewport/Private/Widgets/SAvaLevelViewportStatusBarButtons.h
--- a/Engine/Plugins/Experimental/Avalanche/Source/AvalancheLevelViewport/Private/Widgets/SAvaLevelViewportStatusBarButtons.h
+++ b/Engine/Plugins/Experimental/Avalanche/Source/AvalancheLevelViewport/Private/Widgets/SAvaLevelViewportStatusBarButtons.h
@@ -50,7 +50,6 @@ protected:
 
 	// Post Process
 	FSlateColor GetPostProcessColor() const;
-	const FSlateBrush* GetPostProcessIcon() const;
 	bool GetPostProcessEnabled() const;
 	TSharedRef<SWidget> GetPostProcessMenuContent();
 
