Motion Design
- Added the ability to zoom out.

diff --git a/Engine/Plugins/Experimental/Avalanche/Source/AvalancheViewport/Private/Interaction/AvaCameraZoomController.cpp b/Engine/Plugins/Experimental/Avalanche/Source/AvalancheViewport/Private/Interaction/AvaCameraZoomController.cpp
--- a/Engine/Plugins/Experimental/Avalanche/Source/AvalancheViewport/Private/Interaction/AvaCameraZoomController.cpp
+++ b/Engine/Plugins/Experimental/Avalanche/Source/AvalancheViewport/Private/Interaction/AvaCameraZoomController.cpp
@@ -12,6 +12,13 @@
 namespace UE::AvaViewport::Private
 {
 	const TArray<float> ZoomLevels = {
+		0.65f,
+		0.7f,
+		0.75f,
+		0.8f,
+		0.85f,
+		0.9f,
+		0.95f,
 		1.0f,
 		1.25f,
 		1.5f,
@@ -28,8 +35,7 @@ namespace UE::AvaViewport::Private
 		20.f
 	};
 
-	/** Default zoom level will change. It's 0 for now. */
-	constexpr uint8 DefaultZoomLevel = 0;
+	constexpr uint8 DefaultZoomLevel = 7;
 
 	/** Returns the FOV angle at the given zoom level. Degrees. */
 	float GetFOVForZoomLevel(float InDefaultFOV, uint8 InZoomLevel)
@@ -339,14 +345,14 @@ void FAvaCameraZoomController::CenterOnBox(const FBox& InBoundingBox, const FTra
 		FMath::Max(RequiredHorizontalFOV, RequiredVerticalFOV)
 	);
 
-	if (RequiredZoomLevel == 0)
+	if (RequiredZoomLevel == UE::AvaViewport::Private::DefaultZoomLevel)
 	{
 		Reset();
 		return;
 	}
 
 	// Fake the zoom level being 1 zoom level out and then zoom in.
-	SetZoomLevel(RequiredZoomLevel - 1);
+	SetZoomLevel(FMath::Max<uint8>(RequiredZoomLevel, 1) - 1);
 	ZoomIn();
 	CenterOnPoint(ScreenBoundsCenter);
 }
@@ -427,7 +433,7 @@ float FAvaCameraZoomController::GetFOV() const
 		return DefaultFOV;
 	}
 
-	if (ZoomLevel == 0)
+	if (ZoomLevel == UE::AvaViewport::Private::DefaultZoomLevel)
 	{
 		return DefaultFOV;
 	}
