Motion Design
- Updated cine camera in the itf panel to the ava cine camera.

diff --git a/Engine/Plugins/Experimental/Avalanche/Source/AvalancheEditor/Private/InteractiveTools/AvaInteractiveToolsExtension.cpp b/Engine/Plugins/Experimental/Avalanche/Source/AvalancheEditor/Private/InteractiveTools/AvaInteractiveToolsExtension.cpp
--- a/Engine/Plugins/Experimental/Avalanche/Source/AvalancheEditor/Private/InteractiveTools/AvaInteractiveToolsExtension.cpp
+++ b/Engine/Plugins/Experimental/Avalanche/Source/AvalancheEditor/Private/InteractiveTools/AvaInteractiveToolsExtension.cpp
@@ -9,7 +9,6 @@
 #include "Camera/CameraShakeSourceActor.h"
 #include "CameraRig_Crane.h"
 #include "CameraRig_Rail.h"
-#include "CineCameraActor.h"
 #include "EditorModeManager.h"
 #include "Engine/DirectionalLight.h"
 #include "Engine/PointLight.h"
@@ -19,6 +18,7 @@
 #include "IAvalancheInteractiveToolsModule.h"
 #include "InteractiveTools/AvaCameraActorTool.h"
 #include "MediaPlate.h"
+#include "Viewport/AvaCineCameraActor.h"
 #include "Viewport/AvaPostProcessVolume.h"
 
 namespace UE::AvaEditor::Private
@@ -125,7 +125,7 @@ void FAvaInteractiveToolsExtension::RegisterTools(IAvalancheInteractiveToolsModu
 		Priority += 1000,
 		UAvaCameraActorTool::StaticClass()));
 
-	InModule->RegisterTool(CamerasCategory, UAvaInteractiveToolsActorToolBuilder::CreateToolParameters<ACineCameraActor>( 
+	InModule->RegisterTool(CamerasCategory, UAvaInteractiveToolsActorToolBuilder::CreateToolParameters<AAvaCineCameraActor>( 
 		CamerasCategory, EditorCommands.CineCameraTool, FString(TEXT("Cine Camera Tool")), 
 		Priority += 1000,
 		UAvaCameraActorTool::StaticClass()));
