[Motion Design] Rundown Server API: Adding a command to save a rundown asset.

The rundown assets can be loaded either for playback or by any of the editing commands. It is possible to edit the rundown through editing commands and thus saving the asset on the server should be also possible.

diff --git a/Engine/Plugins/Experimental/Avalanche/Source/AvalancheMediaEditor/Private/Rundown/AvaRundownMessages.h b/Engine/Plugins/Experimental/Avalanche/Source/AvalancheMediaEditor/Private/Rundown/AvaRundownMessages.h
--- a/Engine/Plugins/Experimental/Avalanche/Source/AvalancheMediaEditor/Private/Rundown/AvaRundownMessages.h
+++ b/Engine/Plugins/Experimental/Avalanche/Source/AvalancheMediaEditor/Private/Rundown/AvaRundownMessages.h
@@ -122,11 +122,12 @@ public:
 };
 
 /**
- *	Request that the given rundown be opened.
- *	Only one rundown can be opened at a time. If another rundown
+ *	Request that the given rundown be loaded for playback.
+ *	This will also open an associated playback context.
+ *	Only one rundown can be opened for playback at a time. If another rundown
  *	is opened, it will be closed and all currently playing pages stopped.
  *	If the path is empty, nothing will be done and the server will reply with
- *	a FAvaRundownLoadedRundown message indicating which rundown is currently loaded.
+ *	a FAvaRundownServerMsg message indicating which rundown is currently loaded.
  */
 USTRUCT()
 struct FAvaRundownLoadRundown : public FAvaRundownMsgBase
@@ -138,6 +139,26 @@ public:
 	FString Rundown;
 };
 
+/**
+ * Request that the given rundown be saved to disk.
+ * The rundown asset must have been loaded, either by an edit command
+ * or playback, prior to this command.
+ * Unloaded assets will not be loaded by this command.
+ */
+USTRUCT()
+struct FAvaRundownSaveRundown : public FAvaRundownMsgBase
+{
+	GENERATED_BODY()
+
+public:
+	UPROPERTY()
+	FString Rundown;
+
+	UPROPERTY()
+	bool bOnlyIfIsDirty = false;
+};
+
+
 /**
  * Request the list of pages from the given rundown.
  */
diff --git a/Engine/Plugins/Experimental/Avalanche/Source/AvalancheMediaEditor/Private/Rundown/AvaRundownServer.cpp b/Engine/Plugins/Experimental/Avalanche/Source/AvalancheMediaEditor/Private/Rundown/AvaRundownServer.cpp
--- a/Engine/Plugins/Experimental/Avalanche/Source/AvalancheMediaEditor/Private/Rundown/AvaRundownServer.cpp
+++ b/Engine/Plugins/Experimental/Avalanche/Source/AvalancheMediaEditor/Private/Rundown/AvaRundownServer.cpp
@@ -15,6 +15,7 @@
 #include "Broadcast/OutputDevices/AvaBroadcastOutputRootItem.h"
 #include "Broadcast/OutputDevices/AvaBroadcastOutputTreeItem.h"
 #include "Broadcast/OutputDevices/AvaBroadcastRenderTargetMediaUtils.h"
+#include "Editor.h"
 #include "IAvaMediaModule.h"
 #include "IRemoteControlModule.h"
 #include "ImageUtils.h"
@@ -30,6 +31,7 @@
 #include "Rundown/AvaRundownPagePlayer.h"
 #include "Rundown/AvaRundownPlaybackUtils.h"
 #include "ScopedTransaction.h"
+#include "Subsystems/EditorAssetSubsystem.h"
 #include "TextureResource.h"
 
 DEFINE_LOG_CATEGORY_STATIC(LogAvaRundownServer, Log, All);
@@ -221,6 +223,7 @@ void FAvaRundownServer::Init(const FString& InAssignedHostName)
 	.Handling<FAvaRundownPing>(this, &FAvaRundownServer::HandleRundownPing)
 	.Handling<FAvaRundownGetRundowns>(this, &FAvaRundownServer::HandleGetRundowns)
 	.Handling<FAvaRundownLoadRundown>(this, &FAvaRundownServer::HandleLoadRundown)
+	.Handling<FAvaRundownSaveRundown>(this, &FAvaRundownServer::HandleSaveRundown)
 	.Handling<FAvaRundownCreatePage>(this, &FAvaRundownServer::HandleCreatePage)
 	.Handling<FAvaRundownDeletePage>(this, &FAvaRundownServer::HandleDeletePage)
 	.Handling<FAvaRundownCreateTemplate>(this, &FAvaRundownServer::HandleCreateTemplate)
@@ -487,6 +490,62 @@ void FAvaRundownServer::HandleLoadRundown(const FAvaRundownLoadRundown& InMessag
 		TEXT("Rundown \"%s\" loaded."), *RundownPlaybackCommandData.CurrentRundownPath.ToString());
 }
 
+void FAvaRundownServer::HandleSaveRundown(const FAvaRundownSaveRundown& InMessage, const TSharedRef<IMessageContext, ESPMode::ThreadSafe>& InContext)
+{
+	UEditorAssetSubsystem* EditorAssetSubsystem = GEditor ? GEditor->GetEditorSubsystem<UEditorAssetSubsystem>() : nullptr;
+	
+	if (!EditorAssetSubsystem)
+	{
+		SendMessage(InContext->GetSender(), InMessage.RequestId, ELogVerbosity::Error,
+			TEXT("Saving assets is only available in editor mode."));
+		return;
+	}
+	
+	if (InMessage.Rundown.IsEmpty())
+	{
+		SendMessage(InContext->GetSender(), InMessage.RequestId, ELogVerbosity::Error,
+			TEXT("Rundown asset not specified."));
+		return;
+	}
+
+	const FSoftObjectPath RundownAssetPath(InMessage.Rundown);
+
+	if (!RundownAssetPath.IsValid())
+	{
+		SendMessage(InContext->GetSender(), InMessage.RequestId, ELogVerbosity::Error,
+			TEXT("Rundown asset path \"%s\" is not valid."), *InMessage.Rundown);
+		return;
+	}
+
+	UObject* FoundObject = RundownAssetPath.ResolveObject();
+
+	if (!FoundObject)
+	{
+		SendMessage(InContext->GetSender(), InMessage.RequestId, ELogVerbosity::Error,
+			TEXT("Rundown asset \"%s\" is not loaded."), *InMessage.Rundown);
+		return;
+	}
+
+	UAvaRundown* FoundRundown = Cast<UAvaRundown>(FoundObject);
+	
+	if (!FoundRundown)
+	{
+		SendMessage(InContext->GetSender(), InMessage.RequestId, ELogVerbosity::Error,
+			TEXT("Asset path \"%s\" is loaded but is not a Rundown asset."), *InMessage.Rundown);
+		return;
+	}
+	
+	if (!EditorAssetSubsystem->SaveLoadedAsset(FoundRundown, InMessage.bOnlyIfIsDirty))
+	{
+		SendMessage(InContext->GetSender(), InMessage.RequestId, ELogVerbosity::Error,
+			TEXT("Unable to save asset \"%s\" to location \"%s\"."), *FoundRundown->GetName(), *RundownAssetPath.ToString());
+		return;
+	}
+
+	SendMessage(InContext->GetSender(), InMessage.RequestId, ELogVerbosity::Log,
+		TEXT("Asset \"%s\" save to location \"%s\"."), *FoundRundown->GetName(), *RundownAssetPath.ToString());
+}
+
 void FAvaRundownServer::HandleGetPages(const FAvaRundownGetPages& InMessage, const TSharedRef<IMessageContext, ESPMode::ThreadSafe>& InContext)
 {
 	UAvaRundown* Rundown = GetOrLoadRundownForEdit(InContext->GetSender(), InMessage.RequestId, InMessage.Rundown);
diff --git a/Engine/Plugins/Experimental/Avalanche/Source/AvalancheMediaEditor/Private/Rundown/AvaRundownServer.h b/Engine/Plugins/Experimental/Avalanche/Source/AvalancheMediaEditor/Private/Rundown/AvaRundownServer.h
--- a/Engine/Plugins/Experimental/Avalanche/Source/AvalancheMediaEditor/Private/Rundown/AvaRundownServer.h
+++ b/Engine/Plugins/Experimental/Avalanche/Source/AvalancheMediaEditor/Private/Rundown/AvaRundownServer.h
@@ -46,6 +46,7 @@ public:
 	void HandleRundownPing(const FAvaRundownPing& InMessage, const TSharedRef<IMessageContext, ESPMode::ThreadSafe>& InContext);
 	void HandleGetRundowns(const FAvaRundownGetRundowns& InMessage, const TSharedRef<IMessageContext, ESPMode::ThreadSafe>& InContext);
 	void HandleLoadRundown(const FAvaRundownLoadRundown& InMessage, const TSharedRef<IMessageContext, ESPMode::ThreadSafe>& InContext);
+	void HandleSaveRundown(const FAvaRundownSaveRundown& InMessage, const TSharedRef<IMessageContext, ESPMode::ThreadSafe>& InContext);
 	void HandleGetPages(const FAvaRundownGetPages& InMessage, const TSharedRef<IMessageContext, ESPMode::ThreadSafe>& InContext);
 	void HandleCreatePage(const FAvaRundownCreatePage& InMessage,const TSharedRef<IMessageContext, ESPMode::ThreadSafe>& InContext);
 	void HandleCreateTemplate(const FAvaRundownCreateTemplate& InMessage, const TSharedRef<IMessageContext, ESPMode::ThreadSafe>& InContext);
