MotionDesign : Boolean modifier
- Fix performance hitch when dragging transform values

diff --git a/Engine/Plugins/Experimental/Avalanche/Source/AvalancheModifiers/Private/Modifiers/AvaBooleanModifier.cpp b/Engine/Plugins/Experimental/Avalanche/Source/AvalancheModifiers/Private/Modifiers/AvaBooleanModifier.cpp
--- a/Engine/Plugins/Experimental/Avalanche/Source/AvalancheModifiers/Private/Modifiers/AvaBooleanModifier.cpp
+++ b/Engine/Plugins/Experimental/Avalanche/Source/AvalancheModifiers/Private/Modifiers/AvaBooleanModifier.cpp
@@ -86,8 +86,25 @@ void UAvaBooleanModifier::Apply()
 		return;
 	}
 
-	ApplyInternal();
+	// Update masking material here too to avoid flickering between mask and original materials due to async task
 	UpdateMaskingMaterials();
+
+	// Use async execution since boolean is an heavy operation
+	TWeakObjectPtr<UAvaBooleanModifier> ThisWeak(this);
+	Async(EAsyncExecution::TaskGraphMainThread, [ThisWeak]()
+	{
+		UAvaBooleanModifier* This = ThisWeak.Get();
+
+		if (!IsValid(This))
+		{
+			return;
+		}
+
+		This->ApplyInternal();
+		This->UpdateMaskingMaterials();
+
+		This->Next();
+	});
 }
 
 void UAvaBooleanModifier::OnModifierDisabled(EActorModifierCoreDisableReason InReason)
@@ -298,7 +315,7 @@ void UAvaBooleanModifier::ApplyInternal()
 	UAvaBooleanModifierShared* Shared = GetShared<UAvaBooleanModifierShared>(false);
 	if (!Shared)
 	{
-		Next();
+		Fail(LOCTEXT("InvalidSharedObject", "Invalid boolean modifier shared object"));
 		return;
 	}
 
@@ -355,8 +372,6 @@ void UAvaBooleanModifier::ApplyInternal()
 	{
 		LastTransform = ActorModified->GetActorTransform();
 	}
-
-	Next();
 }
 
 void UAvaBooleanModifier::OnModeChanged()
