[MotionDesign] Mask modifier duplicate crash

diff --git a/Engine/Plugins/Experimental/Avalanche/Source/AvalancheMask/Private/Handling/AvaMaskText3DActorMaterialCollectionHandle.cpp b/Engine/Plugins/Experimental/Avalanche/Source/AvalancheMask/Private/Handling/AvaMaskText3DActorMaterialCollectionHandle.cpp
--- a/Engine/Plugins/Experimental/Avalanche/Source/AvalancheMask/Private/Handling/AvaMaskText3DActorMaterialCollectionHandle.cpp
+++ b/Engine/Plugins/Experimental/Avalanche/Source/AvalancheMask/Private/Handling/AvaMaskText3DActorMaterialCollectionHandle.cpp
@@ -332,7 +332,10 @@ FStructView FAvaMaskText3DActorMaterialCollectionHandle::GetMaterialHandleData(
 {
 	if (InParentHandleData)
 	{
-		return InParentHandleData->GroupMaterialData[InSlotIdx];
+		if (ensure(InParentHandleData->GroupMaterialData.Contains(InSlotIdx)))
+		{
+			return InParentHandleData->GroupMaterialData[InSlotIdx];
+		}
 	}
 
 	return nullptr;
@@ -775,7 +778,10 @@ FStructView FAvaMaskAvaTextActorMaterialCollectionHandle::GetMaterialHandleData(
 {
 	if (InParentHandleData)
 	{
-		return InParentHandleData->GroupMaterialData[InSlotIdx];
+		if (ensure(InParentHandleData->GroupMaterialData.Contains(InSlotIdx)))
+		{
+			return InParentHandleData->GroupMaterialData[InSlotIdx];
+		}
 	}
 
 	return nullptr;	
diff --git a/Engine/Plugins/Experimental/Avalanche/Source/AvalancheMask/Private/Mask2D/AvaMask2DBaseModifier.cpp b/Engine/Plugins/Experimental/Avalanche/Source/AvalancheMask/Private/Mask2D/AvaMask2DBaseModifier.cpp
--- a/Engine/Plugins/Experimental/Avalanche/Source/AvalancheMask/Private/Mask2D/AvaMask2DBaseModifier.cpp
+++ b/Engine/Plugins/Experimental/Avalanche/Source/AvalancheMask/Private/Mask2D/AvaMask2DBaseModifier.cpp
@@ -267,6 +267,7 @@ void UAvaMask2DBaseModifier::OnModifierAdded(EActorModifierCoreEnableReason InRe
 	else
 	{
 		CanvasParamsToLocal();
+		MarkModifierDirty();
 	}
 }
 
diff --git a/Engine/Plugins/Experimental/Avalanche/Source/AvalancheMask/Public/Mask2D/AvaMask2DBaseModifier.h b/Engine/Plugins/Experimental/Avalanche/Source/AvalancheMask/Public/Mask2D/AvaMask2DBaseModifier.h
--- a/Engine/Plugins/Experimental/Avalanche/Source/AvalancheMask/Public/Mask2D/AvaMask2DBaseModifier.h
+++ b/Engine/Plugins/Experimental/Avalanche/Source/AvalancheMask/Public/Mask2D/AvaMask2DBaseModifier.h
@@ -233,10 +233,10 @@ protected:
 
 	bool bIsRestoring = false;
 
-	UPROPERTY()
+	UPROPERTY(DuplicateTransient)
 	TMap<TWeakObjectPtr<UObject>, FInstancedStruct> MaterialCollectionHandleData;
 
-	UPROPERTY()
+	UPROPERTY(DuplicateTransient)
 	TMap<TWeakObjectPtr<UMaterialInterface>, FInstancedStruct> MaterialHandleData;
 
 	TMap<TObjectKey<UObject>, TSharedPtr<IAvaMaskMaterialCollectionHandle>> MaterialCollectionHandles;
