Motion Design: Fix crash when destroying/ending play on a Level while Sequence Players still active with Camera Instances.

diff --git a/Engine/Plugins/Experimental/Avalanche/Source/AvalancheSequence/Private/AvaSequencePlaybackActor.cpp b/Engine/Plugins/Experimental/Avalanche/Source/AvalancheSequence/Private/AvaSequencePlaybackActor.cpp
--- a/Engine/Plugins/Experimental/Avalanche/Source/AvalancheSequence/Private/AvaSequencePlaybackActor.cpp
+++ b/Engine/Plugins/Experimental/Avalanche/Source/AvalancheSequence/Private/AvaSequencePlaybackActor.cpp
@@ -1,7 +1,6 @@
 // Copyright Epic Games, Inc. All Rights Reserved.
 
 #include "AvaSequencePlaybackActor.h"
-#include "Algo/ForEach.h"
 #include "AvaSequencePlayer.h"
 #include "AvaSequenceSubsystem.h"
 #include "Engine/World.h"
@@ -28,19 +27,25 @@ void AAvaSequencePlaybackActor::SetSequenceProvider(IAvaSequenceProvider& InSequ
 
 void AAvaSequencePlaybackActor::CleanupPlayers()
 {
-	auto CleanupPlayer = [](const TObjectPtr<UAvaSequencePlayer>& InPlayerToCleanup)
-	{
-		if (InPlayerToCleanup)
-		{
-			InPlayerToCleanup->Cleanup();
-		}	
-	};
+	// "OnSequenceFinished" is called while cleaning up a player which as a result removes the players from these sets.
+	// To prevent removing element one by one and while iterating these sets, these sets are moved to temp set.
 
-	Algo::ForEach(StoppedSequencePlayers, CleanupPlayer);
-	Algo::ForEach(ActiveSequencePlayers, CleanupPlayer);
+	TSet<TObjectPtr<UAvaSequencePlayer>> SequencePlayers;
+	SequencePlayers.Reserve(StoppedSequencePlayers.Num() + ActiveSequencePlayers.Num());
+
+	SequencePlayers.Append(MoveTemp(StoppedSequencePlayers));
+	SequencePlayers.Append(MoveTemp(ActiveSequencePlayers));
 
 	StoppedSequencePlayers.Reset();
 	ActiveSequencePlayers.Reset();
+
+	for (UAvaSequencePlayer* Player : SequencePlayers)
+	{
+		if (Player)
+		{
+			Player->Cleanup();
+		}
+	}
 }
 
 UAvaSequencePlayer* AAvaSequencePlaybackActor::PlaySequence(UAvaSequence* InSequence, const FAvaSequencePlayParams& InPlaySettings)
diff --git a/Engine/Plugins/Experimental/Avalanche/Source/AvalancheSequence/Private/AvaSequencePlayer.cpp b/Engine/Plugins/Experimental/Avalanche/Source/AvalancheSequence/Private/AvaSequencePlayer.cpp
--- a/Engine/Plugins/Experimental/Avalanche/Source/AvalancheSequence/Private/AvaSequencePlayer.cpp
+++ b/Engine/Plugins/Experimental/Avalanche/Source/AvalancheSequence/Private/AvaSequencePlayer.cpp
@@ -159,7 +159,8 @@ void UAvaSequencePlayer::Cleanup()
 {
 	if (ensureAlwaysMsgf(!IsEvaluating(), TEXT("Calling UAvaSequencePlayer::Cleanup while still evaluating is not allowed!")))
 	{
-		TearDown();	
+		Stop();
+		TearDown();
 	}
 }
 
