[Avalanche/AvalancheEditor] AvaVector customization change
Smaller lock icon and moved to the ValueWidget side.

diff --git a/Engine/Plugins/Experimental/Avalanche/Source/AvalancheEditor/Private/DetailView/Customizations/AvaVectorPropertyTypeCustomization.cpp b/Engine/Plugins/Experimental/Avalanche/Source/AvalancheEditor/Private/DetailView/Customizations/AvaVectorPropertyTypeCustomization.cpp
--- a/Engine/Plugins/Experimental/Avalanche/Source/AvalancheEditor/Private/DetailView/Customizations/AvaVectorPropertyTypeCustomization.cpp
+++ b/Engine/Plugins/Experimental/Avalanche/Source/AvalancheEditor/Private/DetailView/Customizations/AvaVectorPropertyTypeCustomization.cpp
@@ -54,6 +54,8 @@ void FAvaVectorPropertyTypeCustomization::CustomizeHeader(TSharedRef<IPropertyHa
 	}
 
 	bool bPreserveRatio = false;
+	TSharedPtr<SWidget> PreserveRatioWidget = SNullWidget::NullWidget;
+
 	if (StructPropertyHandle->HasMetaData("AllowPreserveRatio"))
 	{
 		bPreserveRatio = true;
@@ -70,15 +72,16 @@ void FAvaVectorPropertyTypeCustomization::CustomizeHeader(TSharedRef<IPropertyHa
 			const FText& ButtonText = GetComboButtonText(ButtonMode);
 			
 			return SNew(SButton)
-					.HAlign(EHorizontalAlignment::HAlign_Fill)
-					.VAlign(EVerticalAlignment::VAlign_Fill)
+					.HAlign(HAlign_Fill)
+					.VAlign(VAlign_Fill)
 					.Cursor(EMouseCursor::Hand)
 					.OnClicked(this, &FAvaVectorPropertyTypeCustomization::OnComboButtonClicked, ButtonMode)
-					.ContentPadding(0.f)
+					.ContentPadding(FMargin(-8.f, 0.f))
 					[
 						SNew(SHorizontalBox)
 						+ SHorizontalBox::Slot()
-						.FillWidth(0.25f)
+						.HAlign(HAlign_Left)
+						.AutoWidth()
 						[
 							SNew(SScaleBox)
 							[
@@ -88,12 +91,12 @@ void FAvaVectorPropertyTypeCustomization::CustomizeHeader(TSharedRef<IPropertyHa
 							]
 						]
 						+ SHorizontalBox::Slot()
-						.FillWidth(1.f)
+						.AutoWidth()
 						[
 							SNew(STextBlock)
 							.Justification(ETextJustify::Center)
 							.Text(ButtonText)
-							.Margin(FMargin(5.f, 0.f))
+							.Margin(FMargin(5.f, 0.f, 0.f, 0.f))
 						]
 					];
 		};
@@ -119,22 +122,16 @@ void FAvaVectorPropertyTypeCustomization::CustomizeHeader(TSharedRef<IPropertyHa
 				ComboBoxButtonBuilder(Mode)
 			];
 		}
-		
-		HeaderRow.NameContent()[
-			SNew(SHorizontalBox)
-			+ SHorizontalBox::Slot()
-			.FillWidth(1.f)
-			[
-				StructPropertyHandle->CreatePropertyNameWidget()
-			]
-			+ SHorizontalBox::Slot()
-			.FillWidth(1.f)
+
+		PreserveRatioWidget = SNew(SBox)
+			.MinDesiredWidth(60.f)
 			[
 				SAssignNew(ComboButton, SComboButton)
-				.HAlign(EHorizontalAlignment::HAlign_Fill)
 				.VAlign(EVerticalAlignment::VAlign_Center)
 				.Cursor(EMouseCursor::Hand)
 				.Method(EPopupMethod::UseCurrentWindow)
+				.HasDownArrow(false)
+				.ContentPadding(FMargin(-5.f, -2.f))
 				.ButtonContent()
 				[
 					SNew(SHorizontalBox)
@@ -161,14 +158,12 @@ void FAvaVectorPropertyTypeCustomization::CustomizeHeader(TSharedRef<IPropertyHa
 				[
 					ButtonVerticalBox
 				]
-			]
-		];
-	}
-	else
-	{
-		HeaderRow.NameContent()[StructPropertyHandle->CreatePropertyNameWidget()];
+			];
 	}
 
+	// Assign Name Widget
+	HeaderRow.NameContent()[StructPropertyHandle->CreatePropertyNameWidget()];
+
 	// fill available space
 	HeaderRow.ValueWidget.HorizontalAlignment = EHorizontalAlignment::HAlign_Fill;
 	HeaderRow.ValueWidget.VerticalAlignment = EVerticalAlignment::VAlign_Fill;
@@ -189,7 +184,7 @@ void FAvaVectorPropertyTypeCustomization::CustomizeHeader(TSharedRef<IPropertyHa
 			}
 			else
 			{
-				RatioMode = ERatioMode::PreserveXYZ;
+				RatioMode = ERatioMode::None;
 			}
 		}
 		
@@ -216,7 +211,19 @@ void FAvaVectorPropertyTypeCustomization::CustomizeHeader(TSharedRef<IPropertyHa
 		
 		HeaderRow.ValueContent()
 		[
-			SNew(SNumericVectorInputBox3D)
+			SNew(SHorizontalBox)
+
+			+ SHorizontalBox::Slot()
+			.AutoWidth()
+			.Padding(0.f, 0.f, 2.f, 0.f)
+			[
+				PreserveRatioWidget.ToSharedRef()
+			]
+
+			+ SHorizontalBox::Slot()
+			.FillWidth(1.f)
+			[
+				SNew(SNumericVectorInputBox3D)
 				.Font(IDetailLayoutBuilder::GetDetailFont())
 				.X(this, &FAvaVectorPropertyTypeCustomization::GetVectorComponent,  static_cast<uint8>(0)) // X
 				.Y(this, &FAvaVectorPropertyTypeCustomization::GetVectorComponent,  static_cast<uint8>(1)) // Y
@@ -237,6 +244,7 @@ void FAvaVectorPropertyTypeCustomization::CustomizeHeader(TSharedRef<IPropertyHa
 				.IsEnabled(this, &FAvaVectorPropertyTypeCustomization::CanEditValue)
 				.OnBeginSliderMovement(this, &FAvaVectorPropertyTypeCustomization::OnBeginSliderMovement)
 				.OnEndSliderMovement(this, &FAvaVectorPropertyTypeCustomization::OnEndSliderMovement)
+			]
 		];
 	}
 	else
@@ -253,7 +261,7 @@ void FAvaVectorPropertyTypeCustomization::CustomizeHeader(TSharedRef<IPropertyHa
 			}
 			else
 			{
-				RatioMode = ERatioMode::PreserveXY;
+				RatioMode = ERatioMode::None;
 			}
 		}
 		
@@ -280,7 +288,19 @@ void FAvaVectorPropertyTypeCustomization::CustomizeHeader(TSharedRef<IPropertyHa
 
 		HeaderRow.ValueContent()
 		[
-			SNew(SNumericVectorInputBox2D)
+			SNew(SHorizontalBox)
+
+			+ SHorizontalBox::Slot()
+			.AutoWidth()
+			.Padding(0.f, 0.f, 2.f, 0.f)
+			[
+				PreserveRatioWidget.ToSharedRef()
+			]
+
+			+ SHorizontalBox::Slot()
+			.FillWidth(1.f)
+			[
+				SNew(SNumericVectorInputBox2D)
 				.Font(IDetailLayoutBuilder::GetDetailFont())
 				.X(this, &FAvaVectorPropertyTypeCustomization::GetVectorComponent, static_cast<uint8>(0)) // X
 				.Y(this, &FAvaVectorPropertyTypeCustomization::GetVectorComponent,  static_cast<uint8>(1)) // Y
@@ -298,6 +318,7 @@ void FAvaVectorPropertyTypeCustomization::CustomizeHeader(TSharedRef<IPropertyHa
 				.IsEnabled(this, &FAvaVectorPropertyTypeCustomization::CanEditValue)
 				.OnBeginSliderMovement(this, &FAvaVectorPropertyTypeCustomization::OnBeginSliderMovement)
 				.OnEndSliderMovement(this, &FAvaVectorPropertyTypeCustomization::OnEndSliderMovement)
+			]
 		];
 	}
 }
