[Motion Design]
 - Scene rig management and UI

diff --git a/Engine/Plugins/Experimental/Avalanche/Avalanche.uplugin b/Engine/Plugins/Experimental/Avalanche/Avalanche.uplugin
--- a/Engine/Plugins/Experimental/Avalanche/Avalanche.uplugin
+++ b/Engine/Plugins/Experimental/Avalanche/Avalanche.uplugin
@@ -228,6 +228,16 @@
 			"Name": "AvalancheAttributeEditor",
 			"Type": "Editor",
 			"LoadingPhase": "Default"
+		},
+		{
+			"Name": "AvalancheSceneRig",
+			"Type": "Runtime",
+			"LoadingPhase": "Default"
+		},
+		{
+			"Name": "AvalancheSceneRigEditor",
+			"Type": "Editor",
+			"LoadingPhase": "Default"
 		}
 	],
 	"Plugins": [
diff --git a/Engine/Plugins/Experimental/Avalanche/Source/Avalanche/Public/AvaSceneSettings.h b/Engine/Plugins/Experimental/Avalanche/Source/Avalanche/Public/AvaSceneSettings.h
--- a/Engine/Plugins/Experimental/Avalanche/Source/Avalanche/Public/AvaSceneSettings.h
+++ b/Engine/Plugins/Experimental/Avalanche/Source/Avalanche/Public/AvaSceneSettings.h
@@ -4,7 +4,9 @@
 
 #include "AvaAttribute.h"
 #include "Containers/Array.h"
+#include "UObject/NoExportTypes.h"
 #include "UObject/Object.h"
+#include "UObject/SoftObjectPath.h"
 #include "AvaSceneSettings.generated.h"
 
 /** Object containing information about its Scene */
@@ -14,11 +16,16 @@ class UAvaSceneSettings : public UObject
 	GENERATED_BODY()
 
 public:
-	static FName GetSceneAttributesName()
+	static FName GetSceneAttributesPropertyName()
 	{
 		return GET_MEMBER_NAME_CHECKED(UAvaSceneSettings, SceneAttributes);
 	}
 
+	static FName GetSceneRigPropertyName()
+	{
+		return GET_MEMBER_NAME_CHECKED(UAvaSceneSettings, SceneRig);
+	}
+
 	/**
 	 * Iterate each valid Scene Attribute of the given Type
 	 * The callable should return true to continue iteration, and false to stop it
@@ -38,7 +45,21 @@ public:
 		}
 	}
 
+	FSoftObjectPath GetSceneRig() const
+	{
+		return SceneRig;
+	}
+
+	void SetSceneRig(const FSoftObjectPath& InSceneRig)
+	{
+		SceneRig = InSceneRig;
+	}
+
+
 private:
 	UPROPERTY(EditAnywhere, Instanced, Category="Scene Attributes")
 	TArray<TObjectPtr<UAvaAttribute>> SceneAttributes;
+
+	UPROPERTY()
+	FSoftObjectPath SceneRig;
 };
diff --git a/Engine/Plugins/Experimental/Avalanche/Source/AvalancheEditor/AvalancheEditor.Build.cs b/Engine/Plugins/Experimental/Avalanche/Source/AvalancheEditor/AvalancheEditor.Build.cs
--- a/Engine/Plugins/Experimental/Avalanche/Source/AvalancheEditor/AvalancheEditor.Build.cs
+++ b/Engine/Plugins/Experimental/Avalanche/Source/AvalancheEditor/AvalancheEditor.Build.cs
@@ -42,6 +42,8 @@ public class AvalancheEditor : ModuleRules
 				"AvalancheOutliner",
 				"AvalanchePropertyAnimator",
 				"AvalancheRemoteControlEditor",
+				"AvalancheSceneRig",
+				"AvalancheSceneRigEditor",
 				"AvalancheSceneTree",
 				"AvalancheSequence",
 				"AvalancheSequencer",
diff --git a/Engine/Plugins/Experimental/Avalanche/Source/AvalancheEditor/Private/AvaEditorIntegration.cpp b/Engine/Plugins/Experimental/Avalanche/Source/AvalancheEditor/Private/AvaEditorIntegration.cpp
--- a/Engine/Plugins/Experimental/Avalanche/Source/AvalancheEditor/Private/AvaEditorIntegration.cpp
+++ b/Engine/Plugins/Experimental/Avalanche/Source/AvalancheEditor/Private/AvaEditorIntegration.cpp
@@ -12,6 +12,7 @@
 #include "Outliner/AvaOutlinerExtension.h"
 #include "RemoteControl/AvaRCExtension.h"
 #include "Scene/AvaSceneExtension.h"
+#include "Scene/AvaSceneRigExtension.h"
 #include "Selection/AvaSelectionProviderExtension.h"
 #include "Sequencer/AvaLevelSequencerExtension.h"
 #include "Transition/AvaTransitionExtension.h"
@@ -33,6 +34,7 @@ TSharedRef<IAvaEditor> FAvaLevelEditorIntegration::BuildEditor()
 		.AddExtension<FAvaInteractiveToolsExtension>()
 		.AddExtension<FAvaLevelMaterialDesignerExtension>()
 		.AddExtension<FAvaSceneExtension>()
+		.AddExtension<FAvaSceneRigExtension>()
 		.AddExtension<FAvaTransitionExtension>()
 		.AddExtension<FAvaAdvancedRenamerExtension>()
 		.AddExtension<FAvaComponentVisualizerExtension>()
diff --git a/Engine/Plugins/Experimental/Avalanche/Source/AvalancheEditor/Private/DetailView/Customizations/AvaSceneSettingsCustomization.cpp b/Engine/Plugins/Experimental/Avalanche/Source/AvalancheEditor/Private/DetailView/Customizations/AvaSceneSettingsCustomization.cpp
--- a/Engine/Plugins/Experimental/Avalanche/Source/AvalancheEditor/Private/DetailView/Customizations/AvaSceneSettingsCustomization.cpp
+++ b/Engine/Plugins/Experimental/Avalanche/Source/AvalancheEditor/Private/DetailView/Customizations/AvaSceneSettingsCustomization.cpp
@@ -4,16 +4,24 @@
 #include "AvaSceneSettings.h"
 #include "DetailLayoutBuilder.h"
 #include "IAvaAttributeEditorModule.h"
+#include "IAvaSceneRigEditorModule.h"
+#include "PropertyHandle.h"
+#include "Templates/SharedPointer.h"
 
 TSharedRef<IDetailCustomization> FAvaSceneSettingsCustomization::MakeInstance()
 {
 	return MakeShared<FAvaSceneSettingsCustomization>();
 }
 
-void FAvaSceneSettingsCustomization::CustomizeDetails(IDetailLayoutBuilder& InDetailBuilder)
+void FAvaSceneSettingsCustomization::CustomizeDetails(IDetailLayoutBuilder& DetailBuilder)
 {
-	if (TSharedPtr<IPropertyHandle> AttributesHandle = InDetailBuilder.GetProperty(UAvaSceneSettings::GetSceneAttributesName()))
+	if (TSharedPtr<IPropertyHandle> SceneRigHandle = DetailBuilder.GetProperty(UAvaSceneSettings::GetSceneRigPropertyName()))
 	{
-		IAvaAttributeEditorModule::Get().CustomizeAttributes(AttributesHandle.ToSharedRef(), InDetailBuilder);	
+		IAvaSceneRigEditorModule::Get().CustomizeSceneRig(SceneRigHandle.ToSharedRef(), DetailBuilder);
+	}
+
+	if (TSharedPtr<IPropertyHandle> AttributesHandle = DetailBuilder.GetProperty(UAvaSceneSettings::GetSceneAttributesPropertyName()))
+	{
+		IAvaAttributeEditorModule::Get().CustomizeAttributes(AttributesHandle.ToSharedRef(), DetailBuilder);
 	}
 }
diff --git a/Engine/Plugins/Experimental/Avalanche/Source/AvalancheEditor/Private/DetailView/Customizations/AvaSceneSettingsCustomization.h b/Engine/Plugins/Experimental/Avalanche/Source/AvalancheEditor/Private/DetailView/Customizations/AvaSceneSettingsCustomization.h
--- a/Engine/Plugins/Experimental/Avalanche/Source/AvalancheEditor/Private/DetailView/Customizations/AvaSceneSettingsCustomization.h
+++ b/Engine/Plugins/Experimental/Avalanche/Source/AvalancheEditor/Private/DetailView/Customizations/AvaSceneSettingsCustomization.h
@@ -3,14 +3,16 @@
 #pragma once
 
 #include "IDetailCustomization.h"
+#include "Templates/SharedPointerFwd.h"
+
+class IDetailLayoutBuilder;
 
 class FAvaSceneSettingsCustomization : public IDetailCustomization
 {
 public:
 	static TSharedRef<IDetailCustomization> MakeInstance();
 
-private:
 	//~ Begin IDetailCustomization
-	virtual void CustomizeDetails(IDetailLayoutBuilder& InDetailBuilder) override;
+	virtual void CustomizeDetails(IDetailLayoutBuilder& DetailBuilder) override;
 	//~ End IDetailCustomization
 };
diff --git a/Engine/Plugins/Experimental/Avalanche/Source/AvalancheEditor/Private/Outliner/AvaOutlinerExtension.cpp b/Engine/Plugins/Experimental/Avalanche/Source/AvalancheEditor/Private/Outliner/AvaOutlinerExtension.cpp
--- a/Engine/Plugins/Experimental/Avalanche/Source/AvalancheEditor/Private/Outliner/AvaOutlinerExtension.cpp
+++ b/Engine/Plugins/Experimental/Avalanche/Source/AvalancheEditor/Private/Outliner/AvaOutlinerExtension.cpp
@@ -8,6 +8,7 @@
 #include "AvaOutlinerSubsystem.h"
 #include "AvaOutlinerTabSpawner.h"
 #include "AvaScene.h"
+#include "AvaSceneRigSubsystem.h"
 #include "AvaSequence.h"
 #include "AvaShapeActor.h"
 #include "AvaTextActor.h"
@@ -22,6 +23,8 @@
 #include "Framework/Docking/LayoutExtender.h"
 #include "IAvaOutliner.h"
 #include "IAvaOutlinerModule.h"
+#include "IAvaSceneRigEditorModule.h"
+#include "Item/AvaOutlinerActor.h"
 #include "Item/AvaOutlinerComponent.h"
 #include "ItemProxies/AvaOutlinerItemProxyRegistry.h"
 #include "LevelEditor.h"
@@ -302,7 +305,7 @@ void FAvaOutlinerExtension::ExtendOutlinerItemContextMenu(UToolMenu* InItemConte
 
 	// Note: Since the Outliner Command List is linked to the Ava Command List (see IAvaOutliner::SetBaseCommandList),
 	// we do NOT need to add the entry with a different Command List
-	const FAvaEditorCommands AvaEditorCommands = FAvaEditorCommands::Get();
+	const FAvaEditorCommands& AvaEditorCommands = FAvaEditorCommands::Get();
 
 	ContextActionsSection->AddMenuEntry(
 		AvaEditorCommands.OpenAdvancedRenamerTool_SelectedActors,
@@ -374,6 +377,20 @@ void FAvaOutlinerExtension::ExtendOutlinerItemFilters(TArray<TSharedPtr<IAvaOutl
 
 	OutItemFilters.Add(MakeShared<FAvaOutlinerItemTypeFilter>(TEXT("SkeletalMesh")
 		, TArray<TSubclassOf<UObject>>{ ASkeletalMeshActor::StaticClass() }));
+
+	OutItemFilters.Add(MakeShared<FAvaOutlinerItemTypeFilter>(TEXT("SceneRig")
+		, FAvaOutlinerItemTypeFilterPassDelegate::CreateSP(this, &FAvaOutlinerExtension::OnSceneRigFilterPass)
+		, FAppStyle::GetBrush(TEXT("LandscapeEditor.NoiseTool"))
+		, LOCTEXT("SceneRigFilterTooltip", "Scene Rig Actors")));
+}
+
+bool FAvaOutlinerExtension::OnSceneRigFilterPass(FAvaOutlinerFilterType InItem) const
+{
+	if (const FAvaOutlinerActor* const OutlinerActor = InItem.CastTo<FAvaOutlinerActor>())
+	{
+		return IAvaSceneRigEditorModule::Get().IsActiveSceneRigActor(GetWorld(), OutlinerActor->GetActor());
+	}
+	return false;
 }
 
 TOptional<EItemDropZone> FAvaOutlinerExtension::OnOutlinerItemCanAcceptDrop(const FDragDropEvent& InDragDropEvent
diff --git a/Engine/Plugins/Experimental/Avalanche/Source/AvalancheEditor/Private/Outliner/AvaOutlinerExtension.h b/Engine/Plugins/Experimental/Avalanche/Source/AvalancheEditor/Private/Outliner/AvaOutlinerExtension.h
--- a/Engine/Plugins/Experimental/Avalanche/Source/AvalancheEditor/Private/Outliner/AvaOutlinerExtension.h
+++ b/Engine/Plugins/Experimental/Avalanche/Source/AvalancheEditor/Private/Outliner/AvaOutlinerExtension.h
@@ -64,6 +64,8 @@ public:
 private:
 	void GroupSelection();
 
+	bool OnSceneRigFilterPass(FAvaOutlinerFilterType InItem) const;
+
 	TSharedPtr<IAvaOutliner> AvaOutliner;
 
 	TSharedRef<FUICommandList> OutlinerCommands;
diff --git a/Engine/Plugins/Experimental/Avalanche/Source/AvalancheEditor/Private/Scene/AvaSceneRigExtension.cpp b/Engine/Plugins/Experimental/Avalanche/Source/AvalancheEditor/Private/Scene/AvaSceneRigExtension.cpp
new file mode 100644
--- /dev/null
+++ b/Engine/Plugins/Experimental/Avalanche/Source/AvalancheEditor/Private/Scene/AvaSceneRigExtension.cpp
@@ -0,0 +1,410 @@
+// Copyright Epic Games, Inc. All Rights Reserved.
+
+#include "AvaSceneRigExtension.h"
+#include "AvaSceneRigEditorCommands.h"
+#include "AvaSceneRigSubsystem.h"
+#include "AvaSceneSettings.h"
+#include "AvaSceneState.h"
+#include "AvaSceneSubsystem.h"
+#include "Engine/Level.h"
+#include "Engine/LevelStreaming.h"
+#include "Engine/World.h"
+#include "IAvaSceneInterface.h"
+#include "IAvaSceneRigEditorModule.h"
+#include "Item/AvaOutlinerActor.h"
+#include "ScopedTransaction.h"
+#include "Styling/AppStyle.h"
+#include "ToolMenus.h"
+#include "ToolMenu.h"
+#include "ToolMenuContext/AvaOutlinerItemsContext.h"
+#include "ToolMenuEntry.h"
+#include "ToolMenuSection.h"
+
+#define LOCTEXT_NAMESPACE "AvaSceneRigExtension"
+
+void FAvaSceneRigExtension::Activate()
+{
+	if (IAvaOutlinerModule::IsLoaded())
+	{
+		OutlinerItemContextMenuDelegate = IAvaOutlinerModule::Get().GetOnExtendOutlinerItemContextMenu()
+			.AddSP(this, &FAvaSceneRigExtension::ExtendOutlinerItemContextMenu);
+	}
+}
+
+void FAvaSceneRigExtension::Deactivate()
+{
+	if (IAvaOutlinerModule::IsLoaded() && OutlinerItemContextMenuDelegate.IsValid())
+	{
+		IAvaOutlinerModule::Get().GetOnExtendOutlinerItemContextMenu().Remove(OutlinerItemContextMenuDelegate);
+		OutlinerItemContextMenuDelegate.Reset();
+	}
+}
+
+void FAvaSceneRigExtension::BindCommands(const TSharedRef<FUICommandList>& InCommandList)
+{
+	FAvaEditorExtension::BindCommands(InCommandList);
+	
+	const FAvaSceneRigEditorCommands& SceneRigEditorCommands = FAvaSceneRigEditorCommands::Get();
+	
+	InCommandList->MapAction(
+		SceneRigEditorCommands.PromptToSaveSceneRigFromOutlinerItems,
+		FExecuteAction::CreateSP(this, &FAvaSceneRigExtension::PromptToSaveSceneRigFromOutlinerItems),
+		FCanExecuteAction::CreateSP(this, &FAvaSceneRigExtension::CanSaveSceneRigFromOutlinerItems));
+
+	InCommandList->MapAction(
+		SceneRigEditorCommands.AddOutlinerItemsToSceneRig,
+		FExecuteAction::CreateSP(this, &FAvaSceneRigExtension::AddOutlinerItemsToSceneRig),
+		FCanExecuteAction::CreateSP(this, &FAvaSceneRigExtension::CanAddOutlinerItemsToSceneRig));
+
+	InCommandList->MapAction(
+		SceneRigEditorCommands.RemoveOutlinerItemsToSceneRig,
+		FExecuteAction::CreateSP(this, &FAvaSceneRigExtension::RemoveOutlinerItemsFromSceneRig),
+		FCanExecuteAction::CreateSP(this, &FAvaSceneRigExtension::CanRemoveOutlinerItemsFromSceneRig));
+
+	CommandListWeak = InCommandList;
+}
+
+void FAvaSceneRigExtension::ExtendOutlinerItemContextMenu(UToolMenu* const InToolMenu)
+{
+	if (!IsValid(InToolMenu))
+	{
+		return;
+	}
+
+	UAvaOutlinerItemsContext* const ItemsContext = InToolMenu->FindContext<UAvaOutlinerItemsContext>();
+	if (!IsValid(ItemsContext))
+	{
+		return;
+	}
+
+	const TConstArrayView<FAvaOutlinerItemWeakPtr> SelectedItems = ItemsContext->GetItems();
+	if (SelectedItems.IsEmpty())
+	{
+		return;
+	}
+
+	ItemsContextWeak = ItemsContext;
+
+	const TArray<AActor*> SelectedActors = OutlinerItemsToActors(ItemsContextWeak.Get()->GetItems(), true);
+	if (!SelectedActors.IsEmpty() && UAvaSceneRigSubsystem::AreActorsSupported(SelectedActors))
+	{
+		const FAvaSceneRigEditorCommands& SceneRigEditorCommands = FAvaSceneRigEditorCommands::Get();
+
+		FToolMenuSection& NewSection = InToolMenu->AddSection(TEXT("SceneRig"), LOCTEXT("SceneRig", "Scene Rig")
+			, FToolMenuInsert(TEXT("ContextActions"), EToolMenuInsertType::After));
+
+		FToolMenuEntry SceneRigMenuEntry = NewSection.AddSubMenu(TEXT("SceneRig")
+			, LOCTEXT("SceneRig", "Scene Rig")
+			, LOCTEXT("SceneRigToolTip", "Scene Rig")
+			, FNewToolMenuDelegate::CreateSP(this, &FAvaSceneRigExtension::CreateSubMenu)
+			, false
+			, FSlateIcon(FAppStyle::GetAppStyleSetName(), TEXT("LandscapeEditor.NoiseTool")));
+	}
+}
+
+void FAvaSceneRigExtension::CreateSubMenu(UToolMenu* const InToolMenu)
+{
+	InToolMenu->AddDynamicSection(TEXT("SceneRig"), FNewToolMenuDelegate::CreateSPLambda(this, [this](UToolMenu* const InToolMenu)
+		{
+			const TSharedPtr<FUICommandList> CommandList = CommandListWeak.Pin();
+			const FAvaSceneRigEditorCommands& SceneRigEditorCommands = FAvaSceneRigEditorCommands::Get();
+
+			FToolMenuSection& NewSection = InToolMenu->AddSection(TEXT("SceneRig"), LOCTEXT("SceneRigActions", "Scene Rig"));
+
+			UWorld* const World = GetWorld();
+			if (IsValid(World))
+			{
+				UAvaSceneRigSubsystem* const SceneRigSubsystem = UAvaSceneRigSubsystem::ForWorld(World);
+				if (IsValid(SceneRigSubsystem))
+				{
+					const ULevelStreaming* const StreamingLevel = SceneRigSubsystem->FindFirstActiveSceneRig();
+					if (IsValid(StreamingLevel))
+					{
+						TArray<AActor*> SelectedActors;
+
+						const UAvaOutlinerItemsContext* const ItemsContext = ItemsContextWeak.Get();
+						if (IsValid(ItemsContext))
+						{
+							const TConstArrayView<FAvaOutlinerItemWeakPtr> SelectedItems = ItemsContext->GetItems();
+							SelectedActors = OutlinerItemsToActors(ItemsContext->GetItems(), true);
+						}
+
+						if (!SelectedActors.IsEmpty())
+						{
+							NewSection.AddMenuEntryWithCommandList(SceneRigEditorCommands.RemoveOutlinerItemsToSceneRig
+								, CommandList
+								, LOCTEXT("RemoveFromSceneRig", "Remove from Scene Rig")
+								, SceneRigEditorCommands.RemoveOutlinerItemsToSceneRig->GetDescription()
+								, FSlateIcon(FAppStyle::GetAppStyleSetName(), TEXT("MeshPaint.NextTexture.Small")));
+
+							NewSection.AddMenuEntryWithCommandList(SceneRigEditorCommands.AddOutlinerItemsToSceneRig
+								, CommandList
+								, LOCTEXT("MoveToSceneRig", "Move to Scene Rig")
+								, SceneRigEditorCommands.AddOutlinerItemsToSceneRig->GetDescription()
+								, FSlateIcon(FAppStyle::GetAppStyleSetName(), TEXT("MeshPaint.PreviousTexture.Small")));
+
+							NewSection.AddSeparator(NAME_None);
+						}
+					}
+				}
+			}
+
+			NewSection.AddMenuEntryWithCommandList(SceneRigEditorCommands.PromptToSaveSceneRigFromOutlinerItems
+				, CommandList
+				, SceneRigEditorCommands.PromptToSaveSceneRigFromOutlinerItems->GetLabel()
+				, SceneRigEditorCommands.PromptToSaveSceneRigFromOutlinerItems->GetDescription()
+				, FSlateIcon(FAppStyle::GetAppStyleSetName(), TEXT("AssetEditor.SaveAssetAs")));
+		}));
+}
+
+void FAvaSceneRigExtension::PromptToSaveSceneRigFromOutlinerItems()
+{
+	if (!CanSaveSceneRigFromOutlinerItems())
+	{
+		return;
+	}
+
+	const UAvaOutlinerItemsContext* const ItemsContext = ItemsContextWeak.Get();
+	if (!IsValid(ItemsContext))
+	{
+		return;
+	}
+
+	UWorld* const World = GetWorld();
+	if (!IsValid(World))
+	{
+		return;
+	}
+
+	const TArray<AActor*> SelectedActors = OutlinerItemsToActors(ItemsContext->GetItems());
+	if (SelectedActors.IsEmpty())
+	{
+		return;
+	}
+
+	IAvaSceneRigEditorModule& SceneRigEditorModule = IAvaSceneRigEditorModule::Get();
+	
+	const FSoftObjectPath NewAssetPath = SceneRigEditorModule.CreateSceneRigAssetWithDialog();
+	if (NewAssetPath.IsValid())
+	{
+		if (IsValid(SceneRigEditorModule.SetActiveSceneRig(World, NewAssetPath)))
+		{
+			SceneRigEditorModule.AddActiveSceneRigActors(World, SelectedActors);
+		}
+	}
+}
+
+bool FAvaSceneRigExtension::CanSaveSceneRigFromOutlinerItems() const
+{
+	TArray<AActor*> SelectedActors;
+
+	if (const TObjectPtr<UAvaOutlinerItemsContext> ItemsContext = ItemsContextWeak.Get())
+	{
+		SelectedActors = OutlinerItemsToActors(ItemsContext->GetItems());
+	}
+
+	return !SelectedActors.IsEmpty() && UAvaSceneRigSubsystem::AreActorsSupported(SelectedActors);
+}
+
+void FAvaSceneRigExtension::AddOutlinerItemsToSceneRig()
+{
+	if (!CanAddOutlinerItemsToSceneRig())
+	{
+		return;
+	}
+
+	const UAvaOutlinerItemsContext* const ItemsContext = ItemsContextWeak.Get();
+	if (!IsValid(ItemsContext))
+	{
+		return;
+	}
+
+	const TArray<AActor*> SelectedActors = OutlinerItemsToActors(ItemsContext->GetItems());
+	if (SelectedActors.IsEmpty())
+	{
+		return;
+	}
+
+	UWorld* const World = GetWorld();
+	if (!IsValid(World))
+	{
+		return;
+	}
+
+	UAvaSceneRigSubsystem* const SceneRigSubsystem = UAvaSceneRigSubsystem::ForWorld(World);
+	if (!IsValid(SceneRigSubsystem))
+	{
+		return;
+	}
+
+	IAvaSceneRigEditorModule& SceneRigEditorModule = IAvaSceneRigEditorModule::Get();
+
+	ULevelStreaming* const SceneRig = SceneRigSubsystem->FindFirstActiveSceneRig();
+
+	if (IsValid(SceneRig))
+	{
+		FScopedTransaction Transaction(LOCTEXT("AddSceneRig", "Add Scene Rig Actors"));
+		SceneRig->Modify();
+
+		SceneRigEditorModule.AddActiveSceneRigActors(World, SelectedActors);
+	}
+	else
+	{
+		const FSoftObjectPath NewAssetPath = SceneRigEditorModule.CreateSceneRigAssetWithDialog();
+		if (NewAssetPath.IsValid())
+		{
+			if (IsValid(SceneRigEditorModule.SetActiveSceneRig(World, NewAssetPath)))
+			{
+				SceneRigEditorModule.AddActiveSceneRigActors(World, SelectedActors);
+			}
+		}
+	}
+}
+
+bool FAvaSceneRigExtension::CanAddOutlinerItemsToSceneRig() const
+{
+	const UAvaOutlinerItemsContext* const ItemsContext = ItemsContextWeak.Get();
+	if (!IsValid(ItemsContext))
+	{
+		return false;
+	}
+
+	const TArray<AActor*> SelectedActors = OutlinerItemsToActors(ItemsContext->GetItems());
+	if (SelectedActors.IsEmpty() || !UAvaSceneRigSubsystem::AreActorsSupported(SelectedActors))
+	{
+		return false;
+	}
+
+	const UAvaSceneRigSubsystem* const SceneRigSubsystem = UAvaSceneRigSubsystem::ForWorld(GetWorld());
+	if (!IsValid(SceneRigSubsystem))
+	{
+		return false;
+	}
+
+	UWorld* const SceneRigAsset = SceneRigSubsystem->FindFirstActiveSceneRigAsset();
+	if (!IsValid(SceneRigAsset) || !IsValid(SceneRigAsset->PersistentLevel))
+	{
+		return false;
+	}
+
+	return !UAvaSceneRigSubsystem::AreAllActorsInLevel(SceneRigAsset->PersistentLevel, SelectedActors);
+}
+
+void FAvaSceneRigExtension::RemoveOutlinerItemsFromSceneRig()
+{
+	if (!CanRemoveOutlinerItemsFromSceneRig())
+	{
+		return;
+	}
+
+	const UAvaOutlinerItemsContext* const ItemsContext = ItemsContextWeak.Get();
+	if (!IsValid(ItemsContext))
+	{
+		return;
+	}
+
+	UWorld* const World = GetWorld();
+	if (!IsValid(World))
+	{
+		return;
+	}
+	
+	UAvaSceneRigSubsystem* const SceneRigSubsystem = UAvaSceneRigSubsystem::ForWorld(World);
+	if (!IsValid(SceneRigSubsystem))
+	{
+		return;
+	}
+
+	ULevelStreaming* const SceneRig = SceneRigSubsystem->FindFirstActiveSceneRig();
+	if (!IsValid(SceneRig))
+	{
+		return;
+	}
+
+	const TArray<AActor*> SelectedActors = OutlinerItemsToActors(ItemsContext->GetItems());
+
+	FScopedTransaction Transaction(LOCTEXT("RemoveSceneRig", "Remove Scene Rig"));
+	SceneRig->Modify();
+
+	IAvaSceneRigEditorModule::Get().RemoveActiveSceneRigActors(World, SelectedActors);
+}
+
+bool FAvaSceneRigExtension::CanRemoveOutlinerItemsFromSceneRig() const
+{
+	const UAvaOutlinerItemsContext* const ItemsContext = ItemsContextWeak.Get();
+	if (!IsValid(ItemsContext))
+	{
+		return false;
+	}
+
+	const TArray<AActor*> SelectedActors = OutlinerItemsToActors(ItemsContext->GetItems());
+
+	if (SelectedActors.IsEmpty() || !UAvaSceneRigSubsystem::AreActorsSupported(SelectedActors))
+	{
+		return false;
+	}
+
+	const UAvaSceneRigSubsystem* const SceneRigSubsystem = UAvaSceneRigSubsystem::ForWorld(GetWorld());
+	if (!IsValid(SceneRigSubsystem))
+	{
+		return false;
+	}
+
+	UWorld* const SceneRigAsset = SceneRigSubsystem->FindFirstActiveSceneRigAsset();
+	if (!IsValid(SceneRigAsset) || !IsValid(SceneRigAsset->PersistentLevel))
+	{
+		return false;
+	}
+
+	return UAvaSceneRigSubsystem::AreSomeActorsInLevel(SceneRigAsset->PersistentLevel, SelectedActors);
+}
+
+TArray<AActor*> FAvaSceneRigExtension::OutlinerItemsToActors(const TConstArrayView<FAvaOutlinerItemWeakPtr>& InOutlinerItems
+	, const bool bInIncludeLocked)
+{
+	TArray<AActor*> OutActors;
+
+	for (const FAvaOutlinerItemWeakPtr& ItemWeak : InOutlinerItems)
+	{
+		if (const FAvaOutlinerItemPtr& Item = ItemWeak.Pin())
+		{
+			if (const FAvaOutlinerActor* const ActorItem = Item->CastTo<FAvaOutlinerActor>())
+			{
+				if (bInIncludeLocked || !ActorItem->IsLocked())
+				{
+					AActor* const Actor = ActorItem->GetActor();
+					if (IsValid(Actor))
+					{
+						OutActors.Add(Actor);
+					}
+				}
+			}
+		}
+	}
+
+	return OutActors;
+}
+
+UAvaSceneState* FAvaSceneRigExtension::GetSceneState(UWorld* const InWorld)
+{
+	if (!IsValid(InWorld))
+	{
+		return nullptr;
+	}
+
+	const UAvaSceneSubsystem* const SceneSubsystem = InWorld->GetSubsystem<UAvaSceneSubsystem>();
+	if (!IsValid(SceneSubsystem))
+	{
+		return nullptr;
+	}
+
+	IAvaSceneInterface* const AvaScene = SceneSubsystem->GetSceneInterface();
+	if (!AvaScene)
+	{
+		return nullptr;
+	}
+
+	return AvaScene->GetSceneState();
+}
+
+#undef LOCTEXT_NAMESPACE
diff --git a/Engine/Plugins/Experimental/Avalanche/Source/AvalancheEditor/Private/Scene/AvaSceneRigExtension.h b/Engine/Plugins/Experimental/Avalanche/Source/AvalancheEditor/Private/Scene/AvaSceneRigExtension.h
new file mode 100644
--- /dev/null
+++ b/Engine/Plugins/Experimental/Avalanche/Source/AvalancheEditor/Private/Scene/AvaSceneRigExtension.h
@@ -0,0 +1,52 @@
+// Copyright Epic Games, Inc. All Rights Reserved.
+
+#pragma once
+
+#include "AvaOutlinerDefines.h"
+#include "Containers/ArrayView.h"
+#include "IAvaEditorExtension.h"
+#include "IAvaOutlinerModule.h"
+
+class AActor;
+class FUICommandList;
+class IAvaEditor;
+class UAvaOutlinerItemsContext;
+class UAvaSceneState;
+class UToolMenu;
+
+class FAvaSceneRigExtension : public FAvaEditorExtension
+{
+public:
+	UE_AVA_INHERITS(FAvaSceneRigExtension, FAvaEditorExtension);
+
+	static UAvaSceneState* GetSceneState(UWorld* const InWorld);
+
+	void PromptToSaveSceneRigFromOutlinerItems();
+	bool CanSaveSceneRigFromOutlinerItems() const;
+
+	void AddOutlinerItemsToSceneRig();
+	bool CanAddOutlinerItemsToSceneRig() const;
+
+	void RemoveOutlinerItemsFromSceneRig();
+	bool CanRemoveOutlinerItemsFromSceneRig() const;
+
+	//~ Begin IAvaEditorExtension
+	virtual void Activate() override;
+	virtual void Deactivate() override;
+	virtual void BindCommands(const TSharedRef<FUICommandList>& InCommandList) override;
+	//~ End IAvaEditorExtension
+
+protected:
+	static TArray<AActor*> OutlinerItemsToActors(const TConstArrayView<FAvaOutlinerItemWeakPtr>& InOutlinerItems
+		, const bool bInIncludeLocked = false);
+
+	void ExtendOutlinerItemContextMenu(UToolMenu* const InToolMenu);
+
+	void CreateSubMenu(UToolMenu* const InToolMenu);
+
+	FDelegateHandle OutlinerItemContextMenuDelegate;
+
+	TWeakPtr<FUICommandList> CommandListWeak;
+
+	TWeakObjectPtr<UAvaOutlinerItemsContext> ItemsContextWeak;
+};
diff --git a/Engine/Plugins/Experimental/Avalanche/Source/AvalancheOutliner/Private/AvaOutliner.cpp b/Engine/Plugins/Experimental/Avalanche/Source/AvalancheOutliner/Private/AvaOutliner.cpp
--- a/Engine/Plugins/Experimental/Avalanche/Source/AvalancheOutliner/Private/AvaOutliner.cpp
+++ b/Engine/Plugins/Experimental/Avalanche/Source/AvalancheOutliner/Private/AvaOutliner.cpp
@@ -14,6 +14,8 @@
 #include "Data/AvaOutlinerVersion.h"
 #include "Editor/Transactor.h"
 #include "EditorModeManager.h"
+#include "Engine/Level.h"
+#include "Engine/LevelStreaming.h"
 #include "Engine/World.h"
 #include "EngineUtils.h"
 #include "Framework/MultiBox/MultiBoxBuilder.h"
@@ -100,9 +102,14 @@ bool FAvaOutliner::IsActorAllowedInOutliner(const AActor* InActor) const
 		return false;
 	}
 
-	// Only consider Actors that are in the Outliner World
-	const UWorld* const World = InActor->GetWorld();
-	if (!IsValid(World) || World != GetWorld())
+	const UWorld* const ActorWorld = InActor->GetWorld();
+	if (!IsValid(ActorWorld))
+	{
+		return false;
+	}
+
+	const UWorld* const OutlinerWorld = GetWorld();
+	if (!IsValid(OutlinerWorld))
 	{
 		return false;
 	}
@@ -113,11 +120,41 @@ bool FAvaOutliner::IsActorAllowedInOutliner(const AActor* InActor) const
 		return false;
 	}
 
+	// Only consider Actors that are in the Outliner World
+	if (ActorWorld != OutlinerWorld)
+	{
+		return false;
+	}
+
+	// Check if actor is in a streaming level
+	for (ULevelStreaming* const LevelStreaming : OutlinerWorld->GetStreamingLevels())
+	{
+		if (const TSoftObjectPtr<UWorld>& WorldAsset = LevelStreaming->GetWorldAsset())
+		{
+			if (WorldAsset && WorldAsset->PersistentLevel
+				&& WorldAsset->PersistentLevel->Actors.Contains(InActor))
+			{
+				return !IsDefaultWorldActorToHide(WorldAsset.Get(), InActor);
+			}
+		}
+	}
+
+	// Make sure the Actor is none of these Default World Actors
+	return !IsDefaultWorldActorToHide(ActorWorld, InActor);
+}
+
+bool FAvaOutliner::IsDefaultWorldActorToHide(const UWorld* const InWorld, const AActor* const InActor) const
+{
+	if (!IsValid(InWorld) || !IsValid(InActor))
+	{
+		return false;
+	}
+
 	// Make sure the Actor is none of these Default World Actors
-	return InActor != World->GetDefaultPhysicsVolume()
-		&& InActor != World->GetDefaultBrush()
-		&& InActor != World->GetWorldSettings()
-		&& InActor != World->MyParticleEventManager;
+	return InActor == InWorld->GetDefaultPhysicsVolume()
+		|| InActor == InWorld->GetDefaultBrush()
+		|| InActor == InWorld->GetWorldSettings()
+		|| InActor == InWorld->MyParticleEventManager;
 }
 
 bool FAvaOutliner::IsComponentAllowedInOutliner(const USceneComponent* InComponent) const
@@ -344,33 +381,55 @@ void FAvaOutliner::Refresh()
 	SceneOutlinerParentMap.Reset();
 	if (UWorld* const World = GetWorld())
 	{
-		ULevelInstanceSubsystem* const LevelInstanceSubsystem = World->GetSubsystem<ULevelInstanceSubsystem>();
+		auto AddActorToParentMap = [this](UWorld* const InWorld, AActor* const InActor)
+			{
+				if (!IsValid(InWorld) || !IsValid(InActor))
+				{
+					return;
+				}
+
+				ULevelInstanceSubsystem* const LevelInstanceSubsystem = InWorld->GetSubsystem<ULevelInstanceSubsystem>();
+
+				const ULevel* Level = InActor->GetLevel();
+				AActor* Parent = InActor->GetSceneOutlinerParent();
+
+				// Try to find the Level Instance Actor to use as Parent for actors that aren't attached to anything
+				// and belong to sub-levels
+				if (!Parent && Level != InWorld->PersistentLevel)
+				{
+					if (const ILevelInstanceInterface* LevelInstance = LevelInstanceSubsystem->GetOwningLevelInstance(Level))
+					{
+						if (const ULevelInstanceComponent* LevelInstanceComponent = LevelInstance->GetLevelInstanceComponent())
+						{
+							Parent = LevelInstanceComponent->GetOwner();
+						}
+					}
+				}
+
+				SceneOutlinerParentMap.FindOrAdd(Parent).AddUnique(InActor);
+				FindOrAdd<FAvaOutlinerActor>(InActor);
+			};
 
 		// 1) Update the Scene Outliner Parent Map for its Parent to know this Actor
 		// 2) Make sure this Actor has an assigned Outliner Item
 		for (AActor* const Actor : TActorRange<AActor>(World))
 		{
-			const ULevel* Level  = Actor->GetLevel();
-			AActor* const Parent = Actor->GetSceneOutlinerParent();
+			AddActorToParentMap(World, Actor);
+		}
 
-			// Try to find the Level Instance Actor to use as Parent for actors that aren't attached to anything
-			// and belong to sub-levels
-			if (!Parent && Level != World->PersistentLevel)
+		// Add all actors from streaming levels
+		for (ULevelStreaming* const LevelStreaming : World->GetStreamingLevels())
+		{
+			if (const TSoftObjectPtr<UWorld>& WorldAsset = LevelStreaming->GetWorldAsset())
 			{
-				if (const ILevelInstanceInterface* LevelInstance = LevelInstanceSubsystem->GetOwningLevelInstance(Level))
+				if (WorldAsset && WorldAsset->PersistentLevel)
 				{
-					if (const ULevelInstanceComponent* LevelInstanceComponent = LevelInstance->GetLevelInstanceComponent())
+					for (AActor* const Actor : WorldAsset->PersistentLevel->Actors)
 					{
-						SceneOutlinerParentMap.FindOrAdd(LevelInstanceComponent->GetOwner()).Add(Actor);
+						AddActorToParentMap(WorldAsset.Get(), Actor);
 					}
 				}
 			}
-			else
-			{
-				SceneOutlinerParentMap.FindOrAdd(Parent).Add(Actor);
-			}
-
-			FindOrAdd<FAvaOutlinerActor>(Actor);
 		}
 	}
 
diff --git a/Engine/Plugins/Experimental/Avalanche/Source/AvalancheOutliner/Private/AvaOutliner.h b/Engine/Plugins/Experimental/Avalanche/Source/AvalancheOutliner/Private/AvaOutliner.h
--- a/Engine/Plugins/Experimental/Avalanche/Source/AvalancheOutliner/Private/AvaOutliner.h
+++ b/Engine/Plugins/Experimental/Avalanche/Source/AvalancheOutliner/Private/AvaOutliner.h
@@ -13,7 +13,9 @@
 #include "ItemProxies/AvaOutlinerItemProxyRegistry.h"
 #include "ItemProxies/IAvaOutlinerItemProxyFactory.h"
 #include "TickableEditorObject.h"
+#include "UObject/WeakObjectPtrTemplates.h"
 
+class AActor;
 class FAvaEditorSelection;
 class FAvaOutlinerTreeRoot;
 class FAvaOutlinerView;
@@ -24,6 +26,7 @@ class IAvaOutlinerAction;
 class IAvaOutlinerProvider;
 class IAvaOutlinerView;
 class UAvaOutlinerSubsystem;
+class ULevelStreaming;
 enum class EItemDropZone;
 struct FAttachmentTransformRules;
 struct FAvaOutlinerSaveState;
@@ -101,6 +104,7 @@ public:
 	virtual void OnObjectSelectionChanged(const FAvaEditorSelection& InEditorSelection) override;
 	virtual UWorld* GetWorld() const override;
 	virtual const FAvaOutlinerItemProxyRegistry& GetItemProxyRegistry() const override;
+	virtual bool IsDefaultWorldActorToHide(const UWorld* const InWorld, const AActor* const InActor) const override;
 	//~ End IAvaOutliner
 
 	FAvaOutlinerItemProxyRegistry& GetItemProxyRegistry();
diff --git a/Engine/Plugins/Experimental/Avalanche/Source/AvalancheOutliner/Private/AvaOutlinerView.h b/Engine/Plugins/Experimental/Avalanche/Source/AvalancheOutliner/Private/AvaOutlinerView.h
--- a/Engine/Plugins/Experimental/Avalanche/Source/AvalancheOutliner/Private/AvaOutlinerView.h
+++ b/Engine/Plugins/Experimental/Avalanche/Source/AvalancheOutliner/Private/AvaOutlinerView.h
@@ -182,6 +182,8 @@ public:
 	
 	const TArray<TSharedPtr<IAvaOutlinerItemFilter>>& GetCustomItemFilters() const { return CustomItemFilters; }
 
+	const TSet<TSharedPtr<IAvaOutlinerItemFilter>>& GetActiveItemFilters() const { return ActiveItemFilters; }
+
 	TSharedRef<FAvaOutlinerTextFilter> GetTextFilter() const { return TextFilter; }
 	
 	/** Gets the Currently Selected Items in the Tree View */
diff --git a/Engine/Plugins/Experimental/Avalanche/Source/AvalancheOutliner/Private/Filters/AvaOutlinerItemTypeFilter.cpp b/Engine/Plugins/Experimental/Avalanche/Source/AvalancheOutliner/Private/Filters/AvaOutlinerItemTypeFilter.cpp
--- a/Engine/Plugins/Experimental/Avalanche/Source/AvalancheOutliner/Private/Filters/AvaOutlinerItemTypeFilter.cpp
+++ b/Engine/Plugins/Experimental/Avalanche/Source/AvalancheOutliner/Private/Filters/AvaOutlinerItemTypeFilter.cpp
@@ -127,6 +127,11 @@ bool FAvaOutlinerItemTypeFilterData::PassesFilterText(const FAvaOutlinerObject&
 
 bool FAvaOutlinerItemTypeFilter::PassesFilter(FAvaOutlinerFilterType InItem) const
 {
+	if (FilterPassDelegate.IsBound())
+	{
+		return FilterPassDelegate.Execute(InItem);
+	}
+
 	auto ItemPassesFilter = [this](FAvaOutlinerFilterType InItem)
 	{
 		if (const FAvaOutlinerObject* ObjectItem = InItem.CastTo<FAvaOutlinerObject>())
diff --git a/Engine/Plugins/Experimental/Avalanche/Source/AvalancheOutliner/Private/Selection/AvaOutlinerScopedSelection.cpp b/Engine/Plugins/Experimental/Avalanche/Source/AvalancheOutliner/Private/Selection/AvaOutlinerScopedSelection.cpp
--- a/Engine/Plugins/Experimental/Avalanche/Source/AvalancheOutliner/Private/Selection/AvaOutlinerScopedSelection.cpp
+++ b/Engine/Plugins/Experimental/Avalanche/Source/AvalancheOutliner/Private/Selection/AvaOutlinerScopedSelection.cpp
@@ -136,3 +136,8 @@ void FAvaOutlinerScopedSelection::SyncSelections()
 	UE::AvaOutliner::Private::SyncSelection(EditorModeTools.GetSelectedComponents(), SelectedComponents);
 	UE::AvaOutliner::Private::SyncSelection(EditorModeTools.GetSelectedObjects(), SelectedObjects);
 }
+
+const FEditorModeTools& FAvaOutlinerScopedSelection::GetEditorModeTools() const
+{
+	return EditorModeTools;
+}
diff --git a/Engine/Plugins/Experimental/Avalanche/Source/AvalancheOutliner/Public/Filters/AvaOutlinerItemTypeFilter.h b/Engine/Plugins/Experimental/Avalanche/Source/AvalancheOutliner/Public/Filters/AvaOutlinerItemTypeFilter.h
--- a/Engine/Plugins/Experimental/Avalanche/Source/AvalancheOutliner/Public/Filters/AvaOutlinerItemTypeFilter.h
+++ b/Engine/Plugins/Experimental/Avalanche/Source/AvalancheOutliner/Public/Filters/AvaOutlinerItemTypeFilter.h
@@ -80,6 +80,8 @@ private:
 	EClassFlags RestrictedClassFlags = CLASS_None;
 };
 
+DECLARE_DELEGATE_RetVal_OneParam(bool, FAvaOutlinerItemTypeFilterPassDelegate, FAvaOutlinerFilterType);
+
 class AVALANCHEOUTLINER_API FAvaOutlinerItemTypeFilter : public IAvaOutlinerItemFilter
 {
 public:
@@ -101,6 +103,16 @@ public:
 	{
 	}
 
+	explicit FAvaOutlinerItemTypeFilter(FName InFilterId
+			, const FAvaOutlinerItemTypeFilterPassDelegate& InFilterPassDelegate
+			, const FSlateBrush* InIconBrush = nullptr
+			, const FText& InTooltipText = FText::GetEmpty())
+		: FilterId(InFilterId)
+		, FilterData({}, EAvaOutlinerTypeFilterMode::MatchesType, InIconBrush, InTooltipText, CLASS_None, CLASS_None)
+		, FilterPassDelegate(InFilterPassDelegate)
+	{
+	}
+
 protected:
 	//~ Begin IFilter
 	virtual FChangedEvent& OnChanged() override { return ChangedEvent; }
@@ -118,4 +130,6 @@ protected:
 	FAvaOutlinerItemTypeFilterData FilterData;
 
 	FChangedEvent ChangedEvent;
+
+	FAvaOutlinerItemTypeFilterPassDelegate FilterPassDelegate;
 };
diff --git a/Engine/Plugins/Experimental/Avalanche/Source/AvalancheOutliner/Public/IAvaOutliner.h b/Engine/Plugins/Experimental/Avalanche/Source/AvalancheOutliner/Public/IAvaOutliner.h
--- a/Engine/Plugins/Experimental/Avalanche/Source/AvalancheOutliner/Public/IAvaOutliner.h
+++ b/Engine/Plugins/Experimental/Avalanche/Source/AvalancheOutliner/Public/IAvaOutliner.h
@@ -192,6 +192,9 @@ public:
 	/** Gets the World the Outliner is working with */
 	virtual UWorld* GetWorld() const = 0;
 
+	/** Returns true if the given actor is a default actor of the given world that will be hidden from the outliner. */
+	virtual bool IsDefaultWorldActorToHide(const UWorld* const InWorld, const AActor* const InActor) const = 0;
+
 protected:
 	virtual const FAvaOutlinerItemProxyRegistry& GetItemProxyRegistry() const = 0;
 };
diff --git a/Engine/Plugins/Experimental/Avalanche/Source/AvalancheOutliner/Public/Selection/AvaOutlinerScopedSelection.h b/Engine/Plugins/Experimental/Avalanche/Source/AvalancheOutliner/Public/Selection/AvaOutlinerScopedSelection.h
--- a/Engine/Plugins/Experimental/Avalanche/Source/AvalancheOutliner/Public/Selection/AvaOutlinerScopedSelection.h
+++ b/Engine/Plugins/Experimental/Avalanche/Source/AvalancheOutliner/Public/Selection/AvaOutlinerScopedSelection.h
@@ -29,6 +29,8 @@ public:
 
 	AVALANCHEOUTLINER_API bool IsSelected(const UObject* InObject) const;
 
+	AVALANCHEOUTLINER_API const FEditorModeTools& GetEditorModeTools() const;
+
 private:
 	void SyncSelections();
 
diff --git a/Engine/Plugins/Experimental/Avalanche/Source/AvalancheSceneRig/AvalancheSceneRig.Build.cs b/Engine/Plugins/Experimental/Avalanche/Source/AvalancheSceneRig/AvalancheSceneRig.Build.cs
new file mode 100644
--- /dev/null
+++ b/Engine/Plugins/Experimental/Avalanche/Source/AvalancheSceneRig/AvalancheSceneRig.Build.cs
@@ -0,0 +1,27 @@
+// Copyright Epic Games, Inc. All Rights Reserved.
+
+using UnrealBuildTool;
+
+public class AvalancheSceneRig : ModuleRules
+{
+	public AvalancheSceneRig(ReadOnlyTargetRules Target) : base(Target)
+	{
+		PCHUsage = ModuleRules.PCHUsageMode.UseExplicitOrSharedPCHs;
+
+		PublicDependencyModuleNames.AddRange(
+			new string[]
+			{
+				"Core",
+				"CoreUObject"
+			}
+		);
+
+		PrivateDependencyModuleNames.AddRange(
+			new string[]
+			{
+				"AssetRegistry",
+				"Engine"
+			}
+		);
+	}
+}
diff --git a/Engine/Plugins/Experimental/Avalanche/Source/AvalancheSceneRig/Private/AvaSceneRigModule.cpp b/Engine/Plugins/Experimental/Avalanche/Source/AvalancheSceneRig/Private/AvaSceneRigModule.cpp
new file mode 100644
--- /dev/null
+++ b/Engine/Plugins/Experimental/Avalanche/Source/AvalancheSceneRig/Private/AvaSceneRigModule.cpp
@@ -0,0 +1,5 @@
+// Copyright Epic Games, Inc. All Rights Reserved.
+
+#include "Modules/ModuleManager.h"
+
+IMPLEMENT_MODULE(FDefaultModuleImpl, AvalancheSceneRig)
diff --git a/Engine/Plugins/Experimental/Avalanche/Source/AvalancheSceneRig/Private/AvaSceneRigSubsystem.cpp b/Engine/Plugins/Experimental/Avalanche/Source/AvalancheSceneRig/Private/AvaSceneRigSubsystem.cpp
new file mode 100644
--- /dev/null
+++ b/Engine/Plugins/Experimental/Avalanche/Source/AvalancheSceneRig/Private/AvaSceneRigSubsystem.cpp
@@ -0,0 +1,267 @@
+// Copyright Epic Games, Inc. All Rights Reserved.
+
+#include "AvaSceneRigSubsystem.h"
+#include "AssetRegistry/AssetRegistryModule.h"
+#include "Containers/Array.h"
+#include "Engine/Level.h"
+#include "Engine/LevelStreaming.h"
+#include "Engine/LevelStreamingDynamic.h"
+#include "Engine/World.h"
+#include "Interfaces/Interface_AssetUserData.h"
+#include "UObject/Class.h"
+#include "UObject/SoftObjectPtr.h"
+#include "UObject/NoExportTypes.h"
+
+DEFINE_LOG_CATEGORY(AvaSceneRigSubsystemLog);
+
+#define LOCTEXT_NAMESPACE "AvaSceneRigSubsystem"
+
+void UAvaSceneRigSubsystem::Initialize(FSubsystemCollectionBase& Collection)
+{
+	Super::Initialize(Collection);
+}
+
+void UAvaSceneRigSubsystem::Deinitialize()
+{
+	Super::Deinitialize();
+}
+
+UAvaSceneRigSubsystem* UAvaSceneRigSubsystem::ForWorld(const UWorld* const InWorld)
+{
+	if (IsValid(InWorld))
+	{
+		return InWorld->GetSubsystem<UAvaSceneRigSubsystem>();
+	}
+	return nullptr;
+}
+
+bool UAvaSceneRigSubsystem::ShouldCreateSubsystem(UObject* const InOuter) const
+{
+	if (!IsValid(InOuter))
+	{
+		return false;
+	}
+
+	const FAssetData AssetData(InOuter);
+
+	// Only create scene rig subsystems for Motion Design scenes
+	const FAssetDataTagMapSharedView::FFindTagResult SceneTagResult = AssetData.TagsAndValues.FindTag(TEXT("MotionDesignScene"));
+	return SceneTagResult.IsSet() && SceneTagResult.GetValue().Equals(TEXT("Enabled"));
+}
+
+bool UAvaSceneRigSubsystem::IsSceneRigAssetData(const FAssetData& InAssetData)
+{
+	if (!InAssetData.IsValid())
+	{
+		return false;
+	}
+
+	UWorld* const SceneRigAsset = Cast<UWorld>(InAssetData.GetAsset());
+	if (!IsValid(SceneRigAsset) || !IsValid(SceneRigAsset->PersistentLevel))
+	{
+		return false;
+	}
+
+	if (IInterface_AssetUserData* AssetUserData = Cast<IInterface_AssetUserData>(SceneRigAsset->PersistentLevel))
+	{
+		UAvaSceneRigData* const SceneRigData = AssetUserData->GetAssetUserData<UAvaSceneRigData>();
+		if (IsValid(SceneRigData))
+		{
+			return true;
+		}
+	}
+
+	return false;
+}
+
+bool UAvaSceneRigSubsystem::IsSceneRigAsset(UObject* const InObject)
+{
+	return IsSceneRigAssetData(FAssetData(InObject));
+}
+
+FString UAvaSceneRigSubsystem::GetSceneRigAssetSuffix()
+{
+	static const TCHAR* Suffix = TEXT("_SceneRig");
+	return Suffix;
+}
+
+TSet<FName> UAvaSceneRigSubsystem::GetSupportedActorClassNames()
+{
+	static const TSet<FName> ClassNames =
+		{
+			// Cameras
+			TEXT("CameraActor"),
+			TEXT("CineCameraActor"),
+			TEXT("AvaCineCameraActor"),
+			// Lights
+			TEXT("SkyLight"),
+			TEXT("DirectionalLight"),
+			TEXT("PointLight"),
+			TEXT("RectLight"),
+			TEXT("SpotLight"),
+			// Misc
+			TEXT("AvaNullActor"),
+			TEXT("PostProcessVolume")
+		};
+	return ClassNames;
+}
+
+bool UAvaSceneRigSubsystem::IsSupportedActorClassName(const FName InName)
+{
+	return GetSupportedActorClassNames().Contains(InName);
+}
+
+bool UAvaSceneRigSubsystem::IsSupportedActorClass(const UClass* InClass)
+{
+	if (IsValid(InClass))
+	{
+		return GetSupportedActorClassNames().Contains(InClass->GetFName());
+	}
+	return false;
+}
+
+bool UAvaSceneRigSubsystem::AreActorsSupported(const TArray<AActor*>& InActors)
+{
+	bool bAllItemClassesSupported = !InActors.IsEmpty();
+
+	for (const AActor* const Actor : InActors)
+	{
+		if (!UAvaSceneRigSubsystem::IsSupportedActorClass(Actor->GetClass()))
+		{
+			bAllItemClassesSupported = false;
+			break;
+		}
+	}
+
+	return bAllItemClassesSupported;
+}
+
+ULevelStreaming* UAvaSceneRigSubsystem::SceneRigFromActor(AActor* const InActor)
+{
+	const UWorld* const World = InActor->GetLevel()->GetWorld();
+	if (!IsValid(World))
+	{
+		return nullptr;
+	}
+
+	const UAvaSceneRigSubsystem* const SceneRigSubsystem = ForWorld(World);
+	if (!IsValid(SceneRigSubsystem))
+	{
+		return nullptr;
+	}
+
+	return SceneRigSubsystem->FindFirstActiveSceneRig();
+}
+
+bool UAvaSceneRigSubsystem::AreAllActorsInLevel(ULevel* const InLevel, const TArray<AActor*>& InActors)
+{
+	if (!IsValid(InLevel))
+	{
+		return false;
+	}
+
+	for (const AActor* const Actor : InActors)
+	{
+		if (!InLevel->Actors.Contains(Actor))
+		{
+			return false;
+		}
+	}
+
+	return true;
+}
+
+bool UAvaSceneRigSubsystem::AreSomeActorsInLevel(ULevel* const InLevel, const TArray<AActor*>& InActors)
+{
+	if (!IsValid(InLevel))
+	{
+		return false;
+	}
+
+	for (const AActor* const Actor : InActors)
+	{
+		if (InLevel->Actors.Contains(Actor))
+		{
+			return true;
+		}
+	}
+
+	return false;
+}
+
+TArray<ULevelStreaming*> UAvaSceneRigSubsystem::FindAllSceneRigs() const
+{
+	TArray<ULevelStreaming*> OutStreamingLevels;
+
+	const FAssetRegistryModule& AssetRegistryModule = FModuleManager::Get().LoadModuleChecked<FAssetRegistryModule>(TEXT("AssetRegistry"));
+
+	for (ULevelStreaming* const LevelStreaming : GetWorldRef().GetStreamingLevels())
+	{
+		if (const TSoftObjectPtr<UWorld>& WorldAsset = LevelStreaming->GetWorldAsset())
+		{
+			if (IsSceneRigAsset(WorldAsset.Get()))
+			{
+				OutStreamingLevels.Add(LevelStreaming);
+			}
+		}
+	}
+
+	return OutStreamingLevels;
+}
+
+ULevelStreaming* UAvaSceneRigSubsystem::FindFirstActiveSceneRig() const
+{
+	const TArray<ULevelStreaming*> SceneRigs = FindAllSceneRigs();
+	if (SceneRigs.IsEmpty())
+	{
+		return nullptr;
+	}
+
+	return SceneRigs[0];
+}
+
+UWorld* UAvaSceneRigSubsystem::FindFirstActiveSceneRigAsset() const
+{
+	ULevelStreaming* const ActiveSceneRig = UAvaSceneRigSubsystem::FindFirstActiveSceneRig();
+	if (!IsValid(ActiveSceneRig))
+	{
+		return nullptr;
+	}
+
+	return ActiveSceneRig->GetWorldAsset().Get();
+}
+
+bool UAvaSceneRigSubsystem::IsActiveSceneRigActor(AActor* const InActor) const
+{
+	ULevelStreaming* const SceneRig = FindFirstActiveSceneRig();
+	if (IsValid(SceneRig))
+	{
+		const UWorld* const WorldAsset = SceneRig->GetWorldAsset().Get();
+		if (IsValid(WorldAsset) && IsValid(WorldAsset->PersistentLevel))
+		{
+			return WorldAsset->PersistentLevel->Actors.Contains(InActor);
+		}
+	}
+	return false;
+}
+
+void UAvaSceneRigSubsystem::ForEachActiveSceneRigActor(TFunction<void(AActor* const InActor)> InFunction) const
+{
+	ULevelStreaming* const ActiveSceneRig = FindFirstActiveSceneRig();
+	if (IsValid(ActiveSceneRig))
+	{
+		UWorld* const SceneRigAsset = ActiveSceneRig->GetWorldAsset().Get();
+		if (IsValid(SceneRigAsset) && IsValid(SceneRigAsset->PersistentLevel))
+		{
+			for (AActor* const Actor : SceneRigAsset->PersistentLevel->Actors)
+			{
+				if (IsValid(Actor))
+				{
+					InFunction(Actor);
+				}
+			}
+		}
+	}
+}
+
+#undef LOCTEXT_NAMESPACE
diff --git a/Engine/Plugins/Experimental/Avalanche/Source/AvalancheSceneRig/Public/AvaSceneRigAssetTags.h b/Engine/Plugins/Experimental/Avalanche/Source/AvalancheSceneRig/Public/AvaSceneRigAssetTags.h
new file mode 100644
--- /dev/null
+++ b/Engine/Plugins/Experimental/Avalanche/Source/AvalancheSceneRig/Public/AvaSceneRigAssetTags.h
@@ -0,0 +1,15 @@
+// Copyright Epic Games, Inc. All Rights Reserved.
+
+#pragma once
+
+#include "HAL/Platform.h"
+
+namespace UE::AvaSceneRig::AssetTags
+{
+	constexpr const TCHAR* SceneRig = TEXT("AvalancheSceneRig");
+
+	namespace Values
+	{
+		constexpr const TCHAR* Enabled = TEXT("Enabled");
+	}
+}
diff --git a/Engine/Plugins/Experimental/Avalanche/Source/AvalancheSceneRig/Public/AvaSceneRigSubsystem.h b/Engine/Plugins/Experimental/Avalanche/Source/AvalancheSceneRig/Public/AvaSceneRigSubsystem.h
new file mode 100644
--- /dev/null
+++ b/Engine/Plugins/Experimental/Avalanche/Source/AvalancheSceneRig/Public/AvaSceneRigSubsystem.h
@@ -0,0 +1,91 @@
+// Copyright Epic Games, Inc. All Rights Reserved.
+
+#pragma once
+
+#include "Containers/ContainersFwd.h"
+#include "Delegates/IDelegateInstance.h"
+#include "Engine/AssetUserData.h"
+#include "Subsystems/WorldSubsystem.h"
+#include "Templates/Function.h"
+#include "AvaSceneRigSubsystem.generated.h"
+
+class AActor;
+class FName;
+class UClass;
+class ULevel;
+class ULevelStreaming;
+class UObject;
+class UWorld;
+struct FAssetData;
+
+DECLARE_LOG_CATEGORY_EXTERN(AvaSceneRigSubsystemLog, Log, All);
+
+UCLASS(MinimalAPI)
+class UAvaSceneRigData : public UAssetUserData
+{
+	GENERATED_BODY()
+};
+
+UCLASS(MinimalAPI)
+class UAvaSceneRigSubsystem : public UWorldSubsystem
+{
+	GENERATED_BODY()
+
+public:
+	/** Returns the Scene Rig Subsystem for a world. */
+	AVALANCHESCENERIG_API static UAvaSceneRigSubsystem* ForWorld(const UWorld* const InWorld);
+
+	/** Returns true if the asset data is for a Scene Rig level. */
+	AVALANCHESCENERIG_API static bool IsSceneRigAssetData(const FAssetData& InAssetData);
+
+	/** Returns true if the asset object is Scene Rig level. */
+	AVALANCHESCENERIG_API static bool IsSceneRigAsset(UObject* const InObject);
+
+	/** Returns the Scene Rig level suffix string. Recommended to append to level asset name to differentiate them from other levels. */
+	AVALANCHESCENERIG_API static FString GetSceneRigAssetSuffix();
+
+	/** Get the list of supported object class names that are allowed to be instanced into a scene rig level. */
+	AVALANCHESCENERIG_API static TSet<FName> GetSupportedActorClassNames();
+
+	/** Returns true if the class name is a supported Scene Rig class. */
+	AVALANCHESCENERIG_API static bool IsSupportedActorClassName(const FName InName);
+
+	/** Returns true if the class is a supported Scene Rig class. */
+	AVALANCHESCENERIG_API static bool IsSupportedActorClass(const UClass* InClass);
+
+	/** Returns true if the list of actors are supported Scene Rig classes. */
+	AVALANCHESCENERIG_API static bool AreActorsSupported(const TArray<AActor*>& InActors);
+	
+	/** Returns the active Scene Rig that an actor belongs to. */
+	AVALANCHESCENERIG_API static ULevelStreaming* SceneRigFromActor(AActor* const InActor);
+
+	/** Returns true if all actors exist in the level. */
+	AVALANCHESCENERIG_API static bool AreAllActorsInLevel(ULevel* const InLevel, const TArray<AActor*>& InActors);
+
+	/** Returns true if at least one actor exists in the level. */
+	AVALANCHESCENERIG_API static bool AreSomeActorsInLevel(ULevel* const InLevel, const TArray<AActor*>& InActors);
+
+	/** Find all scene rig streaming levels in the current persistent level. */
+	AVALANCHESCENERIG_API TArray<ULevelStreaming*> FindAllSceneRigs() const;
+
+	/** Returns the first active Scene Rig object in the persistent level. NOTE: There should only be one Scene Rig in a world at any time. */
+	AVALANCHESCENERIG_API ULevelStreaming* FindFirstActiveSceneRig() const;
+
+	/** Returns all Scene Rig asset level paths that exist in the persistent level. NOTE: There should only be one Scene Rig in a world at any time. */
+	AVALANCHESCENERIG_API UWorld* FindFirstActiveSceneRigAsset() const;
+
+	/** Removes a list of actors from the active Scene Rig. */
+	AVALANCHESCENERIG_API bool IsActiveSceneRigActor(AActor* const InActor) const;
+
+	AVALANCHESCENERIG_API void ForEachActiveSceneRigActor(TFunction<void(AActor* const InActor)> InFunction) const;
+
+	//~ Begin UObject
+	virtual void Initialize(FSubsystemCollectionBase& Collection) override;
+	virtual void Deinitialize() override;
+	//~ End UObject
+
+private:
+	//~ Begin WorldSubsystem
+	virtual bool ShouldCreateSubsystem(UObject* const InOuter) const override;
+	//~ End WorldSubsystem
+};
diff --git a/Engine/Plugins/Experimental/Avalanche/Source/AvalancheSceneRigEditor/AvalancheSceneRigEditor.Build.cs b/Engine/Plugins/Experimental/Avalanche/Source/AvalancheSceneRigEditor/AvalancheSceneRigEditor.Build.cs
new file mode 100644
--- /dev/null
+++ b/Engine/Plugins/Experimental/Avalanche/Source/AvalancheSceneRigEditor/AvalancheSceneRigEditor.Build.cs
@@ -0,0 +1,36 @@
+// Copyright Epic Games, Inc. All Rights Reserved.
+
+using UnrealBuildTool;
+
+public class AvalancheSceneRigEditor : ModuleRules
+{
+	public AvalancheSceneRigEditor(ReadOnlyTargetRules Target) : base(Target)
+	{
+		PCHUsage = ModuleRules.PCHUsageMode.UseExplicitOrSharedPCHs;
+
+		PublicDependencyModuleNames.AddRange(
+			new string[]
+			{
+				"Core",
+				"CoreUObject"
+			}
+		);
+
+		PrivateDependencyModuleNames.AddRange(
+			new string[]
+			{
+				"AssetRegistry",
+				"AssetTools",
+				"Avalanche",
+				"AvalancheOutliner",
+				"AvalancheSceneRig",
+				"Engine",
+				"PropertyEditor",
+				"Slate",
+				"SlateCore",
+				"ToolMenus",
+				"UnrealEd"
+			}
+		);
+	}
+}
diff --git a/Engine/Plugins/Experimental/Avalanche/Source/AvalancheSceneRigEditor/Private/AvaSceneRigEditorCommands.cpp b/Engine/Plugins/Experimental/Avalanche/Source/AvalancheSceneRigEditor/Private/AvaSceneRigEditorCommands.cpp
new file mode 100644
--- /dev/null
+++ b/Engine/Plugins/Experimental/Avalanche/Source/AvalancheSceneRigEditor/Private/AvaSceneRigEditorCommands.cpp
@@ -0,0 +1,28 @@
+// Copyright Epic Games, Inc. All Rights Reserved.
+
+#include "AvaSceneRigEditorCommands.h"
+
+#define LOCTEXT_NAMESPACE "AvaSceneRigEditorCommands"
+
+void FAvaSceneRigEditorCommands::RegisterCommands()
+{
+	UI_COMMAND(PromptToSaveSceneRigFromOutlinerItems
+		, "Save as new Scene Rig..."
+		, "Save the selected actors as a new scene rig level asset and sets it as the active scene rig."
+		, EUserInterfaceActionType::Button
+		, FInputChord());
+
+	UI_COMMAND(AddOutlinerItemsToSceneRig
+		, "Move to active Scene Rig"
+		, "Move the selected actors to the active scene rig."
+		, EUserInterfaceActionType::Button
+		, FInputChord());
+
+	UI_COMMAND(RemoveOutlinerItemsToSceneRig
+		, "Remove from active Scene Rig"
+		, "Remove the selected actors from the active scene rig."
+		, EUserInterfaceActionType::Button
+		, FInputChord());
+}
+
+#undef LOCTEXT_NAMESPACE
diff --git a/Engine/Plugins/Experimental/Avalanche/Source/AvalancheSceneRigEditor/Private/AvaSceneRigEditorModule.cpp b/Engine/Plugins/Experimental/Avalanche/Source/AvalancheSceneRigEditor/Private/AvaSceneRigEditorModule.cpp
new file mode 100644
--- /dev/null
+++ b/Engine/Plugins/Experimental/Avalanche/Source/AvalancheSceneRigEditor/Private/AvaSceneRigEditorModule.cpp
@@ -0,0 +1,425 @@
+// Copyright Epic Games, Inc. All Rights Reserved.
+
+#include "AvaSceneRigEditorModule.h"
+#include "AssetRegistry/AssetRegistryModule.h"
+#include "AssetToolsModule.h"
+#include "AvaSceneRigAssetTags.h"
+#include "AvaSceneRigEditorCommands.h"
+#include "AvaSceneRigSubsystem.h"
+#include "AvaSceneSettings.h"
+#include "AvaSceneSubsystem.h"
+#include "ContentBrowserModule.h"
+#include "DetailCategoryBuilder.h"
+#include "DetailLayoutBuilder.h"
+#include "DetailWidgetRow.h"
+#include "Editor.h"
+#include "EditorDirectories.h"
+#include "EditorLevelUtils.h"
+#include "Engine/Level.h"
+#include "Engine/LevelStreaming.h"
+#include "Engine/LevelStreamingDynamic.h"
+#include "Engine/World.h"
+#include "Factories/WorldFactory.h"
+#include "FileHelpers.h"
+#include "GameFramework/Actor.h"
+#include "IAvaOutliner.h"
+#include "IAvaOutlinerModule.h"
+#include "IAvaSceneInterface.h"
+#include "IAvaSceneRigEditorModule.h"
+#include "IContentBrowserSingleton.h"
+#include "Interfaces/Interface_AssetUserData.h"
+#include "Item/AvaOutlinerActor.h"
+#include "Item/AvaOutlinerItemProxy.h"
+#include "ItemProxies/AvaOutlinerItemProxyRegistry.h"
+#include "LevelUtils.h"
+#include "Misc/MessageDialog.h"
+#include "Outliner/AvaOutlinerSceneRigProxy.h"
+#include "Subsystems/EditorAssetSubsystem.h"
+#include "Templates/SharedPointer.h"
+#include "Widgets/SSceneRigPicker.h"
+#include "UObject/GCObjectScopeGuard.h"
+#include "UObject/Object.h"
+#include "UObject/Package.h"
+#include "UObject/SoftObjectPath.h"
+
+DEFINE_LOG_CATEGORY(AvaSceneRigEditorLog);
+
+#define LOCTEXT_NAMESPACE "AvaSceneRigEditorModule"
+
+void FAvaSceneRigEditorModule::StartupModule()
+{
+	FAvaSceneRigEditorCommands::Register();
+
+	if (IAvaOutlinerModule::IsLoaded())
+	{
+		FAvaOutlinerItemProxyRegistry& ItemProxyRegistry = IAvaOutlinerModule::Get().GetItemProxyRegistry();
+		ItemProxyRegistry.RegisterItemProxyWithDefaultFactory<FAvaOutlinerSceneRigProxy, 0>();
+
+		RegisterOutlinerItems();
+	}
+}
+
+void FAvaSceneRigEditorModule::ShutdownModule()
+{
+	FAvaSceneRigEditorCommands::Unregister();
+	
+	if (IAvaOutlinerModule::IsLoaded())
+	{
+		FAvaOutlinerItemProxyRegistry& ItemProxyRegistry = IAvaOutlinerModule::Get().GetItemProxyRegistry();
+		ItemProxyRegistry.UnregisterItemProxyFactory<FAvaOutlinerSceneRigProxy>();
+
+		UnregisterOutlinerItems();
+	}
+}
+
+void FAvaSceneRigEditorModule::RegisterOutlinerItems()
+{
+	if (!IAvaOutlinerModule::IsLoaded())
+	{
+		return;
+	}
+
+	OutlinerProxiesExtensionDelegateHandle = IAvaOutlinerModule::Get().GetOnExtendItemProxiesForItem().AddLambda(
+		[this](IAvaOutliner& InOutliner, const FAvaOutlinerItemPtr& InItem, TArray<TSharedPtr<FAvaOutlinerItemProxy>>& OutItemProxies)
+		{
+			if (InItem->IsA<FAvaOutlinerActor>())
+			{
+				if (const TSharedPtr<FAvaOutlinerItemProxy> SceneRigProxy = InOutliner.GetOrCreateItemProxy<FAvaOutlinerSceneRigProxy>(InItem))
+				{
+					OutItemProxies.Add(SceneRigProxy);
+				}
+			}
+		});
+}
+
+void FAvaSceneRigEditorModule::UnregisterOutlinerItems()
+{
+	if (!IAvaOutlinerModule::IsLoaded())
+	{
+		return;
+	}
+
+	IAvaOutlinerModule& OutlinerModule = IAvaOutlinerModule::Get();
+
+	FAvaOutlinerItemProxyRegistry& ItemProxyRegistry = OutlinerModule.GetItemProxyRegistry();
+	ItemProxyRegistry.UnregisterItemProxyFactory<FAvaOutlinerSceneRigProxy>();
+
+	OutlinerModule.GetOnExtendItemProxiesForItem().Remove(OutlinerProxiesExtensionDelegateHandle);
+	OutlinerProxiesExtensionDelegateHandle.Reset();
+}
+
+void FAvaSceneRigEditorModule::CustomizeSceneRig(const TSharedRef<IPropertyHandle>& InSceneRigHandle, IDetailLayoutBuilder& DetailBuilder)
+{
+	TArray<TWeakObjectPtr<UObject>> ObjectsBeingCustomized;
+	DetailBuilder.GetObjectsBeingCustomized(ObjectsBeingCustomized);
+
+	if (ObjectsBeingCustomized.IsEmpty() || !ObjectsBeingCustomized[0].IsValid())
+	{
+		return;
+	}
+
+	InSceneRigHandle->MarkHiddenByCustomization();
+
+	IDetailCategoryBuilder& SceneRigCategory = DetailBuilder.EditCategory(TEXT("Scene Rig"));
+
+	SceneRigCategory.AddCustomRow(LOCTEXT("SceneRig", "Scene Rig"))
+		.WholeRowContent()
+		.HAlign(HAlign_Fill)
+		[
+			SNew(SSceneRigPicker, ObjectsBeingCustomized[0])
+		];
+}
+
+ULevelStreaming* FAvaSceneRigEditorModule::SetActiveSceneRig(UWorld* const InWorld, const FSoftObjectPath& InSceneRigAssetPath) const
+{
+	if (!IsValid(InWorld) || !InSceneRigAssetPath.IsValid())
+	{
+		return nullptr;
+	}
+
+	if (!UAvaSceneRigSubsystem::IsSceneRigAssetData(InSceneRigAssetPath.TryLoad()))
+	{
+		return nullptr;
+	}
+
+	const FString LevelPackageName = InSceneRigAssetPath.GetLongPackageName();
+
+	ULevelStreaming* OutStreamingLevel = OutStreamingLevel = FLevelUtils::FindStreamingLevel(InWorld, *LevelPackageName);
+	if (IsValid(OutStreamingLevel))
+	{
+		if (!IsValid(OutStreamingLevel->GetLoadedLevel()))
+		{
+			InWorld->LoadSecondaryLevels();
+			ensure(OutStreamingLevel->GetLoadedLevel());
+		}
+
+		return OutStreamingLevel;
+	}
+
+	InWorld->Modify();
+
+	RemoveAllSceneRigs(InWorld);
+
+	OutStreamingLevel = EditorLevelUtils::AddLevelToWorld(InWorld, *LevelPackageName, ULevelStreamingDynamic::StaticClass());
+
+	if (!IsValid(OutStreamingLevel))
+	{
+		return nullptr;
+	}
+
+	// Adding the level to the world sets it as the current. We don't want this.
+	InWorld->SetCurrentLevel(InWorld->PersistentLevel);
+
+	OutStreamingLevel->LevelColor = FLinearColor::Yellow;
+
+	UAvaSceneSettings* const SceneSettings = GetSceneSettings(InWorld);
+	if (IsValid(SceneSettings))
+	{
+		SceneSettings->SetSceneRig(InSceneRigAssetPath);
+	}
+	
+	OnSceneRigChangedDelegate.Broadcast(InWorld, OutStreamingLevel);
+
+	return OutStreamingLevel;
+}
+
+FSoftObjectPath FAvaSceneRigEditorModule::GetActiveSceneRig(UWorld* const InWorld) const
+{
+	const UAvaSceneSettings* const SceneSettings = GetSceneSettings(InWorld);
+	if (!IsValid(SceneSettings))
+	{
+		FSoftObjectPath();
+	}
+
+	return SceneSettings->GetSceneRig();
+}
+
+bool FAvaSceneRigEditorModule::IsActiveSceneRigActor(UWorld* const InWorld, AActor* const InActor) const
+{
+	const FSoftObjectPath SceneRigPath = GetActiveSceneRig(InWorld);
+
+	UWorld* const SceneRigAsset = Cast<UWorld>(SceneRigPath.TryLoad());
+	if (!IsValid(SceneRigAsset) || !IsValid(SceneRigAsset->PersistentLevel))
+	{
+		return false;
+	}
+
+	return SceneRigAsset->PersistentLevel->Actors.Contains(InActor);
+}
+
+bool FAvaSceneRigEditorModule::RemoveAllSceneRigs(UWorld* const InWorld) const
+{
+	UAvaSceneRigSubsystem* const SceneRigSubsystem = InWorld->GetSubsystem<UAvaSceneRigSubsystem>();
+	if (!IsValid(SceneRigSubsystem))
+	{
+		return false;
+	}
+
+	for (const ULevelStreaming* const LevelStreaming : SceneRigSubsystem->FindAllSceneRigs())
+	{
+		const TSoftObjectPtr<UWorld>& WorldAsset = LevelStreaming->GetWorldAsset();
+		if (!WorldAsset.IsValid())
+		{
+			continue;
+		}
+
+		if (UAvaSceneRigSubsystem::IsSceneRigAsset(WorldAsset.Get()))
+		{
+			WorldAsset->Modify();
+
+			EditorLevelUtils::RemoveLevelFromWorld(LevelStreaming->GetLoadedLevel());
+		}
+	}
+
+	UAvaSceneSettings* const SceneSettings = GetSceneSettings(InWorld);
+	if (IsValid(SceneSettings))
+	{
+		SceneSettings->SetSceneRig(FSoftObjectPath());
+	}
+
+	OnSceneRigChangedDelegate.Broadcast(InWorld, nullptr);
+
+	return true;
+}
+
+void FAvaSceneRigEditorModule::AddActiveSceneRigActors(UWorld* const InWorld, const TArray<AActor*>& InActors) const
+{
+	if (!IsValid(InWorld) || InActors.IsEmpty())
+	{
+		return;
+	}
+
+	if (!UAvaSceneRigSubsystem::AreActorsSupported(InActors))
+	{
+		return;
+	}
+
+	const UAvaSceneRigSubsystem* const SceneRigSubsystem = UAvaSceneRigSubsystem::ForWorld(InWorld);
+	if (!IsValid(SceneRigSubsystem))
+	{
+		return;
+	}
+
+	ULevelStreaming* const SceneRig = SceneRigSubsystem->FindFirstActiveSceneRig();
+	if (!IsValid(SceneRig))
+	{
+		return;
+	}
+
+	TArray<AActor*> MovedActors;
+	const int32 NumActorsMoved = EditorLevelUtils::MoveActorsToLevel(InActors, SceneRig->GetLoadedLevel());
+
+	if (NumActorsMoved > 0)
+	{
+		OnSceneRigActorsAddedDelegate.Broadcast(InWorld, MovedActors);
+	}
+}
+
+void FAvaSceneRigEditorModule::RemoveActiveSceneRigActors(UWorld* const InWorld, const TArray<AActor*>& InActors) const
+{
+	if (!IsValid(InWorld) || !IsValid(InWorld->PersistentLevel) || InActors.IsEmpty())
+	{
+		return;
+	}
+
+	const UAvaSceneRigSubsystem* const SceneRigSubsystem = UAvaSceneRigSubsystem::ForWorld(InWorld);
+	if (!IsValid(SceneRigSubsystem))
+	{
+		return;
+	}
+
+	ULevelStreaming* const SceneRig = SceneRigSubsystem->FindFirstActiveSceneRig();
+	if (!IsValid(SceneRig))
+	{
+		return;
+	}
+
+	TArray<AActor*> MovedActors;
+	const int32 NumActorsMoved = EditorLevelUtils::MoveActorsToLevel(InActors, InWorld->PersistentLevel, true, true, false, &MovedActors);
+
+	if (NumActorsMoved > 0)
+	{
+		OnSceneRigActorsRemovedDelegate.Broadcast(InWorld, MovedActors);
+	}
+}
+
+UAvaSceneSettings* FAvaSceneRigEditorModule::GetSceneSettings(UWorld* const InWorld) const
+{
+	if (!IsValid(InWorld))
+	{
+		return nullptr;
+	}
+
+	const UAvaSceneSubsystem* const SceneSubsystem = InWorld->GetSubsystem<UAvaSceneSubsystem>();
+	if (!IsValid(SceneSubsystem))
+	{
+		return nullptr;
+	}
+
+	const IAvaSceneInterface* const AvaScene = SceneSubsystem->GetSceneInterface();
+	if (!AvaScene)
+	{
+		return nullptr;
+	}
+
+	return AvaScene->GetSceneSettings();
+}
+
+FSoftObjectPath FAvaSceneRigEditorModule::CreateSceneRigAssetWithDialog() const
+{
+	if (!GEditor)
+	{
+		return nullptr;
+	}
+
+	UWorldFactory* const WorldFactory = NewObject<UWorldFactory>();
+	if (!IsValid(WorldFactory))
+	{
+		return nullptr;
+	}
+
+	FGCObjectScopeGuard DontGCFactory(WorldFactory);
+
+	FAssetToolsModule& AssetToolsModule = FAssetToolsModule::GetModule();
+	IAssetTools& AssetTools = AssetToolsModule.Get();
+
+	// Determine the starting path. Try to use the most recently used directory
+	FString AssetPath;
+
+	const FString DefaultFilesystemDirectory = FEditorDirectories::Get().GetLastDirectory(ELastDirectory::NEW_ASSET);
+	if (DefaultFilesystemDirectory.IsEmpty() || !FPackageName::TryConvertFilenameToLongPackageName(DefaultFilesystemDirectory, AssetPath))
+	{
+		// No saved path, just use the game content root
+		AssetPath = TEXT("/Game");
+	}
+
+	const FString NewAssetPath = AssetPath / TEXT("New") + UAvaSceneRigSubsystem::GetSceneRigAssetSuffix();
+
+	FString PackageName;
+	FString AssetName;
+	AssetTools.CreateUniqueAssetName(NewAssetPath, TEXT(""), PackageName, AssetName);
+
+	FContentBrowserModule& ContentBrowserModule = FModuleManager::LoadModuleChecked<FContentBrowserModule>("ContentBrowser");
+
+	// Ask the user for the path to save to
+	FSaveAssetDialogConfig SaveAssetDialogConfig;
+	SaveAssetDialogConfig.DialogTitleOverride = LOCTEXT("SaveAssetDialogTitle", "Save Scene Rig Asset As");
+	SaveAssetDialogConfig.DefaultAssetName = AssetName;
+	SaveAssetDialogConfig.ExistingAssetPolicy = ESaveAssetDialogExistingAssetPolicy::Disallow;
+	SaveAssetDialogConfig.AssetClassNames.Add(WorldFactory->GetSupportedClass()->GetClassPathName());
+
+	const FString SaveObjectPath = ContentBrowserModule.Get().CreateModalSaveAssetDialog(SaveAssetDialogConfig);
+	if (SaveObjectPath.IsEmpty())
+	{
+		return FSoftObjectPath();
+	}
+
+	// Add "_SceneRig" suffix if needed and create asset
+	const FString SavePackageName = FPackageName::ObjectPathToPackageName(SaveObjectPath);
+	const FString SavePackagePath = FPaths::GetPath(SavePackageName);
+	const FString Suffix = UAvaSceneRigSubsystem::GetSceneRigAssetSuffix();
+
+	FString SaveAssetName = FPaths::GetBaseFilename(SavePackageName);
+	if (!SaveAssetName.EndsWith(Suffix, ESearchCase::IgnoreCase))
+	{
+		SaveAssetName.Append(Suffix);
+	}
+	
+	UWorld* const NewSceneRigAsset = StaticCast<UWorld*>(
+		AssetTools.CreateAsset(SaveAssetName, SavePackagePath, UWorld::StaticClass(), WorldFactory));
+
+	if (!IsValid(NewSceneRigAsset) || !IsValid(NewSceneRigAsset->PersistentLevel))
+	{
+		UE_LOG(AvaSceneRigEditorLog, Warning, TEXT("Failed to create new Scene Rig asset!"));
+		return FSoftObjectPath();
+	}
+
+	// Add asset data to specify this level as a Scene Rig
+	if (IInterface_AssetUserData* AssetUserData = Cast<IInterface_AssetUserData>(NewSceneRigAsset->PersistentLevel))
+	{
+		UAvaSceneRigData* const SceneRigData = NewObject<UAvaSceneRigData>(NewSceneRigAsset->PersistentLevel, NAME_None, RF_Public | RF_Transactional);
+		check(IsValid(SceneRigData));
+		AssetUserData->AddAssetUserData(SceneRigData);
+	}
+
+	UEditorAssetSubsystem* const EditorAssetSubsystem = GEditor->GetEditorSubsystem<UEditorAssetSubsystem>();
+	if (IsValid(EditorAssetSubsystem))
+	{
+		using namespace UE::AvaSceneRig;
+		EditorAssetSubsystem->SetMetadataTag(NewSceneRigAsset, AssetTags::SceneRig, AssetTags::Values::Enabled);
+	}
+
+	// Save package
+	TArray<UPackage*> PackagesToSave;
+	PackagesToSave.Add(NewSceneRigAsset->GetPackage());
+	UEditorLoadingAndSavingUtils::SavePackages(PackagesToSave, true);
+
+	// Save last directory path for dialogs
+	FEditorDirectories::Get().SetLastDirectory(ELastDirectory::NEW_ASSET, SavePackagePath);
+
+	return NewSceneRigAsset->GetPathName();
+}
+
+#undef LOCTEXT_NAMESPACE
+
+IMPLEMENT_MODULE(FAvaSceneRigEditorModule, AvalancheSceneRigEditor)
diff --git a/Engine/Plugins/Experimental/Avalanche/Source/AvalancheSceneRigEditor/Private/AvaSceneRigEditorModule.h b/Engine/Plugins/Experimental/Avalanche/Source/AvalancheSceneRigEditor/Private/AvaSceneRigEditorModule.h
new file mode 100644
--- /dev/null
+++ b/Engine/Plugins/Experimental/Avalanche/Source/AvalancheSceneRigEditor/Private/AvaSceneRigEditorModule.h
@@ -0,0 +1,54 @@
+// Copyright Epic Games, Inc. All Rights Reserved.
+
+#pragma once
+
+#include "Containers/Array.h"
+#include "Delegates/IDelegateInstance.h"
+#include "IAvaSceneRigEditorModule.h"
+#include "Templates/SharedPointer.h"
+#include "UObject/NoExportTypes.h"
+
+DECLARE_LOG_CATEGORY_EXTERN(AvaSceneRigEditorLog, Log, All);
+
+class AActor;
+class IDetailLayoutBuilder;
+class IPropertyHandle;
+class UAvaSceneSettings;
+class ULevelStreaming;
+class UWorld;
+struct FSoftObjectPath;
+
+class FAvaSceneRigEditorModule : public IAvaSceneRigEditorModule
+{
+public:
+	//~ Begin IAvaSceneRigEditorModule
+	virtual void CustomizeSceneRig(const TSharedRef<IPropertyHandle>& InSceneRigHandle, IDetailLayoutBuilder& DetailBuilder) override;
+	virtual ULevelStreaming* SetActiveSceneRig(UWorld* const InWorld, const FSoftObjectPath& InSceneRigAssetPath) const override;
+	virtual FSoftObjectPath GetActiveSceneRig(UWorld* const InWorld) const override;
+	virtual bool IsActiveSceneRigActor(UWorld* const InWorld, AActor* const InActor) const override;
+	virtual bool RemoveAllSceneRigs(UWorld* const InWorld) const override;
+	virtual void AddActiveSceneRigActors(UWorld* const InWorld, const TArray<AActor*>& InActors) const override;
+	virtual void RemoveActiveSceneRigActors(UWorld* const InWorld, const TArray<AActor*>& InActors) const override;
+	virtual FSoftObjectPath CreateSceneRigAssetWithDialog() const override;
+	virtual FOnSceneRigChanged& OnSceneRigChanged() override { return OnSceneRigChangedDelegate; }
+	virtual FOnSceneRigActorsAdded& OnSceneRigActorsAdded() override { return OnSceneRigActorsAddedDelegate; }
+	virtual FOnSceneRigActorsRemoved& OnSceneRigActorsRemoved() override { return OnSceneRigActorsRemovedDelegate; }
+	//~ End IAvaSceneRigEditorModule
+
+	//~ Begin IModuleInterface
+	virtual void StartupModule() override;
+	virtual void ShutdownModule() override;
+	//~ End IModuleInterface
+
+private:
+	void RegisterOutlinerItems();
+	void UnregisterOutlinerItems();
+
+	UAvaSceneSettings* GetSceneSettings(UWorld* const InWorld) const;
+
+	FDelegateHandle OutlinerProxiesExtensionDelegateHandle;
+
+	FOnSceneRigChanged OnSceneRigChangedDelegate;
+	FOnSceneRigActorsAdded OnSceneRigActorsAddedDelegate;
+	FOnSceneRigActorsRemoved OnSceneRigActorsRemovedDelegate;
+};
diff --git a/Engine/Plugins/Experimental/Avalanche/Source/AvalancheSceneRigEditor/Private/Outliner/AvaOutlinerSceneRig.cpp b/Engine/Plugins/Experimental/Avalanche/Source/AvalancheSceneRigEditor/Private/Outliner/AvaOutlinerSceneRig.cpp
new file mode 100644
--- /dev/null
+++ b/Engine/Plugins/Experimental/Avalanche/Source/AvalancheSceneRigEditor/Private/Outliner/AvaOutlinerSceneRig.cpp
@@ -0,0 +1,78 @@
+// Copyright Epic Games, Inc. All Rights Reserved.
+
+#include "AvaOutlinerSceneRig.h"
+#include "AvaSceneRigSubsystem.h"
+#include "EditorModeManager.h"
+#include "Engine/LevelStreaming.h"
+#include "Engine/World.h"
+#include "Framework/Docking/TabManager.h"
+#include "Input/Reply.h"
+#include "Selection/AvaOutlinerScopedSelection.h"
+#include "Styling/AppStyle.h"
+#include "Toolkits/IToolkitHost.h"
+#include "UObject/NameTypes.h"
+#include "Widgets/Docking/SDockTab.h"
+
+#define LOCTEXT_NAMESPACE "AvaOutlinerSceneRig"
+
+FAvaOutlinerSceneRig::FAvaOutlinerSceneRig(IAvaOutliner& InOutliner, ULevelStreaming* const InSceneRig, const FAvaOutlinerItemPtr& InReferencingItem)
+	: Super(InOutliner, InSceneRig, InReferencingItem, TEXT("SceneRig"))
+	, StreamingLevelWeak(InSceneRig)
+{
+	Icon = FSlateIcon(FAppStyle::GetAppStyleSetName(), TEXT("LandscapeEditor.NoiseTool"));
+}
+
+FText FAvaOutlinerSceneRig::GetDisplayName() const
+{
+	if (StreamingLevelWeak.IsValid())
+	{
+		const UWorld* const WorldAsset = StreamingLevelWeak->GetWorldAsset().Get();
+		if (IsValid(WorldAsset))
+		{
+			return FText::Format(LOCTEXT("Tooltip", "This actor is a member of the Scene Rig: {0}")
+				, FText::FromString(WorldAsset->GetName()));
+		}
+	}
+	return FText();
+}
+
+FSlateIcon FAvaOutlinerSceneRig::GetIcon() const
+{
+	return Icon;
+}
+
+void FAvaOutlinerSceneRig::Select(FAvaOutlinerScopedSelection& InSelection) const
+{
+	if (!StreamingLevelWeak.IsValid())
+	{
+		return;
+	}
+
+	const UWorld* const WorldAsset = StreamingLevelWeak->GetWorldAsset().Get();
+	if (!IsValid(WorldAsset))
+	{
+		return;
+	}
+
+	if (ReferencingItemWeak.IsValid())
+	{
+		ReferencingItemWeak.Pin()->Select(InSelection);
+	}
+
+	if (const TSharedPtr<IToolkitHost> ToolkitHost = InSelection.GetEditorModeTools().GetToolkitHost())
+	{
+		if (const TSharedPtr<FTabManager> TabManager = ToolkitHost->GetTabManager())
+		{
+			static const FName TabName = TEXT("AvalancheSceneSettingsTabSpawner");
+			TabManager->TryInvokeTab(TabName);
+		}
+	}
+}
+
+void FAvaOutlinerSceneRig::SetObject_Impl(UObject* InObject)
+{
+	FAvaOutlinerObject::SetObject_Impl(InObject);
+	StreamingLevelWeak = Cast<ULevelStreaming>(InObject);
+}
+
+#undef LOCTEXT_NAMESPACE
diff --git a/Engine/Plugins/Experimental/Avalanche/Source/AvalancheSceneRigEditor/Private/Outliner/AvaOutlinerSceneRig.h b/Engine/Plugins/Experimental/Avalanche/Source/AvalancheSceneRigEditor/Private/Outliner/AvaOutlinerSceneRig.h
new file mode 100644
--- /dev/null
+++ b/Engine/Plugins/Experimental/Avalanche/Source/AvalancheSceneRigEditor/Private/Outliner/AvaOutlinerSceneRig.h
@@ -0,0 +1,42 @@
+// Copyright Epic Games, Inc. All Rights Reserved.
+
+#pragma once
+
+#include "Item/AvaOutlinerItem.h"
+#include "Item/AvaOutlinerObjectReference.h"
+#include "Templates/SharedPointer.h"
+#include "Textures/SlateIcon.h"
+#include "UObject/WeakObjectPtr.h"
+#include "UObject/WeakObjectPtrTemplates.h"
+
+class FAvaOutlinerScopedSelection;
+class FText;
+class ULevelStreaming;
+class UObject;
+
+/**
+ * Item in Outliner representing a Scene Rig object.
+ * Inherits from FAvaOutlinerObjectReference as multiple objects can be in the same Scene Rig.
+ */
+class FAvaOutlinerSceneRig : public FAvaOutlinerObjectReference
+{
+public:
+	UE_AVA_INHERITS_WITH_SUPER(FAvaOutlinerSceneRig, FAvaOutlinerObjectReference);
+
+	FAvaOutlinerSceneRig(IAvaOutliner& InOutliner, ULevelStreaming* const InSceneRig, const FAvaOutlinerItemPtr& InReferencingItem);
+
+	//~ Begin IAvaOutlinerItem
+	virtual FText GetDisplayName() const override;
+	virtual FSlateIcon GetIcon() const override;
+	virtual void Select(FAvaOutlinerScopedSelection& InSelection) const override;
+	//~ End IAvaOutlinerItem
+	
+private:
+	//~ Begin FAvaOutlinerObjectItem
+	virtual void SetObject_Impl(UObject* InObject) override;
+	//~ End FAvaOutlinerObjectItem
+
+	FSlateIcon Icon;
+
+	TWeakObjectPtr<ULevelStreaming> StreamingLevelWeak;
+};
diff --git a/Engine/Plugins/Experimental/Avalanche/Source/AvalancheSceneRigEditor/Private/Outliner/AvaOutlinerSceneRigProxy.cpp b/Engine/Plugins/Experimental/Avalanche/Source/AvalancheSceneRigEditor/Private/Outliner/AvaOutlinerSceneRigProxy.cpp
new file mode 100644
--- /dev/null
+++ b/Engine/Plugins/Experimental/Avalanche/Source/AvalancheSceneRigEditor/Private/Outliner/AvaOutlinerSceneRigProxy.cpp
@@ -0,0 +1,171 @@
+// Copyright Epic Games, Inc. All Rights Reserved.
+
+#include "AvaOutlinerSceneRigProxy.h"
+#include "AvaOutlinerSceneRig.h"
+#include "AvaSceneRigSubsystem.h"
+#include "AvaTypeSharedPointer.h"
+#include "Engine/Level.h"
+#include "Engine/LevelStreaming.h"
+#include "Engine/World.h"
+#include "GameFramework/Actor.h"
+#include "IAvaOutliner.h"
+#include "IAvaSceneRigEditorModule.h"
+#include "Item/AvaOutlinerActor.h"
+#include "Item/AvaOutlinerItemParameters.h"
+#include "Selection/AvaOutlinerScopedSelection.h"
+#include "Styling/AppStyle.h"
+
+#define LOCTEXT_NAMESPACE "AvaOutlinerSceneRigProxy"
+
+FAvaOutlinerSceneRigProxy::FAvaOutlinerSceneRigProxy(IAvaOutliner& InOutliner, const FAvaOutlinerItemPtr& InParentItem)
+	: Super(InOutliner, InParentItem)
+{
+	Icon = FSlateIcon(FAppStyle::GetAppStyleSetName(), TEXT("LandscapeEditor.NoiseTool"));
+
+	UAvaSceneRigSubsystem* const SceneRigSubsystem = UAvaSceneRigSubsystem::ForWorld(Outliner.GetWorld());
+	if (IsValid(SceneRigSubsystem))
+	{
+		StreamingLevelWeak = SceneRigSubsystem->FindFirstActiveSceneRig();
+	}
+}
+
+FAvaOutlinerSceneRigProxy::~FAvaOutlinerSceneRigProxy()
+{
+	UnbindDelegates();
+}
+
+void FAvaOutlinerSceneRigProxy::OnItemRegistered()
+{
+	Super::OnItemRegistered();
+
+	BindDelegates();
+}
+
+void FAvaOutlinerSceneRigProxy::OnItemUnregistered()
+{
+	Super::OnItemUnregistered();
+	
+	UnbindDelegates();
+}
+
+FText FAvaOutlinerSceneRigProxy::GetDisplayName() const
+{
+	return LOCTEXT("DisplayName", "Scene Rig");
+}
+
+FSlateIcon FAvaOutlinerSceneRigProxy::GetIcon() const
+{
+	return Icon;
+}
+
+FText FAvaOutlinerSceneRigProxy::GetIconTooltipText() const
+{
+	return FText();
+}
+
+void FAvaOutlinerSceneRigProxy::GetProxiedItems(const TSharedRef<IAvaOutlinerItem>& InParent, TArray<FAvaOutlinerItemPtr>& OutChildren, bool bInRecursive)
+{
+	if (const FAvaOutlinerActor* const ActorItem = InParent->CastTo<FAvaOutlinerActor>())
+	{
+		AActor* const Actor = ActorItem->GetActor();
+		if (IsValid(Actor))
+		{
+			ULevelStreaming* const ActiveSceneRig = StreamingLevelWeak.Get();
+			if (IsValid(ActiveSceneRig))
+			{
+				const UWorld* const World = ActiveSceneRig->GetWorldAsset().Get();
+				if (IsValid(World) && IsValid(World->PersistentLevel))
+				{
+					if (World->PersistentLevel->Actors.Contains(Actor))
+					{
+						FAvaOutlinerItemPtr SceneRigItem = Outliner.FindOrAdd<FAvaOutlinerSceneRig>(ActiveSceneRig, InParent);
+						SceneRigItem->SetParent(SharedThis(this));
+
+						OutChildren.Add(SceneRigItem);
+
+						if (bInRecursive)
+						{
+							SceneRigItem->FindChildren(OutChildren, bInRecursive);
+						}
+					}
+				}
+			}
+		}
+	}
+}
+
+void FAvaOutlinerSceneRigProxy::BindDelegates()
+{
+	UnbindDelegates();
+
+	IAvaSceneRigEditorModule& SceneRigEditorModule = IAvaSceneRigEditorModule::Get();
+
+	OnSceneRigChangedHandle = SceneRigEditorModule.OnSceneRigChanged().AddSP(this, &FAvaOutlinerSceneRigProxy::OnSceneRigChanged);
+	OnSceneRigActorAddedHandle = SceneRigEditorModule.OnSceneRigActorsAdded().AddSP(this, &FAvaOutlinerSceneRigProxy::OnSceneRigActorAdded);
+	OnSceneRigActorRemovedHandle = SceneRigEditorModule.OnSceneRigActorsRemoved().AddSP(this, &FAvaOutlinerSceneRigProxy::OnSceneRigActorRemoved);
+}
+
+void FAvaOutlinerSceneRigProxy::UnbindDelegates()
+{
+	IAvaSceneRigEditorModule& SceneRigEditorModule = IAvaSceneRigEditorModule::Get();
+
+	SceneRigEditorModule.OnSceneRigChanged().Remove(OnSceneRigChangedHandle);
+	OnSceneRigChangedHandle.Reset();
+
+	SceneRigEditorModule.OnSceneRigActorsAdded().Remove(OnSceneRigActorAddedHandle);
+	OnSceneRigActorAddedHandle.Reset();
+
+	SceneRigEditorModule.OnSceneRigActorsRemoved().Remove(OnSceneRigActorRemovedHandle);
+	OnSceneRigActorRemovedHandle.Reset();
+}
+
+void FAvaOutlinerSceneRigProxy::OnSceneRigChanged(UWorld* const InWorld, ULevelStreaming* const InSceneRig)
+{
+	StreamingLevelWeak = InSceneRig;
+
+	Outliner.RequestRefresh();
+	RefreshChildren();
+}
+
+void FAvaOutlinerSceneRigProxy::OnSceneRigActorAdded(UWorld* const InWorld, const TArray<AActor*>& InActors)
+{
+	Outliner.RequestRefresh();
+	RefreshChildren();
+}
+
+void FAvaOutlinerSceneRigProxy::OnSceneRigActorRemoved(UWorld* const InWorld, const TArray<AActor*>& InActors)
+{
+	Outliner.RequestRefresh();
+	RefreshChildren();
+}
+
+ULevelStreaming* FAvaOutlinerSceneRigProxy::GetSceneRigAsset() const
+{
+	const FAvaOutlinerItemPtr Parent = GetParent();
+	if (!Parent.IsValid())
+	{
+		return nullptr;
+	}
+
+	const FAvaOutlinerActor* const ActorItem = Parent->CastTo<FAvaOutlinerActor>();
+	if (!ActorItem)
+	{
+		return nullptr;
+	}
+
+	UAvaSceneRigSubsystem* const SceneRigSubsystem = UAvaSceneRigSubsystem::ForWorld(Outliner.GetWorld());
+	if (!IsValid(SceneRigSubsystem))
+	{
+		return nullptr;
+	}
+
+	ULevelStreaming* const StreamingLevel = SceneRigSubsystem->SceneRigFromActor(ActorItem->GetActor());
+	if (!IsValid(StreamingLevel))
+	{
+		return nullptr;
+	}
+
+	return SceneRigSubsystem->FindFirstActiveSceneRig();
+}
+
+#undef LOCTEXT_NAMESPACE
diff --git a/Engine/Plugins/Experimental/Avalanche/Source/AvalancheSceneRigEditor/Private/Outliner/AvaOutlinerSceneRigProxy.h b/Engine/Plugins/Experimental/Avalanche/Source/AvalancheSceneRigEditor/Private/Outliner/AvaOutlinerSceneRigProxy.h
new file mode 100644
--- /dev/null
+++ b/Engine/Plugins/Experimental/Avalanche/Source/AvalancheSceneRigEditor/Private/Outliner/AvaOutlinerSceneRigProxy.h
@@ -0,0 +1,51 @@
+// Copyright Epic Games, Inc. All Rights Reserved.
+
+#pragma once
+
+#include "Delegates/IDelegateInstance.h"
+#include "Item/AvaOutlinerItemProxy.h"
+#include "Templates/SharedPointer.h"
+#include "Textures/SlateIcon.h"
+
+class UAvaSceneRigSubsystem;
+class ULevelStreaming;
+
+class FAvaOutlinerSceneRigProxy : public FAvaOutlinerItemProxy
+{
+public:
+	UE_AVA_INHERITS_WITH_SUPER(FAvaOutlinerSceneRigProxy, FAvaOutlinerItemProxy);
+
+	FAvaOutlinerSceneRigProxy(IAvaOutliner& InOutliner, const FAvaOutlinerItemPtr& InParentItem);
+
+	virtual ~FAvaOutlinerSceneRigProxy() override;
+
+	//~ Begin IAvaOutlinerItem
+	virtual void OnItemRegistered() override;
+	virtual void OnItemUnregistered() override;
+	virtual FText GetDisplayName() const override;
+	virtual FSlateIcon GetIcon() const override;
+	virtual FText GetIconTooltipText() const override;
+	//~ End IAvaOutlinerItem
+
+	//~ Begin FAvaOutlinerItemProxy
+	virtual void GetProxiedItems(const TSharedRef<IAvaOutlinerItem>& InParent, TArray<FAvaOutlinerItemPtr>& OutChildren, bool bInRecursive) override;
+	//~ End FAvaOutlinerItemProxy
+
+private:
+	void BindDelegates();
+	void UnbindDelegates();
+
+	void OnSceneRigChanged(UWorld* const InWorld, ULevelStreaming* const InSceneRig);
+	void OnSceneRigActorAdded(UWorld* const InWorld, const TArray<AActor*>& InActors);
+	void OnSceneRigActorRemoved(UWorld* const InWorld, const TArray<AActor*>& InActors);
+
+	ULevelStreaming* GetSceneRigAsset() const;
+
+	FSlateIcon Icon;
+
+	FDelegateHandle OnSceneRigChangedHandle;
+	FDelegateHandle OnSceneRigActorAddedHandle;
+	FDelegateHandle OnSceneRigActorRemovedHandle;
+
+	TWeakObjectPtr<ULevelStreaming> StreamingLevelWeak;
+};
diff --git a/Engine/Plugins/Experimental/Avalanche/Source/AvalancheSceneRigEditor/Private/Widgets/SSceneRigPicker.cpp b/Engine/Plugins/Experimental/Avalanche/Source/AvalancheSceneRigEditor/Private/Widgets/SSceneRigPicker.cpp
new file mode 100644
--- /dev/null
+++ b/Engine/Plugins/Experimental/Avalanche/Source/AvalancheSceneRigEditor/Private/Widgets/SSceneRigPicker.cpp
@@ -0,0 +1,124 @@
+// Copyright Epic Games, Inc. All Rights Reserved.
+
+#include "SSceneRigPicker.h"
+#include "AssetRegistry/AssetData.h"
+#include "AvaSceneRigSubsystem.h"
+#include "DetailCategoryBuilder.h"
+#include "DetailLayoutBuilder.h"
+#include "DetailWidgetRow.h"
+#include "Engine/World.h"
+#include "IAvaSceneRigEditorModule.h"
+#include "IDetailChildrenBuilder.h"
+#include "IDetailGroup.h"
+#include "Misc/Attribute.h"
+#include "PropertyCustomizationHelpers.h"
+#include "PropertyEditorModule.h"
+#include "PropertyHandle.h"
+#include "Styling/AppStyle.h"
+#include "UObject/Object.h"
+#include "UObject/SoftObjectPath.h"
+#include "UObject/UnrealType.h"
+#include "Widgets/SBoxPanel.h"
+
+#define LOCTEXT_NAMESPACE "SSceneRigPicker"
+
+void SSceneRigPicker::Construct(const FArguments& InArgs, const TWeakObjectPtr<UObject>& InObjectBeingCustomized)
+{
+	ObjectBeingCustomized = InObjectBeingCustomized;
+
+	ChildSlot
+		[
+			SNew(SHorizontalBox)
+			+ SHorizontalBox::Slot()
+			.FillWidth(1.f)
+			.VAlign(VAlign_Center)
+			[
+				SNew(SObjectPropertyEntryBox)
+				.AllowedClass(UWorld::StaticClass())
+				.DisplayBrowse(true)
+				.DisplayUseSelected(true)
+				.DisplayCompactSize(false)
+				.EnableContentPicker(true)
+				.AllowClear(true)
+				.AllowCreate(false)
+				.DisplayThumbnail(true)
+				.OnShouldFilterAsset(this, &SSceneRigPicker::ShouldFilterAsset)
+				.ObjectPath(this, &SSceneRigPicker::GetObjectPath)
+				.OnObjectChanged(this, &SSceneRigPicker::OnObjectChanged)
+			]
+			+ SHorizontalBox::Slot()
+			.AutoWidth()
+			.VAlign(VAlign_Center)
+			[
+				PropertyCustomizationHelpers::MakeAddButton(FSimpleDelegate::CreateSP(this, &SSceneRigPicker::OnAddNewSceneRigClick)
+					, LOCTEXT("AddSceneRigTooltip", "Create and add a new scene rig to the level.\n\n"
+						"If a scene rig already exists in the level, it will be replaced."))
+			]
+			+ SHorizontalBox::Slot()
+			.AutoWidth()
+			.VAlign(VAlign_Center)
+			[
+				PropertyCustomizationHelpers::MakeRemoveButton(FSimpleDelegate::CreateSP(this, &SSceneRigPicker::OnRemoveSceneRigClick)
+					, LOCTEXT("RemoveSceneRigTooltip", "Remove scene rig from the level.")
+					, TAttribute<bool>::CreateSP(this, &SSceneRigPicker::IsRemoveButtonEnabled))
+			]
+		];
+}
+
+bool SSceneRigPicker::ShouldFilterAsset(const FAssetData& InAssetData) const
+{
+	// Return false to show this asset
+	return !UAvaSceneRigSubsystem::IsSceneRigAssetData(InAssetData);
+}
+
+FString SSceneRigPicker::GetObjectPath() const
+{
+	return IAvaSceneRigEditorModule::Get().GetActiveSceneRig(GetObjectWorld()).ToString();
+}
+
+void SSceneRigPicker::OnObjectChanged(const FAssetData& InAssetData) const
+{
+	if (InAssetData.IsValid())
+	{
+		IAvaSceneRigEditorModule::Get().SetActiveSceneRig(GetObjectWorld(), InAssetData.GetSoftObjectPath());
+	}
+	else
+	{
+		IAvaSceneRigEditorModule::Get().RemoveAllSceneRigs(GetObjectWorld());
+	}
+}
+
+void SSceneRigPicker::OnAddNewSceneRigClick() const
+{
+	IAvaSceneRigEditorModule& SceneRigEditorModule = IAvaSceneRigEditorModule::Get();
+
+	const FSoftObjectPath NewSceneRigPath = SceneRigEditorModule.CreateSceneRigAssetWithDialog();
+	if (!NewSceneRigPath.IsValid())
+	{
+		return;
+	}
+
+	SceneRigEditorModule.SetActiveSceneRig(GetObjectWorld(), NewSceneRigPath);
+}
+
+void SSceneRigPicker::OnRemoveSceneRigClick() const
+{
+	IAvaSceneRigEditorModule::Get().RemoveAllSceneRigs(GetObjectWorld());
+}
+
+bool SSceneRigPicker::IsRemoveButtonEnabled() const
+{
+	return IAvaSceneRigEditorModule::Get().GetActiveSceneRig(GetObjectWorld()).IsAsset();
+}
+
+UWorld* SSceneRigPicker::GetObjectWorld() const
+{
+	if (!ObjectBeingCustomized.IsValid())
+	{
+		return nullptr;
+	}
+
+	return ObjectBeingCustomized->GetWorld();
+}
+
+#undef LOCTEXT_NAMESPACE
diff --git a/Engine/Plugins/Experimental/Avalanche/Source/AvalancheSceneRigEditor/Private/Widgets/SSceneRigPicker.h b/Engine/Plugins/Experimental/Avalanche/Source/AvalancheSceneRigEditor/Private/Widgets/SSceneRigPicker.h
new file mode 100644
--- /dev/null
+++ b/Engine/Plugins/Experimental/Avalanche/Source/AvalancheSceneRigEditor/Private/Widgets/SSceneRigPicker.h
@@ -0,0 +1,32 @@
+// Copyright Epic Games, Inc. All Rights Reserved.
+
+#pragma once
+
+#include "Containers/StringFwd.h"
+#include "UObject/WeakObjectPtrTemplates.h"
+#include "Widgets/SCompoundWidget.h"
+
+class UObject;
+class UWorld;
+struct FAssetData;
+
+class SSceneRigPicker : public SCompoundWidget
+{
+	SLATE_BEGIN_ARGS(SSceneRigPicker) {}
+	SLATE_END_ARGS()
+
+	void Construct(const FArguments& InArgs, const TWeakObjectPtr<UObject>& InObjectBeingCustomized);
+
+protected:
+	bool ShouldFilterAsset(const FAssetData& InAssetData) const;
+	FString GetObjectPath() const;
+	void OnObjectChanged(const FAssetData& InAssetData) const;
+
+	void OnAddNewSceneRigClick() const;
+	void OnRemoveSceneRigClick() const;
+	bool IsRemoveButtonEnabled() const;
+
+	UWorld* GetObjectWorld() const;
+
+	TWeakObjectPtr<UObject> ObjectBeingCustomized;
+};
diff --git a/Engine/Plugins/Experimental/Avalanche/Source/AvalancheSceneRigEditor/Public/AvaSceneRigEditorCommands.h b/Engine/Plugins/Experimental/Avalanche/Source/AvalancheSceneRigEditor/Public/AvaSceneRigEditorCommands.h
new file mode 100644
--- /dev/null
+++ b/Engine/Plugins/Experimental/Avalanche/Source/AvalancheSceneRigEditor/Public/AvaSceneRigEditorCommands.h
@@ -0,0 +1,26 @@
+// Copyright Epic Games, Inc. All Rights Reserved.
+
+#pragma once
+
+#include "Framework/Commands/Commands.h"
+#include "Styling/AppStyle.h"
+
+class AVALANCHESCENERIGEDITOR_API FAvaSceneRigEditorCommands : public TCommands<FAvaSceneRigEditorCommands>
+{
+public:
+	FAvaSceneRigEditorCommands()
+		: TCommands<FAvaSceneRigEditorCommands>(TEXT("AvaSceneRigCommands")
+		, NSLOCTEXT("AvaSceneRigCommands", "AvaSceneRigCommands", "Avalanche Scene Rig Commands")
+		, NAME_None
+		, FAppStyle::GetAppStyleSetName())
+	{}
+
+	//~ Begin TCommands
+	virtual void RegisterCommands() override;	
+	//~ End TCommands
+	
+	TSharedPtr<FUICommandInfo> PromptToSaveSceneRigFromOutlinerItems;
+
+	TSharedPtr<FUICommandInfo> AddOutlinerItemsToSceneRig;
+	TSharedPtr<FUICommandInfo> RemoveOutlinerItemsToSceneRig;
+};
diff --git a/Engine/Plugins/Experimental/Avalanche/Source/AvalancheSceneRigEditor/Public/IAvaSceneRigEditorModule.h b/Engine/Plugins/Experimental/Avalanche/Source/AvalancheSceneRigEditor/Public/IAvaSceneRigEditorModule.h
new file mode 100644
--- /dev/null
+++ b/Engine/Plugins/Experimental/Avalanche/Source/AvalancheSceneRigEditor/Public/IAvaSceneRigEditorModule.h
@@ -0,0 +1,70 @@
+// Copyright Epic Games, Inc. All Rights Reserved.
+
+#pragma once
+
+#include "Containers/ContainersFwd.h"
+#include "HAL/Platform.h"
+#include "Modules/ModuleInterface.h"
+#include "Modules/ModuleManager.h"
+#include "Templates/SharedPointerFwd.h"
+
+class AActor;
+class IDetailLayoutBuilder;
+class IPropertyHandle;
+class ULevel;
+class ULevelStreaming;
+class UToolMenu;
+class UWorld;
+struct FSoftObjectPath;
+
+DECLARE_EVENT_TwoParams(IAvaSceneRigEditorModule, FOnSceneRigChanged, UWorld*, ULevelStreaming*);
+DECLARE_EVENT_TwoParams(IAvaSceneRigEditorModule, FOnSceneRigActorsAdded, UWorld*, const TArray<AActor*>&);
+DECLARE_EVENT_TwoParams(IAvaSceneRigEditorModule, FOnSceneRigActorsRemoved, UWorld*, const TArray<AActor*>&);
+
+class IAvaSceneRigEditorModule : public IModuleInterface
+{
+	static constexpr const TCHAR* ModuleName = TEXT("AvalancheSceneRigEditor");
+
+public:
+	static bool IsLoaded()
+	{
+		return FModuleManager::Get().IsModuleLoaded(ModuleName);
+	}
+
+	static IAvaSceneRigEditorModule& Get()
+	{
+		return FModuleManager::LoadModuleChecked<IAvaSceneRigEditorModule>(ModuleName);
+	}
+
+	virtual void CustomizeSceneRig(const TSharedRef<IPropertyHandle>& InSceneRigHandle, IDetailLayoutBuilder& DetailBuilder) = 0;
+
+	/** Adds a Scene Rig streaming level to the world from a path to a Scene Rig level asset. */
+	virtual ULevelStreaming* SetActiveSceneRig(UWorld* const InWorld, const FSoftObjectPath& InSceneRigAssetPath) const = 0;
+
+	/** As opposed to UAvaSceneSubsystem::FindFirstActiveSceneRig, this returns the cached Scene State Scene Rig. */
+	virtual FSoftObjectPath GetActiveSceneRig(UWorld* const InWorld) const = 0;
+
+	/** Returns true if the specified actor is a member of the active Scene Rig. */
+	virtual bool IsActiveSceneRigActor(UWorld* const InWorld, AActor* const InActor) const = 0;
+
+	/** Removes all Scene Rig objects from the persistent level. */
+	virtual bool RemoveAllSceneRigs(UWorld* const InWorld) const = 0;
+
+	/** Adds a list of actors from another streaming level to the specified Scene Rig. */
+	virtual void AddActiveSceneRigActors(UWorld* const InWorld, const TArray<AActor*>& InActors) const = 0;
+
+	/** Removes a list of actors from the active Scene Rig. */
+	virtual void RemoveActiveSceneRigActors(UWorld* const InWorld, const TArray<AActor*>& InActors) const = 0;
+
+	/**
+	 * Creates a new Scene Rig level asset. Asks the user for a location to save the new asset to.
+	 * Optionally adds actors to the new Scene Rig if a world and actors are specified.
+	 */
+	virtual FSoftObjectPath CreateSceneRigAssetWithDialog() const = 0;
+
+	virtual FOnSceneRigChanged& OnSceneRigChanged() = 0;
+
+	virtual FOnSceneRigActorsAdded& OnSceneRigActorsAdded() = 0;
+
+	virtual FOnSceneRigActorsRemoved& OnSceneRigActorsRemoved() = 0;
+};
