[Motion Design]
 - Fix static analysis warning in FAvaSceneRigEditorModule::SetActiveSceneRig

diff --git a/Engine/Plugins/Experimental/Avalanche/Source/AvalancheSceneRigEditor/Private/AvaSceneRigEditorModule.cpp b/Engine/Plugins/Experimental/Avalanche/Source/AvalancheSceneRigEditor/Private/AvaSceneRigEditorModule.cpp
--- a/Engine/Plugins/Experimental/Avalanche/Source/AvalancheSceneRigEditor/Private/AvaSceneRigEditorModule.cpp
+++ b/Engine/Plugins/Experimental/Avalanche/Source/AvalancheSceneRigEditor/Private/AvaSceneRigEditorModule.cpp
@@ -144,7 +144,7 @@ ULevelStreaming* FAvaSceneRigEditorModule::SetActiveSceneRig(UWorld* const InWor
 
 	const FString LevelPackageName = InSceneRigAssetPath.GetLongPackageName();
 
-	ULevelStreaming* OutStreamingLevel = OutStreamingLevel = FLevelUtils::FindStreamingLevel(InWorld, *LevelPackageName);
+	ULevelStreaming* OutStreamingLevel = FLevelUtils::FindStreamingLevel(InWorld, *LevelPackageName);
 	if (IsValid(OutStreamingLevel))
 	{
 		if (!IsValid(OutStreamingLevel->GetLoadedLevel()))
