[MotionDesign] GizmoComponent crash on duplication fix

diff --git a/Engine/Plugins/Experimental/Avalanche/Source/Avalanche/Private/Framework/AvaGizmoComponent.cpp b/Engine/Plugins/Experimental/Avalanche/Source/Avalanche/Private/Framework/AvaGizmoComponent.cpp
--- a/Engine/Plugins/Experimental/Avalanche/Source/Avalanche/Private/Framework/AvaGizmoComponent.cpp
+++ b/Engine/Plugins/Experimental/Avalanche/Source/Avalanche/Private/Framework/AvaGizmoComponent.cpp
@@ -65,11 +65,7 @@ void UAvaGizmoComponent::OnRegister()
 	{
 		if (UWorld* World = Owner->GetWorld())
 		{
-			if (bIsGizmoEnabled && !bAreGizmoValuesApplied)
-			{
-				StoreComponentValues();
-				ApplyGizmoValues();
-			}
+			// We don't store or apply values here - they'll be called in the cb below
 			
 			PostRegisterComponentsHandle = World->AddOnPostRegisterAllActorComponentsHandler(
 				FOnPostRegisterAllActorComponents::FDelegate::CreateUObject(this, &UAvaGizmoComponent::OnPostRegisterParentComponents));
@@ -88,6 +84,7 @@ void UAvaGizmoComponent::OnUnregister()
 		if (UWorld* World = Owner->GetWorld())
 		{
 			World->RemoveOnPostRegisterAllActorComponentsHandler(PostRegisterComponentsHandle);
+			PostRegisterComponentsHandle.Reset();
 		}
 	}
 
@@ -115,7 +112,7 @@ void UAvaGizmoComponent::OnComponentDestroyed(bool bDestroyingHierarchy)
 	// Only restore if destroying this single component, not on actor teardown
 	if (!bDestroyingHierarchy)
 	{
-		RestoreComponentValues();	
+		RestoreComponentValues();
 	}
 
 	Super::OnComponentDestroyed(bDestroyingHierarchy);
@@ -133,50 +130,44 @@ void UAvaGizmoComponent::ApplyGizmoValues()
 	{
 		bAreGizmoValuesApplied = true;
 
-		TArray<UPrimitiveComponent*> PrimitiveComponents;
-		Owner->GetComponents<UPrimitiveComponent>(PrimitiveComponents, bApplyToChildActors);
-
-		for (UPrimitiveComponent* Component : PrimitiveComponents)
-		{
-			if (!IsValid(Component))
+		ForEachPrimitiveComponent(
+			[this](UPrimitiveComponent* InComponent)
 			{
-				continue;
-			}
-		
-			Component->bVisibleInReflectionCaptures = false;
-			Component->bVisibleInRealTimeSkyCaptures = false;
-			Component->bVisibleInRayTracing = false;
-			Component->SetVisibleInRayTracing(false);
-			Component->SetVisibleInSceneCaptureOnly(false);
-
-			Component->SetAffectDistanceFieldLighting(false);
-			Component->SetAffectDynamicIndirectLighting(false);
-			Component->SetAffectIndirectLightingWhileHidden(false);
-
-			Component->SetCastShadow(bCastShadow);
-			Component->SetVisibility(bIsVisibleInEditor, true);
-			Component->SetHiddenInGame(bIsHiddenInGame);
-
-			Component->SetRenderCustomDepth(bSetStencil);
-			Component->SetCustomDepthStencilValue(StencilId);	
+				InComponent->bVisibleInReflectionCaptures = false;
+				InComponent->bVisibleInRealTimeSkyCaptures = false;
+				InComponent->bVisibleInRayTracing = false;
+				InComponent->SetVisibleInRayTracing(false);
+				InComponent->SetVisibleInSceneCaptureOnly(false);
+				
+				InComponent->SetAffectDistanceFieldLighting(false);
+				InComponent->SetAffectDynamicIndirectLighting(false);
+				InComponent->SetAffectIndirectLightingWhileHidden(false);
 
-			Component->SetRenderInMainPass(bRenderInMainPass);
-			Component->SetRenderInDepthPass(bRenderDepth);
+				InComponent->SetCastShadow(bCastShadow);
+				InComponent->SetVisibility(bIsVisibleInEditor, true);
+				InComponent->SetHiddenInGame(bIsHiddenInGame);
 
-			if (UDynamicMeshComponent* DynamicMeshComponent = Cast<UDynamicMeshComponent>(Component))
-			{
-				DynamicMeshComponent->SetEnableWireframeRenderPass(bDrawWireframe);
-				DynamicMeshComponent->WireframeColor = WireframeColor;
-			}			
+				InComponent->SetRenderCustomDepth(bSetStencil);
+				InComponent->SetCustomDepthStencilValue(StencilId);	
+				
+				InComponent->SetRenderInMainPass(bRenderInMainPass);
+				InComponent->SetRenderInDepthPass(bRenderDepth);
+				
+				if (Material)
+				{
+					for (int32 MaterialIdx = 0; MaterialIdx < InComponent->GetNumMaterials(); ++MaterialIdx)
+					{
+						InComponent->SetMaterial(MaterialIdx, Material);
+					}
+				}
 
-			if (Material)
-			{
-				for (int32 MaterialIdx = 0; MaterialIdx < Component->GetNumMaterials(); ++MaterialIdx)
+				
+				if (UDynamicMeshComponent* DynamicMeshComponent = Cast<UDynamicMeshComponent>(InComponent))
 				{
-					Component->SetMaterial(MaterialIdx, Material);
+					DynamicMeshComponent->WireframeColor = WireframeColor;
+					DynamicMeshComponent->SetEnableWireframeRenderPass(bDrawWireframe);
 				}
-			}
-		}
+			});
 
 		if (!FUObjectThreadContext::Get().IsRoutingPostLoad)
 		{
@@ -331,7 +322,17 @@ void UAvaGizmoComponent::SetMaterial(UMaterialInterface* InMaterial)
 	if (Material != InMaterial)
 	{
 		Material = InMaterial;
-		ApplyGizmoValues();
+		if (Material)
+		{
+			ApplyToEachPrimitiveComponent(
+			[Material = Material](UPrimitiveComponent* InComponent)
+				{
+					for (int32 MaterialIdx = 0; MaterialIdx < InComponent->GetNumMaterials(); ++MaterialIdx)
+					{
+						InComponent->SetMaterial(MaterialIdx, Material);
+					}
+				});
+		}
 	}
 }
 
@@ -363,7 +364,11 @@ void UAvaGizmoComponent::SetCastShadow(bool bInValue)
 	if (bCastShadow != bInValue)
 	{
 		bCastShadow = bInValue;
-		ApplyGizmoValues();
+		ApplyToEachPrimitiveComponent(
+		[bCastShadow = bCastShadow](UPrimitiveComponent* InComponent)
+			{
+				InComponent->SetCastShadow(bCastShadow);
+			});
 	}
 }
 
@@ -377,7 +382,11 @@ void UAvaGizmoComponent::SetVisibleInEditor(bool bInValue)
 	if (bIsVisibleInEditor != bInValue)
 	{
 		bIsVisibleInEditor = bInValue;
-		ApplyGizmoValues();
+		ApplyToEachPrimitiveComponent(
+		[bIsVisibleInEditor = bIsVisibleInEditor](UPrimitiveComponent* InComponent)
+			{
+				InComponent->SetVisibility(bIsVisibleInEditor, true);
+			});
 	}
 }
 
@@ -391,7 +400,11 @@ void UAvaGizmoComponent::SetHiddenInGame(bool bInValue)
 	if (bIsHiddenInGame != bInValue)
 	{
 		bIsHiddenInGame = bInValue;
-		ApplyGizmoValues();
+		ApplyToEachPrimitiveComponent(
+		[bIsHiddenInGame = bIsHiddenInGame](UPrimitiveComponent* InComponent)
+			{
+				InComponent->SetHiddenInGame(bIsHiddenInGame);
+			});
 	}
 }
 
@@ -405,7 +418,11 @@ void UAvaGizmoComponent::SetRenderInMainPass(const bool bInValue)
 	if (bRenderInMainPass != bInValue)
 	{
 		bRenderInMainPass = bInValue;
-		ApplyGizmoValues();
+		ApplyToEachPrimitiveComponent(
+		[bRenderInMainPass = bRenderInMainPass](UPrimitiveComponent* InComponent)
+			{
+				InComponent->SetRenderInMainPass(bRenderInMainPass);
+			});
 	}
 }
 
@@ -419,7 +436,11 @@ void UAvaGizmoComponent::SetRenderDepth(const bool bInValue)
 	if (bRenderDepth != bInValue)
 	{
 		bRenderDepth = bInValue;
-		ApplyGizmoValues();
+		ApplyToEachPrimitiveComponent(
+		[bRenderDepth = bRenderDepth](UPrimitiveComponent* InComponent)
+			{
+				InComponent->SetRenderInDepthPass(bRenderDepth);
+			});
 	}
 }
 
@@ -433,7 +454,14 @@ void UAvaGizmoComponent::SetShowWireframe(bool bInValue)
 	if (bDrawWireframe != bInValue)
 	{
 		bDrawWireframe = bInValue;
-		ApplyGizmoValues();
+		ApplyToEachPrimitiveComponent(
+		[bDrawWireframe = bDrawWireframe](UPrimitiveComponent* InComponent)
+			{
+				if (UDynamicMeshComponent* DynamicMeshComponent = Cast<UDynamicMeshComponent>(InComponent))
+				{
+					DynamicMeshComponent->SetEnableWireframeRenderPass(bDrawWireframe);
+				}
+			});
 	}
 }
 
@@ -442,7 +470,14 @@ void UAvaGizmoComponent::SetWireframeColor(const FLinearColor& InColor)
 	if (WireframeColor != InColor)
 	{
 		WireframeColor = InColor;
-		ApplyGizmoValues();
+		ApplyToEachPrimitiveComponent(
+		[WireframeColor = WireframeColor](UPrimitiveComponent* InComponent)
+			{
+				if (UDynamicMeshComponent* DynamicMeshComponent = Cast<UDynamicMeshComponent>(InComponent))
+				{
+					DynamicMeshComponent->WireframeColor = WireframeColor;
+				}
+			});
 	}
 }
 
@@ -456,7 +491,11 @@ void UAvaGizmoComponent::SetsStencil(const bool bInSetStencil)
 	if (bSetStencil != bInSetStencil)
 	{
 		bSetStencil = bInSetStencil;
-		ApplyGizmoValues();
+		ApplyToEachPrimitiveComponent(
+	[bSetStencil = bSetStencil](UPrimitiveComponent* InComponent)
+			{
+				InComponent->SetRenderCustomDepth(bSetStencil);	
+			});
 	}
 }
 
@@ -470,7 +509,11 @@ void UAvaGizmoComponent::SetStencilId(const uint8 InStencilId)
 	if (StencilId != InStencilId)
 	{
 		StencilId = InStencilId;
-		ApplyGizmoValues();
+		ApplyToEachPrimitiveComponent(
+	[StencilId = StencilId](UPrimitiveComponent* InComponent)
+			{
+				InComponent->SetCustomDepthStencilValue(StencilId);	
+			});
 	}
 }
 
@@ -481,57 +524,50 @@ void UAvaGizmoComponent::StoreComponentValues()
 		if (const AActor* Owner = GetOwner();
 			IsValid(Owner))
 		{
-			TArray<UPrimitiveComponent*> PrimitiveComponents;
-        	Owner->GetComponents<UPrimitiveComponent>(PrimitiveComponents, bApplyToChildActors);
-
 			bool bWasOneOrMorePrimitiveComponents = false;
-        	for (UPrimitiveComponent* Component : PrimitiveComponents)
-        	{
-        		if (!IsValid(Component))
-        		{
-        			continue;
-        		}
-        		
-				bWasOneOrMorePrimitiveComponents = true;
-				
-				FSoftComponentReference ComponentReference;
-				ComponentReference.PathToComponent = Component->GetPathName(Owner);
-				
-				FAvaGizmoComponentPrimitiveValues& StoredComponentValues = ComponentValues.FindOrAdd(ComponentReference);
-				
-				StoredComponentValues.bVisibleInReflectionCaptures = Component->bVisibleInReflectionCaptures;
-				StoredComponentValues.bVisibleInRealTimeSkyCaptures = Component->bVisibleInRealTimeSkyCaptures;
-				StoredComponentValues.bVisibleInRayTracing = Component->bVisibleInRayTracing;
-				
-				StoredComponentValues.bAffectDistanceFieldLighting = Component->bAffectDistanceFieldLighting;
-				StoredComponentValues.bAffectDynamicIndirectLighting = Component->bAffectDynamicIndirectLighting;
-				StoredComponentValues.bAffectIndirectLightingWhileHidden = Component->bAffectIndirectLightingWhileHidden;
-				
-				StoredComponentValues.bCastShadow = Component->CastShadow;
-				StoredComponentValues.bIsVisible = Component->IsVisible();
-				StoredComponentValues.bIsHiddenInGame = Component->bHiddenInGame;
+			ForEachPrimitiveComponent(
+				[this, Owner, &bWasOneOrMorePrimitiveComponents](UPrimitiveComponent* InComponent)
+				{
+					bWasOneOrMorePrimitiveComponents = true;
 
-				StoredComponentValues.bRendersCustomDepth = Component->bRenderCustomDepth > 0;
-				StoredComponentValues.CustomStencilId = Component->CustomDepthStencilValue;
+					FSoftComponentReference ComponentReference;
+					ComponentReference.PathToComponent = InComponent->GetPathName(Owner);
 
-				StoredComponentValues.bRendersInMainPass = Component->bRenderInMainPass;
-				StoredComponentValues.bRendersDepth = Component->bRenderInDepthPass;
+					FAvaGizmoComponentPrimitiveValues& StoredComponentValues = ComponentValues.FindOrAdd(ComponentReference);
 
-				// @todo: move this impl specific stuff to a Handler/Provider subsystem
-				if (UDynamicMeshComponent* DynamicMeshComponent = Cast<UDynamicMeshComponent>(Component))
-				{
-					StoredComponentValues.bDrawWireframe = DynamicMeshComponent->GetEnableWireframeRenderPass();
-					StoredComponentValues.WireframeColor = DynamicMeshComponent->WireframeColor;
-				}
+					StoredComponentValues.bVisibleInReflectionCaptures = InComponent->bVisibleInReflectionCaptures;
+					StoredComponentValues.bVisibleInRealTimeSkyCaptures = InComponent->bVisibleInRealTimeSkyCaptures;
+					StoredComponentValues.bVisibleInRayTracing = InComponent->bVisibleInRayTracing;
 
-				TArray<UMaterialInterface*> UsedMaterials;
-				Component->GetUsedMaterials(UsedMaterials);
+					StoredComponentValues.bAffectDistanceFieldLighting = InComponent->bAffectDistanceFieldLighting;
+					StoredComponentValues.bAffectDynamicIndirectLighting = InComponent->bAffectDynamicIndirectLighting;
+					StoredComponentValues.bAffectIndirectLightingWhileHidden = InComponent->bAffectIndirectLightingWhileHidden;
 
-				Algo::Transform(UsedMaterials, StoredComponentValues.Materials, [](UMaterialInterface* InMaterial)
-				{
-					return InMaterial;
+					StoredComponentValues.bCastShadow = InComponent->CastShadow;
+					StoredComponentValues.bIsVisible = InComponent->IsVisible();
+					StoredComponentValues.bIsHiddenInGame = InComponent->bHiddenInGame;
+
+					StoredComponentValues.bRendersCustomDepth = InComponent->bRenderCustomDepth > 0;
+					StoredComponentValues.CustomStencilId = InComponent->CustomDepthStencilValue;
+
+					StoredComponentValues.bRendersInMainPass = InComponent->bRenderInMainPass;
+					StoredComponentValues.bRendersDepth = InComponent->bRenderInDepthPass;
+
+					// @todo: move this impl specific stuff to a Handler/Provider subsystem
+					if (UDynamicMeshComponent* DynamicMeshComponent = Cast<UDynamicMeshComponent>(InComponent))
+					{
+						StoredComponentValues.bDrawWireframe = DynamicMeshComponent->GetEnableWireframeRenderPass();
+						StoredComponentValues.WireframeColor = DynamicMeshComponent->WireframeColor;
+					}
+
+					TArray<UMaterialInterface*> UsedMaterials;
+					InComponent->GetUsedMaterials(UsedMaterials);
+
+					Algo::Transform(UsedMaterials, StoredComponentValues.Materials, [](UMaterialInterface* InMaterial)
+					{
+						return InMaterial;
+					});
 				});
-			}
 
 			if (!bWasOneOrMorePrimitiveComponents)
 			{
@@ -551,10 +587,46 @@ void UAvaGizmoComponent::OnPostRegisterParentComponents(AActor* InActor)
 			if (!bAreGizmoValuesApplied)
 			{
 				StoreComponentValues();
-				
+
 				// Allows for all other components to do stuff first
 				ApplyGizmoValues();
 			}
 		}
 	}
 }
+	
+void UAvaGizmoComponent::ForEachPrimitiveComponent(TFunctionRef<void(UPrimitiveComponent*)> InFunc)
+{
+	if (AActor* Owner = GetOwner();
+		IsValid(Owner))
+	{
+		TArray<UPrimitiveComponent*> PrimitiveComponents;
+		Owner->GetComponents<UPrimitiveComponent>(PrimitiveComponents, bApplyToChildActors);
+	
+		for (UPrimitiveComponent* Component : PrimitiveComponents)
+		{
+			if (!IsValid(Component))
+			{
+				continue;
+			}
+	
+			InFunc(Component);
+		}
+	}
+}
+
+void UAvaGizmoComponent::ApplyToEachPrimitiveComponent(TFunctionRef<void(UPrimitiveComponent*)> InFunc)
+{
+	if (!bIsGizmoEnabled || bIsRestoringValues)
+	{
+		return;
+	}
+	
+	if (AActor* Owner = GetOwner();
+		IsValid(Owner))
+	{
+		bAreGizmoValuesApplied = true;
+	
+		ForEachPrimitiveComponent(MoveTemp(InFunc));
+	}
+}
diff --git a/Engine/Plugins/Experimental/Avalanche/Source/Avalanche/Public/Framework/AvaGizmoComponent.h b/Engine/Plugins/Experimental/Avalanche/Source/Avalanche/Public/Framework/AvaGizmoComponent.h
--- a/Engine/Plugins/Experimental/Avalanche/Source/Avalanche/Public/Framework/AvaGizmoComponent.h
+++ b/Engine/Plugins/Experimental/Avalanche/Source/Avalanche/Public/Framework/AvaGizmoComponent.h
@@ -179,12 +179,17 @@ public:
 	/** Sets the custom stencil id, if applicable. */
 	void SetStencilId(const uint8 InStencilId);
 
-protected:
+private:
 	/** Stores all primitive component values (that this overrides). Call before applying gizmo values. */
 	void StoreComponentValues();
 
 	void OnPostRegisterParentComponents(AActor* InActor);
 
+	void ForEachPrimitiveComponent(TFunctionRef<void(UPrimitiveComponent*)> InFunc);
+	
+	/** Variation of ForEachPrimitiveComponent that checks various flags. */
+	void ApplyToEachPrimitiveComponent(TFunctionRef<void(UPrimitiveComponent*)> InFunc);
+
 protected:
 	FDelegateHandle PostRegisterComponentsHandle;
 
diff --git a/Engine/Plugins/Experimental/Avalanche/Source/AvalancheMask/Private/Handling/IAvaMaskMaterialCollectionHandle.h b/Engine/Plugins/Experimental/Avalanche/Source/AvalancheMask/Private/Handling/IAvaMaskMaterialCollectionHandle.h
--- a/Engine/Plugins/Experimental/Avalanche/Source/AvalancheMask/Private/Handling/IAvaMaskMaterialCollectionHandle.h
+++ b/Engine/Plugins/Experimental/Avalanche/Source/AvalancheMask/Private/Handling/IAvaMaskMaterialCollectionHandle.h
@@ -240,7 +240,7 @@ bool TAvaMaskMaterialCollectionHandle<HandleDataType>::ApplyOriginalState(const
 						SetMaterial(InComponent, InSlotIdx, InMaterial);
 					});
 			}
-			UE_LOG(LogAvaMask, Warning, TEXT("MaterialData not found for slot %u"), InSlotIdx);
+			UE_LOG(LogAvaMask, Display, TEXT("MaterialData not found for slot %u"), InSlotIdx);
 			return false;
 		});
 
@@ -281,7 +281,7 @@ bool TAvaMaskMaterialCollectionHandle<HandleDataType>::ApplyModifiedState(
 						SetMaterial(InComponent, InSlotIdx, InMaterial);
 					});
 			}
-			UE_LOG(LogAvaMask, Warning, TEXT("MaterialData not found for slot %u"), InSlotIdx);
+			UE_LOG(LogAvaMask, Display, TEXT("MaterialData not found for slot %u"), InSlotIdx);
 			return false;
 		});
 
diff --git a/Engine/Plugins/Experimental/Avalanche/Source/AvalancheMask/Private/Mask2D/AvaMask2DBaseModifier.cpp b/Engine/Plugins/Experimental/Avalanche/Source/AvalancheMask/Private/Mask2D/AvaMask2DBaseModifier.cpp
--- a/Engine/Plugins/Experimental/Avalanche/Source/AvalancheMask/Private/Mask2D/AvaMask2DBaseModifier.cpp
+++ b/Engine/Plugins/Experimental/Avalanche/Source/AvalancheMask/Private/Mask2D/AvaMask2DBaseModifier.cpp
@@ -330,7 +330,9 @@ void UAvaMask2DBaseModifier::RestorePreState()
 	Super::RestorePreState();
 
 	TGuardValue<bool> RestoreStateGuard(bIsRestoring, true);
-	for (const TPair<TWeakObjectPtr<AActor>, FAvaMask2DActorData>& ActorDataPair : ActorData)
+	
+	TMap<TWeakObjectPtr<AActor>, FAvaMask2DActorData> ActorDataCopy = ActorData;
+	for (const TPair<TWeakObjectPtr<AActor>, FAvaMask2DActorData>& ActorDataPair : ActorDataCopy)
 	{
 		if (AActor* Actor = ActorDataPair.Key.Get())
 		{
