[Motion Design]
 - Show outliner filters active at a glance
 - Shift + Click outliner filter button to clear filters

diff --git a/Engine/Plugins/Experimental/Avalanche/Source/AvalancheOutliner/Private/Slate/SAvaOutlinerItemFilters.cpp b/Engine/Plugins/Experimental/Avalanche/Source/AvalancheOutliner/Private/Slate/SAvaOutlinerItemFilters.cpp
--- a/Engine/Plugins/Experimental/Avalanche/Source/AvalancheOutliner/Private/Slate/SAvaOutlinerItemFilters.cpp
+++ b/Engine/Plugins/Experimental/Avalanche/Source/AvalancheOutliner/Private/Slate/SAvaOutlinerItemFilters.cpp
@@ -5,8 +5,11 @@
 #include "AvaOutlinerStyle.h"
 #include "AvaOutlinerView.h"
 #include "Filters/IAvaOutlinerItemFilter.h"
+#include "Framework/Application/SlateApplication.h"
+#include "Internationalization/Text.h"
 #include "Slate/SAvaOutliner.h"
 #include "Styling/AvaOutlinerStyleUtils.h"
+#include "Styling/SlateBrush.h"
 #include "Styling/StyleColors.h"
 #include "Styling/ToolBarStyle.h"
 #include "Widgets/Colors/SColorBlock.h"
@@ -17,6 +20,7 @@
 #include "Widgets/Layout/SScaleBox.h"
 #include "Widgets/Layout/SScrollBox.h"
 #include "Widgets/Layout/SSeparator.h"
+#include "Widgets/SOverlay.h"
 
 #define LOCTEXT_NAMESPACE "SAvaOutlinerItemFilters"
 
@@ -69,11 +73,25 @@ void SAvaOutlinerItemFilters::Construct(const FArguments& InArgs, const TSharedR
 		[
 			SNew(SButton)
 			.ButtonStyle(&GetFilterItemMenuButtonStyle())
+			.ToolTipText(this, &SAvaOutlinerItemFilters::GetFilterItemMenuButtonToolTip)
 			.OnClicked(this, &SAvaOutlinerItemFilters::ToggleShowItemFilters)
 			[
-				SNew(SImage)
-				.Image(FAvaOutlinerStyle::Get().GetBrush(TEXT("AvaOutliner.FilterIcon")))
-				.DesiredSizeOverride(FVector2D(24.f))
+				SNew(SOverlay)
+				+ SOverlay::Slot()
+				[
+					SNew(SImage)
+					.Image(FAvaOutlinerStyle::Get().GetBrush(TEXT("AvaOutliner.FilterIcon")))
+					.DesiredSizeOverride(FVector2D(24.f))
+				]
+				+ SOverlay::Slot()
+				.HAlign(HAlign_Right)
+				.VAlign(VAlign_Bottom)
+				[
+					SNew(SImage)
+					.Image(this, &SAvaOutlinerItemFilters::GetFilterItemMenuButtonIcon)
+					.DesiredSizeOverride(FVector2D(12.f))
+					.RenderTransform(FSlateRenderTransform(FVector2f(2.f, 0.f)))
+				]
 			]
 		];
 
@@ -318,6 +336,22 @@ FReply SAvaOutlinerItemFilters::ToggleShowItemFilters()
 {
 	if (const TSharedPtr<FAvaOutlinerView> OutlinerView = OutlinerViewWeak.Pin())
 	{
+		if (FSlateApplication::Get().GetModifierKeys().IsShiftDown())
+		{
+			for (const TSharedPtr<IAvaOutlinerItemFilter>& Filter : OutlinerView->GetItemFilters())
+			{
+				OutlinerView->DisableItemFilter(Filter, false);
+			}
+			for (const TSharedPtr<IAvaOutlinerItemFilter>& Filter : OutlinerView->GetCustomItemFilters())
+			{
+				OutlinerView->DisableItemFilter(Filter, false);
+			}
+
+			OutlinerView->Refresh();
+
+			return FReply::Handled();
+		}
+
 		OutlinerView->ToggleShowItemFilters();
 		OnShowItemFiltersChanged(*OutlinerView);
 
@@ -380,4 +414,55 @@ FReply SAvaOutlinerItemFilters::DeselectAll()
 	return FReply::Unhandled();
 }
 
+FText SAvaOutlinerItemFilters::GetFilterItemMenuButtonToolTip() const
+{
+	if (const TSharedPtr<FAvaOutlinerView> OutlinerView = OutlinerViewWeak.Pin())
+	{
+		FText FilterText;
+
+		for (const TSharedPtr<IAvaOutlinerItemFilter>& Filter : OutlinerView->GetActiveItemFilters())
+		{
+			if (OutlinerView->IsItemFilterEnabled(Filter))
+			{
+				FilterText = FText::Format(LOCTEXT("ItemFilterText", "{0}\n  {1}"), FilterText, Filter->GetTooltipText());
+			}
+		}
+
+		return FilterText.IsEmpty()
+			? LOCTEXT("NoActiveFiltersTooltip", "No active filters")
+			: FText::Format(LOCTEXT("CustomItemFilterText", "Active Filters:{0}\n\n"
+				"Shift + Click to clear all active filters"), FilterText);
+	}
+
+	return FText();
+}
+
+const FSlateBrush* SAvaOutlinerItemFilters::GetFilterItemMenuButtonIcon() const
+{
+	if (const TSharedPtr<FAvaOutlinerView> OutlinerView = OutlinerViewWeak.Pin())
+	{
+		const TSet<TSharedPtr<IAvaOutlinerItemFilter>>& ActiveFilters = OutlinerView->GetActiveItemFilters();
+		const int32 ActiveFilterCount = ActiveFilters.Num();
+
+		const FSlateBrush* OutBrush = nullptr;
+
+		for (const TSharedPtr<IAvaOutlinerItemFilter>& ItemFilter : ActiveFilters)
+		{
+			if (OutBrush)
+			{
+				OutBrush = FAppStyle::GetBrush(TEXT("MessageLog.Warning"));
+				break;
+			}
+			else
+			{
+				OutBrush = ItemFilter->GetIconBrush();
+			}
+		}
+
+		return OutBrush;
+	}
+
+	return nullptr;
+}
+
 #undef LOCTEXT_NAMESPACE
diff --git a/Engine/Plugins/Experimental/Avalanche/Source/AvalancheOutliner/Private/Slate/SAvaOutlinerItemFilters.h b/Engine/Plugins/Experimental/Avalanche/Source/AvalancheOutliner/Private/Slate/SAvaOutlinerItemFilters.h
--- a/Engine/Plugins/Experimental/Avalanche/Source/AvalancheOutliner/Private/Slate/SAvaOutlinerItemFilters.h
+++ b/Engine/Plugins/Experimental/Avalanche/Source/AvalancheOutliner/Private/Slate/SAvaOutlinerItemFilters.h
@@ -7,10 +7,12 @@
 #include "Widgets/SCompoundWidget.h"
 
 class FAvaOutlinerView;
+class FText;
 class IAvaOutlinerItemFilter;
 class SAvaOutlinerItemFilters;
 class SBox;
 class SScrollBox;
+struct FSlateBrush;
 
 class SAvaOutlinerItemFilters : public SCompoundWidget
 {
@@ -44,6 +46,9 @@ public:
 	FReply DeselectAll();
 
 private:
+	FText GetFilterItemMenuButtonToolTip() const;
+	const FSlateBrush* GetFilterItemMenuButtonIcon() const;
+
 	TWeakPtr<FAvaOutlinerView> OutlinerViewWeak;
 
 	TSharedPtr<SBox> ItemFilterBox;
