Motion Design Sequencer:
- Fix undoing a sequence causing a sequence to not be properly recovered and sometimes crashing

diff --git a/Engine/Plugins/Experimental/Avalanche/Source/AvalancheSequence/Private/AvaSequenceController.cpp b/Engine/Plugins/Experimental/Avalanche/Source/AvalancheSequence/Private/AvaSequenceController.cpp
--- a/Engine/Plugins/Experimental/Avalanche/Source/AvalancheSequence/Private/AvaSequenceController.cpp
+++ b/Engine/Plugins/Experimental/Avalanche/Source/AvalancheSequence/Private/AvaSequenceController.cpp
@@ -75,7 +75,7 @@ FAvaSequenceController::FAvaSequenceController(UAvaSequence& InSequence, IAvaSeq
 
 FAvaSequenceController::~FAvaSequenceController()
 {
-	if (UAvaSequence* const Sequence = SequenceWeak.Get())
+	if (UAvaSequence* const Sequence = SequenceWeak.Get(/*bEvenIfPendingKill*/true))
 	{
 		Sequence->GetOnTreeNodeUpdated().RemoveAll(this);
 	}
diff --git a/Engine/Plugins/Experimental/Avalanche/Source/AvalancheSequencer/Private/AvaSequencer.cpp b/Engine/Plugins/Experimental/Avalanche/Source/AvalancheSequencer/Private/AvaSequencer.cpp
--- a/Engine/Plugins/Experimental/Avalanche/Source/AvalancheSequencer/Private/AvaSequencer.cpp
+++ b/Engine/Plugins/Experimental/Avalanche/Source/AvalancheSequencer/Private/AvaSequencer.cpp
@@ -920,6 +920,7 @@ void FAvaSequencer::DeleteSequence_Execute()
 		check(Item.IsValid());
 		if (UAvaSequence* Sequence = Item->GetSequence())
 		{
+			Sequence->Modify();
 			SequenceProvider->RemoveSequence(Sequence);
 			RemovedSequences.Add(Sequence);
 		}
