[SVG Importer] Geometry Modifier not shown properly right after duplicating SVG related actors like Joined SVG

Add base class which can be used by SVG Shape Components to signal when the dynamic mesh has been updated.

AvaSVGEditor module registers to new delegate, so that it can update modifier stack right after the event.

diff --git a/Engine/Plugins/Experimental/Avalanche/Source/AvalancheSVGEditor/Private/AvaSVGEditorModule.cpp b/Engine/Plugins/Experimental/Avalanche/Source/AvalancheSVGEditor/Private/AvaSVGEditorModule.cpp
--- a/Engine/Plugins/Experimental/Avalanche/Source/AvalancheSVGEditor/Private/AvaSVGEditorModule.cpp
+++ b/Engine/Plugins/Experimental/Avalanche/Source/AvalancheSVGEditor/Private/AvaSVGEditorModule.cpp
@@ -2,6 +2,7 @@
 
 #include "AvaSVGEditorModule.h"
 #include "AvaInteractiveToolsDelegates.h"
+#include "GameFramework/Actor.h"
 #include "IAvalancheInteractiveToolsModule.h"
 #include "Modifiers/ActorModifierCoreStack.h"
 #include "Modifiers/AvaBevelModifier.h"
@@ -19,6 +20,7 @@ void FAvaSVGEditorModule::StartupModule()
 	FAvaInteractiveToolsDelegates::GetRegisterToolsDelegate().AddRaw(this, &FAvaSVGEditorModule::RegisterTools);
 
 	USVGEngineSubsystem::OnSVGActorSplit().BindRaw(this, &FAvaSVGEditorModule::OnSVGActorSplit);
+	USVGEngineSubsystem::OnSVGShapesUpdated().BindRaw(this, &FAvaSVGEditorModule::OnSVGShapesUpdated);
 }
 
 void FAvaSVGEditorModule::ShutdownModule()
@@ -26,6 +28,7 @@ void FAvaSVGEditorModule::ShutdownModule()
 	FAvaInteractiveToolsDelegates::GetRegisterToolsDelegate().RemoveAll(this);
 
 	USVGEngineSubsystem::OnSVGActorSplit().Unbind();
+	USVGEngineSubsystem::OnSVGShapesUpdated().Unbind();
 }
 
 void FAvaSVGEditorModule::RegisterTools(IAvalancheInteractiveToolsModule* InModule)
@@ -127,6 +130,25 @@ void FAvaSVGEditorModule::OnSVGActorSplit(ASVGShapesParentActor* InSVGShapesPare
 	}
 }
 
+void FAvaSVGEditorModule::OnSVGShapesUpdated(AActor* InActor) const
+{
+	const UActorModifierCoreSubsystem* ModifierCoreSubsystem = UActorModifierCoreSubsystem::Get();
+
+	if (!ModifierCoreSubsystem)
+	{
+		return;
+	}
+
+	UActorModifierCoreStack* ModifierStack = ModifierCoreSubsystem->GetActorModifierStack(InActor);
+
+	if (!ModifierStack)
+	{
+		return;
+	}
+
+ 	ModifierStack->MarkModifierDirty();
+}
+
 #undef LOCTEXT_NAMESPACE
 
 IMPLEMENT_MODULE(FAvaSVGEditorModule, AvalancheSVGEditorModule)
diff --git a/Engine/Plugins/Experimental/Avalanche/Source/AvalancheSVGEditor/Private/AvaSVGEditorModule.h b/Engine/Plugins/Experimental/Avalanche/Source/AvalancheSVGEditor/Private/AvaSVGEditorModule.h
--- a/Engine/Plugins/Experimental/Avalanche/Source/AvalancheSVGEditor/Private/AvaSVGEditorModule.h
+++ b/Engine/Plugins/Experimental/Avalanche/Source/AvalancheSVGEditor/Private/AvaSVGEditorModule.h
@@ -5,6 +5,7 @@
 #include "Containers/ContainersFwd.h"
 #include "Modules/ModuleManager.h"
 
+class AActor;
 class ASVGShapesParentActor;
 class IAvalancheInteractiveToolsModule;
 class UEdMode;
@@ -20,4 +21,5 @@ private:
 	void RegisterTools(IAvalancheInteractiveToolsModule* InModule);
 
 	void OnSVGActorSplit(ASVGShapesParentActor* InSVGShapesParent);
+	void OnSVGShapesUpdated(AActor* InActor) const;
 };
diff --git a/Engine/Plugins/Experimental/SVGImporter/Source/SVGImporter/Private/ProceduralMeshes/JoinedSVGDynamicMeshComponent.cpp b/Engine/Plugins/Experimental/SVGImporter/Source/SVGImporter/Private/ProceduralMeshes/JoinedSVGDynamicMeshComponent.cpp
--- a/Engine/Plugins/Experimental/SVGImporter/Source/SVGImporter/Private/ProceduralMeshes/JoinedSVGDynamicMeshComponent.cpp
+++ b/Engine/Plugins/Experimental/SVGImporter/Source/SVGImporter/Private/ProceduralMeshes/JoinedSVGDynamicMeshComponent.cpp
@@ -8,6 +8,7 @@
 #include "Materials/MaterialInstanceDynamic.h"
 #include "Materials/MaterialInterface.h"
 #include "Misc/Paths.h"
+#include "SVGActor.h"
 #include "SVGImporterUtils.h"
 #include "UObject/ConstructorHelpers.h"
 
@@ -148,6 +149,8 @@ void UJoinedSVGDynamicMeshComponent::LoadStoredMesh()
 	{
 		StoreCurrentMesh();
 	}
+
+	MarkSVGMeshUpdated();
 }
 
 void UJoinedSVGDynamicMeshComponent::PostLoad()
diff --git a/Engine/Plugins/Experimental/SVGImporter/Source/SVGImporter/Private/ProceduralMeshes/SVGBaseDynamicMeshComponent.cpp b/Engine/Plugins/Experimental/SVGImporter/Source/SVGImporter/Private/ProceduralMeshes/SVGBaseDynamicMeshComponent.cpp
new file mode 100644
--- /dev/null
+++ b/Engine/Plugins/Experimental/SVGImporter/Source/SVGImporter/Private/ProceduralMeshes/SVGBaseDynamicMeshComponent.cpp
@@ -0,0 +1,29 @@
+// Copyright Epic Games, Inc. All Rights Reserved.
+
+#include "ProceduralMeshes/SVGBaseDynamicMeshComponent.h"
+#include "SVGEngineSubsystem.h"
+
+USVGBaseDynamicMeshComponent::USVGBaseDynamicMeshComponent()
+{
+	bMeshHasBeenUpdated = false;
+}
+
+void USVGBaseDynamicMeshComponent::OnRegister()
+{
+	Super::OnRegister();
+
+	if (bMeshHasBeenUpdated)
+	{
+		if (AActor* Owner = GetOwner())
+		{
+			USVGEngineSubsystem::OnSVGShapesUpdated().Execute(Owner);
+		}
+
+		bMeshHasBeenUpdated = false;
+	}
+}
+
+void USVGBaseDynamicMeshComponent::MarkSVGMeshUpdated()
+{
+	bMeshHasBeenUpdated = true;
+}
diff --git a/Engine/Plugins/Experimental/SVGImporter/Source/SVGImporter/Private/ProceduralMeshes/SVGDynamicMeshComponent.cpp b/Engine/Plugins/Experimental/SVGImporter/Source/SVGImporter/Private/ProceduralMeshes/SVGDynamicMeshComponent.cpp
--- a/Engine/Plugins/Experimental/SVGImporter/Source/SVGImporter/Private/ProceduralMeshes/SVGDynamicMeshComponent.cpp
+++ b/Engine/Plugins/Experimental/SVGImporter/Source/SVGImporter/Private/ProceduralMeshes/SVGDynamicMeshComponent.cpp
@@ -1,7 +1,6 @@
 // Copyright Epic Games, Inc. All Rights Reserved.
 
 #include "ProceduralMeshes/SVGDynamicMeshComponent.h"
-
 #include "DynamicMesh/MeshTransforms.h"
 #include "GeometryScript/GeometryScriptSelectionTypes.h"
 #include "GeometryScript/MeshSimplifyFunctions.h"
@@ -364,6 +363,8 @@ void USVGDynamicMeshComponent::LoadStoredMesh()
 		RegenerateMesh();
 		StoreCurrentMesh();
 	}
+
+	MarkSVGMeshUpdated();
 }
 
 void USVGDynamicMeshComponent::PostLoad()
diff --git a/Engine/Plugins/Experimental/SVGImporter/Source/SVGImporter/Private/USVGEngineSubsystem.cpp b/Engine/Plugins/Experimental/SVGImporter/Source/SVGImporter/Private/USVGEngineSubsystem.cpp
--- a/Engine/Plugins/Experimental/SVGImporter/Source/SVGImporter/Private/USVGEngineSubsystem.cpp
+++ b/Engine/Plugins/Experimental/SVGImporter/Source/SVGImporter/Private/USVGEngineSubsystem.cpp
@@ -4,6 +4,7 @@
 #include "Engine/Engine.h"
 
 FOnSVGActorSplit USVGEngineSubsystem::OnSVGActorSplitDelegate;
+FOnSVGShapesUpdated USVGEngineSubsystem::OnSVGShapesUpdatedDelegate;
 
 USVGEngineSubsystem* USVGEngineSubsystem::Get()
 {
diff --git a/Engine/Plugins/Experimental/SVGImporter/Source/SVGImporter/Public/ProceduralMeshes/JoinedSVGDynamicMeshComponent.h b/Engine/Plugins/Experimental/SVGImporter/Source/SVGImporter/Public/ProceduralMeshes/JoinedSVGDynamicMeshComponent.h
--- a/Engine/Plugins/Experimental/SVGImporter/Source/SVGImporter/Public/ProceduralMeshes/JoinedSVGDynamicMeshComponent.h
+++ b/Engine/Plugins/Experimental/SVGImporter/Source/SVGImporter/Public/ProceduralMeshes/JoinedSVGDynamicMeshComponent.h
@@ -2,7 +2,7 @@
 
 #pragma once
 
-#include "Components/DynamicMeshComponent.h"
+#include "SVGBaseDynamicMeshComponent.h"
 #include "UObject/ObjectPtr.h"
 #include "JoinedSVGDynamicMeshComponent.generated.h"
 
@@ -66,7 +66,7 @@ enum class EJoinedSVGMeshColoring : uint8
 };
 
 UCLASS(MinimalAPI, ClassGroup=(Custom), meta=(BlueprintSpawnableComponent))
-class UJoinedSVGDynamicMeshComponent : public UDynamicMeshComponent
+class UJoinedSVGDynamicMeshComponent : public USVGBaseDynamicMeshComponent
 {
 	friend class FJoinedSVGDynamicMeshComponentCustomization;
 
diff --git a/Engine/Plugins/Experimental/SVGImporter/Source/SVGImporter/Public/ProceduralMeshes/SVGBaseDynamicMeshComponent.h b/Engine/Plugins/Experimental/SVGImporter/Source/SVGImporter/Public/ProceduralMeshes/SVGBaseDynamicMeshComponent.h
new file mode 100644
--- /dev/null
+++ b/Engine/Plugins/Experimental/SVGImporter/Source/SVGImporter/Public/ProceduralMeshes/SVGBaseDynamicMeshComponent.h
@@ -0,0 +1,26 @@
+// Copyright Epic Games, Inc. All Rights Reserved.
+
+#pragma once
+
+#include "Components/DynamicMeshComponent.h"
+#include "SVGBaseDynamicMeshComponent.generated.h"
+
+UCLASS()
+class USVGBaseDynamicMeshComponent : public UDynamicMeshComponent
+{
+	GENERATED_BODY()
+
+public:
+	// Sets default values for this component's properties
+	USVGBaseDynamicMeshComponent();
+
+	//~ Begin UActorComponent
+	virtual void OnRegister() override;
+	//~ End UActorComponent
+
+protected:
+	void MarkSVGMeshUpdated();
+
+private:
+	bool bMeshHasBeenUpdated;
+};
diff --git a/Engine/Plugins/Experimental/SVGImporter/Source/SVGImporter/Public/ProceduralMeshes/SVGDynamicMeshComponent.h b/Engine/Plugins/Experimental/SVGImporter/Source/SVGImporter/Public/ProceduralMeshes/SVGDynamicMeshComponent.h
--- a/Engine/Plugins/Experimental/SVGImporter/Source/SVGImporter/Public/ProceduralMeshes/SVGDynamicMeshComponent.h
+++ b/Engine/Plugins/Experimental/SVGImporter/Source/SVGImporter/Public/ProceduralMeshes/SVGDynamicMeshComponent.h
@@ -2,7 +2,7 @@
 
 #pragma once
 
-#include "Components/DynamicMeshComponent.h"
+#include "SVGBaseDynamicMeshComponent.h"
 #include "SVGDynamicMeshComponent.generated.h"
 
 class UMaterial;
@@ -43,7 +43,7 @@ DECLARE_DELEGATE(FSVGMeshActionDelegate)
 class UStaticMesh;
 
 UCLASS(MinimalAPI, ClassGroup=(SVG))
-class USVGDynamicMeshComponent : public UDynamicMeshComponent
+class USVGDynamicMeshComponent : public USVGBaseDynamicMeshComponent
 {
 	GENERATED_BODY()
 
diff --git a/Engine/Plugins/Experimental/SVGImporter/Source/SVGImporter/Public/SVGEngineSubsystem.h b/Engine/Plugins/Experimental/SVGImporter/Source/SVGImporter/Public/SVGEngineSubsystem.h
--- a/Engine/Plugins/Experimental/SVGImporter/Source/SVGImporter/Public/SVGEngineSubsystem.h
+++ b/Engine/Plugins/Experimental/SVGImporter/Source/SVGImporter/Public/SVGEngineSubsystem.h
@@ -10,6 +10,7 @@ class ASVGShapesParentActor;
 
 DECLARE_DELEGATE_OneParam(FSVGActorComponentsReady, ASVGActor*)
 DECLARE_DELEGATE_OneParam(FOnSVGActorSplit, ASVGShapesParentActor*)
+DECLARE_DELEGATE_OneParam(FOnSVGShapesUpdated, AActor*)
 
 UCLASS()
 class USVGEngineSubsystem : public UEngineSubsystem
@@ -22,9 +23,11 @@ public:
 	FSVGActorComponentsReady& GetSVGActorComponentsReadyDelegate() { return SVGActorComponentsReady; }
 
 	static FOnSVGActorSplit& OnSVGActorSplit() { return OnSVGActorSplitDelegate; }
+	static FOnSVGShapesUpdated& OnSVGShapesUpdated() { return OnSVGShapesUpdatedDelegate; }
 
 protected:
 	SVGIMPORTER_API static FOnSVGActorSplit OnSVGActorSplitDelegate;
+	SVGIMPORTER_API static FOnSVGShapesUpdated OnSVGShapesUpdatedDelegate;
 
 	FSVGActorComponentsReady SVGActorComponentsReady;
 };
