[MotionDesign] GizmoComponent delayed application

diff --git a/Engine/Plugins/Experimental/Avalanche/Source/Avalanche/Private/Framework/AvaGizmoComponent.cpp b/Engine/Plugins/Experimental/Avalanche/Source/Avalanche/Private/Framework/AvaGizmoComponent.cpp
--- a/Engine/Plugins/Experimental/Avalanche/Source/Avalanche/Private/Framework/AvaGizmoComponent.cpp
+++ b/Engine/Plugins/Experimental/Avalanche/Source/Avalanche/Private/Framework/AvaGizmoComponent.cpp
@@ -2,6 +2,7 @@
 
 #include "Framework/AvaGizmoComponent.h"
 
+#include "Async/Async.h"
 #include "AvaLog.h"
 #include "Components/DynamicMeshComponent.h"
 #include "Components/PrimitiveComponent.h"
@@ -589,7 +590,13 @@ void UAvaGizmoComponent::OnPostRegisterParentComponents(AActor* InActor)
 				StoreComponentValues();
 
 				// Allows for all other components to do stuff first
-				ApplyGizmoValues();
+				AsyncTask(ENamedThreads::GameThread, [this]()
+				{
+					if (this)
+					{
+						ApplyGizmoValues();
+					}
+				});
 			}
 		}
 	}
