[MotionDesign] GizmoComponent CIS fix

diff --git a/Engine/Plugins/Experimental/Avalanche/Source/Avalanche/Private/Framework/AvaGizmoComponent.cpp b/Engine/Plugins/Experimental/Avalanche/Source/Avalanche/Private/Framework/AvaGizmoComponent.cpp
--- a/Engine/Plugins/Experimental/Avalanche/Source/Avalanche/Private/Framework/AvaGizmoComponent.cpp
+++ b/Engine/Plugins/Experimental/Avalanche/Source/Avalanche/Private/Framework/AvaGizmoComponent.cpp
@@ -590,11 +590,12 @@ void UAvaGizmoComponent::OnPostRegisterParentComponents(AActor* InActor)
 				StoreComponentValues();
 
 				// Allows for all other components to do stuff first
-				AsyncTask(ENamedThreads::GameThread, [this]()
+				AsyncTask(ENamedThreads::GameThread, [WeakThis = TWeakObjectPtr<UAvaGizmoComponent>(this)]()
 				{
-					if (this)
+					if (UAvaGizmoComponent* StrongThis = WeakThis.Get();
+						UObjectInitialized() && StrongThis)
 					{
-						ApplyGizmoValues();
+						StrongThis->ApplyGizmoValues();
 					}
 				});
 			}
