Removing include of the current file.

diff --git a/Engine/Plugins/Experimental/Avalanche/Source/AvalancheMedia/Internal/Broadcast/AvaBroadcastDevice.h b/Engine/Plugins/Experimental/Avalanche/Source/AvalancheMedia/Internal/Broadcast/AvaBroadcastDevice.h
--- a/Engine/Plugins/Experimental/Avalanche/Source/AvalancheMedia/Internal/Broadcast/AvaBroadcastDevice.h
+++ b/Engine/Plugins/Experimental/Avalanche/Source/AvalancheMedia/Internal/Broadcast/AvaBroadcastDevice.h
@@ -2,7 +2,6 @@
 
 #pragma once
 
-#include "AvaBroadcastDevice.h"
 #include "Containers/UnrealString.h"
 #include "MediaIOCoreDefinitions.h"
 #include "Internationalization/Text.h"
