Motion Design:
- Fix issue with AvaSequences doing unwanted cleanup when closing the Director Blueprint, or closing any editor that had some underlying PreviewWorld.
- This causes data loss on any work done after this cleanup (as the object is still being referenced, it's just renamed to transient package)

diff --git a/Engine/Plugins/Experimental/Avalanche/Source/AvalancheSequence/Private/AvaSequence.cpp b/Engine/Plugins/Experimental/Avalanche/Source/AvalancheSequence/Private/AvaSequence.cpp
--- a/Engine/Plugins/Experimental/Avalanche/Source/AvalancheSequence/Private/AvaSequence.cpp
+++ b/Engine/Plugins/Experimental/Avalanche/Source/AvalancheSequence/Private/AvaSequence.cpp
@@ -634,6 +634,12 @@ void UAvaSequence::OnOuterWorldRenamed(const TCHAR* InName, UObject* InNewOuter,
 
 void UAvaSequence::OnWorldCleanup(UWorld* InWorld, bool bInSessionEnded, bool bInCleanupResources)
 {
+	// Ignore cleanups from other worlds
+	if (GetTypedOuter<UWorld>() != InWorld)
+	{
+		return;
+	}
+
 	if (!bInCleanupResources)
 	{
 		return;
