[Motion Design]
 - Fix ease curve tool not reading key values after save/load
 - Filter out invalid keys when averaging tangents

diff --git a/Engine/Plugins/Experimental/Avalanche/Source/AvalancheSequencer/Private/AvaSequencer.cpp b/Engine/Plugins/Experimental/Avalanche/Source/AvalancheSequencer/Private/AvaSequencer.cpp
--- a/Engine/Plugins/Experimental/Avalanche/Source/AvalancheSequencer/Private/AvaSequencer.cpp
+++ b/Engine/Plugins/Experimental/Avalanche/Source/AvalancheSequencer/Private/AvaSequencer.cpp
@@ -1723,4 +1723,12 @@ TSharedRef<FAvaEaseCurveTool> FAvaSequencer::GetEaseCurveTool() const
 	return EaseCurveTool.ToSharedRef(); 
 }
 
+void FAvaSequencer::UpdateEaseCurveTool() const
+{
+	if (EaseCurveTool.IsValid())
+	{
+		EaseCurveTool->UpdateEaseCurveFromSequencerKeySelections();
+	}
+}
+
 #undef LOCTEXT_NAMESPACE
diff --git a/Engine/Plugins/Experimental/Avalanche/Source/AvalancheSequencer/Private/AvaSequencer.h b/Engine/Plugins/Experimental/Avalanche/Source/AvalancheSequencer/Private/AvaSequencer.h
--- a/Engine/Plugins/Experimental/Avalanche/Source/AvalancheSequencer/Private/AvaSequencer.h
+++ b/Engine/Plugins/Experimental/Avalanche/Source/AvalancheSequencer/Private/AvaSequencer.h
@@ -182,6 +182,8 @@ public:
 	void FixBindingHierarchy();
 
 	TSharedRef<FAvaEaseCurveTool> GetEaseCurveTool() const;
+
+	void UpdateEaseCurveTool() const;
 	
 	TArrayView<TWeakObjectPtr<>> ResolveBoundObjects(const FGuid& InBindingId, class UMovieSceneSequence* Sequence) const;
 
diff --git a/Engine/Plugins/Experimental/Avalanche/Source/AvalancheSequencer/Private/DetailsView/SAvaKeyFrameEdit.cpp b/Engine/Plugins/Experimental/Avalanche/Source/AvalancheSequencer/Private/DetailsView/SAvaKeyFrameEdit.cpp
--- a/Engine/Plugins/Experimental/Avalanche/Source/AvalancheSequencer/Private/DetailsView/SAvaKeyFrameEdit.cpp
+++ b/Engine/Plugins/Experimental/Avalanche/Source/AvalancheSequencer/Private/DetailsView/SAvaKeyFrameEdit.cpp
@@ -79,6 +79,8 @@ void SAvaKeyFrameEdit::Construct(const FArguments& InArgs, const TSharedRef<FAva
 
 	const TSharedPtr<FAvaSequencer> AvaSequencer = AvaSequencerWeak.Pin();
 
+	InSequencer->UpdateEaseCurveTool();
+
 	ChildSlot
 		[
 			SNew(SScrollBox)
diff --git a/Engine/Plugins/Experimental/Avalanche/Source/AvalancheSequencer/Private/EaseCurveTool/AvaEaseCurveTool.cpp b/Engine/Plugins/Experimental/Avalanche/Source/AvalancheSequencer/Private/EaseCurveTool/AvaEaseCurveTool.cpp
--- a/Engine/Plugins/Experimental/Avalanche/Source/AvalancheSequencer/Private/EaseCurveTool/AvaEaseCurveTool.cpp
+++ b/Engine/Plugins/Experimental/Avalanche/Source/AvalancheSequencer/Private/EaseCurveTool/AvaEaseCurveTool.cpp
@@ -32,7 +32,7 @@ using namespace UE::Sequencer;
 FAvaEaseCurveTool::FAvaEaseCurveTool(const TSharedRef<FAvaSequencer>& InSequencer)
 	: AvaSequencerWeak(InSequencer)
 {
-	EaseCurve = NewObject<UAvaEaseCurve>(GetTransientPackage(), NAME_None, RF_Transactional);
+	EaseCurve = NewObject<UAvaEaseCurve>(GetTransientPackage(), NAME_None, RF_Transient | RF_Transactional);
 
 	const TSharedRef<ISequencer> Sequencer = InSequencer->GetSequencer();
 
@@ -74,27 +74,27 @@ void FAvaEaseCurveTool::CacheSelectionData()
 			continue;
 		}
 
-		FMovieSceneChannelHandle ChannelHandle = KeyArea->GetChannel();
-		if (ChannelHandle.GetChannelTypeName() != FMovieSceneDoubleChannel::StaticStruct()->GetFName())
+		FMovieSceneDoubleChannel* const DoubleChannel = KeyArea->GetChannel().Cast<FMovieSceneDoubleChannel>().Get();
+		if (!DoubleChannel)
 		{
 			continue;
 		}
 
 		FKeyDataCache::FChannelData& Entry = NewKeyCache.ChannelKeyData.FindOrAdd(ChannelModel->GetChannelName());
 		Entry.ChannelModel = ChannelModel;
-		Entry.DoubleChannel = static_cast<FMovieSceneDoubleChannel*>(ChannelHandle.Get());
+		Entry.DoubleChannel = DoubleChannel;
 		Entry.Section = ChannelModel->GetSection();
 		Entry.KeyHandles.Add(Key);
 
 		NewKeyCache.TotalSelectedKeys++;
 
-		const TArrayView<const FFrameNumber> ChannelTimes = Entry.DoubleChannel->GetTimes();
+		const TArrayView<const FFrameNumber> ChannelTimes = DoubleChannel->GetTimes();
 		const int32 AllKeyCount = ChannelTimes.Num();
 		const int32 SelectedKeyCount = Entry.KeyHandles.Num();
 
 		if (SelectedKeyCount == 1 && NewKeyCache.TotalSelectedKeys == 1)
 		{
-			const int32 KeyIndex = Entry.DoubleChannel->GetIndex(Entry.KeyHandles[0]);
+			const int32 KeyIndex = DoubleChannel->GetIndex(Entry.KeyHandles[0]);
 			if (KeyIndex == AllKeyCount - 1)
 			{
 				NewKeyCache.bIsLastOnlySelectedKey = true;
@@ -450,8 +450,16 @@ void FAvaEaseCurveTool::UpdateEaseCurveFromSequencerKeySelections()
 						, ChannelTimes[NextKeyIndex], ChannelValues[NextKeyIndex].Value
 						, DisplayRate, TickResolution);
 				}
+				else
+				{
+					ScaledTangents = FAvaEaseCurveTangents(ChannelTimes[KeyIndex].Value, ChannelValues[KeyIndex].Value, 0.0, 0.0);
+				}
 
-				KeySetTangents.Add(ScaledTangents);
+				if (NextKeyIndex != INDEX_NONE
+					&& ChannelValues[KeyIndex].Value != ChannelValues[NextKeyIndex].Value)
+				{
+					KeySetTangents.Add(ScaledTangents);
+				}
 			}
 		}
 	}
