Motion Design
- Swapped position of the isolate actors and toggle bounding boxes button

diff --git a/Engine/Plugins/Experimental/Avalanche/Source/AvalancheLevelViewport/Private/Widgets/SAvaLevelViewportStatusBarButtons.cpp b/Engine/Plugins/Experimental/Avalanche/Source/AvalancheLevelViewport/Private/Widgets/SAvaLevelViewportStatusBarButtons.cpp
--- a/Engine/Plugins/Experimental/Avalanche/Source/AvalancheLevelViewport/Private/Widgets/SAvaLevelViewportStatusBarButtons.cpp
+++ b/Engine/Plugins/Experimental/Avalanche/Source/AvalancheLevelViewport/Private/Widgets/SAvaLevelViewportStatusBarButtons.cpp
@@ -345,11 +345,11 @@ void SAvaLevelViewportStatusBarButtons::PopulateViewportButtons(TSharedPtr<SHori
 		[
 			ViewportStatusBarButton::MakeButton(
 				this,
-				CommandsRef.ToggleBoundingBoxes,
-				FAvaLevelViewportStyle::Get().GetBrush(TEXT("Button.BoundingBoxes")),
-				&SAvaLevelViewportStatusBarButtons::ToggleBoundingBoxes,
-				&SAvaLevelViewportStatusBarButtons::GetToggleBoundingBoxesEnabled,
-				&SAvaLevelViewportStatusBarButtons::GetToggleBoundingBoxesColor
+				CommandsRef.ToggleIsolateActors,
+				FAvaLevelViewportStyle::Get().GetBrush(TEXT("Button.IsolateActors")),
+				&SAvaLevelViewportStatusBarButtons::ToggleIsolateActors,
+				&SAvaLevelViewportStatusBarButtons::GetToggleIsolateActorsEnabled,
+				&SAvaLevelViewportStatusBarButtons::GetToggleIsolateActorsColor
 			)
 		];
 
@@ -359,11 +359,11 @@ void SAvaLevelViewportStatusBarButtons::PopulateViewportButtons(TSharedPtr<SHori
 		[
 			ViewportStatusBarButton::MakeButton(
 				this,
-				CommandsRef.ToggleIsolateActors,
-				FAvaLevelViewportStyle::Get().GetBrush(TEXT("Button.IsolateActors")),
-				&SAvaLevelViewportStatusBarButtons::ToggleIsolateActors,
-				&SAvaLevelViewportStatusBarButtons::GetToggleIsolateActorsEnabled,
-				&SAvaLevelViewportStatusBarButtons::GetToggleIsolateActorsColor
+				CommandsRef.ToggleBoundingBoxes,
+				FAvaLevelViewportStyle::Get().GetBrush(TEXT("Button.BoundingBoxes")),
+				&SAvaLevelViewportStatusBarButtons::ToggleBoundingBoxes,
+				&SAvaLevelViewportStatusBarButtons::GetToggleBoundingBoxesEnabled,
+				&SAvaLevelViewportStatusBarButtons::GetToggleBoundingBoxesColor
 			)
 		];
 
