[Motion Design]
Rundown Editor: Resetting page RC values to default

diff --git a/Engine/Plugins/Experimental/Avalanche/Source/AvalancheMedia/Private/Playable/AvaPlayableRemoteControlValues.cpp b/Engine/Plugins/Experimental/Avalanche/Source/AvalancheMedia/Private/Playable/AvaPlayableRemoteControlValues.cpp
--- a/Engine/Plugins/Experimental/Avalanche/Source/AvalancheMedia/Private/Playable/AvaPlayableRemoteControlValues.cpp
+++ b/Engine/Plugins/Experimental/Avalanche/Source/AvalancheMedia/Private/Playable/AvaPlayableRemoteControlValues.cpp
@@ -70,6 +70,63 @@ namespace UE::AvaPlayableRemoteControlValues::Private
 		return true;
 	}
 
+	bool HasSameValueAndDefault(const FAvaPlayableRemoteControlValue& InValue, const FAvaPlayableRemoteControlValue& InOtherValue)
+	{
+		return InOtherValue.IsSameValueAs(InValue.Value) && InOtherValue.bIsDefault == InValue.bIsDefault;
+	}
+
+	/** Compares values and default status. */
+	bool HasSameValuesAndDefaults(const TMap<FGuid, FAvaPlayableRemoteControlValue>& InValues, const TMap<FGuid, FAvaPlayableRemoteControlValue>& InOtherValues)
+	{
+		// If values count differ, consider as different
+		if (InValues.Num() != InOtherValues.Num())
+		{
+			return false;
+		}
+		// Both Value maps have the same count, so one cannot be a subset of another, therefore, a single find pass should determine equality 
+		for (const TPair<FGuid, FAvaPlayableRemoteControlValue>& Pair : InValues)
+		{
+			const FAvaPlayableRemoteControlValue* FoundValue = InOtherValues.Find(Pair.Key);
+			if (!FoundValue || !HasSameValueAndDefault(Pair.Value, *FoundValue))
+			{
+				// Other's Value wasn't found, or was different from the value of this.
+				return false;
+			}
+		}
+		return true;
+	}
+
+	bool ResetValues(const TMap<FGuid, FAvaPlayableRemoteControlValue>& InValues, TMap<FGuid, FAvaPlayableRemoteControlValue>& OutValues, bool bInIsDefaults)
+	{
+		TMap<FGuid, FAvaPlayableRemoteControlValue> ResetValues = InValues;
+		if (bInIsDefaults)
+		{
+			for (TPair<FGuid, FAvaPlayableRemoteControlValue>& Value : ResetValues)
+			{
+				Value.Value.bIsDefault = true;
+			}
+		}
+		
+		const bool bModified = !HasSameValuesAndDefaults(OutValues, ResetValues);
+		OutValues = ResetValues;
+	
+		return bModified;
+	}
+
+	bool ResetValue(const FAvaPlayableRemoteControlValue& InValue, FAvaPlayableRemoteControlValue& OutValue, bool bInIsDefaults)
+	{
+		FAvaPlayableRemoteControlValue ResetValue = InValue;
+		if (bInIsDefaults)
+		{
+			ResetValue.bIsDefault = true;
+		}
+
+		const bool bModified = !HasSameValueAndDefault(OutValue, ResetValue);
+		OutValue = ResetValue;
+
+		return bModified;
+	}
+
 	EAvaPlayableRemoteControlChanges ToRemoteControlChanges(bool bInModified, EAvaPlayableRemoteControlChanges InModifiedChanges)
 	{
 		return bInModified ? InModifiedChanges : EAvaPlayableRemoteControlChanges::None;
@@ -196,6 +253,35 @@ EAvaPlayableRemoteControlChanges FAvaPlayableRemoteControlValues::UpdateRemoteCo
 		| ToRemoteControlChanges(UpdateValues(InRemoteControlValues.ControllerValues, ControllerValues, bInUpdateDefaults), EAvaPlayableRemoteControlChanges::ControllerValues); 
 }
 
+EAvaPlayableRemoteControlChanges FAvaPlayableRemoteControlValues::ResetRemoteControlValues(const FAvaPlayableRemoteControlValues& InReferenceValues, bool bInIsDefaults)
+{
+	using namespace UE::AvaPlayableRemoteControlValues::Private;
+	return ToRemoteControlChanges(ResetValues(InReferenceValues.EntityValues, EntityValues, bInIsDefaults), EAvaPlayableRemoteControlChanges::EntityValues)
+		| ToRemoteControlChanges(ResetValues(InReferenceValues.ControllerValues, ControllerValues, bInIsDefaults), EAvaPlayableRemoteControlChanges::ControllerValues);
+}
+
+EAvaPlayableRemoteControlChanges FAvaPlayableRemoteControlValues::ResetRemoteControlEntityValue(const FGuid& InId
+	, const FAvaPlayableRemoteControlValue& InReferenceValue, bool bInIsDefaults)
+{
+	if (EntityValues.Contains(InId))
+	{
+		using namespace UE::AvaPlayableRemoteControlValues::Private;
+		return ToRemoteControlChanges(ResetValue(InReferenceValue, EntityValues[InId], bInIsDefaults), EAvaPlayableRemoteControlChanges::EntityValues);
+	}
+	return EAvaPlayableRemoteControlChanges::None;
+}
+
+EAvaPlayableRemoteControlChanges FAvaPlayableRemoteControlValues::ResetRemoteControlControllerValue(const FGuid& InId
+	, const FAvaPlayableRemoteControlValue& InReferenceValue, bool bInIsDefaults)
+{
+	if (ControllerValues.Contains(InId))
+	{
+		using namespace UE::AvaPlayableRemoteControlValues::Private;
+		return ToRemoteControlChanges(ResetValue(InReferenceValue, ControllerValues[InId], bInIsDefaults), EAvaPlayableRemoteControlChanges::ControllerValues);
+	}
+	return EAvaPlayableRemoteControlChanges::None;
+}
+
 bool FAvaPlayableRemoteControlValues::SetEntityValue(const FGuid& InId, const URemoteControlPreset* InRemoteControlPreset, bool bInIsDefault)
 {
 	const TSharedPtr<const FRemoteControlEntity> Entity = InRemoteControlPreset->GetExposedEntity<FRemoteControlEntity>(InId).Pin();
diff --git a/Engine/Plugins/Experimental/Avalanche/Source/AvalancheMedia/Private/Rundown/AvaRundown.cpp b/Engine/Plugins/Experimental/Avalanche/Source/AvalancheMedia/Private/Rundown/AvaRundown.cpp
--- a/Engine/Plugins/Experimental/Avalanche/Source/AvalancheMedia/Private/Rundown/AvaRundown.cpp
+++ b/Engine/Plugins/Experimental/Avalanche/Source/AvalancheMedia/Private/Rundown/AvaRundown.cpp
@@ -653,6 +653,52 @@ EAvaPlayableRemoteControlChanges UAvaRundown::UpdateRemoteControlValues(int32 In
 	return EAvaPlayableRemoteControlChanges::None;
 }
 
+EAvaPlayableRemoteControlChanges UAvaRundown::ResetRemoteControlValues(int32 InPageId, bool bInUseTemplateValues, bool bInIsDefault)
+{
+	FAvaRundownPage& Page = GetPage(InPageId);
+	if (Page.IsValidPage())
+	{
+		EAvaPlayableRemoteControlChanges Changes = Page.ResetRemoteControlValues(this, bInUseTemplateValues, bInIsDefault);
+		if (Changes != EAvaPlayableRemoteControlChanges::None)
+		{
+			NotifyPageRemoteControlValueChanged(InPageId, Changes);
+		}
+		return Changes;
+	}
+	return EAvaPlayableRemoteControlChanges::None;
+}
+
+
+EAvaPlayableRemoteControlChanges UAvaRundown::ResetRemoteControlControllerValue(int32 InPageId, const FGuid& InControllerId, bool bInUseTemplateValues, bool bInIsDefault)
+{
+	FAvaRundownPage& Page = GetPage(InPageId);
+	if (Page.IsValidPage())
+	{
+		EAvaPlayableRemoteControlChanges Changes = Page.ResetRemoteControlControllerValue(this, InControllerId, bInUseTemplateValues, bInIsDefault);
+		if (Changes != EAvaPlayableRemoteControlChanges::None)
+		{
+			NotifyPageRemoteControlValueChanged(InPageId, Changes);
+		}
+		return Changes;
+	}
+	return EAvaPlayableRemoteControlChanges::None;
+}
+
+EAvaPlayableRemoteControlChanges UAvaRundown::ResetRemoteControlEntityValue(int32 InPageId, const FGuid& InEntityId, bool bInUseTemplateValues, bool bInIsDefault)
+{
+	FAvaRundownPage& Page = GetPage(InPageId);
+	if (Page.IsValidPage())
+	{
+		EAvaPlayableRemoteControlChanges Changes = Page.ResetRemoteControlEntityValue(this, InEntityId, bInUseTemplateValues, bInIsDefault);
+		if (Changes != EAvaPlayableRemoteControlChanges::None)
+		{
+			NotifyPageRemoteControlValueChanged(InPageId, Changes);
+		}
+		return Changes;
+	}
+	return EAvaPlayableRemoteControlChanges::None;
+}
+
 void UAvaRundown::InvalidateManagedInstanceCacheForPages(const TArray<int32>& InPageIds) const
 {
 	if (!IAvaMediaModule::IsModuleLoaded())
@@ -2080,5 +2126,4 @@ void UAvaRundown::RemoveStoppedPagePlayers()
 	PagePlayers.RemoveAll([](const UAvaRundownPagePlayer* InPagePlayer) { return !InPagePlayer || InPagePlayer->IsPlaying() == false;});
 }
 
-
 #undef LOCTEXT_NAMESPACE
diff --git a/Engine/Plugins/Experimental/Avalanche/Source/AvalancheMedia/Private/Rundown/AvaRundownPage.cpp b/Engine/Plugins/Experimental/Avalanche/Source/AvalancheMedia/Private/Rundown/AvaRundownPage.cpp
--- a/Engine/Plugins/Experimental/Avalanche/Source/AvalancheMedia/Private/Rundown/AvaRundownPage.cpp
+++ b/Engine/Plugins/Experimental/Avalanche/Source/AvalancheMedia/Private/Rundown/AvaRundownPage.cpp
@@ -12,6 +12,7 @@
 #include "Playback/AvaPlaybackManager.h"
 #include "RCVirtualProperty.h"
 #include "Rundown/AvaRundown.h"
+#include "Rundown/AvaRundownManagedInstance.h"
 #include "Rundown/AvaRundownManagedInstanceCache.h"
 #include "Rundown/AvaRundownPageAssetUtils.h"
 #include "Rundown/AvaRundownPagePlayer.h"
@@ -541,6 +542,199 @@ void FAvaRundownPage::SetRemoteControlControllerValue(const FGuid& InId, const F
 	RemoteControlValues.SetControllerValue(InId, InValue);
 }
 
+bool FAvaRundownPage::GetDefaultRemoteControlValues(const UAvaRundown* InRundown, bool bInUseTemplateValues, FAvaPlayableRemoteControlValues& OutValues) const
+{
+	if (bInUseTemplateValues)
+	{
+		if (!IsTemplate())
+		{
+			const FAvaRundownPage& PageTemplate = ResolveTemplate(InRundown);
+			if (PageTemplate.IsValidPage())
+			{
+				OutValues = PageTemplate.GetRemoteControlValues();
+				return true;
+			}
+		}
+
+		return false;
+	}
+
+	FAvaRundownManagedInstanceCache& ManagedInstanceCache = IAvaMediaModule::Get().GetManagedInstanceCache();
+
+	for (const FSoftObjectPath& PageAssetPath : GetAssetPaths(InRundown))
+	{
+		const TSharedPtr<FAvaRundownManagedInstance> ManagedInstance = ManagedInstanceCache.GetOrLoadInstance(PageAssetPath);
+		if (!ManagedInstance)
+		{
+			return false;
+		}
+
+		OutValues.Merge(ManagedInstance->GetDefaultRemoteControlValues());
+	}
+
+	return true;
+}
+
+bool FAvaRundownPage::GetDefaultEntityValue(const UAvaRundown* InRundown, const FGuid& InId, bool bInUseTemplateValues, FAvaPlayableRemoteControlValue& OutValue) const
+{
+	if (bInUseTemplateValues)
+	{
+		if (!IsTemplate())
+		{
+			const FAvaRundownPage& PageTemplate = ResolveTemplate(InRundown);
+			if (PageTemplate.IsValidPage())
+			{
+				const FAvaPlayableRemoteControlValues& PageValues = PageTemplate.GetRemoteControlValues();
+				if (PageValues.HasEntityValue(InId))
+				{
+					OutValue = *PageValues.GetEntityValue(InId);
+					return true;
+				}
+			}
+		}
+
+		return false;
+	}
+
+	FAvaRundownManagedInstanceCache& ManagedInstanceCache = IAvaMediaModule::Get().GetManagedInstanceCache();
+
+	for (const FSoftObjectPath& PageAssetPath : GetAssetPaths(InRundown))
+	{
+		const TSharedPtr<FAvaRundownManagedInstance> ManagedInstance = ManagedInstanceCache.GetOrLoadInstance(PageAssetPath);
+		if (!ManagedInstance)
+		{
+			return false;
+		}
+
+		const FAvaPlayableRemoteControlValues& DefaultValues = ManagedInstance->GetDefaultRemoteControlValues();
+		if (DefaultValues.HasEntityValue(InId))
+		{
+			const FAvaPlayableRemoteControlValue* const EntityValue = DefaultValues.GetEntityValue(InId);
+			if (EntityValue)
+			{
+				OutValue = *EntityValue;
+				return true;
+			}
+		}
+	}
+
+	return false;
+}
+
+bool FAvaRundownPage::GetDefaultControllerValue(const UAvaRundown* InRundown, const FGuid& InId, bool bInUseTemplateValues, FAvaPlayableRemoteControlValue& OutValue) const
+{
+	if (bInUseTemplateValues)
+	{
+		if (!IsTemplate())
+		{
+			const FAvaRundownPage& PageTemplate = ResolveTemplate(InRundown);
+			if (PageTemplate.IsValidPage())
+			{
+				const FAvaPlayableRemoteControlValues& PageValues = PageTemplate.GetRemoteControlValues();
+				if (PageValues.HasControllerValue(InId))
+				{
+					OutValue = *PageValues.GetControllerValue(InId);
+					return true;
+				}
+			}
+		}
+
+		return false;
+	}
+
+	FAvaRundownManagedInstanceCache& ManagedInstanceCache = IAvaMediaModule::Get().GetManagedInstanceCache();
+
+	for (const FSoftObjectPath& PageAssetPath : GetAssetPaths(InRundown))
+	{
+		const TSharedPtr<FAvaRundownManagedInstance> ManagedInstance = ManagedInstanceCache.GetOrLoadInstance(PageAssetPath);
+		if (!ManagedInstance)
+		{
+			return false;
+		}
+
+		const FAvaPlayableRemoteControlValues& DefaultValues = ManagedInstance->GetDefaultRemoteControlValues();
+		if (DefaultValues.HasControllerValue(InId))
+		{
+			const FAvaPlayableRemoteControlValue* const ControllerValue = DefaultValues.GetControllerValue(InId);
+			if (ControllerValue)
+			{
+				OutValue = *ControllerValue;
+				return true;
+			}
+		}
+	}
+
+	return false;
+}
+
+bool FAvaRundownPage::IsDefaultEntityValue(const UAvaRundown* InRundown, const FGuid& InId, bool bInUseTemplateValues) const
+{
+	FAvaPlayableRemoteControlValue DefaultValue;
+	if (GetDefaultEntityValue(InRundown, InId, bInUseTemplateValues, DefaultValue))
+	{
+		if (RemoteControlValues.HasEntityValue(InId))
+		{
+			if (const FAvaPlayableRemoteControlValue* Value = RemoteControlValues.GetEntityValue(InId))
+			{
+				if (Value->IsSameValueAs(DefaultValue))
+				{
+					return true;
+				}
+			}
+		}
+	}
+	return false;
+}
+
+bool FAvaRundownPage::IsDefaultControllerValue(const UAvaRundown* InRundown, const FGuid& InId, bool bInUseTemplateValues) const
+{
+	FAvaPlayableRemoteControlValue DefaultValue;
+	if (GetDefaultControllerValue(InRundown, InId, bInUseTemplateValues, DefaultValue))
+	{
+		if (RemoteControlValues.HasControllerValue(InId))
+		{
+			if (const FAvaPlayableRemoteControlValue* Value = RemoteControlValues.GetControllerValue(InId))
+			{
+				if (Value->IsSameValueAs(DefaultValue))
+				{
+					return true;
+				}
+			}
+		}
+	}
+	return false;
+}
+
+EAvaPlayableRemoteControlChanges FAvaRundownPage::ResetRemoteControlValues(const UAvaRundown* InRundown, bool bInUseTemplateValues, bool bInIsDefault)
+{
+	FAvaPlayableRemoteControlValues DefaultValues;
+	if (GetDefaultRemoteControlValues(InRundown, bInUseTemplateValues, DefaultValues))
+	{
+		return RemoteControlValues.ResetRemoteControlValues(DefaultValues, bInIsDefault);
+	}
+	return EAvaPlayableRemoteControlChanges::None;
+}
+
+EAvaPlayableRemoteControlChanges FAvaRundownPage::ResetRemoteControlEntityValue(const UAvaRundown* InRundown, const FGuid& InId, bool bInUseTemplateValues, bool bInIsDefault)
+{
+	FAvaPlayableRemoteControlValue DefaultValue;
+	if (GetDefaultEntityValue(InRundown, InId, bInUseTemplateValues, DefaultValue))
+	{
+		return RemoteControlValues.ResetRemoteControlEntityValue(InId, DefaultValue, bInIsDefault);
+	}
+	return EAvaPlayableRemoteControlChanges::None;
+}
+
+EAvaPlayableRemoteControlChanges FAvaRundownPage::ResetRemoteControlControllerValue(const UAvaRundown* InRundown, const FGuid& InId, bool bInUseTemplateValues, bool bInIsDefault)
+{
+	FAvaPlayableRemoteControlValue DefaultValue;
+	if (GetDefaultControllerValue(InRundown, InId, bInUseTemplateValues, DefaultValue))
+	{
+		return RemoteControlValues.ResetRemoteControlControllerValue(InId, DefaultValue, bInIsDefault);
+	}
+	return EAvaPlayableRemoteControlChanges::None;
+}
+
 void FAvaRundownPage::PostLoad()
 {
 }
diff --git a/Engine/Plugins/Experimental/Avalanche/Source/AvalancheMedia/Public/Playable/AvaPlayableRemoteControlValues.h b/Engine/Plugins/Experimental/Avalanche/Source/AvalancheMedia/Public/Playable/AvaPlayableRemoteControlValues.h
--- a/Engine/Plugins/Experimental/Avalanche/Source/AvalancheMedia/Public/Playable/AvaPlayableRemoteControlValues.h
+++ b/Engine/Plugins/Experimental/Avalanche/Source/AvalancheMedia/Public/Playable/AvaPlayableRemoteControlValues.h
@@ -105,7 +105,7 @@ struct AVALANCHEMEDIA_API FAvaPlayableRemoteControlValues
 	bool HasSameControllerValues(const FAvaPlayableRemoteControlValues& InOther) const;
 
 	/**
-	 *	Removes the extra values compared to the given reference values.
+	 * Removes the extra values compared to the given reference values.
 	 * @return flags indicating what changed.
 	 */
 	EAvaPlayableRemoteControlChanges PruneRemoteControlValues(const FAvaPlayableRemoteControlValues& InReferenceValues);
@@ -119,7 +119,33 @@ struct AVALANCHEMEDIA_API FAvaPlayableRemoteControlValues
 	 * @return flags indicating what changed.
 	 */
 	EAvaPlayableRemoteControlChanges UpdateRemoteControlValues(const FAvaPlayableRemoteControlValues& InReferenceValues, bool bInUpdateDefaults);
+
+	/**
+	 * Reset the values to the reference.
+	 * @param InReferenceValues Values to be reset to.
+	 * @param bInIsDefaults If true, consider all values as "default". Otherwise, preserves the "default" status from reference.
+	 * @return
+	 */
+	EAvaPlayableRemoteControlChanges ResetRemoteControlValues(const FAvaPlayableRemoteControlValues& InReferenceValues, bool bInIsDefaults);
 	
+	/**
+	 * Reset the controller value to the reference.
+	 * @param InReferenceValue Value to be reset to.
+	 * @param bInIsDefaults If true, consider all values as "default". Otherwise, preserves the "default" status from reference.
+	 * @return
+	 */
+	EAvaPlayableRemoteControlChanges ResetRemoteControlControllerValue(const FGuid& InId
+		, const FAvaPlayableRemoteControlValue& InReferenceValue, bool bInIsDefaults);
+
+	/**
+	 * Reset the entity value to the reference.
+	 * @param InReferenceValue Value to be reset to.
+	 * @param bInIsDefaults If true, consider all values as "default". Otherwise, preserves the "default" status from reference.
+	 * @return
+	 */
+	EAvaPlayableRemoteControlChanges ResetRemoteControlEntityValue(const FGuid& InId
+		, const FAvaPlayableRemoteControlValue& InReferenceValue, bool bInIsDefaults);
+
 	/**
 	 * Refreshes the EntitiesControlledByController set.
 	 */
diff --git a/Engine/Plugins/Experimental/Avalanche/Source/AvalancheMedia/Public/Rundown/AvaRundown.h b/Engine/Plugins/Experimental/Avalanche/Source/AvalancheMedia/Public/Rundown/AvaRundown.h
--- a/Engine/Plugins/Experimental/Avalanche/Source/AvalancheMedia/Public/Rundown/AvaRundown.h
+++ b/Engine/Plugins/Experimental/Avalanche/Source/AvalancheMedia/Public/Rundown/AvaRundown.h
@@ -382,6 +382,10 @@ public:
 	bool SetRemoteControlEntityValue(int32 InPageId, const FGuid& InId, const FAvaPlayableRemoteControlValue& InValue);
 	bool SetRemoteControlControllerValue(int32 InPageId, const FGuid& InId,const FAvaPlayableRemoteControlValue& InValue);
 	EAvaPlayableRemoteControlChanges UpdateRemoteControlValues(int32 InPageId, const FAvaPlayableRemoteControlValues& InRemoteControlValues, bool bInUpdateDefaults);
+	
+	EAvaPlayableRemoteControlChanges ResetRemoteControlValues(int32 InPageId, bool bInUseTemplateValues, bool bInIsDefault);
+	EAvaPlayableRemoteControlChanges ResetRemoteControlControllerValue(int32 InPageId, const FGuid& InControllerId, bool bInUseTemplateValues, bool bInIsDefault);
+	EAvaPlayableRemoteControlChanges ResetRemoteControlEntityValue(int32 InPageId, const FGuid& InEntityId, bool bInUseTemplateValues, bool bInIsDefault);
 
 	void InvalidateManagedInstanceCacheForPages(const TArray<int32>& InPageIds) const;
 
diff --git a/Engine/Plugins/Experimental/Avalanche/Source/AvalancheMedia/Public/Rundown/AvaRundownPage.h b/Engine/Plugins/Experimental/Avalanche/Source/AvalancheMedia/Public/Rundown/AvaRundownPage.h
--- a/Engine/Plugins/Experimental/Avalanche/Source/AvalancheMedia/Public/Rundown/AvaRundownPage.h
+++ b/Engine/Plugins/Experimental/Avalanche/Source/AvalancheMedia/Public/Rundown/AvaRundownPage.h
@@ -8,6 +8,7 @@
 #include "AvaRundownPage.generated.h"
 
 class UAvaRundown;
+class URemoteControlPreset;
 
 UENUM()
 enum class EAvaRundownPageStatus : uint8
@@ -79,7 +80,7 @@ public:
 	bool IsValidPage() const;
 
 	void Rename(const FString& InNewName);
-	void RenameFriendlyName (const FString& InNewName);
+	void RenameFriendlyName(const FString& InNewName);
 	
 	int32 GetPageId() const { return PageId; }
 	void SetPageId(int32 InPageId) { PageId = InPageId; } // Do not use lightly.
@@ -173,7 +174,7 @@ public:
 
 	EAvaPlayableRemoteControlChanges PruneRemoteControlValues(const FAvaPlayableRemoteControlValues& InRemoteControlValues);
 	EAvaPlayableRemoteControlChanges UpdateRemoteControlValues(const FAvaPlayableRemoteControlValues& InRemoteControlValues, bool bInUpdateDefaults);
-	
+
 	bool HasRemoteControlEntityValue(const FGuid& InId) const { return RemoteControlValues.HasEntityValue(InId); }
 	const FAvaPlayableRemoteControlValue* GetRemoteControlEntityValue(const FGuid& InId) const { return RemoteControlValues.GetEntityValue(InId); }
 	void SetRemoteControlEntityValue(const FGuid& InId, const FAvaPlayableRemoteControlValue& InValue);
@@ -181,8 +182,18 @@ public:
 	bool HasRemoteControlControllerValue(const FGuid& InId) const { return RemoteControlValues.HasControllerValue(InId); }
 	const FAvaPlayableRemoteControlValue* GetRemoteControlControllerValue(const FGuid& InId) const { return RemoteControlValues.GetControllerValue(InId); }
 	void SetRemoteControlControllerValue(const FGuid& InId, const FAvaPlayableRemoteControlValue& InValue);
-	
+
 	const FAvaPlayableRemoteControlValues& GetRemoteControlValues() const { return RemoteControlValues; }
+	bool GetDefaultRemoteControlValues(const UAvaRundown* InRundown, bool bInUseTemplateValues, FAvaPlayableRemoteControlValues& OutValues) const;
+	bool GetDefaultEntityValue(const UAvaRundown* InRundown, const FGuid& InId, bool bInUseTemplateValues, FAvaPlayableRemoteControlValue& OutValue) const;
+	bool GetDefaultControllerValue(const UAvaRundown* InRundown, const FGuid& InId, bool bInUseTemplateValues, FAvaPlayableRemoteControlValue& OutValue) const;
+
+	bool IsDefaultEntityValue(const UAvaRundown* InRundown, const FGuid& InId, bool bInUseTemplateValues) const;
+	bool IsDefaultControllerValue(const UAvaRundown* InRundown, const FGuid& InId, bool bInUseTemplateValues) const;
+
+	EAvaPlayableRemoteControlChanges ResetRemoteControlValues(const UAvaRundown* InRundown, bool bInUseTemplateValues, bool bInIsDefault);
+	EAvaPlayableRemoteControlChanges ResetRemoteControlEntityValue(const UAvaRundown* InRundown, const FGuid& InId, bool bInUseTemplateValues, bool bInIsDefault);
+	EAvaPlayableRemoteControlChanges ResetRemoteControlControllerValue(const UAvaRundown* InRundown, const FGuid& InId, bool bInUseTemplateValues, bool bInIsDefault);
 
 	friend FORCEINLINE uint32 GetTypeHash(const FAvaRundownPage& InPage)
 	{
diff --git a/Engine/Plugins/Experimental/Avalanche/Source/AvalancheMediaEditor/Private/Rundown/AvaRundownCommands.cpp b/Engine/Plugins/Experimental/Avalanche/Source/AvalancheMediaEditor/Private/Rundown/AvaRundownCommands.cpp
--- a/Engine/Plugins/Experimental/Avalanche/Source/AvalancheMediaEditor/Private/Rundown/AvaRundownCommands.cpp
+++ b/Engine/Plugins/Experimental/Avalanche/Source/AvalancheMediaEditor/Private/Rundown/AvaRundownCommands.cpp
@@ -143,6 +143,18 @@ void FAvaRundownCommands::RegisterCommands()
 		, "Takes the Next Page in the Page List"
 		, EUserInterfaceActionType::Button
 		, FInputChord(EKeys::Subtract));
+
+	UI_COMMAND(ResetValuesToDefaults
+		, "Reset To Defaults"
+		, "Reset selected values to default values"
+		, EUserInterfaceActionType::Button
+		, FInputChord());
+
+	UI_COMMAND(ResetValuesToTemplate
+		, "Reset To Template"
+		, "Reset selected values to template values"
+		, EUserInterfaceActionType::Button
+		, FInputChord());
 }
 
 #undef LOCTEXT_NAMESPACE
diff --git a/Engine/Plugins/Experimental/Avalanche/Source/AvalancheMediaEditor/Private/Rundown/AvaRundownCommands.h b/Engine/Plugins/Experimental/Avalanche/Source/AvalancheMediaEditor/Private/Rundown/AvaRundownCommands.h
--- a/Engine/Plugins/Experimental/Avalanche/Source/AvalancheMediaEditor/Private/Rundown/AvaRundownCommands.h
+++ b/Engine/Plugins/Experimental/Avalanche/Source/AvalancheMediaEditor/Private/Rundown/AvaRundownCommands.h
@@ -51,6 +51,10 @@ public:
 	TSharedPtr<FUICommandInfo> PreviewContinue;
 	TSharedPtr<FUICommandInfo> PreviewPlayNext;
 	TSharedPtr<FUICommandInfo> TakeToProgram;
+
+	/** Remote Control Commands */
+	TSharedPtr<FUICommandInfo> ResetValuesToDefaults;
+	TSharedPtr<FUICommandInfo> ResetValuesToTemplate;
 };
 
 #undef LOCTEXT_NAMESPACE
diff --git a/Engine/Plugins/Experimental/Avalanche/Source/AvalancheMediaEditor/Private/Rundown/DetailsView/SAvaRundownPageDetails.cpp b/Engine/Plugins/Experimental/Avalanche/Source/AvalancheMediaEditor/Private/Rundown/DetailsView/SAvaRundownPageDetails.cpp
--- a/Engine/Plugins/Experimental/Avalanche/Source/AvalancheMediaEditor/Private/Rundown/DetailsView/SAvaRundownPageDetails.cpp
+++ b/Engine/Plugins/Experimental/Avalanche/Source/AvalancheMediaEditor/Private/Rundown/DetailsView/SAvaRundownPageDetails.cpp
@@ -2,6 +2,7 @@
 
 #include "SAvaRundownPageDetails.h"
 #include "Async/Async.h"
+#include "Framework/Application/SlateApplication.h"
 #include "IAvaMediaModule.h"
 #include "Input/Reply.h"
 #include "Internationalization/Text.h"
@@ -31,6 +32,12 @@ void SAvaRundownPageDetails::Construct(const FArguments& InArgs, const TSharedPt
 	InRundownEditor->GetOnPageEvent().AddSP(this, &SAvaRundownPageDetails::OnPageEvent);
 	IAvaMediaModule::Get().GetManagedInstanceCache().OnEntryInvalidated.AddSP(this, &SAvaRundownPageDetails::OnManagedInstanceCacheEntryInvalidated);
 
+	UAvaRundown* const Rundown = InRundownEditor->GetRundown();
+	if (IsValid(Rundown))
+	{
+		Rundown->GetOnPagesChanged().AddSP(this, &SAvaRundownPageDetails::OnRundownPagesChanged);
+	}
+
 	TSharedRef<SHorizontalBox> AnimationHeader = SNew(SHorizontalBox);
 	{
 		static const TArray<FText> HeaderTitles
@@ -144,7 +151,7 @@ void SAvaRundownPageDetails::Construct(const FArguments& InArgs, const TSharedPt
 				[
 					SNew(SButton)
 					.ContentPadding(0)
-					.ButtonStyle(FAppStyle::Get(), "NoBorder")
+					.ButtonStyle(FAppStyle::Get(), TEXT("SimpleButton"))
 					.OnClicked(this, &SAvaRundownPageDetails::ToggleExposedPropertiesVisibility)
 					.ToolTipText(LOCTEXT("VisibilityButtonToolTip", "Toggle Exposed Properties Visibility"))
 					.Content()
@@ -185,6 +192,15 @@ SAvaRundownPageDetails::~SAvaRundownPageDetails()
 	{
 		IAvaMediaModule::Get().GetManagedInstanceCache().OnEntryInvalidated.RemoveAll(this);
 	}
+
+	if (const TSharedPtr<FAvaRundownEditor> RundownEditor = RundownEditorWeak.Pin())
+	{
+		UAvaRundown* const Rundown = RundownEditor->GetRundown();
+		if (IsValid(Rundown))
+		{
+			Rundown->GetOnPagesChanged().RemoveAll(this);
+		}
+	}
 }
 
 void SAvaRundownPageDetails::OnPageEvent(const TArray<int32>& InSelectedPageIds, UE::AvaRundown::EPageEvent InPageEvent)
@@ -429,4 +445,14 @@ FReply SAvaRundownPageDetails::DuplicateSelectedPage()
 	return FReply::Handled();
 }
 
+void SAvaRundownPageDetails::OnRundownPagesChanged(const UAvaRundown* InRundown, const FAvaRundownPage& InPage, const EAvaRundownPageChanges InChanges)
+{
+	// Refreshing the page while the mouse is captured will result in losing the capture
+	// and ending any drag event that is actively changing the value.
+	if (!FSlateApplication::Get().GetMouseCaptureWindow())
+	{
+		RefreshSelectedPage();
+	}
+}
+
 #undef LOCTEXT_NAMESPACE
diff --git a/Engine/Plugins/Experimental/Avalanche/Source/AvalancheMediaEditor/Private/Rundown/DetailsView/SAvaRundownPageDetails.h b/Engine/Plugins/Experimental/Avalanche/Source/AvalancheMediaEditor/Private/Rundown/DetailsView/SAvaRundownPageDetails.h
--- a/Engine/Plugins/Experimental/Avalanche/Source/AvalancheMediaEditor/Private/Rundown/DetailsView/SAvaRundownPageDetails.h
+++ b/Engine/Plugins/Experimental/Avalanche/Source/AvalancheMediaEditor/Private/Rundown/DetailsView/SAvaRundownPageDetails.h
@@ -13,9 +13,11 @@ class FReply;
 class FText;
 class SAvaRundownPageRemoteControlProps;
 class SAvaRundownRCControllerPanel;
+class UAvaRundown;
 struct FAvaRundownPage;
 struct FSoftObjectPath;
 struct FSlateBrush;
+enum class EAvaRundownPageChanges : uint8;
 
 class SAvaRundownPageDetails : public SCompoundWidget
 {
@@ -56,6 +58,8 @@ protected:
 
 	FReply DuplicateSelectedPage();
 
+	void OnRundownPagesChanged(const UAvaRundown* InRundown, const FAvaRundownPage& InPage, const EAvaRundownPageChanges InChanges);
+
 private:
 	TWeakPtr<FAvaRundownEditor> RundownEditorWeak;
 
diff --git a/Engine/Plugins/Experimental/Avalanche/Source/AvalancheMediaEditor/Private/Rundown/Pages/AvaRundownPageContextMenu.cpp b/Engine/Plugins/Experimental/Avalanche/Source/AvalancheMediaEditor/Private/Rundown/Pages/AvaRundownPageContextMenu.cpp
--- a/Engine/Plugins/Experimental/Avalanche/Source/AvalancheMediaEditor/Private/Rundown/Pages/AvaRundownPageContextMenu.cpp
+++ b/Engine/Plugins/Experimental/Avalanche/Source/AvalancheMediaEditor/Private/Rundown/Pages/AvaRundownPageContextMenu.cpp
@@ -117,6 +117,38 @@ void FAvaRundownPageContextMenu::PopulatePageContextMenu(UToolMenu& InMenu, UAva
 				FSlateIcon(FAppStyle::GetAppStyleSetName(), "MainFrame.SaveAll")
 			);
 		}
+
+		if (PageListReference.Type == EAvaRundownPageListType::Instance)
+		{
+			Section.AddSubMenu(TEXT("ResetPages"),
+				LOCTEXT("ResetPagesSubMenu", "Reset Pages ..."),
+				LOCTEXT("ResetPagesSubMenuTooltip", "Reset selected pages to default values."),
+				FNewMenuDelegate::CreateLambda([&RundownCommands](FMenuBuilder& InMenuBuilder)
+					{
+						InMenuBuilder.AddMenuEntry(RundownCommands.ResetValuesToDefaults,
+							TEXT("InstanceResetPages"),
+							LOCTEXT("InstanceResetPages", "Reset to Default"),
+							LOCTEXT("InstanceResetPagesTooltip", "Reset all selected pages to default values."),
+							FSlateIcon(FAppStyle::GetAppStyleSetName(), TEXT("PropertyWindow.DiffersFromDefault")));
+
+						InMenuBuilder.AddMenuEntry(RundownCommands.ResetValuesToTemplate,
+							TEXT("TemplateResetPages"),
+							LOCTEXT("TemplateResetPages", "Reset to Template"),
+							LOCTEXT("TemplateResetPagesTooltip", "Reset all selected pages to template values."),
+							FSlateIcon(FAppStyle::GetAppStyleSetName(), TEXT("PropertyWindow.DiffersFromDefault")));
+					}),
+				false,
+				FSlateIcon(FAppStyle::GetAppStyleSetName(), TEXT("PropertyWindow.DiffersFromDefault"))
+			);
+		}
+		else if (PageListReference.Type == EAvaRundownPageListType::Template)
+		{
+			Section.AddMenuEntry(RundownCommands.ResetValuesToDefaults,
+				LOCTEXT("ResetPages", "Reset to Default"),
+				LOCTEXT("ResetPagesTooltip", "Reset all selected template pages to default values."),
+				FSlateIcon(FAppStyle::GetAppStyleSetName(), TEXT("PropertyWindow.DiffersFromDefault"))
+			);
+		}
 	}
 
 	// Show Control Actions
diff --git a/Engine/Plugins/Experimental/Avalanche/Source/AvalancheMediaEditor/Private/Rundown/Pages/Slate/SAvaRundownInstancedPageList.cpp b/Engine/Plugins/Experimental/Avalanche/Source/AvalancheMediaEditor/Private/Rundown/Pages/Slate/SAvaRundownInstancedPageList.cpp
--- a/Engine/Plugins/Experimental/Avalanche/Source/AvalancheMediaEditor/Private/Rundown/Pages/Slate/SAvaRundownInstancedPageList.cpp
+++ b/Engine/Plugins/Experimental/Avalanche/Source/AvalancheMediaEditor/Private/Rundown/Pages/Slate/SAvaRundownInstancedPageList.cpp
@@ -239,6 +239,14 @@ void SAvaRundownInstancedPageList::BindCommands()
 			FExecuteAction::CreateSP(this, &SAvaRundownPageList::ExportSelectedPagesToExternalFile, TEXT("xml")),
 			FCanExecuteAction::CreateSP(this, &SAvaRundownPageList::CanExportSelectedPagesToExternalFile, TEXT("xml")));
 
+		CommandList->MapAction(RundownCommands.ResetValuesToDefaults,
+			FExecuteAction::CreateSP(this, &SAvaRundownInstancedPageList::ResetPagesToDefaults, false),
+			FCanExecuteAction::CreateSP(this, &SAvaRundownInstancedPageList::CanResetPagesToDefaults, false));
+
+		CommandList->MapAction(RundownCommands.ResetValuesToTemplate,
+			FExecuteAction::CreateSP(this, &SAvaRundownInstancedPageList::ResetPagesToDefaults, true),
+			FCanExecuteAction::CreateSP(this, &SAvaRundownInstancedPageList::CanResetPagesToDefaults, true));
+
 		CommandList->MapAction(RundownCommands.Play,
 			FExecuteAction::CreateSP(this, &SAvaRundownInstancedPageList::PlaySelectedPage),
 			FCanExecuteAction::CreateSP(this, &SAvaRundownInstancedPageList::CanPlaySelectedPage));
@@ -1371,4 +1379,48 @@ int32 SAvaRundownInstancedPageList::GetPageIdToTakeNext() const
 	return FAvaRundownPlaybackUtils::GetPageIdToPlayNext(GetRundown(), GetPageListReference(), /*bInPreview*/ false, NAME_None);
 }
 
+void SAvaRundownInstancedPageList::ResetPagesToDefaults(bool bInResetToTemplate)
+{
+	if (UAvaRundown* Rundown = GetRundown())
+	{
+		FScopedTransaction Transaction(LOCTEXT("ResetPagesTransaction", "Reset Pages"));
+		Rundown->Modify();
+
+		for (int32 PageId : SelectedPageIds)
+		{
+			Rundown->ResetRemoteControlValues(PageId, bInResetToTemplate, /*bInIsDefault=*/false);
+		}
+	}
+}
+
+bool SAvaRundownInstancedPageList::CanResetPagesToDefaults(bool bInResetToTemplate) const
+{
+	bool bContainsDifferentValues = false;
+
+	if (UAvaRundown* Rundown = GetRundown())
+	{
+		for (const int32 PageId : SelectedPageIds)
+		{
+			const FAvaRundownPage& Page = Rundown->GetPage(PageId);
+
+			if (!Page.IsValidPage() || !Page.IsEnabled() || Page.IsTemplate())
+			{
+				return false;
+			}
+
+			FAvaPlayableRemoteControlValues DefaultValues;
+			if (Page.GetDefaultRemoteControlValues(Rundown, bInResetToTemplate, DefaultValues))
+			{
+				const FAvaPlayableRemoteControlValues& PageValues = Page.GetRemoteControlValues();
+				if (!(PageValues.HasSameEntityValues(DefaultValues) && PageValues.HasSameControllerValues(DefaultValues)))
+				{
+					bContainsDifferentValues = true;
+				}
+			}
+		}
+	}
+
+	return bContainsDifferentValues;
+}
+
 #undef LOCTEXT_NAMESPACE
diff --git a/Engine/Plugins/Experimental/Avalanche/Source/AvalancheMediaEditor/Private/Rundown/Pages/Slate/SAvaRundownInstancedPageList.h b/Engine/Plugins/Experimental/Avalanche/Source/AvalancheMediaEditor/Private/Rundown/Pages/Slate/SAvaRundownInstancedPageList.h
--- a/Engine/Plugins/Experimental/Avalanche/Source/AvalancheMediaEditor/Private/Rundown/Pages/Slate/SAvaRundownInstancedPageList.h
+++ b/Engine/Plugins/Experimental/Avalanche/Source/AvalancheMediaEditor/Private/Rundown/Pages/Slate/SAvaRundownInstancedPageList.h
@@ -97,4 +97,7 @@ protected:
 	void PlayNextPageNoReturn() const { PlayNextPage(); }
 
 	virtual TArray<int32> AddPastedPages(const TArray<FAvaRundownPage>& InPages) override;
+
+	void ResetPagesToDefaults(bool bInResetToTemplate);
+	bool CanResetPagesToDefaults(bool bInResetToTemplate) const;
 };
diff --git a/Engine/Plugins/Experimental/Avalanche/Source/AvalancheMediaEditor/Private/Rundown/Pages/Slate/SAvaRundownTemplatePageList.cpp b/Engine/Plugins/Experimental/Avalanche/Source/AvalancheMediaEditor/Private/Rundown/Pages/Slate/SAvaRundownTemplatePageList.cpp
--- a/Engine/Plugins/Experimental/Avalanche/Source/AvalancheMediaEditor/Private/Rundown/Pages/Slate/SAvaRundownTemplatePageList.cpp
+++ b/Engine/Plugins/Experimental/Avalanche/Source/AvalancheMediaEditor/Private/Rundown/Pages/Slate/SAvaRundownTemplatePageList.cpp
@@ -3,10 +3,12 @@
 #include "SAvaRundownTemplatePageList.h"
 
 #include "IAvaMediaEditorModule.h"
+#include "IAvaMediaModule.h"
 #include "Rundown/AvaRundown.h"
 #include "Rundown/AvaRundownCommands.h"
 #include "Rundown/AvaRundownEditor.h"
 #include "Rundown/AvaRundownEditorUtils.h"
+#include "Rundown/AvaRundownManagedInstanceCache.h"
 #include "Rundown/AvaRundownPage.h"
 #include "Rundown/Factories/Filters/AvaRundownFactoriesUtils.h"
 #include "Rundown/Pages/Columns/AvaRundownPageAssetSelectorColumn.h"
@@ -16,6 +18,7 @@
 #include "Rundown/Pages/Columns/AvaRundownPageThumbnailColumn.h"
 #include "Rundown/Pages/Columns/AvaRundownPageTransitionLayerColumn.h"
 #include "Rundown/Pages/PageViews/AvaRundownTemplatePageViewImpl.h"
+#include "ScopedTransaction.h"
 #include "Widgets/SNullWidget.h"
 #include "Widgets/Views/SListView.h"
 
@@ -147,6 +150,10 @@ void SAvaRundownTemplatePageList::BindCommands()
 			FExecuteAction::CreateSP(this, &SAvaRundownPageList::EditSelectedPageSource),
 			FCanExecuteAction::CreateSP(this, &SAvaRundownPageList::CanEditSelectedPageSource));
 
+		CommandList->MapAction(RundownCommands.ResetValuesToDefaults,
+			FExecuteAction::CreateSP(this, &SAvaRundownTemplatePageList::ResetPagesToDefaults),
+			FCanExecuteAction::CreateSP(this, &SAvaRundownTemplatePageList::CanResetPagesToDefaults));
+
 		CommandList->MapAction(RundownCommands.PreviewFrame,
 			FExecuteAction::CreateSP(this, &SAvaRundownPageList::PreviewPlaySelectedPage, true),
 			FCanExecuteAction::CreateSP(this, &SAvaRundownPageList::CanPreviewPlaySelectedPage));
@@ -500,4 +507,63 @@ void SAvaRundownTemplatePageList::OnTemplatePageListChanged(const FAvaRundownPag
 	Refresh();
 }
 
+void SAvaRundownTemplatePageList::ResetPagesToDefaults()
+{
+	UAvaRundown* Rundown = GetValidRundown();
+	
+	if (SelectedPageIds.IsEmpty() || !Rundown)
+	{
+		return;
+	}
+
+	FScopedTransaction Transaction(LOCTEXT("ResetPagesTransaction", "Reset Pages"));
+	Rundown->Modify();
+	
+	for (const int32 SelectedPageId : SelectedPageIds)
+	{
+		const FAvaRundownPage& Page = Rundown->GetPage(SelectedPageId);
+
+		if (!Page.IsValidPage() || !Page.IsEnabled() || !Page.IsTemplate())
+		{
+			continue;
+		}
+		
+		Rundown->ResetRemoteControlValues(SelectedPageId, /*bInUseTemplateValues=*/false, /*bInIsDefault=*/false);
+	}
+}
+
+bool SAvaRundownTemplatePageList::CanResetPagesToDefaults() const
+{
+	const UAvaRundown* const Rundown = GetValidRundown();
+
+	if (SelectedPageIds.IsEmpty() || !Rundown)
+	{
+		return false;
+	}
+
+	bool bContainsDifferentValues = false;
+
+	for (const int32 SelectedPageId : SelectedPageIds)
+	{
+		const FAvaRundownPage& Page = Rundown->GetPage(SelectedPageId);
+
+		if (!Page.IsValidPage() || !Page.IsEnabled() || !Page.IsTemplate())
+		{
+			return false;
+		}
+
+		FAvaPlayableRemoteControlValues DefaultValues;
+		if (Page.GetDefaultRemoteControlValues(Rundown, /*bInUseTemplateValues=*/false, DefaultValues))
+		{
+			const FAvaPlayableRemoteControlValues& PageValues = Page.GetRemoteControlValues();
+			if (!(PageValues.HasSameEntityValues(DefaultValues) && PageValues.HasSameControllerValues(DefaultValues)))
+			{
+				bContainsDifferentValues = true;
+			}
+		}
+	}
+
+	return bContainsDifferentValues;
+}
+
 #undef LOCTEXT_NAMESPACE
diff --git a/Engine/Plugins/Experimental/Avalanche/Source/AvalancheMediaEditor/Private/Rundown/Pages/Slate/SAvaRundownTemplatePageList.h b/Engine/Plugins/Experimental/Avalanche/Source/AvalancheMediaEditor/Private/Rundown/Pages/Slate/SAvaRundownTemplatePageList.h
--- a/Engine/Plugins/Experimental/Avalanche/Source/AvalancheMediaEditor/Private/Rundown/Pages/Slate/SAvaRundownTemplatePageList.h
+++ b/Engine/Plugins/Experimental/Avalanche/Source/AvalancheMediaEditor/Private/Rundown/Pages/Slate/SAvaRundownTemplatePageList.h
@@ -47,6 +47,9 @@ public:
 protected:
 	virtual TArray<int32> AddPastedPages(const TArray<FAvaRundownPage>& InPages) override;
 
+	void ResetPagesToDefaults();
+	bool CanResetPagesToDefaults() const;
+
 private:
 	void OnTemplatePageListChanged(const FAvaRundownPageListChangeParams& InParams);
 };
