Motion Design
- Ability to snap the line tool to 30 degrees increments by holding shift.

diff --git a/Engine/Plugins/Experimental/Avalanche/Source/AvalancheInteractiveTools/Private/Planners/AvaInteractiveToolsToolViewportPointListPlanner.cpp b/Engine/Plugins/Experimental/Avalanche/Source/AvalancheInteractiveTools/Private/Planners/AvaInteractiveToolsToolViewportPointListPlanner.cpp
--- a/Engine/Plugins/Experimental/Avalanche/Source/AvalancheInteractiveTools/Private/Planners/AvaInteractiveToolsToolViewportPointListPlanner.cpp
+++ b/Engine/Plugins/Experimental/Avalanche/Source/AvalancheInteractiveTools/Private/Planners/AvaInteractiveToolsToolViewportPointListPlanner.cpp
@@ -12,6 +12,11 @@
 #include "Tools/AvaInteractiveToolsToolBase.h"
 #include "ViewportClient/IAvaViewportClient.h"
 
+void UAvaInteractiveToolsToolViewportPointListPlanner::OverrideCurrentViewportPosition(const FVector2f InPosition)
+{
+	CurrentViewportPosition = InPosition;
+}
+
 void UAvaInteractiveToolsToolViewportPointListPlanner::AddViewportPosition(const FVector2f& InViewportPosition)
 {
 	ViewportPositions.Add(InViewportPosition);
diff --git a/Engine/Plugins/Experimental/Avalanche/Source/AvalancheInteractiveTools/Public/Planners/AvaInteractiveToolsToolViewportPointListPlanner.h b/Engine/Plugins/Experimental/Avalanche/Source/AvalancheInteractiveTools/Public/Planners/AvaInteractiveToolsToolViewportPointListPlanner.h
--- a/Engine/Plugins/Experimental/Avalanche/Source/AvalancheInteractiveTools/Public/Planners/AvaInteractiveToolsToolViewportPointListPlanner.h
+++ b/Engine/Plugins/Experimental/Avalanche/Source/AvalancheInteractiveTools/Public/Planners/AvaInteractiveToolsToolViewportPointListPlanner.h
@@ -20,10 +20,17 @@ class AVALANCHEINTERACTIVETOOLS_API UAvaInteractiveToolsToolViewportPointListPla
 
 public:
 	const FVector2f& GetCurrentViewportPosition() const { return CurrentViewportPosition; }
+
+	void OverrideCurrentViewportPosition(const FVector2f InPosition);
+
 	const TArray<FVector2f>& GetViewportPositions() const { return ViewportPositions; }
+
 	TArray<FVector2f>& GetViewportPositions() { return ViewportPositions; }
+
 	void AddViewportPosition(const FVector2f& InViewportPosition);
+
 	EAvaInteractiveToolsToolViewportPointListPlannerLineStatus GetLineStatus() const { return LineStatus; }
+
 	void SetLineStatus(EAvaInteractiveToolsToolViewportPointListPlannerLineStatus InLineStatus);
 
 	//~ Begin UAvaInteractiveToolsToolViewportPlanner
diff --git a/Engine/Plugins/Experimental/Avalanche/Source/AvalancheShapesEditor/Private/AvaShapesEditorCommands.cpp b/Engine/Plugins/Experimental/Avalanche/Source/AvalancheShapesEditor/Private/AvaShapesEditorCommands.cpp
--- a/Engine/Plugins/Experimental/Avalanche/Source/AvalancheShapesEditor/Private/AvaShapesEditorCommands.cpp
+++ b/Engine/Plugins/Experimental/Avalanche/Source/AvalancheShapesEditor/Private/AvaShapesEditorCommands.cpp
@@ -49,7 +49,7 @@ void FAvaShapesEditorCommands::Register2DCommands()
 
 	UI_COMMAND(Tool_Shape_Line
 		, "Line"
-		, "Create a 2D Line Shape in the viewport."
+		, "Create a 2D Line Shape in the viewport.\n\n- Holding shift will snap the line to 15 degree increments."
 		, EUserInterfaceActionType::ToggleButton
 		, FInputChord());
 
diff --git a/Engine/Plugins/Experimental/Avalanche/Source/AvalancheShapesEditor/Private/Tools/AvaShapesEditorShapeToolLine.cpp b/Engine/Plugins/Experimental/Avalanche/Source/AvalancheShapesEditor/Private/Tools/AvaShapesEditorShapeToolLine.cpp
--- a/Engine/Plugins/Experimental/Avalanche/Source/AvalancheShapesEditor/Private/Tools/AvaShapesEditorShapeToolLine.cpp
+++ b/Engine/Plugins/Experimental/Avalanche/Source/AvalancheShapesEditor/Private/Tools/AvaShapesEditorShapeToolLine.cpp
@@ -5,8 +5,54 @@
 #include "AvaShapesEditorCommands.h"
 #include "Builders/AvaInteractiveToolsToolBuilder.h"
 #include "DynamicMeshes/AvaShapeLineDynMesh.h"
+#include "Framework/Application/SlateApplication.h"
 #include "Planners/AvaInteractiveToolsToolViewportPointListPlanner.h"
 
+namespace UE::AvaShapesEditor::Private
+{
+	constexpr float LineSnapAngle = 15.f;
+
+	// In angle in degrees
+	FVector2f SnapToNearestAngle(const FVector2f& InStartPosition, const FVector2f& InEndPosition, float InSnapAngle)
+	{
+		if (InSnapAngle < 1 || InSnapAngle > 90)
+		{
+			return InEndPosition;
+		}
+
+		static const FVector2f UpVector = FVector2f(1, 0);
+
+		const FVector2f Vector = InEndPosition - InStartPosition;
+		const FVector2f UnitVector = Vector.GetSafeNormal();
+		const float DotProduct = UpVector.Dot(UnitVector);
+		float Angle = FMath::RadiansToDegrees(FMath::Acos(DotProduct));
+		const float Mod = FMath::Fmod(Angle, InSnapAngle);
+
+		if (FMath::IsNearlyZero(Mod))
+		{
+			return InEndPosition;
+		}
+
+		Angle -= Mod;
+
+		if (FMath::Abs(Mod) >= (InSnapAngle *  0.5f))
+		{
+			Angle += InSnapAngle;
+		}
+
+		const float AngleRadians = FMath::DegreesToRadians(Angle);
+
+		// Create snapped vector in correct quadrant
+		// May fail if snap angle is not a factor of 180. So let's not do that?
+		FVector2f SnappedVector = {
+			FMath::Cos(AngleRadians),
+			FMath::Sin(AngleRadians) * (Vector.Y >= 0 ? 1.f : -1.f)
+		};
+
+		return InStartPosition + (SnappedVector * Vector.Length());
+	}
+}
+
 UAvaShapesEditorShapeToolLine::UAvaShapesEditorShapeToolLine()
 {
 	ViewportPlannerClass = UAvaInteractiveToolsToolViewportPointListPlanner::StaticClass();
@@ -55,7 +101,17 @@ void UAvaShapesEditorShapeToolLine::OnViewportPlannerUpdate()
 
 				if (PreviewActor)
 				{
-					SetLineEnds(Cast<AAvaShapeActor>(PreviewActor), Positions[0], PointListPlanner->GetCurrentViewportPosition());
+					LineEndLocation = PointListPlanner->GetCurrentViewportPosition();
+
+					if (FSlateApplication::Get().GetModifierKeys().IsShiftDown())
+					{
+						using namespace UE::AvaShapesEditor::Private;
+
+						LineEndLocation = SnapToNearestAngle(Positions[0], LineEndLocation, LineSnapAngle);
+						PointListPlanner->OverrideCurrentViewportPosition(LineEndLocation);
+					}
+
+					SetLineEnds(Cast<AAvaShapeActor>(PreviewActor), Positions[0], LineEndLocation);
 				}
 				break;
 			
diff --git a/Engine/Plugins/Experimental/Avalanche/Source/AvalancheShapesEditor/Private/Tools/AvaShapesEditorShapeToolLine.h b/Engine/Plugins/Experimental/Avalanche/Source/AvalancheShapesEditor/Private/Tools/AvaShapesEditorShapeToolLine.h
--- a/Engine/Plugins/Experimental/Avalanche/Source/AvalancheShapesEditor/Private/Tools/AvaShapesEditorShapeToolLine.h
+++ b/Engine/Plugins/Experimental/Avalanche/Source/AvalancheShapesEditor/Private/Tools/AvaShapesEditorShapeToolLine.h
@@ -24,6 +24,8 @@ protected:
 	// The minimum dimension
 	static constexpr float MinDim = 5;
 
+	FVector2f LineEndLocation;
+
 	//~ Begin UAvaShapesEditorShapeToolBase
 	virtual void InitShape(UAvaShapeDynamicMeshBase* InShape) const override;
 	virtual void SetShapeSize(AAvaShapeActor* InShapeActor, const FVector2D& InShapeSize) const override;
