MotionDesign : Shapes
- Fix shape duplication issue where "use single material" property would reset (reported in beta)

diff --git a/Engine/Plugins/Experimental/Avalanche/Source/AvalancheShapes/Private/DynamicMeshes/AvaShapeDynMeshBase.cpp b/Engine/Plugins/Experimental/Avalanche/Source/AvalancheShapes/Private/DynamicMeshes/AvaShapeDynMeshBase.cpp
--- a/Engine/Plugins/Experimental/Avalanche/Source/AvalancheShapes/Private/DynamicMeshes/AvaShapeDynMeshBase.cpp
+++ b/Engine/Plugins/Experimental/Avalanche/Source/AvalancheShapes/Private/DynamicMeshes/AvaShapeDynMeshBase.cpp
@@ -1732,11 +1732,15 @@ void UAvaShapeDynamicMeshBase::PostEditImport()
 	// reset cache component
 	CachedComponent.Reset();
 
+	const bool bUseSingleMaterial = bUsePrimaryMaterialEverywhere;
+
 	// Update material options
 	for (const int32 MeshIdx : GetMeshesIndexes())
 	{
 		OnMaterialTypeChanged(MeshIdx);
 	}
+
+	SetUsePrimaryMaterialEverywhere(bUseSingleMaterial);
 }
 
 void UAvaShapeDynamicMeshBase::PostDuplicate(EDuplicateMode::Type DuplicateMode)
