[Motion Design] Statically initialized FGCObject in global scope crashing editor in Debug Game

Moving static member DefaultFontObjects to a static function, to avoid having it static in global space.

diff --git a/Engine/Plugins/Experimental/Avalanche/Source/AvalancheText/Private/Font/AvaFont.cpp b/Engine/Plugins/Experimental/Avalanche/Source/AvalancheText/Private/Font/AvaFont.cpp
--- a/Engine/Plugins/Experimental/Avalanche/Source/AvalancheText/Private/Font/AvaFont.cpp
+++ b/Engine/Plugins/Experimental/Avalanche/Source/AvalancheText/Private/Font/AvaFont.cpp
@@ -6,7 +6,12 @@
 
 #define LOCTEXT_NAMESPACE "AvaFont"
 
-FAvaFont::FAvaDefaultFontObjects FAvaFont::DefaultFontObjects;
+FAvaFont::FAvaDefaultFontObjects& FAvaFont::GetDefaultFontObjects()
+{
+	static FAvaDefaultFontObjects DefaultFontObjects;
+
+	return DefaultFontObjects;
+}
 
 UFont* FAvaFont::GetDefaultFont()
 {
@@ -15,7 +20,7 @@ UFont* FAvaFont::GetDefaultFont()
 		return nullptr;
 	}
 
-	if (!DefaultFontObjects.AvaDefaultFont)
+	if (!GetDefaultFontObjects().AvaDefaultFont)
 	{
 		const FAssetRegistryModule& AssetRegistryModule = FModuleManager::LoadModuleChecked<FAssetRegistryModule>("AssetRegistry");
 
@@ -52,11 +57,11 @@ UFont* FAvaFont::GetDefaultFont()
 
 		if (ensureMsgf(DefaultFont, TEXT("MotionDesignFont: cannot load any font to be used as default.")))
 		{
-			DefaultFontObjects.AvaDefaultFont = DefaultFont;
+			GetDefaultFontObjects().AvaDefaultFont = DefaultFont;
 		}
 	}
 
-	return DefaultFontObjects.AvaDefaultFont.Get();
+	return GetDefaultFontObjects().AvaDefaultFont.Get();
 }
 
 FString FAvaFont::GenerateFontFormattedString(const FString& InFontName, const FString& InFontObjectPathName)
@@ -358,16 +363,16 @@ void FAvaFont::PostSerialize(const FArchive& Ar)
 
 UAvaFontObject* FAvaFont::GetDefaultAvaFontObject()
 {
-	if (!DefaultFontObjects.AvaDefaultFontObject)
+	if (!GetDefaultFontObjects().AvaDefaultFontObject)
 	{
 		if (UFont* DefaultFont = GetDefaultFont())
 		{
-			DefaultFontObjects.AvaDefaultFontObject = NewObject<UAvaFontObject>(DefaultFont);
-			DefaultFontObjects.AvaDefaultFontObject->InitProjectFont(DefaultFont, DefaultFont->GetName());
+			GetDefaultFontObjects().AvaDefaultFontObject = NewObject<UAvaFontObject>(DefaultFont);
+			GetDefaultFontObjects().AvaDefaultFontObject->InitProjectFont(DefaultFont, DefaultFont->GetName());
 		}
 	}
 
-	return DefaultFontObjects.AvaDefaultFontObject.Get();
+	return GetDefaultFontObjects().AvaDefaultFontObject.Get();
 }
 
 void FAvaFont::RefreshName()
diff --git a/Engine/Plugins/Experimental/Avalanche/Source/AvalancheText/Public/Font/AvaFont.h b/Engine/Plugins/Experimental/Avalanche/Source/AvalancheText/Public/Font/AvaFont.h
--- a/Engine/Plugins/Experimental/Avalanche/Source/AvalancheText/Public/Font/AvaFont.h
+++ b/Engine/Plugins/Experimental/Avalanche/Source/AvalancheText/Public/Font/AvaFont.h
@@ -163,6 +163,9 @@ private:
 		//~ End FGCObject
 	};
 
+	/**
+	 * Used to handle font state and possible loading issues
+	 */
 	enum EFontAssetState
 	{
 		DefaultFont,   // Font is meant to be the default one
@@ -170,20 +173,27 @@ private:
 		FallbackFont   // Referenced font asset is missing
 	};
 
+	/** Returns the default font object for Ava Fonts */
 	static UAvaFontObject* GetDefaultAvaFontObject();
 
+	/** Returns a static FAvaDefaultFontObjects struct holding defaults values for fonts */
+	static FAvaDefaultFontObjects& GetDefaultFontObjects();
+
+	/** Search for a font asset based on a font name */
 	static UFont* GetFontByName(const FString& InFontName);
 
+	/** Initializes default values */
 	void InitDefaults();
 
+	/** Ensures the Font Name property is up to date with current font object */
 	void RefreshName();
 
+	/** Refreshes asset state */
 	void RefreshAssetState();
 
+	/** Initialize this Avalanche Font using a UFont asset */
 	void InitFromFont(UFont* InFont);
 
-	static FAvaDefaultFontObjects DefaultFontObjects;
-
 	/**
 	 * Deprecated - used to reference the font used by FAvaFont.
 	 * Older assets still have the font stored in this field. Calling GetFont() will move the font to the currently used MotionDesignFontObject property
