[Motion Design] Static NullChannel in FAvaBroadcastOutputChannel Crashing in Debug Game

1 - Wrapping static properties in static functions:
- FAvaBroadcastOutputChannel::GetNullChannel()
- FAvaBroadcastProfile::GetNullProfile()  --> NullProfile was currently causing no issue, but it looked similar to the NulChannel pattern, and the structs holds a TWeakObjectPtr UPROPERTY, so it got wrapped too.

2 - Adding NoInit constructors for FAvaBroadcastOutputChannel and FAvaViewportQualitySettings, to be used when creating NullChannel.

diff --git a/Engine/Plugins/Experimental/Avalanche/Source/Avalanche/Private/Viewport/AvaViewportQualitySettings.cpp b/Engine/Plugins/Experimental/Avalanche/Source/Avalanche/Private/Viewport/AvaViewportQualitySettings.cpp
--- a/Engine/Plugins/Experimental/Avalanche/Source/Avalanche/Private/Viewport/AvaViewportQualitySettings.cpp
+++ b/Engine/Plugins/Experimental/Avalanche/Source/Avalanche/Private/Viewport/AvaViewportQualitySettings.cpp
@@ -10,6 +10,10 @@ FAvaViewportQualitySettings::FAvaViewportQualitySettings()
 	Features = DefaultFeatures();
 }
 
+FAvaViewportQualitySettings::FAvaViewportQualitySettings(ENoInit NoInit)
+{
+}
+
 FAvaViewportQualitySettings::FAvaViewportQualitySettings(const bool bInUseAllFeatures)
 {
 	Features = AllFeatures(bInUseAllFeatures);
diff --git a/Engine/Plugins/Experimental/Avalanche/Source/Avalanche/Public/Viewport/AvaViewportQualitySettings.h b/Engine/Plugins/Experimental/Avalanche/Source/Avalanche/Public/Viewport/AvaViewportQualitySettings.h
--- a/Engine/Plugins/Experimental/Avalanche/Source/Avalanche/Public/Viewport/AvaViewportQualitySettings.h
+++ b/Engine/Plugins/Experimental/Avalanche/Source/Avalanche/Public/Viewport/AvaViewportQualitySettings.h
@@ -63,6 +63,7 @@ public:
 	static void SortFeaturesByDisplayText(TArray<FAvaViewportQualitySettingsFeature>& InFeatures);
 
 	FAvaViewportQualitySettings();
+	FAvaViewportQualitySettings(ENoInit NoInit);
 	FAvaViewportQualitySettings(const bool bInUseAllFeatures);
 	FAvaViewportQualitySettings(const FEngineShowFlags& InShowFlags);
 	FAvaViewportQualitySettings(const TArray<FAvaViewportQualitySettingsFeature>& InFeatures);
diff --git a/Engine/Plugins/Experimental/Avalanche/Source/AvalancheMedia/Private/Broadcast/AvaBroadcast.cpp b/Engine/Plugins/Experimental/Avalanche/Source/AvalancheMedia/Private/Broadcast/AvaBroadcast.cpp
--- a/Engine/Plugins/Experimental/Avalanche/Source/AvalancheMedia/Private/Broadcast/AvaBroadcast.cpp
+++ b/Engine/Plugins/Experimental/Avalanche/Source/AvalancheMedia/Private/Broadcast/AvaBroadcast.cpp
@@ -301,7 +301,7 @@ FAvaBroadcastProfile& UAvaBroadcast::GetProfile(FName InProfileName)
 	{
 		return Profiles[InProfileName];
 	}
-	return FAvaBroadcastProfile::NullProfile;
+	return FAvaBroadcastProfile::GetNullProfile();
 }
 
 const FAvaBroadcastProfile& UAvaBroadcast::GetProfile(FName InProfileName) const
@@ -310,7 +310,7 @@ const FAvaBroadcastProfile& UAvaBroadcast::GetProfile(FName InProfileName) const
 	{
 		return Profiles[InProfileName];
 	}
-	return FAvaBroadcastProfile::NullProfile;
+	return FAvaBroadcastProfile::GetNullProfile();
 }
 
 #if WITH_EDITOR
diff --git a/Engine/Plugins/Experimental/Avalanche/Source/AvalancheMedia/Private/Broadcast/AvaBroadcastProfile.cpp b/Engine/Plugins/Experimental/Avalanche/Source/AvalancheMedia/Private/Broadcast/AvaBroadcastProfile.cpp
--- a/Engine/Plugins/Experimental/Avalanche/Source/AvalancheMedia/Private/Broadcast/AvaBroadcastProfile.cpp
+++ b/Engine/Plugins/Experimental/Avalanche/Source/AvalancheMedia/Private/Broadcast/AvaBroadcastProfile.cpp
@@ -6,7 +6,12 @@
 
 #define LOCTEXT_NAMESPACE "AvaBroadcastProfile"
 
-FAvaBroadcastProfile FAvaBroadcastProfile::NullProfile(nullptr, NAME_None);
+FAvaBroadcastProfile& FAvaBroadcastProfile::GetNullProfile()
+{
+	static FAvaBroadcastProfile NullProfile(nullptr, NAME_None);
+
+	return NullProfile;
+}
 
 FAvaBroadcastProfile::FAvaBroadcastProfile(UAvaBroadcast* InBroadcast, FName InProfileName)
 	: ParentBroadcastWeak(InBroadcast)
@@ -95,7 +100,7 @@ bool FAvaBroadcastProfile::IsBroadcastingAllChannels() const
 
 bool FAvaBroadcastProfile::IsValidProfile() const
 {
-	return this != &FAvaBroadcastProfile::NullProfile && !ProfileName.IsNone();
+	return this != &FAvaBroadcastProfile::GetNullProfile() && !ProfileName.IsNone();
 }
 
 namespace UE::AvaBroadcastProfile::Private
@@ -266,7 +271,7 @@ const FAvaBroadcastOutputChannel& FAvaBroadcastProfile::GetLocalChannel(FName In
 		return Channels[ChannelIndex];
 	}
 	
-	return FAvaBroadcastOutputChannel::NullChannel;
+	return FAvaBroadcastOutputChannel::GetNullChannel();
 }
 
 FAvaBroadcastOutputChannel& FAvaBroadcastProfile::GetLocalChannelMutable(FName InChannelName)
@@ -278,7 +283,7 @@ FAvaBroadcastOutputChannel& FAvaBroadcastProfile::GetLocalChannelMutable(FName I
 		return Channels[ChannelIndex];
 	}
 	
-	return FAvaBroadcastOutputChannel::NullChannel;
+	return FAvaBroadcastOutputChannel::GetNullChannel();
 }
 
 const FAvaBroadcastOutputChannel& FAvaBroadcastProfile::GetChannel(FName InChannelName) const
diff --git a/Engine/Plugins/Experimental/Avalanche/Source/AvalancheMedia/Private/Broadcast/Channel/AvaBroadcastOutputChannel.cpp b/Engine/Plugins/Experimental/Avalanche/Source/AvalancheMedia/Private/Broadcast/Channel/AvaBroadcastOutputChannel.cpp
--- a/Engine/Plugins/Experimental/Avalanche/Source/AvalancheMedia/Private/Broadcast/Channel/AvaBroadcastOutputChannel.cpp
+++ b/Engine/Plugins/Experimental/Avalanche/Source/AvalancheMedia/Private/Broadcast/Channel/AvaBroadcastOutputChannel.cpp
@@ -62,6 +62,11 @@ void FAvaBroadcastCapture::OnStateChangedNative()
 	}
 }
 
+FAvaBroadcastOutputChannel::FAvaBroadcastOutputChannel(ENoInit NoInit)
+	: QualitySettings(NoInit)
+{
+}
+
 namespace UE::AvaMedia::Private
 {
 	bool IsCapturing(const UMediaCapture* InMediaCapture)
@@ -171,7 +176,13 @@ namespace UE::AvaMedia::Private
 	}
 }
 
-FAvaBroadcastOutputChannel FAvaBroadcastOutputChannel::NullChannel;
+FAvaBroadcastOutputChannel& FAvaBroadcastOutputChannel::GetNullChannel()
+{
+	static FAvaBroadcastOutputChannel NullChannel(NoInit);
+
+	return NullChannel;
+}
+
 FAvaBroadcastOutputChannel::FOnAvaChannelChanged FAvaBroadcastOutputChannel::OnChannelChanged;
 FAvaBroadcastOutputChannel::FMediaOutputStateChanged FAvaBroadcastOutputChannel::OnMediaOutputStateChanged;
 
@@ -620,7 +631,7 @@ void FAvaBroadcastOutputChannel::DrawPlaceholderWidget() const
 
 bool FAvaBroadcastOutputChannel::IsValidChannel() const
 {
-	return this != &FAvaBroadcastOutputChannel::NullChannel && GetChannelName() != NAME_None;
+	return this != &FAvaBroadcastOutputChannel::GetNullChannel() && GetChannelName() != NAME_None;
 }
 
 bool FAvaBroadcastOutputChannel::StartChannelBroadcast()
diff --git a/Engine/Plugins/Experimental/Avalanche/Source/AvalancheMedia/Public/Broadcast/AvaBroadcastProfile.h b/Engine/Plugins/Experimental/Avalanche/Source/AvalancheMedia/Public/Broadcast/AvaBroadcastProfile.h
--- a/Engine/Plugins/Experimental/Avalanche/Source/AvalancheMedia/Public/Broadcast/AvaBroadcastProfile.h
+++ b/Engine/Plugins/Experimental/Avalanche/Source/AvalancheMedia/Public/Broadcast/AvaBroadcastProfile.h
@@ -15,13 +15,12 @@ struct AVALANCHEMEDIA_API FAvaBroadcastProfile
 	friend UAvaBroadcast;
 	
 public:	
+	static FAvaBroadcastProfile& GetNullProfile();
 	FAvaBroadcastProfile() : FAvaBroadcastProfile(nullptr, NAME_None) {}
 	FAvaBroadcastProfile(UAvaBroadcast* InBroadcast, FName InProfileName);
 
 	void BeginDestroy();
-	
-	static FAvaBroadcastProfile NullProfile;
-	
+
 	FName GetName() const { return ProfileName; }
 	UAvaBroadcast& GetBroadcast() const;
 	
diff --git a/Engine/Plugins/Experimental/Avalanche/Source/AvalancheMedia/Public/Broadcast/Channel/AvaBroadcastOutputChannel.h b/Engine/Plugins/Experimental/Avalanche/Source/AvalancheMedia/Public/Broadcast/Channel/AvaBroadcastOutputChannel.h
--- a/Engine/Plugins/Experimental/Avalanche/Source/AvalancheMedia/Public/Broadcast/Channel/AvaBroadcastOutputChannel.h
+++ b/Engine/Plugins/Experimental/Avalanche/Source/AvalancheMedia/Public/Broadcast/Channel/AvaBroadcastOutputChannel.h
@@ -69,6 +69,7 @@ struct AVALANCHEMEDIA_API FAvaBroadcastOutputChannel
 
 public:
 	FAvaBroadcastOutputChannel() = default;
+	FAvaBroadcastOutputChannel(ENoInit NoInit);
 	explicit FAvaBroadcastOutputChannel(FAvaBroadcastProfile* InProfile) : Profile(InProfile) {}
 	~FAvaBroadcastOutputChannel();
 
@@ -76,7 +77,7 @@ public:
 	void ReleasePlaceholderRenderTargets();
 	void ReleaseOutputs();
 	
-	static FAvaBroadcastOutputChannel NullChannel;
+	static FAvaBroadcastOutputChannel& GetNullChannel();
 	
 	static void DuplicateChannel(const FAvaBroadcastOutputChannel& InSourceChannel, FAvaBroadcastOutputChannel& OutTargetChannel);
 
