[Motion Design]
 - Reset individual rundown editor selection RC values

diff --git a/Engine/Plugins/Experimental/Avalanche/Source/AvalancheMediaEditor/Private/Rundown/DetailsView/RemoteControl/Controllers/AvaRundownPageControllerContextMenu.cpp b/Engine/Plugins/Experimental/Avalanche/Source/AvalancheMediaEditor/Private/Rundown/DetailsView/RemoteControl/Controllers/AvaRundownPageControllerContextMenu.cpp
new file mode 100644
--- /dev/null
+++ b/Engine/Plugins/Experimental/Avalanche/Source/AvalancheMediaEditor/Private/Rundown/DetailsView/RemoteControl/Controllers/AvaRundownPageControllerContextMenu.cpp
@@ -0,0 +1,249 @@
+// Copyright Epic Games, Inc. All Rights Reserved.
+
+#include "AvaRundownPageControllerContextMenu.h"
+#include "AvaRundownRCControllerItem.h"
+#include "Controller/RCController.h"
+#include "IAvaMediaEditorModule.h"
+#include "Rundown/AvaRundownCommands.h"
+#include "Rundown/AvaRundownEditor.h"
+#include "Rundown/AvaRundownPage.h"
+#include "SAvaRundownRCControllerPanel.h"
+#include "ScopedTransaction.h"
+#include "ToolMenu.h"
+#include "ToolMenus.h"
+#include "ToolMenuSection.h"
+#include "UObject/WeakObjectPtrTemplates.h"
+
+#define LOCTEXT_NAMESPACE "AvaRundownPageControllerContextMenu"
+
+FAvaRundownPageControllerContextMenu::FAvaRundownPageControllerContextMenu(const TWeakPtr<FUICommandList>& InCommandListWeak)
+{
+	BindCommands(InCommandListWeak);
+}
+
+void FAvaRundownPageControllerContextMenu::BindCommands(const TWeakPtr<FUICommandList>& InCommandListWeak)
+{
+	TSharedPtr<FUICommandList> CommandList = InCommandListWeak.Pin();
+	if (!CommandList.IsValid())
+	{
+		return;
+	}
+
+	CommandListWeak = InCommandListWeak;
+
+	const FAvaRundownCommands& RundownCommands = FAvaRundownCommands::Get();
+
+	CommandList->MapAction(RundownCommands.ResetValuesToDefaults,
+		FExecuteAction::CreateRaw(this, &FAvaRundownPageControllerContextMenu::ResetValuesToDefaults, false),
+		FCanExecuteAction::CreateRaw(this, &FAvaRundownPageControllerContextMenu::CanResetValuesToDefaults, false));
+
+	CommandList->MapAction(RundownCommands.ResetValuesToTemplate,
+		FExecuteAction::CreateRaw(this, &FAvaRundownPageControllerContextMenu::ResetValuesToDefaults, true),
+		FCanExecuteAction::CreateRaw(this, &FAvaRundownPageControllerContextMenu::CanResetValuesToDefaults, true));
+}
+
+TSharedRef<SWidget> FAvaRundownPageControllerContextMenu::GeneratePageContextMenuWidget(const TWeakPtr<FAvaRundownEditor>& InRundownEditorWeak
+	, const FAvaRundownPage& InRundownPage
+	, const TWeakPtr<SAvaRundownRCControllerPanel>& InControllerListWidgetWeak)
+{
+	UToolMenus* const ToolMenus = UToolMenus::Get();
+	check(ToolMenus);
+	
+	static const FName ContextMenuName = TEXT("RundownRCControllerPanel");
+
+	if (!ToolMenus->IsMenuRegistered(ContextMenuName))
+	{
+		UToolMenu* const ContextMenu = ToolMenus->RegisterMenu(ContextMenuName, NAME_None, EMultiBoxType::Menu);
+		check(ContextMenu);
+
+		ContextMenu->AddDynamicSection(TEXT("PopulateContextMenu"), FNewToolMenuDelegate::CreateLambda(
+			[this](UToolMenu* const InMenu)
+			{
+				if (IsValid(InMenu))
+				{
+					PopulatePageContextMenu(*InMenu, InMenu->FindContext<UAvaRundownPageControllerContext>());
+				}
+			}));
+	}
+
+	UAvaRundownPageControllerContext* const ContextObject = NewObject<UAvaRundownPageControllerContext>();
+	check(ContextObject);
+	ContextObject->InitContext(InRundownEditorWeak, InRundownPage.GetPageId(), InControllerListWidgetWeak);
+
+	TSharedPtr<FExtender> Extender;
+
+	// Compatibility with IAvaMediaEditorModule Rundown Menu Extensibility Manager
+	TSharedPtr<FUICommandList> CommandList = CommandListWeak.Pin();
+	if (CommandList.IsValid())
+	{
+		TSharedPtr<FExtensibilityManager> MenuExtensibility = IAvaMediaEditorModule::Get().GetRundownMenuExtensibilityManager();
+		if (MenuExtensibility.IsValid())
+		{
+			if (TSharedPtr<FAvaRundownEditor> RundownEditor = InRundownEditorWeak.Pin())
+			{
+				if (const TArray<UObject*>* EditingObjects = RundownEditor->GetObjectsCurrentlyBeingEdited())
+				{
+					Extender = MenuExtensibility->GetAllExtenders(CommandList.ToSharedRef(), *EditingObjects);
+				}
+			}
+		}
+	}
+
+	FToolMenuContext Context(CommandList, Extender, ContextObject);
+	return ToolMenus->GenerateWidget(ContextMenuName, Context);
+}
+
+void FAvaRundownPageControllerContextMenu::PopulatePageContextMenu(UToolMenu& InMenu, UAvaRundownPageControllerContext* const InContext)
+{
+	CurrentContext = InContext;
+
+	if (!IsValid(InContext))
+	{
+		return;
+	}
+
+	TSharedPtr<SAvaRundownRCControllerPanel> ListWidget = CurrentContext->GetControllerListWidget();
+	if (ListWidget.IsValid() && ListWidget->GetSelectedControllerItems().Num() == 0)
+	{
+		return;
+	}
+
+	UAvaRundown* const Rundown = GetContextRundown();
+	if (!IsValid(Rundown))
+	{
+		return;
+	}
+
+	const FAvaRundownPage& RundownPage = Rundown->GetPage(CurrentContext->GetRundownPageId());
+	if (!RundownPage.IsValidPage())
+	{
+		return;
+	}
+
+	const FAvaRundownCommands& RundownCommands = FAvaRundownCommands::Get();
+
+	FToolMenuSection& Section = InMenu.FindOrAddSection(TEXT("ControllerActions"), LOCTEXT("ControllerActions", "Controller Actions"));
+
+	Section.AddMenuEntry(TEXT("ResetControllerToDefaults"),
+		RundownCommands.ResetValuesToDefaults,
+		LOCTEXT("ResetControllerToDefaults", "Reset to Defaults"),
+		LOCTEXT("ResetControllerToDefaultsToolTip", "Reset selected controller values to the defaults."),
+		FSlateIcon(FAppStyle::GetAppStyleSetName(), TEXT("PropertyWindow.DiffersFromDefault")));
+
+	if (!RundownPage.IsTemplate())
+	{
+		Section.AddMenuEntry(TEXT("ResetControllerToTemplate"),
+			RundownCommands.ResetValuesToTemplate,
+			LOCTEXT("ResetControllerToTemplate", "Reset to Template"),
+			LOCTEXT("ResetControllerToTemplateToolTip", "Reset selected controller values to the defaults."),
+			FSlateIcon(FAppStyle::GetAppStyleSetName(), TEXT("PropertyWindow.DiffersFromDefault")));
+	}
+}
+
+TArray<TWeakPtr<FAvaRundownRCControllerItem>> FAvaRundownPageControllerContextMenu::GetSelectedControllerItems() const
+{
+	TArray<TWeakPtr<FAvaRundownRCControllerItem>> OutItems;
+
+	if (CurrentContext.IsValid())
+	{
+		if (TSharedPtr<SAvaRundownRCControllerPanel> ControllerListWidget = CurrentContext->GetControllerListWidget())
+		{
+			for (const FAvaRundownRCControllerItemPtr& ControllerItem : ControllerListWidget->GetSelectedControllerItems())
+			{
+				if (ControllerItem.IsValid())
+				{
+					OutItems.Add(ControllerItem);
+				}
+			}
+		}
+	}
+
+	return OutItems;
+}
+
+bool FAvaRundownPageControllerContextMenu::HasValidSelectedItems() const
+{
+	return GetSelectedControllerItems().Num() > 0;
+}
+
+void FAvaRundownPageControllerContextMenu::ResetValuesToDefaults(const bool bInUseTemplateValues)
+{
+	UAvaRundown* const Rundown = GetContextRundown();
+	if (!IsValid(Rundown))
+	{
+		return;
+	}
+
+	FScopedTransaction Transaction(LOCTEXT("ResetControllerValuesTransaction", "Reset RC Controller Values"));
+	Rundown->Modify();
+
+	TArray<TWeakPtr<FAvaRundownRCControllerItem>> SelectedControllerItems = GetSelectedControllerItems();
+	for (const TWeakPtr<FAvaRundownRCControllerItem>& ControllerItemWeak : SelectedControllerItems)
+	{
+		if (TSharedPtr<FAvaRundownRCControllerItem> ControllerItem = ControllerItemWeak.Pin())
+		{
+			const URCController* const Controller = ControllerItem->GetController();
+			if (IsValid(Controller))
+			{
+				Rundown->ResetRemoteControlControllerValue(CurrentContext->GetRundownPageId()
+					, Controller->Id, bInUseTemplateValues, /*bInIsDefault=*/false);
+			}
+		}
+	}
+}
+
+bool FAvaRundownPageControllerContextMenu::CanResetValuesToDefaults(const bool bInUseTemplateValues) const
+{
+	UAvaRundown* const Rundown = GetContextRundown();
+	if (!IsValid(Rundown))
+	{
+		return false;
+	}
+
+	const FAvaRundownPage& RundownPage = Rundown->GetPage(CurrentContext->GetRundownPageId());
+	if (!RundownPage.IsValidPage())
+	{
+		return false;
+	}
+
+	const TArray<TWeakPtr<FAvaRundownRCControllerItem>> SelectedItems = GetSelectedControllerItems();
+	if (SelectedItems.Num() == 0)
+	{
+		return false;
+	}
+
+	for (const TWeakPtr<FAvaRundownRCControllerItem>& ControllerItemWeak : SelectedItems)
+	{
+		if (const TSharedPtr<FAvaRundownRCControllerItem> ControllerItem = ControllerItemWeak.Pin())
+		{
+			const URCController* const Controller = ControllerItem->GetController();
+			if (IsValid(Controller))
+			{
+				if (!RundownPage.IsDefaultControllerValue(Rundown, Controller->Id, bInUseTemplateValues))
+				{
+					return true;
+				}
+			}
+		}
+	}
+
+	return false;
+}
+
+UAvaRundown* FAvaRundownPageControllerContextMenu::GetContextRundown() const
+{
+	if (!CurrentContext.IsValid())
+	{
+		return nullptr;
+	}
+
+	const TSharedPtr<FAvaRundownEditor> RundownEditor = CurrentContext->GetRundownEditor();
+	if (!RundownEditor.IsValid())
+	{
+		return nullptr;
+	}
+
+	return RundownEditor->GetRundown();
+}
+
+#undef LOCTEXT_NAMESPACE
diff --git a/Engine/Plugins/Experimental/Avalanche/Source/AvalancheMediaEditor/Private/Rundown/DetailsView/RemoteControl/Controllers/AvaRundownPageControllerContextMenu.h b/Engine/Plugins/Experimental/Avalanche/Source/AvalancheMediaEditor/Private/Rundown/DetailsView/RemoteControl/Controllers/AvaRundownPageControllerContextMenu.h
new file mode 100644
--- /dev/null
+++ b/Engine/Plugins/Experimental/Avalanche/Source/AvalancheMediaEditor/Private/Rundown/DetailsView/RemoteControl/Controllers/AvaRundownPageControllerContextMenu.h
@@ -0,0 +1,44 @@
+// Copyright Epic Games, Inc. All Rights Reserved.
+
+#pragma once
+
+#include "Rundown/Pages/AvaRundownPageControllerContext.h"
+#include "Templates/SharedPointer.h"
+#include "UObject/WeakObjectPtrTemplatesFwd.h"
+
+class FAvaRundownEditor;
+class FAvaRundownRCControllerItem;
+class FUICommandList;
+class SAvaRundownRCControllerPanel;
+class SWidget;
+class UAvaRundownPageControllerContext;
+class UToolMenu;
+struct FAvaRundownPage;
+
+class FAvaRundownPageControllerContextMenu : public TSharedFromThis<FAvaRundownPageControllerContextMenu>
+{
+public:
+	FAvaRundownPageControllerContextMenu(const TWeakPtr<FUICommandList>& InCommandList);
+
+	TSharedRef<SWidget> GeneratePageContextMenuWidget(const TWeakPtr<FAvaRundownEditor>& InRundownEditorWeak
+		, const FAvaRundownPage& InRundownPage
+		, const TWeakPtr<SAvaRundownRCControllerPanel>& InControllerListWidgetWeak);
+
+	bool HasValidSelectedItems() const;
+
+private:
+	void PopulatePageContextMenu(UToolMenu& InMenu, UAvaRundownPageControllerContext* const InContext);
+
+	void BindCommands(const TWeakPtr<FUICommandList>& InCommandListWeak);
+
+	void ResetValuesToDefaults(const bool bInUseTemplateValues);
+	bool CanResetValuesToDefaults(const bool bInUseTemplateValues) const;
+
+	TArray<TWeakPtr<FAvaRundownRCControllerItem>> GetSelectedControllerItems() const;
+
+	UAvaRundown* GetContextRundown() const;
+
+	TWeakPtr<FUICommandList> CommandListWeak;
+
+	TWeakObjectPtr<UAvaRundownPageControllerContext> CurrentContext;
+};
diff --git a/Engine/Plugins/Experimental/Avalanche/Source/AvalancheMediaEditor/Private/Rundown/DetailsView/RemoteControl/Controllers/AvaRundownRCControllerItem.cpp b/Engine/Plugins/Experimental/Avalanche/Source/AvalancheMediaEditor/Private/Rundown/DetailsView/RemoteControl/Controllers/AvaRundownRCControllerItem.cpp
--- a/Engine/Plugins/Experimental/Avalanche/Source/AvalancheMediaEditor/Private/Rundown/DetailsView/RemoteControl/Controllers/AvaRundownRCControllerItem.cpp
+++ b/Engine/Plugins/Experimental/Avalanche/Source/AvalancheMediaEditor/Private/Rundown/DetailsView/RemoteControl/Controllers/AvaRundownRCControllerItem.cpp
@@ -1,16 +1,18 @@
 // Copyright Epic Games, Inc. All Rights Reserved.
 
 #include "AvaRundownRCControllerItem.h"
-
 #include "Controller/RCController.h"
 #include "IRemoteControlUIModule.h"
 #include "RCVirtualProperty.h"
 #include "SAvaRundownRCControllerItemRow.h"
+#include "UObject/WeakObjectPtrTemplates.h"
 
 FAvaRundownRCControllerItem::FAvaRundownRCControllerItem(int32 InInstanceIndex, FName InAssetName, URCController* InController, const TSharedRef<IDetailTreeNode>& InTreeNode)
 {
 	AssetName = InAssetName;
 	InstanceIndex = InInstanceIndex;
+
+	Controller = InController;
 	
 	if (InController)
 	{
@@ -35,3 +37,8 @@ TSharedRef<ITableRow> FAvaRundownRCControllerItem::CreateWidget(TSharedRef<SAvaR
 {
 	return SNew(SAvaRundownRCControllerItemRow, InControllerPanel, InOwnerTable, SharedThis(this));
 }
+
+URCController* FAvaRundownRCControllerItem::GetController() const
+{
+	return Controller.Get();
+}
diff --git a/Engine/Plugins/Experimental/Avalanche/Source/AvalancheMediaEditor/Private/Rundown/DetailsView/RemoteControl/Controllers/AvaRundownRCControllerItem.h b/Engine/Plugins/Experimental/Avalanche/Source/AvalancheMediaEditor/Private/Rundown/DetailsView/RemoteControl/Controllers/AvaRundownRCControllerItem.h
--- a/Engine/Plugins/Experimental/Avalanche/Source/AvalancheMediaEditor/Private/Rundown/DetailsView/RemoteControl/Controllers/AvaRundownRCControllerItem.h
+++ b/Engine/Plugins/Experimental/Avalanche/Source/AvalancheMediaEditor/Private/Rundown/DetailsView/RemoteControl/Controllers/AvaRundownRCControllerItem.h
@@ -5,6 +5,7 @@
 #include "IDetailTreeNode.h"
 #include "Internationalization/Text.h"
 #include "Templates/SharedPointer.h"
+#include "UObject/WeakObjectPtrTemplatesFwd.h"
 
 class ITableRow;
 class SAvaRundownRCControllerPanel;
@@ -24,6 +25,8 @@ public:
 	int32 GetInstanceIndex() const { return InstanceIndex; }
 	int32 GetDisplayIndex() const { return DisplayIndex; }
 	FName GetAssetName() const { return AssetName; }
+
+	URCController* GetController() const;
 	
 private:
 	int32 InstanceIndex = 0;
@@ -31,4 +34,6 @@ private:
 	FName AssetName;
 	FText DisplayNameText;
 	FNodeWidgets NodeWidgets;
+
+	TWeakObjectPtr<URCController> Controller;
 };
diff --git a/Engine/Plugins/Experimental/Avalanche/Source/AvalancheMediaEditor/Private/Rundown/DetailsView/RemoteControl/Controllers/SAvaRundownRCControllerPanel.cpp b/Engine/Plugins/Experimental/Avalanche/Source/AvalancheMediaEditor/Private/Rundown/DetailsView/RemoteControl/Controllers/SAvaRundownRCControllerPanel.cpp
--- a/Engine/Plugins/Experimental/Avalanche/Source/AvalancheMediaEditor/Private/Rundown/DetailsView/RemoteControl/Controllers/SAvaRundownRCControllerPanel.cpp
+++ b/Engine/Plugins/Experimental/Avalanche/Source/AvalancheMediaEditor/Private/Rundown/DetailsView/RemoteControl/Controllers/SAvaRundownRCControllerPanel.cpp
@@ -1,18 +1,21 @@
 // Copyright Epic Games, Inc. All Rights Reserved.
 
 #include "SAvaRundownRCControllerPanel.h"
-
 #include "AvaRundownRCControllerItem.h"
 #include "Behaviour/Builtin/Path/RCSetAssetByPathBehaviour.h"
 #include "Controller/RCController.h"
+#include "IAvaMediaEditorModule.h"
 #include "IDetailTreeNode.h"
 #include "IPropertyRowGenerator.h"
 #include "Playable/AvaPlayableRemoteControl.h"
 #include "RCVirtualProperty.h"
+#include "Rundown/AvaRundownCommands.h"
 #include "Rundown/AvaRundownEditor.h"
 #include "Rundown/AvaRundownEditorUtils.h"
 #include "Rundown/AvaRundownManagedInstanceCache.h"
 #include "Rundown/AvaRundownPage.h"
+#include "Rundown/DetailsView/RemoteControl/Controllers/AvaRundownPageControllerContextMenu.h"
+#include "Rundown/Pages/AvaRundownPageControllerContext.h"
 #include "Widgets/Views/SHeaderRow.h"
 #include "Widgets/Views/SListView.h"
 
@@ -34,6 +37,10 @@ void SAvaRundownRCControllerPanel::Construct(const FArguments& InArgs, const TSh
 	RundownEditorWeak = InRundownEditor;
 	ActivePageId = FAvaRundownPage::InvalidPageId;
 
+	CommandList = MakeShared<FUICommandList>();
+
+	ContextMenu = MakeShared<FAvaRundownPageControllerContextMenu>(CommandList);
+
 	ChildSlot
 	[
 		SNew(SBorder)
@@ -42,7 +49,8 @@ void SAvaRundownRCControllerPanel::Construct(const FArguments& InArgs, const TSh
 		[
 			SAssignNew(ControllerContainer, SListView<FAvaRundownRCControllerItemPtr>)
 			.ListItemsSource(&ControllerItems)
-			.SelectionMode(ESelectionMode::None)
+			.SelectionMode(ESelectionMode::Multi)
+			.OnContextMenuOpening(this, &SAvaRundownRCControllerPanel::GetContextMenuContent)
 			.OnGenerateRow(this, &SAvaRundownRCControllerPanel::OnGenerateControllerRow)
 			.HeaderRow(
 				SNew(SHeaderRow)
@@ -218,7 +226,7 @@ void SAvaRundownRCControllerPanel::Refresh(const TArray<int32>& InSelectedPageId
 
 	ActivePageId = InSelectedPageIds.IsEmpty() ? FAvaRundownPage::InvalidPageId : InSelectedPageIds[0];
 
-	const UAvaRundown* Rundown = GetRundown();
+	UAvaRundown* Rundown = GetRundown();
 	const FAvaRundownPage& Page = GetActivePage(Rundown);
 	
 	if (!Page.IsValidPage())
@@ -430,6 +438,21 @@ FAvaRundownPage& SAvaRundownRCControllerPanel::GetActivePageMutable(UAvaRundown*
 	return FAvaRundownPage::NullPage;
 }
 
+const TArray<FAvaRundownRCControllerItemPtr> SAvaRundownRCControllerPanel::GetSelectedControllerItems() const
+{
+	return ControllerContainer->GetSelectedItems();
+}
+
+TSharedPtr<SWidget> SAvaRundownRCControllerPanel::GetContextMenuContent()
+{
+	const TArray<FAvaRundownRCControllerItemPtr> SelectedItems = GetSelectedControllerItems();
+	if (SelectedItems.Num() > 0)
+	{
+		return ContextMenu->GeneratePageContextMenuWidget(RundownEditorWeak, GetActivePageMutable(), SharedThis(this));
+	}
+	return SNullWidget::NullWidget;
+}
+
 SAvaRundownRCControllerPanel::FPropertyRowGeneratorWrapper::FPropertyRowGeneratorWrapper(SAvaRundownRCControllerPanel* InParentPanel)
 	: ParentPanel(InParentPanel)
 {
diff --git a/Engine/Plugins/Experimental/Avalanche/Source/AvalancheMediaEditor/Private/Rundown/DetailsView/RemoteControl/Controllers/SAvaRundownRCControllerPanel.h b/Engine/Plugins/Experimental/Avalanche/Source/AvalancheMediaEditor/Private/Rundown/DetailsView/RemoteControl/Controllers/SAvaRundownRCControllerPanel.h
--- a/Engine/Plugins/Experimental/Avalanche/Source/AvalancheMediaEditor/Private/Rundown/DetailsView/RemoteControl/Controllers/SAvaRundownRCControllerPanel.h
+++ b/Engine/Plugins/Experimental/Avalanche/Source/AvalancheMediaEditor/Private/Rundown/DetailsView/RemoteControl/Controllers/SAvaRundownRCControllerPanel.h
@@ -8,7 +8,9 @@
 
 class FAvaRundownEditor;
 class FAvaRundownManagedInstance;
+class FAvaRundownPageControllerContextMenu;
 class FAvaRundownRCControllerItem;
+class FUICommandList;
 class IPropertyRowGenerator;
 class ITableRow;
 class SAvaRundownRCControllerPanel;
@@ -51,6 +53,8 @@ public:
 	
 	TSharedRef<ITableRow> OnGenerateControllerRow(FAvaRundownRCControllerItemPtr InItem, const TSharedRef<STableViewBase>& InOwnerTable);
 
+	const TArray<FAvaRundownRCControllerItemPtr> GetSelectedControllerItems() const;
+
 private:
 	static FAvaRundownRCControllerHeaderRowExtensionDelegate HeaderRowExtensionDelegate;
 	static TMap<FName, TArray<FAvaRundownRCControllerTableRowExtensionDelegate>> TableRowExtensionDelegates;
@@ -78,6 +82,8 @@ private:
 	FAvaRundownPage& GetActivePageMutable(UAvaRundown* InRundown) const;
 	FAvaRundownPage& GetActivePageMutable() const { return GetActivePageMutable(GetRundown()); }
 
+	TSharedPtr<SWidget> GetContextMenuContent();
+
 	TWeakPtr<FAvaRundownEditor> RundownEditorWeak;
 
 	class FPropertyRowGeneratorWrapper : public FNotifyHook
@@ -104,4 +110,8 @@ private:
 	TArray<FAvaRundownRCControllerItemPtr> ControllerItems;
 
 	int32 ActivePageId = -1;
+
+	TSharedPtr<FUICommandList> CommandList;
+
+	TSharedPtr<FAvaRundownPageControllerContextMenu> ContextMenu;
 };
diff --git a/Engine/Plugins/Experimental/Avalanche/Source/AvalancheMediaEditor/Private/Rundown/DetailsView/RemoteControl/Properties/AvaRundownPagePropertyContextMenu.cpp b/Engine/Plugins/Experimental/Avalanche/Source/AvalancheMediaEditor/Private/Rundown/DetailsView/RemoteControl/Properties/AvaRundownPagePropertyContextMenu.cpp
new file mode 100644
--- /dev/null
+++ b/Engine/Plugins/Experimental/Avalanche/Source/AvalancheMediaEditor/Private/Rundown/DetailsView/RemoteControl/Properties/AvaRundownPagePropertyContextMenu.cpp
@@ -0,0 +1,256 @@
+// Copyright Epic Games, Inc. All Rights Reserved.
+
+#include "AvaRundownPagePropertyContextMenu.h"
+#include "AvaRundownRCPropertyItem.h"
+#include "IAvaMediaEditorModule.h"
+#include "Rundown/AvaRundownCommands.h"
+#include "Rundown/AvaRundownEditor.h"
+#include "Rundown/AvaRundownPage.h"
+#include "Rundown/Pages/AvaRundownPagePropertyContext.h"
+#include "SAvaRundownPageRemoteControlProps.h"
+#include "ScopedTransaction.h"
+#include "ToolMenu.h"
+#include "ToolMenus.h"
+#include "ToolMenuSection.h"
+#include "UObject/WeakObjectPtrTemplates.h"
+
+#define LOCTEXT_NAMESPACE "AvaRundownPagePropertyContextMenu"
+
+FAvaRundownPagePropertyContextMenu::FAvaRundownPagePropertyContextMenu(const TWeakPtr<FUICommandList>& InCommandListWeak)
+{
+	BindCommands(InCommandListWeak);
+}
+
+void FAvaRundownPagePropertyContextMenu::BindCommands(const TWeakPtr<FUICommandList>& InCommandListWeak)
+{
+	TSharedPtr<FUICommandList> CommandList = InCommandListWeak.Pin();
+	if (!CommandList.IsValid())
+	{
+		return;
+	}
+
+	CommandListWeak = InCommandListWeak;
+
+	const FAvaRundownCommands& RundownCommands = FAvaRundownCommands::Get();
+
+	CommandList->MapAction(RundownCommands.ResetValuesToDefaults,
+		FExecuteAction::CreateRaw(this, &FAvaRundownPagePropertyContextMenu::ResetValuesToDefaults, false),
+		FCanExecuteAction::CreateRaw(this, &FAvaRundownPagePropertyContextMenu::CanResetValuesToDefaults, false));
+
+	CommandList->MapAction(RundownCommands.ResetValuesToTemplate,
+		FExecuteAction::CreateRaw(this, &FAvaRundownPagePropertyContextMenu::ResetValuesToDefaults, true),
+		FCanExecuteAction::CreateRaw(this, &FAvaRundownPagePropertyContextMenu::CanResetValuesToDefaults, true));
+}
+
+TSharedRef<SWidget> FAvaRundownPagePropertyContextMenu::GeneratePageContextMenuWidget(const TWeakPtr<FAvaRundownEditor>& InRundownEditorWeak
+	, const FAvaRundownPage& InRundownPage
+	, const TWeakPtr<SAvaRundownPageRemoteControlProps>& InPropertyListWidgetWeak)
+{
+	UToolMenus* const ToolMenus = UToolMenus::Get();
+	check(ToolMenus);
+
+	const FAvaRundownCommands& RundownCommands = FAvaRundownCommands::Get();
+
+	static const FName ContextMenuName = TEXT("RundownPageRemoteControlProps");
+
+	if (!ToolMenus->IsMenuRegistered(ContextMenuName))
+	{
+		UToolMenu* const ContextMenu = ToolMenus->RegisterMenu(ContextMenuName, NAME_None, EMultiBoxType::Menu);
+		check(ContextMenu);
+
+		ContextMenu->AddDynamicSection(TEXT("PopulateContextMenu"), FNewToolMenuDelegate::CreateLambda(
+			[this](UToolMenu* const InMenu)
+			{
+				if (IsValid(InMenu))
+				{
+					PopulatePageContextMenu(*InMenu, InMenu->FindContext<UAvaRundownPagePropertyContext>());
+				}
+			}));
+	}
+	
+	UAvaRundownPagePropertyContext* const ContextObject = NewObject<UAvaRundownPagePropertyContext>();
+	check(ContextObject);
+	ContextObject->InitContext(InRundownEditorWeak, InRundownPage.GetPageId(), InPropertyListWidgetWeak);
+
+	TSharedPtr<FExtender> Extender;
+
+	// Compatibility with IAvaMediaEditorModule Rundown Menu Extensibility Manager
+	TSharedPtr<FUICommandList> CommandList = CommandListWeak.Pin();
+	if (CommandList.IsValid())
+	{
+		TSharedPtr<FExtensibilityManager> MenuExtensibility = IAvaMediaEditorModule::Get().GetRundownMenuExtensibilityManager();
+		if (MenuExtensibility.IsValid())
+		{
+			if (TSharedPtr<FAvaRundownEditor> RundownEditor = InRundownEditorWeak.Pin())
+			{
+				if (const TArray<UObject*>* EditingObjects = RundownEditor->GetObjectsCurrentlyBeingEdited())
+				{
+					Extender = MenuExtensibility->GetAllExtenders(CommandList.ToSharedRef(), *EditingObjects);
+				}
+			}
+		}
+	}
+
+	FToolMenuContext Context(CommandList, Extender, ContextObject);
+	return ToolMenus->GenerateWidget(ContextMenuName, Context);
+}
+
+void FAvaRundownPagePropertyContextMenu::PopulatePageContextMenu(UToolMenu& InMenu, UAvaRundownPagePropertyContext* const InContext)
+{
+	CurrentContext = InContext;
+
+	if (!IsValid(InContext))
+	{
+		return;
+	}
+
+	const TSharedPtr<SAvaRundownPageRemoteControlProps> ListWidget = CurrentContext->GetPropertyListWidget();
+	if (ListWidget.IsValid() && ListWidget->GetSelectedPropertyItems().Num() == 0)
+	{
+		return;
+	}
+
+	UAvaRundown* const Rundown = GetContextRundown();
+	if (!IsValid(Rundown))
+	{
+		return;
+	}
+
+	const FAvaRundownPage& RundownPage = Rundown->GetPage(CurrentContext->GetRundownPageId());
+	if (!RundownPage.IsValidPage())
+	{
+		return;
+	}
+
+	const FAvaRundownCommands& RundownCommands = FAvaRundownCommands::Get();
+
+	FToolMenuSection& Section = InMenu.FindOrAddSection(TEXT("PropertyActions"), LOCTEXT("PropertyActions", "Property Actions"));
+
+	Section.AddMenuEntry(TEXT("ResetPropertyToDefaults"),
+		RundownCommands.ResetValuesToDefaults,
+		LOCTEXT("ResetPropertyToDefaults", "Reset to Defaults"),
+		LOCTEXT("ResetPropertyToDefaultsToolTip", "Reset selected entity property values to the defaults."),
+		FSlateIcon(FAppStyle::GetAppStyleSetName(), TEXT("PropertyWindow.DiffersFromDefault")));
+
+	if (!RundownPage.IsTemplate())
+	{
+		Section.AddMenuEntry(TEXT("ResetPropertyToTemplate"),
+			RundownCommands.ResetValuesToTemplate,
+			LOCTEXT("ResetPropertyToTemplate", "Reset to Template"),
+			LOCTEXT("ResetPropertyToTemplateToolTip", "Reset selected entity property values to the defaults."),
+			FSlateIcon(FAppStyle::GetAppStyleSetName(), TEXT("PropertyWindow.DiffersFromDefault")));
+	}
+}
+
+TArray<TWeakPtr<FAvaRundownRCPropertyItem>> FAvaRundownPagePropertyContextMenu::GetSelectedPropertyItems() const
+{
+	TArray<TWeakPtr<FAvaRundownRCPropertyItem>> OutItems;
+
+	if (CurrentContext.IsValid())
+	{
+		if (TSharedPtr<SAvaRundownPageRemoteControlProps> PropertyListWidget = CurrentContext->GetPropertyListWidget())
+		{
+			for (const FAvaRundownRCPropertyItemPtr& PropertyItem : PropertyListWidget->GetSelectedPropertyItems())
+			{
+				if (PropertyItem.IsValid())
+				{
+					OutItems.Add(PropertyItem);
+				}
+			}
+		}
+	}
+
+	return OutItems;
+}
+
+bool FAvaRundownPagePropertyContextMenu::HasValidSelectedItems() const
+{
+	return GetSelectedPropertyItems().Num() > 0;
+}
+
+void FAvaRundownPagePropertyContextMenu::ResetValuesToDefaults(const bool bInUseTemplateValues)
+{
+	UAvaRundown* const Rundown = GetContextRundown();
+	if (!IsValid(Rundown))
+	{
+		return;
+	}
+
+	FScopedTransaction Transaction(LOCTEXT("ResetPropertyValuesTransaction", "Reset RC Property Values"));
+	Rundown->Modify();
+
+	const TArray<TWeakPtr<FAvaRundownRCPropertyItem>> SelectedPropertyItems = GetSelectedPropertyItems();
+	for (const TWeakPtr<FAvaRundownRCPropertyItem>& PropertyItemWeak : SelectedPropertyItems)
+	{
+		if (const TSharedPtr<FAvaRundownRCPropertyItem> PropertyItem = PropertyItemWeak.Pin())
+		{
+			const TSharedPtr<FRemoteControlEntity> Entity = PropertyItem->GetEntity();
+			if (Entity.IsValid())
+			{
+				Rundown->ResetRemoteControlEntityValue(CurrentContext->GetRundownPageId()
+					, Entity->GetId(), bInUseTemplateValues, /*bInIsDefault=*/false);
+			}
+		}
+	}
+}
+
+bool FAvaRundownPagePropertyContextMenu::CanResetValuesToDefaults(const bool bInUseTemplateValues) const
+{
+	UAvaRundown* const Rundown = GetContextRundown();
+	if (!IsValid(Rundown))
+	{
+		return false;
+	}
+
+	const FAvaRundownPage& RundownPage = Rundown->GetPage(CurrentContext->GetRundownPageId());
+	if (!RundownPage.IsValidPage())
+	{
+		return false;
+	}
+
+	const TArray<TWeakPtr<FAvaRundownRCPropertyItem>> SelectedItems = GetSelectedPropertyItems();
+	if (SelectedItems.Num() == 0)
+	{
+		return false;
+	}
+
+	bool bContainsDifferentValues = false;
+
+	for (const TWeakPtr<FAvaRundownRCPropertyItem>& PropertyItemWeak : SelectedItems)
+	{
+		const TSharedPtr<FAvaRundownRCPropertyItem> PropertyItem = PropertyItemWeak.Pin();
+		if (PropertyItem.IsValid())
+		{
+			if (PropertyItem->IsEntityControlled())
+			{
+				return false;
+			}
+
+			const TSharedPtr<FRemoteControlEntity> Entity = PropertyItem->GetEntity();
+			if (Entity.IsValid() && !RundownPage.IsDefaultEntityValue(Rundown, Entity->GetId(), bInUseTemplateValues))
+			{
+				bContainsDifferentValues = true;
+			}
+		}
+	}
+
+	return bContainsDifferentValues;
+}
+
+UAvaRundown* FAvaRundownPagePropertyContextMenu::GetContextRundown() const
+{
+	if (!CurrentContext.IsValid())
+	{
+		return nullptr;
+	}
+
+	const TSharedPtr<FAvaRundownEditor> RundownEditor = CurrentContext->GetRundownEditor();
+	if (!RundownEditor.IsValid())
+	{
+		return nullptr;
+	}
+
+	return RundownEditor->GetRundown();
+}
+
+#undef LOCTEXT_NAMESPACE
diff --git a/Engine/Plugins/Experimental/Avalanche/Source/AvalancheMediaEditor/Private/Rundown/DetailsView/RemoteControl/Properties/AvaRundownPagePropertyContextMenu.h b/Engine/Plugins/Experimental/Avalanche/Source/AvalancheMediaEditor/Private/Rundown/DetailsView/RemoteControl/Properties/AvaRundownPagePropertyContextMenu.h
new file mode 100644
--- /dev/null
+++ b/Engine/Plugins/Experimental/Avalanche/Source/AvalancheMediaEditor/Private/Rundown/DetailsView/RemoteControl/Properties/AvaRundownPagePropertyContextMenu.h
@@ -0,0 +1,44 @@
+// Copyright Epic Games, Inc. All Rights Reserved.
+
+#pragma once
+
+#include "Rundown/Pages/AvaRundownPageControllerContext.h"
+#include "Templates/SharedPointer.h"
+#include "UObject/WeakObjectPtrTemplatesFwd.h"
+
+class FAvaRundownEditor;
+class FAvaRundownRCPropertyItem;
+class FUICommandList;
+class SAvaRundownPageRemoteControlProps;
+class SWidget;
+class UAvaRundownPagePropertyContext;
+class UToolMenu;
+struct FAvaRundownPage;
+
+class FAvaRundownPagePropertyContextMenu : public TSharedFromThis<FAvaRundownPagePropertyContextMenu>
+{
+public:
+	FAvaRundownPagePropertyContextMenu(const TWeakPtr<FUICommandList>& InCommandListWeak);
+
+	TSharedRef<SWidget> GeneratePageContextMenuWidget(const TWeakPtr<FAvaRundownEditor>& InRundownEditorWeak
+		, const FAvaRundownPage& InRundownPage
+		, const TWeakPtr<SAvaRundownPageRemoteControlProps>& InPropertyListWidgetWeak);
+
+	bool HasValidSelectedItems() const;
+
+private:
+	void PopulatePageContextMenu(UToolMenu& InMenu, UAvaRundownPagePropertyContext* const InContext);
+
+	void BindCommands(const TWeakPtr<FUICommandList>& InCommandListWeak);
+
+	void ResetValuesToDefaults(const bool bInUseTemplateValues);
+	bool CanResetValuesToDefaults(const bool bInUseTemplateValues) const;
+
+	TArray<TWeakPtr<FAvaRundownRCPropertyItem>> GetSelectedPropertyItems() const;
+
+	UAvaRundown* GetContextRundown() const;
+
+	TWeakPtr<FUICommandList> CommandListWeak;
+
+	TWeakObjectPtr<UAvaRundownPagePropertyContext> CurrentContext;
+};
diff --git a/Engine/Plugins/Experimental/Avalanche/Source/AvalancheMediaEditor/Private/Rundown/DetailsView/RemoteControl/Properties/SAvaRundownPageRemoteControlProps.cpp b/Engine/Plugins/Experimental/Avalanche/Source/AvalancheMediaEditor/Private/Rundown/DetailsView/RemoteControl/Properties/SAvaRundownPageRemoteControlProps.cpp
--- a/Engine/Plugins/Experimental/Avalanche/Source/AvalancheMediaEditor/Private/Rundown/DetailsView/RemoteControl/Properties/SAvaRundownPageRemoteControlProps.cpp
+++ b/Engine/Plugins/Experimental/Avalanche/Source/AvalancheMediaEditor/Private/Rundown/DetailsView/RemoteControl/Properties/SAvaRundownPageRemoteControlProps.cpp
@@ -1,16 +1,19 @@
 // Copyright Epic Games, Inc. All Rights Reserved.
 
 #include "SAvaRundownPageRemoteControlProps.h"
-
 #include "AvaRundownRCPropertyItem.h"
+#include "IAvaMediaEditorModule.h"
 #include "Playable/AvaPlayableRemoteControl.h"
 #include "Playable/AvaPlayableRemoteControlValues.h"
 #include "RemoteControlEntity.h"
 #include "RemoteControlPreset.h"
+#include "Rundown/AvaRundownCommands.h"
 #include "Rundown/AvaRundownEditor.h"
 #include "Rundown/AvaRundownEditorUtils.h"
 #include "Rundown/AvaRundownManagedInstanceCache.h"
 #include "Rundown/AvaRundownPage.h"
+#include "Rundown/DetailsView/RemoteControl/Properties/AvaRundownPagePropertyContextMenu.h"
+#include "Rundown/Pages/AvaRundownPagePropertyContext.h"
 #include "SAvaRundownRCPropertyItemRow.h"
 #include "SlateOptMacros.h"
 #include "UObject/NameTypes.h"
@@ -20,8 +23,8 @@
 
 #define LOCTEXT_NAMESPACE "SAvaRundownPageRemoteControlProps"
 
-const FName SAvaRundownPageRemoteControlProps::PropertyColumnName = "PropertyColumn";
-const FName SAvaRundownPageRemoteControlProps::ValueColumnName = "ValueColumn";
+const FName SAvaRundownPageRemoteControlProps::PropertyColumnName = TEXT("PropertyColumn");
+const FName SAvaRundownPageRemoteControlProps::ValueColumnName = TEXT("ValueColumn");
 
 FAvaRundownRCPropertyHeaderRowExtensionDelegate SAvaRundownPageRemoteControlProps::HeaderRowExtensionDelegate;
 TMap<FName, TArray<FAvaRundownRCPropertyTableRowExtensionDelegate>> SAvaRundownPageRemoteControlProps::TableRowExtensionDelegates;
@@ -38,6 +41,10 @@ void SAvaRundownPageRemoteControlProps::Construct(const FArguments& InArgs, TSha
 	RundownEditorWeak = InRundownEditor;
 	ActivePageId = FAvaRundownPage::InvalidPageId;
 
+	CommandList = MakeShared<FUICommandList>();
+
+	ContextMenu = MakeShared<FAvaRundownPagePropertyContextMenu>(CommandList);
+
 	TSharedRef<SHeaderRow> HeaderRow =
 		SNew(SHeaderRow)
 		+ SHeaderRow::Column(PropertyColumnName)
@@ -57,7 +64,8 @@ void SAvaRundownPageRemoteControlProps::Construct(const FArguments& InArgs, TSha
 		[
 			SAssignNew(PropertyContainer, SListView<FAvaRundownRCPropertyItemPtr>)
 			.ListItemsSource(&PropertyItems)
-			.SelectionMode(ESelectionMode::None)
+			.SelectionMode(ESelectionMode::Multi)
+			.OnContextMenuOpening(this, &SAvaRundownPageRemoteControlProps::GetContextMenuContent)
 			.OnGenerateRow(this, &SAvaRundownPageRemoteControlProps::OnGenerateControllerRow)
 			.HeaderRow(HeaderRow)
 		]
@@ -442,4 +450,19 @@ bool SAvaRundownPageRemoteControlProps::SetSelectedPageEntityValue(const TShared
 	return false;
 }
 
+const TArray<FAvaRundownRCPropertyItemPtr> SAvaRundownPageRemoteControlProps::GetSelectedPropertyItems() const
+{
+	return PropertyContainer->GetSelectedItems();
+}
+
+TSharedPtr<SWidget> SAvaRundownPageRemoteControlProps::GetContextMenuContent()
+{
+	TArray<FAvaRundownRCPropertyItemPtr> SelectedItems = GetSelectedPropertyItems();
+	if (SelectedItems.Num() > 0)
+	{
+		return ContextMenu->GeneratePageContextMenuWidget(RundownEditorWeak, *GetActivePage(), SharedThis(this));
+	}
+	return SNullWidget::NullWidget;
+}
+
 #undef LOCTEXT_NAMESPACE
diff --git a/Engine/Plugins/Experimental/Avalanche/Source/AvalancheMediaEditor/Private/Rundown/DetailsView/RemoteControl/Properties/SAvaRundownPageRemoteControlProps.h b/Engine/Plugins/Experimental/Avalanche/Source/AvalancheMediaEditor/Private/Rundown/DetailsView/RemoteControl/Properties/SAvaRundownPageRemoteControlProps.h
--- a/Engine/Plugins/Experimental/Avalanche/Source/AvalancheMediaEditor/Private/Rundown/DetailsView/RemoteControl/Properties/SAvaRundownPageRemoteControlProps.h
+++ b/Engine/Plugins/Experimental/Avalanche/Source/AvalancheMediaEditor/Private/Rundown/DetailsView/RemoteControl/Properties/SAvaRundownPageRemoteControlProps.h
@@ -12,8 +12,10 @@
 
 class FAvaRundownEditor;
 class FAvaRundownManagedInstance;
+class FAvaRundownPagePropertyContextMenu;
 class FAvaRundownRCPropertyItem;
 class FName;
+class FUICommandList;
 class ITableRow;
 class SAvaRundownPageRemoteControlProps;
 class STableViewBase;
@@ -55,6 +57,8 @@ public:
 	/** Refreshes the content of this widget. */
 	void Refresh(const TArray<int32>& InSelectedPageIds);
 
+	const TArray<FAvaRundownRCPropertyItemPtr> GetSelectedPropertyItems() const;
+
 private:
 	static FAvaRundownRCPropertyHeaderRowExtensionDelegate HeaderRowExtensionDelegate;
 	static TMap<FName, TArray<FAvaRundownRCPropertyTableRowExtensionDelegate>> TableRowExtensionDelegates;
@@ -91,6 +95,8 @@ private:
 
 	void RefreshTable(const TSet<FGuid>& InEntityIds = TSet<FGuid>());
 
+	TSharedPtr<SWidget> GetContextMenuContent();
+
 	TWeakPtr<FAvaRundownEditor> RundownEditorWeak;
 	
 	TArray<TSharedPtr<FAvaRundownManagedInstance>> ManagedInstances;
@@ -102,4 +108,8 @@ private:
 	TArray<FAvaRundownRCPropertyItemPtr> PropertyItems;
 
 	int32 ActivePageId = -1;
+
+	TSharedPtr<FUICommandList> CommandList;
+
+	TSharedPtr<FAvaRundownPagePropertyContextMenu> ContextMenu;
 };
diff --git a/Engine/Plugins/Experimental/Avalanche/Source/AvalancheMediaEditor/Private/Rundown/DetailsView/RemoteControl/Properties/SAvaRundownRCPropertyItemRow.cpp b/Engine/Plugins/Experimental/Avalanche/Source/AvalancheMediaEditor/Private/Rundown/DetailsView/RemoteControl/Properties/SAvaRundownRCPropertyItemRow.cpp
--- a/Engine/Plugins/Experimental/Avalanche/Source/AvalancheMediaEditor/Private/Rundown/DetailsView/RemoteControl/Properties/SAvaRundownRCPropertyItemRow.cpp
+++ b/Engine/Plugins/Experimental/Avalanche/Source/AvalancheMediaEditor/Private/Rundown/DetailsView/RemoteControl/Properties/SAvaRundownRCPropertyItemRow.cpp
@@ -3,6 +3,7 @@
 #include "SAvaRundownRCPropertyItemRow.h"
 #include "AvaRundownPageRemoteControlWidgetUtils.h"
 #include "AvaRundownRCPropertyItem.h"
+#include "GameFramework/Actor.h"
 #include "IDetailTreeNode.h"
 #include "Internationalization/Text.h"
 #include "IPropertyRowGenerator.h"
@@ -10,6 +11,7 @@
 #include "PropertyEditorModule.h"
 #include "RemoteControlEntity.h"
 #include "RemoteControlField.h"
+#include "RemoteControlPreset.h"
 #include "Rundown/DetailsView/RemoteControl/Properties/AvaRundownPageRemoteControlWidgetUtils.h"
 #include "Widgets/Layout/SScissorRectBox.h"
 #include "Widgets/SNullWidget.h"
@@ -37,16 +39,10 @@ TSharedRef<SWidget> SAvaRundownRCPropertyItemRow::GenerateWidgetForColumn(const
 	{
 		if (InColumnName == SAvaRundownPageRemoteControlProps::PropertyColumnName)
 		{
-			return SNew(SScissorRectBox)
-				[
-					SNew(SBox)
-					.VAlign(EVerticalAlignment::VAlign_Center)
-					.Padding(3.f, 2.f, 3.f, 2.f)
-					[
-						SNew(STextBlock)
-						.Text(GetFieldLabel())
-					]
-				];
+			return SNew(STextBlock)
+				.Margin(FMargin(8.f, 2.f, 0.f, 2.f))
+				.Text(GetFieldLabel())
+				.ToolTipText(this, &SAvaRundownRCPropertyItemRow::GetPropertyTooltipText);
 		}
 		else if (InColumnName == SAvaRundownPageRemoteControlProps::ValueColumnName)
 		{
@@ -152,4 +148,56 @@ TSharedRef<SWidget> SAvaRundownRCPropertyItemRow::CreateValue()
 	return SNullWidget::NullWidget;
 }
 
+FText SAvaRundownRCPropertyItemRow::GetPropertyTooltipText() const
+{
+	TSharedPtr<const FAvaRundownRCPropertyItem> ItemPtr = ItemPtrWeak.Pin();
+
+	FText OwnerText = LOCTEXT("InvalidOwnerText", "(Invalid)");
+	FText SubobjectPathText = LOCTEXT("InvalidSubobjectPathText", "(Invalid)");
+
+	const TSharedPtr<FRemoteControlEntity> Entity = ItemPtr->GetEntity();
+	if (Entity.IsValid())
+	{
+		const FString BindingPath = Entity->GetLastBindingPath().ToString();
+
+		FName OwnerName;
+		if (UObject* Object = Entity->GetBoundObject())
+		{
+			if (AActor* OwnerActor = Object->GetTypedOuter<AActor>())
+			{
+				OwnerText = FText::FromString(OwnerActor->GetActorLabel());
+				OwnerName = OwnerActor->GetFName();
+			}
+			else if (AActor* Actor = Cast<AActor>(Object))
+			{
+				OwnerText = FText::FromString(Actor->GetActorLabel());
+				OwnerName = Object->GetFName();
+			}
+			else
+			{
+				OwnerText = FText::FromString(Object->GetName());
+				OwnerName = Object->GetFName();
+			}
+		}
+		else
+		{
+			static const FString PersistentLevelString = TEXT(":PersistentLevel.");
+			const int32 PersistentLevelIndex = BindingPath.Find(PersistentLevelString);
+			if (PersistentLevelIndex != INDEX_NONE)
+			{
+				OwnerText = FText::FromName(OwnerName);
+				OwnerName = *BindingPath.RightChop(PersistentLevelIndex + PersistentLevelString.Len());
+			}
+		}
+
+		const int32 OwnerNameIndex = BindingPath.Find(OwnerName.ToString() + TEXT("."));
+		if (OwnerNameIndex != INDEX_NONE)
+		{
+			SubobjectPathText = FText::FromString(*BindingPath.RightChop(OwnerNameIndex + OwnerName.GetStringLength() + 1));
+		}
+	}
+
+	return FText::Format(LOCTEXT("PropertyTooltipText", "Owner: {0}\nSubobjectPath: {1}"), OwnerText, SubobjectPathText);
+}
+
 #undef LOCTEXT_NAMESPACE
diff --git a/Engine/Plugins/Experimental/Avalanche/Source/AvalancheMediaEditor/Private/Rundown/DetailsView/RemoteControl/Properties/SAvaRundownRCPropertyItemRow.h b/Engine/Plugins/Experimental/Avalanche/Source/AvalancheMediaEditor/Private/Rundown/DetailsView/RemoteControl/Properties/SAvaRundownRCPropertyItemRow.h
--- a/Engine/Plugins/Experimental/Avalanche/Source/AvalancheMediaEditor/Private/Rundown/DetailsView/RemoteControl/Properties/SAvaRundownRCPropertyItemRow.h
+++ b/Engine/Plugins/Experimental/Avalanche/Source/AvalancheMediaEditor/Private/Rundown/DetailsView/RemoteControl/Properties/SAvaRundownRCPropertyItemRow.h
@@ -27,14 +27,16 @@ public:
 	void UpdateValue();
 
 protected:
+	/** Get this field's label. */
+	FText GetFieldLabel() const;
+
+	TSharedRef<SWidget> CreateValue();
+
+	FText GetPropertyTooltipText() const;
+
 	TWeakPtr<const FAvaRundownRCPropertyItem> ItemPtrWeak;
 	TWeakPtr<SAvaRundownPageRemoteControlProps> PropertyPanelWeak;
 	TSharedPtr<IPropertyRowGenerator> Generator;
 	TSharedPtr<SBox> ValueContainer;
 	TSharedPtr<SWidget> ValueWidget;
-
-	/** Get this field's label. */
-	FText GetFieldLabel() const;
-
-	TSharedRef<SWidget> CreateValue();
 };
diff --git a/Engine/Plugins/Experimental/Avalanche/Source/AvalancheMediaEditor/Public/Rundown/Pages/AvaRundownPageControllerContext.h b/Engine/Plugins/Experimental/Avalanche/Source/AvalancheMediaEditor/Public/Rundown/Pages/AvaRundownPageControllerContext.h
new file mode 100644
--- /dev/null
+++ b/Engine/Plugins/Experimental/Avalanche/Source/AvalancheMediaEditor/Public/Rundown/Pages/AvaRundownPageControllerContext.h
@@ -0,0 +1,50 @@
+// Copyright Epic Games, Inc. All Rights Reserved.
+
+#pragma once
+
+#include "Rundown/AvaRundownPage.h"
+#include "Templates/SharedPointer.h"
+#include "UObject/Object.h"
+#include "AvaRundownPageControllerContext.generated.h"
+
+class FAvaRundownEditor;
+class FAvaRundownPageControllerContextMenu;
+class SAvaRundownRCControllerPanel;
+
+UCLASS(MinimalAPI)
+class UAvaRundownPageControllerContext : public UObject
+{
+	GENERATED_BODY()
+
+public:
+	void InitContext(const TWeakPtr<FAvaRundownEditor>& InRundownEditorWeak
+		, const int32 InRundownPageId
+		, const TWeakPtr<SAvaRundownRCControllerPanel>& InControllerListWidgetWeak)
+	{
+		RundownEditorWeak = InRundownEditorWeak;
+		RundownPageId = InRundownPageId;
+		ControllerListWidgetWeak = InControllerListWidgetWeak;
+	}
+
+	TSharedPtr<FAvaRundownEditor> GetRundownEditor() const
+	{
+		return RundownEditorWeak.Pin();
+	}
+
+	int32 GetRundownPageId() const
+	{
+		return RundownPageId;
+	}
+
+	TSharedPtr<SAvaRundownRCControllerPanel> GetControllerListWidget() const
+	{
+		return ControllerListWidgetWeak.Pin();
+	}
+
+private:
+	TWeakPtr<FAvaRundownEditor> RundownEditorWeak;
+
+	int32 RundownPageId = FAvaRundownPage::InvalidPageId;
+
+	TWeakPtr<SAvaRundownRCControllerPanel> ControllerListWidgetWeak;
+};
diff --git a/Engine/Plugins/Experimental/Avalanche/Source/AvalancheMediaEditor/Public/Rundown/Pages/AvaRundownPagePropertyContext.h b/Engine/Plugins/Experimental/Avalanche/Source/AvalancheMediaEditor/Public/Rundown/Pages/AvaRundownPagePropertyContext.h
new file mode 100644
--- /dev/null
+++ b/Engine/Plugins/Experimental/Avalanche/Source/AvalancheMediaEditor/Public/Rundown/Pages/AvaRundownPagePropertyContext.h
@@ -0,0 +1,50 @@
+// Copyright Epic Games, Inc. All Rights Reserved.
+
+#pragma once
+
+#include "Rundown/AvaRundownPage.h"
+#include "Templates/SharedPointer.h"
+#include "UObject/Object.h"
+#include "AvaRundownPagePropertyContext.generated.h"
+
+class FAvaRundownEditor;
+class FAvaRundownPagePropertyContextMenu;
+class SAvaRundownPageRemoteControlProps;
+
+UCLASS(MinimalAPI)
+class UAvaRundownPagePropertyContext : public UObject
+{
+	GENERATED_BODY()
+
+public:
+	void InitContext(const TWeakPtr<FAvaRundownEditor>& InRundownEditorWeak
+		, const int32 InRundownPageId
+		, const TWeakPtr<SAvaRundownPageRemoteControlProps>& InPropertyListWidgetWeak)
+	{
+		RundownEditorWeak = InRundownEditorWeak;
+		RundownPageId = InRundownPageId;
+		PropertyListWidgetWeak = InPropertyListWidgetWeak;
+	}
+
+	TSharedPtr<FAvaRundownEditor> GetRundownEditor() const
+	{
+		return RundownEditorWeak.Pin();
+	}
+
+	int32 GetRundownPageId() const
+	{
+		return RundownPageId;
+	}
+
+	TSharedPtr<SAvaRundownPageRemoteControlProps> GetPropertyListWidget() const
+	{
+		return PropertyListWidgetWeak.Pin();
+	}
+
+private:
+	TWeakPtr<FAvaRundownEditor> RundownEditorWeak;
+
+	int32 RundownPageId = FAvaRundownPage::InvalidPageId;
+
+	TWeakPtr<SAvaRundownPageRemoteControlProps> PropertyListWidgetWeak;
+};
