[Motion Design]
 - Auto save/load Rundown Editor RC property values list visibility

diff --git a/Engine/Plugins/Experimental/Avalanche/Source/AvalancheMediaEditor/Private/AvaMediaEditorSettings.h b/Engine/Plugins/Experimental/Avalanche/Source/AvalancheMediaEditor/Private/AvaMediaEditorSettings.h
--- a/Engine/Plugins/Experimental/Avalanche/Source/AvalancheMediaEditor/Private/AvaMediaEditorSettings.h
+++ b/Engine/Plugins/Experimental/Avalanche/Source/AvalancheMediaEditor/Private/AvaMediaEditorSettings.h
@@ -30,6 +30,9 @@ public:
 	 */
 	UPROPERTY(Config, EditAnywhere, Category = "Broadcast", meta = (DisplayName = "Show All Media Output Classes"))
 	bool bBroadcastShowAllMediaOutputClasses = false;
+
+	UPROPERTY(Config, EditAnywhere, Category = "Playback", meta = (DisplayName = "Show RC Property List"))
+	bool bPlaybackShowPropertyList = false;
 	
 	UPROPERTY(Config, EditAnywhere, Category = "Playback", meta = (DisplayName = "Default Node Color"))
 	FLinearColor PlaybackDefaultNodeColor;
diff --git a/Engine/Plugins/Experimental/Avalanche/Source/AvalancheMediaEditor/Private/Rundown/DetailsView/SAvaRundownPageDetails.cpp b/Engine/Plugins/Experimental/Avalanche/Source/AvalancheMediaEditor/Private/Rundown/DetailsView/SAvaRundownPageDetails.cpp
--- a/Engine/Plugins/Experimental/Avalanche/Source/AvalancheMediaEditor/Private/Rundown/DetailsView/SAvaRundownPageDetails.cpp
+++ b/Engine/Plugins/Experimental/Avalanche/Source/AvalancheMediaEditor/Private/Rundown/DetailsView/SAvaRundownPageDetails.cpp
@@ -2,6 +2,7 @@
 
 #include "SAvaRundownPageDetails.h"
 #include "Async/Async.h"
+#include "AvaMediaEditorSettings.h"
 #include "Framework/Application/SlateApplication.h"
 #include "IAvaMediaModule.h"
 #include "Input/Reply.h"
@@ -174,7 +175,7 @@ void SAvaRundownPageDetails::Construct(const FArguments& InArgs, const TSharedPt
 			.AutoHeight()
 			[
 				SAssignNew(RemoteControlProps, SAvaRundownPageRemoteControlProps, InRundownEditor)
-				.Visibility(EVisibility::Collapsed)
+				.Visibility(UAvaMediaEditorSettings::Get().bPlaybackShowPropertyList ? EVisibility::Visible : EVisibility::Collapsed)
 			]
 		]
 	];
@@ -257,7 +258,11 @@ void SAvaRundownPageDetails::OnManagedInstanceCacheEntryInvalidated(const FSoftO
 
 FReply SAvaRundownPageDetails::ToggleExposedPropertiesVisibility()
 {
-	if (RemoteControlProps->GetVisibility() == EVisibility::Collapsed)
+	UAvaMediaEditorSettings& MediaEditorSettings = UAvaMediaEditorSettings::GetMutable();
+	MediaEditorSettings.bPlaybackShowPropertyList = !MediaEditorSettings.bPlaybackShowPropertyList;
+	MediaEditorSettings.SaveConfig();
+
+	if (MediaEditorSettings.bPlaybackShowPropertyList)
 	{
 		RemoteControlProps->SetVisibility(EVisibility::SelfHitTestInvisible);
 	}
@@ -271,13 +276,13 @@ FReply SAvaRundownPageDetails::ToggleExposedPropertiesVisibility()
 
 const FSlateBrush* SAvaRundownPageDetails::GetExposedPropertiesVisibilityBrush() const
 {
-	if (RemoteControlProps->GetVisibility() == EVisibility::Collapsed)
+	if (UAvaMediaEditorSettings::Get().bPlaybackShowPropertyList)
 	{
-		return FAppStyle::GetBrush("Level.NotVisibleHighlightIcon16x");
+		return FAppStyle::GetBrush(TEXT("Level.VisibleHighlightIcon16x"));
 	}
 	else
 	{
-		return FAppStyle::GetBrush("Level.VisibleHighlightIcon16x");
+		return FAppStyle::GetBrush(TEXT("Level.NotVisibleHighlightIcon16x"));
 	}
 }
 
