Motion Design: Fix issue with Preview Frame playing from PreviewFrame to end, rather than staying at that Preview Frame.

diff --git a/Engine/Plugins/Experimental/Avalanche/Source/AvalancheMedia/Private/Playable/AvaPlayable.cpp b/Engine/Plugins/Experimental/Avalanche/Source/AvalancheMedia/Private/Playable/AvaPlayable.cpp
--- a/Engine/Plugins/Experimental/Avalanche/Source/AvalancheMedia/Private/Playable/AvaPlayable.cpp
+++ b/Engine/Plugins/Experimental/Avalanche/Source/AvalancheMedia/Private/Playable/AvaPlayable.cpp
@@ -180,21 +180,23 @@ EAvaPlayableCommandResult UAvaPlayable::ExecuteAnimationCommand(EAvaPlaybackAnim
 		// Remark: if the command doesn't specify the sequence name, we run the command on all the sequences.
 		if (Sequence && (Sequence->GetFName() == InAnimPlaySettings.AnimationName || InAnimPlaySettings.AnimationName.IsNone()))
 		{
-			if (InAnimAction == EAvaPlaybackAnimAction::Play || InAnimAction == EAvaPlaybackAnimAction::PreviewFrame)
-			{
-				UAvaSequencePlayer* const SequencePlayer = PlaybackObject->PlaySequence(Sequence, InAnimPlaySettings.AsPlayParams());
-				if (InAnimAction == EAvaPlaybackAnimAction::PreviewFrame && SequencePlayer)
-				{
-					SequencePlayer->PreviewFrame();
-				}
-			}
-			else if (InAnimAction == EAvaPlaybackAnimAction::Continue)
+			switch (InAnimAction)
 			{
+			case EAvaPlaybackAnimAction::Play:
+				PlaybackObject->PlaySequence(Sequence, InAnimPlaySettings.AsPlayParams());
+				break;
+
+			case EAvaPlaybackAnimAction::Continue:
 				PlaybackObject->ContinueSequence(Sequence);
-			}
-			else if (InAnimAction == EAvaPlaybackAnimAction::Stop)
-			{
+				break;
+
+			case EAvaPlaybackAnimAction::Stop:
 				PlaybackObject->StopSequence(Sequence);
+				break;
+
+			case EAvaPlaybackAnimAction::PreviewFrame:
+				PlaybackObject->PreviewFrame(Sequence);
+				break;
 			}
 		}
 	}
diff --git a/Engine/Plugins/Experimental/Avalanche/Source/AvalancheSequence/Private/AvaSequencePlaybackActor.cpp b/Engine/Plugins/Experimental/Avalanche/Source/AvalancheSequence/Private/AvaSequencePlaybackActor.cpp
--- a/Engine/Plugins/Experimental/Avalanche/Source/AvalancheSequence/Private/AvaSequencePlaybackActor.cpp
+++ b/Engine/Plugins/Experimental/Avalanche/Source/AvalancheSequence/Private/AvaSequencePlaybackActor.cpp
@@ -77,6 +77,38 @@ UAvaSequencePlayer* AAvaSequencePlaybackActor::PlaySequence(UAvaSequence* InSequ
 	return nullptr;
 }
 
+UAvaSequencePlayer* AAvaSequencePlaybackActor::PreviewFrame(UAvaSequence* InSequence)
+{
+	if (!IsValid(InSequence))
+	{
+		return nullptr;
+	}
+
+	const FAvaMark* PreviewMark = InSequence->GetPreviewMark();
+	if (!PreviewMark)
+	{
+		UE_LOG(LogAvaSequencePlayback, Warning
+			, TEXT("Failed to preview Sequence '%s' ('%s'). Missing Preview Mark.")
+			, *InSequence->GetLabel().ToString()
+			, *InSequence->GetName());
+		return nullptr;
+	}
+
+	FFrameTime PreviewPosition;
+	{
+		FMovieSceneSequencePlaybackParams PlaybackParams;
+		PlaybackParams.MarkedFrame = PreviewMark->GetLabel();
+		PlaybackParams.PositionType = EMovieScenePositionType::MarkedFrame;
+
+		PreviewPosition = PlaybackParams.GetPlaybackPosition(InSequence);
+	}
+
+	FAvaSequencePlayParams PlaySettings;
+	PlaySettings.Start = PlaySettings.End = FAvaSequenceTime(PreviewPosition);
+
+	return PlaySequence(InSequence, PlaySettings);
+}
+
 UAvaSequencePlayer* AAvaSequencePlaybackActor::PlaySequenceBySoftReference(TSoftObjectPtr<UAvaSequence> InSequence, FAvaSequencePlayParams InPlaySettings)
 {
 	if (UAvaSequence* const ResolvedSequence = InSequence.Get())
diff --git a/Engine/Plugins/Experimental/Avalanche/Source/AvalancheSequence/Private/AvaSequencePlaybackActor.h b/Engine/Plugins/Experimental/Avalanche/Source/AvalancheSequence/Private/AvaSequencePlaybackActor.h
--- a/Engine/Plugins/Experimental/Avalanche/Source/AvalancheSequence/Private/AvaSequencePlaybackActor.h
+++ b/Engine/Plugins/Experimental/Avalanche/Source/AvalancheSequence/Private/AvaSequencePlaybackActor.h
@@ -29,6 +29,7 @@ protected:
 	virtual ULevel* GetPlaybackLevel() const { return GetLevel(); }
 	virtual void CleanupPlayers() override;
 	virtual UAvaSequencePlayer* PlaySequence(UAvaSequence* InSequence, const FAvaSequencePlayParams& InPlaySettings = FAvaSequencePlayParams()) override;
+	virtual UAvaSequencePlayer* PreviewFrame(UAvaSequence* InSequence) override;
 	virtual UAvaSequencePlayer* PlaySequenceBySoftReference(TSoftObjectPtr<UAvaSequence> InSequence, FAvaSequencePlayParams InPlaySettings) override;
 	virtual TArray<UAvaSequencePlayer*> PlaySequencesByLabel(FName InSequenceLabel, FAvaSequencePlayParams InPlaySettings) override;
 	virtual TArray<UAvaSequencePlayer*> PlaySequencesBySoftReference(const TArray<TSoftObjectPtr<UAvaSequence>>& InSequences, FAvaSequencePlayParams InPlaySettings) override;
diff --git a/Engine/Plugins/Experimental/Avalanche/Source/AvalancheSequence/Public/AvaSequencePlaybackObject.h b/Engine/Plugins/Experimental/Avalanche/Source/AvalancheSequence/Public/AvaSequencePlaybackObject.h
--- a/Engine/Plugins/Experimental/Avalanche/Source/AvalancheSequence/Public/AvaSequencePlaybackObject.h
+++ b/Engine/Plugins/Experimental/Avalanche/Source/AvalancheSequence/Public/AvaSequencePlaybackObject.h
@@ -42,6 +42,14 @@ public:
 
 	virtual UAvaSequencePlayer* PlaySequence(UAvaSequence* InSequence, const FAvaSequencePlayParams& InPlaySettings = FAvaSequencePlayParams()) = 0;
 
+	/**
+	 * Evaluates the Preview Frame of a Sequence.
+	 * Does nothing if the Sequence has no preview frame.
+	 * @param InSequence the sequence to preview
+	 * @return the player instantiated for the Sequence, or null if Sequence was not valid or did not have a preview mark
+	 */
+	virtual UAvaSequencePlayer* PreviewFrame(UAvaSequence* InSequence) = 0;
+
 	/**
 	 * Plays a single sequence by its soft reference
 	 * @param InSequence soft reference of the sequence to play
diff --git a/Engine/Plugins/Experimental/Avalanche/Source/AvalancheSequence/Public/AvaSequencePlayer.h b/Engine/Plugins/Experimental/Avalanche/Source/AvalancheSequence/Public/AvaSequencePlayer.h
--- a/Engine/Plugins/Experimental/Avalanche/Source/AvalancheSequence/Public/AvaSequencePlayer.h
+++ b/Engine/Plugins/Experimental/Avalanche/Source/AvalancheSequence/Public/AvaSequencePlayer.h
@@ -42,6 +42,7 @@ public:
 
 	void ContinueSequence();
 
+	UE_DEPRECATED(5.5, "Use IAvaSequencePlaybackObject::PreviewFrame instead")
 	void PreviewFrame();
 
 	/** Jump to the given Frame, in Tick Resolution space */
