[Avalanche/DynamicMaterial/GeometryMask/StormSync] Removing Optimization macro
Removing BEGIN_SLATE_FUNCTION_BUILD_OPTIMIZATION macro in mentioned plugins

diff --git a/Engine/Plugins/Experimental/Avalanche/Source/AvalancheEditor/Private/Alignment/SAvaAnchorAlignment.cpp b/Engine/Plugins/Experimental/Avalanche/Source/AvalancheEditor/Private/Alignment/SAvaAnchorAlignment.cpp
--- a/Engine/Plugins/Experimental/Avalanche/Source/AvalancheEditor/Private/Alignment/SAvaAnchorAlignment.cpp
+++ b/Engine/Plugins/Experimental/Avalanche/Source/AvalancheEditor/Private/Alignment/SAvaAnchorAlignment.cpp
@@ -9,7 +9,6 @@
 
 #define LOCTEXT_NAMESPACE "SAvaAnchorAlignment"
 
-BEGIN_SLATE_FUNCTION_BUILD_OPTIMIZATION
 void SAvaAnchorAlignment::Construct(const FArguments& InArgs)
 {
 	Anchors = InArgs._Anchors;
@@ -53,7 +52,6 @@ void SAvaAnchorAlignment::Construct(const FArguments& InArgs)
 		]
 	];
 }
-END_SLATE_FUNCTION_BUILD_OPTIMIZATION
 
 EVisibility SAvaAnchorAlignment::GetHorizontalVisibility() const
 {
diff --git a/Engine/Plugins/Experimental/Avalanche/Source/AvalancheEditor/Private/Alignment/SAvaDepthAlignment.cpp b/Engine/Plugins/Experimental/Avalanche/Source/AvalancheEditor/Private/Alignment/SAvaDepthAlignment.cpp
--- a/Engine/Plugins/Experimental/Avalanche/Source/AvalancheEditor/Private/Alignment/SAvaDepthAlignment.cpp
+++ b/Engine/Plugins/Experimental/Avalanche/Source/AvalancheEditor/Private/Alignment/SAvaDepthAlignment.cpp
@@ -8,8 +8,6 @@
 
 #define LOCTEXT_NAMESPACE "SAvaDepthAlignment"
 
-BEGIN_SLATE_FUNCTION_BUILD_OPTIMIZATION
-
 void SAvaDepthAlignment::Construct(const FArguments& InArgs)
 {
 	Alignment = InArgs._Alignment;
@@ -35,8 +33,6 @@ void SAvaDepthAlignment::Construct(const FArguments& InArgs)
 	];
 }
 
-END_SLATE_FUNCTION_BUILD_OPTIMIZATION
-
 EAvaDepthAlignment SAvaDepthAlignment::GetCurrentAlignment() const
 {
 	return Alignment.Get(EAvaDepthAlignment::Center);
diff --git a/Engine/Plugins/Experimental/Avalanche/Source/AvalancheEditor/Private/Alignment/SAvaHorizontalAlignment.cpp b/Engine/Plugins/Experimental/Avalanche/Source/AvalancheEditor/Private/Alignment/SAvaHorizontalAlignment.cpp
--- a/Engine/Plugins/Experimental/Avalanche/Source/AvalancheEditor/Private/Alignment/SAvaHorizontalAlignment.cpp
+++ b/Engine/Plugins/Experimental/Avalanche/Source/AvalancheEditor/Private/Alignment/SAvaHorizontalAlignment.cpp
@@ -8,8 +8,6 @@
 
 #define LOCTEXT_NAMESPACE "SAvaHorizontalAlignment"
 
-BEGIN_SLATE_FUNCTION_BUILD_OPTIMIZATION
-
 void SAvaHorizontalAlignment::Construct(const FArguments& InArgs)
 {
 	Alignment = InArgs._Alignment;
@@ -35,8 +33,6 @@ void SAvaHorizontalAlignment::Construct(const FArguments& InArgs)
 	];
 }
 
-END_SLATE_FUNCTION_BUILD_OPTIMIZATION
-
 EAvaHorizontalAlignment SAvaHorizontalAlignment::GetCurrentAlignment() const
 {
 	return Alignment.Get(EAvaHorizontalAlignment::Center);
diff --git a/Engine/Plugins/Experimental/Avalanche/Source/AvalancheEditor/Private/Alignment/SAvaVerticalAlignment.cpp b/Engine/Plugins/Experimental/Avalanche/Source/AvalancheEditor/Private/Alignment/SAvaVerticalAlignment.cpp
--- a/Engine/Plugins/Experimental/Avalanche/Source/AvalancheEditor/Private/Alignment/SAvaVerticalAlignment.cpp
+++ b/Engine/Plugins/Experimental/Avalanche/Source/AvalancheEditor/Private/Alignment/SAvaVerticalAlignment.cpp
@@ -8,8 +8,6 @@
 
 #define LOCTEXT_NAMESPACE "SAvaVerticalAlignment"
 
-BEGIN_SLATE_FUNCTION_BUILD_OPTIMIZATION
-
 void SAvaVerticalAlignment::Construct(const FArguments& InArgs)
 {
 	Alignment = InArgs._Alignment;
@@ -35,8 +33,6 @@ void SAvaVerticalAlignment::Construct(const FArguments& InArgs)
 	];
 }
 
-END_SLATE_FUNCTION_BUILD_OPTIMIZATION
-
 EAvaVerticalAlignment SAvaVerticalAlignment::GetCurrentAlignment() const
 {
 	return Alignment.Get(EAvaVerticalAlignment::Center);
diff --git a/Engine/Plugins/Experimental/Avalanche/Source/AvalancheLevelViewport/Private/Customizations/AvaLevelViewportComponentTransformDetails.cpp b/Engine/Plugins/Experimental/Avalanche/Source/AvalancheLevelViewport/Private/Customizations/AvaLevelViewportComponentTransformDetails.cpp
--- a/Engine/Plugins/Experimental/Avalanche/Source/AvalancheLevelViewport/Private/Customizations/AvaLevelViewportComponentTransformDetails.cpp
+++ b/Engine/Plugins/Experimental/Avalanche/Source/AvalancheLevelViewport/Private/Customizations/AvaLevelViewportComponentTransformDetails.cpp
@@ -374,8 +374,6 @@ FUIAction FAvaLevelViewportComponentTransformDetails::CreatePasteAction(EAvaLeve
 	return FUIAction(FExecuteAction::CreateSP(MutableThis, &FAvaLevelViewportComponentTransformDetails::OnPaste, TransformField));
 }
 
-BEGIN_SLATE_FUNCTION_BUILD_OPTIMIZATION
-
 TSharedPtr<SWidget> FAvaLevelViewportComponentTransformDetails::GetTransformBody()
 {
 	FSlateFontInfo FontInfo = IDetailLayoutBuilder::GetDetailFont();
@@ -499,8 +497,6 @@ void FAvaLevelViewportComponentTransformDetails::SetTransform(EAvaLevelViewportT
 	OnEndChange(TransformField, EAxisList::XYZ);
 }
 
-END_SLATE_FUNCTION_BUILD_OPTIMIZATION
-
 void FAvaLevelViewportComponentTransformDetails::OnSelectionChanged(const TArray<UObject*>& InSelectedObjects)
 {
 	SelectedObjects.Empty();
diff --git a/Engine/Plugins/Experimental/Avalanche/Source/AvalancheMediaEditor/Private/Rundown/DetailsView/RemoteControl/Properties/AvaRundownPageRemoteControlWidgetUtils.cpp b/Engine/Plugins/Experimental/Avalanche/Source/AvalancheMediaEditor/Private/Rundown/DetailsView/RemoteControl/Properties/AvaRundownPageRemoteControlWidgetUtils.cpp
--- a/Engine/Plugins/Experimental/Avalanche/Source/AvalancheMediaEditor/Private/Rundown/DetailsView/RemoteControl/Properties/AvaRundownPageRemoteControlWidgetUtils.cpp
+++ b/Engine/Plugins/Experimental/Avalanche/Source/AvalancheMediaEditor/Private/Rundown/DetailsView/RemoteControl/Properties/AvaRundownPageRemoteControlWidgetUtils.cpp
@@ -8,8 +8,6 @@
 
 #define LOCTEXT_NAMESPACE "AvaRundownPageRemoteControlWidgetUtils"
 
-BEGIN_SLATE_FUNCTION_BUILD_OPTIMIZATION
-
 TSharedRef<SWidget> FAvaRundownPageRemoteControlWidgetUtils::CreateNodeValueWidget(const TSharedPtr<IDetailTreeNode>& InNode)
 {
 	FNodeWidgets NodeWidgets = InNode->CreateNodeWidgets();
@@ -109,6 +107,4 @@ TSharedPtr<IDetailTreeNode> FAvaRundownPageRemoteControlWidgetUtils::FindNode(co
 	return nullptr;
 }
 
-END_SLATE_FUNCTION_BUILD_OPTIMIZATION
-
 #undef LOCTEXT_NAMESPACE
diff --git a/Engine/Plugins/Experimental/Avalanche/Source/AvalancheMediaEditor/Private/Rundown/DetailsView/RemoteControl/Properties/SAvaRundownPageRemoteControlProps.cpp b/Engine/Plugins/Experimental/Avalanche/Source/AvalancheMediaEditor/Private/Rundown/DetailsView/RemoteControl/Properties/SAvaRundownPageRemoteControlProps.cpp
--- a/Engine/Plugins/Experimental/Avalanche/Source/AvalancheMediaEditor/Private/Rundown/DetailsView/RemoteControl/Properties/SAvaRundownPageRemoteControlProps.cpp
+++ b/Engine/Plugins/Experimental/Avalanche/Source/AvalancheMediaEditor/Private/Rundown/DetailsView/RemoteControl/Properties/SAvaRundownPageRemoteControlProps.cpp
@@ -29,8 +29,6 @@ const FName SAvaRundownPageRemoteControlProps::ValueColumnName = TEXT("ValueColu
 FAvaRundownRCPropertyHeaderRowExtensionDelegate SAvaRundownPageRemoteControlProps::HeaderRowExtensionDelegate;
 TMap<FName, TArray<FAvaRundownRCPropertyTableRowExtensionDelegate>> SAvaRundownPageRemoteControlProps::TableRowExtensionDelegates;
 
-BEGIN_SLATE_FUNCTION_BUILD_OPTIMIZATION
-
 TArray<FAvaRundownRCPropertyTableRowExtensionDelegate>& SAvaRundownPageRemoteControlProps::GetTableRowExtensionDelegates(FName InExtensionName)
 {
 	return TableRowExtensionDelegates.FindOrAdd(InExtensionName);
@@ -74,8 +72,6 @@ void SAvaRundownPageRemoteControlProps::Construct(const FArguments& InArgs, TSha
 	Refresh({});
 }
 
-END_SLATE_FUNCTION_BUILD_OPTIMIZATION
-
 SAvaRundownPageRemoteControlProps::~SAvaRundownPageRemoteControlProps()
 {
 }
diff --git a/Engine/Plugins/Experimental/Avalanche/Source/AvalancheTextEditor/Private/Slate/AvaFontSearchSettingsMenu.cpp b/Engine/Plugins/Experimental/Avalanche/Source/AvalancheTextEditor/Private/Slate/AvaFontSearchSettingsMenu.cpp
--- a/Engine/Plugins/Experimental/Avalanche/Source/AvalancheTextEditor/Private/Slate/AvaFontSearchSettingsMenu.cpp
+++ b/Engine/Plugins/Experimental/Avalanche/Source/AvalancheTextEditor/Private/Slate/AvaFontSearchSettingsMenu.cpp
@@ -10,7 +10,6 @@
 
 #define LOCTEXT_NAMESPACE "AvaFontSearchSettingsMenu"
 
-BEGIN_SLATE_FUNCTION_BUILD_OPTIMIZATION
 void SAvaFontSearchSettingsMenu::Construct(const FArguments& InArgs)
 {
 	FontSelector = InArgs._AvaFontSelector;
@@ -54,7 +53,6 @@ void SAvaFontSearchSettingsMenu::Construct(const FArguments& InArgs)
 		]
 	];
 }
-END_SLATE_FUNCTION_BUILD_OPTIMIZATION
 
 TSharedRef<SWidget> SAvaFontSearchSettingsMenu::CreateSettingsMenu() const
 {
diff --git a/Engine/Plugins/Experimental/Avalanche/Source/AvalancheTextEditor/Private/Slate/SAvaFontField.cpp b/Engine/Plugins/Experimental/Avalanche/Source/AvalancheTextEditor/Private/Slate/SAvaFontField.cpp
--- a/Engine/Plugins/Experimental/Avalanche/Source/AvalancheTextEditor/Private/Slate/SAvaFontField.cpp
+++ b/Engine/Plugins/Experimental/Avalanche/Source/AvalancheTextEditor/Private/Slate/SAvaFontField.cpp
@@ -19,8 +19,6 @@ namespace UE::AvaTextEditor::Private
 
 #define LOCTEXT_NAMESPACE "AvaFontField"
 
-BEGIN_SLATE_FUNCTION_BUILD_OPTIMIZATION
-
 void SAvaFontField::Construct(const FArguments& InArgs)
 {
 	FontView = InArgs._FontView;
@@ -96,8 +94,6 @@ void SAvaFontField::Construct(const FArguments& InArgs)
 	UpdateFont(FontView.Get());
 }
 
-END_SLATE_FUNCTION_BUILD_OPTIMIZATION
-
 void SAvaFontField::UpdateFont(const TSharedPtr<FAvaFontView>& InFontView)
 {
 	if (!InFontView.IsValid())
diff --git a/Engine/Plugins/Experimental/Avalanche/Source/AvalancheTextEditor/Private/Slate/SAvaTextAlignmentWidget.cpp b/Engine/Plugins/Experimental/Avalanche/Source/AvalancheTextEditor/Private/Slate/SAvaTextAlignmentWidget.cpp
--- a/Engine/Plugins/Experimental/Avalanche/Source/AvalancheTextEditor/Private/Slate/SAvaTextAlignmentWidget.cpp
+++ b/Engine/Plugins/Experimental/Avalanche/Source/AvalancheTextEditor/Private/Slate/SAvaTextAlignmentWidget.cpp
@@ -18,8 +18,6 @@
 
 #define LOCTEXT_NAMESPACE "SAvaTextAlignmentWidget"
 
-BEGIN_SLATE_FUNCTION_BUILD_OPTIMIZATION
-
 void SAvaTextAlignmentWidget::Construct(const FArguments& InArgs)
 {
 	TextAlignmentPropertyHandle = InArgs._TextAlignmentPropertyHandle;
@@ -88,8 +86,6 @@ void SAvaTextAlignmentWidget::Construct(const FArguments& InArgs)
 	TextAlignmentPropertyHandle.Get()->SetOnPropertyValueChanged(FSimpleDelegate::CreateSP(this, &SAvaTextAlignmentWidget::OnAlignmentPropertyChanged));
 }
 
-END_SLATE_FUNCTION_BUILD_OPTIMIZATION
-
 TSharedRef<SButton> SAvaTextAlignmentWidget::GetHorizontalAlignmentButton(TSharedPtr<SButton>& OutButton, EText3DHorizontalTextAlignment InHorizontalAlignment, FName Image, FText Tooltip)
 {
 	const ISlateStyle* MotionDesignEditorStyle = FSlateStyleRegistry::FindSlateStyle("AvaEditor");
diff --git a/Engine/Plugins/Experimental/DynamicMaterial/Source/DynamicMaterialEditor/Private/Slate/Layers/SDMLayerEffectsItem.cpp b/Engine/Plugins/Experimental/DynamicMaterial/Source/DynamicMaterialEditor/Private/Slate/Layers/SDMLayerEffectsItem.cpp
--- a/Engine/Plugins/Experimental/DynamicMaterial/Source/DynamicMaterialEditor/Private/Slate/Layers/SDMLayerEffectsItem.cpp
+++ b/Engine/Plugins/Experimental/DynamicMaterial/Source/DynamicMaterialEditor/Private/Slate/Layers/SDMLayerEffectsItem.cpp
@@ -15,8 +15,6 @@
 
 #define LOCTEXT_NAMESPACE "SDMLayerEffectsItem"
 
-BEGIN_SLATE_FUNCTION_BUILD_OPTIMIZATION
-
 FDMEffectsLayerItem::FDMEffectsLayerItem(UDMMaterialEffect* InMaterialEffect)
 {
 	MaterialEffectWeak = InMaterialEffect;
@@ -115,8 +113,6 @@ TSharedRef<SWidget> SDMLayerEffectsItem::CreateLayerRemoveButton()
 		];
 }
 
-END_SLATE_FUNCTION_BUILD_OPTIMIZATION
-
 FCursorReply SDMLayerEffectsItem::OnCursorQuery(const FGeometry& InMyGeometry, const FPointerEvent& InCursorEvent) const
 {
 	if (InCursorEvent.IsMouseButtonDown(EKeys::LeftMouseButton))
diff --git a/Engine/Plugins/Experimental/DynamicMaterial/Source/DynamicMaterialEditor/Private/Slate/Layers/SDMSlotLayerItem.cpp b/Engine/Plugins/Experimental/DynamicMaterial/Source/DynamicMaterialEditor/Private/Slate/Layers/SDMSlotLayerItem.cpp
--- a/Engine/Plugins/Experimental/DynamicMaterial/Source/DynamicMaterialEditor/Private/Slate/Layers/SDMSlotLayerItem.cpp
+++ b/Engine/Plugins/Experimental/DynamicMaterial/Source/DynamicMaterialEditor/Private/Slate/Layers/SDMSlotLayerItem.cpp
@@ -33,8 +33,6 @@
 
 #define LOCTEXT_NAMESPACE "SDMSlotLayer"
 
-BEGIN_SLATE_FUNCTION_BUILD_OPTIMIZATION
-
 void SDMSlotLayerItem::Construct(const FArguments& InArgs, const TSharedPtr<SDMSlot>& InSlotWidget, const TSharedRef<STableViewBase>& InLayerView, 
 	const TSharedPtr<FDMMaterialLayerReference>& InLayerItem)
 {
@@ -461,8 +459,6 @@ TSharedRef<SWidget> SDMSlotLayerItem::CreateEffectsToggleButton()
 		];
 }
 
-END_SLATE_FUNCTION_BUILD_OPTIMIZATION
-
 UDMMaterialStage* SDMSlotLayerItem::GetBaseStage() const
 {
 	if (const UDMMaterialLayerObject* Layer = GetLayer())
diff --git a/Engine/Plugins/Experimental/DynamicMaterial/Source/DynamicMaterialEditor/Private/Slate/Properties/SDMMaterialParameters.cpp b/Engine/Plugins/Experimental/DynamicMaterial/Source/DynamicMaterialEditor/Private/Slate/Properties/SDMMaterialParameters.cpp
--- a/Engine/Plugins/Experimental/DynamicMaterial/Source/DynamicMaterialEditor/Private/Slate/Properties/SDMMaterialParameters.cpp
+++ b/Engine/Plugins/Experimental/DynamicMaterial/Source/DynamicMaterialEditor/Private/Slate/Properties/SDMMaterialParameters.cpp
@@ -23,8 +23,6 @@
 
 #define LOCTEXT_NAMESPACE "SDMMaterialParameters"
 
-BEGIN_SLATE_FUNCTION_BUILD_OPTIMIZATION
-
 void SDMMaterialParameters::Construct(const FArguments& InArgs, TWeakObjectPtr<UDynamicMaterialModel> InBuilderWeak)
 {
 	ensure(InBuilderWeak.IsValid());
@@ -216,8 +214,6 @@ TSharedRef<SWidget> SDMMaterialParameters::CreateValuesWidget()
 	return ValuesList;
 }
 
-END_SLATE_FUNCTION_BUILD_OPTIMIZATION
-
 void SDMMaterialParameters::RefreshWidgets()
 {
 	if (AreaBody.IsValid())
diff --git a/Engine/Plugins/Experimental/DynamicMaterial/Source/DynamicMaterialEditor/Private/Slate/SDMEditor.cpp b/Engine/Plugins/Experimental/DynamicMaterial/Source/DynamicMaterialEditor/Private/Slate/SDMEditor.cpp
--- a/Engine/Plugins/Experimental/DynamicMaterial/Source/DynamicMaterialEditor/Private/Slate/SDMEditor.cpp
+++ b/Engine/Plugins/Experimental/DynamicMaterial/Source/DynamicMaterialEditor/Private/Slate/SDMEditor.cpp
@@ -479,8 +479,6 @@ TSharedRef<SWidget> SDMEditor::CreateActorMaterialSlotSelector(const AActor* InA
 		];
 }
 
-BEGIN_SLATE_FUNCTION_BUILD_OPTIMIZATION
-
 TSharedRef<SWidget> SDMEditor::GetEmptyContent()
 {
 	return SNew(SBox)
@@ -877,8 +875,6 @@ TSharedRef<SWidget> SDMEditor::CreateSlotPickerWidget()
 	return SlotSelector;
 }
 
-END_SLATE_FUNCTION_BUILD_OPTIMIZATION
-
 TSharedPtr<SDMSlot> SDMEditor::GetSlotWidget(UDMMaterialSlot* Slot) const
 {
 	if (ensure(IsValid(Slot)))
diff --git a/Engine/Plugins/Experimental/DynamicMaterial/Source/DynamicMaterialEditor/Private/Slate/SDMSlot.cpp b/Engine/Plugins/Experimental/DynamicMaterial/Source/DynamicMaterialEditor/Private/Slate/SDMSlot.cpp
--- a/Engine/Plugins/Experimental/DynamicMaterial/Source/DynamicMaterialEditor/Private/Slate/SDMSlot.cpp
+++ b/Engine/Plugins/Experimental/DynamicMaterial/Source/DynamicMaterialEditor/Private/Slate/SDMSlot.cpp
@@ -78,8 +78,6 @@ SDMSlot::~SDMSlot()
 	SDMEditor::ClearPropertyHandles(this);
 }
 
-BEGIN_SLATE_FUNCTION_BUILD_OPTIMIZATION
-
 void SDMSlot::Construct(const FArguments& InArgs, const TSharedRef<SDMEditor>& InEditor, UDMMaterialSlot* InSlot)
 {
 	EditorWidgetWeak = InEditor;
@@ -870,8 +868,6 @@ TSharedRef<SWidget> SDMSlot::CreateSlotSettingsRow()
 	return Row;
 }
 
-END_SLATE_FUNCTION_BUILD_OPTIMIZATION
-
 void SDMSlot::InvalidateMainWidget()
 {
 	bInvalidateMainWidget = true;
diff --git a/Engine/Plugins/Experimental/DynamicMaterial/Source/DynamicMaterialEditor/Private/Slate/SDMStage.cpp b/Engine/Plugins/Experimental/DynamicMaterial/Source/DynamicMaterialEditor/Private/Slate/SDMStage.cpp
--- a/Engine/Plugins/Experimental/DynamicMaterial/Source/DynamicMaterialEditor/Private/Slate/SDMStage.cpp
+++ b/Engine/Plugins/Experimental/DynamicMaterial/Source/DynamicMaterialEditor/Private/Slate/SDMStage.cpp
@@ -34,8 +34,6 @@
 
 #define LOCTEXT_NAMESPACE "SDMStage"
 
-BEGIN_SLATE_FUNCTION_BUILD_OPTIMIZATION
-
 SDMStage::~SDMStage()
 {
 	SDMEditor::ClearPropertyHandles(this);
@@ -170,8 +168,6 @@ TSharedRef<SWidget> SDMStage::CreateTextBlockBackground()
 		];
 }
 
-END_SLATE_FUNCTION_BUILD_OPTIMIZATION
-
 EVisibility SDMStage::GetBorderVisibility() const
 {
 	if (bInteractable.Get())
diff --git a/Engine/Plugins/Experimental/DynamicMaterial/Source/DynamicMaterialEditor/Private/Slate/SDMStatusBar.cpp b/Engine/Plugins/Experimental/DynamicMaterial/Source/DynamicMaterialEditor/Private/Slate/SDMStatusBar.cpp
--- a/Engine/Plugins/Experimental/DynamicMaterial/Source/DynamicMaterialEditor/Private/Slate/SDMStatusBar.cpp
+++ b/Engine/Plugins/Experimental/DynamicMaterial/Source/DynamicMaterialEditor/Private/Slate/SDMStatusBar.cpp
@@ -13,8 +13,6 @@
 
 #define LOCTEXT_NAMESPACE "SDMStatusBar"
 
-BEGIN_SLATE_FUNCTION_BUILD_OPTIMIZATION
-
 void SDMStatusBar::Construct(const FArguments& InArgs, const TWeakObjectPtr<UDynamicMaterialModel>& InEditorModel, const TSharedRef<SDMEditor>& InEditorWidget)
 {
 	SetCanTick(true);
@@ -88,8 +86,6 @@ SWrapBox::FSlot::FSlotArguments SDMStatusBar::CreateStatsWrapBoxEntry(TAttribute
 		);
 }
 
-END_SLATE_FUNCTION_BUILD_OPTIMIZATION
-
 FText SDMStatusBar::GetNumMaterialSlotsText() const
 {
 	return FText::Format(LOCTEXT("NumMaterialSlots_Text", "Material Slots {0}"), CachedSlotCount);
diff --git a/Engine/Plugins/Experimental/DynamicMaterial/Source/DynamicMaterialEditor/Private/Slate/SDMToolBar.cpp b/Engine/Plugins/Experimental/DynamicMaterial/Source/DynamicMaterialEditor/Private/Slate/SDMToolBar.cpp
--- a/Engine/Plugins/Experimental/DynamicMaterial/Source/DynamicMaterialEditor/Private/Slate/SDMToolBar.cpp
+++ b/Engine/Plugins/Experimental/DynamicMaterial/Source/DynamicMaterialEditor/Private/Slate/SDMToolBar.cpp
@@ -27,8 +27,6 @@
 
 #define LOCTEXT_NAMESPACE "SDMToolBar"
 
-BEGIN_SLATE_FUNCTION_BUILD_OPTIMIZATION
-
 void SDMToolBar::Construct(const FArguments& InArgs)
 {
 	SetCanTick(true);
@@ -218,8 +216,6 @@ TSharedRef<SWidget> SDMToolBar::GenerateSelectedMaterialSlotRow(TSharedPtr<FDMOb
 	return SNullWidget::NullWidget;
 }
 
-END_SLATE_FUNCTION_BUILD_OPTIMIZATION
-
 FText SDMToolBar::GetSlotDisplayName(TSharedPtr<FDMObjectMaterialProperty> InSlot) const
 {
 	return InSlot->GetPropertyName(false);
diff --git a/Engine/Plugins/Experimental/DynamicMaterial/Source/DynamicMaterialEditor/Private/Slate/SMaterialToolTip.cpp b/Engine/Plugins/Experimental/DynamicMaterial/Source/DynamicMaterialEditor/Private/Slate/SMaterialToolTip.cpp
--- a/Engine/Plugins/Experimental/DynamicMaterial/Source/DynamicMaterialEditor/Private/Slate/SMaterialToolTip.cpp
+++ b/Engine/Plugins/Experimental/DynamicMaterial/Source/DynamicMaterialEditor/Private/Slate/SMaterialToolTip.cpp
@@ -12,8 +12,6 @@
 
 #define LOCTEXT_NAMESPACE "SMaterialToolTip"
 
-BEGIN_SLATE_FUNCTION_BUILD_OPTIMIZATION
-
 SMaterialToolTip::SMaterialToolTip()
 	: MaterialBrush(FSlateMaterialBrush(FVector2D(1.0f, 1.0f)))
 {
@@ -106,6 +104,4 @@ FOptionalSize SMaterialToolTip::GetMaterialSizeY() const
 	return MaterialSize.Get().Y;
 }
 
-END_SLATE_FUNCTION_BUILD_OPTIMIZATION
-
 #undef LOCTEXT_NAMESPACE
diff --git a/Engine/Plugins/Experimental/GeometryMask/Source/GeometryMaskEditor/Private/Widgets/SGMECanvasItem.cpp b/Engine/Plugins/Experimental/GeometryMask/Source/GeometryMaskEditor/Private/Widgets/SGMECanvasItem.cpp
--- a/Engine/Plugins/Experimental/GeometryMask/Source/GeometryMaskEditor/Private/Widgets/SGMECanvasItem.cpp
+++ b/Engine/Plugins/Experimental/GeometryMask/Source/GeometryMaskEditor/Private/Widgets/SGMECanvasItem.cpp
@@ -9,8 +9,6 @@
 #include "Widgets/SOverlay.h"
 #include "Widgets/Text/STextBlock.h"
 
-BEGIN_SLATE_FUNCTION_BUILD_OPTIMIZATION
-
 void SGMECanvasItem::Construct(
 	const FArguments& InArgs,
 	const TSharedRef<STableViewBase>& InOwnerTableView,
@@ -51,5 +49,3 @@ FOptionalSize SGMECanvasItem::GetAspectRatio()
 
 	return AspectRatio;
 }
-
-END_SLATE_FUNCTION_BUILD_OPTIMIZATION
diff --git a/Engine/Plugins/Experimental/GeometryMask/Source/GeometryMaskEditor/Private/Widgets/SGMECanvasList.cpp b/Engine/Plugins/Experimental/GeometryMask/Source/GeometryMaskEditor/Private/Widgets/SGMECanvasList.cpp
--- a/Engine/Plugins/Experimental/GeometryMask/Source/GeometryMaskEditor/Private/Widgets/SGMECanvasList.cpp
+++ b/Engine/Plugins/Experimental/GeometryMask/Source/GeometryMaskEditor/Private/Widgets/SGMECanvasList.cpp
@@ -9,8 +9,6 @@
 #include "Widgets/Layout/SSeparator.h"
 #include "Widgets/Views/SListView.h"
 
-BEGIN_SLATE_FUNCTION_BUILD_OPTIMIZATION
-
 #define LOCTEXT_NAMESPACE "SGMECanvasList"
 
 SGMECanvasList::~SGMECanvasList()
@@ -102,5 +100,3 @@ void SGMECanvasList::Refresh()
 }
 
 #undef LOCTEXT_NAMESPACE
-
-END_SLATE_FUNCTION_BUILD_OPTIMIZATION
diff --git a/Engine/Plugins/Experimental/GeometryMask/Source/GeometryMaskEditor/Private/Widgets/SGMEResourceItem.cpp b/Engine/Plugins/Experimental/GeometryMask/Source/GeometryMaskEditor/Private/Widgets/SGMEResourceItem.cpp
--- a/Engine/Plugins/Experimental/GeometryMask/Source/GeometryMaskEditor/Private/Widgets/SGMEResourceItem.cpp
+++ b/Engine/Plugins/Experimental/GeometryMask/Source/GeometryMaskEditor/Private/Widgets/SGMEResourceItem.cpp
@@ -6,8 +6,6 @@
 #include "SlateOptMacros.h"
 #include "ViewModels/GMEResourceItemViewModel.h"
 
-BEGIN_SLATE_FUNCTION_BUILD_OPTIMIZATION
-
 void SGMEResourceItem::Construct(
 	const FArguments& InArgs,
 	const TSharedRef<STableViewBase>& InOwnerTableView,
@@ -50,5 +48,3 @@ FOptionalSize SGMEResourceItem::GetAspectRatio()
 
 	return AspectRatio;
 }
-
-END_SLATE_FUNCTION_BUILD_OPTIMIZATION
diff --git a/Engine/Plugins/Experimental/GeometryMask/Source/GeometryMaskEditor/Private/Widgets/SGMEResourceList.cpp b/Engine/Plugins/Experimental/GeometryMask/Source/GeometryMaskEditor/Private/Widgets/SGMEResourceList.cpp
--- a/Engine/Plugins/Experimental/GeometryMask/Source/GeometryMaskEditor/Private/Widgets/SGMEResourceList.cpp
+++ b/Engine/Plugins/Experimental/GeometryMask/Source/GeometryMaskEditor/Private/Widgets/SGMEResourceList.cpp
@@ -8,8 +8,6 @@
 #include "ViewModels/GMEResourceListViewModel.h"
 #include "Widgets/Layout/SSeparator.h"
 
-BEGIN_SLATE_FUNCTION_BUILD_OPTIMIZATION
-
 #define LOCTEXT_NAMESPACE "SGMEResourceList"
 
 SGMEResourceList::~SGMEResourceList()
@@ -138,5 +136,3 @@ void SGMEResourceList::Refresh()
 }
 
 #undef LOCTEXT_NAMESPACE
-
-END_SLATE_FUNCTION_BUILD_OPTIMIZATION
diff --git a/Engine/Plugins/Experimental/GeometryMask/Source/GeometryMaskEditor/Private/Widgets/SGeometryMaskCanvasPreview.cpp b/Engine/Plugins/Experimental/GeometryMask/Source/GeometryMaskEditor/Private/Widgets/SGeometryMaskCanvasPreview.cpp
--- a/Engine/Plugins/Experimental/GeometryMask/Source/GeometryMaskEditor/Private/Widgets/SGeometryMaskCanvasPreview.cpp
+++ b/Engine/Plugins/Experimental/GeometryMask/Source/GeometryMaskEditor/Private/Widgets/SGeometryMaskCanvasPreview.cpp
@@ -33,8 +33,6 @@ namespace UE::GeometryMaskEditor::Private
 	static FName FeatherParameterName = TEXT("Feather");
 }
 
-BEGIN_SLATE_FUNCTION_BUILD_OPTIMIZATION
-
 SGeometryMaskCanvasPreview::SGeometryMaskCanvasPreview()
 	: PreviewMaterialPath(UE::GeometryMaskEditor::Private::PreviewMaterialPath)
 	, DefaultTexturePath(UE::GeometryMaskEditor::Private::DefaultTexturePath)
@@ -319,5 +317,3 @@ void SGeometryMaskCanvasPreview::UpdateBrush(const UGeometryMaskCanvas* InCanvas
 		AspectRatio.Set(ImageSize.X / ImageSize.Y);
 	}
 }
-
-END_SLATE_FUNCTION_BUILD_OPTIMIZATION
diff --git a/Engine/Plugins/Experimental/StormSync/Source/StormSyncEditor/Private/Slate/SStormSyncExportWizard.cpp b/Engine/Plugins/Experimental/StormSync/Source/StormSyncEditor/Private/Slate/SStormSyncExportWizard.cpp
--- a/Engine/Plugins/Experimental/StormSync/Source/StormSyncEditor/Private/Slate/SStormSyncExportWizard.cpp
+++ b/Engine/Plugins/Experimental/StormSync/Source/StormSyncEditor/Private/Slate/SStormSyncExportWizard.cpp
@@ -28,8 +28,6 @@ namespace StormSync::ExportWizardInternal
 	static FString Default_FileExtension = TEXT(".spak");
 }
 
-BEGIN_SLATE_FUNCTION_BUILD_OPTIMIZATION
-
 void SStormSyncExportWizard::Construct(const FArguments& InArgs, const TArray<FName>& InInitialPackageNames, const TArray<FName>& InPackageNames, const FOnExportWizardCompleted& InOnExportWizardCompleted)
 {
 	OnWizardCompleted = InOnExportWizardCompleted;
@@ -241,8 +239,6 @@ void SStormSyncExportWizard::Construct(const FArguments& InArgs, const TArray<FN
 	];
 }
 
-END_SLATE_FUNCTION_BUILD_OPTIMIZATION
-
 void SStormSyncExportWizard::OpenWizard(const TArray<FName>& InInitialPackageNames, const TArray<FName>& PackageNames, const FOnExportWizardCompleted& InOnExportWizardCompleted)
 {
 	const TSharedRef<SWindow> ReportWindow = SNew(SWindow)
diff --git a/Engine/Plugins/Experimental/StormSync/Source/StormSyncEditor/Private/Slate/SStormSyncFileDependencyWidgetRow.cpp b/Engine/Plugins/Experimental/StormSync/Source/StormSyncEditor/Private/Slate/SStormSyncFileDependencyWidgetRow.cpp
--- a/Engine/Plugins/Experimental/StormSync/Source/StormSyncEditor/Private/Slate/SStormSyncFileDependencyWidgetRow.cpp
+++ b/Engine/Plugins/Experimental/StormSync/Source/StormSyncEditor/Private/Slate/SStormSyncFileDependencyWidgetRow.cpp
@@ -7,8 +7,6 @@
 #include "StormSyncImportTypes.h"
 #include "Widgets/Text/STextBlock.h"
 
-BEGIN_SLATE_FUNCTION_BUILD_OPTIMIZATION
-
 void SStormSyncFileDependencyWidgetRow::Construct(const FArguments& InArgs, const TSharedRef<STableViewBase>& InOwnerTableView)
 {
 	WidgetItem = InArgs._Item;
@@ -19,8 +17,6 @@ void SStormSyncFileDependencyWidgetRow::Construct(const FArguments& InArgs, cons
 	);
 }
 
-END_SLATE_FUNCTION_BUILD_OPTIMIZATION
-
 TSharedRef<SWidget> SStormSyncFileDependencyWidgetRow::GenerateWidgetForColumn(const FName& InColumnName)
 {
 	TSharedRef<SBox> BoxWrapper = SNew(SBox)
diff --git a/Engine/Plugins/Experimental/StormSync/Source/StormSyncEditor/Private/Slate/SStormSyncImportWizard.cpp b/Engine/Plugins/Experimental/StormSync/Source/StormSyncEditor/Private/Slate/SStormSyncImportWizard.cpp
--- a/Engine/Plugins/Experimental/StormSync/Source/StormSyncEditor/Private/Slate/SStormSyncImportWizard.cpp
+++ b/Engine/Plugins/Experimental/StormSync/Source/StormSyncEditor/Private/Slate/SStormSyncImportWizard.cpp
@@ -12,8 +12,6 @@
 
 #define LOCTEXT_NAMESPACE "StormSyncImportWizard"
 
-BEGIN_SLATE_FUNCTION_BUILD_OPTIMIZATION
-
 void SStormSyncImportWizard::Construct(const FArguments& InArgs, const TArray<FStormSyncImportFileInfo>& InFilesToImport, const TArray<FStormSyncImportFileInfo>& InBufferFiles)
 {
 	FilesToImport = InFilesToImport;
@@ -162,8 +160,6 @@ void SStormSyncImportWizard::Construct(const FArguments& InArgs, const TArray<FS
 	];
 }
 
-END_SLATE_FUNCTION_BUILD_OPTIMIZATION
-
 void SStormSyncImportWizard::InitListSources()
 {
 	for (const FStormSyncImportFileInfo& FileInfo : FilesToImport)
diff --git a/Engine/Plugins/Experimental/StormSync/Source/StormSyncEditor/Private/Slate/SStormSyncProgressBarNotification.cpp b/Engine/Plugins/Experimental/StormSync/Source/StormSyncEditor/Private/Slate/SStormSyncProgressBarNotification.cpp
--- a/Engine/Plugins/Experimental/StormSync/Source/StormSyncEditor/Private/Slate/SStormSyncProgressBarNotification.cpp
+++ b/Engine/Plugins/Experimental/StormSync/Source/StormSyncEditor/Private/Slate/SStormSyncProgressBarNotification.cpp
@@ -11,8 +11,6 @@
 
 #define LOCTEXT_NAMESPACE "SStormSyncProgressBarNotification"
 
-BEGIN_SLATE_FUNCTION_BUILD_OPTIMIZATION
-
 void SStormSyncProgressBarNotification::Construct(const FArguments& InArgs)
 {
 	TotalBytes = InArgs._TotalBytes;
@@ -68,8 +66,6 @@ void SStormSyncProgressBarNotification::Construct(const FArguments& InArgs)
 	];
 }
 
-END_SLATE_FUNCTION_BUILD_OPTIMIZATION
-
 float SStormSyncProgressBarNotification::GetPercent() const
 {
 	if (TotalBytes == 0)
diff --git a/Engine/Plugins/Experimental/StormSync/Source/StormSyncEditor/Private/Slate/SStormSyncReportDialog.cpp b/Engine/Plugins/Experimental/StormSync/Source/StormSyncEditor/Private/Slate/SStormSyncReportDialog.cpp
--- a/Engine/Plugins/Experimental/StormSync/Source/StormSyncEditor/Private/Slate/SStormSyncReportDialog.cpp
+++ b/Engine/Plugins/Experimental/StormSync/Source/StormSyncEditor/Private/Slate/SStormSyncReportDialog.cpp
@@ -92,8 +92,6 @@ FStormSyncPackageReportNode::FChildrenState FStormSyncPackageReportNode::AddPack
 	return ChildrenState;
 }
 
-BEGIN_SLATE_FUNCTION_BUILD_OPTIMIZATION
-
 void SStormSyncReportDialog::Construct(const FArguments& InArgs, const FText& InReportMessage, TArray<FStormSyncReportPackageData>& InPackageNames)
 {
 	FolderOpenBrush = FAppStyle::GetBrush("ContentBrowser.AssetTreeFolderOpen");
@@ -144,8 +142,6 @@ void SStormSyncReportDialog::Construct(const FArguments& InArgs, const FText& In
 	}
 }
 
-END_SLATE_FUNCTION_BUILD_OPTIMIZATION
-
 void SStormSyncReportDialog::CloseDialog()
 {
 	const TSharedPtr<SWindow> Window = FSlateApplication::Get().FindWidgetWindow(AsShared());
diff --git a/Engine/Plugins/Experimental/StormSync/Source/StormSyncEditor/Private/Slate/Status/SStormSyncStatusWidget.cpp b/Engine/Plugins/Experimental/StormSync/Source/StormSyncEditor/Private/Slate/Status/SStormSyncStatusWidget.cpp
--- a/Engine/Plugins/Experimental/StormSync/Source/StormSyncEditor/Private/Slate/Status/SStormSyncStatusWidget.cpp
+++ b/Engine/Plugins/Experimental/StormSync/Source/StormSyncEditor/Private/Slate/Status/SStormSyncStatusWidget.cpp
@@ -14,8 +14,6 @@
 
 #define LOCTEXT_NAMESPACE "StormSyncStatusWidget"
 
-BEGIN_SLATE_FUNCTION_BUILD_OPTIMIZATION
-
 void SStormSyncStatusWidget::Construct(const FArguments& InArgs, const TSharedRef<FStormSyncTransportStatusResponse>& InStatusResponse)
 {
 	FStormSyncTransportStatusResponse& StatusResponse = InStatusResponse.Get();
@@ -190,8 +188,6 @@ void SStormSyncStatusWidget::Construct(const FArguments& InArgs, const TSharedRe
 	];
 }
 
-END_SLATE_FUNCTION_BUILD_OPTIMIZATION
-
 TSharedRef<SWindow> SStormSyncStatusWidget::CreateWindow(const TSharedRef<FStormSyncTransportStatusResponse>& InStatusResponse)
 {
 	return SNew(SWindow)
