Motion Design
- Changed the position of the high res screenshot button.
- Changed the align actor button icon.

diff --git a/Engine/Plugins/Experimental/Avalanche/Source/AvalancheLevelViewport/Private/Widgets/SAvaLevelViewportStatusBarButtons.cpp b/Engine/Plugins/Experimental/Avalanche/Source/AvalancheLevelViewport/Private/Widgets/SAvaLevelViewportStatusBarButtons.cpp
--- a/Engine/Plugins/Experimental/Avalanche/Source/AvalancheLevelViewport/Private/Widgets/SAvaLevelViewportStatusBarButtons.cpp
+++ b/Engine/Plugins/Experimental/Avalanche/Source/AvalancheLevelViewport/Private/Widgets/SAvaLevelViewportStatusBarButtons.cpp
@@ -271,7 +271,7 @@ void SAvaLevelViewportStatusBarButtons::PopulateActorButtons(TSharedPtr<SHorizon
 	TSharedRef<SComboButton> AlignmentButton = ViewportStatusBarButton::MakeMenuButton(
 		LOCTEXT("ActorAlign", "Align Actors"),
 		FOnGetContent::CreateSP(this, &SAvaLevelViewportStatusBarButtons::GetActorAlignmentMenuContent),
-		FAppStyle::Get().GetBrush(TEXT("Icons.Layout")),
+		FAvaLevelViewportStyle::Get().GetBrush(TEXT("Icons.Alignment.Left")),
 		TAttribute<FSlateColor>::CreateSP(this, &SAvaLevelViewportStatusBarButtons::GetActorAlignmentColor)
 	);
 
@@ -295,20 +295,6 @@ void SAvaLevelViewportStatusBarButtons::PopulateViewportButtons(TSharedPtr<SHori
 	const FAvaLevelViewportCommands& CommandsRef = FAvaLevelViewportCommands::Get();
 	const FLevelViewportCommands& ViewportActionsRef = FLevelViewportCommands::Get();
 
-	InContainer->AddSlot()
-		.AutoWidth()
-		.Padding(ViewportStatusBarButton::Padding)
-		[
-			ViewportStatusBarButton::MakeButton(
-				this,
-				ViewportActionsRef.HighResScreenshot,
-				FSlateIcon(FAppStyle::GetAppStyleSetName(), "LevelViewport.HighResScreenshot").GetIcon(),
-				&SAvaLevelViewportStatusBarButtons::HighResScreenshot,
-				&SAvaLevelViewportStatusBarButtons::GetHighResScreenshotEnabled,
-				&SAvaLevelViewportStatusBarButtons::GetHighResScreenshotColor
-			)
-		];
-
 	InContainer->AddSlot()
 		.AutoWidth()
 		.Padding(ViewportStatusBarButton::Padding)
@@ -480,6 +466,20 @@ void SAvaLevelViewportStatusBarButtons::PopulateViewportButtons(TSharedPtr<SHori
 			SnapButton
 		];
 
+	InContainer->AddSlot()
+		.AutoWidth()
+		.Padding(ViewportStatusBarButton::Padding)
+		[
+			ViewportStatusBarButton::MakeButton(
+				this,
+				ViewportActionsRef.HighResScreenshot,
+				FSlateIcon(FAppStyle::GetAppStyleSetName(), "LevelViewport.HighResScreenshot").GetIcon(),
+				&SAvaLevelViewportStatusBarButtons::HighResScreenshot,
+				&SAvaLevelViewportStatusBarButtons::GetHighResScreenshotEnabled,
+				&SAvaLevelViewportStatusBarButtons::GetHighResScreenshotColor
+			)
+		];
+
 	TSharedRef<SComboButton> ViewportInfoButton = ViewportStatusBarButton::MakeMenuButton(
 		LOCTEXT("ViewportInfomation", "Viewport Information"),
 		FOnGetContent::CreateSP(this, &SAvaLevelViewportStatusBarButtons::GetViewportInfoWidget),
