[Motion Design]
 - Change how a Scene Rig tag is stored

diff --git a/Engine/Plugins/Experimental/Avalanche/Source/AvalancheSceneRig/Private/AvaSceneRigSubsystem.cpp b/Engine/Plugins/Experimental/Avalanche/Source/AvalancheSceneRig/Private/AvaSceneRigSubsystem.cpp
--- a/Engine/Plugins/Experimental/Avalanche/Source/AvalancheSceneRig/Private/AvaSceneRigSubsystem.cpp
+++ b/Engine/Plugins/Experimental/Avalanche/Source/AvalancheSceneRig/Private/AvaSceneRigSubsystem.cpp
@@ -1,30 +1,69 @@
 // Copyright Epic Games, Inc. All Rights Reserved.
 
 #include "AvaSceneRigSubsystem.h"
+#include "AvaSceneRigAssetTags.h"
 #include "AssetRegistry/AssetRegistryModule.h"
 #include "Containers/Array.h"
 #include "Engine/Level.h"
 #include "Engine/LevelStreaming.h"
 #include "Engine/LevelStreamingDynamic.h"
 #include "Engine/World.h"
-#include "Interfaces/Interface_AssetUserData.h"
+#include "UObject/AssetRegistryTagsContext.h"
 #include "UObject/Class.h"
-#include "UObject/SoftObjectPtr.h"
 #include "UObject/NoExportTypes.h"
+#include "UObject/Object.h"
+#include "UObject/SoftObjectPtr.h"
 
 DEFINE_LOG_CATEGORY(AvaSceneRigSubsystemLog);
 
 #define LOCTEXT_NAMESPACE "AvaSceneRigSubsystem"
 
-void UAvaSceneRigSubsystem::Initialize(FSubsystemCollectionBase& Collection)
+void UAvaSceneRigSubsystem::Initialize(FSubsystemCollectionBase& InOutCollection)
 {
-	Super::Initialize(Collection);
+	Super::Initialize(InOutCollection);
+
+#if WITH_EDITOR
+	if (!HasAnyFlags(EObjectFlags::RF_ClassDefaultObject))
+	{
+		WorldTagGetterDelegate = UObject::FAssetRegistryTag::OnGetExtraObjectTagsWithContext.AddUObject(this, &UAvaSceneRigSubsystem::OnGetWorldTags);
+	}
+#endif
 }
 
 void UAvaSceneRigSubsystem::Deinitialize()
 {
 	Super::Deinitialize();
+
+#if WITH_EDITOR
+	UObject::FAssetRegistryTag::OnGetExtraObjectTagsWithContext.Remove(WorldTagGetterDelegate);
+	WorldTagGetterDelegate.Reset();
+#endif
+}
+
+#if WITH_EDITOR
+void UAvaSceneRigSubsystem::OnGetWorldTags(FAssetRegistryTagsContext InContext) const
+{
+	// Outer of this subsystem should always be the Motion Design level the scene rig exists in
+	const UWorld* const OuterWorld = GetTypedOuter<UWorld>();
+	if (!IsValid(OuterWorld))
+	{
+		return;
+	}
+
+	// The context object should be the scene rig world asset
+	const UWorld* const SceneRigWorld = Cast<UWorld>(InContext.GetObject());
+	if (!IsValid(SceneRigWorld) || SceneRigWorld == OuterWorld)
+	{
+		return;
+	}
+
+	// NOTE: Are there other cases where a context world object has the motion design level as outer?
+	// May need to add an additional check here in that case.
+
+	using namespace UE::AvaSceneRig::AssetTags;
+	InContext.AddTag(UObject::FAssetRegistryTag(SceneRig, Values::Enabled, UObject::FAssetRegistryTag::TT_Alphabetical));
 }
+#endif
 
 UAvaSceneRigSubsystem* UAvaSceneRigSubsystem::ForWorld(const UWorld* const InWorld)
 {
@@ -56,22 +95,12 @@ bool UAvaSceneRigSubsystem::IsSceneRigAssetData(const FAssetData& InAssetData)
 		return false;
 	}
 
-	UWorld* const SceneRigAsset = Cast<UWorld>(InAssetData.GetAsset());
-	if (!IsValid(SceneRigAsset) || !IsValid(SceneRigAsset->PersistentLevel))
-	{
-		return false;
-	}
+	using namespace UE::AvaSceneRig::AssetTags;
+	
+	FString TagValue;
+	InAssetData.GetTagValue(SceneRig, TagValue);
 
-	if (IInterface_AssetUserData* AssetUserData = Cast<IInterface_AssetUserData>(SceneRigAsset->PersistentLevel))
-	{
-		UAvaSceneRigData* const SceneRigData = AssetUserData->GetAssetUserData<UAvaSceneRigData>();
-		if (IsValid(SceneRigData))
-		{
-			return true;
-		}
-	}
-
-	return false;
+	return TagValue.Equals(Values::Enabled);
 }
 
 bool UAvaSceneRigSubsystem::IsSceneRigAsset(UObject* const InObject)
diff --git a/Engine/Plugins/Experimental/Avalanche/Source/AvalancheSceneRig/Public/AvaSceneRigSubsystem.h b/Engine/Plugins/Experimental/Avalanche/Source/AvalancheSceneRig/Public/AvaSceneRigSubsystem.h
--- a/Engine/Plugins/Experimental/Avalanche/Source/AvalancheSceneRig/Public/AvaSceneRigSubsystem.h
+++ b/Engine/Plugins/Experimental/Avalanche/Source/AvalancheSceneRig/Public/AvaSceneRigSubsystem.h
@@ -10,6 +10,7 @@
 #include "AvaSceneRigSubsystem.generated.h"
 
 class AActor;
+class FAssetRegistryTagsContext;
 class FName;
 class UClass;
 class ULevel;
@@ -20,12 +21,6 @@ struct FAssetData;
 
 DECLARE_LOG_CATEGORY_EXTERN(AvaSceneRigSubsystemLog, Log, All);
 
-UCLASS(MinimalAPI)
-class UAvaSceneRigData : public UAssetUserData
-{
-	GENERATED_BODY()
-};
-
 UCLASS(MinimalAPI)
 class UAvaSceneRigSubsystem : public UWorldSubsystem
 {
@@ -80,12 +75,20 @@ public:
 	AVALANCHESCENERIG_API void ForEachActiveSceneRigActor(TFunction<void(AActor* const InActor)> InFunction) const;
 
 	//~ Begin UObject
-	virtual void Initialize(FSubsystemCollectionBase& Collection) override;
+	virtual void Initialize(FSubsystemCollectionBase& InOutCollection) override;
 	virtual void Deinitialize() override;
 	//~ End UObject
 
 private:
+#if WITH_EDITOR
+	void OnGetWorldTags(FAssetRegistryTagsContext InContext) const;
+#endif
+
 	//~ Begin WorldSubsystem
 	virtual bool ShouldCreateSubsystem(UObject* const InOuter) const override;
 	//~ End WorldSubsystem
+
+#if WITH_EDITOR
+	FDelegateHandle WorldTagGetterDelegate;
+#endif
 };
diff --git a/Engine/Plugins/Experimental/Avalanche/Source/AvalancheSceneRigEditor/Private/AvaSceneRigEditorModule.cpp b/Engine/Plugins/Experimental/Avalanche/Source/AvalancheSceneRigEditor/Private/AvaSceneRigEditorModule.cpp
--- a/Engine/Plugins/Experimental/Avalanche/Source/AvalancheSceneRigEditor/Private/AvaSceneRigEditorModule.cpp
+++ b/Engine/Plugins/Experimental/Avalanche/Source/AvalancheSceneRigEditor/Private/AvaSceneRigEditorModule.cpp
@@ -394,21 +394,6 @@ FSoftObjectPath FAvaSceneRigEditorModule::CreateSceneRigAssetWithDialog() const
 		return FSoftObjectPath();
 	}
 
-	// Add asset data to specify this level as a Scene Rig
-	if (IInterface_AssetUserData* AssetUserData = Cast<IInterface_AssetUserData>(NewSceneRigAsset->PersistentLevel))
-	{
-		UAvaSceneRigData* const SceneRigData = NewObject<UAvaSceneRigData>(NewSceneRigAsset->PersistentLevel, NAME_None, RF_Public | RF_Transactional);
-		check(IsValid(SceneRigData));
-		AssetUserData->AddAssetUserData(SceneRigData);
-	}
-
-	UEditorAssetSubsystem* const EditorAssetSubsystem = GEditor->GetEditorSubsystem<UEditorAssetSubsystem>();
-	if (IsValid(EditorAssetSubsystem))
-	{
-		using namespace UE::AvaSceneRig;
-		EditorAssetSubsystem->SetMetadataTag(NewSceneRigAsset, AssetTags::SceneRig, AssetTags::Values::Enabled);
-	}
-
 	// Save package
 	TArray<UPackage*> PackagesToSave;
 	PackagesToSave.Add(NewSceneRigAsset->GetPackage());
