Motion Design: Fix Media Style Toolbars showing invalid entries in Toolbars when collapsed

diff --git a/Engine/Plugins/Experimental/Avalanche/Source/AvalancheMediaEditor/Private/AvaMediaEditorStyle.cpp b/Engine/Plugins/Experimental/Avalanche/Source/AvalancheMediaEditor/Private/AvaMediaEditorStyle.cpp
--- a/Engine/Plugins/Experimental/Avalanche/Source/AvalancheMediaEditor/Private/AvaMediaEditorStyle.cpp
+++ b/Engine/Plugins/Experimental/Avalanche/Source/AvalancheMediaEditor/Private/AvaMediaEditorStyle.cpp
@@ -17,6 +17,8 @@
 FAvaMediaEditorStyle::FAvaMediaEditorStyle()
 	: FSlateStyleSet(TEXT("AvaMediaEditor"))
 {
+	SetParentStyleName(FAppStyle::GetAppStyleSetName());
+
 	const FVector2D Icon12x12(12.0f, 12.0f);
 	const FVector2D Icon16x16(16.0f, 16.0f);
 	const FVector2D Icon20x20(20.0f, 20.0f);
