[MotionDesign]
 - Rundown Editor: Add display text for toolbar "Alpha Preview" entry

diff --git a/Engine/Plugins/Experimental/Avalanche/Source/AvalancheMediaEditor/Private/Rundown/Preview/SAvaRundownPagePreview.cpp b/Engine/Plugins/Experimental/Avalanche/Source/AvalancheMediaEditor/Private/Rundown/Preview/SAvaRundownPagePreview.cpp
--- a/Engine/Plugins/Experimental/Avalanche/Source/AvalancheMediaEditor/Private/Rundown/Preview/SAvaRundownPagePreview.cpp
+++ b/Engine/Plugins/Experimental/Avalanche/Source/AvalancheMediaEditor/Private/Rundown/Preview/SAvaRundownPagePreview.cpp
@@ -157,8 +157,8 @@ TSharedRef<SWidget> SAvaRundownPagePreview::CreatePagePreviewToolBar(const TShar
 					FExecuteAction::CreateSP(this, &SAvaRundownPagePreview::HandleCheckerboardActionExecute),
 					FCanExecuteAction::CreateLambda([]	{ return true;}))
 				, NAME_None
-				, FText()
-				, LOCTEXT("ToggleAlpha_ToolTip", "Toggle alpha preview (checker board).")
+				, LOCTEXT("AlphaPreview", "Alpha Preview")
+				, LOCTEXT("AlphaPreview_ToolTip", "Toggle alpha preview (checker board).")
 				, FSlateIcon(FAppStyle::GetAppStyleSetName(), "Checkerboard")
 			);
 
