MotionDesign : ActorModifiers
- Fix LogEditorClassViewer warning with modifier shared object caused by missing class specifier (Hidden) on children

diff --git a/Engine/Plugins/Experimental/Avalanche/Source/AvalancheModifiers/Public/Shared/AvaBooleanModifierShared.h b/Engine/Plugins/Experimental/Avalanche/Source/AvalancheModifiers/Public/Shared/AvaBooleanModifierShared.h
--- a/Engine/Plugins/Experimental/Avalanche/Source/AvalancheModifiers/Public/Shared/AvaBooleanModifierShared.h
+++ b/Engine/Plugins/Experimental/Avalanche/Source/AvalancheModifiers/Public/Shared/AvaBooleanModifierShared.h
@@ -48,7 +48,7 @@ struct FAvaBooleanModifierSharedChannelInfo
  * Singleton class for boolean modifiers to share data
  * Transient because we can rebuild it, no need to save the asset
  */
-UCLASS(Transient)
+UCLASS(Transient, Hidden)
 class UAvaBooleanModifierShared : public UActorModifierCoreSharedObject
 {
 	GENERATED_BODY()
diff --git a/Engine/Plugins/Experimental/Avalanche/Source/AvalancheModifiers/Public/Shared/AvaTransformModifierShared.h b/Engine/Plugins/Experimental/Avalanche/Source/AvalancheModifiers/Public/Shared/AvaTransformModifierShared.h
--- a/Engine/Plugins/Experimental/Avalanche/Source/AvalancheModifiers/Public/Shared/AvaTransformModifierShared.h
+++ b/Engine/Plugins/Experimental/Avalanche/Source/AvalancheModifiers/Public/Shared/AvaTransformModifierShared.h
@@ -102,7 +102,7 @@ struct FAvaTransformSharedActorState
  * We want to save the state of that actor once before any modifier changes it
  * and restore it when no other modifier is watching it
  */
-UCLASS()
+UCLASS(Hidden)
 class UAvaTransformModifierShared : public UActorModifierCoreSharedObject
 {
 	GENERATED_BODY()
diff --git a/Engine/Plugins/Experimental/Avalanche/Source/AvalancheModifiers/Public/Shared/AvaTranslucentPriorityModifierShared.h b/Engine/Plugins/Experimental/Avalanche/Source/AvalancheModifiers/Public/Shared/AvaTranslucentPriorityModifierShared.h
--- a/Engine/Plugins/Experimental/Avalanche/Source/AvalancheModifiers/Public/Shared/AvaTranslucentPriorityModifierShared.h
+++ b/Engine/Plugins/Experimental/Avalanche/Source/AvalancheModifiers/Public/Shared/AvaTranslucentPriorityModifierShared.h
@@ -61,7 +61,7 @@ struct FAvaTranslucentPriorityModifierComponentState
 /**
  * Singleton class for translucent priority modifiers to share data about component state
  */
-UCLASS()
+UCLASS(Hidden)
 class UAvaTranslucentPriorityModifierShared : public UActorModifierCoreSharedObject
 {
 	GENERATED_BODY()
diff --git a/Engine/Plugins/Experimental/Avalanche/Source/AvalancheModifiers/Public/Shared/AvaVisibilityModifierShared.h b/Engine/Plugins/Experimental/Avalanche/Source/AvalancheModifiers/Public/Shared/AvaVisibilityModifierShared.h
--- a/Engine/Plugins/Experimental/Avalanche/Source/AvalancheModifiers/Public/Shared/AvaVisibilityModifierShared.h
+++ b/Engine/Plugins/Experimental/Avalanche/Source/AvalancheModifiers/Public/Shared/AvaVisibilityModifierShared.h
@@ -112,7 +112,7 @@ struct FAvaVisibilitySharedActorState
  * We want to save the state of that actor once before any modifier changes it
  * and restore it when no other modifier is watching it
  */
-UCLASS()
+UCLASS(Hidden)
 class UAvaVisibilityModifierShared : public UActorModifierCoreSharedObject
 {
 	GENERATED_BODY()
