MotionDesign : ActorModifiers
- Fix reference actor reset issue caused during property migration on load in autosize, autofollow, lookat modifiers

diff --git a/Engine/Plugins/Experimental/Avalanche/Source/AvalancheModifiers/Private/Modifiers/AvaAutoFollowModifier.cpp b/Engine/Plugins/Experimental/Avalanche/Source/AvalancheModifiers/Private/Modifiers/AvaAutoFollowModifier.cpp
--- a/Engine/Plugins/Experimental/Avalanche/Source/AvalancheModifiers/Private/Modifiers/AvaAutoFollowModifier.cpp
+++ b/Engine/Plugins/Experimental/Avalanche/Source/AvalancheModifiers/Private/Modifiers/AvaAutoFollowModifier.cpp
@@ -57,6 +57,8 @@ void UAvaAutoFollowModifier::OnModifierAdded(EActorModifierCoreEnableReason InRe
 	{
 		SceneExtension->TrackSceneTree(0, &ReferenceActor);
 	}
+
+	bDeprecatedPropertiesMigrated = true;
 }
 
 void UAvaAutoFollowModifier::OnModifierEnabled(EActorModifierCoreEnableReason InReason)
@@ -173,7 +175,9 @@ void UAvaAutoFollowModifier::Apply()
 
 void UAvaAutoFollowModifier::PostLoad()
 {
-	if (!bDeprecatedPropertiesMigrated)
+	if (!bDeprecatedPropertiesMigrated
+		&& ReferenceActor.ReferenceContainer == EAvaReferenceContainer::Other
+		&& ReferenceActor.ReferenceActorWeak == nullptr)
 	{
 		ReferenceActor.ReferenceContainer = ReferenceContainer_DEPRECATED;
 		ReferenceActor.ReferenceActorWeak = ReferenceActorWeak_DEPRECATED;
@@ -328,7 +332,7 @@ void UAvaAutoFollowModifier::OnRenderStateUpdated(AActor* InActor, UActorCompone
 	{
 		return;
 	}
-	
+
 	const AActor* FollowedActor = ReferenceActor.ReferenceActorWeak.Get();
 	const UAvaTransformModifierShared* LayoutShared = GetShared<UAvaTransformModifierShared>(false);
 	
@@ -336,17 +340,21 @@ void UAvaAutoFollowModifier::OnRenderStateUpdated(AActor* InActor, UActorCompone
 	{
 		return;
 	}
-	
+
 	const bool bIsReferenceActor = FollowedActor == InActor;
 	const bool bIsAttachedToReferenceActor = InActor->IsAttachedTo(FollowedActor);
-	const bool bModifierDirtyable = IsModifierDirtyable();
 
-	if (!bIsReferenceActor && !bIsAttachedToReferenceActor && !bModifierDirtyable)
+	if (!bIsReferenceActor && !bIsAttachedToReferenceActor)
 	{
 		return;
 	}
-	
-	MarkModifierDirty();
+
+	const bool bModifierDirtyable = IsModifierDirtyable();
+
+	if (bModifierDirtyable)
+	{
+		MarkModifierDirty();
+	}
 }
 
 void UAvaAutoFollowModifier::OnSceneTreeTrackedActorChanged(int32 InIdx, AActor* InPreviousActor, AActor* InNewActor)
diff --git a/Engine/Plugins/Experimental/Avalanche/Source/AvalancheModifiers/Private/Modifiers/AvaLookAtModifier.cpp b/Engine/Plugins/Experimental/Avalanche/Source/AvalancheModifiers/Private/Modifiers/AvaLookAtModifier.cpp
--- a/Engine/Plugins/Experimental/Avalanche/Source/AvalancheModifiers/Private/Modifiers/AvaLookAtModifier.cpp
+++ b/Engine/Plugins/Experimental/Avalanche/Source/AvalancheModifiers/Private/Modifiers/AvaLookAtModifier.cpp
@@ -20,6 +20,8 @@ void UAvaLookAtModifier::OnModifierAdded(EActorModifierCoreEnableReason InReason
 	{
 		SceneExtension->TrackSceneTree(0, &ReferenceActor);
 	}
+
+	bDeprecatedPropertiesMigrated = true;
 }
 
 void UAvaLookAtModifier::OnModifierEnabled(EActorModifierCoreEnableReason InReason)
@@ -107,7 +109,9 @@ void UAvaLookAtModifier::OnModifiedActorTransformed()
 
 void UAvaLookAtModifier::PostLoad()
 {
-	if (!bDeprecatedPropertiesMigrated)
+	if (!bDeprecatedPropertiesMigrated
+		&& ReferenceActor.ReferenceContainer == EAvaReferenceContainer::Other
+		&& ReferenceActor.ReferenceActorWeak == nullptr)
 	{
 		ReferenceActor.ReferenceContainer = EAvaReferenceContainer::Other;
 		ReferenceActor.ReferenceActorWeak = ReferenceActorWeak;
