Motion Design: add blueprint type to tag handles

diff --git a/Engine/Plugins/Experimental/Avalanche/Source/AvalancheTag/Public/AvaTagHandle.h b/Engine/Plugins/Experimental/Avalanche/Source/AvalancheTag/Public/AvaTagHandle.h
--- a/Engine/Plugins/Experimental/Avalanche/Source/AvalancheTag/Public/AvaTagHandle.h
+++ b/Engine/Plugins/Experimental/Avalanche/Source/AvalancheTag/Public/AvaTagHandle.h
@@ -12,7 +12,7 @@ class UAvaTagCollection;
  * Handle to a Tag in a particular Source.
  * This should be used by the UStructs/UObjects to properly reference a particular FAvaTag.
  */
-USTRUCT()
+USTRUCT(BlueprintType)
 struct AVALANCHETAG_API FAvaTagHandle
 {
 	GENERATED_BODY()
@@ -51,7 +51,7 @@ struct AVALANCHETAG_API FAvaTagHandle
 		return HashCombineFast(GetTypeHash(InHandle.Source), GetTypeHash(InHandle.TagId));
 	}
 
-	UPROPERTY(EditAnywhere, Category="Tag")
+	UPROPERTY(EditAnywhere, Category = "Tag")
 	TObjectPtr<const UAvaTagCollection> Source;
 
 	UPROPERTY()
diff --git a/Engine/Plugins/Experimental/Avalanche/Source/AvalancheTag/Public/AvaTagHandleContainer.h b/Engine/Plugins/Experimental/Avalanche/Source/AvalancheTag/Public/AvaTagHandleContainer.h
--- a/Engine/Plugins/Experimental/Avalanche/Source/AvalancheTag/Public/AvaTagHandleContainer.h
+++ b/Engine/Plugins/Experimental/Avalanche/Source/AvalancheTag/Public/AvaTagHandleContainer.h
@@ -13,7 +13,7 @@ struct FAvaTagHandle;
  * Handle to a multiple tags in a particular Source.
  * This should be used by the UStructs/UObjects to properly reference a multiple FAvaTags.
  */
-USTRUCT()
+USTRUCT(BlueprintType)
 struct AVALANCHETAG_API FAvaTagHandleContainer
 {
 	GENERATED_BODY()
@@ -43,11 +43,11 @@ struct AVALANCHETAG_API FAvaTagHandleContainer
 	/** Returns an array of resolved tags through this container's tag ids and source tag collection */
 	TArray<FAvaTag> ResolveTags() const;
 
-	UPROPERTY(EditAnywhere, Category="Tag")
+	UPROPERTY(EditAnywhere, Category = "Tag")
 	TObjectPtr<const UAvaTagCollection> Source;
 
 private:
-	UPROPERTY(EditAnywhere, Category="Tag")
+	UPROPERTY(EditAnywhere, Category = "Tag")
 	TArray<FAvaTagId> TagIds;
 };
 
diff --git a/Engine/Plugins/Experimental/Avalanche/Source/AvalancheTag/Public/AvaTagSoftHandle.h b/Engine/Plugins/Experimental/Avalanche/Source/AvalancheTag/Public/AvaTagSoftHandle.h
--- a/Engine/Plugins/Experimental/Avalanche/Source/AvalancheTag/Public/AvaTagSoftHandle.h
+++ b/Engine/Plugins/Experimental/Avalanche/Source/AvalancheTag/Public/AvaTagSoftHandle.h
@@ -13,7 +13,7 @@ struct FAvaTagHandle;
  * Soft reference handle to a Tag in a particular Source Tag Collection
  * This should be used when needing to soft reference a particular FAvaTag.
  */
-USTRUCT()
+USTRUCT(BlueprintType)
 struct AVALANCHETAG_API FAvaTagSoftHandle
 {
 	GENERATED_BODY()
@@ -44,7 +44,7 @@ struct AVALANCHETAG_API FAvaTagSoftHandle
 		return !Source.IsNull() && TagId.IsValid();
 	}
 
-	UPROPERTY(EditAnywhere, Category="Tag")
+	UPROPERTY(EditAnywhere, Category = "Tag")
 	TSoftObjectPtr<const UAvaTagCollection> Source;
 
 	UPROPERTY()
