[Motion Design]
 - Mass renumbering of rundown pages

diff --git a/Engine/Plugins/Experimental/Avalanche/Source/AvalancheMedia/Private/Rundown/AvaRundown.cpp b/Engine/Plugins/Experimental/Avalanche/Source/AvalancheMedia/Private/Rundown/AvaRundown.cpp
--- a/Engine/Plugins/Experimental/Avalanche/Source/AvalancheMedia/Private/Rundown/AvaRundown.cpp
+++ b/Engine/Plugins/Experimental/Avalanche/Source/AvalancheMedia/Private/Rundown/AvaRundown.cpp
@@ -518,6 +518,23 @@ bool UAvaRundown::CanRemovePages(const TArray<int32>& InPageIds) const
 	return InPageIds.Num() > 0;
 }
 
+bool UAvaRundown::RenumberPageIds(const TArray<int32>& InPageIds, const FAvaRundownPageIdGeneratorParams& InIdParams)
+{
+	int32 CurrentId = InIdParams.ReferenceId;
+
+	for (const int32 PageId : InPageIds)
+	{
+		const int32 NewId = GenerateUniquePageId(CurrentId, InIdParams.Increment);
+
+		// If a page re-number fails, ignore it and continue on
+		RenumberPageId(PageId, NewId);
+		
+		CurrentId += InIdParams.Increment;
+	}
+
+	return true;
+}
+
 bool UAvaRundown::RenumberPageId(int32 InPageId, int32 InNewPageId)
 {
 	if (!CanRenumberPageId(InPageId, InNewPageId))
diff --git a/Engine/Plugins/Experimental/Avalanche/Source/AvalancheMedia/Public/Rundown/AvaRundown.h b/Engine/Plugins/Experimental/Avalanche/Source/AvalancheMedia/Public/Rundown/AvaRundown.h
--- a/Engine/Plugins/Experimental/Avalanche/Source/AvalancheMedia/Public/Rundown/AvaRundown.h
+++ b/Engine/Plugins/Experimental/Avalanche/Source/AvalancheMedia/Public/Rundown/AvaRundown.h
@@ -375,6 +375,7 @@ public:
 	int32 RemovePages(const TArray<int32>& InPageIds);
 	bool CanRemovePages(const TArray<int32>& InPageIds) const;
 
+	bool RenumberPageIds(const TArray<int32>& InPageIds, const FAvaRundownPageIdGeneratorParams& InIdParams);
 	bool RenumberPageId(int32 InPageId, int32 InNewPageId);
 	bool CanRenumberPageId(int32 InPageId) const;
 	bool CanRenumberPageId(int32 InPageId, int32 InNewPageId) const;
diff --git a/Engine/Plugins/Experimental/Avalanche/Source/AvalancheMediaEditor/AvalancheMediaEditor.Build.cs b/Engine/Plugins/Experimental/Avalanche/Source/AvalancheMediaEditor/AvalancheMediaEditor.Build.cs
--- a/Engine/Plugins/Experimental/Avalanche/Source/AvalancheMediaEditor/AvalancheMediaEditor.Build.cs
+++ b/Engine/Plugins/Experimental/Avalanche/Source/AvalancheMediaEditor/AvalancheMediaEditor.Build.cs
@@ -53,6 +53,7 @@ public class AvalancheMediaEditor : ModuleRules
 				"Slate",
 				"SlateCore",
 				"ToolMenus",
+				"ToolWidgets",
 				"UnrealEd",
 				"XmlSerialization",
 			}
diff --git a/Engine/Plugins/Experimental/Avalanche/Source/AvalancheMediaEditor/Private/Rundown/AvaRundownCommands.cpp b/Engine/Plugins/Experimental/Avalanche/Source/AvalancheMediaEditor/Private/Rundown/AvaRundownCommands.cpp
--- a/Engine/Plugins/Experimental/Avalanche/Source/AvalancheMediaEditor/Private/Rundown/AvaRundownCommands.cpp
+++ b/Engine/Plugins/Experimental/Avalanche/Source/AvalancheMediaEditor/Private/Rundown/AvaRundownCommands.cpp
@@ -31,8 +31,8 @@ void FAvaRundownCommands::RegisterCommands()
 		, FInputChord(EKeys::Delete));
 
 	UI_COMMAND(RenumberPage
-		, "Renumber Page"
-		, "Renumbers the Id of the Selected Page"
+		, "Renumber Pages"
+		, "Renumbers the Id of the Selected Pages"
 		, EUserInterfaceActionType::Button
 		, FInputChord(EKeys::F3));
 	
diff --git a/Engine/Plugins/Experimental/Avalanche/Source/AvalancheMediaEditor/Private/Rundown/Pages/Slate/SAvaRundownInstancedPageList.cpp b/Engine/Plugins/Experimental/Avalanche/Source/AvalancheMediaEditor/Private/Rundown/Pages/Slate/SAvaRundownInstancedPageList.cpp
--- a/Engine/Plugins/Experimental/Avalanche/Source/AvalancheMediaEditor/Private/Rundown/Pages/Slate/SAvaRundownInstancedPageList.cpp
+++ b/Engine/Plugins/Experimental/Avalanche/Source/AvalancheMediaEditor/Private/Rundown/Pages/Slate/SAvaRundownInstancedPageList.cpp
@@ -216,8 +216,8 @@ void SAvaRundownInstancedPageList::BindCommands()
 			FCanExecuteAction::CreateSP(this, &SAvaRundownPageList::CanRemoveSelectedPages));
 
 		CommandList->MapAction(RundownCommands.RenumberPage,
-			FExecuteAction::CreateSP(this, &SAvaRundownPageList::RenumberSelectedPage),
-			FCanExecuteAction::CreateSP(this, &SAvaRundownPageList::CanRenumberSelectedPage));
+			FExecuteAction::CreateSP(this, &SAvaRundownPageList::RenumberSelectedPages),
+			FCanExecuteAction::CreateSP(this, &SAvaRundownPageList::CanRenumberSelectedPages));
 
 		CommandList->MapAction(RundownCommands.ReimportPage,
 			FExecuteAction::CreateSP(this, &SAvaRundownPageList::ReimportSelectedPage),
diff --git a/Engine/Plugins/Experimental/Avalanche/Source/AvalancheMediaEditor/Private/Rundown/Pages/Slate/SAvaRundownPageList.cpp b/Engine/Plugins/Experimental/Avalanche/Source/AvalancheMediaEditor/Private/Rundown/Pages/Slate/SAvaRundownPageList.cpp
--- a/Engine/Plugins/Experimental/Avalanche/Source/AvalancheMediaEditor/Private/Rundown/Pages/Slate/SAvaRundownPageList.cpp
+++ b/Engine/Plugins/Experimental/Avalanche/Source/AvalancheMediaEditor/Private/Rundown/Pages/Slate/SAvaRundownPageList.cpp
@@ -5,6 +5,7 @@
 #include "AvaAssetTags.h"
 #include "DragAndDrop/AssetDragDropOp.h"
 #include "Editor/EditorWidgets/Public/SAssetSearchBox.h"
+#include "Framework/Application/SlateApplication.h"
 #include "Framework/Commands/GenericCommands.h"
 #include "Framework/MultiBox/MultiBoxBuilder.h"
 #include "HAL/PlatformApplicationMisc.h"
@@ -22,6 +23,7 @@
 #include "Rundown/Pages/AvaRundownPageContextMenu.h"
 #include "Rundown/Pages/PageViews/AvaRundownPageViewImpl.h"
 #include "Rundown/Pages/Slate/SAvaRundownPageViewRow.h"
+#include "SAvaRundownRenumberPages.h"
 #include "ScopedTransaction.h"
 #include "Widgets/SBoxPanel.h"
 #include "Widgets/SNullWidget.h"
@@ -227,7 +229,6 @@ void SAvaRundownPageList::OnPageEvent(const TArray<int32>& InSelectedPageIds, UE
 
 void SAvaRundownPageList::OnPageItemActionRequested(const TArray<int32>& InPageIds, IAvaRundownPageView::FOnPageAction& (IAvaRundownPageView::* InFunc)())
 {
-	//Only allow one page id to be renamed at a time
 	const TSet<int32> PageIdSet(InPageIds);
 	for (const FAvaRundownPageViewPtr& PageView : PageViews)
 	{
@@ -610,24 +611,57 @@ void SAvaRundownPageList::RenameSelectedPage()
 	}
 }
 
-bool SAvaRundownPageList::CanRenumberSelectedPage() const
+bool SAvaRundownPageList::CanRenumberSelectedPages() const
 {
-	if (SelectedPageIds.Num() == 1)
+	const int32 SelectedPageCount = SelectedPageIds.Num();
+	if (SelectedPageCount == 0)
 	{
-		if (const UAvaRundown* Rundown = GetRundown(); IsValid(Rundown))
+		return false;
+	}
+
+	const UAvaRundown* const Rundown = GetRundown();
+	if (!IsValid(Rundown))
+	{
+		return false;
+	}
+
+	for (const int32 PageId : SelectedPageIds)
+	{
+		if (Rundown->CanRenumberPageId(PageId))
 		{
-			return Rundown->CanRenumberPageId(SelectedPageIds[0]);
+			return true;
 		}
 	}
+
 	return false;
 }
 
-void SAvaRundownPageList::RenumberSelectedPage()
+void SAvaRundownPageList::RenumberSelectedPages()
 {
-	if (CanRenumberSelectedPage())
+	const int32 SelectedCount = SelectedPageIds.Num();
+	if (SelectedCount == 0)
+	{
+		return;
+	}
+	
+	if (SelectedCount == 1)
 	{
+		// Begin re-numbering inline
 		OnPageEvent(SelectedPageIds, UE::AvaRundown::EPageEvent::RenumberRequest);
 	}
+	else
+	{
+		// Bring up dialog for re-numbering multiple selection
+		const TSharedPtr<SAvaRundownRenumberPages> RenumberPagesDialog = SNew(SAvaRundownRenumberPages)
+			.OnAccept_Lambda([this](const int32 InBaseNumber, const int32 InIncrement)
+				{
+					if (UAvaRundown* const Rundown = GetRundown())
+					{
+						Rundown->RenumberPageIds(SelectedPageIds, FAvaRundownPageIdGeneratorParams(InBaseNumber, InIncrement));
+					}
+				});
+		FSlateApplication::Get().AddModalWindow(RenumberPagesDialog.ToSharedRef(), SharedThis(this));
+	}
 }
 
 bool SAvaRundownPageList::CanReimportSelectedPage() const
diff --git a/Engine/Plugins/Experimental/Avalanche/Source/AvalancheMediaEditor/Private/Rundown/Pages/Slate/SAvaRundownPageList.h b/Engine/Plugins/Experimental/Avalanche/Source/AvalancheMediaEditor/Private/Rundown/Pages/Slate/SAvaRundownPageList.h
--- a/Engine/Plugins/Experimental/Avalanche/Source/AvalancheMediaEditor/Private/Rundown/Pages/Slate/SAvaRundownPageList.h
+++ b/Engine/Plugins/Experimental/Avalanche/Source/AvalancheMediaEditor/Private/Rundown/Pages/Slate/SAvaRundownPageList.h
@@ -105,8 +105,8 @@ public:
 	bool CanRenameSelectedPage() const;
 	void RenameSelectedPage();
 
-	bool CanRenumberSelectedPage() const;
-	void RenumberSelectedPage();
+	bool CanRenumberSelectedPages() const;
+	void RenumberSelectedPages();
 
 	bool CanReimportSelectedPage() const;
 	void ReimportSelectedPage() const;
diff --git a/Engine/Plugins/Experimental/Avalanche/Source/AvalancheMediaEditor/Private/Rundown/Pages/Slate/SAvaRundownRenumberPages.cpp b/Engine/Plugins/Experimental/Avalanche/Source/AvalancheMediaEditor/Private/Rundown/Pages/Slate/SAvaRundownRenumberPages.cpp
new file mode 100644
--- /dev/null
+++ b/Engine/Plugins/Experimental/Avalanche/Source/AvalancheMediaEditor/Private/Rundown/Pages/Slate/SAvaRundownRenumberPages.cpp
@@ -0,0 +1,152 @@
+// Copyright Epic Games, Inc. All Rights Reserved.
+
+#include "SAvaRundownRenumberPages.h"
+#include "Framework/Application/SlateApplication.h"
+#include "Input/Reply.h"
+#include "SPrimaryButton.h"
+#include "Widgets/Images/SImage.h"
+#include "Widgets/Input/SButton.h"
+#include "Widgets/Input/SNumericEntryBox.h"
+#include "Widgets/Layout/SBox.h"
+#include "Widgets/Layout/SUniformGridPanel.h"
+#include "Widgets/SBoxPanel.h"
+#include "Widgets/Text/STextBlock.h"
+
+#define LOCTEXT_NAMESPACE "SAvaRundownRenumberPages"
+
+void SAvaRundownRenumberPages::Construct(const FArguments& InArgs)
+{
+	OnAccept = InArgs._OnAccept;
+
+	SWindow::Construct(SWindow::FArguments()
+		.Title(LOCTEXT("DialogTitle", "Renumber Pages"))
+		.AutoCenter(EAutoCenter::PrimaryWorkArea)
+		.SizingRule(ESizingRule::Autosized)
+		.SupportsMinimize(false)
+		.SupportsMaximize(false));
+	
+	SetContent(
+		SNew(SVerticalBox)
+		+ SVerticalBox::Slot()
+		.AutoHeight()
+		.Padding(10.f)
+		[
+			ConstructBaseNumberWidget()
+		]
+		+ SVerticalBox::Slot()
+		.AutoHeight()
+		.Padding(10.f, 0.f, 10.f, 10.f)
+		[
+			ConstructIncrementWidget()
+		]
+		+ SVerticalBox::Slot()
+		.AutoHeight()
+		.Padding(10.f, 5.f, 10.f, 10.f)
+		.HAlign(HAlign_Right)
+		[
+			SNew(SUniformGridPanel)
+			.SlotPadding(FMargin(5.f, 0.f))
+			+ SUniformGridPanel::Slot(0, 0)
+			[
+				SNew(SPrimaryButton)
+				.Text(LOCTEXT("OkButtonText", "OK"))
+				.OnClicked(this, &SAvaRundownRenumberPages::HandleAcceptClick)
+			]
+			+ SUniformGridPanel::Slot(1, 0)
+			[
+				SNew(SButton)
+				.HAlign(HAlign_Center)
+				.Text(LOCTEXT("CancelButtonText", "Cancel"))
+				.OnClicked(this, &SAvaRundownRenumberPages::HandleCancelClick)
+			]
+		]
+	);
+}
+
+TSharedRef<SWidget> SAvaRundownRenumberPages::ConstructBaseNumberWidget()
+{
+	return SNew(SHorizontalBox)
+		+ SHorizontalBox::Slot()
+		.FillWidth(0.5f)
+		.HAlign(HAlign_Right)
+		.VAlign(VAlign_Center)
+		[
+			SNew(STextBlock)
+			.Text(LOCTEXT("BaseNumberLabel", "Base Number:"))
+		]
+		+ SHorizontalBox::Slot()
+		.FillWidth(0.5f)
+		.Padding(10.f, 0.f, 0.f, 0.f)
+		[
+			SNew(SNumericEntryBox<int32>)
+			.AllowSpin(true)
+			.MinValue(0)
+			.MinSliderValue(0)
+			.MaxSliderValue(MaxSliderValue)
+			.Value(this, &SAvaRundownRenumberPages::GetBaseNumber)
+			.OnValueChanged(this, &SAvaRundownRenumberPages::HandleBaseNumberChanged)
+		];
+}
+
+TSharedRef<SWidget> SAvaRundownRenumberPages::ConstructIncrementWidget()
+{
+	return SNew(SHorizontalBox)
+		+ SHorizontalBox::Slot()
+		.FillWidth(0.5f)
+		.HAlign(HAlign_Right)
+		.VAlign(VAlign_Center)
+		[
+			SNew(STextBlock)
+			.Text(LOCTEXT("IncrementLabel", "Increment:"))
+		]
+		+ SHorizontalBox::Slot()
+		.FillWidth(0.5f)
+		.Padding(10.f, 0.f, 0.f, 0.f)
+		[
+			SNew(SNumericEntryBox<int32>)
+			.AllowSpin(true)
+			.MinValue(1)
+			.MinSliderValue(1)
+			.MaxSliderValue(MaxSliderValue)
+			.Value(this, &SAvaRundownRenumberPages::GetIncrement)
+			.OnValueChanged(this, &SAvaRundownRenumberPages::HandleIncrementChanged)
+		];
+}
+
+TOptional<int32> SAvaRundownRenumberPages::GetBaseNumber() const
+{
+	return BaseNumber;
+}
+
+void SAvaRundownRenumberPages::HandleBaseNumberChanged(const int32 InNewValue)
+{
+	BaseNumber = InNewValue;
+}
+
+TOptional<int32> SAvaRundownRenumberPages::GetIncrement() const
+{
+	return Increment;
+}
+
+void SAvaRundownRenumberPages::HandleIncrementChanged(const int32 InNewValue)
+{
+	Increment = InNewValue;
+}
+
+FReply SAvaRundownRenumberPages::HandleAcceptClick()
+{
+	OnAccept.ExecuteIfBound(BaseNumber, Increment);
+
+	RequestDestroyWindow();
+
+	return FReply::Handled();
+}
+
+FReply SAvaRundownRenumberPages::HandleCancelClick()
+{
+	RequestDestroyWindow();
+
+	return FReply::Handled();
+}
+
+#undef LOCTEXT_NAMESPACE
diff --git a/Engine/Plugins/Experimental/Avalanche/Source/AvalancheMediaEditor/Private/Rundown/Pages/Slate/SAvaRundownRenumberPages.h b/Engine/Plugins/Experimental/Avalanche/Source/AvalancheMediaEditor/Private/Rundown/Pages/Slate/SAvaRundownRenumberPages.h
new file mode 100644
--- /dev/null
+++ b/Engine/Plugins/Experimental/Avalanche/Source/AvalancheMediaEditor/Private/Rundown/Pages/Slate/SAvaRundownRenumberPages.h
@@ -0,0 +1,41 @@
+// Copyright Epic Games, Inc. All Rights Reserved.
+
+#pragma once
+
+#include "Misc/Optional.h"
+#include "Widgets/SWindow.h"
+
+class FReply;
+
+class SAvaRundownRenumberPages : public SWindow
+{
+public:
+	DECLARE_DELEGATE_TwoParams(FOnRundownRenumberPagesResult, const int32 /*BaseNumber*/, const int32 /*Increment*/)
+
+	SLATE_BEGIN_ARGS(SAvaRundownRenumberPages)
+	{}
+		SLATE_EVENT(FOnRundownRenumberPagesResult, OnAccept)
+	SLATE_END_ARGS()
+
+	void Construct(const FArguments& InArgs);
+
+protected:
+	static constexpr int32 MaxSliderValue = 1000;
+
+	TSharedRef<SWidget> ConstructBaseNumberWidget();
+	TSharedRef<SWidget> ConstructIncrementWidget();
+
+	TOptional<int32> GetBaseNumber() const;
+	void HandleBaseNumberChanged(const int32 InNewValue);
+
+	TOptional<int32> GetIncrement() const;
+	void HandleIncrementChanged(const int32 InNewValue);
+
+	FReply HandleAcceptClick();
+	FReply HandleCancelClick();
+
+	FOnRundownRenumberPagesResult OnAccept;
+
+	int32 BaseNumber = 0;
+	int32 Increment = 1;
+};
diff --git a/Engine/Plugins/Experimental/Avalanche/Source/AvalancheMediaEditor/Private/Rundown/Pages/Slate/SAvaRundownTemplatePageList.cpp b/Engine/Plugins/Experimental/Avalanche/Source/AvalancheMediaEditor/Private/Rundown/Pages/Slate/SAvaRundownTemplatePageList.cpp
--- a/Engine/Plugins/Experimental/Avalanche/Source/AvalancheMediaEditor/Private/Rundown/Pages/Slate/SAvaRundownTemplatePageList.cpp
+++ b/Engine/Plugins/Experimental/Avalanche/Source/AvalancheMediaEditor/Private/Rundown/Pages/Slate/SAvaRundownTemplatePageList.cpp
@@ -139,8 +139,8 @@ void SAvaRundownTemplatePageList::BindCommands()
 			FCanExecuteAction::CreateSP(this, &SAvaRundownPageList::CanRemoveSelectedPages));
 
 		CommandList->MapAction(RundownCommands.RenumberPage,
-			FExecuteAction::CreateSP(this, &SAvaRundownPageList::RenumberSelectedPage),
-			FCanExecuteAction::CreateSP(this, &SAvaRundownPageList::CanRenumberSelectedPage));
+			FExecuteAction::CreateSP(this, &SAvaRundownPageList::RenumberSelectedPages),
+			FCanExecuteAction::CreateSP(this, &SAvaRundownPageList::CanRenumberSelectedPages));
 
 		CommandList->MapAction(RundownCommands.ReimportPage,
 			FExecuteAction::CreateSP(this, &SAvaRundownPageList::ReimportSelectedPage),
