Motion Design: Added Blueprint conversion support for Tag Handles and way to get Active Sequence Players

diff --git a/Engine/Plugins/Experimental/Avalanche/Source/AvalancheSequence/Private/AvaSequenceLibrary.cpp b/Engine/Plugins/Experimental/Avalanche/Source/AvalancheSequence/Private/AvaSequenceLibrary.cpp
new file mode 100644
--- /dev/null
+++ b/Engine/Plugins/Experimental/Avalanche/Source/AvalancheSequence/Private/AvaSequenceLibrary.cpp
@@ -0,0 +1,46 @@
+// Copyright Epic Games, Inc. All Rights Reserved.
+
+#include "AvaSequenceLibrary.h"
+#include "AvaSequencePlaybackObject.h"
+#include "AvaSequenceSubsystem.h"
+#include "Engine/Engine.h"
+#include "Engine/Level.h"
+#include "Engine/World.h"
+
+TScriptInterface<IAvaSequencePlaybackObject> UAvaSequenceLibrary::GetPlaybackObject(const UObject* InWorldContextObject)
+{
+	if (!InWorldContextObject || !GEngine)
+	{
+		return nullptr;
+	}
+
+	UWorld* World = GEngine->GetWorldFromContextObject(InWorldContextObject, EGetWorldErrorMode::LogAndReturnNull);
+	if (!World)
+	{
+		return nullptr;
+	}
+
+	UAvaSequenceSubsystem* SequenceSubsystem = World->GetSubsystem<UAvaSequenceSubsystem>();
+	if (!SequenceSubsystem)
+	{
+		return nullptr;
+	}
+
+	ULevel* Level = InWorldContextObject->GetTypedOuter<ULevel>();
+	if (!Level)
+	{
+		Level = World->PersistentLevel;
+		if (!Level)
+		{
+			return nullptr;
+		}
+	}
+
+	IAvaSequencePlaybackObject* PlaybackObject = SequenceSubsystem->FindPlaybackObject(Level);
+	if (!PlaybackObject)
+	{
+		return nullptr;
+	}
+
+	return TScriptInterface<IAvaSequencePlaybackObject>(PlaybackObject->ToUObject());	
+}
diff --git a/Engine/Plugins/Experimental/Avalanche/Source/AvalancheSequence/Public/AvaSequenceLibrary.h b/Engine/Plugins/Experimental/Avalanche/Source/AvalancheSequence/Public/AvaSequenceLibrary.h
new file mode 100644
--- /dev/null
+++ b/Engine/Plugins/Experimental/Avalanche/Source/AvalancheSequence/Public/AvaSequenceLibrary.h
@@ -0,0 +1,19 @@
+// Copyright Epic Games, Inc. All Rights Reserved.
+
+#pragma once
+
+#include "Kismet/BlueprintFunctionLibrary.h"
+#include "AvaSequenceLibrary.generated.h"
+
+class IAvaSequencePlaybackObject;
+template<typename InInterfaceType> class TScriptInterface;
+
+UCLASS(MinimalAPI, DisplayName = "Motion Design Sequence Library", meta=(ScriptName = "MotionDesignSequenceLibrary"))
+class UAvaSequenceLibrary : public UBlueprintFunctionLibrary
+{
+	GENERATED_BODY()
+
+public:
+	UFUNCTION(BlueprintCallable, Category = "Motion Design|Sequence", meta=(WorldContext = "InWorldContextObject"))
+	static AVALANCHESEQUENCE_API TScriptInterface<IAvaSequencePlaybackObject> GetPlaybackObject(const UObject* InWorldContextObject);
+};
diff --git a/Engine/Plugins/Experimental/Avalanche/Source/AvalancheSequence/Public/AvaSequencePlaybackObject.h b/Engine/Plugins/Experimental/Avalanche/Source/AvalancheSequence/Public/AvaSequencePlaybackObject.h
--- a/Engine/Plugins/Experimental/Avalanche/Source/AvalancheSequence/Public/AvaSequencePlaybackObject.h
+++ b/Engine/Plugins/Experimental/Avalanche/Source/AvalancheSequence/Public/AvaSequencePlaybackObject.h
@@ -155,5 +155,6 @@ public:
 	virtual TArray<UAvaSequencePlayer*> GetSequencePlayersByTag(const FAvaTag& InTag, bool bInExactMatch) const = 0;
 
 	/** Retrieves all Active Sequence Players */
+	UFUNCTION(BlueprintCallable, DisplayName = "Get Active Sequence Players", Category = "Playback")
 	virtual TArray<UAvaSequencePlayer*> GetAllSequencePlayers() const = 0;
 };
diff --git a/Engine/Plugins/Experimental/Avalanche/Source/AvalancheTag/AvalancheTag.Build.cs b/Engine/Plugins/Experimental/Avalanche/Source/AvalancheTag/AvalancheTag.Build.cs
--- a/Engine/Plugins/Experimental/Avalanche/Source/AvalancheTag/AvalancheTag.Build.cs
+++ b/Engine/Plugins/Experimental/Avalanche/Source/AvalancheTag/AvalancheTag.Build.cs
@@ -13,6 +13,7 @@ public class AvalancheTag : ModuleRules
             {
                 "Core",
                 "CoreUObject",
+                "Engine",
             }
         );
     }
diff --git a/Engine/Plugins/Experimental/Avalanche/Source/AvalancheTag/Private/AvaTagLibrary.cpp b/Engine/Plugins/Experimental/Avalanche/Source/AvalancheTag/Private/AvaTagLibrary.cpp
new file mode 100644
--- /dev/null
+++ b/Engine/Plugins/Experimental/Avalanche/Source/AvalancheTag/Private/AvaTagLibrary.cpp
@@ -0,0 +1,25 @@
+// Copyright Epic Games, Inc. All Rights Reserved.
+
+#include "AvaTagLibrary.h"
+#include "AvaTagHandle.h"
+#include "AvaTagHandleContainer.h"
+#include "AvaTagSoftHandle.h"
+
+FAvaTag UAvaTagLibrary::ResolveTagHandle(const FAvaTagHandle& InTagHandle)
+{
+	if (const FAvaTag* Tag = InTagHandle.GetTag())
+	{
+		return *Tag;
+	}
+	return FAvaTag();
+}
+
+TArray<FAvaTag> UAvaTagLibrary::ResolveTagHandles(const FAvaTagHandleContainer& InTagHandleContainer)
+{
+	return InTagHandleContainer.ResolveTags();
+}
+
+FAvaTagHandle UAvaTagLibrary::ResolveTagSoftHandle(const FAvaTagSoftHandle& InTagSoftHandle)
+{
+	return InTagSoftHandle.MakeTagHandle();
+}
diff --git a/Engine/Plugins/Experimental/Avalanche/Source/AvalancheTag/Public/AvaTagLibrary.h b/Engine/Plugins/Experimental/Avalanche/Source/AvalancheTag/Public/AvaTagLibrary.h
new file mode 100644
--- /dev/null
+++ b/Engine/Plugins/Experimental/Avalanche/Source/AvalancheTag/Public/AvaTagLibrary.h
@@ -0,0 +1,28 @@
+// Copyright Epic Games, Inc. All Rights Reserved.
+
+#pragma once
+
+#include "Containers/ContainersFwd.h"
+#include "Kismet/BlueprintFunctionLibrary.h"
+#include "AvaTagLibrary.generated.h"
+
+struct FAvaTag;
+struct FAvaTagHandle;
+struct FAvaTagHandleContainer;
+struct FAvaTagSoftHandle;
+
+UCLASS(MinimalAPI, DisplayName = "Motion Design Tag Library", meta=(ScriptName = "MotionDesignTagLibrary"))
+class UAvaTagLibrary : public UBlueprintFunctionLibrary
+{
+	GENERATED_BODY()
+
+public:
+	UFUNCTION(BlueprintPure, Category = "Motion Design|Tag", meta=(CompactNodeTitle = "->", BlueprintAutocast))
+	static AVALANCHETAG_API FAvaTag ResolveTagHandle(const FAvaTagHandle& InTagHandle);
+
+	UFUNCTION(BlueprintPure, Category = "Motion Design|Tag", meta=(CompactNodeTitle = "->", BlueprintAutocast))
+	static AVALANCHETAG_API TArray<FAvaTag> ResolveTagHandles(const FAvaTagHandleContainer& InTagHandleContainer);
+
+	UFUNCTION(BlueprintPure, Category = "Motion Design|Tag", meta=(CompactNodeTitle = "->", BlueprintAutocast))
+	static AVALANCHETAG_API FAvaTagHandle ResolveTagSoftHandle(const FAvaTagSoftHandle& InTagSoftHandle);
+};
