Motion Design: Added Sequence Playback function to return whether there are any Active Sequence players

diff --git a/Engine/Plugins/Experimental/Avalanche/Source/AvalancheSequence/Private/AvaSequencePlaybackActor.cpp b/Engine/Plugins/Experimental/Avalanche/Source/AvalancheSequence/Private/AvaSequencePlaybackActor.cpp
--- a/Engine/Plugins/Experimental/Avalanche/Source/AvalancheSequence/Private/AvaSequencePlaybackActor.cpp
+++ b/Engine/Plugins/Experimental/Avalanche/Source/AvalancheSequence/Private/AvaSequencePlaybackActor.cpp
@@ -427,6 +427,11 @@ TArray<UAvaSequencePlayer*> AAvaSequencePlaybackActor::GetAllSequencePlayers() c
 	return SequencePlayers;
 }
 
+bool AAvaSequencePlaybackActor::HasActiveSequencePlayers() const
+{
+	return !ActiveSequencePlayers.IsEmpty();
+}
+
 TArray<UAvaSequencePlayer*> AAvaSequencePlaybackActor::GetSequencePlayersByLabels(const TArray<FName>& InSequenceLabels) const
 {
 	const TArray<UAvaSequence*> Sequences = GetSequencesByLabel(InSequenceLabels);
diff --git a/Engine/Plugins/Experimental/Avalanche/Source/AvalancheSequence/Private/AvaSequencePlaybackActor.h b/Engine/Plugins/Experimental/Avalanche/Source/AvalancheSequence/Private/AvaSequencePlaybackActor.h
--- a/Engine/Plugins/Experimental/Avalanche/Source/AvalancheSequence/Private/AvaSequencePlaybackActor.h
+++ b/Engine/Plugins/Experimental/Avalanche/Source/AvalancheSequence/Private/AvaSequencePlaybackActor.h
@@ -47,6 +47,7 @@ protected:
 	virtual TArray<UAvaSequencePlayer*> GetSequencePlayersByLabels(const TArray<FName>& InSequenceLabels) const override;
 	virtual TArray<UAvaSequencePlayer*> GetSequencePlayersByTag(const FAvaTag& InTag, bool bInExactMatch) const override;
 	virtual TArray<UAvaSequencePlayer*> GetAllSequencePlayers() const override;
+	virtual bool HasActiveSequencePlayers() const override;
 	virtual void UpdateCameraCut(const UE::MovieScene::FOnCameraCutUpdatedParams& InCameraCutParams) override;
 	virtual UObject* CreateDirectorInstance(IMovieScenePlayer& InPlayer, FMovieSceneSequenceID InSequenceID) override;
 	virtual FOnCameraCut& GetOnCameraCut() override { return OnCameraCut; }
diff --git a/Engine/Plugins/Experimental/Avalanche/Source/AvalancheSequence/Public/AvaSequencePlaybackObject.h b/Engine/Plugins/Experimental/Avalanche/Source/AvalancheSequence/Public/AvaSequencePlaybackObject.h
--- a/Engine/Plugins/Experimental/Avalanche/Source/AvalancheSequence/Public/AvaSequencePlaybackObject.h
+++ b/Engine/Plugins/Experimental/Avalanche/Source/AvalancheSequence/Public/AvaSequencePlaybackObject.h
@@ -157,4 +157,8 @@ public:
 	/** Retrieves all Active Sequence Players */
 	UFUNCTION(BlueprintCallable, DisplayName = "Get Active Sequence Players", Category = "Playback")
 	virtual TArray<UAvaSequencePlayer*> GetAllSequencePlayers() const = 0;
+
+	/** Returns true if there are any Active Sequence Players */
+	UFUNCTION(BlueprintCallable, DisplayName = "Has Active Sequence Players", Category = "Playback")
+	virtual bool HasActiveSequencePlayers() const = 0;
 };
