[Motion Design]
 - EaseCurvePreview widget: Add background color option
 - Add extra contrast to hover ease curve preset item

diff --git a/Engine/Plugins/Experimental/Avalanche/Source/AvalancheSequencer/Private/EaseCurveTool/Widgets/SAvaEaseCurvePresetGroupItem.cpp b/Engine/Plugins/Experimental/Avalanche/Source/AvalancheSequencer/Private/EaseCurveTool/Widgets/SAvaEaseCurvePresetGroupItem.cpp
--- a/Engine/Plugins/Experimental/Avalanche/Source/AvalancheSequencer/Private/EaseCurveTool/Widgets/SAvaEaseCurvePresetGroupItem.cpp
+++ b/Engine/Plugins/Experimental/Avalanche/Source/AvalancheSequencer/Private/EaseCurveTool/Widgets/SAvaEaseCurvePresetGroupItem.cpp
@@ -137,6 +137,7 @@ void SAvaEaseCurvePresetGroupItem::Construct(const FArguments& InArgs, const TSh
 						.CurveThickness(1.5f)
 						.Tangents(Preset->Tangents)
 						.CustomToolTip(true)
+						.BackgroundColor(FStyleColors::Dropdown.GetSpecifiedColor())
 						.UnderCurveColor(FStyleColors::SelectInactive.GetSpecifiedColor())
 						.DisplayRate(InArgs._DisplayRate)
 					]
diff --git a/Engine/Plugins/Experimental/Avalanche/Source/AvalancheSequencer/Private/EaseCurveTool/Widgets/SAvaEaseCurvePreview.cpp b/Engine/Plugins/Experimental/Avalanche/Source/AvalancheSequencer/Private/EaseCurveTool/Widgets/SAvaEaseCurvePreview.cpp
--- a/Engine/Plugins/Experimental/Avalanche/Source/AvalancheSequencer/Private/EaseCurveTool/Widgets/SAvaEaseCurvePreview.cpp
+++ b/Engine/Plugins/Experimental/Avalanche/Source/AvalancheSequencer/Private/EaseCurveTool/Widgets/SAvaEaseCurvePreview.cpp
@@ -20,6 +20,7 @@ void SAvaEaseCurvePreview::Construct(const FArguments& InArgs)
 	PreviewSize = InArgs._PreviewSize;
 	CanExpandPreview = InArgs._CanExpandPreview;
 	CurveThickness = InArgs._CurveThickness;
+	BackgroundColor = InArgs._BackgroundColor;
 	CurveColor = InArgs._CurveColor;
 	StraightColor = InArgs._StraightColor;
 	UnderCurveColor = InArgs._UnderCurveColor;
@@ -216,6 +217,13 @@ int32 SAvaEaseCurvePreview::OnPaint(const FPaintArgs& InArgs, const FGeometry& I
 	const FVector2D AdditionalAboveHeight(BelowZeroValue * PreviewSize, AboveOneValue * PreviewSize);
 	const FPaintGeometry AbovePaintGeometry = InAllottedGeometry.ToPaintGeometry(InAllottedGeometry.GetLocalSize(), FSlateLayoutTransform(AdditionalAboveHeight));
 
+	// Draw background color
+	if (BackgroundColor.IsSet())
+	{
+		FSlateDrawElement::MakeBox(OutDrawElements, LayerId++, InAllottedGeometry.ToPaintGeometry()
+			, FAppStyle::GetBrush(TEXT("WhiteTexture")), DrawEffects, BackgroundColor.GetValue());
+	}
+
 	// Draw 0 and 1 line
 	{
 		constexpr FLinearColor NormalLimitLineColor = FLinearColor(0.1f, 0.1f, 0.1f, 1.f);
@@ -232,7 +240,7 @@ int32 SAvaEaseCurvePreview::OnPaint(const FPaintArgs& InArgs, const FGeometry& I
 			ZeroValueLinePoints.Add(FVector2D(LocalWidth - OneFifthWidth, ZeroOutputY));
 			ZeroValueLinePoints.Add(FVector2D(LocalWidth, ZeroOutputY));
 			TArray<FLinearColor> ZeroValueLineColors;
-			FSlateDrawElement::MakeLines(OutDrawElements, InLayerId, AbovePaintGeometry
+			FSlateDrawElement::MakeLines(OutDrawElements, LayerId++, AbovePaintGeometry
 				, ZeroValueLinePoints, ValueLineColors, DrawEffects, FLinearColor::White, false, InAllottedGeometry.Scale);
 		}
 
@@ -244,7 +252,7 @@ int32 SAvaEaseCurvePreview::OnPaint(const FPaintArgs& InArgs, const FGeometry& I
 			OneValueLinePoints.Add(FVector2D(OneFifthWidth, OneOutputY));
 			OneValueLinePoints.Add(FVector2D(LocalWidth - OneFifthWidth, OneOutputY));
 			OneValueLinePoints.Add(FVector2D(LocalWidth, OneOutputY));
-			FSlateDrawElement::MakeLines(OutDrawElements, InLayerId, AbovePaintGeometry
+			FSlateDrawElement::MakeLines(OutDrawElements, LayerId++, AbovePaintGeometry
 				, OneValueLinePoints, ValueLineColors, DrawEffects, FLinearColor::White, false, InAllottedGeometry.Scale);
 		}
 	}
diff --git a/Engine/Plugins/Experimental/Avalanche/Source/AvalancheSequencer/Private/EaseCurveTool/Widgets/SAvaEaseCurvePreview.h b/Engine/Plugins/Experimental/Avalanche/Source/AvalancheSequencer/Private/EaseCurveTool/Widgets/SAvaEaseCurvePreview.h
--- a/Engine/Plugins/Experimental/Avalanche/Source/AvalancheSequencer/Private/EaseCurveTool/Widgets/SAvaEaseCurvePreview.h
+++ b/Engine/Plugins/Experimental/Avalanche/Source/AvalancheSequencer/Private/EaseCurveTool/Widgets/SAvaEaseCurvePreview.h
@@ -44,6 +44,7 @@ public:
 		SLATE_ARGUMENT(bool, CanExpandPreview)
 		SLATE_ARGUMENT(float, CurveThickness)
 	
+		SLATE_ARGUMENT(TOptional<FLinearColor>, BackgroundColor)
 		SLATE_ARGUMENT(FLinearColor, CurveColor)
 		SLATE_ARGUMENT(FLinearColor, StraightColor)
 		SLATE_ARGUMENT(TOptional<FLinearColor>, UnderCurveColor)
@@ -93,6 +94,7 @@ protected:
 	bool CanExpandPreview = false;
 	float CurveThickness = 1.f;
 
+	TOptional<FLinearColor> BackgroundColor;
 	FLinearColor CurveColor;
 	FLinearColor StraightColor;
 	TOptional<FLinearColor> UnderCurveColor;
