Motion Design: Fix crash when opening a sub-sequence in motion design sequencer.

diff --git a/Engine/Plugins/Experimental/Avalanche/Source/AvalancheSequencer/Private/AvaSequencer.cpp b/Engine/Plugins/Experimental/Avalanche/Source/AvalancheSequencer/Private/AvaSequencer.cpp
--- a/Engine/Plugins/Experimental/Avalanche/Source/AvalancheSequencer/Private/AvaSequencer.cpp
+++ b/Engine/Plugins/Experimental/Avalanche/Source/AvalancheSequencer/Private/AvaSequencer.cpp
@@ -230,6 +230,14 @@ void FAvaSequencer::EnsureSequencer()
 
 	Sequencer = Provider.GetExternalSequencer();
 
+	// External Implementation could call GetSequencer again (e.g. to get the underlying sequencer widget),
+	// so need to give priority to that call and initialize from there.
+	// If this is the case, SequencerWeak is now initialized/valid and should return early to avoid double init.
+	if (SequencerWeak.IsValid())
+	{
+		return;
+	}
+
 	if (Sequencer.IsValid())
 	{
 		checkf(Sequencer.GetSharedReferenceCount() > 1
