[Motion Design] Rundown Page Detail: fixing cascading ui refresh causing entity values to be reset (and be uneditable).

diff --git a/Engine/Plugins/Experimental/Avalanche/Source/AvalancheMediaEditor/Private/Rundown/DetailsView/SAvaRundownPageDetails.cpp b/Engine/Plugins/Experimental/Avalanche/Source/AvalancheMediaEditor/Private/Rundown/DetailsView/SAvaRundownPageDetails.cpp
--- a/Engine/Plugins/Experimental/Avalanche/Source/AvalancheMediaEditor/Private/Rundown/DetailsView/SAvaRundownPageDetails.cpp
+++ b/Engine/Plugins/Experimental/Avalanche/Source/AvalancheMediaEditor/Private/Rundown/DetailsView/SAvaRundownPageDetails.cpp
@@ -188,20 +188,16 @@ SAvaRundownPageDetails::~SAvaRundownPageDetails()
 	if (const TSharedPtr<FAvaRundownEditor> RundownEditor = RundownEditorWeak.Pin())
 	{
 		RundownEditor->GetOnPageEvent().RemoveAll(this);
-	}
-	if (IAvaMediaModule::IsModuleLoaded())
-	{
-		IAvaMediaModule::Get().GetManagedInstanceCache().OnEntryInvalidated.RemoveAll(this);
-	}
-
-	if (const TSharedPtr<FAvaRundownEditor> RundownEditor = RundownEditorWeak.Pin())
-	{
 		UAvaRundown* const Rundown = RundownEditor->GetRundown();
 		if (IsValid(Rundown))
 		{
 			Rundown->GetOnPagesChanged().RemoveAll(this);
 		}
 	}
+	if (IAvaMediaModule::IsModuleLoaded())
+	{
+		IAvaMediaModule::Get().GetManagedInstanceCache().OnEntryInvalidated.RemoveAll(this);
+	}
 }
 
 void SAvaRundownPageDetails::OnPageEvent(const TArray<int32>& InSelectedPageIds, UE::AvaRundown::EPageEvent InPageEvent)
@@ -235,20 +231,10 @@ void SAvaRundownPageDetails::OnManagedInstanceCacheEntryInvalidated(const FSoftO
 				{
 					if (SelectedPage.GetAssetPath(Rundown) == InAssetPath)
 					{
-						bRefreshSelectedPageQueued = true;
 						// Queue a refresh on next tick.
 						// We don't want to refresh immediately to avoid issues with
 						// cascading events within the managed instance cache.
-						TWeakPtr<SWidget> ThisWeak(AsShared());
-						AsyncTask(ENamedThreads::GameThread, [ThisWeak]()
-							{
-								if (const TSharedPtr<SWidget> ThisWidget = ThisWeak.Pin())
-								{
-									SAvaRundownPageDetails* AvaPageDetails = static_cast<SAvaRundownPageDetails*>(ThisWidget.Get());
-									AvaPageDetails->RefreshSelectedPage();
-									AvaPageDetails->bRefreshSelectedPageQueued = false;
-								}
-							});
+						QueueUpdateAndRefreshSelectedPage();
 					}
 				}
 			}
@@ -310,16 +296,46 @@ FAvaRundownPage& SAvaRundownPageDetails::GetMutableSelectedPage() const
 	return FAvaRundownPage::NullPage;
 }
 
-void SAvaRundownPageDetails::RefreshSelectedPage()
+void SAvaRundownPageDetails::QueueRefreshSelectedPage()
 {
-	const FAvaRundownPage& SelectedPage = GetSelectedPage();
+	if (bRefreshSelectedPageQueued)
+	{
+		return;
+	}
+	bRefreshSelectedPageQueued = true;
+	
+	FTSTicker::GetCoreTicker().AddTicker(FTickerDelegate::CreateSPLambda(this, [this](float)
+	{
+		const FAvaRundownPage& SelectedPage = GetSelectedPage();
+		if (SelectedPage.IsValidPage())
+		{
+			RemoteControlProps->Refresh({SelectedPage.GetPageId()});
+			RCControllerPanel->Refresh({SelectedPage.GetPageId()});
+		}
+		bRefreshSelectedPageQueued = false;
+		return false;
+	}));
+}
 
-	if (SelectedPage.IsValidPage())
+void SAvaRundownPageDetails::QueueUpdateAndRefreshSelectedPage()
+{
+	if (bUpdateAndRefreshSelectedPageQueued)
 	{
-		OnPageSelectionChanged({SelectedPage.GetPageId()});
-		RemoteControlProps->UpdateDefaultValuesAndRefresh({SelectedPage.GetPageId()});
-		RCControllerPanel->Refresh({SelectedPage.GetPageId()});
+		return;
 	}
+	bUpdateAndRefreshSelectedPageQueued = true;
+
+	FTSTicker::GetCoreTicker().AddTicker(FTickerDelegate::CreateSPLambda(this, [this](float)
+	{
+		const FAvaRundownPage& SelectedPage = GetSelectedPage();
+		if (SelectedPage.IsValidPage())
+		{
+			RemoteControlProps->UpdateDefaultValuesAndRefresh({SelectedPage.GetPageId()});
+			RemoteControlProps->Refresh({SelectedPage.GetPageId()});
+		}
+		bUpdateAndRefreshSelectedPageQueued = true;
+		return false;
+	}));
 }
 
 bool SAvaRundownPageDetails::HasSelectedPage() const
@@ -454,9 +470,10 @@ void SAvaRundownPageDetails::OnRundownPagesChanged(const UAvaRundown* InRundown,
 {
 	// Refreshing the page while the mouse is captured will result in losing the capture
 	// and ending any drag event that is actively changing the value.
-	if (!FSlateApplication::Get().GetMouseCaptureWindow())
+	if (!FSlateApplication::Get().GetMouseCaptureWindow() && InPage.GetPageId() == ActivePageId)
 	{
-		RefreshSelectedPage();
+		// Queue a refresh on next tick to avoid issues with cascading events.
+		QueueRefreshSelectedPage();
 	}
 }
 
diff --git a/Engine/Plugins/Experimental/Avalanche/Source/AvalancheMediaEditor/Private/Rundown/DetailsView/SAvaRundownPageDetails.h b/Engine/Plugins/Experimental/Avalanche/Source/AvalancheMediaEditor/Private/Rundown/DetailsView/SAvaRundownPageDetails.h
--- a/Engine/Plugins/Experimental/Avalanche/Source/AvalancheMediaEditor/Private/Rundown/DetailsView/SAvaRundownPageDetails.h
+++ b/Engine/Plugins/Experimental/Avalanche/Source/AvalancheMediaEditor/Private/Rundown/DetailsView/SAvaRundownPageDetails.h
@@ -42,7 +42,8 @@ protected:
 	const FAvaRundownPage& GetSelectedPage() const;
 	FAvaRundownPage& GetMutableSelectedPage() const;
 
-	void RefreshSelectedPage();
+	void QueueRefreshSelectedPage();
+	void QueueUpdateAndRefreshSelectedPage();
 
 	bool HasSelectedPage() const;
 
@@ -68,6 +69,7 @@ private:
 	TSharedPtr<SAvaRundownRCControllerPanel> RCControllerPanel;
 
 	bool bRefreshSelectedPageQueued = false;
+	bool bUpdateAndRefreshSelectedPageQueued = false;
 
 	int32 ActivePageId;
 };
