[Motion Design]
 - Scene Rig: Allow modules to register supported actor classes

diff --git a/Engine/Plugins/Experimental/Avalanche/Source/AvalancheEditor/Private/AvaEditorModule.cpp b/Engine/Plugins/Experimental/Avalanche/Source/AvalancheEditor/Private/AvaEditorModule.cpp
--- a/Engine/Plugins/Experimental/Avalanche/Source/AvalancheEditor/Private/AvaEditorModule.cpp
+++ b/Engine/Plugins/Experimental/Avalanche/Source/AvalancheEditor/Private/AvaEditorModule.cpp
@@ -5,16 +5,26 @@
 #include "AvaEditorActorUtils.h"
 #include "AvaEditorCommands.h"
 #include "AvaEditorIntegration.h"
+#include "AvaSceneRigSubsystem.h"
 #include "AvaShapeActor.h"
+#include "Camera/CameraActor.h"
+#include "CineCameraActor.h"
 #include "ColorPicker/AvaViewportColorPickerActorClassRegistry.h"
 #include "ColorPicker/IAvaViewportColorPickerAdapter.h"
 #include "Components/LightComponentBase.h"
+#include "Containers/Set.h"
 #include "DynamicMeshes/AvaShapeDynMeshBase.h"
 #include "Editor/UnrealEdEngine.h"
+#include "Engine/DirectionalLight.h"
 #include "Engine/Light.h"
+#include "Engine/PointLight.h"
+#include "Engine/PostProcessVolume.h"
+#include "Engine/RectLight.h"
 #include "Engine/SkyLight.h"
+#include "Engine/SpotLight.h"
 #include "Engine/Texture2D.h"
 #include "Framework/Application/SlateApplication.h"
+#include "Framework/AvaNullActor.h"
 #include "IAssetTools.h"
 #include "IAvaOutlinerModule.h"
 #include "Icon/AvaOutlinerObjectIconCustomization.h"
@@ -24,6 +34,7 @@
 #include "SVGImporter/AvaOutlinerSVGActorContextMenu.h"
 #include "Styling/SlateIconFinder.h"
 #include "UnrealEdGlobals.h"
+#include "Viewport/AvaCineCameraActor.h"
 #include "Viewport/AvaViewportQualitySettings.h"
 
 // Details View
@@ -105,6 +116,8 @@ void FAvaEditorModule::StartupModule()
 	// Note: ASkylight Does not extend from ALight
 	FAvaViewportColorPickerActorClassRegistry::RegisterClassAdapter<ALight, FAvaViewportColorPickerLightAdapter>();
 	FAvaViewportColorPickerActorClassRegistry::RegisterClassAdapter<ASkyLight, FAvaViewportColorPickerLightAdapter>();
+
+	UAvaSceneRigSubsystem::RegisterSupportedActorClasses(DefaultSceneRigActorClasses());
 }
 
 void FAvaEditorModule::ShutdownModule()
@@ -123,6 +136,8 @@ void FAvaEditorModule::ShutdownModule()
 			IAvaOutlinerModule::Get().UnregisterOverriddenIcon<FAvaOutlinerActor>(AAvaShapeActor::StaticClass()->GetFName());
 		}
 	}
+
+	UAvaSceneRigSubsystem::UnregisterSupportedActorClasses(DefaultSceneRigActorClasses());
 }
 
 void FAvaEditorModule::CreateAvaLevelEditor()
@@ -376,6 +391,27 @@ void FAvaEditorModule::RegisterLevelTemplates()
 	}
 }
 
+const TSet<TSubclassOf<AActor>>& FAvaEditorModule::DefaultSceneRigActorClasses()
+{
+	static const TSet<TSubclassOf<AActor>> Classes =
+	{
+		// Cameras
+		ACameraActor::StaticClass(),
+		ACineCameraActor::StaticClass(),
+		AAvaCineCameraActor::StaticClass(),
+		// Lights
+		ASkyLight::StaticClass(),
+		ADirectionalLight::StaticClass(),
+		APointLight::StaticClass(),
+		ARectLight::StaticClass(),
+		ASpotLight::StaticClass(),
+		// Misc
+		AAvaNullActor::StaticClass(),
+		APostProcessVolume::StaticClass()
+	};
+	return Classes;
+}
+
 #undef LOCTEXT_NAMESPACE
 
 IMPLEMENT_MODULE(FAvaEditorModule, AvalancheEditor)
diff --git a/Engine/Plugins/Experimental/Avalanche/Source/AvalancheEditor/Private/AvaEditorModule.h b/Engine/Plugins/Experimental/Avalanche/Source/AvalancheEditor/Private/AvaEditorModule.h
--- a/Engine/Plugins/Experimental/Avalanche/Source/AvalancheEditor/Private/AvaEditorModule.h
+++ b/Engine/Plugins/Experimental/Avalanche/Source/AvalancheEditor/Private/AvaEditorModule.h
@@ -3,8 +3,10 @@
 #pragma once
 
 #include "AssetTypeCategories.h"
+#include "Containers/ContainersFwd.h"
 #include "Modules/ModuleInterface.h"
 #include "Modules/ModuleManager.h"
+#include "Templates/SubclassOf.h"
 #include "Toolkits/AssetEditorToolkit.h"
 
 class FAvaOutlinerItem;
@@ -25,6 +27,8 @@ public:
 	//~ End IAvaEditorModule
 
 	static FSlateIcon GetOutlinerShapeActorIcon(TSharedPtr<const FAvaOutlinerItem> InItem);
+
+	static const TSet<TSubclassOf<AActor>>& DefaultSceneRigActorClasses();
 	
 private:
 	/** Motion Design Level Editor */
diff --git a/Engine/Plugins/Experimental/Avalanche/Source/AvalancheSceneRig/Private/AvaSceneRigSubsystem.cpp b/Engine/Plugins/Experimental/Avalanche/Source/AvalancheSceneRig/Private/AvaSceneRigSubsystem.cpp
--- a/Engine/Plugins/Experimental/Avalanche/Source/AvalancheSceneRig/Private/AvaSceneRigSubsystem.cpp
+++ b/Engine/Plugins/Experimental/Avalanche/Source/AvalancheSceneRig/Private/AvaSceneRigSubsystem.cpp
@@ -18,6 +18,8 @@ DEFINE_LOG_CATEGORY(AvaSceneRigSubsystemLog);
 
 #define LOCTEXT_NAMESPACE "AvaSceneRigSubsystem"
 
+TSet<TSubclassOf<AActor>> UAvaSceneRigSubsystem::SupportedActorClasses = {};
+
 void UAvaSceneRigSubsystem::Initialize(FSubsystemCollectionBase& InOutCollection)
 {
 	Super::Initialize(InOutCollection);
@@ -114,39 +116,34 @@ FString UAvaSceneRigSubsystem::GetSceneRigAssetSuffix()
 	return Suffix;
 }
 
-TSet<FName> UAvaSceneRigSubsystem::GetSupportedActorClassNames()
+void UAvaSceneRigSubsystem::RegisterSupportedActorClasses(const TSet<TSubclassOf<AActor>>& InClasses)
+{
+	for (const TSubclassOf<AActor>& Class : InClasses)
+	{
+		SupportedActorClasses.Add(Class);
+	}
+}
+
+void UAvaSceneRigSubsystem::UnregisterSupportedActorClasses(const TSet<TSubclassOf<AActor>>& InClasses)
 {
-	static const TSet<FName> ClassNames =
+	for (const TSubclassOf<AActor>& Class : InClasses)
+	{
+		const FSetElementId SetId = SupportedActorClasses.FindId(Class);
+		if (SetId.IsValidId())
 		{
-			// Cameras
-			TEXT("CameraActor"),
-			TEXT("CineCameraActor"),
-			TEXT("AvaCineCameraActor"),
-			// Lights
-			TEXT("SkyLight"),
-			TEXT("DirectionalLight"),
-			TEXT("PointLight"),
-			TEXT("RectLight"),
-			TEXT("SpotLight"),
-			// Misc
-			TEXT("AvaNullActor"),
-			TEXT("PostProcessVolume")
-		};
-	return ClassNames;
+			SupportedActorClasses.Remove(SetId);
+		}
+	}
 }
 
-bool UAvaSceneRigSubsystem::IsSupportedActorClassName(const FName InName)
+const TSet<TSubclassOf<AActor>>& UAvaSceneRigSubsystem::GetSupportedActorClasses()
 {
-	return GetSupportedActorClassNames().Contains(InName);
+	return SupportedActorClasses;
 }
 
-bool UAvaSceneRigSubsystem::IsSupportedActorClass(const UClass* InClass)
+bool UAvaSceneRigSubsystem::IsSupportedActorClass(const TSubclassOf<AActor>& InClass)
 {
-	if (IsValid(InClass))
-	{
-		return GetSupportedActorClassNames().Contains(InClass->GetFName());
-	}
-	return false;
+	return GetSupportedActorClasses().Contains(InClass);
 }
 
 bool UAvaSceneRigSubsystem::AreActorsSupported(const TArray<AActor*>& InActors)
diff --git a/Engine/Plugins/Experimental/Avalanche/Source/AvalancheSceneRig/Public/AvaSceneRigSubsystem.h b/Engine/Plugins/Experimental/Avalanche/Source/AvalancheSceneRig/Public/AvaSceneRigSubsystem.h
--- a/Engine/Plugins/Experimental/Avalanche/Source/AvalancheSceneRig/Public/AvaSceneRigSubsystem.h
+++ b/Engine/Plugins/Experimental/Avalanche/Source/AvalancheSceneRig/Public/AvaSceneRigSubsystem.h
@@ -7,6 +7,7 @@
 #include "Engine/AssetUserData.h"
 #include "Subsystems/WorldSubsystem.h"
 #include "Templates/Function.h"
+#include "Templates/SubclassOf.h"
 #include "AvaSceneRigSubsystem.generated.h"
 
 class AActor;
@@ -39,14 +40,17 @@ public:
 	/** Returns the Scene Rig level suffix string. Recommended to append to level asset name to differentiate them from other levels. */
 	AVALANCHESCENERIG_API static FString GetSceneRigAssetSuffix();
 
-	/** Get the list of supported object class names that are allowed to be instanced into a scene rig level. */
-	AVALANCHESCENERIG_API static TSet<FName> GetSupportedActorClassNames();
+	/** Registers actor classes that are allowed to be added to a Scene Rig. */
+	AVALANCHESCENERIG_API static void RegisterSupportedActorClasses(const TSet<TSubclassOf<AActor>>& InClasses);
 
-	/** Returns true if the class name is a supported Scene Rig class. */
-	AVALANCHESCENERIG_API static bool IsSupportedActorClassName(const FName InName);
+	/** Unregisters actor classes that are allowed to be added to a Scene Rig. */
+	AVALANCHESCENERIG_API static void UnregisterSupportedActorClasses(const TSet<TSubclassOf<AActor>>& InClasses);
+
+	/** Returns the list of supported object classes that are allowed to be added to a scene rig level. */
+	AVALANCHESCENERIG_API static const TSet<TSubclassOf<AActor>>& GetSupportedActorClasses();
 
 	/** Returns true if the class is a supported Scene Rig class. */
-	AVALANCHESCENERIG_API static bool IsSupportedActorClass(const UClass* InClass);
+	AVALANCHESCENERIG_API static bool IsSupportedActorClass(const TSubclassOf<AActor>& InClass);
 
 	/** Returns true if the list of actors are supported Scene Rig classes. */
 	AVALANCHESCENERIG_API static bool AreActorsSupported(const TArray<AActor*>& InActors);
@@ -91,4 +95,6 @@ private:
 #if WITH_EDITOR
 	FDelegateHandle WorldTagGetterDelegate;
 #endif
+
+	static TSet<TSubclassOf<AActor>> SupportedActorClasses;
 };
