[Motion Design] EaseCurveTool
 - Toast message when trying to ease key values that are equal
 - Fix issue with preset combo box not being updated on key selection
 - Fix transactions when dragging
 - Better tangent handle selection

diff --git a/Engine/Plugins/Experimental/Avalanche/Source/AvalancheSequencer/Private/AvaSequencer.cpp b/Engine/Plugins/Experimental/Avalanche/Source/AvalancheSequencer/Private/AvaSequencer.cpp
--- a/Engine/Plugins/Experimental/Avalanche/Source/AvalancheSequencer/Private/AvaSequencer.cpp
+++ b/Engine/Plugins/Experimental/Avalanche/Source/AvalancheSequencer/Private/AvaSequencer.cpp
@@ -1731,12 +1731,4 @@ TSharedRef<FAvaEaseCurveTool> FAvaSequencer::GetEaseCurveTool() const
 	return EaseCurveTool.ToSharedRef(); 
 }
 
-void FAvaSequencer::UpdateEaseCurveTool() const
-{
-	if (EaseCurveTool.IsValid())
-	{
-		EaseCurveTool->UpdateEaseCurveFromSequencerKeySelections();
-	}
-}
-
 #undef LOCTEXT_NAMESPACE
diff --git a/Engine/Plugins/Experimental/Avalanche/Source/AvalancheSequencer/Private/AvaSequencer.h b/Engine/Plugins/Experimental/Avalanche/Source/AvalancheSequencer/Private/AvaSequencer.h
--- a/Engine/Plugins/Experimental/Avalanche/Source/AvalancheSequencer/Private/AvaSequencer.h
+++ b/Engine/Plugins/Experimental/Avalanche/Source/AvalancheSequencer/Private/AvaSequencer.h
@@ -182,8 +182,6 @@ public:
 	void FixBindingHierarchy();
 
 	TSharedRef<FAvaEaseCurveTool> GetEaseCurveTool() const;
-
-	void UpdateEaseCurveTool() const;
 	
 	TArrayView<TWeakObjectPtr<>> ResolveBoundObjects(const FGuid& InBindingId, class UMovieSceneSequence* Sequence) const;
 
diff --git a/Engine/Plugins/Experimental/Avalanche/Source/AvalancheSequencer/Private/DetailsView/SAvaKeyFrameEdit.cpp b/Engine/Plugins/Experimental/Avalanche/Source/AvalancheSequencer/Private/DetailsView/SAvaKeyFrameEdit.cpp
--- a/Engine/Plugins/Experimental/Avalanche/Source/AvalancheSequencer/Private/DetailsView/SAvaKeyFrameEdit.cpp
+++ b/Engine/Plugins/Experimental/Avalanche/Source/AvalancheSequencer/Private/DetailsView/SAvaKeyFrameEdit.cpp
@@ -79,8 +79,6 @@ void SAvaKeyFrameEdit::Construct(const FArguments& InArgs, const TSharedRef<FAva
 
 	const TSharedPtr<FAvaSequencer> AvaSequencer = AvaSequencerWeak.Pin();
 
-	InSequencer->UpdateEaseCurveTool();
-
 	ChildSlot
 		[
 			SNew(SScrollBox)
diff --git a/Engine/Plugins/Experimental/Avalanche/Source/AvalancheSequencer/Private/EaseCurveTool/AvaEaseCurveStyle.cpp b/Engine/Plugins/Experimental/Avalanche/Source/AvalancheSequencer/Private/EaseCurveTool/AvaEaseCurveStyle.cpp
--- a/Engine/Plugins/Experimental/Avalanche/Source/AvalancheSequencer/Private/EaseCurveTool/AvaEaseCurveStyle.cpp
+++ b/Engine/Plugins/Experimental/Avalanche/Source/AvalancheSequencer/Private/EaseCurveTool/AvaEaseCurveStyle.cpp
@@ -58,6 +58,7 @@ FAvaEaseCurveStyle::FAvaEaseCurveStyle()
 	Set("ToolToggleButton", ToolToggleButtonStyle);
 
 	Set("Editor.LabelFont", FSlateFontInfo(FCoreStyle::GetDefaultFont(), 7, TEXT("Regular")));
+	Set("Editor.ErrorFont", FSlateFontInfo(FCoreStyle::GetDefaultFont(), 9, TEXT("Italic")));
 
 	FSlateStyleRegistry::RegisterSlateStyle(*this);
 }
diff --git a/Engine/Plugins/Experimental/Avalanche/Source/AvalancheSequencer/Private/EaseCurveTool/AvaEaseCurveSubsystem.cpp b/Engine/Plugins/Experimental/Avalanche/Source/AvalancheSequencer/Private/EaseCurveTool/AvaEaseCurveSubsystem.cpp
--- a/Engine/Plugins/Experimental/Avalanche/Source/AvalancheSequencer/Private/EaseCurveTool/AvaEaseCurveSubsystem.cpp
+++ b/Engine/Plugins/Experimental/Avalanche/Source/AvalancheSequencer/Private/EaseCurveTool/AvaEaseCurveSubsystem.cpp
@@ -539,13 +539,14 @@ TSharedPtr<FAvaEaseCurvePreset> UAvaEaseCurveSubsystem::FindPreset(const FString
 	return nullptr;
 }
 
-TSharedPtr<FAvaEaseCurvePreset> UAvaEaseCurveSubsystem::FindPresetByTangents(const FAvaEaseCurveTangents& InTangents)
+TSharedPtr<FAvaEaseCurvePreset> UAvaEaseCurveSubsystem::FindPresetByTangents(const FAvaEaseCurveTangents& InTangents, const double InErrorTolerance)
 {
 	for (const TPair<FString, TArray<TSharedPtr<FAvaEaseCurvePreset>>>& Preset : Presets)
 	{
-		const TSharedPtr<FAvaEaseCurvePreset>* PresetPtr = Preset.Value.FindByPredicate([&InTangents](const TSharedPtr<FAvaEaseCurvePreset>& InPreset)
+		const TSharedPtr<FAvaEaseCurvePreset>* PresetPtr = Preset.Value.FindByPredicate([&InTangents, InErrorTolerance]
+			(const TSharedPtr<FAvaEaseCurvePreset>& InPreset)
 			{
-				return InPreset->Tangents == InTangents;
+				return InPreset->Tangents.IsNearlyEqual(InTangents, InErrorTolerance);
 			});
 		if (PresetPtr)
 		{
diff --git a/Engine/Plugins/Experimental/Avalanche/Source/AvalancheSequencer/Private/EaseCurveTool/AvaEaseCurveSubsystem.h b/Engine/Plugins/Experimental/Avalanche/Source/AvalancheSequencer/Private/EaseCurveTool/AvaEaseCurveSubsystem.h
--- a/Engine/Plugins/Experimental/Avalanche/Source/AvalancheSequencer/Private/EaseCurveTool/AvaEaseCurveSubsystem.h
+++ b/Engine/Plugins/Experimental/Avalanche/Source/AvalancheSequencer/Private/EaseCurveTool/AvaEaseCurveSubsystem.h
@@ -50,7 +50,7 @@ public:
 	bool ChangePresetCategory(const TSharedPtr<FAvaEaseCurvePreset>& InPreset, const FString& InNewCategory) const;
 
 	TSharedPtr<FAvaEaseCurvePreset> FindPreset(const FString& InName);
-	TSharedPtr<FAvaEaseCurvePreset> FindPresetByTangents(const FAvaEaseCurveTangents& InTangents);
+	TSharedPtr<FAvaEaseCurvePreset> FindPresetByTangents(const FAvaEaseCurveTangents& InTangents, const double InErrorTolerance = 0.01);
 
 	bool DoesPresetCategoryExist(const FString& InCategory);
 	bool RenamePresetCategory(const FString& InCategory, const FString& InNewCategoryName);
diff --git a/Engine/Plugins/Experimental/Avalanche/Source/AvalancheSequencer/Private/EaseCurveTool/AvaEaseCurveTangents.cpp b/Engine/Plugins/Experimental/Avalanche/Source/AvalancheSequencer/Private/EaseCurveTool/AvaEaseCurveTangents.cpp
--- a/Engine/Plugins/Experimental/Avalanche/Source/AvalancheSequencer/Private/EaseCurveTool/AvaEaseCurveTangents.cpp
+++ b/Engine/Plugins/Experimental/Avalanche/Source/AvalancheSequencer/Private/EaseCurveTool/AvaEaseCurveTangents.cpp
@@ -56,6 +56,14 @@ FAvaEaseCurveTangents::FAvaEaseCurveTangents(const FString& InTangentsString)
 	FAvaEaseCurveTangents::FromString(InTangentsString, *this);
 }
 
+bool FAvaEaseCurveTangents::IsNearlyEqual(const FAvaEaseCurveTangents& InOther, const double InErrorTolerance) const
+{
+	return FMath::IsNearlyEqual(Start, InOther.Start, InErrorTolerance)
+		&& FMath::IsNearlyEqual(StartWeight, InOther.StartWeight, InErrorTolerance)
+		&& FMath::IsNearlyEqual(End, InOther.End, InErrorTolerance)
+		&& FMath::IsNearlyEqual(EndWeight, InOther.EndWeight, InErrorTolerance);
+}
+
 FText FAvaEaseCurveTangents::ToDisplayText() const
 {
 	const FNumberFormattingOptions NumberFormat = DefaultNumberFormattingOptions();
diff --git a/Engine/Plugins/Experimental/Avalanche/Source/AvalancheSequencer/Private/EaseCurveTool/AvaEaseCurveTangents.h b/Engine/Plugins/Experimental/Avalanche/Source/AvalancheSequencer/Private/EaseCurveTool/AvaEaseCurveTangents.h
--- a/Engine/Plugins/Experimental/Avalanche/Source/AvalancheSequencer/Private/EaseCurveTool/AvaEaseCurveTangents.h
+++ b/Engine/Plugins/Experimental/Avalanche/Source/AvalancheSequencer/Private/EaseCurveTool/AvaEaseCurveTangents.h
@@ -4,6 +4,7 @@
 
 #include "Curves/RichCurve.h"
 #include "Channels/MovieSceneDoubleChannel.h"
+#include "Math/UnrealMathUtility.h"
 #include "AvaEaseCurveTangents.generated.h"
 
 class FString;
@@ -57,6 +58,8 @@ public:
 		return !(*this == InRhs);
 	}
 
+	bool IsNearlyEqual(const FAvaEaseCurveTangents& InOther, const double InErrorTolerance = UE_DOUBLE_SMALL_NUMBER) const;
+
 	FText ToDisplayText() const;
 	FString ToDisplayString() const;
 
diff --git a/Engine/Plugins/Experimental/Avalanche/Source/AvalancheSequencer/Private/EaseCurveTool/AvaEaseCurveTool.cpp b/Engine/Plugins/Experimental/Avalanche/Source/AvalancheSequencer/Private/EaseCurveTool/AvaEaseCurveTool.cpp
--- a/Engine/Plugins/Experimental/Avalanche/Source/AvalancheSequencer/Private/EaseCurveTool/AvaEaseCurveTool.cpp
+++ b/Engine/Plugins/Experimental/Avalanche/Source/AvalancheSequencer/Private/EaseCurveTool/AvaEaseCurveTool.cpp
@@ -8,7 +8,6 @@
 #include "Dialogs/DlgPickAssetPath.h"
 #include "EaseCurveTool/AvaEaseCurveToolCommands.h"
 #include "EaseCurveTool/AvaEaseCurveToolSettings.h"
-#include "EaseCurveTool/Widgets/SAvaEaseCurveEditor.h"
 #include "EaseCurveTool/Widgets/SAvaEaseCurveTool.h"
 #include "EngineAnalytics.h"
 #include "Factories/CurveFactory.h"
@@ -121,11 +120,12 @@ FString FAvaEaseCurveTool::GetReferencerName() const
 
 TSharedRef<SWidget> FAvaEaseCurveTool::GenerateWidget()
 {
-	CacheSelectionData();
+	UpdateEaseCurveFromSequencerKeySelections();
 
 	if (!ToolWidget.IsValid())
 	{
 		ToolWidget = SNew(SAvaEaseCurveTool, SharedThis(this))
+			.InitialTangents(GetEaseCurveTangents())
 			.Visibility(this, &FAvaEaseCurveTool::GetVisibility)
 			.ToolOperation(this, &FAvaEaseCurveTool::GetToolOperation);
 	}
@@ -153,12 +153,8 @@ FAvaEaseCurveTangents FAvaEaseCurveTool::GetEaseCurveTangents() const
 	return EaseCurve->GetTangents();
 }
 
-void FAvaEaseCurveTool::SetEaseCurveTangents(const FAvaEaseCurveTangents& InTangents, const EOperation InOperation
-	, const bool bInBroadcastUpdate, const bool bInSetSequencerTangents)
+void FAvaEaseCurveTool::SetEaseCurveTangents_Internal(const FAvaEaseCurveTangents& InTangents, const EOperation InOperation, const bool bInBroadcastUpdate)
 {
-	const FScopedTransaction Transaction(LOCTEXT("SetEaseCurveTangents", "Set Ease Curve Tangents"));
-	EaseCurve->Modify();
-
 	switch (InOperation)
 	{
 	case EOperation::InOut:
@@ -176,6 +172,15 @@ void FAvaEaseCurveTool::SetEaseCurveTangents(const FAvaEaseCurveTangents& InTang
 	{
 		EaseCurve->BroadcastUpdate();
 	}
+}
+
+void FAvaEaseCurveTool::SetEaseCurveTangents(const FAvaEaseCurveTangents& InTangents, const EOperation InOperation
+	, const bool bInBroadcastUpdate, const bool bInSetSequencerTangents)
+{
+	const FScopedTransaction Transaction(LOCTEXT("SetEaseCurveTangents", "Set Ease Curve Tangents"));
+	EaseCurve->Modify();
+
+	SetEaseCurveTangents_Internal(InTangents, InOperation, bInBroadcastUpdate);
 
 	if (bInSetSequencerTangents)
 	{
@@ -203,7 +208,9 @@ void FAvaEaseCurveTool::ResetEaseCurveTangents(const EOperation InOperation)
 	const FScopedTransaction Transaction(TransactionText);
 	EaseCurve->ModifyOwner();
 
-	SetEaseCurveTangents(FAvaEaseCurveTangents(0.0, 0.0, 0.0, 0.0), InOperation, true, true);
+	const FAvaEaseCurveTangents ZeroTangents;
+	SetEaseCurveTangents_Internal(ZeroTangents, InOperation, /*bInBroadcastUpdate=*/true);
+	SetSequencerKeySelectionTangents(ZeroTangents, InOperation);
 }
 
 void FAvaEaseCurveTool::FlattenOrStraightenTangents(const EOperation InOperation, const bool bInFlattenTangents) const
@@ -254,11 +261,6 @@ void FAvaEaseCurveTool::FlattenOrStraightenTangents(const EOperation InOperation
 	EaseCurve->BroadcastUpdate();
 }
 
-void FAvaEaseCurveTool::ApplyEaseCurveToSequencerKeySelections()
-{
-	SetSequencerKeySelectionTangents(GetEaseCurveTangents(), OperationMode);
-}
-
 void FAvaEaseCurveTool::ApplyQuickEaseToSequencerKeySelections(const EOperation InOperation)
 {
 	const UAvaEaseCurveToolSettings* const Settings = GetDefault<UAvaEaseCurveToolSettings>();
@@ -271,11 +273,12 @@ void FAvaEaseCurveTool::ApplyQuickEaseToSequencerKeySelections(const EOperation
 		return;
 	}
 
-	SetEaseCurveTangents(Tangents, InOperation, true, true);
+	SetEaseCurveTangents(Tangents, InOperation, /*bInBroadcastUpdate=*/true, /*bInSetSequencerTangents=*/true);
 
+	// Update ease curve combobox widget and zoom graph editor
 	if (ToolWidget.IsValid())
 	{
-		ToolWidget->ZoomToFit();
+		ToolWidget->SetTangents(Tangents, InOperation, /*bInSetEaseCurve=*/false, /*bInBroadcastUpdate=*/false, /*bInSetSequencerTangents=*/false);
 	}
 
 	if (FEngineAnalytics::IsAvailable())
@@ -308,39 +311,28 @@ void FAvaEaseCurveTool::SetSequencerKeySelectionTangents(const FAvaEaseCurveTang
 
 	const FFrameRate DisplayRate = GetDisplayRate();
 	const FFrameRate TickResolution = GetTickResolution();
-
 	const UAvaEaseCurveToolSettings* const EaseCurveToolSettings = GetDefault<UAvaEaseCurveToolSettings>();
 	check(EaseCurveToolSettings);
 	const bool bAutoFlipTangents = EaseCurveToolSettings->GetAutoFlipTangents();
 
-	for (const TPair<FName, FKeyDataCache::FChannelData>& ChannelEntry : KeyCache.ChannelKeyData)
-	{
-		for (const FKeyHandle& KeyHandleToEdit : ChannelEntry.Value.KeyHandles)
+	KeyCache.ForEachEaseableKey(/*bInIncludeEqualValueKeys=*/true, [&DisplayRate, &TickResolution, &InTangents, InOperation, bAutoFlipTangents](const FKeyHandle& InKeyHandle
+		, const FKeyHandle& InNextKeyHandle, const FKeyDataCache::FChannelData& InChannelData)
 		{
-			TMovieSceneChannelData<FMovieSceneDoubleValue> ChannelData = ChannelEntry.Value.DoubleChannel->GetData();
+			TMovieSceneChannelData<FMovieSceneDoubleValue> ChannelData = InChannelData.DoubleChannel->GetData();
 
-			const int32 KeyIndex = ChannelData.GetIndex(KeyHandleToEdit);
-			if (KeyIndex == INDEX_NONE)
-			{
-				continue;
-			}
-
-			// InTangents should always be a normalized range, allowing for weights in the range of 0 - 10
-			const TArrayView<FFrameNumber> ChannelTimes = ChannelData.GetTimes();
 			TArrayView<FMovieSceneDoubleValue> ChannelValues = ChannelData.GetValues();
+			const TArrayView<FFrameNumber> ChannelTimes = ChannelData.GetTimes();
 			const int32 KeyCount = ChannelValues.Num();
 
-			// If there is a keyframe after this keyframe that we are editing, we check if the that keyframe value is less
-			// than or greater than this keyframe value. If less, flip the tangent (if option is set).
-			int32 NextKeyIndex = INDEX_NONE;
-			bool bIncreasingValue = false;
-			if (KeyIndex + 1 < KeyCount)
-			{
-				NextKeyIndex = KeyIndex + 1;
-				bIncreasingValue = ChannelValues[NextKeyIndex].Value >= ChannelValues[KeyIndex].Value;
-			}
+			const int32 KeyIndex = ChannelData.GetIndex(InKeyHandle);
+			const int32 NextKeyIndex = ChannelData.GetIndex(InNextKeyHandle);
 
 			FAvaEaseCurveTangents ScaledTangents = InTangents;
+
+			// If there is a key frame after this key frame that we are editing, we check if the that key frame value is less
+			// than or greater than this key frame value. If less, flip the tangent (if option is set).
+			const bool bIncreasingValue = ChannelValues[NextKeyIndex].Value >= ChannelValues[KeyIndex].Value;
+
 			if (bAutoFlipTangents && !bIncreasingValue)
 			{
 				ScaledTangents.Start *= -1.f;
@@ -348,16 +340,13 @@ void FAvaEaseCurveTool::SetSequencerKeySelectionTangents(const FAvaEaseCurveTang
 			}
 
 			// Scale normalized tangents to time/value range
-			if (NextKeyIndex != INDEX_NONE)
-			{
-				ScaledTangents.ScaleUp(ChannelTimes[KeyIndex], ChannelValues[KeyIndex].Value
-					, ChannelTimes[NextKeyIndex], ChannelValues[NextKeyIndex].Value
-					, DisplayRate, TickResolution);
-			}
+			ScaledTangents.ScaleUp(ChannelTimes[KeyIndex], ChannelValues[KeyIndex].Value
+				, ChannelTimes[NextKeyIndex], ChannelValues[NextKeyIndex].Value
+				, DisplayRate, TickResolution);
 
 			const FScopedTransaction Transaction(LOCTEXT("SetSequencerCurveTangents", "Set Sequencer Curve Tangents"));
-			ChannelEntry.Value.Section->Modify();
-			ChannelEntry.Value.Section->MarkAsChanged();
+			InChannelData.Section->Modify();
+			InChannelData.Section->MarkAsChanged();
 
 			// Set this keys leave tangent
 			if (InOperation == EOperation::Out || InOperation == EOperation::InOut)
@@ -378,19 +367,15 @@ void FAvaEaseCurveTool::SetSequencerKeySelectionTangents(const FAvaEaseCurveTang
 				ChannelValues[NextKeyIndex].Tangent.ArriveTangent = ScaledTangents.End;
 				ChannelValues[NextKeyIndex].Tangent.ArriveTangentWeight = ScaledTangents.EndWeight;
 			}
-		}
-	}
+
+			return true;
+		});
 }
 
 void FAvaEaseCurveTool::UpdateEaseCurveFromSequencerKeySelections()
 {
 	CacheSelectionData();
 	
-	if (KeyCache.TotalSelectedKeys == 0)
-	{
-		return;
-	}
-	
 	const FFrameRate DisplayRate = GetDisplayRate();
 	const FFrameRate TickResolution = GetTickResolution();
 	const UAvaEaseCurveToolSettings* const EaseCurveToolSettings = GetDefault<UAvaEaseCurveToolSettings>();
@@ -398,73 +383,58 @@ void FAvaEaseCurveTool::UpdateEaseCurveFromSequencerKeySelections()
 	const bool bAutoFlipTangents = EaseCurveToolSettings->GetAutoFlipTangents();
 
 	TArray<FAvaEaseCurveTangents> KeySetTangents;
+	TArray<FAvaEaseCurveTangents> ChangingTangents;
 
-	for (const TPair<FName, FKeyDataCache::FChannelData>& ChannelEntry : KeyCache.ChannelKeyData)
-	{
-		for (const FKeyHandle& KeyHandleToEdit : ChannelEntry.Value.KeyHandles)
+	KeyCache.ForEachEaseableKey(/*bInIncludeEqualValueKeys=*/true, [&DisplayRate, &TickResolution, bAutoFlipTangents, &KeySetTangents, &ChangingTangents]
+		(const FKeyHandle& InKeyHandle, const FKeyHandle& InNextKeyHandle, const FKeyDataCache::FChannelData& InChannelData)
 		{
-			TMovieSceneChannelData<FMovieSceneDoubleValue> ChannelData = ChannelEntry.Value.DoubleChannel->GetData();
+			TMovieSceneChannelData<FMovieSceneDoubleValue> ChannelData = InChannelData.DoubleChannel->GetData();
+
+			const TArrayView<FMovieSceneDoubleValue> ChannelValues = ChannelData.GetValues();
 			const TArrayView<FFrameNumber> ChannelTimes = ChannelData.GetTimes();
-			const TArrayView<const FMovieSceneDoubleValue> ChannelValues = ChannelData.GetValues();
 			const int32 KeyCount = ChannelValues.Num();
 
-			const int32 KeyIndex = ChannelData.GetIndex(KeyHandleToEdit);
-			if (KeyIndex == INDEX_NONE)
+			const int32 KeyIndex = ChannelData.GetIndex(InKeyHandle);
+			const int32 NextKeyIndex = ChannelData.GetIndex(InNextKeyHandle);
+
+			// If there is a key frame after this key frame that we are editing, we check if the that key frame value is less
+			// than or greater than this key frame value. If less, flip the tangent (if option is set).
+			const bool bIncreasingValue = ChannelValues[NextKeyIndex].Value >= ChannelValues[KeyIndex].Value;
+
+			FAvaEaseCurveTangents Tangents = FAvaEaseCurveTangents(ChannelValues[KeyIndex], ChannelValues[NextKeyIndex]);
+
+			if (bAutoFlipTangents && !bIncreasingValue)
 			{
-				continue;
+				Tangents.Start *= -1.f;
+				Tangents.End *= -1.f;
 			}
 
-			if (ChannelValues.IsValidIndex(KeyIndex))
+			// Scale time/value to normalized tangent range
+			FAvaEaseCurveTangents ScaledTangents = Tangents;
+			ScaledTangents.Normalize(ChannelTimes[KeyIndex], ChannelValues[KeyIndex].Value
+				, ChannelTimes[NextKeyIndex], ChannelValues[NextKeyIndex].Value
+				, DisplayRate, TickResolution);
+
+			KeySetTangents.Add(ScaledTangents);
+			if (ChannelValues[KeyIndex].Value != ChannelValues[NextKeyIndex].Value)
 			{
-				// If there is a keyframe after this keyframe that we are editing, we check if the that keyframe value is less
-				// than or greater than this keyframe value. If less, flip the tangent (if option is set).
-				int32 NextKeyIndex = INDEX_NONE;
-				bool bIncreasingValue = false;
-				if (KeyIndex + 1 < KeyCount)
-				{
-					NextKeyIndex = KeyIndex + 1;
-					bIncreasingValue = ChannelValues[NextKeyIndex].Value >= ChannelValues[KeyIndex].Value;
-				}
+				ChangingTangents.Add(ScaledTangents);
+			}
 
-				FAvaEaseCurveTangents Tangents;
-				if (ChannelValues.IsValidIndex(NextKeyIndex))
-				{
-					Tangents = FAvaEaseCurveTangents(ChannelValues[KeyIndex], ChannelValues[NextKeyIndex]);
-				}
-				else
-				{
-					Tangents = FAvaEaseCurveTangents(ChannelValues[KeyIndex].Tangent.LeaveTangent, ChannelValues[KeyIndex].Tangent.LeaveTangentWeight, 0.f, 0.f);
-				}
+			return true;
+		});
 
-				if (bAutoFlipTangents && !bIncreasingValue)
-				{
-					Tangents.Start *= -1.f;
-					Tangents.End *= -1.f;
-				}
+	const FAvaEaseCurveTangents AverageTangents = FAvaEaseCurveTangents::Average(KeySetTangents);
+	const FAvaEaseCurveTangents AverageChangingTangents = FAvaEaseCurveTangents::Average(ChangingTangents);
 
-				// Scale time/value to normalized tangent range
-				FAvaEaseCurveTangents ScaledTangents = Tangents;
-				if (NextKeyIndex != INDEX_NONE)
-				{
-					ScaledTangents.Normalize(ChannelTimes[KeyIndex], ChannelValues[KeyIndex].Value
-						, ChannelTimes[NextKeyIndex], ChannelValues[NextKeyIndex].Value
-						, DisplayRate, TickResolution);
-				}
-				else
-				{
-					ScaledTangents = FAvaEaseCurveTangents(ChannelTimes[KeyIndex].Value, ChannelValues[KeyIndex].Value, 0.0, 0.0);
-				}
+	SetEaseCurveTangents(AverageChangingTangents, EOperation::InOut, /*bInBroadcastUpdate=*/true, false);
 
-				if (NextKeyIndex != INDEX_NONE
-					&& ChannelValues[KeyIndex].Value != ChannelValues[NextKeyIndex].Value)
-				{
-					KeySetTangents.Add(ScaledTangents);
-				}
-			}
-		}
+	// Update the preset combobox widget
+	if (ToolWidget.IsValid())
+	{
+		ToolWidget->SetTangents(AverageChangingTangents, EOperation::InOut
+			, /*bInSetEaseCurve=*/false, /*bInBroadcastUpdate=*/false, /*bInSetSequencerTangents*/false);
 	}
-
-	SetEaseCurveTangents(FAvaEaseCurveTangents::Average(KeySetTangents), EOperation::InOut, true, false);
 }
 
 UCurveBase* FAvaEaseCurveTool::CreateCurveAsset() const
@@ -787,4 +757,69 @@ void FAvaEaseCurveTool::SelectPreviousChannelKey()
 	UpdateEaseCurveFromSequencerKeySelections();
 }
 
+bool FAvaEaseCurveTool::HasCachedKeysToEase()
+{
+	bool bEaseableKeys = false;
+
+	KeyCache.ForEachEaseableKey(/*bInIncludeEqualValueKeys=*/false, [&bEaseableKeys](const FKeyHandle& InKeyHandle
+		, const FKeyHandle& InNextKeyHandle, const FKeyDataCache::FChannelData& InChannelData)
+		{
+			bEaseableKeys = true;
+			return false;
+		});
+
+	return bEaseableKeys;
+}
+
+void FAvaEaseCurveTool::FKeyDataCache::ForEachEaseableKey(const bool bInIncludeEqualValueKeys
+	, TFunctionRef<bool(const FKeyHandle& /*InKeyHandle*/, const FKeyHandle& /*InNextKeyHandle*/, const FChannelData&)> InCallable)
+{
+	for (const TPair<FName, FKeyDataCache::FChannelData>& ChannelEntry : ChannelKeyData)
+	{
+		TMovieSceneChannelData<FMovieSceneDoubleValue> ChannelData = ChannelEntry.Value.DoubleChannel->GetData();
+		const TArrayView<const FMovieSceneDoubleValue> ChannelValues = ChannelData.GetValues();
+		const int32 KeyCount = ChannelValues.Num();
+		
+		for (const FKeyHandle& KeyHandle : ChannelEntry.Value.KeyHandles)
+		{
+			if (KeyHandle == FKeyHandle::Invalid())
+			{
+				continue;
+			}
+
+			const int32 KeyIndex = ChannelData.GetIndex(KeyHandle);
+			if (KeyIndex == INDEX_NONE)
+			{
+				continue;
+			}
+			
+			// If there is no key after the selected key, we don't need to process.
+			// The arrive tangents of this key will be set by the previous key's processing.
+			int32 NextKeyIndex = KeyIndex + 1;
+			NextKeyIndex = (NextKeyIndex < KeyCount) ? NextKeyIndex : INDEX_NONE;
+			if (NextKeyIndex == INDEX_NONE)
+			{
+				continue;
+			}
+
+			// Need to check if the next key index is valid, otherwise GetHandle() will fail.
+			const FKeyHandle NextKeyHandle = ChannelData.GetHandle(NextKeyIndex);
+			if (NextKeyHandle == FKeyHandle::Invalid())
+			{
+				continue;
+			}
+
+			if (!bInIncludeEqualValueKeys && ChannelValues[KeyIndex].Value == ChannelValues[NextKeyIndex].Value)
+			{
+				continue;
+			}
+			
+			if (!InCallable(KeyHandle, NextKeyHandle, ChannelEntry.Value))
+			{
+				return;
+			}
+		}
+	}
+}
+
 #undef LOCTEXT_NAMESPACE
diff --git a/Engine/Plugins/Experimental/Avalanche/Source/AvalancheSequencer/Private/EaseCurveTool/AvaEaseCurveTool.h b/Engine/Plugins/Experimental/Avalanche/Source/AvalancheSequencer/Private/EaseCurveTool/AvaEaseCurveTool.h
--- a/Engine/Plugins/Experimental/Avalanche/Source/AvalancheSequencer/Private/EaseCurveTool/AvaEaseCurveTool.h
+++ b/Engine/Plugins/Experimental/Avalanche/Source/AvalancheSequencer/Private/EaseCurveTool/AvaEaseCurveTool.h
@@ -74,23 +74,22 @@ public:
 
 	FAvaEaseCurveTangents GetEaseCurveTangents() const;
 
-	void SetEaseCurveTangents(const FAvaEaseCurveTangents& InTangents, const EOperation InOperation,
-		const bool bBroadcastUpdate, const bool bInSetSequencerTangents);
+	/**
+	 * Sets the internal ease curve tangents and optionally broadcasts a change event for the curve object.
+	 * This is different from SetEaseCurveTangents_Internal in that it performs undo/redo transactions and
+	 * optionally sets the selected tangents in the actual sequence.
+	 */
+	void SetEaseCurveTangents(const FAvaEaseCurveTangents& InTangents, const EOperation InOperation, const bool bInBroadcastUpdate, const bool bInSetSequencerTangents);
 
 	void ResetEaseCurveTangents(const EOperation InOperation);
 
 	void FlattenOrStraightenTangents(const EOperation InOperation, const bool bInFlattenTangents) const;
 
-	FORCEINLINE EOperation GetOperation() const { return OperationMode; }
-	FORCEINLINE void SetOperation(const EOperation InOperation) { OperationMode = InOperation; }
-
 	/** Creates a new external float curve from the internet curve editor curve. */
 	UCurveBase* CreateCurveAsset() const;
 
 	void SetSequencerKeySelectionTangents(const FAvaEaseCurveTangents& InTangents, const EOperation InOperation = EOperation::InOut);
 
-	void ApplyEaseCurveToSequencerKeySelections();
-
 	void ApplyQuickEaseToSequencerKeySelections(const EOperation InOperation = EOperation::InOut);
 
 	/** Updates the ease curve graph view based on the active sequencer key selection. */
@@ -122,6 +121,8 @@ public:
 	void SelectNextChannelKey();
 	void SelectPreviousChannelKey();
 
+	bool HasCachedKeysToEase();
+
 	//~ Begin FGCObject
 	virtual void AddReferencedObjects(FReferenceCollector& Collector) override;
 	virtual FString GetReferencerName() const override;
@@ -133,6 +134,12 @@ public:
 	//~ End FEditorUndoClient
 
 protected:
+	/** 
+	 * Sets the internal ease curve tangents and optionally broadcasts a change event for the curve object.
+	 * Changing the internal ease curve tangents will be directly reflected in the ease curve editor graph.
+	 */
+	void SetEaseCurveTangents_Internal(const FAvaEaseCurveTangents& InTangents, const EOperation InOperation, const bool bInBroadcastUpdate);
+
 	TWeakPtr<FAvaSequencer> AvaSequencerWeak;
 	TWeakPtr<UE::Sequencer::FSequencerSelection> SequencerSelectionWeak;
 
@@ -154,6 +161,9 @@ protected:
 			TArray<FKeyHandle> KeyHandles;
 		};
 
+		void ForEachEaseableKey(const bool bInIncludeEqualValueKeys
+			, TFunctionRef<bool(const FKeyHandle& /*InKeyHandle*/, const FKeyHandle& /*InNextKeyHandle*/, const FChannelData&)> InCallable);
+
 		TMap<FName, FChannelData> ChannelKeyData;
 
 		int32 TotalSelectedKeys = 0;
@@ -165,7 +175,6 @@ protected:
 		bool bIsLastOnlySelectedKey = false;
 	};
 	FKeyDataCache KeyCache;
-	bool bAreKeysCached = false;
 
 private:
 	void CacheSelectionData();
diff --git a/Engine/Plugins/Experimental/Avalanche/Source/AvalancheSequencer/Private/EaseCurveTool/Widgets/SAvaEaseCurveEditor.cpp b/Engine/Plugins/Experimental/Avalanche/Source/AvalancheSequencer/Private/EaseCurveTool/Widgets/SAvaEaseCurveEditor.cpp
--- a/Engine/Plugins/Experimental/Avalanche/Source/AvalancheSequencer/Private/EaseCurveTool/Widgets/SAvaEaseCurveEditor.cpp
+++ b/Engine/Plugins/Experimental/Avalanche/Source/AvalancheSequencer/Private/EaseCurveTool/Widgets/SAvaEaseCurveEditor.cpp
@@ -36,11 +36,13 @@
 namespace UE::EaseCurveTool::Private
 {
 	static const FVector2D KeyHitSize = FVector2D(8.f);
-	static const FVector2D TangentHitSize = FVector2D(16.f);
+	static const FVector2D TangentHitSize = FVector2D(20.f);
+	static const FVector2D TangentDrawSize = FVector2D(16.f);
 	static const FVector2D CurveHitSize = FVector2D(6.f);
 
 	static const FVector2D HalfKeyHitSize = KeyHitSize * 0.5f;
 	static const FVector2D HalfTangentHitSize = TangentHitSize * 0.5f;
+	static const FVector2D HalfTangentDrawSize = TangentDrawSize * 0.5f;
 	static const FVector2D HalfCurveHitSize = CurveHitSize * 0.5f;
 
 	static constexpr float FitMargin = 0.1f;
@@ -87,6 +89,7 @@ void SAvaEaseCurveEditor::Construct(const FArguments& InArgs, const TObjectPtr<U
 	CurveThickness = InArgs._CurveThickness;
 	CurveColor = InArgs._CurveColor;
 	Operation = InArgs._Operation;
+	ShowEqualValueKeyError = InArgs._ShowEqualValueKeyError;
 
 	StartText = InArgs._StartText;
 	StartTooltipText = InArgs._StartTooltipText;
@@ -176,51 +179,71 @@ int32 SAvaEaseCurveEditor::OnPaint(const FPaintArgs& InArgs, const FGeometry& In
 			WhiteBrush, DrawEffects, NormalAreaColor);
 	}
 
-	InLayerId = PaintGrid(ScaleInfo, InAllottedGeometry, OutDrawElements, ++InLayerId, InMyCullingRect, DrawEffects);
-
-	if (StartText.IsSet() || EndText.IsSet())
+	if (ShowEqualValueKeyError.Get(false))
 	{
-		const FSlateFontInfo FontInfo = FAvaEaseCurveStyle::Get().GetFontStyle(TEXT("Editor.LabelFont"));
-		const FLinearColor TextColor = FLinearColor(0.1f, 0.1f, 0.1f, 1.f);
-		const FVector2D TextOffset = FVector2D(4.f);
+		InLayerId = PaintGrid(ScaleInfo, InAllottedGeometry, OutDrawElements, ++InLayerId, InMyCullingRect, DrawEffects);
 
-		if (StartText.IsSet())
+		if (StartText.IsSet() || EndText.IsSet())
 		{
-			const FText& Text = StartText.Get();
-			const TSharedRef<FSlateFontMeasure> FontMeasure = FSlateApplication::Get().GetRenderer()->GetFontMeasureService();
-			const FVector2D TextSize = FontMeasure->Measure(Text, FontInfo);
+			const FSlateFontInfo FontInfo = FAvaEaseCurveStyle::Get().GetFontStyle(TEXT("Editor.LabelFont"));
+			const FLinearColor TextColor = FLinearColor(0.1f, 0.1f, 0.1f, 1.f);
+			const FVector2D TextOffset = FVector2D(4.f);
+
+			if (StartText.IsSet())
+			{
+				const FText& Text = StartText.Get();
+				const TSharedRef<FSlateFontMeasure> FontMeasure = FSlateApplication::Get().GetRenderer()->GetFontMeasureService();
+				const FVector2D TextSize = FontMeasure->Measure(Text, FontInfo);
 
-			const FVector2D ActualOffset(ZeroInputX - (TextOffset.X + TextSize.X), ZeroOutputY - (TextOffset.Y + TextSize.Y));
+				const FVector2D ActualOffset(ZeroInputX - (TextOffset.X + TextSize.X), ZeroOutputY - (TextOffset.Y + TextSize.Y));
+
+				FSlateDrawElement::MakeText(OutDrawElements, ++InLayerId
+					, InAllottedGeometry.ToPaintGeometry(InAllottedGeometry.Size, FSlateLayoutTransform(ActualOffset))
+					, Text, FontInfo, DrawEffects, TextColor);
+			}
 
-			FSlateDrawElement::MakeText(OutDrawElements, ++InLayerId
-				, InAllottedGeometry.ToPaintGeometry(InAllottedGeometry.Size, FSlateLayoutTransform(ActualOffset))
-				, Text, FontInfo, DrawEffects, TextColor);
+			if (EndText.IsSet())
+			{
+				const FVector2D ActualOffset(OneInputX + TextOffset.X, OneOutputY + TextOffset.Y);
+
+				FSlateDrawElement::MakeText(OutDrawElements, ++InLayerId
+					, InAllottedGeometry.ToPaintGeometry(InAllottedGeometry.Size, FSlateLayoutTransform(ActualOffset))
+					, EndText.Get(), FontInfo, DrawEffects, TextColor);
+			}
 		}
 
-		if (EndText.IsSet())
+		PaintNormalBounds(ScaleInfo, InAllottedGeometry, OutDrawElements, ++InLayerId, InMyCullingRect, DrawEffects, InWidgetStyle);
+
+		PaintCurve(ScaleInfo, InAllottedGeometry, OutDrawElements, ++InLayerId, InMyCullingRect, DrawEffects, InWidgetStyle);
+
+		InLayerId = PaintKeys(ScaleInfo, InAllottedGeometry, OutDrawElements, ++InLayerId, InMyCullingRect, DrawEffects, InWidgetStyle);
+
+		if (DragState == EDragState::MarqueeSelect)
 		{
-			const FVector2D ActualOffset(OneInputX + TextOffset.X, OneOutputY + TextOffset.Y);
+			const FVector2D MarqueTopLeft(FMath::Min(MouseDownLocation.X, MouseMoveLocation.X), FMath::Min(MouseDownLocation.Y, MouseMoveLocation.Y));
+			const FVector2D MarqueBottomRight(FMath::Max(MouseDownLocation.X, MouseMoveLocation.X), FMath::Max(MouseDownLocation.Y, MouseMoveLocation.Y));
 
-			FSlateDrawElement::MakeText(OutDrawElements, ++InLayerId
-				, InAllottedGeometry.ToPaintGeometry(InAllottedGeometry.Size, FSlateLayoutTransform(ActualOffset))
-				, EndText.Get(), FontInfo, DrawEffects, TextColor);
+			FSlateDrawElement::MakeBox(OutDrawElements, ++InLayerId
+				, InAllottedGeometry.ToPaintGeometry(MarqueBottomRight - MarqueTopLeft, FSlateLayoutTransform(MarqueTopLeft))
+				, FAppStyle::GetBrush(TEXT("MarqueeSelection")));
 		}
 	}
+	else
+	{
+		const FSlateFontInfo FontInfo = FAvaEaseCurveStyle::Get().GetFontStyle(TEXT("Editor.ErrorFont"));
+		const FText Text = LOCTEXT("EqualValueKeys", "No different key values to create ease curve!");
 
-	PaintNormalBounds(ScaleInfo, InAllottedGeometry, OutDrawElements, ++InLayerId, InMyCullingRect, DrawEffects, InWidgetStyle);
+		const TSharedRef<FSlateFontMeasure> FontMeasure = FSlateApplication::Get().GetRenderer()->GetFontMeasureService();
+		const FVector2D TextSize = FontMeasure->Measure(Text, FontInfo);
 
-	PaintCurve(ScaleInfo, InAllottedGeometry, OutDrawElements, ++InLayerId, InMyCullingRect, DrawEffects, InWidgetStyle);
+		const float HalfInputX = ScaleInfo.InputToLocalX(0.5f);
+		const float HalfOutputY = ScaleInfo.OutputToLocalY(0.5f);
 
-	InLayerId = PaintKeys(ScaleInfo, InAllottedGeometry, OutDrawElements, ++InLayerId, InMyCullingRect, DrawEffects, InWidgetStyle);
+		const FVector2D ActualOffset(HalfInputX - (TextSize.X * 0.5f), HalfOutputY - (TextSize.Y * 0.5f));
 
-	if (DragState == EDragState::MarqueeSelect)
-	{
-		const FVector2D MarqueTopLeft(FMath::Min(MouseDownLocation.X, MouseMoveLocation.X), FMath::Min(MouseDownLocation.Y, MouseMoveLocation.Y));
-		const FVector2D MarqueBottomRight(FMath::Max(MouseDownLocation.X, MouseMoveLocation.X), FMath::Max(MouseDownLocation.Y, MouseMoveLocation.Y));
-
-		FSlateDrawElement::MakeBox(OutDrawElements, ++InLayerId
-			, InAllottedGeometry.ToPaintGeometry(MarqueBottomRight - MarqueTopLeft, FSlateLayoutTransform(MarqueTopLeft))
-			, FAppStyle::GetBrush(TEXT("MarqueeSelection")));
+		FSlateDrawElement::MakeText(OutDrawElements, ++InLayerId
+			, InAllottedGeometry.ToPaintGeometry(InAllottedGeometry.Size, FSlateLayoutTransform(ActualOffset))
+			, Text, FontInfo, DrawEffects, FStyleColors::Foreground.GetSpecifiedColor());
 	}
 
 	return InLayerId;
@@ -354,7 +377,6 @@ void SAvaEaseCurveEditor::PaintCurve(const FTrackScaleInfo& InScaleInfo, const F
 	int32 LayerId, const FSlateRect& MyCullingRect, ESlateDrawEffect DrawEffects, const FWidgetStyle& InWidgetStyle) const
 {
 	const float NumKeys = EaseCurve->FloatCurve.GetNumKeys();
-	
 	if (NumKeys != 2)
 	{
 		return;
@@ -550,11 +572,11 @@ int32 SAvaEaseCurveEditor::PaintTangentHandle(const FGeometry& InAllottedGeometr
 	// Draw tangent handle
 	{
 		const FSlateBrush* ImageBrush = FAppStyle::GetBrush(TEXT("Icons.BulletPoint"));
-		const FVector2D TangentIconLocation = InTangentLocation - HalfTangentHitSize;
+		const FVector2D TangentIconLocation = InTangentLocation - HalfTangentDrawSize;
 		const FLinearColor HandleColor = bInSelected ? FStyleColors::AccentBlue.GetSpecifiedColor() : FLinearColor::White;
 
 		FSlateDrawElement::MakeBox(OutDrawElements, InLayerId++
-			, InAllottedGeometry.ToPaintGeometry(TangentHitSize, FSlateLayoutTransform(TangentIconLocation))
+			, InAllottedGeometry.ToPaintGeometry(TangentDrawSize, FSlateLayoutTransform(TangentIconLocation))
 			, ImageBrush, InDrawEffects, HandleColor * InWidgetStyle.GetColorAndOpacityTint());
 	}
 
@@ -579,8 +601,6 @@ TOptional<bool> SAvaEaseCurveEditor::OnQueryShowFocus(const EFocusCause InFocusC
 FReply SAvaEaseCurveEditor::OnMouseButtonDown(const FGeometry& InMyGeometry, const FPointerEvent& InMouseEvent)
 {
 	const FKey EffectingButton = InMouseEvent.GetEffectingButton();
-	
-	OnDragEnd.ExecuteIfBound();
 
 	DragState = EDragState::PreDrag;
 
@@ -593,8 +613,6 @@ FReply SAvaEaseCurveEditor::OnMouseButtonDown(const FGeometry& InMyGeometry, con
 			FSlateApplication::Get().SetKeyboardFocus(SharedThis(this), EFocusCause::SetDirectly);
 		}
 
-		OnMouseMove(InMyGeometry, InMouseEvent);
-
 		return FReply::Handled().CaptureMouse(SharedThis(this));
 	}
 	else if (EffectingButton == EKeys::RightMouseButton)
@@ -873,7 +891,7 @@ void SAvaEaseCurveEditor::TryStartDrag(const FGeometry& InMyGeometry, const FPoi
 		if (bLeftMouseButton)
 		{
 			// Check if we should start dragging a tangent.
-			const FSelectedTangent HitTangent = HitTestTangentHandle(InMyGeometry, InMouseEvent);
+			const FSelectedTangent HitTangent = HitTestTangentHandles(InMyGeometry, InMouseEvent);
 			if (HitTangent.IsValid())
 			{
 				SelectedTangent = HitTangent;
@@ -970,7 +988,7 @@ void SAvaEaseCurveEditor::ProcessClick(const FGeometry& InMyGeometry, const FPoi
 
 	if (EffectingButton == EKeys::LeftMouseButton)
 	{
-		const FSelectedTangent HitTangent = HitTestTangentHandle(InMyGeometry, InMouseEvent);
+		const FSelectedTangent HitTangent = HitTestTangentHandles(InMyGeometry, InMouseEvent);
 		if (HitTangent.IsValid())
 		{
 			SelectedTangent = HitTangent;
@@ -1133,46 +1151,57 @@ bool SAvaEaseCurveEditor::HitTestCurves(const FGeometry& InMyGeometry, const FPo
 	return false;
 }
 
-SAvaEaseCurveEditor::FSelectedTangent SAvaEaseCurveEditor::HitTestTangentHandle(const FGeometry& InMyGeometry, const FPointerEvent& InMouseEvent) const
+SAvaEaseCurveEditor::FSelectedTangent SAvaEaseCurveEditor::HitTestTangentHandles(const FGeometry& InMyGeometry, const FPointerEvent& InMouseEvent) const
 {
 	const FTrackScaleInfo ScaleInfo(ViewMinInput, ViewMaxInput, ViewMinOutput, ViewMaxOutput, InMyGeometry.GetLocalSize());
 	const FVector2D MousePosition = InMyGeometry.AbsoluteToLocal(InMouseEvent.GetScreenSpacePosition());
 
-	auto IsInside = [&MousePosition](const FVector2D& InTangent) -> bool
+	auto IsInside = [&MousePosition](const FVector2D& InPosition) -> bool
 		{
-			return MousePosition.Y > (InTangent.Y - HalfTangentHitSize.Y)
-				&& MousePosition.Y < (InTangent.Y + HalfTangentHitSize.Y)
-				&& MousePosition.X > (InTangent.X - HalfTangentHitSize.X)
-				&& MousePosition.X < (InTangent.X + HalfTangentHitSize.X);
+			return MousePosition.Y > (InPosition.Y - HalfTangentHitSize.Y)
+				&& MousePosition.Y < (InPosition.Y + HalfTangentHitSize.Y)
+				&& MousePosition.X > (InPosition.X - HalfTangentHitSize.X)
+				&& MousePosition.X < (InPosition.X + HalfTangentHitSize.X);
 		};
 
 	FSelectedTangent OutTangent;
 
-	for (auto It(EaseCurve->FloatCurve.GetKeyHandleIterator()); It; ++It)
-	{
-		const FKeyHandle KeyHandle = *It;
-		const int32 KeyIndex = EaseCurve->FloatCurve.GetIndexSafe(KeyHandle);
-		check(KeyIndex != INDEX_NONE)
+	const FKeyHandle StartKeyHandle = EaseCurve->FloatCurve.GetFirstKeyHandle();
+	check(StartKeyHandle != FKeyHandle::Invalid());
+	const FKeyHandle EndKeyHandle = EaseCurve->FloatCurve.GetLastKeyHandle();
+	check(EndKeyHandle != FKeyHandle::Invalid());
+
+	const FVector2D StartTangentScreenLocation = GetLeaveTangentScreenLocation(ScaleInfo, StartKeyHandle);
+	const FVector2D EndTangentScreenLocation = GetArriveTangentScreenLocation(ScaleInfo, EndKeyHandle);
+
+	const FVector2D StartKeyPosition(EaseCurve->FloatCurve.GetKeyTime(StartKeyHandle), EaseCurve->FloatCurve.GetKeyValue(StartKeyHandle));
+	const FVector2D EndKeyPosition(EaseCurve->FloatCurve.GetKeyTime(EndKeyHandle), EaseCurve->FloatCurve.GetKeyValue(EndKeyHandle));
 
-		if (KeyIndex == 0)
+	const float StartDeltaPosition = (StartKeyPosition - StartTangentScreenLocation).Size();
+	const float EndDeltaPosition = (EndKeyPosition - EndTangentScreenLocation).Size();
+
+	const bool bIsInsideStartHandle = IsInside(StartTangentScreenLocation);
+	const bool bIsInsideEndHandle = IsInside(EndTangentScreenLocation);
+
+	if (bIsInsideStartHandle && bIsInsideEndHandle)
+	{
+		if (StartDeltaPosition < EndDeltaPosition)
 		{
-			const FVector2D LeaveTangentScreenLocation = GetLeaveTangentScreenLocation(ScaleInfo, KeyHandle);
-			if (IsInside(LeaveTangentScreenLocation))
-			{
-				OutTangent = FSelectedTangent(KeyHandle, false);
-				break;
-			}
+			OutTangent = FSelectedTangent(StartKeyHandle, false);
 		}
-		else if (KeyIndex == 1)
+		else if (EndDeltaPosition < StartDeltaPosition)
 		{
-			const FVector2D ArriveTangentScreenLocation = GetArriveTangentScreenLocation(ScaleInfo, KeyHandle);
-			if (IsInside(ArriveTangentScreenLocation))
-			{
-				OutTangent = FSelectedTangent(KeyHandle, true);
-				break;
-			}
+			OutTangent = FSelectedTangent(EndKeyHandle, true);
 		}
 	}
+	else if (bIsInsideStartHandle)
+	{
+		OutTangent = FSelectedTangent(StartKeyHandle, false);
+	}
+	else if (bIsInsideEndHandle)
+	{
+		OutTangent = FSelectedTangent(EndKeyHandle, true);
+	}
 
 	return OutTangent;
 }
diff --git a/Engine/Plugins/Experimental/Avalanche/Source/AvalancheSequencer/Private/EaseCurveTool/Widgets/SAvaEaseCurveEditor.h b/Engine/Plugins/Experimental/Avalanche/Source/AvalancheSequencer/Private/EaseCurveTool/Widgets/SAvaEaseCurveEditor.h
--- a/Engine/Plugins/Experimental/Avalanche/Source/AvalancheSequencer/Private/EaseCurveTool/Widgets/SAvaEaseCurveEditor.h
+++ b/Engine/Plugins/Experimental/Avalanche/Source/AvalancheSequencer/Private/EaseCurveTool/Widgets/SAvaEaseCurveEditor.h
@@ -67,6 +67,7 @@ public:
 		SLATE_ARGUMENT(float, CurveThickness)
 		SLATE_ARGUMENT(FLinearColor, CurveColor)
 		SLATE_ATTRIBUTE(FAvaEaseCurveTool::EOperation, Operation)
+		SLATE_ATTRIBUTE(bool, ShowEqualValueKeyError)
 
 		SLATE_ATTRIBUTE(FText, StartText)
 		SLATE_ATTRIBUTE(FText, StartTooltipText)
@@ -222,7 +223,7 @@ protected:
 	void SetOutputMinMax(const float InNewMin, const float InNewMax);
 
 	bool HitTestCurves(const FGeometry& InMyGeometry, const FPointerEvent& InMouseEvent) const;
-	FSelectedTangent HitTestTangentHandle(const FGeometry& InMyGeometry, const FPointerEvent& InMouseEvent) const;
+	FSelectedTangent HitTestTangentHandles(const FGeometry& InMyGeometry, const FPointerEvent& InMouseEvent) const;
 	bool HitTestKey(const FGeometry& InMyGeometry, const FPointerEvent& InMouseEvent, const FVector2D& InInputPosition) const;
 
 	TObjectPtr<UAvaEaseCurve> EaseCurve;
@@ -246,6 +247,7 @@ protected:
 	float CurveThickness = 1.f;
 	FLinearColor CurveColor;
 	TAttribute<FAvaEaseCurveTool::EOperation> Operation;
+	TAttribute<bool> ShowEqualValueKeyError;
 
 	TAttribute<FText> StartText;
 	TAttribute<FText> StartTooltipText;
diff --git a/Engine/Plugins/Experimental/Avalanche/Source/AvalancheSequencer/Private/EaseCurveTool/Widgets/SAvaEaseCurveTool.cpp b/Engine/Plugins/Experimental/Avalanche/Source/AvalancheSequencer/Private/EaseCurveTool/Widgets/SAvaEaseCurveTool.cpp
--- a/Engine/Plugins/Experimental/Avalanche/Source/AvalancheSequencer/Private/EaseCurveTool/Widgets/SAvaEaseCurveTool.cpp
+++ b/Engine/Plugins/Experimental/Avalanche/Source/AvalancheSequencer/Private/EaseCurveTool/Widgets/SAvaEaseCurveTool.cpp
@@ -166,6 +166,8 @@ void SAvaEaseCurveTool::Construct(const FArguments& InArgs, const TSharedRef<FAv
 	{
 		GEditor->RegisterForUndo(this);
 	}
+
+	CurvePresetWidget->SetSelectedItem(InArgs._InitialTangents);
 }
 
 TSharedRef<SWidget> SAvaEaseCurveTool::ConstructCurveEditorPanel()
@@ -181,11 +183,15 @@ TSharedRef<SWidget> SAvaEaseCurveTool::ConstructCurveEditorPanel()
 			[
 				SAssignNew(CurveEaseEditorWidget, SAvaEaseCurveEditor, EaseCurveTool->GetToolCurve())
 				.DisplayRate(EaseCurveToolRef, &FAvaEaseCurveTool::GetDisplayRate)
-				.Operation(EaseCurveToolRef, &FAvaEaseCurveTool::GetOperation)
+				.Operation(EaseCurveToolRef, &FAvaEaseCurveTool::GetToolOperation)
 				.DesiredSize_Lambda([this]() -> FVector2D
 					{
 						return FVector2D(CurrentGraphSize);
 					})
+				.ShowEqualValueKeyError_Lambda([this]() -> bool
+					{
+						return EaseCurveTool->HasCachedKeysToEase();
+					})
 				.OnTangentsChanged(this, &SAvaEaseCurveTool::HandleEditorTangentsChanged)
 				.GridSnap_UObject(GetDefault<UAvaEaseCurveToolSettings>(), &UAvaEaseCurveToolSettings::GetGridSnap)
 				.GridSize_UObject(GetDefault<UAvaEaseCurveToolSettings>(), &UAvaEaseCurveToolSettings::GetGridSize)
@@ -403,17 +409,19 @@ void SAvaEaseCurveTool::OnEditorDragStart() const
 void SAvaEaseCurveTool::OnEditorDragEnd() const
 {
 	EaseCurveTool->EndTransaction();
+
+	if (!EaseCurveTool->HasCachedKeysToEase())
+	{
+		ResetTangentsAndNotify();
+	}
 }
 
 void SAvaEaseCurveTool::SetTangents(const FAvaEaseCurveTangents& InTangents, FAvaEaseCurveTool::EOperation InOperation
 	, const bool bInSetEaseCurve, const bool bInBroadcastUpdate, const bool bInSetSequencerTangents) const
 {
-	if (CurvePresetWidget.IsValid())
+	if (CurvePresetWidget.IsValid() && !CurvePresetWidget->SetSelectedItem(InTangents))
 	{
-		if (!CurvePresetWidget->SetSelectedItem(InTangents))
-		{
-			CurvePresetWidget->ClearSelection();
-		}
+		CurvePresetWidget->ClearSelection();
 	}
 
 	// To change the graph UI tangents, we need to change the ease curve object tangents and the graph will reflect.
@@ -478,6 +486,12 @@ void SAvaEaseCurveTool::OnEndTangentWeightSpinBoxChanged(const float InNewValue)
 
 void SAvaEaseCurveTool::OnPresetChanged(const TSharedPtr<FAvaEaseCurvePreset>& InPreset) const
 {
+	if (!EaseCurveTool->HasCachedKeysToEase())
+	{
+		ResetTangentsAndNotify();
+		return;
+	}
+
 	SetTangents(InPreset->Tangents, ToolOperation.Get(), true, true, true);
 
 	FSlateApplication::Get().SetAllUserFocus(CurveEaseEditorWidget);
@@ -521,7 +535,7 @@ void SAvaEaseCurveTool::BindCommands()
 
 	CommandList->MapAction(EaseCurveToolCommands.Refresh, FExecuteAction::CreateSP(EaseCurveToolRef, &FAvaEaseCurveTool::UpdateEaseCurveFromSequencerKeySelections));
 	
-	CommandList->MapAction(EaseCurveToolCommands.Apply, FExecuteAction::CreateSP(EaseCurveToolRef, &FAvaEaseCurveTool::ApplyEaseCurveToSequencerKeySelections));
+	CommandList->MapAction(EaseCurveToolCommands.Apply, FExecuteAction::CreateSP(this, &SAvaEaseCurveTool::ApplyTangents));
 
 	CommandList->MapAction(EaseCurveToolCommands.ZoomToFit, FExecuteAction::CreateSP(this, &SAvaEaseCurveTool::ZoomToFit));
 
@@ -979,4 +993,18 @@ void SAvaEaseCurveTool::ResetToDefaultPresets()
 	}
 }
 
+void SAvaEaseCurveTool::ApplyTangents()
+{
+	EaseCurveTool->SetEaseCurveTangents(EaseCurveTool->GetEaseCurveTangents(), EaseCurveTool->GetToolOperation(), /*bInBroadcastUpdate=*/true, /*bInSetSequencerTangents=*/true);
+}
+
+void SAvaEaseCurveTool::ResetTangentsAndNotify() const
+{
+	CurvePresetWidget->ClearSelection();
+
+	SetTangents(FAvaEaseCurveTangents(), FAvaEaseCurveTool::EOperation::InOut, true, true, false);
+
+	FAvaEaseCurveTool::ShowNotificationMessage(LOCTEXT("EqualValueKeys", "No different key values to create ease curve!"));
+}
+
 #undef LOCTEXT_NAMESPACE
diff --git a/Engine/Plugins/Experimental/Avalanche/Source/AvalancheSequencer/Private/EaseCurveTool/Widgets/SAvaEaseCurveTool.h b/Engine/Plugins/Experimental/Avalanche/Source/AvalancheSequencer/Private/EaseCurveTool/Widgets/SAvaEaseCurveTool.h
--- a/Engine/Plugins/Experimental/Avalanche/Source/AvalancheSequencer/Private/EaseCurveTool/Widgets/SAvaEaseCurveTool.h
+++ b/Engine/Plugins/Experimental/Avalanche/Source/AvalancheSequencer/Private/EaseCurveTool/Widgets/SAvaEaseCurveTool.h
@@ -84,6 +84,10 @@ protected:
 
 	void ResetToDefaultPresets();
 
+	void ApplyTangents();
+
+	void ResetTangentsAndNotify() const;
+
 	//~ Begin SWidget
 	virtual FReply OnKeyDown(const FGeometry& MyGeometry, const FKeyEvent& InKeyEvent) override;
 	//~ End SWidget
