Fix lambda capture that was accidentally copying a FTickableEditorObject which is not supported because of the auto registration

diff --git a/Engine/Plugins/Experimental/Avalanche/Source/AvalancheOutliner/Private/Data/AvaOutlinerSaveState.cpp b/Engine/Plugins/Experimental/Avalanche/Source/AvalancheOutliner/Private/Data/AvaOutlinerSaveState.cpp
--- a/Engine/Plugins/Experimental/Avalanche/Source/AvalancheOutliner/Private/Data/AvaOutlinerSaveState.cpp
+++ b/Engine/Plugins/Experimental/Avalanche/Source/AvalancheOutliner/Private/Data/AvaOutlinerSaveState.cpp
@@ -59,7 +59,7 @@ void FAvaOutlinerSaveState::Serialize(FAvaOutliner& InOutliner, FArchive& Ar)
 	if (Ar.IsSaving())
 	{
 		SaveSceneTree(InOutliner, /*bInResetTree*/true);
-		InOutliner.ForEachOutlinerView([InOutliner, this](const TSharedPtr<FAvaOutlinerView>& InOutlinerView)
+		InOutliner.ForEachOutlinerView([&InOutliner, this](const TSharedPtr<FAvaOutlinerView>& InOutlinerView)
 		{
 			SaveOutlinerViewState(InOutliner, *InOutlinerView);
 		});
@@ -106,7 +106,7 @@ void FAvaOutlinerSaveState::Serialize(FAvaOutliner& InOutliner, FArchive& Ar)
 		ensure(SceneTree || OutlinerVersion < FAvaOutlinerVersion::SceneTree);
 		LoadSceneTree(InOutliner.GetTreeRoot(), SceneTree, OutlinerWorld);
 
-		InOutliner.ForEachOutlinerView([InOutliner, this](const TSharedPtr<FAvaOutlinerView>& InOutlinerView)
+		InOutliner.ForEachOutlinerView([&InOutliner, this](const TSharedPtr<FAvaOutlinerView>& InOutlinerView)
 		{
 			LoadOutlinerViewState(InOutliner, *InOutlinerView);
 		});
