[Motion Design] Game Viewport Playable Group Refactor

Summary:

Allow Motion Design playback to occur in a Game Viewport (PIE, Game, etc). Augment the use cases for rundown and playback graphs, including using within nDisplay's workflow.

Refactor Details:

- Playable Group is now more an abstract class. All derived classes have been moved to separate cpp/h.
- The AvaGameInstance implementation has been moved to a derived class: AvaGameInstancePlayableGroup.
- AvaGameViewportPlayableGroup implements the new pre-existing game instance/viewport hook.
- AvaPlayableGroupSubsystem is used to track when a game instance is stopped and unregister it from any playable group it would have been attached to.
- (WIP) PlayableGroupManager: PlayableGroup factory has been updated. Using the presence of GameViewportMediaOutput in the channel to hook to game viewport playable group. This method may be revisited later.

Other:

- GameViewportMediaOutput (WIP):

This is the temporary way we setup a channel to render in a game viewport using dummy GameViewportMediaOutput.
This part is still work in progress.

- Display Cluster Support (WIP):

if the playback server is running as a cluster node, it will append the node's name to the playback server name to be identifiable in the host list in the broadcast editor. This may be a temporary approach until we find a better solution.

- Server name and status in broadcast output tile:

Added more information to indicate which server the output is part of. This had been erroneously removed in a previous refactor.

diff --git a/Engine/Plugins/Experimental/Avalanche/Source/AvalancheMedia/Internal/Broadcast/OutputDevices/AvaBroadcastOutputUtils.h b/Engine/Plugins/Experimental/Avalanche/Source/AvalancheMedia/Internal/Broadcast/OutputDevices/AvaBroadcastOutputUtils.h
--- a/Engine/Plugins/Experimental/Avalanche/Source/AvalancheMedia/Internal/Broadcast/OutputDevices/AvaBroadcastOutputUtils.h
+++ b/Engine/Plugins/Experimental/Avalanche/Source/AvalancheMedia/Internal/Broadcast/OutputDevices/AvaBroadcastOutputUtils.h
@@ -23,6 +23,9 @@ namespace UE::AvaBroadcastOutputUtils
 	/** Retrieves the device provider name from the MediaIOCustomLayout class meta data. */
 	AVALANCHEMEDIA_API FName GetDeviceProviderName(const UMediaOutput* InMediaOutput);
 
+	/** Returns true if the given media output is a game viewport output. */
+	AVALANCHEMEDIA_API bool IsGameViewportOutput(const UMediaOutput* InMediaOutput);
+
 	/** Serializes the given MediaOutput object into a FAvaBroadcastOutputData. */ 
 	AVALANCHEMEDIA_API FAvaBroadcastOutputData CreateMediaOutputData(UMediaOutput* InMediaOutput);
 
diff --git a/Engine/Plugins/Experimental/Avalanche/Source/AvalancheMedia/Private/AvaMediaModule.cpp b/Engine/Plugins/Experimental/Avalanche/Source/AvalancheMedia/Private/AvaMediaModule.cpp
--- a/Engine/Plugins/Experimental/Avalanche/Source/AvalancheMedia/Private/AvaMediaModule.cpp
+++ b/Engine/Plugins/Experimental/Avalanche/Source/AvalancheMedia/Private/AvaMediaModule.cpp
@@ -27,6 +27,11 @@ namespace UE::AvaMediaModule::Private
 		return FParse::Value(FCommandLine::Get(),TEXT("MotionDesignPlaybackServerStart="), OutPlaybackServerName) ||
 		 FParse::Param(FCommandLine::Get(), TEXT("MotionDesignPlaybackServerStart"));
 	}
+
+	bool IsDisplayClusterNode(FString& OutDisplayClusterNodeName)
+	{
+		return FParse::Value(FCommandLine::Get(),TEXT("dc_node="), OutDisplayClusterNodeName);
+	}
 }
 
 FAvaMediaModule::FAvaMediaModule()
@@ -374,6 +379,16 @@ void FAvaMediaModule::PostEngineInit()
 	const bool bIsServerManuallyStarted = IsPlaybackServerManuallyStarted(PlaybackServerName);
 	const bool bIsClientManuallyStarted = FParse::Param(FCommandLine::Get(), TEXT("MotionDesignPlaybackClientStart"));
 
+	FString DisplayClusterNode;
+	if (IsDisplayClusterNode(DisplayClusterNode))
+	{
+		if (PlaybackServerName.IsEmpty())
+		{
+			PlaybackServerName = FPlatformProcess::ComputerName();
+		}
+		PlaybackServerName += FString::Printf(TEXT("_%s"), *DisplayClusterNode);
+	}
+	
 	// Adding a command to suppress the client from auto-starting. This is used when spawning
 	// extra server process from the same project, while preventing extra clients.
 	const bool bIsClientAutoStartSuppressed = FParse::Param(FCommandLine::Get(), TEXT("MotionDesignPlaybackClientSuppress"));
diff --git a/Engine/Plugins/Experimental/Avalanche/Source/AvalancheMedia/Private/Broadcast/Channel/AvaBroadcastOutputChannel.cpp b/Engine/Plugins/Experimental/Avalanche/Source/AvalancheMedia/Private/Broadcast/Channel/AvaBroadcastOutputChannel.cpp
--- a/Engine/Plugins/Experimental/Avalanche/Source/AvalancheMedia/Private/Broadcast/Channel/AvaBroadcastOutputChannel.cpp
+++ b/Engine/Plugins/Experimental/Avalanche/Source/AvalancheMedia/Private/Broadcast/Channel/AvaBroadcastOutputChannel.cpp
@@ -768,7 +768,7 @@ void FAvaBroadcastOutputChannel::StopChannelBroadcast()
 UMediaOutput* FAvaBroadcastOutputChannel::AddMediaOutput(const UClass* InMediaOutputClass, const FAvaBroadcastMediaOutputInfo& InOutputInfo)
 {
 	// Don't add a remote output to a local preview channel.
-	if (InOutputInfo.IsValid() && InOutputInfo.IsRemote() && GetChannelType() == EAvaBroadcastChannelType::Preview)
+	if (InOutputInfo.IsRemote() && GetChannelType() == EAvaBroadcastChannelType::Preview)
 	{
 		return nullptr;
 	}
@@ -799,7 +799,15 @@ UMediaOutput* FAvaBroadcastOutputChannel::AddMediaOutput(const UClass* InMediaOu
 			}
 		}
 
-		AddMediaOutput(MediaOutput, InOutputInfo);
+		FAvaBroadcastMediaOutputInfo OutputInfo = InOutputInfo;
+
+		// If the device was not enumerated, we try to get the device name from the MediaOutput object.
+		if (OutputInfo.DeviceName.IsNone())
+		{
+			OutputInfo.DeviceName = FName(UE::AvaBroadcastOutputUtils::GetDeviceName(MediaOutput));
+		}
+
+		AddMediaOutput(MediaOutput, OutputInfo);
 		
 		return MediaOutput;
 	}
diff --git a/Engine/Plugins/Experimental/Avalanche/Source/AvalancheMedia/Private/Broadcast/OutputDevices/AvaBroadcastOutputUtils.cpp b/Engine/Plugins/Experimental/Avalanche/Source/AvalancheMedia/Private/Broadcast/OutputDevices/AvaBroadcastOutputUtils.cpp
--- a/Engine/Plugins/Experimental/Avalanche/Source/AvalancheMedia/Private/Broadcast/OutputDevices/AvaBroadcastOutputUtils.cpp
+++ b/Engine/Plugins/Experimental/Avalanche/Source/AvalancheMedia/Private/Broadcast/OutputDevices/AvaBroadcastOutputUtils.cpp
@@ -3,6 +3,7 @@
 #include "Broadcast/OutputDevices/AvaBroadcastOutputUtils.h"
 
 #include "AvaMediaModule.h"
+#include "Broadcast/OutputDevices/AvaGameViewportMediaOutput.h"
 #include "MediaOutput.h"
 #include "Playback/AvaPlaybackMessages.h"
 #include "Serialization/MemoryReader.h"
@@ -76,6 +77,11 @@ FName UE::AvaBroadcastOutputUtils::GetDeviceProviderName(const UMediaOutput* InM
 	return FName();
 }
 
+bool UE::AvaBroadcastOutputUtils::IsGameViewportOutput(const UMediaOutput* InMediaOutput)
+{
+	return InMediaOutput->IsA<UAvaGameViewportMediaOutput>();
+}
+
 FAvaBroadcastOutputData UE::AvaBroadcastOutputUtils::CreateMediaOutputData(UMediaOutput* InMediaOutput)
 {
 	FAvaBroadcastOutputData MediaOutputData;
diff --git a/Engine/Plugins/Experimental/Avalanche/Source/AvalancheMedia/Private/Broadcast/OutputDevices/AvaGameViewportMediaCapture.cpp b/Engine/Plugins/Experimental/Avalanche/Source/AvalancheMedia/Private/Broadcast/OutputDevices/AvaGameViewportMediaCapture.cpp
new file mode 100644
--- /dev/null
+++ b/Engine/Plugins/Experimental/Avalanche/Source/AvalancheMedia/Private/Broadcast/OutputDevices/AvaGameViewportMediaCapture.cpp
@@ -0,0 +1,25 @@
+// Copyright Epic Games, Inc. All Rights Reserved.
+
+#include "AvaGameViewportMediaCapture.h"
+
+#include "AvaGameViewportMediaOutput.h"
+#include "Engine/TextureRenderTarget2D.h"
+#include "Slate/SceneViewport.h"
+#include "TextureResource.h"
+
+bool UAvaGameViewportMediaCapture::InitializeCapture()
+{
+	return true;
+}
+
+bool UAvaGameViewportMediaCapture::PostInitializeCaptureViewport(TSharedPtr<FSceneViewport>& InSceneViewport)
+{
+	SetState(EMediaCaptureState::Capturing);
+	return true;
+}
+
+bool UAvaGameViewportMediaCapture::PostInitializeCaptureRenderTarget(UTextureRenderTarget2D* InRenderTarget)
+{
+	SetState(EMediaCaptureState::Capturing);
+	return true;
+}
\ No newline at end of file
diff --git a/Engine/Plugins/Experimental/Avalanche/Source/AvalancheMedia/Private/Broadcast/OutputDevices/AvaGameViewportMediaCapture.h b/Engine/Plugins/Experimental/Avalanche/Source/AvalancheMedia/Private/Broadcast/OutputDevices/AvaGameViewportMediaCapture.h
new file mode 100644
--- /dev/null
+++ b/Engine/Plugins/Experimental/Avalanche/Source/AvalancheMedia/Private/Broadcast/OutputDevices/AvaGameViewportMediaCapture.h
@@ -0,0 +1,22 @@
+// Copyright Epic Games, Inc. All Rights Reserved.
+
+#pragma once
+
+#include "HAL/CriticalSection.h"
+#include "MediaCapture.h"
+
+#include "AvaGameViewportMediaCapture.generated.h"
+
+UCLASS()
+class UAvaGameViewportMediaCapture : public UMediaCapture
+{
+	GENERATED_BODY()
+
+protected:
+	//~ Begin UMediaCapture
+	virtual bool ShouldCaptureRHIResource() const override { return true; }
+	virtual bool InitializeCapture() override;
+	virtual bool PostInitializeCaptureViewport(TSharedPtr<FSceneViewport>& InSceneViewport) override;
+	virtual bool PostInitializeCaptureRenderTarget(UTextureRenderTarget2D* InRenderTarget) override;
+	//~ End UMediaCapture
+};
diff --git a/Engine/Plugins/Experimental/Avalanche/Source/AvalancheMedia/Private/Broadcast/OutputDevices/AvaGameViewportMediaOutput.cpp b/Engine/Plugins/Experimental/Avalanche/Source/AvalancheMedia/Private/Broadcast/OutputDevices/AvaGameViewportMediaOutput.cpp
new file mode 100644
--- /dev/null
+++ b/Engine/Plugins/Experimental/Avalanche/Source/AvalancheMedia/Private/Broadcast/OutputDevices/AvaGameViewportMediaOutput.cpp
@@ -0,0 +1,17 @@
+// Copyright Epic Games, Inc. All Rights Reserved.
+
+#include "AvaGameViewportMediaOutput.h"
+
+#include "AvaGameViewportMediaCapture.h"
+
+DEFINE_LOG_CATEGORY(LogAvaGameViewportMedia);
+
+UMediaCapture* UAvaGameViewportMediaOutput::CreateMediaCaptureImpl()
+{
+	UMediaCapture* Result = NewObject<UAvaGameViewportMediaCapture>();
+	if (Result)
+	{
+		Result->SetMediaOutput(this);
+	}
+	return Result;
+}
\ No newline at end of file
diff --git a/Engine/Plugins/Experimental/Avalanche/Source/AvalancheMedia/Private/Broadcast/OutputDevices/AvaGameViewportMediaOutput.h b/Engine/Plugins/Experimental/Avalanche/Source/AvalancheMedia/Private/Broadcast/OutputDevices/AvaGameViewportMediaOutput.h
new file mode 100644
--- /dev/null
+++ b/Engine/Plugins/Experimental/Avalanche/Source/AvalancheMedia/Private/Broadcast/OutputDevices/AvaGameViewportMediaOutput.h
@@ -0,0 +1,32 @@
+// Copyright Epic Games, Inc. All Rights Reserved.
+
+#pragma once
+
+#include "MediaCapture.h"
+#include "AvaGameViewportMediaOutput.generated.h"
+
+DECLARE_LOG_CATEGORY_EXTERN(LogAvaGameViewportMedia, Log, All);
+
+UCLASS(BlueprintType, ClassGroup = "Motion Design Broadcast",
+	meta = (DisplayName = "Motion Design Game Viewport Media Output"))
+class UAvaGameViewportMediaOutput : public UMediaOutput
+{
+	GENERATED_BODY()
+
+	//~ Begin UMediaOutput
+	public:
+	virtual FIntPoint GetRequestedSize() const override { return UMediaOutput::RequestCaptureSourceSize; }
+	virtual EPixelFormat GetRequestedPixelFormat() const override { return EPixelFormat::PF_B8G8R8A8; }
+	virtual EMediaCaptureConversionOperation GetConversionOperation(EMediaCaptureSourceType InSourceType) const override { return EMediaCaptureConversionOperation::NONE; }
+protected:
+	virtual UMediaCapture* CreateMediaCaptureImpl() override;
+	//~ End UMediaOutput
+
+public:
+	/**
+	 * The source name is a property that exists to provide the "device name" for
+	 * displaying in the broadcast editor. 
+	 */
+	UPROPERTY()
+	FString SourceName = TEXT("Game Viewport");
+};
\ No newline at end of file
diff --git a/Engine/Plugins/Experimental/Avalanche/Source/AvalancheMedia/Private/Playable/AvaPlayableGroup.cpp b/Engine/Plugins/Experimental/Avalanche/Source/AvalancheMedia/Private/Playable/AvaPlayableGroup.cpp
--- a/Engine/Plugins/Experimental/Avalanche/Source/AvalancheMedia/Private/Playable/AvaPlayableGroup.cpp
+++ b/Engine/Plugins/Experimental/Avalanche/Source/AvalancheMedia/Private/Playable/AvaPlayableGroup.cpp
@@ -2,10 +2,14 @@
 
 #include "Playable/AvaPlayableGroup.h"
 
+#include "Engine/GameInstance.h"
 #include "Engine/ViewportStatsSubsystem.h"
-#include "Framework/AvaGameInstance.h"
+#include "Engine/World.h"
 #include "Playable/AvaPlayable.h"
 #include "Playable/AvaPlayableGroupManager.h"
+#include "Playable/PlayableGroups/AvaGameInstancePlayableGroup.h"
+#include "Playable/PlayableGroups/AvaGameViewportPlayableGroup.h"
+#include "Playable/PlayableGroups/AvaRemoteProxyPlayableGroup.h"
 #include "Playable/Transition/AvaPlayableTransition.h"
 #include "Playback/AvaPlaybackUtils.h"
 #include "UObject/Package.h"
@@ -27,99 +31,27 @@ namespace UE::AvaMedia::PlayableGroup::Private
 	}
 }
 
-UPackage* UAvaPlayableGroup::MakeGameInstancePackage(const FSoftObjectPath& InSourceAssetPath, const FName& InChannelName)
-{
-	// Remote Control Preset will be registered with the package name.
-	// We want a package name unique to the game instance and that will be
-	// human readable since it will show up in the Web Remote Control page.
-	
-	FString InstancePackageName = TEXT("/Temp/");	// Keep it short for web page.
-
-	// Using channel name, since we should have one instance per channel.
-	InstancePackageName += InChannelName.ToString();
-
-	// In order to keep things short, we remove "/Game" since we added the channel name instead.
-	FString InstanceSubPath;
-	if (!InSourceAssetPath.GetLongPackageName().Split(TEXT("/Game"), nullptr, &InstanceSubPath))
-	{
-		InstanceSubPath = InSourceAssetPath.GetLongPackageName();
-	}
-
-	// This may happen if the original asset path is not specified.
-	if (InstanceSubPath.IsEmpty())
-	{
-		// Add something to get a valid path at least.
-		InstanceSubPath = TEXT("InvalidAssetName");
-	}
-	
-	InstancePackageName += InstanceSubPath;
-
-	return MakeInstancePackage(InstancePackageName);
-}
-
-UPackage* UAvaPlayableGroup::MakeSharedInstancePackage(const FName& InChannelName)
-{
-	// Remote Control Preset will be registered with the package name.
-	// We want a package name unique to the game instance and that will be
-	// human readable since it will show up in the Web Remote Control page.
-	
-	FString SharedPackageName = TEXT("/Temp/");	// Keep it short for web page.
-
-	// Using channel name, since we should have one instance per channel.
-	SharedPackageName += InChannelName.ToString();
-
-	// Shared for all levels.
-	SharedPackageName += TEXT("/SharedLevels");
-	
-	return MakeInstancePackage(SharedPackageName);
-}
-
-UPackage* UAvaPlayableGroup::MakeInstancePackage(const FString& InInstancePackageName)
-{
-	UPackage* InstancePackage = CreatePackage(*InInstancePackageName);
-	if (InstancePackage)
-	{
-		InstancePackage->SetFlags(RF_Transient);
-	}
-	else
-	{
-		// Note: The outer will fallback to GEngine in that case.
-		UE_LOG(LogAvaPlayable, Error, TEXT("Unable to create package \"%s\" for Motion Design Game Instance."), *InInstancePackageName);
-	}
-	return InstancePackage;
-}
-
 UAvaPlayableGroup* UAvaPlayableGroup::MakePlayableGroup(UObject* InOuter, const FPlayableGroupCreationInfo& InPlayableGroupInfo)
 {
 	UObject* Outer = InOuter ? InOuter : GetTransientPackage();
 	
-	UAvaPlayableGroup* GameInstanceGroup;
+	UAvaPlayableGroup* GameInstanceGroup = nullptr;
 	if (InPlayableGroupInfo.bIsRemoteProxy)
 	{
 		// Remote Proxy group doesn't have a game instance.
-		GameInstanceGroup = NewObject<UAvaPlayableRemoteProxyGroup>(Outer);
+		GameInstanceGroup = NewObject<UAvaRemoteProxyPlayableGroup>(Outer);
 		GameInstanceGroup->ParentPlayableGroupManagerWeak = InPlayableGroupInfo.PlayableGroupManager;
 	}
 	else
 	{
-		GameInstanceGroup = NewObject<UAvaPlayableGroup>(Outer);
-		GameInstanceGroup->ParentPlayableGroupManagerWeak = InPlayableGroupInfo.PlayableGroupManager;
-		
-		if (InPlayableGroupInfo.bIsSharedGroup)
+		if (InPlayableGroupInfo.GameInstance)
 		{
-			GameInstanceGroup->GameInstancePackage = MakeSharedInstancePackage(InPlayableGroupInfo.ChannelName);
+			GameInstanceGroup = UAvaGameViewportPlayableGroup::Create(InOuter, InPlayableGroupInfo.GameInstance, InPlayableGroupInfo.PlayableGroupManager);
 		}
 		else
 		{
-			// We can create the package even if the name is null, it will have a generic name. But that will be considered an error.
-			if (!InPlayableGroupInfo.SourceAssetPath.IsNull())
-			{
-				UE_LOG(LogAvaPlayable, Error, TEXT("Creating game instance package for asset with unspecified name."));
-			}
-			GameInstanceGroup->GameInstancePackage = MakeGameInstancePackage(InPlayableGroupInfo.SourceAssetPath, InPlayableGroupInfo.ChannelName);
+			GameInstanceGroup = UAvaGameInstancePlayableGroup::Create(InOuter, InPlayableGroupInfo);
 		}
-		
-		GameInstanceGroup->GameInstance = UAvaGameInstance::Create(GameInstanceGroup->GameInstancePackage);
 	}
 	return GameInstanceGroup;
 }
@@ -273,60 +205,6 @@ bool UAvaPlayableGroup::HasTransitions() const
 	return !PlayableTransitions.IsEmpty();
 }
 
-bool UAvaPlayableGroup::ConditionalCreateWorld()
-{
-	if (!GameInstance)
-	{
-		return false;
-	}
-
-	bool bWorldWasCreated = false;
-
-	if (!GameInstance->IsWorldCreated())
-	{
-		bWorldWasCreated = GameInstance->CreateWorld();
-	}
-	
-	// Make sure we register our delegates to this world.
-	if (GameInstance->GetPlayWorld())
-	{
-		ConditionalRegisterWorldDelegates(GameInstance->GetPlayWorld());
-	}
-	
-	return bWorldWasCreated;
-}
-
-bool UAvaPlayableGroup::ConditionalBeginPlay(const FAvaInstancePlaySettings& InWorldPlaySettings)
-{
-	bool bHasBegunPlay = false;
-	if (!GameInstance)
-	{
-		return bHasBegunPlay;
-	}
-	
-	// Make sure we don't have pending unload or stop requests left over in the game instance.
-	GameInstance->CancelWorldRequests();
-
-	if (!GameInstance->IsWorldPlaying())
-	{
-		bHasBegunPlay = GameInstance->BeginPlayWorld(InWorldPlaySettings);
-	}
-	else
-	{
-		GameInstance->UpdateRenderTarget(InWorldPlaySettings.RenderTarget);
-		GameInstance->UpdateSceneViewportSize(InWorldPlaySettings.ViewportSize);
-	}
-	return bHasBegunPlay;
-}
-
-void UAvaPlayableGroup::RequestEndPlayWorld(bool bInForceImmediate)
-{
-	if (GameInstance)
-	{
-		GameInstance->RequestEndPlayWorld(bInForceImmediate);
-	}
-}
-
 void UAvaPlayableGroup::SetLastAppliedCameraPlayable(UAvaPlayable* InPlayable)
 {
 	LastAppliedCameraPlayableWeak = InPlayable;
@@ -361,56 +239,25 @@ bool UAvaPlayableGroup::UpdateCameraSetup()
 	return false;
 }
 
-bool UAvaPlayableGroup::IsWorldPlaying() const
-{
-	return GameInstance ? GameInstance->IsWorldPlaying() : false;
-}
-
-bool UAvaPlayableGroup::IsRenderTargetReady() const
-{
-	return GameInstance ? GameInstance->IsRenderTargetReady() : false;
-}
-
 UTextureRenderTarget2D* UAvaPlayableGroup::GetRenderTarget() const
 {
-	return GameInstance ? GameInstance->GetRenderTarget() : RenderTarget.Get();
+	return ManagedRenderTarget.Get();
 }
 
-UGameInstance* UAvaPlayableGroup::GetGameInstance() const
-{
-	return GameInstance;
-}
-	
-UWorld* UAvaPlayableGroup::GetPlayWorld() const
+UTextureRenderTarget2D* UAvaPlayableGroup::GetManagedRenderTarget() const
 {
-	return GameInstance ? GameInstance->GetPlayWorld() : nullptr;
+	return ManagedRenderTarget.Get();
 }
 
-bool UAvaPlayableGroup::ConditionalRequestUnloadWorld(bool bForceImmediate)
+
+void UAvaPlayableGroup::SetManagedRenderTarget(UTextureRenderTarget2D* InManageRenderTarget)
 {
-	if (!GameInstance)
-	{
-		return false;
-	}
-	
-	if (!HasPlayables())
-	{
-		UnregisterWorldDelegates(GameInstance->GetPlayWorld());
-		GameInstance->RequestUnloadWorld(bForceImmediate);
-		return true;
-	}
-	return false;
+	ManagedRenderTarget = InManageRenderTarget;
 }
 
-void UAvaPlayableGroup::QueueCameraCut()
+UWorld* UAvaPlayableGroup::GetPlayWorld() const
 {
-	if (GameInstance)
-	{
-		if (UAvaGameViewportClient* GameViewportClient = GameInstance->GetAvaGameViewportClient())
-		{
-			GameViewportClient->SetCameraCutThisFrame();
-		}
-	}
+	return GameInstance ? GameInstance->GetWorld() : nullptr;
 }
 
 void UAvaPlayableGroup::NotifyLevelStreaming(UAvaPlayable* InPlayable)
@@ -629,19 +476,4 @@ bool UAvaPlayableGroup::DisplayTransitions(FText& OutText, FLinearColor& OutColo
 	return false;
 }
 
-bool UAvaPlayableRemoteProxyGroup::ConditionalBeginPlay(const FAvaInstancePlaySettings& InWorldPlaySettings)
-{
-	if (!bIsPlaying)
-	{
-		bIsPlaying = true;
-		return true;
-	}
-	return false;
-}
-
-void UAvaPlayableRemoteProxyGroup::RequestEndPlayWorld(bool bInForceImmediate)
-{
-	bIsPlaying = false;
-}
-
 #undef LOCTEXT_NAMESPACE
\ No newline at end of file
diff --git a/Engine/Plugins/Experimental/Avalanche/Source/AvalancheMedia/Private/Playable/AvaPlayableGroupManager.cpp b/Engine/Plugins/Experimental/Avalanche/Source/AvalancheMedia/Private/Playable/AvaPlayableGroupManager.cpp
--- a/Engine/Plugins/Experimental/Avalanche/Source/AvalancheMedia/Private/Playable/AvaPlayableGroupManager.cpp
+++ b/Engine/Plugins/Experimental/Avalanche/Source/AvalancheMedia/Private/Playable/AvaPlayableGroupManager.cpp
@@ -3,20 +3,78 @@
 #include "Playable/AvaPlayableGroupManager.h"
 
 #include "Broadcast/AvaBroadcast.h"
+#include "Broadcast/OutputDevices/AvaBroadcastOutputUtils.h"
 #include "Broadcast/OutputDevices/AvaBroadcastRenderTargetMediaUtils.h"
+#include "Engine/Engine.h"
 #include "Framework/AvaGameInstance.h"
 #include "Misc/CoreDelegates.h"
 #include "Misc/TimeGuard.h"
 #include "Playable/AvaPlayableGroup.h"
+#include "Playable/PlayableGroups/AvaRemoteProxyPlayableGroup.h"
 #include "UObject/Package.h"
 
 #define LOCTEXT_NAMESPACE "AvaPlayableGroupManager"
 
-UAvaPlayableGroup* UAvaPlayableGroupChannelManager::GetOrCreateSharedLevelGroup(bool bInIsRemoteProxy)
+namespace UE::AvaMedia::PlayableGroupManager::Private
 {
-	if (UAvaPlayableGroup* ExistingPlayableGroup = bInIsRemoteProxy ? SharedRemoteProxyLevelGroupWeak.Get() : SharedLevelGroupWeak.Get())
+	UGameViewportClient* FindGameViewportClient()
 	{
-		return ExistingPlayableGroup;
+		if (GEngine->GameViewport)
+		{
+			return GEngine->GameViewport;
+		}
+		const TIndirectArray<FWorldContext>& WorldContexts = GEngine->GetWorldContexts();
+		for (const FWorldContext& Context : WorldContexts)
+		{
+			if ((Context.WorldType == EWorldType::PIE) && Context.World() != nullptr && Context.GameViewport != nullptr)
+			{
+				return Context.GameViewport;
+			}
+		}
+		return nullptr;
+	}
+
+	bool HasLocalGameViewportOutput(const FAvaBroadcastOutputChannel& InBroadcastChannel)
+	{
+		if (!InBroadcastChannel.IsValidChannel())
+		{
+			return false;
+		}
+		
+		for (const UMediaOutput* MediaOutput : InBroadcastChannel.GetMediaOutputs())
+		{
+			if (AvaBroadcastOutputUtils::IsGameViewportOutput(MediaOutput) && !InBroadcastChannel.IsMediaOutputRemote(MediaOutput))
+			{
+				return true;
+			}
+		}
+		return false;
+	}
+}
+
+UAvaPlayableGroup* UAvaPlayableGroupChannelManager::GetOrCreatePlayableGroup(UGameInstance* InExistingGameInstance, bool bInIsRemoteProxy)
+{
+	// See if there is a matching playable group.
+	for (const TWeakObjectPtr<UAvaPlayableGroup>& PlayableGroupWeak : PlayableGroupsWeak)
+	{
+		if (UAvaPlayableGroup* PlayableGroup = PlayableGroupWeak.Get())
+		{
+			if (InExistingGameInstance && PlayableGroup->GetGameInstance() == InExistingGameInstance)
+			{
+				return PlayableGroup;
+			}
+
+			const bool bIsRemoteProxy = PlayableGroup->IsA<UAvaRemoteProxyPlayableGroup>();
+			if (bInIsRemoteProxy && bIsRemoteProxy)
+			{
+				return PlayableGroup;
+			}
+			
+			if (InExistingGameInstance == nullptr && !bInIsRemoteProxy && !bIsRemoteProxy)
+			{
+				return PlayableGroup;
+			}
+		}
 	}
 	
 	UAvaPlayableGroup::FPlayableGroupCreationInfo PlayableGroupCreationInfo;
@@ -24,19 +82,13 @@ UAvaPlayableGroup* UAvaPlayableGroupChannelManager::GetOrCreateSharedLevelGroup(
 	PlayableGroupCreationInfo.ChannelName = ChannelName;
 	PlayableGroupCreationInfo.bIsRemoteProxy = bInIsRemoteProxy;
 	PlayableGroupCreationInfo.bIsSharedGroup = true;
-
+	PlayableGroupCreationInfo.GameInstance = InExistingGameInstance;
+	
 	UAvaPlayableGroup* NewPlayableGroup = UAvaPlayableGroup::MakePlayableGroup(GetPlayableGroupManager(), PlayableGroupCreationInfo);
 
 	// Keep track of the shared playable group.
-	if (bInIsRemoteProxy)
-	{
-		SharedRemoteProxyLevelGroupWeak = NewPlayableGroup;
-	}
-	else
-	{
-		SharedLevelGroupWeak = NewPlayableGroup;
-	}
-	
+	PlayableGroupsWeak.Add(NewPlayableGroup);
+
 	return NewPlayableGroup;
 }
 
@@ -45,10 +97,15 @@ UAvaPlayableGroupManager* UAvaPlayableGroupChannelManager::GetPlayableGroupManag
 	return Cast<UAvaPlayableGroupManager>(GetOuter());	
 }
 
+void UAvaPlayableGroupChannelManager::GetPlayableGroups(TArray<TWeakObjectPtr<UAvaPlayableGroup>>& OutGroups) const
+{
+	OutGroups.Append(PlayableGroupsWeak);
+}
+
+
 void UAvaPlayableGroupChannelManager::Shutdown()
 {
-	SharedRemoteProxyLevelGroupWeak.Reset();
-	SharedLevelGroupWeak.Reset();
+	PlayableGroupsWeak.Reset();
 }
 
 void UAvaPlayableGroupChannelManager::BeginDestroy()
@@ -96,6 +153,25 @@ UAvaPlayableGroupChannelManager* UAvaPlayableGroupManager::FindOrAddChannelManag
 	return ChannelManager;
 }
 
+UAvaPlayableGroup* UAvaPlayableGroupManager::GetOrCreateSharedPlayableGroup(const FName& InChannelName, bool bInIsRemoteProxy)
+{
+	using namespace UE::AvaMedia::PlayableGroupManager::Private;
+
+	UGameInstance* ExistingGameInstance = nullptr;
+	
+	const FAvaBroadcastOutputChannel& BroadcastChannel =  UAvaBroadcast::Get().GetCurrentProfile().GetChannel(InChannelName);
+	if (HasLocalGameViewportOutput(BroadcastChannel) && !bInIsRemoteProxy)
+	{
+		if (const UGameViewportClient* GameViewport = FindGameViewportClient())
+		{
+			ExistingGameInstance = GameViewport->GetGameInstance();
+		}
+	}
+
+	UAvaPlayableGroupChannelManager* ChannelManager = FindOrAddChannelManager(InChannelName);
+	return ChannelManager ? ChannelManager->GetOrCreatePlayableGroup(ExistingGameInstance, bInIsRemoteProxy) : nullptr;
+}
+
 void UAvaPlayableGroupManager::RegisterForLevelStreamingUpdate(UAvaPlayableGroup* InPlayableGroup)
 {
 	if (ensure(!bIsUpdatingStreaming))
@@ -128,6 +204,29 @@ void UAvaPlayableGroupManager::UnregisterFromTransitionTicking(UAvaPlayableGroup
 	}
 }
 
+TArray<TWeakObjectPtr<UAvaPlayableGroup>> UAvaPlayableGroupManager::GetPlayableGroups(FName InChannelName) const
+{
+	TArray<TWeakObjectPtr<UAvaPlayableGroup>> Groups;
+	if (InChannelName.IsNone())
+	{
+		for (const TPair<FName, TObjectPtr<UAvaPlayableGroupChannelManager>>& ChannelManager : ChannelManagers)
+		{
+			if (ChannelManager.Value)
+			{
+				ChannelManager.Value->GetPlayableGroups(Groups);
+			}
+		}
+	}
+	else
+	{
+		if (const UAvaPlayableGroupChannelManager* ChannelManager = FindChannelManager(InChannelName))
+		{
+			ChannelManager->GetPlayableGroups(Groups);
+		}
+	}
+	return Groups;
+}
+
 void UAvaPlayableGroupManager::BeginDestroy()
 {
 	Shutdown();
diff --git a/Engine/Plugins/Experimental/Avalanche/Source/AvalancheMedia/Private/Playable/AvaPlayableGroupManager.h b/Engine/Plugins/Experimental/Avalanche/Source/AvalancheMedia/Private/Playable/AvaPlayableGroupManager.h
--- a/Engine/Plugins/Experimental/Avalanche/Source/AvalancheMedia/Private/Playable/AvaPlayableGroupManager.h
+++ b/Engine/Plugins/Experimental/Avalanche/Source/AvalancheMedia/Private/Playable/AvaPlayableGroupManager.h
@@ -10,6 +10,7 @@
 class UAvaGameInstance;
 class UAvaPlayableGroup;
 class UAvaPlayableGroupManager;
+class UGameInstance;
 
 /**
  * Manager for the shared playable groups per channel.
@@ -20,9 +21,11 @@ class UAvaPlayableGroupChannelManager : public UObject
 	GENERATED_BODY()
 
 public:
-	UAvaPlayableGroup* GetOrCreateSharedLevelGroup(bool bInIsRemoteProxy);
+	UAvaPlayableGroup* GetOrCreatePlayableGroup(UGameInstance* InExistingGameInstance, bool bInIsRemoteProxy);
 
 	UAvaPlayableGroupManager* GetPlayableGroupManager() const;
+
+	void GetPlayableGroups(TArray<TWeakObjectPtr<UAvaPlayableGroup>>& OutGroups) const;
 	
 protected:
 	//~ Begin UObject
@@ -35,20 +38,11 @@ protected:
 	FName ChannelName;
 
 	/**
-	 * This is the shared playable group for this channel.
-	 * 
-	 * For now all the levels go in the same group for a given channel.
-	 * This a weak ptr because it is only a cache. The ownership
-	 * of the group is with the playable objects.
-	 */
-	TWeakObjectPtr<UAvaPlayableGroup> SharedLevelGroupWeak;
-
-	/**
-	 * For remote proxy playables, we need a specialized group
-	 * that doesn't implement the local game instance but still
-	 * provides the correct logic to emulate functionality.
+	 * Shared playable groups for this channel.
+	 * The manager only keeps weak pointers.
+	 * Groups are owned by their respective playables. 
 	 */
-	TWeakObjectPtr<UAvaPlayableGroup> SharedRemoteProxyLevelGroupWeak;
+	TArray<TWeakObjectPtr<UAvaPlayableGroup>> PlayableGroupsWeak;
 
 	friend class UAvaPlayableGroupManager;
 };
@@ -78,11 +72,13 @@ public:
 
 	UAvaPlayableGroupChannelManager* FindOrAddChannelManager(const FName& InChannelName);
 
-	UAvaPlayableGroup* GetOrCreateSharedLevelGroup(const FName& InChannelName, bool bInIsRemoteProxy)
-	{
-		UAvaPlayableGroupChannelManager* ChannelManager = FindOrAddChannelManager(InChannelName);
-		return ChannelManager ? ChannelManager->GetOrCreateSharedLevelGroup(bInIsRemoteProxy) : nullptr;
-	}
+	/**
+	 * @brief Returns the playable group for the given channel.
+	 * @param InChannelName Channel name, should correspond to a configured broadcast channel.
+	 * @param bInIsRemoteProxy Indicate if the implementation is local or remote.
+	 * @return Playable group.
+	 */
+	UAvaPlayableGroup* GetOrCreateSharedPlayableGroup(const FName& InChannelName, bool bInIsRemoteProxy);
 
 	void RegisterForLevelStreamingUpdate(UAvaPlayableGroup* InPlayableGroup);
 	void UnregisterFromLevelStreamingUpdate(UAvaPlayableGroup* InPlayableGroup);
@@ -90,17 +86,23 @@ public:
 	void RegisterForTransitionTicking(UAvaPlayableGroup* InPlayableGroup);
 	void UnregisterFromTransitionTicking(UAvaPlayableGroup* InPlayableGroup);
 
+	/**
+	 * Returns playable groups part of the given channel, or all playable groups if channel is None.
+	 */
+	TArray<TWeakObjectPtr<UAvaPlayableGroup>> GetPlayableGroups(FName InChannelName = NAME_None) const;
+
 protected:
 	//~ Begin UObject
 	virtual void BeginDestroy() override;
 	//~ End UObject
 
+	// Move to UAvaGameInstancePlayableGroup
 	void OnGameInstanceEndPlay(UAvaGameInstance* InGameInstance, FName InChannelName);
 
 	void UpdateLevelStreaming();
 	
 	void TickTransitions(double InDeltaSeconds);
-	
+
 protected:
 	UPROPERTY(Transient)
 	TMap<FName, TObjectPtr<UAvaPlayableGroupChannelManager>> ChannelManagers;
diff --git a/Engine/Plugins/Experimental/Avalanche/Source/AvalancheMedia/Private/Playable/AvaPlayableGroupSubsystem.cpp b/Engine/Plugins/Experimental/Avalanche/Source/AvalancheMedia/Private/Playable/AvaPlayableGroupSubsystem.cpp
new file mode 100644
--- /dev/null
+++ b/Engine/Plugins/Experimental/Avalanche/Source/AvalancheMedia/Private/Playable/AvaPlayableGroupSubsystem.cpp
@@ -0,0 +1,32 @@
+// Copyright Epic Games, Inc. All Rights Reserved.
+
+#include "AvaPlayableGroupSubsystem.h"
+
+#include "Playable/AvaPlayableGroupManager.h"
+#include "Playable/PlayableGroups/AvaGameViewportPlayableGroup.h"
+
+void UAvaPlayableGroupSubsystem::Initialize(FSubsystemCollectionBase& Collection)
+{
+	Super::Initialize(Collection);
+}
+
+void UAvaPlayableGroupSubsystem::Deinitialize()
+{
+	if (PlayableGroupManager)
+	{
+		TArray<TWeakObjectPtr<UAvaPlayableGroup>> PlayableGroupsWeak = PlayableGroupManager->GetPlayableGroups();
+		for (const TWeakObjectPtr<UAvaPlayableGroup>& PlayableGroupWeak : PlayableGroupsWeak)
+		{
+			UAvaGameViewportPlayableGroup* GameViewportPlayableGroup = Cast<UAvaGameViewportPlayableGroup>(PlayableGroupWeak.Get());
+	
+			if (GameViewportPlayableGroup && GameViewportPlayableGroup->GetGameInstance() == GetGameInstance())
+			{
+				GameViewportPlayableGroup->DetachGameInstance();
+			}
+		}
+	}
+	
+	PlayableGroupManager = nullptr;
+	
+	Super::Deinitialize();
+}
diff --git a/Engine/Plugins/Experimental/Avalanche/Source/AvalancheMedia/Private/Playable/AvaPlayableGroupSubsystem.h b/Engine/Plugins/Experimental/Avalanche/Source/AvalancheMedia/Private/Playable/AvaPlayableGroupSubsystem.h
new file mode 100644
--- /dev/null
+++ b/Engine/Plugins/Experimental/Avalanche/Source/AvalancheMedia/Private/Playable/AvaPlayableGroupSubsystem.h
@@ -0,0 +1,24 @@
+// Copyright Epic Games, Inc. All Rights Reserved.
+
+#pragma once
+
+#include "Subsystems/GameInstanceSubsystem.h"
+
+#include "AvaPlayableGroupSubsystem.generated.h"
+
+class UAvaPlayableGroupManager;
+
+UCLASS()
+class UAvaPlayableGroupSubsystem : public UGameInstanceSubsystem
+{
+	GENERATED_BODY()
+
+public:
+	//~ Begin UGameInstanceSubsystem
+	virtual void Initialize(FSubsystemCollectionBase& Collection) override;
+	virtual void Deinitialize() override;
+	//~ End UGameInstanceSubsystem
+
+	UPROPERTY()
+	TObjectPtr<UAvaPlayableGroupManager> PlayableGroupManager;
+};
diff --git a/Engine/Plugins/Experimental/Avalanche/Source/AvalancheMedia/Private/Playable/PlayableGroups/AvaGameInstancePlayableGroup.cpp b/Engine/Plugins/Experimental/Avalanche/Source/AvalancheMedia/Private/Playable/PlayableGroups/AvaGameInstancePlayableGroup.cpp
new file mode 100644
--- /dev/null
+++ b/Engine/Plugins/Experimental/Avalanche/Source/AvalancheMedia/Private/Playable/PlayableGroups/AvaGameInstancePlayableGroup.cpp
@@ -0,0 +1,215 @@
+// Copyright Epic Games, Inc. All Rights Reserved.
+
+#include "AvaGameInstancePlayableGroup.h"
+
+#include "Framework/AvaGameInstance.h"
+#include "Playable/AvaPlayable.h"
+#include "Playable/AvaPlayableGroupManager.h"
+#include "UObject/Package.h"
+
+namespace UE::AvaMedia::PlayableGroup::Private
+{
+	UPackage* MakeInstancePackage(const FString& InInstancePackageName)
+	{
+		UPackage* InstancePackage = CreatePackage(*InInstancePackageName);
+		if (InstancePackage)
+		{
+			InstancePackage->SetFlags(RF_Transient);
+		}
+		else
+		{
+			// Note: The outer will fallback to GEngine in that case.
+			UE_LOG(LogAvaPlayable, Error, TEXT("Unable to create package \"%s\" for Motion Design Game Instance."), *InInstancePackageName);
+		}
+		return InstancePackage;
+	}
+
+	UPackage* MakeGameInstancePackage(const FSoftObjectPath& InSourceAssetPath, const FName& InChannelName)
+	{
+		// Remote Control Preset will be registered with the package name.
+		// We want a package name unique to the game instance and that will be
+		// human readable since it will show up in the Web Remote Control page.
+		
+		FString InstancePackageName = TEXT("/Temp/");	// Keep it short for web page.
+
+		// Using channel name, since we should have one instance per channel.
+		InstancePackageName += InChannelName.ToString();
+
+		// In order to keep things short, we remove "/Game" since we added the channel name instead.
+		FString InstanceSubPath;
+		if (!InSourceAssetPath.GetLongPackageName().Split(TEXT("/Game"), nullptr, &InstanceSubPath))
+		{
+			InstanceSubPath = InSourceAssetPath.GetLongPackageName();
+		}
+
+		// This may happen if the original asset path is not specified.
+		if (InstanceSubPath.IsEmpty())
+		{
+			// Add something to get a valid path at least.
+			InstanceSubPath = TEXT("InvalidAssetName");
+		}
+		
+		InstancePackageName += InstanceSubPath;
+
+		return MakeInstancePackage(InstancePackageName);
+	}
+
+	UPackage* MakeSharedInstancePackage(const FName& InChannelName)
+	{
+		// Remote Control Preset will be registered with the package name.
+		// We want a package name unique to the game instance and that will be
+		// human readable since it will show up in the Web Remote Control page.
+		
+		FString SharedPackageName = TEXT("/Temp/");	// Keep it short for web page.
+
+		// Using channel name, since we should have one instance per channel.
+		SharedPackageName += InChannelName.ToString();
+
+		// Shared for all levels.
+		SharedPackageName += TEXT("/SharedLevels");
+		
+		return MakeInstancePackage(SharedPackageName);
+	}
+}
+
+UAvaGameInstancePlayableGroup* UAvaGameInstancePlayableGroup::Create(UObject* InOuter, const FPlayableGroupCreationInfo& InPlayableGroupInfo)
+{
+	using namespace UE::AvaMedia::PlayableGroup::Private;
+	
+	UAvaGameInstancePlayableGroup* GameInstanceGroup = NewObject<UAvaGameInstancePlayableGroup>(InOuter);
+	GameInstanceGroup->ParentPlayableGroupManagerWeak = InPlayableGroupInfo.PlayableGroupManager;
+		
+	if (InPlayableGroupInfo.bIsSharedGroup)
+	{
+		GameInstanceGroup->GameInstancePackage = MakeSharedInstancePackage(InPlayableGroupInfo.ChannelName);
+	}
+	else
+	{
+		// We can create the package even if the name is null, it will have a generic name. But that will be considered an error.
+		if (!InPlayableGroupInfo.SourceAssetPath.IsNull())
+		{
+			UE_LOG(LogAvaPlayable, Error, TEXT("Creating game instance package for asset with unspecified name."));
+		}
+		GameInstanceGroup->GameInstancePackage = MakeGameInstancePackage(InPlayableGroupInfo.SourceAssetPath, InPlayableGroupInfo.ChannelName);
+	}
+		
+	GameInstanceGroup->GameInstance = UAvaGameInstance::Create(GameInstanceGroup->GameInstancePackage);
+
+	return GameInstanceGroup;
+}
+
+UAvaGameInstance* UAvaGameInstancePlayableGroup::GetAvaGameInstance() const
+{
+	return Cast<UAvaGameInstance>(GameInstance);
+}
+
+bool UAvaGameInstancePlayableGroup::ConditionalCreateWorld()
+{
+	UAvaGameInstance* AvaGameInstance = GetAvaGameInstance();
+	
+	if (!AvaGameInstance)
+	{
+		return false;
+	}
+
+	bool bWorldWasCreated = false;
+
+	if (!AvaGameInstance->IsWorldCreated())
+	{
+		bWorldWasCreated = AvaGameInstance->CreateWorld();
+	}
+	
+	// Make sure we register our delegates to this world.
+	if (AvaGameInstance->GetPlayWorld())
+	{
+		ConditionalRegisterWorldDelegates(AvaGameInstance->GetPlayWorld());
+	}
+	
+	return bWorldWasCreated;
+}
+
+bool UAvaGameInstancePlayableGroup::ConditionalBeginPlay(const FAvaInstancePlaySettings& InWorldPlaySettings)
+{
+	UAvaGameInstance* AvaGameInstance = GetAvaGameInstance();
+	
+	bool bHasBegunPlay = false;
+	if (!AvaGameInstance)
+	{
+		return bHasBegunPlay;
+	}
+	
+	// Make sure we don't have pending unload or stop requests left over in the game instance.
+	AvaGameInstance->CancelWorldRequests();
+
+	if (!AvaGameInstance->IsWorldPlaying())
+	{
+		bHasBegunPlay = AvaGameInstance->BeginPlayWorld(InWorldPlaySettings);
+	}
+	else
+	{
+		AvaGameInstance->UpdateRenderTarget(InWorldPlaySettings.RenderTarget);
+		AvaGameInstance->UpdateSceneViewportSize(InWorldPlaySettings.ViewportSize);
+	}
+	return bHasBegunPlay;
+}
+
+void UAvaGameInstancePlayableGroup::RequestEndPlayWorld(bool bInForceImmediate)
+{
+	if (UAvaGameInstance* AvaGameInstance = GetAvaGameInstance())
+	{
+		AvaGameInstance->RequestEndPlayWorld(bInForceImmediate);
+	}
+}
+
+bool UAvaGameInstancePlayableGroup::IsWorldPlaying() const
+{
+	const UAvaGameInstance* AvaGameInstance = GetAvaGameInstance();
+	return AvaGameInstance ? AvaGameInstance->IsWorldCreated() : false;
+}
+
+bool UAvaGameInstancePlayableGroup::IsRenderTargetReady() const
+{
+	const UAvaGameInstance* AvaGameInstance = GetAvaGameInstance();
+	return AvaGameInstance ? AvaGameInstance->IsRenderTargetReady() : false;
+}
+
+UTextureRenderTarget2D* UAvaGameInstancePlayableGroup::GetRenderTarget() const
+{
+	const UAvaGameInstance* AvaGameInstance = GetAvaGameInstance();
+	return AvaGameInstance ? AvaGameInstance->GetRenderTarget() : ManagedRenderTarget.Get();
+}
+
+UWorld* UAvaGameInstancePlayableGroup::GetPlayWorld() const
+{
+	const UAvaGameInstance* AvaGameInstance = GetAvaGameInstance();
+	return AvaGameInstance ? AvaGameInstance->GetPlayWorld() : nullptr;
+}
+
+bool UAvaGameInstancePlayableGroup::ConditionalRequestUnloadWorld(bool bForceImmediate)
+{
+	UAvaGameInstance* AvaGameInstance = GetAvaGameInstance();
+	
+	if (!AvaGameInstance)
+	{
+		return false;
+	}
+	
+	if (!HasPlayables())
+	{
+		UnregisterWorldDelegates(AvaGameInstance->GetPlayWorld());
+		AvaGameInstance->RequestUnloadWorld(bForceImmediate);
+		return true;
+	}
+	return false;
+}
+
+void UAvaGameInstancePlayableGroup::QueueCameraCut()
+{
+	if (const UAvaGameInstance* AvaGameInstance = GetAvaGameInstance())
+	{
+		if (UAvaGameViewportClient* GameViewportClient = AvaGameInstance->GetAvaGameViewportClient())
+		{
+			GameViewportClient->SetCameraCutThisFrame();
+		}
+	}
+}
\ No newline at end of file
diff --git a/Engine/Plugins/Experimental/Avalanche/Source/AvalancheMedia/Private/Playable/PlayableGroups/AvaGameInstancePlayableGroup.h b/Engine/Plugins/Experimental/Avalanche/Source/AvalancheMedia/Private/Playable/PlayableGroups/AvaGameInstancePlayableGroup.h
new file mode 100644
--- /dev/null
+++ b/Engine/Plugins/Experimental/Avalanche/Source/AvalancheMedia/Private/Playable/PlayableGroups/AvaGameInstancePlayableGroup.h
@@ -0,0 +1,44 @@
+// Copyright Epic Games, Inc. All Rights Reserved.
+
+#pragma once
+
+#include "Playable/AvaPlayableGroup.h"
+#include "AvaGameInstancePlayableGroup.generated.h"
+
+class UAvaGameInstance;
+
+/**
+ * This playable group implements a container that owns a
+ * motion design game instance. It is intended to be used for
+ * broadcast only self contained world simulation and rendering.
+ *
+ * Render Target:
+ * The render target provided as argument of BeginPlay
+ * is used to override the game viewport's canvas. So the game will render
+ * in the provided render target directly.
+ */
+UCLASS()
+class UAvaGameInstancePlayableGroup : public UAvaPlayableGroup
+{
+	GENERATED_BODY()
+
+public:
+	static UAvaGameInstancePlayableGroup* Create(UObject* InOuter, const FPlayableGroupCreationInfo& InPlayableGroupInfo);
+	
+	UAvaGameInstance* GetAvaGameInstance() const;
+
+	//~ Begin UAvaPlayableGroup
+	virtual bool ConditionalCreateWorld() override;
+	virtual bool ConditionalBeginPlay(const FAvaInstancePlaySettings& InWorldPlaySettings) override;
+	virtual void RequestEndPlayWorld(bool bInForceImmediate) override;
+	virtual bool ConditionalRequestUnloadWorld(bool bForceImmediate) override;
+	virtual void QueueCameraCut() override;
+	virtual bool IsWorldPlaying() const override;
+	virtual bool IsRenderTargetReady() const override;
+	virtual UTextureRenderTarget2D* GetRenderTarget() const override;
+	virtual UWorld* GetPlayWorld() const override;
+	//~ End UAvaPlayableGroup
+
+	UPROPERTY(Transient)
+	TObjectPtr<UPackage> GameInstancePackage;
+};
diff --git a/Engine/Plugins/Experimental/Avalanche/Source/AvalancheMedia/Private/Playable/PlayableGroups/AvaGameViewportPlayableGroup.cpp b/Engine/Plugins/Experimental/Avalanche/Source/AvalancheMedia/Private/Playable/PlayableGroups/AvaGameViewportPlayableGroup.cpp
new file mode 100644
--- /dev/null
+++ b/Engine/Plugins/Experimental/Avalanche/Source/AvalancheMedia/Private/Playable/PlayableGroups/AvaGameViewportPlayableGroup.cpp
@@ -0,0 +1,149 @@
+// Copyright Epic Games, Inc. All Rights Reserved.
+
+#include "AvaGameViewportPlayableGroup.h"
+
+#include "Broadcast/OutputDevices/AvaBroadcastRenderTargetMediaUtils.h"
+#include "Engine/GameInstance.h"
+#include "Engine/TextureRenderTarget2D.h"
+#include "Framework/AvaGameInstance.h"
+#include "Playable/AvaPlayableGroupManager.h"
+#include "Playable/AvaPlayableGroupSubsystem.h"
+#include "TextureResource.h"
+#if WITH_EDITOR
+#include "LevelEditorViewport.h"
+#endif
+
+UAvaGameViewportPlayableGroup* UAvaGameViewportPlayableGroup::Create(
+	UObject* InOuter,
+	UGameInstance* InGameInstance,
+	UAvaPlayableGroupManager* InPlayableGroupManager)
+{
+	if (InGameInstance)
+	{
+		UAvaGameViewportPlayableGroup* PlayableGroup = NewObject<UAvaGameViewportPlayableGroup>(InOuter);
+		PlayableGroup->ParentPlayableGroupManagerWeak = InPlayableGroupManager;
+		PlayableGroup->GameInstance = InGameInstance;
+		if (UAvaPlayableGroupSubsystem* PlayableGroupSubsystem = InGameInstance->GetSubsystem<UAvaPlayableGroupSubsystem>())
+		{
+			PlayableGroupSubsystem->PlayableGroupManager = InPlayableGroupManager;
+		}
+		return PlayableGroup;
+	}
+	return nullptr;
+}
+
+UAvaGameViewportPlayableGroup::~UAvaGameViewportPlayableGroup()
+{
+	UnregisterViewportDelegates();
+}
+
+void UAvaGameViewportPlayableGroup::DetachGameInstance()
+{
+	UnregisterViewportDelegates();
+	GameInstance = nullptr;
+	bIsPlaying = false;
+}
+
+bool UAvaGameViewportPlayableGroup::ConditionalBeginPlay(const FAvaInstancePlaySettings& InWorldPlaySettings)
+{
+	ViewportRenderTargetWeak = InWorldPlaySettings.RenderTarget;
+
+	if (!bIsPlaying)
+	{
+		RegisterViewportDelegates();
+		bIsPlaying = true;
+		return true;
+	}
+	return false;
+}
+
+void UAvaGameViewportPlayableGroup::RequestEndPlayWorld(bool bInForceImmediate)
+{
+	UnregisterViewportDelegates();
+	bIsPlaying = false;
+}
+
+UTextureRenderTarget2D* UAvaGameViewportPlayableGroup::GetRenderTarget() const
+{
+	if (UTextureRenderTarget2D* ViewportRenderTarget = ViewportRenderTargetWeak.Get())
+	{
+		return ViewportRenderTarget;
+	}
+	return ManagedRenderTarget;
+}
+
+void UAvaGameViewportPlayableGroup::RegisterViewportDelegates()
+{
+	if (!UGameViewportClient::OnViewportRendered().IsBoundToObject(this))
+	{
+		UGameViewportClient::OnViewportRendered().AddUObject(this, &UAvaGameViewportPlayableGroup::OnViewportRendered);
+	}
+}
+
+void UAvaGameViewportPlayableGroup::UnregisterViewportDelegates()
+{
+	UGameViewportClient::OnViewportRendered().RemoveAll(this);
+}
+
+void UAvaGameViewportPlayableGroup::OnViewportRendered(FViewport* InViewport)
+{
+	// Main path for capturing the viewport
+	if (GameInstance)
+	{
+		if (const UGameViewportClient* GameViewportClient = GameInstance->GetGameViewportClient())
+		{
+			if (GameViewportClient->Viewport == InViewport)
+			{
+				CopyViewportRenderTarget(InViewport);
+				return;
+			}
+		}
+	}
+
+#if WITH_EDITOR
+	// Fallback for "in editor" PIE (i.e. when PIE window is part of level editor).
+	if (GCurrentLevelEditingViewportClient && GCurrentLevelEditingViewportClient->Viewport == InViewport)
+	{
+		CopyViewportRenderTarget(InViewport);
+	}
+#endif
+}
+
+void UAvaGameViewportPlayableGroup::CopyViewportRenderTarget(FViewport* InViewport) const
+{
+	if (!ViewportRenderTargetWeak.IsValid())
+	{
+		return;
+	}
+
+	TWeakObjectPtr<UTextureRenderTarget2D> RenderTargetWeak = ViewportRenderTargetWeak;
+	
+	const FTextureRHIRef SourceRef = InViewport->GetRenderTargetTexture();
+	if (!SourceRef.IsValid())
+	{
+		return;
+	}
+	
+	ENQUEUE_RENDER_COMMAND(CopyViewportRenderTarget)([SourceRef, RenderTargetWeak](FRHICommandListImmediate& RHICmdList) 
+	{
+		UTextureRenderTarget2D* Target = RenderTargetWeak.Get();
+		if (!Target)
+		{
+			return;
+		}
+
+		const FTextureRenderTargetResource* TargetResource = Target->GetRenderTargetResource();
+		if (!TargetResource)
+		{
+			return;
+		}
+
+		const FTextureRHIRef TargetRef = TargetResource->GetRenderTargetTexture();
+		if (!TargetRef.IsValid())
+		{
+			return;
+		}
+
+		UE::AvaBroadcastRenderTargetMediaUtils::CopyTexture(RHICmdList, SourceRef, TargetRef);
+	});
+}
\ No newline at end of file
diff --git a/Engine/Plugins/Experimental/Avalanche/Source/AvalancheMedia/Private/Playable/PlayableGroups/AvaGameViewportPlayableGroup.h b/Engine/Plugins/Experimental/Avalanche/Source/AvalancheMedia/Private/Playable/PlayableGroups/AvaGameViewportPlayableGroup.h
new file mode 100644
--- /dev/null
+++ b/Engine/Plugins/Experimental/Avalanche/Source/AvalancheMedia/Private/Playable/PlayableGroups/AvaGameViewportPlayableGroup.h
@@ -0,0 +1,57 @@
+// Copyright Epic Games, Inc. All Rights Reserved.
+
+#pragma once
+
+#include "Playable/AvaPlayableGroup.h"
+#include "AvaGameViewportPlayableGroup.generated.h"
+
+class FViewport;
+
+/**
+ * This playable group implements a wrapper for a pre-existing game instance.
+ * It is intended to be used for injecting motion design's playable in an existing
+ * game instance that is not specifically implemented for motion design playback.
+ *
+ * Render Target:
+ * The render target provided as argument of BeginPlay
+ * is used to copy the viewport render in when the viewport is rendered,
+ * but before EndFrame, so it can still be used with media captures for
+ * chaining outputs to other devices.
+ */
+UCLASS()
+class UAvaGameViewportPlayableGroup : public UAvaPlayableGroup
+{
+	GENERATED_BODY()
+
+public:
+	static UAvaGameViewportPlayableGroup* Create(
+		UObject* InOuter,
+		UGameInstance* InGameInstance,
+		UAvaPlayableGroupManager* InPlayableGroupManager);
+
+	virtual ~UAvaGameViewportPlayableGroup() override;
+
+	void DetachGameInstance();
+	
+	//~ Begin UAvaPlayableGroup
+	virtual bool ConditionalBeginPlay(const FAvaInstancePlaySettings& InWorldPlaySettings) override;
+	virtual void RequestEndPlayWorld(bool bInForceImmediate) override;
+	virtual UTextureRenderTarget2D* GetRenderTarget() const override;
+	//~ End UAvaPlayableGroup
+	
+private:
+	void RegisterViewportDelegates();
+	void UnregisterViewportDelegates();
+	void CopyViewportRenderTarget(FViewport* InViewport) const;
+	void OnViewportRendered(FViewport* InViewport);
+	
+private:
+	/** Internal copy of the viewport render target. */
+	TWeakObjectPtr<UTextureRenderTarget2D> ViewportRenderTargetWeak;
+
+	/**
+	 * Keep track if the playable group is "playing".
+	 * ConditionalBeginPlay will return true only if the group was not already playing. 
+	 */
+	bool bIsPlaying = false;
+};
\ No newline at end of file
diff --git a/Engine/Plugins/Experimental/Avalanche/Source/AvalancheMedia/Private/Playable/PlayableGroups/AvaRemoteProxyPlayableGroup.cpp b/Engine/Plugins/Experimental/Avalanche/Source/AvalancheMedia/Private/Playable/PlayableGroups/AvaRemoteProxyPlayableGroup.cpp
new file mode 100644
--- /dev/null
+++ b/Engine/Plugins/Experimental/Avalanche/Source/AvalancheMedia/Private/Playable/PlayableGroups/AvaRemoteProxyPlayableGroup.cpp
@@ -0,0 +1,18 @@
+// Copyright Epic Games, Inc. All Rights Reserved.
+
+#include "AvaRemoteProxyPlayableGroup.h"
+
+bool UAvaRemoteProxyPlayableGroup::ConditionalBeginPlay(const FAvaInstancePlaySettings& InWorldPlaySettings)
+{
+	if (!bIsPlaying)
+	{
+		bIsPlaying = true;
+		return true;
+	}
+	return false;
+}
+
+void UAvaRemoteProxyPlayableGroup::RequestEndPlayWorld(bool bInForceImmediate)
+{
+	bIsPlaying = false;
+}
\ No newline at end of file
diff --git a/Engine/Plugins/Experimental/Avalanche/Source/AvalancheMedia/Private/Playable/PlayableGroups/AvaRemoteProxyPlayableGroup.h b/Engine/Plugins/Experimental/Avalanche/Source/AvalancheMedia/Private/Playable/PlayableGroups/AvaRemoteProxyPlayableGroup.h
new file mode 100644
--- /dev/null
+++ b/Engine/Plugins/Experimental/Avalanche/Source/AvalancheMedia/Private/Playable/PlayableGroups/AvaRemoteProxyPlayableGroup.h
@@ -0,0 +1,23 @@
+// Copyright Epic Games, Inc. All Rights Reserved.
+
+#pragma once
+
+#include "Playable/AvaPlayableGroup.h"
+#include "AvaRemoteProxyPlayableGroup.generated.h"
+
+/**
+ *	Remote Proxy Playable Group doesn't have a game instance because it is executed remotely.
+ */
+UCLASS()
+class UAvaRemoteProxyPlayableGroup : public UAvaPlayableGroup
+{
+	GENERATED_BODY()
+	
+public:
+	//~ Begin UAvaPlayableGroup
+	virtual bool ConditionalBeginPlay(const FAvaInstancePlaySettings& InWorldPlaySettings) override;
+	virtual void RequestEndPlayWorld(bool bInForceImmediate) override;
+	//~ End UAvaPlayableGroup
+	
+	bool bIsPlaying = false;
+};
\ No newline at end of file
diff --git a/Engine/Plugins/Experimental/Avalanche/Source/AvalancheMedia/Private/Playable/Playables/AvaPlayableLevelStreaming.cpp b/Engine/Plugins/Experimental/Avalanche/Source/AvalancheMedia/Private/Playable/Playables/AvaPlayableLevelStreaming.cpp
--- a/Engine/Plugins/Experimental/Avalanche/Source/AvalancheMedia/Private/Playable/Playables/AvaPlayableLevelStreaming.cpp
+++ b/Engine/Plugins/Experimental/Avalanche/Source/AvalancheMedia/Private/Playable/Playables/AvaPlayableLevelStreaming.cpp
@@ -546,7 +546,7 @@ bool UAvaPlayableLevelStreaming::InitPlayable(const FPlayableCreationInfo& InPla
 {
 	// For now, we share all the levels in the same instance group. We may do sub-grouping later.
 	PlayableGroup = InPlayableInfo.PlayableGroup ?
-		InPlayableInfo.PlayableGroup : InPlayableInfo.PlayableGroupManager->GetOrCreateSharedLevelGroup(InPlayableInfo.ChannelName, false);
+		InPlayableInfo.PlayableGroup : InPlayableInfo.PlayableGroupManager->GetOrCreateSharedPlayableGroup(InPlayableInfo.ChannelName, false);
 
 	const bool bInitSuccess = Super::InitPlayable(InPlayableInfo);
 	
diff --git a/Engine/Plugins/Experimental/Avalanche/Source/AvalancheMedia/Private/Playable/Playables/AvaPlayableRemoteProxy.cpp b/Engine/Plugins/Experimental/Avalanche/Source/AvalancheMedia/Private/Playable/Playables/AvaPlayableRemoteProxy.cpp
--- a/Engine/Plugins/Experimental/Avalanche/Source/AvalancheMedia/Private/Playable/Playables/AvaPlayableRemoteProxy.cpp
+++ b/Engine/Plugins/Experimental/Avalanche/Source/AvalancheMedia/Private/Playable/Playables/AvaPlayableRemoteProxy.cpp
@@ -188,7 +188,7 @@ bool UAvaPlayableRemoteProxy::InitPlayable(const FPlayableCreationInfo& InPlayab
 	switch (InPlayableInfo.SourceAsset.GetAssetType())
 	{
 	case EMotionDesignAssetType::World:
-		PlayableGroup = InPlayableInfo.PlayableGroupManager->GetOrCreateSharedLevelGroup(InPlayableInfo.ChannelName, bIsRemoteProxy);
+		PlayableGroup = InPlayableInfo.PlayableGroupManager->GetOrCreateSharedPlayableGroup(InPlayableInfo.ChannelName, bIsRemoteProxy);
 		break;
 	default:
 		UE_LOG(LogAvaPlayable, Error, TEXT("Asset \"%s\" is an unsupported type."), *InPlayableInfo.SourceAsset.ToSoftObjectPath().ToString());
diff --git a/Engine/Plugins/Experimental/Avalanche/Source/AvalancheMedia/Private/Playback/AvaPlaybackGraph.cpp b/Engine/Plugins/Experimental/Avalanche/Source/AvalancheMedia/Private/Playback/AvaPlaybackGraph.cpp
--- a/Engine/Plugins/Experimental/Avalanche/Source/AvalancheMedia/Private/Playback/AvaPlaybackGraph.cpp
+++ b/Engine/Plugins/Experimental/Avalanche/Source/AvalancheMedia/Private/Playback/AvaPlaybackGraph.cpp
@@ -865,12 +865,12 @@ bool UAvaPlaybackGraph::RefreshPreview(const FAvaSoftAssetPtr& InSourceAsset, FA
 
 	UAvaPlayable* const Playable = FindOrLoadPlayable(InSourceAsset, InChannelName);
 
-	if (!Playable)
+	if (!Playable || !Playable->GetPlayableGroup())
 	{
 		return false;
 	}
 	
-	UTextureRenderTarget2D* RenderTarget = Playable->GetPlayableGroup()->RenderTarget;
+	UTextureRenderTarget2D* RenderTarget = Playable->GetPlayableGroup()->GetManagedRenderTarget();
 
 	if (!RenderTarget)
 	{
@@ -883,7 +883,7 @@ bool UAvaPlaybackGraph::RefreshPreview(const FAvaSoftAssetPtr& InSourceAsset, FA
 			FAvaBroadcastOutputChannel::GetDefaultMediaOutputFormat(),
 			IAvaMediaModule::Get().GetBroadcastSettings().GetChannelClearColor());
 
-		Playable->GetPlayableGroup()->RenderTarget = RenderTarget;
+		Playable->GetPlayableGroup()->SetManagedRenderTarget(RenderTarget);
 	}
 	
 	const FIntPoint ViewportSize = UE::AvaMediaRenderTargetUtils::GetRenderTargetSize(RenderTarget);
@@ -914,7 +914,7 @@ UTextureRenderTarget2D* UAvaPlaybackGraph::UpdatePlaybackRenderTarget(const UAva
 			const FName ChannelName = InChannel.GetChannelName();
 
 			// See if we have a managed render target for this asset in this channel.
-			RenderTarget = InPlayable->GetPlayableGroup()->RenderTarget;
+			RenderTarget = InPlayable->GetPlayableGroup()->GetManagedRenderTarget();
 			
 			if (!IsValid(RenderTarget))
 			{
@@ -922,7 +922,7 @@ UTextureRenderTarget2D* UAvaPlaybackGraph::UpdatePlaybackRenderTarget(const UAva
 				RenderTarget = UE::AvaMediaRenderTargetUtils::CreateDefaultRenderTarget(PlaybackRenderTargetBaseName);
 
 				// This render target is now managed by the playable's instance group.
-				InPlayable->GetPlayableGroup()->RenderTarget = RenderTarget;
+				InPlayable->GetPlayableGroup()->SetManagedRenderTarget(RenderTarget);
 			}
 		}
 		
diff --git a/Engine/Plugins/Experimental/Avalanche/Source/AvalancheMedia/Public/Playable/AvaPlayableGroup.h b/Engine/Plugins/Experimental/Avalanche/Source/AvalancheMedia/Public/Playable/AvaPlayableGroup.h
--- a/Engine/Plugins/Experimental/Avalanche/Source/AvalancheMedia/Public/Playable/AvaPlayableGroup.h
+++ b/Engine/Plugins/Experimental/Avalanche/Source/AvalancheMedia/Public/Playable/AvaPlayableGroup.h
@@ -10,7 +10,6 @@
 #include "UObject/WeakInterfacePtr.h"
 #include "AvaPlayableGroup.generated.h"
 
-class UAvaGameInstance;
 class UAvaPlayable;
 class UAvaPlayableGroupManager;
 class UAvaPlayableTransition;
@@ -19,15 +18,20 @@ class UTextureRenderTarget2D;
 struct FAvaInstancePlaySettings;
 
 /**
- * This defines a game instance playable group.
- * It tracks and manage the game instance and playables state.
+ * This class defines the interface and base of a playable group.
  *
- * The design goal of this class is to remove all the playback management
- * from UAvaGameInstance and move it to the playable framework. This
- * should allow us to hook the playback framework to any game instance, including PIE
- * so it can work with any work flow (editor, PIE, game, nDisplay, etc).
+ * A playable group is intended to group playables according to the
+ * underlying rendering implementation. In most cases, it corresponds
+ * to a game instance, either owned or not, local to the process or remote.
  *
- * Also, ideally, the playable class itself should be "game instance" agnostic
+ * It tracks and manage the playables state, transitions and
+ * visibility constraints.
+ *
+ * The design goal of this class is to allow hooking the playable framework
+ * to any game instance, including PIE so it can work with any work flow
+ * (editor, PIE, game, nDisplay, etc).
+ *
+ * Ideally, the playable class itself should be "game instance" agnostic
  * and do all it's bidding on it's container through this class.
  */
 UCLASS()
@@ -50,8 +54,10 @@ public:
 		FName ChannelName;
 		/** Indicate if the group is for remote proxy playables. */
 		bool bIsRemoteProxy = false;
-		/** Indicate if the group is shared. */
+		/** Indicate if the group is shared for multiple playables. If so, it will be registered in the given playable group manager. */
 		bool bIsSharedGroup = false;
+		/** Existing Game Instance. In this case, the playable group will not own the game instance. */
+		UGameInstance* GameInstance = nullptr;
 	};
 	
 	static UAvaPlayableGroup* MakePlayableGroup(UObject* InOuter, const FPlayableGroupCreationInfo& InPlayableGroupInfo);
@@ -84,46 +90,66 @@ public:
 	 * Creates the game instance's world if it wasn't already.
 	 * @return true if the world was created. false if nothing was done.
 	 */
-	bool ConditionalCreateWorld();
+	virtual bool ConditionalCreateWorld() { return true; }
 
 	/**
 	 * Begin playing the game instance's world if it wasn't already.
 	 * @return true if the BeginPlay was done (i.e. on the state transition only), false otherwise.
 	 */
-	virtual bool ConditionalBeginPlay(const FAvaInstancePlaySettings& InWorldPlaySettings);
+	virtual bool ConditionalBeginPlay(const FAvaInstancePlaySettings& InWorldPlaySettings) { return false; }
 
-	virtual void RequestEndPlayWorld(bool bInForceImmediate);
+	virtual void RequestEndPlayWorld(bool bInForceImmediate) {}
 
 	/** Keep track of the last playable that applied it's camera in the viewport/controller. */
 	void SetLastAppliedCameraPlayable(UAvaPlayable* InPlayable);
 	
-	/** Search for the first remaining playing playable and use it's camera. */
+	/** Determines which camera, from currently playing playables, should be used and use it. */
 	bool UpdateCameraSetup();
 
-	bool IsWorldPlaying() const;
+	virtual bool IsWorldPlaying() const { return true; }
 
-	bool IsRenderTargetReady() const;
+	virtual bool IsRenderTargetReady() const { return true; }
 
 	/**
-	 * Current logic for the render target, use the game instance's if there, fallback to internal one if not.
+	 * Current logic for the render target: use the game instance's if present, fallback to internal one if not.
 	 */
-	UTextureRenderTarget2D* GetRenderTarget() const;
+	virtual UTextureRenderTarget2D* GetRenderTarget() const;
 
-	/** Return a vanilla game instance, we want to eventually support any game instance. */
-	UGameInstance* GetGameInstance() const;
-	
-	UWorld* GetPlayWorld() const;
+	/**
+	 * Returns the currently managed render target.
+	 */
+	virtual UTextureRenderTarget2D* GetManagedRenderTarget() const;
+
+	/**
+	 * The playback graph determines if this playable group will render in a broadcast channel's
+	 * render target or an offscreen one. In the later case, the playable group keeps
+	 * track of that render target.
+	 * 
+	 * @remark The playable group does not automatically render in the current "managed" render target.
+	 * The render target this group will render into is determined by the arguments of ConditionalBeginPlay.
+	 */
+	virtual void SetManagedRenderTarget(UTextureRenderTarget2D* InManageRenderTarget);
+
+	/**
+	 * Returns this group's game instance, if it has one.
+	 */
+	UGameInstance* GetGameInstance() const { return GameInstance; }
+
+	/**
+	* Returns this group's play world, if it has one.
+	*/
+	virtual UWorld* GetPlayWorld() const;
 	
 	/**
 	 * Unloads the game instance's world if no more playables are loaded.
 	 * @return true if the world was unloaded. false if nothing was done.
 	 */ 
-	bool ConditionalRequestUnloadWorld(bool bForceImmediate);
+	virtual bool ConditionalRequestUnloadWorld(bool bForceImmediate) { return true; }
 
 	/**
 	 * @brief Queue a camera cut for the next rendered frame.
 	 */
-	void QueueCameraCut();
+	virtual void QueueCameraCut() {}
 
 	/**
 	 * @brief Notify the playable group that a playable is loading an asset.
@@ -150,21 +176,16 @@ protected:
 	bool DisplayPlayingAssets(FText& OutText, FLinearColor& OutColor);
 	bool DisplayTransitions(FText& OutText, FLinearColor& OutColor);
 
-	static UPackage* MakeGameInstancePackage(const FSoftObjectPath& InSourceAssetPath, const FName& InChannelName);
-	static UPackage* MakeSharedInstancePackage(const FName& InChannelName);
-	static UPackage* MakeInstancePackage(const FString& InInstancePackageName);
-
-public:	
-	UPROPERTY(Transient)
-	TObjectPtr<UTextureRenderTarget2D> RenderTarget;	// Optional
-	
+protected:
+	/**
+	 * Managed Render Target for this playable group.
+	 */
 	UPROPERTY(Transient)
-	TObjectPtr<UAvaGameInstance> GameInstance;
+	TObjectPtr<UTextureRenderTarget2D> ManagedRenderTarget;
 
 	UPROPERTY(Transient)
-	TObjectPtr<UPackage> GameInstancePackage;
+	TObjectPtr<UGameInstance> GameInstance;
 
-protected:
 	/** PlayableGroup Manager handling this playable group. */
 	TWeakObjectPtr<UAvaPlayableGroupManager> ParentPlayableGroupManagerWeak;
 	
@@ -207,19 +228,4 @@ protected:
 	TArray<FVisibilityRequest> VisibilityRequests;
 
 	TArray<TWeakInterfacePtr<IAvaPlayableVisibilityConstraint>> VisibilityConstraints;
-};
-
-/**
- *	Remote Proxy Playable Group doesn't have a game instance.
- */
-UCLASS()
-class AVALANCHEMEDIA_API UAvaPlayableRemoteProxyGroup : public UAvaPlayableGroup
-{
-	GENERATED_BODY()
-	
-public:
-	virtual bool ConditionalBeginPlay(const FAvaInstancePlaySettings& InWorldPlaySettings) override;
-	virtual void RequestEndPlayWorld(bool bInForceImmediate) override;
-	
-	bool bIsPlaying = false;
 };
\ No newline at end of file
diff --git a/Engine/Plugins/Experimental/Avalanche/Source/AvalancheMediaEditor/Private/Broadcast/ChannelGrid/AvaBroadcastOutputTileItem.cpp b/Engine/Plugins/Experimental/Avalanche/Source/AvalancheMediaEditor/Private/Broadcast/ChannelGrid/AvaBroadcastOutputTileItem.cpp
--- a/Engine/Plugins/Experimental/Avalanche/Source/AvalancheMediaEditor/Private/Broadcast/ChannelGrid/AvaBroadcastOutputTileItem.cpp
+++ b/Engine/Plugins/Experimental/Avalanche/Source/AvalancheMediaEditor/Private/Broadcast/ChannelGrid/AvaBroadcastOutputTileItem.cpp
@@ -197,9 +197,10 @@ void FAvaBroadcastOutputTileItem::OnBroadcastChanged(EAvaBroadcastChange InChang
 void FAvaBroadcastOutputTileItem::UpdateInfo()
 {
 	MediaOutputDisplayText = FindLatestDisplayText();
-	
-	const EAvaBroadcastOutputState OutputState = GetChannel().GetMediaOutputState(MediaOutput.Get());
-	const EAvaBroadcastIssueSeverity Severity = GetChannel().GetMediaOutputIssueSeverity(OutputState, MediaOutput.Get());
+
+	const FAvaBroadcastOutputChannel& Channel = GetChannel();
+	const EAvaBroadcastOutputState OutputState = Channel.GetMediaOutputState(MediaOutput.Get());
+	const EAvaBroadcastIssueSeverity Severity = Channel.GetMediaOutputIssueSeverity(OutputState, MediaOutput.Get());
 	switch(OutputState)
 	{
 	case EAvaBroadcastOutputState::Offline:
@@ -235,6 +236,26 @@ void FAvaBroadcastOutputTileItem::UpdateInfo()
 	}
 
 	FString AllMessages;
+
+	// Include the status of the server in the tooltip.
+	const FAvaBroadcastMediaOutputInfo& OutputInfo = Channel.GetMediaOutputInfo(MediaOutput.Get());
+	if (OutputInfo.IsValid())
+	{
+		if (OutputInfo.IsRemote())
+		{
+			AllMessages.Append(FString::Printf(TEXT("Remote Output on server \"%s\".\n"), *OutputInfo.ServerName));
+		}
+		else
+		{
+			AllMessages.Append(TEXT("Local Output.\n"));
+		}
+	}
+	else
+	{
+		AllMessages.Append(TEXT("Invalid Output Info.\n"));
+	}
+	
+	
 	switch(Severity)
 	{
 	case EAvaBroadcastIssueSeverity::Errors:
@@ -244,7 +265,14 @@ void FAvaBroadcastOutputTileItem::UpdateInfo()
 		AllMessages.Append(TEXT("Warning: \n"));
 		break;
 	default:
-		AllMessages.Append(TEXT("Healthy\n"));
+		if (OutputState == EAvaBroadcastOutputState::Offline)
+		{
+			AllMessages.Append(TEXT("Offline\n"));
+		}
+		else
+		{
+			AllMessages.Append(TEXT("Healthy\n"));
+		}
 		break;
 	}
 	
@@ -260,8 +288,27 @@ void FAvaBroadcastOutputTileItem::UpdateInfo()
 FText FAvaBroadcastOutputTileItem::FindLatestDisplayText() const
 {
 	check(MediaOutput.IsValid());
-	const FString DeviceName = UE::AvaBroadcastOutputUtils::GetDeviceName(MediaOutput.Get());
-	return !DeviceName.IsEmpty() ? FText::FromString(DeviceName) : FText::FromName(MediaOutput->GetFName());
+
+	FString ServerName;
+	
+	const FAvaBroadcastOutputChannel& Channel = GetChannel();
+	if (Channel.IsValidChannel())
+	{
+		const FAvaBroadcastMediaOutputInfo& OutputInfo = Channel.GetMediaOutputInfo(MediaOutput.Get());
+		if (OutputInfo.IsValid() && OutputInfo.IsRemote())
+		{
+			ServerName = OutputInfo.ServerName + TEXT(": ");
+		}
+	}
+
+	FString DeviceName = UE::AvaBroadcastOutputUtils::GetDeviceName(MediaOutput.Get());
+	
+	if (DeviceName.IsEmpty())
+	{
+		DeviceName = MediaOutput->GetFName().ToString();
+	}
+
+	return ServerName.IsEmpty() ? FText::FromString(DeviceName) : FText::FromString(ServerName + DeviceName);
 }
 
 #undef LOCTEXT_NAMESPACE
