[MotionDesign] Check for UObjectInitialized() before calling UAvaSceneRigSubsystem::UnregisterSupportedActorClasses

diff --git a/Engine/Plugins/Experimental/Avalanche/Source/AvalancheEditor/Private/AvaEditorModule.cpp b/Engine/Plugins/Experimental/Avalanche/Source/AvalancheEditor/Private/AvaEditorModule.cpp
--- a/Engine/Plugins/Experimental/Avalanche/Source/AvalancheEditor/Private/AvaEditorModule.cpp
+++ b/Engine/Plugins/Experimental/Avalanche/Source/AvalancheEditor/Private/AvaEditorModule.cpp
@@ -135,9 +135,9 @@ void FAvaEditorModule::ShutdownModule()
 		{
 			IAvaOutlinerModule::Get().UnregisterOverriddenIcon<FAvaOutlinerActor>(AAvaShapeActor::StaticClass()->GetFName());
 		}
-	}
 
-	UAvaSceneRigSubsystem::UnregisterSupportedActorClasses(DefaultSceneRigActorClasses());
+		UAvaSceneRigSubsystem::UnregisterSupportedActorClasses(DefaultSceneRigActorClasses());
+	}
 }
 
 void FAvaEditorModule::CreateAvaLevelEditor()
