[Motion Design] Rundown Server API: Adding PageSummary and FriendlyName.

diff --git a/Engine/Plugins/Experimental/Avalanche/Source/AvalancheMediaEditor/Private/Rundown/AvaRundownMessages.h b/Engine/Plugins/Experimental/Avalanche/Source/AvalancheMediaEditor/Private/Rundown/AvaRundownMessages.h
--- a/Engine/Plugins/Experimental/Avalanche/Source/AvalancheMediaEditor/Private/Rundown/AvaRundownMessages.h
+++ b/Engine/Plugins/Experimental/Avalanche/Source/AvalancheMediaEditor/Private/Rundown/AvaRundownMessages.h
@@ -242,6 +242,12 @@ public:
 	UPROPERTY()
 	FString PageName;
 
+	UPROPERTY()
+	FString PageSummary;
+
+	UPROPERTY()
+	FString FriendlyName;
+
 	UPROPERTY()
 	bool IsTemplate = false;
 
diff --git a/Engine/Plugins/Experimental/Avalanche/Source/AvalancheMediaEditor/Private/Rundown/AvaRundownServer.cpp b/Engine/Plugins/Experimental/Avalanche/Source/AvalancheMediaEditor/Private/Rundown/AvaRundownServer.cpp
--- a/Engine/Plugins/Experimental/Avalanche/Source/AvalancheMediaEditor/Private/Rundown/AvaRundownServer.cpp
+++ b/Engine/Plugins/Experimental/Avalanche/Source/AvalancheMediaEditor/Private/Rundown/AvaRundownServer.cpp
@@ -45,6 +45,8 @@ namespace UE::AvaRundownServer::Private
 		FAvaRundownPageInfo PageInfo;
 		PageInfo.PageId = InPage.GetPageId();
 		PageInfo.PageName = InPage.GetPageName();
+		PageInfo.PageSummary = InPage.GetPageSummary().ToString();
+		PageInfo.FriendlyName = InPage.GetPageFriendlyName().ToString();
 		PageInfo.IsTemplate = InPage.IsTemplate();
 		PageInfo.TemplateId = InPage.GetTemplateId();
 		PageInfo.CombinedTemplateIds = InPage.GetCombinedTemplateIds();
