[Motion Design] Font Loading Mac OS

Adds system font loading for Mac OS.
Loading *.ttf and *.otf files as for Windows, with the available typefaces.

diff --git a/Engine/Plugins/Experimental/Avalanche/Source/AvalancheTextEditor/AvalancheTextEditor.Build.cs b/Engine/Plugins/Experimental/Avalanche/Source/AvalancheTextEditor/AvalancheTextEditor.Build.cs
--- a/Engine/Plugins/Experimental/Avalanche/Source/AvalancheTextEditor/AvalancheTextEditor.Build.cs
+++ b/Engine/Plugins/Experimental/Avalanche/Source/AvalancheTextEditor/AvalancheTextEditor.Build.cs
@@ -65,5 +65,10 @@ public class AvalancheTextEditor : ModuleRules
 				"AvalancheEditorCore"
 			});
 		}
+
+		if (Target.Platform == UnrealTargetPlatform.Mac)
+		{
+			PublicFrameworks.AddRange(new string[] { "Foundation", "CoreText", "ApplicationServices"});
+		}
 	}
 }
diff --git a/Engine/Plugins/Experimental/Avalanche/Source/AvalancheTextEditor/Private/Platforms/Mac/MacSystemFontLoading.cpp b/Engine/Plugins/Experimental/Avalanche/Source/AvalancheTextEditor/Private/Platforms/Mac/MacSystemFontLoading.cpp
--- a/Engine/Plugins/Experimental/Avalanche/Source/AvalancheTextEditor/Private/Platforms/Mac/MacSystemFontLoading.cpp
+++ b/Engine/Plugins/Experimental/Avalanche/Source/AvalancheTextEditor/Private/Platforms/Mac/MacSystemFontLoading.cpp
@@ -1,14 +1,123 @@
 // Copyright Epic Games, Inc. All Rights Reserved.
 
 #include "MacSystemFontLoading.h"
+#include "Font/AvaFontManagerSubsystem.h"
+#include <ApplicationServices/ApplicationServices.h>
+#include <CoreText/CoreText.h>
+#include <Foundation/Foundation.h>
 
 namespace UE::Ava::Private::Fonts
 {
 	void GetSystemFontInfo(TMap<FString, FSystemFontsRetrieveParams>& OutFontsInfo)
 	{
+		// Cache font collection for fast use in future calls
+		static CTFontCollectionRef Collection = nullptr;
+		if (!Collection)
+		{
+			Collection = CTFontCollectionCreateFromAvailableFonts(nullptr);
+		}
+
+		CFArrayRef FontDescriptors = CTFontCollectionCreateMatchingFontDescriptors(Collection);
+		CFIndex FontsCount = CFArrayGetCount(FontDescriptors);
+		for (CFIndex CurrentFontIndex = 0; CurrentFontIndex < FontsCount; CurrentFontIndex++)
+		{
+			CTFontDescriptorRef FontDescriptorRef = (CTFontDescriptorRef) CFArrayGetValueAtIndex(FontDescriptors, CurrentFontIndex);
+			
+			FString FontPath;
+			if (NSURL* URLAttribute = (NSURL*) CTFontDescriptorCopyAttribute(FontDescriptorRef, kCTFontURLAttribute))
+			{
+				if (NSString* URLString = URLAttribute.absoluteString)
+				{
+					FontPath = FString(URLString.UTF8String);
+				}
+				
+				CFRelease(URLAttribute);
+			}
+			
+			if (FontPath.IsEmpty())
+			{
+				continue;
+			}
+
+			FString FontFaceName;
+			if (NSString* FontFaceNameAttribute = (NSString*) CTFontDescriptorCopyAttribute(FontDescriptorRef, kCTFontStyleNameAttribute))
+			{
+				FontFaceName = FString(FontFaceNameAttribute.UTF8String);
+				CFRelease(FontFaceNameAttribute);
+			}
+			else
+			{
+				continue;
+			}
+
+			FString FontFamilyName;
+			if (NSString* FontFamilyNameAttribute = (NSString*) CTFontDescriptorCopyAttribute(FontDescriptorRef, kCTFontFamilyNameAttribute))
+			{
+				FontFamilyName = FString(FontFamilyNameAttribute.UTF8String);
+				CFRelease(FontFamilyNameAttribute);
+			}
+			else
+			{
+				continue;
+			}
+					
+			FString SanitizedFamilyName = FontFamilyName;
+			UAvaFontManagerSubsystem::SanitizeString(SanitizedFamilyName);
+		
+			FontFaceName.RemoveFromStart(FontFamilyName);
+			FontFaceName.RemoveFromStart(TEXT(" "));
+
+			if (FontFaceName.IsEmpty())
+			{
+				FontFaceName = TEXT("Regular");
+			}	
+
+			// Clean the url to properly find the file
+			FontPath.RemoveFromStart("file://");
+			FontPath.ReplaceInline(TEXT("%20"), TEXT(" "));
+
+			if (FPaths::FileExists(FontPath) && UAvaFontManagerSubsystem::IsSupportedFontFile(FontPath))
+			{
+				if (!OutFontsInfo.Contains(SanitizedFamilyName))
+				{
+					FSystemFontsRetrieveParams FontRetrieveParams;
+					FontRetrieveParams.FontFamilyName = FontFamilyName;
+					FontRetrieveParams.AddFontFace(FontFaceName, FontPath);
+					OutFontsInfo.Add(SanitizedFamilyName, FontRetrieveParams);
+				}
+				else
+				{
+					OutFontsInfo[SanitizedFamilyName].AddFontFace(FontFaceName, FontPath);
+				}
+			}
+		}
+		
+		CFRelease(FontDescriptors);
 	}
 
 	void ListAvailableFontFiles()
 	{
+		TMap<FString, FSystemFontsRetrieveParams> FontsInfoMap;
+		GetSystemFontInfo(FontsInfoMap);
+
+		if (FontsInfoMap.IsEmpty())
+		{
+			return;
+		}
+
+		UE_LOG(LogAvaFont, Log, TEXT("Font Manager Subsystem: listing system fonts and their typefaces:"));
+		for (const TPair<FString, FSystemFontsRetrieveParams>& FontsInfoPair : FontsInfoMap)
+		{
+			const FSystemFontsRetrieveParams& FontParameters = FontsInfoPair.Value;
+			UE_LOG(LogAvaFont, Log, TEXT("== Font: %s =="), *FontParameters.FontFamilyName);
+
+			int32 FontFaceIndex = 0;
+			for (const FString& FontFaceName : FontParameters.GetFontFaceNames())
+			{
+				const FString& FontFacePath = FontParameters.GetFontFacePaths()[FontFaceIndex];
+				UE_LOG(LogAvaFont, Log, TEXT("\t\tFace Name: %s found at %s"), *FontFaceName, *FontFacePath);
+				FontFaceIndex++;
+			}
+		}
 	}
 }
