Motion Design
- Fixed an issue where a jump cut was being acted upon each time the camera cut track was updated.

diff --git a/Engine/Plugins/Experimental/Avalanche/Source/AvalancheEditor/Private/Viewport/AvaLevelViewportExtension.cpp b/Engine/Plugins/Experimental/Avalanche/Source/AvalancheEditor/Private/Viewport/AvaLevelViewportExtension.cpp
--- a/Engine/Plugins/Experimental/Avalanche/Source/AvalancheEditor/Private/Viewport/AvaLevelViewportExtension.cpp
+++ b/Engine/Plugins/Experimental/Avalanche/Source/AvalancheEditor/Private/Viewport/AvaLevelViewportExtension.cpp
@@ -475,6 +475,19 @@ void FAvaLevelViewportExtension::UnbindCameraCutDelegate()
 
 void FAvaLevelViewportExtension::OnCameraCut(UObject* InCameraObject, bool bInJumpCut)
 {
+	/**
+	 * The event is always called with bInJumpCut set to true. However, it's only really a jump cut
+	 * if ther camera object changes. Here we're checking to see if there was a change and ignoring
+	 * the camera cut update if the camera hasn't changed.
+	 *
+	 * Always updating the camera and sending the jump cut signal to the base viewport class causes
+	 * rendering artifacts (such as flickering).
+	 */
+	if (InCameraObject == LastCameraCutActorWeak.Get())
+	{
+		return;
+	}
+
 	LastCameraCutActorWeak = nullptr;
 
 	if (!IsValid(InCameraObject))
