MotionDesign : ClonerEffector
- Add missing visualizer support for torus and radial effector type
- Fix effector property undo not updating dynamic mesh visualizer
- Saving uassets against 5.5 with specific CL version to avoid warning

#virtualized

diff --git a/Engine/Plugins/Experimental/Avalanche/Source/AvalancheComponentVisualizers/Private/AvaVisBase.cpp b/Engine/Plugins/Experimental/Avalanche/Source/AvalancheComponentVisualizers/Private/AvaVisBase.cpp
--- a/Engine/Plugins/Experimental/Avalanche/Source/AvalancheComponentVisualizers/Private/AvaVisBase.cpp
+++ b/Engine/Plugins/Experimental/Avalanche/Source/AvalancheComponentVisualizers/Private/AvaVisBase.cpp
@@ -41,7 +41,7 @@ bool FAvaVisualizerBase::GetCustomInputCoordinateSystem(const FEditorViewportCli
 	return true;
 }
 
-bool FAvaVisualizerBase::HandleInputDelta(FEditorViewportClient* InViewportClient, FViewport* InViewport, 
+bool FAvaVisualizerBase::HandleInputDelta(FEditorViewportClient* InViewportClient, FViewport* InViewport,
 	FVector& InDeltaTranslate, FRotator& InDeltaRotate, FVector& InDeltaScale)
 {
 	if (!GetEditedComponent() || !GetEditedComponent()->GetOwner())
@@ -328,12 +328,12 @@ float FAvaVisualizerBase::GetIconSizeScale(const FSceneView* InView, const FVect
 
 	static constexpr float DefaultSize = 12.f;
 	float Size = DefaultSize;
-	
-	if (UAvaComponentVisualizersSettings* Settings = UAvaComponentVisualizersSettings::Get())
+
+	if (const UAvaComponentVisualizersSettings* Settings = UAvaComponentVisualizersSettings::Get())
 	{
 		if (Settings->SpriteSize > 0)
 		{
-			Size = UAvaComponentVisualizersSettings::Get()->SpriteSize;
+			Size = Settings->SpriteSize;
 		}
 	}
 
@@ -407,7 +407,7 @@ FVector FAvaVisualizerBase::GetLocalVector(FEditorViewportClient* InViewport, co
 	if (WidgetMode == UE::Widget::WM_None)
 	{
 		return OutVector;
-	}	
+	}
 
 	EAxisList::Type AxisList = InViewport->GetCurrentWidgetAxis(); // Fallback value
 
@@ -465,7 +465,7 @@ bool FAvaVisualizerBase::IsMouseOverComponent(const UActorComponent* Component,
 	};
 	FVector WorldOrigin;
 	FVector WorldDirection;
-	
+
 	InView->DeprojectFVector2D(MouseLocation, WorldOrigin, WorldDirection);
 
 	FVector CameraRotationVector = (WorldOrigin - InView->ViewLocation).GetUnsafeNormal();
@@ -521,7 +521,14 @@ void FAvaVisualizerBase::TrackingStoppedInternal(FEditorViewportClient* InViewpo
 	SnapOperation.Reset();
 }
 
-void FAvaVisualizerBase::NotifyPropertyModified(UObject* InObject, FProperty* InProperty, 
+void FAvaVisualizerBase::ModifyProperty(UObject* InObject, FProperty* InProperty, EPropertyChangeType::Type InPropertyChangeType, TFunctionRef<void()> InFunction)
+{
+	bHasBeenModified = InObject->Modify();
+	InFunction();
+	NotifyPropertyModified(InObject, InProperty, InPropertyChangeType);
+}
+
+void FAvaVisualizerBase::NotifyPropertyModified(UObject* InObject, FProperty* InProperty,
 	EPropertyChangeType::Type InPropertyChangeType, FProperty* InMemberProperty)
 {
 	TArray<FProperty*> Properties;
@@ -570,7 +577,7 @@ void FAvaVisualizerBase::NotifyPropertyChainModified(UObject* InObject, FPropert
 	{
 		return;
 	}
-	
+
 	FPropertyChangedEvent PropertyChangedEvent(InProperty);
 	FEditPropertyChain PropertyChain;
 	FPropertyChangedChainEvent PropertyChangedChainEvent(PropertyChain, PropertyChangedEvent);
@@ -579,10 +586,10 @@ void FAvaVisualizerBase::NotifyPropertyChainModified(UObject* InObject, FPropert
 	{
 		PropertyChangedChainEvent.PropertyChain.AddHead(ChainProperty);
 	}
-	
+
 	PropertyChangedChainEvent.SetActiveMemberProperty(InMemberChainProperties.Last());
 	PropertyChangedChainEvent.ObjectIteratorIndex = 0;
-	
+
 	TArray<TMap<FString, int32>> ArrayIndices;
 	ArrayIndices.SetNum(1);
 	ArrayIndices[0].Add(InMemberChainProperties.Last()->GetFName().ToString(), InContainerIdx);
diff --git a/Engine/Plugins/Experimental/Avalanche/Source/AvalancheComponentVisualizers/Public/AvaVisBase.h b/Engine/Plugins/Experimental/Avalanche/Source/AvalancheComponentVisualizers/Public/AvaVisBase.h
--- a/Engine/Plugins/Experimental/Avalanche/Source/AvalancheComponentVisualizers/Public/AvaVisBase.h
+++ b/Engine/Plugins/Experimental/Avalanche/Source/AvalancheComponentVisualizers/Public/AvaVisBase.h
@@ -189,4 +189,6 @@ protected:
 
 	virtual void TrackingStartedInternal(FEditorViewportClient* InViewportClient);
 	virtual void TrackingStoppedInternal(FEditorViewportClient* InViewportClient);
+
+	void ModifyProperty(UObject* InObject, FProperty* InProperty, EPropertyChangeType::Type InPropertyChangeType, TFunctionRef<void()> InFunction);
 };
diff --git a/Engine/Plugins/Experimental/Avalanche/Source/AvalancheEffectorsEditor/Private/Effector/AvaEffectorActorVis.cpp b/Engine/Plugins/Experimental/Avalanche/Source/AvalancheEffectorsEditor/Private/Effector/AvaEffectorActorVis.cpp
--- a/Engine/Plugins/Experimental/Avalanche/Source/AvalancheEffectorsEditor/Private/Effector/AvaEffectorActorVis.cpp
+++ b/Engine/Plugins/Experimental/Avalanche/Source/AvalancheEffectorsEditor/Private/Effector/AvaEffectorActorVis.cpp
@@ -20,27 +20,60 @@ FAvaEffectorActorVisualizer::FAvaEffectorActorVisualizer()
 	: FAvaVisualizerBase()
 {
 	using namespace UE::AvaCore;
+
+	// Sphere
 	InnerRadiusProperty  = GetProperty<ACEEffectorActor>(GET_MEMBER_NAME_CHECKED(ACEEffectorActor, InnerRadius));
 	OuterRadiusProperty  = GetProperty<ACEEffectorActor>(GET_MEMBER_NAME_CHECKED(ACEEffectorActor, OuterRadius));
+
+	// Box
 	InnerExtentProperty  = GetProperty<ACEEffectorActor>(GET_MEMBER_NAME_CHECKED(ACEEffectorActor, InnerExtent));
 	OuterExtentProperty  = GetProperty<ACEEffectorActor>(GET_MEMBER_NAME_CHECKED(ACEEffectorActor, OuterExtent));
+
+	// Plane
 	PlaneSpacingProperty = GetProperty<ACEEffectorActor>(GET_MEMBER_NAME_CHECKED(ACEEffectorActor, PlaneSpacing));
+
+	// Radial
+	RadialAngleProperty = GetProperty<ACEEffectorActor>(GET_MEMBER_NAME_CHECKED(ACEEffectorActor, RadialAngle));
+	RadialMinRadiusProperty = GetProperty<ACEEffectorActor>(GET_MEMBER_NAME_CHECKED(ACEEffectorActor, RadialMinRadius));
+	RadialMaxRadiusProperty = GetProperty<ACEEffectorActor>(GET_MEMBER_NAME_CHECKED(ACEEffectorActor, RadialMaxRadius));
+
+	// Torus
+	TorusRadiusProperty = GetProperty<ACEEffectorActor>(GET_MEMBER_NAME_CHECKED(ACEEffectorActor, TorusRadius));
+	TorusInnerRadiusProperty = GetProperty<ACEEffectorActor>(GET_MEMBER_NAME_CHECKED(ACEEffectorActor, TorusInnerRadius));
+	TorusOuterRadiusProperty = GetProperty<ACEEffectorActor>(GET_MEMBER_NAME_CHECKED(ACEEffectorActor, TorusOuterRadius));
 }
 
 void FAvaEffectorActorVisualizer::StoreInitialValues()
 {
 	Super::StoreInitialValues();
-	
-	if (GetEditedComponent() == nullptr)
+
+	const ACEEffectorActor* EffectorActor = EffectorActorWeak.Get();
+
+	if (GetEditedComponent() == nullptr || !EffectorActor)
 	{
 		return;
 	}
 
-	InitialInnerRadius = EffectorActorWeak->GetInnerRadius();
-	InitialOuterRadius = EffectorActorWeak->GetOuterRadius();
-	InitialInnerExtent = EffectorActorWeak->GetInnerExtent();
-	InitialOuterExtent = EffectorActorWeak->GetOuterExtent();
-	InitialPlaneSpacing = EffectorActorWeak->GetPlaneSpacing();
+	// Sphere
+	InitialInnerRadius = EffectorActor->GetInnerRadius();
+	InitialOuterRadius = EffectorActor->GetOuterRadius();
+
+	// Box
+	InitialInnerExtent = EffectorActor->GetInnerExtent();
+	InitialOuterExtent = EffectorActor->GetOuterExtent();
+
+	// Plane
+	InitialPlaneSpacing = EffectorActor->GetPlaneSpacing();
+
+	// Radial
+	InitialRadialAngle = EffectorActor->GetRadialAngle();
+	InitialRadialMinRadius = EffectorActor->GetRadialMinRadius();
+	InitialRadialMaxRadius = EffectorActor->GetRadialMaxRadius();
+
+	// Torus
+	InitialTorusRadius = EffectorActor->GetTorusRadius();
+	InitialTorusInnerRadius = EffectorActor->GetTorusInnerRadius();
+	InitialTorusOuterRadius = EffectorActor->GetTorusOuterRadius();
 }
 
 FBox FAvaEffectorActorVisualizer::GetComponentBounds(const UActorComponent* InComponent) const
@@ -49,21 +82,35 @@ FBox FAvaEffectorActorVisualizer::GetComponentBounds(const UActorComponent* InCo
 	{
 		if (const ACEEffectorActor* EffectorActor = Cast<ACEEffectorActor>(EffectorComponent->GetOwner()))
 		{
-			if (EffectorActor->GetType() == ECEClonerEffectorType::Box)
+			const ECEClonerEffectorType Type = EffectorActor->GetType();
+
+			if (Type == ECEClonerEffectorType::Box)
 			{
 				return FBox(-EffectorActor->GetOuterExtent(), EffectorActor->GetOuterExtent());
 			}
-			if (EffectorActor->GetType() == ECEClonerEffectorType::Sphere)
+
+			if (Type == ECEClonerEffectorType::Sphere)
 			{
 				return FBox(-FVector(EffectorActor->GetOuterRadius() / 2), FVector(EffectorActor->GetOuterRadius() / 2));
 			}
-			if (EffectorActor->GetType() == ECEClonerEffectorType::Plane)
+
+			if (Type == ECEClonerEffectorType::Plane)
 			{
 				return FBox(-FVector(EffectorActor->GetPlaneSpacing() / 2), FVector(EffectorActor->GetPlaneSpacing() / 2));
 			}
+
+			if (Type == ECEClonerEffectorType::Radial)
+			{
+				return FBox(-FVector(EffectorActor->GetRadialMaxRadius() / 2), FVector(EffectorActor->GetRadialMaxRadius() / 2));
+			}
+
+			if (Type == ECEClonerEffectorType::Torus)
+			{
+				return FBox(-FVector((EffectorActor->GetTorusRadius() + EffectorActor->GetTorusOuterRadius()) / 2), FVector((EffectorActor->GetTorusRadius() + EffectorActor->GetTorusOuterRadius()) / 2));
+			}
 		}
 	}
-	
+
 	return Super::GetComponentBounds(InComponent);
 }
 
@@ -76,61 +123,139 @@ bool FAvaEffectorActorVisualizer::HandleInputDeltaInternal(FEditorViewportClient
 
 	if (ACEEffectorActor* EffectorActor = EffectorActorWeak.Get())
 	{
+		const ECEClonerEffectorType Type = EffectorActor->GetType();
+
 		if (GetViewportWidgetMode(InViewportClient) == UE::Widget::WM_Translate)
 		{
 			if (GetViewportWidgetAxisList(InViewportClient) & EAxisList::XYZ)
 			{
-				if (EffectorActor->GetType() == ECEClonerEffectorType::Box)
+				if (Type == ECEClonerEffectorType::Box)
 				{
-					if (bEditingInnerZone)
+					if (EditingHandleType == HandleTypeInnerZone)
 					{
-						EffectorActor->SetInnerExtent(InitialInnerExtent + InAccumulatedTranslation);
-						EffectorActor->Modify();
-						bHasBeenModified = true;
-						NotifyPropertyModified(EffectorActor, InnerExtentProperty, EPropertyChangeType::Interactive);
+						ModifyProperty(EffectorActor, InnerExtentProperty, EPropertyChangeType::Interactive, [this, &EffectorActor, &InAccumulatedTranslation]()
+						{
+							EffectorActor->SetInnerExtent(InitialInnerExtent + InAccumulatedTranslation);
+						});
 					}
-					else if (bEditingOuterZone)
+					else if (EditingHandleType == HandleTypeOuterZone)
 					{
-						EffectorActor->SetOuterExtent(InitialOuterExtent + InAccumulatedTranslation);
-						EffectorActor->Modify();
-						bHasBeenModified = true;
-						NotifyPropertyModified(EffectorActor, OuterExtentProperty, EPropertyChangeType::Interactive);
+						ModifyProperty(EffectorActor, OuterExtentProperty, EPropertyChangeType::Interactive, [this, &EffectorActor, &InAccumulatedTranslation]()
+						{
+							EffectorActor->SetOuterExtent(InitialOuterExtent + InAccumulatedTranslation);
+						});
 					}
-					
+
 					return true;
 				}
 			}
+
 			if (GetViewportWidgetAxisList(InViewportClient) & EAxisList::Y)
 			{
-				if (EffectorActor->GetType() == ECEClonerEffectorType::Plane)
+				if (Type == ECEClonerEffectorType::Plane)
+				{
+					if (EditingHandleType == HandleTypeInnerZone || EditingHandleType == HandleTypeOuterZone)
+					{
+						ModifyProperty(EffectorActor, PlaneSpacingProperty, EPropertyChangeType::Interactive, [this, &EffectorActor, &InAccumulatedTranslation]()
+						{
+							EffectorActor->SetPlaneSpacing(InitialPlaneSpacing + InAccumulatedTranslation.Y);
+						});
+					}
+
+					return true;
+				}
+
+				if (Type == ECEClonerEffectorType::Sphere)
 				{
-					if (bEditingOuterZone || bEditingInnerZone)
+					if (EditingHandleType == HandleTypeInnerZone)
 					{
-						EffectorActor->SetPlaneSpacing(InitialPlaneSpacing + InAccumulatedTranslation.Y);
-						EffectorActor->Modify();
-						bHasBeenModified = true;
-						NotifyPropertyModified(EffectorActor, PlaneSpacingProperty, EPropertyChangeType::Interactive);
+						ModifyProperty(EffectorActor, InnerRadiusProperty, EPropertyChangeType::Interactive, [this, &EffectorActor, &InAccumulatedTranslation]()
+						{
+							EffectorActor->SetInnerRadius(InitialInnerRadius + InAccumulatedTranslation.Y);
+						});
 					}
-					
+					else if (EditingHandleType == HandleTypeOuterZone)
+					{
+						ModifyProperty(EffectorActor, OuterRadiusProperty, EPropertyChangeType::Interactive, [this, &EffectorActor, &InAccumulatedTranslation]
+						{
+							EffectorActor->SetOuterRadius(InitialOuterRadius + InAccumulatedTranslation.Y);
+						});
+					}
+
+					return true;
+				}
+
+				if (Type == ECEClonerEffectorType::Radial)
+				{
+					if (EditingHandleType == HandleTypeInnerZone)
+					{
+						ModifyProperty(EffectorActor, RadialMinRadiusProperty, EPropertyChangeType::Interactive, [this, &EffectorActor, &InAccumulatedTranslation]
+						{
+							EffectorActor->SetRadialMinRadius(InitialRadialMinRadius + InAccumulatedTranslation.Y);
+						});
+					}
+					else if (EditingHandleType == HandleTypeOuterZone)
+					{
+						ModifyProperty(EffectorActor, RadialMaxRadiusProperty, EPropertyChangeType::Interactive, [this, &EffectorActor, &InAccumulatedTranslation]()
+						{
+							EffectorActor->SetRadialMaxRadius(InitialRadialMaxRadius + InAccumulatedTranslation.Y);
+						});
+					}
+
+					return true;
+				}
+
+				if (Type == ECEClonerEffectorType::Torus)
+				{
+					if (EditingHandleType == HandleTypeRadius)
+					{
+						ModifyProperty(EffectorActor, TorusRadiusProperty, EPropertyChangeType::Interactive, [this, &EffectorActor, &InAccumulatedTranslation]()
+						{
+							EffectorActor->SetTorusRadius(InitialTorusRadius + InAccumulatedTranslation.Y);
+						});
+					}
+
 					return true;
 				}
-				if (EffectorActor->GetType() == ECEClonerEffectorType::Sphere)
+			}
+
+			if (GetViewportWidgetAxisList(InViewportClient) & EAxisList::Z)
+			{
+				if (Type == ECEClonerEffectorType::Torus)
 				{
-					if (bEditingInnerZone)
+					if (EditingHandleType == HandleTypeInnerZone)
 					{
-						EffectorActor->SetInnerRadius(InitialInnerRadius + InAccumulatedTranslation.Y);
-						EffectorActor->Modify();
-						bHasBeenModified = true;
-						NotifyPropertyModified(EffectorActor, InnerRadiusProperty, EPropertyChangeType::Interactive);
+						ModifyProperty(EffectorActor, TorusInnerRadiusProperty, EPropertyChangeType::Interactive, [this, &EffectorActor, &InAccumulatedTranslation]()
+						{
+							EffectorActor->SetTorusInnerRadius(InitialTorusInnerRadius + InAccumulatedTranslation.Z);
+						});
 					}
-					else if (bEditingOuterZone)
+					else if (EditingHandleType == HandleTypeOuterZone)
 					{
-						EffectorActor->SetOuterRadius(InitialOuterRadius + InAccumulatedTranslation.Y);
-						EffectorActor->Modify();
-						bHasBeenModified = true;
-						NotifyPropertyModified(EffectorActor, OuterRadiusProperty, EPropertyChangeType::Interactive);
+						ModifyProperty(EffectorActor, TorusOuterRadiusProperty, EPropertyChangeType::Interactive, [this, &EffectorActor, &InAccumulatedTranslation]()
+						{
+							EffectorActor->SetTorusOuterRadius(InitialTorusOuterRadius + InAccumulatedTranslation.Z);
+						});
 					}
-					
+
+					return true;
+				}
+			}
+		}
+		else if (GetViewportWidgetMode(InViewportClient) == UE::Widget::WM_Rotate)
+		{
+			if (GetViewportWidgetAxisList(InViewportClient) & EAxisList::Z)
+			{
+				if (Type == ECEClonerEffectorType::Radial)
+				{
+					if (EditingHandleType == HandleTypeAngle)
+					{
+						ModifyProperty(EffectorActor, RadialAngleProperty, EPropertyChangeType::Interactive, [this, &EffectorActor, &InAccumulatedRotation]()
+						{
+							EffectorActor->SetRadialAngle(InitialRadialAngle + InAccumulatedRotation.Yaw);
+						});
+					}
+
 					return true;
 				}
 			}
@@ -140,7 +265,7 @@ bool FAvaEffectorActorVisualizer::HandleInputDeltaInternal(FEditorViewportClient
 	{
 		EndEditing();
 	}
-	
+
 	return Super::HandleInputDeltaInternal(InViewportClient, InViewport, InAccumulatedTranslation, InAccumulatedRotation, InAccumulatedScale);
 }
 
@@ -164,12 +289,24 @@ void FAvaEffectorActorVisualizer::DrawVisualizationEditing(const UActorComponent
 
 	if (EffectorActor->GetType() != ECEClonerEffectorType::Plane)
 	{
-		DrawZoneButton(EffectorActor, InView, InPDI, InOutIconIndex, true, FAvaVisualizerBase::Inactive);
+		DrawZoneButton(EffectorActor, InView, InPDI, InOutIconIndex, HandleTypeInnerZone, FAvaVisualizerBase::Active);
 		InOutIconIndex++;
 	}
 
-	DrawZoneButton(EffectorActor, InView, InPDI, InOutIconIndex, false, bEditingOuterZone ? FAvaVisualizerBase::Active : FAvaVisualizerBase::Inactive);
+	DrawZoneButton(EffectorActor, InView, InPDI, InOutIconIndex, HandleTypeOuterZone, FAvaVisualizerBase::Active);
 	InOutIconIndex++;
+
+	if (EffectorActor->GetType() == ECEClonerEffectorType::Radial)
+	{
+		DrawZoneButton(EffectorActor, InView, InPDI, InOutIconIndex, HandleTypeAngle, FAvaVisualizerBase::Active);
+		InOutIconIndex++;
+	}
+
+	if (EffectorActor->GetType() == ECEClonerEffectorType::Torus)
+	{
+		DrawZoneButton(EffectorActor, InView, InPDI, InOutIconIndex, HandleTypeRadius, FAvaVisualizerBase::Active);
+		InOutIconIndex++;
+	}
 }
 
 void FAvaEffectorActorVisualizer::DrawVisualizationNotEditing(const UActorComponent* InComponent, const FSceneView* InView, FPrimitiveDrawInterface* InPDI, int32& InOutIconIndex)
@@ -192,53 +329,118 @@ void FAvaEffectorActorVisualizer::DrawVisualizationNotEditing(const UActorCompon
 
 	if (EffectorActor->GetType() != ECEClonerEffectorType::Plane)
 	{
-		DrawZoneButton(EffectorActor, InView, InPDI, InOutIconIndex, true, FAvaVisualizerBase::Inactive);
+		DrawZoneButton(EffectorActor, InView, InPDI, InOutIconIndex, HandleTypeInnerZone, FAvaVisualizerBase::Inactive);
 		InOutIconIndex++;
 	}
-	
-	DrawZoneButton(EffectorActor, InView, InPDI, InOutIconIndex, false, FAvaVisualizerBase::Inactive);
+
+	DrawZoneButton(EffectorActor, InView, InPDI, InOutIconIndex, HandleTypeOuterZone, FAvaVisualizerBase::Inactive);
 	InOutIconIndex++;
+
+	if (EffectorActor->GetType() == ECEClonerEffectorType::Radial)
+	{
+		DrawZoneButton(EffectorActor, InView, InPDI, InOutIconIndex, HandleTypeAngle, FAvaVisualizerBase::Inactive);
+		InOutIconIndex++;
+	}
+
+	if (EffectorActor->GetType() == ECEClonerEffectorType::Torus)
+	{
+		DrawZoneButton(EffectorActor, InView, InPDI, InOutIconIndex, HandleTypeRadius, FAvaVisualizerBase::Inactive);
+		InOutIconIndex++;
+	}
 }
 
-FVector FAvaEffectorActorVisualizer::GetHandleZoneLocation(const ACEEffectorActor* InEffectorActor, bool bInInnerSize) const
+FVector FAvaEffectorActorVisualizer::GetHandleZoneLocation(const ACEEffectorActor* InEffectorActor, int32 InHandleType) const
 {
+	const ECEClonerEffectorType Type = InEffectorActor->GetType();
 	const FVector EffectorScale = InEffectorActor->GetActorScale();
 	const FRotator EffectorRotation = InEffectorActor->GetActorRotation();
 	FVector OutLocation = InEffectorActor->GetActorLocation();
 
-	if (InEffectorActor->GetType() == ECEClonerEffectorType::Box)
+	// To avoid inner/outer handle to be near actor gizmo and hard to select
+	constexpr float MinHandleOffset = 50.f;
+	constexpr float MaxHandleOffset = 100.f;
+
+	if (Type == ECEClonerEffectorType::Box)
 	{
-		if (bInInnerSize)
+		if (InHandleType == HandleTypeInnerZone)
 		{
 			OutLocation += EffectorRotation.RotateVector(InEffectorActor->GetInnerExtent()) * EffectorScale;
 		}
-		else
-			{
+		else if (InHandleType == HandleTypeOuterZone)
+		{
 			OutLocation += EffectorRotation.RotateVector(InEffectorActor->GetOuterExtent()) * EffectorScale;
 		}
 	}
-	else if (InEffectorActor->GetType() == ECEClonerEffectorType::Plane)
+	else if (Type == ECEClonerEffectorType::Plane)
 	{
-		const float ComponentScale = (EffectorRotation.RotateVector(-FVector::YAxisVector) * EffectorScale).Length();
-		OutLocation += EffectorRotation.RotateVector(FVector::YAxisVector) * (InEffectorActor->GetPlaneSpacing() / 2) * ComponentScale;
+		if (InHandleType == HandleTypeInnerZone
+			|| InHandleType == HandleTypeOuterZone)
+		{
+			const float ComponentScale = (EffectorRotation.RotateVector(-FVector::YAxisVector) * EffectorScale).Length();
+			const FVector HandleAxis = EffectorRotation.RotateVector(FVector::YAxisVector);
+
+			OutLocation += HandleAxis * (InEffectorActor->GetPlaneSpacing() / 2) * ComponentScale;
+		}
 	}
-	else if (InEffectorActor->GetType() == ECEClonerEffectorType::Sphere)
+	else if (Type == ECEClonerEffectorType::Sphere)
 	{
 		const float MinComponentScale = FMath::Min<float>(FMath::Min<float>(EffectorScale.X, EffectorScale.Y), EffectorScale.Z);
-		if (bInInnerSize)
+		const FVector HandleAxis = EffectorRotation.RotateVector(FVector::YAxisVector);
+
+		if (InHandleType == HandleTypeInnerZone)
 		{
-			OutLocation += EffectorRotation.RotateVector(FVector::YAxisVector) * InEffectorActor->GetInnerRadius() * MinComponentScale;
+			OutLocation += FVector::Max(HandleAxis * InEffectorActor->GetInnerRadius(), HandleAxis * MinHandleOffset) * MinComponentScale;
 		}
-		else
+		else if (InHandleType == HandleTypeOuterZone)
+		{
+			OutLocation += FVector::Max(HandleAxis * InEffectorActor->GetOuterRadius(), HandleAxis * MaxHandleOffset) * MinComponentScale;
+		}
+	}
+	else if (Type == ECEClonerEffectorType::Radial)
+	{
+		const float MinComponentScale = FMath::Min<float>(FMath::Min<float>(EffectorScale.X, EffectorScale.Y), EffectorScale.Z);
+		const FVector UpHandleAxis = EffectorRotation.RotateVector(FVector::ZAxisVector);
+		const FVector RightHandleAxis = EffectorRotation.RotateVector(FVector::YAxisVector);
+
+		if (InHandleType == HandleTypeInnerZone)
+		{
+			OutLocation += FVector::Max(RightHandleAxis * InEffectorActor->GetRadialMinRadius(), RightHandleAxis * MinHandleOffset) * MinComponentScale;
+		}
+		else if (InHandleType == HandleTypeOuterZone)
+		{
+			OutLocation += FVector::Max(RightHandleAxis * InEffectorActor->GetRadialMaxRadius(), RightHandleAxis * MaxHandleOffset) * MinComponentScale;
+		}
+		else if (InHandleType == HandleTypeAngle)
 		{
-			OutLocation += EffectorRotation.RotateVector(FVector::YAxisVector) * InEffectorActor->GetOuterRadius() * MinComponentScale;
+			OutLocation += UpHandleAxis * MaxHandleOffset * MinComponentScale;
 		}
 	}
-	
+	else if (Type == ECEClonerEffectorType::Torus)
+	{
+		const float MinComponentScale = FMath::Min<float>(FMath::Min<float>(EffectorScale.X, EffectorScale.Y), EffectorScale.Z);
+		const FVector UpHandleAxis = EffectorRotation.RotateVector(FVector::ZAxisVector);
+		const FVector RightHandleAxis = EffectorRotation.RotateVector(FVector::YAxisVector);
+
+		if (InHandleType == HandleTypeInnerZone)
+		{
+			OutLocation += RightHandleAxis * InEffectorActor->GetTorusRadius()
+				+ FVector::Max(UpHandleAxis * InEffectorActor->GetTorusInnerRadius(), UpHandleAxis * MinHandleOffset) * MinComponentScale;
+		}
+		else if (InHandleType == HandleTypeOuterZone)
+		{
+			OutLocation += RightHandleAxis * InEffectorActor->GetTorusRadius()
+				+ FVector::Max(UpHandleAxis * InEffectorActor->GetTorusOuterRadius(), UpHandleAxis * MaxHandleOffset) * MinComponentScale;
+		}
+		else if (InHandleType == HandleTypeRadius)
+		{
+			OutLocation += FVector::Max(RightHandleAxis * InEffectorActor->GetTorusRadius(), RightHandleAxis * MinHandleOffset) * MinComponentScale;
+		}
+	}
+
 	return OutLocation;
 }
 
-void FAvaEffectorActorVisualizer::DrawZoneButton(const ACEEffectorActor* InEffectorActor, const FSceneView* InView, FPrimitiveDrawInterface* InPDI, int32 InIconIndex, bool bInInnerZone, FLinearColor InColor) const
+void FAvaEffectorActorVisualizer::DrawZoneButton(const ACEEffectorActor* InEffectorActor, const FSceneView* InView, FPrimitiveDrawInterface* InPDI, int32 InIconIndex, int32 InHandleType, FLinearColor InColor) const
 {
 	UTexture2D* ZoneSprite = IAvalancheComponentVisualizersModule::Get().GetSettings()->GetVisualizerSprite(UE::AvaShapes::BevelSprite);
 
@@ -251,9 +453,9 @@ void FAvaEffectorActorVisualizer::DrawZoneButton(const ACEEffectorActor* InEffec
 	float IconSize;
 	GetIconMetrics(InView, InIconIndex, IconLocation, IconSize);
 
-	IconLocation = GetHandleZoneLocation(InEffectorActor, bInInnerZone);
+	IconLocation = GetHandleZoneLocation(InEffectorActor, InHandleType);
 
-	InPDI->SetHitProxy(new HAvaEffectorActorZoneHitProxy(InEffectorActor->GetRootComponent(), bInInnerZone));
+	InPDI->SetHitProxy(new HAvaEffectorActorZoneHitProxy(InEffectorActor->GetRootComponent(), InHandleType));
 	InPDI->DrawSprite(IconLocation, IconSize, IconSize, ZoneSprite->GetResource(), InColor, SDPG_Foreground, 0, 0, 0, 0, SE_BLEND_Opaque);
 	InPDI->SetHitProxy(nullptr);
 }
@@ -265,7 +467,7 @@ UActorComponent* FAvaEffectorActorVisualizer::GetEditedComponent() const
 
 TMap<UObject*, TArray<FProperty*>> FAvaEffectorActorVisualizer::GatherEditableProperties(UObject* InObject) const
 {
-	if (UCEEffectorComponent* EffectorComponent = Cast<UCEEffectorComponent>(InObject))
+	if (const UCEEffectorComponent* EffectorComponent = Cast<UCEEffectorComponent>(InObject))
 	{
 		if (ACEEffectorActor* EffectorActor = EffectorComponent->GetOuterACEEffectorActor())
 		{
@@ -279,6 +481,15 @@ TMap<UObject*, TArray<FProperty*>> FAvaEffectorActorVisualizer::GatherEditablePr
 
 				case ECEClonerEffectorType::Sphere:
 					return {{EffectorActor, {InnerRadiusProperty, OuterRadiusProperty}}};
+
+				case ECEClonerEffectorType::Radial:
+					return {{EffectorActor, {RadialAngleProperty, RadialMinRadiusProperty, RadialMaxRadiusProperty}}};
+
+				case ECEClonerEffectorType::Torus:
+					return {{EffectorActor, {TorusRadiusProperty, TorusInnerRadiusProperty, TorusOuterRadiusProperty}}};
+
+				default:
+					return {};
 			}
 		}
 	}
@@ -305,10 +516,9 @@ bool FAvaEffectorActorVisualizer::VisProxyHandleClick(FEditorViewportClient* InV
 	{
 		EndEditing();
 		EffectorActorWeak = Cast<ACEEffectorActor>(InVisProxy->Component->GetOwner());
-		bEditingInnerZone = static_cast<HAvaEffectorActorZoneHitProxy*>(InVisProxy)->bInnerZone;
-		bEditingOuterZone = !bEditingInnerZone;
+		EditingHandleType = static_cast<HAvaEffectorActorZoneHitProxy*>(InVisProxy)->HandleType;
 		StartEditing(InViewportClient, Component);
-		
+
 		return true;
 	}
 
@@ -317,59 +527,118 @@ bool FAvaEffectorActorVisualizer::VisProxyHandleClick(FEditorViewportClient* InV
 
 bool FAvaEffectorActorVisualizer::GetWidgetLocation(const FEditorViewportClient* InViewportClient, FVector& OutLocation) const
 {
-	if (EffectorActorWeak.IsValid())
+	if (const ACEEffectorActor* EffectorActor = EffectorActorWeak.Get())
 	{
-		OutLocation = GetHandleZoneLocation(EffectorActorWeak.Get(), bEditingInnerZone);
+		OutLocation = GetHandleZoneLocation(EffectorActor, EditingHandleType);
 		return true;
 	}
-	
+
 	return Super::GetWidgetLocation(InViewportClient, OutLocation);
 }
 
 bool FAvaEffectorActorVisualizer::GetWidgetMode(const FEditorViewportClient* InViewportClient,
 	UE::Widget::EWidgetMode& OutMode) const
 {
-	if (bEditingOuterZone || bEditingInnerZone)
+	if (EditingHandleType == HandleTypeInnerZone
+		|| EditingHandleType == HandleTypeOuterZone
+		|| EditingHandleType == HandleTypeRadius)
 	{
 		OutMode = UE::Widget::EWidgetMode::WM_Translate;
 		return true;
 	}
-	
+
+	if (EditingHandleType == HandleTypeAngle)
+	{
+		OutMode = UE::Widget::EWidgetMode::WM_Rotate;
+		return true;
+	}
+
 	return Super::GetWidgetMode(InViewportClient, OutMode);
 }
 
 bool FAvaEffectorActorVisualizer::GetWidgetAxisList(const FEditorViewportClient* InViewportClient,
 	UE::Widget::EWidgetMode InWidgetMode, EAxisList::Type& OutAxisList) const
 {
-	if (bEditingOuterZone || bEditingInnerZone)
+	if (EditingHandleType == HandleTypeInnerZone
+		|| EditingHandleType == HandleTypeOuterZone)
 	{
-		if (EffectorActorWeak->GetType() != ECEClonerEffectorType::Box)
+		if (EffectorActorWeak->GetType() == ECEClonerEffectorType::Torus)
 		{
-			OutAxisList = EAxisList::Type::Y;
+			OutAxisList = EAxisList::Type::Z;
 		}
-		else
+		else if (EffectorActorWeak->GetType() == ECEClonerEffectorType::Box)
 		{
 			OutAxisList = EAxisList::Type::XYZ;
 		}
-		
+		else
+		{
+			OutAxisList = EAxisList::Type::Y;
+		}
+
+		return true;
+	}
+
+	if (EditingHandleType == HandleTypeRadius)
+	{
+		if (EffectorActorWeak->GetType() == ECEClonerEffectorType::Torus)
+		{
+			OutAxisList = EAxisList::Type::Y;
+		}
+
 		return true;
 	}
-	
+
+	if (EditingHandleType == HandleTypeAngle)
+	{
+		if (EffectorActorWeak->GetType() == ECEClonerEffectorType::Radial)
+		{
+			OutAxisList = EAxisList::Type::Z;
+		}
+
+		return true;
+	}
+
 	return Super::GetWidgetAxisList(InViewportClient, InWidgetMode, OutAxisList);
 }
 
 bool FAvaEffectorActorVisualizer::GetWidgetAxisListDragOverride(const FEditorViewportClient* InViewportClient,
 	UE::Widget::EWidgetMode InWidgetMode, EAxisList::Type& OutAxisList) const
 {
-	if (bEditingOuterZone || bEditingInnerZone)
+	if (EditingHandleType == HandleTypeInnerZone
+		|| EditingHandleType == HandleTypeOuterZone)
 	{
+		if (EffectorActorWeak->GetType() == ECEClonerEffectorType::Torus)
+		{
+			OutAxisList = EAxisList::Type::Z;
+			return true;
+		}
+
 		if (EffectorActorWeak->GetType() != ECEClonerEffectorType::Box)
 		{
 			OutAxisList = EAxisList::Type::Y;
 			return true;
 		}
 	}
-	
+
+	if (EditingHandleType == HandleTypeRadius)
+	{
+		if (EffectorActorWeak->GetType() == ECEClonerEffectorType::Torus)
+		{
+			OutAxisList = EAxisList::Type::Y;
+			return true;
+		}
+	}
+
+	if (EditingHandleType == HandleTypeAngle)
+	{
+		if (EffectorActorWeak->GetType() == ECEClonerEffectorType::Radial)
+		{
+			OutAxisList = EAxisList::Type::Z;
+		}
+
+		return true;
+	}
+
 	return Super::GetWidgetAxisListDragOverride(InViewportClient, InWidgetMode, OutAxisList);
 }
 
@@ -379,9 +648,9 @@ bool FAvaEffectorActorVisualizer::ResetValue(FEditorViewportClient* InViewportCl
 	{
 		return Super::ResetValue(InViewportClient, InHitProxy);
 	}
-	
+
 	const HAvaEffectorActorZoneHitProxy* ComponentHitProxy = static_cast<HAvaEffectorActorZoneHitProxy*>(InHitProxy);
-	
+
 	if (!ComponentHitProxy->Component.IsValid() || !ComponentHitProxy->Component->IsA<UCEEffectorComponent>())
 	{
 		return Super::ResetValue(InViewportClient, InHitProxy);
@@ -389,64 +658,116 @@ bool FAvaEffectorActorVisualizer::ResetValue(FEditorViewportClient* InViewportCl
 
 	if (ACEEffectorActor* EffectorActor = Cast<ACEEffectorActor>(ComponentHitProxy->Component->GetOwner()))
 	{
-		if (EffectorActor->GetType() == ECEClonerEffectorType::Box)
+		const ECEClonerEffectorType Type = EffectorActor->GetType();
+		const int32 HandleType = ComponentHitProxy->HandleType;
+
+		FScopedTransaction Transaction(LOCTEXT("VisualizerResetValue", "Visualizer Reset Value"));
+
+		if (Type == ECEClonerEffectorType::Box)
+		{
+			if (HandleType == HandleTypeInnerZone)
+			{
+				ModifyProperty(EffectorActor, InnerExtentProperty, EPropertyChangeType::ValueSet, [&EffectorActor]()
+				{
+					EffectorActor->SetInnerExtent(FVector(50.f));
+				});
+			}
+			else if (HandleType == HandleTypeOuterZone)
+			{
+				ModifyProperty(EffectorActor, OuterExtentProperty, EPropertyChangeType::ValueSet, [&EffectorActor]()
+				{
+					EffectorActor->SetOuterExtent(FVector(200.f));
+				});
+			}
+		}
+		else if (Type == ECEClonerEffectorType::Plane)
 		{
-			if (ComponentHitProxy->bInnerZone)
+			if (HandleType == HandleTypeInnerZone
+				|| HandleType == HandleTypeOuterZone)
 			{
-				FScopedTransaction Transaction(LOCTEXT("VisualizerResetValue", "Visualizer Reset Value"));
-				EffectorActor->SetFlags(RF_Transactional);
-				EffectorActor->SetInnerExtent(FVector(50.f));
-				EffectorActor->Modify();
-				NotifyPropertyModified(EffectorActor, InnerExtentProperty, EPropertyChangeType::ValueSet);
+				ModifyProperty(EffectorActor, PlaneSpacingProperty, EPropertyChangeType::ValueSet, [&EffectorActor]()
+				{
+					EffectorActor->SetPlaneSpacing(200.f);
+				});
 			}
-			else
+		}
+		else if (Type == ECEClonerEffectorType::Sphere)
+		{
+			if (HandleType == HandleTypeInnerZone)
 			{
-				FScopedTransaction Transaction(LOCTEXT("VisualizerResetValue", "Visualizer Reset Value"));
-				EffectorActor->SetFlags(RF_Transactional);
-				EffectorActor->SetOuterExtent(FVector(200.f));
-				EffectorActor->Modify();
-				NotifyPropertyModified(EffectorActor, OuterExtentProperty, EPropertyChangeType::ValueSet);
+				ModifyProperty(EffectorActor, InnerRadiusProperty, EPropertyChangeType::ValueSet, [&EffectorActor]()
+				{
+					EffectorActor->SetInnerRadius(50.f);
+				});
+			}
+			else if (HandleType == HandleTypeOuterZone)
+			{
+				ModifyProperty(EffectorActor, OuterRadiusProperty, EPropertyChangeType::ValueSet, [&EffectorActor]()
+				{
+					EffectorActor->SetOuterRadius(200.f);
+				});
 			}
 		}
-		else if (EffectorActor->GetType() == ECEClonerEffectorType::Plane)
+		else if (Type == ECEClonerEffectorType::Radial)
 		{
-			FScopedTransaction Transaction(LOCTEXT("VisualizerResetValue", "Visualizer Reset Value"));
-			EffectorActor->SetFlags(RF_Transactional);
-			EffectorActor->SetPlaneSpacing(200.f);
-			EffectorActor->Modify();
-			NotifyPropertyModified(EffectorActor, PlaneSpacingProperty, EPropertyChangeType::ValueSet);
+			if (HandleType == HandleTypeInnerZone)
+			{
+				ModifyProperty(EffectorActor, RadialMinRadiusProperty, EPropertyChangeType::ValueSet, [&EffectorActor]()
+				{
+					EffectorActor->SetRadialMinRadius(0.f);
+				});
+			}
+			else if (HandleType == HandleTypeOuterZone)
+			{
+				ModifyProperty(EffectorActor, RadialMaxRadiusProperty, EPropertyChangeType::ValueSet, [&EffectorActor]()
+				{
+					EffectorActor->SetRadialMaxRadius(1000.f);
+				});
+			}
+			else if (HandleType == HandleTypeAngle)
+			{
+				ModifyProperty(EffectorActor, RadialAngleProperty, EPropertyChangeType::ValueSet, [&EffectorActor]()
+				{
+					EffectorActor->SetRadialAngle(180.f);
+				});
+			}
 		}
-		else if (EffectorActor->GetType() == ECEClonerEffectorType::Sphere)
+		else if (Type == ECEClonerEffectorType::Torus)
 		{
-			if (ComponentHitProxy->bInnerZone)
+			if (HandleType == HandleTypeInnerZone)
+			{
+				ModifyProperty(EffectorActor, TorusInnerRadiusProperty, EPropertyChangeType::ValueSet, [&EffectorActor]()
+				{
+					EffectorActor->SetTorusInnerRadius(50.f);
+				});
+			}
+			else if (HandleType == HandleTypeOuterZone)
 			{
-				FScopedTransaction Transaction(LOCTEXT("VisualizerResetValue", "Visualizer Reset Value"));
-				EffectorActor->SetFlags(RF_Transactional);
-				EffectorActor->SetInnerRadius(50.f);
-				EffectorActor->Modify();
-				NotifyPropertyModified(EffectorActor, InnerRadiusProperty, EPropertyChangeType::ValueSet);
+				ModifyProperty(EffectorActor, TorusOuterRadiusProperty, EPropertyChangeType::ValueSet, [&EffectorActor]()
+				{
+					EffectorActor->SetTorusOuterRadius(200.f);
+				});
 			}
-			else
+			else if (HandleType == HandleTypeRadius)
 			{
-				FScopedTransaction Transaction(LOCTEXT("VisualizerResetValue", "Visualizer Reset Value"));
-				EffectorActor->SetFlags(RF_Transactional);
-				EffectorActor->SetOuterRadius(200.f);
-				EffectorActor->Modify();
-				NotifyPropertyModified(EffectorActor, OuterRadiusProperty, EPropertyChangeType::ValueSet);
+				ModifyProperty(EffectorActor, TorusRadiusProperty, EPropertyChangeType::ValueSet, [&EffectorActor]()
+				{
+					EffectorActor->SetTorusRadius(250.f);
+				});
 			}
 		}
 	}
-	
+
 	return true;
 }
 
 bool FAvaEffectorActorVisualizer::IsEditing() const
 {
-	if (bEditingInnerZone || bEditingOuterZone)
+	if (EditingHandleType != INDEX_NONE)
 	{
 		return true;
 	}
-	
+
 	return Super::IsEditing();
 }
 
@@ -455,8 +776,7 @@ void FAvaEffectorActorVisualizer::EndEditing()
 	Super::EndEditing();
 
 	EffectorActorWeak.Reset();
-	bEditingInnerZone = false;
-	bEditingOuterZone = false;
+	EditingHandleType = INDEX_NONE;
 }
 
 #undef LOCTEXT_NAMESPACE
diff --git a/Engine/Plugins/Experimental/Avalanche/Source/AvalancheEffectorsEditor/Private/Effector/AvaEffectorActorVis.h b/Engine/Plugins/Experimental/Avalanche/Source/AvalancheEffectorsEditor/Private/Effector/AvaEffectorActorVis.h
--- a/Engine/Plugins/Experimental/Avalanche/Source/AvalancheEffectorsEditor/Private/Effector/AvaEffectorActorVis.h
+++ b/Engine/Plugins/Experimental/Avalanche/Source/AvalancheEffectorsEditor/Private/Effector/AvaEffectorActorVis.h
@@ -11,12 +11,12 @@ struct HAvaEffectorActorZoneHitProxy : HAvaHitProxy
 {
 	DECLARE_HIT_PROXY();
 
-	HAvaEffectorActorZoneHitProxy(const UActorComponent* InComponent, bool bInInnerZone)
+	HAvaEffectorActorZoneHitProxy(const UActorComponent* InComponent, int32 InHandleType)
 		: HAvaHitProxy(InComponent)
-		, bInnerZone(bInInnerZone)
+		, HandleType(InHandleType)
 	{}
 
-	bool bInnerZone = false;
+	int32 HandleType = INDEX_NONE;
 };
 
 /** Custom visualization for effector actor to handle weight zones */
@@ -45,24 +45,53 @@ public:
 	virtual void DrawVisualizationNotEditing(const UActorComponent* InComponent, const FSceneView* InView, FPrimitiveDrawInterface* InPDI, int32& InOutIconIndex) override;
 	//~ End FAvaVisualizerBase
 
-	ACEEffectorActor* GetEffectorActor() const { return EffectorActorWeak.Get(); };
-protected:
+	ACEEffectorActor* GetEffectorActor() const
+	{
+		return EffectorActorWeak.Get();
+	}
 
-	FVector GetHandleZoneLocation(const ACEEffectorActor* InEffectorActor, bool bInInnerSize) const;
-	void DrawZoneButton(const ACEEffectorActor* InEffectorActor, const FSceneView* InView, FPrimitiveDrawInterface* InPDI, int32 InIconIndex, bool bInInnerZone, FLinearColor InColor) const;
+protected:
+	FVector GetHandleZoneLocation(const ACEEffectorActor* InEffectorActor, int32 InHandleType) const;
+	void DrawZoneButton(const ACEEffectorActor* InEffectorActor, const FSceneView* InView, FPrimitiveDrawInterface* InPDI, int32 InIconIndex, int32 InHandleType, FLinearColor InColor) const;
 
 	FProperty* InnerRadiusProperty;
 	FProperty* OuterRadiusProperty;
+
 	FProperty* InnerExtentProperty;
 	FProperty* OuterExtentProperty;
+
 	FProperty* PlaneSpacingProperty;
 
+	FProperty* RadialAngleProperty;
+	FProperty* RadialMinRadiusProperty;
+	FProperty* RadialMaxRadiusProperty;
+
+	FProperty* TorusRadiusProperty;
+	FProperty* TorusInnerRadiusProperty;
+	FProperty* TorusOuterRadiusProperty;
+
 	TWeakObjectPtr<ACEEffectorActor> EffectorActorWeak = nullptr;
+
 	float InitialInnerRadius = 0.f;
 	float InitialOuterRadius = 0.f;
+
 	FVector InitialInnerExtent = FVector(0.f);
 	FVector InitialOuterExtent = FVector(0.f);
+
 	float InitialPlaneSpacing = 0.f;
-	bool bEditingInnerZone = false;
-	bool bEditingOuterZone = false;
+
+	float InitialRadialAngle = 0.f;
+	float InitialRadialMinRadius = 0.f;
+	float InitialRadialMaxRadius = 0.f;
+
+	float InitialTorusRadius = 0.f;
+	float InitialTorusInnerRadius = 0.f;
+	float InitialTorusOuterRadius = 0.f;
+
+	static constexpr int32 HandleTypeInnerZone = 0;
+	static constexpr int32 HandleTypeOuterZone = 1;
+	static constexpr int32 HandleTypeRadius = 2;
+	static constexpr int32 HandleTypeAngle = 3;
+
+	int32 EditingHandleType = INDEX_NONE;
 };
diff --git a/Engine/Plugins/Experimental/ClonerEffector/Source/ClonerEffector/Private/Effector/CEEffectorActor.cpp b/Engine/Plugins/Experimental/ClonerEffector/Source/ClonerEffector/Private/Effector/CEEffectorActor.cpp
--- a/Engine/Plugins/Experimental/ClonerEffector/Source/ClonerEffector/Private/Effector/CEEffectorActor.cpp
+++ b/Engine/Plugins/Experimental/ClonerEffector/Source/ClonerEffector/Private/Effector/CEEffectorActor.cpp
@@ -194,8 +194,7 @@ void ACEEffectorActor::PostDuplicate(EDuplicateMode::Type InDuplicateMode)
 
 void ACEEffectorActor::RegisterToChannel()
 {
-	if (IsValid(this)
-		&& ChannelData.GetIdentifier() == INDEX_NONE)
+	if (IsValid(this) && ChannelData.GetIdentifier() == INDEX_NONE)
 	{
 		// Register this effector to the effector channel
 		if (UCEEffectorSubsystem* EffectorSubsystem = UCEEffectorSubsystem::Get(GetWorld()))
@@ -393,6 +392,8 @@ void ACEEffectorActor::SetRadialMaxRadius(float InRadius)
 
 void ACEEffectorActor::SetTorusRadius(float InRadius)
 {
+	InRadius = FMath::Max(0, InRadius);
+
 	if (FMath::IsNearlyEqual(InRadius, TorusRadius))
 	{
 		return;
@@ -404,6 +405,8 @@ void ACEEffectorActor::SetTorusRadius(float InRadius)
 
 void ACEEffectorActor::SetTorusInnerRadius(float InRadius)
 {
+	InRadius = FMath::Max(0, InRadius);
+
 	if (FMath::IsNearlyEqual(InRadius, TorusInnerRadius))
 	{
 		return;
@@ -415,6 +418,8 @@ void ACEEffectorActor::SetTorusInnerRadius(float InRadius)
 
 void ACEEffectorActor::SetTorusOuterRadius(float InRadius)
 {
+	InRadius = FMath::Max(0, InRadius);
+
 	if (FMath::IsNearlyEqual(InRadius, TorusOuterRadius))
 	{
 		return;
@@ -1009,7 +1014,7 @@ void ACEEffectorActor::OnRadialChanged()
 		return;
 	}
 
-	ChannelData.InnerExtent = FVector::LeftVector;
+	ChannelData.InnerExtent = FVector::RightVector;
 	ChannelData.OuterExtent = FVector(RadialAngle, RadialMinRadius, RadialMaxRadius);
 
 #if WITH_EDITOR
@@ -1031,13 +1036,13 @@ void ACEEffectorActor::OnRadialChanged()
 		UpdateVisualizer(InnerVisualizerId, [this](UDynamicMesh* InMesh)
 		{
 			constexpr FGeometryScriptPrimitiveOptions PrimitiveOptions;
-			UGeometryScriptLibrary_MeshPrimitiveFunctions::AppendDisc(InMesh, PrimitiveOptions, FTransform(FRotator(0, -90, 0)), RadialMaxRadius, 16, 0, 0, RadialAngle / 2, RadialMinRadius);
+			UGeometryScriptLibrary_MeshPrimitiveFunctions::AppendDisc(InMesh, PrimitiveOptions, FTransform(FRotator(0, 90, 0)), RadialMaxRadius, 16, 0, 0, RadialAngle / 2, RadialMinRadius);
 		});
 
 		UpdateVisualizer(OuterVisualizerId, [this](UDynamicMesh* InMesh)
 		{
 			constexpr FGeometryScriptPrimitiveOptions PrimitiveOptions;
-			UGeometryScriptLibrary_MeshPrimitiveFunctions::AppendDisc(InMesh, PrimitiveOptions, FTransform(FRotator(0, -90, 0)), RadialMaxRadius, 16, 0, RadialAngle / 2, RadialAngle, RadialMinRadius);
+			UGeometryScriptLibrary_MeshPrimitiveFunctions::AppendDisc(InMesh, PrimitiveOptions, FTransform(FRotator(0, 90, 0)), RadialMaxRadius, 16, 0, RadialAngle / 2, RadialAngle, RadialMinRadius);
 		});
 	}
 
diff --git a/Engine/Plugins/Experimental/ClonerEffector/Source/ClonerEffector/Public/Effector/CEEffectorActor.h b/Engine/Plugins/Experimental/ClonerEffector/Source/ClonerEffector/Public/Effector/CEEffectorActor.h
--- a/Engine/Plugins/Experimental/ClonerEffector/Source/ClonerEffector/Public/Effector/CEEffectorActor.h
+++ b/Engine/Plugins/Experimental/ClonerEffector/Source/ClonerEffector/Public/Effector/CEEffectorActor.h
@@ -135,7 +135,7 @@ public:
 	}
 
 	UFUNCTION(BlueprintCallable, Category="Effector")
-	void SetTorusRadius(float InRadius);
+	CLONEREFFECTOR_API void SetTorusRadius(float InRadius);
 
 	UFUNCTION(BlueprintPure, Category="Effector")
 	float GetTorusRadius() const
@@ -144,7 +144,7 @@ public:
 	}
 
 	UFUNCTION(BlueprintCallable, Category="Effector")
-	void SetTorusInnerRadius(float InRadius);
+	CLONEREFFECTOR_API void SetTorusInnerRadius(float InRadius);
 
 	UFUNCTION(BlueprintPure, Category="Effector")
 	float GetTorusInnerRadius() const
@@ -153,7 +153,7 @@ public:
 	}
 
 	UFUNCTION(BlueprintCallable, Category="Effector")
-	void SetTorusOuterRadius(float InRadius);
+	CLONEREFFECTOR_API void SetTorusOuterRadius(float InRadius);
 
 	UFUNCTION(BlueprintPure, Category="Effector")
 	float GetTorusOuterRadius() const
@@ -738,7 +738,7 @@ private:
 	TObjectPtr<UMaterialInstanceDynamic> OuterVisualizerMaterial;
 
 	/** Used to hold data related to visualizer for updates */
-	UPROPERTY(Transient)
+	UPROPERTY(Transient, NonTransactional)
 	FInstancedPropertyBag VisualizerData;
 #endif
 
