[Motion Design] Rundown Server API: adding channel and output device status information.

Details:

- FAvaRundownChannel: added channel state and issue severity.
- FAvaRundownOutputDeviceItem: added device state, issue severity and issue messages.

Affected response messages:
FAvaRundownChannelListChanged -> Broadcasted to all connected clients when channel list is modified.
FAvaRundownChannelResponse -> Response to FAvaRundownGetChannel, also now, broadcasted to all connected clients when a channel state is changed.
FAvaRundownChannels -> Response to FAvaRundownGetChannels.

Other:
Broadcast delegates are now setup immediately instead of within the FAvaRundownGetChannels handler.
Clients would expect to receive the broadcast status notification as soon as connection is established.

Tested with [REDACTED_PROFILE_URL]/MotionDesignSandbox/tree/main/Tests/RundownTestClient.
API change is backward compatible.

diff --git a/Engine/Plugins/Experimental/Avalanche/Source/AvalancheMediaEditor/Private/Rundown/AvaRundownMessages.h b/Engine/Plugins/Experimental/Avalanche/Source/AvalancheMediaEditor/Private/Rundown/AvaRundownMessages.h
--- a/Engine/Plugins/Experimental/Avalanche/Source/AvalancheMediaEditor/Private/Rundown/AvaRundownMessages.h
+++ b/Engine/Plugins/Experimental/Avalanche/Source/AvalancheMediaEditor/Private/Rundown/AvaRundownMessages.h
@@ -2,6 +2,7 @@
 
 #pragma once
 
+#include "Broadcast/Channel/AvaBroadcastMediaOutputInfo.h"
 #include "Rundown/AvaRundownPage.h"
 #include "AvaRundownMessages.generated.h"
 
@@ -526,6 +527,18 @@ public:
 	UPROPERTY()
 	FString Name;
 
+	UPROPERTY()
+	FAvaBroadcastMediaOutputInfo OutputInfo;
+	
+	UPROPERTY()
+	EAvaBroadcastOutputState OutputState = EAvaBroadcastOutputState::Invalid;
+	
+	UPROPERTY()
+	EAvaBroadcastIssueSeverity IssueSeverity = EAvaBroadcastIssueSeverity::None;
+
+	UPROPERTY()
+	TArray<FString> IssueMessages;
+
 	/**
 	 * Raw Json string representing a serialized UMediaOutput.
 	 */
@@ -577,6 +590,12 @@ public:
 	UPROPERTY()
 	FString Name;
 
+	UPROPERTY()
+	EAvaBroadcastChannelState State = EAvaBroadcastChannelState::Offline;
+
+	UPROPERTY()
+	EAvaBroadcastIssueSeverity IssueSeverity = EAvaBroadcastIssueSeverity::None;
+
 	UPROPERTY()
 	TArray<FAvaRundownOutputDeviceItem> Devices;
 };
diff --git a/Engine/Plugins/Experimental/Avalanche/Source/AvalancheMediaEditor/Private/Rundown/AvaRundownServer.cpp b/Engine/Plugins/Experimental/Avalanche/Source/AvalancheMediaEditor/Private/Rundown/AvaRundownServer.cpp
--- a/Engine/Plugins/Experimental/Avalanche/Source/AvalancheMediaEditor/Private/Rundown/AvaRundownServer.cpp
+++ b/Engine/Plugins/Experimental/Avalanche/Source/AvalancheMediaEditor/Private/Rundown/AvaRundownServer.cpp
@@ -76,11 +76,17 @@ namespace UE::AvaRundownServer::Private
 	{
 		FAvaRundownChannel Channel;
 		Channel.Name = InChannel.GetChannelName().ToString();
+		Channel.State = InChannel.GetState();
+		Channel.IssueSeverity = InChannel.GetIssueSeverity();
 		const TArray<UMediaOutput*>& MediaOutputs = InChannel.GetMediaOutputs();
 		for (const UMediaOutput* MediaOutput : MediaOutputs)
 		{
 			FAvaRundownOutputDeviceItem DeviceItem;
 			DeviceItem.Name = MediaOutput->GetFName().ToString();
+			DeviceItem.OutputInfo = InChannel.GetMediaOutputInfo(MediaOutput);
+			DeviceItem.OutputState = InChannel.GetMediaOutputState(MediaOutput);
+			DeviceItem.IssueSeverity = InChannel.GetMediaOutputIssueSeverity(DeviceItem.OutputState, MediaOutput);
+			DeviceItem.IssueMessages = InChannel.GetMediaOutputIssueMessages(MediaOutput);
 			DeviceItem.Data = FAvaRundownOutputEditorUtils::SerializeMediaOutput(MediaOutput);
 			Channel.Devices.Push(MoveTemp(DeviceItem));
 		}
@@ -198,6 +204,9 @@ FAvaRundownServer::FAvaRundownServer()
 
 FAvaRundownServer::~FAvaRundownServer()
 {
+	RemoveBroadcastDelegates(&UAvaBroadcast::Get());
+	RemoveEditorDelegates();
+	
 	FMessageEndpoint::SafeRelease(MessageEndpoint);
 	
 	for (IConsoleObject* ConsoleCommand : ConsoleCommands)
@@ -254,6 +263,9 @@ void FAvaRundownServer::Init(const FString& InAssignedHostName)
 		// Subscribe to the server listing requests
 		MessageEndpoint->Subscribe<FAvaRundownPing>();
 
+		SetupBroadcastDelegates(&UAvaBroadcast::Get());
+		SetupEditorDelegates();
+
 		UE_LOG(LogAvaRundownServer, Log, TEXT("Motion Design Rundown Server \"%s\" Started."), *HostName);
 	}
 }
@@ -261,20 +273,22 @@ void FAvaRundownServer::Init(const FString& InAssignedHostName)
 void FAvaRundownServer::SetupBroadcastDelegates(UAvaBroadcast* InBroadcast)
 {
 	RemoveBroadcastDelegates(InBroadcast);
-	InBroadcast->GetOnChannelsListChanged().AddRaw(this, &FAvaRundownServer::OnBroadcastChannelListChanged);
+	InBroadcast->GetOnChannelsListChanged().AddSP(this, &FAvaRundownServer::OnBroadcastChannelListChanged);
+	FAvaBroadcastOutputChannel::GetOnChannelChanged().AddSP(this, &FAvaRundownServer::OnBroadcastChannelChanged);
 }
 
 void FAvaRundownServer::SetupEditorDelegates()
 {
 	RemoveEditorDelegates();
 	const FAssetRegistryModule& AssetRegistryModule = FModuleManager::LoadModuleChecked<FAssetRegistryModule>("AssetRegistry");
-	AssetRegistryModule.Get().OnAssetAdded().AddRaw(this, &FAvaRundownServer::OnAssetAddedOrRemoved);
-	AssetRegistryModule.Get().OnAssetRemoved().AddRaw(this, &FAvaRundownServer::OnAssetAddedOrRemoved);
+	AssetRegistryModule.Get().OnAssetAdded().AddSP(this, &FAvaRundownServer::OnAssetAddedOrRemoved);
+	AssetRegistryModule.Get().OnAssetRemoved().AddSP(this, &FAvaRundownServer::OnAssetAddedOrRemoved);
 }
 
 void FAvaRundownServer::RemoveBroadcastDelegates(UAvaBroadcast* InBroadcast) const
 {
 	InBroadcast->GetOnChannelsListChanged().RemoveAll(this);
+	FAvaBroadcastOutputChannel::GetOnChannelChanged().RemoveAll(this);
 }
 
 void FAvaRundownServer::RemoveEditorDelegates() const
@@ -373,6 +387,11 @@ void FAvaRundownServer::PageAnimSettingsChanged(const UAvaRundown* InRundown, co
 
 void FAvaRundownServer::OnBroadcastChannelListChanged(const FAvaBroadcastProfile& InProfile) const
 {
+	if (ClientAddresses.IsEmpty())
+	{
+		return;
+	}
+
 	FAvaRundownChannelListChanged* ReplyMessage = FMessageEndpoint::MakeMessage<FAvaRundownChannelListChanged>();
 
 	const TArray<FAvaBroadcastOutputChannel*>& OutputChannels = InProfile.GetChannels();
@@ -387,8 +406,26 @@ void FAvaRundownServer::OnBroadcastChannelListChanged(const FAvaBroadcastProfile
 	SendResponse(ReplyMessage, ClientAddresses);
 }
 
+void FAvaRundownServer::OnBroadcastChannelChanged(const FAvaBroadcastOutputChannel& InChannel, EAvaBroadcastChannelChange InChange) const
+{
+	if (ClientAddresses.IsEmpty())
+	{
+		return;
+	}
+
+	FAvaRundownChannelResponse* ReplyMessage = FMessageEndpoint::MakeMessage<FAvaRundownChannelResponse>();
+	
+	ReplyMessage->Channel = UE::AvaRundownServer::Private::SerializeChannel(InChannel);
+	SendResponse(ReplyMessage, ClientAddresses);
+}
+
 void FAvaRundownServer::OnAssetAddedOrRemoved(const FAssetData& InAssetData) const
 {
+	if (ClientAddresses.IsEmpty())
+	{
+		return;
+	}
+	
 	using namespace UE::AvaRundownServer::Private;
 	if (InAssetData.GetClass() == UAvaRundown::StaticClass() || FAvaPlaybackUtils::IsPlayableAsset(InAssetData))
 	{
@@ -405,8 +442,6 @@ void FAvaRundownServer::HandleRundownPing(const FAvaRundownPing& InMessage, cons
 		UE_LOG(LogAvaRundownServer, Log, TEXT("Received Ping request from %s"), *InContext->GetSender().ToString());
 	}
 
-	SetupEditorDelegates();
-	
 	FAvaRundownPong* ReplyMessage = FMessageEndpoint::MakeMessage<FAvaRundownPong>();
 	ReplyMessage->RequestId = InMessage.RequestId;
 	ReplyMessage->bAuto = InMessage.bAuto;
@@ -864,7 +899,6 @@ void FAvaRundownServer::HandleGetChannels(const FAvaRundownGetChannels& InMessag
 		ReplyMessage->Channels.Push(MoveTemp(Channel));
 	}
 	
-	SetupBroadcastDelegates(&Broadcast);
 	SendResponse(ReplyMessage, InContext->GetSender());
 }
 
@@ -1070,12 +1104,8 @@ void FAvaRundownServer::HandleEditChannelDevice(const FAvaRundownEditChannelDevi
 		LogAndSendMessage(InContext->GetSender(), InMessage.RequestId, ELogVerbosity::Error, TEXT("\"EditChannelDevice\" Failed. Reason: Invalid Device \"%s\"."), *InMessage.MediaOutputName);
 		return;
 	}
-	
-	FAvaRundownOutputDeviceItem DeviceItem;
-	DeviceItem.Name = InMessage.MediaOutputName;
-	DeviceItem.Data = InMessage.Data;
 
-	FAvaRundownOutputEditorUtils::EditMediaOutput(MediaOutput, DeviceItem.Data);
+	FAvaRundownOutputEditorUtils::EditMediaOutput(MediaOutput, InMessage.Data);
 
 	Broadcast.SaveBroadcast();
 	
diff --git a/Engine/Plugins/Experimental/Avalanche/Source/AvalancheMediaEditor/Private/Rundown/AvaRundownServer.h b/Engine/Plugins/Experimental/Avalanche/Source/AvalancheMediaEditor/Private/Rundown/AvaRundownServer.h
--- a/Engine/Plugins/Experimental/Avalanche/Source/AvalancheMediaEditor/Private/Rundown/AvaRundownServer.h
+++ b/Engine/Plugins/Experimental/Avalanche/Source/AvalancheMediaEditor/Private/Rundown/AvaRundownServer.h
@@ -37,6 +37,7 @@ public:
 	void PageChannelChanged(const UAvaRundown* InRundown, const FAvaRundownPage& InPage, const FString& InChannelName) const;
 	void PageAnimSettingsChanged(const UAvaRundown* InRundown, const FAvaRundownPage& InPage) const;
 	void OnBroadcastChannelListChanged(const FAvaBroadcastProfile& InProfile) const;
+	void OnBroadcastChannelChanged(const FAvaBroadcastOutputChannel& InChannel, EAvaBroadcastChannelChange InChange) const;
 	void OnAssetAddedOrRemoved(const FAssetData& InAssetData) const;
 
 	/** Returns the endpoint's message address. */
