MotionDesign : BevelModifier
- Allowing higher iteration counts with roundess ratio for flat, inner, outer rounded bevels
- Reducing inset property lag when dragging

diff --git a/Engine/Plugins/Experimental/Avalanche/Source/AvalancheModifiers/Private/Modifiers/AvaBevelModifier.cpp b/Engine/Plugins/Experimental/Avalanche/Source/AvalancheModifiers/Private/Modifiers/AvaBevelModifier.cpp
--- a/Engine/Plugins/Experimental/Avalanche/Source/AvalancheModifiers/Private/Modifiers/AvaBevelModifier.cpp
+++ b/Engine/Plugins/Experimental/Avalanche/Source/AvalancheModifiers/Private/Modifiers/AvaBevelModifier.cpp
@@ -19,60 +19,45 @@ void UAvaBevelModifier::OnModifierCDOSetup(FActorModifierCoreMetadata& InMetadat
 	InMetadata.SetName(TEXT("Bevel"));
 	InMetadata.SetCategory(TEXT("Geometry"));
 #if WITH_EDITOR
-	InMetadata.SetDescription(LOCTEXT("ModifierDescription", "Create chamfered or rounded corners to geometry that smooth edges and corners"));
+	InMetadata.SetDescription(LOCTEXT("ModifierDescription", "Create chamfered or rounded corners on geometry that smooth edges and corners"));
 #endif
 }
 
 void UAvaBevelModifier::Apply()
 {
 	UDynamicMeshComponent* const DynMeshComp = GetMeshComponent();
-	
+
 	if (!IsValid(DynMeshComp))
 	{
 		Fail(LOCTEXT("InvalidDynamicMeshComponent", "Invalid dynamic mesh component on modified actor"));
 		return;
 	}
 
-	if (Inset <= 0.f)
+	if (Inset <= 0.f || DynMeshComp->GetDynamicMesh()->GetTriangleCount() == 0)
 	{
 		Next();
 		return;
 	}
-	
+
 	using namespace UE::Geometry;
-			
-	DynMeshComp->GetDynamicMesh()->EditMesh([this, DynMeshComp](FDynamicMesh3& EditMesh) 
+
+	DynMeshComp->GetDynamicMesh()->EditMesh([this, DynMeshComp](FDynamicMesh3& EditMesh)
 	{
-		if (EditMesh.TriangleCount() > 0)
-		{
-			// weld edges
-			FMergeCoincidentMeshEdges Welder(&EditMesh);
-			Welder.Apply();
-		}
+		// Weld edges
+		FMergeCoincidentMeshEdges Welder(&EditMesh);
+		Welder.Apply();
 
-		float BevelDistance = Inset;
-		float Divider = 4.f;
-				
 		// Apply bevel operator for multiple iterations
-		TArray<int32> NewTriangles;
-		for (int32 ItIdx = 0; ItIdx < Iterations; ItIdx++)
-		{
-			const FGroupTopology Topology(&EditMesh, true);
-					
-			FMeshBevel Bevel;
-			Bevel.InsetDistance = BevelDistance;
-			Bevel.InitializeFromGroupTopology(EditMesh, Topology);
-			Bevel.Apply(EditMesh, nullptr);
-					
-			NewTriangles.Append(Bevel.NewTriangles);
-
-			// lets reduce the distance to avoid reversed triangles and overlapping triangles
-			BevelDistance /= Divider;
-			Divider /= 2.f;
-		}
+		const FGroupTopology Topology(&EditMesh, true);
+		FMeshBevel Bevel;
+		Bevel.InsetDistance = Inset;
+		Bevel.NumSubdivisions = Iterations;
+		Bevel.RoundWeight = Iterations > 0 ? Roundness : 0;
+		Bevel.InitializeFromGroupTopology(EditMesh, Topology);
+		Bevel.Apply(EditMesh, nullptr);
 
 		// Get polygroup layer for back side
-		FDynamicMeshPolygroupAttribute* const BevelPolygroup = FindOrCreatePolygroupLayer(EditMesh, UAvaBevelModifier::BevelPolygroupLayerName, &NewTriangles);
+		FDynamicMeshPolygroupAttribute* const BevelPolygroup = FindOrCreatePolygroupLayer(EditMesh, UAvaBevelModifier::BevelPolygroupLayerName, &Bevel.NewTriangles);
 
 		// TODO : Fix UVs temp, UV modifier needed instead of this
 		FRotator BoxRotation = DynMeshComp->GetComponentTransform().Rotator();
@@ -82,100 +67,104 @@ void UAvaBevelModifier::Apply()
 		FDynamicMeshUVEditor UVEditor(&EditMesh, UVOverlay);
 		const FFrame3d ProjectionFrame(PlaneTransform);
 		FUVEditResult Result;
-		UVEditor.SetTriangleUVsFromBoxProjection(NewTriangles, [](const FVector3d& Pos)
+		UVEditor.SetTriangleUVsFromBoxProjection(Bevel.NewTriangles, [](const FVector3d& Pos)
 		{
 			return Pos;
 		}
 		, ProjectionFrame, MeshBounds.GetSize(), 3, &Result);
-				
-	}, EDynamicMeshChangeType::GeneralEdit, EDynamicMeshAttributeChangeFlags::Unknown, false);	
+
+	}, EDynamicMeshChangeType::GeneralEdit, EDynamicMeshAttributeChangeFlags::Unknown, false);
 
 	Next();
 }
 
 #if WITH_EDITOR
-void UAvaBevelModifier::PostEditChangeProperty(FPropertyChangedEvent& PropertyChangedEvent)
+void UAvaBevelModifier::PostEditChangeProperty(FPropertyChangedEvent& InPropertyChangedEvent)
 {
-	Super::PostEditChangeProperty(PropertyChangedEvent);
-
-	const FName MemberName = PropertyChangedEvent.GetMemberPropertyName();
-	
-	static const FName InsetName = GET_MEMBER_NAME_CHECKED(UAvaBevelModifier, Inset);
-	static const FName IterationsName = GET_MEMBER_NAME_CHECKED(UAvaBevelModifier, Iterations);
-	
-	if (MemberName == InsetName)
+	Super::PostEditChangeProperty(InPropertyChangedEvent);
+
+	const FName MemberName = InPropertyChangedEvent.GetMemberPropertyName();
+
+	if (MemberName == GET_MEMBER_NAME_CHECKED(UAvaBevelModifier, Inset))
 	{
 		OnInsetChanged();
 	}
-	else if (MemberName == IterationsName)
+	else if (MemberName == GET_MEMBER_NAME_CHECKED(UAvaBevelModifier, Iterations))
 	{
 		OnIterationsChanged();
 	}
+	else if (MemberName == GET_MEMBER_NAME_CHECKED(UAvaBevelModifier, Roundness))
+	{
+		OnRoundnessChanged();
+	}
 }
 #endif
 
 void UAvaBevelModifier::SetInset(float InInset)
 {
-	if (Inset == InInset)
+	InInset = FMath::Clamp(InInset, UAvaBevelModifier::MinInset, GetMaxInsetDistance());
+
+	if (FMath::IsNearlyEqual(Inset, InInset))
 	{
 		return;
 	}
 
-	Inset = FMath::Clamp(InInset, UAvaBevelModifier::MinInset, GetMaxBevel());
-	
+	Inset = InInset;
 	OnInsetChanged();
 }
 
 void UAvaBevelModifier::SetIterations(int32 InIterations)
 {
+	InIterations = FMath::Clamp(InIterations, UAvaBevelModifier::MinIterations, UAvaBevelModifier::MaxIterations);
+
 	if (Iterations == InIterations)
 	{
 		return;
 	}
 
-	Iterations = FMath::Clamp(InIterations, UAvaBevelModifier::MinIterations, UAvaBevelModifier::MaxIterations);
-	
+	Iterations = InIterations;
 	OnIterationsChanged();
 }
 
+void UAvaBevelModifier::SetRoundness(float InRoundness)
+{
+	InRoundness = FMath::Clamp(InRoundness, UAvaBevelModifier::MinRoundness, UAvaBevelModifier::MaxRoundness);
+
+	if (FMath::IsNearlyEqual(Roundness, InRoundness))
+	{
+		return;
+	}
+
+	Roundness = InRoundness;
+	OnRoundnessChanged();
+}
+
 void UAvaBevelModifier::OnInsetChanged()
 {
-	Inset = FMath::Min(Inset, GetMaxBevel());
+	Inset = FMath::Min(Inset, GetMaxInsetDistance());
 	MarkModifierDirty();
 }
 
 void UAvaBevelModifier::OnIterationsChanged()
 {
-	Inset = FMath::Min(Inset, GetMaxBevel());
 	MarkModifierDirty();
 }
 
-float UAvaBevelModifier::GetMaxBevel() const
+void UAvaBevelModifier::OnRoundnessChanged()
+{
+	MarkModifierDirty();
+}
+
+float UAvaBevelModifier::GetMaxInsetDistance() const
 {
-	float MaxBevelDistance = 0;
-	
 	if (!PreModifierCachedMesh.IsSet())
 	{
-		return MaxBevelDistance;
+		return 0.f;
 	}
 
 	const FBox Bounds = static_cast<FBox>(PreModifierCachedMesh.GetValue().GetBounds(true));
 	const FVector Size3d = Bounds.GetSize();
-
-	const float MinBevelDistance = FMath::Min3(Size3d.X / 2, Size3d.Y / 2, Size3d.Z / 2);
-	float Divider = 4.f;
-	float BevelDistance = MinBevelDistance;
-	MaxBevelDistance = MinBevelDistance;
-	
-	for (int32 Idx = 1; Idx < Iterations; Idx++)
-	{
-		const float Reduce = BevelDistance / Divider;
-		MaxBevelDistance -= Reduce;
-		BevelDistance = Reduce;
-		Divider /= 2.f;
-	}
-	
-	return MaxBevelDistance;
+	return FMath::Max(0, FMath::Min3(Size3d.X / 2, Size3d.Y / 2, Size3d.Z / 2) - 0.001f);
 }
 
 #undef LOCTEXT_NAMESPACE
diff --git a/Engine/Plugins/Experimental/Avalanche/Source/AvalancheModifiers/Public/Modifiers/AvaBevelModifier.h b/Engine/Plugins/Experimental/Avalanche/Source/AvalancheModifiers/Public/Modifiers/AvaBevelModifier.h
--- a/Engine/Plugins/Experimental/Avalanche/Source/AvalancheModifiers/Public/Modifiers/AvaBevelModifier.h
+++ b/Engine/Plugins/Experimental/Avalanche/Source/AvalancheModifiers/Public/Modifiers/AvaBevelModifier.h
@@ -13,8 +13,10 @@ class UAvaBevelModifier : public UAvaGeometryBaseModifier
 public:
 	static inline const FName BevelPolygroupLayerName = TEXT("Bevel");
 	static constexpr float MinInset = 0;
-	static constexpr int32 MinIterations = 1;
-	static constexpr int32 MaxIterations = 3;
+	static constexpr int32 MinIterations = 0;
+	static constexpr int32 MaxIterations = 10;
+	static constexpr float MinRoundness = -2;
+	static constexpr float MaxRoundness = 2;
 
 	AVALANCHEMODIFIERS_API void SetInset(float InBevel);
 	float GetInset() const
@@ -28,10 +30,16 @@ public:
 		return Iterations;
 	}
 
+	AVALANCHEMODIFIERS_API void SetRoundness(float InRoundness);
+	float GetRoundness() const
+	{
+		return Roundness;
+	}
+
 protected:
 	//~ Begin UObject
 #if WITH_EDITOR
-	virtual void PostEditChangeProperty(FPropertyChangedEvent& PropertyChangedEvent) override;
+	virtual void PostEditChangeProperty(FPropertyChangedEvent& InPropertyChangedEvent) override;
 #endif
 	//~ End UObject
 
@@ -42,12 +50,19 @@ protected:
 
 	void OnInsetChanged();
 	void OnIterationsChanged();
+	void OnRoundnessChanged();
 
-	float GetMaxBevel() const;
+	float GetMaxInsetDistance() const;
 
-	UPROPERTY(EditInstanceOnly, BlueprintReadWrite, Setter="SetInset", Getter="GetInset", Category="Bevel", meta=(ClampMin="0.0", AllowPrivateAccess="true"))
+	/** Distance used on vertices for beveling, clamped between 0 and (min bound size / 2) */
+	UPROPERTY(EditInstanceOnly, BlueprintReadWrite, Setter, Getter, Category="Bevel", meta=(ClampMin="0", AllowPrivateAccess="true"))
 	float Inset = 1.0f;
 
-	UPROPERTY(EditInstanceOnly, BlueprintReadWrite, Setter="SetIterations", Getter="GetIterations", Category="Bevel", meta=(ClampMin="1", ClampMax="3", AllowPrivateAccess="true"))
-	int32 Iterations = 1;
+	/** Amount of subdivisions applied on the bevel, could affect performance the higher this value gets */
+	UPROPERTY(EditInstanceOnly, BlueprintReadWrite, Setter, Getter, Category="Bevel", meta=(ClampMin="0", ClampMax="10", AllowPrivateAccess="true"))
+	int32 Iterations = 0;
+
+	/** Roundness of the beveling when multiple iterations are applied : -2 = inner rounded, 0 = flat, 2 = outer rounded */
+	UPROPERTY(EditInstanceOnly, BlueprintReadWrite, Setter, Getter, Category="Bevel", meta=(ClampMin="-2", ClampMax="2", EditCondition="Iterations > 0", AllowPrivateAccess="true"))
+	float Roundness = 0;
 };
