Motion Design
- Changed from index-based slot method to material-property-based one.

diff --git a/Engine/Plugins/Experimental/Avalanche/Source/AvalancheEditor/Private/MaterialDesigner/AvaMaterialDesignerTextureAssetFactory.cpp b/Engine/Plugins/Experimental/Avalanche/Source/AvalancheEditor/Private/MaterialDesigner/AvaMaterialDesignerTextureAssetFactory.cpp
--- a/Engine/Plugins/Experimental/Avalanche/Source/AvalancheEditor/Private/MaterialDesigner/AvaMaterialDesignerTextureAssetFactory.cpp
+++ b/Engine/Plugins/Experimental/Avalanche/Source/AvalancheEditor/Private/MaterialDesigner/AvaMaterialDesignerTextureAssetFactory.cpp
@@ -121,7 +121,7 @@ AActor* UAvaMaterialDesignerTextureAssetFactory::SpawnActor(UObject* InAsset, UL
 		return InNewActor;
 	}
 
-	const UDMMaterialSlot* Slot = EditorOnlyData->GetSlot(0);
+	const UDMMaterialSlot* Slot = EditorOnlyData->GetSlotForMaterialProperty(EDMMaterialPropertyType::BaseColor);
 
 	if (!Slot)
 	{
