Avalanche
- Fix for status bar buttons not working after changing viewports.

diff --git a/Engine/Plugins/Experimental/Avalanche/Source/AvalancheLevelViewport/Private/Widgets/SAvaLevelViewportStatusBarButtons.cpp b/Engine/Plugins/Experimental/Avalanche/Source/AvalancheLevelViewport/Private/Widgets/SAvaLevelViewportStatusBarButtons.cpp
--- a/Engine/Plugins/Experimental/Avalanche/Source/AvalancheLevelViewport/Private/Widgets/SAvaLevelViewportStatusBarButtons.cpp
+++ b/Engine/Plugins/Experimental/Avalanche/Source/AvalancheLevelViewport/Private/Widgets/SAvaLevelViewportStatusBarButtons.cpp
@@ -540,138 +540,138 @@ TSharedRef<SWidget> SAvaLevelViewportStatusBarButtons::GetPostProcessMenuContent
 
 	if (!ContextMenu)
 	{
-		using namespace UE::Ava::LevelViewportStatusBarButtons::Private;
-
 		ContextMenu = Menus->RegisterMenu(BackgroundMenuName, NAME_None, EMultiBoxType::Menu);
 
-		const FAvaLevelViewportCommands& CommandsRef = FAvaLevelViewportCommands::Get();
+		if (!ContextMenu)
+		{
+			return SNullWidget::NullWidget;
+		}
+	}
 
-		FToolMenuSection& EffectsSection = ContextMenu->AddSection("Effects", LOCTEXT("Effects", "Effects"));
+	using namespace UE::Ava::LevelViewportStatusBarButtons::Private;
 
-		FToolUIAction RGBAction;
-		RGBAction.ExecuteAction = FToolMenuExecuteAction::CreateSP(this, &SAvaLevelViewportStatusBarButtons::TogglePostProcessMenu, EAvaViewportPostProcessType::None);
-		RGBAction.GetActionCheckState = FToolMenuGetActionCheckState::CreateSP(this, &SAvaLevelViewportStatusBarButtons::GetPostProcessActiveMenu, EAvaViewportPostProcessType::None);
-		RGBAction.CanExecuteAction = FToolMenuCanExecuteAction::CreateSP(this, &SAvaLevelViewportStatusBarButtons::GetPostProcessEnabledMenu);
+	const FAvaLevelViewportCommands& CommandsRef = FAvaLevelViewportCommands::Get();
 
-		EffectsSection.AddEntry(FToolMenuEntry::InitMenuEntry(
-			"RGB",
-			LOCTEXT("RGB", "RGB"),
-			CommandsRef.TogglePostProcessNone->GetDescription(),
-			RGBChannelIcon,
-			FToolUIActionChoice(RGBAction),
-			EUserInterfaceActionType::Check
-		));
+	FToolMenuSection& EffectsSection = ContextMenu->FindOrAddSection("Effects", LOCTEXT("Effects", "Effects"));
 
-		FToolUIAction BackgroundAction;
-		BackgroundAction.ExecuteAction = FToolMenuExecuteAction::CreateSP(this, &SAvaLevelViewportStatusBarButtons::TogglePostProcessMenu, EAvaViewportPostProcessType::Background);
-		BackgroundAction.GetActionCheckState = FToolMenuGetActionCheckState::CreateSP(this, &SAvaLevelViewportStatusBarButtons::GetPostProcessActiveMenu, EAvaViewportPostProcessType::Background);
-		BackgroundAction.CanExecuteAction = FToolMenuCanExecuteAction::CreateSP(this, &SAvaLevelViewportStatusBarButtons::GetPostProcessEnabledMenu);
-
-		EffectsSection.AddEntry(FToolMenuEntry::InitMenuEntry(
-			"Background", 
-			LOCTEXT("Background", "Background"), 
-			CommandsRef.TogglePostProcessBackground->GetDescription(),
-			BackgroundIcon,
-			FToolUIActionChoice(BackgroundAction),
-			EUserInterfaceActionType::Check
-		));
+	FToolUIAction RGBAction;
+	RGBAction.ExecuteAction = FToolMenuExecuteAction::CreateSP(this, &SAvaLevelViewportStatusBarButtons::TogglePostProcessMenu, EAvaViewportPostProcessType::None);
+	RGBAction.GetActionCheckState = FToolMenuGetActionCheckState::CreateSP(this, &SAvaLevelViewportStatusBarButtons::GetPostProcessActiveMenu, EAvaViewportPostProcessType::None);
+	RGBAction.CanExecuteAction = FToolMenuCanExecuteAction::CreateSP(this, &SAvaLevelViewportStatusBarButtons::GetPostProcessEnabledMenu);
 
-		FToolUIAction RedAction;
-		RedAction.ExecuteAction = FToolMenuExecuteAction::CreateSP(this, &SAvaLevelViewportStatusBarButtons::TogglePostProcessMenu, EAvaViewportPostProcessType::RedChannel);
-		RedAction.GetActionCheckState = FToolMenuGetActionCheckState::CreateSP(this, &SAvaLevelViewportStatusBarButtons::GetPostProcessActiveMenu, EAvaViewportPostProcessType::RedChannel);
-		RedAction.CanExecuteAction = FToolMenuCanExecuteAction::CreateSP(this, &SAvaLevelViewportStatusBarButtons::GetPostProcessEnabledMenu);
-
-		EffectsSection.AddEntry(FToolMenuEntry::InitMenuEntry(
-			"Red",
-			LOCTEXT("Red", "Red"),
-			CommandsRef.TogglePostProcessChannelRed->GetDescription(),
-			RedChannelIcon,
-			FToolUIActionChoice(RedAction),
-			EUserInterfaceActionType::Check
-		));
+	EffectsSection.AddEntry(FToolMenuEntry::InitMenuEntry(
+		"RGB",
+		LOCTEXT("RGB", "RGB"),
+		CommandsRef.TogglePostProcessNone->GetDescription(),
+		RGBChannelIcon,
+		FToolUIActionChoice(RGBAction),
+		EUserInterfaceActionType::Check
+	));
 
-		FToolUIAction GreenAction;
-		GreenAction.ExecuteAction = FToolMenuExecuteAction::CreateSP(this, &SAvaLevelViewportStatusBarButtons::TogglePostProcessMenu, EAvaViewportPostProcessType::GreenChannel);
-		GreenAction.GetActionCheckState = FToolMenuGetActionCheckState::CreateSP(this, &SAvaLevelViewportStatusBarButtons::GetPostProcessActiveMenu, EAvaViewportPostProcessType::GreenChannel);
-		GreenAction.CanExecuteAction = FToolMenuCanExecuteAction::CreateSP(this, &SAvaLevelViewportStatusBarButtons::GetPostProcessEnabledMenu);
-
-		EffectsSection.AddEntry(FToolMenuEntry::InitMenuEntry(
-			"Green",
-			LOCTEXT("Green", "Green"),
-			CommandsRef.TogglePostProcessChannelGreen->GetDescription(),
-			GreenChannelIcon,
-			FToolUIActionChoice(GreenAction),
-			EUserInterfaceActionType::Check
-		));
+	FToolUIAction BackgroundAction;
+	BackgroundAction.ExecuteAction = FToolMenuExecuteAction::CreateSP(this, &SAvaLevelViewportStatusBarButtons::TogglePostProcessMenu, EAvaViewportPostProcessType::Background);
+	BackgroundAction.GetActionCheckState = FToolMenuGetActionCheckState::CreateSP(this, &SAvaLevelViewportStatusBarButtons::GetPostProcessActiveMenu, EAvaViewportPostProcessType::Background);
+	BackgroundAction.CanExecuteAction = FToolMenuCanExecuteAction::CreateSP(this, &SAvaLevelViewportStatusBarButtons::GetPostProcessEnabledMenu);
+
+	EffectsSection.AddEntry(FToolMenuEntry::InitMenuEntry(
+		"Background", 
+		LOCTEXT("Background", "Background"), 
+		CommandsRef.TogglePostProcessBackground->GetDescription(),
+		BackgroundIcon,
+		FToolUIActionChoice(BackgroundAction),
+		EUserInterfaceActionType::Check
+	));
 
-		FToolUIAction BlueAction;
-		BlueAction.ExecuteAction = FToolMenuExecuteAction::CreateSP(this, &SAvaLevelViewportStatusBarButtons::TogglePostProcessMenu, EAvaViewportPostProcessType::BlueChannel);
-		BlueAction.GetActionCheckState = FToolMenuGetActionCheckState::CreateSP(this, &SAvaLevelViewportStatusBarButtons::GetPostProcessActiveMenu, EAvaViewportPostProcessType::BlueChannel);
-		BlueAction.CanExecuteAction = FToolMenuCanExecuteAction::CreateSP(this, &SAvaLevelViewportStatusBarButtons::GetPostProcessEnabledMenu);
-
-		EffectsSection.AddEntry(FToolMenuEntry::InitMenuEntry(
-			"Blue",
-			LOCTEXT("Blue", "Blue"),
-			CommandsRef.TogglePostProcessChannelBlue->GetDescription(),
-			BlueChannelIcon,
-			FToolUIActionChoice(BlueAction),
-			EUserInterfaceActionType::Check
-		));
+	FToolUIAction RedAction;
+	RedAction.ExecuteAction = FToolMenuExecuteAction::CreateSP(this, &SAvaLevelViewportStatusBarButtons::TogglePostProcessMenu, EAvaViewportPostProcessType::RedChannel);
+	RedAction.GetActionCheckState = FToolMenuGetActionCheckState::CreateSP(this, &SAvaLevelViewportStatusBarButtons::GetPostProcessActiveMenu, EAvaViewportPostProcessType::RedChannel);
+	RedAction.CanExecuteAction = FToolMenuCanExecuteAction::CreateSP(this, &SAvaLevelViewportStatusBarButtons::GetPostProcessEnabledMenu);
+
+	EffectsSection.AddEntry(FToolMenuEntry::InitMenuEntry(
+		"Red",
+		LOCTEXT("Red", "Red"),
+		CommandsRef.TogglePostProcessChannelRed->GetDescription(),
+		RedChannelIcon,
+		FToolUIActionChoice(RedAction),
+		EUserInterfaceActionType::Check
+	));
 
-		FToolUIAction AlphaAction;
-		AlphaAction.ExecuteAction = FToolMenuExecuteAction::CreateSP(this, &SAvaLevelViewportStatusBarButtons::TogglePostProcessMenu, EAvaViewportPostProcessType::AlphaChannel);
-		AlphaAction.GetActionCheckState = FToolMenuGetActionCheckState::CreateSP(this, &SAvaLevelViewportStatusBarButtons::GetPostProcessActiveMenu, EAvaViewportPostProcessType::AlphaChannel);
-		AlphaAction.CanExecuteAction = FToolMenuCanExecuteAction::CreateSP(this, &SAvaLevelViewportStatusBarButtons::GetPostProcessEnabledMenu);
-
-		EffectsSection.AddEntry(FToolMenuEntry::InitMenuEntry(
-			"Alpha",
-			LOCTEXT("Alpha", "Alpha"),
-			CommandsRef.TogglePostProcessChannelAlpha->GetDescription(),
-			AlphaChannelIcon,
-			FToolUIActionChoice(AlphaAction),
-			EUserInterfaceActionType::Check
-		));
+	FToolUIAction GreenAction;
+	GreenAction.ExecuteAction = FToolMenuExecuteAction::CreateSP(this, &SAvaLevelViewportStatusBarButtons::TogglePostProcessMenu, EAvaViewportPostProcessType::GreenChannel);
+	GreenAction.GetActionCheckState = FToolMenuGetActionCheckState::CreateSP(this, &SAvaLevelViewportStatusBarButtons::GetPostProcessActiveMenu, EAvaViewportPostProcessType::GreenChannel);
+	GreenAction.CanExecuteAction = FToolMenuCanExecuteAction::CreateSP(this, &SAvaLevelViewportStatusBarButtons::GetPostProcessEnabledMenu);
+
+	EffectsSection.AddEntry(FToolMenuEntry::InitMenuEntry(
+		"Green",
+		LOCTEXT("Green", "Green"),
+		CommandsRef.TogglePostProcessChannelGreen->GetDescription(),
+		GreenChannelIcon,
+		FToolUIActionChoice(GreenAction),
+		EUserInterfaceActionType::Check
+	));
 
-		FToolUIAction CheckerboardAction;
-		CheckerboardAction.ExecuteAction = FToolMenuExecuteAction::CreateSP(this, &SAvaLevelViewportStatusBarButtons::TogglePostProcessMenu, EAvaViewportPostProcessType::Checkerboard);
-		CheckerboardAction.GetActionCheckState = FToolMenuGetActionCheckState::CreateSP(this, &SAvaLevelViewportStatusBarButtons::GetPostProcessActiveMenu, EAvaViewportPostProcessType::Checkerboard);
-		CheckerboardAction.CanExecuteAction = FToolMenuCanExecuteAction::CreateSP(this, &SAvaLevelViewportStatusBarButtons::GetPostProcessEnabledMenu);
-
-		EffectsSection.AddEntry(FToolMenuEntry::InitMenuEntry(
-			"Checkerboard",
-			LOCTEXT("Checkerboard", "Checkerboard"),
-			CommandsRef.TogglePostProcessCheckerboard->GetDescription(),
-			CheckerboardIcon,
-			FToolUIActionChoice(CheckerboardAction),
-			EUserInterfaceActionType::Check
-		));
+	FToolUIAction BlueAction;
+	BlueAction.ExecuteAction = FToolMenuExecuteAction::CreateSP(this, &SAvaLevelViewportStatusBarButtons::TogglePostProcessMenu, EAvaViewportPostProcessType::BlueChannel);
+	BlueAction.GetActionCheckState = FToolMenuGetActionCheckState::CreateSP(this, &SAvaLevelViewportStatusBarButtons::GetPostProcessActiveMenu, EAvaViewportPostProcessType::BlueChannel);
+	BlueAction.CanExecuteAction = FToolMenuCanExecuteAction::CreateSP(this, &SAvaLevelViewportStatusBarButtons::GetPostProcessEnabledMenu);
+
+	EffectsSection.AddEntry(FToolMenuEntry::InitMenuEntry(
+		"Blue",
+		LOCTEXT("Blue", "Blue"),
+		CommandsRef.TogglePostProcessChannelBlue->GetDescription(),
+		BlueChannelIcon,
+		FToolUIActionChoice(BlueAction),
+		EUserInterfaceActionType::Check
+	));
 
-		FToolMenuSection& OptionsSection = ContextMenu->AddSection("Options", LOCTEXT("Options", "Options"));
+	FToolUIAction AlphaAction;
+	AlphaAction.ExecuteAction = FToolMenuExecuteAction::CreateSP(this, &SAvaLevelViewportStatusBarButtons::TogglePostProcessMenu, EAvaViewportPostProcessType::AlphaChannel);
+	AlphaAction.GetActionCheckState = FToolMenuGetActionCheckState::CreateSP(this, &SAvaLevelViewportStatusBarButtons::GetPostProcessActiveMenu, EAvaViewportPostProcessType::AlphaChannel);
+	AlphaAction.CanExecuteAction = FToolMenuCanExecuteAction::CreateSP(this, &SAvaLevelViewportStatusBarButtons::GetPostProcessEnabledMenu);
+
+	EffectsSection.AddEntry(FToolMenuEntry::InitMenuEntry(
+		"Alpha",
+		LOCTEXT("Alpha", "Alpha"),
+		CommandsRef.TogglePostProcessChannelAlpha->GetDescription(),
+		AlphaChannelIcon,
+		FToolUIActionChoice(AlphaAction),
+		EUserInterfaceActionType::Check
+	));
 
-		if (PostProcessOpacitySlider.IsValid())
-		{
-			OptionsSection.AddEntry(FToolMenuEntry::InitWidget(
-				"PostProcessOpacity",
-				PostProcessOpacitySlider.ToSharedRef(),
-				LOCTEXT("PostProcessOpacity", "Opacity"),
-				true
-			));
-		}
+	FToolUIAction CheckerboardAction;
+	CheckerboardAction.ExecuteAction = FToolMenuExecuteAction::CreateSP(this, &SAvaLevelViewportStatusBarButtons::TogglePostProcessMenu, EAvaViewportPostProcessType::Checkerboard);
+	CheckerboardAction.GetActionCheckState = FToolMenuGetActionCheckState::CreateSP(this, &SAvaLevelViewportStatusBarButtons::GetPostProcessActiveMenu, EAvaViewportPostProcessType::Checkerboard);
+	CheckerboardAction.CanExecuteAction = FToolMenuCanExecuteAction::CreateSP(this, &SAvaLevelViewportStatusBarButtons::GetPostProcessEnabledMenu);
+
+	EffectsSection.AddEntry(FToolMenuEntry::InitMenuEntry(
+		"Checkerboard",
+		LOCTEXT("Checkerboard", "Checkerboard"),
+		CommandsRef.TogglePostProcessCheckerboard->GetDescription(),
+		CheckerboardIcon,
+		FToolUIActionChoice(CheckerboardAction),
+		EUserInterfaceActionType::Check
+	));
 
-		if (BackgroundTextureSelector.IsValid())
-		{
-			OptionsSection.AddEntry(FToolMenuEntry::InitWidget(
-				"PostProcessTexture",
-				BackgroundTextureSelector.ToSharedRef(),
-				LOCTEXT("PostProcessTexture", "Texture"),
-				true
-			));
-		}
+	FToolMenuSection& OptionsSection = ContextMenu->AddSection("Options", LOCTEXT("Options", "Options"));
+
+	if (PostProcessOpacitySlider.IsValid())
+	{
+		OptionsSection.AddEntry(FToolMenuEntry::InitWidget(
+			"PostProcessOpacity",
+			PostProcessOpacitySlider.ToSharedRef(),
+			LOCTEXT("PostProcessOpacity", "Opacity"),
+			true
+		));
 	}
 
-	if (!ContextMenu)
+	if (BackgroundTextureSelector.IsValid())
 	{
-		return SNullWidget::NullWidget;
+		OptionsSection.AddEntry(FToolMenuEntry::InitWidget(
+			"PostProcessTexture",
+			BackgroundTextureSelector.ToSharedRef(),
+			LOCTEXT("PostProcessTexture", "Texture"),
+			true
+		));
 	}
 
 	return Menus->GenerateWidget(ContextMenu);
@@ -874,42 +874,42 @@ TSharedRef<SWidget> SAvaLevelViewportStatusBarButtons::GetSnappingMenuContent()
 
 	if (!ContextMenu)
 	{
-		const FAvaLevelViewportGuideFrameAndWidget FrameAndWidget(ViewportFrameWeak);
+		ContextMenu = Menus->RegisterMenu(SnapMenuName, NAME_None, EMultiBoxType::Menu);
 
-		if (!FrameAndWidget.IsValid())
+		if (!ContextMenu)
 		{
 			return SNullWidget::NullWidget;
 		}
+	}
 
-		ContextMenu = Menus->RegisterMenu(SnapMenuName, NAME_None, EMultiBoxType::Menu);
+	const FAvaLevelViewportGuideFrameAndWidget FrameAndWidget(ViewportFrameWeak);
 
-		FToolMenuSection& Section = ContextMenu->AddSection("SnapTo", LOCTEXT("SnapTo", "Snap To"));
+	if (!FrameAndWidget.IsValid())
+	{
+		return SNullWidget::NullWidget;
+	}
 
-		const FAvaLevelViewportCommands& CommandsRef = FAvaLevelViewportCommands::Get();
+	FToolMenuSection& SnapToSection = ContextMenu->FindOrAddSection("SnapTo", LOCTEXT("SnapTo", "Snap To"));
 
-		Section.AddEntry(FToolMenuEntry::InitMenuEntryWithCommandList(
-			CommandsRef.ToggleGridSnapping,
-			FrameAndWidget.ViewportWidget->GetCommandList(),
-			LOCTEXT("GridSnapping", "Grid")
-		));
+	const FAvaLevelViewportCommands& CommandsRef = FAvaLevelViewportCommands::Get();
 
-		Section.AddEntry(FToolMenuEntry::InitMenuEntryWithCommandList(
-			CommandsRef.ToggleScreenSnapping,
-			FrameAndWidget.ViewportWidget->GetCommandList(),
-			LOCTEXT("ScreenSnapping", "Screen & Guide")
-		));
+	SnapToSection.AddEntry(FToolMenuEntry::InitMenuEntryWithCommandList(
+		CommandsRef.ToggleGridSnapping,
+		FrameAndWidget.ViewportWidget->GetCommandList(),
+		LOCTEXT("GridSnapping", "Grid")
+	));
 
-		Section.AddEntry(FToolMenuEntry::InitMenuEntryWithCommandList(
-			CommandsRef.ToggleActorSnapping,
-			FrameAndWidget.ViewportWidget->GetCommandList(),
-			LOCTEXT("ActorSnapping", "Actor")
-		));
-	}
+	SnapToSection.AddEntry(FToolMenuEntry::InitMenuEntryWithCommandList(
+		CommandsRef.ToggleScreenSnapping,
+		FrameAndWidget.ViewportWidget->GetCommandList(),
+		LOCTEXT("ScreenSnapping", "Screen & Guide")
+	));
 
-	if (!ContextMenu)
-	{
-		return SNullWidget::NullWidget;
-	}
+	SnapToSection.AddEntry(FToolMenuEntry::InitMenuEntryWithCommandList(
+		CommandsRef.ToggleActorSnapping,
+		FrameAndWidget.ViewportWidget->GetCommandList(),
+		LOCTEXT("ActorSnapping", "Actor")
+	));
 
 	return Menus->GenerateWidget(ContextMenu);
 }
@@ -1053,35 +1053,35 @@ TSharedRef<SWidget> SAvaLevelViewportStatusBarButtons::GetGridMenuContent() cons
 	{
 		ContextMenu = Menus->RegisterMenu(GridMenuName, NAME_None, EMultiBoxType::Menu);
 
-		FToolMenuSection& Section = ContextMenu->AddSection("Grid", LOCTEXT("Grid", "Grid"));
+		if (!ContextMenu)
+		{
+			return SNullWidget::NullWidget;
+		}
+	}
 
-		const FAvaLevelViewportCommands& CommandsRef = FAvaLevelViewportCommands::Get();
+	FToolMenuSection& GridSection = ContextMenu->FindOrAddSection("Grid", LOCTEXT("Grid", "Grid"));
 
-		const FAvaLevelViewportGuideFrameAndWidget FrameAndWidget(ViewportFrameWeak);
+	const FAvaLevelViewportCommands& CommandsRef = FAvaLevelViewportCommands::Get();
 
-		if (FrameAndWidget.IsValid())
-		{
-			Section.AddEntry(FToolMenuEntry::InitMenuEntryWithCommandList(
-				CommandsRef.ToggleGridAlwaysVisible,
-				FrameAndWidget.ViewportWidget->GetCommandList(),
-				LOCTEXT("AlwaysShowGrid", "Always On")
-			));
-		}
+	const FAvaLevelViewportGuideFrameAndWidget FrameAndWidget(ViewportFrameWeak);
 
-		if (GridSizeSlider.IsValid())
-		{
-			Section.AddEntry(FToolMenuEntry::InitWidget(
-				"GridSize",
-				GridSizeSlider.ToSharedRef(),
-				LOCTEXT("GridSize", "Size"),
-				true
-			));
-		}
+	if (FrameAndWidget.IsValid())
+	{
+		GridSection.AddEntry(FToolMenuEntry::InitMenuEntryWithCommandList(
+			CommandsRef.ToggleGridAlwaysVisible,
+			FrameAndWidget.ViewportWidget->GetCommandList(),
+			LOCTEXT("AlwaysShowGrid", "Always On")
+		));
 	}
 
-	if (!ContextMenu)
+	if (GridSizeSlider.IsValid())
 	{
-		return SNullWidget::NullWidget;
+		GridSection.AddEntry(FToolMenuEntry::InitWidget(
+			"GridSize",
+			GridSizeSlider.ToSharedRef(),
+			LOCTEXT("GridSize", "Size"),
+			true
+		));
 	}
 
 	return Menus->GenerateWidget(ContextMenu);
@@ -1209,42 +1209,42 @@ TSharedRef<SWidget> SAvaLevelViewportStatusBarButtons::GetTextureOverlayMenuCont
 
 		ContextMenu = Menus->RegisterMenu(TextureOverlayMenuName, NAME_None, EMultiBoxType::Menu);
 
-		FToolMenuSection& OptionsSection = ContextMenu->AddSection("TextureOverlay", LOCTEXT("TextureOverlay", "Texture Overlay"));
-
-		if (TextureOverlayOpacitySlider.IsValid())
+		if (!ContextMenu)
 		{
-			OptionsSection.AddEntry(FToolMenuEntry::InitWidget(
-				"TextureOverlayOpacity",
-				TextureOverlayOpacitySlider.ToSharedRef(),
-				LOCTEXT("TextureOverlayOpacity", "Opacity"),
-				true
-			));
+			return SNullWidget::NullWidget;
 		}
+	}
 
-		if (TextureOverlayTextureSelector.IsValid())
-		{
-			OptionsSection.AddEntry(FToolMenuEntry::InitWidget(
-				"TextureOverlayTexture",
-				TextureOverlayTextureSelector.ToSharedRef(),
-				LOCTEXT("TextureOverlayTexture", "Texture"),
-				true
-			));
-		}
+	FToolMenuSection& OptionsSection = ContextMenu->FindOrAddSection("TextureOverlay", LOCTEXT("TextureOverlay", "Texture Overlay"));
 
-		if (TextureOverlayStretchCheckBox.IsValid())
-		{
-			OptionsSection.AddEntry(FToolMenuEntry::InitWidget(
-				"TextureOverlayStretch",
-				TextureOverlayStretchCheckBox.ToSharedRef(),
-				LOCTEXT("TextureOverlayStretch", "Stretch Texture"),
-				true
-			));
-		}
+	if (TextureOverlayOpacitySlider.IsValid())
+	{
+		OptionsSection.AddEntry(FToolMenuEntry::InitWidget(
+			"TextureOverlayOpacity",
+			TextureOverlayOpacitySlider.ToSharedRef(),
+			LOCTEXT("TextureOverlayOpacity", "Opacity"),
+			true
+		));
 	}
 
-	if (!ContextMenu)
+	if (TextureOverlayTextureSelector.IsValid())
 	{
-		return SNullWidget::NullWidget;
+		OptionsSection.AddEntry(FToolMenuEntry::InitWidget(
+			"TextureOverlayTexture",
+			TextureOverlayTextureSelector.ToSharedRef(),
+			LOCTEXT("TextureOverlayTexture", "Texture"),
+			true
+		));
+	}
+
+	if (TextureOverlayStretchCheckBox.IsValid())
+	{
+		OptionsSection.AddEntry(FToolMenuEntry::InitWidget(
+			"TextureOverlayStretch",
+			TextureOverlayStretchCheckBox.ToSharedRef(),
+			LOCTEXT("TextureOverlayStretch", "Stretch Texture"),
+			true
+		));
 	}
 
 	return Menus->GenerateWidget(ContextMenu);
