Motion Design
- Fixed shape area planner drawing incorrectly.

diff --git a/Engine/Plugins/Experimental/Avalanche/Source/AvalancheInteractiveTools/Private/Planners/AvaInteractiveToolsToolViewportAreaPlanner.cpp b/Engine/Plugins/Experimental/Avalanche/Source/AvalancheInteractiveTools/Private/Planners/AvaInteractiveToolsToolViewportAreaPlanner.cpp
--- a/Engine/Plugins/Experimental/Avalanche/Source/AvalancheInteractiveTools/Private/Planners/AvaInteractiveToolsToolViewportAreaPlanner.cpp
+++ b/Engine/Plugins/Experimental/Avalanche/Source/AvalancheInteractiveTools/Private/Planners/AvaInteractiveToolsToolViewportAreaPlanner.cpp
@@ -187,7 +187,7 @@ void UAvaInteractiveToolsToolViewportAreaPlanner::DrawHUD(FCanvas* InCanvas, ITo
 		return;
 	}
 
-	const FAvaVisibleArea VisibleArea = AvaViewportClient->GetVirtualZoomedVisibleArea();
+	const FAvaVisibleArea VisibleArea = AvaViewportClient->GetZoomedVisibleArea();
 
 	if (!VisibleArea.IsValid())
 	{
diff --git a/Engine/Plugins/Experimental/Avalanche/Source/AvalancheLevelViewport/Private/Widgets/SAvaLevelViewportGuide.cpp b/Engine/Plugins/Experimental/Avalanche/Source/AvalancheLevelViewport/Private/Widgets/SAvaLevelViewportGuide.cpp
--- a/Engine/Plugins/Experimental/Avalanche/Source/AvalancheLevelViewport/Private/Widgets/SAvaLevelViewportGuide.cpp
+++ b/Engine/Plugins/Experimental/Avalanche/Source/AvalancheLevelViewport/Private/Widgets/SAvaLevelViewportGuide.cpp
@@ -502,7 +502,6 @@ bool SAvaLevelViewportGuide::DragUpdate()
 	}
 
 	FVector2f MousePosition = FrameAndClient.ViewportClient->GetConstrainedZoomedViewportMousePosition();
-	FVector2f Offset = FrameAndClient.ViewportClient->GetCachedViewportOffset();
 	FVector2f NewSlotPosition = FVector2f::ZeroVector;
 
 	if (TSharedPtr<FAvaSnapOperation> SnapOperation = FrameAndClient.ViewportClient->GetSnapOperation())
