Motion Design
- Adjusted the way double clicks are detected.

diff --git a/Engine/Plugins/Experimental/Avalanche/Source/AvalancheInteractiveTools/Private/EdMode/AvaInteractiveToolsEdMode.cpp b/Engine/Plugins/Experimental/Avalanche/Source/AvalancheInteractiveTools/Private/EdMode/AvaInteractiveToolsEdMode.cpp
--- a/Engine/Plugins/Experimental/Avalanche/Source/AvalancheInteractiveTools/Private/EdMode/AvaInteractiveToolsEdMode.cpp
+++ b/Engine/Plugins/Experimental/Avalanche/Source/AvalancheInteractiveTools/Private/EdMode/AvaInteractiveToolsEdMode.cpp
@@ -7,7 +7,6 @@
 #include "EdModeInteractiveToolsContext.h"
 #include "EditorModeManager.h"
 #include "EditorModes.h"
-#include "GameFramework/InputSettings.h"
 #include "IAvalancheInteractiveToolsModule.h"
 #include "Selection.h"
 #include "Tools/AvaInteractiveToolsToolBase.h"
@@ -182,6 +181,11 @@ void UAvaInteractiveToolsEdMode::OnToolShutdown(UInteractiveTool* InTool, EToolS
 	}
 }
 
+void UAvaInteractiveToolsEdMode::OnToolActivateEnd()
+{
+	LastToolActivateTime = FApp::GetCurrentTime();
+}
+
 void UAvaInteractiveToolsEdMode::OnActorSelectionChange(const UTypedElementSelectionSet* InSelectionSet)
 {
 	UInteractiveToolManager* ToolManager = GetToolManager();
diff --git a/Engine/Plugins/Experimental/Avalanche/Source/AvalancheInteractiveTools/Private/EdMode/AvaInteractiveToolsEdMode.h b/Engine/Plugins/Experimental/Avalanche/Source/AvalancheInteractiveTools/Private/EdMode/AvaInteractiveToolsEdMode.h
--- a/Engine/Plugins/Experimental/Avalanche/Source/AvalancheInteractiveTools/Private/EdMode/AvaInteractiveToolsEdMode.h
+++ b/Engine/Plugins/Experimental/Avalanche/Source/AvalancheInteractiveTools/Private/EdMode/AvaInteractiveToolsEdMode.h
@@ -30,10 +30,13 @@ public:
 	void OnToolSetup(UInteractiveTool* InTool);
 	void OnToolShutdown(UInteractiveTool* InTool, EToolShutdownType InShutdownType);
 	const FString GetLastActiveTool() const { return LastActiveTool; }
+	double GetLastToolActivateTime() { return LastToolActivateTime; }
+	void OnToolActivateEnd();
 
 protected:
 	FString LastActiveTool;
 	TWeakObjectPtr<UTypedElementSelectionSet> WeakActorSelectionSet;
+	double LastToolActivateTime;
 
 	void OnActorSelectionChange(const UTypedElementSelectionSet* InSelectionSet);
 };
diff --git a/Engine/Plugins/Experimental/Avalanche/Source/AvalancheInteractiveTools/Private/Tools/AvaInteractiveToolsToolBase.cpp b/Engine/Plugins/Experimental/Avalanche/Source/AvalancheInteractiveTools/Private/Tools/AvaInteractiveToolsToolBase.cpp
--- a/Engine/Plugins/Experimental/Avalanche/Source/AvalancheInteractiveTools/Private/Tools/AvaInteractiveToolsToolBase.cpp
+++ b/Engine/Plugins/Experimental/Avalanche/Source/AvalancheInteractiveTools/Private/Tools/AvaInteractiveToolsToolBase.cpp
@@ -2,17 +2,18 @@
 
 #include "Tools/AvaInteractiveToolsToolBase.h"
 #include "AvaInteractiveToolsSettings.h"
-#include "AvaViewportUtils.h"
 #include "AvalancheInteractiveToolsModule.h"
+#include "AvaViewportUtils.h"
 #include "BaseBehaviors/SingleClickBehavior.h"
 #include "BaseBehaviors/SingleKeyCaptureBehavior.h"
 #include "ContextObjectStore.h"
-#include "EdMode/AvaInteractiveToolsEdMode.h"
 #include "Editor/EditorEngine.h"
 #include "EditorViewportClient.h"
+#include "EdMode/AvaInteractiveToolsEdMode.h"
 #include "Engine/World.h"
 #include "EngineAnalytics.h"
 #include "Framework/Application/SlateApplication.h"
+#include "GameFramework/InputSettings.h"
 #include "IAvaInteractiveToolsModeDetailsObject.h"
 #include "IAvaInteractiveToolsModeDetailsObjectProvider.h"
 #include "InteractiveToolManager.h"
@@ -85,6 +86,8 @@ void UAvaInteractiveToolsToolBase::Setup()
 {
 	Super::Setup();
 
+	using namespace UE::AvaInteractiveTools::Private;
+
 	UInteractiveToolManager* ToolManager = GetToolManager();
 	bool bReactivated = false;
 	UAvaInteractiveToolsEdMode* AvaInteractiveToolsEdMode = nullptr;
@@ -95,13 +98,23 @@ void UAvaInteractiveToolsToolBase::Setup()
 
 		if (AvaInteractiveToolsEdMode)
 		{
-			bReactivated = AvaInteractiveToolsEdMode->GetLastActiveTool() == ToolManager->GetActiveToolName(EToolSide::Left);
+			// I found double-clicking inconsistent at the default speed. 
+			static const double ReactivateDelayTime = GetDefault<UInputSettings>()->DoubleClickTime + 0.1;
+
+			const bool bIsDoubleClick = (FApp::GetCurrentTime() - AvaInteractiveToolsEdMode->GetLastToolActivateTime()) <= ReactivateDelayTime;
+			const bool bIsSameTool = AvaInteractiveToolsEdMode->GetLastActiveTool() == ToolManager->GetActiveToolName(EToolSide::Left);
+
+			if (bIsDoubleClick && bIsSameTool && SupportsDefaultAction())
+			{
+				bReactivated = true;
+			}
 		}
 	}
 
 	if (!CanActivate(bReactivated))
 	{
 		ToolManager->PostActiveToolShutdownRequest(this, EToolShutdownType::Cancel);
+		return;
 	}
 
 	if (AvaInteractiveToolsEdMode)
@@ -110,6 +123,11 @@ void UAvaInteractiveToolsToolBase::Setup()
 	}
 
 	Activate(bReactivated);
+
+	if (AvaInteractiveToolsEdMode)
+	{
+		AvaInteractiveToolsEdMode->OnToolActivateEnd();
+	}
 }
 
 void UAvaInteractiveToolsToolBase::Shutdown(EToolShutdownType ShutdownType)
