Motion Design
- Fixed a code analysis issue

diff --git a/Engine/Plugins/Experimental/Avalanche/Source/AvalancheEditor/Private/MaterialDesigner/AvaLevelMaterialDesignerExtension.cpp b/Engine/Plugins/Experimental/Avalanche/Source/AvalancheEditor/Private/MaterialDesigner/AvaLevelMaterialDesignerExtension.cpp
--- a/Engine/Plugins/Experimental/Avalanche/Source/AvalancheEditor/Private/MaterialDesigner/AvaLevelMaterialDesignerExtension.cpp
+++ b/Engine/Plugins/Experimental/Avalanche/Source/AvalancheEditor/Private/MaterialDesigner/AvaLevelMaterialDesignerExtension.cpp
@@ -71,7 +71,7 @@ void FAvaLevelMaterialDesignerExtension::DeinitContentBrowserExtension()
 		{
 			TArray<FContentBrowserMenuExtender_SelectedAssets>& CBMenuExtenderDelegates = ContentBrowserModule->GetAllAssetViewContextMenuExtenders();
 
-			CBMenuExtenderDelegates.RemoveAll([ContentBrowserExtenderDelegateHandle = ContentBrowserExtenderDelegateHandle](const FContentBrowserMenuExtender_SelectedAssets& InDelegate)
+			CBMenuExtenderDelegates.RemoveAll([](const FContentBrowserMenuExtender_SelectedAssets& InDelegate)
 				{
 					return InDelegate.GetHandle() == ContentBrowserExtenderDelegateHandle;
 				});
