Motion Design: fix crash when trying to open a sub-sequence from an asset within Motion Design Sequencer

diff --git a/Engine/Plugins/Experimental/Avalanche/Source/AvalancheEditor/Private/Sequencer/AvaLevelSequencerExtension.cpp b/Engine/Plugins/Experimental/Avalanche/Source/AvalancheEditor/Private/Sequencer/AvaLevelSequencerExtension.cpp
--- a/Engine/Plugins/Experimental/Avalanche/Source/AvalancheEditor/Private/Sequencer/AvaLevelSequencerExtension.cpp
+++ b/Engine/Plugins/Experimental/Avalanche/Source/AvalancheEditor/Private/Sequencer/AvaLevelSequencerExtension.cpp
@@ -154,7 +154,7 @@ TSharedPtr<ISequencer> FAvaLevelSequencerExtension::GetExternalSequencer() const
 void FAvaLevelSequencerExtension::OnViewedSequenceChanged(UAvaSequence* InOldSequence, UAvaSequence* InNewSequence)
 {
 	IAvaSequenceProvider* SequenceProvider = GetSequenceProvider();
-	if (!SequenceProvider)
+	if (!SequenceProvider || !InNewSequence)
 	{
 		return;
 	}
@@ -162,7 +162,7 @@ void FAvaLevelSequencerExtension::OnViewedSequenceChanged(UAvaSequence* InOldSeq
 	UObject* const Outer = SequenceProvider->ToUObject();
 
 	bool bIsOldOwnedByProvider = InOldSequence && InOldSequence->GetOuter() == Outer;
-	bool bIsNewOwnedByProvider = InNewSequence && InNewSequence->GetOuter() == Outer;
+	bool bIsNewOwnedByProvider = InNewSequence->GetOuter() == Outer;
 
 	// Only avoid re-opening the asset editor if it's changing between two sequences that are owned by the provider
 	if (!bIsOldOwnedByProvider || !bIsNewOwnedByProvider)
diff --git a/Engine/Plugins/Experimental/Avalanche/Source/AvalancheSequencer/Private/AvaSequencer.cpp b/Engine/Plugins/Experimental/Avalanche/Source/AvalancheSequencer/Private/AvaSequencer.cpp
--- a/Engine/Plugins/Experimental/Avalanche/Source/AvalancheSequencer/Private/AvaSequencer.cpp
+++ b/Engine/Plugins/Experimental/Avalanche/Source/AvalancheSequencer/Private/AvaSequencer.cpp
@@ -1440,6 +1440,13 @@ TArray<UAvaSequence*> FAvaSequencer::GetSequencesForObject(UObject* InObject) co
 
 TSharedRef<SWidget> FAvaSequencer::CreateSequenceWidget()
 {
+	// Force the SequencerWeak ptr to be invalid if this AvaSequencer doesn't explicitly own the sequencer (i.e. InstancedSequencer is null)
+	// This is to force the sequencer to look for a new sequencer again
+	if (!InstancedSequencer.IsValid())
+	{
+		SequencerWeak.Reset();
+	}
+
 	TSharedRef<ISequencer> Sequencer = GetSequencer();
 
 	TArray<TSharedRef<IAvaSequenceSectionDetails>> Sections =
