Motion Designer
- Changed asset data in tool placeable items from cdo to class

diff --git a/Engine/Plugins/Experimental/Avalanche/Source/AvalancheInteractiveTools/Private/AvalancheInteractiveToolsModule.cpp b/Engine/Plugins/Experimental/Avalanche/Source/AvalancheInteractiveTools/Private/AvalancheInteractiveToolsModule.cpp
--- a/Engine/Plugins/Experimental/Avalanche/Source/AvalancheInteractiveTools/Private/AvalancheInteractiveToolsModule.cpp
+++ b/Engine/Plugins/Experimental/Avalanche/Source/AvalancheInteractiveTools/Private/AvalancheInteractiveToolsModule.cpp
@@ -260,7 +260,7 @@ void FAvalancheInteractiveToolsModule::OnPlacementCategoryRefreshed(FName InCate
 		{
 			PlaceableItem = MakeShared<FPlaceableItem>(
 				Tool.Factory.Get(),
-				FAssetData(Tool.Factory->NewActorClass->GetDefaultObject()),
+				FAssetData(Tool.Factory->NewActorClass),
 				Tool.Priority
 			);
 		}
