[Motion Design] Fix profile rename when there is a pinned channel.

diff --git a/Engine/Plugins/Experimental/Avalanche/Source/AvalancheMedia/Private/Broadcast/AvaBroadcast.cpp b/Engine/Plugins/Experimental/Avalanche/Source/AvalancheMedia/Private/Broadcast/AvaBroadcast.cpp
--- a/Engine/Plugins/Experimental/Avalanche/Source/AvalancheMedia/Private/Broadcast/AvaBroadcast.cpp
+++ b/Engine/Plugins/Experimental/Avalanche/Source/AvalancheMedia/Private/Broadcast/AvaBroadcast.cpp
@@ -266,6 +266,16 @@ bool UAvaBroadcast::RenameProfile(FName InProfileName, FName InNewProfileName)
 		{
 			CurrentProfile = InNewProfileName;
 		}
+
+		// Rename pinned channel's profile.
+		for (TPair<FName, FName>& PinnedChannel : PinnedChannels)
+		{
+			if (PinnedChannel.Value == InProfileName)
+			{
+				PinnedChannel.Value = InNewProfileName;
+			}
+		}
+		
 		return true;
 	}
 	return false;
diff --git a/Engine/Plugins/Experimental/Avalanche/Source/AvalancheMedia/Private/Broadcast/AvaBroadcastProfile.cpp b/Engine/Plugins/Experimental/Avalanche/Source/AvalancheMedia/Private/Broadcast/AvaBroadcastProfile.cpp
--- a/Engine/Plugins/Experimental/Avalanche/Source/AvalancheMedia/Private/Broadcast/AvaBroadcastProfile.cpp
+++ b/Engine/Plugins/Experimental/Avalanche/Source/AvalancheMedia/Private/Broadcast/AvaBroadcastProfile.cpp
@@ -303,8 +303,16 @@ FAvaBroadcastOutputChannel& FAvaBroadcastProfile::GetChannelMutable(FName InChan
 	const FName PinnedChannelProfileName = Broadcast.GetPinnedChannelProfileName(InChannelName);
 	if (PinnedChannelProfileName != NAME_None)
 	{
-		return Broadcast.GetProfile(PinnedChannelProfileName).GetLocalChannelMutable(InChannelName);
-	}	
+		const FAvaBroadcastProfile& PinnedProfile = Broadcast.GetProfile(PinnedChannelProfileName); 
+		if (PinnedProfile.IsValidProfile())
+		{
+			return Broadcast.GetProfile(PinnedChannelProfileName).GetLocalChannelMutable(InChannelName);
+		}
+
+		UE_LOG(LogAvaBroadcast, Error,
+			TEXT("Can't find profile \"%s\" for pinned channel \"%s\". Using local channel instead."),
+			*PinnedChannelProfileName.ToString(), *InChannelName.ToString());	
+	}
 	return GetLocalChannelMutable(InChannelName);
 }
 
