Motion Design: fix missing return succeeded for discard self task

diff --git a/Engine/Plugins/Experimental/Avalanche/Source/AvalancheTransition/Private/Tasks/AvaTransitionDiscardSceneTask.cpp b/Engine/Plugins/Experimental/Avalanche/Source/AvalancheTransition/Private/Tasks/AvaTransitionDiscardSceneTask.cpp
--- a/Engine/Plugins/Experimental/Avalanche/Source/AvalancheTransition/Private/Tasks/AvaTransitionDiscardSceneTask.cpp
+++ b/Engine/Plugins/Experimental/Avalanche/Source/AvalancheTransition/Private/Tasks/AvaTransitionDiscardSceneTask.cpp
@@ -11,6 +11,7 @@ EStateTreeRunStatus FAvaTransitionDiscardSceneTask::EnterState(FStateTreeExecuti
 	if (FAvaTransitionScene* TransitionScene = TransitionContext.GetTransitionScene())
 	{
 		TransitionScene->SetFlags(EAvaTransitionSceneFlags::NeedsDiscard);
+		return EStateTreeRunStatus::Succeeded;
 	}
 
 	return EStateTreeRunStatus::Failed;
