Motion Design: Fix issue with tag attribute added to Scene Settings unable to be removed with the "remove tag" task in TL.

diff --git a/Engine/Plugins/Experimental/Avalanche/Source/Avalanche/Private/AvaScene.cpp b/Engine/Plugins/Experimental/Avalanche/Source/Avalanche/Private/AvaScene.cpp
--- a/Engine/Plugins/Experimental/Avalanche/Source/Avalanche/Private/AvaScene.cpp
+++ b/Engine/Plugins/Experimental/Avalanche/Source/Avalanche/Private/AvaScene.cpp
@@ -74,11 +74,6 @@ AAvaScene::AAvaScene()
 
 	SceneState = CreateDefaultSubobject<UAvaSceneState>(TEXT("SceneState"));
 
-	if (SceneState)
-	{
-		SceneState->SetSceneSettings(SceneSettings);
-	}
-
 	RemoteControlPreset = CreateDefaultSubobject<URemoteControlPreset>(TEXT("RemoteControlPreset"));
 
 	StartupCameraName = NAME_None;
@@ -296,6 +291,20 @@ void AAvaScene::PostLoad()
 	{
 		SceneSubsystem->RegisterSceneInterface(GetLevel(), this);
 	}
+
+	if (SceneState)
+	{
+		SceneState->Initialize(SceneSettings);
+	}
+}
+
+void AAvaScene::PostInitializeComponents()
+{
+	Super::PostInitializeComponents();
+	if (SceneState)
+	{
+		SceneState->Initialize(SceneSettings);
+	}
 }
 
 void AAvaScene::PostDuplicate(EDuplicateMode::Type InDuplicateMode)
diff --git a/Engine/Plugins/Experimental/Avalanche/Source/Avalanche/Private/AvaSceneState.cpp b/Engine/Plugins/Experimental/Avalanche/Source/Avalanche/Private/AvaSceneState.cpp
--- a/Engine/Plugins/Experimental/Avalanche/Source/Avalanche/Private/AvaSceneState.cpp
+++ b/Engine/Plugins/Experimental/Avalanche/Source/Avalanche/Private/AvaSceneState.cpp
@@ -3,61 +3,69 @@
 #include "AvaSceneState.h"
 #include "AvaSceneSettings.h"
 #include "AvaTagHandle.h"
+#include "AvaTagHandleContainer.h"
 #include "Tags/AvaTagAttribute.h"
+#include "Tags/AvaTagAttributeBase.h"
 
-void UAvaSceneState::SetSceneSettings(UAvaSceneSettings* InSceneSettings)
+void UAvaSceneState::Initialize(UAvaSceneSettings* InSceneSettings)
 {
-	SceneSettingsWeak = InSceneSettings;
+	SceneAttributes.Reset();
+	if (InSceneSettings)
+	{
+		SceneAttributes = InSceneSettings->GetSceneAttributes();
+	}
 }
 
 bool UAvaSceneState::AddTagAttribute(const FAvaTagHandle& InTagHandle)
 {
-	if (const FAvaTag* Tag = InTagHandle.GetTag())
+	// Return true if already existing
+	if (ContainsTagAttribute(InTagHandle))
 	{
-		bool bAlreadyInSet;
-		ActiveTagAttributes.Add(*Tag, &bAlreadyInSet);
-		return !bAlreadyInSet;
+		return true;
 	}
-	return false;
-}
 
-bool UAvaSceneState::RemoveTagAttribute(const FAvaTagHandle& InTagHandle)
-{
 	if (const FAvaTag* Tag = InTagHandle.GetTag())
 	{
-		return ActiveTagAttributes.Remove(*Tag) > 0;
+		UAvaTagAttribute* TagAttribute = NewObject<UAvaTagAttribute>(this, NAME_None, RF_Transient);
+		check(TagAttribute);
+		TagAttribute->Tag = InTagHandle;
+
+		SceneAttributes.Add(TagAttribute);
+		return true;
 	}
+
 	return false;
 }
 
-bool UAvaSceneState::ContainsTagAttribute(const FAvaTagHandle& InTagHandle) const
+bool UAvaSceneState::RemoveTagAttribute(const FAvaTagHandle& InTagHandle)
 {
-	const FAvaTag* Tag = InTagHandle.GetTag();
-	if (!Tag)
-	{
-		return false;
-	}
+	uint32 TagsCleared = 0;
 
-	if (ActiveTagAttributes.Contains(*Tag))
+	for (UAvaAttribute* Attribute : SceneAttributes)
 	{
-		return true;
-	}
+		UAvaTagAttributeBase* TagAttribute = Cast<UAvaTagAttributeBase>(Attribute);
+		if (!TagAttribute)
+		{
+			continue;
+		}
 
-	const UAvaSceneSettings* SceneSettings = SceneSettingsWeak.Get();
-	if (!SceneSettings)
-	{
-		return false;
+		// Attempt to clear the given tag handle for the attribute. This will return true if it did remove the entry
+		// Do not remove the attribute itself from the list as it could still have valid tags, or later have valid tags
+		if (TagAttribute->ClearTagHandle(InTagHandle))
+		{
+			++TagsCleared;
+		}
 	}
 
-	bool bFoundTag = false;
+	return TagsCleared > 0;
+}
 
-	SceneSettings->ForEachSceneAttributeOfType<UAvaTagAttributeBase>(
-		[&InTagHandle, &bFoundTag](const UAvaTagAttributeBase& InTagAttribute)
+bool UAvaSceneState::ContainsTagAttribute(const FAvaTagHandle& InTagHandle) const
+{
+	return SceneAttributes.ContainsByPredicate(
+		[&InTagHandle](UAvaAttribute* InAttribute)
 		{
-			bFoundTag = InTagAttribute.ContainsTag(InTagHandle);
-			// continue iteration if tag not found
-			return !bFoundTag;
+			UAvaTagAttributeBase* TagAttribute = Cast<UAvaTagAttributeBase>(InAttribute);
+			return TagAttribute && TagAttribute->ContainsTag(InTagHandle);
 		});
-
-	return bFoundTag;
 }
diff --git a/Engine/Plugins/Experimental/Avalanche/Source/Avalanche/Public/AvaScene.h b/Engine/Plugins/Experimental/Avalanche/Source/Avalanche/Public/AvaScene.h
--- a/Engine/Plugins/Experimental/Avalanche/Source/Avalanche/Public/AvaScene.h
+++ b/Engine/Plugins/Experimental/Avalanche/Source/Avalanche/Public/AvaScene.h
@@ -88,6 +88,7 @@ public:
 
 	//~ Begin UObject
 	virtual void PostLoad() override;
+	virtual void PostInitializeComponents() override;
 	virtual void PostDuplicate(EDuplicateMode::Type InDuplicateMode) override;
 	virtual void PostEditImport() override;
 	virtual void BeginDestroy() override;
diff --git a/Engine/Plugins/Experimental/Avalanche/Source/Avalanche/Public/AvaSceneSettings.h b/Engine/Plugins/Experimental/Avalanche/Source/Avalanche/Public/AvaSceneSettings.h
--- a/Engine/Plugins/Experimental/Avalanche/Source/Avalanche/Public/AvaSceneSettings.h
+++ b/Engine/Plugins/Experimental/Avalanche/Source/Avalanche/Public/AvaSceneSettings.h
@@ -26,23 +26,9 @@ public:
 		return GET_MEMBER_NAME_CHECKED(UAvaSceneSettings, SceneRig);
 	}
 
-	/**
-	 * Iterate each valid Scene Attribute of the given Type
-	 * The callable should return true to continue iteration, and false to stop it
-	 */
-	template<typename InAttributeType>
-	void ForEachSceneAttributeOfType(TFunctionRef<bool(const InAttributeType&)> InCallable) const
+	TConstArrayView<TObjectPtr<UAvaAttribute>> GetSceneAttributes() const
 	{
-		for (const UAvaAttribute* SceneAttribute : SceneAttributes)
-		{
-			if (const InAttributeType* CastedAttribute = Cast<InAttributeType>(SceneAttribute))
-			{
-				if (!InCallable(*CastedAttribute))
-				{
-					break;
-				}
-			}
-		}
+		return SceneAttributes;
 	}
 
 	FSoftObjectPath GetSceneRig() const
@@ -55,7 +41,6 @@ public:
 		SceneRig = InSceneRig;
 	}
 
-
 private:
 	UPROPERTY(EditAnywhere, Instanced, Category="Scene Attributes")
 	TArray<TObjectPtr<UAvaAttribute>> SceneAttributes;
diff --git a/Engine/Plugins/Experimental/Avalanche/Source/Avalanche/Public/AvaSceneState.h b/Engine/Plugins/Experimental/Avalanche/Source/Avalanche/Public/AvaSceneState.h
--- a/Engine/Plugins/Experimental/Avalanche/Source/Avalanche/Public/AvaSceneState.h
+++ b/Engine/Plugins/Experimental/Avalanche/Source/Avalanche/Public/AvaSceneState.h
@@ -2,11 +2,15 @@
 
 #pragma once
 
-#include "AvaTagHandleContainer.h"
+#include "AvaTag.h"
+#include "Containers/Array.h"
+#include "Containers/Set.h"
 #include "UObject/Object.h"
 #include "AvaSceneState.generated.h"
 
+class UAvaAttribute;
 class UAvaSceneSettings;
+struct FAvaTagHandle;
 
 /** Object providing State information of the Scene */
 UCLASS(MinimalAPI)
@@ -15,7 +19,7 @@ class UAvaSceneState : public UObject
 	GENERATED_BODY()
 
 public:
-	void SetSceneSettings(UAvaSceneSettings* InSceneSettings);
+	void Initialize(UAvaSceneSettings* InSceneSettings);
 
 	AVALANCHE_API bool AddTagAttribute(const FAvaTagHandle& InTagHandle);
 
@@ -24,9 +28,7 @@ public:
 	AVALANCHE_API bool ContainsTagAttribute(const FAvaTagHandle& InTagHandle) const;
 
 private:
-	/** Weak pointer to the same scene's Settings. Used to query for default tag attributes of the scene */
-	TWeakObjectPtr<const UAvaSceneSettings> SceneSettingsWeak;
-
-	/** Active Tag Attributes that were added to the scene, separate from Scene Settings */
-	TSet<FAvaTag> ActiveTagAttributes;
+	/** In-play Scene Attributes. Starts as being the Scene Setting Attributes but can added to / removed from while in-play */
+	UPROPERTY(Transient)
+	TArray<TObjectPtr<UAvaAttribute>> SceneAttributes;
 };
diff --git a/Engine/Plugins/Experimental/Avalanche/Source/AvalancheAttribute/Private/Tags/AvaTagAttribute.cpp b/Engine/Plugins/Experimental/Avalanche/Source/AvalancheAttribute/Private/Tags/AvaTagAttribute.cpp
--- a/Engine/Plugins/Experimental/Avalanche/Source/AvalancheAttribute/Private/Tags/AvaTagAttribute.cpp
+++ b/Engine/Plugins/Experimental/Avalanche/Source/AvalancheAttribute/Private/Tags/AvaTagAttribute.cpp
@@ -9,9 +9,35 @@ FText UAvaTagAttribute::GetDisplayName() const
 	return FText::Format(LOCTEXT("DisplayName", "Tag Attribute: {0}"), FText::FromName(Tag.ToName()));
 }
 
+bool UAvaTagAttribute::SetTagHandle(const FAvaTagHandle& InTagHandle)
+{
+	if (Tag.MatchesExact(InTagHandle))
+	{
+		return false;
+	}
+
+	Tag = InTagHandle;
+	return true;
+}
+
+bool UAvaTagAttribute::ClearTagHandle(const FAvaTagHandle& InTagHandle)
+{
+	if (Tag.MatchesExact(InTagHandle))
+	{
+		Tag = FAvaTagHandle();
+		return true;
+	}
+	return false;
+}
+
 bool UAvaTagAttribute::ContainsTag(const FAvaTagHandle& InTagHandle) const
 {
 	return Tag.MatchesTag(InTagHandle);
 }
 
+bool UAvaTagAttribute::HasValidTagHandle() const
+{
+	return Tag.IsValid();
+}
+
 #undef LOCTEXT_NAMESPACE
diff --git a/Engine/Plugins/Experimental/Avalanche/Source/AvalancheAttribute/Private/Tags/AvaTagContainerAttribute.cpp b/Engine/Plugins/Experimental/Avalanche/Source/AvalancheAttribute/Private/Tags/AvaTagContainerAttribute.cpp
--- a/Engine/Plugins/Experimental/Avalanche/Source/AvalancheAttribute/Private/Tags/AvaTagContainerAttribute.cpp
+++ b/Engine/Plugins/Experimental/Avalanche/Source/AvalancheAttribute/Private/Tags/AvaTagContainerAttribute.cpp
@@ -9,11 +9,30 @@ FText UAvaTagContainerAttribute::GetDisplayName() const
 	return FText::Format(LOCTEXT("DisplayName", "Tag Container Attribute: {0}"), FText::FromString(TagContainer.ToString()));
 }
 
+bool UAvaTagContainerAttribute::SetTagHandle(const FAvaTagHandle& InTagHandle)
+{
+	return TagContainer.AddTagHandle(InTagHandle);
+}
+
+bool UAvaTagContainerAttribute::ClearTagHandle(const FAvaTagHandle& InTagHandle)
+{
+	return TagContainer.RemoveTagHandle(InTagHandle);
+}
+
 bool UAvaTagContainerAttribute::ContainsTag(const FAvaTagHandle& InTagHandle) const
 {
 	return TagContainer.ContainsTag(InTagHandle);
 }
 
+bool UAvaTagContainerAttribute::HasValidTagHandle() const
+{
+	return TagContainer.GetTagIds().ContainsByPredicate(
+		[](const FAvaTagId& InTagId)
+		{
+			return InTagId.IsValid();
+		});
+}
+
 void UAvaTagContainerAttribute::SetTagContainer(const FAvaTagHandleContainer& InTagContainer)
 {
 	TagContainer = InTagContainer;
diff --git a/Engine/Plugins/Experimental/Avalanche/Source/AvalancheAttribute/Public/Tags/AvaTagAttribute.h b/Engine/Plugins/Experimental/Avalanche/Source/AvalancheAttribute/Public/Tags/AvaTagAttribute.h
--- a/Engine/Plugins/Experimental/Avalanche/Source/AvalancheAttribute/Public/Tags/AvaTagAttribute.h
+++ b/Engine/Plugins/Experimental/Avalanche/Source/AvalancheAttribute/Public/Tags/AvaTagAttribute.h
@@ -18,7 +18,10 @@ public:
 	//~ End UAvaAttribute
 
 	//~ Begin UAvaTagAttributeBase
+	AVALANCHEATTRIBUTE_API virtual bool SetTagHandle(const FAvaTagHandle& InTagHandle) override;
+	AVALANCHEATTRIBUTE_API virtual bool ClearTagHandle(const FAvaTagHandle& InTagHandle) override;
 	AVALANCHEATTRIBUTE_API virtual bool ContainsTag(const FAvaTagHandle& InTagHandle) const override;
+	AVALANCHEATTRIBUTE_API virtual bool HasValidTagHandle() const;
 	//~ End UAvaTagAttributeBase
 
 	UPROPERTY(EditAnywhere, Category="Attributes")
diff --git a/Engine/Plugins/Experimental/Avalanche/Source/AvalancheAttribute/Public/Tags/AvaTagAttributeBase.h b/Engine/Plugins/Experimental/Avalanche/Source/AvalancheAttribute/Public/Tags/AvaTagAttributeBase.h
--- a/Engine/Plugins/Experimental/Avalanche/Source/AvalancheAttribute/Public/Tags/AvaTagAttributeBase.h
+++ b/Engine/Plugins/Experimental/Avalanche/Source/AvalancheAttribute/Public/Tags/AvaTagAttributeBase.h
@@ -14,8 +14,42 @@ class UAvaTagAttributeBase : public UAvaAttribute
 	GENERATED_BODY()
 
 public:
+	/**
+	 * Sets/Adds the given tag handle to the tag attribute
+	 * @param InTagHandle the tag handle to add
+	 * @return true if the tag handle was added, false if this tag handle already existed in the tag attribute
+	 */
+	virtual bool SetTagHandle(const FAvaTagHandle& InTagHandle)
+	{
+		return false;
+	}
+
+	/**
+	 * Clears/Removes the given tag handle from the tag attribute
+	 * @param InTagHandle the tag handle to remove
+	 * @return true if the tag handle was removed, false if this tag handle was not present in the tag attribute and nothing was removed
+	 */
+	virtual bool ClearTagHandle(const FAvaTagHandle& InTagHandle)
+	{
+		return false;
+	}
+
+	/**
+	 * Checks whether the resolved tag from the given tag handle exists in the tag attribute
+	 * @param InTagHandle the tag handle to resolve into a tag and check if the tag attribute contains it
+	 * @return true if the resolved tag exists in the tag attribute, false otherwise.
+	 */
 	virtual bool ContainsTag(const FAvaTagHandle& InTagHandle) const
 	{
 		return false;
 	}
+
+	/**
+	 * Checks whether the tag attribute has any tag handle that could resolve to a valid tag
+	 * @return true if the tag attribute has a valid tag handle that could resolve to a tag
+	 */
+	virtual bool HasValidTagHandle() const
+	{
+		return false;
+	}
 };
diff --git a/Engine/Plugins/Experimental/Avalanche/Source/AvalancheAttribute/Public/Tags/AvaTagContainerAttribute.h b/Engine/Plugins/Experimental/Avalanche/Source/AvalancheAttribute/Public/Tags/AvaTagContainerAttribute.h
--- a/Engine/Plugins/Experimental/Avalanche/Source/AvalancheAttribute/Public/Tags/AvaTagContainerAttribute.h
+++ b/Engine/Plugins/Experimental/Avalanche/Source/AvalancheAttribute/Public/Tags/AvaTagContainerAttribute.h
@@ -18,7 +18,10 @@ public:
 	//~ End UAvaAttribute
 
 	//~ Begin UAvaTagAttributeBase
+	AVALANCHEATTRIBUTE_API virtual bool SetTagHandle(const FAvaTagHandle& InTagHandle) override;
+	AVALANCHEATTRIBUTE_API virtual bool ClearTagHandle(const FAvaTagHandle& InTagHandle) override;
 	AVALANCHEATTRIBUTE_API virtual bool ContainsTag(const FAvaTagHandle& InTagHandle) const override;
+	AVALANCHEATTRIBUTE_API virtual bool HasValidTagHandle() const;
 	//~ End UAvaTagAttributeBase
 
 	UFUNCTION()
diff --git a/Engine/Plugins/Experimental/Avalanche/Source/AvalancheTag/Public/AvaTagHandleContainer.h b/Engine/Plugins/Experimental/Avalanche/Source/AvalancheTag/Public/AvaTagHandleContainer.h
--- a/Engine/Plugins/Experimental/Avalanche/Source/AvalancheTag/Public/AvaTagHandleContainer.h
+++ b/Engine/Plugins/Experimental/Avalanche/Source/AvalancheTag/Public/AvaTagHandleContainer.h
@@ -43,6 +43,11 @@ struct AVALANCHETAG_API FAvaTagHandleContainer
 	/** Returns an array of resolved tags through this container's tag ids and source tag collection */
 	TArray<FAvaTag> ResolveTags() const;
 
+	TConstArrayView<FAvaTagId> GetTagIds() const
+	{
+		return TagIds;
+	}
+
 	UPROPERTY(EditAnywhere, Category = "Tag")
 	TObjectPtr<const UAvaTagCollection> Source;
 
