Motion Design
- Moved some of the user input dialog files around.
- Added validation.
- Added struct input support.

diff --git a/Engine/Plugins/Experimental/Avalanche/Source/AvalancheEditorCore/Private/Widgets/DataTypes/AvaUserInputDialogDataTypeStruct.cpp b/Engine/Plugins/Experimental/Avalanche/Source/AvalancheEditorCore/Private/Widgets/DataTypes/AvaUserInputDialogDataTypeStruct.cpp
new file mode 100644
--- /dev/null
+++ b/Engine/Plugins/Experimental/Avalanche/Source/AvalancheEditorCore/Private/Widgets/DataTypes/AvaUserInputDialogDataTypeStruct.cpp
@@ -0,0 +1,63 @@
+// Copyright Epic Games, Inc. All Rights Reserved.
+
+#include "Widgets/DataTypes/AvaUserInputDialogDataTypeStruct.h"
+#include "IStructureDetailsView.h"
+#include "Modules/ModuleManager.h"
+#include "PropertyEditorModule.h"
+#include "UObject/StructOnScope.h"
+
+FAvaUserInputDialogDataTypeStruct::FAvaUserInputDialogDataTypeStruct(const FParams& InParams)
+	: StructOnScope(MakeShared<FStructOnScope>(InParams.Struct))
+	, IsValidDelegate(InParams.IsValidDelegate)
+{
+}
+
+const uint8* FAvaUserInputDialogDataTypeStruct::GetStructMemory() const
+{
+	return StructOnScope->GetStructMemory();
+}
+
+const UStruct* FAvaUserInputDialogDataTypeStruct::GetStruct() const
+{
+	return StructOnScope->GetStruct();
+}
+
+TSharedRef<SWidget> FAvaUserInputDialogDataTypeStruct::CreateInputWidget()
+{
+	FPropertyEditorModule& PropertyEditorModule = FModuleManager::Get().LoadModuleChecked<FPropertyEditorModule>("PropertyEditor");
+
+	FDetailsViewArgs DetailsViewArgs;
+	{
+		DetailsViewArgs.bAllowSearch = false;
+		DetailsViewArgs.bHideSelectionTip = true;
+		DetailsViewArgs.bLockable = false;
+		DetailsViewArgs.bSearchInitialKeyFocus = true;
+		DetailsViewArgs.bUpdatesFromSelection = false;
+		DetailsViewArgs.bShowOptions = false;
+		DetailsViewArgs.bShowModifiedPropertiesOption = false;
+		DetailsViewArgs.bShowObjectLabel = false;
+		DetailsViewArgs.bForceHiddenPropertyVisibility = true;
+		DetailsViewArgs.bShowScrollBar = false;
+	}
+
+	FStructureDetailsViewArgs StructureViewArgs;
+	{
+		StructureViewArgs.bShowObjects = true;
+		StructureViewArgs.bShowAssets = true;
+		StructureViewArgs.bShowClasses = true;
+		StructureViewArgs.bShowInterfaces = true;
+	}
+
+	TSharedRef<IStructureDetailsView> DetailsView = PropertyEditorModule.CreateStructureDetailView(DetailsViewArgs, StructureViewArgs, StructOnScope);
+	return DetailsView->GetWidget().ToSharedRef();
+}
+
+bool FAvaUserInputDialogDataTypeStruct::IsValueValid()
+{
+	if (IsValidDelegate.IsBound())
+	{
+		return IsValidDelegate.Execute(StructOnScope);
+	}
+
+	return true;
+}
diff --git a/Engine/Plugins/Experimental/Avalanche/Source/AvalancheEditorCore/Private/Widgets/DataTypes/AvaUserInputDialogDataTypeText.cpp b/Engine/Plugins/Experimental/Avalanche/Source/AvalancheEditorCore/Private/Widgets/DataTypes/AvaUserInputDialogDataTypeText.cpp
new file mode 100644
--- /dev/null
+++ b/Engine/Plugins/Experimental/Avalanche/Source/AvalancheEditorCore/Private/Widgets/DataTypes/AvaUserInputDialogDataTypeText.cpp
@@ -0,0 +1,116 @@
+// Copyright Epic Games, Inc. All Rights Reserved.
+
+#include "Widgets/DataTypes/AvaUserInputDialogDataTypeText.h"
+#include "Widgets/Input/SEditableText.h"
+#include "Widgets/Input/SEditableTextBox.h"
+#include "Widgets/Layout/SBorder.h"
+#include "Widgets/Layout/SBox.h"
+#include "Widgets/SBoxPanel.h"
+#include "Widgets/SOverlay.h"
+#include "Widgets/Text/SMultiLineEditableText.h"
+
+#define LOCTEXT_NAMESPACE "FAvaUserInputDialogTextData"
+
+FAvaUserInputDialogTextData::FAvaUserInputDialogTextData(const FParams& InParams)
+	: Value(InParams.InitialValue)
+	, bAllowMultiline(InParams.bAllowMultiline)
+	, MaxLength(InParams.MaxLength)
+	, OnVerifyDelegate(InParams.OnVerifyDelegate)
+{
+}
+
+const FText& FAvaUserInputDialogTextData::GetValue() const
+{
+	return Value;
+}
+
+TSharedRef<SWidget> FAvaUserInputDialogTextData::CreateInputWidget()
+{
+	if (bAllowMultiline)
+	{
+		TSharedRef<SScrollBar> ScrollBar = SNew(SScrollBar)
+			.Orientation(Orient_Vertical);
+
+		return SNew(SBox)
+			.WidthOverride(500.f)
+			.HeightOverride(200.f)
+			[
+				SNew(SOverlay)
+				+ SOverlay::Slot()
+				[
+					SNew(SBorder)
+					.BorderImage(&FCoreStyle::Get().GetWidgetStyle<FEditableTextBoxStyle>("NormalEditableTextBox").BackgroundImageFocused)
+				]
+				+ SOverlay::Slot()
+				[
+					SNew(SHorizontalBox)
+					+ SHorizontalBox::Slot()
+					.FillWidth(1.f)
+					.Padding(5.f)
+					[
+						SNew(SMultiLineEditableText)
+						.Text(Value)
+						.OnTextChanged(this, &FAvaUserInputDialogTextData::OnTextChanged)
+						.OnTextCommitted(this, &FAvaUserInputDialogTextData::OnTextCommitted)
+						.AllowMultiLine(true)
+						.VScrollBar(ScrollBar)
+					]
+					+ SHorizontalBox::Slot()
+					.AutoWidth()
+					[
+						ScrollBar
+					]
+				]
+			];
+	}
+
+	return SNew(SBox)
+		.WidthOverride(FMath::Max(200.f, MaxLength.IsSet() ? MaxLength.GetValue() * 8.f : 200.f))
+		.HAlign(HAlign_Fill)
+		[
+			SNew(SEditableTextBox)
+			.Text(Value)
+			.OnTextChanged(this, &FAvaUserInputDialogTextData::OnTextChanged)
+			.OnTextCommitted(this, &FAvaUserInputDialogTextData::OnTextCommitted)
+			.OnVerifyTextChanged(this, &FAvaUserInputDialogTextData::OnTextVerify)
+		];
+}
+
+void FAvaUserInputDialogTextData::OnTextChanged(const FText& InValue)
+{
+	Value = InValue;
+}
+
+void FAvaUserInputDialogTextData::OnTextCommitted(const FText& InValue, ETextCommit::Type InCommitType)
+{
+	Value = InValue;
+
+	if (!bAllowMultiline && InCommitType == ETextCommit::OnEnter)
+	{
+		OnUserCommit();
+	}
+}
+
+bool FAvaUserInputDialogTextData::OnTextVerify(const FText& InValue, FText& OutErrorText)
+{
+	if (MaxLength.IsSet() && InValue.ToString().Len() > MaxLength.GetValue())
+	{
+		OutErrorText = LOCTEXT("InputTooLong", "Input too long");
+		return false;
+	}
+
+	if (OnVerifyDelegate.IsBound())
+	{
+		return OnVerifyDelegate.Execute(InValue, OutErrorText);
+	}
+
+	return true;
+}
+
+bool FAvaUserInputDialogTextData::IsValueValid()
+{
+	FText Unused;
+	return OnTextVerify(Value, Unused);
+}
+
+#undef LOCTEXT_NAMESPACE
diff --git a/Engine/Plugins/Experimental/Avalanche/Source/AvalancheEditorCore/Private/Widgets/InputDataTypes/AvaUserInputDataText.cpp b/Engine/Plugins/Experimental/Avalanche/Source/AvalancheEditorCore/Private/Widgets/InputDataTypes/AvaUserInputDataText.cpp
deleted file mode 100644
diff --git a/Engine/Plugins/Experimental/Avalanche/Source/AvalancheEditorCore/Private/Widgets/InputDataTypes/AvaUserInputDataTypeBase.cpp b/Engine/Plugins/Experimental/Avalanche/Source/AvalancheEditorCore/Private/Widgets/InputDataTypes/AvaUserInputDataTypeBase.cpp
deleted file mode 100644
diff --git a/Engine/Plugins/Experimental/Avalanche/Source/AvalancheEditorCore/Private/Widgets/SAvaUserInputDialog.cpp b/Engine/Plugins/Experimental/Avalanche/Source/AvalancheEditorCore/Private/Widgets/SAvaUserInputDialog.cpp
--- a/Engine/Plugins/Experimental/Avalanche/Source/AvalancheEditorCore/Private/Widgets/SAvaUserInputDialog.cpp
+++ b/Engine/Plugins/Experimental/Avalanche/Source/AvalancheEditorCore/Private/Widgets/SAvaUserInputDialog.cpp
@@ -9,22 +9,30 @@
 
 #define LOCTEXT_NAMESPACE "SAvaUserInputDialog"
 
-bool SAvaUserInputDialog::CreateModalDialog(const TSharedPtr<SWidget>& InParent, const FText& InTitle,
-	const FText& InPrompt, const TSharedRef<FAvaUserInputDataTypeBase>& InInputType)
+bool SAvaUserInputDialog::CreateModalDialog(const TSharedRef<FAvaUserInputDialogDataTypeBase>& InInputType, const TSharedPtr<SWidget>& InParent,
+	const TOptional<FText>& InPrompt, const TOptional<FText>& InTitle)
 {
-	TSharedPtr<SWindow> ParentWindow = InParent.IsValid() ? FSlateApplication::Get().FindWidgetWindow(InParent.ToSharedRef()) : nullptr;
-
-	static const FText DefaultTitle = LOCTEXT("DefaultTitle", "User Input Required");
 	static const FText DefaultPrompt = LOCTEXT("DefaultPrompt", "Value requested:");
+	static const FText DefaultTitle = LOCTEXT("DefaultTitle", "User Input Required");
+
+	const FText Prompt = InPrompt.IsSet() && !InPrompt.GetValue().IsEmpty()
+		? InPrompt.GetValue()
+		: DefaultPrompt;
+
+	const FText Title = InTitle.IsSet() && !InTitle.GetValue().IsEmpty()
+		? InTitle.GetValue()
+		: DefaultTitle;
+
+	TSharedPtr<SWindow> ParentWindow = InParent.IsValid() ? FSlateApplication::Get().FindWidgetWindow(InParent.ToSharedRef()) : nullptr;
 
 	TSharedRef<SAvaUserInputDialog> InputDialog = SNew(SAvaUserInputDialog, InInputType)
-		.Prompt(InPrompt.IsEmpty() ? DefaultPrompt : InPrompt);
+		.Prompt(Prompt);
 
 	TSharedRef<SWindow> Window = SNew(SWindow)
 		.SizingRule(ESizingRule::Autosized)
 		.SupportsMinimize(false)
 		.SupportsMaximize(false)
-		.Title(InTitle.IsEmpty() ? DefaultTitle : InTitle)
+		.Title(Title)
 		[
 			InputDialog
 		];
@@ -38,7 +46,7 @@ void SAvaUserInputDialog::PrivateRegisterAttributes(struct FSlateAttributeDescri
 {
 }
 
-void SAvaUserInputDialog::Construct(const FArguments& InArgs, const TSharedRef<FAvaUserInputDataTypeBase>& InInputType)
+void SAvaUserInputDialog::Construct(const FArguments& InArgs, const TSharedRef<FAvaUserInputDialogDataTypeBase>& InInputType)
 {
 	InputType = InInputType;
 	InputType->OnCommit.BindSP(this, &SAvaUserInputDialog::OnUserCommit);
@@ -74,7 +82,8 @@ void SAvaUserInputDialog::Construct(const FArguments& InArgs, const TSharedRef<F
 				SNew(SButton)
 				.Text(LOCTEXT("Accept", "Accept"))
 				.ButtonStyle(FAppStyle::Get(), "PrimaryButton")
-				.OnClicked(this, &SAvaUserInputDialog::OnOkayClicked)
+				.OnClicked(this, &SAvaUserInputDialog::OnAcceptClicked)
+				.IsEnabled(this, &SAvaUserInputDialog::GetAcceptedEnabled)
 			]
 			+ SHorizontalBox::Slot()
 			.AutoWidth()
@@ -89,7 +98,7 @@ void SAvaUserInputDialog::Construct(const FArguments& InArgs, const TSharedRef<F
 	];
 }
 
-TSharedPtr<FAvaUserInputDataTypeBase> SAvaUserInputDialog::GetInputType() const
+TSharedPtr<FAvaUserInputDialogDataTypeBase> SAvaUserInputDialog::GetInputType() const
 {
 	return InputType;
 }
@@ -99,16 +108,28 @@ bool SAvaUserInputDialog::WasAccepted() const
 	return bAccepted;
 }
 
-FReply SAvaUserInputDialog::OnOkayClicked()
+FReply SAvaUserInputDialog::OnAcceptClicked()
 {
-	Close(true);
+	if (!InputType.IsValid())
+	{
+		Close(/* Accepted */ false);
+	}
+	else if (InputType->IsValueValid())
+	{
+		Close(/* Accepted */ true);
+	}
 
 	return FReply::Handled();
 }
 
+bool SAvaUserInputDialog::GetAcceptedEnabled() const
+{
+	return !InputType.IsValid() || InputType->IsValueValid();
+}
+
 FReply SAvaUserInputDialog::OnCancelClicked()
 {
-	Close(false);
+	Close(/* Accepted */ false);
 
 	return FReply::Handled();
 }
diff --git a/Engine/Plugins/Experimental/Avalanche/Source/AvalancheEditorCore/Public/Widgets/DataTypes/AvaUserInputDialogDataTypeBase.h b/Engine/Plugins/Experimental/Avalanche/Source/AvalancheEditorCore/Public/Widgets/DataTypes/AvaUserInputDialogDataTypeBase.h
new file mode 100644
--- /dev/null
+++ b/Engine/Plugins/Experimental/Avalanche/Source/AvalancheEditorCore/Public/Widgets/DataTypes/AvaUserInputDialogDataTypeBase.h
@@ -0,0 +1,34 @@
+// Copyright Epic Games, Inc. All Rights Reserved.
+
+#pragma once
+
+#include "Delegates/Delegate.h"
+#include "Delegates/DelegateCombinations.h"
+#include "Templates/SharedPointer.h"
+
+class SWidget;
+
+struct FAvaUserInputDialogDataTypeBase : public TSharedFromThis<FAvaUserInputDialogDataTypeBase>
+{
+	friend class SAvaUserInputDialog;
+
+	virtual ~FAvaUserInputDialogDataTypeBase() = default;
+
+	virtual TSharedRef<SWidget> CreateInputWidget() = 0;
+
+	virtual bool IsValueValid()
+	{
+		return true;
+	}
+
+protected:
+	FSimpleDelegate OnCommit;
+
+	void OnUserCommit()
+	{
+		if (IsValueValid())
+		{
+			OnCommit.ExecuteIfBound();
+		}
+	}
+};
diff --git a/Engine/Plugins/Experimental/Avalanche/Source/AvalancheEditorCore/Public/Widgets/DataTypes/AvaUserInputDialogDataTypeNumeric.h b/Engine/Plugins/Experimental/Avalanche/Source/AvalancheEditorCore/Public/Widgets/DataTypes/AvaUserInputDialogDataTypeNumeric.h
new file mode 100644
--- /dev/null
+++ b/Engine/Plugins/Experimental/Avalanche/Source/AvalancheEditorCore/Public/Widgets/DataTypes/AvaUserInputDialogDataTypeNumeric.h
@@ -0,0 +1,67 @@
+// Copyright Epic Games, Inc. All Rights Reserved.
+
+#pragma once
+
+#include "Widgets/DataTypes/AvaUserInputDialogDataTypeBase.h"
+#include "Misc/Optional.h"
+#include "Widgets/Input/SSpinBox.h"
+#include "Widgets/Layout/SBox.h"
+
+template<typename InNumericType
+	UE_REQUIRES(TAnd<TIsArithmetic<InNumericType>, TNot<std::is_same<bool, InNumericType>>>::Value)>
+struct FAvaUserInputDialogNumericData : public FAvaUserInputDialogDataTypeBase
+{
+	using FNumericType = InNumericType;
+	using ThisStruct = FAvaUserInputDialogNumericData<FNumericType>;
+
+	struct FParams
+	{
+		FNumericType InitialValue = 0;
+		TOptional<FNumericType> InMinValue = TOptional<FNumericType>();
+		TOptional<FNumericType> InMaxValue = TOptional<FNumericType>();
+	};
+
+	FAvaUserInputDialogNumericData(const FParams& InParams)
+		: Value(InParams.InitialValue)
+		, MinValue(InParams.MinValue)
+		, MaxValue(InParams.MaxValue)
+	{
+	}
+
+	FNumericType GetValue() const { return Value; }
+
+	//~ Begin FAvaUserInputDialogDataTypeBase
+	virtual TSharedRef<SWidget> CreateInputWidget()
+	{
+		return SNew(SBox)
+			.WidthOverride(200.f)
+			.HAlign(HAlign_Fill)
+			[
+				SNew(SSpinBox<FNumericType>)
+				.Value(Value)
+				.MinValue(MinValue)
+				.MaxSliderValue(MinValue)
+				.MaxValue(MaxValue)
+				.MaxSliderValue(MaxValue)
+				.EnableSlider(true)
+				.OnValueChanged(this, &ThisStruct::OnValueChanged)
+				.OnValueCommitted(this, &ThisStruct::OnValueCommitted)
+			];
+	}
+	//~ End FAvaUserInputDialogDataTypeBase
+
+protected:
+	FNumericType Value;
+	TOptional<FNumericType> MinValue;
+	TOptional<FNumericType> MaxValue;
+
+	void OnValueChanged(FNumericType InValue)
+	{
+		Value = InValue;
+	}
+
+	void OnValueCommitted(FNumericType InValue, ETextCommit::Type InCommitType)
+	{
+		Value = InValue;
+	}
+};
\ No newline at end of file
diff --git a/Engine/Plugins/Experimental/Avalanche/Source/AvalancheEditorCore/Public/Widgets/DataTypes/AvaUserInputDialogDataTypeStruct.h b/Engine/Plugins/Experimental/Avalanche/Source/AvalancheEditorCore/Public/Widgets/DataTypes/AvaUserInputDialogDataTypeStruct.h
new file mode 100644
--- /dev/null
+++ b/Engine/Plugins/Experimental/Avalanche/Source/AvalancheEditorCore/Public/Widgets/DataTypes/AvaUserInputDialogDataTypeStruct.h
@@ -0,0 +1,72 @@
+// Copyright Epic Games, Inc. All Rights Reserved.
+
+#pragma once
+
+#include "AvaUserInputDialogDataTypeBase.h"
+#include "Templates/SharedPointer.h"
+#include "UObject/Class.h"
+
+class FStructOnScope;
+
+struct FAvaUserInputDialogDataTypeStruct : FAvaUserInputDialogDataTypeBase
+{
+	DECLARE_DELEGATE_RetVal_OneParam(bool, FIsValidDelegate, TSharedRef<FStructOnScope>)
+
+	struct FParams
+	{
+		UScriptStruct* Struct = nullptr;
+		FIsValidDelegate IsValidDelegate;
+	};
+
+	AVALANCHEEDITORCORE_API FAvaUserInputDialogDataTypeStruct(const FParams& InParams);
+
+	virtual ~FAvaUserInputDialogDataTypeStruct() override = default;
+
+	AVALANCHEEDITORCORE_API const uint8* GetStructMemory() const;
+
+	AVALANCHEEDITORCORE_API const UStruct* GetStruct() const;
+
+	//~ Begin FAvaUserInputDataTypeBase
+	AVALANCHEEDITORCORE_API virtual TSharedRef<SWidget> CreateInputWidget() override;
+	AVALANCHEEDITORCORE_API virtual bool IsValueValid() override;
+	//~ End FAvaUserInputDataTypeBase
+
+protected:
+	TSharedRef<FStructOnScope> StructOnScope;
+	FIsValidDelegate IsValidDelegate;
+};
+
+template<typename InStructType>
+struct TAvaUserInputDialogDataTypeStruct : FAvaUserInputDialogDataTypeStruct
+{
+	struct FParams
+	{
+		FIsValidDelegate IsValidDelegate;
+	};
+
+	explicit TAvaUserInputDialogDataTypeStruct(const FParams& InParams)
+		: FAvaUserInputDialogDataTypeStruct({TBaseStructure<InStructType>::Get(), InParams.IsValidDelegate})
+	{
+	}
+
+	virtual ~TAvaUserInputDialogDataTypeStruct() override = default;
+
+	const InStructType* GetPtr() const
+	{
+		const UStruct* ScriptStruct = GetStruct();
+		if (ScriptStruct && (ScriptStruct == TBaseStructure<InStructType>::Get() || ScriptStruct->IsChildOf(TBaseStructure<InStructType>::Get())))
+		{
+			return reinterpret_cast<const InStructType*>(GetStructMemory());
+		}
+		return nullptr;
+	}
+
+	const InStructType& Get() const
+	{
+		const UStruct* ScriptStruct = GetStruct();
+		const uint8* StructMemory = GetStructMemory();
+		check(StructMemory);
+		check(ScriptStruct && (ScriptStruct == TBaseStructure<InStructType>::Get() || ScriptStruct->IsChildOf(TBaseStructure<InStructType>::Get())));
+		return *reinterpret_cast<const InStructType*>(StructMemory);
+	}
+};
diff --git a/Engine/Plugins/Experimental/Avalanche/Source/AvalancheEditorCore/Public/Widgets/DataTypes/AvaUserInputDialogDataTypeText.h b/Engine/Plugins/Experimental/Avalanche/Source/AvalancheEditorCore/Public/Widgets/DataTypes/AvaUserInputDialogDataTypeText.h
new file mode 100644
--- /dev/null
+++ b/Engine/Plugins/Experimental/Avalanche/Source/AvalancheEditorCore/Public/Widgets/DataTypes/AvaUserInputDialogDataTypeText.h
@@ -0,0 +1,44 @@
+// Copyright Epic Games, Inc. All Rights Reserved.
+
+#pragma once
+
+#include "Widgets/DataTypes/AvaUserInputDialogDataTypeBase.h"
+#include "Framework/SlateDelegates.h"
+#include "Internationalization/Text.h"
+#include "Misc/Optional.h"
+
+namespace ETextCommit { enum Type : int; }
+
+struct FAvaUserInputDialogTextData : public FAvaUserInputDialogDataTypeBase
+{
+	struct FParams
+	{
+		FText InitialValue = FText::GetEmpty();
+		bool bAllowMultiline = false;
+		TOptional<int32> MaxLength = TOptional<int32>();
+		FOnVerifyTextChanged OnVerifyDelegate;
+	};
+
+	AVALANCHEEDITORCORE_API FAvaUserInputDialogTextData(const FParams& InParams);
+
+	virtual ~FAvaUserInputDialogTextData() override = default;
+
+	AVALANCHEEDITORCORE_API const FText& GetValue() const;
+
+	//~ Begin FAvaUserInputDialogDataTypeBase
+	AVALANCHEEDITORCORE_API virtual TSharedRef<SWidget> CreateInputWidget() override;
+	AVALANCHEEDITORCORE_API virtual bool IsValueValid() override;
+	//~ End FAvaUserInputDialogDataTypeBase
+
+protected:
+	FText Value;
+	bool bAllowMultiline;
+	TOptional<int32> MaxLength;
+	FOnVerifyTextChanged OnVerifyDelegate;
+
+	void OnTextChanged(const FText& InValue);
+
+	void OnTextCommitted(const FText& InValue, ETextCommit::Type InCommitType);
+
+	bool OnTextVerify(const FText& InValue, FText& OutErrorText);
+};
\ No newline at end of file
diff --git a/Engine/Plugins/Experimental/Avalanche/Source/AvalancheEditorCore/Public/Widgets/InputDataTypes/AvaUserInputDataNumber.h b/Engine/Plugins/Experimental/Avalanche/Source/AvalancheEditorCore/Public/Widgets/InputDataTypes/AvaUserInputDataNumber.h
deleted file mode 100644
diff --git a/Engine/Plugins/Experimental/Avalanche/Source/AvalancheEditorCore/Public/Widgets/InputDataTypes/AvaUserInputDataText.h b/Engine/Plugins/Experimental/Avalanche/Source/AvalancheEditorCore/Public/Widgets/InputDataTypes/AvaUserInputDataText.h
deleted file mode 100644
diff --git a/Engine/Plugins/Experimental/Avalanche/Source/AvalancheEditorCore/Public/Widgets/InputDataTypes/AvaUserInputDataTypeBase.h b/Engine/Plugins/Experimental/Avalanche/Source/AvalancheEditorCore/Public/Widgets/InputDataTypes/AvaUserInputDataTypeBase.h
deleted file mode 100644
diff --git a/Engine/Plugins/Experimental/Avalanche/Source/AvalancheEditorCore/Public/Widgets/SAvaUserInputDialog.h b/Engine/Plugins/Experimental/Avalanche/Source/AvalancheEditorCore/Public/Widgets/SAvaUserInputDialog.h
--- a/Engine/Plugins/Experimental/Avalanche/Source/AvalancheEditorCore/Public/Widgets/SAvaUserInputDialog.h
+++ b/Engine/Plugins/Experimental/Avalanche/Source/AvalancheEditorCore/Public/Widgets/SAvaUserInputDialog.h
@@ -3,31 +3,33 @@
 #pragma once
 
 #include "Widgets/SCompoundWidget.h"
-#include "Widgets/InputDataTypes/AvaUserInputDataTypeBase.h"
+#include "Widgets/DataTypes/AvaUserInputDialogDataTypeBase.h"
 
 class SAvaUserInputDialog : public SCompoundWidget
 {
 	SLATE_DECLARE_WIDGET_API(SAvaUserInputDialog, SCompoundWidget, AVALANCHEEDITORCORE_API)
 
 public:
-	static AVALANCHEEDITORCORE_API bool CreateModalDialog(const TSharedPtr<SWidget>& InParent, const FText& InTitle, const FText& InPrompt,
-		const TSharedRef<FAvaUserInputDataTypeBase>& InInputType);
+	static AVALANCHEEDITORCORE_API bool CreateModalDialog(const TSharedRef<FAvaUserInputDialogDataTypeBase>& InInputType, const TSharedPtr<SWidget>& InParent = nullptr,
+		const TOptional<FText>& InPrompt = TOptional<FText>(), const TOptional<FText>& InTitle = TOptional<FText>());
 
 	SLATE_BEGIN_ARGS(SAvaUserInputDialog) {}
 		SLATE_ARGUMENT(FText, Prompt)
 	SLATE_END_ARGS()
 
-	AVALANCHEEDITORCORE_API void Construct(const FArguments& InArgs, const TSharedRef<FAvaUserInputDataTypeBase>& InInputType);
+	AVALANCHEEDITORCORE_API void Construct(const FArguments& InArgs, const TSharedRef<FAvaUserInputDialogDataTypeBase>& InInputType);
 
-	AVALANCHEEDITORCORE_API TSharedPtr<FAvaUserInputDataTypeBase> GetInputType() const;
+	AVALANCHEEDITORCORE_API TSharedPtr<FAvaUserInputDialogDataTypeBase> GetInputType() const;
 
 	AVALANCHEEDITORCORE_API bool WasAccepted() const;
 
 private:
-	TSharedPtr<FAvaUserInputDataTypeBase> InputType;
+	TSharedPtr<FAvaUserInputDialogDataTypeBase> InputType;
 	bool bAccepted = false;
 
-	FReply OnOkayClicked();
+	FReply OnAcceptClicked();
+
+	bool GetAcceptedEnabled() const;
 
 	FReply OnCancelClicked();
 
diff --git a/Engine/Plugins/Experimental/Avalanche/Source/AvalancheLevelViewport/Private/SAvaLevelViewport_Menus.cpp b/Engine/Plugins/Experimental/Avalanche/Source/AvalancheLevelViewport/Private/SAvaLevelViewport_Menus.cpp
--- a/Engine/Plugins/Experimental/Avalanche/Source/AvalancheLevelViewport/Private/SAvaLevelViewport_Menus.cpp
+++ b/Engine/Plugins/Experimental/Avalanche/Source/AvalancheLevelViewport/Private/SAvaLevelViewport_Menus.cpp
@@ -22,10 +22,10 @@
 #include "SViewportToolBar.h"
 #include "Viewport/Interaction/IAvaViewportDataProvider.h"
 #include "ViewportClient/AvaLevelViewportClient.h"
+#include "Widgets/DataTypes/AvaUserInputDialogDataTypeText.h"
 #include "Widgets/Images/SImage.h"
 #include "Widgets/Input/SButton.h"
 #include "Widgets/Input/SSpinBox.h"
-#include "Widgets/InputDataTypes/AvaUserInputDataText.h"
 #include "Widgets/SAvaUserInputDialog.h"
 
 #define LOCTEXT_NAMESPACE "SAvaLevelViewport"
@@ -1029,17 +1029,16 @@ void SAvaLevelViewport::ExecuteSaveAsGuidePreset(const FToolMenuContext& InConte
 		return;
 	}
 
-	TSharedRef<FAvaUserInputTextData> TextInput = MakeShared<FAvaUserInputTextData>(
-		LOCTEXT("NewPreset", "NewPreset"),
-		/* Multiline */ false,
-		/* Max Length */ 30
-	);
+	FAvaUserInputDialogTextData::FParams Params;
+	Params.InitialValue = LOCTEXT("NewPreset", "NewPreset");
+	Params.MaxLength = 30;
+
+	TSharedRef<FAvaUserInputDialogTextData> TextInput = MakeShared<FAvaUserInputDialogTextData>(Params);
 
 	const bool bAccepted = SAvaUserInputDialog::CreateModalDialog(
+		TextInput,
 		SharedThis(this),
-		FText::GetEmpty(),
-		LOCTEXT("NewPresetName", "New Preset Name"),
-		TextInput
+		LOCTEXT("NewPresetName", "New Preset Name")
 	);
 
 	if (!bAccepted)
