[Motion Design] Rundown Editor: moving the property setting to the rundown editor settings.

diff --git a/Engine/Plugins/Experimental/Avalanche/Source/AvalancheMediaEditor/Private/AvaMediaEditorSettings.h b/Engine/Plugins/Experimental/Avalanche/Source/AvalancheMediaEditor/Private/AvaMediaEditorSettings.h
--- a/Engine/Plugins/Experimental/Avalanche/Source/AvalancheMediaEditor/Private/AvaMediaEditorSettings.h
+++ b/Engine/Plugins/Experimental/Avalanche/Source/AvalancheMediaEditor/Private/AvaMediaEditorSettings.h
@@ -31,9 +31,6 @@ public:
 	UPROPERTY(Config, EditAnywhere, Category = "Broadcast", meta = (DisplayName = "Show All Media Output Classes"))
 	bool bBroadcastShowAllMediaOutputClasses = false;
 
-	UPROPERTY(Config, EditAnywhere, Category = "Playback", meta = (DisplayName = "Show RC Property List"))
-	bool bPlaybackShowPropertyList = false;
-	
 	UPROPERTY(Config, EditAnywhere, Category = "Playback", meta = (DisplayName = "Default Node Color"))
 	FLinearColor PlaybackDefaultNodeColor;
 
diff --git a/Engine/Plugins/Experimental/Avalanche/Source/AvalancheMediaEditor/Private/Rundown/AvaRundownEditorSettings.h b/Engine/Plugins/Experimental/Avalanche/Source/AvalancheMediaEditor/Private/Rundown/AvaRundownEditorSettings.h
--- a/Engine/Plugins/Experimental/Avalanche/Source/AvalancheMediaEditor/Private/Rundown/AvaRundownEditorSettings.h
+++ b/Engine/Plugins/Experimental/Avalanche/Source/AvalancheMediaEditor/Private/Rundown/AvaRundownEditorSettings.h
@@ -30,6 +30,9 @@ public:
 	UPROPERTY(Config, EditAnywhere, Category="Preview")
 	bool bPreviewCheckerBoard = false;
 
+	UPROPERTY(Config, EditAnywhere, Category = "Page Details", meta = (DisplayName = "Show RC Properties"))
+	bool bPageDetailsShowProperties = false;
+
 	/** Current macro collection used by the rundown editor. */
 	UPROPERTY(Config, EditAnywhere, Category="Macros")
 	TSoftObjectPtr<UAvaRundownMacroCollection> MacroCollection;
diff --git a/Engine/Plugins/Experimental/Avalanche/Source/AvalancheMediaEditor/Private/Rundown/DetailsView/SAvaRundownPageDetails.cpp b/Engine/Plugins/Experimental/Avalanche/Source/AvalancheMediaEditor/Private/Rundown/DetailsView/SAvaRundownPageDetails.cpp
--- a/Engine/Plugins/Experimental/Avalanche/Source/AvalancheMediaEditor/Private/Rundown/DetailsView/SAvaRundownPageDetails.cpp
+++ b/Engine/Plugins/Experimental/Avalanche/Source/AvalancheMediaEditor/Private/Rundown/DetailsView/SAvaRundownPageDetails.cpp
@@ -1,8 +1,8 @@
 // Copyright Epic Games, Inc. All Rights Reserved.
 
 #include "SAvaRundownPageDetails.h"
+
 #include "Async/Async.h"
-#include "AvaMediaEditorSettings.h"
 #include "Framework/Application/SlateApplication.h"
 #include "IAvaMediaModule.h"
 #include "Input/Reply.h"
@@ -10,6 +10,7 @@
 #include "RemoteControl/Controllers/SAvaRundownRCControllerPanel.h"
 #include "Rundown/AvaRundown.h"
 #include "Rundown/AvaRundownEditor.h"
+#include "Rundown/AvaRundownEditorSettings.h"
 #include "Rundown/AvaRundownManagedInstanceCache.h"
 #include "Rundown/AvaRundownPage.h"
 #include "Rundown/DetailsView/RemoteControl/Properties/SAvaRundownPageRemoteControlProps.h"
@@ -25,6 +26,15 @@
 
 #define LOCTEXT_NAMESPACE "SAvaRundownPageDetails"
 
+namespace UE::AvaMedia::RundownEditor::Private
+{
+	bool ShouldPageDetailsShowProperties()
+	{
+		const UAvaRundownEditorSettings* RundownEditorSettings = UAvaRundownEditorSettings::Get();
+		return RundownEditorSettings && RundownEditorSettings->bPageDetailsShowProperties;
+	}
+}
+
 void SAvaRundownPageDetails::Construct(const FArguments& InArgs, const TSharedPtr<FAvaRundownEditor>& InRundownEditor)
 {
 	RundownEditorWeak = InRundownEditor;
@@ -175,7 +185,7 @@ void SAvaRundownPageDetails::Construct(const FArguments& InArgs, const TSharedPt
 			.AutoHeight()
 			[
 				SAssignNew(RemoteControlProps, SAvaRundownPageRemoteControlProps, InRundownEditor)
-				.Visibility(UAvaMediaEditorSettings::Get().bPlaybackShowPropertyList ? EVisibility::Visible : EVisibility::Collapsed)
+				.Visibility(UE::AvaMedia::RundownEditor::Private::ShouldPageDetailsShowProperties() ? EVisibility::Visible : EVisibility::Collapsed)
 			]
 		]
 	];
@@ -244,17 +254,19 @@ void SAvaRundownPageDetails::OnManagedInstanceCacheEntryInvalidated(const FSoftO
 
 FReply SAvaRundownPageDetails::ToggleExposedPropertiesVisibility()
 {
-	UAvaMediaEditorSettings& MediaEditorSettings = UAvaMediaEditorSettings::GetMutable();
-	MediaEditorSettings.bPlaybackShowPropertyList = !MediaEditorSettings.bPlaybackShowPropertyList;
-	MediaEditorSettings.SaveConfig();
-
-	if (MediaEditorSettings.bPlaybackShowPropertyList)
+	if (UAvaRundownEditorSettings* RundownEditorSettings = UAvaRundownEditorSettings::GetMutable())
 	{
-		RemoteControlProps->SetVisibility(EVisibility::SelfHitTestInvisible);
-	}
-	else
-	{
-		RemoteControlProps->SetVisibility(EVisibility::Collapsed);
+		RundownEditorSettings->bPageDetailsShowProperties = !RundownEditorSettings->bPageDetailsShowProperties;
+		RundownEditorSettings->SaveConfig();
+
+		if (RundownEditorSettings->bPageDetailsShowProperties)
+		{
+			RemoteControlProps->SetVisibility(EVisibility::SelfHitTestInvisible);
+		}
+		else
+		{
+			RemoteControlProps->SetVisibility(EVisibility::Collapsed);
+		}
 	}
 
 	return FReply::Handled();
@@ -262,7 +274,7 @@ FReply SAvaRundownPageDetails::ToggleExposedPropertiesVisibility()
 
 const FSlateBrush* SAvaRundownPageDetails::GetExposedPropertiesVisibilityBrush() const
 {
-	if (UAvaMediaEditorSettings::Get().bPlaybackShowPropertyList)
+	if (UE::AvaMedia::RundownEditor::Private::ShouldPageDetailsShowProperties())
 	{
 		return FAppStyle::GetBrush(TEXT("Level.VisibleHighlightIcon16x"));
 	}
