[Motion Design] Fix Broadcast crash in cooked packaged game.

The xml serialization is not supported in packaged game, but the broadcast object still needs to be properly initialized to work with the playback server.

diff --git a/Engine/Plugins/Experimental/Avalanche/Source/AvalancheMedia/Private/Broadcast/AvaBroadcast.cpp b/Engine/Plugins/Experimental/Avalanche/Source/AvalancheMedia/Private/Broadcast/AvaBroadcast.cpp
--- a/Engine/Plugins/Experimental/Avalanche/Source/AvalancheMedia/Private/Broadcast/AvaBroadcast.cpp
+++ b/Engine/Plugins/Experimental/Avalanche/Source/AvalancheMedia/Private/Broadcast/AvaBroadcast.cpp
@@ -68,6 +68,8 @@ UAvaBroadcast& UAvaBroadcast::Get()
 		Broadcast->AddToRoot();
 #if WITH_EDITOR
 		Broadcast->LoadBroadcast();
+#else
+		Broadcast->InitBroadcast();	// todo: find another way to load config in packaged build.
 #endif
 	}
 	
@@ -378,6 +380,13 @@ FString UAvaBroadcast::GetBroadcastSaveFilepath() const
 {
 	return UE::AvaBroadcast::Private::GetXmlSaveFilepath();
 }
+#else
+void UAvaBroadcast::InitBroadcast()
+{
+	CreateProfile(NAME_None, true);
+	EnsureValidCurrentProfile();
+	UpdateProfileNames();
+}
 #endif
 
 void UAvaBroadcast::QueueNotifyChange(EAvaBroadcastChange InChange)
@@ -397,7 +406,7 @@ void UAvaBroadcast::QueueNotifyChange(EAvaBroadcastChange InChange)
 					ThisWeak->OnBroadcastChanged.Broadcast(ThisWeak->QueuedBroadcastChanges);
 					ThisWeak->QueuedBroadcastChanges = EAvaBroadcastChange::None;
 				}
-			});	
+			});
 		}
 	}
 }
diff --git a/Engine/Plugins/Experimental/Avalanche/Source/AvalancheMedia/Public/Broadcast/AvaBroadcast.h b/Engine/Plugins/Experimental/Avalanche/Source/AvalancheMedia/Public/Broadcast/AvaBroadcast.h
--- a/Engine/Plugins/Experimental/Avalanche/Source/AvalancheMedia/Public/Broadcast/AvaBroadcast.h
+++ b/Engine/Plugins/Experimental/Avalanche/Source/AvalancheMedia/Public/Broadcast/AvaBroadcast.h
@@ -85,6 +85,8 @@ public:
 	void LoadBroadcast();
 	void SaveBroadcast();
 	FString GetBroadcastSaveFilepath() const;
+#else
+	void InitBroadcast();
 #endif
 
 	//Rather than Immediately calling the Broadcast Event change, it will queue all the multiple calls and dispatch it once to avoid redundancy
