[Avalanche]
* Fixed compile warning of unreachable code by changing return to #if 0/#endif

diff --git a/Engine/Plugins/Experimental/Avalanche/Source/AvalancheViewport/Private/Interaction/AvaCameraZoomController.cpp b/Engine/Plugins/Experimental/Avalanche/Source/AvalancheViewport/Private/Interaction/AvaCameraZoomController.cpp
--- a/Engine/Plugins/Experimental/Avalanche/Source/AvalancheViewport/Private/Interaction/AvaCameraZoomController.cpp
+++ b/Engine/Plugins/Experimental/Avalanche/Source/AvalancheViewport/Private/Interaction/AvaCameraZoomController.cpp
@@ -284,7 +284,7 @@ void FAvaCameraZoomController::CenterOnPoint(const FVector2f& InPoint)
 void FAvaCameraZoomController::CenterOnBox(const FBox& InBoundingBox, const FTransform& InBoxTransform)
 {
 	// @TODO Update
-	return;
+#if 0
 
 	TSharedPtr<IAvaViewportClient> AvaViewportClient = AvaViewportClientWeak.Pin();
 
@@ -355,6 +355,7 @@ void FAvaCameraZoomController::CenterOnBox(const FBox& InBoundingBox, const FTra
 	SetZoomLevel(FMath::Max<uint8>(RequiredZoomLevel, 1) - 1);
 	ZoomIn();
 	CenterOnPoint(ScreenBoundsCenter);
+#endif
 }
 
 void FAvaCameraZoomController::StartPanning()
