[Motion Design] Fix playable camera setup in game build.

diff --git a/Engine/Plugins/Experimental/Avalanche/Source/AvalancheMedia/Private/Playable/Playables/AvaPlayableLevelStreaming.cpp b/Engine/Plugins/Experimental/Avalanche/Source/AvalancheMedia/Private/Playable/Playables/AvaPlayableLevelStreaming.cpp
--- a/Engine/Plugins/Experimental/Avalanche/Source/AvalancheMedia/Private/Playable/Playables/AvaPlayableLevelStreaming.cpp
+++ b/Engine/Plugins/Experimental/Avalanche/Source/AvalancheMedia/Private/Playable/Playables/AvaPlayableLevelStreaming.cpp
@@ -321,9 +321,7 @@ bool UAvaPlayableLevelStreaming::ApplyCamera()
 				// Old code using ava camera manager. To retire.
 				constexpr bool bIsCanvasController = false;
 				ViewportClient->GetCameraManager()->Init(Scene->GetPlaybackObject(), bIsCanvasController);
-#if WITH_EDITOR
 				ViewportClient->GetCameraManager()->SetViewTarget(CameraActor);
-#endif
 				GetPlayableGroup()->SetLastAppliedCameraPlayable(PlayableOwningCamera);
 				bSetupDone = true;
 				UE_LOG(LogAvaPlayable, Log, TEXT("No player controller found in \"%s\" - Using Ava Camera Manager instead."), *LevelStreaming->PackageNameToLoad.ToString())
