[Motion Design] Adding a setter to the Tag property to make it writable in game build.

diff --git a/Engine/Plugins/Experimental/Avalanche/Source/AvalancheAttribute/Public/Tags/AvaTagAttribute.h b/Engine/Plugins/Experimental/Avalanche/Source/AvalancheAttribute/Public/Tags/AvaTagAttribute.h
--- a/Engine/Plugins/Experimental/Avalanche/Source/AvalancheAttribute/Public/Tags/AvaTagAttribute.h
+++ b/Engine/Plugins/Experimental/Avalanche/Source/AvalancheAttribute/Public/Tags/AvaTagAttribute.h
@@ -21,9 +21,15 @@ public:
 	AVALANCHEATTRIBUTE_API virtual bool SetTagHandle(const FAvaTagHandle& InTagHandle) override;
 	AVALANCHEATTRIBUTE_API virtual bool ClearTagHandle(const FAvaTagHandle& InTagHandle) override;
 	AVALANCHEATTRIBUTE_API virtual bool ContainsTag(const FAvaTagHandle& InTagHandle) const override;
-	AVALANCHEATTRIBUTE_API virtual bool HasValidTagHandle() const;
+	AVALANCHEATTRIBUTE_API virtual bool HasValidTagHandle() const override;
 	//~ End UAvaTagAttributeBase
 
-	UPROPERTY(EditAnywhere, Category="Attributes")
+	UFUNCTION()
+	void SetTag(const FAvaTagHandle& InTag)
+	{
+		Tag = InTag;
+	}
+
+	UPROPERTY(EditAnywhere, Setter, Category="Attributes")
 	FAvaTagHandle Tag;
 };
