Motion Design: Fix font subsystem doing FAssetData::GetAsset instead of first checking if the asset is relevant to it

diff --git a/Engine/Plugins/Experimental/Avalanche/Source/AvalancheTextEditor/Private/Font/AvaFontManagerSubsystem.cpp b/Engine/Plugins/Experimental/Avalanche/Source/AvalancheTextEditor/Private/Font/AvaFontManagerSubsystem.cpp
--- a/Engine/Plugins/Experimental/Avalanche/Source/AvalancheTextEditor/Private/Font/AvaFontManagerSubsystem.cpp
+++ b/Engine/Plugins/Experimental/Avalanche/Source/AvalancheTextEditor/Private/Font/AvaFontManagerSubsystem.cpp
@@ -113,6 +113,11 @@ bool UAvaFontManagerSubsystem::IsSupportedFontFile(const FString& InFontFilePath
 
 void UAvaFontManagerSubsystem::OnAssetsAdded(const FAssetData& InAssetData)
 {
+	if (InAssetData.AssetClassPath != UFont::StaticClass()->GetClassPathName())
+	{
+		return;
+	}
+
 	if (UFont* CurrFont = Cast<UFont>(InAssetData.GetAsset()))
 	{
 		CreateProjectFont(CurrFont);
