Motion Design: Fix crash with viewport.
Recent changes in Level Editor Viewport seem to require Vp client to properly have its widget set

diff --git a/Engine/Plugins/Experimental/Avalanche/Source/AvalancheLevelViewport/Private/SAvaLevelViewport.cpp b/Engine/Plugins/Experimental/Avalanche/Source/AvalancheLevelViewport/Private/SAvaLevelViewport.cpp
--- a/Engine/Plugins/Experimental/Avalanche/Source/AvalancheLevelViewport/Private/SAvaLevelViewport.cpp
+++ b/Engine/Plugins/Experimental/Avalanche/Source/AvalancheLevelViewport/Private/SAvaLevelViewport.cpp
@@ -59,6 +59,9 @@ void SAvaLevelViewport::Construct(const FArguments& InArgs, const FAssetEditorVi
 	VirtualSizeAspectRatio = 0.f;
 	VirtualSizeAspectRatioState = EAvaViewportVirtualSizeAspectRatioState::LockedToCamera;
 
+	TSharedPtr<FAvaLevelViewportClient> ViewportClient = InArgs._ViewportFrame->GetViewportClient();
+	ViewportClient->SetViewportWidget(SharedThis(this));
+
 	Super::Construct(Super::FArguments()
 			.ParentLevelEditor(InArgs._ParentLevelEditor)
 			.LevelEditorViewportClient(InArgs._ViewportFrame->GetViewportClient())
diff --git a/Engine/Plugins/Experimental/Avalanche/Source/AvalancheLevelViewport/Private/SAvaLevelViewportFrame.cpp b/Engine/Plugins/Experimental/Avalanche/Source/AvalancheLevelViewport/Private/SAvaLevelViewportFrame.cpp
--- a/Engine/Plugins/Experimental/Avalanche/Source/AvalancheLevelViewport/Private/SAvaLevelViewportFrame.cpp
+++ b/Engine/Plugins/Experimental/Avalanche/Source/AvalancheLevelViewport/Private/SAvaLevelViewportFrame.cpp
@@ -114,8 +114,6 @@ void SAvaLevelViewportFrame::Construct(const FArguments& InArgs, const FAssetEdi
 			]
 		]
 	];
-
-	ViewportClient->SetViewportWidget(ViewportWidget);
 }
 
 void SAvaLevelViewportFrame::Tick(const FGeometry& AllottedGeometry, const double InCurrentTime,
