Motion Design: trivial valid check added

diff --git a/Engine/Plugins/Experimental/Avalanche/Source/AvalancheLevelViewport/Private/SAvaLevelViewport.cpp b/Engine/Plugins/Experimental/Avalanche/Source/AvalancheLevelViewport/Private/SAvaLevelViewport.cpp
--- a/Engine/Plugins/Experimental/Avalanche/Source/AvalancheLevelViewport/Private/SAvaLevelViewport.cpp
+++ b/Engine/Plugins/Experimental/Avalanche/Source/AvalancheLevelViewport/Private/SAvaLevelViewport.cpp
@@ -60,11 +60,15 @@ void SAvaLevelViewport::Construct(const FArguments& InArgs, const FAssetEditorVi
 	VirtualSizeAspectRatioState = EAvaViewportVirtualSizeAspectRatioState::LockedToCamera;
 
 	TSharedPtr<FAvaLevelViewportClient> ViewportClient = InArgs._ViewportFrame->GetViewportClient();
-	ViewportClient->SetViewportWidget(SharedThis(this));
+
+	if (ViewportClient.IsValid())
+	{
+		ViewportClient->SetViewportWidget(SharedThis(this));
+	}
 
 	Super::Construct(Super::FArguments()
 			.ParentLevelEditor(InArgs._ParentLevelEditor)
-			.LevelEditorViewportClient(InArgs._ViewportFrame->GetViewportClient())
+			.LevelEditorViewportClient(ViewportClient)
 		, InViewportArgs);
 
 	GetMutableDefault<UAvaViewportSettings>()->OnChange.AddSP(SharedThis(this), &SAvaLevelViewport::OnSettingsChanged);
