[MotionDesign] Ease curve tool
 - Add support for float value channels
 - Better sequencer selection updating
 - Remove unneeded select next/previous key actions

diff --git a/Engine/Plugins/Experimental/Avalanche/Source/AvalancheSequencer/Private/AvaSequencer.cpp b/Engine/Plugins/Experimental/Avalanche/Source/AvalancheSequencer/Private/AvaSequencer.cpp
--- a/Engine/Plugins/Experimental/Avalanche/Source/AvalancheSequencer/Private/AvaSequencer.cpp
+++ b/Engine/Plugins/Experimental/Avalanche/Source/AvalancheSequencer/Private/AvaSequencer.cpp
@@ -279,13 +279,13 @@ void FAvaSequencer::EnsureSequencer()
 	const FAvaEaseCurveToolCommands& EaseCurveToolCommands = FAvaEaseCurveToolCommands::Get();
 
 	CommandList->MapAction(EaseCurveToolCommands.QuickEaseIn
-		, FExecuteAction::CreateSP(EaseCurveToolRef, &FAvaEaseCurveTool::ApplyQuickEaseToSequencerKeySelections, FAvaEaseCurveTool::EOperation::In));
+		, FExecuteAction::CreateSP(EaseCurveToolRef, &FAvaEaseCurveTool::ApplyQuickEaseToSequencerKeySelections, EAvaEaseCurveToolOperation::In));
 
 	CommandList->MapAction(EaseCurveToolCommands.QuickEase
-		, FExecuteAction::CreateSP(EaseCurveToolRef, &FAvaEaseCurveTool::ApplyQuickEaseToSequencerKeySelections, FAvaEaseCurveTool::EOperation::InOut));
+		, FExecuteAction::CreateSP(EaseCurveToolRef, &FAvaEaseCurveTool::ApplyQuickEaseToSequencerKeySelections, EAvaEaseCurveToolOperation::InOut));
 
 	CommandList->MapAction(EaseCurveToolCommands.QuickEaseOut
-		, FExecuteAction::CreateSP(EaseCurveToolRef, &FAvaEaseCurveTool::ApplyQuickEaseToSequencerKeySelections, FAvaEaseCurveTool::EOperation::Out));
+		, FExecuteAction::CreateSP(EaseCurveToolRef, &FAvaEaseCurveTool::ApplyQuickEaseToSequencerKeySelections, EAvaEaseCurveToolOperation::Out));
 }
 
 TSharedRef<ISequencer> FAvaSequencer::CreateSequencer()
diff --git a/Engine/Plugins/Experimental/Avalanche/Source/AvalancheSequencer/Private/EaseCurveTool/AvaEaseCurveKeySelection.cpp b/Engine/Plugins/Experimental/Avalanche/Source/AvalancheSequencer/Private/EaseCurveTool/AvaEaseCurveKeySelection.cpp
new file mode 100644
--- /dev/null
+++ b/Engine/Plugins/Experimental/Avalanche/Source/AvalancheSequencer/Private/EaseCurveTool/AvaEaseCurveKeySelection.cpp
@@ -0,0 +1,261 @@
+﻿// Copyright Epic Games, Inc. All Rights Reserved.
+
+#include "AvaEaseCurveKeySelection.h"
+
+#include "Channels/MovieSceneChannel.h"
+#include "Channels/MovieSceneDoubleChannel.h"
+#include "Channels/MovieSceneFloatChannel.h"
+#include "EaseCurveTool/AvaEaseCurveTangents.h"
+#include "EaseCurveTool/AvaEaseCurveTool.h"
+#include "IKeyArea.h"
+#include "MVVM/ViewModels/ChannelModel.h"
+
+using namespace UE::Sequencer;
+
+FAvaEaseCurveKeySelection::FAvaEaseCurveKeySelection(const TWeakPtr<FSequencerSelection>& InSequencerSelectionWeak)
+	: SequencerSelectionWeak(InSequencerSelectionWeak)
+{
+	const TSharedPtr<FSequencerSelection> SequencerSelection = InSequencerSelectionWeak.Pin();
+	if (!SequencerSelection.IsValid())
+	{
+		return;
+	}
+
+	for (const FKeyHandle Key : SequencerSelection->KeySelection)
+	{
+		if (Key == FKeyHandle::Invalid())
+		{
+			continue;
+		}
+
+		const TViewModelPtr<FChannelModel> ChannelModel = SequencerSelection->KeySelection.GetModelForKey(Key);
+		if (!ChannelModel.IsValid())
+		{
+			continue;
+		}
+
+		const TSharedPtr<IKeyArea> KeyArea = ChannelModel->GetKeyArea();
+		if (!KeyArea.IsValid())
+		{
+			continue;
+		}
+
+		const FMovieSceneChannelHandle& ChannelHandle = KeyArea->GetChannel();
+
+		FChannelData& Entry = ChannelKeyData.FindOrAdd(ChannelModel->GetChannelName());
+		Entry.ChannelModel = ChannelModel;
+		Entry.Section = ChannelModel->GetSection();
+		Entry.Channel = ChannelHandle;
+		Entry.KeyHandles.Add(Key);
+
+		TotalSelectedKeys++;
+		
+		const int32 AllKeyCount = ChannelHandle.Get()->GetNumKeys();
+		const int32 SelectedKeyCount = Entry.KeyHandles.Num();
+		
+		if (SelectedKeyCount == 1 && TotalSelectedKeys == 1)
+		{
+			const int32 KeyIndex = ChannelHandle.Get()->GetIndex(Entry.KeyHandles[0]);
+			if (KeyIndex == AllKeyCount - 1)
+			{
+				bIsLastOnlySelectedKey = true;
+			}
+		}
+		else if (SelectedKeyCount > 1)
+		{
+			bAllChannelSingleKeySelections = false;
+		}
+	}
+}
+
+void FAvaEaseCurveKeySelection::ForEachEaseableKey(const bool bInIncludeEqualValueKeys
+	, const TFunctionRef<bool(const FKeyHandle& /*InKeyHandle*/, const FKeyHandle& /*InNextKeyHandle*/, const FChannelData&)>& InCallable)
+{
+	for (const TPair<FName, FChannelData>& ChannelEntry : ChannelKeyData)
+	{
+		const FName ChannelTypeName = ChannelEntry.Value.Channel.GetChannelTypeName();
+		if (ChannelTypeName == FMovieSceneDoubleChannel::StaticStruct()->GetFName())
+		{
+			if (!CheckMatchingValues<FMovieSceneDoubleChannel, FMovieSceneDoubleValue>(bInIncludeEqualValueKeys, ChannelEntry.Value, InCallable))
+			{
+				return;
+			}
+		}
+		else if (ChannelTypeName == FMovieSceneFloatChannel::StaticStruct()->GetFName())
+		{
+			if (!CheckMatchingValues<FMovieSceneFloatChannel, FMovieSceneFloatValue>(bInIncludeEqualValueKeys, ChannelEntry.Value, InCallable))
+			{
+				return;
+			}
+		}
+	}
+}
+
+namespace UE::Ava::EaseCurveTool::Private
+{
+	template <class ChannelHandle, class ChannelValue>
+	void UpdateChannelValues(const FKeyHandle& InKeyHandle, const FKeyHandle& InNextKeyHandle
+		, const FAvaEaseCurveKeySelection::FChannelData& InChannelData
+		, const bool bInAutoFlipTangents
+		, const FFrameRate& InDisplayRate, const FFrameRate& InTickResolution
+		, TArray<FAvaEaseCurveTangents>& OutKeySetTangents
+		, TArray<FAvaEaseCurveTangents>& OutChangingTangents)
+	{
+		TMovieSceneChannelHandle<ChannelHandle> Channel = InChannelData.Channel.Cast<ChannelHandle>();
+		TMovieSceneChannelData<ChannelValue> ChannelData = Channel.Get()->GetData();
+
+		const TArrayView<ChannelValue> ChannelValues = ChannelData.GetValues();
+		const TArrayView<FFrameNumber> ChannelTimes = ChannelData.GetTimes();
+
+		const int32 KeyIndex = ChannelData.GetIndex(InKeyHandle);
+		const int32 NextKeyIndex = ChannelData.GetIndex(InNextKeyHandle);
+
+		// If there is a key frame after this key frame that we are editing, we check if the that key frame value is less
+		// than or greater than this key frame value. If less, flip the tangent (if option is set).
+		const bool bIncreasingValue = ChannelValues[NextKeyIndex].Value >= ChannelValues[KeyIndex].Value;
+
+		FAvaEaseCurveTangents Tangents = FAvaEaseCurveTangents(ChannelValues[KeyIndex], ChannelValues[NextKeyIndex]);
+
+		if (bInAutoFlipTangents && !bIncreasingValue)
+		{
+			Tangents.Start *= -1.f;
+			Tangents.End *= -1.f;
+		}
+
+		// Scale time/value to normalized tangent range
+		FAvaEaseCurveTangents ScaledTangents = Tangents;
+		ScaledTangents.Normalize(ChannelTimes[KeyIndex], ChannelValues[KeyIndex].Value
+			, ChannelTimes[NextKeyIndex], ChannelValues[NextKeyIndex].Value
+			, InDisplayRate, InTickResolution);
+
+		OutKeySetTangents.Add(ScaledTangents);
+		if (ChannelValues[KeyIndex].Value != ChannelValues[NextKeyIndex].Value)
+		{
+			OutChangingTangents.Add(ScaledTangents);
+		}
+	}
+}
+
+FAvaEaseCurveTangents FAvaEaseCurveKeySelection::AverageTangents(const FFrameRate& InDisplayRate
+	, const FFrameRate& InTickResolution
+	, const bool bInAutoFlipTangents)
+{
+	using namespace UE::Ava::EaseCurveTool::Private;
+	
+	TArray<FAvaEaseCurveTangents> KeySetTangents;
+	TArray<FAvaEaseCurveTangents> ChangingTangents;
+
+	ForEachEaseableKey(/*bInIncludeEqualValueKeys=*/true, [&InDisplayRate, &InTickResolution, bInAutoFlipTangents, &KeySetTangents, &ChangingTangents]
+		(const FKeyHandle& InKeyHandle, const FKeyHandle& InNextKeyHandle, const FAvaEaseCurveKeySelection::FChannelData& InChannelData)
+		{
+			const FName ChannelTypeName = InChannelData.Channel.GetChannelTypeName();
+			if (ChannelTypeName == FMovieSceneDoubleChannel::StaticStruct()->GetFName())
+			{
+				UpdateChannelValues<FMovieSceneDoubleChannel, FMovieSceneDoubleValue>(InKeyHandle, InNextKeyHandle, InChannelData
+					, bInAutoFlipTangents, InDisplayRate, InTickResolution, KeySetTangents, ChangingTangents);
+			}
+			else if (ChannelTypeName == FMovieSceneFloatChannel::StaticStruct()->GetFName())
+			{
+				UpdateChannelValues<FMovieSceneFloatChannel, FMovieSceneFloatValue>(InKeyHandle, InNextKeyHandle, InChannelData
+					, bInAutoFlipTangents, InDisplayRate, InTickResolution, KeySetTangents, ChangingTangents);
+			}
+
+			return true;
+		});
+	
+	return FAvaEaseCurveTangents::Average(ChangingTangents);
+}
+
+namespace UE::Ava::EaseCurveTool::Private
+{
+	template <class ChannelHandle, class ChannelValue>
+	void SetChannelValues(const FAvaEaseCurveTangents& InTangents
+		, const EAvaEaseCurveToolOperation InOperation
+		, const FKeyHandle& InKeyHandle, const FKeyHandle& InNextKeyHandle
+		, const FAvaEaseCurveKeySelection::FChannelData& InChannelData
+		, const bool bInAutoFlipTangents
+		, const FFrameRate& InDisplayRate, const FFrameRate& InTickResolution)
+	{
+		TMovieSceneChannelHandle<ChannelHandle> Channel = InChannelData.Channel.Cast<ChannelHandle>();
+		TMovieSceneChannelData<ChannelValue> ChannelData = Channel.Get()->GetData();
+
+		TArrayView<ChannelValue> ChannelValues = ChannelData.GetValues();
+		const TArrayView<FFrameNumber> ChannelTimes = ChannelData.GetTimes();
+
+		const int32 KeyIndex = ChannelData.GetIndex(InKeyHandle);
+		const int32 NextKeyIndex = ChannelData.GetIndex(InNextKeyHandle);
+
+		FAvaEaseCurveTangents ScaledTangents = InTangents;
+
+		// If there is a key frame after this key frame that we are editing, we check if the that key frame value is less
+		// than or greater than this key frame value. If less, flip the tangent (if option is set).
+		const bool bIncreasingValue = ChannelValues[NextKeyIndex].Value >= ChannelValues[KeyIndex].Value;
+
+		if (bInAutoFlipTangents && !bIncreasingValue)
+		{
+			ScaledTangents.Start *= -1.f;
+			ScaledTangents.End *= -1.f;
+		}
+
+		// Scale normalized tangents to time/value range
+		ScaledTangents.ScaleUp(ChannelTimes[KeyIndex], ChannelValues[KeyIndex].Value
+			, ChannelTimes[NextKeyIndex], ChannelValues[NextKeyIndex].Value
+			, InDisplayRate, InTickResolution);
+
+		const FScopedTransaction Transaction(NSLOCTEXT("EaseCurveTool", "SetSequencerCurveTangents", "Set Sequencer Curve Tangents"));
+		InChannelData.Section->Modify();
+		InChannelData.Section->MarkAsChanged();
+
+		// Set this keys leave tangent
+		if (InOperation == EAvaEaseCurveToolOperation::Out || InOperation == EAvaEaseCurveToolOperation::InOut)
+		{
+			ChannelValues[KeyIndex].InterpMode = ERichCurveInterpMode::RCIM_Cubic;
+			ChannelValues[KeyIndex].Tangent.TangentWeightMode = ERichCurveTangentWeightMode::RCTWM_WeightedBoth;
+			ChannelValues[KeyIndex].TangentMode = ERichCurveTangentMode::RCTM_Break;
+			ChannelValues[KeyIndex].Tangent.LeaveTangent = ScaledTangents.Start;
+			ChannelValues[KeyIndex].Tangent.LeaveTangentWeight = ScaledTangents.StartWeight;
+		}
+
+		// Set the next keys arrive tangent
+		if (NextKeyIndex != INDEX_NONE && (InOperation == EAvaEaseCurveToolOperation::In || InOperation == EAvaEaseCurveToolOperation::InOut))
+		{
+			ChannelValues[NextKeyIndex].InterpMode = ERichCurveInterpMode::RCIM_Cubic;
+			ChannelValues[NextKeyIndex].Tangent.TangentWeightMode = ERichCurveTangentWeightMode::RCTWM_WeightedBoth;
+			ChannelValues[NextKeyIndex].TangentMode = ERichCurveTangentMode::RCTM_Break;
+			ChannelValues[NextKeyIndex].Tangent.ArriveTangent = ScaledTangents.End;
+			ChannelValues[NextKeyIndex].Tangent.ArriveTangentWeight = ScaledTangents.EndWeight;
+		}
+	}
+}
+
+void FAvaEaseCurveKeySelection::SetTangents(const FAvaEaseCurveTangents& InTangents
+	, const EAvaEaseCurveToolOperation InOperation
+	, const FFrameRate& InDisplayRate
+	, const FFrameRate& InTickResolution
+	, const bool bInAutoFlipTangents)
+{
+	using namespace UE::Ava::EaseCurveTool::Private;
+	
+	if (TotalSelectedKeys == 0)
+	{
+		return;
+	}
+	
+	ForEachEaseableKey(/*bInIncludeEqualValueKeys=*/true, [&InDisplayRate, &InTickResolution, &InTangents, InOperation, bInAutoFlipTangents](const FKeyHandle& InKeyHandle
+		, const FKeyHandle& InNextKeyHandle, const FAvaEaseCurveKeySelection::FChannelData& InChannelData)
+		{
+			const FName ChannelTypeName = InChannelData.Channel.GetChannelTypeName();
+			if (ChannelTypeName == FMovieSceneDoubleChannel::StaticStruct()->GetFName())
+			{
+				SetChannelValues<FMovieSceneDoubleChannel, FMovieSceneDoubleValue>(InTangents, InOperation, InKeyHandle, InNextKeyHandle, InChannelData
+					, bInAutoFlipTangents, InDisplayRate, InTickResolution);
+			}
+			else if (ChannelTypeName == FMovieSceneFloatChannel::StaticStruct()->GetFName())
+			{
+				SetChannelValues<FMovieSceneFloatChannel, FMovieSceneFloatValue>(InTangents, InOperation, InKeyHandle, InNextKeyHandle, InChannelData
+					, bInAutoFlipTangents, InDisplayRate, InTickResolution);
+			}
+
+			return true;
+		});
+}
diff --git a/Engine/Plugins/Experimental/Avalanche/Source/AvalancheSequencer/Private/EaseCurveTool/AvaEaseCurveKeySelection.h b/Engine/Plugins/Experimental/Avalanche/Source/AvalancheSequencer/Private/EaseCurveTool/AvaEaseCurveKeySelection.h
new file mode 100644
--- /dev/null
+++ b/Engine/Plugins/Experimental/Avalanche/Source/AvalancheSequencer/Private/EaseCurveTool/AvaEaseCurveKeySelection.h
@@ -0,0 +1,109 @@
+﻿// Copyright Epic Games, Inc. All Rights Reserved.
+
+#pragma once
+
+#include "Channels/MovieSceneChannelData.h"
+#include "Channels/MovieSceneChannelHandle.h"
+#include "MVVM/Selection/Selection.h"
+
+struct FAvaEaseCurveTangents;
+class UMovieSceneSection;
+struct FKeyHandle;
+enum class EAvaEaseCurveToolOperation : uint8;
+
+namespace UE::Sequencer
+{
+	class FChannelModel;
+}
+
+struct FAvaEaseCurveKeySelection
+{
+	struct FChannelData
+	{
+		TSharedPtr<UE::Sequencer::FChannelModel> ChannelModel;
+		TObjectPtr<UMovieSceneSection> Section;
+		FMovieSceneChannelHandle Channel;
+		TArray<FKeyHandle> KeyHandles;
+	};
+
+	FAvaEaseCurveKeySelection() {}
+	FAvaEaseCurveKeySelection(const TWeakPtr<UE::Sequencer::FSequencerSelection>& InSequencerSelectionWeak);
+
+	void ForEachEaseableKey(const bool bInIncludeEqualValueKeys
+		, const TFunctionRef<bool(const FKeyHandle& /*InKeyHandle*/, const FKeyHandle& /*InNextKeyHandle*/, const FChannelData&)>& InCallable);
+
+	template <class ChannelHandle, class ChannelValue>
+	bool CheckMatchingValues(const bool bInIncludeEqualValueKeys, const FChannelData& InChannelData
+		, const TFunctionRef<bool(const FKeyHandle& /*InKeyHandle*/, const FKeyHandle& /*InNextKeyHandle*/, const FChannelData&)>& InCallable)
+	{
+		TMovieSceneChannelHandle<ChannelHandle> Channel = InChannelData.Channel.Cast<ChannelHandle>();
+		TMovieSceneChannelData<ChannelValue> ChannelData = Channel.Get()->GetData();
+
+		const TArrayView<ChannelValue> ChannelValues = ChannelData.GetValues();
+		const int32 KeyCount = ChannelValues.Num();
+
+		for (const FKeyHandle& KeyHandle : InChannelData.KeyHandles)
+		{
+			if (KeyHandle == FKeyHandle::Invalid())
+			{
+				continue;
+			}
+
+			const int32 KeyIndex = ChannelData.GetIndex(KeyHandle);
+			if (KeyIndex == INDEX_NONE)
+			{
+				continue;
+			}
+
+			// If there is no key after the selected key, we don't need to process.
+			// The arrive tangents of this key will be set by the previous key's processing.
+			int32 NextKeyIndex = KeyIndex + 1;
+			NextKeyIndex = (NextKeyIndex < KeyCount) ? NextKeyIndex : INDEX_NONE;
+			if (NextKeyIndex == INDEX_NONE)
+			{
+				continue;
+			}
+
+			// Need to check if the next key index is valid, otherwise GetHandle() will fail.
+			const FKeyHandle NextKeyHandle = ChannelData.GetHandle(NextKeyIndex);
+			if (NextKeyHandle == FKeyHandle::Invalid())
+			{
+				continue;
+			}
+
+			if (!bInIncludeEqualValueKeys && ChannelValues[KeyIndex].Value == ChannelValues[NextKeyIndex].Value)
+			{
+				continue;
+			}
+
+			if (!InCallable(KeyHandle, NextKeyHandle, InChannelData))
+			{
+				return false;
+			}
+		}
+
+		return true;
+	}
+
+	FAvaEaseCurveTangents AverageTangents(const FFrameRate& InDisplayRate
+		, const FFrameRate& InTickResolution
+		, const bool bInAutoFlipTangents);
+
+	void SetTangents(const FAvaEaseCurveTangents& InTangents
+		, const EAvaEaseCurveToolOperation InOperation
+		, const FFrameRate& InDisplayRate
+		, const FFrameRate& InTickResolution
+		, const bool bInAutoFlipTangents);
+
+	TWeakPtr<UE::Sequencer::FSequencerSelection> SequencerSelectionWeak;
+
+	TMap<FName, FChannelData> ChannelKeyData;
+
+	int32 TotalSelectedKeys = 0;
+
+	/** Indicates only one selected key on each of the selected channels. */
+	bool bAllChannelSingleKeySelections = true;
+
+	/** Indicates only one key selected and it is the last key of the channel. */
+	bool bIsLastOnlySelectedKey = false;
+};
diff --git a/Engine/Plugins/Experimental/Avalanche/Source/AvalancheSequencer/Private/EaseCurveTool/AvaEaseCurveStyle.cpp b/Engine/Plugins/Experimental/Avalanche/Source/AvalancheSequencer/Private/EaseCurveTool/AvaEaseCurveStyle.cpp
--- a/Engine/Plugins/Experimental/Avalanche/Source/AvalancheSequencer/Private/EaseCurveTool/AvaEaseCurveStyle.cpp
+++ b/Engine/Plugins/Experimental/Avalanche/Source/AvalancheSequencer/Private/EaseCurveTool/AvaEaseCurveStyle.cpp
@@ -36,8 +36,8 @@ FAvaEaseCurveStyle::FAvaEaseCurveStyle()
 	Set("ToolButton.ImageSize", ToolButtonImageSize);
 
 	const FButtonStyle ToolButtonStyle = FButtonStyle(SimpleButtonStyle)
-		.SetHovered(FSlateRoundedBoxBrush(FStyleColors::Hover, 4.0f))
-		.SetPressed(FSlateRoundedBoxBrush(FStyleColors::Hover, 4.0f))
+		.SetHovered(FSlateRoundedBoxBrush(FStyleColors::Hover, 4.f))
+		.SetPressed(FSlateRoundedBoxBrush(FStyleColors::Hover, 4.f))
 		.SetNormalPadding(FMargin(ToolButtonPadding))
 		.SetPressedPadding(FMargin(ToolButtonPadding, ToolButtonPadding + (ToolButtonPadding * 0.5f), ToolButtonPadding, ToolButtonPadding - (ToolButtonPadding * 0.5f)));
 	Set("ToolButton", ToolButtonStyle);
@@ -47,13 +47,15 @@ FAvaEaseCurveStyle::FAvaEaseCurveStyle()
 		.SetPressedPadding(FMargin(0.f));
 	Set("ToolButton.NoPad", ToolButtonNoPadStyle);
 
+	Set("ToolButton.Opaque", new FSlateRoundedBoxBrush(FStyleColors::Dropdown, 4.f));
+
 	const FCheckBoxStyle& ToggleButtonCheckboxStyle = FAppStyle::GetWidgetStyle<FCheckBoxStyle>(TEXT("ToggleButtonCheckbox"));
 
 	const FCheckBoxStyle ToolToggleButtonStyle = FCheckBoxStyle(ToggleButtonCheckboxStyle)
-		.SetUncheckedHoveredImage(FSlateRoundedBoxBrush(FStyleColors::Hover, 4.0f))
-		.SetUncheckedPressedImage(FSlateRoundedBoxBrush(FStyleColors::Hover, 4.0f))
-		.SetCheckedHoveredImage(FSlateRoundedBoxBrush(FStyleColors::PrimaryHover, 4.0f))
-		.SetCheckedPressedImage(FSlateRoundedBoxBrush(FStyleColors::PrimaryPress, 4.0f))
+		.SetUncheckedHoveredImage(FSlateRoundedBoxBrush(FStyleColors::Hover, 4.f))
+		.SetUncheckedPressedImage(FSlateRoundedBoxBrush(FStyleColors::Hover, 4.f))
+		.SetCheckedHoveredImage(FSlateRoundedBoxBrush(FStyleColors::PrimaryHover, 4.f))
+		.SetCheckedPressedImage(FSlateRoundedBoxBrush(FStyleColors::PrimaryPress, 4.f))
 		.SetPadding(ToolButtonPadding);
 	Set("ToolToggleButton", ToolToggleButtonStyle);
 
diff --git a/Engine/Plugins/Experimental/Avalanche/Source/AvalancheSequencer/Private/EaseCurveTool/AvaEaseCurveTangents.h b/Engine/Plugins/Experimental/Avalanche/Source/AvalancheSequencer/Private/EaseCurveTool/AvaEaseCurveTangents.h
--- a/Engine/Plugins/Experimental/Avalanche/Source/AvalancheSequencer/Private/EaseCurveTool/AvaEaseCurveTangents.h
+++ b/Engine/Plugins/Experimental/Avalanche/Source/AvalancheSequencer/Private/EaseCurveTool/AvaEaseCurveTangents.h
@@ -4,6 +4,7 @@
 
 #include "Curves/RichCurve.h"
 #include "Channels/MovieSceneDoubleChannel.h"
+#include "Channels/MovieSceneFloatChannel.h"
 #include "Math/UnrealMathUtility.h"
 #include "AvaEaseCurveTangents.generated.h"
 
@@ -46,6 +47,14 @@ public:
 		: Start(InStartMovieSceneDoubleValue.Tangent.LeaveTangent), StartWeight(InStartMovieSceneDoubleValue.Tangent.LeaveTangentWeight)
 		, End(InEndMovieSceneDoubleValue.Tangent.ArriveTangent), EndWeight(InEndMovieSceneDoubleValue.Tangent.ArriveTangentWeight)
 	{}
+	FAvaEaseCurveTangents(const FMovieSceneFloatValue& InMovieSceneFloatValue)
+		: Start(InMovieSceneFloatValue.Tangent.LeaveTangent), StartWeight(InMovieSceneFloatValue.Tangent.LeaveTangentWeight)
+		, End(InMovieSceneFloatValue.Tangent.ArriveTangent), EndWeight(InMovieSceneFloatValue.Tangent.ArriveTangentWeight)
+	{}
+	FAvaEaseCurveTangents(const FMovieSceneFloatValue& InStartMovieSceneFloatValue, const FMovieSceneFloatValue& InEndMovieSceneFloatValue)
+		: Start(InStartMovieSceneFloatValue.Tangent.LeaveTangent), StartWeight(InStartMovieSceneFloatValue.Tangent.LeaveTangentWeight)
+		, End(InEndMovieSceneFloatValue.Tangent.ArriveTangent), EndWeight(InEndMovieSceneFloatValue.Tangent.ArriveTangentWeight)
+	{}
 	/** Constructor from string consisting of cubic bezier points. Ex. "0.45, 0.34, 0.0, 1.00" */
 	explicit FAvaEaseCurveTangents(const FString& InTangentsString);
 
diff --git a/Engine/Plugins/Experimental/Avalanche/Source/AvalancheSequencer/Private/EaseCurveTool/AvaEaseCurveTool.cpp b/Engine/Plugins/Experimental/Avalanche/Source/AvalancheSequencer/Private/EaseCurveTool/AvaEaseCurveTool.cpp
--- a/Engine/Plugins/Experimental/Avalanche/Source/AvalancheSequencer/Private/EaseCurveTool/AvaEaseCurveTool.cpp
+++ b/Engine/Plugins/Experimental/Avalanche/Source/AvalancheSequencer/Private/EaseCurveTool/AvaEaseCurveTool.cpp
@@ -4,22 +4,18 @@
 #include "AssetRegistry/AssetRegistryModule.h"
 #include "AssetToolsModule.h"
 #include "AvaSequencer.h"
-#include "CurveEditor.h"
 #include "Dialogs/DlgPickAssetPath.h"
-#include "EaseCurveTool/AvaEaseCurveToolCommands.h"
+#include "EaseCurveTool/AvaEaseCurveKeySelection.h"
 #include "EaseCurveTool/AvaEaseCurveToolSettings.h"
 #include "EaseCurveTool/Widgets/SAvaEaseCurveTool.h"
 #include "EngineAnalytics.h"
 #include "Factories/CurveFactory.h"
 #include "Framework/Notifications/NotificationManager.h"
 #include "HAL/PlatformApplicationMisc.h"
-#include "IKeyArea.h"
 #include "ISequencer.h"
 #include "ISettingsModule.h"
 #include "Math/UnrealMathUtility.h"
-#include "MVVM/CurveEditorExtension.h"
 #include "MVVM/Selection/Selection.h"
-#include "MVVM/ViewModels/ChannelModel.h"
 #include "MVVM/ViewModels/SequencerEditorViewModel.h"
 #include "Settings/AvaSequencerSettings.h"
 #include "Widgets/Notifications/SNotificationList.h"
@@ -40,72 +36,22 @@ FAvaEaseCurveTool::FAvaEaseCurveTool(const TSharedRef<FAvaSequencer>& InSequence
 		SequencerSelectionWeak = SequencerViewModel->GetSelection();
 	}
 
+	Sequencer->GetSelectionChangedObjectGuids().AddRaw(this, &FAvaEaseCurveTool::OnSequencerSelectionChanged);
+
 	UpdateEaseCurveFromSequencerKeySelections();
 }
 
-void FAvaEaseCurveTool::CacheSelectionData()
+FAvaEaseCurveTool::~FAvaEaseCurveTool()
 {
-	if (!SequencerSelectionWeak.IsValid())
+	if (const TSharedPtr<FAvaSequencer> AvaSequencer = AvaSequencerWeak.Pin())
 	{
-		return;
-	}
-
-	const TSharedPtr<FSequencerSelection> SequencerSelection = SequencerSelectionWeak.Pin();
-
-	FKeyDataCache NewKeyCache;
-
-	for (const FKeyHandle Key : SequencerSelection->KeySelection)
-	{
-		if (Key == FKeyHandle::Invalid())
-		{
-			continue;
-		}
-
-		TViewModelPtr<FChannelModel> ChannelModel = SequencerSelection->KeySelection.GetModelForKey(Key);
-		if (!ChannelModel.IsValid())
-		{
-			continue;
-		}
-
-		TSharedPtr<IKeyArea> KeyArea = ChannelModel->GetKeyArea();
-		if (!KeyArea.IsValid())
-		{
-			continue;
-		}
-
-		FMovieSceneDoubleChannel* const DoubleChannel = KeyArea->GetChannel().Cast<FMovieSceneDoubleChannel>().Get();
-		if (!DoubleChannel)
-		{
-			continue;
-		}
-
-		FKeyDataCache::FChannelData& Entry = NewKeyCache.ChannelKeyData.FindOrAdd(ChannelModel->GetChannelName());
-		Entry.ChannelModel = ChannelModel;
-		Entry.DoubleChannel = DoubleChannel;
-		Entry.Section = ChannelModel->GetSection();
-		Entry.KeyHandles.Add(Key);
-
-		NewKeyCache.TotalSelectedKeys++;
-
-		const TArrayView<const FFrameNumber> ChannelTimes = DoubleChannel->GetTimes();
-		const int32 AllKeyCount = ChannelTimes.Num();
-		const int32 SelectedKeyCount = Entry.KeyHandles.Num();
-
-		if (SelectedKeyCount == 1 && NewKeyCache.TotalSelectedKeys == 1)
-		{
-			const int32 KeyIndex = DoubleChannel->GetIndex(Entry.KeyHandles[0]);
-			if (KeyIndex == AllKeyCount - 1)
-			{
-				NewKeyCache.bIsLastOnlySelectedKey = true;
-			}
-		}
-		else if (SelectedKeyCount > 1)
-		{
-			NewKeyCache.bAllChannelSingleKeySelections = false;
-		}
+		AvaSequencer->GetSequencer()->GetSelectionChangedObjectGuids().RemoveAll(this);
 	}
+}
 
-	KeyCache = NewKeyCache;
+void FAvaEaseCurveTool::OnSequencerSelectionChanged(TArray<FGuid> InObjectGuids)
+{
+	UpdateEaseCurveFromSequencerKeySelections();
 }
 
 void FAvaEaseCurveTool::AddReferencedObjects(FReferenceCollector& Collector)
@@ -153,17 +99,17 @@ FAvaEaseCurveTangents FAvaEaseCurveTool::GetEaseCurveTangents() const
 	return EaseCurve->GetTangents();
 }
 
-void FAvaEaseCurveTool::SetEaseCurveTangents_Internal(const FAvaEaseCurveTangents& InTangents, const EOperation InOperation, const bool bInBroadcastUpdate)
+void FAvaEaseCurveTool::SetEaseCurveTangents_Internal(const FAvaEaseCurveTangents& InTangents, const EAvaEaseCurveToolOperation InOperation, const bool bInBroadcastUpdate) const
 {
 	switch (InOperation)
 	{
-	case EOperation::InOut:
+	case EAvaEaseCurveToolOperation::InOut:
 		EaseCurve->SetTangents(InTangents);
 		break;
-	case EOperation::In:
+	case EAvaEaseCurveToolOperation::In:
 		EaseCurve->SetEndTangent(InTangents.End, InTangents.EndWeight);
 		break;
-	case EOperation::Out:
+	case EAvaEaseCurveToolOperation::Out:
 		EaseCurve->SetStartTangent(InTangents.Start, InTangents.StartWeight);
 		break;
 	}
@@ -174,7 +120,7 @@ void FAvaEaseCurveTool::SetEaseCurveTangents_Internal(const FAvaEaseCurveTangent
 	}
 }
 
-void FAvaEaseCurveTool::SetEaseCurveTangents(const FAvaEaseCurveTangents& InTangents, const EOperation InOperation
+void FAvaEaseCurveTool::SetEaseCurveTangents(const FAvaEaseCurveTangents& InTangents, const EAvaEaseCurveToolOperation InOperation
 	, const bool bInBroadcastUpdate, const bool bInSetSequencerTangents)
 {
 	const FScopedTransaction Transaction(LOCTEXT("SetEaseCurveTangents", "Set Ease Curve Tangents"));
@@ -188,19 +134,19 @@ void FAvaEaseCurveTool::SetEaseCurveTangents(const FAvaEaseCurveTangents& InTang
 	}
 }
 
-void FAvaEaseCurveTool::ResetEaseCurveTangents(const EOperation InOperation)
+void FAvaEaseCurveTool::ResetEaseCurveTangents(const EAvaEaseCurveToolOperation InOperation)
 {
 	FText TransactionText;
 
 	switch (InOperation)
 	{
-	case EOperation::InOut:
+	case EAvaEaseCurveToolOperation::InOut:
 		TransactionText = LOCTEXT("ResetTangents", "Reset Tangents");
 		break;
-	case EOperation::In:
+	case EAvaEaseCurveToolOperation::In:
 		TransactionText = LOCTEXT("ResetEndTangents", "Reset End Tangents");
 		break;
-	case EOperation::Out:
+	case EAvaEaseCurveToolOperation::Out:
 		TransactionText = LOCTEXT("ResetStartTangents", "Reset Start Tangents");
 		break;
 	}
@@ -213,20 +159,20 @@ void FAvaEaseCurveTool::ResetEaseCurveTangents(const EOperation InOperation)
 	SetSequencerKeySelectionTangents(ZeroTangents, InOperation);
 }
 
-void FAvaEaseCurveTool::FlattenOrStraightenTangents(const EOperation InOperation, const bool bInFlattenTangents) const
+void FAvaEaseCurveTool::FlattenOrStraightenTangents(const EAvaEaseCurveToolOperation InOperation, const bool bInFlattenTangents) const
 {
 	FText TransactionText;
 	if (bInFlattenTangents)
 	{
 		switch (InOperation)
 		{
-		case EOperation::InOut:
+		case EAvaEaseCurveToolOperation::InOut:
 			TransactionText = LOCTEXT("FlattenTangents", "Flatten Tangents");
 			break;
-		case EOperation::In:
+		case EAvaEaseCurveToolOperation::In:
 			TransactionText = LOCTEXT("FlattenEndTangents", "Flatten End Tangents");
 			break;
-		case EOperation::Out:
+		case EAvaEaseCurveToolOperation::Out:
 			TransactionText = LOCTEXT("FlattenStartTangents", "Flatten Start Tangents");
 			break;
 		}
@@ -235,13 +181,13 @@ void FAvaEaseCurveTool::FlattenOrStraightenTangents(const EOperation InOperation
 	{
 		switch (InOperation)
 		{
-		case EOperation::InOut:
+		case EAvaEaseCurveToolOperation::InOut:
 			TransactionText = LOCTEXT("StraightenTangents", "Straighten Tangents");
 			break;
-		case EOperation::In:
+		case EAvaEaseCurveToolOperation::In:
 			TransactionText = LOCTEXT("StraightenEndTangents", "Straighten End Tangents");
 			break;
-		case EOperation::Out:
+		case EAvaEaseCurveToolOperation::Out:
 			TransactionText = LOCTEXT("StraightenStartTangents", "Straighten Start Tangents");
 			break;
 		}
@@ -249,11 +195,11 @@ void FAvaEaseCurveTool::FlattenOrStraightenTangents(const EOperation InOperation
 	const FScopedTransaction Transaction(TransactionText);
 	EaseCurve->ModifyOwner();
 
-	if (InOperation == EOperation::Out || InOperation == EOperation::InOut)
+	if (InOperation == EAvaEaseCurveToolOperation::Out || InOperation == EAvaEaseCurveToolOperation::InOut)
 	{
 		EaseCurve->FlattenOrStraightenTangents(EaseCurve->GetStartKeyHandle(), bInFlattenTangents);
 	}
-	if (InOperation == EOperation::In || InOperation == EOperation::InOut)
+	if (InOperation == EAvaEaseCurveToolOperation::In || InOperation == EAvaEaseCurveToolOperation::InOut)
 	{
 		EaseCurve->FlattenOrStraightenTangents(EaseCurve->GetEndKeyHandle(), bInFlattenTangents);
 	}
@@ -261,7 +207,7 @@ void FAvaEaseCurveTool::FlattenOrStraightenTangents(const EOperation InOperation
 	EaseCurve->BroadcastUpdate();
 }
 
-void FAvaEaseCurveTool::ApplyQuickEaseToSequencerKeySelections(const EOperation InOperation)
+void FAvaEaseCurveTool::ApplyQuickEaseToSequencerKeySelections(const EAvaEaseCurveToolOperation InOperation)
 {
 	const UAvaEaseCurveToolSettings* const Settings = GetDefault<UAvaEaseCurveToolSettings>();
 
@@ -286,13 +232,13 @@ void FAvaEaseCurveTool::ApplyQuickEaseToSequencerKeySelections(const EOperation
 		FString ParamValue;
 		switch (InOperation)
 		{
-		case EOperation::InOut:
+		case EAvaEaseCurveToolOperation::InOut:
 			ParamValue = TEXT("InOut");
 			break;
-		case EOperation::In:
+		case EAvaEaseCurveToolOperation::In:
 			ParamValue = TEXT("In");
 			break;
-		case EOperation::Out:
+		case EAvaEaseCurveToolOperation::Out:
 			ParamValue = TEXT("Out");
 			break;
 		}
@@ -300,139 +246,38 @@ void FAvaEaseCurveTool::ApplyQuickEaseToSequencerKeySelections(const EOperation
 	}
 }
 
-void FAvaEaseCurveTool::SetSequencerKeySelectionTangents(const FAvaEaseCurveTangents& InTangents, const EOperation InOperation)
+void FAvaEaseCurveTool::SetSequencerKeySelectionTangents(const FAvaEaseCurveTangents& InTangents, const EAvaEaseCurveToolOperation InOperation)
 {
-	CacheSelectionData();
+	KeyCache = FAvaEaseCurveKeySelection(SequencerSelectionWeak);
 
 	if (KeyCache.TotalSelectedKeys == 0)
 	{
 		return;
 	}
-
-	const FFrameRate DisplayRate = GetDisplayRate();
-	const FFrameRate TickResolution = GetTickResolution();
+	
 	const UAvaEaseCurveToolSettings* const EaseCurveToolSettings = GetDefault<UAvaEaseCurveToolSettings>();
 	check(EaseCurveToolSettings);
 	const bool bAutoFlipTangents = EaseCurveToolSettings->GetAutoFlipTangents();
 
-	KeyCache.ForEachEaseableKey(/*bInIncludeEqualValueKeys=*/true, [&DisplayRate, &TickResolution, &InTangents, InOperation, bAutoFlipTangents](const FKeyHandle& InKeyHandle
-		, const FKeyHandle& InNextKeyHandle, const FKeyDataCache::FChannelData& InChannelData)
-		{
-			TMovieSceneChannelData<FMovieSceneDoubleValue> ChannelData = InChannelData.DoubleChannel->GetData();
-
-			TArrayView<FMovieSceneDoubleValue> ChannelValues = ChannelData.GetValues();
-			const TArrayView<FFrameNumber> ChannelTimes = ChannelData.GetTimes();
-			const int32 KeyCount = ChannelValues.Num();
-
-			const int32 KeyIndex = ChannelData.GetIndex(InKeyHandle);
-			const int32 NextKeyIndex = ChannelData.GetIndex(InNextKeyHandle);
-
-			FAvaEaseCurveTangents ScaledTangents = InTangents;
-
-			// If there is a key frame after this key frame that we are editing, we check if the that key frame value is less
-			// than or greater than this key frame value. If less, flip the tangent (if option is set).
-			const bool bIncreasingValue = ChannelValues[NextKeyIndex].Value >= ChannelValues[KeyIndex].Value;
-
-			if (bAutoFlipTangents && !bIncreasingValue)
-			{
-				ScaledTangents.Start *= -1.f;
-				ScaledTangents.End *= -1.f;
-			}
-
-			// Scale normalized tangents to time/value range
-			ScaledTangents.ScaleUp(ChannelTimes[KeyIndex], ChannelValues[KeyIndex].Value
-				, ChannelTimes[NextKeyIndex], ChannelValues[NextKeyIndex].Value
-				, DisplayRate, TickResolution);
-
-			const FScopedTransaction Transaction(LOCTEXT("SetSequencerCurveTangents", "Set Sequencer Curve Tangents"));
-			InChannelData.Section->Modify();
-			InChannelData.Section->MarkAsChanged();
-
-			// Set this keys leave tangent
-			if (InOperation == EOperation::Out || InOperation == EOperation::InOut)
-			{
-				ChannelValues[KeyIndex].InterpMode = ERichCurveInterpMode::RCIM_Cubic;
-				ChannelValues[KeyIndex].Tangent.TangentWeightMode = ERichCurveTangentWeightMode::RCTWM_WeightedBoth;
-				ChannelValues[KeyIndex].TangentMode = ERichCurveTangentMode::RCTM_Break;
-				ChannelValues[KeyIndex].Tangent.LeaveTangent = ScaledTangents.Start;
-				ChannelValues[KeyIndex].Tangent.LeaveTangentWeight = ScaledTangents.StartWeight;
-			}
-
-			// Set the next keys arrive tangent
-			if (NextKeyIndex != INDEX_NONE && (InOperation == EOperation::In || InOperation == EOperation::InOut))
-			{
-				ChannelValues[NextKeyIndex].InterpMode = ERichCurveInterpMode::RCIM_Cubic;
-				ChannelValues[NextKeyIndex].Tangent.TangentWeightMode = ERichCurveTangentWeightMode::RCTWM_WeightedBoth;
-				ChannelValues[NextKeyIndex].TangentMode = ERichCurveTangentMode::RCTM_Break;
-				ChannelValues[NextKeyIndex].Tangent.ArriveTangent = ScaledTangents.End;
-				ChannelValues[NextKeyIndex].Tangent.ArriveTangentWeight = ScaledTangents.EndWeight;
-			}
-
-			return true;
-		});
+	KeyCache.SetTangents(InTangents, InOperation, GetDisplayRate(), GetTickResolution(), bAutoFlipTangents);
 }
 
 void FAvaEaseCurveTool::UpdateEaseCurveFromSequencerKeySelections()
 {
-	CacheSelectionData();
+	KeyCache = FAvaEaseCurveKeySelection(SequencerSelectionWeak);
 	
-	const FFrameRate DisplayRate = GetDisplayRate();
-	const FFrameRate TickResolution = GetTickResolution();
 	const UAvaEaseCurveToolSettings* const EaseCurveToolSettings = GetDefault<UAvaEaseCurveToolSettings>();
 	check(EaseCurveToolSettings);
 	const bool bAutoFlipTangents = EaseCurveToolSettings->GetAutoFlipTangents();
+	
+	const FAvaEaseCurveTangents AverageTangents = KeyCache.AverageTangents(GetDisplayRate(), GetTickResolution(), bAutoFlipTangents);
 
-	TArray<FAvaEaseCurveTangents> KeySetTangents;
-	TArray<FAvaEaseCurveTangents> ChangingTangents;
-
-	KeyCache.ForEachEaseableKey(/*bInIncludeEqualValueKeys=*/true, [&DisplayRate, &TickResolution, bAutoFlipTangents, &KeySetTangents, &ChangingTangents]
-		(const FKeyHandle& InKeyHandle, const FKeyHandle& InNextKeyHandle, const FKeyDataCache::FChannelData& InChannelData)
-		{
-			TMovieSceneChannelData<FMovieSceneDoubleValue> ChannelData = InChannelData.DoubleChannel->GetData();
-
-			const TArrayView<FMovieSceneDoubleValue> ChannelValues = ChannelData.GetValues();
-			const TArrayView<FFrameNumber> ChannelTimes = ChannelData.GetTimes();
-			const int32 KeyCount = ChannelValues.Num();
-
-			const int32 KeyIndex = ChannelData.GetIndex(InKeyHandle);
-			const int32 NextKeyIndex = ChannelData.GetIndex(InNextKeyHandle);
-
-			// If there is a key frame after this key frame that we are editing, we check if the that key frame value is less
-			// than or greater than this key frame value. If less, flip the tangent (if option is set).
-			const bool bIncreasingValue = ChannelValues[NextKeyIndex].Value >= ChannelValues[KeyIndex].Value;
-
-			FAvaEaseCurveTangents Tangents = FAvaEaseCurveTangents(ChannelValues[KeyIndex], ChannelValues[NextKeyIndex]);
-
-			if (bAutoFlipTangents && !bIncreasingValue)
-			{
-				Tangents.Start *= -1.f;
-				Tangents.End *= -1.f;
-			}
-
-			// Scale time/value to normalized tangent range
-			FAvaEaseCurveTangents ScaledTangents = Tangents;
-			ScaledTangents.Normalize(ChannelTimes[KeyIndex], ChannelValues[KeyIndex].Value
-				, ChannelTimes[NextKeyIndex], ChannelValues[NextKeyIndex].Value
-				, DisplayRate, TickResolution);
-
-			KeySetTangents.Add(ScaledTangents);
-			if (ChannelValues[KeyIndex].Value != ChannelValues[NextKeyIndex].Value)
-			{
-				ChangingTangents.Add(ScaledTangents);
-			}
-
-			return true;
-		});
-
-	const FAvaEaseCurveTangents AverageTangents = FAvaEaseCurveTangents::Average(KeySetTangents);
-	const FAvaEaseCurveTangents AverageChangingTangents = FAvaEaseCurveTangents::Average(ChangingTangents);
-
-	SetEaseCurveTangents(AverageChangingTangents, EOperation::InOut, /*bInBroadcastUpdate=*/true, false);
+	SetEaseCurveTangents(AverageTangents, EAvaEaseCurveToolOperation::InOut, /*bInBroadcastUpdate=*/true, false);
 
 	// Update the preset combobox widget
 	if (ToolWidget.IsValid())
 	{
-		ToolWidget->SetTangents(AverageChangingTangents, EOperation::InOut
+		ToolWidget->SetTangents(AverageTangents, EAvaEaseCurveToolOperation::InOut
 			, /*bInSetEaseCurve=*/false, /*bInBroadcastUpdate=*/false, /*bInSetSequencerTangents*/false);
 	}
 }
@@ -497,17 +342,17 @@ UCurveBase* FAvaEaseCurveTool::CreateCurveAsset() const
 	return nullptr;
 }
 
-FAvaEaseCurveTool::EOperation FAvaEaseCurveTool::GetToolOperation() const
+EAvaEaseCurveToolOperation FAvaEaseCurveTool::GetToolOperation() const
 {
 	return OperationMode;
 }
 
-void FAvaEaseCurveTool::SetToolOperation(const FAvaEaseCurveTool::EOperation InNewOperation)
+void FAvaEaseCurveTool::SetToolOperation(const EAvaEaseCurveToolOperation InNewOperation)
 {
 	OperationMode = InNewOperation;
 }
 
-bool FAvaEaseCurveTool::IsToolOperation(const FAvaEaseCurveTool::EOperation InNewOperation) const
+bool FAvaEaseCurveTool::IsToolOperation(const EAvaEaseCurveToolOperation InNewOperation) const
 {
 	return OperationMode == InNewOperation;
 }
@@ -517,7 +362,7 @@ bool FAvaEaseCurveTool::CanCopyTangentsToClipboard() const
 	return true;
 }
 
-void FAvaEaseCurveTool::CopyTangentsToClipboard()
+void FAvaEaseCurveTool::CopyTangentsToClipboard() const
 {
 	FPlatformApplicationMisc::ClipboardCopy(*EaseCurve->GetTangents().ToJson());
 
@@ -530,7 +375,7 @@ bool FAvaEaseCurveTool::CanPasteTangentsFromClipboard() const
 	return TangentsFromClipboardPaste(Tangents);
 }
 
-void FAvaEaseCurveTool::PasteTangentsFromClipboard()
+void FAvaEaseCurveTool::PasteTangentsFromClipboard() const
 {
 	FAvaEaseCurveTangents Tangents;
 	if (TangentsFromClipboardPaste(Tangents))
@@ -548,14 +393,14 @@ bool FAvaEaseCurveTool::TangentsFromClipboardPaste(FAvaEaseCurveTangents& OutTan
 	return FAvaEaseCurveTangents::FromString(ClipboardString, OutTangents);
 }
 
-bool FAvaEaseCurveTool::IsKeyInterpMode(const ERichCurveInterpMode InInterpMode, const ERichCurveTangentMode InTangentMode)
+bool FAvaEaseCurveTool::IsKeyInterpMode(const ERichCurveInterpMode InInterpMode, const ERichCurveTangentMode InTangentMode) const
 {
 	const FKeyHandle StartKeyHandle = EaseCurve->GetStartKeyHandle();
 	return (EaseCurve->FloatCurve.GetKeyInterpMode(StartKeyHandle) == InInterpMode
 		&& EaseCurve->FloatCurve.GetKeyTangentMode(StartKeyHandle) == InTangentMode);
 }
 
-void FAvaEaseCurveTool::SetKeyInterpMode(const ERichCurveInterpMode InInterpMode, const ERichCurveTangentMode InTangentMode)
+void FAvaEaseCurveTool::SetKeyInterpMode(const ERichCurveInterpMode InInterpMode, const ERichCurveTangentMode InTangentMode) const
 {
 	const FKeyHandle StartKeyHandle = EaseCurve->GetStartKeyHandle();
 
@@ -625,7 +470,7 @@ void FAvaEaseCurveTool::PostRedo(bool bInSuccess)
 
 void FAvaEaseCurveTool::OpenToolSettings() const
 {
-	if (ISettingsModule* SettingsModule = FModuleManager::GetModulePtr<ISettingsModule>("Settings"))
+	if (ISettingsModule* SettingsModule = FModuleManager::GetModulePtr<ISettingsModule>(TEXT("Settings")))
 	{
 		SettingsModule->ShowViewer(TEXT("Editor"), TEXT("Motion Design"), TEXT("Ease Curve Tool"));
 	}
@@ -633,7 +478,7 @@ void FAvaEaseCurveTool::OpenToolSettings() const
 
 FFrameRate FAvaEaseCurveTool::GetTickResolution() const
 {
-	if (TSharedPtr<FAvaSequencer> AvaSequencer = AvaSequencerWeak.Pin())
+	if (const TSharedPtr<FAvaSequencer> AvaSequencer = AvaSequencerWeak.Pin())
 	{
 		return AvaSequencer->GetSequencer()->GetFocusedTickResolution();
 	}
@@ -643,7 +488,7 @@ FFrameRate FAvaEaseCurveTool::GetTickResolution() const
 
 FFrameRate FAvaEaseCurveTool::GetDisplayRate() const
 {
-	if (TSharedPtr<FAvaSequencer> AvaSequencer = AvaSequencerWeak.Pin())
+	if (const TSharedPtr<FAvaSequencer> AvaSequencer = AvaSequencerWeak.Pin())
 	{
 		return AvaSequencer->GetSequencer()->GetFocusedDisplayRate();
 	}
@@ -659,110 +504,12 @@ void FAvaEaseCurveTool::ShowNotificationMessage(const FText& InMessageText)
 	FSlateNotificationManager::Get().AddNotification(Info);
 }
 
-void FAvaEaseCurveTool::SelectNextChannelKey()
-{
-	// This command is only available for single key selections
-	if (KeyCache.TotalSelectedKeys != 1 || !SequencerSelectionWeak.IsValid())
-	{
-		return;
-	}
-
-	const TSharedPtr<FSequencerSelection> SequencerSelection = SequencerSelectionWeak.Pin();
-
-	int32 SelectedKeyIndex = INDEX_NONE;
-	FKeyHandle SelectedKeyHandle;
-	FKeyHandle NextKeyHandle;
-	TViewModelPtr<FChannelModel> ChannelModel;
-
-	for (const TPair<FName, FKeyDataCache::FChannelData>& ChannelEntry : KeyCache.ChannelKeyData)
-	{
-		if (ChannelEntry.Value.KeyHandles.Num() == 1)
-		{
-			SelectedKeyHandle = ChannelEntry.Value.KeyHandles[0];
-
-			TMovieSceneChannelData<FMovieSceneDoubleValue> ChannelData = ChannelEntry.Value.DoubleChannel->GetData();
-			const TArrayView<const FFrameNumber> KeyTimes = ChannelData.GetTimes();
-
-			SelectedKeyIndex = ChannelData.GetIndex(SelectedKeyHandle);
-			if (KeyTimes.IsValidIndex(SelectedKeyIndex + 1))
-			{
-				NextKeyHandle = ChannelData.GetHandle(SelectedKeyIndex + 1);
-
-				ChannelModel = ChannelEntry.Value.ChannelModel;
-			}
-		}
-		break;
-	}
-
-	if (SelectedKeyIndex == INDEX_NONE
-		|| SelectedKeyHandle == FKeyHandle::Invalid()
-		|| NextKeyHandle == FKeyHandle::Invalid()
-		|| !ChannelModel.IsValid())
-	{
-		return;
-	}
-	
-	SequencerSelection->KeySelection.Deselect(SelectedKeyHandle);
-	SequencerSelection->KeySelection.Select(ChannelModel, NextKeyHandle);
-
-	UpdateEaseCurveFromSequencerKeySelections();
-}
-
-void FAvaEaseCurveTool::SelectPreviousChannelKey()
-{
-	// This command is only available for single key selections
-	if (KeyCache.TotalSelectedKeys != 1 || !SequencerSelectionWeak.IsValid())
-	{
-		return;
-	}
-
-	const TSharedPtr<FSequencerSelection> SequencerSelection = SequencerSelectionWeak.Pin();
-
-	int32 SelectedKeyIndex = INDEX_NONE;
-	FKeyHandle SelectedKeyHandle;
-	FKeyHandle PreviousKeyHandle;
-	TViewModelPtr<FChannelModel> ChannelModel;
-
-	for (const TPair<FName, FKeyDataCache::FChannelData>& ChannelEntry : KeyCache.ChannelKeyData)
-	{
-		if (ChannelEntry.Value.KeyHandles.Num() == 1)
-		{
-			SelectedKeyHandle = ChannelEntry.Value.KeyHandles[0];
-
-			TMovieSceneChannelData<FMovieSceneDoubleValue> ChannelData = ChannelEntry.Value.DoubleChannel->GetData();
-			const TArrayView<const FFrameNumber> KeyTimes = ChannelData.GetTimes();
-
-			SelectedKeyIndex = ChannelData.GetIndex(SelectedKeyHandle);
-			if (KeyTimes.IsValidIndex(SelectedKeyIndex - 1))
-			{
-				PreviousKeyHandle = ChannelData.GetHandle(SelectedKeyIndex - 1);
-
-				ChannelModel = ChannelEntry.Value.ChannelModel;
-			}
-		}
-		break;
-	}
-
-	if (SelectedKeyIndex == INDEX_NONE
-		|| SelectedKeyHandle == FKeyHandle::Invalid()
-		|| PreviousKeyHandle == FKeyHandle::Invalid()
-		|| !ChannelModel.IsValid())
-	{
-		return;
-	}
-
-	SequencerSelection->KeySelection.Deselect(SelectedKeyHandle);
-	SequencerSelection->KeySelection.Select(ChannelModel, PreviousKeyHandle);
-
-	UpdateEaseCurveFromSequencerKeySelections();
-}
-
 bool FAvaEaseCurveTool::HasCachedKeysToEase()
 {
 	bool bEaseableKeys = false;
 
 	KeyCache.ForEachEaseableKey(/*bInIncludeEqualValueKeys=*/false, [&bEaseableKeys](const FKeyHandle& InKeyHandle
-		, const FKeyHandle& InNextKeyHandle, const FKeyDataCache::FChannelData& InChannelData)
+		, const FKeyHandle& InNextKeyHandle, const FAvaEaseCurveKeySelection::FChannelData& InChannelData)
 		{
 			bEaseableKeys = true;
 			return false;
@@ -771,55 +518,4 @@ bool FAvaEaseCurveTool::HasCachedKeysToEase()
 	return bEaseableKeys;
 }
 
-void FAvaEaseCurveTool::FKeyDataCache::ForEachEaseableKey(const bool bInIncludeEqualValueKeys
-	, TFunctionRef<bool(const FKeyHandle& /*InKeyHandle*/, const FKeyHandle& /*InNextKeyHandle*/, const FChannelData&)> InCallable)
-{
-	for (const TPair<FName, FKeyDataCache::FChannelData>& ChannelEntry : ChannelKeyData)
-	{
-		TMovieSceneChannelData<FMovieSceneDoubleValue> ChannelData = ChannelEntry.Value.DoubleChannel->GetData();
-		const TArrayView<const FMovieSceneDoubleValue> ChannelValues = ChannelData.GetValues();
-		const int32 KeyCount = ChannelValues.Num();
-		
-		for (const FKeyHandle& KeyHandle : ChannelEntry.Value.KeyHandles)
-		{
-			if (KeyHandle == FKeyHandle::Invalid())
-			{
-				continue;
-			}
-
-			const int32 KeyIndex = ChannelData.GetIndex(KeyHandle);
-			if (KeyIndex == INDEX_NONE)
-			{
-				continue;
-			}
-			
-			// If there is no key after the selected key, we don't need to process.
-			// The arrive tangents of this key will be set by the previous key's processing.
-			int32 NextKeyIndex = KeyIndex + 1;
-			NextKeyIndex = (NextKeyIndex < KeyCount) ? NextKeyIndex : INDEX_NONE;
-			if (NextKeyIndex == INDEX_NONE)
-			{
-				continue;
-			}
-
-			// Need to check if the next key index is valid, otherwise GetHandle() will fail.
-			const FKeyHandle NextKeyHandle = ChannelData.GetHandle(NextKeyIndex);
-			if (NextKeyHandle == FKeyHandle::Invalid())
-			{
-				continue;
-			}
-
-			if (!bInIncludeEqualValueKeys && ChannelValues[KeyIndex].Value == ChannelValues[NextKeyIndex].Value)
-			{
-				continue;
-			}
-			
-			if (!InCallable(KeyHandle, NextKeyHandle, ChannelEntry.Value))
-			{
-				return;
-			}
-		}
-	}
-}
-
 #undef LOCTEXT_NAMESPACE
diff --git a/Engine/Plugins/Experimental/Avalanche/Source/AvalancheSequencer/Private/EaseCurveTool/AvaEaseCurveTool.h b/Engine/Plugins/Experimental/Avalanche/Source/AvalancheSequencer/Private/EaseCurveTool/AvaEaseCurveTool.h
--- a/Engine/Plugins/Experimental/Avalanche/Source/AvalancheSequencer/Private/EaseCurveTool/AvaEaseCurveTool.h
+++ b/Engine/Plugins/Experimental/Avalanche/Source/AvalancheSequencer/Private/EaseCurveTool/AvaEaseCurveTool.h
@@ -3,10 +3,10 @@
 #pragma once
 
 #include "EaseCurveTool/AvaEaseCurve.h"
+#include "EaseCurveTool/AvaEaseCurveKeySelection.h"
 #include "Curves/RichCurve.h"
 #include "Curves/KeyHandle.h"
 #include "EaseCurveTool/Widgets/SAvaEaseCurvePreset.h"
-#include "Editor/Transactor.h"
 #include "EditorUndoClient.h"
 #include "ScopedTransaction.h"
 #include "Templates/SharedPointer.h"
@@ -26,7 +26,7 @@ class UAvaEaseCurve;
 class UCurveFloat;
 class UMovieSceneSection;
 struct FAvaEaseCurveTangents;
-struct FMovieSceneDoubleChannel;
+struct FGuid;
 
 namespace UE::Sequencer
 {
@@ -34,30 +34,30 @@ namespace UE::Sequencer
 	class FSequencerSelection;
 }
 
+/** Current default and only implemented is DualKeyEdit. */
+enum class EAvaEaseCurveToolMode : uint8
+{
+	/** Edits the selected key's leave tangent and the next key's arrive tangent in the curve editor graph. */
+	DualKeyEdit,
+	/** Edits only the selected key.
+	 * The leave tangent in the curve editor graph will set the sequence key arrive tangent.
+	 * The arrive tangent in the curve editor graph will set the sequence key leave tangent. */
+	SingleKeyEdit
+};
+
+enum class EAvaEaseCurveToolOperation : uint8
+{
+	InOut,
+	In,
+	Out
+};
+
 class FAvaEaseCurveTool
 	: public TSharedFromThis<FAvaEaseCurveTool>
 	, public FGCObject
 	, public FSelfRegisteringEditorUndoClient
 {
 public:
-	/** Current default and only implemented is DualKeyEdit. */
-	enum class EMode : uint8
-	{
-		/** Edits the selected key's leave tangent and the next key's arrive tangent in the curve editor graph. */
-		DualKeyEdit,
-		/** Edits only the selected key.
-		 * The leave tangent in the curve editor graph will set the sequence key arrive tangent.
-		 * The arrive tangent in the curve editor graph will set the sequence key leave tangent. */
-		SingleKeyEdit
-	};
-
-	enum class EOperation : uint8
-	{
-		InOut,
-		In,
-		Out
-	};
-
 	static void ShowNotificationMessage(const FText& InMessageText);
 
 	/** Returns true if the clipboard paste data contains tangent information. */
@@ -65,6 +65,8 @@ public:
 
 	FAvaEaseCurveTool(const TSharedRef<FAvaSequencer>& InSequencer);
 
+	virtual ~FAvaEaseCurveTool() override;
+
 	TSharedRef<SWidget> GenerateWidget();
 
 	EVisibility GetVisibility() const;
@@ -79,33 +81,33 @@ public:
 	 * This is different from SetEaseCurveTangents_Internal in that it performs undo/redo transactions and
 	 * optionally sets the selected tangents in the actual sequence.
 	 */
-	void SetEaseCurveTangents(const FAvaEaseCurveTangents& InTangents, const EOperation InOperation, const bool bInBroadcastUpdate, const bool bInSetSequencerTangents);
+	void SetEaseCurveTangents(const FAvaEaseCurveTangents& InTangents, const EAvaEaseCurveToolOperation InOperation, const bool bInBroadcastUpdate, const bool bInSetSequencerTangents);
 
-	void ResetEaseCurveTangents(const EOperation InOperation);
+	void ResetEaseCurveTangents(const EAvaEaseCurveToolOperation InOperation);
 
-	void FlattenOrStraightenTangents(const EOperation InOperation, const bool bInFlattenTangents) const;
+	void FlattenOrStraightenTangents(const EAvaEaseCurveToolOperation InOperation, const bool bInFlattenTangents) const;
 
 	/** Creates a new external float curve from the internet curve editor curve. */
 	UCurveBase* CreateCurveAsset() const;
 
-	void SetSequencerKeySelectionTangents(const FAvaEaseCurveTangents& InTangents, const EOperation InOperation = EOperation::InOut);
+	void SetSequencerKeySelectionTangents(const FAvaEaseCurveTangents& InTangents, const EAvaEaseCurveToolOperation InOperation = EAvaEaseCurveToolOperation::InOut);
 
-	void ApplyQuickEaseToSequencerKeySelections(const EOperation InOperation = EOperation::InOut);
+	void ApplyQuickEaseToSequencerKeySelections(const EAvaEaseCurveToolOperation InOperation = EAvaEaseCurveToolOperation::InOut);
 
 	/** Updates the ease curve graph view based on the active sequencer key selection. */
 	void UpdateEaseCurveFromSequencerKeySelections();
 
-	EOperation GetToolOperation() const;
-	void SetToolOperation(const EOperation InNewOperation);
-	bool IsToolOperation(const EOperation InNewOperation) const;
+	EAvaEaseCurveToolOperation GetToolOperation() const;
+	void SetToolOperation(const EAvaEaseCurveToolOperation InNewOperation);
+	bool IsToolOperation(const EAvaEaseCurveToolOperation InNewOperation) const;
 
 	bool CanCopyTangentsToClipboard() const;
-	void CopyTangentsToClipboard();
+	void CopyTangentsToClipboard() const;
 	bool CanPasteTangentsFromClipboard() const;
-	void PasteTangentsFromClipboard();
+	void PasteTangentsFromClipboard() const;
 
-	bool IsKeyInterpMode(const ERichCurveInterpMode InInterpMode, const ERichCurveTangentMode InTangentMode);
-	void SetKeyInterpMode(const ERichCurveInterpMode InInterpMode, const ERichCurveTangentMode InTangentMode);
+	bool IsKeyInterpMode(const ERichCurveInterpMode InInterpMode, const ERichCurveTangentMode InTangentMode) const;
+	void SetKeyInterpMode(const ERichCurveInterpMode InInterpMode, const ERichCurveTangentMode InTangentMode) const;
 
 	void BeginTransaction(const FText& InDescription) const;
 	void EndTransaction() const;
@@ -118,9 +120,6 @@ public:
 	FFrameRate GetTickResolution() const;
 	FFrameRate GetDisplayRate() const;
 
-	void SelectNextChannelKey();
-	void SelectPreviousChannelKey();
-
 	bool HasCachedKeysToEase();
 
 	//~ Begin FGCObject
@@ -129,8 +128,8 @@ public:
 	//~ End FGCObject
 
 	//~ Begin FEditorUndoClient
-	virtual void PostUndo(bool bInSuccess);
-	virtual void PostRedo(bool bInSuccess);
+	virtual void PostUndo(bool bInSuccess) override;
+	virtual void PostRedo(bool bInSuccess) override;
 	//~ End FEditorUndoClient
 
 protected:
@@ -138,44 +137,21 @@ protected:
 	 * Sets the internal ease curve tangents and optionally broadcasts a change event for the curve object.
 	 * Changing the internal ease curve tangents will be directly reflected in the ease curve editor graph.
 	 */
-	void SetEaseCurveTangents_Internal(const FAvaEaseCurveTangents& InTangents, const EOperation InOperation, const bool bInBroadcastUpdate);
+	void SetEaseCurveTangents_Internal(const FAvaEaseCurveTangents& InTangents, const EAvaEaseCurveToolOperation InOperation, const bool bInBroadcastUpdate) const;
 
 	TWeakPtr<FAvaSequencer> AvaSequencerWeak;
 	TWeakPtr<UE::Sequencer::FSequencerSelection> SequencerSelectionWeak;
 
 	TObjectPtr<UAvaEaseCurve> EaseCurve;
 
-	EMode ToolMode = EMode::DualKeyEdit;
-	EOperation OperationMode = EOperation::InOut;
+	EAvaEaseCurveToolMode ToolMode = EAvaEaseCurveToolMode::DualKeyEdit;
+	EAvaEaseCurveToolOperation OperationMode = EAvaEaseCurveToolOperation::InOut;
 
 	TSharedPtr<SAvaEaseCurveTool> ToolWidget;
 	
 	/** Cached data set when a new sequencer selection is made. */
-	struct FKeyDataCache
-	{
-		struct FChannelData
-		{
-			TSharedPtr<UE::Sequencer::FChannelModel> ChannelModel;
-			UMovieSceneSection* Section = nullptr;
-			FMovieSceneDoubleChannel* DoubleChannel = nullptr;
-			TArray<FKeyHandle> KeyHandles;
-		};
-
-		void ForEachEaseableKey(const bool bInIncludeEqualValueKeys
-			, TFunctionRef<bool(const FKeyHandle& /*InKeyHandle*/, const FKeyHandle& /*InNextKeyHandle*/, const FChannelData&)> InCallable);
-
-		TMap<FName, FChannelData> ChannelKeyData;
-
-		int32 TotalSelectedKeys = 0;
-
-		/** Indicates only one selected key on each of the selected channels. */
-		bool bAllChannelSingleKeySelections = true;
-
-		/** Indicates only one key selected and it is the last key of the channel. */
-		bool bIsLastOnlySelectedKey = false;
-	};
-	FKeyDataCache KeyCache;
+	FAvaEaseCurveKeySelection KeyCache;
 
 private:
-	void CacheSelectionData();
+	void OnSequencerSelectionChanged(TArray<FGuid> InObjectGuids);
 };
diff --git a/Engine/Plugins/Experimental/Avalanche/Source/AvalancheSequencer/Private/EaseCurveTool/AvaEaseCurveToolCommands.cpp b/Engine/Plugins/Experimental/Avalanche/Source/AvalancheSequencer/Private/EaseCurveTool/AvaEaseCurveToolCommands.cpp
--- a/Engine/Plugins/Experimental/Avalanche/Source/AvalancheSequencer/Private/EaseCurveTool/AvaEaseCurveToolCommands.cpp
+++ b/Engine/Plugins/Experimental/Avalanche/Source/AvalancheSequencer/Private/EaseCurveTool/AvaEaseCurveToolCommands.cpp
@@ -219,24 +219,6 @@ void FAvaEaseCurveToolCommands::RegisterCommands()
 	UI_COMMAND(SetKeyInterpCubicBreak, "Break", "Cubic interpolation - User broken tangents", EUserInterfaceActionType::ToggleButton, FInputChord(EKeys::Three));
 	UI_COMMAND(SetKeyInterpToggleWeighted, "Weighted Tangents", "Toggle weighted tangents for cubic interpolation modes. Only supported on some curve types", EUserInterfaceActionType::ToggleButton, FInputChord(EModifierKey::Control, EKeys::W));
 
-	FUICommandInfo::MakeCommandInfo(SharedThis
-		, SelectNextChannelKey
-		, TEXT("SelectNextChannelKey")
-		, LOCTEXT("SelectNextChannelKey_Label", "Select Next Channel Key")
-		, LOCTEXT("SelectNextChannelKey_ToolTip", "Select the next key on the same channel.")
-		, FSlateIcon(FAppStyle::Get().GetStyleSetName(), TEXT("Wizard.NextIcon"))
-		, EUserInterfaceActionType::Button
-		, FInputChord(EKeys::Period, EModifierKey::Control));
-
-	FUICommandInfo::MakeCommandInfo(SharedThis
-		, SelectPreviousChannelKey
-		, TEXT("SelectsPreviousChannelKey")
-		, LOCTEXT("SelectsPreviousChannelKey_Label", "Selects Previous Channel Key")
-		, LOCTEXT("SelectsPreviousChannelKey_ToolTip", "Select the previous key on the same channel.")
-		, FSlateIcon(FAppStyle::Get().GetStyleSetName(), TEXT("Wizard.BackIcon"))
-		, EUserInterfaceActionType::Button
-		, FInputChord(EKeys::Comma, EModifierKey::Control));
-
 	FUICommandInfo::MakeCommandInfo(SharedThis
 		, QuickEase
 		, TEXT("QuickEase")
diff --git a/Engine/Plugins/Experimental/Avalanche/Source/AvalancheSequencer/Private/EaseCurveTool/AvaEaseCurveToolCommands.h b/Engine/Plugins/Experimental/Avalanche/Source/AvalancheSequencer/Private/EaseCurveTool/AvaEaseCurveToolCommands.h
--- a/Engine/Plugins/Experimental/Avalanche/Source/AvalancheSequencer/Private/EaseCurveTool/AvaEaseCurveToolCommands.h
+++ b/Engine/Plugins/Experimental/Avalanche/Source/AvalancheSequencer/Private/EaseCurveTool/AvaEaseCurveToolCommands.h
@@ -62,9 +62,6 @@ public:
 	TSharedPtr<FUICommandInfo> SetKeyInterpCubicBreak;
 	TSharedPtr<FUICommandInfo> SetKeyInterpToggleWeighted;
 
-	TSharedPtr<FUICommandInfo> SelectNextChannelKey;
-	TSharedPtr<FUICommandInfo> SelectPreviousChannelKey;
-
 	TSharedPtr<FUICommandInfo> QuickEase;
 	TSharedPtr<FUICommandInfo> QuickEaseIn;
 	TSharedPtr<FUICommandInfo> QuickEaseOut;
diff --git a/Engine/Plugins/Experimental/Avalanche/Source/AvalancheSequencer/Private/EaseCurveTool/Widgets/AvaEaseCurveToolContextMenu.cpp b/Engine/Plugins/Experimental/Avalanche/Source/AvalancheSequencer/Private/EaseCurveTool/Widgets/AvaEaseCurveToolContextMenu.cpp
new file mode 100644
--- /dev/null
+++ b/Engine/Plugins/Experimental/Avalanche/Source/AvalancheSequencer/Private/EaseCurveTool/Widgets/AvaEaseCurveToolContextMenu.cpp
@@ -0,0 +1,360 @@
+﻿// Copyright Epic Games, Inc. All Rights Reserved.
+
+#include "AvaEaseCurveToolContextMenu.h"
+#include "EaseCurveTool/AvaEaseCurveToolCommands.h"
+#include "EaseCurveTool/AvaEaseCurveToolSettings.h"
+#include "ToolMenu.h"
+#include "ToolMenus.h"
+#include "Widgets/Input/SNumericEntryBox.h"
+
+#define LOCTEXT_NAMESPACE "AvaEaseCurveToolContextMenu"
+
+namespace UE::Ava::EaseCurveTool::Private
+{
+	template<typename NumericType>
+	void ExtractNumericMetadata(FProperty* InProperty
+		, TOptional<NumericType>& OutMinValue
+		, TOptional<NumericType>& OutMaxValue
+		, TOptional<NumericType>& OutSliderMinValue
+		, TOptional<NumericType>& OutSliderMaxValue
+		, NumericType& OutSliderExponent
+		, NumericType& OutDelta
+		, float& OutShiftMultiplier
+		, float& OutCtrlMultiplier
+		, bool& OutSupportDynamicSliderMaxValue
+		, bool& OutSupportDynamicSliderMinValue)
+	{
+		const FString& MetaUIMinString = InProperty->GetMetaData(TEXT("UIMin"));
+		const FString& MetaUIMaxString = InProperty->GetMetaData(TEXT("UIMax"));
+		const FString& SliderExponentString = InProperty->GetMetaData(TEXT("SliderExponent"));
+		const FString& DeltaString = InProperty->GetMetaData(TEXT("Delta"));
+		const FString& ShiftMultiplierString = InProperty->GetMetaData(TEXT("ShiftMultiplier"));
+		const FString& CtrlMultiplierString = InProperty->GetMetaData(TEXT("CtrlMultiplier"));
+		const FString& SupportDynamicSliderMaxValueString = InProperty->GetMetaData(TEXT("SupportDynamicSliderMaxValue"));
+		const FString& SupportDynamicSliderMinValueString = InProperty->GetMetaData(TEXT("SupportDynamicSliderMinValue"));
+		const FString& ClampMinString = InProperty->GetMetaData(TEXT("ClampMin"));
+		const FString& ClampMaxString = InProperty->GetMetaData(TEXT("ClampMax"));
+
+		// If no UIMin/Max was specified then use the clamp string
+		const FString& UIMinString = MetaUIMinString.Len() ? MetaUIMinString : ClampMinString;
+		const FString& UIMaxString = MetaUIMaxString.Len() ? MetaUIMaxString : ClampMaxString;
+
+		NumericType ClampMin = TNumericLimits<NumericType>::Lowest();
+		NumericType ClampMax = TNumericLimits<NumericType>::Max();
+
+		if (!ClampMinString.IsEmpty())
+		{
+			TTypeFromString<NumericType>::FromString(ClampMin, *ClampMinString);
+		}
+
+		if (!ClampMaxString.IsEmpty())
+		{
+			TTypeFromString<NumericType>::FromString(ClampMax, *ClampMaxString);
+		}
+
+		NumericType UIMin = TNumericLimits<NumericType>::Lowest();
+		NumericType UIMax = TNumericLimits<NumericType>::Max();
+		TTypeFromString<NumericType>::FromString(UIMin, *UIMinString);
+		TTypeFromString<NumericType>::FromString(UIMax, *UIMaxString);
+
+		OutSliderExponent = NumericType(1);
+
+		if (SliderExponentString.Len())
+		{
+			TTypeFromString<NumericType>::FromString(OutSliderExponent, *SliderExponentString);
+		}
+
+		OutDelta = NumericType(0);
+
+		if (DeltaString.Len())
+		{
+			TTypeFromString<NumericType>::FromString(OutDelta, *DeltaString);
+		}
+
+		OutShiftMultiplier = 10.f;
+		if (ShiftMultiplierString.Len())
+		{
+			TTypeFromString<float>::FromString(OutShiftMultiplier, *ShiftMultiplierString);
+		}
+
+		OutCtrlMultiplier = 0.1f;
+		if (CtrlMultiplierString.Len())
+		{
+			TTypeFromString<float>::FromString(OutCtrlMultiplier, *CtrlMultiplierString);
+		}
+
+		const NumericType ActualUIMin = FMath::Max(UIMin, ClampMin);
+		const NumericType ActualUIMax = FMath::Min(UIMax, ClampMax);
+
+		OutMinValue = ClampMinString.Len() ? ClampMin : TOptional<NumericType>();
+		OutMaxValue = ClampMaxString.Len() ? ClampMax : TOptional<NumericType>();
+		OutSliderMinValue = (UIMinString.Len()) ? ActualUIMin : TOptional<NumericType>();
+		OutSliderMaxValue = (UIMaxString.Len()) ? ActualUIMax : TOptional<NumericType>();
+
+		OutSupportDynamicSliderMaxValue = SupportDynamicSliderMaxValueString.Len() > 0 && SupportDynamicSliderMaxValueString.ToBool();
+		OutSupportDynamicSliderMinValue = SupportDynamicSliderMinValueString.Len() > 0 && SupportDynamicSliderMinValueString.ToBool();
+	}
+}
+
+FAvaEaseCurveToolContextMenu::FAvaEaseCurveToolContextMenu(const TWeakPtr<FUICommandList>& InCommandListWeak, const FAvaEaseCurveToolOnGraphSizeChanged& InOnGraphSizeChanged)
+	: CommandListWeak(InCommandListWeak), OnGraphSizeChanged(InOnGraphSizeChanged)
+{
+	const UAvaEaseCurveToolSettings* const EaseCurveToolSettings = GetDefault<UAvaEaseCurveToolSettings>();
+	check(IsValid(EaseCurveToolSettings));
+	GraphSize = EaseCurveToolSettings->GetGraphSize();
+}
+
+TSharedRef<SWidget> FAvaEaseCurveToolContextMenu::GenerateWidget()
+{
+	UToolMenus* const ToolMenus = UToolMenus::Get();
+	if (!ToolMenus)
+	{
+		return SNullWidget::NullWidget;
+	}
+
+	constexpr const TCHAR* MenuName = TEXT("AvaEaseCurveToolMenu");
+
+	if (!ToolMenus->IsMenuRegistered(MenuName))
+	{
+		UToolMenu* const ToolMenu = ToolMenus->RegisterMenu(MenuName, NAME_None, EMultiBoxType::Menu);
+
+		FToolMenuSection& Section = ToolMenu->FindOrAddSection(TEXT("EaseCurveTool"), LOCTEXT("EaseCurveToolActions", "Curve Ease Tool Actions"));
+
+		const FAvaEaseCurveToolCommands& EaseCurveToolCommands = FAvaEaseCurveToolCommands::Get();
+
+		Section.AddSubMenu(TEXT("Settings"),
+			LOCTEXT("SettingsSubMenuLabel", "Settings"),
+			LOCTEXT("SettingsSubMenuToolTip", ""),
+			FNewToolMenuDelegate::CreateSP(this, &FAvaEaseCurveToolContextMenu::PopulateContextMenuSettings),
+			false,
+			FSlateIcon(FAppStyle::Get().GetStyleSetName(), TEXT("Icons.Toolbar.Settings")));
+
+		Section.AddSeparator(NAME_None);
+
+		Section.AddMenuEntry(EaseCurveToolCommands.CreateExternalCurveAsset);
+
+		Section.AddSeparator(NAME_None);
+
+		Section.AddMenuEntry(EaseCurveToolCommands.CopyTangents);
+
+		Section.AddMenuEntry(EaseCurveToolCommands.PasteTangents);
+
+		Section.AddSeparator(NAME_None);
+
+		Section.AddSubMenu(TEXT("StraightenTangents"),
+			LOCTEXT("StraightenTangentsSubMenuLabel", "Straighten Tangents"),
+			LOCTEXT("StraightenTangentsSubMenuToolTip", ""),
+			FNewToolMenuDelegate::CreateLambda([&EaseCurveToolCommands](UToolMenu* InToolMenu)
+				{
+					FToolMenuSection& NewSection = InToolMenu->FindOrAddSection(TEXT("StraightenTangents"));
+					NewSection.AddMenuEntry(EaseCurveToolCommands.StraightenTangents);
+					NewSection.AddSeparator(NAME_None);
+					NewSection.AddMenuEntry(EaseCurveToolCommands.StraightenStartTangent);
+					NewSection.AddMenuEntry(EaseCurveToolCommands.StraightenEndTangent);
+				}),
+			false, FSlateIcon(FAppStyle::Get().GetStyleSetName(), TEXT("GenericCurveEditor.StraightenTangents")));
+
+		Section.AddSubMenu(TEXT("FlattenTangents"),
+			LOCTEXT("FlattenTangentsSubMenuLabel", "Flatten Tangents"),
+			LOCTEXT("FlattenTangentsSubMenuToolTip", ""),
+			FNewToolMenuDelegate::CreateLambda([&EaseCurveToolCommands](UToolMenu* InToolMenu)
+				{
+					FToolMenuSection& NewSection = InToolMenu->FindOrAddSection(TEXT("FlattenTangents"));
+					NewSection.AddMenuEntry(EaseCurveToolCommands.FlattenTangents);
+					NewSection.AddSeparator(NAME_None);
+					NewSection.AddMenuEntry(EaseCurveToolCommands.FlattenStartTangent);
+					NewSection.AddMenuEntry(EaseCurveToolCommands.FlattenEndTangent);
+				}), 
+			false, FSlateIcon(FAppStyle::Get().GetStyleSetName(), TEXT("GenericCurveEditor.FlattenTangents")));
+
+		Section.AddSubMenu(TEXT("ResetTangents"),
+			LOCTEXT("ResetTangentsSubMenuLabel", "Reset Tangents"),
+			LOCTEXT("ResetTangentsSubMenuToolTip", ""),
+			FNewToolMenuDelegate::CreateLambda([&EaseCurveToolCommands](UToolMenu* InToolMenu)
+				{
+					FToolMenuSection& NewSection = InToolMenu->FindOrAddSection(TEXT("ResetTangents"));
+					NewSection.AddMenuEntry(EaseCurveToolCommands.ResetTangents);
+					NewSection.AddSeparator(NAME_None);
+					NewSection.AddMenuEntry(EaseCurveToolCommands.ResetStartTangent);
+					NewSection.AddMenuEntry(EaseCurveToolCommands.ResetEndTangent);
+				}),
+			false, FSlateIcon(FAppStyle::Get().GetStyleSetName(), TEXT("PropertyWindow.DiffersFromDefault")));
+
+		Section.AddSeparator(NAME_None);
+
+		/** @TODO: Only show these in single edit mode(?)
+		Section.AddMenuEntry(EaseCurveToolCommands.SetKeyInterpConstant);
+
+		Section.AddMenuEntry(EaseCurveToolCommands.SetKeyInterpLinear);
+
+		Section.AddMenuEntry(EaseCurveToolCommands.SetKeyInterpCubicAuto);
+
+		Section.AddMenuEntry(EaseCurveToolCommands.SetKeyInterpCubicSmartAuto);
+
+		Section.AddMenuEntry(EaseCurveToolCommands.SetKeyInterpCubicUser);
+
+		Section.AddMenuEntry(EaseCurveToolCommands.SetKeyInterpCubicBreak);
+
+		Section.AddSeparator(NAME_None);*/
+
+		Section.AddMenuEntry(EaseCurveToolCommands.SetOperationToEaseOut);
+
+		Section.AddMenuEntry(EaseCurveToolCommands.SetOperationToEaseInOut);
+
+		Section.AddMenuEntry(EaseCurveToolCommands.SetOperationToEaseIn);
+
+		Section.AddSeparator(NAME_None);
+
+		Section.AddMenuEntry(EaseCurveToolCommands.ToggleGridSnap);
+
+		Section.AddMenuEntry(EaseCurveToolCommands.ZoomToFit);
+
+		Section.AddSeparator(NAME_None);
+
+		Section.AddMenuEntry(EaseCurveToolCommands.Refresh);
+
+		Section.AddMenuEntry(EaseCurveToolCommands.Apply);
+	}
+
+	return ToolMenus->GenerateWidget(MenuName, FToolMenuContext(CommandListWeak.Pin()));
+}
+
+void FAvaEaseCurveToolContextMenu::PopulateContextMenuSettings(UToolMenu* const InToolMenu)
+{
+	if (!IsValid(InToolMenu))
+	{
+		return;
+	}
+
+	const FAvaEaseCurveToolCommands& EaseCurveToolCommands = FAvaEaseCurveToolCommands::Get();
+
+	FToolMenuSection& Section = InToolMenu->FindOrAddSection(TEXT("EaseCurveToolSettings"), LOCTEXT("EaseCurveToolSettingsActions", "Settings"));
+
+	Section.AddMenuEntry(EaseCurveToolCommands.OpenToolSettings);
+
+	Section.AddSeparator(NAME_None);
+
+	Section.AddMenuEntry(EaseCurveToolCommands.ToggleAutoFlipTangents);
+
+	Section.AddSeparator(NAME_None);
+
+	// Graph Size
+	{
+		FProperty* const GraphSizeProperty = UAvaEaseCurveToolSettings::StaticClass()->FindPropertyByName(TEXT("GraphSize"));
+		check(GraphSizeProperty);
+
+		TOptional<int32> MinValue, MaxValue, SliderMinValue, SliderMaxValue;
+		int32 SliderExponent, Delta;
+		float ShiftMultiplier = 10.f;
+		float CtrlMultiplier = 0.1f;
+		bool SupportDynamicSliderMaxValue = false;
+		bool SupportDynamicSliderMinValue = false;
+		UE::Ava::EaseCurveTool::Private::ExtractNumericMetadata(GraphSizeProperty
+			, MinValue, MaxValue
+			, SliderMinValue, SliderMaxValue
+			, SliderExponent, Delta
+			, ShiftMultiplier
+			, CtrlMultiplier
+			, SupportDynamicSliderMaxValue, SupportDynamicSliderMinValue);
+
+		const TSharedRef<SNumericEntryBox<int32>> GraphSizeWidget = SNew(SNumericEntryBox<int32>)
+			.Font(FAppStyle::GetFontStyle(TEXT("MenuItem.Font")))
+			.AllowSpin(true)
+			.MinValue(MinValue)
+			.MaxValue(MaxValue)
+			.MinSliderValue(SliderMinValue)
+			.MaxSliderValue(SliderMaxValue)
+			.SliderExponent(SliderExponent)
+			.Delta(Delta)
+			.ShiftMultiplier(ShiftMultiplier)
+			.CtrlMultiplier(CtrlMultiplier)
+			.SupportDynamicSliderMaxValue(SupportDynamicSliderMaxValue)
+			.SupportDynamicSliderMinValue(SupportDynamicSliderMinValue)
+			.Value_Lambda([this]()
+				{
+					return GraphSize;
+				})
+			.OnValueChanged_Lambda([this](const int32 InNewValue)
+				{
+					GraphSize = InNewValue;
+				
+					OnGraphSizeChanged.ExecuteIfBound(InNewValue);
+				})
+			.OnValueCommitted_Lambda([this](const int32 InNewValue, ETextCommit::Type InCommitType)
+				{
+					GraphSize = InNewValue;
+				
+					OnGraphSizeChanged.ExecuteIfBound(InNewValue);
+
+					UAvaEaseCurveToolSettings* const EaseCurveToolSettings = GetMutableDefault<UAvaEaseCurveToolSettings>();
+					check(IsValid(EaseCurveToolSettings));
+					EaseCurveToolSettings->SetGraphSize(InNewValue);
+					EaseCurveToolSettings->SaveConfig();
+				});
+
+		Section.AddEntry(FToolMenuEntry::InitWidget(TEXT("ToolSize"), GraphSizeWidget, LOCTEXT("ToolSizeLabel", "Tool Size")));
+	}
+
+	// Grid Size
+	{
+		FProperty* const GridSizeProperty = UAvaEaseCurveToolSettings::StaticClass()->FindPropertyByName(TEXT("GridSize"));
+		check(GridSizeProperty);
+
+		TOptional<int32> MinValue, MaxValue, SliderMinValue, SliderMaxValue;
+		int32 SliderExponent, Delta;
+		float ShiftMultiplier = 10.f;
+		float CtrlMultiplier = 0.1f;
+		bool SupportDynamicSliderMaxValue = false;
+		bool SupportDynamicSliderMinValue = false;
+		UE::Ava::EaseCurveTool::Private::ExtractNumericMetadata(GridSizeProperty
+			, MinValue, MaxValue
+			, SliderMinValue, SliderMaxValue
+			, SliderExponent, Delta
+			, ShiftMultiplier
+			, CtrlMultiplier
+			, SupportDynamicSliderMaxValue, SupportDynamicSliderMinValue);
+
+		auto SetEaseCurveToolGridSize = [](const int32 InNewValue)
+			{
+				UAvaEaseCurveToolSettings* const EaseCurveToolSettings = GetMutableDefault<UAvaEaseCurveToolSettings>();
+				check(IsValid(EaseCurveToolSettings));
+				EaseCurveToolSettings->SetGridSize(InNewValue);
+				EaseCurveToolSettings->SaveConfig();
+			};
+
+		const TSharedRef<SNumericEntryBox<int32>> GridSizeWidget = SNew(SNumericEntryBox<int32>)
+			.Font(FAppStyle::GetFontStyle(TEXT("MenuItem.Font")))
+			.AllowSpin(true)
+			.MinValue(MinValue)
+			.MaxValue(MaxValue)
+			.MinSliderValue(SliderMinValue)
+			.MaxSliderValue(SliderMaxValue)
+			.SliderExponent(SliderExponent)
+			.Delta(Delta)
+			.ShiftMultiplier(ShiftMultiplier)
+			.CtrlMultiplier(CtrlMultiplier)
+			.SupportDynamicSliderMaxValue(SupportDynamicSliderMaxValue)
+			.SupportDynamicSliderMinValue(SupportDynamicSliderMinValue)
+			.Value_Lambda([this]()
+				{
+					return GetDefault<UAvaEaseCurveToolSettings>()->GetGridSize();
+				})
+			.OnValueChanged_Lambda(SetEaseCurveToolGridSize)
+			.OnValueCommitted_Lambda([SetEaseCurveToolGridSize](const int32 InNewValue, ETextCommit::Type InCommitType)
+				{
+					SetEaseCurveToolGridSize(InNewValue);
+				});
+
+		Section.AddEntry(FToolMenuEntry::InitWidget(TEXT("GridSize"), GridSizeWidget, LOCTEXT("GridSizeLabel", "Grid Size")));
+	}
+
+	Section.AddSeparator(NAME_None);
+
+	Section.AddMenuEntry(EaseCurveToolCommands.ToggleAutoZoomToFit);
+
+	Section.AddSeparator(NAME_None);
+
+	Section.AddMenuEntry(EaseCurveToolCommands.ResetToDefaultPresets);
+}
+
+#undef LOCTEXT_NAMESPACE
diff --git a/Engine/Plugins/Experimental/Avalanche/Source/AvalancheSequencer/Private/EaseCurveTool/Widgets/AvaEaseCurveToolContextMenu.h b/Engine/Plugins/Experimental/Avalanche/Source/AvalancheSequencer/Private/EaseCurveTool/Widgets/AvaEaseCurveToolContextMenu.h
new file mode 100644
--- /dev/null
+++ b/Engine/Plugins/Experimental/Avalanche/Source/AvalancheSequencer/Private/EaseCurveTool/Widgets/AvaEaseCurveToolContextMenu.h
@@ -0,0 +1,29 @@
+﻿// Copyright Epic Games, Inc. All Rights Reserved.
+
+#pragma once
+
+#include "Delegates/Delegate.h"
+#include "Templates/SharedPointer.h"
+
+class FUICommandList;
+class SWidget;
+class UToolMenu;
+
+DECLARE_DELEGATE_OneParam(FAvaEaseCurveToolOnGraphSizeChanged, const int32 /*InNewSize*/)
+
+class FAvaEaseCurveToolContextMenu : public TSharedFromThis<FAvaEaseCurveToolContextMenu>
+{
+public:
+	FAvaEaseCurveToolContextMenu(const TWeakPtr<FUICommandList>& InCommandListWeak, const FAvaEaseCurveToolOnGraphSizeChanged& InOnGraphSizeChanged);
+
+	TSharedRef<SWidget> GenerateWidget();
+
+protected:
+	void PopulateContextMenuSettings(UToolMenu* const InToolMenu);
+	
+	TWeakPtr<FUICommandList> CommandListWeak;
+
+	FAvaEaseCurveToolOnGraphSizeChanged OnGraphSizeChanged;
+	
+	int32 GraphSize;
+};
diff --git a/Engine/Plugins/Experimental/Avalanche/Source/AvalancheSequencer/Private/EaseCurveTool/Widgets/SAvaEaseCurveEditor.cpp b/Engine/Plugins/Experimental/Avalanche/Source/AvalancheSequencer/Private/EaseCurveTool/Widgets/SAvaEaseCurveEditor.cpp
--- a/Engine/Plugins/Experimental/Avalanche/Source/AvalancheSequencer/Private/EaseCurveTool/Widgets/SAvaEaseCurveEditor.cpp
+++ b/Engine/Plugins/Experimental/Avalanche/Source/AvalancheSequencer/Private/EaseCurveTool/Widgets/SAvaEaseCurveEditor.cpp
@@ -186,7 +186,7 @@ int32 SAvaEaseCurveEditor::OnPaint(const FPaintArgs& InArgs, const FGeometry& In
 		if (StartText.IsSet() || EndText.IsSet())
 		{
 			const FSlateFontInfo FontInfo = FAvaEaseCurveStyle::Get().GetFontStyle(TEXT("Editor.LabelFont"));
-			const FLinearColor TextColor = FLinearColor(0.1f, 0.1f, 0.1f, 1.f);
+			constexpr FLinearColor TextColor = FLinearColor(0.1f, 0.1f, 0.1f, 1.f);
 			const FVector2D TextOffset = FVector2D(4.f);
 
 			if (StartText.IsSet())
@@ -454,21 +454,21 @@ void SAvaEaseCurveEditor::CreateLinesForSegment(const ERichCurveInterpMode InInt
 			const float CurveOut = EaseCurve->FloatCurve.Eval(CurveIn);
 			OutLinePoints.Add(FVector2D(CurveIn, CurveOut));
 			
-			switch (Operation.Get(FAvaEaseCurveTool::EOperation::InOut))
+			switch (Operation.Get(EAvaEaseCurveToolOperation::InOut))
 			{
-			case FAvaEaseCurveTool::EOperation::Out:
+			case EAvaEaseCurveToolOperation::Out:
 			{
 				const float Alpha = CurrentX / CurveLengthX;
 				OutLineColors.Add(FLinearColor::LerpUsingHSV(FadedCurveColor, CurveColor,  1.f - Alpha));
 				break;
 			}
-			case FAvaEaseCurveTool::EOperation::In:
+			case EAvaEaseCurveToolOperation::In:
 			{
 				const float Alpha = CurrentX / CurveLengthX;
 				OutLineColors.Add(FLinearColor::LerpUsingHSV(FadedCurveColor, CurveColor, Alpha));
 				break;
 			}
-			case FAvaEaseCurveTool::EOperation::InOut:
+			case EAvaEaseCurveToolOperation::InOut:
 			default:
 			{
 				OutLineColors.Add(CurveColor);
diff --git a/Engine/Plugins/Experimental/Avalanche/Source/AvalancheSequencer/Private/EaseCurveTool/Widgets/SAvaEaseCurveEditor.h b/Engine/Plugins/Experimental/Avalanche/Source/AvalancheSequencer/Private/EaseCurveTool/Widgets/SAvaEaseCurveEditor.h
--- a/Engine/Plugins/Experimental/Avalanche/Source/AvalancheSequencer/Private/EaseCurveTool/Widgets/SAvaEaseCurveEditor.h
+++ b/Engine/Plugins/Experimental/Avalanche/Source/AvalancheSequencer/Private/EaseCurveTool/Widgets/SAvaEaseCurveEditor.h
@@ -66,7 +66,7 @@ public:
 		SLATE_ARGUMENT(FLinearColor, ExtendedGridColor)
 		SLATE_ARGUMENT(float, CurveThickness)
 		SLATE_ARGUMENT(FLinearColor, CurveColor)
-		SLATE_ATTRIBUTE(FAvaEaseCurveTool::EOperation, Operation)
+		SLATE_ATTRIBUTE(EAvaEaseCurveToolOperation, Operation)
 		SLATE_ATTRIBUTE(bool, ShowEqualValueKeyError)
 
 		SLATE_ATTRIBUTE(FText, StartText)
@@ -246,7 +246,7 @@ protected:
 	FLinearColor ExtendedGridColor;
 	float CurveThickness = 1.f;
 	FLinearColor CurveColor;
-	TAttribute<FAvaEaseCurveTool::EOperation> Operation;
+	TAttribute<EAvaEaseCurveToolOperation> Operation;
 	TAttribute<bool> ShowEqualValueKeyError;
 
 	TAttribute<FText> StartText;
diff --git a/Engine/Plugins/Experimental/Avalanche/Source/AvalancheSequencer/Private/EaseCurveTool/Widgets/SAvaEaseCurvePreset.cpp b/Engine/Plugins/Experimental/Avalanche/Source/AvalancheSequencer/Private/EaseCurveTool/Widgets/SAvaEaseCurvePreset.cpp
--- a/Engine/Plugins/Experimental/Avalanche/Source/AvalancheSequencer/Private/EaseCurveTool/Widgets/SAvaEaseCurvePreset.cpp
+++ b/Engine/Plugins/Experimental/Avalanche/Source/AvalancheSequencer/Private/EaseCurveTool/Widgets/SAvaEaseCurvePreset.cpp
@@ -4,10 +4,7 @@
 #include "EaseCurveTool/AvaEaseCurvePreset.h"
 #include "EaseCurveTool/AvaEaseCurveStyle.h"
 #include "EaseCurveTool/AvaEaseCurveSubsystem.h"
-#include "EaseCurveTool/AvaEaseCurveToolSettings.h"
-#include "EaseCurveTool/Widgets/SAvaEaseCurvePreview.h"
 #include "Editor.h"
-#include "DetailLayoutBuilder.h"
 #include "Framework/SlateDelegates.h"
 #include "Internationalization/Text.h"
 #include "Styling/AppStyle.h"
@@ -15,7 +12,7 @@
 #include "Widgets/Images/SImage.h"
 #include "Widgets/Input/SButton.h"
 #include "Widgets/Input/SCheckBox.h"
-#include "Widgets/Input/SSearchBox.h"
+#include "Widgets/Input/SEditableTextBox.h"
 #include "Widgets/Layout/SScrollBox.h"
 #include "Widgets/Layout/SUniformWrapPanel.h"
 #include "Widgets/Layout/SWidgetSwitcher.h"
@@ -40,9 +37,12 @@ void SAvaEaseCurvePreset::Construct(const FArguments& InArgs)
 			[
 				SNew(SHorizontalBox)
 				+ SHorizontalBox::Slot()
-				.Expose(ComboBoxSlot)
 				[
-					ConstructPresetComboBox()
+					SAssignNew(PresetComboBox, SAvaEaseCurvePresetComboBox)
+					.DisplayRate(DisplayRate)
+					.AllowEditMode(true)
+					.OnPresetChanged(OnPresetChanged)
+					.OnQuickPresetChanged(OnQuickPresetChanged)
 				]
 				+ SHorizontalBox::Slot()
 				.AutoWidth()
@@ -55,7 +55,7 @@ void SAvaEaseCurvePreset::Construct(const FArguments& InArgs)
 					.ToolTipText(LOCTEXT("AddNewPresetToolTip", "Save the current ease curve as a new preset"))
 					.IsEnabled_Lambda([this]()
 						{
-							return !SelectedItem.IsValid();
+							return !PresetComboBox->HasSelection();
 						})
 					.OnClicked(this, &SAvaEaseCurvePreset::OnCreateNewPresetClick)
 					[
@@ -84,7 +84,7 @@ void SAvaEaseCurvePreset::Construct(const FArguments& InArgs)
 					.ToolTipText(LOCTEXT("CancelNewPresetToolTip", "Cancels the current new ease curve preset operation"))
 					.IsEnabled_Lambda([this]()
 						{
-							return !SelectedItem.IsValid();
+							return !PresetComboBox->HasSelection();
 						})
 					.OnClicked(this, &SAvaEaseCurvePreset::OnCancelNewPresetClick)
 					[
@@ -97,281 +97,6 @@ void SAvaEaseCurvePreset::Construct(const FArguments& InArgs)
 		];
 }
 
-TSharedRef<SWidget> SAvaEaseCurvePreset::ConstructPresetComboBox()
-{
-	return SNew(SComboButton)
-		.OnGetMenuContent(this, &SAvaEaseCurvePreset::GeneratePresetDropdown)
-		.OnMenuOpenChanged_Lambda([this](const bool bInOpening)
-			{
-				bIsInEditMode.Set(false);
-			})
-		.ButtonContent()
-		[
-			GenerateSelectedRowWidget()
-		];
-}
-
-TSharedRef<SWidget> SAvaEaseCurvePreset::GeneratePresetDropdown()
-{
-	static const FVector2D ButtonImageSize = FVector2D(FAvaEaseCurveStyle::Get().GetFloat(TEXT("ToolButton.ImageSize")));
-
-	TSharedRef<SWidget> OutWidget = SNew(SVerticalBox)
-		+ SVerticalBox::Slot()
-		.AutoHeight()
-		.Padding(3.f)
-		[
-			SNew(SHorizontalBox)
-			+ SHorizontalBox::Slot()
-			.VAlign(VAlign_Center)
-			[
-				SNew(SSearchBox)
-				.HintText(LOCTEXT("SearchHintLabel", "Search"))
-				.OnTextChanged(this, &SAvaEaseCurvePreset::OnSearchTextChanged)
-			]
-			+ SHorizontalBox::Slot()
-			.AutoWidth()
-			.VAlign(VAlign_Center)
-			.Padding(3.f, 0.f, 0.f, 0.f)
-			[
-				SNew(SButton)
-				.ButtonStyle(FAvaEaseCurveStyle::Get(), TEXT("ToolButton"))
-				.VAlign(VAlign_Center)
-				.ToolTipText(LOCTEXT("ReloadJsonPresetsToolTip", "Reload ease curve presets from Json files"))
-				.OnClicked(this, &SAvaEaseCurvePreset::ReloadJsonPresets)
-				[
-					SNew(SImage)
-					.DesiredSizeOverride(ButtonImageSize)
-					.Image(FAppStyle::GetBrush(TEXT("Icons.Refresh")))
-				]
-			]
-			+ SHorizontalBox::Slot()
-			.AutoWidth()
-			.VAlign(VAlign_Center)
-			.Padding(3.f, 0.f, 0.f, 0.f)
-			[
-				SNew(SButton)
-				.ButtonStyle(FAvaEaseCurveStyle::Get(), TEXT("ToolButton"))
-				.VAlign(VAlign_Center)
-				.ToolTipText(LOCTEXT("ExploreJsonPresetsFolderToolTip", "Opens the folder location for the Json ease curve presets"))
-				.OnClicked(this, &SAvaEaseCurvePreset::ExploreJsonPresetsFolder)
-				[
-					SNew(SImage)
-					.DesiredSizeOverride(ButtonImageSize)
-					.Image(FAppStyle::GetBrush(TEXT("Icons.FolderOpen")))
-				]
-			]
-			+ SHorizontalBox::Slot()
-			.AutoWidth()
-			.VAlign(VAlign_Center)
-			.Padding(3.f, 0.f, 3.f, 0.f)
-			[
-				SNew(SCheckBox)
-				.Style(FAvaEaseCurveStyle::Get(), TEXT("ToolToggleButton"))
-				.Padding(4.f)
-				.ToolTipText(LOCTEXT("ToggleEditModeToolTip", "Enable editing of ease curve presets and categories"))
-				.IsChecked_Lambda([this]()
-					{
-						return bIsInEditMode.Get(false) ? ECheckBoxState::Checked : ECheckBoxState::Unchecked;
-					})
-				.OnCheckStateChanged(this, &SAvaEaseCurvePreset::ToggleEditMode)
-				[
-					SNew(SImage)
-					.DesiredSizeOverride(ButtonImageSize)
-					.Image(FAppStyle::GetBrush(TEXT("Icons.Edit")))
-				]
-			]
-			+ SHorizontalBox::Slot()
-			.AutoWidth()
-			.VAlign(VAlign_Center)
-			[
-				SNew(SButton)
-				.ButtonStyle(FAvaEaseCurveStyle::Get(), TEXT("ToolButton"))
-				.VAlign(VAlign_Center)
-				.ToolTipText(LOCTEXT("CreateCategoryToolTip", "Creates a new empty category"))
-				.Visibility_Lambda([this]()
-					{
-						return bIsInEditMode.Get(false) ? EVisibility::Visible : EVisibility::Collapsed;
-					})
-				.OnClicked(this, &SAvaEaseCurvePreset::CreateNewCategory)
-				[
-					SNew(SImage)
-					.DesiredSizeOverride(ButtonImageSize)
-					.Image(FAppStyle::GetBrush(TEXT("Icons.Plus")))
-				]
-			]
-		]
-		+ SVerticalBox::Slot()
-		.AutoHeight()
-		.HAlign(HAlign_Fill)
-		.Padding(3.f, 0.f, 3.f, 3.f)
-		[
-			SNew(SBox)
-			.MaxDesiredHeight(960.f)
-			[
-				SNew(SScrollBox)
-				+ SScrollBox::Slot()
-				[
-					SAssignNew(GroupWidgetsParent, SBox)
-				]
-			]
-		];
-
-	RegenerateGroupWrapBox();
-
-	return OutWidget;
-}
-
-void SAvaEaseCurvePreset::UpdateGroupsContent()
-{
-	const UAvaEaseCurveSubsystem& EaseCurveSubsystem = UAvaEaseCurveSubsystem::Get();
-	const TArray<FString>& EaseCurveCategories = EaseCurveSubsystem.GetEaseCurveCategories();
-
-	auto GenerateNoPresetsWidget = [](const FText& InText)
-		{
-			return SNew(SBox)
-				.WidthOverride(300.f)
-				.HeightOverride(200.f)
-				.HAlign(HAlign_Center)
-				.VAlign(VAlign_Center)
-				[
-					SNew(STextBlock)
-					.TextStyle(FAppStyle::Get(), TEXT("HintText"))
-					.Text(InText)
-				];
-		};
-
-	if (EaseCurveCategories.Num() == 0)
-	{
-		GroupWidgetsParent->SetContent(GenerateNoPresetsWidget(LOCTEXT("NoPresetsLabel", "No ease curve presets")));
-		return;
-	}
-
-	if (!SearchText.IsEmpty())
-	{
-		int32 TotalVisiblePresets = 0;
-		for (const TSharedPtr<SAvaEaseCurvePresetGroup>& GroupWidget : GroupWidgets)
-		{
-			TotalVisiblePresets += GroupWidget->GetVisiblePresetCount();
-		}
-		if (TotalVisiblePresets == 0)
-		{
-			GroupWidgetsParent->SetContent(GenerateNoPresetsWidget(LOCTEXT("NoPresetsFoundLabel", "No ease curve presets found")));
-			return;
-		}
-	}
-
-	GroupWidgetsParent->SetContent(GroupWrapBox.ToSharedRef());
-}
-
-void SAvaEaseCurvePreset::RegenerateGroupWrapBox()
-{
-	const UAvaEaseCurveSubsystem& EaseCurveSubsystem = UAvaEaseCurveSubsystem::Get();
-	const TArray<FString>& EaseCurveCategories = EaseCurveSubsystem.GetEaseCurveCategories();
-	const int32 CurvePresetCount = EaseCurveCategories.Num();
-
-	GroupWidgets.Empty(CurvePresetCount);
-
-	GroupWrapBox = SNew(SUniformWrapPanel)
-		.HAlign(HAlign_Center)
-		.SlotPadding(FMargin(2.f, 1.f))
-		.EvenRowDistribution(true)
-		.NumColumnsOverride_Lambda([]()
-			{
-				const int32 EaseCurveCategoryCount = UAvaEaseCurveSubsystem::Get().GetEaseCurveCategories().Num();
-				return FMath::Min(5, EaseCurveCategoryCount);
-			});
-
-	for (const FString& Category : EaseCurveCategories)
-	{
-		TSharedRef<SAvaEaseCurvePresetGroup> NewGroupWidget = SNew(SAvaEaseCurvePresetGroup)
-			.CategoryName(Category)
-			.Presets(EaseCurveSubsystem.GetEaseCurvePresets(Category))
-			.SelectedPreset(SelectedItem)
-			.IsEditMode(bIsInEditMode)
-			.DisplayRate(DisplayRate.Get())
-			.OnCategoryDelete(this, &SAvaEaseCurvePreset::HandleCategoryDelete)
-			.OnCategoryRename(this, &SAvaEaseCurvePreset::HandleCategoryRename)
-			.OnPresetDelete(this, &SAvaEaseCurvePreset::HandlePresetDelete)
-			.OnPresetRename(this, &SAvaEaseCurvePreset::HandlePresetRename)
-			.OnBeginPresetMove(this, &SAvaEaseCurvePreset::HandleBeginPresetMove)
-			.OnEndPresetMove(this, &SAvaEaseCurvePreset::HandleEndPresetMove)
-			.OnPresetClick(this, &SAvaEaseCurvePreset::HandlePresetClick);
-
-		GroupWidgets.Add(NewGroupWidget);
-
-		GroupWrapBox->AddSlot()
-			.HAlign(HAlign_Left)
-			[
-				NewGroupWidget
-			];
-	}
-
-	UpdateGroupsContent();
-}
-
-TSharedRef<SWidget> SAvaEaseCurvePreset::GenerateSelectedRowWidget() const
-{
-	if (!SelectedItem.IsValid())
-	{
-		return SNew(STextBlock)
-			.Font(IDetailLayoutBuilder::GetDetailFont())
-			.Text(FText::FromString("Select Preset..."));
-	}
-
-	return SNew(SHorizontalBox)
-		+ SHorizontalBox::Slot()
-		.AutoWidth()
-		.HAlign(HAlign_Right)
-		.VAlign(VAlign_Center)
-		.Padding(0.f, 2.f, 5.f, 2.f)
-		[
-			SNew(SBorder)
-			.BorderBackgroundColor(FStyleColors::White25)
-			[
-				SNew(SAvaEaseCurvePreview)
-				.PreviewSize(12.f)
-				.CustomToolTip(true)
-				.DisplayRate(DisplayRate.Get())
-				.Tangents_Lambda([this]()
-					{
-						return SelectedItem.IsValid() ? SelectedItem->Tangents : FAvaEaseCurveTangents();
-					})
-			]
-		]
-		+ SHorizontalBox::Slot()
-		.HAlign(HAlign_Fill)
-		.VAlign(VAlign_Center)
-		.Padding(0.f, 0.f, 5.f, 0.f)
-		[
-			SNew(SHorizontalBox)
-			+ SHorizontalBox::Slot()
-			.HAlign(HAlign_Left)
-			[
-				SNew(STextBlock)
-				.Font(IDetailLayoutBuilder::GetDetailFont())
-				.ColorAndOpacity(FStyleColors::Foreground)
-				.OverflowPolicy(ETextOverflowPolicy::Ellipsis)
-				.Text_Lambda([this]()
-					{
-						return SelectedItem.IsValid() ? FText::FromString(SelectedItem->Name) : FText();
-					})
-			]
-			+ SHorizontalBox::Slot()
-			.HAlign(HAlign_Right)
-			[
-				SNew(STextBlock)
-				.Font(IDetailLayoutBuilder::GetDetailFont())
-				.ColorAndOpacity(FStyleColors::White25)
-				.OverflowPolicy(ETextOverflowPolicy::Ellipsis)
-				.Text_Lambda([this]()
-					{
-						return SelectedItem.IsValid() ? FText::FromString(SelectedItem->Category) : FText();
-					})
-			]
-
-		];
-}
-
 FReply SAvaEaseCurvePreset::OnCreateNewPresetClick()
 {
 	bIsCreatingNewPreset = true;
@@ -389,21 +114,6 @@ FReply SAvaEaseCurvePreset::OnCancelNewPresetClick()
 	return FReply::Handled();
 }
 
-FReply SAvaEaseCurvePreset::OnDeletePresetClick()
-{
-	if (SelectedItem.IsValid())
-	{
-		if (UAvaEaseCurveSubsystem* EaseCurveSubsystem = GEditor->GetEditorSubsystem<UAvaEaseCurveSubsystem>())
-		{
-			EaseCurveSubsystem->RemovePreset(*SelectedItem);
-		}
-
-		ClearSelection();
-	}
-
-	return FReply::Handled();
-}
-
 FReply SAvaEaseCurvePreset::OnNewPresetKeyDownHandler(const FGeometry& InGeometry, const FKeyEvent& InKeyEvent)
 {
 	if (InKeyEvent.GetKey() == EKeys::Escape)
@@ -426,9 +136,11 @@ void SAvaEaseCurvePreset::OnNewPresetTextCommitted(const FText& InNewText, EText
 		{
 			UAvaEaseCurveSubsystem& EaseCurveSubsystem = UAvaEaseCurveSubsystem::Get();
 
-			SelectedItem = EaseCurveSubsystem.AddPreset(InNewText.ToString(), NewTangents);
-
-			ComboBoxSlot->AttachWidget(ConstructPresetComboBox());
+			const TSharedPtr<FAvaEaseCurvePreset> NewPreset = EaseCurveSubsystem.AddPreset(InNewText.ToString(), NewTangents);
+			if (NewPreset.IsValid())
+			{
+				PresetComboBox->SetSelectedItem(*NewPreset);
+			}
 		}
 	}
 
@@ -436,188 +148,24 @@ void SAvaEaseCurvePreset::OnNewPresetTextCommitted(const FText& InNewText, EText
 	NewPresetNameTextBox->SetText(FText());
 }
 
-void SAvaEaseCurvePreset::OnSearchTextChanged(const FText& InSearchText)
-{
-	SearchText = InSearchText;
-
-	for (const TSharedPtr<SAvaEaseCurvePresetGroup>& GroupWidget : GroupWidgets)
-	{
-		GroupWidget->SetSearchText(SearchText);
-	}
-
-	UpdateGroupsContent();
-}
-
 void SAvaEaseCurvePreset::ClearSelection()
 {
-	SelectedItem.Reset();
-	ComboBoxSlot->AttachWidget(ConstructPresetComboBox());
-}
-
-bool SAvaEaseCurvePreset::GetSelectedItem(FAvaEaseCurvePreset& OutPreset) const
-{
-	if (!SelectedItem.IsValid())
-	{
-		return false;
-	}
-
-	OutPreset = *SelectedItem;
-
-	return true;
+	PresetComboBox->ClearSelection();
 }
 
 bool SAvaEaseCurvePreset::SetSelectedItem(const FString& InName)
 {
-	UAvaEaseCurveSubsystem& EaseCurveSubsystem = UAvaEaseCurveSubsystem::Get();
-
-	const TSharedPtr<FAvaEaseCurvePreset>& FoundItem = EaseCurveSubsystem.FindPreset(InName);
-	if (!FoundItem.IsValid())
-	{
-		return false;
-	}
-
-	SelectedItem = FoundItem;
-	ComboBoxSlot->AttachWidget(ConstructPresetComboBox());
-
-	return true;
-}
-
-bool SAvaEaseCurvePreset::SetSelectedItem(const FAvaEaseCurvePreset& InPreset)
-{
-	return SetSelectedItem(InPreset.Name);
+	return PresetComboBox->SetSelectedItem(InName);
 }
 
 bool SAvaEaseCurvePreset::SetSelectedItem(const FAvaEaseCurveTangents& InTangents)
 {
-	UAvaEaseCurveSubsystem& EaseCurveSubsystem = UAvaEaseCurveSubsystem::Get();
-
-	const TSharedPtr<FAvaEaseCurvePreset>& FoundItem = EaseCurveSubsystem.FindPresetByTangents(InTangents);
-	if (!FoundItem.IsValid())
-	{
-		return false;
-	}
-
-	SelectedItem = FoundItem;
-	ComboBoxSlot->AttachWidget(ConstructPresetComboBox());
-
-	return true;
-}
-
-FReply SAvaEaseCurvePreset::ReloadJsonPresets()
-{
-	UAvaEaseCurveSubsystem::Get().ReloadPresetsFromJson();
-
-	RegenerateGroupWrapBox();
-
-	return FReply::Handled();
-}
-
-FReply SAvaEaseCurvePreset::ExploreJsonPresetsFolder()
-{
-	UAvaEaseCurveSubsystem::Get().ExploreJsonPresetsFolder();
-
-	return FReply::Handled();
-}
-
-FReply SAvaEaseCurvePreset::CreateNewCategory()
-{
-	UAvaEaseCurveSubsystem::Get().AddNewPresetCategory();
-
-	ReloadJsonPresets();
-
-	return FReply::Handled();
-}
-
-void SAvaEaseCurvePreset::ToggleEditMode(const ECheckBoxState bInNewState)
-{
-	bIsInEditMode.Set(bInNewState == ECheckBoxState::Checked ? true : false);
-
-	RegenerateGroupWrapBox();
-}
-
-bool SAvaEaseCurvePreset::HandleCategoryDelete(const FString& InCategoryName)
-{
-	if (!UAvaEaseCurveSubsystem::Get().RemovePresetCategory(InCategoryName))
-	{
-		return false;
-	}
-
-	ReloadJsonPresets();
-
-	return true;
-}
-
-bool SAvaEaseCurvePreset::HandleCategoryRename(const FString& InCategoryName, const FString& InNewName)
-{
-	return UAvaEaseCurveSubsystem::Get().RenamePresetCategory(InCategoryName, InNewName);
-}
-
-bool SAvaEaseCurvePreset::HandlePresetDelete(const TSharedPtr<FAvaEaseCurvePreset>& InPreset)
-{
-	return UAvaEaseCurveSubsystem::Get().RemovePreset(*InPreset);
-}
-
-bool SAvaEaseCurvePreset::HandlePresetRename(const TSharedPtr<FAvaEaseCurvePreset>& InPreset, const FString& InNewName)
-{
-	return UAvaEaseCurveSubsystem::Get().RenamePreset(InPreset->Category, InPreset->Name, InNewName);
-}
-
-bool SAvaEaseCurvePreset::HandleBeginPresetMove(const TSharedPtr<FAvaEaseCurvePreset>& InPreset, const FString& InNewCategoryName)
-{
-	for (const TSharedPtr<SAvaEaseCurvePresetGroup>& GroupWidget : GroupWidgets)
-	{
-		if (!GroupWidget->GetCategoryName().Equals(InNewCategoryName))
-		{
-			GroupWidget->NotifyCanDrop(true);
-		}
-	}
-
-	return true;
+	return PresetComboBox->SetSelectedItem(InTangents);
 }
 
-bool SAvaEaseCurvePreset::HandleEndPresetMove(const TSharedPtr<FAvaEaseCurvePreset>& InPreset, const FString& InNewCategoryName)
-{
-	for (const TSharedPtr<SAvaEaseCurvePresetGroup>& GroupWidget : GroupWidgets)
-	{
-		if (!GroupWidget->GetCategoryName().Equals(InNewCategoryName))
-		{
-			GroupWidget->NotifyCanDrop(false);
-		}
-	}
-
-	if (!InPreset.IsValid() || InPreset->Category.Equals(InNewCategoryName))
-	{
-		return false;
-	}
-
-	if (!UAvaEaseCurveSubsystem::Get().ChangePresetCategory(InPreset, InNewCategoryName))
-	{
-		return false;
-	}
-
-	ReloadJsonPresets();
-
-	return true;
-}
-
-bool SAvaEaseCurvePreset::HandlePresetClick(const TSharedPtr<FAvaEaseCurvePreset>& InPreset, const FModifierKeysState& InModifierKeys)
+bool SAvaEaseCurvePreset::SetSelectedItem(const FAvaEaseCurvePreset& InPreset)
 {
-	if (InModifierKeys.IsShiftDown())
-	{
-		UAvaEaseCurveToolSettings* const EaseCurveToolSettings = GetMutableDefault<UAvaEaseCurveToolSettings>();
-		EaseCurveToolSettings->SetQuickEaseTangents(InPreset->Tangents.ToJson());
-		EaseCurveToolSettings->SaveConfig();
-
-		OnQuickPresetChanged.ExecuteIfBound(InPreset);
-
-		return true;
-	}
-
-	SetSelectedItem(*InPreset);
-
-	OnPresetChanged.ExecuteIfBound(InPreset);
-
-	return true;
+	return SetSelectedItem(InPreset.Name);
 }
 
 #undef LOCTEXT_NAMESPACE
diff --git a/Engine/Plugins/Experimental/Avalanche/Source/AvalancheSequencer/Private/EaseCurveTool/Widgets/SAvaEaseCurvePreset.h b/Engine/Plugins/Experimental/Avalanche/Source/AvalancheSequencer/Private/EaseCurveTool/Widgets/SAvaEaseCurvePreset.h
--- a/Engine/Plugins/Experimental/Avalanche/Source/AvalancheSequencer/Private/EaseCurveTool/Widgets/SAvaEaseCurvePreset.h
+++ b/Engine/Plugins/Experimental/Avalanche/Source/AvalancheSequencer/Private/EaseCurveTool/Widgets/SAvaEaseCurvePreset.h
@@ -3,6 +3,7 @@
 #pragma once
 
 #include "EaseCurveTool/AvaEaseCurveTangents.h"
+#include "EaseCurveTool/Widgets/SAvaEaseCurvePresetComboBox.h"
 #include "EaseCurveTool/Widgets/SAvaEaseCurvePresetGroup.h"
 #include "Templates/SharedPointer.h"
 #include "Widgets/Input/SComboBox.h"
@@ -14,7 +15,6 @@ class FText;
 class SUniformWrapPanel;
 struct FAvaEaseCurvePreset;
 
-DECLARE_DELEGATE_OneParam(FAvaOnPresetChanged, const TSharedPtr<FAvaEaseCurvePreset>&)
 DECLARE_DELEGATE_RetVal_OneParam(bool, FAvaOnGetNewPresetTangents, FAvaEaseCurveTangents& /*InTangents*/)
 
 class SAvaEaseCurvePreset : public SCompoundWidget
@@ -23,6 +23,7 @@ public:
 	SLATE_BEGIN_ARGS(SAvaEaseCurvePreset)
 		: _DisplayRate(FFrameRate(30, 1))
 	{}
+		/** Display rate used to draw the ease curve preview. */
 		SLATE_ATTRIBUTE(FFrameRate, DisplayRate)
 		SLATE_EVENT(FAvaOnPresetChanged, OnPresetChanged)
 		SLATE_EVENT(FAvaOnPresetChanged, OnQuickPresetChanged)
@@ -32,60 +33,25 @@ public:
 	void Construct(const FArguments& InArgs);
 
 	void ClearSelection();
-	bool GetSelectedItem(FAvaEaseCurvePreset& OutPreset) const;
+	
 	bool SetSelectedItem(const FString& InName);
-	bool SetSelectedItem(const FAvaEaseCurvePreset& InPreset);
 	bool SetSelectedItem(const FAvaEaseCurveTangents& InTangents);
+	bool SetSelectedItem(const FAvaEaseCurvePreset& InPreset);
 
 protected:
-	TSharedRef<SWidget> ConstructPresetComboBox();
-	TSharedRef<SWidget> GeneratePresetDropdown();
-	
-	TSharedRef<SWidget> GenerateSelectedRowWidget() const;
-
-	void UpdateGroupsContent();
-	void RegenerateGroupWrapBox();
-
 	FReply OnCreateNewPresetClick();
 	FReply OnCancelNewPresetClick();
-	FReply OnDeletePresetClick();
-
+	
 	FReply OnNewPresetKeyDownHandler(const FGeometry& InGeometry, const FKeyEvent& InKeyEvent);
 	void OnNewPresetTextCommitted(const FText& InNewText, ETextCommit::Type InCommitType);
 
-	void OnSearchTextChanged(const FText& InSearchText);
-
-	FReply ReloadJsonPresets();
-	FReply ExploreJsonPresetsFolder();
-	FReply CreateNewCategory();
-
-	void ToggleEditMode(const ECheckBoxState bInNewState);
-
-	bool HandleCategoryDelete(const FString& InCategoryName);
-	bool HandleCategoryRename(const FString& InCategoryName, const FString& InNewName);
-	bool HandlePresetDelete(const TSharedPtr<FAvaEaseCurvePreset>& InPreset);
-	bool HandlePresetRename(const TSharedPtr<FAvaEaseCurvePreset>& InPreset, const FString& InNewName);
-	bool HandleBeginPresetMove(const TSharedPtr<FAvaEaseCurvePreset>& InPreset, const FString& InNewCategoryName);
-	bool HandleEndPresetMove(const TSharedPtr<FAvaEaseCurvePreset>& InPreset, const FString& InNewCategoryName);
-	bool HandlePresetClick(const TSharedPtr<FAvaEaseCurvePreset>& InPreset, const FModifierKeysState& InModifierKeys);
-
 	TAttribute<FFrameRate> DisplayRate;
-
 	FAvaOnPresetChanged OnPresetChanged;
 	FAvaOnPresetChanged OnQuickPresetChanged;
 	FAvaOnGetNewPresetTangents OnGetNewPresetTangents;
-
-	TSharedPtr<SBox> GroupWidgetsParent;
-	TSharedPtr<SUniformWrapPanel> GroupWrapBox;
-	TArray<TSharedPtr<SAvaEaseCurvePresetGroup>> GroupWidgets;
-
-	SHorizontalBox::FSlot* ComboBoxSlot = nullptr;
+	
+	TSharedPtr<SAvaEaseCurvePresetComboBox> PresetComboBox;
 	TSharedPtr<SEditableTextBox> NewPresetNameTextBox;
 
 	bool bIsCreatingNewPreset = false;
-	TAttribute<bool> bIsInEditMode = false;
-
-	TSharedPtr<FAvaEaseCurvePreset> SelectedItem;
-
-	FText SearchText;
 };
diff --git a/Engine/Plugins/Experimental/Avalanche/Source/AvalancheSequencer/Private/EaseCurveTool/Widgets/SAvaEaseCurvePresetComboBox.cpp b/Engine/Plugins/Experimental/Avalanche/Source/AvalancheSequencer/Private/EaseCurveTool/Widgets/SAvaEaseCurvePresetComboBox.cpp
new file mode 100644
--- /dev/null
+++ b/Engine/Plugins/Experimental/Avalanche/Source/AvalancheSequencer/Private/EaseCurveTool/Widgets/SAvaEaseCurvePresetComboBox.cpp
@@ -0,0 +1,533 @@
+// Copyright Epic Games, Inc. All Rights Reserved.
+
+#include "SAvaEaseCurvePresetComboBox.h"
+#include "EaseCurveTool/AvaEaseCurvePreset.h"
+#include "EaseCurveTool/AvaEaseCurveStyle.h"
+#include "EaseCurveTool/AvaEaseCurveSubsystem.h"
+#include "EaseCurveTool/AvaEaseCurveToolSettings.h"
+#include "EaseCurveTool/Widgets/SAvaEaseCurvePreview.h"
+#include "Editor.h"
+#include "DetailLayoutBuilder.h"
+#include "Framework/SlateDelegates.h"
+#include "Internationalization/Text.h"
+#include "Styling/AppStyle.h"
+#include "Styling/StyleColors.h"
+#include "Widgets/Images/SImage.h"
+#include "Widgets/Input/SButton.h"
+#include "Widgets/Input/SCheckBox.h"
+#include "Widgets/Input/SSearchBox.h"
+#include "Widgets/Layout/SScrollBox.h"
+#include "Widgets/Layout/SUniformWrapPanel.h"
+
+#define LOCTEXT_NAMESPACE "SAvaEaseCurvePresetComboBox"
+
+void SAvaEaseCurvePresetComboBox::Construct(const FArguments& InArgs)
+{
+	DisplayRate = InArgs._DisplayRate;
+	bAllowEditMode = InArgs._AllowEditMode;
+	OnPresetChanged = InArgs._OnPresetChanged;
+	OnQuickPresetChanged = InArgs._OnQuickPresetChanged;
+	
+	ChildSlot
+		[
+			SNew(SComboButton)
+			.OnGetMenuContent(this, &SAvaEaseCurvePresetComboBox::GeneratePresetDropdown)
+			.OnMenuOpenChanged_Lambda([this](const bool bInOpening)
+				{
+					bEditMode.Set(false);
+				})
+			.ButtonContent()
+			[
+				SAssignNew(SelectedRowContainer, SBox)
+			]
+		];
+}
+
+TSharedRef<SWidget> SAvaEaseCurvePresetComboBox::GenerateSearchRowWidget()
+{
+	static const FVector2D ButtonImageSize = FVector2D(FAvaEaseCurveStyle::Get().GetFloat(TEXT("ToolButton.ImageSize")));
+	
+	TSharedRef<SHorizontalBox> RowWidget = SNew(SHorizontalBox)
+		+ SHorizontalBox::Slot()
+		.VAlign(VAlign_Center)
+		[
+			SNew(SSearchBox)
+			.HintText(LOCTEXT("SearchHintLabel", "Search"))
+			.OnTextChanged(this, &SAvaEaseCurvePresetComboBox::OnSearchTextChanged)
+		]
+		+ SHorizontalBox::Slot()
+		.AutoWidth()
+		.VAlign(VAlign_Center)
+		.Padding(3.f, 0.f, 0.f, 0.f)
+		[
+			SNew(SButton)
+			.ButtonStyle(FAvaEaseCurveStyle::Get(), TEXT("ToolButton"))
+			.VAlign(VAlign_Center)
+			.ToolTipText(LOCTEXT("ReloadJsonPresetsToolTip", "Reload ease curve presets from Json files"))
+			.OnClicked(this, &SAvaEaseCurvePresetComboBox::ReloadJsonPresets)
+			[
+				SNew(SImage)
+				.DesiredSizeOverride(ButtonImageSize)
+				.Image(FAppStyle::GetBrush(TEXT("Icons.Refresh")))
+			]
+		];
+	
+	if (bAllowEditMode)
+	{
+		RowWidget->AddSlot()
+			.AutoWidth()
+			.VAlign(VAlign_Center)
+			.Padding(3.f, 0.f, 0.f, 0.f)
+			[
+				SNew(SButton)
+				.ButtonStyle(FAvaEaseCurveStyle::Get(), TEXT("ToolButton"))
+				.VAlign(VAlign_Center)
+				.ToolTipText(LOCTEXT("ExploreJsonPresetsFolderToolTip", "Opens the folder location for the Json ease curve presets"))
+				.OnClicked(this, &SAvaEaseCurvePresetComboBox::ExploreJsonPresetsFolder)
+				[
+					SNew(SImage)
+					.DesiredSizeOverride(ButtonImageSize)
+					.Image(FAppStyle::GetBrush(TEXT("Icons.FolderOpen")))
+				]
+			];
+		
+		RowWidget->AddSlot()
+			.AutoWidth()
+			.VAlign(VAlign_Center)
+			.Padding(3.f, 0.f, 3.f, 0.f)
+			[
+				SNew(SCheckBox)
+				.Style(FAvaEaseCurveStyle::Get(), TEXT("ToolToggleButton"))
+				.Padding(4.f)
+				.ToolTipText(LOCTEXT("ToggleEditModeToolTip", "Enable editing of ease curve presets and categories"))
+				.IsChecked_Lambda([this]()
+					{
+						return bEditMode.Get(false) ? ECheckBoxState::Checked : ECheckBoxState::Unchecked;
+					})
+				.OnCheckStateChanged(this, &SAvaEaseCurvePresetComboBox::ToggleEditMode)
+				[
+					SNew(SImage)
+					.DesiredSizeOverride(ButtonImageSize)
+					.Image(FAppStyle::GetBrush(TEXT("Icons.Edit")))
+				]
+			];
+		
+		RowWidget->AddSlot()
+			.AutoWidth()
+			.VAlign(VAlign_Center)
+			[
+				SNew(SButton)
+				.ButtonStyle(FAvaEaseCurveStyle::Get(), TEXT("ToolButton"))
+				.VAlign(VAlign_Center)
+				.ToolTipText(LOCTEXT("CreateCategoryToolTip", "Creates a new empty category"))
+				.Visibility_Lambda([this]()
+					{
+						return bEditMode.Get(false) ? EVisibility::Visible : EVisibility::Collapsed;
+					})
+				.OnClicked(this, &SAvaEaseCurvePresetComboBox::CreateNewCategory)
+				[
+					SNew(SImage)
+					.DesiredSizeOverride(ButtonImageSize)
+					.Image(FAppStyle::GetBrush(TEXT("Icons.Plus")))
+				]
+			];
+	}
+
+	return RowWidget;
+}
+
+TSharedRef<SWidget> SAvaEaseCurvePresetComboBox::GeneratePresetDropdown()
+{
+	TSharedRef<SWidget> OutWidget = SNew(SVerticalBox)
+		+ SVerticalBox::Slot()
+		.AutoHeight()
+		.Padding(3.f)
+		[
+			GenerateSearchRowWidget()
+		]
+		+ SVerticalBox::Slot()
+		.AutoHeight()
+		.HAlign(HAlign_Fill)
+		.Padding(3.f, 0.f, 3.f, 3.f)
+		[
+			SNew(SBox)
+			.MaxDesiredHeight(960.f)
+			[
+				SNew(SScrollBox)
+				+ SScrollBox::Slot()
+				[
+					SAssignNew(GroupWidgetsParent, SBox)
+				]
+			]
+		];
+
+	RegenerateGroupWrapBox();
+
+	return OutWidget;
+}
+
+void SAvaEaseCurvePresetComboBox::UpdateGroupsContent()
+{
+	const UAvaEaseCurveSubsystem& EaseCurveSubsystem = UAvaEaseCurveSubsystem::Get();
+	const TArray<FString>& EaseCurveCategories = EaseCurveSubsystem.GetEaseCurveCategories();
+
+	auto GenerateNoPresetsWidget = [](const FText& InText)
+		{
+			return SNew(SBox)
+				.WidthOverride(300.f)
+				.HeightOverride(200.f)
+				.HAlign(HAlign_Center)
+				.VAlign(VAlign_Center)
+				[
+					SNew(STextBlock)
+					.TextStyle(FAppStyle::Get(), TEXT("HintText"))
+					.Text(InText)
+				];
+		};
+
+	if (EaseCurveCategories.Num() == 0)
+	{
+		GroupWidgetsParent->SetContent(GenerateNoPresetsWidget(LOCTEXT("NoPresetsLabel", "No ease curve presets")));
+		return;
+	}
+
+	if (!SearchText.IsEmpty())
+	{
+		int32 TotalVisiblePresets = 0;
+		for (const TSharedPtr<SAvaEaseCurvePresetGroup>& GroupWidget : GroupWidgets)
+		{
+			TotalVisiblePresets += GroupWidget->GetVisiblePresetCount();
+		}
+		if (TotalVisiblePresets == 0)
+		{
+			GroupWidgetsParent->SetContent(GenerateNoPresetsWidget(LOCTEXT("NoPresetsFoundLabel", "No ease curve presets found")));
+			return;
+		}
+	}
+
+	GroupWidgetsParent->SetContent(GroupWrapBox.ToSharedRef());
+}
+
+void SAvaEaseCurvePresetComboBox::RegenerateGroupWrapBox()
+{
+	const UAvaEaseCurveSubsystem& EaseCurveSubsystem = UAvaEaseCurveSubsystem::Get();
+	const TArray<FString>& EaseCurveCategories = EaseCurveSubsystem.GetEaseCurveCategories();
+	const int32 CurvePresetCount = EaseCurveCategories.Num();
+
+	GroupWidgets.Empty(CurvePresetCount);
+
+	GroupWrapBox = SNew(SUniformWrapPanel)
+		.HAlign(HAlign_Center)
+		.SlotPadding(FMargin(2.f, 1.f))
+		.EvenRowDistribution(true)
+		.NumColumnsOverride_Lambda([]()
+			{
+				const int32 EaseCurveCategoryCount = UAvaEaseCurveSubsystem::Get().GetEaseCurveCategories().Num();
+				return FMath::Min(5, EaseCurveCategoryCount);
+			});
+
+	for (const FString& Category : EaseCurveCategories)
+	{
+		TSharedRef<SAvaEaseCurvePresetGroup> NewGroupWidget = SNew(SAvaEaseCurvePresetGroup)
+			.CategoryName(Category)
+			.Presets(EaseCurveSubsystem.GetEaseCurvePresets(Category))
+			.SelectedPreset(SelectedItem)
+			.IsEditMode(bEditMode)
+			.DisplayRate(DisplayRate.Get())
+			.OnCategoryDelete(this, &SAvaEaseCurvePresetComboBox::HandleCategoryDelete)
+			.OnCategoryRename(this, &SAvaEaseCurvePresetComboBox::HandleCategoryRename)
+			.OnPresetDelete(this, &SAvaEaseCurvePresetComboBox::HandlePresetDelete)
+			.OnPresetRename(this, &SAvaEaseCurvePresetComboBox::HandlePresetRename)
+			.OnBeginPresetMove(this, &SAvaEaseCurvePresetComboBox::HandleBeginPresetMove)
+			.OnEndPresetMove(this, &SAvaEaseCurvePresetComboBox::HandleEndPresetMove)
+			.OnPresetClick(this, &SAvaEaseCurvePresetComboBox::HandlePresetClick)
+			.OnSetQuickEase(this, &SAvaEaseCurvePresetComboBox::HandleSetQuickEase);
+
+		GroupWidgets.Add(NewGroupWidget);
+
+		GroupWrapBox->AddSlot()
+			.HAlign(HAlign_Left)
+			[
+				NewGroupWidget
+			];
+	}
+
+	UpdateGroupsContent();
+}
+
+void SAvaEaseCurvePresetComboBox::GenerateSelectedRowWidget()
+{
+	TSharedPtr<SWidget> OutRowWidget;
+	
+	if (!SelectedItem.IsValid())
+	{
+		OutRowWidget = SNew(STextBlock)
+			.Font(IDetailLayoutBuilder::GetDetailFont())
+			.Text(FText::FromString("Select Preset..."));
+	}
+	else
+	{
+		OutRowWidget = SNew(SHorizontalBox)
+			+ SHorizontalBox::Slot()
+			.AutoWidth()
+			.HAlign(HAlign_Right)
+			.VAlign(VAlign_Center)
+			.Padding(0.f, 2.f, 5.f, 2.f)
+			[
+				SNew(SBorder)
+				.BorderBackgroundColor(FStyleColors::White25)
+				[
+					SNew(SAvaEaseCurvePreview)
+					.PreviewSize(12.f)
+					.CustomToolTip(true)
+					.DisplayRate(DisplayRate.Get())
+					.Tangents_Lambda([this]()
+						{
+							return SelectedItem.IsValid() ? SelectedItem->Tangents : FAvaEaseCurveTangents();
+						})
+				]
+			]
+			+ SHorizontalBox::Slot()
+			.HAlign(HAlign_Fill)
+			.VAlign(VAlign_Center)
+			.Padding(0.f, 0.f, 5.f, 0.f)
+			[
+				SNew(SHorizontalBox)
+				+ SHorizontalBox::Slot()
+				.HAlign(HAlign_Left)
+				[
+					SNew(STextBlock)
+					.Font(IDetailLayoutBuilder::GetDetailFont())
+					.ColorAndOpacity(FStyleColors::Foreground)
+					.OverflowPolicy(ETextOverflowPolicy::Ellipsis)
+					.Text_Lambda([this]()
+						{
+							return SelectedItem.IsValid() ? FText::FromString(SelectedItem->Name) : FText();
+						})
+				]
+				+ SHorizontalBox::Slot()
+				.HAlign(HAlign_Right)
+				[
+					SNew(STextBlock)
+					.Font(IDetailLayoutBuilder::GetDetailFont())
+					.ColorAndOpacity(FStyleColors::White25)
+					.OverflowPolicy(ETextOverflowPolicy::Ellipsis)
+					.Text_Lambda([this]()
+						{
+							return SelectedItem.IsValid() ? FText::FromString(SelectedItem->Category) : FText();
+						})
+				]
+			];
+	}
+
+	SelectedRowContainer->SetContent(OutRowWidget.ToSharedRef());
+}
+
+FReply SAvaEaseCurvePresetComboBox::OnDeletePresetClick()
+{
+	if (SelectedItem.IsValid())
+	{
+		if (UAvaEaseCurveSubsystem* EaseCurveSubsystem = GEditor->GetEditorSubsystem<UAvaEaseCurveSubsystem>())
+		{
+			EaseCurveSubsystem->RemovePreset(*SelectedItem);
+		}
+
+		ClearSelection();
+	}
+
+	return FReply::Handled();
+}
+
+void SAvaEaseCurvePresetComboBox::OnSearchTextChanged(const FText& InSearchText)
+{
+	SearchText = InSearchText;
+
+	for (const TSharedPtr<SAvaEaseCurvePresetGroup>& GroupWidget : GroupWidgets)
+	{
+		GroupWidget->SetSearchText(SearchText);
+	}
+
+	UpdateGroupsContent();
+}
+
+bool SAvaEaseCurvePresetComboBox::HasSelection() const
+{
+	return SelectedItem.IsValid();
+}
+
+void SAvaEaseCurvePresetComboBox::ClearSelection()
+{
+	SelectedItem.Reset();
+
+	GenerateSelectedRowWidget();
+}
+
+bool SAvaEaseCurvePresetComboBox::GetSelectedItem(FAvaEaseCurvePreset& OutPreset) const
+{
+	if (!SelectedItem.IsValid())
+	{
+		return false;
+	}
+
+	OutPreset = *SelectedItem;
+
+	return true;
+}
+
+bool SAvaEaseCurvePresetComboBox::SetSelectedItem(const FString& InName)
+{
+	UAvaEaseCurveSubsystem& EaseCurveSubsystem = UAvaEaseCurveSubsystem::Get();
+
+	const TSharedPtr<FAvaEaseCurvePreset>& FoundItem = EaseCurveSubsystem.FindPreset(InName);
+	if (!FoundItem.IsValid())
+	{
+		return false;
+	}
+
+	SelectedItem = FoundItem;
+
+	GenerateSelectedRowWidget();
+
+	return true;
+}
+
+bool SAvaEaseCurvePresetComboBox::SetSelectedItem(const FAvaEaseCurveTangents& InTangents)
+{
+	UAvaEaseCurveSubsystem& EaseCurveSubsystem = UAvaEaseCurveSubsystem::Get();
+
+	const TSharedPtr<FAvaEaseCurvePreset>& FoundItem = EaseCurveSubsystem.FindPresetByTangents(InTangents);
+	if (!FoundItem.IsValid())
+	{
+		return false;
+	}
+
+	SelectedItem = FoundItem;
+
+	GenerateSelectedRowWidget();
+	
+	return true;
+}
+
+bool SAvaEaseCurvePresetComboBox::SetSelectedItem(const FAvaEaseCurvePreset& InPreset)
+{
+	return SetSelectedItem(InPreset.Name);
+}
+
+FReply SAvaEaseCurvePresetComboBox::ReloadJsonPresets()
+{
+	UAvaEaseCurveSubsystem::Get().ReloadPresetsFromJson();
+
+	RegenerateGroupWrapBox();
+
+	return FReply::Handled();
+}
+
+FReply SAvaEaseCurvePresetComboBox::ExploreJsonPresetsFolder()
+{
+	UAvaEaseCurveSubsystem::Get().ExploreJsonPresetsFolder();
+
+	return FReply::Handled();
+}
+
+FReply SAvaEaseCurvePresetComboBox::CreateNewCategory()
+{
+	UAvaEaseCurveSubsystem::Get().AddNewPresetCategory();
+
+	ReloadJsonPresets();
+
+	return FReply::Handled();
+}
+
+void SAvaEaseCurvePresetComboBox::ToggleEditMode(const ECheckBoxState bInNewState)
+{
+	bEditMode.Set(bInNewState == ECheckBoxState::Checked ? true : false);
+
+	RegenerateGroupWrapBox();
+}
+
+bool SAvaEaseCurvePresetComboBox::HandleCategoryDelete(const FString& InCategoryName)
+{
+	if (!UAvaEaseCurveSubsystem::Get().RemovePresetCategory(InCategoryName))
+	{
+		return false;
+	}
+
+	ReloadJsonPresets();
+
+	return true;
+}
+
+bool SAvaEaseCurvePresetComboBox::HandleCategoryRename(const FString& InCategoryName, const FString& InNewName)
+{
+	return UAvaEaseCurveSubsystem::Get().RenamePresetCategory(InCategoryName, InNewName);
+}
+
+bool SAvaEaseCurvePresetComboBox::HandlePresetDelete(const TSharedPtr<FAvaEaseCurvePreset>& InPreset)
+{
+	return UAvaEaseCurveSubsystem::Get().RemovePreset(*InPreset);
+}
+
+bool SAvaEaseCurvePresetComboBox::HandlePresetRename(const TSharedPtr<FAvaEaseCurvePreset>& InPreset, const FString& InNewName)
+{
+	return UAvaEaseCurveSubsystem::Get().RenamePreset(InPreset->Category, InPreset->Name, InNewName);
+}
+
+bool SAvaEaseCurvePresetComboBox::HandleBeginPresetMove(const TSharedPtr<FAvaEaseCurvePreset>& InPreset, const FString& InNewCategoryName)
+{
+	for (const TSharedPtr<SAvaEaseCurvePresetGroup>& GroupWidget : GroupWidgets)
+	{
+		if (!GroupWidget->GetCategoryName().Equals(InNewCategoryName))
+		{
+			GroupWidget->NotifyCanDrop(true);
+		}
+	}
+
+	return true;
+}
+
+bool SAvaEaseCurvePresetComboBox::HandleEndPresetMove(const TSharedPtr<FAvaEaseCurvePreset>& InPreset, const FString& InNewCategoryName)
+{
+	for (const TSharedPtr<SAvaEaseCurvePresetGroup>& GroupWidget : GroupWidgets)
+	{
+		if (!GroupWidget->GetCategoryName().Equals(InNewCategoryName))
+		{
+			GroupWidget->NotifyCanDrop(false);
+		}
+	}
+
+	if (!InPreset.IsValid() || InPreset->Category.Equals(InNewCategoryName))
+	{
+		return false;
+	}
+
+	if (!UAvaEaseCurveSubsystem::Get().ChangePresetCategory(InPreset, InNewCategoryName))
+	{
+		return false;
+	}
+
+	ReloadJsonPresets();
+
+	return true;
+}
+
+bool SAvaEaseCurvePresetComboBox::HandlePresetClick(const TSharedPtr<FAvaEaseCurvePreset>& InPreset)
+{
+	SetSelectedItem(*InPreset);
+
+	OnPresetChanged.ExecuteIfBound(InPreset);
+
+	return true;
+}
+
+bool SAvaEaseCurvePresetComboBox::HandleSetQuickEase(const TSharedPtr<FAvaEaseCurvePreset>& InPreset)
+{
+	UAvaEaseCurveToolSettings* const EaseCurveToolSettings = GetMutableDefault<UAvaEaseCurveToolSettings>();
+	EaseCurveToolSettings->SetQuickEaseTangents(InPreset->Tangents.ToJson());
+	EaseCurveToolSettings->SaveConfig();
+
+	OnQuickPresetChanged.ExecuteIfBound(InPreset);
+
+	return true;
+}
+
+#undef LOCTEXT_NAMESPACE
diff --git a/Engine/Plugins/Experimental/Avalanche/Source/AvalancheSequencer/Private/EaseCurveTool/Widgets/SAvaEaseCurvePresetComboBox.h b/Engine/Plugins/Experimental/Avalanche/Source/AvalancheSequencer/Private/EaseCurveTool/Widgets/SAvaEaseCurvePresetComboBox.h
new file mode 100644
--- /dev/null
+++ b/Engine/Plugins/Experimental/Avalanche/Source/AvalancheSequencer/Private/EaseCurveTool/Widgets/SAvaEaseCurvePresetComboBox.h
@@ -0,0 +1,88 @@
+// Copyright Epic Games, Inc. All Rights Reserved.
+
+#pragma once
+
+#include "EaseCurveTool/AvaEaseCurveTangents.h"
+#include "EaseCurveTool/Widgets/SAvaEaseCurvePresetGroup.h"
+#include "Templates/SharedPointer.h"
+#include "Widgets/Input/SComboBox.h"
+#include "Widgets/SBoxPanel.h"
+#include "Widgets/SCompoundWidget.h"
+
+class FName;
+class FText;
+class SUniformWrapPanel;
+struct FAvaEaseCurvePreset;
+
+DECLARE_DELEGATE_OneParam(FAvaOnPresetChanged, const TSharedPtr<FAvaEaseCurvePreset>&)
+
+class SAvaEaseCurvePresetComboBox : public SCompoundWidget
+{
+public:
+	SLATE_BEGIN_ARGS(SAvaEaseCurvePresetComboBox)
+		: _DisplayRate(FFrameRate(30, 1))
+		, _AllowEditMode(false)
+	{}
+		/** Display rate used to draw the ease curve preview. */
+		SLATE_ATTRIBUTE(FFrameRate, DisplayRate)
+		SLATE_ARGUMENT(bool, AllowEditMode)
+		SLATE_EVENT(FAvaOnPresetChanged, OnPresetChanged)
+		SLATE_EVENT(FAvaOnPresetChanged, OnQuickPresetChanged)
+	SLATE_END_ARGS()
+
+	void Construct(const FArguments& InArgs);
+
+	bool HasSelection() const;
+	
+	void ClearSelection();
+	
+	bool GetSelectedItem(FAvaEaseCurvePreset& OutPreset) const;
+	
+	bool SetSelectedItem(const FString& InName);
+	bool SetSelectedItem(const FAvaEaseCurveTangents& InTangents);
+	bool SetSelectedItem(const FAvaEaseCurvePreset& InPreset);
+
+protected:
+	TSharedRef<SWidget> GenerateSearchRowWidget();
+	TSharedRef<SWidget> GeneratePresetDropdown();
+
+	void GenerateSelectedRowWidget();
+
+	void UpdateGroupsContent();
+	void RegenerateGroupWrapBox();
+	
+	FReply OnDeletePresetClick();
+
+	void OnSearchTextChanged(const FText& InSearchText);
+
+	FReply ReloadJsonPresets();
+	FReply ExploreJsonPresetsFolder();
+	FReply CreateNewCategory();
+
+	void ToggleEditMode(const ECheckBoxState bInNewState);
+
+	bool HandleCategoryDelete(const FString& InCategoryName);
+	bool HandleCategoryRename(const FString& InCategoryName, const FString& InNewName);
+	bool HandlePresetDelete(const TSharedPtr<FAvaEaseCurvePreset>& InPreset);
+	bool HandlePresetRename(const TSharedPtr<FAvaEaseCurvePreset>& InPreset, const FString& InNewName);
+	bool HandleBeginPresetMove(const TSharedPtr<FAvaEaseCurvePreset>& InPreset, const FString& InNewCategoryName);
+	bool HandleEndPresetMove(const TSharedPtr<FAvaEaseCurvePreset>& InPreset, const FString& InNewCategoryName);
+	bool HandlePresetClick(const TSharedPtr<FAvaEaseCurvePreset>& InPreset);
+	bool HandleSetQuickEase(const TSharedPtr<FAvaEaseCurvePreset>& InPreset);
+
+	TAttribute<FFrameRate> DisplayRate;
+	bool bAllowEditMode = false;
+	FAvaOnPresetChanged OnPresetChanged;
+	FAvaOnPresetChanged OnQuickPresetChanged;
+
+	TSharedPtr<SBox> SelectedRowContainer;
+	TSharedPtr<SBox> GroupWidgetsParent;
+	TSharedPtr<SUniformWrapPanel> GroupWrapBox;
+	TArray<TSharedPtr<SAvaEaseCurvePresetGroup>> GroupWidgets;
+
+	TSharedPtr<FAvaEaseCurvePreset> SelectedItem;
+
+	FText SearchText;
+
+	TAttribute<bool> bEditMode;
+};
diff --git a/Engine/Plugins/Experimental/Avalanche/Source/AvalancheSequencer/Private/EaseCurveTool/Widgets/SAvaEaseCurvePresetGroup.cpp b/Engine/Plugins/Experimental/Avalanche/Source/AvalancheSequencer/Private/EaseCurveTool/Widgets/SAvaEaseCurvePresetGroup.cpp
--- a/Engine/Plugins/Experimental/Avalanche/Source/AvalancheSequencer/Private/EaseCurveTool/Widgets/SAvaEaseCurvePresetGroup.cpp
+++ b/Engine/Plugins/Experimental/Avalanche/Source/AvalancheSequencer/Private/EaseCurveTool/Widgets/SAvaEaseCurvePresetGroup.cpp
@@ -34,6 +34,7 @@ void SAvaEaseCurvePresetGroup::Construct(const FArguments& InArgs)
 	OnBeginPresetMove = InArgs._OnBeginPresetMove;
 	OnEndPresetMove = InArgs._OnEndPresetMove;
 	OnPresetClick = InArgs._OnPresetClick;
+	OnSetQuickEase = InArgs._OnSetQuickEase;
 
 	ChildSlot
 		[
@@ -145,11 +146,12 @@ TSharedRef<ITableRow> SAvaEaseCurvePresetGroup::GeneratePresetWidget(const TShar
 		.IsEditMode(bIsEditMode)
 		.DisplayRate(DisplayRate)
 		.IsSelected(this, &SAvaEaseCurvePresetGroup::IsSelected, InPreset)
-		.OnClick(this, &SAvaEaseCurvePresetGroup::HandlePresetClick)
 		.OnDelete(this, &SAvaEaseCurvePresetGroup::HandlePresetDelete)
 		.OnRename(this, &SAvaEaseCurvePresetGroup::HandlePresetRename)
 		.OnBeginMove(this, &SAvaEaseCurvePresetGroup::HandlePresetBeginMove)
-		.OnEndMove(this, &SAvaEaseCurvePresetGroup::HandlePresetEndMove);
+		.OnEndMove(this, &SAvaEaseCurvePresetGroup::HandlePresetEndMove)
+		.OnClick(this, &SAvaEaseCurvePresetGroup::HandlePresetClick)
+		.OnSetQuickEase(this, &SAvaEaseCurvePresetGroup::HandleSetQuickEase);
 	PresetWidgetsMap.Add(InPreset, NewPresetWidget);
 
 	return NewPresetWidget;
@@ -287,11 +289,23 @@ bool SAvaEaseCurvePresetGroup::HandlePresetEndMove(const TSharedPtr<FAvaEaseCurv
 	return false;
 }
 
-bool SAvaEaseCurvePresetGroup::HandlePresetClick(const TSharedPtr<FAvaEaseCurvePreset>& InPreset, const FModifierKeysState& InModifierKeys) const
+bool SAvaEaseCurvePresetGroup::HandlePresetClick(const TSharedPtr<FAvaEaseCurvePreset>& InPreset) const
 {
 	if (OnPresetClick.IsBound())
 	{
-		OnPresetClick.Execute(InPreset, InModifierKeys);
+		OnPresetClick.Execute(InPreset);
+
+		return true;
+	}
+
+	return false;
+}
+
+bool SAvaEaseCurvePresetGroup::HandleSetQuickEase(const TSharedPtr<FAvaEaseCurvePreset>& InPreset) const
+{
+	if (OnSetQuickEase.IsBound())
+	{
+		OnSetQuickEase.Execute(InPreset);
 
 		return true;
 	}
diff --git a/Engine/Plugins/Experimental/Avalanche/Source/AvalancheSequencer/Private/EaseCurveTool/Widgets/SAvaEaseCurvePresetGroup.h b/Engine/Plugins/Experimental/Avalanche/Source/AvalancheSequencer/Private/EaseCurveTool/Widgets/SAvaEaseCurvePresetGroup.h
--- a/Engine/Plugins/Experimental/Avalanche/Source/AvalancheSequencer/Private/EaseCurveTool/Widgets/SAvaEaseCurvePresetGroup.h
+++ b/Engine/Plugins/Experimental/Avalanche/Source/AvalancheSequencer/Private/EaseCurveTool/Widgets/SAvaEaseCurvePresetGroup.h
@@ -32,6 +32,7 @@ public:
 		SLATE_EVENT(FAvaEaseCurvePresetMoveDelegate, OnBeginPresetMove)
 		SLATE_EVENT(FAvaEaseCurvePresetMoveDelegate, OnEndPresetMove)
 		SLATE_EVENT(FAvaEaseCurvePresetClickDelegate, OnPresetClick)
+		SLATE_EVENT(FAvaEaseCurvePresetDelegate, OnSetQuickEase)
 	SLATE_END_ARGS()
 
 	void Construct(const FArguments& InArgs);
@@ -63,13 +64,14 @@ protected:
 
 	void HandleCategoryRenameCommitted(const FText& InNewText, ETextCommit::Type InCommitType);
 	FReply HandleCategoryDelete() const;
-
-	bool HandlePresetClick(const TSharedPtr<FAvaEaseCurvePreset>& InPreset, const FModifierKeysState& InModifierKeys) const;
+	
 	bool HandlePresetDelete(const TSharedPtr<FAvaEaseCurvePreset>& InPreset);
 	bool HandlePresetRename(const TSharedPtr<FAvaEaseCurvePreset>& InPreset, const FString& InNewName);
 	bool HandlePresetBeginMove(const TSharedPtr<FAvaEaseCurvePreset>& InPreset, const FString& InNewCategoryName) const;
 	bool HandlePresetEndMove(const TSharedPtr<FAvaEaseCurvePreset>& InPreset, const FString& InNewCategoryName) const;
-
+	bool HandlePresetClick(const TSharedPtr<FAvaEaseCurvePreset>& InPreset) const;
+	bool HandleSetQuickEase(const TSharedPtr<FAvaEaseCurvePreset>& InPreset) const;
+	
 	//~ Begin SWidget
 	virtual void OnDragEnter(const FGeometry& InGeometry, const FDragDropEvent& InDragDropEvent) override;
 	virtual void OnDragLeave(const FDragDropEvent& InDragDropEvent) override;
@@ -90,6 +92,7 @@ protected:
 	FAvaEaseCurvePresetMoveDelegate OnBeginPresetMove;
 	FAvaEaseCurvePresetMoveDelegate OnEndPresetMove;
 	FAvaEaseCurvePresetClickDelegate OnPresetClick;
+	FAvaEaseCurvePresetDelegate OnSetQuickEase;
 
 	TArray<TSharedPtr<FAvaEaseCurvePreset>> VisiblePresets;
 
diff --git a/Engine/Plugins/Experimental/Avalanche/Source/AvalancheSequencer/Private/EaseCurveTool/Widgets/SAvaEaseCurvePresetGroupItem.cpp b/Engine/Plugins/Experimental/Avalanche/Source/AvalancheSequencer/Private/EaseCurveTool/Widgets/SAvaEaseCurvePresetGroupItem.cpp
--- a/Engine/Plugins/Experimental/Avalanche/Source/AvalancheSequencer/Private/EaseCurveTool/Widgets/SAvaEaseCurvePresetGroupItem.cpp
+++ b/Engine/Plugins/Experimental/Avalanche/Source/AvalancheSequencer/Private/EaseCurveTool/Widgets/SAvaEaseCurvePresetGroupItem.cpp
@@ -28,11 +28,12 @@ void SAvaEaseCurvePresetGroupItem::Construct(const FArguments& InArgs, const TSh
 	Preset = InArgs._Preset;
 	bIsEditMode = InArgs._IsEditMode;
 	IsSelected = InArgs._IsSelected;
-	OnClick = InArgs._OnClick;
 	OnDelete = InArgs._OnDelete;
 	OnRename = InArgs._OnRename;
 	OnBeginMove = InArgs._OnBeginMove;
 	OnEndMove = InArgs._OnEndMove;
+	OnClick = InArgs._OnClick;
+	OnSetQuickEase = InArgs._OnSetQuickEase;
 
 	const FText ItemTooltipText = FText::Format(LOCTEXT("ItemTooltip", "{0}\n\nShift + Click to set as active quick preset"), FText::FromString(Preset->Name));
 
@@ -63,33 +64,17 @@ void SAvaEaseCurvePresetGroupItem::Construct(const FArguments& InArgs, const TSh
 							{
 								return IsEditMode() ? 1 : 0;
 							})
+						// Normal mode
 						+ SWidgetSwitcher::Slot()
 						[
-							SNew(SHorizontalBox)
-							+ SHorizontalBox::Slot()
-							.AutoWidth()
-							.VAlign(VAlign_Center)
-							[
-								SNew(SBox)
-								.Padding(0.f, 0.f, 3.f, 0.f)
-								.Visibility(this, &SAvaEaseCurvePresetGroupItem::GetQuickPresetIconVisibility)
-								[
-									SNew(SImage)
-									.DesiredSizeOverride(FVector2D(10.f))
-									.Image(FAppStyle::GetBrush(TEXT("Icons.Adjust")))
-									.ToolTipText(this, &SAvaEaseCurvePresetGroupItem::GetQuickPresetIconToolTip)
-								]
-							]
-							+ SHorizontalBox::Slot()
-							[
-								SNew(STextBlock)
-								.TextStyle(FAppStyle::Get(), TEXT("Menu.Label"))
-								.OverflowPolicy(ETextOverflowPolicy::Ellipsis)
-								.ColorAndOpacity(FStyleColors::White)
-								.Text(FText::FromString(Preset->Name))
-								.ToolTipText(ItemTooltipText)
-							]
+							SNew(STextBlock)
+							.TextStyle(FAppStyle::Get(), TEXT("Menu.Label"))
+							.OverflowPolicy(ETextOverflowPolicy::Ellipsis)
+							.ColorAndOpacity(FStyleColors::White)
+							.Text(FText::FromString(Preset->Name))
+							.ToolTipText(ItemTooltipText)
 						]
+						// Edit mode
 						+ SWidgetSwitcher::Slot()
 						[
 							SNew(SHorizontalBox)
@@ -121,6 +106,32 @@ void SAvaEaseCurvePresetGroupItem::Construct(const FArguments& InArgs, const TSh
 						]
 					]
 				]
+				// Quick Ease Button
+				+ SHorizontalBox::Slot()
+				.AutoWidth()
+				.HAlign(HAlign_Right)
+				.VAlign(VAlign_Center)
+				.Padding(8.f, 0.f, 0.f, 0.f)
+				[
+					SNew(SBorder)
+					.Padding(0.f)
+					.BorderImage(FAvaEaseCurveStyle::Get().GetBrush(TEXT("ToolButton.Opaque")))
+					[
+						SNew(SButton)
+						.VAlign(VAlign_Center)
+						.ButtonStyle(FAvaEaseCurveStyle::Get(), TEXT("ToolButton.NoPad"))
+						.ToolTipText(this, &SAvaEaseCurvePresetGroupItem::GetQuickPresetIconToolTip)
+						.Visibility(this, &SAvaEaseCurvePresetGroupItem::GetQuickPresetIconVisibility)
+						.OnClicked(this, &SAvaEaseCurvePresetGroupItem::HandleSetQuickEase)
+						[
+							SNew(SImage)
+							.DesiredSizeOverride(FVector2D(10.f))
+							.Image(FAppStyle::GetBrush(TEXT("Icons.Adjust")))
+							.ColorAndOpacity(this, &SAvaEaseCurvePresetGroupItem::GetQuickPresetIconColor)
+						]
+					]
+				]
+				// Preview Image
 				+ SHorizontalBox::Slot()
 				.AutoWidth()
 				.HAlign(HAlign_Right)
@@ -224,7 +235,7 @@ FReply SAvaEaseCurvePresetGroupItem::OnMouseButtonDown(const FGeometry& InGeomet
 
 	if (OnClick.IsBound())
 	{
-		OnClick.Execute(Preset, InMouseEvent.GetModifierKeys());
+		OnClick.Execute(Preset);
 	}
 
 	return FReply::Handled();
@@ -234,7 +245,7 @@ FReply SAvaEaseCurvePresetGroupItem::OnMouseButtonUp(const FGeometry& InGeometry
 {
 	TriggerEndMove();
 
-	return STableRow<TSharedPtr<FAvaEaseCurvePreset>>::OnMouseButtonUp(InGeometry, InMouseEvent);
+	return STableRow::OnMouseButtonUp(InGeometry, InMouseEvent);
 }
 
 FReply SAvaEaseCurvePresetGroupItem::OnDragDetected(const FGeometry& MyGeometry, const FPointerEvent& MouseEvent)
@@ -254,7 +265,7 @@ FReply SAvaEaseCurvePresetGroupItem::OnDrop(const FGeometry& InGeometry, const F
 {
 	bIsDragging = false;
 
-	return STableRow<TSharedPtr<FAvaEaseCurvePreset>>::OnDrop(InGeometry, InDragDropEvent);
+	return STableRow::OnDrop(InGeometry, InDragDropEvent);
 }
 
 void SAvaEaseCurvePresetGroupItem::TriggerBeginMove()
@@ -277,22 +288,21 @@ void SAvaEaseCurvePresetGroupItem::TriggerEndMove()
 	bIsDragging = false;
 }
 
-EVisibility SAvaEaseCurvePresetGroupItem::GetQuickPresetIconVisibility() const
+FSlateColor SAvaEaseCurvePresetGroupItem::GetQuickPresetIconColor() const
 {
-	const UAvaEaseCurveToolSettings* const Settings = GetDefault<UAvaEaseCurveToolSettings>();
-
-	FAvaEaseCurveTangents Tangents;
-	if (!FAvaEaseCurveTangents::FromString(Settings->GetQuickEaseTangents(), Tangents))
-	{
-		return EVisibility::Hidden;
-	}
+	return (IsHovered() && IsQuickEasePreset()) ? FStyleColors::Select : FSlateColor::UseStyle();
+}
 
-	return (Preset->Tangents == Tangents) ? EVisibility::Visible : EVisibility::Collapsed;
+EVisibility SAvaEaseCurvePresetGroupItem::GetQuickPresetIconVisibility() const
+{
+	return (IsHovered() || IsQuickEasePreset()) ? EVisibility::Visible : EVisibility::Collapsed;
 }
 
 FText SAvaEaseCurvePresetGroupItem::GetQuickPresetIconToolTip() const
 {
-	static const FText QuickEaseText = LOCTEXT("QuickEaseIconTooltip", "Active Quick Ease Preset");
+	static const FText QuickEaseText = IsQuickEasePreset()
+		? LOCTEXT("ActiveQuickEaseIconTooltip", "Active Quick Ease Preset")
+		: LOCTEXT("QuickEaseIconTooltip", "Set to Active Quick Ease Preset");
 
 	const FAvaEaseCurveToolCommands& EaseCurveToolCommands = FAvaEaseCurveToolCommands::Get();
 
@@ -319,4 +329,34 @@ FText SAvaEaseCurvePresetGroupItem::GetQuickPresetIconToolTip() const
 	return CommandText.IsEmpty() ? QuickEaseText : FText::Format(LOCTEXT("QuickEasePresetIconTooltip", "{0}\n\n{1}"), QuickEaseText, CommandText);
 }
 
+FReply SAvaEaseCurvePresetGroupItem::HandleSetQuickEase()
+{
+	if (OnSetQuickEase.IsBound())
+	{
+		OnSetQuickEase.Execute(Preset);
+	}
+
+	return FReply::Handled();
+}
+
+bool SAvaEaseCurvePresetGroupItem::IsQuickEasePreset() const
+{
+	if (!Preset.IsValid())
+	{
+		return false;
+	}
+	
+	const UAvaEaseCurveToolSettings* const Settings = GetDefault<UAvaEaseCurveToolSettings>();
+	check(IsValid(Settings));
+	
+	FAvaEaseCurveTangents Tangents;
+	
+	if (!FAvaEaseCurveTangents::FromString(Settings->GetQuickEaseTangents(), Tangents))
+	{
+		return false;
+	}
+
+	return Tangents.IsNearlyEqual(Preset->Tangents);
+}
+
 #undef LOCTEXT_NAMESPACE
diff --git a/Engine/Plugins/Experimental/Avalanche/Source/AvalancheSequencer/Private/EaseCurveTool/Widgets/SAvaEaseCurvePresetGroupItem.h b/Engine/Plugins/Experimental/Avalanche/Source/AvalancheSequencer/Private/EaseCurveTool/Widgets/SAvaEaseCurvePresetGroupItem.h
--- a/Engine/Plugins/Experimental/Avalanche/Source/AvalancheSequencer/Private/EaseCurveTool/Widgets/SAvaEaseCurvePresetGroupItem.h
+++ b/Engine/Plugins/Experimental/Avalanche/Source/AvalancheSequencer/Private/EaseCurveTool/Widgets/SAvaEaseCurvePresetGroupItem.h
@@ -13,7 +13,7 @@ class STableViewBase;
 struct FAvaEaseCurvePreset;
 
 DECLARE_DELEGATE_RetVal_OneParam(bool, FAvaEaseCurvePresetDelegate, const TSharedPtr<FAvaEaseCurvePreset>& /*InPreset*/)
-DECLARE_DELEGATE_RetVal_TwoParams(bool, FAvaEaseCurvePresetClickDelegate, const TSharedPtr<FAvaEaseCurvePreset>& /*InPreset*/, const FModifierKeysState& /*InModifierKeys*/)
+DECLARE_DELEGATE_RetVal_OneParam(bool, FAvaEaseCurvePresetClickDelegate, const TSharedPtr<FAvaEaseCurvePreset>& /*InPreset*/)
 DECLARE_DELEGATE_RetVal_TwoParams(bool, FAvaEaseCurvePresetRenameDelegate, const TSharedPtr<FAvaEaseCurvePreset>& /*InPreset*/, const FString& /*InNewName*/)
 DECLARE_DELEGATE_RetVal_TwoParams(bool, FAvaEaseCurvePresetMoveDelegate, const TSharedPtr<FAvaEaseCurvePreset>& /*InPreset*/, const FString& /*InNewCategoryName*/)
 
@@ -30,6 +30,7 @@ public:
 		SLATE_EVENT(FAvaEaseCurvePresetClickDelegate, OnClick)
 		SLATE_EVENT(FAvaEaseCurvePresetDelegate, OnDelete)
 		SLATE_EVENT(FAvaEaseCurvePresetRenameDelegate, OnRename)
+		SLATE_EVENT(FAvaEaseCurvePresetDelegate, OnSetQuickEase)
 		SLATE_EVENT(FAvaEaseCurvePresetMoveDelegate, OnBeginMove)
 		SLATE_EVENT(FAvaEaseCurvePresetMoveDelegate, OnEndMove)
 	SLATE_END_ARGS()
@@ -44,8 +45,6 @@ public:
 	void TriggerEndMove();
 
 protected:
-	void HandlePresetClick() const;
-
 	void HandleRenameTextCommitted(const FText& InNewText, ETextCommit::Type InCommitType) const;
 
 	FReply HandleDeleteClick() const;
@@ -55,9 +54,14 @@ protected:
 	EVisibility GetBorderVisibility() const;
 	const FSlateBrush* GetBackgroundImage() const;
 
+	FSlateColor GetQuickPresetIconColor() const;
 	EVisibility GetQuickPresetIconVisibility() const;
 	FText GetQuickPresetIconToolTip() const;
 
+	bool IsQuickEasePreset() const;
+
+	FReply HandleSetQuickEase();
+
 	//~ Begin SWidget
 	virtual FReply OnMouseButtonDown(const FGeometry& InGeometry, const FPointerEvent& InMouseEvent) override;
 	virtual FReply OnMouseButtonUp(const FGeometry& InGeometry, const FPointerEvent& InMouseEvent) override;
@@ -68,9 +72,12 @@ protected:
 	TSharedPtr<FAvaEaseCurvePreset> Preset;
 	TAttribute<bool> bIsEditMode;
 	TAttribute<bool> IsSelected;
+	
 	FAvaEaseCurvePresetClickDelegate OnClick;
 	FAvaEaseCurvePresetDelegate OnDelete;
 	FAvaEaseCurvePresetRenameDelegate OnRename;
+	FAvaEaseCurvePresetDelegate OnSetQuickEase;
+	
 	FAvaEaseCurvePresetMoveDelegate OnBeginMove;
 	FAvaEaseCurvePresetMoveDelegate OnEndMove;
 
diff --git a/Engine/Plugins/Experimental/Avalanche/Source/AvalancheSequencer/Private/EaseCurveTool/Widgets/SAvaEaseCurveTangents.cpp b/Engine/Plugins/Experimental/Avalanche/Source/AvalancheSequencer/Private/EaseCurveTool/Widgets/SAvaEaseCurveTangents.cpp
new file mode 100644
--- /dev/null
+++ b/Engine/Plugins/Experimental/Avalanche/Source/AvalancheSequencer/Private/EaseCurveTool/Widgets/SAvaEaseCurveTangents.cpp
@@ -0,0 +1,219 @@
+﻿// Copyright Epic Games, Inc. All Rights Reserved.
+
+#include "SAvaEaseCurveTangents.h"
+#include "DetailLayoutBuilder.h"
+#include "Widgets/Input/SSpinBox.h"
+#include "Widgets/Layout/SBox.h"
+#include "Widgets/Layout/SWrapBox.h"
+#include "Widgets/SBoxPanel.h"
+#include "Widgets/Text/STextBlock.h"
+
+#define LOCTEXT_NAMESPACE "SAvaEaseCurveTool"
+
+void SAvaEaseCurveTangents::Construct(const FArguments& InArgs)
+{
+	Tangents = InArgs._InitialTangents;
+	OnStartTangentChanged = InArgs._OnStartTangentChanged;
+	OnStartWeightChanged = InArgs._OnStartWeightChanged;
+	OnEndTangentChanged = InArgs._OnEndTangentChanged;
+	OnEndWeightChanged = InArgs._OnEndWeightChanged;
+	OnBeginSliderMovement = InArgs._OnBeginSliderMovement;
+	OnEndSliderMovement = InArgs._OnEndSliderMovement;
+	
+	constexpr float WrapSize = 120.f;
+
+	constexpr double MinTangent = -180.0f;
+	constexpr double MaxTangent = 180.0f;
+	constexpr double MinWeight = 0.0f;
+	constexpr double MaxWeight = 10.0f;
+
+	ChildSlot
+	[
+		SNew(SWrapBox)
+		.UseAllottedSize(true)
+		.HAlign(HAlign_Center)
+		+ SWrapBox::Slot()
+		.FillLineWhenSizeLessThan(WrapSize)
+		[
+			SNew(SHorizontalBox)
+			+ SHorizontalBox::Slot()
+			.AutoWidth()
+			.VAlign(VAlign_Center)
+			.Padding(0.f, 1.f, 0.f, 0.f)
+			[
+				SNew(SBox)
+				.WidthOverride(26.f)
+				.HAlign(HAlign_Center)
+				[
+					SNew(STextBlock)
+					.Font(IDetailLayoutBuilder::GetDetailFont())
+					.Text(LOCTEXT("OutLabel", "Out"))
+				]
+			]
+			+ SHorizontalBox::Slot()
+			.FillWidth(1.f)
+			.HAlign(HAlign_Fill)
+			.VAlign(VAlign_Center)
+			[
+				ConstructTangentNumBox(LOCTEXT("OutTangentWeightLabel", "W")
+					, LOCTEXT("OutTangentWeightToolTip", "Out Tangent Weight")
+					, TAttribute<double>::CreateSP(this, &SAvaEaseCurveTangents::GetStartTangentWeight)
+					, SNumericEntryBox<double>::FOnValueChanged::CreateSP(this, &SAvaEaseCurveTangents::OnStartTangentWeightSpinBoxChanged)
+					, MinWeight, MaxWeight)
+			]
+			+ SHorizontalBox::Slot()
+			.FillWidth(1.f)
+			.HAlign(HAlign_Fill)
+			.VAlign(VAlign_Center)
+			.Padding(2.f, 0.f, 0.f, 0.f)
+			[
+				ConstructTangentNumBox(LOCTEXT("OutTangentLabel", "T")
+					, LOCTEXT("OutTangentToolTip", "Out Tangent")
+					, TAttribute<double>::CreateSP(this, &SAvaEaseCurveTangents::GetStartTangent)
+					, SNumericEntryBox<double>::FOnValueChanged::CreateSP(this, &SAvaEaseCurveTangents::OnStartTangentSpinBoxChanged)
+					, MinTangent, MaxTangent)
+			]
+		]
+		+ SWrapBox::Slot()
+		.FillLineWhenSizeLessThan(WrapSize)
+		[
+			SNew(SHorizontalBox)
+			+ SHorizontalBox::Slot()
+			.AutoWidth()
+			.VAlign(VAlign_Center)
+			.Padding(0.f, 1.f, 0.f, 0.f)
+			[
+				SNew(SBox)
+				.WidthOverride(26.f)
+				.HAlign(HAlign_Center)
+				[
+					SNew(STextBlock)
+					.Font(IDetailLayoutBuilder::GetDetailFont())
+					.Text(LOCTEXT("InLabel", "In"))
+				]
+			]
+			+ SHorizontalBox::Slot()
+			.FillWidth(1.f)
+			.HAlign(HAlign_Fill)
+			.VAlign(VAlign_Center)
+			[
+				ConstructTangentNumBox(LOCTEXT("InTangentWeightLabel", "W")
+					, LOCTEXT("InTangentWeightToolTip", "In Tangent Weight")
+					, TAttribute<double>::CreateSP(this, &SAvaEaseCurveTangents::GetEndTangentWeight)
+					, SNumericEntryBox<double>::FOnValueChanged::CreateSP(this, &SAvaEaseCurveTangents::OnEndTangentWeightSpinBoxChanged)
+					, MinWeight, MaxWeight)
+			]
+			+ SHorizontalBox::Slot()
+			.FillWidth(1.f)
+			.HAlign(HAlign_Fill)
+			.VAlign(VAlign_Center)
+			.Padding(2.f, 0.f, 0.f, 0.f)
+			[
+				ConstructTangentNumBox(LOCTEXT("InTangentLabel", "T")
+					, LOCTEXT("InTangentToolTip", "In Tangent")
+					, TAttribute<double>::CreateSP(this, &SAvaEaseCurveTangents::GetEndTangent)
+					, SNumericEntryBox<double>::FOnValueChanged::CreateSP(this, &SAvaEaseCurveTangents::OnEndTangentSpinBoxChanged)
+					, MinTangent, MaxTangent)
+			]
+		]
+	];
+}
+
+TSharedRef<SWidget> SAvaEaseCurveTangents::ConstructTangentNumBox(const FText& InLabel
+	, const FText& InToolTip
+	, const TAttribute<double>& InValue
+	, const SNumericEntryBox<double>::FOnValueChanged& InOnValueChanged
+	, const TOptional<double>& InMinSliderValue
+	, const TOptional<double>& InMaxSliderValue) const
+{
+	return SNew(SBox)
+		.MaxDesiredWidth(100.f)
+		.HAlign(HAlign_Fill)
+		[
+			SNew(SHorizontalBox)
+			+ SHorizontalBox::Slot()
+			.AutoWidth()
+			.HAlign(HAlign_Right)
+			[
+				SNew(STextBlock)
+				.MinDesiredWidth(8.f)
+				.Margin(FMargin(2.f, 5.f, 2.f, 3.f))
+				.Font(IDetailLayoutBuilder::GetDetailFont())
+				.Text(InLabel)
+				.ToolTipText(InToolTip)
+			]
+			+ SHorizontalBox::Slot()
+			[
+				SNew(SSpinBox<double>)
+				.Font(IDetailLayoutBuilder::GetDetailFont())
+				.MinSliderValue(InMinSliderValue)
+				.MaxSliderValue(InMaxSliderValue)
+				.Delta(0.00001f)
+				.WheelStep(0.001f)
+				.MinFractionalDigits(4)
+				.MaxFractionalDigits(6)
+				.MinDesiredWidth(70.f)
+				.Value(InValue)
+				.OnBeginSliderMovement(OnBeginSliderMovement)
+				.OnEndSliderMovement(OnEndSliderMovement)
+				.OnValueChanged_Lambda([this, InOnValueChanged](const double InNewValue)
+					{
+						InOnValueChanged.ExecuteIfBound(InNewValue);
+					})
+				.OnValueCommitted_Lambda([InOnValueChanged](const double InNewValue, ETextCommit::Type InCommitType)
+					{
+						InOnValueChanged.ExecuteIfBound(InNewValue);
+					})
+			]
+		];
+}
+
+double SAvaEaseCurveTangents::GetStartTangent() const
+{
+	return Tangents.Start;
+}
+
+double SAvaEaseCurveTangents::GetStartTangentWeight() const
+{
+	return Tangents.StartWeight;
+}
+
+double SAvaEaseCurveTangents::GetEndTangent() const
+{
+	return Tangents.End;
+}
+
+double SAvaEaseCurveTangents::GetEndTangentWeight() const
+{
+	return Tangents.EndWeight;
+}
+
+void SAvaEaseCurveTangents::OnStartTangentSpinBoxChanged(const double InNewValue)
+{
+	Tangents.Start = InNewValue;
+
+	OnStartTangentChanged.Execute(InNewValue);
+}
+
+void SAvaEaseCurveTangents::OnStartTangentWeightSpinBoxChanged(const double InNewValue)
+{
+	Tangents.StartWeight = InNewValue;
+	
+	OnStartWeightChanged.Execute(InNewValue);
+}
+
+void SAvaEaseCurveTangents::OnEndTangentSpinBoxChanged(const double InNewValue)
+{
+	Tangents.End = InNewValue;
+
+	OnEndTangentChanged.Execute(InNewValue);
+}
+
+void SAvaEaseCurveTangents::OnEndTangentWeightSpinBoxChanged(const double InNewValue)
+{
+	Tangents.EndWeight = InNewValue;
+
+	OnEndWeightChanged.Execute(InNewValue);
+}
+
+#undef LOCTEXT_NAMESPACE
diff --git a/Engine/Plugins/Experimental/Avalanche/Source/AvalancheSequencer/Private/EaseCurveTool/Widgets/SAvaEaseCurveTangents.h b/Engine/Plugins/Experimental/Avalanche/Source/AvalancheSequencer/Private/EaseCurveTool/Widgets/SAvaEaseCurveTangents.h
new file mode 100644
--- /dev/null
+++ b/Engine/Plugins/Experimental/Avalanche/Source/AvalancheSequencer/Private/EaseCurveTool/Widgets/SAvaEaseCurveTangents.h
@@ -0,0 +1,51 @@
+﻿// Copyright Epic Games, Inc. All Rights Reserved.
+
+#pragma once
+
+#include "EaseCurveTool/AvaEaseCurveTangents.h"
+#include "Widgets/Input/SNumericEntryBox.h"
+
+class SAvaEaseCurveTangents : public SCompoundWidget
+{
+public:
+	SLATE_BEGIN_ARGS(SAvaEaseCurveTangents)
+	{}
+		SLATE_ARGUMENT(FAvaEaseCurveTangents, InitialTangents)
+		SLATE_EVENT(SNumericEntryBox<double>::FOnValueChanged, OnStartTangentChanged)
+		SLATE_EVENT(SNumericEntryBox<double>::FOnValueChanged, OnStartWeightChanged)
+		SLATE_EVENT(SNumericEntryBox<double>::FOnValueChanged, OnEndTangentChanged)
+		SLATE_EVENT(SNumericEntryBox<double>::FOnValueChanged, OnEndWeightChanged)
+		SLATE_EVENT(FSimpleDelegate, OnBeginSliderMovement)
+		SLATE_EVENT(SNumericEntryBox<double>::FOnValueChanged, OnEndSliderMovement)
+	SLATE_END_ARGS()
+
+	void Construct(const FArguments& InArgs);
+
+	double GetStartTangent() const;
+	double GetStartTangentWeight() const;
+	double GetEndTangent() const;
+	double GetEndTangentWeight() const;
+
+protected:
+	TSharedRef<SWidget> ConstructTangentNumBox(const FText& InLabel
+		, const FText& InToolTip
+		, const TAttribute<double>& InValue
+		, const SNumericEntryBox<double>::FOnValueChanged& InOnValueChanged
+		, const TOptional<double>& InMinSliderValue
+		, const TOptional<double>& InMaxSliderValue) const;
+
+	void OnStartTangentSpinBoxChanged(const double InNewValue);
+	void OnStartTangentWeightSpinBoxChanged(const double InNewValue);
+	void OnEndTangentSpinBoxChanged(const double InNewValue);
+	void OnEndTangentWeightSpinBoxChanged(const double InNewValue);
+
+	FAvaEaseCurveTangents Tangents;
+
+	SNumericEntryBox<double>::FOnValueChanged OnStartTangentChanged;
+	SNumericEntryBox<double>::FOnValueChanged OnStartWeightChanged;
+	SNumericEntryBox<double>::FOnValueChanged OnEndTangentChanged;
+	SNumericEntryBox<double>::FOnValueChanged OnEndWeightChanged;
+
+	FSimpleDelegate OnBeginSliderMovement;
+	SNumericEntryBox<double>::FOnValueChanged OnEndSliderMovement;
+};
diff --git a/Engine/Plugins/Experimental/Avalanche/Source/AvalancheSequencer/Private/EaseCurveTool/Widgets/SAvaEaseCurveTool.cpp b/Engine/Plugins/Experimental/Avalanche/Source/AvalancheSequencer/Private/EaseCurveTool/Widgets/SAvaEaseCurveTool.cpp
--- a/Engine/Plugins/Experimental/Avalanche/Source/AvalancheSequencer/Private/EaseCurveTool/Widgets/SAvaEaseCurveTool.cpp
+++ b/Engine/Plugins/Experimental/Avalanche/Source/AvalancheSequencer/Private/EaseCurveTool/Widgets/SAvaEaseCurveTool.cpp
@@ -3,7 +3,6 @@
 #include "SAvaEaseCurveTool.h"
 #include "Curves/KeyHandle.h"
 #include "Curves/RichCurve.h"
-#include "DetailLayoutBuilder.h"
 #include "EaseCurveTool/AvaEaseCurvePreset.h"
 #include "EaseCurveTool/AvaEaseCurveTool.h"
 #include "EaseCurveTool/AvaEaseCurveToolCommands.h"
@@ -11,14 +10,15 @@
 #include "EaseCurveTool/AvaEaseCurveSubsystem.h"
 #include "EaseCurveTool/Widgets/SAvaEaseCurveEditor.h"
 #include "EaseCurveTool/Widgets/SAvaEaseCurvePreset.h"
+#include "EaseCurveTool/Widgets/AvaEaseCurveToolContextMenu.h"
 #include "Editor.h"
 #include "EngineAnalytics.h"
+#include "SAvaEaseCurveTangents.h"
 #include "Framework/Commands/GenericCommands.h"
+#include "Framework/Commands/UICommandList.h"
 #include "Framework/MultiBox/MultiBoxBuilder.h"
-#include "IAvaSequencer.h"
 #include "Math/UnrealMathUtility.h"
 #include "Misc/MessageDialog.h"
-#include "SCurveEditor.h"
 #include "Styling/AppStyle.h"
 #include "Styling/StyleColors.h"
 #include "ToolMenu.h"
@@ -26,100 +26,11 @@
 #include "Widgets/Images/SImage.h"
 #include "Widgets/Input/SButton.h"
 #include "Widgets/Input/SComboBox.h"
-#include "Widgets/Input/SSpinBox.h"
 #include "Widgets/Layout/SScrollBox.h"
-#include "Widgets/Layout/SWrapBox.h"
 #include "Widgets/SNullWidget.h"
 
 #define LOCTEXT_NAMESPACE "SAvaEaseCurveTool"
 
-namespace UE::AvaSequencer
-{
-	template<typename NumericType>
-	void ExtractNumericMetadata(FProperty* InProperty
-		, TOptional<NumericType>& OutMinValue
-		, TOptional<NumericType>& OutMaxValue
-		, TOptional<NumericType>& OutSliderMinValue
-		, TOptional<NumericType>& OutSliderMaxValue
-		, NumericType& OutSliderExponent
-		, NumericType& OutDelta
-		, float& OutShiftMultiplier
-		, float& OutCtrlMultiplier
-		, bool& OutSupportDynamicSliderMaxValue
-		, bool& OutSupportDynamicSliderMinValue)
-	{
-		const FString& MetaUIMinString = InProperty->GetMetaData(TEXT("UIMin"));
-		const FString& MetaUIMaxString = InProperty->GetMetaData(TEXT("UIMax"));
-		const FString& SliderExponentString = InProperty->GetMetaData(TEXT("SliderExponent"));
-		const FString& DeltaString = InProperty->GetMetaData(TEXT("Delta"));
-		const FString& ShiftMultiplierString = InProperty->GetMetaData(TEXT("ShiftMultiplier"));
-		const FString& CtrlMultiplierString = InProperty->GetMetaData(TEXT("CtrlMultiplier"));
-		const FString& SupportDynamicSliderMaxValueString = InProperty->GetMetaData(TEXT("SupportDynamicSliderMaxValue"));
-		const FString& SupportDynamicSliderMinValueString = InProperty->GetMetaData(TEXT("SupportDynamicSliderMinValue"));
-		const FString& ClampMinString = InProperty->GetMetaData(TEXT("ClampMin"));
-		const FString& ClampMaxString = InProperty->GetMetaData(TEXT("ClampMax"));
-
-		// If no UIMin/Max was specified then use the clamp string
-		const FString& UIMinString = MetaUIMinString.Len() ? MetaUIMinString : ClampMinString;
-		const FString& UIMaxString = MetaUIMaxString.Len() ? MetaUIMaxString : ClampMaxString;
-
-		NumericType ClampMin = TNumericLimits<NumericType>::Lowest();
-		NumericType ClampMax = TNumericLimits<NumericType>::Max();
-
-		if (!ClampMinString.IsEmpty())
-		{
-			TTypeFromString<NumericType>::FromString(ClampMin, *ClampMinString);
-		}
-
-		if (!ClampMaxString.IsEmpty())
-		{
-			TTypeFromString<NumericType>::FromString(ClampMax, *ClampMaxString);
-		}
-
-		NumericType UIMin = TNumericLimits<NumericType>::Lowest();
-		NumericType UIMax = TNumericLimits<NumericType>::Max();
-		TTypeFromString<NumericType>::FromString(UIMin, *UIMinString);
-		TTypeFromString<NumericType>::FromString(UIMax, *UIMaxString);
-
-		OutSliderExponent = NumericType(1);
-
-		if (SliderExponentString.Len())
-		{
-			TTypeFromString<NumericType>::FromString(OutSliderExponent, *SliderExponentString);
-		}
-
-		OutDelta = NumericType(0);
-
-		if (DeltaString.Len())
-		{
-			TTypeFromString<NumericType>::FromString(OutDelta, *DeltaString);
-		}
-
-		OutShiftMultiplier = 10.f;
-		if (ShiftMultiplierString.Len())
-		{
-			TTypeFromString<float>::FromString(OutShiftMultiplier, *ShiftMultiplierString);
-		}
-
-		OutCtrlMultiplier = 0.1f;
-		if (CtrlMultiplierString.Len())
-		{
-			TTypeFromString<float>::FromString(OutCtrlMultiplier, *CtrlMultiplierString);
-		}
-
-		const NumericType ActualUIMin = FMath::Max(UIMin, ClampMin);
-		const NumericType ActualUIMax = FMath::Min(UIMax, ClampMax);
-
-		OutMinValue = ClampMinString.Len() ? ClampMin : TOptional<NumericType>();
-		OutMaxValue = ClampMaxString.Len() ? ClampMax : TOptional<NumericType>();
-		OutSliderMinValue = (UIMinString.Len()) ? ActualUIMin : TOptional<NumericType>();
-		OutSliderMaxValue = (UIMaxString.Len()) ? ActualUIMax : TOptional<NumericType>();
-
-		OutSupportDynamicSliderMaxValue = SupportDynamicSliderMaxValueString.Len() > 0 && SupportDynamicSliderMaxValueString.ToBool();
-		OutSupportDynamicSliderMinValue = SupportDynamicSliderMinValueString.Len() > 0 && SupportDynamicSliderMinValueString.ToBool();
-	}
-}
-
 void SAvaEaseCurveTool::Construct(const FArguments& InArgs, const TSharedRef<FAvaEaseCurveTool>& InEaseCurveTool)
 {
 	ToolMode = InArgs._ToolMode;
@@ -156,7 +67,20 @@ void SAvaEaseCurveTool::Construct(const FArguments& InArgs, const TSharedRef<FAv
 			.HAlign(HAlign_Fill)
 			.Padding(0.f, 3.f, 0.f, 0.f)
 			[
-				ConstructInputBoxes()
+				SNew(SAvaEaseCurveTangents)
+				.InitialTangents(GetTangents())
+				.OnStartTangentChanged(this, &SAvaEaseCurveTool::OnStartTangentSpinBoxChanged)
+				.OnStartWeightChanged(this, &SAvaEaseCurveTool::OnStartTangentWeightSpinBoxChanged)
+				.OnEndTangentChanged(this, &SAvaEaseCurveTool::OnEndTangentSpinBoxChanged)
+				.OnEndWeightChanged(this, &SAvaEaseCurveTool::OnEndTangentWeightSpinBoxChanged)
+				.OnBeginSliderMovement_Lambda([this]()
+					{
+						EaseCurveTool->BeginTransaction(LOCTEXT("SliderDragStartLabel", "Ease Curve Slider Drag"));
+					})
+				.OnEndSliderMovement_Lambda([this](const float InNewValue)
+					{
+						EaseCurveTool->EndTransaction();
+					})
 			]
 		];
 
@@ -175,7 +99,7 @@ TSharedRef<SWidget> SAvaEaseCurveTool::ConstructCurveEditorPanel()
 	CurrentGraphSize = GetDefault<UAvaEaseCurveToolSettings>()->GetGraphSize();
 
 	const TSharedRef<FAvaEaseCurveTool> EaseCurveToolRef = EaseCurveTool.ToSharedRef();
-
+	
 	return SNew(SBorder)
 		[
 			SNew(SOverlay)
@@ -195,7 +119,15 @@ TSharedRef<SWidget> SAvaEaseCurveTool::ConstructCurveEditorPanel()
 				.OnTangentsChanged(this, &SAvaEaseCurveTool::HandleEditorTangentsChanged)
 				.GridSnap_UObject(GetDefault<UAvaEaseCurveToolSettings>(), &UAvaEaseCurveToolSettings::GetGridSnap)
 				.GridSize_UObject(GetDefault<UAvaEaseCurveToolSettings>(), &UAvaEaseCurveToolSettings::GetGridSize)
-				.GetContextMenuContent(this, &SAvaEaseCurveTool::CreateContextMenuContent)
+				.GetContextMenuContent_Lambda([this]()
+					{
+						TSharedPtr<FAvaEaseCurveToolContextMenu> ContextMenu = MakeShared<FAvaEaseCurveToolContextMenu>(CommandList
+							, FAvaEaseCurveToolOnGraphSizeChanged::CreateLambda([this](const int32 InNewSize)
+								{
+									CurrentGraphSize = InNewSize;
+								}));
+						return ContextMenu->GenerateWidget();
+					})
 				.StartText(this, &SAvaEaseCurveTool::GetStartText)
 				.StartTooltipText(this, &SAvaEaseCurveTool::GetStartTooltipText)
 				.EndText(this, &SAvaEaseCurveTool::GetEndText)
@@ -211,7 +143,7 @@ TSharedRef<SWidget> SAvaEaseCurveTool::ConstructCurveEditorPanel()
 				SNew(SImage)
 				.Visibility(EVisibility::HitTestInvisible)
 				.ColorAndOpacity(FLinearColor(0.f, 0.f, 0.f, 0.5f))
-				.Image(&(FCoreStyle::Get().GetWidgetStyle<FScrollBoxStyle>("ScrollBox").TopShadowBrush))
+				.Image(&(FCoreStyle::Get().GetWidgetStyle<FScrollBoxStyle>(TEXT("ScrollBox")).TopShadowBrush))
 			]
 			+ SOverlay::Slot()
 			.HAlign(HAlign_Fill)
@@ -220,7 +152,7 @@ TSharedRef<SWidget> SAvaEaseCurveTool::ConstructCurveEditorPanel()
 				SNew(SImage)
 				.Visibility(EVisibility::HitTestInvisible)
 				.ColorAndOpacity(FLinearColor(0.f, 0.f, 0.f, 0.5f))
-				.Image(&(FCoreStyle::Get().GetWidgetStyle<FScrollBoxStyle>("ScrollBox").BottomShadowBrush))
+				.Image(&(FCoreStyle::Get().GetWidgetStyle<FScrollBoxStyle>(TEXT("ScrollBox")).BottomShadowBrush))
 			]
 			+ SOverlay::Slot()
 			.HAlign(HAlign_Left)
@@ -229,7 +161,7 @@ TSharedRef<SWidget> SAvaEaseCurveTool::ConstructCurveEditorPanel()
 				SNew(SImage)
 				.Visibility(EVisibility::HitTestInvisible)
 				.ColorAndOpacity(FLinearColor(0.f, 0.f, 0.f, 0.5f))
-				.Image(&(FCoreStyle::Get().GetWidgetStyle<FScrollBoxStyle>("ScrollBox").LeftShadowBrush))
+				.Image(&(FCoreStyle::Get().GetWidgetStyle<FScrollBoxStyle>(TEXT("ScrollBox")).LeftShadowBrush))
 			]
 			+ SOverlay::Slot()
 			.HAlign(HAlign_Right)
@@ -238,160 +170,7 @@ TSharedRef<SWidget> SAvaEaseCurveTool::ConstructCurveEditorPanel()
 				SNew(SImage)
 				.Visibility(EVisibility::HitTestInvisible)
 				.ColorAndOpacity(FLinearColor(0.f, 0.f, 0.f, 0.5f))
-				.Image(&(FCoreStyle::Get().GetWidgetStyle<FScrollBoxStyle>("ScrollBox").RightShadowBrush))
-			]
-		];
-}
-
-TSharedRef<SWidget> SAvaEaseCurveTool::ConstructInputBoxes()
-{
-	constexpr float WrapSize = 120.f;
-
-	constexpr float MinTangent = -180.0f;
-	constexpr float MaxTangent = 180.0f;
-	constexpr float MinWeight = 0.0f;
-	constexpr float MaxWeight = 10.0f;
-
-	return SNew(SWrapBox)
-		.UseAllottedSize(true)
-		.HAlign(HAlign_Center)
-		+ SWrapBox::Slot()
-		.FillLineWhenSizeLessThan(WrapSize)
-		[
-			SNew(SHorizontalBox)
-			+ SHorizontalBox::Slot()
-			.AutoWidth()
-			.VAlign(VAlign_Center)
-			.Padding(0.f, 1.f, 0.f, 0.f)
-			[
-				SNew(SBox)
-				.WidthOverride(26.f)
-				.HAlign(HAlign_Center)
-				[
-					SNew(STextBlock)
-					.Font(IDetailLayoutBuilder::GetDetailFont())
-					.Text(LOCTEXT("OutLabel", "Out"))
-				]
-			]
-			+ SHorizontalBox::Slot()
-			.FillWidth(1.f)
-			.HAlign(HAlign_Fill)
-			.VAlign(VAlign_Center)
-			[
-				ConstructTangentNumBox(LOCTEXT("OutTangentWeightLabel", "W")
-					, LOCTEXT("OutTangentWeightToolTip", "Out Tangent Weight")
-					, TAttribute<float>::CreateSP(this, &SAvaEaseCurveTool::GetStartTangentWeight)
-					, SNumericEntryBox<float>::FOnValueChanged::CreateSP(this, &SAvaEaseCurveTool::OnStartTangentWeightSpinBoxChanged)
-					, MinWeight, MaxWeight)
-			]
-			+ SHorizontalBox::Slot()
-			.FillWidth(1.f)
-			.HAlign(HAlign_Fill)
-			.VAlign(VAlign_Center)
-			.Padding(2.f, 0.f, 0.f, 0.f)
-			[
-				ConstructTangentNumBox(LOCTEXT("OutTangentLabel", "T")
-					, LOCTEXT("OutTangentToolTip", "Out Tangent")
-					, TAttribute<float>::CreateSP(this, &SAvaEaseCurveTool::GetStartTangent)
-					, SNumericEntryBox<float>::FOnValueChanged::CreateSP(this, &SAvaEaseCurveTool::OnStartTangentSpinBoxChanged)
-					, MinTangent, MaxTangent)
-			]
-		]
-		+ SWrapBox::Slot()
-		.FillLineWhenSizeLessThan(WrapSize)
-		[
-			SNew(SHorizontalBox)
-			+ SHorizontalBox::Slot()
-			.AutoWidth()
-			.VAlign(VAlign_Center)
-			.Padding(0.f, 1.f, 0.f, 0.f)
-			[
-				SNew(SBox)
-				.WidthOverride(26.f)
-				.HAlign(HAlign_Center)
-				[
-					SNew(STextBlock)
-					.Font(IDetailLayoutBuilder::GetDetailFont())
-					.Text(LOCTEXT("InLabel", "In"))
-				]
-			]
-			+ SHorizontalBox::Slot()
-			.FillWidth(1.f)
-			.HAlign(HAlign_Fill)
-			.VAlign(VAlign_Center)
-			[
-				ConstructTangentNumBox(LOCTEXT("InTangentWeightLabel", "W")
-					, LOCTEXT("InTangentWeightToolTip", "In Tangent Weight")
-					, TAttribute<float>::CreateSP(this, &SAvaEaseCurveTool::GetEndTangentWeight)
-					, SNumericEntryBox<float>::FOnValueChanged::CreateSP(this, &SAvaEaseCurveTool::OnEndTangentWeightSpinBoxChanged)
-					, MinWeight, MaxWeight)
-			]
-			+ SHorizontalBox::Slot()
-			.FillWidth(1.f)
-			.HAlign(HAlign_Fill)
-			.VAlign(VAlign_Center)
-			.Padding(2.f, 0.f, 0.f, 0.f)
-			[
-				ConstructTangentNumBox(LOCTEXT("InTangentLabel", "T")
-					, LOCTEXT("InTangentToolTip", "In Tangent")
-					, TAttribute<float>::CreateSP(this, &SAvaEaseCurveTool::GetEndTangent)
-					, SNumericEntryBox<float>::FOnValueChanged::CreateSP(this, &SAvaEaseCurveTool::OnEndTangentSpinBoxChanged)
-					, MinTangent, MaxTangent)
-			]
-		];
-}
-
-TSharedRef<SWidget> SAvaEaseCurveTool::ConstructTangentNumBox(const FText& InLabel
-	, const FText& InToolTip
-	, const TAttribute<float>& InValue
-	, const SNumericEntryBox<float>::FOnValueChanged& InOnValueChanged
-	, const TOptional<float>& InMinSliderValue
-	, const TOptional<float>& InMaxSliderValue) const
-{
-	return SNew(SBox)
-		.MaxDesiredWidth(100.f)
-		.HAlign(HAlign_Fill)
-		[
-			SNew(SHorizontalBox)
-			+ SHorizontalBox::Slot()
-			.AutoWidth()
-			.HAlign(HAlign_Right)
-			[
-				SNew(STextBlock)
-				.MinDesiredWidth(8.f)
-				.Margin(FMargin(2.f, 5.f, 2.f, 3.f))
-				.Font(IDetailLayoutBuilder::GetDetailFont())
-				.Text(InLabel)
-				.ToolTipText(InToolTip)
-			]
-			+ SHorizontalBox::Slot()
-			[
-				SNew(SSpinBox<float>)
-				.Font(IDetailLayoutBuilder::GetDetailFont())
-				.MinSliderValue(InMinSliderValue)
-				.MaxSliderValue(InMaxSliderValue)
-				.Delta(0.00001f)
-				.WheelStep(0.001f)
-				.MinFractionalDigits(4)
-				.MaxFractionalDigits(6)
-				.MinDesiredWidth(70.f)
-				.Value(InValue)
-				.OnBeginSliderMovement_Lambda([this]()
-					{
-						EaseCurveTool->BeginTransaction(LOCTEXT("SliderDragStartLabel", "Ease Curve Slider Drag"));
-					})
-				.OnEndSliderMovement_Lambda([this](const float InNewValue)
-					{
-						EaseCurveTool->EndTransaction();
-					})
-				.OnValueChanged_Lambda([this, InOnValueChanged](const float InNewValue)
-					{
-						InOnValueChanged.ExecuteIfBound(InNewValue);
-					})
-				.OnValueCommitted_Lambda([InOnValueChanged](const float InNewValue, ETextCommit::Type InCommitType)
-					{
-						InOnValueChanged.ExecuteIfBound(InNewValue);
-					})
+				.Image(&(FCoreStyle::Get().GetWidgetStyle<FScrollBoxStyle>(TEXT("ScrollBox")).RightShadowBrush))
 			]
 		];
 }
@@ -416,7 +195,7 @@ void SAvaEaseCurveTool::OnEditorDragEnd() const
 	}
 }
 
-void SAvaEaseCurveTool::SetTangents(const FAvaEaseCurveTangents& InTangents, FAvaEaseCurveTool::EOperation InOperation
+void SAvaEaseCurveTool::SetTangents(const FAvaEaseCurveTangents& InTangents, EAvaEaseCurveToolOperation InOperation
 	, const bool bInSetEaseCurve, const bool bInBroadcastUpdate, const bool bInSetSequencerTangents) const
 {
 	if (CurvePresetWidget.IsValid() && !CurvePresetWidget->SetSelectedItem(InTangents))
@@ -436,48 +215,33 @@ void SAvaEaseCurveTool::SetTangents(const FAvaEaseCurveTangents& InTangents, FAv
 	}
 }
 
-float SAvaEaseCurveTool::GetStartTangent() const
-{
-	return EaseCurveTool->GetEaseCurveTangents().Start;
-}
-
-float SAvaEaseCurveTool::GetStartTangentWeight() const
+FAvaEaseCurveTangents SAvaEaseCurveTool::GetTangents() const
 {
-	return EaseCurveTool->GetEaseCurveTangents().StartWeight;
+	return EaseCurveTool->GetEaseCurveTangents();
 }
 
-float SAvaEaseCurveTool::GetEndTangent() const
-{
-	return EaseCurveTool->GetEaseCurveTangents().End;
-}
-
-float SAvaEaseCurveTool::GetEndTangentWeight() const
-{
-	return EaseCurveTool->GetEaseCurveTangents().EndWeight;
-}
-
-void SAvaEaseCurveTool::OnStartTangentSpinBoxChanged(const float InNewValue) const
+void SAvaEaseCurveTool::OnStartTangentSpinBoxChanged(const double InNewValue) const
 {
 	FAvaEaseCurveTangents NewTangents = EaseCurveTool->GetEaseCurveTangents();
 	NewTangents.Start = InNewValue;
 	SetTangents(NewTangents, ToolOperation.Get(), true, true, true);
 }
 
-void SAvaEaseCurveTool::OnStartTangentWeightSpinBoxChanged(const float InNewValue) const
+void SAvaEaseCurveTool::OnStartTangentWeightSpinBoxChanged(const double InNewValue) const
 {
 	FAvaEaseCurveTangents NewTangents = EaseCurveTool->GetEaseCurveTangents();
 	NewTangents.StartWeight = InNewValue;
 	SetTangents(NewTangents, ToolOperation.Get(), true, true, true);
 }
 
-void SAvaEaseCurveTool::OnEndTangentSpinBoxChanged(const float InNewValue) const
+void SAvaEaseCurveTool::OnEndTangentSpinBoxChanged(const double InNewValue) const
 {
 	FAvaEaseCurveTangents NewTangents = EaseCurveTool->GetEaseCurveTangents();
 	NewTangents.End = InNewValue;
 	SetTangents(NewTangents, ToolOperation.Get(), true, true, true);
 }
 
-void SAvaEaseCurveTool::OnEndTangentWeightSpinBoxChanged(const float InNewValue) const
+void SAvaEaseCurveTool::OnEndTangentWeightSpinBoxChanged(const double InNewValue) const
 {
 	FAvaEaseCurveTangents NewTangents = EaseCurveTool->GetEaseCurveTangents();
 	NewTangents.EndWeight = InNewValue;
@@ -492,7 +256,10 @@ void SAvaEaseCurveTool::OnPresetChanged(const TSharedPtr<FAvaEaseCurvePreset>& I
 		return;
 	}
 
-	SetTangents(InPreset->Tangents, ToolOperation.Get(), true, true, true);
+	if (InPreset.IsValid())
+	{
+		SetTangents(InPreset->Tangents, ToolOperation.Get(), true, true, true);
+	}
 
 	FSlateApplication::Get().SetAllUserFocus(CurveEaseEditorWidget);
 
@@ -554,51 +321,47 @@ void SAvaEaseCurveTool::BindCommands()
 		, FCanExecuteAction()
 		, FIsActionChecked::CreateUObject(GetDefault<UAvaEaseCurveToolSettings>(), &UAvaEaseCurveToolSettings::GetAutoZoomToFit));
 
-	CommandList->MapAction(EaseCurveToolCommands.SelectNextChannelKey, FExecuteAction::CreateSP(EaseCurveToolRef, &FAvaEaseCurveTool::SelectNextChannelKey));
-	
-	CommandList->MapAction(EaseCurveToolCommands.SelectPreviousChannelKey, FExecuteAction::CreateSP(EaseCurveToolRef, &FAvaEaseCurveTool::SelectPreviousChannelKey));
-
 	CommandList->MapAction(EaseCurveToolCommands.SetOperationToEaseOut
-		, FExecuteAction::CreateSP(EaseCurveToolRef, &FAvaEaseCurveTool::SetToolOperation, FAvaEaseCurveTool::EOperation::Out)
+		, FExecuteAction::CreateSP(EaseCurveToolRef, &FAvaEaseCurveTool::SetToolOperation, EAvaEaseCurveToolOperation::Out)
 		, FCanExecuteAction()
-		, FIsActionChecked::CreateSP(EaseCurveToolRef, &FAvaEaseCurveTool::IsToolOperation, FAvaEaseCurveTool::EOperation::Out));
+		, FIsActionChecked::CreateSP(EaseCurveToolRef, &FAvaEaseCurveTool::IsToolOperation, EAvaEaseCurveToolOperation::Out));
 
 	CommandList->MapAction(EaseCurveToolCommands.SetOperationToEaseInOut
-		, FExecuteAction::CreateSP(EaseCurveToolRef, &FAvaEaseCurveTool::SetToolOperation, FAvaEaseCurveTool::EOperation::InOut)
+		, FExecuteAction::CreateSP(EaseCurveToolRef, &FAvaEaseCurveTool::SetToolOperation, EAvaEaseCurveToolOperation::InOut)
 		, FCanExecuteAction()
-		, FIsActionChecked::CreateSP(EaseCurveToolRef, &FAvaEaseCurveTool::IsToolOperation, FAvaEaseCurveTool::EOperation::InOut));
+		, FIsActionChecked::CreateSP(EaseCurveToolRef, &FAvaEaseCurveTool::IsToolOperation, EAvaEaseCurveToolOperation::InOut));
 
 	CommandList->MapAction(EaseCurveToolCommands.SetOperationToEaseIn
-		, FExecuteAction::CreateSP(EaseCurveToolRef, &FAvaEaseCurveTool::SetToolOperation, FAvaEaseCurveTool::EOperation::In)
+		, FExecuteAction::CreateSP(EaseCurveToolRef, &FAvaEaseCurveTool::SetToolOperation, EAvaEaseCurveToolOperation::In)
 		, FCanExecuteAction()
-		, FIsActionChecked::CreateSP(EaseCurveToolRef, &FAvaEaseCurveTool::IsToolOperation, FAvaEaseCurveTool::EOperation::In));
+		, FIsActionChecked::CreateSP(EaseCurveToolRef, &FAvaEaseCurveTool::IsToolOperation, EAvaEaseCurveToolOperation::In));
 
 	CommandList->MapAction(EaseCurveToolCommands.ResetTangents
-		, FExecuteAction::CreateSP(EaseCurveToolRef, &FAvaEaseCurveTool::ResetEaseCurveTangents, FAvaEaseCurveTool::EOperation::InOut));
+		, FExecuteAction::CreateSP(EaseCurveToolRef, &FAvaEaseCurveTool::ResetEaseCurveTangents, EAvaEaseCurveToolOperation::InOut));
 
 	CommandList->MapAction(EaseCurveToolCommands.ResetStartTangent
-		, FExecuteAction::CreateSP(EaseCurveToolRef, &FAvaEaseCurveTool::ResetEaseCurveTangents, FAvaEaseCurveTool::EOperation::Out));
+		, FExecuteAction::CreateSP(EaseCurveToolRef, &FAvaEaseCurveTool::ResetEaseCurveTangents, EAvaEaseCurveToolOperation::Out));
 
 	CommandList->MapAction(EaseCurveToolCommands.ResetEndTangent
-		, FExecuteAction::CreateSP(EaseCurveToolRef, &FAvaEaseCurveTool::ResetEaseCurveTangents, FAvaEaseCurveTool::EOperation::In));
+		, FExecuteAction::CreateSP(EaseCurveToolRef, &FAvaEaseCurveTool::ResetEaseCurveTangents, EAvaEaseCurveToolOperation::In));
 
 	CommandList->MapAction(EaseCurveToolCommands.FlattenTangents
-		, FExecuteAction::CreateSP(EaseCurveToolRef, &FAvaEaseCurveTool::FlattenOrStraightenTangents, FAvaEaseCurveTool::EOperation::InOut, true));
+		, FExecuteAction::CreateSP(EaseCurveToolRef, &FAvaEaseCurveTool::FlattenOrStraightenTangents, EAvaEaseCurveToolOperation::InOut, true));
 
 	CommandList->MapAction(EaseCurveToolCommands.FlattenStartTangent
-		, FExecuteAction::CreateSP(EaseCurveToolRef, &FAvaEaseCurveTool::FlattenOrStraightenTangents, FAvaEaseCurveTool::EOperation::Out, true));
+		, FExecuteAction::CreateSP(EaseCurveToolRef, &FAvaEaseCurveTool::FlattenOrStraightenTangents, EAvaEaseCurveToolOperation::Out, true));
 
 	CommandList->MapAction(EaseCurveToolCommands.FlattenEndTangent
-		, FExecuteAction::CreateSP(EaseCurveToolRef, &FAvaEaseCurveTool::FlattenOrStraightenTangents, FAvaEaseCurveTool::EOperation::In, true));
+		, FExecuteAction::CreateSP(EaseCurveToolRef, &FAvaEaseCurveTool::FlattenOrStraightenTangents, EAvaEaseCurveToolOperation::In, true));
 
 	CommandList->MapAction(EaseCurveToolCommands.StraightenTangents
-		, FExecuteAction::CreateSP(EaseCurveToolRef, &FAvaEaseCurveTool::FlattenOrStraightenTangents, FAvaEaseCurveTool::EOperation::InOut, false));
+		, FExecuteAction::CreateSP(EaseCurveToolRef, &FAvaEaseCurveTool::FlattenOrStraightenTangents, EAvaEaseCurveToolOperation::InOut, false));
 
 	CommandList->MapAction(EaseCurveToolCommands.StraightenStartTangent
-		, FExecuteAction::CreateSP(EaseCurveToolRef, &FAvaEaseCurveTool::FlattenOrStraightenTangents, FAvaEaseCurveTool::EOperation::Out, false));
+		, FExecuteAction::CreateSP(EaseCurveToolRef, &FAvaEaseCurveTool::FlattenOrStraightenTangents, EAvaEaseCurveToolOperation::Out, false));
 
 	CommandList->MapAction(EaseCurveToolCommands.StraightenEndTangent
-		, FExecuteAction::CreateSP(EaseCurveToolRef, &FAvaEaseCurveTool::FlattenOrStraightenTangents, FAvaEaseCurveTool::EOperation::In, false));
+		, FExecuteAction::CreateSP(EaseCurveToolRef, &FAvaEaseCurveTool::FlattenOrStraightenTangents, EAvaEaseCurveToolOperation::In, false));
 
 	CommandList->MapAction(EaseCurveToolCommands.CopyTangents
 		, FExecuteAction::CreateSP(EaseCurveToolRef, &FAvaEaseCurveTool::CopyTangentsToClipboard)
@@ -646,268 +409,13 @@ void SAvaEaseCurveTool::BindCommands()
 		FIsActionChecked::CreateSP(EaseCurveToolRef, &FAvaEaseCurveTool::IsKeyInterpMode, RCIM_Cubic, RCTM_Break));
 
 	CommandList->MapAction(EaseCurveToolCommands.QuickEase
-		, FExecuteAction::CreateSP(EaseCurveToolRef, &FAvaEaseCurveTool::ApplyQuickEaseToSequencerKeySelections, FAvaEaseCurveTool::EOperation::InOut));
+		, FExecuteAction::CreateSP(EaseCurveToolRef, &FAvaEaseCurveTool::ApplyQuickEaseToSequencerKeySelections, EAvaEaseCurveToolOperation::InOut));
 
 	CommandList->MapAction(EaseCurveToolCommands.QuickEaseIn
-		, FExecuteAction::CreateSP(EaseCurveToolRef, &FAvaEaseCurveTool::ApplyQuickEaseToSequencerKeySelections, FAvaEaseCurveTool::EOperation::In));
+		, FExecuteAction::CreateSP(EaseCurveToolRef, &FAvaEaseCurveTool::ApplyQuickEaseToSequencerKeySelections, EAvaEaseCurveToolOperation::In));
 
 	CommandList->MapAction(EaseCurveToolCommands.QuickEaseOut
-		, FExecuteAction::CreateSP(EaseCurveToolRef, &FAvaEaseCurveTool::ApplyQuickEaseToSequencerKeySelections, FAvaEaseCurveTool::EOperation::Out));
-}
-
-TSharedRef<SWidget> SAvaEaseCurveTool::CreateContextMenuContent()
-{
-	UToolMenus* const ToolMenus = UToolMenus::Get();
-	if (!ToolMenus)
-	{
-		return SNullWidget::NullWidget;
-	}
-
-	constexpr const TCHAR* MenuName = TEXT("AvaEaseCurveToolMenu");
-
-	if (!ToolMenus->IsMenuRegistered(MenuName))
-	{
-		UToolMenu* const ToolMenu = ToolMenus->RegisterMenu(MenuName, NAME_None, EMultiBoxType::Menu);
-
-		FToolMenuSection& Section = ToolMenu->FindOrAddSection(TEXT("EaseCurveTool"), LOCTEXT("EaseCurveToolActions", "Curve Ease Tool Actions"));
-
-		const FAvaEaseCurveToolCommands& EaseCurveToolCommands = FAvaEaseCurveToolCommands::Get();
-
-		Section.AddSubMenu(TEXT("Settings"),
-			LOCTEXT("SettingsSubMenuLabel", "Settings"),
-			LOCTEXT("SettingsSubMenuToolTip", ""),
-			FNewToolMenuDelegate::CreateSP(this, &SAvaEaseCurveTool::MakeContextMenuSettings),
-			false,
-			FSlateIcon(FAppStyle::Get().GetStyleSetName(), TEXT("Icons.Toolbar.Settings")));
-
-		Section.AddSeparator(NAME_None);
-
-		Section.AddMenuEntry(EaseCurveToolCommands.CreateExternalCurveAsset);
-
-		Section.AddSeparator(NAME_None);
-
-		Section.AddMenuEntry(EaseCurveToolCommands.CopyTangents);
-
-		Section.AddMenuEntry(EaseCurveToolCommands.PasteTangents);
-
-		Section.AddSeparator(NAME_None);
-
-		Section.AddSubMenu(TEXT("StraightenTangents"),
-			LOCTEXT("StraightenTangentsSubMenuLabel", "Straighten Tangents"),
-			LOCTEXT("StraightenTangentsSubMenuToolTip", ""),
-			FNewToolMenuDelegate::CreateLambda([&EaseCurveToolCommands](UToolMenu* InToolMenu)
-				{
-					FToolMenuSection& NewSection = InToolMenu->FindOrAddSection(TEXT("StraightenTangents"));
-					NewSection.AddMenuEntry(EaseCurveToolCommands.StraightenTangents);
-					NewSection.AddSeparator(NAME_None);
-					NewSection.AddMenuEntry(EaseCurveToolCommands.StraightenStartTangent);
-					NewSection.AddMenuEntry(EaseCurveToolCommands.StraightenEndTangent);
-				}),
-			false, FSlateIcon(FAppStyle::Get().GetStyleSetName(), TEXT("GenericCurveEditor.StraightenTangents")));
-
-		Section.AddSubMenu(TEXT("FlattenTangents"),
-			LOCTEXT("FlattenTangentsSubMenuLabel", "Flatten Tangents"),
-			LOCTEXT("FlattenTangentsSubMenuToolTip", ""),
-			FNewToolMenuDelegate::CreateLambda([&EaseCurveToolCommands](UToolMenu* InToolMenu)
-				{
-					FToolMenuSection& NewSection = InToolMenu->FindOrAddSection(TEXT("FlattenTangents"));
-					NewSection.AddMenuEntry(EaseCurveToolCommands.FlattenTangents);
-					NewSection.AddSeparator(NAME_None);
-					NewSection.AddMenuEntry(EaseCurveToolCommands.FlattenStartTangent);
-					NewSection.AddMenuEntry(EaseCurveToolCommands.FlattenEndTangent);
-				}), 
-			false, FSlateIcon(FAppStyle::Get().GetStyleSetName(), TEXT("GenericCurveEditor.FlattenTangents")));
-
-		Section.AddSubMenu(TEXT("ResetTangents"),
-			LOCTEXT("ResetTangentsSubMenuLabel", "Reset Tangents"),
-			LOCTEXT("ResetTangentsSubMenuToolTip", ""),
-			FNewToolMenuDelegate::CreateLambda([&EaseCurveToolCommands](UToolMenu* InToolMenu)
-				{
-					FToolMenuSection& NewSection = InToolMenu->FindOrAddSection(TEXT("ResetTangents"));
-					NewSection.AddMenuEntry(EaseCurveToolCommands.ResetTangents);
-					NewSection.AddSeparator(NAME_None);
-					NewSection.AddMenuEntry(EaseCurveToolCommands.ResetStartTangent);
-					NewSection.AddMenuEntry(EaseCurveToolCommands.ResetEndTangent);
-				}),
-			false, FSlateIcon(FAppStyle::Get().GetStyleSetName(), TEXT("PropertyWindow.DiffersFromDefault")));
-
-		Section.AddSeparator(NAME_None);
-
-		/** @TODO: Only show these in single edit mode(?)
-		Section.AddMenuEntry(EaseCurveToolCommands.SetKeyInterpConstant);
-
-		Section.AddMenuEntry(EaseCurveToolCommands.SetKeyInterpLinear);
-
-		Section.AddMenuEntry(EaseCurveToolCommands.SetKeyInterpCubicAuto);
-
-		Section.AddMenuEntry(EaseCurveToolCommands.SetKeyInterpCubicSmartAuto);
-
-		Section.AddMenuEntry(EaseCurveToolCommands.SetKeyInterpCubicUser);
-
-		Section.AddMenuEntry(EaseCurveToolCommands.SetKeyInterpCubicBreak);
-
-		Section.AddSeparator(NAME_None);*/
-
-		Section.AddMenuEntry(EaseCurveToolCommands.SetOperationToEaseOut);
-
-		Section.AddMenuEntry(EaseCurveToolCommands.SetOperationToEaseInOut);
-
-		Section.AddMenuEntry(EaseCurveToolCommands.SetOperationToEaseIn);
-
-		Section.AddSeparator(NAME_None);
-
-		Section.AddMenuEntry(EaseCurveToolCommands.SelectPreviousChannelKey);
-
-		Section.AddMenuEntry(EaseCurveToolCommands.SelectNextChannelKey);
-
-		Section.AddSeparator(NAME_None);
-
-		Section.AddMenuEntry(EaseCurveToolCommands.ToggleGridSnap);
-
-		Section.AddMenuEntry(EaseCurveToolCommands.ZoomToFit);
-
-		Section.AddSeparator(NAME_None);
-
-		Section.AddMenuEntry(EaseCurveToolCommands.Refresh);
-
-		Section.AddMenuEntry(EaseCurveToolCommands.Apply);
-	}
-
-	return ToolMenus->GenerateWidget(MenuName, FToolMenuContext(CommandList));
-}
-
-void SAvaEaseCurveTool::MakeContextMenuSettings(UToolMenu* const InToolMenu)
-{
-	if (!IsValid(InToolMenu))
-	{
-		return;
-	}
-
-	const FAvaEaseCurveToolCommands& EaseCurveToolCommands = FAvaEaseCurveToolCommands::Get();
-
-	FToolMenuSection& Section = InToolMenu->FindOrAddSection(TEXT("EaseCurveToolSettings"), LOCTEXT("EaseCurveToolSettingsActions", "Settings"));
-
-	Section.AddMenuEntry(EaseCurveToolCommands.OpenToolSettings);
-
-	Section.AddSeparator(NAME_None);
-
-	Section.AddMenuEntry(EaseCurveToolCommands.ToggleAutoFlipTangents);
-
-	Section.AddSeparator(NAME_None);
-
-	// Graph Size
-	{
-		FProperty* const GraphSizeProperty = UAvaEaseCurveToolSettings::StaticClass()->FindPropertyByName(TEXT("GraphSize"));
-		check(GraphSizeProperty);
-
-		TOptional<int32> MinValue, MaxValue, SliderMinValue, SliderMaxValue;
-		int32 SliderExponent, Delta;
-		float ShiftMultiplier = 10.f;
-		float CtrlMultiplier = 0.1f;
-		bool SupportDynamicSliderMaxValue = false;
-		bool SupportDynamicSliderMinValue = false;
-		UE::AvaSequencer::ExtractNumericMetadata(GraphSizeProperty
-			, MinValue, MaxValue
-			, SliderMinValue, SliderMaxValue
-			, SliderExponent, Delta
-			, ShiftMultiplier
-			, CtrlMultiplier
-			, SupportDynamicSliderMaxValue, SupportDynamicSliderMinValue);
-
-		const TSharedRef<SNumericEntryBox<int32>> GraphSizeWidget = SNew(SNumericEntryBox<int32>)
-			.Font(FAppStyle::GetFontStyle(TEXT("MenuItem.Font")))
-			.AllowSpin(true)
-			.MinValue(MinValue)
-			.MaxValue(MaxValue)
-			.MinSliderValue(SliderMinValue)
-			.MaxSliderValue(SliderMaxValue)
-			.SliderExponent(SliderExponent)
-			.Delta(Delta)
-			.ShiftMultiplier(ShiftMultiplier)
-			.CtrlMultiplier(CtrlMultiplier)
-			.SupportDynamicSliderMaxValue(SupportDynamicSliderMaxValue)
-			.SupportDynamicSliderMinValue(SupportDynamicSliderMinValue)
-			.Value_Lambda([this]()
-				{
-					return CurrentGraphSize;
-				})
-			.OnValueChanged_Lambda([this](const int32 InNewValue)
-				{
-					CurrentGraphSize = InNewValue;
-				})
-			.OnValueCommitted_Lambda([this](const int32 InNewValue, ETextCommit::Type InCommitType)
-				{
-					CurrentGraphSize = InNewValue;
-
-					UAvaEaseCurveToolSettings* const EaseCurveToolSettings = GetMutableDefault<UAvaEaseCurveToolSettings>();
-					check(EaseCurveToolSettings);
-					EaseCurveToolSettings->SetGraphSize(CurrentGraphSize);
-					EaseCurveToolSettings->SaveConfig();
-				});
-
-		Section.AddEntry(FToolMenuEntry::InitWidget(TEXT("ToolSize"), GraphSizeWidget, LOCTEXT("ToolSizeLabel", "Tool Size")));
-	}
-
-	// Grid Size
-	{
-		FProperty* const GridSizeProperty = UAvaEaseCurveToolSettings::StaticClass()->FindPropertyByName(TEXT("GridSize"));
-		check(GridSizeProperty);
-
-		TOptional<int32> MinValue, MaxValue, SliderMinValue, SliderMaxValue;
-		int32 SliderExponent, Delta;
-		float ShiftMultiplier = 10.f;
-		float CtrlMultiplier = 0.1f;
-		bool SupportDynamicSliderMaxValue = false;
-		bool SupportDynamicSliderMinValue = false;
-		UE::AvaSequencer::ExtractNumericMetadata(GridSizeProperty
-			, MinValue, MaxValue
-			, SliderMinValue, SliderMaxValue
-			, SliderExponent, Delta
-			, ShiftMultiplier
-			, CtrlMultiplier
-			, SupportDynamicSliderMaxValue, SupportDynamicSliderMinValue);
-
-		auto SetEaseCurveToolGridSize = [](const int32 InNewValue)
-			{
-				UAvaEaseCurveToolSettings* const EaseCurveToolSettings = GetMutableDefault<UAvaEaseCurveToolSettings>();
-				check(EaseCurveToolSettings);
-				EaseCurveToolSettings->SetGridSize(InNewValue);
-				EaseCurveToolSettings->SaveConfig();
-			};
-
-		const TSharedRef<SNumericEntryBox<int32>> GridSizeWidget = SNew(SNumericEntryBox<int32>)
-			.Font(FAppStyle::GetFontStyle(TEXT("MenuItem.Font")))
-			.AllowSpin(true)
-			.MinValue(MinValue)
-			.MaxValue(MaxValue)
-			.MinSliderValue(SliderMinValue)
-			.MaxSliderValue(SliderMaxValue)
-			.SliderExponent(SliderExponent)
-			.Delta(Delta)
-			.ShiftMultiplier(ShiftMultiplier)
-			.CtrlMultiplier(CtrlMultiplier)
-			.SupportDynamicSliderMaxValue(SupportDynamicSliderMaxValue)
-			.SupportDynamicSliderMinValue(SupportDynamicSliderMinValue)
-			.Value_Lambda([this]()
-				{
-					return GetDefault<UAvaEaseCurveToolSettings>()->GetGridSize();
-				})
-			.OnValueChanged_Lambda(SetEaseCurveToolGridSize)
-			.OnValueCommitted_Lambda([SetEaseCurveToolGridSize](const int32 InNewValue, ETextCommit::Type InCommitType)
-				{
-					SetEaseCurveToolGridSize(InNewValue);
-				});
-
-		Section.AddEntry(FToolMenuEntry::InitWidget(TEXT("GridSize"), GridSizeWidget, LOCTEXT("GridSizeLabel", "Grid Size")));
-	}
-
-	Section.AddSeparator(NAME_None);
-
-	Section.AddMenuEntry(EaseCurveToolCommands.ToggleAutoZoomToFit);
-
-	Section.AddSeparator(NAME_None);
-
-	Section.AddMenuEntry(EaseCurveToolCommands.ResetToDefaultPresets);
+		, FExecuteAction::CreateSP(EaseCurveToolRef, &FAvaEaseCurveTool::ApplyQuickEaseToSequencerKeySelections, EAvaEaseCurveToolOperation::Out));
 }
 
 void SAvaEaseCurveTool::UndoAction()
@@ -954,28 +462,28 @@ FReply SAvaEaseCurveTool::OnKeyDown(const FGeometry& MyGeometry, const FKeyEvent
 
 FText SAvaEaseCurveTool::GetStartText() const
 {
-	return (ToolMode.Get(FAvaEaseCurveTool::EMode::DualKeyEdit) == FAvaEaseCurveTool::EMode::DualKeyEdit)
+	return (ToolMode.Get(EAvaEaseCurveToolMode::DualKeyEdit) == EAvaEaseCurveToolMode::DualKeyEdit)
 		? LOCTEXT("StartText", "Leave")
 		: LOCTEXT("ArriveText", "Arrive");
 }
 
 FText SAvaEaseCurveTool::GetStartTooltipText() const
 {
-	return (ToolMode.Get(FAvaEaseCurveTool::EMode::DualKeyEdit) == FAvaEaseCurveTool::EMode::DualKeyEdit)
+	return (ToolMode.Get(EAvaEaseCurveToolMode::DualKeyEdit) == EAvaEaseCurveToolMode::DualKeyEdit)
 		? LOCTEXT("StartTooltipText", "Start: The selected key's leave tangent")
 		: LOCTEXT("ArriveTooltipText", "Arrive");
 }
 
 FText SAvaEaseCurveTool::GetEndText() const
 {
-	return (ToolMode.Get(FAvaEaseCurveTool::EMode::DualKeyEdit) == FAvaEaseCurveTool::EMode::DualKeyEdit)
+	return (ToolMode.Get(EAvaEaseCurveToolMode::DualKeyEdit) == EAvaEaseCurveToolMode::DualKeyEdit)
 		? LOCTEXT("EndText", "Arrive")
 		: LOCTEXT("LeaveText", "Leave");
 }
 
 FText SAvaEaseCurveTool::GetEndTooltipText() const
 {
-	return (ToolMode.Get(FAvaEaseCurveTool::EMode::DualKeyEdit) == FAvaEaseCurveTool::EMode::DualKeyEdit)
+	return (ToolMode.Get(EAvaEaseCurveToolMode::DualKeyEdit) == EAvaEaseCurveToolMode::DualKeyEdit)
 		? LOCTEXT("EndTooltipText", "End: The next key's arrive tangent")
 		: LOCTEXT("LeaveTooltipText", "Leave");
 }
@@ -1002,7 +510,7 @@ void SAvaEaseCurveTool::ResetTangentsAndNotify() const
 {
 	CurvePresetWidget->ClearSelection();
 
-	SetTangents(FAvaEaseCurveTangents(), FAvaEaseCurveTool::EOperation::InOut, true, true, false);
+	SetTangents(FAvaEaseCurveTangents(), EAvaEaseCurveToolOperation::InOut, true, true, false);
 
 	FAvaEaseCurveTool::ShowNotificationMessage(LOCTEXT("EqualValueKeys", "No different key values to create ease curve!"));
 }
diff --git a/Engine/Plugins/Experimental/Avalanche/Source/AvalancheSequencer/Private/EaseCurveTool/Widgets/SAvaEaseCurveTool.h b/Engine/Plugins/Experimental/Avalanche/Source/AvalancheSequencer/Private/EaseCurveTool/Widgets/SAvaEaseCurveTool.h
--- a/Engine/Plugins/Experimental/Avalanche/Source/AvalancheSequencer/Private/EaseCurveTool/Widgets/SAvaEaseCurveTool.h
+++ b/Engine/Plugins/Experimental/Avalanche/Source/AvalancheSequencer/Private/EaseCurveTool/Widgets/SAvaEaseCurveTool.h
@@ -6,11 +6,12 @@
 #include "EditorUndoClient.h"
 #include "Templates/SharedPointer.h"
 #include "Widgets/DeclarativeSyntaxSupport.h"
-#include "Widgets/Input/SNumericEntryBox.h"
 #include "Widgets/SCompoundWidget.h"
 
+class FAvaEaseCurveToolContextMenu;
 class FCurveEditor;
 class FText;
+class FUICommandList;
 class SAvaEaseCurveEditor;
 class SAvaEaseCurvePreset;
 class UCurveBase;
@@ -23,24 +24,23 @@ class SAvaEaseCurveTool
 	, public FEditorUndoClient
 {
 public:
+	static constexpr int32 DefaultGraphSize = 200;
+	
 	SLATE_BEGIN_ARGS(SAvaEaseCurveTool)
-		: _ToolMode(FAvaEaseCurveTool::EMode::DualKeyEdit)
-		, _ToolOperation(FAvaEaseCurveTool::EOperation::InOut)
+		: _ToolMode(EAvaEaseCurveToolMode::DualKeyEdit)
+		, _ToolOperation(EAvaEaseCurveToolOperation::InOut)
 	{}
-		SLATE_ATTRIBUTE(FAvaEaseCurveTool::EMode, ToolMode)
-		SLATE_ATTRIBUTE(FAvaEaseCurveTool::EOperation, ToolOperation)
+		SLATE_ATTRIBUTE(EAvaEaseCurveToolMode, ToolMode)
+		SLATE_ATTRIBUTE(EAvaEaseCurveToolOperation, ToolOperation)
 		SLATE_ARGUMENT(FAvaEaseCurveTangents, InitialTangents)
 	SLATE_END_ARGS()
 
 	void Construct(const FArguments& InArgs, const TSharedRef<FAvaEaseCurveTool>& InEaseCurveTool);
 	
-	void SetTangents(const FAvaEaseCurveTangents& InTangents, FAvaEaseCurveTool::EOperation InOperation,
+	void SetTangents(const FAvaEaseCurveTangents& InTangents, EAvaEaseCurveToolOperation InOperation,
 		const bool bInSetEaseCurve, const bool bInBroadcastUpdate, const bool bInSetSequencerTangents) const;
 
-	float GetStartTangent() const;
-	float GetStartTangentWeight() const;
-	float GetEndTangent() const;
-	float GetEndTangentWeight() const;
+	FAvaEaseCurveTangents GetTangents() const;
 
 	FKeyHandle GetSelectedKeyHandle() const;
 
@@ -48,29 +48,19 @@ public:
 
 protected:
 	TSharedRef<SWidget> ConstructCurveEditorPanel();
-	TSharedRef<SWidget> ConstructInputBoxes();
-	TSharedRef<SWidget> ConstructTangentNumBox(const FText& InLabel
-		, const FText& InToolTip
-		, const TAttribute<float>& InValue
-		, const SNumericEntryBox<float>::FOnValueChanged& InOnValueChanged
-		, const TOptional<float>& InMinSliderValue
-		, const TOptional<float>& InMaxSliderValue) const;
 
 	void HandleEditorTangentsChanged(const FAvaEaseCurveTangents& InTangents) const;
 
-	void OnStartTangentSpinBoxChanged(const float InNewValue) const;
-	void OnStartTangentWeightSpinBoxChanged(const float InNewValue) const;
-	void OnEndTangentSpinBoxChanged(const float InNewValue) const;
-	void OnEndTangentWeightSpinBoxChanged(const float InNewValue) const;
+	void OnStartTangentSpinBoxChanged(const double InNewValue) const;
+	void OnStartTangentWeightSpinBoxChanged(const double InNewValue) const;
+	void OnEndTangentSpinBoxChanged(const double InNewValue) const;
+	void OnEndTangentWeightSpinBoxChanged(const double InNewValue) const;
 
 	void OnPresetChanged(const TSharedPtr<FAvaEaseCurvePreset>& InPreset) const;
 	void OnQuickPresetChanged(const TSharedPtr<FAvaEaseCurvePreset>& InPreset) const;
 
 	void BindCommands();
 
-	TSharedRef<SWidget> CreateContextMenuContent();
-	void MakeContextMenuSettings(UToolMenu* const InToolMenu);
-
 	void UndoAction();
 	void RedoAction();
 
@@ -96,11 +86,11 @@ protected:
 
 	TSharedPtr<FAvaEaseCurveTool> EaseCurveTool;
 
-	TAttribute<FAvaEaseCurveTool::EMode> ToolMode;
-	TAttribute<FAvaEaseCurveTool::EOperation> ToolOperation;
+	TAttribute<EAvaEaseCurveToolMode> ToolMode;
+	TAttribute<EAvaEaseCurveToolOperation> ToolOperation;
 
 	TSharedPtr<SAvaEaseCurveEditor> CurveEaseEditorWidget;
 	TSharedPtr<SAvaEaseCurvePreset> CurvePresetWidget;
 
-	float CurrentGraphSize = 200.f;
+	int32 CurrentGraphSize = DefaultGraphSize;
 };
