[Motion Design] IAvaRemoteControlInterface for blueprint API.

diff --git a/Engine/Plugins/Experimental/Avalanche/Source/AvalancheMedia/Private/Playable/AvaPlayable.cpp b/Engine/Plugins/Experimental/Avalanche/Source/AvalancheMedia/Private/Playable/AvaPlayable.cpp
--- a/Engine/Plugins/Experimental/Avalanche/Source/AvalancheMedia/Private/Playable/AvaPlayable.cpp
+++ b/Engine/Plugins/Experimental/Avalanche/Source/AvalancheMedia/Private/Playable/AvaPlayable.cpp
@@ -256,6 +256,8 @@ EAvaPlayableCommandResult UAvaPlayable::UpdateRemoteControlCommand(const TShared
 	// during page edition and the resulting entity values are already captured.
 	InRemoteControlValues->ApplyEntityValuesToRemoteControlPreset(RemoteControlPreset);
 
+	OnRemoteControlValuesApplied();
+
 	return EAvaPlayableCommandResult::Executed;
 }
 
diff --git a/Engine/Plugins/Experimental/Avalanche/Source/AvalancheMedia/Private/Playable/Playables/AvaPlayableLevelStreaming.cpp b/Engine/Plugins/Experimental/Avalanche/Source/AvalancheMedia/Private/Playable/Playables/AvaPlayableLevelStreaming.cpp
--- a/Engine/Plugins/Experimental/Avalanche/Source/AvalancheMedia/Private/Playable/Playables/AvaPlayableLevelStreaming.cpp
+++ b/Engine/Plugins/Experimental/Avalanche/Source/AvalancheMedia/Private/Playable/Playables/AvaPlayableLevelStreaming.cpp
@@ -13,6 +13,7 @@
 #include "Framework/AvaInstanceSettings.h"
 #include "Framework/AvaSoftAssetPtr.h"
 #include "IAvaMediaModule.h"
+#include "IAvaRemoteControlInterface.h"
 #include "Playable/AvaPlayableGroup.h"
 #include "Playable/AvaPlayableGroupManager.h"
 #include "Playback/AvaPlaybackUtils.h"
@@ -625,6 +626,34 @@ void UAvaPlayableLevelStreaming::OnEndPlay()
 	}
 }
 
+void UAvaPlayableLevelStreaming::OnRemoteControlValuesApplied()
+{
+	const ULevel* const Level = LevelStreaming->GetLoadedLevel();
+	if (!IsValid(Level))
+	{
+		return;
+	}
+
+	auto NotifyRemoteControlValuesApplied = 
+		[](UObject* InObject)
+		{
+			check(InObject);
+			if (InObject->GetClass()->ImplementsInterface(UAvaRemoteControlInterface::StaticClass()))
+			{
+				IAvaRemoteControlInterface::Execute_OnValuesApplied(InObject);
+			}
+		};
+
+	for (AActor* Actor : Level->Actors)
+	{
+		if (IsValid(Actor))
+		{
+			NotifyRemoteControlValuesApplied(Actor);
+			Actor->ForEachComponent(/*bNestedComps*/false, NotifyRemoteControlValuesApplied);
+		}
+	}
+}
+
 void UAvaPlayableLevelStreaming::BeginDestroy()
 {
 	UnbindDelegates();
diff --git a/Engine/Plugins/Experimental/Avalanche/Source/AvalancheMedia/Private/Playable/Playables/AvaPlayableLevelStreaming.h b/Engine/Plugins/Experimental/Avalanche/Source/AvalancheMedia/Private/Playable/Playables/AvaPlayableLevelStreaming.h
--- a/Engine/Plugins/Experimental/Avalanche/Source/AvalancheMedia/Private/Playable/Playables/AvaPlayableLevelStreaming.h
+++ b/Engine/Plugins/Experimental/Avalanche/Source/AvalancheMedia/Private/Playable/Playables/AvaPlayableLevelStreaming.h
@@ -33,6 +33,7 @@ protected:
 	virtual bool InitPlayable(const FPlayableCreationInfo& InPlayableInfo) override;
 	virtual void OnPlay() override;
 	virtual void OnEndPlay() override;
+	virtual void OnRemoteControlValuesApplied() override;
 	//~ End UAvaPlayable
 
 	//~ Begin UObject
diff --git a/Engine/Plugins/Experimental/Avalanche/Source/AvalancheMedia/Public/Playable/AvaPlayable.h b/Engine/Plugins/Experimental/Avalanche/Source/AvalancheMedia/Public/Playable/AvaPlayable.h
--- a/Engine/Plugins/Experimental/Avalanche/Source/AvalancheMedia/Public/Playable/AvaPlayable.h
+++ b/Engine/Plugins/Experimental/Avalanche/Source/AvalancheMedia/Public/Playable/AvaPlayable.h
@@ -164,6 +164,9 @@ protected:
 	/** Called by EndPlay for derived classes implementation. */
 	virtual void OnEndPlay() {}
 
+	/** Called by UpdateRemoteControlCommand for derived classes implementation. */
+	virtual void OnRemoteControlValuesApplied() {}
+	
 	void HandleOnSequenceStarted(UAvaSequencePlayer* InSequencePlayer, UAvaSequence* InSequence);
 	void HandleOnSequenceFinished(UAvaSequencePlayer* InSequencePlayer, UAvaSequence* InSequence);
 
diff --git a/Engine/Plugins/Experimental/Avalanche/Source/AvalancheRemoteControl/Public/IAvaRemoteControlInterface.h b/Engine/Plugins/Experimental/Avalanche/Source/AvalancheRemoteControl/Public/IAvaRemoteControlInterface.h
new file mode 100644
--- /dev/null
+++ b/Engine/Plugins/Experimental/Avalanche/Source/AvalancheRemoteControl/Public/IAvaRemoteControlInterface.h
@@ -0,0 +1,28 @@
+// Copyright Epic Games, Inc. All Rights Reserved.
+
+#pragma once
+
+#include "UObject/Interface.h"
+#include "IAvaRemoteControlInterface.generated.h"
+
+/**
+ * Interface for runtime remote control events.
+ */
+UINTERFACE(MinimalAPI, Blueprintable)
+class UAvaRemoteControlInterface : public UInterface
+{
+	GENERATED_BODY()
+};
+
+class AVALANCHEREMOTECONTROL_API IAvaRemoteControlInterface
+{
+	GENERATED_BODY()
+
+public:
+	/**
+	 * Called when the remote control values are applied.
+	 */
+	UFUNCTION(BlueprintNativeEvent, Category = "Remote Control")
+	void OnValuesApplied();
+	virtual void OnValuesApplied_Implementation() {}
+};
\ No newline at end of file
