[StateTreeDebugger] enabled debugger on all desktop platforms

diff --git a/Engine/Plugins/Experimental/Avalanche/Source/AvalancheTransitionEditor/AvalancheTransitionEditor.Build.cs b/Engine/Plugins/Experimental/Avalanche/Source/AvalancheTransitionEditor/AvalancheTransitionEditor.Build.cs
--- a/Engine/Plugins/Experimental/Avalanche/Source/AvalancheTransitionEditor/AvalancheTransitionEditor.Build.cs
+++ b/Engine/Plugins/Experimental/Avalanche/Source/AvalancheTransitionEditor/AvalancheTransitionEditor.Build.cs
@@ -1,5 +1,6 @@
 // Copyright Epic Games, Inc. All Rights Reserved.
 
+using System.Linq;
 using UnrealBuildTool;
 
 public class AvalancheTransitionEditor : ModuleRules
@@ -41,7 +42,8 @@ public class AvalancheTransitionEditor : ModuleRules
             }
         );
 
-        if (Target.Platform == UnrealTargetPlatform.Win64 && (Target.Configuration != UnrealTargetConfiguration.Shipping || Target.bBuildEditor))
+        // Until we split use of TraceServices and TraceAnalysis in the debugger we enable only on platforms supporting both at the moment
+        if (Target.Platform.IsInGroup(UnrealPlatformGroup.Desktop) && (Target.Configuration != UnrealTargetConfiguration.Shipping || Target.bBuildEditor))
         {
             PrivateDefinitions.Add("WITH_STATETREE_DEBUGGER=1");
             PrivateDependencyModuleNames.AddRange(
diff --git a/Engine/Plugins/Runtime/StateTree/Source/StateTreeEditorModule/StateTreeEditorModule.Build.cs b/Engine/Plugins/Runtime/StateTree/Source/StateTreeEditorModule/StateTreeEditorModule.Build.cs
--- a/Engine/Plugins/Runtime/StateTree/Source/StateTreeEditorModule/StateTreeEditorModule.Build.cs
+++ b/Engine/Plugins/Runtime/StateTree/Source/StateTreeEditorModule/StateTreeEditorModule.Build.cs
@@ -60,14 +60,7 @@ namespace UnrealBuildTool.Rules
 				"MessageLog",
 			});
 
-			if (Target.Platform == UnrealTargetPlatform.Win64)
-			{
-				PublicDefinitions.Add("WITH_STATETREE_DEBUGGER=1");
-			}
-			else
-			{
-				PublicDefinitions.Add("WITH_STATETREE_DEBUGGER=0");
-			}
+			PublicDefinitions.Add("WITH_STATETREE_DEBUGGER=1");
 		}
 	}
 }
diff --git a/Engine/Plugins/Runtime/StateTree/Source/StateTreeModule/StateTreeModule.Build.cs b/Engine/Plugins/Runtime/StateTree/Source/StateTreeModule/StateTreeModule.Build.cs
--- a/Engine/Plugins/Runtime/StateTree/Source/StateTreeModule/StateTreeModule.Build.cs
+++ b/Engine/Plugins/Runtime/StateTree/Source/StateTreeModule/StateTreeModule.Build.cs
@@ -1,5 +1,7 @@
 // Copyright Epic Games, Inc. All Rights Reserved.
 
+using System.Linq;
+
 namespace UnrealBuildTool.Rules
 {
 	public class StateTreeModule : ModuleRules
@@ -44,8 +46,8 @@ namespace UnrealBuildTool.Rules
 				);
 			}
 
-			if (Target.Platform == UnrealTargetPlatform.Win64 && 
-				(Target.Configuration != UnrealTargetConfiguration.Shipping || Target.bBuildEditor))
+			// Until we split use of TraceServices and TraceAnalysis in the debugger we enable only on platforms supporting both at the moment
+			if (Target.Platform.IsInGroup(UnrealPlatformGroup.Desktop) && (Target.Configuration != UnrealTargetConfiguration.Shipping || Target.bBuildEditor))
 			{
 				PublicDefinitions.Add("WITH_STATETREE_DEBUGGER=1");
 				PublicDependencyModuleNames.AddRange(
