[Backout] - CL33232683
Original CL Desc
-----------------------------------------------------------------
[StateTreeDebugger] enabled debugger on all desktop platforms

diff --git a/Engine/Plugins/Experimental/Avalanche/Source/AvalancheTransitionEditor/AvalancheTransitionEditor.Build.cs b/Engine/Plugins/Experimental/Avalanche/Source/AvalancheTransitionEditor/AvalancheTransitionEditor.Build.cs
--- a/Engine/Plugins/Experimental/Avalanche/Source/AvalancheTransitionEditor/AvalancheTransitionEditor.Build.cs
+++ b/Engine/Plugins/Experimental/Avalanche/Source/AvalancheTransitionEditor/AvalancheTransitionEditor.Build.cs
@@ -1,6 +1,5 @@
 // Copyright Epic Games, Inc. All Rights Reserved.
 
-using System.Linq;
 using UnrealBuildTool;
 
 public class AvalancheTransitionEditor : ModuleRules
@@ -42,8 +41,7 @@ public class AvalancheTransitionEditor : ModuleRules
             }
         );
 
-        // Until we split use of TraceServices and TraceAnalysis in the debugger we enable only on platforms supporting both at the moment
-        if (Target.Platform.IsInGroup(UnrealPlatformGroup.Desktop) && (Target.Configuration != UnrealTargetConfiguration.Shipping || Target.bBuildEditor))
+        if (Target.Platform == UnrealTargetPlatform.Win64 && (Target.Configuration != UnrealTargetConfiguration.Shipping || Target.bBuildEditor))
         {
             PrivateDefinitions.Add("WITH_STATETREE_DEBUGGER=1");
             PrivateDependencyModuleNames.AddRange(
diff --git a/Engine/Plugins/Runtime/StateTree/Source/StateTreeEditorModule/StateTreeEditorModule.Build.cs b/Engine/Plugins/Runtime/StateTree/Source/StateTreeEditorModule/StateTreeEditorModule.Build.cs
--- a/Engine/Plugins/Runtime/StateTree/Source/StateTreeEditorModule/StateTreeEditorModule.Build.cs
+++ b/Engine/Plugins/Runtime/StateTree/Source/StateTreeEditorModule/StateTreeEditorModule.Build.cs
@@ -60,7 +60,14 @@ namespace UnrealBuildTool.Rules
 				"MessageLog",
 			});
 
-			PublicDefinitions.Add("WITH_STATETREE_DEBUGGER=1");
+			if (Target.Platform == UnrealTargetPlatform.Win64)
+			{
+				PublicDefinitions.Add("WITH_STATETREE_DEBUGGER=1");
+			}
+			else
+			{
+				PublicDefinitions.Add("WITH_STATETREE_DEBUGGER=0");
+			}
 		}
 	}
 }
diff --git a/Engine/Plugins/Runtime/StateTree/Source/StateTreeModule/StateTreeModule.Build.cs b/Engine/Plugins/Runtime/StateTree/Source/StateTreeModule/StateTreeModule.Build.cs
--- a/Engine/Plugins/Runtime/StateTree/Source/StateTreeModule/StateTreeModule.Build.cs
+++ b/Engine/Plugins/Runtime/StateTree/Source/StateTreeModule/StateTreeModule.Build.cs
@@ -1,7 +1,5 @@
 // Copyright Epic Games, Inc. All Rights Reserved.
 
-using System.Linq;
-
 namespace UnrealBuildTool.Rules
 {
 	public class StateTreeModule : ModuleRules
@@ -46,8 +44,8 @@ namespace UnrealBuildTool.Rules
 				);
 			}
 
-			// Until we split use of TraceServices and TraceAnalysis in the debugger we enable only on platforms supporting both at the moment
-			if (Target.Platform.IsInGroup(UnrealPlatformGroup.Desktop) && (Target.Configuration != UnrealTargetConfiguration.Shipping || Target.bBuildEditor))
+			if (Target.Platform == UnrealTargetPlatform.Win64 && 
+				(Target.Configuration != UnrealTargetConfiguration.Shipping || Target.bBuildEditor))
 			{
 				PublicDefinitions.Add("WITH_STATETREE_DEBUGGER=1");
 				PublicDependencyModuleNames.AddRange(
