[Motion Design] Adding Sequence Label in logs to be able to identify the sequences since the name is not user facing anywhere in the UI.

diff --git a/Engine/Plugins/Experimental/Avalanche/Source/AvalancheMedia/Private/Playable/AvaPlayable.cpp b/Engine/Plugins/Experimental/Avalanche/Source/AvalancheMedia/Private/Playable/AvaPlayable.cpp
--- a/Engine/Plugins/Experimental/Avalanche/Source/AvalancheMedia/Private/Playable/AvaPlayable.cpp
+++ b/Engine/Plugins/Experimental/Avalanche/Source/AvalancheMedia/Private/Playable/AvaPlayable.cpp
@@ -59,6 +59,17 @@ namespace UE::AvaPlayable::Private
 		}
 		return TEXT("(nullptr)");
 	}
+	
+	FString GetPrettySequenceInfo(const UAvaSequence* InSequence)
+	{
+		if (InSequence)
+		{
+			return FString::Printf(TEXT("Name:%s, Label:%s"),
+				*InSequence->GetFName().ToString(),
+				*InSequence->GetLabel().ToString());
+		}
+		return TEXT("(nullptr)");
+	}		
 
 	FString GetBriefFrameInfo()
 	{
@@ -350,8 +361,8 @@ void UAvaPlayable::HandleOnSequenceStarted(UAvaSequencePlayer* InSequencePlayer,
 	if (HasSequence(InSequence))
 	{
 		using namespace UE::AvaPlayable::Private;
-		UE_LOG(LogAvaPlayable, Verbose, TEXT("%s Playable {%s}: Sequence \"%s\" started."),
-			*GetBriefFrameInfo(), *GetPrettyPlayableInfo(this), *InSequence->GetFName().ToString());
+		UE_LOG(LogAvaPlayable, Verbose, TEXT("%s Playable {%s}: Sequence {%s} started."),
+			*GetBriefFrameInfo(), *GetPrettyPlayableInfo(this),  *GetPrettySequenceInfo(InSequence));
 		OnSequenceEventDelegate.Broadcast(this, InSequence->GetFName(), EAvaPlayableSequenceEventType::Started);
 	}
 }
@@ -361,8 +372,8 @@ void UAvaPlayable::HandleOnSequenceFinished(UAvaSequencePlayer* InSequencePlayer
 	if (HasSequence(InSequence))
 	{
 		using namespace UE::AvaPlayable::Private;
-		UE_LOG(LogAvaPlayable, Verbose, TEXT("%s Playable {%s}: Sequence \"%s\" finished."),
-			*GetBriefFrameInfo(), *GetPrettyPlayableInfo(this), *InSequence->GetFName().ToString());
+		UE_LOG(LogAvaPlayable, Verbose, TEXT("%s Playable {%s}: Sequence {%s} finished."),
+			*GetBriefFrameInfo(), *GetPrettyPlayableInfo(this),  *GetPrettySequenceInfo(InSequence));
 		OnSequenceEventDelegate.Broadcast(this, InSequence->GetFName(), EAvaPlayableSequenceEventType::Finished);
 	}
 }
