[Motion Design] Changed DisplayName to "Motion Design" instead of "Ava".

diff --git a/Engine/Plugins/Experimental/Avalanche/Source/AvalancheRemoteControl/Public/IAvaRemoteControlInterface.h b/Engine/Plugins/Experimental/Avalanche/Source/AvalancheRemoteControl/Public/IAvaRemoteControlInterface.h
--- a/Engine/Plugins/Experimental/Avalanche/Source/AvalancheRemoteControl/Public/IAvaRemoteControlInterface.h
+++ b/Engine/Plugins/Experimental/Avalanche/Source/AvalancheRemoteControl/Public/IAvaRemoteControlInterface.h
@@ -8,7 +8,7 @@
 /**
  * Interface for runtime remote control events.
  */
-UINTERFACE(MinimalAPI, Blueprintable)
+UINTERFACE(MinimalAPI, Blueprintable, meta=(DisplayName="Motion Design Remote Control Interface"))
 class UAvaRemoteControlInterface : public UInterface
 {
 	GENERATED_BODY()
