[Motion Design] Fix potential crash when a pending visibility test has an expired playable.

diff --git a/Engine/Plugins/Experimental/Avalanche/Source/AvalancheMedia/Private/Playable/AvaPlayableGroup.cpp b/Engine/Plugins/Experimental/Avalanche/Source/AvalancheMedia/Private/Playable/AvaPlayableGroup.cpp
--- a/Engine/Plugins/Experimental/Avalanche/Source/AvalancheMedia/Private/Playable/AvaPlayableGroup.cpp
+++ b/Engine/Plugins/Experimental/Avalanche/Source/AvalancheMedia/Private/Playable/AvaPlayableGroup.cpp
@@ -282,6 +282,7 @@ void UAvaPlayableGroup::FVisibilityRequest::Execute(const UAvaPlayableGroup* InP
 		UE_LOG(LogAvaPlayable, Error,
 			TEXT("%s Failed to Set Visibility to \"%s\" because the playable has become stale. Playable Group: \"%s\"."),
 			*GetBriefFrameInfo(), bShouldBeVisible ? TEXT("true") : TEXT("false"), *InPlayableGroup->GetFullName());
+		return;
 	}
 
 	Playable->SetShouldBeVisible(bShouldBeVisible);
