[MotionDesign] EaseCurveTool:
 - Fix ensure when transacting
 - Fix context menu settings submenu not showing

diff --git a/Engine/Plugins/Experimental/Avalanche/Source/AvalancheSequencer/Private/EaseCurveTool/AvaEaseCurveKeySelection.h b/Engine/Plugins/Experimental/Avalanche/Source/AvalancheSequencer/Private/EaseCurveTool/AvaEaseCurveKeySelection.h
--- a/Engine/Plugins/Experimental/Avalanche/Source/AvalancheSequencer/Private/EaseCurveTool/AvaEaseCurveKeySelection.h
+++ b/Engine/Plugins/Experimental/Avalanche/Source/AvalancheSequencer/Private/EaseCurveTool/AvaEaseCurveKeySelection.h
@@ -36,6 +36,11 @@ struct FAvaEaseCurveKeySelection
 	bool CheckMatchingValues(const bool bInIncludeEqualValueKeys, const FChannelData& InChannelData
 		, const TFunctionRef<bool(const FKeyHandle& /*InKeyHandle*/, const FKeyHandle& /*InNextKeyHandle*/, const FChannelData&)>& InCallable)
 	{
+		if (!InChannelData.Channel.Get())
+		{
+			return false;
+		}
+		
 		TMovieSceneChannelHandle<ChannelHandle> Channel = InChannelData.Channel.Cast<ChannelHandle>();
 		TMovieSceneChannelData<ChannelValue> ChannelData = Channel.Get()->GetData();
 
diff --git a/Engine/Plugins/Experimental/Avalanche/Source/AvalancheSequencer/Private/EaseCurveTool/AvaEaseCurveTool.cpp b/Engine/Plugins/Experimental/Avalanche/Source/AvalancheSequencer/Private/EaseCurveTool/AvaEaseCurveTool.cpp
--- a/Engine/Plugins/Experimental/Avalanche/Source/AvalancheSequencer/Private/EaseCurveTool/AvaEaseCurveTool.cpp
+++ b/Engine/Plugins/Experimental/Avalanche/Source/AvalancheSequencer/Private/EaseCurveTool/AvaEaseCurveTool.cpp
@@ -15,7 +15,6 @@
 #include "ISequencer.h"
 #include "ISettingsModule.h"
 #include "Math/UnrealMathUtility.h"
-#include "MVVM/Selection/Selection.h"
 #include "MVVM/ViewModels/SequencerEditorViewModel.h"
 #include "Settings/AvaSequencerSettings.h"
 #include "Widgets/Notifications/SNotificationList.h"
@@ -123,7 +122,12 @@ void FAvaEaseCurveTool::SetEaseCurveTangents_Internal(const FAvaEaseCurveTangent
 void FAvaEaseCurveTool::SetEaseCurveTangents(const FAvaEaseCurveTangents& InTangents, const EAvaEaseCurveToolOperation InOperation
 	, const bool bInBroadcastUpdate, const bool bInSetSequencerTangents)
 {
-	const FScopedTransaction Transaction(LOCTEXT("SetEaseCurveTangents", "Set Ease Curve Tangents"));
+	if (InTangents == GetEaseCurveTangents())
+	{
+		return;
+	}
+
+	const FScopedTransaction Transaction(LOCTEXT("SetEaseCurveTangents", "Set Ease Curve Tangents"), !GIsTransacting);
 	EaseCurve->Modify();
 
 	SetEaseCurveTangents_Internal(InTangents, InOperation, bInBroadcastUpdate);
diff --git a/Engine/Plugins/Experimental/Avalanche/Source/AvalancheSequencer/Private/EaseCurveTool/Widgets/SAvaEaseCurveTool.cpp b/Engine/Plugins/Experimental/Avalanche/Source/AvalancheSequencer/Private/EaseCurveTool/Widgets/SAvaEaseCurveTool.cpp
--- a/Engine/Plugins/Experimental/Avalanche/Source/AvalancheSequencer/Private/EaseCurveTool/Widgets/SAvaEaseCurveTool.cpp
+++ b/Engine/Plugins/Experimental/Avalanche/Source/AvalancheSequencer/Private/EaseCurveTool/Widgets/SAvaEaseCurveTool.cpp
@@ -15,6 +15,7 @@
 #include "EngineAnalytics.h"
 #include "SAvaEaseCurveTangents.h"
 #include "Framework/Commands/GenericCommands.h"
+#include "Framework/Commands/UIAction.h"
 #include "Framework/Commands/UICommandList.h"
 #include "Framework/MultiBox/MultiBoxBuilder.h"
 #include "Math/UnrealMathUtility.h"
@@ -97,6 +98,12 @@ void SAvaEaseCurveTool::Construct(const FArguments& InArgs, const TSharedRef<FAv
 TSharedRef<SWidget> SAvaEaseCurveTool::ConstructCurveEditorPanel()
 {
 	CurrentGraphSize = GetDefault<UAvaEaseCurveToolSettings>()->GetGraphSize();
+	
+	ContextMenu = MakeShared<FAvaEaseCurveToolContextMenu>(CommandList
+		, FAvaEaseCurveToolOnGraphSizeChanged::CreateLambda([this](const int32 InNewSize)
+			{
+				CurrentGraphSize = InNewSize;
+			}));
 
 	const TSharedRef<FAvaEaseCurveTool> EaseCurveToolRef = EaseCurveTool.ToSharedRef();
 	
@@ -119,15 +126,7 @@ TSharedRef<SWidget> SAvaEaseCurveTool::ConstructCurveEditorPanel()
 				.OnTangentsChanged(this, &SAvaEaseCurveTool::HandleEditorTangentsChanged)
 				.GridSnap_UObject(GetDefault<UAvaEaseCurveToolSettings>(), &UAvaEaseCurveToolSettings::GetGridSnap)
 				.GridSize_UObject(GetDefault<UAvaEaseCurveToolSettings>(), &UAvaEaseCurveToolSettings::GetGridSize)
-				.GetContextMenuContent_Lambda([this]()
-					{
-						TSharedPtr<FAvaEaseCurveToolContextMenu> ContextMenu = MakeShared<FAvaEaseCurveToolContextMenu>(CommandList
-							, FAvaEaseCurveToolOnGraphSizeChanged::CreateLambda([this](const int32 InNewSize)
-								{
-									CurrentGraphSize = InNewSize;
-								}));
-						return ContextMenu->GenerateWidget();
-					})
+				.GetContextMenuContent(ContextMenu.ToSharedRef(), &FAvaEaseCurveToolContextMenu::GenerateWidget)
 				.StartText(this, &SAvaEaseCurveTool::GetStartText)
 				.StartTooltipText(this, &SAvaEaseCurveTool::GetStartTooltipText)
 				.EndText(this, &SAvaEaseCurveTool::GetEndText)
diff --git a/Engine/Plugins/Experimental/Avalanche/Source/AvalancheSequencer/Private/EaseCurveTool/Widgets/SAvaEaseCurveTool.h b/Engine/Plugins/Experimental/Avalanche/Source/AvalancheSequencer/Private/EaseCurveTool/Widgets/SAvaEaseCurveTool.h
--- a/Engine/Plugins/Experimental/Avalanche/Source/AvalancheSequencer/Private/EaseCurveTool/Widgets/SAvaEaseCurveTool.h
+++ b/Engine/Plugins/Experimental/Avalanche/Source/AvalancheSequencer/Private/EaseCurveTool/Widgets/SAvaEaseCurveTool.h
@@ -93,4 +93,6 @@ protected:
 	TSharedPtr<SAvaEaseCurvePreset> CurvePresetWidget;
 
 	int32 CurrentGraphSize = DefaultGraphSize;
+
+	TSharedPtr<FAvaEaseCurveToolContextMenu> ContextMenu;
 };
