[Motion Design] Fix Pinned Channel Resolve in post load.

The profile name was not set until the end of post load which caused the profiles to be seen as invalid during the pinned channel resolution pass. Moving the profile name setup earlier fixes that problem.

diff --git a/Engine/Plugins/Experimental/Avalanche/Source/AvalancheMedia/Private/Broadcast/AvaBroadcast.cpp b/Engine/Plugins/Experimental/Avalanche/Source/AvalancheMedia/Private/Broadcast/AvaBroadcast.cpp
--- a/Engine/Plugins/Experimental/Avalanche/Source/AvalancheMedia/Private/Broadcast/AvaBroadcast.cpp
+++ b/Engine/Plugins/Experimental/Avalanche/Source/AvalancheMedia/Private/Broadcast/AvaBroadcast.cpp
@@ -337,6 +337,9 @@ void UAvaBroadcast::LoadBroadcast()
 		UE::AvaMediaSerializationUtils::SerializeObject(InputFormatter, this);
 		FileReader->Close();
 	}
+
+	// Set the profile names early because it is needed to resolve the pinned channels below.
+	UpdateProfileNames();
 	
 	if (Profiles.Num() > 0)
 	{
@@ -353,7 +356,6 @@ void UAvaBroadcast::LoadBroadcast()
 	}
 	
 	EnsureValidCurrentProfile();
-	UpdateProfileNames();
 }
 
 void UAvaBroadcast::SaveBroadcast()
