[Motion Design] Moving Rundown Server to runtime module.

Details:
All functions still work in game mode/build except for saving the rundown asset which needs the editor.
Added IAvaRundownServer to expose some of the functions for tests, but could also be used to extend the server. I had planned to have some functions in the editor module, but turns out I didn't need to.
I moved the device enumeration (BroadcastOutputTree classes) to runtime as well. The server's ability to enumerate the device is needed in game mode/build too as it turns out from use cases in projects.

The rundown server auto start is still only available in editor mode. Game mode rundown server is only available by launching with the command line -MotionDesignRundownServerStart to avoid conflicts.

Tested with python client in editor mode and game build (with FanFuel project). Will require further testing (device enumeration, broadcast config) for game build, but the basics are working at least.

diff --git a/Engine/Plugins/Experimental/Avalanche/Config/DefaultAvalanche.ini b/Engine/Plugins/Experimental/Avalanche/Config/DefaultAvalanche.ini
--- a/Engine/Plugins/Experimental/Avalanche/Config/DefaultAvalanche.ini
+++ b/Engine/Plugins/Experimental/Avalanche/Config/DefaultAvalanche.ini
@@ -274,5 +274,32 @@
 +PackageRedirects=(OldName="/Avalanche/EditorResources/M_AvalancheChannelVisualizer_Blue.M_AvalancheChannelVisualizer_Blue",NewName="/Avalanche/EditorResources/M_ChannelVisualizer_Blue.M_ChannelVisualizer_Blue")
 +PackageRedirects=(OldName="/Avalanche/EditorResources/M_AvalancheChannelVisualizer_Alpha.M_AvalancheChannelVisualizer_Alpha",NewName="/Avalanche/EditorResources/M_ChannelVisualizer_Alpha.M_ChannelVisualizer_Alpha")
 
+; Redirects for rundown server requests for legacy clients.
++StructRedirects=(OldName="/Script/AvalancheMediaEditor.AvaRundownPing",NewName="/Script/AvalancheMedia.AvaRundownPing")
++StructRedirects=(OldName="/Script/AvalancheMediaEditor.AvaRundownGetRundowns",NewName="/Script/AvalancheMedia.AvaRundownGetRundowns")
++StructRedirects=(OldName="/Script/AvalancheMediaEditor.AvaRundownLoadRundown",NewName="/Script/AvalancheMedia.AvaRundownLoadRundown")
++StructRedirects=(OldName="/Script/AvalancheMediaEditor.AvaRundownSaveRundown",NewName="/Script/AvalancheMedia.AvaRundownSaveRundown")
++StructRedirects=(OldName="/Script/AvalancheMediaEditor.AvaRundownGetPages",NewName="/Script/AvalancheMedia.AvaRundownGetPages")
++StructRedirects=(OldName="/Script/AvalancheMediaEditor.AvaRundownCreatePage",NewName="/Script/AvalancheMedia.AvaRundownCreatePage")
++StructRedirects=(OldName="/Script/AvalancheMediaEditor.AvaRundownCreateTemplate",NewName="/Script/AvalancheMedia.AvaRundownCreateTemplate")
++StructRedirects=(OldName="/Script/AvalancheMediaEditor.AvaRundownDeletePage",NewName="/Script/AvalancheMedia.AvaRundownDeletePage")
++StructRedirects=(OldName="/Script/AvalancheMediaEditor.AvaRundownDeleteTemplate",NewName="/Script/AvalancheMedia.AvaRundownDeleteTemplate")
++StructRedirects=(OldName="/Script/AvalancheMediaEditor.AvaRundownChangeTemplateBP",NewName="/Script/AvalancheMedia.AvaRundownChangeTemplateBP")
++StructRedirects=(OldName="/Script/AvalancheMediaEditor.AvaRundownGetPageDetails",NewName="/Script/AvalancheMedia.AvaRundownGetPageDetails")
++StructRedirects=(OldName="/Script/AvalancheMediaEditor.AvaRundownPageChangeChannel",NewName="/Script/AvalancheMedia.AvaRundownPageChangeChannel")
++StructRedirects=(OldName="/Script/AvalancheMediaEditor.AvaRundownUpdatePageFromRCP",NewName="/Script/AvalancheMedia.AvaRundownUpdatePageFromRCP")
++StructRedirects=(OldName="/Script/AvalancheMediaEditor.AvaRundownPageAction",NewName="/Script/AvalancheMedia.AvaRundownPageAction")
++StructRedirects=(OldName="/Script/AvalancheMediaEditor.AvaRundownPagePreviewAction",NewName="/Script/AvalancheMedia.AvaRundownPagePreviewAction")
++StructRedirects=(OldName="/Script/AvalancheMediaEditor.AvaRundownPageActions",NewName="/Script/AvalancheMedia.AvaRundownPageActions")
++StructRedirects=(OldName="/Script/AvalancheMediaEditor.AvaRundownPagePreviewActions",NewName="/Script/AvalancheMedia.AvaRundownPagePreviewActions")
++StructRedirects=(OldName="/Script/AvalancheMediaEditor.AvaRundownGetChannel",NewName="/Script/AvalancheMedia.AvaRundownGetChannel")
++StructRedirects=(OldName="/Script/AvalancheMediaEditor.AvaRundownGetChannels",NewName="/Script/AvalancheMedia.AvaRundownGetChannels")
++StructRedirects=(OldName="/Script/AvalancheMediaEditor.AvaRundownChannelAction",NewName="/Script/AvalancheMedia.AvaRundownChannelAction")
++StructRedirects=(OldName="/Script/AvalancheMediaEditor.AvaRundownAddChannelDevice",NewName="/Script/AvalancheMedia.AvaRundownAddChannelDevice")
++StructRedirects=(OldName="/Script/AvalancheMediaEditor.AvaRundownEditChannelDevice",NewName="/Script/AvalancheMedia.AvaRundownEditChannelDevice")
++StructRedirects=(OldName="/Script/AvalancheMediaEditor.AvaRundownRemoveChannelDevice",NewName="/Script/AvalancheMedia.AvaRundownRemoveChannelDevice")
++StructRedirects=(OldName="/Script/AvalancheMediaEditor.AvaRundownGetChannelImage",NewName="/Script/AvalancheMedia.AvaRundownGetChannelImage")
++StructRedirects=(OldName="/Script/AvalancheMediaEditor.AvaRundownGetDevices",NewName="/Script/AvalancheMedia.AvaRundownGetDevices")
+
 [/Script/Avalanche.AvaAnimation]
 DefaultCompletionMode=KeepState
diff --git a/Engine/Plugins/Experimental/Avalanche/Source/AvalancheMedia/AvalancheMedia.Build.cs b/Engine/Plugins/Experimental/Avalanche/Source/AvalancheMedia/AvalancheMedia.Build.cs
--- a/Engine/Plugins/Experimental/Avalanche/Source/AvalancheMedia/AvalancheMedia.Build.cs
+++ b/Engine/Plugins/Experimental/Avalanche/Source/AvalancheMedia/AvalancheMedia.Build.cs
@@ -40,6 +40,7 @@ public class AvalancheMedia : ModuleRules
 				"AvalancheRemoteControl",
 				"HeadMountedDisplay",
 				"Json",
+				"ImageCore",
 				"InputCore",
 				"Projects",
 				"RHI",
diff --git a/Engine/Plugins/Experimental/Avalanche/Source/AvalancheMedia/Internal/Broadcast/OutputDevices/AvaBroadcastOutputRootItem.h b/Engine/Plugins/Experimental/Avalanche/Source/AvalancheMedia/Internal/Broadcast/OutputDevices/AvaBroadcastOutputRootItem.h
new file mode 100644
--- /dev/null
+++ b/Engine/Plugins/Experimental/Avalanche/Source/AvalancheMedia/Internal/Broadcast/OutputDevices/AvaBroadcastOutputRootItem.h
@@ -0,0 +1,25 @@
+// Copyright Epic Games, Inc. All Rights Reserved.
+
+#pragma once
+
+#include "AvaBroadcastOutputTreeItem.h"
+
+class AVALANCHEMEDIA_API FAvaBroadcastOutputRootItem : public FAvaBroadcastOutputTreeItem
+{
+public:
+	UE_AVA_INHERITS_WITH_SUPER(FAvaBroadcastOutputRootItem, FAvaBroadcastOutputTreeItem);
+
+	FAvaBroadcastOutputRootItem()
+		: Super(nullptr)
+	{}
+
+	//~ Begin IAvaBroadcastOutputTreeItem
+	virtual void RefreshChildren(const FRefreshChildrenParams& InParams) override;
+private:
+	virtual FText GetDisplayName() const override;
+	virtual const FSlateBrush* GetIconBrush() const override;
+	virtual TSharedPtr<SWidget> GenerateRowWidget() override;
+	virtual FReply OnDragDetected(const FGeometry& MyGeometry, const FPointerEvent& MouseEvent) override;
+	virtual UMediaOutput* AddMediaOutputToChannel(FName InTargetChannel, const FAvaBroadcastMediaOutputInfo& InOutputInfo) override;
+	//~ End IAvaBroadcastOutputTreeItem
+};
diff --git a/Engine/Plugins/Experimental/Avalanche/Source/AvalancheMedia/Internal/Broadcast/OutputDevices/AvaBroadcastOutputTreeItem.h b/Engine/Plugins/Experimental/Avalanche/Source/AvalancheMedia/Internal/Broadcast/OutputDevices/AvaBroadcastOutputTreeItem.h
new file mode 100644
--- /dev/null
+++ b/Engine/Plugins/Experimental/Avalanche/Source/AvalancheMedia/Internal/Broadcast/OutputDevices/AvaBroadcastOutputTreeItem.h
@@ -0,0 +1,82 @@
+// Copyright Epic Games, Inc. All Rights Reserved.
+
+#pragma once
+
+#include "AvaType.h"
+#include "Broadcast/Channel/AvaBroadcastMediaOutputInfo.h"
+
+class FAvaBroadcastOutputTreeItem;
+class FDragDropOperation;
+class FReply;
+class IAvaBroadcastOutputTreeItem;
+class SWidget;
+class UMediaOutput;
+struct FGeometry;
+struct FPointerEvent;
+struct FSlateBrush;
+
+using FAvaOutputTreeItemPtr = TSharedPtr<IAvaBroadcastOutputTreeItem>;
+
+class IAvaBroadcastOutputTreeItem : public IAvaTypeCastable, public TSharedFromThis<IAvaBroadcastOutputTreeItem>
+{
+public:
+	UE_AVA_INHERITS(IAvaBroadcastOutputTreeItem, IAvaTypeCastable);
+
+	virtual FText GetDisplayName() const = 0;
+
+	virtual const FSlateBrush* GetIconBrush() const = 0;
+
+	struct FRefreshChildrenParams
+	{
+		/** Include all the media output classes. If false, only classes with a device provider will be included. */
+		bool bShowAllMediaOutputClasses = false;
+	};
+	
+	/** Refreshes what the Children are of this Item. (not recursive!) */
+	virtual void RefreshChildren(const FRefreshChildrenParams& InParams) = 0;
+
+	virtual TSharedPtr<SWidget> GenerateRowWidget() = 0;
+	
+	virtual const TWeakPtr<FAvaBroadcastOutputTreeItem>& GetParent() const = 0;
+
+	virtual const TArray<FAvaOutputTreeItemPtr>& GetChildren() const = 0;
+
+	virtual FReply OnDragDetected(const FGeometry& MyGeometry, const FPointerEvent& MouseEvent) = 0;
+
+	/** Returns true if it is valid to add this item to the given channel. */
+	virtual bool IsValidToDropInChannel(FName InTargetChannelName) = 0;
+
+	virtual UMediaOutput* AddMediaOutputToChannel(FName InTargetChannel, const FAvaBroadcastMediaOutputInfo& InOutputInfo) = 0;
+
+	/** Delegating the drag and drop operation to editor module. */
+	DECLARE_DELEGATE_RetVal_OneParam(TSharedRef<FDragDropOperation>, FOnCreateDragDropOperation, const FAvaOutputTreeItemPtr&)
+	virtual FOnCreateDragDropOperation& OnCreateDragDropOperation() = 0;
+};
+
+class AVALANCHEMEDIA_API FAvaBroadcastOutputTreeItem : public IAvaBroadcastOutputTreeItem
+{
+public:
+	UE_AVA_INHERITS(FAvaBroadcastOutputTreeItem, IAvaBroadcastOutputTreeItem);
+
+	FAvaBroadcastOutputTreeItem(const TSharedPtr<FAvaBroadcastOutputTreeItem>& InParent)
+		: ParentWeak(InParent)
+	{
+	}
+
+	//~ Begin IAvaBroadcastOutputTreeItem
+	virtual ~FAvaBroadcastOutputTreeItem() override {};
+	virtual const TWeakPtr<FAvaBroadcastOutputTreeItem>& GetParent() const override;
+	virtual const TArray<FAvaOutputTreeItemPtr>& GetChildren() const override;
+	virtual FReply OnDragDetected(const FGeometry& MyGeometry, const FPointerEvent& MouseEvent) override;
+	virtual bool IsValidToDropInChannel(FName InTargetChannelName) override { return true; }
+	virtual FOnCreateDragDropOperation& OnCreateDragDropOperation() override { return OnCreateDragDropOperationDelegate; }
+	//~ End IAvaBroadcastOutputTreeItem
+
+	/** Calls RefreshChildren() on the tree of item. */
+	static void RefreshTree(const FAvaOutputTreeItemPtr& InItem, const FRefreshChildrenParams& InParams);
+
+protected:
+	TWeakPtr<FAvaBroadcastOutputTreeItem> ParentWeak;
+	TArray<FAvaOutputTreeItemPtr> Children;
+	FOnCreateDragDropOperation OnCreateDragDropOperationDelegate;
+};
diff --git a/Engine/Plugins/Experimental/Avalanche/Source/AvalancheMedia/Internal/Rundown/AvaRundownMessages.h b/Engine/Plugins/Experimental/Avalanche/Source/AvalancheMedia/Internal/Rundown/AvaRundownMessages.h
new file mode 100644
--- /dev/null
+++ b/Engine/Plugins/Experimental/Avalanche/Source/AvalancheMedia/Internal/Rundown/AvaRundownMessages.h
@@ -0,0 +1,745 @@
+// Copyright Epic Games, Inc. All Rights Reserved.
+
+#pragma once
+
+#include "Broadcast/Channel/AvaBroadcastMediaOutputInfo.h"
+#include "Rundown/AvaRundownPage.h"
+#include "AvaRundownMessages.generated.h"
+
+namespace EAvaRundownApiVersion
+{
+	/**
+	 * Defines the protocol version of the Rundown Server API.
+	 *
+	 * API versioning is used to provide legacy support either on
+	 * the client side or server side for non compatible changes.
+	 * Clients can request a version of the API that they where implemented against,
+	 * if the server can still honor the request it will accept.
+	 */
+	enum Type
+	{
+		Initial = 1,
+		/**
+		 * The rundown server has been moved to the runtime module.
+		 * All message scripts paths moved from AvalancheMediaEditor to AvalancheMedia.
+		 * However, all server requests messages have been added to core redirect, so
+		 * previous path will still get through, but all response messages will be the new path.
+		 * Clients can still issue a ping with the old path and will get a response.
+		 */ 
+		MoveToRuntime = 2,
+
+		// -----<new versions can be added before this line>-------------------------------------------------
+		// - this needs to be the last line (see note below)
+		VersionPlusOne,
+		LatestVersion = VersionPlusOne - 1
+	};
+}
+
+
+USTRUCT()
+struct FAvaRundownMsgBase
+{
+	GENERATED_BODY()
+public:
+	UPROPERTY()
+	int32 RequestId = -1;
+};
+
+USTRUCT()
+struct FAvaRundownServerMsg : public FAvaRundownMsgBase
+{
+	GENERATED_BODY()
+public:
+	UPROPERTY()
+	FString Verbosity;
+
+	UPROPERTY()
+	FString Text;
+};
+
+/** Request published by client to discover servers. */
+USTRUCT()
+struct FAvaRundownPing : public FAvaRundownMsgBase
+{
+	GENERATED_BODY()
+public:
+	/** True if the request originates from an automatic timer. False if requests originates from user interaction. */
+	UPROPERTY()
+	bool bAuto = true;
+
+	/**
+	 * API Version the client has been implemented against.
+	 * If none (-1) the server will consider the initial version is requested.
+	 */
+	UPROPERTY()
+	int32 RequestedApiVersion = -1;
+};
+
+/** Response sent by server to client to be discovered. */
+USTRUCT()
+struct FAvaRundownPong : public FAvaRundownMsgBase
+{
+	GENERATED_BODY()
+
+public:
+	/** True if it is a reply to an auto ping. Mirrors the bAuto flag from Ping message. */
+	UPROPERTY()
+	bool bAuto = true;
+
+	/**
+	 * API Version the server will communicate with for this client.
+	 * The server may honor the requested version if possible.
+	 * Versions newer than server implementation will obviously not be honored either.
+	 * Clients should expect an older server to reply with an older version.
+	 */
+	UPROPERTY()
+	int32 ApiVersion = -1;
+
+	/** Minimum API Version the server implements. */
+	UPROPERTY()
+	int32 MinimumApiVersion = -1;
+
+	/** Latest API Version the server support. */
+	UPROPERTY()
+	int32 LatestApiVersion = -1;
+
+	UPROPERTY()
+	FString HostName;
+};
+
+/**
+ *	Request list of rundown that can be opened on the current server.
+ */
+USTRUCT()
+struct FAvaRundownGetRundowns : public FAvaRundownMsgBase
+{
+	GENERATED_BODY()
+};
+
+/**
+ *	List of all rundowns.
+ *	Expected Response from FAvaRundownGetRundowns.
+ */
+USTRUCT()
+struct FAvaRundownRundowns : public FAvaRundownMsgBase
+{
+	GENERATED_BODY()
+
+public:
+	UPROPERTY()
+	TArray<FString> Rundowns;
+};
+
+/**
+ *	Request that the given rundown be loaded for playback.
+ *	This will also open an associated playback context.
+ *	Only one rundown can be opened for playback at a time. If another rundown
+ *	is opened, it will be closed and all currently playing pages stopped.
+ *	If the path is empty, nothing will be done and the server will reply with
+ *	a FAvaRundownServerMsg message indicating which rundown is currently loaded.
+ */
+USTRUCT()
+struct FAvaRundownLoadRundown : public FAvaRundownMsgBase
+{
+	GENERATED_BODY()
+
+public:
+	UPROPERTY()
+	FString Rundown;
+};
+
+/**
+ * Request that the given rundown be saved to disk.
+ * The rundown asset must have been loaded, either by an edit command
+ * or playback, prior to this command.
+ * Unloaded assets will not be loaded by this command.
+ */
+USTRUCT()
+struct FAvaRundownSaveRundown : public FAvaRundownMsgBase
+{
+	GENERATED_BODY()
+
+public:
+	UPROPERTY()
+	FString Rundown;
+
+	UPROPERTY()
+	bool bOnlyIfIsDirty = false;
+};
+
+
+/**
+ * Request the list of pages from the given rundown.
+ */
+USTRUCT()
+struct FAvaRundownGetPages : public FAvaRundownMsgBase
+{
+	GENERATED_BODY()
+public:
+	UPROPERTY()
+	FString Rundown;
+};
+
+USTRUCT()
+struct FAvaRundownCreatePage : public FAvaRundownMsgBase
+{
+	GENERATED_BODY()
+
+	UPROPERTY()
+	FString Rundown;
+
+	UPROPERTY()
+	int32 TemplateId = FAvaRundownPage::InvalidPageId;
+};
+
+USTRUCT()
+struct FAvaRundownDeletePage : public FAvaRundownMsgBase
+{
+	GENERATED_BODY()
+
+	UPROPERTY()
+	FString Rundown;
+
+	UPROPERTY()
+	int32 PageId = FAvaRundownPage::InvalidPageId;
+};
+
+USTRUCT()
+struct FAvaRundownCreateTemplate : public FAvaRundownMsgBase
+{
+	GENERATED_BODY()
+
+	UPROPERTY()
+	FString Rundown;
+};
+
+USTRUCT()
+struct FAvaRundownDeleteTemplate : public FAvaRundownMsgBase
+{
+	GENERATED_BODY()
+
+	UPROPERTY()
+	FString Rundown;
+
+	UPROPERTY()
+	int32 PageId = FAvaRundownPage::InvalidPageId;
+};
+
+USTRUCT()
+struct FAvaRundownChangeTemplateBP : public FAvaRundownMsgBase
+{
+	GENERATED_BODY()
+public:
+	UPROPERTY()
+	FString Rundown;
+
+	UPROPERTY()
+	int32 TemplateId = FAvaRundownPage::InvalidPageId;
+
+	UPROPERTY()
+	FString AssetPath;
+};
+
+USTRUCT()
+struct FAvaRundownPageInfo
+{
+	GENERATED_BODY()
+public:
+	UPROPERTY()
+	int32 PageId = FAvaRundownPage::InvalidPageId;
+
+	UPROPERTY()
+	FString PageName;
+
+	UPROPERTY()
+	FString PageSummary;
+
+	UPROPERTY()
+	FString FriendlyName;
+
+	UPROPERTY()
+	bool IsTemplate = false;
+
+	UPROPERTY()
+	int32 TemplateId = FAvaRundownPage::InvalidPageId;
+
+	UPROPERTY()
+	TArray<int32> CombinedTemplateIds;
+
+	UPROPERTY()
+	FSoftObjectPath AssetPath;
+
+	UPROPERTY()
+	TArray<FAvaRundownChannelPageStatus> Statuses;
+
+	UPROPERTY()
+	FString TransitionLayerName;
+
+	UPROPERTY()
+	FString OutputChannel;
+
+	UPROPERTY()
+	bool bIsEnabled = false;
+
+	UPROPERTY()
+	bool bIsPlaying = false;
+};
+
+/*
+ * List of pages from the current rundown.
+ */
+USTRUCT()
+struct FAvaRundownPages : public FAvaRundownMsgBase
+{
+	GENERATED_BODY()
+public:
+	UPROPERTY()
+	TArray<FAvaRundownPageInfo> Pages;
+};
+
+/**
+ * Request the page details from the given rundown.
+ */
+USTRUCT()
+struct FAvaRundownGetPageDetails : public FAvaRundownMsgBase
+{
+	GENERATED_BODY()
+
+public:
+	UPROPERTY()
+	FString Rundown;
+
+	UPROPERTY()
+	int32 PageId = FAvaRundownPage::InvalidPageId;
+
+	/** This will request that a managed asset instance gets loaded to be
+	 * accessible through WebRC. */
+	UPROPERTY()
+	bool bLoadRemoteControlPreset = false;
+};
+
+/**
+ *	Server response to FAvaRundownGetPageDetails request.
+ */
+USTRUCT()
+struct FAvaRundownPageDetails : public FAvaRundownMsgBase
+{
+	GENERATED_BODY()
+
+public:
+	UPROPERTY()
+	FString Rundown;
+
+	UPROPERTY()
+	FAvaRundownPageInfo PageInfo;
+
+	UPROPERTY()
+	FAvaPlayableRemoteControlValues RemoteControlValues;
+
+	/** Name of the remote control preset to resolve through WebRC API. */
+	UPROPERTY()
+	FString RemoteControlPresetName;
+
+	UPROPERTY()
+	FString RemoteControlPresetId;
+};
+
+USTRUCT()
+struct FAvaRundownPagesStatuses : public FAvaRundownMsgBase
+{
+	GENERATED_BODY()
+
+public:
+	UPROPERTY()
+	FString Rundown;
+
+	UPROPERTY()
+	FAvaRundownPageInfo PageInfo;
+};
+
+USTRUCT()
+struct FAvaRundownPageListChanged : public FAvaRundownMsgBase
+{
+	GENERATED_BODY()
+
+	UPROPERTY()
+	FString Rundown;
+
+	/** See EAvaPageListChange flags. */
+	UPROPERTY()
+	uint8 ChangeType = 0;
+
+	UPROPERTY();
+	TArray<int32> AffectedPages;
+};
+
+USTRUCT()
+struct FAvaRundownPageBlueprintChanged : public FAvaRundownMsgBase
+{
+	GENERATED_BODY()
+
+	UPROPERTY()
+	FString Rundown;
+
+	UPROPERTY()
+	int32 PageId = FAvaRundownPage::InvalidPageId;
+
+	UPROPERTY()
+	FString BlueprintPath;
+};
+
+USTRUCT()
+struct FAvaRundownPageChannelChanged : public FAvaRundownMsgBase
+{
+	GENERATED_BODY()
+
+	UPROPERTY()
+	FString Rundown;
+
+	UPROPERTY()
+	int32 PageId = FAvaRundownPage::InvalidPageId;
+
+	UPROPERTY()
+	FString ChannelName;
+};
+
+USTRUCT()
+struct FAvaRundownPageAnimSettingsChanged : public FAvaRundownMsgBase
+{
+	GENERATED_BODY()
+
+	UPROPERTY()
+	FString Rundown;
+
+	UPROPERTY()
+	int32 PageId = FAvaRundownPage::InvalidPageId;
+};
+
+USTRUCT()
+struct FAvaRundownPageChangeChannel : public FAvaRundownMsgBase
+{
+	GENERATED_BODY()
+
+	UPROPERTY()
+	FString Rundown;
+
+	UPROPERTY()
+	int32 PageId = FAvaRundownPage::InvalidPageId;
+
+	UPROPERTY()
+	FString ChannelName;
+};
+
+/** This is a request to save the managed RCP back to the corresponding page. */
+USTRUCT()
+struct FAvaRundownUpdatePageFromRCP : public FAvaRundownMsgBase
+{
+	GENERATED_BODY()
+
+public:
+	/** Unregister the Remote Control Preset from the WebRC. */
+	UPROPERTY()
+	bool bUnregister = false;
+};
+
+/** Supported Page actions for playback. */
+UENUM()
+enum class EAvaRundownPageActions
+{
+	None,
+	Load,
+	Unload,
+	Play,
+	PlayNext,
+	Stop,
+	ForceStop,
+	Continue,
+	UpdateValues,
+	TakeToProgram
+};
+
+USTRUCT()
+struct FAvaRundownPageAction : public FAvaRundownMsgBase
+{
+	GENERATED_BODY()
+public:
+	UPROPERTY()
+	int32 PageId = FAvaRundownPage::InvalidPageId;
+
+	UPROPERTY()
+	EAvaRundownPageActions Action = EAvaRundownPageActions::None;
+};
+
+USTRUCT()
+struct FAvaRundownPagePreviewAction : public FAvaRundownPageAction
+{
+	GENERATED_BODY()
+public:
+	/** Specify which preview channel to use. If left empty, the rundown's default preview channel is used. */
+	UPROPERTY()
+	FString PreviewChannelName;
+};
+
+/**
+ * Command to execute an action on multiple pages at the same time.
+ * This is necessary for pages to be part of the same transition.
+ */
+USTRUCT()
+struct FAvaRundownPageActions : public FAvaRundownMsgBase
+{
+	GENERATED_BODY()
+public:
+	UPROPERTY()
+	TArray<int32> PageIds;
+
+	UPROPERTY()
+	EAvaRundownPageActions Action = EAvaRundownPageActions::None;
+};
+
+USTRUCT()
+struct FAvaRundownPagePreviewActions : public FAvaRundownPageActions
+{
+	GENERATED_BODY()
+public:
+	/** Specify which preview channel to use. If left empty, the rundown's default preview channel is used. */
+	UPROPERTY()
+	FString PreviewChannelName;
+};
+
+UENUM()
+enum class EAvaRundownPageEvents
+{
+	None,
+	AnimStarted,
+	AnimPaused,
+	AnimFinished
+};
+
+USTRUCT()
+struct FAvaRundownPageEvent : public FAvaRundownMsgBase
+{
+	GENERATED_BODY()
+public:
+	UPROPERTY()
+	int32 PageId = FAvaRundownPage::InvalidPageId;
+
+	UPROPERTY()
+	EAvaRundownPageEvents Event = EAvaRundownPageEvents::None;
+};
+
+USTRUCT()
+struct FAvaRundownOutputDeviceItem
+{
+	GENERATED_BODY()
+public:
+	UPROPERTY()
+	FString Name;
+
+	UPROPERTY()
+	FAvaBroadcastMediaOutputInfo OutputInfo;
+	
+	UPROPERTY()
+	EAvaBroadcastOutputState OutputState = EAvaBroadcastOutputState::Invalid;
+	
+	UPROPERTY()
+	EAvaBroadcastIssueSeverity IssueSeverity = EAvaBroadcastIssueSeverity::None;
+
+	UPROPERTY()
+	TArray<FString> IssueMessages;
+
+	/**
+	 * Raw Json string representing a serialized UMediaOutput.
+	 */
+	UPROPERTY()
+	FString Data;
+};
+
+USTRUCT()
+struct FAvaRundownOutputClassItem
+{
+	GENERATED_BODY()
+public:
+	UPROPERTY()
+	FString Name;
+
+	UPROPERTY()
+	TArray<FAvaRundownOutputDeviceItem> Devices;
+};
+
+USTRUCT()
+struct FAvaRundownDevicesList : public FAvaRundownMsgBase
+{
+	GENERATED_BODY()
+public:
+	UPROPERTY()
+	TArray<FAvaRundownOutputClassItem> DeviceClasses;
+};
+
+USTRUCT()
+struct FAvaRundownGetChannel : public FAvaRundownMsgBase
+{
+	GENERATED_BODY()
+public:
+	UPROPERTY()
+	FString ChannelName;
+};
+
+USTRUCT()
+struct FAvaRundownGetChannels : public FAvaRundownMsgBase
+{
+	GENERATED_BODY()
+};
+
+USTRUCT()
+struct FAvaRundownChannel
+{
+	GENERATED_BODY()
+public:
+	UPROPERTY()
+	FString Name;
+
+	UPROPERTY()
+	EAvaBroadcastChannelState State = EAvaBroadcastChannelState::Offline;
+
+	UPROPERTY()
+	EAvaBroadcastIssueSeverity IssueSeverity = EAvaBroadcastIssueSeverity::None;
+
+	UPROPERTY()
+	TArray<FAvaRundownOutputDeviceItem> Devices;
+};
+
+USTRUCT()
+struct FAvaRundownChannelListChanged : public FAvaRundownMsgBase
+{
+	GENERATED_BODY()
+
+	UPROPERTY()
+	TArray<FAvaRundownChannel> Channels;
+};
+
+USTRUCT()
+struct FAvaRundownChannelResponse : public FAvaRundownMsgBase
+{
+	GENERATED_BODY()
+public:
+	UPROPERTY()
+	FAvaRundownChannel Channel;
+};
+
+USTRUCT()
+struct FAvaRundownChannels : public FAvaRundownMsgBase
+{
+	GENERATED_BODY()
+public:
+	UPROPERTY()
+	TArray<FAvaRundownChannel> Channels;
+};
+
+USTRUCT()
+struct FAvaRundownAssetsChanged : public FAvaRundownMsgBase
+{
+	GENERATED_BODY()
+
+	UPROPERTY()
+	FString AssetName;
+};
+
+// Channel actions
+UENUM()
+enum class EAvaRundownChannelActions
+{
+	None,
+	Start,
+	Stop
+};
+
+USTRUCT()
+struct FAvaRundownChannelAction : public FAvaRundownMsgBase
+{
+	GENERATED_BODY()
+public:
+	UPROPERTY()
+	FString ChannelName;
+
+	UPROPERTY()
+	EAvaRundownChannelActions Action = EAvaRundownChannelActions::None;
+};
+
+USTRUCT()
+struct FAvaRundownGetDevices : public FAvaRundownMsgBase
+{
+	GENERATED_BODY()
+
+	UPROPERTY()
+	bool bShowAllMediaOutputClasses = false;
+};
+
+USTRUCT()
+struct FAvaRundownAddChannelDevice : public FAvaRundownMsgBase
+{
+	GENERATED_BODY()
+public:
+	UPROPERTY()
+	FString ChannelName;
+
+	UPROPERTY()
+	FString MediaOutputName;
+};
+
+USTRUCT()
+struct FAvaRundownEditChannelDevice : public FAvaRundownMsgBase
+{
+	GENERATED_BODY()
+public:
+	UPROPERTY()
+	FString ChannelName;
+
+	UPROPERTY()
+	FString MediaOutputName;
+
+	UPROPERTY()
+	FString Data;
+};
+
+USTRUCT()
+struct FAvaRundownRemoveChannelDevice : public FAvaRundownMsgBase
+{
+	GENERATED_BODY()
+public:
+	UPROPERTY()
+	FString ChannelName;
+
+	UPROPERTY()
+	FString MediaOutputName;
+};
+
+/* No difference from FAvaRundownOutputDeviceItem except this is meant to return a single device response. */
+USTRUCT()
+struct FAvaRundownOutputDeviceItemResponse : public FAvaRundownMsgBase
+{
+	GENERATED_BODY()
+public:
+	UPROPERTY()
+	FString Name;
+
+	UPROPERTY()
+	FString Data;
+};
+
+USTRUCT()
+struct FAvaRundownGetChannelImage : public FAvaRundownMsgBase
+{
+	GENERATED_BODY()
+public:
+	UPROPERTY()
+	FString ChannelName;
+};
+
+USTRUCT()
+struct FAvaRundownChannelImage : public FAvaRundownMsgBase
+{
+	GENERATED_BODY()
+public:
+	UPROPERTY()
+	TArray<uint8> ImageData;
+};
diff --git a/Engine/Plugins/Experimental/Avalanche/Source/AvalancheMedia/Internal/Rundown/AvaRundownPlaybackUtils.h b/Engine/Plugins/Experimental/Avalanche/Source/AvalancheMedia/Internal/Rundown/AvaRundownPlaybackUtils.h
new file mode 100644
--- /dev/null
+++ b/Engine/Plugins/Experimental/Avalanche/Source/AvalancheMedia/Internal/Rundown/AvaRundownPlaybackUtils.h
@@ -0,0 +1,19 @@
+// Copyright Epic Games, Inc. All Rights Reserved.
+
+#pragma once
+
+#include "Rundown/AvaRundownPage.h"
+
+class UAvaRundown;
+
+struct AVALANCHEMEDIA_API FAvaRundownPlaybackUtils
+{
+	static bool IsPageIdValid(int32 InPageId)
+	{
+		return InPageId != FAvaRundownPage::InvalidPageId;
+	}
+	
+	static int32 GetPageIdToPlayNext(const UAvaRundown* InRundown, const FAvaRundownPageListReference& InPageListReference, bool bInPreview, FName InPreviewChannelName);
+	
+	static TArray<int32> GetPagesToTakeToProgram(const UAvaRundown* InRundown, const TArray<int32>& InSelectedPageIds, FName InPreviewChannel = NAME_None);
+};
\ No newline at end of file
diff --git a/Engine/Plugins/Experimental/Avalanche/Source/AvalancheMedia/Internal/Rundown/AvaRundownServerMediaOutputUtils.h b/Engine/Plugins/Experimental/Avalanche/Source/AvalancheMedia/Internal/Rundown/AvaRundownServerMediaOutputUtils.h
new file mode 100644
--- /dev/null
+++ b/Engine/Plugins/Experimental/Avalanche/Source/AvalancheMedia/Internal/Rundown/AvaRundownServerMediaOutputUtils.h
@@ -0,0 +1,20 @@
+// Copyright Epic Games, Inc. All Rights Reserved.
+
+#pragma once
+
+#include "Containers/UnrealString.h"
+
+class UMediaOutput;
+
+/*
+ *	Class designed to deal with Serialization and Editing
+ *	of Media Output device data coming from Rundown Server
+ *
+ *	It is meant to give full control and provide all information to Rundown server
+ *	in a form that is usable outside of unreal engine
+ */
+struct AVALANCHEMEDIA_API FAvaRundownServerMediaOutputUtils
+{
+	static FString SerializeMediaOutput(const UMediaOutput* InMediaOutput);
+	static void EditMediaOutput(UMediaOutput* InMediaOutput, const FString& InDeviceData);
+};
diff --git a/Engine/Plugins/Experimental/Avalanche/Source/AvalancheMedia/Private/AvaMediaModule.cpp b/Engine/Plugins/Experimental/Avalanche/Source/AvalancheMedia/Private/AvaMediaModule.cpp
--- a/Engine/Plugins/Experimental/Avalanche/Source/AvalancheMedia/Private/AvaMediaModule.cpp
+++ b/Engine/Plugins/Experimental/Avalanche/Source/AvalancheMedia/Private/AvaMediaModule.cpp
@@ -32,6 +32,13 @@ namespace UE::AvaMediaModule::Private
 	{
 		return FParse::Value(FCommandLine::Get(),TEXT("dc_node="), OutDisplayClusterNodeName);
 	}
+
+	// Command line parsing helper.
+	bool IsRundownServerManuallyStarted(FString& OutRundownServerName)
+	{
+		return FParse::Value(FCommandLine::Get(), TEXT("MotionDesignRundownServerStart="), OutRundownServerName) ||
+			FParse::Param(FCommandLine::Get(), TEXT("MotionDesignRundownServerStart"));
+	}
 }
 
 FAvaMediaModule::FAvaMediaModule()
@@ -98,6 +105,18 @@ void FAvaMediaModule::StartupModule()
 				FConsoleCommandWithArgsDelegate::CreateRaw(this, &FAvaMediaModule::StopHttpPlaybackServerCommand),
 				ECVF_Default
 				));
+	ConsoleCmds.Add(IConsoleManager::Get().RegisterConsoleCommand(
+				TEXT("MotionDesignRundownServer.Start"),
+				TEXT("Starts the rundown server."),
+				FConsoleCommandWithArgsDelegate::CreateRaw(this, &FAvaMediaModule::StartRundownServerCommand),
+				ECVF_Default
+				));
+	ConsoleCmds.Add(IConsoleManager::Get().RegisterConsoleCommand(
+				TEXT("MotionDesignRundownServer.Stop"),
+				TEXT("Stops the rundown server."),
+				FConsoleCommandWithArgsDelegate::CreateRaw(this, &FAvaMediaModule::StopRundownServerCommand),
+				ECVF_Default
+				));
 	ConsoleCmds.Add(IConsoleManager::Get().RegisterConsoleCommand(
 				TEXT("MotionDesignPlaybackDevices.Save"),
 				TEXT("Save Device Providers data."),
@@ -133,7 +152,7 @@ void FAvaMediaModule::StartupModule()
 	FCoreDelegates::OnEnginePreExit.AddRaw(this, &FAvaMediaModule::EnginePreExit);
 	
 	FString DummyServerName;
-	if (!AvaPlaybackServer.IsValid() && IsPlaybackServerManuallyStarted(DummyServerName))
+	if (IsPlaybackServerManuallyStarted(DummyServerName) || IsRundownServerManuallyStarted(DummyServerName))
 	{
 		// Prevent throttling when server is started.
 		// This has to be done before any SLevelViewport are ticked since the cvar value is cached on first tick.
@@ -196,6 +215,24 @@ void FAvaMediaModule::StopPlaybackServer()
 	StopPlaybackServerCommand({});
 }
 
+void FAvaMediaModule::StartRundownServer(const FString& InRundownServerName)
+{
+	StartRundownServerCommand({InRundownServerName});
+}
+
+void FAvaMediaModule::StopRundownServer()
+{
+	StopRundownServerCommand({});
+}
+
+TSharedPtr<IAvaRundownServer> FAvaMediaModule::MakeDetachedRundownServer(const FString& InServerName)
+{
+	const TSharedRef<FAvaRundownServer, ESPMode::ThreadSafe> DetachedRundownServer = MakeShared<FAvaRundownServer>();
+	DetachedRundownServer->Init(InServerName);
+	OnRundownServerStarted.Broadcast(DetachedRundownServer);
+	return DetachedRundownServer;
+}
+
 IAvaPlaybackClient& FAvaMediaModule::GetPlaybackClient()
 {
 	if (AvaPlaybackClient.IsValid())
@@ -447,6 +484,14 @@ void FAvaMediaModule::PostEngineInit()
 	{
 		StartHttpPlaybackServerCommand({});
 	}
+
+	// Allow for specification of the rundown server name in the command line.
+	// Note: auto start of rundown server is only done in the editor module.
+	FString RundownServerName;
+	if (IsRundownServerManuallyStarted(RundownServerName))
+	{
+		StartRundownServerCommand({RundownServerName});
+	}
 }
 
 void FAvaMediaModule::EnginePreExit()
@@ -456,6 +501,7 @@ void FAvaMediaModule::EnginePreExit()
 
 void FAvaMediaModule::StopAllServices()
 {
+	StopRundownServerCommand({});
 	StopPlaybackServerCommand({});
 	StopPlaybackClientCommand({});
 	
@@ -468,6 +514,36 @@ void FAvaMediaModule::StopAllServices()
 	ManagedInstanceCache.Reset();
 }
 
+void FAvaMediaModule::StartRundownServerCommand(const TArray<FString>& Args)
+{
+	if (RundownServer)
+	{
+		UE_LOG(LogAvaMedia, Log, TEXT("Rundown Server is already started."));
+		return;
+	}
+	
+	RundownServer = MakeShared<FAvaRundownServer>();
+	
+	// Remark: Only the module's rundown server register console commands to avoid
+	// conflicts with temporary servers (for testing).
+	RundownServer->RegisterConsoleCommands();
+
+	RundownServer->Init(Args.Num() > 0 ? Args[0] : TEXT(""));
+	OnRundownServerStarted.Broadcast(RundownServer);
+
+	UE_LOG(LogAvaMedia, Log, TEXT("Rundown Server Started."));
+}
+
+void FAvaMediaModule::StopRundownServerCommand(const TArray<FString>& Args)
+{
+	if (RundownServer)
+	{
+		UE_LOG(LogAvaMedia, Log, TEXT("Stopping Rundown Server..."));
+		OnRundownServerStopping.Broadcast(RundownServer);
+	}
+	RundownServer.Reset();
+}
+
 void FAvaMediaModule::StartPlaybackServerCommand(const TArray<FString>& InArgs)
 {
 	// Starting a playback server in the same process as playback client is forbidden.
diff --git a/Engine/Plugins/Experimental/Avalanche/Source/AvalancheMedia/Private/AvaMediaModule.h b/Engine/Plugins/Experimental/Avalanche/Source/AvalancheMedia/Private/AvaMediaModule.h
--- a/Engine/Plugins/Experimental/Avalanche/Source/AvalancheMedia/Private/AvaMediaModule.h
+++ b/Engine/Plugins/Experimental/Avalanche/Source/AvalancheMedia/Private/AvaMediaModule.h
@@ -14,6 +14,7 @@
 #include "Playback/AvaPlaybackServerProcess.h"
 #include "Playback/Http/AvaPlaybackHttpServer.h"
 #include "Rundown/AvaRundownManagedInstanceCache.h"
+#include "Rundown/AvaRundownServer.h"
 
 class FAvaMediaModule : public IAvaMediaModule
 {
@@ -55,6 +56,13 @@ public:
 	virtual FOnAvaPlaybackServerStarted& GetOnAvaPlaybackServerStarted() override { return OnAvaPlaybackServerStarted; }
 	virtual FOnAvaPlaybackServerStopped& GetAvaPlaybackServerStopped() override { return OnAvaPlaybackServerStopped; }
 	virtual FGetEditorViewportClient& GetEditorViewportClientDelegate() override { return GetEditorViewportClient; }
+	virtual FOnRundownServerEvent& GetOnRundownServerStarted() override { return OnRundownServerStarted; }
+	virtual FOnRundownServerEvent& GetOnRundownServerStopping() override { return OnRundownServerStopping; }
+	virtual bool IsRundownServerStarted() const override { return RundownServer.IsValid();}
+	virtual void StartRundownServer(const FString& InRundownServerName) override;
+	virtual void StopRundownServer() override;
+	virtual TSharedPtr<IAvaRundownServer> GetRundownServer() const override { return RundownServer; }
+	virtual TSharedPtr<IAvaRundownServer> MakeDetachedRundownServer(const FString& InServerName) override;
 	virtual IAvaBroadcastDeviceProviderProxyManager& GetDeviceProviderProxyManager() override;
 	//~ End IAvaMediaModule
 	
@@ -64,6 +72,8 @@ private:
 	void StopAllServices();
 	
 	// Command handlers
+	void StartRundownServerCommand(const TArray<FString>& Args);
+	void StopRundownServerCommand(const TArray<FString>& Args);
 	void StartPlaybackServerCommand(const TArray<FString>& InArgs);
 	void StopPlaybackServerCommand(const TArray<FString>& InArgs);
 	void StartPlaybackClientCommand(const TArray<FString>& InArgs);
@@ -134,7 +144,8 @@ private:
 		FAvaMediaModule* ParentModule = nullptr;
 	};
 	FBroadcastSettingsBridge BroadcastSettingsBridge;
-	
+
+	TSharedPtr<FAvaRundownServer> RundownServer;
 	TSharedPtr<FAvaPlaybackServer> AvaPlaybackServer;	
 	TSharedPtr<FAvaPlaybackClient> AvaPlaybackClient;
 	TSharedPtr<FAvaPlaybackServerProcess> LocalPlaybackServerProcess;
@@ -150,5 +161,7 @@ private:
 	FOnAvaPlaybackClientStopped OnAvaPlaybackClientStopped;
 	FOnAvaPlaybackServerStarted OnAvaPlaybackServerStarted;
 	FOnAvaPlaybackServerStopped OnAvaPlaybackServerStopped;
+	FOnRundownServerEvent OnRundownServerStarted;
+	FOnRundownServerEvent OnRundownServerStopping;
 	FGetEditorViewportClient GetEditorViewportClient;
 };
diff --git a/Engine/Plugins/Experimental/Avalanche/Source/AvalancheMedia/Private/Broadcast/OutputDevices/AvaBroadcastOutputClassItem.cpp b/Engine/Plugins/Experimental/Avalanche/Source/AvalancheMedia/Private/Broadcast/OutputDevices/AvaBroadcastOutputClassItem.cpp
new file mode 100644
--- /dev/null
+++ b/Engine/Plugins/Experimental/Avalanche/Source/AvalancheMedia/Private/Broadcast/OutputDevices/AvaBroadcastOutputClassItem.cpp
@@ -0,0 +1,214 @@
+// Copyright Epic Games, Inc. All Rights Reserved.
+
+#include "AvaBroadcastOutputClassItem.h"
+
+#include "AvaBroadcastOutputDeviceItem.h"
+#include "AvaBroadcastOutputServerItem.h"
+#include "Broadcast/AvaBroadcast.h"
+#include "Broadcast/OutputDevices/AvaBroadcastDeviceProviderData.h"
+#include "Broadcast/OutputDevices/AvaBroadcastOutputUtils.h"
+#include "Broadcast/OutputDevices/IAvaBroadcastDeviceProviderProxyManager.h"
+#include "IAvaMediaModule.h"
+#include "IMediaIOCoreDeviceProvider.h"
+#include "IMediaIOCoreModule.h"
+#include "MediaIOCoreDefinitions.h"
+#include "MediaOutput.h"
+#include "Slate/SAvaBroadcastOutputTreeItem.h"
+#include "Styling/SlateIconFinder.h"
+
+#if WITH_EDITOR
+#include "ScopedTransaction.h"
+#endif
+
+#define LOCTEXT_NAMESPACE "AvaBroadcastOutputClassItem"
+
+namespace UE::AvaOutputClassItem::Private
+{
+	IMediaIOCoreDeviceProvider* GetDeviceProvider(const UClass* InMediaOutputClass)
+	{
+		const FName DeviceProviderName = UE::AvaBroadcastOutputUtils::GetDeviceProviderName(InMediaOutputClass);
+		if (!IMediaIOCoreModule::IsAvailable() || DeviceProviderName.IsNone())
+		{
+			return nullptr;
+		}
+
+		// Remark: even if the class has a device provider name, it is possible
+		// that there is no corresponding device provider.
+		return IMediaIOCoreModule::Get().GetDeviceProvider(DeviceProviderName);
+	}
+
+	void UpdateDeviceName(UAvaBroadcast& InBroadcast, const FName& InChannelName, const UMediaOutput* InMediaOutput, const FString& InDeviceName)
+	{
+		FAvaBroadcastOutputChannel& Channel = InBroadcast.GetCurrentProfile().GetChannelMutable(InChannelName);
+		if (Channel.IsValidChannel())
+		{
+			if (FAvaBroadcastMediaOutputInfo* ExistingOutputInfo = Channel.GetMediaOutputInfoMutable(InMediaOutput))
+			{
+				ExistingOutputInfo->DeviceName = FName(InDeviceName);	
+			}
+		}
+	}
+}
+
+FText FAvaBroadcastOutputClassItem::GetDisplayName() const
+{
+#if WITH_EDITOR
+	check(OutputClass.IsValid());
+	return OutputClass->GetDisplayNameText();
+#else
+	return LOCTEXT("DisplayName_NotAvailable", "NotAvailable");	// Not available in game build.
+#endif
+}
+
+const FSlateBrush* FAvaBroadcastOutputClassItem::GetIconBrush() const
+{
+	check(OutputClass.IsValid());
+	return FSlateIconFinder::FindIconBrushForClass(OutputClass.Get());
+}
+
+void FAvaBroadcastOutputClassItem::RefreshChildren(const FRefreshChildrenParams& InParams)
+{
+	IMediaIOCoreDeviceProvider* DeviceProvider = UE::AvaOutputClassItem::Private::GetDeviceProvider(OutputClass.Get());
+
+	if (!DeviceProvider)
+	{
+		Children.Reset();
+		return;
+	}
+
+	TSharedPtr<FAvaBroadcastOutputTreeItem> Parent = ParentWeak.Pin();
+	if (!Parent.IsValid())
+	{
+		return;
+	}
+
+	TSharedPtr<FAvaBroadcastOutputServerItem> ParentServerItem = StaticCastSharedPtr<FAvaBroadcastOutputServerItem>(Parent);
+	check(ParentServerItem.IsValid());
+
+	const FName DeviceProviderName = DeviceProvider->GetFName();
+	const IAvaBroadcastDeviceProviderProxyManager& DeviceProviderProxyManager = IAvaMediaModule::Get().GetDeviceProviderProxyManager();
+
+	TArray<FMediaIOConfiguration> OutputConfigs = DeviceProvider->GetConfigurations(false, true);
+
+	// Fill current devices from output configs of this device provider.
+	TSet<FAvaBroadcastDevice> CurrentDevices;
+	CurrentDevices.Reserve(OutputConfigs.Num());
+
+	for (const FMediaIOConfiguration& Config : OutputConfigs)
+	{
+		const FString ServerName = DeviceProviderProxyManager.FindServerNameForDevice(DeviceProviderName, Config.MediaConnection.Device.DeviceName);
+		const bool bIsLocal = DeviceProviderProxyManager.IsLocalDevice(DeviceProviderName, Config.MediaConnection.Device.DeviceName);
+		const FAvaBroadcastDeviceProviderData* DeviceProviderDataForServer = ParentServerItem->GetDeviceProviderData(DeviceProviderName);
+
+		if ((!DeviceProviderDataForServer && bIsLocal)
+			 || (DeviceProviderDataForServer && !bIsLocal && ServerName == ParentServerItem->GetServerName()))
+		{
+			CurrentDevices.Add({ Config.MediaConnection.Device, DeviceProviderName, ServerName });
+		}
+	}
+
+	//A set to contain the Media Devices of the Current Children
+	TSet<FAvaBroadcastDevice> SeenDevices;
+	SeenDevices.Reserve(CurrentDevices.Num());
+	
+	//Remove Existing Children that are Invalid
+	for (TArray<FAvaOutputTreeItemPtr>::TIterator ItemIt = Children.CreateIterator(); ItemIt; ++ItemIt)
+	{
+		FAvaOutputTreeItemPtr Item(*ItemIt);
+
+		//Remove Invalid Pointers or Items that are not Device Items since a Class Item can only contain Device Items
+		if (!Item.IsValid() || !Item->IsA<FAvaBroadcastOutputDeviceItem>())
+		{
+			ItemIt.RemoveCurrent();
+			continue;
+		}
+
+		const TSharedPtr<FAvaBroadcastOutputDeviceItem> DeviceItem = StaticCastSharedPtr<FAvaBroadcastOutputDeviceItem>(Item);
+		const FAvaBroadcastDevice& UnderlyingDevice = DeviceItem->GetDevice();
+
+		if (UnderlyingDevice.IsValid() && CurrentDevices.Contains(UnderlyingDevice))
+		{
+			SeenDevices.Add(UnderlyingDevice);
+		}
+		else
+		{
+			//Remove if there's no valid Underlying Device or it's no longer in the set
+			ItemIt.RemoveCurrent();
+		}
+	}
+	
+	//Append the Devices that are not already in the Original Children List
+	{
+		TSharedPtr<FAvaBroadcastOutputClassItem> This = SharedThis(this);
+		TArray<FAvaBroadcastDevice> NewDevices = CurrentDevices.Difference(SeenDevices).Array();
+		Children.Reserve(Children.Num() + NewDevices.Num());
+		
+		for (const FAvaBroadcastDevice& Device : NewDevices)
+		{
+			TSharedPtr<FAvaBroadcastOutputDeviceItem> DeviceItem = MakeShared<FAvaBroadcastOutputDeviceItem>(This, Device);
+			Children.Add(DeviceItem);
+		}
+	}
+}
+
+TSharedPtr<SWidget> FAvaBroadcastOutputClassItem::GenerateRowWidget()
+{
+	return SNew(SAvaBroadcastOutputTreeItem, SharedThis(this));
+}
+
+UMediaOutput* FAvaBroadcastOutputClassItem::AddMediaOutputToChannel(FName InTargetChannel, const FAvaBroadcastMediaOutputInfo& InOutputInfo)
+{
+	check(OutputClass.IsValid());
+
+	UAvaBroadcast& Broadcast = UAvaBroadcast::Get();
+
+#if WITH_EDITOR
+	FScopedTransaction Transaction(LOCTEXT("AddMediaOutput", "Add Media Output"));
+#endif
+	Broadcast.Modify();
+
+	FAvaBroadcastMediaOutputInfo OutputInfo = InOutputInfo;
+	
+	// Support non-enumerated devices (i.e. no device provider).
+	if (!OutputInfo.IsValid() || !OutputInfo.Guid.IsValid())
+	{
+		if (TSharedPtr<FAvaBroadcastOutputTreeItem> Parent = ParentWeak.Pin())
+		{
+			const TSharedPtr<FAvaBroadcastOutputServerItem> ParentServerItem = StaticCastSharedPtr<FAvaBroadcastOutputServerItem>(Parent);
+			check(ParentServerItem.IsValid());
+
+			// Fill the device info with available data.
+			OutputInfo.Guid = FGuid::NewGuid();
+			OutputInfo.ServerName = ParentServerItem->GetServerName();
+
+			// Fetch the device provider name (if there is one) from the class's meta data.
+			OutputInfo.DeviceProviderName = UE::AvaBroadcastOutputUtils::GetDeviceProviderName(OutputClass.Get());
+
+			// Device name:
+			// There is no enumerated device name. It will be updated from the media output object below.
+		}
+	}
+	
+	UMediaOutput* const MediaOutput = Broadcast.GetCurrentProfile().AddChannelMediaOutput(InTargetChannel, OutputClass.Get(), OutputInfo);
+
+	// If the device name was not provided as input, we need to update it from the Media Output object.
+	if (MediaOutput && OutputInfo.DeviceName.IsNone())
+	{
+		const FString DeviceName = UE::AvaBroadcastOutputUtils::GetDeviceName(MediaOutput);
+		if (!DeviceName.IsEmpty())
+		{
+			UE::AvaOutputClassItem::Private::UpdateDeviceName(Broadcast, InTargetChannel, MediaOutput, DeviceName);
+		}
+	}
+
+#if WITH_EDITOR
+	if (!MediaOutput)
+	{
+		Transaction.Cancel();
+	}
+#endif
+	
+	return MediaOutput;
+}
+
+#undef LOCTEXT_NAMESPACE
diff --git a/Engine/Plugins/Experimental/Avalanche/Source/AvalancheMedia/Private/Broadcast/OutputDevices/AvaBroadcastOutputClassItem.h b/Engine/Plugins/Experimental/Avalanche/Source/AvalancheMedia/Private/Broadcast/OutputDevices/AvaBroadcastOutputClassItem.h
new file mode 100644
--- /dev/null
+++ b/Engine/Plugins/Experimental/Avalanche/Source/AvalancheMedia/Private/Broadcast/OutputDevices/AvaBroadcastOutputClassItem.h
@@ -0,0 +1,35 @@
+// Copyright Epic Games, Inc. All Rights Reserved.
+
+#pragma once
+
+#include "Broadcast/OutputDevices/AvaBroadcastOutputTreeItem.h"
+#include "UObject/WeakObjectPtr.h"
+#include "UObject/WeakObjectPtrTemplates.h"
+
+class FAvaBroadcastOutputClassItem : public FAvaBroadcastOutputTreeItem
+{
+public:
+	UE_AVA_INHERITS_WITH_SUPER(FAvaBroadcastOutputClassItem, FAvaBroadcastOutputTreeItem);
+
+	FAvaBroadcastOutputClassItem(const TSharedPtr<FAvaBroadcastOutputTreeItem>& InParent, UClass* InOutputClass)
+		: Super(InParent)
+		, OutputClass(InOutputClass)
+	{
+	}
+
+	UClass* GetOutputClass() const
+	{
+		return OutputClass.Get();
+	}
+
+	//~ Begin IAvaBroadcastOutputTreeItem
+	virtual FText GetDisplayName() const override;
+	virtual const FSlateBrush* GetIconBrush() const override;
+	virtual void RefreshChildren(const FRefreshChildrenParams& InParams) override;
+	virtual TSharedPtr<SWidget> GenerateRowWidget() override;
+	virtual UMediaOutput* AddMediaOutputToChannel(FName InTargetChannel, const FAvaBroadcastMediaOutputInfo& InOutputInfo) override;
+	//~ End IAvaBroadcastOutputTreeItem
+
+protected:
+	TWeakObjectPtr<UClass> OutputClass;
+};
diff --git a/Engine/Plugins/Experimental/Avalanche/Source/AvalancheMedia/Private/Broadcast/OutputDevices/AvaBroadcastOutputDeviceItem.cpp b/Engine/Plugins/Experimental/Avalanche/Source/AvalancheMedia/Private/Broadcast/OutputDevices/AvaBroadcastOutputDeviceItem.cpp
new file mode 100644
--- /dev/null
+++ b/Engine/Plugins/Experimental/Avalanche/Source/AvalancheMedia/Private/Broadcast/OutputDevices/AvaBroadcastOutputDeviceItem.cpp
@@ -0,0 +1,106 @@
+// Copyright Epic Games, Inc. All Rights Reserved.
+
+#include "AvaBroadcastOutputDeviceItem.h"
+
+#include "AvaBroadcastOutputClassItem.h"
+#include "Broadcast/AvaBroadcast.h"
+#include "Broadcast/Channel/AvaBroadcastMediaOutputInfo.h"
+#include "IMediaIOCoreDeviceProvider.h"
+#include "IMediaIOCoreModule.h"
+#include "MediaOutput.h"
+#include "Widgets/Text/STextBlock.h"
+
+#define LOCTEXT_NAMESPACE "AvaBroadcastOutputDeviceItem"
+
+FText FAvaBroadcastOutputDeviceItem::GetDisplayName() const
+{
+	return Device.GetDisplayNameText();
+}
+
+const FSlateBrush* FAvaBroadcastOutputDeviceItem::GetIconBrush() const
+{
+	return nullptr;
+}
+
+void FAvaBroadcastOutputDeviceItem::RefreshChildren(const FRefreshChildrenParams& InParams)
+{
+	//There shouldn't be any Childrens here
+	Children.Reset();
+}
+
+TSharedPtr<SWidget> FAvaBroadcastOutputDeviceItem::GenerateRowWidget()
+{
+	return SNew(STextBlock)
+		.Text(GetDisplayName());
+}
+
+bool FAvaBroadcastOutputDeviceItem::IsValidToDropInChannel(FName InTargetChannelName)
+{
+	// It is invalid to drop a "remote" device in a preview channel.
+	if (UAvaBroadcast::Get().GetChannelType(InTargetChannelName) == EAvaBroadcastChannelType::Preview)
+	{
+		return !FAvaBroadcastMediaOutputInfo::IsRemote(Device.GetServerName());
+	}
+	return true;
+}
+
+UMediaOutput* FAvaBroadcastOutputDeviceItem::AddMediaOutputToChannel(FName InTargetChannel, const FAvaBroadcastMediaOutputInfo& /*InOutputInfo*/)
+{
+	if (TSharedPtr<FAvaBroadcastOutputTreeItem> Parent = ParentWeak.Pin())
+	{
+		FAvaBroadcastMediaOutputInfo OutputInfo;
+		OutputInfo.Guid = FGuid::NewGuid();
+		OutputInfo.ServerName = Device.GetServerName();
+		OutputInfo.DeviceProviderName = Device.GetDeviceProviderName();
+		OutputInfo.DeviceName = Device.GetDevice().DeviceName;
+
+		UMediaOutput* const MediaOutput = Parent->AddMediaOutputToChannel(InTargetChannel, OutputInfo);
+		
+		if (MediaOutput)
+		{
+			const UClass* const MediaOutputClass = MediaOutput->GetClass();
+			const IMediaIOCoreDeviceProvider* const DeviceProvider = IMediaIOCoreModule::Get().GetDeviceProvider(Device.GetDeviceProviderName());
+			
+			if (DeviceProvider)
+			{
+				FStructProperty* const Property = FindFProperty<FStructProperty>(MediaOutputClass, TEXT("OutputConfiguration"));
+				if (Property && Property->Struct->IsChildOf(FMediaIOOutputConfiguration::StaticStruct()))
+				{
+					if (FMediaIOOutputConfiguration* const OutputConfig = Property->ContainerPtrToValuePtr<FMediaIOOutputConfiguration>(MediaOutput))
+					{
+						{
+							// Initial "default" setup. 
+							// This leaves the OutputConfig object in a possibly invalid state
+							// because it is very likely the selected device's config doesn't match the one returned by
+							// GetDefaultOutputConfiguration().
+							*OutputConfig = DeviceProvider->GetDefaultOutputConfiguration();
+							OutputConfig->MediaConfiguration.MediaConnection.Device = Device.GetDevice();
+						}
+
+						// Fetch all possible output configurations.
+						TArray<FMediaIOOutputConfiguration> OutputConfigs = DeviceProvider->GetOutputConfigurations();
+
+						// Select one that match the device. 
+						// For now just get the first one.
+						for (const FMediaIOOutputConfiguration& Config : OutputConfigs)
+						{
+							if (Config.MediaConfiguration.MediaConnection.Device.DeviceName == Device.GetDevice().DeviceName)
+							{
+								*OutputConfig = Config;
+								break; // get the first one.
+							}
+						}
+#if WITH_EDITOR
+						FPropertyChangedEvent PropertyChangedEvent(Property, EPropertyChangeType::ValueSet);
+						MediaOutput->PostEditChangeProperty(PropertyChangedEvent);
+#endif
+					}
+				}
+			}
+		}
+		return MediaOutput;
+	}
+	return nullptr;
+}
+
+#undef LOCTEXT_NAMESPACE
diff --git a/Engine/Plugins/Experimental/Avalanche/Source/AvalancheMedia/Private/Broadcast/OutputDevices/AvaBroadcastOutputDeviceItem.h b/Engine/Plugins/Experimental/Avalanche/Source/AvalancheMedia/Private/Broadcast/OutputDevices/AvaBroadcastOutputDeviceItem.h
new file mode 100644
--- /dev/null
+++ b/Engine/Plugins/Experimental/Avalanche/Source/AvalancheMedia/Private/Broadcast/OutputDevices/AvaBroadcastOutputDeviceItem.h
@@ -0,0 +1,37 @@
+// Copyright Epic Games, Inc. All Rights Reserved.
+
+#pragma once
+
+#include "Broadcast/OutputDevices/AvaBroadcastOutputTreeItem.h"
+#include "Broadcast/AvaBroadcastDevice.h"
+
+class FAvaBroadcastOutputClassItem;
+
+class FAvaBroadcastOutputDeviceItem : public FAvaBroadcastOutputTreeItem
+{
+public:
+	UE_AVA_INHERITS_WITH_SUPER(FAvaBroadcastOutputDeviceItem, FAvaBroadcastOutputTreeItem);
+
+	FAvaBroadcastOutputDeviceItem(const TSharedPtr<FAvaBroadcastOutputTreeItem>& InParent, const FAvaBroadcastDevice& InDevice)
+		: Super(InParent)
+		, Device(InDevice)
+	{
+	}
+
+	const FAvaBroadcastDevice& GetDevice() const
+	{
+		return Device;
+	}
+
+	//~ Begin IAvaBroadcastOutputTreeItem
+	virtual FText GetDisplayName() const override;
+	virtual const FSlateBrush* GetIconBrush() const override;
+	virtual void RefreshChildren(const FRefreshChildrenParams& InParams) override;
+	virtual TSharedPtr<SWidget> GenerateRowWidget() override;
+	virtual bool IsValidToDropInChannel(FName InTargetChannelName) override;
+	virtual UMediaOutput* AddMediaOutputToChannel(FName InTargetChannel, const FAvaBroadcastMediaOutputInfo& InOutputInfo) override;
+	//~ End IAvaBroadcastOutputTreeItem
+
+protected:
+	FAvaBroadcastDevice Device;
+};
diff --git a/Engine/Plugins/Experimental/Avalanche/Source/AvalancheMedia/Private/Broadcast/OutputDevices/AvaBroadcastOutputRootItem.cpp b/Engine/Plugins/Experimental/Avalanche/Source/AvalancheMedia/Private/Broadcast/OutputDevices/AvaBroadcastOutputRootItem.cpp
new file mode 100644
--- /dev/null
+++ b/Engine/Plugins/Experimental/Avalanche/Source/AvalancheMedia/Private/Broadcast/OutputDevices/AvaBroadcastOutputRootItem.cpp
@@ -0,0 +1,103 @@
+// Copyright Epic Games, Inc. All Rights Reserved.
+
+#include "Broadcast/OutputDevices/AvaBroadcastOutputRootItem.h"
+
+#include "Broadcast/OutputDevices/AvaBroadcastOutputServerItem.h"
+#include "Broadcast/OutputDevices/IAvaBroadcastDeviceProviderProxyManager.h"
+#include "IAvaMediaModule.h"
+#include "Input/Reply.h"
+#include "MediaOutput.h"
+#include "UObject/UObjectIterator.h"
+
+void FAvaBroadcastOutputRootItem::RefreshChildren(const FRefreshChildrenParams& InParams)
+{
+	const IAvaBroadcastDeviceProviderProxyManager& DeviceProviderProxyManager = IAvaMediaModule::Get().GetDeviceProviderProxyManager();
+
+	const TSet<FString> CurrentServers = DeviceProviderProxyManager.GetServerNames();
+
+	TSet<FString> SeenServerNames;
+	SeenServerNames.Reserve(CurrentServers.Num());
+
+	//Add local server item
+	const TSharedPtr<FAvaBroadcastOutputServerItem> OutputLocalServerItem = MakeShared<FAvaBroadcastOutputServerItem>(DeviceProviderProxyManager.GetLocalServerName(), nullptr);
+	Children.Add(OutputLocalServerItem);
+
+	//Remove Existing Children that are Invalid
+	for (TArray<FAvaOutputTreeItemPtr>::TIterator ItemIt = Children.CreateIterator(); ItemIt; ++ItemIt)
+	{
+		FAvaOutputTreeItemPtr Item(*ItemIt);
+
+		//Always going to display local devices group
+		if (Item == OutputLocalServerItem)
+		{
+			continue;
+		}
+
+		//Remove Invalid Pointers or Items that are not Output Server Items since Root can only have the Server Items as Top Level
+		if (!Item.IsValid() || !Item->IsA<FAvaBroadcastOutputServerItem>())
+		{
+			ItemIt.RemoveCurrent();
+			continue;
+		}
+
+		const TSharedPtr<FAvaBroadcastOutputServerItem> OutputServerItem = StaticCastSharedPtr<FAvaBroadcastOutputServerItem>(Item);
+		const FString UnderlyingServerName = OutputServerItem->GetServerName();
+
+		if (!UnderlyingServerName.IsEmpty() && CurrentServers.Contains(UnderlyingServerName))
+		{
+			SeenServerNames.Add(UnderlyingServerName);
+		}
+		else
+		{
+			//Remove if there's no valid Underlying OutputClass or it's no longer in the set
+			ItemIt.RemoveCurrent();
+		}
+	}
+	
+	//Append the New Servers that are not already in the Original Children List
+	{
+		TSet<FString> NewServers = CurrentServers.Difference(SeenServerNames);
+		Children.Reserve(Children.Num() + NewServers.Num());
+		
+		for (const FString& ServerName : NewServers)
+		{
+			if (const TSharedPtr<const FAvaBroadcastDeviceProviderDataList> ServerDeviceProviderDataList = DeviceProviderProxyManager.GetDeviceProviderDataListForServer(ServerName))
+			{
+				TSharedPtr<FAvaBroadcastOutputServerItem> OutputServerItem = MakeShared<FAvaBroadcastOutputServerItem>(ServerName, ServerDeviceProviderDataList);
+				Children.Add(OutputServerItem);
+			}
+		}
+	}
+}
+
+FText FAvaBroadcastOutputRootItem::GetDisplayName() const
+{
+	//Root shouldn't need to give a Display Name!
+	check(0);
+	return FText::GetEmpty();
+}
+
+const FSlateBrush* FAvaBroadcastOutputRootItem::GetIconBrush() const
+{
+	check(0);
+	return nullptr;
+}
+
+TSharedPtr<SWidget> FAvaBroadcastOutputRootItem::GenerateRowWidget()
+{
+	//Root shouldn't need to Generate a Row Widget
+	check(0);
+	return nullptr;
+}
+
+FReply FAvaBroadcastOutputRootItem::OnDragDetected(const FGeometry& MyGeometry, const FPointerEvent& MouseEvent)
+{
+	check(0);
+	return FReply::Unhandled();
+}
+
+UMediaOutput* FAvaBroadcastOutputRootItem::AddMediaOutputToChannel(FName InTargetChannel, const FAvaBroadcastMediaOutputInfo& InOutputInfo)
+{
+	check(0);
+	return nullptr;
+}
diff --git a/Engine/Plugins/Experimental/Avalanche/Source/AvalancheMedia/Private/Broadcast/OutputDevices/AvaBroadcastOutputServerItem.cpp b/Engine/Plugins/Experimental/Avalanche/Source/AvalancheMedia/Private/Broadcast/OutputDevices/AvaBroadcastOutputServerItem.cpp
new file mode 100644
--- /dev/null
+++ b/Engine/Plugins/Experimental/Avalanche/Source/AvalancheMedia/Private/Broadcast/OutputDevices/AvaBroadcastOutputServerItem.cpp
@@ -0,0 +1,128 @@
+// Copyright Epic Games, Inc. All Rights Reserved.
+
+#include "AvaBroadcastOutputServerItem.h"
+
+#include "AvaBroadcastOutputClassItem.h"
+#include "Broadcast/AvaBroadcast.h"
+#include "Broadcast/OutputDevices/AvaBroadcastOutputUtils.h"
+#include "IMediaIOCoreModule.h"
+#include "MediaOutput.h"
+#include "Slate/SAvaBroadcastOutputTreeItem.h"
+#include "Styling/AppStyle.h"
+#include "UObject/UObjectIterator.h"
+
+#if WITH_EDITOR
+#include "ScopedTransaction.h"
+#endif
+
+#define LOCTEXT_NAMESPACE "AvaBroadcastOutputServerItem"
+
+FString FAvaBroadcastOutputServerItem::GetServerName() const
+{
+	return ServerName;
+}
+
+const FAvaBroadcastDeviceProviderData* FAvaBroadcastOutputServerItem::GetDeviceProviderData(FName InDeviceProviderName) const
+{
+	if (!DeviceProviderDataList.IsValid())
+	{
+		return nullptr;
+	}
+
+	for (const FAvaBroadcastDeviceProviderData& DeviceProviderData : DeviceProviderDataList->DeviceProviders)
+	{
+		if (DeviceProviderData.Name == InDeviceProviderName)
+		{
+			return &DeviceProviderData;
+		}
+	}
+
+	return nullptr;
+}
+
+FText FAvaBroadcastOutputServerItem::GetDisplayName() const
+{
+	return FText::FromString(ServerName);
+}
+
+const FSlateBrush* FAvaBroadcastOutputServerItem::GetIconBrush() const
+{
+	return FAppStyle::Get().GetBrush(TEXT("Icons.Server"));
+}
+
+void FAvaBroadcastOutputServerItem::RefreshChildren(const FRefreshChildrenParams& InParams)
+{
+	if (!IMediaIOCoreModule::IsAvailable())
+	{
+		Children.Reset();
+		return;
+	}
+
+	// Make a set of all allowed Media Output Classes.
+	// It may have all classes or be restricted to those with the MediaIOCustomLayout MetaData.
+	TSet<UClass*> CurrentOutputClasses;
+	for (UClass* const Class : TObjectRange<UClass>())
+	{
+		const bool bIsMediaOutputClass = Class->IsChildOf(UMediaOutput::StaticClass()) && Class != UMediaOutput::StaticClass();
+		const bool bHasDeviceProvider = UE::AvaBroadcastOutputUtils::HasDeviceProviderName(Class);
+
+		if (bIsMediaOutputClass && (InParams.bShowAllMediaOutputClasses || bHasDeviceProvider))
+		{
+			CurrentOutputClasses.Add(Class);
+		}
+	}
+
+	TSet<UClass*> SeenOutputClasses;
+	SeenOutputClasses.Reserve(CurrentOutputClasses.Num());
+
+	//Remove Existing Children that are Invalid
+	for (TArray<FAvaOutputTreeItemPtr>::TIterator ItemIt = Children.CreateIterator(); ItemIt; ++ItemIt)
+	{
+		FAvaOutputTreeItemPtr Item(*ItemIt);
+
+		//Remove Invalid Pointers or Items that are not Output Class Items since Root can only have the Class Items as Top Level
+		if (!Item.IsValid() || !Item->IsA<FAvaBroadcastOutputClassItem>())
+		{
+			ItemIt.RemoveCurrent();
+			continue;
+		}
+
+		const TSharedPtr<FAvaBroadcastOutputClassItem> OutputClassItem = StaticCastSharedPtr<FAvaBroadcastOutputClassItem>(Item);
+		UClass* const UnderlyingOutputClass = OutputClassItem->GetOutputClass();
+
+		if (UnderlyingOutputClass && CurrentOutputClasses.Contains(UnderlyingOutputClass))
+		{
+			SeenOutputClasses.Add(UnderlyingOutputClass);
+		}
+		else
+		{
+			//Remove if there's no valid Underlying OutputClass or it's no longer in the set
+			ItemIt.RemoveCurrent();
+		}
+	}
+
+	//Append the New Output Classes that are not already in the Original Children List
+	{
+		TSharedPtr<FAvaBroadcastOutputServerItem> This = SharedThis(this);
+		TArray<UClass*> NewOutputClasses = CurrentOutputClasses.Difference(SeenOutputClasses).Array();
+		Children.Reserve(Children.Num() + NewOutputClasses.Num());
+
+		for (UClass* const NewOutputClass : NewOutputClasses)
+		{
+			TSharedPtr<FAvaBroadcastOutputClassItem> OutputClassItem = MakeShared<FAvaBroadcastOutputClassItem>(This, NewOutputClass);
+			Children.Add(OutputClassItem);
+		}
+	}
+}
+
+TSharedPtr<SWidget> FAvaBroadcastOutputServerItem::GenerateRowWidget()
+{
+	return SNew(SAvaBroadcastOutputTreeItem, SharedThis(this));
+}
+
+UMediaOutput* FAvaBroadcastOutputServerItem::AddMediaOutputToChannel(FName InTargetChannel, const FAvaBroadcastMediaOutputInfo& InOutputInfo)
+{
+	return nullptr;
+}
+
+#undef LOCTEXT_NAMESPACE
diff --git a/Engine/Plugins/Experimental/Avalanche/Source/AvalancheMedia/Private/Broadcast/OutputDevices/AvaBroadcastOutputServerItem.h b/Engine/Plugins/Experimental/Avalanche/Source/AvalancheMedia/Private/Broadcast/OutputDevices/AvaBroadcastOutputServerItem.h
new file mode 100644
--- /dev/null
+++ b/Engine/Plugins/Experimental/Avalanche/Source/AvalancheMedia/Private/Broadcast/OutputDevices/AvaBroadcastOutputServerItem.h
@@ -0,0 +1,37 @@
+// Copyright Epic Games, Inc. All Rights Reserved.
+
+#pragma once
+
+#include "Broadcast/OutputDevices/AvaBroadcastOutputTreeItem.h"
+#include "Broadcast/OutputDevices/AvaBroadcastDeviceProviderData.h"
+#include "UObject/WeakObjectPtr.h"
+#include "UObject/WeakObjectPtrTemplates.h"
+
+class FAvaBroadcastOutputServerItem : public FAvaBroadcastOutputTreeItem
+{
+public:
+	UE_AVA_INHERITS_WITH_SUPER(FAvaBroadcastOutputServerItem, FAvaBroadcastOutputTreeItem);
+
+	FAvaBroadcastOutputServerItem(const FString& InServerName, const TSharedPtr<const FAvaBroadcastDeviceProviderDataList>& InDeviceProviderDataList)
+		: Super(nullptr)
+		, ServerName(InServerName)
+		, DeviceProviderDataList(InDeviceProviderDataList)
+	{
+	}
+
+	FString GetServerName() const;
+	const FAvaBroadcastDeviceProviderData* GetDeviceProviderData(FName InDeviceProviderName) const;
+
+	//~ Begin IAvaBroadcastOutputTreeItem
+	virtual FText GetDisplayName() const override;
+	virtual const FSlateBrush* GetIconBrush() const override;
+	virtual void RefreshChildren(const FRefreshChildrenParams& InParams) override;
+	virtual TSharedPtr<SWidget> GenerateRowWidget() override;
+	virtual bool IsValidToDropInChannel(FName InTargetChannelName) override { return false; }
+	virtual UMediaOutput* AddMediaOutputToChannel(FName InTargetChannel, const FAvaBroadcastMediaOutputInfo& InOutputInfo) override;
+	//~ End IAvaBroadcastOutputTreeItem
+
+protected:
+	FString ServerName;
+	const TSharedPtr<const FAvaBroadcastDeviceProviderDataList> DeviceProviderDataList;
+};
diff --git a/Engine/Plugins/Experimental/Avalanche/Source/AvalancheMedia/Private/Broadcast/OutputDevices/AvaBroadcastOutputTreeItem.cpp b/Engine/Plugins/Experimental/Avalanche/Source/AvalancheMedia/Private/Broadcast/OutputDevices/AvaBroadcastOutputTreeItem.cpp
new file mode 100644
--- /dev/null
+++ b/Engine/Plugins/Experimental/Avalanche/Source/AvalancheMedia/Private/Broadcast/OutputDevices/AvaBroadcastOutputTreeItem.cpp
@@ -0,0 +1,43 @@
+// Copyright Epic Games, Inc. All Rights Reserved.
+
+#include "Broadcast/OutputDevices/AvaBroadcastOutputTreeItem.h"
+
+#include "Input/Reply.h"
+
+const TWeakPtr<FAvaBroadcastOutputTreeItem>& FAvaBroadcastOutputTreeItem::GetParent() const
+{
+	return ParentWeak;
+}
+
+const TArray<FAvaOutputTreeItemPtr>& FAvaBroadcastOutputTreeItem::GetChildren() const
+{
+	return Children;
+}
+
+FReply FAvaBroadcastOutputTreeItem::OnDragDetected(const FGeometry& MyGeometry, const FPointerEvent& MouseEvent)
+{
+	if (MouseEvent.IsMouseButtonDown(EKeys::LeftMouseButton))
+	{
+		if (OnCreateDragDropOperationDelegate.IsBound())
+		{
+			return FReply::Handled().BeginDragDrop(OnCreateDragDropOperationDelegate.Execute(SharedThis(this)));
+		}
+	}
+	return FReply::Unhandled();
+}
+
+void FAvaBroadcastOutputTreeItem::RefreshTree(const FAvaOutputTreeItemPtr& InItem, const FRefreshChildrenParams& InParams)
+{
+	TArray<FAvaOutputTreeItemPtr> ItemsRemainingToRefresh;
+	ItemsRemainingToRefresh.Add(InItem);
+		
+	while (ItemsRemainingToRefresh.Num() > 0)
+	{
+		FAvaOutputTreeItemPtr Item = ItemsRemainingToRefresh.Pop();
+		if (Item.IsValid())
+		{
+			Item->RefreshChildren(InParams);
+			ItemsRemainingToRefresh.Append(Item->GetChildren());
+		}
+	}
+}
diff --git a/Engine/Plugins/Experimental/Avalanche/Source/AvalancheMedia/Private/Broadcast/OutputDevices/Slate/SAvaBroadcastOutputTreeItem.cpp b/Engine/Plugins/Experimental/Avalanche/Source/AvalancheMedia/Private/Broadcast/OutputDevices/Slate/SAvaBroadcastOutputTreeItem.cpp
new file mode 100644
--- /dev/null
+++ b/Engine/Plugins/Experimental/Avalanche/Source/AvalancheMedia/Private/Broadcast/OutputDevices/Slate/SAvaBroadcastOutputTreeItem.cpp
@@ -0,0 +1,33 @@
+// Copyright Epic Games, Inc. All Rights Reserved.
+
+#include "SAvaBroadcastOutputTreeItem.h"
+
+#include "Broadcast/OutputDevices/AvaBroadcastOutputTreeItem.h"
+#include "Widgets/Images/SImage.h"
+#include "Widgets/SBoxPanel.h"
+#include "Widgets/Text/STextBlock.h"
+
+void SAvaBroadcastOutputTreeItem::Construct(const FArguments& InArgs, const TSharedPtr<FAvaBroadcastOutputTreeItem>& InOutputTreeItem)
+{
+	check(InOutputTreeItem.IsValid());
+	OutputTreeItemWeak = InOutputTreeItem;
+
+	ChildSlot
+		[
+			SNew(SHorizontalBox)
+			+ SHorizontalBox::Slot()
+			.AutoWidth()
+			[
+				SNew(SImage)
+				.Image(InOutputTreeItem->GetIconBrush())
+			]
+			+ SHorizontalBox::Slot()
+			.FillWidth(1.f)
+			.Padding(5.f, 0.f, 0.f, 0.f)
+			.VAlign(VAlign_Center)
+			[
+				SNew(STextBlock)
+				.Text(InOutputTreeItem->GetDisplayName())
+			]
+		];
+}
diff --git a/Engine/Plugins/Experimental/Avalanche/Source/AvalancheMedia/Private/Broadcast/OutputDevices/Slate/SAvaBroadcastOutputTreeItem.h b/Engine/Plugins/Experimental/Avalanche/Source/AvalancheMedia/Private/Broadcast/OutputDevices/Slate/SAvaBroadcastOutputTreeItem.h
new file mode 100644
--- /dev/null
+++ b/Engine/Plugins/Experimental/Avalanche/Source/AvalancheMedia/Private/Broadcast/OutputDevices/Slate/SAvaBroadcastOutputTreeItem.h
@@ -0,0 +1,19 @@
+// Copyright Epic Games, Inc. All Rights Reserved.
+
+#pragma once
+
+#include "Widgets/SCompoundWidget.h"
+
+class FAvaBroadcastOutputTreeItem;
+
+class SAvaBroadcastOutputTreeItem : public SCompoundWidget
+{
+public:
+	SLATE_BEGIN_ARGS(SAvaBroadcastOutputTreeItem){}
+	SLATE_END_ARGS()
+	
+	void Construct(const FArguments& InArgs, const TSharedPtr<FAvaBroadcastOutputTreeItem>& InOutputTreeItem);
+	
+protected:
+	TWeakPtr<FAvaBroadcastOutputTreeItem> OutputTreeItemWeak;
+};
diff --git a/Engine/Plugins/Experimental/Avalanche/Source/AvalancheMedia/Private/Rundown/AvaRundownPlaybackUtils.cpp b/Engine/Plugins/Experimental/Avalanche/Source/AvalancheMedia/Private/Rundown/AvaRundownPlaybackUtils.cpp
new file mode 100644
--- /dev/null
+++ b/Engine/Plugins/Experimental/Avalanche/Source/AvalancheMedia/Private/Rundown/AvaRundownPlaybackUtils.cpp
@@ -0,0 +1,70 @@
+// Copyright Epic Games, Inc. All Rights Reserved.
+
+#include "Rundown/AvaRundownPlaybackUtils.h"
+
+#include "Rundown/AvaRundown.h"
+
+namespace UE::AvaRundownPlaybackUtils::Private
+{
+	template <typename Predicate>
+	TArray<int32> FilterSelectedOrPreviewingPages(const UAvaRundown* InRundown, const TArray<int32>& InSelectedPageIds, FName InPreviewChannelName, Predicate InPredicate)
+	{
+		if (InRundown)
+		{
+			if (InSelectedPageIds.Num())
+			{
+				return ::Invoke(InPredicate, InRundown, InSelectedPageIds);
+			}
+			return ::Invoke(InPredicate, InRundown, InRundown->GetPreviewingPageIds(InPreviewChannelName));
+		}
+		return {};
+	}
+}
+
+// Selecting the "next page" implies we know what the last page played was and we go to the next.
+// Normally, the last page played was the one currently selected, but the user might decide to select
+// another page for editing purpose, but still want to hit play next from the last played page.
+// This is why the "play head" is now added in the playback context for the given list.
+int32 FAvaRundownPlaybackUtils::GetPageIdToPlayNext(const UAvaRundown* InRundown, const FAvaRundownPageListReference& InPageListReference, bool bInPreview, FName InPreviewChannelName)
+{
+	// If there is no playback context, it means no pages are playing and can't have a next page if nothing is playing.
+	if (InRundown && InRundown->GetPageListPlaybackContextCollection())
+	{
+		if (const TSharedPtr<FAvaRundownPageListPlaybackContext> PlaybackContext = InRundown->GetPageListPlaybackContextCollection()->GetContext(bInPreview, InPreviewChannelName))
+		{
+			// Check if the current play head page is playing.
+			const bool bIsHeadPlaying = bInPreview ? InRundown->IsPagePreviewing(PlaybackContext->PlayHeadPageId) : InRundown->IsPagePlaying(PlaybackContext->PlayHeadPageId); 
+			
+			if (bIsHeadPlaying)
+			{
+				const int32 NextPageId = InRundown->GetNextPage(PlaybackContext->PlayHeadPageId, InPageListReference).GetPageId();
+				if (IsPageIdValid(NextPageId) && InRundown->CanPlayPage(NextPageId, true))
+				{
+					return NextPageId;
+				}
+			}
+		}
+	}
+	return FAvaRundownPage::InvalidPageId;
+}
+
+TArray<int32> FAvaRundownPlaybackUtils::GetPagesToTakeToProgram(const UAvaRundown* InRundown, const TArray<int32>& InSelectedPageIds, FName InPreviewChannel)
+{
+	using namespace UE::AvaRundownPlaybackUtils::Private;
+	auto KeepPagesToTakeToProgram = [](const UAvaRundown* InRundown, const TArray<int32>& InPageIds)
+	{
+		TArray<int32> OutPageIds;
+		OutPageIds.Reserve(InPageIds.Num());
+		for (int32 PageId : InPageIds)
+		{
+			// We can't take templates to program. Only instanced pages.
+			// Note: CanPlayPage returns true for playing pages, so need to make sure it is not already playing.
+			if (InRundown->CanPlayPage(PageId, false) && !InRundown->IsPagePlaying(PageId))
+			{
+				OutPageIds.Add(PageId);
+			}
+		}
+		return OutPageIds;
+	};
+	return FilterSelectedOrPreviewingPages(InRundown, InSelectedPageIds, InPreviewChannel, KeepPagesToTakeToProgram);
+}
diff --git a/Engine/Plugins/Experimental/Avalanche/Source/AvalancheMedia/Private/Rundown/AvaRundownServer.cpp b/Engine/Plugins/Experimental/Avalanche/Source/AvalancheMedia/Private/Rundown/AvaRundownServer.cpp
new file mode 100644
--- /dev/null
+++ b/Engine/Plugins/Experimental/Avalanche/Source/AvalancheMedia/Private/Rundown/AvaRundownServer.cpp
@@ -0,0 +1,1787 @@
+// Copyright Epic Games, Inc. All Rights Reserved.
+
+#include "Rundown/AvaRundownServer.h"
+
+#include "AssetRegistry/AssetData.h"
+#include "AssetRegistry/AssetRegistryModule.h"
+#include "AssetRegistry/IAssetRegistry.h"
+#include "Async/Async.h"
+#include "AvaMediaMessageUtils.h"
+#include "AvaMediaRenderTargetUtils.h"
+#include "AvaRemoteControlUtils.h"
+#include "Broadcast/AvaBroadcast.h"
+#include "Broadcast/OutputDevices/AvaBroadcastOutputClassItem.h"
+#include "Broadcast/OutputDevices/AvaBroadcastOutputDeviceItem.h"
+#include "Broadcast/OutputDevices/AvaBroadcastOutputRootItem.h"
+#include "Broadcast/OutputDevices/AvaBroadcastOutputTreeItem.h"
+#include "Broadcast/OutputDevices/AvaBroadcastRenderTargetMediaUtils.h"
+#include "IAvaMediaModule.h"
+#include "IRemoteControlModule.h"
+#include "ImageUtils.h"
+#include "MediaOutput.h"
+#include "MessageEndpointBuilder.h"
+#include "Misc/FileHelper.h"
+#include "Playback/AvaPlaybackManager.h"
+#include "Playback/AvaPlaybackUtils.h"
+#include "RenderingThread.h"
+#include "Rundown/AvaRundown.h"
+#include "Rundown/AvaRundownManagedInstanceCache.h"
+#include "Rundown/AvaRundownPagePlayer.h"
+#include "Rundown/AvaRundownPlaybackUtils.h"
+#include "Rundown/AvaRundownServerMediaOutputUtils.h"
+#include "TextureResource.h"
+
+#if WITH_EDITOR
+#include "Editor.h"
+#include "ScopedTransaction.h"
+#include "Subsystems/EditorAssetSubsystem.h"
+#endif
+
+DEFINE_LOG_CATEGORY_STATIC(LogAvaRundownServer, Log, All);
+
+#define LOCTEXT_NAMESPACE "AvaRundownServer"
+
+namespace UE::AvaRundownServer::Private
+{
+	inline FAvaRundownPageInfo GetPageInfo(const UAvaRundown* InRundown, const FAvaRundownPage& InPage)
+	{
+		FAvaRundownPageInfo PageInfo;
+		PageInfo.PageId = InPage.GetPageId();
+		PageInfo.PageName = InPage.GetPageName();
+		PageInfo.PageSummary = InPage.GetPageSummary().ToString();
+		PageInfo.FriendlyName = InPage.GetPageFriendlyName().ToString();
+		PageInfo.IsTemplate = InPage.IsTemplate();
+		PageInfo.TemplateId = InPage.GetTemplateId();
+		PageInfo.CombinedTemplateIds = InPage.GetCombinedTemplateIds();
+		PageInfo.AssetPath = InPage.GetAssetPath(InRundown);	// Todo: combo templates
+		PageInfo.Statuses = InPage.GetPageStatuses(InRundown);
+		PageInfo.TransitionLayerName = InPage.GetTransitionLayer(InRundown).ToString(); 	// Todo: combo templates
+		PageInfo.OutputChannel = InPage.GetChannelName().ToString();
+		PageInfo.bIsEnabled = InPage.IsEnabled();
+		PageInfo.bIsPlaying = InRundown->IsPagePlaying(InPage);
+		return PageInfo;
+	}
+
+	/**
+	 * Utility function load a rundown asset in memory.
+	 */
+	static TStrongObjectPtr<UAvaRundown> LoadRundown(const FSoftObjectPath& InRundownPath)
+	{
+		UObject* Object = InRundownPath.ResolveObject();
+		if (!Object)
+		{
+			Object = InRundownPath.TryLoad();
+		}
+		return TStrongObjectPtr<UAvaRundown>(Cast<UAvaRundown>(Object));
+	}
+
+	static FAvaRundownChannel SerializeChannel(const FAvaBroadcastOutputChannel& InChannel)
+	{
+		FAvaRundownChannel Channel;
+		Channel.Name = InChannel.GetChannelName().ToString();
+		Channel.State = InChannel.GetState();
+		Channel.IssueSeverity = InChannel.GetIssueSeverity();
+		const TArray<UMediaOutput*>& MediaOutputs = InChannel.GetMediaOutputs();
+		for (const UMediaOutput* MediaOutput : MediaOutputs)
+		{
+			FAvaRundownOutputDeviceItem DeviceItem;
+			DeviceItem.Name = MediaOutput->GetFName().ToString();
+			DeviceItem.OutputInfo = InChannel.GetMediaOutputInfo(MediaOutput);
+			DeviceItem.OutputState = InChannel.GetMediaOutputState(MediaOutput);
+			DeviceItem.IssueSeverity = InChannel.GetMediaOutputIssueSeverity(DeviceItem.OutputState, MediaOutput);
+			DeviceItem.IssueMessages = InChannel.GetMediaOutputIssueMessages(MediaOutput);
+			DeviceItem.Data = FAvaRundownServerMediaOutputUtils::SerializeMediaOutput(MediaOutput);
+			Channel.Devices.Push(MoveTemp(DeviceItem));
+		}
+
+		return Channel;
+	}
+	
+	// recursively search device and children
+	static FAvaOutputTreeItemPtr RecursiveFindOutputTreeItem(const FAvaOutputTreeItemPtr& InOutputTreeItem, const FString& InDeviceName)
+	{
+		if (!InOutputTreeItem->IsA<FAvaBroadcastOutputRootItem>() && InDeviceName == InOutputTreeItem->GetDisplayName().ToString())
+		{
+			return InOutputTreeItem;
+		}
+		
+		for (const TSharedPtr<IAvaBroadcastOutputTreeItem>& Child : InOutputTreeItem->GetChildren())
+		{
+			if (FAvaOutputTreeItemPtr TreeItem = RecursiveFindOutputTreeItem(Child, InDeviceName))
+			{
+				return TreeItem;
+			}
+		}
+		
+		return nullptr;
+	}
+
+	static UMediaOutput* FindChannelMediaOutput(const FAvaBroadcastOutputChannel& InOutputChannel, const FString& InOutputMediaName)
+	{
+		const TArray<UMediaOutput*>& MediaOutputs = InOutputChannel.GetMediaOutputs();
+		for (UMediaOutput* MediaOutput : MediaOutputs)
+		{
+			if (MediaOutput->GetName() == InOutputMediaName)
+			{
+				return MediaOutput;
+			}
+		}
+		
+		return nullptr;
+	}
+	
+	inline TArray<int32> GetPlayingPages(const UAvaRundown* InRundown, bool bInIsPreview, FName InChannelName)
+	{
+		return bInIsPreview ? InRundown->GetPreviewingPageIds(InChannelName) : InRundown->GetPlayingPageIds(InChannelName);
+	}
+
+	static bool ContinuePages(UAvaRundown* InRundown, const TArray<int32>& InPageIds, bool bInIsPreview, FName InPreviewChannelName, FString& OutFailureReason)
+	{
+		bool bSuccess = false;
+		for (const int32 PageId : InPageIds)
+		{
+			if (InRundown->CanContinuePage(PageId, bInIsPreview, InPreviewChannelName))
+			{
+				bSuccess |= InRundown->ContinuePage(PageId, bInIsPreview, InPreviewChannelName);
+			}
+			else if (bInIsPreview)
+			{
+				OutFailureReason.Appendf(TEXT("PageId %d was not previewing on channel \"%s\". "), PageId, *InPreviewChannelName.ToString());
+			}
+			else
+			{
+				OutFailureReason.Appendf(TEXT("PageId %d was not playing. "), PageId);
+			}
+		}
+		return bSuccess;
+	}
+
+	static bool UpdatePagesValues(const UAvaRundown* InRundown, const TArray<int32>& InPageIds, bool bInIsPreview, FName InPreviewChannelName)
+	{
+		bool bSuccess = false;
+		for (const int32 PageId : InPageIds)
+		{
+			bSuccess |= InRundown->PushRuntimeRemoteControlValues(PageId, bInIsPreview, InPreviewChannelName);
+		}
+		return bSuccess;
+	}
+}
+
+/**
+ * Holds the render target for copying the channel image.
+ * The render target needs to be held for many frames until it is done.
+ */
+struct FAvaRundownServer::FChannelImage
+{
+	// Optional temporary render target for converting pixel format.
+	TStrongObjectPtr<UTextureRenderTarget2D> RenderTarget;
+
+	// Pixels readback from the render target. Format is PF_B8G8R8A8 (for now).
+	TArray<FColor> RawPixels;
+	int32 SizeX = 0;
+	int32 SizeY = 0;
+
+	void UpdateRenderTarget(int32 InSizeX, int32 InSizeY, EPixelFormat InFormat, const FLinearColor& InClearColor)
+	{
+		if (!RenderTarget.IsValid())
+		{
+			static const FName ChannelImageRenderTargetBaseName = TEXT("AvaRundownServer_ChannelImageRenderTarget");
+			RenderTarget.Reset(UE::AvaMediaRenderTargetUtils::CreateDefaultRenderTarget(ChannelImageRenderTargetBaseName));
+		}
+	
+		UE::AvaMediaRenderTargetUtils::UpdateRenderTarget(RenderTarget.Get(), FIntPoint(InSizeX, InSizeY), InFormat, InClearColor);
+	}
+
+	void UpdateRawPixels(int32 InSizeX, int32 InSizeY)
+	{
+		SizeX = InSizeX;
+		SizeY = InSizeY;
+		RawPixels.SetNum(InSizeX * InSizeY);
+	}
+};
+
+FAvaRundownServer::FAvaRundownServer()
+{
+	
+}
+
+FAvaRundownServer::~FAvaRundownServer()
+{
+	RemoveBroadcastDelegates(&UAvaBroadcast::Get());
+	RemoveEditorDelegates();
+	
+	FMessageEndpoint::SafeRelease(MessageEndpoint);
+	
+	for (IConsoleObject* ConsoleCommand : ConsoleCommands)
+	{
+		IConsoleManager::Get().UnregisterConsoleObject(ConsoleCommand);
+	}
+	ConsoleCommands.Empty();
+}
+
+const FMessageAddress& FAvaRundownServer::GetMessageAddress() const
+{
+	if (MessageEndpoint.IsValid())
+	{
+		return MessageEndpoint->GetAddress();
+	}
+	static FMessageAddress InvalidMessageAddress;
+	return InvalidMessageAddress;
+}
+
+void FAvaRundownServer::Init(const FString& InAssignedHostName)
+{
+	HostName = InAssignedHostName.IsEmpty() ? FPlatformProcess::ComputerName() : InAssignedHostName;
+	
+	MessageEndpoint = FMessageEndpoint::Builder("MotionDesignRundownServer")
+	.Handling<FAvaRundownPing>(this, &FAvaRundownServer::HandleRundownPing)
+	.Handling<FAvaRundownGetRundowns>(this, &FAvaRundownServer::HandleGetRundowns)
+	.Handling<FAvaRundownLoadRundown>(this, &FAvaRundownServer::HandleLoadRundown)
+	.Handling<FAvaRundownSaveRundown>(this, &FAvaRundownServer::HandleSaveRundown)
+	.Handling<FAvaRundownCreatePage>(this, &FAvaRundownServer::HandleCreatePage)
+	.Handling<FAvaRundownDeletePage>(this, &FAvaRundownServer::HandleDeletePage)
+	.Handling<FAvaRundownCreateTemplate>(this, &FAvaRundownServer::HandleCreateTemplate)
+	.Handling<FAvaRundownDeleteTemplate>(this, &FAvaRundownServer::HandleDeleteTemplate)
+	.Handling<FAvaRundownChangeTemplateBP>(this, &FAvaRundownServer::HandleChangeTemplateBP)
+	.Handling<FAvaRundownGetPages>(this, &FAvaRundownServer::HandleGetPages)
+	.Handling<FAvaRundownGetPageDetails>(this, &FAvaRundownServer::HandleGetPageDetails)
+	.Handling<FAvaRundownPageChangeChannel>(this, &FAvaRundownServer::HandleChangePageChannel)
+	.Handling<FAvaRundownUpdatePageFromRCP>(this, &FAvaRundownServer::HandleUpdatePageFromRCP)
+	.Handling<FAvaRundownPageAction>(this, &FAvaRundownServer::HandlePageAction)
+	.Handling<FAvaRundownPagePreviewAction>(this, &FAvaRundownServer::HandlePagePreviewAction)
+	.Handling<FAvaRundownPageActions>(this, &FAvaRundownServer::HandlePageActions)
+	.Handling<FAvaRundownPagePreviewActions>(this, &FAvaRundownServer::HandlePagePreviewActions)
+	.Handling<FAvaRundownGetChannel>(this, &FAvaRundownServer::HandleGetChannel)
+	.Handling<FAvaRundownGetChannels>(this, &FAvaRundownServer::HandleGetChannels)
+	.Handling<FAvaRundownChannelAction>(this, &FAvaRundownServer::HandleChannelAction)
+	.Handling<FAvaRundownGetDevices>(this, &FAvaRundownServer::HandleGetDevices)
+	.Handling<FAvaRundownAddChannelDevice>(this, &FAvaRundownServer::HandleAddChannelDevice)
+	.Handling<FAvaRundownEditChannelDevice>(this, &FAvaRundownServer::HandleEditChannelDevice)
+	.Handling<FAvaRundownRemoveChannelDevice>(this, &FAvaRundownServer::HandleRemoveChannelDevice)
+	.Handling<FAvaRundownGetChannelImage>(this, &FAvaRundownServer::HandleGetChannelImage)
+	.NotificationHandling(FOnBusNotification::CreateRaw(this, &FAvaRundownServer::OnMessageBusNotification));
+	
+	if (MessageEndpoint.IsValid())
+	{
+		// Subscribe to the server listing requests
+		MessageEndpoint->Subscribe<FAvaRundownPing>();
+
+		SetupBroadcastDelegates(&UAvaBroadcast::Get());
+		SetupEditorDelegates();
+
+		UE_LOG(LogAvaRundownServer, Log, TEXT("Motion Design Rundown Server \"%s\" Started."), *HostName);
+	}
+}
+
+void FAvaRundownServer::SetupBroadcastDelegates(UAvaBroadcast* InBroadcast)
+{
+	RemoveBroadcastDelegates(InBroadcast);
+	InBroadcast->GetOnChannelsListChanged().AddSP(this, &FAvaRundownServer::OnBroadcastChannelListChanged);
+	FAvaBroadcastOutputChannel::GetOnChannelChanged().AddSP(this, &FAvaRundownServer::OnBroadcastChannelChanged);
+}
+
+void FAvaRundownServer::SetupEditorDelegates()
+{
+	RemoveEditorDelegates();
+	const FAssetRegistryModule& AssetRegistryModule = FModuleManager::LoadModuleChecked<FAssetRegistryModule>("AssetRegistry");
+	AssetRegistryModule.Get().OnAssetAdded().AddSP(this, &FAvaRundownServer::OnAssetAddedOrRemoved);
+	AssetRegistryModule.Get().OnAssetRemoved().AddSP(this, &FAvaRundownServer::OnAssetAddedOrRemoved);
+}
+
+void FAvaRundownServer::RemoveBroadcastDelegates(UAvaBroadcast* InBroadcast) const
+{
+	InBroadcast->GetOnChannelsListChanged().RemoveAll(this);
+	FAvaBroadcastOutputChannel::GetOnChannelChanged().RemoveAll(this);
+}
+
+void FAvaRundownServer::RemoveEditorDelegates() const
+{
+	const FAssetRegistryModule& AssetRegistryModule = FModuleManager::LoadModuleChecked<FAssetRegistryModule>("AssetRegistry");
+	AssetRegistryModule.Get().OnAssetAdded().RemoveAll(this);
+	AssetRegistryModule.Get().OnAssetRemoved().RemoveAll(this);
+}
+
+void FAvaRundownServer::OnPageListChanged(const FAvaRundownPageListChangeParams& InParams) const
+{
+	FAvaRundownPageListChanged* ReplyMessage = FMessageEndpoint::MakeMessage<FAvaRundownPageListChanged>();
+	ReplyMessage->Rundown = FSoftObjectPath(InParams.Rundown).ToString();
+	ReplyMessage->ChangeType = static_cast<uint8>(InParams.ChangeType);
+	ReplyMessage->AffectedPages = InParams.AffectedPages;
+	
+	SendResponse(ReplyMessage, ClientAddresses);
+}
+
+void FAvaRundownServer::OnPagesChanged(const UAvaRundown* InRundown, const FAvaRundownPage& InPage, EAvaRundownPageChanges InChange) const
+{
+	switch(InChange)
+	{
+	case EAvaRundownPageChanges::AnimationSettings:
+		{
+			PageAnimSettingsChanged(InRundown, InPage);
+			break;
+		}
+	case EAvaRundownPageChanges::Blueprint:
+		{
+			PageBlueprintChanged(InRundown, InPage, InPage.GetAssetPath(InRundown).ToString());
+			break;
+		}
+	case EAvaRundownPageChanges::Status:
+		{
+			PageStatusChanged(InRundown, InPage);
+			break;
+		}
+	case EAvaRundownPageChanges::Channel:
+		{
+			PageChannelChanged(InRundown, InPage, InPage.GetChannelName().ToString());
+			break;
+		}
+	case EAvaRundownPageChanges::RemoteControlValues:
+		{
+			break;
+		}
+	case EAvaRundownPageChanges::All:
+		{
+			break;
+		}
+	case EAvaRundownPageChanges::None:
+		{
+			break;
+		}
+	default:
+		{
+			break;
+		}
+	}
+	
+}
+void FAvaRundownServer::PageStatusChanged(const UAvaRundown* InRundown, const FAvaRundownPage& InPage) const
+{
+	FAvaRundownPagesStatuses* ReplyMessage = FMessageEndpoint::MakeMessage<FAvaRundownPagesStatuses>();
+	ReplyMessage->Rundown = FSoftObjectPath(InRundown).ToString();
+	ReplyMessage->PageInfo = UE::AvaRundownServer::Private::GetPageInfo(InRundown, InPage);
+	SendResponse(ReplyMessage, ClientAddresses);
+}
+
+void FAvaRundownServer::PageBlueprintChanged(const UAvaRundown* InRundown, const FAvaRundownPage& InPage, const FString& InBlueprintPath) const
+{
+	FAvaRundownPageBlueprintChanged* ReplyMessage = FMessageEndpoint::MakeMessage<FAvaRundownPageBlueprintChanged>();
+	ReplyMessage->Rundown = FSoftObjectPath(InRundown).ToString();
+	ReplyMessage->PageId = InPage.GetPageId();
+	ReplyMessage->BlueprintPath = InBlueprintPath;
+	SendResponse(ReplyMessage, ClientAddresses);
+}
+
+void FAvaRundownServer::PageChannelChanged(const UAvaRundown* InRundown, const FAvaRundownPage& InPage, const FString& InChannelName) const
+{
+	FAvaRundownPageChannelChanged* ReplyMessage = FMessageEndpoint::MakeMessage<FAvaRundownPageChannelChanged>();
+	ReplyMessage->Rundown = FSoftObjectPath(InRundown).ToString();
+	ReplyMessage->PageId = InPage.GetPageId();
+	ReplyMessage->ChannelName = InChannelName;
+	SendResponse(ReplyMessage, ClientAddresses);
+}
+
+void FAvaRundownServer::PageAnimSettingsChanged(const UAvaRundown* InRundown, const FAvaRundownPage& InPage) const
+{
+	FAvaRundownPageAnimSettingsChanged* ReplyMessage = FMessageEndpoint::MakeMessage<FAvaRundownPageAnimSettingsChanged>();
+	ReplyMessage->Rundown = FSoftObjectPath(InRundown).ToString();
+	ReplyMessage->PageId = InPage.GetPageId();
+	SendResponse(ReplyMessage, ClientAddresses);
+}
+
+void FAvaRundownServer::OnBroadcastChannelListChanged(const FAvaBroadcastProfile& InProfile) const
+{
+	if (ClientAddresses.IsEmpty())
+	{
+		return;
+	}
+
+	FAvaRundownChannelListChanged* ReplyMessage = FMessageEndpoint::MakeMessage<FAvaRundownChannelListChanged>();
+
+	const TArray<FAvaBroadcastOutputChannel*>& OutputChannels = InProfile.GetChannels();
+	
+	ReplyMessage->Channels.Reserve(OutputChannels.Num());
+
+	for (const FAvaBroadcastOutputChannel* OutputChannel : OutputChannels)
+	{
+		FAvaRundownChannel Channel = UE::AvaRundownServer::Private::SerializeChannel(*OutputChannel);
+		ReplyMessage->Channels.Push(MoveTemp(Channel));
+	}
+	SendResponse(ReplyMessage, ClientAddresses);
+}
+
+void FAvaRundownServer::OnBroadcastChannelChanged(const FAvaBroadcastOutputChannel& InChannel, EAvaBroadcastChannelChange InChange) const
+{
+	if (ClientAddresses.IsEmpty())
+	{
+		return;
+	}
+
+	FAvaRundownChannelResponse* ReplyMessage = FMessageEndpoint::MakeMessage<FAvaRundownChannelResponse>();
+	
+	ReplyMessage->Channel = UE::AvaRundownServer::Private::SerializeChannel(InChannel);
+	SendResponse(ReplyMessage, ClientAddresses);
+}
+
+void FAvaRundownServer::OnAssetAddedOrRemoved(const FAssetData& InAssetData) const
+{
+	if (ClientAddresses.IsEmpty())
+	{
+		return;
+	}
+	
+	using namespace UE::AvaRundownServer::Private;
+	if (InAssetData.GetClass() == UAvaRundown::StaticClass() || FAvaPlaybackUtils::IsPlayableAsset(InAssetData))
+	{
+		FAvaRundownAssetsChanged* ReplyMessage = FMessageEndpoint::MakeMessage<FAvaRundownAssetsChanged>();
+		ReplyMessage->AssetName = InAssetData.AssetName.ToString();
+		SendResponse(ReplyMessage, ClientAddresses);
+	}
+}
+
+void FAvaRundownServer::HandleRundownPing(const FAvaRundownPing& InMessage, const TSharedRef<IMessageContext, ESPMode::ThreadSafe>& InContext)
+{
+	if (!InMessage.bAuto)
+	{
+		UE_LOG(LogAvaRundownServer, Log, TEXT("Received Ping request from %s"), *InContext->GetSender().ToString());
+	}
+
+	FAvaRundownPong* ReplyMessage = FMessageEndpoint::MakeMessage<FAvaRundownPong>();
+	ReplyMessage->RequestId = InMessage.RequestId;
+	ReplyMessage->bAuto = InMessage.bAuto;
+
+	// We still support the initial version.
+	constexpr int32 CurrentMinimumApiVersion = EAvaRundownApiVersion::Initial;
+
+	// Consider clients that didn't request a version to be the "initial" version.
+	const int32 RequestedApiVersion = InMessage.RequestedApiVersion != -1 ? InMessage.RequestedApiVersion : EAvaRundownApiVersion::Initial;
+
+	// Determine the version we will communicate with this client.
+	int32 HonoredApiVersion = EAvaRundownApiVersion::LatestVersion;
+	
+	if (RequestedApiVersion >= CurrentMinimumApiVersion && RequestedApiVersion <= EAvaRundownApiVersion::LatestVersion)
+	{
+		HonoredApiVersion = RequestedApiVersion;
+	}
+	
+	ReplyMessage->ApiVersion = HonoredApiVersion;
+	ReplyMessage->MinimumApiVersion = CurrentMinimumApiVersion;
+	ReplyMessage->LatestApiVersion = EAvaRundownApiVersion::LatestVersion;
+	ReplyMessage->HostName = HostName;
+
+	FClientInfo& ClientInfo = GetOrAddClientInfo(InContext->GetSender());
+	ClientInfo.ApiVersion = HonoredApiVersion;
+
+	SendResponse(ReplyMessage, InContext->GetSender());
+}
+
+void FAvaRundownServer::HandleGetRundowns(const FAvaRundownGetRundowns& InMessage, const TSharedRef<IMessageContext, ESPMode::ThreadSafe>& InContext)
+{
+	FAvaRundownRundowns* ReplyMessage = FMessageEndpoint::MakeMessage<FAvaRundownRundowns>();
+
+	ReplyMessage->RequestId = InMessage.RequestId;
+	
+	// List all the play list.
+	if (IAssetRegistry* AssetRegistry = IAssetRegistry::Get())
+	{
+		TArray<FAssetData> Assets;
+		if (AssetRegistry->GetAssetsByClass(UAvaRundown::StaticClass()->GetClassPathName(), Assets))
+		{
+			ReplyMessage->Rundowns.Reserve(Assets.Num());		
+			for (const FAssetData& Data : Assets)
+			{
+				ReplyMessage->Rundowns.Add(Data.ToSoftObjectPath().ToString());
+			}
+		}
+	}
+
+	SendResponse(ReplyMessage, InContext->GetSender());
+}
+
+void FAvaRundownServer::HandleLoadRundown(const FAvaRundownLoadRundown& InMessage, const TSharedRef<IMessageContext, ESPMode::ThreadSafe>& InContext)
+{
+	// If the requested path is empty, we assume this is a request for information only.
+	if (!InMessage.Rundown.IsEmpty())
+	{
+		const FSoftObjectPath NewRundownPath(InMessage.Rundown);
+		const UAvaRundown* Rundown = RundownPlaybackCommandData.GetOrLoadRundown(NewRundownPath,
+			[this](UAvaRundown* InRundown)
+			{
+				RundownPlaybackCommandData.ClosePlaybackContext();
+				RundownPlaybackCommandData.RemoveRundownDelegates(this, InRundown);
+			},
+			[this](UAvaRundown* InRundown)
+			{
+				RundownPlaybackCommandData.SetupRundownDelegates(this, InRundown);
+				if (InRundown)
+				{
+					InRundown->InitializePlaybackContext();
+				}
+			});
+		
+		if (!Rundown)
+		{
+			SendMessage(InContext->GetSender(), InMessage.RequestId, ELogVerbosity::Error,
+				TEXT("Rundown \"%s\" not loaded."), *InMessage.Rundown);
+			return;
+		}
+	}
+
+	SendMessage(InContext->GetSender(), InMessage.RequestId, ELogVerbosity::Log,
+		TEXT("Rundown \"%s\" loaded."), *RundownPlaybackCommandData.CurrentRundownPath.ToString());
+}
+
+void FAvaRundownServer::HandleSaveRundown(const FAvaRundownSaveRundown& InMessage, const TSharedRef<IMessageContext, ESPMode::ThreadSafe>& InContext)
+{
+#if WITH_EDITOR
+	UEditorAssetSubsystem* EditorAssetSubsystem = GEditor ? GEditor->GetEditorSubsystem<UEditorAssetSubsystem>() : nullptr;
+	
+	if (!EditorAssetSubsystem)
+	{
+		SendMessage(InContext->GetSender(), InMessage.RequestId, ELogVerbosity::Error,
+			TEXT("Saving assets is only available in editor mode."));
+		return;
+	}
+	
+	if (InMessage.Rundown.IsEmpty())
+	{
+		SendMessage(InContext->GetSender(), InMessage.RequestId, ELogVerbosity::Error,
+			TEXT("Rundown asset not specified."));
+		return;
+	}
+
+	const FSoftObjectPath RundownAssetPath(InMessage.Rundown);
+
+	if (!RundownAssetPath.IsValid())
+	{
+		SendMessage(InContext->GetSender(), InMessage.RequestId, ELogVerbosity::Error,
+			TEXT("Rundown asset path \"%s\" is not valid."), *InMessage.Rundown);
+		return;
+	}
+
+	UObject* FoundObject = RundownAssetPath.ResolveObject();
+
+	if (!FoundObject)
+	{
+		SendMessage(InContext->GetSender(), InMessage.RequestId, ELogVerbosity::Error,
+			TEXT("Rundown asset \"%s\" is not loaded."), *InMessage.Rundown);
+		return;
+	}
+
+	UAvaRundown* FoundRundown = Cast<UAvaRundown>(FoundObject);
+	
+	if (!FoundRundown)
+	{
+		SendMessage(InContext->GetSender(), InMessage.RequestId, ELogVerbosity::Error,
+			TEXT("Asset path \"%s\" is loaded but is not a Rundown asset."), *InMessage.Rundown);
+		return;
+	}
+	
+	if (!EditorAssetSubsystem->SaveLoadedAsset(FoundRundown, InMessage.bOnlyIfIsDirty))
+	{
+		SendMessage(InContext->GetSender(), InMessage.RequestId, ELogVerbosity::Error,
+			TEXT("Unable to save asset \"%s\" to location \"%s\"."), *FoundRundown->GetName(), *RundownAssetPath.ToString());
+		return;
+	}
+
+	SendMessage(InContext->GetSender(), InMessage.RequestId, ELogVerbosity::Log,
+		TEXT("Asset \"%s\" save to location \"%s\"."), *FoundRundown->GetName(), *RundownAssetPath.ToString());
+#else
+	SendMessage(InContext->GetSender(), InMessage.RequestId, ELogVerbosity::Error, TEXT("Save rundown is not available in game build."));
+#endif
+}
+
+void FAvaRundownServer::HandleGetPages(const FAvaRundownGetPages& InMessage, const TSharedRef<IMessageContext, ESPMode::ThreadSafe>& InContext)
+{
+	UAvaRundown* Rundown = GetOrLoadRundownForEdit(InContext->GetSender(), InMessage.RequestId, InMessage.Rundown);
+	if (!Rundown)
+	{
+		return;
+	}
+	
+	FAvaRundownPages* ReplyMessage = FMessageEndpoint::MakeMessage<FAvaRundownPages>();
+	ReplyMessage->RequestId = InMessage.RequestId;
+	ReplyMessage->Pages.Reserve(Rundown->GetInstancedPages().Pages.Num() + Rundown->GetTemplatePages().Pages.Num());
+	for (const FAvaRundownPage& Page : Rundown->GetInstancedPages().Pages)
+	{
+		FAvaRundownPageInfo PageInfo = UE::AvaRundownServer::Private::GetPageInfo(Rundown, Page);
+		ReplyMessage->Pages.Add(MoveTemp(PageInfo));
+	}
+	
+	for (const FAvaRundownPage& Page : Rundown->GetTemplatePages().Pages)
+	{
+		FAvaRundownPageInfo PageInfo = UE::AvaRundownServer::Private::GetPageInfo(Rundown, Page);
+		ReplyMessage->Pages.Add(MoveTemp(PageInfo));
+	}
+	
+	SendResponse(ReplyMessage, InContext->GetSender());
+}
+
+void FAvaRundownServer::HandleCreatePage(const FAvaRundownCreatePage& InMessage,
+	const TSharedRef<IMessageContext, ESPMode::ThreadSafe>& InContext)
+{
+	UAvaRundown* Rundown = GetOrLoadRundownForEdit(InContext->GetSender(), InMessage.RequestId, InMessage.Rundown);
+	if (!Rundown)
+	{
+		return;
+	}
+
+	const FAvaRundownPage& Template = Rundown->GetPage(InMessage.TemplateId);
+    if (!Template.IsValidPage() || !Template.IsTemplate())
+    {
+    	LogAndSendMessage(InContext->GetSender(), InMessage.RequestId, ELogVerbosity::Error, TEXT("Template %d is not valid or is not a template"), InMessage.TemplateId);
+    	return;
+    }
+	
+	const int32 PageId = Rundown->AddPageFromTemplate(InMessage.TemplateId);
+	if (PageId != FAvaRundownPage::InvalidPageId)
+	{
+		SendMessage(InContext->GetSender(), InMessage.RequestId, ELogVerbosity::Log, TEXT("Page %d Created from Template %d"), PageId, InMessage.TemplateId);
+	}
+	else
+	{
+		SendMessage(InContext->GetSender(), InMessage.RequestId, ELogVerbosity::Error, TEXT("Failed to create a page from Template %d"), InMessage.TemplateId);
+	}
+}
+
+void FAvaRundownServer::HandleDeletePage(const FAvaRundownDeletePage& InMessage, const TSharedRef<IMessageContext, ESPMode::ThreadSafe>& InContext)
+{
+	UAvaRundown* Rundown = GetOrLoadRundownForEdit(InContext->GetSender(), InMessage.RequestId, InMessage.Rundown);
+	if (!Rundown)
+	{
+		return;
+	}
+
+	const FAvaRundownPage& Page = Rundown->GetPage(InMessage.PageId);
+	if (!Page.IsValidPage())
+	{
+		LogAndSendMessage(InContext->GetSender(), InMessage.RequestId, ELogVerbosity::Error, TEXT("page %d is not valid"), InMessage.PageId);
+		return;
+	}
+
+	Rundown->RemovePage(InMessage.PageId);
+	SendMessage(InContext->GetSender(), InMessage.RequestId, ELogVerbosity::Log, TEXT("Page %d deleted"), InMessage.PageId);
+}
+
+void FAvaRundownServer::HandleDeleteTemplate(const FAvaRundownDeleteTemplate& InMessage,
+	const TSharedRef<IMessageContext, ESPMode::ThreadSafe>& InContext)
+{
+	UAvaRundown* Rundown = GetOrLoadRundownForEdit(InContext->GetSender(), InMessage.RequestId, InMessage.Rundown);
+	if (!Rundown)
+	{
+		return;
+	}
+
+	const FAvaRundownPage& Page = Rundown->GetPage(InMessage.PageId);
+	if (!Page.IsValidPage())
+	{
+		LogAndSendMessage(InContext->GetSender(), InMessage.RequestId, ELogVerbosity::Error, TEXT("page %d is not valid"), InMessage.PageId);
+		return;
+	}
+
+	if (!Page.IsTemplate())
+	{
+		LogAndSendMessage(InContext->GetSender(), InMessage.RequestId, ELogVerbosity::Error, TEXT("page %d is not a template"), InMessage.PageId);
+		return;
+	}
+
+	if (!Page.GetInstancedIds().IsEmpty())
+	{
+		LogAndSendMessage(InContext->GetSender(), InMessage.RequestId, ELogVerbosity::Error, TEXT("Template has instanced pages"), InMessage.PageId);
+		return;
+	}
+
+	Rundown->RemovePage(InMessage.PageId);
+	SendMessage(InContext->GetSender(), InMessage.RequestId, ELogVerbosity::Log, TEXT("Page template %d deleted"), InMessage.PageId);
+}
+
+void FAvaRundownServer::HandleCreateTemplate(const FAvaRundownCreateTemplate& InMessage,
+	const TSharedRef<IMessageContext, ESPMode::ThreadSafe>& InContext)
+{
+	UAvaRundown* Rundown = GetOrLoadRundownForEdit(InContext->GetSender(), InMessage.RequestId, InMessage.Rundown);
+	if (!Rundown)
+	{
+		return;
+	}
+	
+	Rundown->AddTemplate();
+	SendMessage(InContext->GetSender(), InMessage.RequestId, ELogVerbosity::Log, TEXT("Template Created"));
+}
+
+
+void FAvaRundownServer::HandleChangeTemplateBP(const FAvaRundownChangeTemplateBP& InMessage,
+	const TSharedRef<IMessageContext, ESPMode::ThreadSafe>& InContext)
+{
+	UAvaRundown* Rundown = GetOrLoadRundownForEdit(InContext->GetSender(), InMessage.RequestId, InMessage.Rundown);
+	if (!Rundown)
+	{
+		return;
+	}
+
+	FAvaRundownPage& Page = Rundown->GetPage(InMessage.TemplateId);
+	if (Page.IsValidPage() && Page.IsTemplate())
+	{
+		if (Page.UpdateAsset(InMessage.AssetPath))
+		{
+			SendMessage(InContext->GetSender(), InMessage.RequestId, ELogVerbosity::Log, TEXT("Blueprint change of template: %d to %s"), InMessage.TemplateId, *InMessage.AssetPath);
+			Rundown->GetOnPagesChanged().Broadcast(Rundown, Page, EAvaRundownPageChanges::Blueprint);
+			return;
+		}
+	}
+	LogAndSendMessage(InContext->GetSender(), InMessage.RequestId, ELogVerbosity::Error, TEXT("Blueprint change of template: %d to %s failed."), InMessage.TemplateId, *InMessage.AssetPath);
+}
+
+void FAvaRundownServer::HandleGetPageDetails(const FAvaRundownGetPageDetails& InMessage, const TSharedRef<IMessageContext, ESPMode::ThreadSafe>& InContext)
+{	
+	UAvaRundown* Rundown = GetOrLoadRundownForEdit(InContext->GetSender(), InMessage.RequestId, InMessage.Rundown);
+	if (!Rundown)
+	{
+		return;
+	}
+
+	const FAvaRundownPage& Page = Rundown->GetPage(InMessage.PageId);
+	if (!Page.IsValidPage())
+	{
+		LogAndSendMessage(InContext->GetSender(), InMessage.RequestId, ELogVerbosity::Error, TEXT("\"PageDetails\" not available: PageId %d is invalid."), InMessage.PageId);
+		return;
+	}
+
+	if (Page.GetAssetPath(Rundown).IsNull())
+	{
+		LogAndSendMessage(InContext->GetSender(), InMessage.RequestId, ELogVerbosity::Error, TEXT("Page has no asset selected"), InMessage.PageId);
+		return;
+	}
+	
+	if (InMessage.bLoadRemoteControlPreset)
+	{
+		FAvaRundownManagedInstanceCache& ManagedInstanceCache = IAvaMediaModule::Get().GetManagedInstanceCache();
+		const TSharedPtr<FAvaRundownManagedInstance> ManagedInstance
+			= ManagedInstanceCache.GetOrLoadInstance(Page.GetAssetPath(Rundown));
+		
+		if (ManagedInstance.IsValid())
+		{
+			RundownEditCommandData.SaveCurrentRemoteControlPresetToPage(true);
+
+			// Applying the controller values can break the WYSIWYG of the editor,
+			// in case multiple controllers set the same exposed entity with different values.
+			// There is no guaranty that the controller actions are self consistent.
+			// To avoid this issue, we apply the controllers first, and then
+			// restore the entity values in a second pass.
+			
+			Page.GetRemoteControlValues().ApplyControllerValuesToRemoteControlPreset(ManagedInstance->GetRemoteControlPreset(), true);
+			Page.GetRemoteControlValues().ApplyEntityValuesToRemoteControlPreset(ManagedInstance->GetRemoteControlPreset());
+
+			// Register the RC Preset to Remote Control module to make it available through WebRC.
+			FAvaRemoteControlUtils::RegisterRemoteControlPreset(ManagedInstance->GetRemoteControlPreset(), /*bInEnsureUniqueId*/ true);
+
+			// Keep track of what is currently registered.
+			RundownEditCommandData.ManagedPageId = InMessage.PageId;
+			RundownEditCommandData.ManagedInstance = ManagedInstance;
+		}
+	}
+	
+	FAvaRundownPageDetails* ReplyMessage = FMessageEndpoint::MakeMessage<FAvaRundownPageDetails>();
+	ReplyMessage->RequestId = InMessage.RequestId;
+	ReplyMessage->Rundown = InMessage.Rundown;
+	ReplyMessage->PageInfo = UE::AvaRundownServer::Private::GetPageInfo(Rundown, Page);
+	ReplyMessage->RemoteControlValues = Page.GetRemoteControlValues();
+	if (InMessage.bLoadRemoteControlPreset && RundownEditCommandData.ManagedInstance.IsValid() && RundownEditCommandData.ManagedInstance->GetRemoteControlPreset())
+	{
+		ReplyMessage->RemoteControlPresetName = RundownEditCommandData.ManagedInstance->GetRemoteControlPreset()->GetPresetName().ToString();
+		ReplyMessage->RemoteControlPresetId = RundownEditCommandData.ManagedInstance->GetRemoteControlPreset()->GetPresetId().ToString();
+	}
+	SendResponse(ReplyMessage, InContext->GetSender());
+}
+
+void FAvaRundownServer::HandleChangePageChannel(const FAvaRundownPageChangeChannel& InMessage,
+	const TSharedRef<IMessageContext, ESPMode::ThreadSafe>& InContext)
+{
+	UAvaBroadcast& Broadcast = UAvaBroadcast::Get();
+	if (InMessage.ChannelName.IsEmpty())
+	{
+		LogAndSendMessage(InContext->GetSender(), InMessage.RequestId, ELogVerbosity::Error, TEXT("No Channel Name Provided"));
+		return;
+	}
+
+	const FName ChannelName(InMessage.ChannelName);
+	const FAvaBroadcastOutputChannel& Channel = Broadcast.GetCurrentProfile().GetChannel(ChannelName);
+	if (!Channel.IsValidChannel())
+	{
+		LogAndSendMessage(InContext->GetSender(), InMessage.RequestId, ELogVerbosity::Error, TEXT("%s is not a valid channel"), *ChannelName.ToString());
+		return;
+	}
+
+	UAvaRundown* Rundown = GetOrLoadRundownForEdit(InContext->GetSender(), InMessage.RequestId, InMessage.Rundown);
+	if (!Rundown)
+	{
+		return;
+	}
+
+	FAvaRundownPage& Page = Rundown->GetPage(InMessage.PageId);
+	if (!Page.IsValidPage())
+	{
+		LogAndSendMessage(InContext->GetSender(), InMessage.RequestId, ELogVerbosity::Error, TEXT("PageId %d is invalid."), InMessage.PageId);
+		return;
+	}
+
+	if (Page.GetChannelName() == ChannelName)
+	{
+		SendMessage(InContext->GetSender(), InMessage.RequestId, ELogVerbosity::Log, TEXT("Same Channel Selected"), InMessage.PageId);
+		return;
+	}
+
+	Page.SetChannelName(ChannelName);
+	Rundown->GetOnPagesChanged().Broadcast(Rundown, Page, EAvaRundownPageChanges::Channel);
+	SendMessage(InContext->GetSender(), InMessage.RequestId, ELogVerbosity::Log, TEXT("Channel changed"));
+}
+
+void FAvaRundownServer::HandleUpdatePageFromRCP(const FAvaRundownUpdatePageFromRCP& InMessage, const TSharedRef<IMessageContext, ESPMode::ThreadSafe>& InContext)
+{
+	// Note that this doesn't save the rundown.
+	RundownEditCommandData.SaveCurrentRemoteControlPresetToPage(InMessage.bUnregister);
+	if (InMessage.bUnregister)
+	{
+		RundownEditCommandData.ManagedPageId = FAvaRundownPage::InvalidPageId;
+		RundownEditCommandData.ManagedInstance.Reset();
+	}
+	SendMessage(InContext->GetSender(), InMessage.RequestId, ELogVerbosity::Log, TEXT("\"UpdatePageFromRCP\" Ok."));
+}
+
+void FAvaRundownServer::HandlePageAction(const FAvaRundownPageAction& InMessage, const TSharedRef<IMessageContext, ESPMode::ThreadSafe>& InContext)
+{
+	const FRequestInfo RequestInfo = {InMessage.RequestId, InContext->GetSender()};
+	HandlePageActions(RequestInfo, {InMessage.PageId}, false, FName(), InMessage.Action);
+}
+
+void FAvaRundownServer::HandlePagePreviewAction(const FAvaRundownPagePreviewAction& InMessage,
+	const TSharedRef<IMessageContext, ESPMode::ThreadSafe>& InContext)
+{
+	const FRequestInfo RequestInfo = {InMessage.RequestId, InContext->GetSender()};
+	HandlePageActions(RequestInfo, {InMessage.PageId}, true, FName(InMessage.PreviewChannelName), InMessage.Action);
+}
+
+void FAvaRundownServer::HandlePageActions(const FAvaRundownPageActions& InMessage, const TSharedRef<IMessageContext, ESPMode::ThreadSafe>& InContext)
+{
+	const FRequestInfo RequestInfo = {InMessage.RequestId, InContext->GetSender()};
+	HandlePageActions(RequestInfo, InMessage.PageIds, false, FName(), InMessage.Action);
+}
+
+void FAvaRundownServer::HandlePagePreviewActions(const FAvaRundownPagePreviewActions& InMessage,
+	const TSharedRef<IMessageContext, ESPMode::ThreadSafe>& InContext)
+{
+	const FRequestInfo RequestInfo = {InMessage.RequestId, InContext->GetSender()};
+	HandlePageActions(RequestInfo, InMessage.PageIds, true, FName(InMessage.PreviewChannelName), InMessage.Action);
+}
+
+void FAvaRundownServer::HandleGetChannel(const FAvaRundownGetChannel& InMessage, const TSharedRef<IMessageContext, ESPMode::ThreadSafe>& InContext)
+{
+	const FName ChannelName(InMessage.ChannelName);
+	
+	const UAvaBroadcast& Broadcast = UAvaBroadcast::Get();
+	const FAvaBroadcastOutputChannel& Channel = Broadcast.GetCurrentProfile().GetChannel(ChannelName);
+
+	if (!Channel.IsValidChannel())
+	{
+		LogAndSendMessage(InContext->GetSender(), InMessage.RequestId, ELogVerbosity::Error, TEXT("\"GetChannel\" Channel \"%s\" not found."), *InMessage.ChannelName);
+		return;
+	}
+	
+	FAvaRundownChannelResponse* ReplyMessage = FMessageEndpoint::MakeMessage<FAvaRundownChannelResponse>();
+	ReplyMessage->RequestId = InMessage.RequestId;
+	ReplyMessage->Channel = UE::AvaRundownServer::Private::SerializeChannel(Channel);
+	SendResponse(ReplyMessage, InContext->GetSender());
+}
+
+void FAvaRundownServer::HandleGetChannels(const FAvaRundownGetChannels& InMessage, const TSharedRef<IMessageContext, ESPMode::ThreadSafe>& InContext)
+{
+	UAvaBroadcast& Broadcast = UAvaBroadcast::Get();
+	FAvaRundownChannels* ReplyMessage = FMessageEndpoint::MakeMessage<FAvaRundownChannels>();
+	ReplyMessage->RequestId = InMessage.RequestId;
+	
+	const TArray<FAvaBroadcastOutputChannel*>& Channels = Broadcast.GetCurrentProfile().GetChannels();
+	ReplyMessage->Channels.Reserve(Channels.Num());
+
+	for (const FAvaBroadcastOutputChannel* OutputChannel : Channels)
+	{
+		FAvaRundownChannel Channel = UE::AvaRundownServer::Private::SerializeChannel(*OutputChannel);
+		ReplyMessage->Channels.Push(MoveTemp(Channel));
+	}
+	
+	SendResponse(ReplyMessage, InContext->GetSender());
+}
+
+void FAvaRundownServer::HandleChannelAction(const FAvaRundownChannelAction& InMessage, const TSharedRef<IMessageContext, ESPMode::ThreadSafe>& InContext)
+{
+	UAvaBroadcast& Broadcast = UAvaBroadcast::Get();
+	if (InMessage.Action == EAvaRundownChannelActions::Start)
+	{
+		if (InMessage.ChannelName.IsEmpty())
+		{
+			Broadcast.StartBroadcast();
+			SendMessage(InContext->GetSender(), InMessage.RequestId, ELogVerbosity::Log, TEXT("\"ChannelAction\" Ok."));
+		}
+		else
+		{
+			const FName ChannelName(InMessage.ChannelName);
+			FAvaBroadcastOutputChannel& Channel = Broadcast.GetCurrentProfile().GetChannelMutable(ChannelName);
+			if (Channel.IsValidChannel())
+			{
+				Channel.StartChannelBroadcast();
+				SendMessage(InContext->GetSender(), InMessage.RequestId, ELogVerbosity::Log, TEXT("\"ChannelAction\" Ok."));
+			}
+			else
+			{
+				LogAndSendMessage(InContext->GetSender(), InMessage.RequestId, ELogVerbosity::Error, TEXT("\"ChannelAction\" Failed. Reason: Invalid Channel \"%s\"."), *InMessage.ChannelName);
+			}
+		}
+	}
+	else if (InMessage.Action == EAvaRundownChannelActions::Stop)
+	{
+		if (InMessage.ChannelName.IsEmpty())
+		{
+			Broadcast.StopBroadcast();
+			SendMessage(InContext->GetSender(), InMessage.RequestId, ELogVerbosity::Log, TEXT("\"ChannelAction\" Ok."));
+		}
+		else
+		{
+			const FName ChannelName(InMessage.ChannelName);
+			FAvaBroadcastOutputChannel& Channel = Broadcast.GetCurrentProfile().GetChannelMutable(ChannelName);
+			if (Channel.IsValidChannel())
+			{
+				Channel.StopChannelBroadcast();
+				SendMessage(InContext->GetSender(), InMessage.RequestId, ELogVerbosity::Log, TEXT("\"ChannelAction\" Ok."));
+			}
+			else
+			{
+				LogAndSendMessage(InContext->GetSender(), InMessage.RequestId, ELogVerbosity::Error, TEXT("\"ChannelAction\" Failed. Reason: Invalid Channel \"%s\"."), *InMessage.ChannelName);
+			}
+		}
+	}
+	else
+	{
+		LogAndSendMessage(InContext->GetSender(), InMessage.RequestId, ELogVerbosity::Error, TEXT("\"ChannelAction\" Failed. Reason: Invalid Action (must be \"Start\" or \"Stop\"."));
+	}
+}
+
+void FAvaRundownServer::HandleGetChannelImage(const FAvaRundownGetChannelImage& InMessage, const TSharedRef<IMessageContext, ESPMode::ThreadSafe>& InContext)
+{
+	if (InMessage.ChannelName.IsEmpty())
+	{
+		SendMessage(InContext->GetSender(), InMessage.RequestId, ELogVerbosity::Error, TEXT("\"GetChannelImage\" Failed. Reason: Invalid ChannelName."));
+		return;
+	}
+	
+	const FName ChannelName(InMessage.ChannelName);
+	const FAvaBroadcastOutputChannel& Channel = UAvaBroadcast::Get().GetCurrentProfile().GetChannel(ChannelName);
+	
+	if (!Channel.IsValidChannel())
+	{
+		SendMessage(InContext->GetSender(), InMessage.RequestId, ELogVerbosity::Error, TEXT("\"GetChannelImage\" Failed. Reason: Invalid Channel \"%s\"."), *InMessage.ChannelName);
+		return;
+	}
+
+	const FRequestInfo RequestInfo = {InMessage.RequestId, InContext->GetSender()};
+	
+	UTextureRenderTarget2D* ChannelRenderTarget = Channel.GetCurrentRenderTarget(true);
+	
+	// If the channel's render target is not the desired format, we will need to convert it.
+	TSharedPtr<FChannelImage> ChannelImage;
+	if (AvailableChannelImages.Num() > 0)
+	{
+		ChannelImage = AvailableChannelImages.Pop();
+	}
+	else
+	{
+		ChannelImage = MakeShared<FChannelImage>();
+	}
+
+	if (ChannelRenderTarget->GetFormat() != PF_B8G8R8A8)
+	{
+		ChannelImage->UpdateRenderTarget(ChannelRenderTarget->SizeX, ChannelRenderTarget->SizeY, PF_B8G8R8A8, ChannelRenderTarget->ClearColor);
+	}
+	else
+	{
+		ChannelImage->RenderTarget.Reset(); // No need for conversion.
+	}
+	
+	TWeakPtr<FAvaRundownServer> WeakRundownServer(SharedThis(this));
+
+	// The conversion is done by the GPU in the render thread.
+	ENQUEUE_RENDER_COMMAND(FAvaConvertChannelImage)(
+		[ChannelRenderTarget, ChannelImage, RequestInfo, WeakRundownServer](FRHICommandListImmediate& RHICmdList)
+		{
+			FRHITexture* SourceRHI = ChannelRenderTarget->GetResource()->GetTexture2DRHI();
+			FRHITexture* ReadbackRHI = SourceRHI;
+
+			// Convert if needed.
+			if (ChannelImage->RenderTarget.IsValid())
+			{
+				FRHITexture* DestinationRHI = ChannelImage->RenderTarget->GetResource()->GetTexture2DRHI();
+				UE::AvaBroadcastRenderTargetMediaUtils::CopyTexture(RHICmdList, SourceRHI, DestinationRHI);
+				ReadbackRHI = DestinationRHI;
+			}
+
+			// Reading Render Target pixels in the render thread to avoid a flush render commands.
+			const FReadSurfaceDataFlags ReadDataFlags(RCM_UNorm, CubeFace_MAX);
+			const FIntRect SourceRect = FIntRect(0, 0, ChannelRenderTarget->SizeX, ChannelRenderTarget->SizeY);
+
+			ChannelImage->UpdateRawPixels(SourceRect.Width(), SourceRect.Height());
+			
+			RHICmdList.ReadSurfaceData(ReadbackRHI, SourceRect, ChannelImage->RawPixels, ReadDataFlags);
+
+			// When the converted render target is ready, we resume the work in the game thread.
+			AsyncTask(ENamedThreads::GameThread, [WeakRundownServer, RequestInfo, ChannelImage]()
+			{
+				if (const TSharedPtr<FAvaRundownServer> RundownServer = WeakRundownServer.Pin())
+				{
+					RundownServer->FinishGetChannelImage(RequestInfo, ChannelImage);
+				}
+			});
+		});
+}
+
+void FAvaRundownServer::HandleAddChannelDevice(const FAvaRundownAddChannelDevice& InMessage, const TSharedRef<IMessageContext, ESPMode::ThreadSafe>& InContext)
+{
+	UAvaBroadcast& Broadcast = UAvaBroadcast::Get();
+	if (InMessage.ChannelName.IsEmpty() || InMessage.MediaOutputName.IsEmpty())
+	{
+		LogAndSendMessage(InContext->GetSender(), InMessage.RequestId, ELogVerbosity::Error, TEXT("\"AddChannelDevice\" Failed. Reason: One or more Empty Parameters."));
+		return;
+	}
+
+	const FName ChannelName(InMessage.ChannelName);
+	const FAvaBroadcastOutputChannel& OutputChannel = Broadcast.GetCurrentProfile().GetChannel(ChannelName);
+	if (!OutputChannel.IsValidChannel())
+	{
+		LogAndSendMessage(InContext->GetSender(), InMessage.RequestId, ELogVerbosity::Error, TEXT("\"AddChannelDevice\" Failed. Reason: Invalid Channel \"%s\"."), *InMessage.ChannelName);
+		return;    
+	}
+
+	/*
+		We're essentially replicating the UI editor here. The editor:
+		1. Builds an output tree
+		2. Allows drag-and-drop of output/devices to a channel
+		3. AddMediaOutputToChannel() is called
+
+		We don't have immediate drag-and-drop information here, since this is called externally, so we'll rebuild a tree,
+		and recursively search for a match, and then issue the same AddMediaOutputToChannel call the editor UI would've called.
+
+		This won't be called frequently, so it's equivalent to an end-user opening up and adding a device to a channel via
+		the broadcast window (tree rebuild -> drag and drop item)
+	*/
+	const FAvaOutputTreeItemPtr OutputDevices = MakeShared<FAvaBroadcastOutputRootItem>();
+	IAvaBroadcastOutputTreeItem::FRefreshChildrenParams RefreshDevicesParams;
+	RefreshDevicesParams.bShowAllMediaOutputClasses = true; // Listing all classes so the specified device is present.
+	FAvaBroadcastOutputTreeItem::RefreshTree(OutputDevices, RefreshDevicesParams);
+	FAvaOutputTreeItemPtr TreeItem = UE::AvaRundownServer::Private::RecursiveFindOutputTreeItem(OutputDevices, InMessage.MediaOutputName);
+	
+	if (!TreeItem.IsValid())
+	{
+		LogAndSendMessage(InContext->GetSender(), InMessage.RequestId, ELogVerbosity::Error, TEXT("\"AddChannelDevice\" Failed. Reason: Invalid Device \"%s\"."), *InMessage.MediaOutputName);
+		return;
+	}
+
+	const FAvaBroadcastMediaOutputInfo OutputInfo;
+	const UMediaOutput* OutputDevice = TreeItem->AddMediaOutputToChannel(OutputChannel.GetChannelName(), OutputInfo);
+#if WITH_EDITOR
+	Broadcast.SaveBroadcast();
+#endif
+	LogAndSendMessage(InContext->GetSender(), InMessage.RequestId, ELogVerbosity::Log, TEXT("\"AddChannelDevice\" successfully added device \"%s\""), *OutputDevice->GetFName().ToString());
+}
+
+void FAvaRundownServer::HandleEditChannelDevice(const FAvaRundownEditChannelDevice& InMessage,
+	const TSharedRef<IMessageContext, ESPMode::ThreadSafe>& InContext)
+{
+	UAvaBroadcast& Broadcast = UAvaBroadcast::Get();
+	if (InMessage.ChannelName.IsEmpty() || InMessage.MediaOutputName.IsEmpty())
+	{
+		LogAndSendMessage(InContext->GetSender(), InMessage.RequestId, ELogVerbosity::Error, TEXT("\"EditChannelDevice\" Failed. Reason: One or more Empty Parameters."));
+		return;
+	}
+
+	const FName ChannelName(InMessage.ChannelName);
+	const FAvaBroadcastOutputChannel& OutputChannel = Broadcast.GetCurrentProfile().GetChannel(ChannelName);
+	if (!OutputChannel.IsValidChannel())
+	{
+		LogAndSendMessage(InContext->GetSender(), InMessage.RequestId, ELogVerbosity::Error, TEXT("\"EditChannelDevice\" Failed. Reason: Invalid Channel \"%s\"."), *InMessage.ChannelName);
+		return;    
+	}
+
+	UMediaOutput* MediaOutput = UE::AvaRundownServer::Private::FindChannelMediaOutput(OutputChannel, InMessage.MediaOutputName);
+	
+	if (!MediaOutput)
+	{
+		LogAndSendMessage(InContext->GetSender(), InMessage.RequestId, ELogVerbosity::Error, TEXT("\"EditChannelDevice\" Failed. Reason: Invalid Device \"%s\"."), *InMessage.MediaOutputName);
+		return;
+	}
+
+	FAvaRundownServerMediaOutputUtils::EditMediaOutput(MediaOutput, InMessage.Data);
+	
+#if WITH_EDITOR
+	Broadcast.SaveBroadcast();
+#endif
+	
+	SendMessage(InContext->GetSender(), InMessage.RequestId, ELogVerbosity::Log, TEXT("\"EditChannelDevice\". Successfully edited device \"%s\" on \"%s\""), *InMessage.MediaOutputName, *InMessage.ChannelName); 
+}
+
+void FAvaRundownServer::HandleRemoveChannelDevice(const FAvaRundownRemoveChannelDevice& InMessage,
+	const TSharedRef<IMessageContext, ESPMode::ThreadSafe>& InContext)
+{
+	UAvaBroadcast& Broadcast = UAvaBroadcast::Get();
+	if (InMessage.ChannelName.IsEmpty() || InMessage.MediaOutputName.IsEmpty())
+	{
+		LogAndSendMessage(InContext->GetSender(), InMessage.RequestId, ELogVerbosity::Error, TEXT("\"RemoveChannelDevice\" Failed. Reason: One or more Empty Parameters."));
+		return;
+	}
+	const FName ChannelName(InMessage.ChannelName);
+	const FAvaBroadcastOutputChannel& OutputChannel = Broadcast.GetCurrentProfile().GetChannel(ChannelName);
+	if (!OutputChannel.IsValidChannel())
+	{
+		LogAndSendMessage(InContext->GetSender(), InMessage.RequestId, ELogVerbosity::Error, TEXT("\"RemoveChannelDevice\" Failed. Reason: Invalid Channel \"%s\"."), *InMessage.ChannelName);
+		return;    
+	}
+
+	UMediaOutput* MediaOutput = UE::AvaRundownServer::Private::FindChannelMediaOutput(OutputChannel, InMessage.MediaOutputName);
+	
+	if (!MediaOutput)
+	{
+		LogAndSendMessage(InContext->GetSender(), InMessage.RequestId, ELogVerbosity::Error, TEXT("\"RemoveChannelDevice\" Failed. Reason: Invalid Device \"%s\"."), *InMessage.MediaOutputName);
+		return;
+	}
+	
+#if WITH_EDITOR
+	FScopedTransaction Transaction(LOCTEXT("RemoveMediaOutput", "Remove Media Output"));
+	Broadcast.Modify();
+#endif
+
+	const int32 RemovedCount = Broadcast.GetCurrentProfile().RemoveChannelMediaOutputs(ChannelName, TArray{MediaOutput});
+
+	if (RemovedCount == 0)
+	{
+#if WITH_EDITOR
+		Transaction.Cancel();
+#endif
+		LogAndSendMessage(InContext->GetSender(), InMessage.RequestId, ELogVerbosity::Error, TEXT("\"RemoveChannelDevice\" Didn't remove device."));
+		return;
+	}
+
+#if WITH_EDITOR
+	Broadcast.SaveBroadcast();
+#endif
+	
+	LogAndSendMessage(InContext->GetSender(), InMessage.RequestId, ELogVerbosity::Log, TEXT("\"RemoveChannelDevice\" Removed Device \"%s\""), *InMessage.MediaOutputName);
+}
+
+void FAvaRundownServer::HandleGetDevices(const FAvaRundownGetDevices& InMessage,
+	const TSharedRef<IMessageContext, ESPMode::ThreadSafe>& InContext)
+{
+	FAvaRundownDevicesList* ReplyMessage = FMessageEndpoint::MakeMessage<FAvaRundownDevicesList>();
+	ReplyMessage->RequestId = InMessage.RequestId;
+	const FAvaOutputTreeItemPtr OutputDevices = MakeShared<FAvaBroadcastOutputRootItem>();
+	IAvaBroadcastOutputTreeItem::FRefreshChildrenParams RefreshDevicesParams;
+	RefreshDevicesParams.bShowAllMediaOutputClasses = InMessage.bShowAllMediaOutputClasses;
+	FAvaBroadcastOutputTreeItem::RefreshTree(OutputDevices, RefreshDevicesParams);
+	// OutputDevices here aren't literally a physical device, just a construct representing
+	// output. This convention was pulled from the SAvaBroadcastOutputDevices->RefreshDevices() call
+	for (const TSharedPtr<IAvaBroadcastOutputTreeItem>& ClassItem : OutputDevices->GetChildren())
+	{
+		if (const FAvaBroadcastOutputClassItem* AvaOutputClassItem = ClassItem->CastTo<FAvaBroadcastOutputClassItem>())
+		{
+			FAvaRundownOutputClassItem OutputClassItem;
+			OutputClassItem.Name = ClassItem->GetDisplayName().ToString();
+
+			const TArray<FAvaOutputTreeItemPtr>& Children = AvaOutputClassItem->GetChildren();
+			if (Children.Num() > 0)
+			{
+				for (const TSharedPtr<IAvaBroadcastOutputTreeItem>& OutputDeviceItem : Children)
+				{
+					if (OutputDeviceItem->IsA<FAvaBroadcastOutputDeviceItem>())
+					{
+						FAvaRundownOutputDeviceItem DeviceItem;
+						DeviceItem.Name = OutputDeviceItem->GetDisplayName().ToString();
+						// Intentionally leaving DeviceItem.Data blank, as it's not usable data by itself
+						// .Data will be filled out on a GetChannels call, where it becomes usable
+
+						OutputClassItem.Devices.Push(DeviceItem);
+					}
+				}
+			}
+			else
+			{
+				FAvaRundownOutputDeviceItem DeviceItem;
+				DeviceItem.Name = OutputClassItem.Name;
+				OutputClassItem.Devices.Push(DeviceItem);
+			}
+
+			ReplyMessage->DeviceClasses.Push(OutputClassItem);
+		}
+	}
+	SendResponse(ReplyMessage, InContext->GetSender());
+}
+
+void FAvaRundownServer::LogAndSendMessage(const FMessageAddress& InSender, int32 InRequestId, ELogVerbosity::Type InVerbosity, const TCHAR* InFormat, ...) const
+{
+	TCHAR TempString[1024];
+	va_list Args;
+	va_start(Args, InFormat);
+	FCString::GetVarArgs(TempString, UE_ARRAY_COUNT(TempString), InFormat, Args);
+	va_end(Args);
+
+	// The UE_LOG macro adds ELogVerbosity:: to the verbosity, which prevents
+	// us from using it with a variable.
+	switch (InVerbosity)
+	{
+	case ELogVerbosity::Type::Log:
+		UE_LOG(LogAvaRundownServer, Log, TEXT("%s"), TempString);
+		break;
+	case ELogVerbosity::Type::Display:
+		UE_LOG(LogAvaRundownServer, Display, TEXT("%s"), TempString);
+		break;
+	case ELogVerbosity::Type::Warning:
+		UE_LOG(LogAvaRundownServer, Warning, TEXT("%s"), TempString);
+		break;
+	case ELogVerbosity::Type::Error:
+		UE_LOG(LogAvaRundownServer, Error, TEXT("%s"), TempString);
+		break;
+	default:
+		UE_LOG(LogAvaRundownServer, Log, TEXT("%s"), TempString);
+		break;
+	}
+	
+	// Send the error message to the client.
+	SendMessageImpl(InSender, InRequestId, InVerbosity, TempString);
+}
+
+void FAvaRundownServer::SendMessage(const FMessageAddress& InSender, int32 InRequestId, ELogVerbosity::Type InVerbosity, const TCHAR* InFormat, ...) const
+{
+	TCHAR TempString[1024];
+	va_list Args;
+	va_start(Args, InFormat);
+	FCString::GetVarArgs(TempString, UE_ARRAY_COUNT(TempString), InFormat, Args);
+	va_end(Args);
+
+	// Send the error message to the client.
+	SendMessageImpl(InSender, InRequestId, InVerbosity, TempString);
+}
+
+void FAvaRundownServer::SendMessageImpl(const FMessageAddress& InSender, int32 InRequestId, ELogVerbosity::Type InVerbosity, const TCHAR* InMsg) const
+{
+	FAvaRundownServerMsg* ErrorMessage = FMessageEndpoint::MakeMessage<FAvaRundownServerMsg>();
+	ErrorMessage->RequestId = InRequestId;
+	ErrorMessage->Verbosity = ToString(InVerbosity);
+	ErrorMessage->Text = InMsg;
+	SendResponse(ErrorMessage, InSender);
+}
+
+void FAvaRundownServer::RegisterConsoleCommands()
+{
+	if (ConsoleCommands.Num() != 0)
+	{
+		return;
+	}
+		
+	ConsoleCommands.Add(IConsoleManager::Get().RegisterConsoleCommand(
+			TEXT("MotionDesignRundownServer.Status"),
+			TEXT("Display current status of all server info."),
+			FConsoleCommandWithArgsDelegate::CreateRaw(this, &FAvaRundownServer::ShowStatusCommand),
+			ECVF_Default
+			));
+
+}
+
+void FAvaRundownServer::ShowStatusCommand(const TArray<FString>& InArgs)
+{
+	UE_LOG(LogAvaRundownServer, Display, TEXT("Rundown Server: \"%s\""), *HostName);
+	UE_LOG(LogAvaRundownServer, Display, TEXT("- Endpoint Bus Address: \"%s\""), MessageEndpoint.IsValid() ? *MessageEndpoint->GetAddress().ToString() : TEXT("Invalid"));
+	UE_LOG(LogAvaRundownServer, Display, TEXT("- Computer: \"%s\""), *HostName);
+
+	for (const TPair<FMessageAddress, TSharedPtr<FClientInfo>>& Client : Clients)
+	{
+		const FClientInfo& ClientInfo = *Client.Value;
+		UE_LOG(LogAvaRundownServer, Display, TEXT("Connected Client: \"%s\""), *ClientInfo.Address.ToString());
+		UE_LOG(LogAvaRundownServer, Display, TEXT("   - Api Version: %d"), ClientInfo.ApiVersion);
+	}
+	
+	UE_LOG(LogAvaRundownServer, Display, TEXT("Rundown Caches:"));
+	UE_LOG(LogAvaRundownServer, Display, TEXT("- Editing Rundown: \"%s\""), *RundownEditCommandData.CurrentRundownPath.ToString());
+	UE_LOG(LogAvaRundownServer, Display, TEXT("- Editing PageId: \"%d\""), RundownEditCommandData.ManagedPageId);
+	UE_LOG(LogAvaRundownServer, Display, TEXT("- Playing Rundown: \"%s\""), *RundownPlaybackCommandData.CurrentRundownPath.ToString());
+
+	if (RundownPlaybackCommandData.CurrentRundown.IsValid())
+	{
+		TArray<int32> PlayingPages = RundownPlaybackCommandData.CurrentRundown->GetPlayingPageIds();
+		for (const int32 PlayingPageId : PlayingPages)
+		{
+			UE_LOG(LogAvaRundownServer, Display, TEXT("- Playing PageId: \"%d\""), PlayingPageId);
+		}
+		TArray<int32> PreviewingPages = RundownPlaybackCommandData.CurrentRundown->GetPreviewingPageIds();
+		for (const int32 PreviewingPageId : PreviewingPages)
+		{
+			UE_LOG(LogAvaRundownServer, Display, TEXT("- Previewing PageId: \"%d\""), PreviewingPageId);
+		}
+	}
+}
+
+void FAvaRundownServer::FinishGetChannelImage(const FRequestInfo& InRequestInfo, const TSharedPtr<FChannelImage>& InChannelImage)
+{
+	FAvaRundownChannelImage* ReplyMessage = FMessageEndpoint::MakeMessage<FAvaRundownChannelImage>();
+	ReplyMessage->RequestId = InRequestInfo.RequestId;
+	FImage Image;
+	bool bSuccess = false;
+
+	// Note: replacing FImageUtils::GetRenderTargetImage since we already have the raw pixels.
+	{
+		constexpr EPixelFormat Format = PF_B8G8R8A8;
+		const int32 ImageBytes = CalculateImageBytes(InChannelImage->SizeX, InChannelImage->SizeY, 0, Format);
+		Image.RawData.AddUninitialized(ImageBytes);
+		FMemory::Memcpy( Image.RawData.GetData(), InChannelImage->RawPixels.GetData(), InChannelImage->RawPixels.Num() * sizeof(FColor) );
+		Image.SizeX = InChannelImage->SizeX;
+		Image.SizeY = InChannelImage->SizeY;
+		Image.NumSlices = 1;
+		Image.Format = ERawImageFormat::BGRA8;
+		Image.GammaSpace = EGammaSpace::sRGB;
+	}
+
+	// TODO: profile this.
+	// Options: resize the render target on the gpu prior to reading pixels.
+	{
+		FImage ResizedImage;
+		Image.ResizeTo(ResizedImage, Image.GetWidth() * .25f, Image.GetHeight() * .25f, Image.Format, EGammaSpace::Linear);
+
+		TArray64<uint8> CompressedData;
+		if (FImageUtils::CompressImage(CompressedData,TEXT("JPEG"), ResizedImage, 95))
+		{
+			const uint32 SafeMessageSizeLimit = UE::AvaMediaMessageUtils::GetSafeMessageSizeLimit();
+			if (CompressedData.Num() > SafeMessageSizeLimit)
+			{
+				LogAndSendMessage(InRequestInfo.Sender, InRequestInfo.RequestId, ELogVerbosity::Error,
+					TEXT("\"GetChannelImage\" Failed. Reason: (DataSize: %d) is larger that the safe size limit for udp segmenter (%d)."),
+					CompressedData.Num(), SafeMessageSizeLimit);
+				return;
+			}
+
+			ReplyMessage->ImageData.Append(CompressedData.GetData(), CompressedData.GetAllocatedSize());
+			bSuccess = true;
+		}
+	}
+	
+	if (!bSuccess)
+	{
+		LogAndSendMessage(InRequestInfo.Sender, InRequestInfo.RequestId, ELogVerbosity::Error,
+			TEXT("\"GetChannelImage\" Failed. Reason: Unable to retrieve Channel Image."));
+		return;
+	}
+
+	SendResponse(ReplyMessage, InRequestInfo.Sender);
+	
+	// Put the image back in the pool of available images for next request. (or we could abandon it)
+	AvailableChannelImages.Add(InChannelImage);
+}
+
+void FAvaRundownServer::HandlePageActions(const FRequestInfo& InRequestInfo, const TArray<int32>& InPageIds,
+	bool bInIsPreview, FName InPreviewChannelName, EAvaRundownPageActions InAction) const
+{
+	using namespace UE::AvaRundownServer::Private;
+	
+	if (!RundownPlaybackCommandData.CurrentRundown.IsValid())
+	{
+		LogAndSendMessage(InRequestInfo.Sender, InRequestInfo.RequestId, ELogVerbosity::Error,
+			TEXT("\"PageAction\" Failed. Reason: no play list currently loaded."));
+		return;
+	}
+
+	UAvaRundown* Rundown = RundownPlaybackCommandData.CurrentRundown.Get();
+
+	{
+		// Validate the pages - the command will be considered a failure (as a whole) if it contains invalid pages.
+		FString InvalidPages;
+		for (const int32 PageId : InPageIds)
+		{
+			const FAvaRundownPage& Page = Rundown->GetPage(PageId);
+			if (!Page.IsValidPage())
+			{
+				InvalidPages.Appendf(TEXT("%s%d"), InvalidPages.IsEmpty() ? TEXT("") : TEXT(", "), PageId);
+			}
+		}
+
+		if (!InvalidPages.IsEmpty())
+		{
+			LogAndSendMessage(InRequestInfo.Sender, InRequestInfo.RequestId, ELogVerbosity::Error,
+				TEXT("\"PageAction\" Failed. Reason: PageIds {%s} are invalid."), *InvalidPages);
+			return;
+		}
+	}
+
+	const FName PreviewChannelName = !InPreviewChannelName.IsNone() ? InPreviewChannelName : UAvaRundown::GetDefaultPreviewChannelName();
+	// Todo: support program channel name in command.
+	const FName CommandChannelName = bInIsPreview ? InPreviewChannelName : NAME_None; 
+	
+	bool bSuccess = false;
+	FString FailureReason;
+	switch (InAction)
+	{
+		case EAvaRundownPageActions::Load:
+			for (const int32 PageId : InPageIds)
+			{
+				bSuccess |= Rundown->GetPageLoadingManager().RequestLoadPage(PageId, bInIsPreview, PreviewChannelName);
+			}
+			break;
+		case EAvaRundownPageActions::Unload:
+			for (const int32 PageId : InPageIds)
+			{
+				const FAvaRundownPage& Page =  Rundown->GetPage(PageId);
+				if (Page.IsValidPage())
+				{
+					bSuccess |= Rundown->UnloadPage(PageId, (bInIsPreview ? PreviewChannelName : Page.GetChannelName()).ToString());
+				}
+			}
+			break;
+		case EAvaRundownPageActions::Play:
+			bSuccess = !Rundown->PlayPages(InPageIds, bInIsPreview ? EAvaRundownPagePlayType::PreviewFromStart : EAvaRundownPagePlayType::PlayFromStart, PreviewChannelName).IsEmpty();
+			break;
+		case EAvaRundownPageActions::PlayNext:
+			{
+				const int32 NextPageId = FAvaRundownPlaybackUtils::GetPageIdToPlayNext(Rundown, UAvaRundown::InstancePageList, bInIsPreview, PreviewChannelName);
+				if (FAvaRundownPlaybackUtils::IsPageIdValid(NextPageId))
+				{
+					bSuccess = Rundown->PlayPage(NextPageId, bInIsPreview ? EAvaRundownPagePlayType::PreviewFromFrame : EAvaRundownPagePlayType::PlayFromStart);
+				}
+				break;
+			}
+		case EAvaRundownPageActions::Stop:
+			if (InPageIds.IsEmpty())
+			{
+				// If the list of pages is empty, we will stop all the playing pages.
+				const TArray<int32> PageIds = GetPlayingPages(Rundown, bInIsPreview, CommandChannelName);
+				bSuccess = !Rundown->StopPages(PageIds, EAvaRundownPageStopOptions::Default, bInIsPreview).IsEmpty();
+			}
+			else
+			{
+				bSuccess = !Rundown->StopPages(InPageIds, EAvaRundownPageStopOptions::Default, bInIsPreview).IsEmpty();
+			}
+			break;
+		case EAvaRundownPageActions::ForceStop:
+			if (InPageIds.IsEmpty())
+			{
+				// If the list of pages is empty, we will stop all the playing pages.
+				const TArray<int32> PageIds = GetPlayingPages(Rundown, bInIsPreview, CommandChannelName);
+				bSuccess = !Rundown->StopPages(PageIds, EAvaRundownPageStopOptions::ForceNoTransition, bInIsPreview).IsEmpty();
+			}
+			else
+			{
+				bSuccess = !Rundown->StopPages(InPageIds, EAvaRundownPageStopOptions::ForceNoTransition, bInIsPreview).IsEmpty();
+			}
+			break;
+		case EAvaRundownPageActions::Continue:
+			if (InPageIds.IsEmpty())
+			{
+				// If the list of pages is empty, we will continue all the playing pages.
+				const TArray<int32> PageIds = GetPlayingPages(Rundown, bInIsPreview, CommandChannelName);
+				bSuccess = ContinuePages(Rundown, PageIds, bInIsPreview, PreviewChannelName, FailureReason);
+			}
+			else
+			{
+				bSuccess = ContinuePages(Rundown, InPageIds, bInIsPreview, PreviewChannelName, FailureReason);
+			}
+			break;
+		case EAvaRundownPageActions::UpdateValues:
+			if (InPageIds.IsEmpty())
+			{
+				// If the list of pages is empty, we will continue all the playing pages.
+				const TArray<int32> PageIds = GetPlayingPages(Rundown, bInIsPreview, CommandChannelName);
+				bSuccess = UpdatePagesValues(Rundown, PageIds, bInIsPreview, PreviewChannelName);
+			}
+			else
+			{
+				bSuccess = UpdatePagesValues(Rundown, InPageIds, bInIsPreview, PreviewChannelName);
+			}
+			break;
+		case EAvaRundownPageActions::TakeToProgram:
+			{
+				const TArray<int32> PageIds = FAvaRundownPlaybackUtils::GetPagesToTakeToProgram(Rundown, InPageIds, PreviewChannelName);
+				Rundown->PlayPages(PageIds, EAvaRundownPagePlayType::PlayFromStart);
+			}
+			break;
+		default:
+			FailureReason.Appendf(TEXT("Invalid action. "));
+			break;
+	}
+
+	const TCHAR* CommandName = bInIsPreview ? TEXT("PagePreviewAction") : TEXT("PageAction");
+
+	// For multi-page commands, we consider a partial success as success.
+	// Remote applications are notified of the page status with FAvaRundownPagesStatuses.
+	//
+	// Todo:
+	// For pages that failed to execute the command, the failure reason is not sent
+	// to remote applications. Given the more complex status information, we would
+	// probably need a response message for this command with additional error information.
+	
+	if (bSuccess)
+	{
+		SendMessage(InRequestInfo.Sender, InRequestInfo.RequestId, ELogVerbosity::Log,
+			TEXT("\"%s\" Ok."), CommandName);
+	}
+	else if (!FailureReason.IsEmpty())
+	{
+		LogAndSendMessage(InRequestInfo.Sender, InRequestInfo.RequestId, ELogVerbosity::Error,
+			TEXT("\"%s\" Failed. Reason: %s"), CommandName, *FailureReason);
+	}
+	else
+	{
+		LogAndSendMessage(InRequestInfo.Sender, InRequestInfo.RequestId, ELogVerbosity::Error,
+			TEXT("\"%s\" Failed."), CommandName);
+	}
+}
+
+UAvaRundown* FAvaRundownServer::GetOrLoadRundownForEdit(const FMessageAddress& InSender, int32 InRequestId, const FString& InRundownPath)
+{
+	UAvaRundown* Rundown;
+	
+	if (!InRundownPath.IsEmpty())
+	{
+		// If a path is specified, the rundown gets reloaded
+		// unless it was already loaded from a previous editing command.
+		// This will not affect the currently loaded rundown for playback.
+		const FSoftObjectPath NewRundownPath(InRundownPath);
+		
+		Rundown = RundownEditCommandData.GetOrLoadRundown(NewRundownPath,
+			[this](UAvaRundown* InRundown)
+			{
+				RundownEditCommandData.SaveCurrentRemoteControlPresetToPage(true);
+				RundownEditCommandData.RemoveRundownDelegates(this, InRundown);
+			},
+			[this](UAvaRundown* InRundown)
+			{
+				RundownEditCommandData.SetupRundownDelegates(this, InRundown);
+			});
+
+		if (!Rundown)
+		{
+			LogAndSendMessage(InSender, InRequestId, ELogVerbosity::Error, TEXT("Failed to load Rundown \"%s\"."), *InRundownPath);
+		}
+	}
+	else
+	{
+		// If the path is not specified, we assume it is using the previously loaded rundown.
+		Rundown = RundownEditCommandData.CurrentRundown.Get();
+
+		// Note: for backward compatibility with QA python script, we allow this command to use the current "playback" rundown as fallback.
+		if (!Rundown)
+		{
+			Rundown = RundownPlaybackCommandData.CurrentRundown.Get();
+
+			// Update the edit data accordingly.
+			RundownEditCommandData.CurrentRundown.Reset(Rundown);
+			RundownEditCommandData.CurrentRundownPath = RundownPlaybackCommandData.CurrentRundownPath;
+		}
+
+		if (!Rundown)
+		{
+			LogAndSendMessage(InSender, InRequestId, ELogVerbosity::Error, TEXT("No rundown path specified and no rundown currently loaded."));
+		}
+	}
+	return Rundown;
+}
+
+void FAvaRundownServer::OnMessageBusNotification(const FMessageBusNotification& InNotification)
+{
+	// This is called when the websocket client disconnects.
+	if (InNotification.NotificationType == EMessageBusNotification::Unregistered)
+	{
+		TWeakPtr<FAvaRundownServer> ServerWeak = SharedThis(this);
+		auto RemoveClient = [ServerWeak, RegistrationAddress = InNotification.RegistrationAddress]()
+		{
+			if (const TSharedPtr<FAvaRundownServer> Server = ServerWeak.Pin())
+			{
+				if (Server->Clients.Contains(RegistrationAddress))
+				{
+					UE_LOG(LogAvaRundownServer, Log, TEXT("Client \"%s\" disconnected."), *RegistrationAddress.ToString());
+					Server->Clients.Remove(RegistrationAddress);
+					Server->RefreshClientAddresses();
+				}
+			}
+		};
+
+		if (IsInGameThread())
+		{
+			RemoveClient();
+		}
+		else
+		{
+			Async(EAsyncExecution::TaskGraphMainThread, MoveTemp(RemoveClient));
+		}
+	}
+}
+
+void FAvaRundownServer::RefreshClientAddresses()
+{
+	ClientAddresses.Reset(Clients.Num());
+	for (const TPair<FMessageAddress, TSharedPtr<FClientInfo>>& Client : Clients)
+	{
+		ClientAddresses.Add(Client.Key);
+	}
+}
+
+UAvaRundown* FAvaRundownServer::FRundownCache::GetOrLoadRundown(const FSoftObjectPath& InRundownPath,
+	const FRundownEventFunction InUnloadCurrentRundownFunction,
+	const FRundownEventFunction InNewRundownLoadedFunction)
+{
+	if (CurrentRundownPath != InRundownPath)
+	{
+		const TStrongObjectPtr<UAvaRundown> NewRundown = UE::AvaRundownServer::Private::LoadRundown(InRundownPath);
+		if (NewRundown.IsValid())
+		{
+			if (CurrentRundown.Get())
+			{
+				InUnloadCurrentRundownFunction(CurrentRundown.Get());
+			}
+			InNewRundownLoadedFunction(NewRundown.Get());
+			CurrentRundown = NewRundown;
+			CurrentRundownPath = InRundownPath;
+		}
+		else
+		{
+			return nullptr;
+		}
+	}
+	return CurrentRundown.Get();
+}
+
+void FAvaRundownServer::FRundownCache::SetupRundownDelegates(FAvaRundownServer* InRundownServer, UAvaRundown* InRundown)
+{
+	RemoveRundownDelegates(InRundownServer, InRundown);
+
+	if (!InRundown)
+	{
+		return;
+	}
+
+	FAvaPlaybackManager& Manager = IAvaMediaModule::Get().GetLocalPlaybackManager();
+	TWeakObjectPtr<UAvaRundown> RundownWeak(InRundown);
+	TWeakPtr<FAvaRundownServer> RundownServerWeak = InRundownServer->AsShared();
+	OnPlaybackInstanceStatusChangedDelegateHandle = 
+		Manager.OnPlaybackInstanceStatusChanged.AddLambda([RundownWeak, RundownServerWeak](const FAvaPlaybackInstance& InPlaybackInstance)
+		{
+			UAvaRundown* Rundown = RundownWeak.Get();
+			const TSharedPtr<FAvaRundownServer> RundownServer = RundownServerWeak.Pin();
+			if (IsValid(Rundown) && RundownServer.IsValid())
+			{
+				const int32 PageId = UAvaRundownPagePlayer::GetPageIdFromInstanceUserData(InPlaybackInstance.GetInstanceUserData());
+				const FAvaRundownPage Page = Rundown->GetPage(PageId);
+				if (Page.IsValidPage())
+				{
+					RundownServer->PageStatusChanged(Rundown, Page);
+				}
+			}
+		});
+
+	InRundown->GetOnPagesChanged().AddRaw(InRundownServer, &FAvaRundownServer::OnPagesChanged);
+	InRundown->GetOnInstancedPageListChanged().AddRaw(InRundownServer, &FAvaRundownServer::OnPageListChanged);
+	InRundown->GetOnTemplatePageListChanged().AddRaw(InRundownServer, &FAvaRundownServer::OnPageListChanged);
+}
+
+void FAvaRundownServer::FRundownCache::RemoveRundownDelegates(const FAvaRundownServer* InRundownServer, UAvaRundown* InRundown) const
+{
+	FAvaPlaybackManager& Manager = IAvaMediaModule::Get().GetLocalPlaybackManager();
+	Manager.OnPlaybackInstanceStatusChanged.Remove(OnPlaybackInstanceStatusChangedDelegateHandle);
+
+	if (InRundown)
+	{
+		InRundown->GetOnPagesChanged().RemoveAll(InRundownServer);
+		InRundown->GetOnInstancedPageListChanged().RemoveAll(InRundownServer);
+		InRundown->GetOnTemplatePageListChanged().RemoveAll(InRundownServer);
+	}
+}
+
+FAvaRundownServer::FRundownEditCommandData::~FRundownEditCommandData()
+{
+	SaveCurrentRemoteControlPresetToPage(true);
+}
+
+void FAvaRundownServer::FRundownEditCommandData::SaveCurrentRemoteControlPresetToPage(bool bInUnregister)
+{
+	if (!ManagedInstance.IsValid() || !ManagedInstance->GetRemoteControlPreset())
+	{
+		return;
+	}
+
+	// Check if the RCP was registered.
+	IRemoteControlModule& RemoteControlModule = IRemoteControlModule::Get();
+	const FName CurrentPresetName = ManagedInstance->GetRemoteControlPreset()->GetPresetName();
+	const URemoteControlPreset* ResolvedPreset = RemoteControlModule.ResolvePreset(CurrentPresetName);
+	if (ResolvedPreset != ManagedInstance->GetRemoteControlPreset())
+	{
+		return;
+	}
+
+	if (bInUnregister)
+	{
+		// Unregister from RC module.
+		RemoteControlModule.UnregisterEmbeddedPreset(CurrentPresetName);
+	}
+
+	if (!CurrentRundown.IsValid())
+	{
+		return;
+	}
+
+	// Save the modified values to the page.
+	FAvaRundownPage& ManagedPage = CurrentRundown->GetPage(ManagedPageId);
+	if (!ManagedPage.IsValidPage())
+	{
+		return;
+	}
+
+	constexpr bool bIsDefault = false;
+	FAvaPlayableRemoteControlValues NewValues;
+	NewValues.CopyFrom(ManagedInstance->GetRemoteControlPreset(), bIsDefault);
+
+	// UpdateRemoteControlValues does half the job by ensuring that missing values are added and
+	// extra values are removed. But it doesn't change existing values.
+	EAvaPlayableRemoteControlChanges RemoteControlChanges = ManagedPage.UpdateRemoteControlValues(NewValues, bIsDefault);
+
+	// Modify existing values if different.
+	for (const TPair<FGuid, FAvaPlayableRemoteControlValue>& NewValue : NewValues.EntityValues)
+	{
+		const FAvaPlayableRemoteControlValue* ExistingValue = ManagedPage.GetRemoteControlEntityValue(NewValue.Key);
+		
+		if (ExistingValue && !NewValue.Value.IsSameValueAs(*ExistingValue))
+		{
+			ManagedPage.SetRemoteControlEntityValue(NewValue.Key, NewValue.Value);
+			RemoteControlChanges |= EAvaPlayableRemoteControlChanges::EntityValues;
+		}
+	}
+	for (const TPair<FGuid, FAvaPlayableRemoteControlValue>& NewValue : NewValues.ControllerValues)
+	{
+		const FAvaPlayableRemoteControlValue* ExistingValue = ManagedPage.GetRemoteControlControllerValue(NewValue.Key);
+
+		if (ExistingValue && !NewValue.Value.IsSameValueAs(*ExistingValue))
+		{
+			ManagedPage.SetRemoteControlControllerValue(NewValue.Key, NewValue.Value);
+			RemoteControlChanges |= EAvaPlayableRemoteControlChanges::ControllerValues;
+		}
+	}
+
+	if (RemoteControlChanges != EAvaPlayableRemoteControlChanges::None)
+	{
+		CurrentRundown->NotifyPageRemoteControlValueChanged(ManagedPageId, RemoteControlChanges);
+	}
+}
+
+FAvaRundownServer::FRundownPlaybackCommandData::~FRundownPlaybackCommandData()
+{
+	ClosePlaybackContext();
+}
+
+void FAvaRundownServer::FRundownPlaybackCommandData::ClosePlaybackContext()
+{
+	if (CurrentRundown.IsValid())
+	{
+		// Stop all playing pages.
+		CurrentRundown->ClosePlaybackContext(true);
+		CurrentRundown.Reset();
+		CurrentRundownPath.Reset();
+	}
+}
+
+#undef LOCTEXT_NAMESPACE
\ No newline at end of file
diff --git a/Engine/Plugins/Experimental/Avalanche/Source/AvalancheMedia/Private/Rundown/AvaRundownServer.h b/Engine/Plugins/Experimental/Avalanche/Source/AvalancheMedia/Private/Rundown/AvaRundownServer.h
new file mode 100644
--- /dev/null
+++ b/Engine/Plugins/Experimental/Avalanche/Source/AvalancheMedia/Private/Rundown/AvaRundownServer.h
@@ -0,0 +1,230 @@
+// Copyright Epic Games, Inc. All Rights Reserved.
+
+#pragma once
+
+#include "Broadcast/AvaBroadcastProfile.h"
+#include "MessageEndpoint.h"
+#include "Rundown/AvaRundown.h"
+#include "Rundown/AvaRundownMessages.h"
+#include "Rundown/IAvaRundownServer.h"
+#include "UObject/SoftObjectPath.h"
+#include "UObject/StrongObjectPtr.h"
+
+class FAvaRundownManagedInstance;
+class UMediaOutput;
+
+/**
+ * Implements a rundown server that listens to commands on message bus.
+ * The intention is to run a web socket transport bridge so the messages can
+ * come from external applications.
+ */
+class FAvaRundownServer : public TSharedFromThis<FAvaRundownServer>, public IAvaRundownServer
+{
+public:
+	FAvaRundownServer();
+	virtual ~FAvaRundownServer();
+
+	//~ Begin IAvaRundownServer
+	virtual const FString& GetName() const override { return HostName; }
+	virtual const FMessageAddress& GetMessageAddress() const override;
+	virtual TArray<FMessageAddress> GetClientAddresses() const override { return ClientAddresses; }
+	//~ End IAvaRundownServer
+
+	void Init(const FString& InAssignedHostName);
+
+	void SetupBroadcastDelegates(UAvaBroadcast* InBroadcast);
+	void SetupEditorDelegates();
+	void RemoveBroadcastDelegates(UAvaBroadcast* InBroadcast) const;
+	void RemoveEditorDelegates() const;
+	
+	void OnPagesChanged(const UAvaRundown* InRundown, const FAvaRundownPage& InPage, EAvaRundownPageChanges InChange) const;
+	void OnPageListChanged(const FAvaRundownPageListChangeParams& InParams) const;
+	void PageBlueprintChanged(const UAvaRundown* InRundown, const FAvaRundownPage& InPage, const FString& InBlueprintPath) const;
+	void PageStatusChanged(const UAvaRundown* InRundown, const FAvaRundownPage& InPage) const;
+	void PageChannelChanged(const UAvaRundown* InRundown, const FAvaRundownPage& InPage, const FString& InChannelName) const;
+	void PageAnimSettingsChanged(const UAvaRundown* InRundown, const FAvaRundownPage& InPage) const;
+	void OnBroadcastChannelListChanged(const FAvaBroadcastProfile& InProfile) const;
+	void OnBroadcastChannelChanged(const FAvaBroadcastOutputChannel& InChannel, EAvaBroadcastChannelChange InChange) const;
+	void OnAssetAddedOrRemoved(const FAssetData& InAssetData) const;
+	
+	// Message handlers
+	void HandleRundownPing(const FAvaRundownPing& InMessage, const TSharedRef<IMessageContext, ESPMode::ThreadSafe>& InContext);
+	void HandleGetRundowns(const FAvaRundownGetRundowns& InMessage, const TSharedRef<IMessageContext, ESPMode::ThreadSafe>& InContext);
+	void HandleLoadRundown(const FAvaRundownLoadRundown& InMessage, const TSharedRef<IMessageContext, ESPMode::ThreadSafe>& InContext);
+	void HandleSaveRundown(const FAvaRundownSaveRundown& InMessage, const TSharedRef<IMessageContext, ESPMode::ThreadSafe>& InContext);
+	void HandleGetPages(const FAvaRundownGetPages& InMessage, const TSharedRef<IMessageContext, ESPMode::ThreadSafe>& InContext);
+	void HandleCreatePage(const FAvaRundownCreatePage& InMessage,const TSharedRef<IMessageContext, ESPMode::ThreadSafe>& InContext);
+	void HandleCreateTemplate(const FAvaRundownCreateTemplate& InMessage, const TSharedRef<IMessageContext, ESPMode::ThreadSafe>& InContext);
+	void HandleDeletePage(const FAvaRundownDeletePage& InMessage,const TSharedRef<IMessageContext, ESPMode::ThreadSafe>& InContext);
+	void HandleDeleteTemplate(const FAvaRundownDeleteTemplate& InMessage,const TSharedRef<IMessageContext, ESPMode::ThreadSafe>& InContext);
+	void HandleChangeTemplateBP(const FAvaRundownChangeTemplateBP& InMessage, const TSharedRef<IMessageContext, ESPMode::ThreadSafe>& InContext);
+	void HandleGetPageDetails(const FAvaRundownGetPageDetails& InMessage, const TSharedRef<IMessageContext, ESPMode::ThreadSafe>& InContext);
+	void HandleChangePageChannel(const FAvaRundownPageChangeChannel& InMessage, const TSharedRef<IMessageContext, ESPMode::ThreadSafe>& InContext);
+	void HandleUpdatePageFromRCP(const FAvaRundownUpdatePageFromRCP& InMessage, const TSharedRef<IMessageContext, ESPMode::ThreadSafe>& InContext);
+	void HandlePageAction(const FAvaRundownPageAction& InMessage, const TSharedRef<IMessageContext, ESPMode::ThreadSafe>& InContext);
+	void HandlePagePreviewAction(const FAvaRundownPagePreviewAction& InMessage, const TSharedRef<IMessageContext, ESPMode::ThreadSafe>& InContext);
+	void HandlePageActions(const FAvaRundownPageActions& InMessage, const TSharedRef<IMessageContext, ESPMode::ThreadSafe>& InContext);
+	void HandlePagePreviewActions(const FAvaRundownPagePreviewActions& InMessage, const TSharedRef<IMessageContext, ESPMode::ThreadSafe>& InContext);
+
+	void HandleGetChannel(const FAvaRundownGetChannel& InMessage, const TSharedRef<IMessageContext, ESPMode::ThreadSafe>& InContext);
+	void HandleGetChannels(const FAvaRundownGetChannels& InMessage, const TSharedRef<IMessageContext, ESPMode::ThreadSafe>& InContext);
+	void HandleChannelAction(const FAvaRundownChannelAction& InMessage, const TSharedRef<IMessageContext, ESPMode::ThreadSafe>& InContext);
+	void HandleAddChannelDevice(const FAvaRundownAddChannelDevice& InMessage, const TSharedRef<IMessageContext, ESPMode::ThreadSafe>& InContext);
+	void HandleEditChannelDevice(const FAvaRundownEditChannelDevice& InMessage, const TSharedRef<IMessageContext, ESPMode::ThreadSafe>& InContext);
+	void HandleRemoveChannelDevice(const FAvaRundownRemoveChannelDevice& InMessage, const TSharedRef<IMessageContext, ESPMode::ThreadSafe>& InContext);
+	void HandleGetChannelImage(const FAvaRundownGetChannelImage& InMessage, const TSharedRef<IMessageContext, ESPMode::ThreadSafe>& InContext);
+	
+	void HandleGetDevices(const FAvaRundownGetDevices& InMessage, const TSharedRef<IMessageContext, ESPMode::ThreadSafe>& InContext);
+	
+	void LogAndSendMessage(const FMessageAddress& InSender, int32 InRequestId, ELogVerbosity::Type InVerbosity, const TCHAR* InFormat, ...) const;
+	void SendMessage(const FMessageAddress& InSender, int32 InRequestId, ELogVerbosity::Type InVerbosity, const TCHAR* InFormat, ...) const;
+	void SendMessageImpl(const FMessageAddress& InSender, int32 InRequestId, ELogVerbosity::Type InVerbosity, const TCHAR* InMsg) const;
+
+	void RegisterConsoleCommands();
+	
+	void ShowStatusCommand(const TArray<FString>& InArgs);
+
+protected:
+	struct FRequestInfo
+	{
+		int32 RequestId;
+		FMessageAddress Sender;
+	};
+	
+	struct FChannelImage;
+	void FinishGetChannelImage(const FRequestInfo& InRequestInfo, const TSharedPtr<FChannelImage>& InChannelImage);
+
+	void HandlePageActions(const FRequestInfo& InRequestInfo, const TArray<int32>& InPageIds,
+		bool bInIsPreview, FName InPreviewChannelName, EAvaRundownPageActions InAction) const;
+	
+	/**
+	 * Helper function to retrieve the appropriate rundown for editing commands.
+	 */
+	UAvaRundown* GetOrLoadRundownForEdit(const FMessageAddress& InSender, int32 InRequestId, const FString& InRundownPath);
+
+	template<typename MessageType>
+	void SendResponse(MessageType* InMessage, const FMessageAddress& InRecipient, EMessageFlags InFlags = EMessageFlags::None) const
+	{
+		if (MessageEndpoint)
+		{
+			MessageEndpoint->Send(InMessage, InRecipient);
+		}
+	}
+	
+	template<typename MessageType>
+	void SendResponse(MessageType* InMessage, const TArray<FMessageAddress>& InRecipients, EMessageFlags InFlags = EMessageFlags::None) const
+	{
+		if (MessageEndpoint)
+		{
+			MessageEndpoint->Send(InMessage, MessageType::StaticStruct(), InFlags, nullptr,
+				InRecipients, FTimespan::Zero(), FDateTime::MaxValue());
+		}
+	}
+
+	void OnMessageBusNotification(const FMessageBusNotification& InNotification);
+
+private:
+	FString HostName;	
+	TSharedPtr<FMessageEndpoint, ESPMode::ThreadSafe> MessageEndpoint;
+	TArray<IConsoleObject*> ConsoleCommands;
+
+	/** Keep information on connected clients. */
+	struct FClientInfo
+	{
+		FMessageAddress Address;
+		/** Api version for communication with this client. */
+		int32 ApiVersion = -1;
+
+		explicit FClientInfo(const FMessageAddress& InAddress) : Address(InAddress) {}
+	};
+
+	/** Keep track of remote clients context information. */
+	TMap<FMessageAddress, TSharedPtr<FClientInfo>> Clients;
+
+	/** Array of just the client addresses for sending responses. */
+	TArray<FMessageAddress> ClientAddresses;
+	
+	FClientInfo* GetClientInfo(const FMessageAddress& InAddress) const
+	{
+		const TSharedPtr<FClientInfo>* ClientInfoPtr = Clients.Find(InAddress);
+		return ClientInfoPtr ? ClientInfoPtr->Get() : nullptr;
+	}
+	
+	FClientInfo& GetOrAddClientInfo(const FMessageAddress& InAddress)
+	{
+		if (const TSharedPtr<FClientInfo>* ExistingClientInfo = Clients.Find(InAddress))
+		{
+			return *ExistingClientInfo->Get();
+		}
+
+		const TSharedPtr<FClientInfo> NewClientInfo = MakeShared<FClientInfo>(InAddress);
+		Clients.Add(InAddress, NewClientInfo);
+		RefreshClientAddresses();
+		return *NewClientInfo;
+	}
+
+	void RefreshClientAddresses();
+	
+	/** Pool of images that can be recycled. */
+	TArray<TSharedPtr<FChannelImage>> AvailableChannelImages;
+	
+	/**
+	 * Keeps a current rundown cached for commands operating on rundown.
+	 * There is only one "current" rundown at a given time.
+	 */
+	struct FRundownCache
+	{
+		/** Currently loaded/cached rundown's path. */
+		FSoftObjectPath CurrentRundownPath;
+		
+		/** Currently loaded/cached rundown object. */
+		TStrongObjectPtr<UAvaRundown> CurrentRundown;
+
+		FDelegateHandle OnPlaybackInstanceStatusChangedDelegateHandle;
+		
+		using FRundownEventFunction = TFunctionRef<void(UAvaRundown*)>;
+		
+		/**
+		 * Returns requested rundown specified by InRundownPath. Will load it if necessary or returned the cached one if it is the same.
+		 * If the new rundown fails to load, the return value is nullptr, and the previous rundown will remain loaded.
+		 * @param InRundownPath	Requested rundown path. 
+		 * @param InUnloadCurrentRundownFunction	Function called, with the previously cached rundown, when a new rundown is loaded
+		 *											and previous rundown needs to be unloaded.
+		 * @param InNewRundownLoadedFunction	Function called, with the new rundown, when a new rundown is loaded.
+		 */
+		UAvaRundown* GetOrLoadRundown(const FSoftObjectPath& InRundownPath, 
+			const FRundownEventFunction InUnloadCurrentRundownFunction,
+			const FRundownEventFunction InNewRundownLoadedFunction);
+
+		void SetupRundownDelegates(FAvaRundownServer* InRundownServer, UAvaRundown* InRundown);
+		void RemoveRundownDelegates(const FAvaRundownServer* InRundownServer, UAvaRundown* InRundown) const;
+	};
+	
+	/** Cached data for page editing commands (GetPages and GetPageDetails). */
+	struct FRundownEditCommandData : public FRundownCache
+	{
+		/** PageId of the current managed ava asset. */
+		int32 ManagedPageId = FAvaRundownPage::InvalidPageId;
+		TSharedPtr<FAvaRundownManagedInstance> ManagedInstance;
+
+		~FRundownEditCommandData();
+		
+		/**
+		 * Checks if previous RCP was registered.
+		 * If so, save modified values to corresponding page.
+ 		 * This may result in the Rundown to be modified.
+ 		 * Will also unregister RCP from RC Module if requested.
+		 */
+		void SaveCurrentRemoteControlPresetToPage(bool bInUnregister);
+	};
+	FRundownEditCommandData RundownEditCommandData;
+
+	/** Cached data for playback commands (LoadRundown and PageAction). */
+	struct FRundownPlaybackCommandData: public FRundownCache
+	{
+		~FRundownPlaybackCommandData();
+
+		void ClosePlaybackContext();
+	};
+	FRundownPlaybackCommandData RundownPlaybackCommandData;
+};
diff --git a/Engine/Plugins/Experimental/Avalanche/Source/AvalancheMedia/Private/Rundown/AvaRundownServerMediaOutputUtils.cpp b/Engine/Plugins/Experimental/Avalanche/Source/AvalancheMedia/Private/Rundown/AvaRundownServerMediaOutputUtils.cpp
new file mode 100644
--- /dev/null
+++ b/Engine/Plugins/Experimental/Avalanche/Source/AvalancheMedia/Private/Rundown/AvaRundownServerMediaOutputUtils.cpp
@@ -0,0 +1,286 @@
+// Copyright Epic Games, Inc. All Rights Reserved.
+
+#include "Rundown/AvaRundownServerMediaOutputUtils.h"
+
+#include "Dom/JsonObject.h"
+#include "MediaOutput.h"
+#include "Serialization/JsonReader.h"
+#include "Serialization/JsonSerializer.h"
+#include "Serialization/JsonWriter.h"
+
+DEFINE_LOG_CATEGORY_STATIC(LogAvaRundownServerMediaOutputUtils, Warning, All);
+
+namespace UE::AvaMedia::RundownServerMediaOutputUtils::Private
+{
+	// Serialization helper functions
+	static TSharedPtr<FJsonObject> ParsePropertyInfo(TFieldIterator<FProperty> InProperty, const void* InOwnerObject);
+	static TSharedPtr<FJsonObject> ParseElementaryPropertyInfo(TFieldIterator<FProperty> InProperty, const void* InOwnerObject);
+	static TSharedPtr<FJsonObject> ParseEnumPropertyInfo(TFieldIterator<FProperty> InProperty, const void* InOwnerObject);
+	static TSharedPtr<FJsonObject> ParseStructPropertyInfo(TFieldIterator<FProperty> InProperty, const void* InOwnerObject);
+
+	// Helper functions for editing Media Output Devices
+	static void SetProperty(void* InOwnerObject, const TSharedPtr<FJsonObject>& InPropertyObject, const TFieldIterator<FProperty>& InProperty);
+	static void SetElementaryProperty(void* InOwnerObject, const TSharedPtr<FJsonObject>& InPropertyObject, TFieldIterator<FProperty> InProperty);
+	static void SetEnumProperty(void* InOwnerObject, const TSharedPtr<FJsonObject>& InPropertyObject, TFieldIterator<FProperty> InProperty);
+	static void SetStructProperties(void* InOwnerObject, const TSharedPtr<FJsonObject>& InPropertyObject, TFieldIterator<FProperty> InProperty);
+
+	FString GetPropertyDisplayNameString(TFieldIterator<FProperty> InProperty)
+	{
+#if WITH_EDITORONLY_DATA
+		return InProperty->GetDisplayNameText().ToString();
+#else
+		return TEXT("NotAvailable"); // Not available in game build.
+	#endif
+	}
+}
+
+FString FAvaRundownServerMediaOutputUtils::SerializeMediaOutput(const UMediaOutput* InMediaOutput)
+{
+	using namespace UE::AvaMedia::RundownServerMediaOutputUtils::Private;
+	
+	const TSharedPtr<FJsonObject> OutputObject = MakeShareable(new FJsonObject);
+
+	OutputObject->SetStringField("Class", InMediaOutput->GetClass()->GetPathName());
+	OutputObject->SetStringField("Name", InMediaOutput->GetName());
+
+	TArray<TSharedPtr<FJsonValue>> PropertyArray;
+	for (TFieldIterator<FProperty> Iterator(InMediaOutput->GetClass()); Iterator; ++Iterator)
+	{
+		PropertyArray.Add(MakeShareable(new FJsonValueObject(ParsePropertyInfo(Iterator, InMediaOutput))));
+	}
+
+	const TSharedPtr<FJsonObject> PropertiesObject = MakeShareable(new FJsonObject);
+	OutputObject->SetArrayField("Properties", PropertyArray);
+
+	FString OutputString;
+	const TSharedRef<TJsonWriter<>> JsonWriter = TJsonWriterFactory<>::Create(&OutputString);
+	FJsonSerializer::Serialize(OutputObject.ToSharedRef(), JsonWriter);
+	
+	return OutputString;
+}
+
+void FAvaRundownServerMediaOutputUtils::EditMediaOutput(UMediaOutput* InMediaOutput, const FString& InDeviceData)
+{
+	using namespace UE::AvaMedia::RundownServerMediaOutputUtils::Private;
+	
+	const TSharedRef<TJsonReader<TCHAR>> JsonReader = TJsonReaderFactory<TCHAR>::Create(InDeviceData);
+	if (TSharedPtr<FJsonObject> NewDeviceData; FJsonSerializer::Deserialize(JsonReader, NewDeviceData))
+	{
+		for (TArray<TSharedPtr<FJsonValue>> PropertyValues = NewDeviceData->GetArrayField(TEXT("Properties"));
+		     TSharedPtr<FJsonValue> PropertyValue : PropertyValues)
+		{
+			const TSharedPtr<FJsonObject> PropertyObject = PropertyValue.Get()->AsObject();
+
+			for (TFieldIterator<FProperty> PropertyIterator(InMediaOutput->GetClass());
+			     PropertyIterator; ++PropertyIterator)
+			{
+				if (PropertyObject->GetStringField(TEXT("Name")) == PropertyIterator->GetName())
+				{
+					SetProperty(InMediaOutput, PropertyObject, PropertyIterator);
+				}
+			}
+		}
+	}
+}
+
+TSharedPtr<FJsonObject> UE::AvaMedia::RundownServerMediaOutputUtils::Private::ParsePropertyInfo(TFieldIterator<FProperty> InProperty,
+                                                                       const void* InOwnerObject)
+{
+	const FFieldClass* FieldClass = InProperty->GetClass();
+	if (FieldClass == FIntProperty::StaticClass() || FieldClass == FNameProperty::StaticClass() ||
+		FieldClass == FBoolProperty::StaticClass() || FieldClass == FStrProperty::StaticClass())
+	{
+		return ParseElementaryPropertyInfo(InProperty, InOwnerObject);
+	}
+	if (FieldClass == FEnumProperty::StaticClass())
+	{
+		return ParseEnumPropertyInfo(InProperty, InOwnerObject);
+	}
+	if (FieldClass == FStructProperty::StaticClass())
+	{
+		return ParseStructPropertyInfo(InProperty, InOwnerObject);
+	}
+	return MakeShareable(new FJsonObject);
+}
+
+TSharedPtr<FJsonObject> UE::AvaMedia::RundownServerMediaOutputUtils::Private::ParseElementaryPropertyInfo(
+	TFieldIterator<FProperty> InProperty, const void* InOwnerObject)
+{
+	const TSharedPtr<FJsonObject> PropertyObject = MakeShareable(new FJsonObject);
+	if (InProperty->GetClass() == FIntProperty::StaticClass())
+	{
+		const FIntProperty* IntProperty = CastField<FIntProperty>(InProperty->GetOwnerProperty());
+		const int32 Value = IntProperty->GetPropertyValue_InContainer(InOwnerObject);
+		PropertyObject->SetStringField("Name", InProperty->GetName());
+		PropertyObject->SetStringField("DisplayName", GetPropertyDisplayNameString(InProperty));
+		PropertyObject->SetStringField("Type", "number");
+		PropertyObject->SetNumberField("Value", Value);
+	}
+	else if (InProperty->GetClass() == FBoolProperty::StaticClass())
+	{
+		const FBoolProperty* BoolProperty = CastField<FBoolProperty>(InProperty->GetOwnerProperty());
+		const bool Value = BoolProperty->GetPropertyValue_InContainer(InOwnerObject);
+		PropertyObject->SetStringField("Name", InProperty->GetName());
+		PropertyObject->SetStringField("DisplayName", GetPropertyDisplayNameString(InProperty));
+		PropertyObject->SetStringField("Type", "boolean");
+		PropertyObject->SetBoolField("Value", Value);
+	}
+	else if (InProperty->GetClass() == FStrProperty::StaticClass())
+	{
+		const FStrProperty* TextProperty = CastField<FStrProperty>(InProperty->GetOwnerProperty());
+		const FString Value = TextProperty->GetPropertyValue_InContainer(InOwnerObject);
+		PropertyObject->SetStringField("Name", InProperty->GetName());
+		PropertyObject->SetStringField("DisplayName", GetPropertyDisplayNameString(InProperty));
+		PropertyObject->SetStringField("Type", "string");
+		PropertyObject->SetStringField("Value", Value);
+	}
+	else if (InProperty->GetClass() == FNameProperty::StaticClass())
+	{
+		const FNameProperty* NameProperty = CastField<FNameProperty>(InProperty->GetOwnerProperty());
+		const FName Value = NameProperty->GetPropertyValue_InContainer(InOwnerObject);
+		PropertyObject->SetStringField("Name", InProperty->GetName());
+		PropertyObject->SetStringField("DisplayName", GetPropertyDisplayNameString(InProperty));
+		PropertyObject->SetStringField("Type", "name");
+		PropertyObject->SetStringField("Value", Value.ToString());
+	}
+	else
+	{
+		UE_LOG(LogAvaRundownServerMediaOutputUtils, Warning, TEXT("Property type: %s is not supported"), *InProperty->GetClass()->GetName());
+	}
+
+	return PropertyObject;
+}
+
+TSharedPtr<FJsonObject> UE::AvaMedia::RundownServerMediaOutputUtils::Private::ParseEnumPropertyInfo(
+	TFieldIterator<FProperty> InProperty, const void* InOwnerObject)
+{
+	using namespace UE::AvaMedia::RundownServerMediaOutputUtils::Private;
+	
+	const TSharedPtr<FJsonObject> PropertyObject = MakeShareable(new FJsonObject);
+	const FEnumProperty* EnumProperty = CastField<FEnumProperty>(InProperty->GetOwnerProperty());
+
+	PropertyObject->SetStringField("Name", InProperty->GetName());
+	PropertyObject->SetStringField("DisplayName", GetPropertyDisplayNameString(InProperty));
+	PropertyObject->SetStringField("Type", "enum");
+	PropertyObject->SetStringField("Class", EnumProperty->GetEnum()->GetName());
+
+	const int32 MaxIndex = EnumProperty->GetEnum()->NumEnums();
+	TArray<TSharedPtr<FJsonValue>> EnumArray;
+	for (int32 Index = 0; Index < MaxIndex; ++Index)
+	{
+		const TSharedPtr<FJsonObject> EnumValueObject = MakeShareable(new FJsonObject);
+		EnumValueObject->SetStringField(FString::FromInt(Index),
+		                                *EnumProperty->GetEnum()->GetAuthoredNameStringByIndex(Index));
+		EnumArray.Add(MakeShareable(new FJsonValueObject(EnumValueObject)));
+		PropertyObject->SetArrayField("EnumValues", EnumArray);
+	}
+	const uint8* Value = EnumProperty->ContainerPtrToValuePtr<uint8>(InOwnerObject);
+
+	PropertyObject->SetStringField("Value", *EnumProperty->GetEnum()->GetAuthoredNameStringByValue(*Value));
+	return PropertyObject;
+}
+
+TSharedPtr<FJsonObject> UE::AvaMedia::RundownServerMediaOutputUtils::Private::ParseStructPropertyInfo(
+	TFieldIterator<FProperty> InProperty, const void* InOwnerObject)
+{
+	using namespace UE::AvaMedia::RundownServerMediaOutputUtils::Private;
+	
+	const TSharedPtr<FJsonObject> PropertyObject = MakeShareable(new FJsonObject);
+	const FStructProperty* StructProperty = CastField<FStructProperty>(InProperty->GetOwnerProperty());
+
+	PropertyObject->SetStringField("Name", InProperty->GetName());
+	PropertyObject->SetStringField("DisplayName", GetPropertyDisplayNameString(InProperty));
+	PropertyObject->SetStringField("Type", "struct");
+
+	const UScriptStruct* Struct = StructProperty->Struct;
+	PropertyObject->SetStringField("Class", Struct->GetName());
+
+	const void* StructPtr = StructProperty->ContainerPtrToValuePtr<void>(InOwnerObject);
+
+	TArray<TSharedPtr<FJsonValue>> PropertyArray;
+	for (TFieldIterator<FProperty> StructIterator(Struct); StructIterator; ++StructIterator)
+	{
+		PropertyArray.Add(MakeShareable(new FJsonValueObject(ParsePropertyInfo(StructIterator, StructPtr))));
+	}
+	PropertyObject->SetArrayField("StructProperties", PropertyArray);
+
+	return PropertyObject;
+}
+
+void UE::AvaMedia::RundownServerMediaOutputUtils::Private::SetProperty(void* InOwnerObject, const TSharedPtr<FJsonObject>& InPropertyObject,
+                                              const TFieldIterator<FProperty>& InProperty)
+{
+	if (InPropertyObject->GetStringField(TEXT("Type")) == "number" || InPropertyObject->GetStringField(TEXT("Type")) == "string"
+		|| InPropertyObject->GetStringField(TEXT("Type")) == "boolean" || InPropertyObject->GetStringField(TEXT("Type")) == "name")
+	{
+		SetElementaryProperty(InOwnerObject, InPropertyObject, InProperty);
+	}
+	else if (InPropertyObject->GetStringField(TEXT("Type")) == "enum")
+	{
+		SetEnumProperty(InOwnerObject, InPropertyObject, InProperty);
+	}
+	else if (InPropertyObject->GetStringField(TEXT("Type")) == "struct")
+	{
+		SetStructProperties(InOwnerObject, InPropertyObject, InProperty);
+	}
+}
+
+void UE::AvaMedia::RundownServerMediaOutputUtils::Private::SetElementaryProperty(void* InOwnerObject,
+	const TSharedPtr<FJsonObject>& InPropertyObject, TFieldIterator<FProperty> InProperty)
+{
+	if (InPropertyObject->GetStringField(TEXT("Type")) == "number")
+	{
+		const int32 NumberValue = InPropertyObject->GetNumberField(TEXT("Value"));
+		InProperty->GetOwnerProperty()->SetValue_InContainer(InOwnerObject, &NumberValue);
+	}
+	else if (InPropertyObject->GetStringField(TEXT("Type")) == "string")
+	{
+		const FString StringValue = InPropertyObject->GetStringField(TEXT("Value"));
+		InProperty->GetOwnerProperty()->SetValue_InContainer(InOwnerObject, &StringValue);
+	}
+	else if (InPropertyObject->GetStringField(TEXT("Type")) == "boolean")
+	{
+		const bool BoolValue = InPropertyObject->GetBoolField(TEXT("Value"));
+		InProperty->GetOwnerProperty()->SetValue_InContainer(InOwnerObject, &BoolValue);
+	}
+	else if (InPropertyObject->GetStringField(TEXT("Type")) == "name")
+	{
+		const FName NameValue = FName(InPropertyObject->GetStringField(TEXT("Value")));
+		InProperty->GetOwnerProperty()->SetValue_InContainer(InOwnerObject, &NameValue);
+	}
+}
+
+void UE::AvaMedia::RundownServerMediaOutputUtils::Private::SetEnumProperty(void* InOwnerObject,
+	const TSharedPtr<FJsonObject>& InPropertyObject, TFieldIterator<FProperty> InProperty)
+{
+	const FString ValueName = InPropertyObject->GetStringField(TEXT("Value"));
+
+	const FEnumProperty* EnumProperty = CastField<FEnumProperty>(
+		InProperty->GetOwnerProperty());
+	const int32 EnumValue = EnumProperty->GetEnum()->GetValueByName(FName(ValueName));
+	EnumProperty->GetOwnerProperty()->SetValue_InContainer(InOwnerObject, &EnumValue);
+}
+
+void UE::AvaMedia::RundownServerMediaOutputUtils::Private::SetStructProperties(void* InOwnerObject,
+	const TSharedPtr<FJsonObject>& InPropertyObject, TFieldIterator<FProperty> InProperty)
+{
+	const FStructProperty* StructProperty = CastField<FStructProperty>(
+		InProperty->GetOwnerProperty());
+
+	const UScriptStruct* Struct = StructProperty->Struct;
+	void* StructPtr = StructProperty->ContainerPtrToValuePtr<void>(InOwnerObject);
+
+	for (TArray<TSharedPtr<FJsonValue>> StructPropertyValues = InPropertyObject->GetArrayField(TEXT("StructProperties"));
+			 TSharedPtr<FJsonValue> StructPropertyValue : StructPropertyValues)
+	{
+		const TSharedPtr<FJsonObject> StructPropertyObject = StructPropertyValue.Get()->AsObject();
+
+		for (TFieldIterator<FProperty> StructIterator(Struct); StructIterator; ++StructIterator)
+		{
+			if (StructPropertyObject->GetStringField(TEXT("Name")) == StructIterator->GetName())
+			{
+				SetProperty(StructPtr, StructPropertyObject, StructIterator);
+			}
+		}
+	}
+}
diff --git a/Engine/Plugins/Experimental/Avalanche/Source/AvalancheMedia/Public/IAvaMediaModule.h b/Engine/Plugins/Experimental/Avalanche/Source/AvalancheMedia/Public/IAvaMediaModule.h
--- a/Engine/Plugins/Experimental/Avalanche/Source/AvalancheMedia/Public/IAvaMediaModule.h
+++ b/Engine/Plugins/Experimental/Avalanche/Source/AvalancheMedia/Public/IAvaMediaModule.h
@@ -18,6 +18,7 @@ class IAvaBroadcastSettings;
 class IAvaMediaSyncProvider;
 class IAvaPlaybackClient;
 class IAvaPlaybackServer;
+class IAvaRundownServer;
 class IMediaIOCoreDeviceProvider;
 class UWorld;
 struct FAvaInstanceSettings;
@@ -173,6 +174,39 @@ public:
 	DECLARE_DELEGATE_OneParam(FGetEditorViewportClient, FCommonViewportClient** );
 	virtual FGetEditorViewportClient& GetEditorViewportClientDelegate() = 0;
 
+	DECLARE_MULTICAST_DELEGATE_OneParam(FOnRundownServerEvent, TSharedPtr<IAvaRundownServer>);
+	virtual FOnRundownServerEvent& GetOnRundownServerStarted() = 0;
+	virtual FOnRundownServerEvent& GetOnRundownServerStopping() = 0;
+
+	/**
+	 * @brief Returns true if the rundown server is started.
+	 */
+	virtual bool IsRundownServerStarted() const = 0;
+
+	/**
+	 * @brief Starts the rundown server (if not already started).
+	 * @param InServerName Optional server name. If empty, the host (computer) name will be used.
+	 */
+	virtual void StartRundownServer(const FString& InServerName) = 0;
+
+	/**
+	 * @brief Stops the rundown server.
+	 */
+	virtual void StopRundownServer() = 0;
+	
+	/**
+	 * @brief Returns currently running rundown server. 
+	 */
+	virtual TSharedPtr<IAvaRundownServer> GetRundownServer() const = 0;
+
+	/**
+	 * Creates a rundown server that is not managed by the module.
+	 * @param InServerName Optional server name. if empty, the host name will be used.d
+	 * @return Created server.
+	 * @remark For internal use only (testing). Detached servers will interfere with the managed one. 
+	 */
+	virtual TSharedPtr<IAvaRundownServer> MakeDetachedRundownServer(const FString& InServerName) = 0;
+	
 	/**
 	 * Access the device provider proxy manager.
 	 */
diff --git a/Engine/Plugins/Experimental/Avalanche/Source/AvalancheMedia/Public/Rundown/IAvaRundownServer.h b/Engine/Plugins/Experimental/Avalanche/Source/AvalancheMedia/Public/Rundown/IAvaRundownServer.h
new file mode 100644
--- /dev/null
+++ b/Engine/Plugins/Experimental/Avalanche/Source/AvalancheMedia/Public/Rundown/IAvaRundownServer.h
@@ -0,0 +1,20 @@
+// Copyright Epic Games, Inc. All Rights Reserved.
+
+#pragma once
+
+#include "IMessageContext.h"
+
+class IAvaRundownServer
+{
+public:
+	virtual ~IAvaRundownServer() = default;
+	
+	/** Returns the server's name. */
+	virtual const FString& GetName() const = 0;
+
+	/** Returns the endpoint's message address. */
+	virtual const FMessageAddress& GetMessageAddress() const = 0;
+
+	/** Returns the list of connected client addresses. */
+	virtual TArray<FMessageAddress> GetClientAddresses() const = 0;
+};
\ No newline at end of file
diff --git a/Engine/Plugins/Experimental/Avalanche/Source/AvalancheMediaEditor/AvalancheMediaEditor.Build.cs b/Engine/Plugins/Experimental/Avalanche/Source/AvalancheMediaEditor/AvalancheMediaEditor.Build.cs
--- a/Engine/Plugins/Experimental/Avalanche/Source/AvalancheMediaEditor/AvalancheMediaEditor.Build.cs
+++ b/Engine/Plugins/Experimental/Avalanche/Source/AvalancheMediaEditor/AvalancheMediaEditor.Build.cs
@@ -34,7 +34,6 @@ public class AvalancheMediaEditor : ModuleRules
 				"ContentBrowser",
 				"EditorWidgets",
 				"GraphEditor",
-				"ImageCore",
 				"InputCore",
 				"Json",
 				"JsonSerialization",
diff --git a/Engine/Plugins/Experimental/Avalanche/Source/AvalancheMediaEditor/Private/AvaMediaEditorModule.cpp b/Engine/Plugins/Experimental/Avalanche/Source/AvalancheMediaEditor/Private/AvaMediaEditorModule.cpp
--- a/Engine/Plugins/Experimental/Avalanche/Source/AvalancheMediaEditor/Private/AvaMediaEditorModule.cpp
+++ b/Engine/Plugins/Experimental/Avalanche/Source/AvalancheMediaEditor/Private/AvaMediaEditorModule.cpp
@@ -48,13 +48,6 @@ namespace UE::AvaMediaEditorModule::Private
 		static const FName MenuName(TEXT("LevelEditor.StatusBar.ToolBar"));
 		static const FName SectionName(TEXT("MotionDesign"));
 	}
-	
-	// Command line parsing helper.
-	bool IsRundownServerManuallyStarted(FString& OutRundownServerName)
-	{
-		return FParse::Value(FCommandLine::Get(), TEXT("MotionDesignRundownServerStart="), OutRundownServerName) ||
-			FParse::Param(FCommandLine::Get(), TEXT("MotionDesignRundownServerStart"));
-	}
 
 	static void GetEditorViewportClient(FCommonViewportClient** OutViewportClient)
 	{
@@ -95,50 +88,21 @@ void FAvaMediaEditorModule::StartupModule()
 	// register ours after, i.e. once all modules are loaded.
 	FCoreDelegates::OnAllModuleLoadingPhasesComplete.AddRaw(this, &FAvaMediaEditorModule::RegisterCustomizations);
 	FCoreDelegates::OnPostEngineInit.AddRaw(this, &FAvaMediaEditorModule::PostEngineInit);
-	FCoreDelegates::OnEnginePreExit.AddRaw(this, &FAvaMediaEditorModule::EnginePreExit);
 
 	// Register Map Change Events
 	FLevelEditorModule& LevelEditor = FModuleManager::LoadModuleChecked<FLevelEditorModule>("LevelEditor");
 	LevelEditor.OnMapChanged().AddRaw(this, &FAvaMediaEditorModule::HandleMapChanged);
 
-	ConsoleCmds.Add(IConsoleManager::Get().RegisterConsoleCommand(
-		TEXT("MotionDesignRundownServer.Start"),
-		TEXT("Starts the rundown server."),
-		FConsoleCommandWithArgsDelegate::CreateRaw(this, &FAvaMediaEditorModule::StartRundownServerCommand),
-		ECVF_Default
-	));
-	ConsoleCmds.Add(IConsoleManager::Get().RegisterConsoleCommand(
-		TEXT("MotionDesignRundownServer.Stop"),
-		TEXT("Stops the rundown server."),
-		FConsoleCommandWithArgsDelegate::CreateRaw(this, &FAvaMediaEditorModule::StopRundownServerCommand),
-		ECVF_Default
-	));
-
 	IAvaMediaModule::Get().GetEditorViewportClientDelegate().BindStatic(&GetEditorViewportClient);
 
-	FString DummyServerName;
-	if (IsRundownServerManuallyStarted(DummyServerName))
-	{
-		// Prevent throttling when the server is started.
-		// This has to be done before any SLevelViewport are ticked since the cvar value is cached on first tick.
-		static const FSlateThrottleManager& ThrottleManager = FSlateThrottleManager::Get();
-		if (IConsoleVariable* AllowThrottling = IConsoleManager::Get().FindConsoleVariable(TEXT("Slate.bAllowThrottling")))
-		{
-			AllowThrottling->Set(0);
-			UE_LOG(LogAvaMediaEditor, Log, TEXT("Setting Slate.bAllowThrottling to false."));
-		}
-	}
 	RegisterRundownFilterExpressionFactories();
 	RegisterRundownFilterSuggestionFactories();
 }
 
 void FAvaMediaEditorModule::ShutdownModule()
 {
-	StopAllServices();
-
 	FCoreDelegates::OnAllModuleLoadingPhasesComplete.RemoveAll(this);
 	FCoreDelegates::OnPostEngineInit.RemoveAll(this);
-	FCoreDelegates::OnEnginePreExit.RemoveAll(this);
 
 	// Unregister Map Change Events
 	if (FLevelEditorModule* LevelEditor = FModuleManager::GetModulePtr<FLevelEditorModule>("LevelEditor"))
@@ -164,12 +128,6 @@ void FAvaMediaEditorModule::ShutdownModule()
 		FEdGraphUtilities::UnregisterVisualPinConnectionFactory(PlaybackConnectionFactory);
 		PlaybackConnectionFactory.Reset();
 	}
-
-	for (IConsoleObject* ConsoleCmd : ConsoleCmds)
-	{
-		IConsoleManager::Get().UnregisterConsoleObject(ConsoleCmd);
-	}
-	ConsoleCmds.Empty();
 }
 
 TSharedPtr<FExtensibilityManager> FAvaMediaEditorModule::GetBroadcastToolBarExtensibilityManager()
@@ -321,66 +279,22 @@ void FAvaMediaEditorModule::UnregisterCustomizations() const
 	PropertyModule.UnregisterCustomPropertyTypeLayout(FAvaRundownMacroKeyBinding::StaticStruct()->GetFName());
 }
 
-void FAvaMediaEditorModule::StartRundownServerCommand(const TArray<FString>& Args)
-{
-	if (RundownServer)
-	{
-		UE_LOG(LogAvaMediaEditor, Log, TEXT("Rundown Server is already started."));
-		return;
-	}
-	
-	RundownServer = MakeShared<FAvaRundownServer>();
-	
-	// Remark: Only the module's rundown server register console commands to avoid
-	// conflicts with temporary servers (for testing).
-	RundownServer->RegisterConsoleCommands();
-
-	RundownServer->Init(Args.Num() > 0 ? Args[0] : TEXT(""));
-	OnRundownServerStarted.Broadcast();
-
-	UE_LOG(LogAvaMediaEditor, Log, TEXT("Rundown Server Started."));
-}
-
-void FAvaMediaEditorModule::StopRundownServerCommand(const TArray<FString>& Args)
-{
-	if (RundownServer)
-	{
-		UE_LOG(LogAvaMediaEditor, Log, TEXT("Stopping Rundown Server..."));
-		OnRundownServerStopped.Broadcast();
-	}
-	RundownServer.Reset();
-}
-
 void FAvaMediaEditorModule::PostEngineInit()
 {
 	using namespace UE::AvaMediaEditorModule::Private;
 
-	// This needs to happen late in the loading process, otherwise it fails.
 	const UAvaRundownEditorSettings* Settings = UAvaRundownEditorSettings::Get();
 
-	// Allow for specification of the server name in the command line.
-	// Command line has priority over project settings.
-	FString ServerName;
-	if (IsRundownServerManuallyStarted(ServerName))
-	{
-		StartRundownServerCommand({ServerName});
-	}
-	else if (Settings && Settings->bAutoStartRundownServer)
+	if (Settings && Settings->bAutoStartRundownServer)
 	{
-		StartRundownServerCommand({Settings->RundownServerName});
+		IAvaMediaModule& AvaModule = IAvaMediaModule::Get();
+		if (!AvaModule.IsRundownServerStarted())
+		{
+			AvaModule.StartRundownServer(Settings->RundownServerName);
+		}
 	}
 }
 
-void FAvaMediaEditorModule::EnginePreExit()
-{
-	StopAllServices();
-}
-
-void FAvaMediaEditorModule::StopAllServices()
-{
-	StopRundownServerCommand({});
-}
-
 void FAvaMediaEditorModule::HandleMapChanged(UWorld* InWorld, EMapChangeType InMapChangeType)
 {
 	EAvaMediaMapChangeType EventType;
diff --git a/Engine/Plugins/Experimental/Avalanche/Source/AvalancheMediaEditor/Private/AvaMediaEditorModule.h b/Engine/Plugins/Experimental/Avalanche/Source/AvalancheMediaEditor/Private/AvaMediaEditorModule.h
--- a/Engine/Plugins/Experimental/Avalanche/Source/AvalancheMediaEditor/Private/AvaMediaEditorModule.h
+++ b/Engine/Plugins/Experimental/Avalanche/Source/AvalancheMediaEditor/Private/AvaMediaEditorModule.h
@@ -3,7 +3,6 @@
 #pragma once
 
 #include "IAvaMediaEditorModule.h"
-#include "Rundown/AvaRundownServer.h"
 #include "Templates/UnrealTypeTraits.h"
 
 class IAvaRundownFilterExpressionFactory;
@@ -26,9 +25,6 @@ public:
 	virtual TSharedPtr<FExtensibilityManager> GetPlaybackToolBarExtensibilityManager() override;
 	virtual TSharedPtr<FExtensibilityManager> GetRundownToolBarExtensibilityManager() override;
 	virtual TSharedPtr<FExtensibilityManager> GetRundownMenuExtensibilityManager() override;
-	virtual FOnRundownServerStarted& GetOnRundownServerStarted() override { return OnRundownServerStarted; }
-	virtual FOnRundownServerStopped& GetOnRundownServerStopped() override { return OnRundownServerStopped; }
-	virtual TSharedPtr<FAvaRundownServer> GetRundownServer() const override { return RundownServer; }
 	virtual bool CanFilterSupportComparisonOperation(const FName& InFilterKey, ETextFilterComparisonOperation InOperation, EAvaRundownSearchListType InRundownSearchListType) const override;
 	virtual bool FilterExpression(const FName& InFilterKey, const FAvaRundownPage& InItem, const FAvaRundownTextFilterArgs& InArgs) const override;
 	virtual TArray<TSharedPtr<IAvaRundownFilterSuggestionFactory>> GetSimpleSuggestions(EAvaRundownSearchListType InSuggestionType) const override;
@@ -51,13 +47,8 @@ protected:
 	/** Unregister details view customizations. */
 	void UnregisterCustomizations() const;
 
-	void StartRundownServerCommand(const TArray<FString>& Args);
-	void StopRundownServerCommand(const TArray<FString>& Args);
-
 private:
 	void PostEngineInit();
-	void EnginePreExit();
-	void StopAllServices();
 	void HandleMapChanged(UWorld* InWorld, EMapChangeType InMapChangeType);
 	
 	template <
@@ -86,13 +77,6 @@ private:
 
 	TSharedPtr<FGraphPanelPinConnectionFactory> PlaybackConnectionFactory;
 
-	TSharedPtr<FAvaRundownServer> RundownServer;
-
-	FOnRundownServerStarted OnRundownServerStarted;
-	FOnRundownServerStopped OnRundownServerStopped;
-
-	TArray<IConsoleObject*> ConsoleCmds;
-
 	/** Holds all the RundownFilterExpressionFactory */
 	TMap<FName, TSharedPtr<IAvaRundownFilterExpressionFactory>> FilterExpressionFactories;
 
diff --git a/Engine/Plugins/Experimental/Avalanche/Source/AvalancheMediaEditor/Private/Broadcast/OutputDevices/AvaBroadcastOutputClassItem.cpp b/Engine/Plugins/Experimental/Avalanche/Source/AvalancheMediaEditor/Private/Broadcast/OutputDevices/AvaBroadcastOutputClassItem.cpp
deleted file mode 100644
diff --git a/Engine/Plugins/Experimental/Avalanche/Source/AvalancheMediaEditor/Private/Broadcast/OutputDevices/AvaBroadcastOutputClassItem.h b/Engine/Plugins/Experimental/Avalanche/Source/AvalancheMediaEditor/Private/Broadcast/OutputDevices/AvaBroadcastOutputClassItem.h
deleted file mode 100644
diff --git a/Engine/Plugins/Experimental/Avalanche/Source/AvalancheMediaEditor/Private/Broadcast/OutputDevices/AvaBroadcastOutputDeviceItem.cpp b/Engine/Plugins/Experimental/Avalanche/Source/AvalancheMediaEditor/Private/Broadcast/OutputDevices/AvaBroadcastOutputDeviceItem.cpp
deleted file mode 100644
diff --git a/Engine/Plugins/Experimental/Avalanche/Source/AvalancheMediaEditor/Private/Broadcast/OutputDevices/AvaBroadcastOutputDeviceItem.h b/Engine/Plugins/Experimental/Avalanche/Source/AvalancheMediaEditor/Private/Broadcast/OutputDevices/AvaBroadcastOutputDeviceItem.h
deleted file mode 100644
diff --git a/Engine/Plugins/Experimental/Avalanche/Source/AvalancheMediaEditor/Private/Broadcast/OutputDevices/AvaBroadcastOutputRootItem.cpp b/Engine/Plugins/Experimental/Avalanche/Source/AvalancheMediaEditor/Private/Broadcast/OutputDevices/AvaBroadcastOutputRootItem.cpp
deleted file mode 100644
diff --git a/Engine/Plugins/Experimental/Avalanche/Source/AvalancheMediaEditor/Private/Broadcast/OutputDevices/AvaBroadcastOutputRootItem.h b/Engine/Plugins/Experimental/Avalanche/Source/AvalancheMediaEditor/Private/Broadcast/OutputDevices/AvaBroadcastOutputRootItem.h
deleted file mode 100644
diff --git a/Engine/Plugins/Experimental/Avalanche/Source/AvalancheMediaEditor/Private/Broadcast/OutputDevices/AvaBroadcastOutputServerItem.cpp b/Engine/Plugins/Experimental/Avalanche/Source/AvalancheMediaEditor/Private/Broadcast/OutputDevices/AvaBroadcastOutputServerItem.cpp
deleted file mode 100644
diff --git a/Engine/Plugins/Experimental/Avalanche/Source/AvalancheMediaEditor/Private/Broadcast/OutputDevices/AvaBroadcastOutputServerItem.h b/Engine/Plugins/Experimental/Avalanche/Source/AvalancheMediaEditor/Private/Broadcast/OutputDevices/AvaBroadcastOutputServerItem.h
deleted file mode 100644
diff --git a/Engine/Plugins/Experimental/Avalanche/Source/AvalancheMediaEditor/Private/Broadcast/OutputDevices/AvaBroadcastOutputTreeItem.cpp b/Engine/Plugins/Experimental/Avalanche/Source/AvalancheMediaEditor/Private/Broadcast/OutputDevices/AvaBroadcastOutputTreeItem.cpp
deleted file mode 100644
diff --git a/Engine/Plugins/Experimental/Avalanche/Source/AvalancheMediaEditor/Private/Broadcast/OutputDevices/AvaBroadcastOutputTreeItem.h b/Engine/Plugins/Experimental/Avalanche/Source/AvalancheMediaEditor/Private/Broadcast/OutputDevices/AvaBroadcastOutputTreeItem.h
deleted file mode 100644
diff --git a/Engine/Plugins/Experimental/Avalanche/Source/AvalancheMediaEditor/Private/Broadcast/OutputDevices/DragDropOps/AvaBroadcastOutputTreeItemDragDropOp.cpp b/Engine/Plugins/Experimental/Avalanche/Source/AvalancheMediaEditor/Private/Broadcast/OutputDevices/DragDropOps/AvaBroadcastOutputTreeItemDragDropOp.cpp
--- a/Engine/Plugins/Experimental/Avalanche/Source/AvalancheMediaEditor/Private/Broadcast/OutputDevices/DragDropOps/AvaBroadcastOutputTreeItemDragDropOp.cpp
+++ b/Engine/Plugins/Experimental/Avalanche/Source/AvalancheMediaEditor/Private/Broadcast/OutputDevices/DragDropOps/AvaBroadcastOutputTreeItemDragDropOp.cpp
@@ -3,13 +3,13 @@
 #include "AvaBroadcastOutputTreeItemDragDropOp.h"
 
 #include "Broadcast/AvaBroadcast.h"
-#include "Broadcast/OutputDevices/AvaBroadcastOutputClassItem.h"
+#include "Broadcast/OutputDevices/AvaBroadcastOutputTreeItem.h"
 #include "MediaOutput.h"
 #include "ScopedTransaction.h"
 
 #define LOCTEXT_NAMESPACE "AvaBroadcastOutputTreeItemDragDropOp"
 
-TSharedRef<FAvaBroadcastOutputTreeItemDragDropOp> FAvaBroadcastOutputTreeItemDragDropOp::New(const TSharedPtr<FAvaBroadcastOutputTreeItem>& InOutputTreeItem)
+TSharedRef<FAvaBroadcastOutputTreeItemDragDropOp> FAvaBroadcastOutputTreeItemDragDropOp::New(const TSharedPtr<IAvaBroadcastOutputTreeItem>& InOutputTreeItem)
 {
 	TSharedRef<FAvaBroadcastOutputTreeItemDragDropOp> DragDropOp = MakeShared<FAvaBroadcastOutputTreeItemDragDropOp>();
 	DragDropOp->Init(InOutputTreeItem);
@@ -23,7 +23,7 @@ bool FAvaBroadcastOutputTreeItemDragDropOp::IsValidToDropInChannel(FName InTarge
 
 FReply FAvaBroadcastOutputTreeItemDragDropOp::OnChannelDrop(FName InTargetChannelName)
 {
-	if (const TSharedPtr<FAvaBroadcastOutputTreeItem> Item = GetOutputTreeItem())
+	if (const TSharedPtr<IAvaBroadcastOutputTreeItem> Item = GetOutputTreeItem())
 	{
 		UAvaBroadcast& Broadcast = UAvaBroadcast::Get();
 		
@@ -43,7 +43,7 @@ FReply FAvaBroadcastOutputTreeItemDragDropOp::OnChannelDrop(FName InTargetChanne
 	return FReply::Unhandled();
 }
 
-void FAvaBroadcastOutputTreeItemDragDropOp::Init(const TSharedPtr<FAvaBroadcastOutputTreeItem>& InOutputTreeItem)
+void FAvaBroadcastOutputTreeItemDragDropOp::Init(const TSharedPtr<IAvaBroadcastOutputTreeItem>& InOutputTreeItem)
 {
 	OutputTreeItem = InOutputTreeItem;
 	CurrentHoverText = InOutputTreeItem->GetDisplayName();
diff --git a/Engine/Plugins/Experimental/Avalanche/Source/AvalancheMediaEditor/Private/Broadcast/OutputDevices/DragDropOps/AvaBroadcastOutputTreeItemDragDropOp.h b/Engine/Plugins/Experimental/Avalanche/Source/AvalancheMediaEditor/Private/Broadcast/OutputDevices/DragDropOps/AvaBroadcastOutputTreeItemDragDropOp.h
--- a/Engine/Plugins/Experimental/Avalanche/Source/AvalancheMediaEditor/Private/Broadcast/OutputDevices/DragDropOps/AvaBroadcastOutputTreeItemDragDropOp.h
+++ b/Engine/Plugins/Experimental/Avalanche/Source/AvalancheMediaEditor/Private/Broadcast/OutputDevices/DragDropOps/AvaBroadcastOutputTreeItemDragDropOp.h
@@ -4,24 +4,24 @@
 
 #include "DragAndDrop/DecoratedDragDropOp.h"
 
-class FAvaBroadcastOutputTreeItem;
+class IAvaBroadcastOutputTreeItem;
 
 class FAvaBroadcastOutputTreeItemDragDropOp : public FDecoratedDragDropOp
 {
 public:
-	DRAG_DROP_OPERATOR_TYPE(FAvaOutputClassDragDropOp, FDecoratedDragDropOp)
+	DRAG_DROP_OPERATOR_TYPE(FAvaBroadcastOutputTreeItemDragDropOp, FDecoratedDragDropOp)
 
-	static TSharedRef<FAvaBroadcastOutputTreeItemDragDropOp> New(const TSharedPtr<FAvaBroadcastOutputTreeItem>& InOutputClassItem);
+	static TSharedRef<FAvaBroadcastOutputTreeItemDragDropOp> New(const TSharedPtr<IAvaBroadcastOutputTreeItem>& InOutputClassItem);
 
 	bool IsValidToDropInChannel(FName InTargetChannelName) const;
 
-	TSharedPtr<FAvaBroadcastOutputTreeItem> GetOutputTreeItem() const { return OutputTreeItem; }
+	TSharedPtr<IAvaBroadcastOutputTreeItem> GetOutputTreeItem() const { return OutputTreeItem; }
 
 	FReply OnChannelDrop(FName InTargetChannelName);
 
 protected:
-	void Init(const TSharedPtr<FAvaBroadcastOutputTreeItem>& InOutputClassItem);
+	void Init(const TSharedPtr<IAvaBroadcastOutputTreeItem>& InOutputClassItem);
 
 	/** Keep Reference Count while Drag Dropping */
-	TSharedPtr<FAvaBroadcastOutputTreeItem> OutputTreeItem;
+	TSharedPtr<IAvaBroadcastOutputTreeItem> OutputTreeItem;
 };
diff --git a/Engine/Plugins/Experimental/Avalanche/Source/AvalancheMediaEditor/Private/Broadcast/OutputDevices/Slate/SAvaBroadcastOutputDevices.cpp b/Engine/Plugins/Experimental/Avalanche/Source/AvalancheMediaEditor/Private/Broadcast/OutputDevices/Slate/SAvaBroadcastOutputDevices.cpp
--- a/Engine/Plugins/Experimental/Avalanche/Source/AvalancheMediaEditor/Private/Broadcast/OutputDevices/Slate/SAvaBroadcastOutputDevices.cpp
+++ b/Engine/Plugins/Experimental/Avalanche/Source/AvalancheMediaEditor/Private/Broadcast/OutputDevices/Slate/SAvaBroadcastOutputDevices.cpp
@@ -5,6 +5,7 @@
 #include "AvaMediaEditorSettings.h"
 #include "Broadcast/AvaBroadcast.h"
 #include "Broadcast/OutputDevices/AvaBroadcastOutputRootItem.h"
+#include "Broadcast/OutputDevices/DragDropOps/AvaBroadcastOutputTreeItemDragDropOp.h"
 #include "Widgets/Views/STreeView.h"
 
 TMap<uint32, bool> SAvaBroadcastOutputDevices::ItemExpansionStates = {};
@@ -45,6 +46,12 @@ TSharedRef<ITableRow> SAvaBroadcastOutputDevices::OnGenerateItemRow(FAvaOutputTr
 	, const TSharedRef<STableViewBase>& InOwnerTable)
 {
 	check(InItem.IsValid());
+
+	InItem->OnCreateDragDropOperation().BindLambda([](const FAvaOutputTreeItemPtr& InItem)
+	{
+		return FAvaBroadcastOutputTreeItemDragDropOp::New(InItem);
+	});
+
 	return SNew(STableRow<FAvaOutputTreeItemPtr>, InOwnerTable)
 		.ShowWires(false)
 		.Padding(FMargin(5.f, 5.f))
@@ -103,7 +110,9 @@ void SAvaBroadcastOutputDevices::RefreshOutputDevices()
 	check(RootItem.IsValid());
 	
 	//Refresh Items
-	FAvaBroadcastOutputTreeItem::RefreshTree(RootItem);
+	IAvaBroadcastOutputTreeItem::FRefreshChildrenParams RefreshDevicesParams;
+	RefreshDevicesParams.bShowAllMediaOutputClasses = UAvaMediaEditorSettings::Get().bBroadcastShowAllMediaOutputClasses;
+	FAvaBroadcastOutputTreeItem::RefreshTree(RootItem, RefreshDevicesParams);
 	TopLevelItems = RootItem->GetChildren();
 	OutputTree->RequestTreeRefresh();
 
diff --git a/Engine/Plugins/Experimental/Avalanche/Source/AvalancheMediaEditor/Private/Broadcast/OutputDevices/Slate/SAvaBroadcastOutputTreeItem.cpp b/Engine/Plugins/Experimental/Avalanche/Source/AvalancheMediaEditor/Private/Broadcast/OutputDevices/Slate/SAvaBroadcastOutputTreeItem.cpp
deleted file mode 100644
diff --git a/Engine/Plugins/Experimental/Avalanche/Source/AvalancheMediaEditor/Private/Broadcast/OutputDevices/Slate/SAvaBroadcastOutputTreeItem.h b/Engine/Plugins/Experimental/Avalanche/Source/AvalancheMediaEditor/Private/Broadcast/OutputDevices/Slate/SAvaBroadcastOutputTreeItem.h
deleted file mode 100644
diff --git a/Engine/Plugins/Experimental/Avalanche/Source/AvalancheMediaEditor/Private/Rundown/AvaRundownEditorSettings.h b/Engine/Plugins/Experimental/Avalanche/Source/AvalancheMediaEditor/Private/Rundown/AvaRundownEditorSettings.h
--- a/Engine/Plugins/Experimental/Avalanche/Source/AvalancheMediaEditor/Private/Rundown/AvaRundownEditorSettings.h
+++ b/Engine/Plugins/Experimental/Avalanche/Source/AvalancheMediaEditor/Private/Rundown/AvaRundownEditorSettings.h
@@ -52,7 +52,11 @@ public:
 	UPROPERTY(Config, EditAnywhere, Category="Page Actions")
 	EAvaRundownPageSet UpdateValuesActionPageSet = EAvaRundownPageSet::SelectedOrPlaying;
 
-	/** Whether rundown server is started automatically. */
+	/**
+	 * Whether rundown server is started automatically when the editor is launched.
+	 * For game mode or packaged games, the rundown server can be launched with
+	 * the command line -MotionDesignRundownServerStart[=ServerName].
+	 */
 	UPROPERTY(Config, EditAnywhere, Category="Server")
 	bool bAutoStartRundownServer = false;
 
diff --git a/Engine/Plugins/Experimental/Avalanche/Source/AvalancheMediaEditor/Private/Rundown/AvaRundownMessages.h b/Engine/Plugins/Experimental/Avalanche/Source/AvalancheMediaEditor/Private/Rundown/AvaRundownMessages.h
deleted file mode 100644
diff --git a/Engine/Plugins/Experimental/Avalanche/Source/AvalancheMediaEditor/Private/Rundown/AvaRundownPlaybackUtils.cpp b/Engine/Plugins/Experimental/Avalanche/Source/AvalancheMediaEditor/Private/Rundown/AvaRundownPlaybackUtils.cpp
deleted file mode 100644
diff --git a/Engine/Plugins/Experimental/Avalanche/Source/AvalancheMediaEditor/Private/Rundown/AvaRundownPlaybackUtils.h b/Engine/Plugins/Experimental/Avalanche/Source/AvalancheMediaEditor/Private/Rundown/AvaRundownPlaybackUtils.h
deleted file mode 100644
diff --git a/Engine/Plugins/Experimental/Avalanche/Source/AvalancheMediaEditor/Private/Rundown/AvaRundownServer.cpp b/Engine/Plugins/Experimental/Avalanche/Source/AvalancheMediaEditor/Private/Rundown/AvaRundownServer.cpp
deleted file mode 100644
diff --git a/Engine/Plugins/Experimental/Avalanche/Source/AvalancheMediaEditor/Private/Rundown/AvaRundownServer.h b/Engine/Plugins/Experimental/Avalanche/Source/AvalancheMediaEditor/Private/Rundown/AvaRundownServer.h
deleted file mode 100644
diff --git a/Engine/Plugins/Experimental/Avalanche/Source/AvalancheMediaEditor/Private/Rundown/MediaOutputEditorUtils/AvaRundownOutputEditorUtils.cpp b/Engine/Plugins/Experimental/Avalanche/Source/AvalancheMediaEditor/Private/Rundown/MediaOutputEditorUtils/AvaRundownOutputEditorUtils.cpp
deleted file mode 100644
diff --git a/Engine/Plugins/Experimental/Avalanche/Source/AvalancheMediaEditor/Private/Rundown/MediaOutputEditorUtils/AvaRundownOutputEditorUtils.h b/Engine/Plugins/Experimental/Avalanche/Source/AvalancheMediaEditor/Private/Rundown/MediaOutputEditorUtils/AvaRundownOutputEditorUtils.h
deleted file mode 100644
diff --git a/Engine/Plugins/Experimental/Avalanche/Source/AvalancheMediaEditor/Private/Tests/AvaRundownServerTests.cpp b/Engine/Plugins/Experimental/Avalanche/Source/AvalancheMediaEditor/Private/Tests/AvaRundownServerTests.cpp
--- a/Engine/Plugins/Experimental/Avalanche/Source/AvalancheMediaEditor/Private/Tests/AvaRundownServerTests.cpp
+++ b/Engine/Plugins/Experimental/Avalanche/Source/AvalancheMediaEditor/Private/Tests/AvaRundownServerTests.cpp
@@ -3,13 +3,15 @@
 #include "Broadcast/AvaBroadcast.h"
 #include "Broadcast/OutputDevices/AvaBroadcastOutputRootItem.h"
 #include "Broadcast/OutputDevices/AvaBroadcastOutputTreeItem.h"
-#include "IAvaMediaEditorModule.h"
+#include "IAvaMediaModule.h"
 #include "MediaOutput.h"
 #include "MessageEndpointBuilder.h"
 #include "Misc/AutomationTest.h"
 #include "Misc/FileHelper.h"
-#include "Rundown/AvaRundownServer.h"
-#include "Rundown/MediaOutputEditorUtils/AvaRundownOutputEditorUtils.h"
+#include "Rundown/AvaRundownMessages.h"
+#include "Rundown/AvaRundownServerMediaOutputUtils.h"
+#include "Rundown/IAvaRundownServer.h"
+#include "UObject/Package.h"
 #include "UObject/UObjectIterator.h"
 
 DEFINE_LOG_CATEGORY_STATIC(LogAvaRundownServerTests, Log, All);
@@ -87,16 +89,15 @@ namespace UE::AvaRundownServerTests
 		}
 	};
 
-	TSharedPtr<FAvaRundownServer> GetOrCreateRundownServer()
+	TSharedPtr<IAvaRundownServer> GetOrCreateRundownServer()
 	{
-		TSharedPtr<FAvaRundownServer> RundownServer = IAvaMediaEditorModule::Get().GetRundownServer();
+		TSharedPtr<IAvaRundownServer> RundownServer = IAvaMediaModule::Get().GetRundownServer();
 		if (!RundownServer)
 		{
 			UE_LOG(LogAvaRundownServerTests, Log, TEXT("Rundown Server not started. Starting one temporarily for the test."));
 
 			// Start a temporary server. It will be deleted when the last latent command is finished.
-			RundownServer = MakeShared<FAvaRundownServer>();
-			RundownServer->Init(TEXT(""));
+			RundownServer = IAvaMediaModule::Get().MakeDetachedRundownServer(TEXT(""));
 		}
 		return RundownServer;
 	}
@@ -132,7 +133,7 @@ namespace UE::AvaRundownServerTests
 	}
 }
 
-DEFINE_LATENT_AUTOMATION_COMMAND_THREE_PARAMETER(FRundownServerWaitForResponse, int32, RequestId, TSharedPtr<UE::AvaRundownServerTests::FTestClient>, TestClient, TSharedPtr<FAvaRundownServer>, RundownServer);
+DEFINE_LATENT_AUTOMATION_COMMAND_THREE_PARAMETER(FRundownServerWaitForResponse, int32, RequestId, TSharedPtr<UE::AvaRundownServerTests::FTestClient>, TestClient, TSharedPtr<IAvaRundownServer>, RundownServer);
 bool FRundownServerWaitForResponse::Update()
 {
 	if (!TestClient.IsValid())
@@ -157,7 +158,7 @@ bool FRundownServerGetChannelImage::RunTest(const FString& Parameters)
 {
 	using namespace UE::AvaRundownServerTests;
 	
-	const TSharedPtr<FAvaRundownServer> RundownServer = GetOrCreateRundownServer();
+	const TSharedPtr<IAvaRundownServer> RundownServer = GetOrCreateRundownServer();
 	
 	// Build a test client for this request.
 	const TSharedPtr<FTestClient> TestClient = MakeShared<FTestClient>();
@@ -183,7 +184,7 @@ bool FRundownServerAddChannelDevice::RunTest(const FString& Parameters)
 	// Need to backup the broadcast config because we are about to make changes.
 	BackupBroadcastConfig();
 	
-	const TSharedPtr<FAvaRundownServer> RundownServer = GetOrCreateRundownServer();
+	const TSharedPtr<IAvaRundownServer> RundownServer = GetOrCreateRundownServer();
 
 	// Build a test client for this request.
 	const TSharedPtr<FTestClient> TestClient = MakeShared<FTestClient>();
@@ -198,7 +199,7 @@ bool FRundownServerAddChannelDevice::RunTest(const FString& Parameters)
 
 	// Find an existing output so the request succeeds.
 	const FAvaOutputTreeItemPtr OutputDevices = MakeShared<FAvaBroadcastOutputRootItem>();
-	FAvaBroadcastOutputTreeItem::RefreshTree(OutputDevices);
+	FAvaBroadcastOutputTreeItem::RefreshTree(OutputDevices, IAvaBroadcastOutputTreeItem::FRefreshChildrenParams());
 	
 	if (OutputDevices->GetChildren().Num() > 0)
 	{
@@ -225,17 +226,16 @@ bool FRundownServerAddChannelDevice::RunTest(const FString& Parameters)
 IMPLEMENT_SIMPLE_AUTOMATION_TEST(FRundownServerMediaOutputSerialization, "MotionDesign.RundownServer.MediaOutputSerialization", (EAutomationTestFlags::EditorContext | EAutomationTestFlags::ProductFilter));
 bool FRundownServerMediaOutputSerialization::RunTest(const FString& Parameters)
 {
-
-	for (UClass* const Class : TObjectRange<UClass>())
+	for (const UClass* const Class : TObjectRange<UClass>())
 	{
 		// For now, we test all the media output classes, technically, this should work with all of them.
 		// The result of this test depend on the plugins that are enabled. 
-		// The blackmagic plugin should be enabled at least.
+		// The Blackmagic plugin should be enabled at least.
 		const bool bIsMediaOutputClass = Class->IsChildOf(UMediaOutput::StaticClass()) && Class != UMediaOutput::StaticClass();
 		if (bIsMediaOutputClass)
 		{
-			UMediaOutput* const MediaOutput = NewObject<UMediaOutput>(GetTransientPackage(), Class, NAME_None, RF_Transactional);
-			FString MediaOutputJson = FAvaRundownOutputEditorUtils::SerializeMediaOutput(MediaOutput);
+			const UMediaOutput* const MediaOutput = NewObject<UMediaOutput>(GetTransientPackage(), Class, NAME_None, RF_Transactional);
+			FString MediaOutputJson = FAvaRundownServerMediaOutputUtils::SerializeMediaOutput(MediaOutput);
 
 			const FString SaveName = FString::Printf(TEXT("%sSerializationTest_%s.json"), *FPaths::ProjectSavedDir(), *Class->GetName());
 			UE_LOG(LogAvaRundownServerTests, Display, TEXT("Serializing Media Output \"%s\" to \"%s\""), *Class->GetName(), *SaveName);
diff --git a/Engine/Plugins/Experimental/Avalanche/Source/AvalancheMediaEditor/Public/IAvaMediaEditorModule.h b/Engine/Plugins/Experimental/Avalanche/Source/AvalancheMediaEditor/Public/IAvaMediaEditorModule.h
--- a/Engine/Plugins/Experimental/Avalanche/Source/AvalancheMediaEditor/Public/IAvaMediaEditorModule.h
+++ b/Engine/Plugins/Experimental/Avalanche/Source/AvalancheMediaEditor/Public/IAvaMediaEditorModule.h
@@ -8,7 +8,6 @@
 #include "Toolkits/AssetEditorToolkit.h"
 #include "Templates/SharedPointer.h"
 
-class FAvaRundownServer;
 class IAvaRundownFilterSuggestionFactory;
 class FExtensibilityManager;
 class FName;
@@ -58,14 +57,6 @@ public:
 	 */
 	virtual TSharedPtr<FExtensibilityManager> GetRundownMenuExtensibilityManager() = 0;
 
-	DECLARE_MULTICAST_DELEGATE(FOnRundownServerStarted);
-	virtual FOnRundownServerStarted& GetOnRundownServerStarted() = 0;
-
-	DECLARE_MULTICAST_DELEGATE(FOnRundownServerStopped);
-	virtual FOnRundownServerStopped& GetOnRundownServerStopped() = 0;
-
-	virtual TSharedPtr<FAvaRundownServer> GetRundownServer() const = 0;
-
 	/**
  	* Check if current rundown filter expression factory support the comparison operation
  	* @param InFilterKey Filter key to get the rundown filter expression factory needed
