Motion Design: Implemented RC Signature Customization for Motion Design Outliner drag drops

diff --git a/Engine/Plugins/Experimental/Avalanche/Source/AvalancheEditor/Private/RemoteControl/AvaRCExtension.cpp b/Engine/Plugins/Experimental/Avalanche/Source/AvalancheEditor/Private/RemoteControl/AvaRCExtension.cpp
--- a/Engine/Plugins/Experimental/Avalanche/Source/AvalancheEditor/Private/RemoteControl/AvaRCExtension.cpp
+++ b/Engine/Plugins/Experimental/Avalanche/Source/AvalancheEditor/Private/RemoteControl/AvaRCExtension.cpp
@@ -1,9 +1,11 @@
 // Copyright Epic Games, Inc. All Rights Reserved.
 
 #include "AvaRCExtension.h"
+#include "AvaRCSignatureCustomization.h"
 #include "Editor.h"
 #include "Framework/Docking/LayoutExtender.h"
 #include "IAvaSceneInterface.h"
+#include "IRemoteControlUIModule.h"
 #include "LevelEditor.h"
 #include "RemoteControlPreset.h"
 #include "RemoteControlTrackerComponent.h"
@@ -21,12 +23,14 @@ URemoteControlPreset* FAvaRCExtension::GetRemoteControlPreset() const
 
 void FAvaRCExtension::Activate()
 {
+	RegisterSignatureCustomization();
 	OpenRemoteControlTab();
 }
 
 void FAvaRCExtension::Deactivate()
 {
 	CloseRemoteControlTab();
+	UnregisterSignatureCustomization();
 }
 
 void FAvaRCExtension::ExtendLevelEditorLayout(FLayoutExtender& InExtender) const
@@ -84,4 +88,27 @@ void FAvaRCExtension::CloseRemoteControlTab() const
 	}
 }
 
+void FAvaRCExtension::RegisterSignatureCustomization()
+{
+	UnregisterSignatureCustomization();
+
+	SignatureCustomization = MakeShared<FAvaRCSignatureCustomization>();
+	IRemoteControlUIModule::Get().RegisterSignatureCustomization(SignatureCustomization);
+}
+
+void FAvaRCExtension::UnregisterSignatureCustomization()
+{
+	if (!SignatureCustomization.IsValid())
+	{
+		return;
+	}
+
+	if (IRemoteControlUIModule* RCUIModule = FModuleManager::GetModulePtr<IRemoteControlUIModule>(TEXT("RemoteControlUI")))
+	{
+		RCUIModule->UnregisterSignatureCustomization(SignatureCustomization);
+	}
+
+	SignatureCustomization.Reset();
+}
+
 #undef LOCTEXT_NAMESPACE
diff --git a/Engine/Plugins/Experimental/Avalanche/Source/AvalancheEditor/Private/RemoteControl/AvaRCExtension.h b/Engine/Plugins/Experimental/Avalanche/Source/AvalancheEditor/Private/RemoteControl/AvaRCExtension.h
--- a/Engine/Plugins/Experimental/Avalanche/Source/AvalancheEditor/Private/RemoteControl/AvaRCExtension.h
+++ b/Engine/Plugins/Experimental/Avalanche/Source/AvalancheEditor/Private/RemoteControl/AvaRCExtension.h
@@ -4,8 +4,9 @@
 
 #include "IAvaEditorExtension.h"
 
-class URemoteControlTrackerComponent;
+class IRCSignatureCustomization;
 class URemoteControlPreset;
+class URemoteControlTrackerComponent;
 
 class FAvaRCExtension: public FAvaEditorExtension
 {
@@ -25,6 +26,11 @@ public:
 
 protected:
 	void OpenRemoteControlTab() const;
-
 	void CloseRemoteControlTab() const;
+
+	void RegisterSignatureCustomization();
+	void UnregisterSignatureCustomization();
+
+private:
+	TSharedPtr<IRCSignatureCustomization> SignatureCustomization;
 };
diff --git a/Engine/Plugins/Experimental/Avalanche/Source/AvalancheEditor/Private/RemoteControl/AvaRCSignatureCustomization.cpp b/Engine/Plugins/Experimental/Avalanche/Source/AvalancheEditor/Private/RemoteControl/AvaRCSignatureCustomization.cpp
new file mode 100644
--- /dev/null
+++ b/Engine/Plugins/Experimental/Avalanche/Source/AvalancheEditor/Private/RemoteControl/AvaRCSignatureCustomization.cpp
@@ -0,0 +1,24 @@
+// Copyright Epic Games, Inc. All Rights Reserved.
+
+#include "AvaRCSignatureCustomization.h"
+#include "DragDropOps/AvaOutlinerItemDragDropOp.h"
+#include "UI/Signature/IRCSignatureItem.h"
+
+bool FAvaRCSignatureCustomization::CanAcceptDrop(const FDragDropEvent& InDragDropEvent, IRCSignatureItem* InSignatureItem) const
+{
+	return InDragDropEvent.GetOperationAs<FAvaOutlinerItemDragDropOp>().IsValid();
+}
+
+FReply FAvaRCSignatureCustomization::AcceptDrop(const FDragDropEvent& InDragDropEvent, IRCSignatureItem* InSignatureItem) const
+{
+	// Motion Design Outliner Drag Drop
+	if (TSharedPtr<FAvaOutlinerItemDragDropOp> OutlinerDragDrop = InDragDropEvent.GetOperationAs<FAvaOutlinerItemDragDropOp>())
+	{
+		TArray<TWeakObjectPtr<AActor>> DragDropActors;
+		OutlinerDragDrop->GetDragDropOpActors(DragDropActors);
+		InSignatureItem->ApplySignature(DragDropActors);
+		return FReply::Handled();
+	}
+
+	return FReply::Unhandled();
+}
diff --git a/Engine/Plugins/Experimental/Avalanche/Source/AvalancheEditor/Private/RemoteControl/AvaRCSignatureCustomization.h b/Engine/Plugins/Experimental/Avalanche/Source/AvalancheEditor/Private/RemoteControl/AvaRCSignatureCustomization.h
new file mode 100644
--- /dev/null
+++ b/Engine/Plugins/Experimental/Avalanche/Source/AvalancheEditor/Private/RemoteControl/AvaRCSignatureCustomization.h
@@ -0,0 +1,13 @@
+// Copyright Epic Games, Inc. All Rights Reserved.
+
+#pragma once
+
+#include "UI/Signature/IRCSignatureCustomization.h"
+
+class FAvaRCSignatureCustomization : public IRCSignatureCustomization
+{
+	//~ Begin IRCSignatureCustomization
+	virtual bool CanAcceptDrop(const FDragDropEvent& InDragDropEvent, IRCSignatureItem* InSignatureItem) const override;
+	virtual FReply AcceptDrop(const FDragDropEvent& InDragDropEvent, IRCSignatureItem* InSignatureItem) const override;
+	//~ End IRCSignatureCustomization
+};
diff --git a/Engine/Plugins/Experimental/Avalanche/Source/AvalancheOutliner/Private/DragDropOps/AvaOutlinerItemDragDropOp.cpp b/Engine/Plugins/Experimental/Avalanche/Source/AvalancheOutliner/Private/DragDropOps/AvaOutlinerItemDragDropOp.cpp
--- a/Engine/Plugins/Experimental/Avalanche/Source/AvalancheOutliner/Private/DragDropOps/AvaOutlinerItemDragDropOp.cpp
+++ b/Engine/Plugins/Experimental/Avalanche/Source/AvalancheOutliner/Private/DragDropOps/AvaOutlinerItemDragDropOp.cpp
@@ -54,23 +54,19 @@ TSharedRef<FAvaOutlinerItemDragDropOp> FAvaOutlinerItemDragDropOp::New(const TAr
 	return DragDropOp;
 }
 
-void FAvaOutlinerItemDragDropOp::GetDragDropOpActors(TSharedPtr<FDragDropOperation> InDragDropOp, TArray<TWeakObjectPtr<AActor>>& OutActors)
+void FAvaOutlinerItemDragDropOp::GetDragDropOpActors(TArray<TWeakObjectPtr<AActor>>& OutActors) const
 {
-	if (InDragDropOp.IsValid() && InDragDropOp->IsOfType<FAvaOutlinerItemDragDropOp>())
+	for (const FAvaOutlinerItemPtr& Item : Items)
 	{
-		const TSharedPtr<FAvaOutlinerItemDragDropOp> OutlinerItemDragDropOp = StaticCastSharedPtr<FAvaOutlinerItemDragDropOp>(InDragDropOp);
-		for (const FAvaOutlinerItemPtr& Item : OutlinerItemDragDropOp->Items)
+		if (!Item.IsValid())
 		{
-			if (!Item.IsValid())
-			{
-				continue;
-			}
-			if (FAvaOutlinerActor* const ActorItem = Item->CastTo<FAvaOutlinerActor>())
+			continue;
+		}
+		if (FAvaOutlinerActor* const ActorItem = Item->CastTo<FAvaOutlinerActor>())
+		{
+			if (AActor* const UnderlyingActor = ActorItem->GetActor())
 			{
-				if (AActor* const UnderlyingActor = ActorItem->GetActor())
-				{
-					OutActors.Add(UnderlyingActor);
-				}
+				OutActors.Add(UnderlyingActor);
 			}
 		}
 	}
diff --git a/Engine/Plugins/Experimental/Avalanche/Source/AvalancheOutliner/Public/DragDropOps/AvaOutlinerItemDragDropOp.h b/Engine/Plugins/Experimental/Avalanche/Source/AvalancheOutliner/Public/DragDropOps/AvaOutlinerItemDragDropOp.h
--- a/Engine/Plugins/Experimental/Avalanche/Source/AvalancheOutliner/Public/DragDropOps/AvaOutlinerItemDragDropOp.h
+++ b/Engine/Plugins/Experimental/Avalanche/Source/AvalancheOutliner/Public/DragDropOps/AvaOutlinerItemDragDropOp.h
@@ -37,7 +37,7 @@ public:
 	}
 
 	/** Retrieves the Actors found in the DragDropOp, assuming this is of type FAvaOutlinerItemDragDropOp */
-	static void GetDragDropOpActors(TSharedPtr<FDragDropOperation> InDragDropOp, TArray<TWeakObjectPtr<AActor>>& OutActors);
+	AVALANCHEOUTLINER_API void GetDragDropOpActors(TArray<TWeakObjectPtr<AActor>>& OutActors) const;
 
 	/** Called when the FAvaOutlinerItemDragDropOp has been created and Initialized in FAvaOutlinerItemDragDropOp::Init */
 	DECLARE_MULTICAST_DELEGATE_OneParam(FOnItemDragDropOpInitialized, FAvaOutlinerItemDragDropOp&)
