[Motion Design] Clean up of IAvaMediaModule: removing GetPlaybackServerInternal() from module interface.

diff --git a/Engine/Plugins/Experimental/Avalanche/Source/AvalancheMedia/Private/AvaMediaModule.cpp b/Engine/Plugins/Experimental/Avalanche/Source/AvalancheMedia/Private/AvaMediaModule.cpp
--- a/Engine/Plugins/Experimental/Avalanche/Source/AvalancheMedia/Private/AvaMediaModule.cpp
+++ b/Engine/Plugins/Experimental/Avalanche/Source/AvalancheMedia/Private/AvaMediaModule.cpp
@@ -610,6 +610,12 @@ void FAvaMediaModule::StopPlaybackClientCommand(const TArray<FString>& InArgs)
 
 void FAvaMediaModule::StartHttpPlaybackServerCommand(const TArray<FString>& InArgs)
 {
+	if (!AvaPlaybackServer.IsValid())
+	{
+		UE_LOG(LogAvaMedia, Error, TEXT("Playback Server must be started in order to start Http Playback Server."));
+		return;
+	}
+	
 	if (!AvaPlaybackHttpPlaybackServer)
 	{
 		AvaPlaybackHttpPlaybackServer = MakeShared<FAvaPlaybackHttpServer>();
@@ -623,7 +629,7 @@ void FAvaMediaModule::StartHttpPlaybackServerCommand(const TArray<FString>& InAr
 	if (!AvaPlaybackHttpPlaybackServer->IsRunning())
 	{
 		const int32 Port = GetDefault<UAvaMediaSettings>()->HttpServerPort;
-		AvaPlaybackHttpPlaybackServer->Start(Port);
+		AvaPlaybackHttpPlaybackServer->Start(AvaPlaybackServer, Port);
 		UE_LOG(LogAvaMedia, Log, TEXT("Http Playback Server Started"));
 	}
 }
diff --git a/Engine/Plugins/Experimental/Avalanche/Source/AvalancheMedia/Private/AvaMediaModule.h b/Engine/Plugins/Experimental/Avalanche/Source/AvalancheMedia/Private/AvaMediaModule.h
--- a/Engine/Plugins/Experimental/Avalanche/Source/AvalancheMedia/Private/AvaMediaModule.h
+++ b/Engine/Plugins/Experimental/Avalanche/Source/AvalancheMedia/Private/AvaMediaModule.h
@@ -32,7 +32,6 @@ public:
 	virtual void StopPlaybackServer() override;
 	
 	virtual IAvaPlaybackClient& GetPlaybackClient() override;
-	virtual TSharedPtr<FAvaPlaybackServer> GetPlaybackServerInternal() const override { return AvaPlaybackServer; }
 	virtual IAvaPlaybackServer* GetPlaybackServer() const override { return AvaPlaybackServer.Get(); }
 	virtual const IMediaIOCoreDeviceProvider* GetDeviceProvider(FName InProviderName, const FMediaIOOutputConfiguration* InMediaIOOutputConfiguration) const override;
 	virtual TArray<const IMediaIOCoreDeviceProvider*> GetDeviceProvidersForServer(const FString& InServerName) const override;
diff --git a/Engine/Plugins/Experimental/Avalanche/Source/AvalancheMedia/Private/Playback/AvaPlaybackServer.cpp b/Engine/Plugins/Experimental/Avalanche/Source/AvalancheMedia/Private/Playback/AvaPlaybackServer.cpp
--- a/Engine/Plugins/Experimental/Avalanche/Source/AvalancheMedia/Private/Playback/AvaPlaybackServer.cpp
+++ b/Engine/Plugins/Experimental/Avalanche/Source/AvalancheMedia/Private/Playback/AvaPlaybackServer.cpp
@@ -596,7 +596,7 @@ void FAvaPlaybackServer::HandlePlayableTransitionStartRequest(const FAvaPlayback
 {
 	check(PlaybackInstanceTransitions);
 	
-	UAvaPlaybackServerTransition* Transition = NewObject<UAvaPlaybackServerTransition>();
+	UAvaPlaybackServerTransition* Transition = UAvaPlaybackServerTransition::MakeNew(AsShared());
 	Transition->SetTransitionId(InMessage.TransitionId);
 	Transition->SetChannelName(FName(InMessage.ChannelName));
 	Transition->SetClientName(GetClientNameSafe(InContext->GetSender()));
diff --git a/Engine/Plugins/Experimental/Avalanche/Source/AvalancheMedia/Private/Playback/Http/AvaPlaybackHttpServer.cpp b/Engine/Plugins/Experimental/Avalanche/Source/AvalancheMedia/Private/Playback/Http/AvaPlaybackHttpServer.cpp
--- a/Engine/Plugins/Experimental/Avalanche/Source/AvalancheMedia/Private/Playback/Http/AvaPlaybackHttpServer.cpp
+++ b/Engine/Plugins/Experimental/Avalanche/Source/AvalancheMedia/Private/Playback/Http/AvaPlaybackHttpServer.cpp
@@ -2,7 +2,6 @@
 
 #include "Playback/Http/AvaPlaybackHttpServer.h"
 
-#include "AvaMediaModule.h"
 #include "HttpServerModule.h"
 #include "HttpServerResponse.h"
 #include "Playback/AvaPlaybackServer.h"
@@ -39,23 +38,20 @@ FAvaPlaybackHttpRouteBuilder::FAvaPlaybackHttpRouteBuilder(
 {
 }
 
-void FAvaPlaybackHttpServer::Start(int32 InPortToUse)
+void FAvaPlaybackHttpServer::Start(const TSharedPtr<FAvaPlaybackServer>& InPlaybackServer, int32 InPortToUse)
 {
 	PortToUse = InPortToUse;
 	HttpRouter = FHttpServerModule::Get().GetHttpRouter(PortToUse);
 
-	RegisterRoutes();
+	RegisterRoutes(InPlaybackServer);
 
 	FHttpServerModule::Get().StartAllListeners();
 }
 
-void FAvaPlaybackHttpServer::RegisterRoutes()
+void FAvaPlaybackHttpServer::RegisterRoutes(const TSharedPtr<FAvaPlaybackServer>& InPlaybackServer)
 {
-	const TSharedPtr<FAvaPlaybackServer> MediaPlaybackServer = FModuleManager::GetModulePtr<FAvaMediaModule>(UE_MODULE_NAME)->GetPlaybackServerInternal();
-	check(MediaPlaybackServer.IsValid());
-	
 	// Map as per FAvaPlaybackServer::Init(const FString& AssignedServerName)
-	FAvaPlaybackHttpServer::FBuilder(AsShared(), MediaPlaybackServer)
+	FAvaPlaybackHttpServer::FBuilder(AsShared(), InPlaybackServer)
 	.Route<FAvaPlaybackPing>(TEXT("/playback/ping"), EHttpServerRequestVerbs::VERB_POST, &FAvaPlaybackServer::HandlePlaybackPing)
 	.Route<FAvaPlaybackDeviceProviderDataRequest>(TEXT("/playback/devices"), EHttpServerRequestVerbs::VERB_POST, &FAvaPlaybackServer::HandleDeviceProviderDataRequest)
 	.Route<FAvaPlaybackRequest>(TEXT("/playback"), EHttpServerRequestVerbs::VERB_POST, &FAvaPlaybackServer::HandlePlaybackRequest)
diff --git a/Engine/Plugins/Experimental/Avalanche/Source/AvalancheMedia/Private/Playback/Http/AvaPlaybackHttpServer.h b/Engine/Plugins/Experimental/Avalanche/Source/AvalancheMedia/Private/Playback/Http/AvaPlaybackHttpServer.h
--- a/Engine/Plugins/Experimental/Avalanche/Source/AvalancheMedia/Private/Playback/Http/AvaPlaybackHttpServer.h
+++ b/Engine/Plugins/Experimental/Avalanche/Source/AvalancheMedia/Private/Playback/Http/AvaPlaybackHttpServer.h
@@ -135,9 +135,9 @@ public:
 	FAvaPlaybackHttpServer() = default;
 	virtual ~FAvaPlaybackHttpServer() = default;
 
-	void Start(int32 InPortToUse = 10123);
+	void Start(const TSharedPtr<FAvaPlaybackServer>& InPlaybackServer, int32 InPortToUse = 10123);
 
-	void RegisterRoutes();
+	void RegisterRoutes(const TSharedPtr<FAvaPlaybackServer>& InPlaybackServer);
 
 	/**
 	 * Try to get a route registered under given friendly name. Returns false if could not be found.
diff --git a/Engine/Plugins/Experimental/Avalanche/Source/AvalancheMedia/Private/Playback/Transition/AvaPlaybackServerTransition.cpp b/Engine/Plugins/Experimental/Avalanche/Source/AvalancheMedia/Private/Playback/Transition/AvaPlaybackServerTransition.cpp
--- a/Engine/Plugins/Experimental/Avalanche/Source/AvalancheMedia/Private/Playback/Transition/AvaPlaybackServerTransition.cpp
+++ b/Engine/Plugins/Experimental/Avalanche/Source/AvalancheMedia/Private/Playback/Transition/AvaPlaybackServerTransition.cpp
@@ -2,7 +2,6 @@
 
 #include "Playback/Transition/AvaPlaybackServerTransition.h"
 
-#include "IAvaMediaModule.h"
 #include "Playable/Transition/AvaPlayableTransition.h"
 #include "Playback/AvaPlaybackManager.h"
 #include "Playback/AvaPlaybackServer.h"
@@ -45,6 +44,13 @@ namespace UE::AvaPlaybackServerTransition::Private
 	}
 }
 
+UAvaPlaybackServerTransition* UAvaPlaybackServerTransition::MakeNew(const TSharedPtr<FAvaPlaybackServer>& InPlaybackServer)
+{
+	UAvaPlaybackServerTransition* NewTransition = NewObject<UAvaPlaybackServerTransition>();
+	NewTransition->PlaybackServerWeak = InPlaybackServer.ToWeakPtr();
+	return NewTransition;
+}
+
 void UAvaPlaybackServerTransition::SetEnterValues(const TArray<FAvaPlayableRemoteControlValues>& InEnterValues)
 {
 	EnterValues.Reserve(InEnterValues.Num());
@@ -234,7 +240,7 @@ void UAvaPlaybackServerTransition::Stop()
 	UnregisterFromPlayableTransitionEvent();
 	
 	// Remove transition from server.
-	if (const TSharedPtr<FAvaPlaybackServer> PlaybackServer = IAvaMediaModule::Get().GetPlaybackServerInternal())
+	if (const TSharedPtr<FAvaPlaybackServer> PlaybackServer = PlaybackServerWeak.Pin())
 	{
 		if (!PlaybackServer->RemovePlaybackInstanceTransition(TransitionId))
 		{
@@ -315,7 +321,7 @@ void UAvaPlaybackServerTransition::OnTransitionEvent(UAvaPlayable* InPlayable, U
 		return;
 	}
 
-	const TSharedPtr<FAvaPlaybackServer> PlaybackServer = IAvaMediaModule::Get().GetPlaybackServerInternal();
+	const TSharedPtr<FAvaPlaybackServer> PlaybackServer = PlaybackServerWeak.Pin();
 
 	// Find the page player for this playable
 	if (const TSharedPtr<FAvaPlaybackInstance> Instance = FindInstanceForPlayable(InPlayable))
diff --git a/Engine/Plugins/Experimental/Avalanche/Source/AvalancheMedia/Private/Playback/Transition/AvaPlaybackServerTransition.h b/Engine/Plugins/Experimental/Avalanche/Source/AvalancheMedia/Private/Playback/Transition/AvaPlaybackServerTransition.h
--- a/Engine/Plugins/Experimental/Avalanche/Source/AvalancheMedia/Private/Playback/Transition/AvaPlaybackServerTransition.h
+++ b/Engine/Plugins/Experimental/Avalanche/Source/AvalancheMedia/Private/Playback/Transition/AvaPlaybackServerTransition.h
@@ -26,6 +26,8 @@ class UAvaPlaybackServerTransition : public UAvaPlaybackTransition
 	GENERATED_BODY()
 	
 public:
+	static UAvaPlaybackServerTransition* MakeNew(const TSharedPtr<FAvaPlaybackServer>& InPlaybackServer);
+	
 	void SetChannelName(const FName& InChannelName) { ChannelName = InChannelName; }
 	void SetTransitionId(const FGuid& InTransitionId) { TransitionId = InTransitionId; }
 	void SetClientName(const FString& InClientName) { ClientName = InClientName; }
@@ -73,6 +75,8 @@ protected:
 	void UpdateChannelName(const FAvaPlaybackInstance* InPlaybackInstance);
 
 protected:
+	TWeakPtr<FAvaPlaybackServer> PlaybackServerWeak;
+	
 	FString ClientName;
 	FName ChannelName;
 	FGuid TransitionId;
diff --git a/Engine/Plugins/Experimental/Avalanche/Source/AvalancheMedia/Public/IAvaMediaModule.h b/Engine/Plugins/Experimental/Avalanche/Source/AvalancheMedia/Public/IAvaMediaModule.h
--- a/Engine/Plugins/Experimental/Avalanche/Source/AvalancheMedia/Public/IAvaMediaModule.h
+++ b/Engine/Plugins/Experimental/Avalanche/Source/AvalancheMedia/Public/IAvaMediaModule.h
@@ -9,7 +9,6 @@
 DECLARE_LOG_CATEGORY_EXTERN(LogAvaMedia, Log, All);
 
 class FAvaPlaybackManager;
-class FAvaPlaybackServer;
 class FAvaRundownManagedInstanceCache;
 class FCommonViewportClient;
 class FName;
@@ -81,7 +80,6 @@ public:
 	virtual void StopPlaybackServer() = 0;
 
 	virtual IAvaPlaybackClient& GetPlaybackClient() = 0;
-	virtual TSharedPtr<FAvaPlaybackServer> GetPlaybackServerInternal() const = 0;
 	virtual IAvaPlaybackServer* GetPlaybackServer() const = 0;
 	virtual const IMediaIOCoreDeviceProvider* GetDeviceProvider(FName InProviderName, const FMediaIOOutputConfiguration* InMediaIOOutputConfiguration) const = 0;
 	virtual TArray<const IMediaIOCoreDeviceProvider*> GetDeviceProvidersForServer(const FString& InServerName) const = 0;
