[Motion Design] Media Module Settings (QOL): If auto start is selected, either for playback client or rundown server, it will ask if the user wants to launch the service right away.

Both playback client and rundown server are services we want to auto-enable in an editor session typically. It is seems natural and convenient (at least to me) that it would ask to start the service right away.

diff --git a/Engine/Plugins/Experimental/Avalanche/Source/AvalancheMedia/Private/AvaMediaModule.cpp b/Engine/Plugins/Experimental/Avalanche/Source/AvalancheMedia/Private/AvaMediaModule.cpp
--- a/Engine/Plugins/Experimental/Avalanche/Source/AvalancheMedia/Private/AvaMediaModule.cpp
+++ b/Engine/Plugins/Experimental/Avalanche/Source/AvalancheMedia/Private/AvaMediaModule.cpp
@@ -4,6 +4,7 @@
 
 #include "Application/ThrottleManager.h"
 #include "AvaMediaSettings.h"
+#include "Broadcast/AvaBroadcast.h"
 #include "Broadcast/OutputDevices/AvaBroadcastRenderTargetMediaUtils.h"
 #include "IMediaIOCoreModule.h"
 #include "Interfaces/IPluginManager.h"
@@ -59,7 +60,8 @@ void FAvaMediaModule::StartupModule()
 
 	ConsoleCmds.Add(IConsoleManager::Get().RegisterConsoleCommand(
 				TEXT("MotionDesignPlaybackServer.Start"),
-				TEXT("Starts the playback server."),
+				TEXT("Starts the playback server. Arguments: [ServerName]\n")
+				TEXT("ServerName\tSpecifies the server name. Will use computer name if left empty."),
 				FConsoleCommandWithArgsDelegate::CreateRaw(this, &FAvaMediaModule::StartPlaybackServerCommand),
 				ECVF_Default
 				));
@@ -71,7 +73,8 @@ void FAvaMediaModule::StartupModule()
 				));
 	ConsoleCmds.Add(IConsoleManager::Get().RegisterConsoleCommand(
 				TEXT("MotionDesignPlaybackClient.Start"),
-				TEXT("Starts the playback client."),
+				TEXT("Starts the playback client. Arguments: [-Force]\n")
+				TEXT("-Force\tStops playback server and force client to start."),
 				FConsoleCommandWithArgsDelegate::CreateRaw(this, &FAvaMediaModule::StartPlaybackClientCommand),
 				ECVF_Default
 				));
@@ -107,7 +110,8 @@ void FAvaMediaModule::StartupModule()
 				));
 	ConsoleCmds.Add(IConsoleManager::Get().RegisterConsoleCommand(
 				TEXT("MotionDesignRundownServer.Start"),
-				TEXT("Starts the rundown server."),
+				TEXT("Starts the rundown server. Arguments: [ServerName]\n")
+				TEXT("ServerName\tSpecifies the server name. Will use computer name if left empty."),
 				FConsoleCommandWithArgsDelegate::CreateRaw(this, &FAvaMediaModule::StartRundownServerCommand),
 				ECVF_Default
 				));
@@ -119,7 +123,8 @@ void FAvaMediaModule::StartupModule()
 				));
 	ConsoleCmds.Add(IConsoleManager::Get().RegisterConsoleCommand(
 				TEXT("MotionDesignPlaybackDevices.Save"),
-				TEXT("Save Device Providers data."),
+				TEXT("Save Device Providers data. Arguments: [ServerName]\n")
+				TEXT("ServerName\tSpecifies the server name. Will use computer name if left empty."),
 				FConsoleCommandWithArgsDelegate::CreateRaw(this, &FAvaMediaModule::SaveDeviceProvidersCommand),
 				ECVF_Default
 				));
@@ -197,7 +202,16 @@ void FAvaMediaModule::ShutdownModule()
 
 void FAvaMediaModule::StartPlaybackClient()
 {
-	StartPlaybackClientCommand({});
+	TArray<FString> CommandArgs;
+
+	// In editor mode, the client has priority over the server.
+	const bool bIsEditorMode = GIsEditor && !IsRunningGame();
+	if (bIsEditorMode)
+	{
+		CommandArgs.Add(TEXT("-Force"));
+	}
+
+	StartPlaybackClientCommand(CommandArgs);
 }
 
 void FAvaMediaModule::StopPlaybackClient()
@@ -576,26 +590,51 @@ void FAvaMediaModule::StopPlaybackServerCommand(const TArray<FString>& InArgs)
 
 void FAvaMediaModule::StartPlaybackClientCommand(const TArray<FString>& InArgs)
 {
-	// Starting a playback server in the same process as playback client is forbidden.
-	if (AvaPlaybackServer)
+	if (AvaPlaybackClient)
 	{
-		UE_LOG(LogAvaMedia, Error, TEXT("A Playback Client can't be started in the same process as a Playback Server."));
+		UE_LOG(LogAvaMedia, Log, TEXT("Playback client already started."));
 		return;
 	}
-
-	if (!AvaPlaybackClient)
+	
+	bool bForce = false;
+	for (const FString& Arg : InArgs)
+	{
+		if (Arg.Compare(TEXT("-Force"), ESearchCase::IgnoreCase) == 0)
+		{
+			bForce = true;
+		}
+	}
+	
+	// Starting a playback server in the same process as playback client is forbidden.
+	if (AvaPlaybackServer)
 	{
-		using namespace UE::AvaPlaybackClient::Delegates;
-		if (GetOnConnectionEvent().IsBoundToObject(this))
+		if (!bForce)
 		{
-			GetOnConnectionEvent().AddRaw(this, &FAvaMediaModule::OnAvaPlaybackClientConnectionEvent);
+			UE_LOG(LogAvaMedia, Error, TEXT("A playback client can't be started in the same process as a playback server."));
+			return;
 		}
 
-		AvaPlaybackClient = MakeShared<FAvaPlaybackClient>(this);
-		AvaPlaybackClient->Init();
-		OnAvaPlaybackClientStarted.Broadcast();
-		UE_LOG(LogAvaMedia, Log, TEXT("Playback Client Started"));
+		UE_LOG(LogAvaMedia, Warning, TEXT("Playback server has been stopped in editor mode in favor of playback client."));
+		StopPlaybackServer();
+
+#if WITH_EDITOR
+		// Stopping the playback server requires a reload of the broadcast client config.
+		UAvaBroadcast& Broadcast = UAvaBroadcast::Get();
+		Broadcast.LoadBroadcast();
+		Broadcast.QueueNotifyChange(EAvaBroadcastChange::CurrentProfile); // Force a refresh of broadcast editor (if opened).
+#endif
+	}
+
+	using namespace UE::AvaPlaybackClient::Delegates;
+	if (GetOnConnectionEvent().IsBoundToObject(this))
+	{
+		GetOnConnectionEvent().AddRaw(this, &FAvaMediaModule::OnAvaPlaybackClientConnectionEvent);
 	}
+
+	AvaPlaybackClient = MakeShared<FAvaPlaybackClient>(this);
+	AvaPlaybackClient->Init();
+	OnAvaPlaybackClientStarted.Broadcast();
+	UE_LOG(LogAvaMedia, Log, TEXT("Playback client started"));
 }
 
 void FAvaMediaModule::StopPlaybackClientCommand(const TArray<FString>& InArgs)
diff --git a/Engine/Plugins/Experimental/Avalanche/Source/AvalancheMedia/Private/AvaMediaSettings.cpp b/Engine/Plugins/Experimental/Avalanche/Source/AvalancheMedia/Private/AvaMediaSettings.cpp
--- a/Engine/Plugins/Experimental/Avalanche/Source/AvalancheMedia/Private/AvaMediaSettings.cpp
+++ b/Engine/Plugins/Experimental/Avalanche/Source/AvalancheMedia/Private/AvaMediaSettings.cpp
@@ -1,6 +1,10 @@
 // Copyright Epic Games, Inc. All Rights Reserved.
 
 #include "AvaMediaSettings.h"
+#include "IAvaMediaModule.h"
+#include "Misc/MessageDialog.h"
+
+#define LOCTEXT_NAMESPACE "AvaMediaSettings"
 
 UAvaMediaSettings::UAvaMediaSettings()
 {
@@ -45,5 +49,39 @@ ELogVerbosity::Type UAvaMediaSettings::ToLogVerbosity(EAvaMediaLogVerbosity InAv
 	}
 }
 
+#if WITH_EDITOR
+void UAvaMediaSettings::PostEditChangeProperty(FPropertyChangedEvent& InPropertyChangedEvent)
+{
+	Super::PostEditChangeProperty(InPropertyChangedEvent);
+
+	// Check if playback client is auto-start.
+	static const FName AutoStartPlaybackClientPropertyName = GET_MEMBER_NAME_CHECKED(UAvaMediaSettings, bAutoStartPlaybackClient);
+
+	if (InPropertyChangedEvent.MemberProperty->GetFName() == AutoStartPlaybackClientPropertyName)
+	{
+		IAvaMediaModule& AvaMediaModule = IAvaMediaModule::Get();
+		if (bAutoStartPlaybackClient && !AvaMediaModule.IsPlaybackClientStarted())
+		{
+			const FText MessageText = LOCTEXT("StartPlaybackClientQuestion", "Do you want to start the playback client now?");
+			const EAppReturnType::Type Reply = FMessageDialog::Open(EAppMsgType::YesNo, EAppReturnType::Yes, MessageText);
+
+			if (Reply == EAppReturnType::Yes)
+			{
+				AvaMediaModule.StartPlaybackClient();
+			}
+		}
+		else if (!bAutoStartPlaybackClient && AvaMediaModule.IsPlaybackClientStarted())
+		{
+			const FText MessageText = LOCTEXT("StopPlaybackClientQuestion", "Playback Client is currently running. Do you want to stop it now?");
+			const EAppReturnType::Type Reply = FMessageDialog::Open(EAppMsgType::YesNo, EAppReturnType::Yes, MessageText);
 
+			if (Reply == EAppReturnType::Yes)
+			{
+				AvaMediaModule.StopPlaybackClient();
+			}
+		}
+	}
+}
+#endif
 
+#undef LOCTEXT_NAMESPACE
diff --git a/Engine/Plugins/Experimental/Avalanche/Source/AvalancheMedia/Public/AvaMediaSettings.h b/Engine/Plugins/Experimental/Avalanche/Source/AvalancheMedia/Public/AvaMediaSettings.h
--- a/Engine/Plugins/Experimental/Avalanche/Source/AvalancheMedia/Public/AvaMediaSettings.h
+++ b/Engine/Plugins/Experimental/Avalanche/Source/AvalancheMedia/Public/AvaMediaSettings.h
@@ -85,6 +85,12 @@ public:
 
 	static ELogVerbosity::Type ToLogVerbosity(EAvaMediaLogVerbosity InAvaMediaLogVerbosity);
 
+	//~ Begin UObject
+#if WITH_EDITOR
+	virtual void PostEditChangeProperty(struct FPropertyChangedEvent& InPropertyChangedEvent) override;
+#endif
+	//~ End UObject
+	
 	/** Specifies the background clear color for the channel. */
 	UPROPERTY(config, EditAnywhere, Category = "Broadcast")
 	FLinearColor ChannelClearColor = FLinearColor::Black;
diff --git a/Engine/Plugins/Experimental/Avalanche/Source/AvalancheMedia/Public/IAvaMediaModule.h b/Engine/Plugins/Experimental/Avalanche/Source/AvalancheMedia/Public/IAvaMediaModule.h
--- a/Engine/Plugins/Experimental/Avalanche/Source/AvalancheMedia/Public/IAvaMediaModule.h
+++ b/Engine/Plugins/Experimental/Avalanche/Source/AvalancheMedia/Public/IAvaMediaModule.h
@@ -55,6 +55,7 @@ public:
 
 	/**
 	 * @brief Starts the playback client (if not already started).
+	 * @remark In editor mode, this will stop the playback server. 
 	 */
 	virtual void StartPlaybackClient() = 0;
 
diff --git a/Engine/Plugins/Experimental/Avalanche/Source/AvalancheMediaEditor/Private/Broadcast/AvaBroadcastEditor.cpp b/Engine/Plugins/Experimental/Avalanche/Source/AvalancheMediaEditor/Private/Broadcast/AvaBroadcastEditor.cpp
--- a/Engine/Plugins/Experimental/Avalanche/Source/AvalancheMediaEditor/Private/Broadcast/AvaBroadcastEditor.cpp
+++ b/Engine/Plugins/Experimental/Avalanche/Source/AvalancheMediaEditor/Private/Broadcast/AvaBroadcastEditor.cpp
@@ -535,18 +535,8 @@ void FAvaBroadcastEditor::StartPlaybackClientAction()
 		{
 			return;
 		}
-
-		AvaMediaModule.StopPlaybackServer();
-
-		// Stopping the playback server requires a reload of the broadcast client config.
-		if (BroadcastEditor && BroadcastEditor->GetBroadcastObject())
-		{
-			BroadcastEditor->GetBroadcastObject()->LoadBroadcast();
-			// Force a refresh of the channels tab.
-			BroadcastEditor->GetBroadcastObject()->QueueNotifyChange(EAvaBroadcastChange::CurrentProfile);
-		}
 	}
-		
+
 	AvaMediaModule.StartPlaybackClient();
 }
 
diff --git a/Engine/Plugins/Experimental/Avalanche/Source/AvalancheMediaEditor/Private/Rundown/AvaRundownEditorSettings.cpp b/Engine/Plugins/Experimental/Avalanche/Source/AvalancheMediaEditor/Private/Rundown/AvaRundownEditorSettings.cpp
--- a/Engine/Plugins/Experimental/Avalanche/Source/AvalancheMediaEditor/Private/Rundown/AvaRundownEditorSettings.cpp
+++ b/Engine/Plugins/Experimental/Avalanche/Source/AvalancheMediaEditor/Private/Rundown/AvaRundownEditorSettings.cpp
@@ -1,6 +1,10 @@
 // Copyright Epic Games, Inc. All Rights Reserved.
 
 #include "Rundown/AvaRundownEditorSettings.h"
+#include "IAvaMediaModule.h"
+#include "Misc/MessageDialog.h"
+
+#define LOCTEXT_NAMESPACE "AvaRundownEditorSettings"
 
 UAvaRundownEditorSettings::UAvaRundownEditorSettings()
 {
@@ -23,4 +27,39 @@ UAvaRundownEditorSettings* UAvaRundownEditorSettings::GetMutable()
 		DefaultSettings->SetFlags(RF_Transactional);
 	}
 	return DefaultSettings;
-}
\ No newline at end of file
+}
+
+void UAvaRundownEditorSettings::PostEditChangeProperty(struct FPropertyChangedEvent& InPropertyChangedEvent)
+{
+	Super::PostEditChangeProperty(InPropertyChangedEvent);
+
+	// Check if server is auto-start.
+	static const FName AutoStartRundownServerPropertyName = GET_MEMBER_NAME_CHECKED(UAvaRundownEditorSettings, bAutoStartRundownServer);
+
+	if (InPropertyChangedEvent.MemberProperty->GetFName() == AutoStartRundownServerPropertyName)
+	{
+		IAvaMediaModule& AvaMediaModule = IAvaMediaModule::Get();
+		if (bAutoStartRundownServer && !AvaMediaModule.IsRundownServerStarted())
+		{
+			const FText MessageText = LOCTEXT("StartRundownServerQuestion", "Do you want to start rundown server now?");
+			const EAppReturnType::Type Reply = FMessageDialog::Open(EAppMsgType::YesNo, EAppReturnType::Yes, MessageText);
+
+			if (Reply == EAppReturnType::Yes)
+			{
+				AvaMediaModule.StartRundownServer(RundownServerName);
+			}
+		}
+		else if (!bAutoStartRundownServer && AvaMediaModule.IsRundownServerStarted())
+		{
+			const FText MessageText = LOCTEXT("StopRundownServerQuestion", "Rundown Server is currently running. Do you want to stop it now?");
+			const EAppReturnType::Type Reply = FMessageDialog::Open(EAppMsgType::YesNo, EAppReturnType::Yes, MessageText);
+
+			if (Reply == EAppReturnType::Yes)
+			{
+				AvaMediaModule.StopRundownServer();
+			}
+		}
+	}
+}
+
+#undef LOCTEXT_NAMESPACE
diff --git a/Engine/Plugins/Experimental/Avalanche/Source/AvalancheMediaEditor/Private/Rundown/AvaRundownEditorSettings.h b/Engine/Plugins/Experimental/Avalanche/Source/AvalancheMediaEditor/Private/Rundown/AvaRundownEditorSettings.h
--- a/Engine/Plugins/Experimental/Avalanche/Source/AvalancheMediaEditor/Private/Rundown/AvaRundownEditorSettings.h
+++ b/Engine/Plugins/Experimental/Avalanche/Source/AvalancheMediaEditor/Private/Rundown/AvaRundownEditorSettings.h
@@ -19,6 +19,10 @@ public:
 	static const UAvaRundownEditorSettings* Get();
 	static UAvaRundownEditorSettings* GetMutable();
 
+	//~ Begin UObject
+	virtual void PostEditChangeProperty(struct FPropertyChangedEvent& InPropertyChangedEvent) override;
+	//~ End UObject
+	
 	/**
 	 * Configuring the default page action when closing the editor.
 	 * By default it will stop the pages.
