[StateTreeDebugger] enabled debugger on all desktop platforms

diff --git a/Engine/Plugins/Experimental/Avalanche/Source/AvalancheTransitionEditor/AvalancheTransitionEditor.Build.cs b/Engine/Plugins/Experimental/Avalanche/Source/AvalancheTransitionEditor/AvalancheTransitionEditor.Build.cs
--- a/Engine/Plugins/Experimental/Avalanche/Source/AvalancheTransitionEditor/AvalancheTransitionEditor.Build.cs
+++ b/Engine/Plugins/Experimental/Avalanche/Source/AvalancheTransitionEditor/AvalancheTransitionEditor.Build.cs
@@ -41,7 +41,8 @@ public class AvalancheTransitionEditor : ModuleRules
             }
         );
 
-        if (Target.Platform == UnrealTargetPlatform.Win64 && (Target.Configuration != UnrealTargetConfiguration.Shipping || Target.bBuildEditor))
+        // Until we split use of TraceServices and TraceAnalysis in the debugger we enable only on platforms supporting both at the moment
+        if (Target.Platform.IsInGroup(UnrealPlatformGroup.Desktop) && (Target.Configuration != UnrealTargetConfiguration.Shipping || Target.bBuildEditor))
         {
             PrivateDefinitions.Add("WITH_STATETREE_DEBUGGER=1");
             PrivateDependencyModuleNames.AddRange(
diff --git a/Engine/Plugins/Runtime/StateTree/Source/StateTreeEditorModule/StateTreeEditorModule.Build.cs b/Engine/Plugins/Runtime/StateTree/Source/StateTreeEditorModule/StateTreeEditorModule.Build.cs
--- a/Engine/Plugins/Runtime/StateTree/Source/StateTreeEditorModule/StateTreeEditorModule.Build.cs
+++ b/Engine/Plugins/Runtime/StateTree/Source/StateTreeEditorModule/StateTreeEditorModule.Build.cs
@@ -60,14 +60,7 @@ namespace UnrealBuildTool.Rules
 				"MessageLog",
 			});
 
-			if (Target.Platform == UnrealTargetPlatform.Win64)
-			{
-				PublicDefinitions.Add("WITH_STATETREE_DEBUGGER=1");
-			}
-			else
-			{
-				PublicDefinitions.Add("WITH_STATETREE_DEBUGGER=0");
-			}
+			PublicDefinitions.Add("WITH_STATETREE_DEBUGGER=1");
 		}
 	}
 }
diff --git a/Engine/Plugins/Runtime/StateTree/Source/StateTreeModule/Private/Debugger/StateTreeDebugger.cpp b/Engine/Plugins/Runtime/StateTree/Source/StateTreeModule/Private/Debugger/StateTreeDebugger.cpp
--- a/Engine/Plugins/Runtime/StateTree/Source/StateTreeModule/Private/Debugger/StateTreeDebugger.cpp
+++ b/Engine/Plugins/Runtime/StateTree/Source/StateTreeModule/Private/Debugger/StateTreeDebugger.cpp
@@ -890,7 +890,8 @@ bool FStateTreeDebugger::ProcessEvent(const FStateTreeInstanceDebugId InstanceId
 			bShouldAddFrameToSpans = true;
 
 			const UE::StateTreeDebugger::FInstanceEventCollection::FContiguousTraceInfo& TraceInfo =
-				ExistingCollection->ContiguousTracesData.Emplace_GetRef(ExistingCollection->FrameSpans.Num()-1);
+				ExistingCollection->ContiguousTracesData.Emplace_GetRef(
+					UE::StateTreeDebugger::FInstanceEventCollection::FContiguousTraceInfo(ExistingCollection->FrameSpans.Num()-1));
 			FrameToAddInSpans.Index += ExistingCollection->FrameSpans[TraceInfo.LastSpanIndex].Frame.Index + 1;
 		}
 	}
diff --git a/Engine/Plugins/Runtime/StateTree/Source/StateTreeModule/Public/Debugger/StateTreeDebuggerTypes.h b/Engine/Plugins/Runtime/StateTree/Source/StateTreeModule/Public/Debugger/StateTreeDebuggerTypes.h
--- a/Engine/Plugins/Runtime/StateTree/Source/StateTreeModule/Public/Debugger/StateTreeDebuggerTypes.h
+++ b/Engine/Plugins/Runtime/StateTree/Source/StateTreeModule/Public/Debugger/StateTreeDebuggerTypes.h
@@ -141,6 +141,11 @@ struct STATETREEMODULE_API FInstanceEventCollection
 
 	struct FContiguousTraceInfo
 	{
+		explicit FContiguousTraceInfo(int32 LastSpanIndex)
+			: LastSpanIndex(LastSpanIndex)
+		{
+		}
+
 		/** Indicates the index of the last spans of the trace and from which the frame index will be used to offset new events since their frames will restart at 0. */
 		int32 LastSpanIndex = INDEX_NONE;
 	};
diff --git a/Engine/Plugins/Runtime/StateTree/Source/StateTreeModule/StateTreeModule.Build.cs b/Engine/Plugins/Runtime/StateTree/Source/StateTreeModule/StateTreeModule.Build.cs
--- a/Engine/Plugins/Runtime/StateTree/Source/StateTreeModule/StateTreeModule.Build.cs
+++ b/Engine/Plugins/Runtime/StateTree/Source/StateTreeModule/StateTreeModule.Build.cs
@@ -44,8 +44,8 @@ namespace UnrealBuildTool.Rules
 				);
 			}
 
-			if (Target.Platform == UnrealTargetPlatform.Win64 && 
-				(Target.Configuration != UnrealTargetConfiguration.Shipping || Target.bBuildEditor))
+			// Until we split use of TraceServices and TraceAnalysis in the debugger we enable only on platforms supporting both at the moment
+			if (Target.Platform.IsInGroup(UnrealPlatformGroup.Desktop) && (Target.Configuration != UnrealTargetConfiguration.Shipping || Target.bBuildEditor))
 			{
 				PublicDefinitions.Add("WITH_STATETREE_DEBUGGER=1");
 				PublicDependencyModuleNames.AddRange(
