Motion Design: fix fonts being created with conflcting names

diff --git a/Engine/Plugins/Experimental/Avalanche/Source/AvalancheTextEditor/Private/Font/AvaFontManagerSubsystem.cpp b/Engine/Plugins/Experimental/Avalanche/Source/AvalancheTextEditor/Private/Font/AvaFontManagerSubsystem.cpp
--- a/Engine/Plugins/Experimental/Avalanche/Source/AvalancheTextEditor/Private/Font/AvaFontManagerSubsystem.cpp
+++ b/Engine/Plugins/Experimental/Avalanche/Source/AvalancheTextEditor/Private/Font/AvaFontManagerSubsystem.cpp
@@ -669,6 +669,12 @@ UAvaFontObject* UAvaFontManagerSubsystem::CreateProjectFont(UFont* InSourceFont,
 					ProjectFontsPackage->AddToRoot();
 				}
 
+				// If the font object already exist, early return
+				if (UAvaFontObject* ExistingAvaFont = FindObject<UAvaFontObject>(ProjectFontsPackage, *SanitizedFontName))
+				{
+					return ExistingAvaFont;
+				}
+
 				UAvaFontObject* NewAvaFont = NewObject<UAvaFontObject>(ProjectFontsPackage, *SanitizedFontName, RF_Public | RF_Standalone | RF_Transient);
 				NewAvaFont->AddToRoot();
 
