[Motion Design] Rundown Server: Adding more broadcast config commands.

diff --git a/Engine/Plugins/Experimental/Avalanche/Source/AvalancheMedia/Internal/Rundown/AvaRundownMessages.h b/Engine/Plugins/Experimental/Avalanche/Source/AvalancheMedia/Internal/Rundown/AvaRundownMessages.h
--- a/Engine/Plugins/Experimental/Avalanche/Source/AvalancheMedia/Internal/Rundown/AvaRundownMessages.h
+++ b/Engine/Plugins/Experimental/Avalanche/Source/AvalancheMedia/Internal/Rundown/AvaRundownMessages.h
@@ -4,6 +4,7 @@
 
 #include "Broadcast/Channel/AvaBroadcastMediaOutputInfo.h"
 #include "Rundown/AvaRundownPage.h"
+#include "Viewport/AvaViewportQualitySettings.h"
 #include "AvaRundownMessages.generated.h"
 
 namespace EAvaRundownApiVersion
@@ -527,6 +528,107 @@ public:
 	EAvaRundownPageEvents Event = EAvaRundownPageEvents::None;
 };
 
+USTRUCT()
+struct FAvaRundownGetProfiles : public FAvaRundownMsgBase
+{
+	GENERATED_BODY()
+};
+
+USTRUCT()
+struct FAvaRundownProfiles : public FAvaRundownMsgBase
+{
+	GENERATED_BODY()
+
+	/** List of all profiles. */
+	UPROPERTY()
+	TArray<FString> Profiles;
+
+	/** Current Active Profile. */
+	UPROPERTY()
+	FString CurrentProfile;
+};
+
+/**
+ * Creates a new empty profile with the given name.
+ * Fails if the profile already exist.
+ */
+USTRUCT()
+struct FAvaRundownCreateProfile : public FAvaRundownMsgBase
+{
+	GENERATED_BODY()
+public:
+	UPROPERTY()
+	FString ProfileName;
+
+	/**
+	 * If true the created profile is make "current".
+	 */
+	UPROPERTY()
+	bool bMakeCurrent = true;
+};
+
+/**
+ * Duplicates an existing profile.
+ * Fails if the new profile name already exist.
+ * Fails if the source profile does not exist.
+ */
+USTRUCT()
+struct FAvaRundownDuplicateProfile : public FAvaRundownMsgBase
+{
+	GENERATED_BODY()
+public:
+	UPROPERTY()
+	FString SourceProfileName;
+
+	UPROPERTY()
+	FString NewProfileName;
+
+	/**
+	 * If true the created profile is make "current".
+	 */
+	UPROPERTY()
+	bool bMakeCurrent = true;
+};
+
+USTRUCT()
+struct FAvaRundownRenameProfile : public FAvaRundownMsgBase
+{
+	GENERATED_BODY()
+public:
+	UPROPERTY()
+	FString OldProfileName;
+
+	UPROPERTY()
+	FString NewProfileName;
+};
+
+/**
+ * Delete the specified profile.
+ * Fails if profile to be deleted is the current profile.
+ */
+USTRUCT()
+struct FAvaRundownDeleteProfile : public FAvaRundownMsgBase
+{
+	GENERATED_BODY()
+public:
+	UPROPERTY()
+	FString ProfileName;
+};
+
+/**
+ * Specified profile is made "current".
+ * The current profile becomes the context for all other broadcasts commands.
+ * Fails if some channels are currently broadcasting.
+ */
+USTRUCT()
+struct FAvaRundownSetCurrentProfile : public FAvaRundownMsgBase
+{
+	GENERATED_BODY()
+public:
+	UPROPERTY()
+	FString ProfileName;
+};
+
 USTRUCT()
 struct FAvaRundownOutputDeviceItem
 {
@@ -598,6 +700,9 @@ public:
 	UPROPERTY()
 	FString Name;
 
+	UPROPERTY()
+	EAvaBroadcastChannelType Type = EAvaBroadcastChannelType::Program;
+
 	UPROPERTY()
 	EAvaBroadcastChannelState State = EAvaBroadcastChannelState::Offline;
 
@@ -665,6 +770,38 @@ public:
 	EAvaRundownChannelActions Action = EAvaRundownChannelActions::None;
 };
 
+UENUM()
+enum class EAvaRundownChannelEditActions
+{
+	None,
+	Add,
+	Remove
+};
+
+USTRUCT()
+struct FAvaRundownChannelEditAction : public FAvaRundownMsgBase
+{
+	GENERATED_BODY()
+public:
+	UPROPERTY()
+	FString ChannelName;
+
+	UPROPERTY()
+	EAvaRundownChannelEditActions Action = EAvaRundownChannelEditActions::None;
+};
+
+USTRUCT()
+struct FAvaRundownRenameChannel : public FAvaRundownMsgBase
+{
+	GENERATED_BODY()
+public:
+	UPROPERTY()
+	FString OldChannelName;
+
+	UPROPERTY()
+	FString NewChannelName;
+};
+
 USTRUCT()
 struct FAvaRundownGetDevices : public FAvaRundownMsgBase
 {
@@ -743,3 +880,44 @@ public:
 	UPROPERTY()
 	TArray<uint8> ImageData;
 };
+
+/**
+ * Queries the given channel's quality settings.
+ * Response message is FAvaRundownChannelQualitySettings.
+ */
+USTRUCT()
+struct FAvaRundownGetChannelQualitySettings : public FAvaRundownMsgBase
+{
+	GENERATED_BODY()
+public:
+	UPROPERTY()
+	FString ChannelName;
+};
+
+/** Response to FAvaRundownGetChannelQualitySettings. */
+USTRUCT()
+struct FAvaRundownChannelQualitySettings : public FAvaRundownMsgBase
+{
+	GENERATED_BODY()
+public:
+	UPROPERTY()
+	FString ChannelName;
+
+	/** Advanced viewport client engine features indexed by FEngineShowFlags names. */
+	UPROPERTY()
+	TArray<FAvaViewportQualitySettingsFeature> Features;
+};
+
+/** Sets the given channel's quality settings. */
+USTRUCT()
+struct FAvaRundownSetChannelQualitySettings : public FAvaRundownMsgBase
+{
+	GENERATED_BODY()
+public:
+	UPROPERTY()
+	FString ChannelName;
+
+	/** Advanced viewport client engine features indexed by FEngineShowFlags names. */
+	UPROPERTY()
+	TArray<FAvaViewportQualitySettingsFeature> Features;
+};
\ No newline at end of file
diff --git a/Engine/Plugins/Experimental/Avalanche/Source/AvalancheMedia/Private/Rundown/AvaRundownServer.cpp b/Engine/Plugins/Experimental/Avalanche/Source/AvalancheMedia/Private/Rundown/AvaRundownServer.cpp
--- a/Engine/Plugins/Experimental/Avalanche/Source/AvalancheMedia/Private/Rundown/AvaRundownServer.cpp
+++ b/Engine/Plugins/Experimental/Avalanche/Source/AvalancheMedia/Private/Rundown/AvaRundownServer.cpp
@@ -80,6 +80,7 @@ namespace UE::AvaRundownServer::Private
 		FAvaRundownChannel Channel;
 		Channel.Name = InChannel.GetChannelName().ToString();
 		Channel.State = InChannel.GetState();
+		Channel.Type = InChannel.GetChannelType();
 		Channel.IssueSeverity = InChannel.GetIssueSeverity();
 		const TArray<UMediaOutput*>& MediaOutputs = InChannel.GetMediaOutputs();
 		for (const UMediaOutput* MediaOutput : MediaOutputs)
@@ -251,14 +252,24 @@ void FAvaRundownServer::Init(const FString& InAssignedHostName)
 	.Handling<FAvaRundownPagePreviewAction>(this, &FAvaRundownServer::HandlePagePreviewAction)
 	.Handling<FAvaRundownPageActions>(this, &FAvaRundownServer::HandlePageActions)
 	.Handling<FAvaRundownPagePreviewActions>(this, &FAvaRundownServer::HandlePagePreviewActions)
+	.Handling<FAvaRundownGetProfiles>(this, &FAvaRundownServer::HandleGetProfiles)
+	.Handling<FAvaRundownDuplicateProfile>(this, &FAvaRundownServer::HandleDuplicateProfile)
+	.Handling<FAvaRundownCreateProfile>(this, &FAvaRundownServer::HandleCreateProfile)
+	.Handling<FAvaRundownRenameProfile>(this, &FAvaRundownServer::HandleRenameProfile)	
+	.Handling<FAvaRundownDeleteProfile>(this, &FAvaRundownServer::HandleDeleteProfile)
+	.Handling<FAvaRundownSetCurrentProfile>(this, &FAvaRundownServer::HandleSetCurrentProfile)
 	.Handling<FAvaRundownGetChannel>(this, &FAvaRundownServer::HandleGetChannel)
 	.Handling<FAvaRundownGetChannels>(this, &FAvaRundownServer::HandleGetChannels)
 	.Handling<FAvaRundownChannelAction>(this, &FAvaRundownServer::HandleChannelAction)
+	.Handling<FAvaRundownChannelEditAction>(this, &FAvaRundownServer::HandleChannelEditAction)
+	.Handling<FAvaRundownRenameChannel>(this, &FAvaRundownServer::HandleRenameChannel)
 	.Handling<FAvaRundownGetDevices>(this, &FAvaRundownServer::HandleGetDevices)
 	.Handling<FAvaRundownAddChannelDevice>(this, &FAvaRundownServer::HandleAddChannelDevice)
 	.Handling<FAvaRundownEditChannelDevice>(this, &FAvaRundownServer::HandleEditChannelDevice)
 	.Handling<FAvaRundownRemoveChannelDevice>(this, &FAvaRundownServer::HandleRemoveChannelDevice)
 	.Handling<FAvaRundownGetChannelImage>(this, &FAvaRundownServer::HandleGetChannelImage)
+	.Handling<FAvaRundownGetChannelQualitySettings>(this, &FAvaRundownServer::HandleGetChannelQualitySettings)
+	.Handling<FAvaRundownSetChannelQualitySettings>(this, &FAvaRundownServer::HandleSetChannelQualitySettings)
 	.NotificationHandling(FOnBusNotification::CreateRaw(this, &FAvaRundownServer::OnMessageBusNotification));
 	
 	if (MessageEndpoint.IsValid())
@@ -452,8 +463,8 @@ void FAvaRundownServer::HandleRundownPing(const FAvaRundownPing& InMessage, cons
 	// We still support the initial version.
 	constexpr int32 CurrentMinimumApiVersion = EAvaRundownApiVersion::Initial;
 
-	// Consider clients that didn't request a version to be the "initial" version.
-	const int32 RequestedApiVersion = InMessage.RequestedApiVersion != -1 ? InMessage.RequestedApiVersion : EAvaRundownApiVersion::Initial;
+	// Consider clients that didn't request a version to be the latest version.
+	const int32 RequestedApiVersion = InMessage.RequestedApiVersion != -1 ? InMessage.RequestedApiVersion : EAvaRundownApiVersion::LatestVersion;
 
 	// Determine the version we will communicate with this client.
 	int32 HonoredApiVersion = EAvaRundownApiVersion::LatestVersion;
@@ -872,6 +883,164 @@ void FAvaRundownServer::HandlePagePreviewActions(const FAvaRundownPagePreviewAct
 	HandlePageActions(RequestInfo, InMessage.PageIds, true, FName(InMessage.PreviewChannelName), InMessage.Action);
 }
 
+void FAvaRundownServer::HandleGetProfiles(const FAvaRundownGetProfiles& InMessage, const TSharedRef<IMessageContext, ESPMode::ThreadSafe>& InContext)
+{
+	const UAvaBroadcast& Broadcast = UAvaBroadcast::Get();
+	FAvaRundownProfiles* ReplyMessage = FMessageEndpoint::MakeMessage<FAvaRundownProfiles>();
+	ReplyMessage->RequestId = InMessage.RequestId;
+	ReplyMessage->Profiles.Reserve(Broadcast.GetProfiles().Num());
+	for (const TPair<FName, FAvaBroadcastProfile>& Profile : Broadcast.GetProfiles())
+	{
+		ReplyMessage->Profiles.Add(Profile.Key.ToString());
+	}
+	ReplyMessage->CurrentProfile = Broadcast.GetCurrentProfileName().ToString();
+	SendResponse(ReplyMessage, InContext->GetSender());
+}
+
+void FAvaRundownServer::HandleCreateProfile(const FAvaRundownCreateProfile& InMessage, const TSharedRef<IMessageContext, ESPMode::ThreadSafe>& InContext)
+{
+	UAvaBroadcast& Broadcast = UAvaBroadcast::Get();
+	const FName ProfileName(InMessage.ProfileName);
+	
+	if (Broadcast.GetProfile(ProfileName).IsValidProfile())
+	{
+		LogAndSendMessage(InContext->GetSender(), InMessage.RequestId, ELogVerbosity::Error,
+			TEXT("\"CreateProfile\" Failed. Reason: Profile \"%s\" already exist."), *InMessage.ProfileName);
+		return;
+	}
+
+	Broadcast.CreateProfile(ProfileName, InMessage.bMakeCurrent);	// Always succeed apparently.
+	LogAndSendMessage(InContext->GetSender(), InMessage.RequestId, ELogVerbosity::Log,
+		TEXT("\"CreateProfile\" Profile \"%s\" created."), *InMessage.ProfileName);
+}
+
+void FAvaRundownServer::HandleDuplicateProfile(const FAvaRundownDuplicateProfile& InMessage, const TSharedRef<IMessageContext, ESPMode::ThreadSafe>& InContext)
+{
+	UAvaBroadcast& Broadcast = UAvaBroadcast::Get();
+	const FName SourceProfileName(InMessage.SourceProfileName);
+	const FName NewProfileName(InMessage.NewProfileName);
+	
+	if (!Broadcast.GetProfile(SourceProfileName).IsValidProfile())
+	{
+		LogAndSendMessage(InContext->GetSender(), InMessage.RequestId, ELogVerbosity::Error,
+			TEXT("\"DuplicateProfile\" Failed. Reason: Source Profile \"%s\" does not exist."), *InMessage.SourceProfileName);
+		return;
+	}
+
+	if (Broadcast.GetProfile(NewProfileName).IsValidProfile())
+	{
+		LogAndSendMessage(InContext->GetSender(), InMessage.RequestId, ELogVerbosity::Error,
+			TEXT("\"DuplicateProfile\" Failed. Reason: Destination Profile \"%s\" already exist."), *InMessage.NewProfileName);
+		return;
+	}
+
+	if (NewProfileName.IsNone())
+	{
+		LogAndSendMessage(InContext->GetSender(), InMessage.RequestId, ELogVerbosity::Error,
+			TEXT("\"DuplicateProfile\" Failed. Reason: Destination Profile Name is empty."));
+		return;
+	}
+	
+	if (!Broadcast.DuplicateProfile(NewProfileName,  SourceProfileName, InMessage.bMakeCurrent))
+	{
+		LogAndSendMessage(InContext->GetSender(), InMessage.RequestId, ELogVerbosity::Error,
+			TEXT("\"DuplicateProfile\" Failed to duplicate \"%s\" from \"%s\" (Reason unknown)."), *InMessage.NewProfileName, *InMessage.SourceProfileName);
+		return;
+	}
+	
+	LogAndSendMessage(InContext->GetSender(), InMessage.RequestId, ELogVerbosity::Log,
+		TEXT("\"DuplicateProfile\" Profile \"%s\" duplicated from \"%s\"."), *InMessage.NewProfileName, *InMessage.SourceProfileName);
+}
+
+void FAvaRundownServer::HandleRenameProfile(const FAvaRundownRenameProfile& InMessage, const TSharedRef<IMessageContext, ESPMode::ThreadSafe>& InContext)
+{
+	UAvaBroadcast& Broadcast = UAvaBroadcast::Get();
+	const FName OldProfileName(InMessage.OldProfileName);
+	const FName NewProfileName(InMessage.NewProfileName);
+	FText FailReason;
+
+	// CanRenameProfile doesn't check if the profile exists.
+	if (!Broadcast.GetProfile(OldProfileName).IsValidProfile())
+	{
+		LogAndSendMessage(InContext->GetSender(), InMessage.RequestId, ELogVerbosity::Error,
+			TEXT("\"RenameProfile\" Failed. Reason: Profile \"%s\" does not exist."), *InMessage.OldProfileName);
+		return;
+	}
+
+	if (!Broadcast.CanRenameProfile(OldProfileName, NewProfileName, &FailReason))
+	{
+		LogAndSendMessage(InContext->GetSender(), InMessage.RequestId, ELogVerbosity::Error,
+			TEXT("\"RenameProfile\" Failed to rename profile \"%s\" to \"%s\". Reason: %s."),
+			*InMessage.OldProfileName, *InMessage.NewProfileName, *FailReason.ToString());
+		return;
+	}
+
+	Broadcast.RenameProfile(OldProfileName, NewProfileName);
+	LogAndSendMessage(InContext->GetSender(), InMessage.RequestId, ELogVerbosity::Log,
+		TEXT("\"RenameProfile\" Profile \"%s\" renamed to \"%s\"."), *InMessage.OldProfileName, *InMessage.NewProfileName);
+}
+
+void FAvaRundownServer::HandleDeleteProfile(const FAvaRundownDeleteProfile& InMessage, const TSharedRef<IMessageContext, ESPMode::ThreadSafe>& InContext)
+{
+	UAvaBroadcast& Broadcast = UAvaBroadcast::Get();
+	const FName ProfileName(InMessage.ProfileName);
+	
+	if (!Broadcast.GetProfile(ProfileName).IsValidProfile())
+	{
+		LogAndSendMessage(InContext->GetSender(), InMessage.RequestId, ELogVerbosity::Error,
+			TEXT("\"DeleteProfile\" Failed. Reason: Profile \"%s\" does not exist."), *InMessage.ProfileName);
+		return;
+	}
+
+	if (Broadcast.GetCurrentProfileName() == ProfileName)
+	{
+		LogAndSendMessage(InContext->GetSender(), InMessage.RequestId, ELogVerbosity::Error,
+			TEXT("\"DeleteProfile\" Failed. Reason: Profile \"%s\" is the currently active profile and can't be deleted."), *InMessage.ProfileName);
+		return;
+	}
+	
+	
+	if (!Broadcast.RemoveProfile(ProfileName))
+	{
+		LogAndSendMessage(InContext->GetSender(), InMessage.RequestId, ELogVerbosity::Error,
+			TEXT("\"DeleteProfile\" Failed to delete profile \"%s\" (Reason unknown)."), *InMessage.ProfileName);
+		return;
+	}
+	
+	LogAndSendMessage(InContext->GetSender(), InMessage.RequestId, ELogVerbosity::Log,
+		TEXT("\"DeleteProfile\" Profile \"%s\" deleted."), *InMessage.ProfileName);
+}
+
+void FAvaRundownServer::HandleSetCurrentProfile(const FAvaRundownSetCurrentProfile& InMessage, const TSharedRef<IMessageContext, ESPMode::ThreadSafe>& InContext)
+{
+	UAvaBroadcast& Broadcast = UAvaBroadcast::Get();
+	const FName ProfileName(InMessage.ProfileName);
+
+	if(Broadcast.IsBroadcastingAnyChannel())
+	{
+		LogAndSendMessage(InContext->GetSender(), InMessage.RequestId, ELogVerbosity::Error,
+			TEXT("\"SetCurrentProfile\" Failed. Reason: Channels are currently broadcasting."));
+		return;
+	}
+	
+	if (!Broadcast.GetProfile(ProfileName).IsValidProfile())
+	{
+		LogAndSendMessage(InContext->GetSender(), InMessage.RequestId, ELogVerbosity::Error,
+			TEXT("\"SetCurrentProfile\" Failed. Reason: Profile \"%s\" does not exist."), *InMessage.ProfileName);
+		return;
+	}
+
+	if (!Broadcast.SetCurrentProfile(ProfileName))
+	{
+		LogAndSendMessage(InContext->GetSender(), InMessage.RequestId, ELogVerbosity::Error,
+			TEXT("\"SetCurrentProfile\" Failed to set current profile \"%s\" (Reason unknown)."), *InMessage.ProfileName);
+		return;
+	}
+
+	LogAndSendMessage(InContext->GetSender(), InMessage.RequestId, ELogVerbosity::Log,
+		TEXT("\"SetCurrentProfile\" Profile \"%s\" is current."), *InMessage.ProfileName);
+}
+
 void FAvaRundownServer::HandleGetChannel(const FAvaRundownGetChannel& InMessage, const TSharedRef<IMessageContext, ESPMode::ThreadSafe>& InContext)
 {
 	const FName ChannelName(InMessage.ChannelName);
@@ -1039,6 +1208,90 @@ void FAvaRundownServer::HandleGetChannelImage(const FAvaRundownGetChannelImage&
 		});
 }
 
+void FAvaRundownServer::HandleChannelEditAction(const FAvaRundownChannelEditAction& InMessage, const TSharedRef<IMessageContext, ESPMode::ThreadSafe>& InContext)
+{
+	UAvaBroadcast& Broadcast = UAvaBroadcast::Get();
+
+	if (InMessage.ChannelName.IsEmpty())
+	{
+		LogAndSendMessage(InContext->GetSender(), InMessage.RequestId, ELogVerbosity::Error,
+			TEXT("\"ChannelEditAction\" Failed. Reason: Empty Channel Name."));
+		return;
+	}
+
+	const FName ChannelName(InMessage.ChannelName);
+
+	if (InMessage.Action == EAvaRundownChannelEditActions::Add)
+	{
+		if (Broadcast.GetCurrentProfile().GetChannel(ChannelName).IsValidChannel())
+		{
+			LogAndSendMessage(InContext->GetSender(), InMessage.RequestId, ELogVerbosity::Error,
+				TEXT("\"ChannelEditAction\" Add Failed. Reason: Channel \"%s\" already exist."), *ChannelName.ToString());
+			return;
+		}
+		
+		Broadcast.GetCurrentProfile().AddChannel(ChannelName);	// This function doesn't fail apparently.
+		LogAndSendMessage(InContext->GetSender(), InMessage.RequestId, ELogVerbosity::Log,
+			TEXT("\"ChannelEditAction\" Add Channel %s succeeded."), *ChannelName.ToString());
+		return;
+	}
+
+	if (InMessage.Action == EAvaRundownChannelEditActions::Remove)
+	{
+		if(!Broadcast.GetCurrentProfile().RemoveChannel(ChannelName))
+		{
+			LogAndSendMessage(InContext->GetSender(), InMessage.RequestId, ELogVerbosity::Error,
+				TEXT("\"ChannelEditAction\" Remove Failed. Reason: Channel \"%s\" didn't exist in profile."), *ChannelName.ToString());
+			return;
+		}
+
+		LogAndSendMessage(InContext->GetSender(), InMessage.RequestId, ELogVerbosity::Log,
+			TEXT("\"ChannelEditAction\" Remove Channel %s succeeded."), *ChannelName.ToString());
+		return;
+	}
+	
+	LogAndSendMessage(InContext->GetSender(), InMessage.RequestId, ELogVerbosity::Error,
+		TEXT("\"ChannelEditAction\" Failed. Reason: Unknown action."));
+}
+
+void FAvaRundownServer::HandleRenameChannel(const FAvaRundownRenameChannel& InMessage, const TSharedRef<IMessageContext, ESPMode::ThreadSafe>& InContext)
+{
+	UAvaBroadcast& Broadcast = UAvaBroadcast::Get();
+	const FName OldChannelName(InMessage.OldChannelName);
+	const FName NewChannelName(InMessage.NewChannelName);
+	
+	if (InMessage.NewChannelName.IsEmpty())
+    {
+    	LogAndSendMessage(InContext->GetSender(), InMessage.RequestId, ELogVerbosity::Error,
+    		TEXT("\"ChannelEditAction\" Failed. Reason: Empty New Channel Name."));
+    	return;
+    }
+	
+	if (!Broadcast.GetCurrentProfile().GetChannel(OldChannelName).IsValidChannel())
+	{
+		LogAndSendMessage(InContext->GetSender(), InMessage.RequestId, ELogVerbosity::Error,
+			TEXT("\"RenameChannel\" Failed. Reason: Channel \"%s\" does not exist."), *OldChannelName.ToString());
+		return;
+	}
+
+	if (Broadcast.GetCurrentProfile().GetChannel(NewChannelName).IsValidChannel())
+	{
+		LogAndSendMessage(InContext->GetSender(), InMessage.RequestId, ELogVerbosity::Error,
+			TEXT("\"RenameChannel\" Failed. Reason: Channel \"%s\" already exist."), *NewChannelName.ToString());
+		return;
+	}
+
+	if(!Broadcast.RenameChannel(OldChannelName, NewChannelName))
+	{
+		LogAndSendMessage(InContext->GetSender(), InMessage.RequestId, ELogVerbosity::Error,
+			TEXT("\"RenameChannel\" Failed to rename channel \"%s\" to \"%s\" (Unknown reason)."), *OldChannelName.ToString(), *NewChannelName.ToString());
+		return;
+	}
+
+	LogAndSendMessage(InContext->GetSender(), InMessage.RequestId, ELogVerbosity::Log,
+		TEXT("\"RenameChannel\" Channel \"%s\" rename to \"%s\"."), *OldChannelName.ToString(), *NewChannelName.ToString());
+}
+
 void FAvaRundownServer::HandleAddChannelDevice(const FAvaRundownAddChannelDevice& InMessage, const TSharedRef<IMessageContext, ESPMode::ThreadSafe>& InContext)
 {
 	UAvaBroadcast& Broadcast = UAvaBroadcast::Get();
@@ -1171,6 +1424,37 @@ void FAvaRundownServer::HandleRemoveChannelDevice(const FAvaRundownRemoveChannel
 	LogAndSendMessage(InContext->GetSender(), InMessage.RequestId, ELogVerbosity::Log, TEXT("\"RemoveChannelDevice\" Removed Device \"%s\""), *InMessage.MediaOutputName);
 }
 
+void FAvaRundownServer::HandleGetChannelQualitySettings(const FAvaRundownGetChannelQualitySettings& InMessage, const TSharedRef<IMessageContext, ESPMode::ThreadSafe>& InContext)
+{
+	UAvaBroadcast& Broadcast = UAvaBroadcast::Get();
+	const FAvaBroadcastOutputChannel& Channel = Broadcast.GetCurrentProfile().GetChannel(FName(InMessage.ChannelName));
+	if (!Channel.IsValidChannel())
+	{
+		LogAndSendMessage(InContext->GetSender(), InMessage.RequestId, ELogVerbosity::Error, TEXT("\"GetChannelQualitySettings\" Failed. Reason: Invalid Channel \"%s\"."), *InMessage.ChannelName);
+		return;
+	}
+
+	FAvaRundownChannelQualitySettings* ReplyMessage = FMessageEndpoint::MakeMessage<FAvaRundownChannelQualitySettings>();
+	ReplyMessage->RequestId = InMessage.RequestId;
+	ReplyMessage->ChannelName = InMessage.ChannelName;
+	ReplyMessage->Features = Channel.GetViewportQualitySettings().Features;
+	SendResponse(ReplyMessage, InContext->GetSender());
+}
+
+void FAvaRundownServer::HandleSetChannelQualitySettings(const FAvaRundownSetChannelQualitySettings& InMessage, const TSharedRef<IMessageContext, ESPMode::ThreadSafe>& InContext)
+{
+	UAvaBroadcast& Broadcast = UAvaBroadcast::Get();
+	FAvaBroadcastOutputChannel& Channel = Broadcast.GetCurrentProfile().GetChannelMutable(FName(InMessage.ChannelName));
+	if (!Channel.IsValidChannel())
+	{
+		LogAndSendMessage(InContext->GetSender(), InMessage.RequestId, ELogVerbosity::Error, TEXT("\"SetChannelQualitySettings\" Failed. Reason: Invalid Channel \"%s\"."), *InMessage.ChannelName);
+		return;
+	}
+
+	Channel.SetViewportQualitySettings(FAvaViewportQualitySettings(InMessage.Features));
+	LogAndSendMessage(InContext->GetSender(), InMessage.RequestId, ELogVerbosity::Log, TEXT("\"SetChannelQualitySettings\" Channel \"%s\" success."), *InMessage.ChannelName);
+}
+
 void FAvaRundownServer::HandleGetDevices(const FAvaRundownGetDevices& InMessage,
 	const TSharedRef<IMessageContext, ESPMode::ThreadSafe>& InContext)
 {
diff --git a/Engine/Plugins/Experimental/Avalanche/Source/AvalancheMedia/Private/Rundown/AvaRundownServer.h b/Engine/Plugins/Experimental/Avalanche/Source/AvalancheMedia/Private/Rundown/AvaRundownServer.h
--- a/Engine/Plugins/Experimental/Avalanche/Source/AvalancheMedia/Private/Rundown/AvaRundownServer.h
+++ b/Engine/Plugins/Experimental/Avalanche/Source/AvalancheMedia/Private/Rundown/AvaRundownServer.h
@@ -47,7 +47,7 @@ public:
 	void OnBroadcastChannelChanged(const FAvaBroadcastOutputChannel& InChannel, EAvaBroadcastChannelChange InChange) const;
 	void OnAssetAddedOrRemoved(const FAssetData& InAssetData) const;
 	
-	// Message handlers
+	// Rundown message handlers
 	void HandleRundownPing(const FAvaRundownPing& InMessage, const TSharedRef<IMessageContext, ESPMode::ThreadSafe>& InContext);
 	void HandleGetRundowns(const FAvaRundownGetRundowns& InMessage, const TSharedRef<IMessageContext, ESPMode::ThreadSafe>& InContext);
 	void HandleLoadRundown(const FAvaRundownLoadRundown& InMessage, const TSharedRef<IMessageContext, ESPMode::ThreadSafe>& InContext);
@@ -66,13 +66,24 @@ public:
 	void HandlePageActions(const FAvaRundownPageActions& InMessage, const TSharedRef<IMessageContext, ESPMode::ThreadSafe>& InContext);
 	void HandlePagePreviewActions(const FAvaRundownPagePreviewActions& InMessage, const TSharedRef<IMessageContext, ESPMode::ThreadSafe>& InContext);
 
+	// Broadcast message handlers
+	void HandleGetProfiles(const FAvaRundownGetProfiles& InMessage, const TSharedRef<IMessageContext, ESPMode::ThreadSafe>& InContext);
+	void HandleCreateProfile(const FAvaRundownCreateProfile& InMessage, const TSharedRef<IMessageContext, ESPMode::ThreadSafe>& InContext);
+	void HandleDuplicateProfile(const FAvaRundownDuplicateProfile& InMessage, const TSharedRef<IMessageContext, ESPMode::ThreadSafe>& InContext);
+	void HandleRenameProfile(const FAvaRundownRenameProfile& InMessage, const TSharedRef<IMessageContext, ESPMode::ThreadSafe>& InContext);
+	void HandleDeleteProfile(const FAvaRundownDeleteProfile& InMessage, const TSharedRef<IMessageContext, ESPMode::ThreadSafe>& InContext);
+	void HandleSetCurrentProfile(const FAvaRundownSetCurrentProfile& InMessage, const TSharedRef<IMessageContext, ESPMode::ThreadSafe>& InContext);
 	void HandleGetChannel(const FAvaRundownGetChannel& InMessage, const TSharedRef<IMessageContext, ESPMode::ThreadSafe>& InContext);
 	void HandleGetChannels(const FAvaRundownGetChannels& InMessage, const TSharedRef<IMessageContext, ESPMode::ThreadSafe>& InContext);
 	void HandleChannelAction(const FAvaRundownChannelAction& InMessage, const TSharedRef<IMessageContext, ESPMode::ThreadSafe>& InContext);
+	void HandleChannelEditAction(const FAvaRundownChannelEditAction& InMessage, const TSharedRef<IMessageContext, ESPMode::ThreadSafe>& InContext);
+	void HandleRenameChannel(const FAvaRundownRenameChannel& InMessage, const TSharedRef<IMessageContext, ESPMode::ThreadSafe>& InContext);
 	void HandleAddChannelDevice(const FAvaRundownAddChannelDevice& InMessage, const TSharedRef<IMessageContext, ESPMode::ThreadSafe>& InContext);
 	void HandleEditChannelDevice(const FAvaRundownEditChannelDevice& InMessage, const TSharedRef<IMessageContext, ESPMode::ThreadSafe>& InContext);
 	void HandleRemoveChannelDevice(const FAvaRundownRemoveChannelDevice& InMessage, const TSharedRef<IMessageContext, ESPMode::ThreadSafe>& InContext);
 	void HandleGetChannelImage(const FAvaRundownGetChannelImage& InMessage, const TSharedRef<IMessageContext, ESPMode::ThreadSafe>& InContext);
+	void HandleGetChannelQualitySettings(const FAvaRundownGetChannelQualitySettings& InMessage, const TSharedRef<IMessageContext, ESPMode::ThreadSafe>& InContext);
+	void HandleSetChannelQualitySettings(const FAvaRundownSetChannelQualitySettings& InMessage, const TSharedRef<IMessageContext, ESPMode::ThreadSafe>& InContext);
 	
 	void HandleGetDevices(const FAvaRundownGetDevices& InMessage, const TSharedRef<IMessageContext, ESPMode::ThreadSafe>& InContext);
 	
