[Motion Design] Minor include cleanup: AvaMediaModule.h.

diff --git a/Engine/Plugins/Experimental/Avalanche/Source/AvalancheMedia/Private/Broadcast/OutputDevices/AvaBroadcastOutputUtils.cpp b/Engine/Plugins/Experimental/Avalanche/Source/AvalancheMedia/Private/Broadcast/OutputDevices/AvaBroadcastOutputUtils.cpp
--- a/Engine/Plugins/Experimental/Avalanche/Source/AvalancheMedia/Private/Broadcast/OutputDevices/AvaBroadcastOutputUtils.cpp
+++ b/Engine/Plugins/Experimental/Avalanche/Source/AvalancheMedia/Private/Broadcast/OutputDevices/AvaBroadcastOutputUtils.cpp
@@ -2,8 +2,8 @@
 
 #include "Broadcast/OutputDevices/AvaBroadcastOutputUtils.h"
 
-#include "AvaMediaModule.h"
 #include "Broadcast/OutputDevices/AvaGameViewportMediaOutput.h"
+#include "MediaIOCoreDefinitions.h"
 #include "MediaOutput.h"
 #include "Playback/AvaPlaybackMessages.h"
 #include "Serialization/MemoryReader.h"
diff --git a/Engine/Plugins/Experimental/Avalanche/Source/AvalancheMedia/Private/Playback/AvaPlaybackClient.cpp b/Engine/Plugins/Experimental/Avalanche/Source/AvalancheMedia/Private/Playback/AvaPlaybackClient.cpp
--- a/Engine/Plugins/Experimental/Avalanche/Source/AvalancheMedia/Private/Playback/AvaPlaybackClient.cpp
+++ b/Engine/Plugins/Experimental/Avalanche/Source/AvalancheMedia/Private/Playback/AvaPlaybackClient.cpp
@@ -5,16 +5,19 @@
 #include "AssetRegistry/AssetData.h"
 #include "AssetRegistry/IAssetRegistry.h"
 #include "AvaMediaMessageUtils.h"
-#include "AvaMediaModule.h"
 #include "AvaMediaSettings.h"
 #include "Broadcast/AvaBroadcast.h"
+#include "Broadcast/IAvaBroadcastSettings.h"
+#include "Broadcast/OutputDevices/AvaBroadcastDeviceProviderProxy.h"
 #include "Broadcast/OutputDevices/AvaBroadcastOutputUtils.h"
+#include "IAvaMediaModule.h"
 #include "IAvaModule.h"
 #include "IMediaIOCoreModule.h"
 #include "MediaOutput.h"
 #include "MessageEndpointBuilder.h"
 #include "Misc/CoreDelegates.h"
 #include "Misc/Paths.h"
+#include "Playback/AvaPlaybackClientDelegates.h"
 #include "UObject/ObjectSaveContext.h"
 #include "UObject/Package.h"
 
@@ -29,7 +32,7 @@ namespace UE::AvaPlaybackClient::Private
 	}
 }
 
-FAvaPlaybackClient::FAvaPlaybackClient(FAvaMediaModule* InParentModule)
+FAvaPlaybackClient::FAvaPlaybackClient(IAvaMediaModule* InParentModule)
 	: ParentModule(InParentModule)
 {
 	UPackage::PreSavePackageWithContextEvent.AddRaw(this, &FAvaPlaybackClient::OnPreSavePackage);
diff --git a/Engine/Plugins/Experimental/Avalanche/Source/AvalancheMedia/Private/Playback/AvaPlaybackClient.h b/Engine/Plugins/Experimental/Avalanche/Source/AvalancheMedia/Private/Playback/AvaPlaybackClient.h
--- a/Engine/Plugins/Experimental/Avalanche/Source/AvalancheMedia/Private/Playback/AvaPlaybackClient.h
+++ b/Engine/Plugins/Experimental/Avalanche/Source/AvalancheMedia/Private/Playback/AvaPlaybackClient.h
@@ -9,14 +9,14 @@
 #include "Playback/AvaPlaybackMessages.h"
 #include "Playback/IAvaPlaybackClient.h"
 
-class FAvaMediaModule;
+class IAvaMediaModule;
 struct FAvaBroadcastOutputChannel;
 
 // Listens to command to play and broadcast
 class FAvaPlaybackClient : public IAvaPlaybackClient
 {
 public:
-	FAvaPlaybackClient(FAvaMediaModule* InParentModule);
+	FAvaPlaybackClient(IAvaMediaModule* InParentModule);
 	virtual ~FAvaPlaybackClient() override;
 
 	void Init();
@@ -160,7 +160,7 @@ protected:
 	FString GetServerNameForMediaOutputFallback(const UMediaOutput* InMediaOutput) const;
 
 private:
-	FAvaMediaModule* ParentModule;
+	IAvaMediaModule* ParentModule;
 	
 	TArray<IConsoleObject*> ConsoleCommands;
 	
diff --git a/Engine/Plugins/Experimental/Avalanche/Source/AvalancheMedia/Private/Playback/AvaPlaybackUtils.cpp b/Engine/Plugins/Experimental/Avalanche/Source/AvalancheMedia/Private/Playback/AvaPlaybackUtils.cpp
--- a/Engine/Plugins/Experimental/Avalanche/Source/AvalancheMedia/Private/Playback/AvaPlaybackUtils.cpp
+++ b/Engine/Plugins/Experimental/Avalanche/Source/AvalancheMedia/Private/Playback/AvaPlaybackUtils.cpp
@@ -4,7 +4,8 @@
 
 #include "AssetRegistry/AssetData.h"
 #include "AvaAssetTags.h"
-#include "AvaMediaModule.h"
+#include "Framework/AvaSoftAssetPtr.h"
+#include "IAvaMediaModule.h"
 #include "HAL/PlatformFileManager.h"
 #include "Misc/Paths.h"
 #include "RenderingThread.h"
