[Motion Design] Rundown API: Fixing device enumeration.

The structure of the tree has changed, there is server node now. The traversal code needs to be updated.

For now, instead of creating one more level in the message, I just added the server name to the class item, the control app will need to deal with that.

diff --git a/Engine/Plugins/Experimental/Avalanche/Source/AvalancheMedia/Internal/Rundown/AvaRundownMessages.h b/Engine/Plugins/Experimental/Avalanche/Source/AvalancheMedia/Internal/Rundown/AvaRundownMessages.h
--- a/Engine/Plugins/Experimental/Avalanche/Source/AvalancheMedia/Internal/Rundown/AvaRundownMessages.h
+++ b/Engine/Plugins/Experimental/Avalanche/Source/AvalancheMedia/Internal/Rundown/AvaRundownMessages.h
@@ -661,9 +661,14 @@ struct FAvaRundownOutputClassItem
 {
 	GENERATED_BODY()
 public:
+	/** Class name */
 	UPROPERTY()
 	FString Name;
 
+	/** Name of the server this class was seen on. */
+	UPROPERTY()
+	FString Server;
+
 	UPROPERTY()
 	TArray<FAvaRundownOutputDeviceItem> Devices;
 };
diff --git a/Engine/Plugins/Experimental/Avalanche/Source/AvalancheMedia/Private/Rundown/AvaRundownServer.cpp b/Engine/Plugins/Experimental/Avalanche/Source/AvalancheMedia/Private/Rundown/AvaRundownServer.cpp
--- a/Engine/Plugins/Experimental/Avalanche/Source/AvalancheMedia/Private/Rundown/AvaRundownServer.cpp
+++ b/Engine/Plugins/Experimental/Avalanche/Source/AvalancheMedia/Private/Rundown/AvaRundownServer.cpp
@@ -13,6 +13,7 @@
 #include "Broadcast/OutputDevices/AvaBroadcastOutputClassItem.h"
 #include "Broadcast/OutputDevices/AvaBroadcastOutputDeviceItem.h"
 #include "Broadcast/OutputDevices/AvaBroadcastOutputRootItem.h"
+#include "Broadcast/OutputDevices/AvaBroadcastOutputServerItem.h"
 #include "Broadcast/OutputDevices/AvaBroadcastOutputTreeItem.h"
 #include "Broadcast/OutputDevices/AvaBroadcastRenderTargetMediaUtils.h"
 #include "IAvaMediaModule.h"
@@ -1466,37 +1467,49 @@ void FAvaRundownServer::HandleGetDevices(const FAvaRundownGetDevices& InMessage,
 	FAvaBroadcastOutputTreeItem::RefreshTree(OutputDevices, RefreshDevicesParams);
 	// OutputDevices here aren't literally a physical device, just a construct representing
 	// output. This convention was pulled from the SAvaBroadcastOutputDevices->RefreshDevices() call
-	for (const TSharedPtr<IAvaBroadcastOutputTreeItem>& ClassItem : OutputDevices->GetChildren())
+	for (const TSharedPtr<IAvaBroadcastOutputTreeItem>& ServerItem : OutputDevices->GetChildren())
 	{
-		if (const FAvaBroadcastOutputClassItem* AvaOutputClassItem = ClassItem->CastTo<FAvaBroadcastOutputClassItem>())
+		const FAvaBroadcastOutputServerItem* OutputServerItem = ServerItem->CastTo<FAvaBroadcastOutputServerItem>();
+		if (!OutputServerItem)
 		{
+			continue;
+		}
+		
+		for (const TSharedPtr<IAvaBroadcastOutputTreeItem>& ClassItem : ServerItem->GetChildren())
+		{
+			const FAvaBroadcastOutputClassItem* AvaOutputClassItem = ClassItem->CastTo<FAvaBroadcastOutputClassItem>();
+			if (!AvaOutputClassItem)
+			{
+				continue;
+			}
+			
 			FAvaRundownOutputClassItem OutputClassItem;
 			OutputClassItem.Name = ClassItem->GetDisplayName().ToString();
-
-			const TArray<FAvaOutputTreeItemPtr>& Children = AvaOutputClassItem->GetChildren();
-			if (Children.Num() > 0)
+			OutputClassItem.Server = OutputServerItem->GetServerName();
+			
+			for (const TSharedPtr<IAvaBroadcastOutputTreeItem>& OutputDeviceItem : AvaOutputClassItem->GetChildren())
 			{
-				for (const TSharedPtr<IAvaBroadcastOutputTreeItem>& OutputDeviceItem : Children)
+				if (!OutputDeviceItem->IsA<FAvaBroadcastOutputDeviceItem>())
 				{
-					if (OutputDeviceItem->IsA<FAvaBroadcastOutputDeviceItem>())
-					{
-						FAvaRundownOutputDeviceItem DeviceItem;
-						DeviceItem.Name = OutputDeviceItem->GetDisplayName().ToString();
-						// Intentionally leaving DeviceItem.Data blank, as it's not usable data by itself
-						// .Data will be filled out on a GetChannels call, where it becomes usable
-
-						OutputClassItem.Devices.Push(DeviceItem);
-					}
+					continue;
 				}
+				
+				FAvaRundownOutputDeviceItem DeviceItem;
+				DeviceItem.Name = OutputDeviceItem->GetDisplayName().ToString();
+				// Intentionally leaving DeviceItem.Data blank, as it's not usable data by itself
+				// .Data will be filled out on a GetChannels call, where it becomes usable
+
+				OutputClassItem.Devices.Push(MoveTemp(DeviceItem));
 			}
-			else
+
+			if (OutputClassItem.Devices.IsEmpty())
 			{
 				FAvaRundownOutputDeviceItem DeviceItem;
 				DeviceItem.Name = OutputClassItem.Name;
-				OutputClassItem.Devices.Push(DeviceItem);
+				OutputClassItem.Devices.Push(MoveTemp(DeviceItem));
 			}
 
-			ReplyMessage->DeviceClasses.Push(OutputClassItem);
+			ReplyMessage->DeviceClasses.Push(MoveTemp(OutputClassItem));
 		}
 	}
 	SendResponse(ReplyMessage, InContext->GetSender());
