[Motion Design] Broadcast "in place" serialization

Using json serialization to save the broadcast config in a single json file including the sub objects.

The reason for doing this is to have the possibility to serialize in game build.

diff --git a/Engine/Plugins/Experimental/Avalanche/Source/AvalancheMedia/AvalancheMedia.Build.cs b/Engine/Plugins/Experimental/Avalanche/Source/AvalancheMedia/AvalancheMedia.Build.cs
--- a/Engine/Plugins/Experimental/Avalanche/Source/AvalancheMedia/AvalancheMedia.Build.cs
+++ b/Engine/Plugins/Experimental/Avalanche/Source/AvalancheMedia/AvalancheMedia.Build.cs
@@ -40,6 +40,7 @@ public class AvalancheMedia : ModuleRules
 				"AvalancheRemoteControl",
 				"HeadMountedDisplay",
 				"Json",
+				"JsonUtilities",
 				"ImageCore",
 				"InputCore",
 				"Projects",
diff --git a/Engine/Plugins/Experimental/Avalanche/Source/AvalancheMedia/Private/AvaMediaModule.cpp b/Engine/Plugins/Experimental/Avalanche/Source/AvalancheMedia/Private/AvaMediaModule.cpp
--- a/Engine/Plugins/Experimental/Avalanche/Source/AvalancheMedia/Private/AvaMediaModule.cpp
+++ b/Engine/Plugins/Experimental/Avalanche/Source/AvalancheMedia/Private/AvaMediaModule.cpp
@@ -617,12 +617,10 @@ void FAvaMediaModule::StartPlaybackClientCommand(const TArray<FString>& InArgs)
 		UE_LOG(LogAvaMedia, Warning, TEXT("Playback server has been stopped in editor mode in favor of playback client."));
 		StopPlaybackServer();
 
-#if WITH_EDITOR
 		// Stopping the playback server requires a reload of the broadcast client config.
 		UAvaBroadcast& Broadcast = UAvaBroadcast::Get();
 		Broadcast.LoadBroadcast();
 		Broadcast.QueueNotifyChange(EAvaBroadcastChange::CurrentProfile); // Force a refresh of broadcast editor (if opened).
-#endif
 	}
 
 	using namespace UE::AvaPlaybackClient::Delegates;
diff --git a/Engine/Plugins/Experimental/Avalanche/Source/AvalancheMedia/Private/Broadcast/AvaBroadcast.cpp b/Engine/Plugins/Experimental/Avalanche/Source/AvalancheMedia/Private/Broadcast/AvaBroadcast.cpp
--- a/Engine/Plugins/Experimental/Avalanche/Source/AvalancheMedia/Private/Broadcast/AvaBroadcast.cpp
+++ b/Engine/Plugins/Experimental/Avalanche/Source/AvalancheMedia/Private/Broadcast/AvaBroadcast.cpp
@@ -2,11 +2,9 @@
 
 #include "Broadcast/AvaBroadcast.h"
 #include "Async/Async.h"
-#include "AvaMediaSerializationUtils.h"
+#include "AvaBroadcastSerialization.h"
 #include "Containers/UnrealString.h"
 #include "Delegates/IDelegateInstance.h"
-#include "Formatters/XmlArchiveInputFormatter.h"
-#include "Formatters/XmlArchiveOutputFormatter.h"
 #include "HAL/FileManager.h"
 #include "IAvaMediaModule.h"
 #include "Misc/Paths.h"
@@ -19,33 +17,43 @@ DEFINE_LOG_CATEGORY(LogAvaBroadcast);
 
 namespace UE::AvaBroadcast::Private
 {
-	// Tested with server in game and editor modes: This function is called only when the server receives some broadcast queries.
-	static FString GetXmlSaveFilepath()
+	static FString EnsureExtension(const FString& InFilename, const FString& InExtension)
+	{
+		if (!FPaths::GetExtension(InFilename).Equals(InExtension, ESearchCase::IgnoreCase))
+		{
+			// Appending extension instead of using SetExtension to respect
+			// dot naming convention for xml/yaml config files.
+			return InFilename + TEXT(".") + InExtension;
+		}
+		return InFilename;
+	}
+
+	static FString GetConfigFilepath(const FString& InExtension)
 	{
 		FString BroadcastConfigName;
 
 		// Allowing command line specification of the broadcast configuration file name. This allows
 		// starting the same project (from a shared location) with different configurations.
-		if (FParse::Value(FCommandLine::Get(),TEXT("MotionDesignBroadcastConfig="), BroadcastConfigName))
-		{
-			// Ensure it has an xml extension.
-			if (!FPaths::GetExtension(BroadcastConfigName).Equals(TEXT("xml"), ESearchCase::IgnoreCase))
-			{
-				// Appending extension instead of using SetExtension to respect
-				// dot naming convention for xml/yaml config files.
-				BroadcastConfigName += TEXT(".xml");
-			}
-		}
-		else
+		if (!FParse::Value(FCommandLine::Get(),TEXT("MotionDesignBroadcastConfig="), BroadcastConfigName))
 		{
 			// When launching the server from the same project location, we want to avoid loading the same
 			// broadcast configuration as the client. The server needs a clean configuration.
 			const bool bIsServerRunning = (IAvaMediaModule::IsModuleLoaded() && IAvaMediaModule::Get().IsPlaybackServerStarted());	
-			BroadcastConfigName = bIsServerRunning ? TEXT("MotionDesignServerBroadcastConfig.xml") : TEXT("MotionDesignBroadcastConfig.xml");
+			BroadcastConfigName = bIsServerRunning ? TEXT("MotionDesignServerBroadcastConfig") : TEXT("MotionDesignBroadcastConfig");
 		}
-		
-		return FPaths::ProjectConfigDir() / BroadcastConfigName;
+		return EnsureExtension(FPaths::ProjectConfigDir() / BroadcastConfigName, InExtension);
 	}
+
+	static FString GetXmlSaveFilepath()
+	{
+		return GetConfigFilepath(TEXT("xml"));
+	}
+
+	static FString GetJsonSaveFilepath()
+	{
+		return GetConfigFilepath(TEXT("json"));
+	}
+
 }
 
 UAvaBroadcast& UAvaBroadcast::Get()
@@ -66,11 +74,7 @@ UAvaBroadcast& UAvaBroadcast::Get()
 			, RF_Transactional | RF_Standalone);
 
 		Broadcast->AddToRoot();
-#if WITH_EDITOR
 		Broadcast->LoadBroadcast();
-#else
-		Broadcast->InitBroadcast();	// todo: find another way to load config in packaged build.
-#endif
 	}
 	
 	check(Broadcast);
@@ -325,17 +329,25 @@ const FAvaBroadcastProfile& UAvaBroadcast::GetProfile(FName InProfileName) const
 	return FAvaBroadcastProfile::GetNullProfile();
 }
 
-#if WITH_EDITOR
 void UAvaBroadcast::LoadBroadcast()
 {
 	using namespace UE::AvaBroadcast::Private;
-	TUniquePtr<FArchive> FileReader(IFileManager::Get().CreateFileReader(*GetXmlSaveFilepath()));
 
-	if (FileReader.IsValid())
+	LoadedConfigFilepath = GetJsonSaveFilepath();
+	bool bConfigLoaded = FAvaBroadcastSerialization::LoadBroadcastFromJson(LoadedConfigFilepath, this);
+
+#if WITH_EDITOR
+	if (!bConfigLoaded)
+	{
+		// Fallback to legacy xml format.
+		LoadedConfigFilepath = GetXmlSaveFilepath();
+		bConfigLoaded = FAvaBroadcastSerialization::LoadBroadcastFromXml(LoadedConfigFilepath, this);
+	}
+#endif
+
+	if (!bConfigLoaded)
 	{
-		FXmlArchiveInputFormatter InputFormatter(*FileReader, this);
-		UE::AvaMediaSerializationUtils::SerializeObject(InputFormatter, this);
-		FileReader->Close();
+		LoadedConfigFilepath.Reset();
 	}
 
 	// Set the profile names early because it is needed to resolve the pinned channels below.
@@ -360,19 +372,24 @@ void UAvaBroadcast::LoadBroadcast()
 
 void UAvaBroadcast::SaveBroadcast()
 {
-	bool bIsBroadcastSaved = false;
 	using namespace UE::AvaBroadcast::Private;
-	TUniquePtr<FArchive> FileWriter(IFileManager::Get().CreateFileWriter(*GetXmlSaveFilepath()));
-	if (FileWriter.IsValid())	
+	
+	bool bIsBroadcastSaved = FAvaBroadcastSerialization::SaveBroadcastToJson(this, GetJsonSaveFilepath());
+
+#if WITH_EDITOR
+	// In case of failure, fallback to xml format.
+	// Temporary until the json format is battle tested.
+	if (!bIsBroadcastSaved)
+	{
+		bIsBroadcastSaved = FAvaBroadcastSerialization::SaveBroadcastToXml(this, GetXmlSaveFilepath());
+	}
+#endif
+	
+	if (bIsBroadcastSaved)
 	{
-		FXmlArchiveOutputFormatter XmlOutput(*FileWriter);
-		XmlOutput.SerializeObjectsInPlace(true); // We want the media outputs nested.
-		UE::AvaMediaSerializationUtils::SerializeObject(XmlOutput, this);
-		bIsBroadcastSaved = XmlOutput.SaveDocumentToInnerArchive();
 		GetPackage()->SetDirtyFlag(false);
-		FileWriter->Close();
 	}
-	if (!bIsBroadcastSaved)
+	else
 	{
 		UE_LOG(LogAvaBroadcast, Error, TEXT("Failed to save broadcast configuration.")); 
 	}
@@ -380,16 +397,13 @@ void UAvaBroadcast::SaveBroadcast()
 
 FString UAvaBroadcast::GetBroadcastSaveFilepath() const
 {
-	return UE::AvaBroadcast::Private::GetXmlSaveFilepath();
-}
-#else
-void UAvaBroadcast::InitBroadcast()
-{
-	CreateProfile(NAME_None, true);
-	EnsureValidCurrentProfile();
-	UpdateProfileNames();
+	if (!LoadedConfigFilepath.IsEmpty())
+	{
+		return LoadedConfigFilepath;
+	}
+	
+	return UE::AvaBroadcast::Private::GetJsonSaveFilepath();
 }
-#endif
 
 void UAvaBroadcast::QueueNotifyChange(EAvaBroadcastChange InChange)
 {
diff --git a/Engine/Plugins/Experimental/Avalanche/Source/AvalancheMedia/Private/Broadcast/AvaBroadcastSerialization.cpp b/Engine/Plugins/Experimental/Avalanche/Source/AvalancheMedia/Private/Broadcast/AvaBroadcastSerialization.cpp
new file mode 100644
--- /dev/null
+++ b/Engine/Plugins/Experimental/Avalanche/Source/AvalancheMedia/Private/Broadcast/AvaBroadcastSerialization.cpp
@@ -0,0 +1,265 @@
+// Copyright Epic Games, Inc. All Rights Reserved.
+
+#include "AvaBroadcastSerialization.h"
+
+#include "AvaMediaSerializationUtils.h"
+#include "Broadcast/AvaBroadcast.h"
+#include "Formatters/XmlArchiveInputFormatter.h"
+#include "Formatters/XmlArchiveOutputFormatter.h"
+#include "JsonObjectConverter.h"
+#include "Misc/FileHelper.h"
+
+namespace UE::AvaMedia::BroadcastSerialization::Private
+{
+	struct FJsonObjectSerializerContext
+	{
+		TArray<TSharedPtr<FJsonValue>> ReferencedObjects;
+
+		TSharedPtr<FJsonValue> ExportPropertyCallback(FProperty* InProperty, const void* InValue)
+		{
+			TSharedPtr<FJsonValue> Result = nullptr;
+
+			if (const FObjectProperty* ObjectProperty = CastField<FObjectProperty>(InProperty))
+			{
+				const UObject* const Object = ObjectProperty->GetObjectPropertyValue_InContainer(InValue);
+				if (IsValid(Object))
+				{
+					const TSoftObjectPtr<UClass> ObjectClass =  Object->GetClass();
+					const TSharedPtr<FJsonObject> JsonObject = MakeShared<FJsonObject>();
+					// Do we need to keep the path to the outer?
+					JsonObject->SetStringField(TEXT("Class"), ObjectClass.ToSoftObjectPath().ToString());
+					JsonObject->SetStringField(TEXT("Name"), Object->GetName());
+					JsonObject->SetStringField(TEXT("Flags"), FString::Printf(TEXT("%d"), static_cast<uint32>(Object->GetFlags())));
+					JsonObject->SetObjectField(TEXT("Object"), SerializeToJson(Object));
+
+					ReferencedObjects.Add(MakeShared<FJsonValueObject>(JsonObject));
+				}
+			}
+		
+			return Result;
+		}
+
+		TSharedPtr<FJsonObject> SerializeToJson(const UObject* InObject)
+		{
+			FJsonObjectConverter::CustomExportCallback CustomCB;
+			CustomCB.BindRaw(this, &FJsonObjectSerializerContext::ExportPropertyCallback);
+
+			const TSharedPtr<FJsonObject> JsonObject = MakeShared<FJsonObject>();
+			constexpr int64 CheckFlags = 0;
+			constexpr int64 SkipFlags = 0;
+			FJsonObjectConverter::UStructToJsonObject(InObject->GetClass(), InObject, JsonObject.ToSharedRef(), CheckFlags, SkipFlags, &CustomCB, EJsonObjectConversionFlags::SkipStandardizeCase);
+
+			return JsonObject;
+		}
+	};
+	
+	UObject* LoadReferencedObject(UObject* InRootObject, const TSharedPtr<FJsonObject>& InReferenceObject)
+	{
+		const FSoftObjectPath ClassPath(InReferenceObject->GetStringField(TEXT("Class")));
+		const TSoftObjectPtr<UClass> ClassPtr(ClassPath);
+		UClass* const ObjectClass = ClassPtr.LoadSynchronous();
+		
+		if (!ObjectClass)
+		{
+			return nullptr;
+		}
+
+		const TSharedPtr<FJsonObject> JsonObjectData = InReferenceObject->GetObjectField(TEXT("Object"));
+		if (!JsonObjectData)
+		{
+			return nullptr;
+		}
+		
+		const FString OriginalObjectName = InReferenceObject->GetStringField(TEXT("Name"));
+
+		// Try to find the object, it may exist already.
+		//UObject* NestedObject = StaticFindObject(ObjectClass, InRootObject, *OriginalObjectName, /*bExactClass*/ true );
+		UObject* NestedObject = static_cast<UObject*>(FindObjectWithOuter(InRootObject, ObjectClass, FName(OriginalObjectName)));
+
+		if (!NestedObject)
+		{
+			// We need to create the object
+			const FString FlagsString = InReferenceObject->GetStringField(TEXT("Flags"));
+			const uint32 Flags = static_cast<uint32>(FCString::Atoi(*FlagsString));
+			const EObjectFlags ObjectFlags = static_cast<EObjectFlags>(Flags);
+			
+			NestedObject = NewObject<UObject>(InRootObject, ObjectClass, FName(OriginalObjectName), ObjectFlags);
+		}
+
+		if (!FJsonObjectConverter::JsonObjectToUStruct(JsonObjectData.ToSharedRef(), ObjectClass, NestedObject))
+		{
+			return nullptr;
+		}
+
+		return NestedObject;
+	}
+
+	void FixupReferences(const TSharedPtr<FJsonValue>& InValue, const FString& InOldReference, const FString& InNewReference);
+
+	void FixupReferences(const TArray<TSharedPtr<FJsonValue>>& InValues, const FString& InOldReference, const FString& InNewReference)
+	{
+		for (const TSharedPtr<FJsonValue>& Value : InValues)
+		{
+			FixupReferences(Value, InOldReference, InNewReference);
+		}
+	}
+	
+	void FixupReferences(const TSharedPtr<FJsonObject>& InJsonObject, const FString& InOldReference, const FString& InNewReference)
+	{
+		for (TPair<FString, TSharedPtr<FJsonValue>>& Value : InJsonObject->Values)
+		{
+			FixupReferences(Value.Value, InOldReference, InNewReference);
+		}
+	}
+
+	// Can't seem to find a way to set the string value other than this.
+	class FJsonStringFixupHelper : public FJsonValueString
+	{
+	public:
+		FJsonStringFixupHelper(const FString& InString) : FJsonValueString(InString) {}
+			
+		void FixupString(const FString& InOldReference, const FString& InNewReference)
+		{
+			if (Value.Contains(InOldReference, ESearchCase::CaseSensitive))
+			{
+				Value = Value.Replace(*InOldReference, *InNewReference, ESearchCase::CaseSensitive);
+			}
+		}
+	};
+
+	void FixupReferences(const TSharedPtr<FJsonValue>& InValue, const FString& InOldReference, const FString& InNewReference)
+	{
+		switch(InValue->Type)
+		{
+		case EJson::String:
+			static_cast<FJsonStringFixupHelper*>(InValue.Get())->FixupString(InOldReference, InNewReference);	
+			break;
+		case EJson::Array:
+			FixupReferences(InValue->AsArray(), InOldReference, InNewReference);
+			break;
+		case EJson::Object:
+			FixupReferences(InValue->AsObject(), InOldReference, InNewReference);
+			break;
+		default:
+			break;
+		}
+	}
+}
+
+bool FAvaBroadcastSerialization::SaveBroadcastToJson(const UAvaBroadcast* InBroadcast, const FString& InFilename)
+{
+	using namespace UE::AvaMedia::BroadcastSerialization::Private;
+	FJsonObjectSerializerContext SerializerContext;
+	
+	const TSharedPtr<FJsonObject> JsonBroadcast = SerializerContext.SerializeToJson(InBroadcast);
+
+	const FSoftObjectPath BroadcastPath(InBroadcast);
+	
+	const TSharedPtr<FJsonObject> JsonRoot = MakeShared<FJsonObject>();
+	JsonRoot->SetStringField(TEXT("BroadcastPath"), BroadcastPath.ToString());	// Used to fixup reference when loading.
+	JsonRoot->SetObjectField(TEXT("Broadcast"), JsonBroadcast);
+	JsonRoot->SetArrayField(TEXT("ReferencedObjects"), SerializerContext.ReferencedObjects);
+	
+	FString OutJsonString;
+	const TSharedRef<TJsonWriter<>> Writer = TJsonWriterFactory<>::Create(&OutJsonString);
+	if (!FJsonSerializer::Serialize(JsonRoot.ToSharedRef(), Writer))
+	{
+		return false;
+	}
+
+	if (!FFileHelper::SaveStringToFile(OutJsonString, *InFilename))
+	{
+		UE_LOG(LogAvaBroadcast, Error, TEXT("Couldn't save data to file: %s"), *InFilename);
+		return false;
+	}
+
+	return true;
+}
+
+bool FAvaBroadcastSerialization::LoadBroadcastFromJson(const FString& InFilename, UAvaBroadcast* OutBroadcast)
+{
+	FString JsonText;
+
+	// Load json text to the string object
+	if (!FFileHelper::LoadFileToString(JsonText, *InFilename))
+	{
+		UE_LOG(LogAvaBroadcast, Error, TEXT("Couldn't read file: %s"), *InFilename);
+		return false;
+	}
+
+	TSharedPtr<FJsonObject> JsonRoot = MakeShared<FJsonObject>();
+
+	const TSharedRef<TJsonReader<TCHAR>> JsonReader = TJsonReaderFactory<TCHAR>::Create(JsonText);
+	if (!FJsonSerializer::Deserialize(JsonReader, JsonRoot))
+	{
+		UE_LOG(LogAvaBroadcast, Error, TEXT("Couldn't parse json data from file: %s"), *InFilename);
+		return false;
+	}
+
+	const TSharedPtr<FJsonObject> BroadcastJson = JsonRoot->GetObjectField(TEXT("Broadcast"));
+	if (!BroadcastJson)
+	{
+		return false;
+	}
+
+	using namespace UE::AvaMedia::BroadcastSerialization::Private;
+	
+	const FSoftObjectPath BroadcastPath(OutBroadcast);
+	const FSoftObjectPath OrigBroadcastPath(JsonRoot->GetStringField(TEXT("BroadcastPath")));
+
+	if (BroadcastPath != OrigBroadcastPath)
+	{
+		// In case the root object is not the same, we need to fixup the references.
+		FixupReferences(BroadcastJson, OrigBroadcastPath.ToString(), BroadcastPath.ToString());
+	}
+	
+	TArray<TSharedPtr<FJsonValue>> ReferencedObjects = JsonRoot->GetArrayField(TEXT("ReferencedObjects"));
+
+	// The save traversal produces leaf objects last, so we reverse to get them first for loading. 
+	Algo::Reverse(ReferencedObjects);
+	
+	for (TSharedPtr<FJsonValue> ReferenceObjectValue : ReferencedObjects)
+	{
+		const TSharedPtr<FJsonObject> ReferenceObject = ReferenceObjectValue.Get()->AsObject(); 
+		if (BroadcastPath != OrigBroadcastPath)
+		{
+			FixupReferences(ReferenceObject, OrigBroadcastPath.ToString(), BroadcastPath.ToString());
+		}
+		LoadReferencedObject(OutBroadcast, ReferenceObject);
+	}
+	
+	return FJsonObjectConverter::JsonObjectToUStruct(BroadcastJson.ToSharedRef(), OutBroadcast->GetClass(), OutBroadcast);;
+}
+
+bool FAvaBroadcastSerialization::SaveBroadcastToXml(UAvaBroadcast* InBroadcast, const FString& InFilename)
+{
+	bool bIsBroadcastSaved = false;
+#if WITH_EDITOR
+	const TUniquePtr<FArchive> FileWriter(IFileManager::Get().CreateFileWriter(*InFilename));
+	if (FileWriter.IsValid())	
+	{
+		FXmlArchiveOutputFormatter XmlOutput(*FileWriter);
+		XmlOutput.SerializeObjectsInPlace(true); // We want the media outputs nested.
+		UE::AvaMediaSerializationUtils::SerializeObject(XmlOutput, InBroadcast);
+		bIsBroadcastSaved = XmlOutput.SaveDocumentToInnerArchive();
+		FileWriter->Close();
+	}
+#endif
+	return bIsBroadcastSaved;
+}
+
+bool FAvaBroadcastSerialization::LoadBroadcastFromXml(const FString& InFilename, UAvaBroadcast* OutBroadcast)
+{
+	bool bIsBroadcastLoaded = false;
+#if WITH_EDITOR
+	const TUniquePtr<FArchive> FileReader(IFileManager::Get().CreateFileReader(*InFilename));
+	if (FileReader.IsValid())
+	{
+		FXmlArchiveInputFormatter InputFormatter(*FileReader, OutBroadcast);
+		UE::AvaMediaSerializationUtils::SerializeObject(InputFormatter, OutBroadcast);
+		FileReader->Close();
+		bIsBroadcastLoaded = true;
+	}
+#endif	
+	return bIsBroadcastLoaded;
+}
diff --git a/Engine/Plugins/Experimental/Avalanche/Source/AvalancheMedia/Private/Broadcast/AvaBroadcastSerialization.h b/Engine/Plugins/Experimental/Avalanche/Source/AvalancheMedia/Private/Broadcast/AvaBroadcastSerialization.h
new file mode 100644
--- /dev/null
+++ b/Engine/Plugins/Experimental/Avalanche/Source/AvalancheMedia/Private/Broadcast/AvaBroadcastSerialization.h
@@ -0,0 +1,16 @@
+// Copyright Epic Games, Inc. All Rights Reserved.
+
+#pragma once
+
+#include "Containers/UnrealString.h"
+
+class UAvaBroadcast;
+
+struct FAvaBroadcastSerialization
+{
+	static bool SaveBroadcastToJson(const UAvaBroadcast* InBroadcast, const FString& InFilename);
+	static bool LoadBroadcastFromJson(const FString& InFilename, UAvaBroadcast* OutBroadcast);
+
+	static bool SaveBroadcastToXml(UAvaBroadcast* InBroadcast, const FString& InFilename);
+	static bool LoadBroadcastFromXml(const FString& InFilename, UAvaBroadcast* OutBroadcast);
+};
\ No newline at end of file
diff --git a/Engine/Plugins/Experimental/Avalanche/Source/AvalancheMedia/Private/Rundown/AvaRundownServer.cpp b/Engine/Plugins/Experimental/Avalanche/Source/AvalancheMedia/Private/Rundown/AvaRundownServer.cpp
--- a/Engine/Plugins/Experimental/Avalanche/Source/AvalancheMedia/Private/Rundown/AvaRundownServer.cpp
+++ b/Engine/Plugins/Experimental/Avalanche/Source/AvalancheMedia/Private/Rundown/AvaRundownServer.cpp
@@ -1336,9 +1336,8 @@ void FAvaRundownServer::HandleAddChannelDevice(const FAvaRundownAddChannelDevice
 
 	const FAvaBroadcastMediaOutputInfo OutputInfo;
 	const UMediaOutput* OutputDevice = TreeItem->AddMediaOutputToChannel(OutputChannel.GetChannelName(), OutputInfo);
-#if WITH_EDITOR
 	Broadcast.SaveBroadcast();
-#endif
+
 	LogAndSendMessage(InContext->GetSender(), InMessage.RequestId, ELogVerbosity::Log, TEXT("\"AddChannelDevice\" successfully added device \"%s\""), *OutputDevice->GetFName().ToString());
 }
 
@@ -1370,9 +1369,7 @@ void FAvaRundownServer::HandleEditChannelDevice(const FAvaRundownEditChannelDevi
 
 	FAvaRundownServerMediaOutputUtils::EditMediaOutput(MediaOutput, InMessage.Data);
 	
-#if WITH_EDITOR
 	Broadcast.SaveBroadcast();
-#endif
 	
 	SendMessage(InContext->GetSender(), InMessage.RequestId, ELogVerbosity::Log, TEXT("\"EditChannelDevice\". Successfully edited device \"%s\" on \"%s\""), *InMessage.MediaOutputName, *InMessage.ChannelName); 
 }
@@ -1418,9 +1415,7 @@ void FAvaRundownServer::HandleRemoveChannelDevice(const FAvaRundownRemoveChannel
 		return;
 	}
 
-#if WITH_EDITOR
 	Broadcast.SaveBroadcast();
-#endif
 	
 	LogAndSendMessage(InContext->GetSender(), InMessage.RequestId, ELogVerbosity::Log, TEXT("\"RemoveChannelDevice\" Removed Device \"%s\""), *InMessage.MediaOutputName);
 }
diff --git a/Engine/Plugins/Experimental/Avalanche/Source/AvalancheMedia/Public/Broadcast/AvaBroadcast.h b/Engine/Plugins/Experimental/Avalanche/Source/AvalancheMedia/Public/Broadcast/AvaBroadcast.h
--- a/Engine/Plugins/Experimental/Avalanche/Source/AvalancheMedia/Public/Broadcast/AvaBroadcast.h
+++ b/Engine/Plugins/Experimental/Avalanche/Source/AvalancheMedia/Public/Broadcast/AvaBroadcast.h
@@ -81,13 +81,9 @@ public:
 
 	const FAvaBroadcastProfile& GetCurrentProfile() const { return GetProfile(CurrentProfile); }
 
-#if WITH_EDITOR
 	void LoadBroadcast();
 	void SaveBroadcast();
 	FString GetBroadcastSaveFilepath() const;
-#else
-	void InitBroadcast();
-#endif
 
 	//Rather than Immediately calling the Broadcast Event change, it will queue all the multiple calls and dispatch it once to avoid redundancy
 	void QueueNotifyChange(EAvaBroadcastChange InChange);
@@ -197,6 +193,9 @@ protected:
 	/** Maps a channel name to a profile. Indicating which profile to use for a pinned channel. */
 	UPROPERTY()
 	TMap<FName, FName> PinnedChannels;
+
+	/** Full path to the last loaded config file. */
+	FString LoadedConfigFilepath;
 	
 #if WITH_EDITOR
 	bool bCanShowPreview = true;
