[Motion Design] Temporary solution to fix device enumeration in game build.

The media output class information: device provider and display name are compiled out in the game build.

The device provider name is needed to find the corresponding device provider in the list.
The dislpay name is used to show in the control/config app.

The proposed temporary solution is to produce this information in a registry file when packaging the game.

diff --git a/Engine/Plugins/Experimental/Avalanche/Source/AvalancheMedia/Private/AvaMediaModule.cpp b/Engine/Plugins/Experimental/Avalanche/Source/AvalancheMedia/Private/AvaMediaModule.cpp
--- a/Engine/Plugins/Experimental/Avalanche/Source/AvalancheMedia/Private/AvaMediaModule.cpp
+++ b/Engine/Plugins/Experimental/Avalanche/Source/AvalancheMedia/Private/AvaMediaModule.cpp
@@ -5,6 +5,7 @@
 #include "Application/ThrottleManager.h"
 #include "AvaMediaSettings.h"
 #include "Broadcast/AvaBroadcast.h"
+#include "Broadcast/OutputDevices/AvaBroadcastDeviceProviderRegistry.h"
 #include "Broadcast/OutputDevices/AvaBroadcastRenderTargetMediaUtils.h"
 #include "IMediaIOCoreModule.h"
 #include "Interfaces/IPluginManager.h"
@@ -423,6 +424,11 @@ void FAvaMediaModule::PostEngineInit()
 	const UAvaMediaSettings& Settings = UAvaMediaSettings::Get();
 
 	ManagedInstanceCache = MakeShared<FAvaRundownManagedInstanceCache>();
+
+#if WITH_EDITOR
+	// Initialise the device provider registry.
+	FAvaBroadcastDeviceProviderRegistry::Get();
+#endif
 	
 	// Allow for specification of the server name in the command line.
 	// Command line has priority over project settings. 
diff --git a/Engine/Plugins/Experimental/Avalanche/Source/AvalancheMedia/Private/Broadcast/OutputDevices/AvaBroadcastDeviceProviderRegistry.cpp b/Engine/Plugins/Experimental/Avalanche/Source/AvalancheMedia/Private/Broadcast/OutputDevices/AvaBroadcastDeviceProviderRegistry.cpp
new file mode 100644
--- /dev/null
+++ b/Engine/Plugins/Experimental/Avalanche/Source/AvalancheMedia/Private/Broadcast/OutputDevices/AvaBroadcastDeviceProviderRegistry.cpp
@@ -0,0 +1,139 @@
+// Copyright Epic Games, Inc. All Rights Reserved.
+
+#include "AvaBroadcastDeviceProviderRegistry.h"
+
+#include "Backends/JsonStructDeserializerBackend.h"
+#include "Backends/JsonStructSerializerBackend.h"
+#include "Broadcast/OutputDevices/AvaBroadcastOutputUtils.h"
+#include "HAL/FileManager.h"
+#include "MediaOutput.h"
+#include "Misc/Paths.h"
+#include "StructDeserializer.h"
+#include "StructSerializer.h"
+#include "UObject/UObjectIterator.h"
+
+namespace UE::AvaMedia::BroadcastDeviceProviderRegistry::Private
+{
+	void InitializeData(FAvaBroadcastDeviceProviderRegistryData& OutData);
+	bool SaveData(const FAvaBroadcastDeviceProviderRegistryData& OutData);
+	bool LoadData(FAvaBroadcastDeviceProviderRegistryData& OutData);
+}
+
+const FAvaBroadcastDeviceProviderRegistry& FAvaBroadcastDeviceProviderRegistry::Get()
+{
+	static FAvaBroadcastDeviceProviderRegistry Instance;
+	return Instance;
+}
+
+FAvaBroadcastDeviceProviderRegistry::FAvaBroadcastDeviceProviderRegistry()
+{
+	using namespace UE::AvaMedia::BroadcastDeviceProviderRegistry::Private;
+#if WITH_EDITORONLY_DATA
+	InitializeData(Data);
+	SaveData(Data);
+#else
+	LoadData(Data);
+#endif
+}
+
+bool FAvaBroadcastDeviceProviderRegistry::HasDeviceProviderName(const UClass* InMediaOutputClass) const
+{
+	if (InMediaOutputClass)
+	{
+		return Data.DeviceProviderNames.Contains(InMediaOutputClass->GetFName());
+	}
+	return false;
+}
+
+FName FAvaBroadcastDeviceProviderRegistry::GetDeviceProviderName(const UClass* InMediaOutputClass) const
+{
+	if (InMediaOutputClass)
+	{
+		if (const FName* FoundName = Data.DeviceProviderNames.Find(InMediaOutputClass->GetFName()))
+		{
+			return *FoundName;	
+		}
+	}
+	return FName();
+}
+
+FText FAvaBroadcastDeviceProviderRegistry::GetOutputClassDisplayText(const UClass* InMediaOutputClass) const
+{
+	if (InMediaOutputClass)
+	{
+		if (const FText* FoundText = Data.OutputClassDisplayNames.Find(InMediaOutputClass->GetFName()))
+		{
+			return *FoundText;
+		}
+	}
+	return FText();
+}
+
+
+namespace UE::AvaMedia::BroadcastDeviceProviderRegistry::Private
+{
+	void InitializeData(FAvaBroadcastDeviceProviderRegistryData& OutData)
+	{
+		OutData.DeviceProviderNames.Reset();
+		OutData.OutputClassDisplayNames.Reset();
+		
+#if WITH_EDITORONLY_DATA
+		for (const UClass* const Class : TObjectRange<UClass>())
+		{
+			const bool bIsMediaOutputClass = Class->IsChildOf(UMediaOutput::StaticClass()) && Class != UMediaOutput::StaticClass();
+			if (bIsMediaOutputClass)
+			{
+				OutData.OutputClassDisplayNames.Add(Class->GetFName(), Class->GetDisplayNameText());
+				
+				const FName DeviceProviderName = UE::AvaBroadcastOutputUtils::GetDeviceProviderName(Class);
+				if (!DeviceProviderName.IsNone())
+				{
+					OutData.DeviceProviderNames.Add(Class->GetFName(), DeviceProviderName);
+				}
+			}
+		}
+#endif		
+	}
+
+	static FString GetRegistryFilepath()
+	{
+		const FString RegistryName = TEXT("BroadcastDeviceProviderRegistry.json");
+		return FPaths::ProjectConfigDir() / RegistryName;
+	}
+	
+	bool LoadData(FAvaBroadcastDeviceProviderRegistryData& OutData, FArchive& InArchive)
+	{
+		FJsonStructDeserializerBackend Backend(InArchive);
+		return FStructDeserializer::Deserialize(&OutData, *FAvaBroadcastDeviceProviderRegistryData::StaticStruct(), Backend);
+	}
+
+	bool LoadData(FAvaBroadcastDeviceProviderRegistryData& OutData)
+	{
+		bool bLoaded = false;
+		const TUniquePtr<FArchive> FileReader(IFileManager::Get().CreateFileReader(*GetRegistryFilepath()));
+		if (FileReader)
+		{
+			bLoaded =LoadData(OutData, *FileReader);
+			FileReader->Close();
+		}
+		return bLoaded;
+	}
+	
+	void SaveData(const FAvaBroadcastDeviceProviderRegistryData& OutData, FArchive& InArchive)
+	{
+		FJsonStructSerializerBackend Backend(InArchive, EStructSerializerBackendFlags::Default);
+		FStructSerializer::Serialize(&OutData, *FAvaBroadcastDeviceProviderRegistryData::StaticStruct(), Backend);
+	}
+
+	bool SaveData(const FAvaBroadcastDeviceProviderRegistryData& OutData)
+	{
+		const TUniquePtr<FArchive> FileWriter(IFileManager::Get().CreateFileWriter(*GetRegistryFilepath()));
+		if (FileWriter)
+		{
+			SaveData(OutData, *FileWriter);
+			FileWriter->Close();
+			return true;
+		}
+		return false;
+	}
+}
diff --git a/Engine/Plugins/Experimental/Avalanche/Source/AvalancheMedia/Private/Broadcast/OutputDevices/AvaBroadcastDeviceProviderRegistry.h b/Engine/Plugins/Experimental/Avalanche/Source/AvalancheMedia/Private/Broadcast/OutputDevices/AvaBroadcastDeviceProviderRegistry.h
new file mode 100644
--- /dev/null
+++ b/Engine/Plugins/Experimental/Avalanche/Source/AvalancheMedia/Private/Broadcast/OutputDevices/AvaBroadcastDeviceProviderRegistry.h
@@ -0,0 +1,40 @@
+// Copyright Epic Games, Inc. All Rights Reserved.
+
+#pragma once
+
+#include "Containers/UnrealString.h"
+#include "UObject/NameTypes.h"
+
+#include "AvaBroadcastDeviceProviderRegistry.generated.h"
+
+class UClass;
+
+USTRUCT()
+struct FAvaBroadcastDeviceProviderRegistryData
+{
+	GENERATED_BODY()
+
+	/** Maps output class name to the device provider. */
+	UPROPERTY()
+	TMap<FName, FName> DeviceProviderNames;
+
+	UPROPERTY()
+	TMap<FName, FText> OutputClassDisplayNames;
+};
+
+class FAvaBroadcastDeviceProviderRegistry
+{
+public:
+	static const FAvaBroadcastDeviceProviderRegistry& Get();
+
+	bool HasDeviceProviderName(const UClass* InMediaOutputClass) const;
+	FName GetDeviceProviderName(const UClass* InMediaOutputClass) const;
+	FText GetOutputClassDisplayText(const UClass* InMediaOutputClass) const;
+
+protected:
+	FAvaBroadcastDeviceProviderRegistry();
+	~FAvaBroadcastDeviceProviderRegistry() = default;
+
+private:
+	FAvaBroadcastDeviceProviderRegistryData Data;
+};
\ No newline at end of file
diff --git a/Engine/Plugins/Experimental/Avalanche/Source/AvalancheMedia/Private/Broadcast/OutputDevices/AvaBroadcastOutputClassItem.cpp b/Engine/Plugins/Experimental/Avalanche/Source/AvalancheMedia/Private/Broadcast/OutputDevices/AvaBroadcastOutputClassItem.cpp
--- a/Engine/Plugins/Experimental/Avalanche/Source/AvalancheMedia/Private/Broadcast/OutputDevices/AvaBroadcastOutputClassItem.cpp
+++ b/Engine/Plugins/Experimental/Avalanche/Source/AvalancheMedia/Private/Broadcast/OutputDevices/AvaBroadcastOutputClassItem.cpp
@@ -2,6 +2,7 @@
 
 #include "AvaBroadcastOutputClassItem.h"
 
+#include "AvaBroadcastDeviceProviderRegistry.h"
 #include "AvaBroadcastOutputDeviceItem.h"
 #include "AvaBroadcastOutputServerItem.h"
 #include "Broadcast/AvaBroadcast.h"
@@ -56,7 +57,11 @@ FText FAvaBroadcastOutputClassItem::GetDisplayName() const
 	check(OutputClass.IsValid());
 	return OutputClass->GetDisplayNameText();
 #else
-	return LOCTEXT("DisplayName_NotAvailable", "NotAvailable");	// Not available in game build.
+	if (OutputClass.IsValid())
+	{
+		return FAvaBroadcastDeviceProviderRegistry::Get().GetOutputClassDisplayText(OutputClass.Get());
+	}
+	return LOCTEXT("DisplayName_NotAvailable", "NotAvailable");
 #endif
 }
 
diff --git a/Engine/Plugins/Experimental/Avalanche/Source/AvalancheMedia/Private/Broadcast/OutputDevices/AvaBroadcastOutputUtils.cpp b/Engine/Plugins/Experimental/Avalanche/Source/AvalancheMedia/Private/Broadcast/OutputDevices/AvaBroadcastOutputUtils.cpp
--- a/Engine/Plugins/Experimental/Avalanche/Source/AvalancheMedia/Private/Broadcast/OutputDevices/AvaBroadcastOutputUtils.cpp
+++ b/Engine/Plugins/Experimental/Avalanche/Source/AvalancheMedia/Private/Broadcast/OutputDevices/AvaBroadcastOutputUtils.cpp
@@ -2,6 +2,7 @@
 
 #include "Broadcast/OutputDevices/AvaBroadcastOutputUtils.h"
 
+#include "AvaBroadcastDeviceProviderRegistry.h"
 #include "Broadcast/OutputDevices/AvaGameViewportMediaOutput.h"
 #include "MediaIOCoreDefinitions.h"
 #include "MediaOutput.h"
@@ -48,7 +49,7 @@ bool UE::AvaBroadcastOutputUtils::HasDeviceProviderName(const UClass* InMediaOut
 		&& InMediaOutputClass->IsChildOf(UMediaOutput::StaticClass())
 		&& InMediaOutputClass->HasMetaData(Private::NAME_MediaIOCustomLayout);
 #else
-	return false;
+	return FAvaBroadcastDeviceProviderRegistry::Get().HasDeviceProviderName(InMediaOutputClass);
 #endif
 }
 
@@ -59,22 +60,27 @@ FName UE::AvaBroadcastOutputUtils::GetDeviceProviderName(const UClass* InMediaOu
 	{
 		return FName(InMediaOutputClass->GetMetaData(Private::NAME_MediaIOCustomLayout));
 	}
-#endif
 	return FName();
+#else
+	return FAvaBroadcastDeviceProviderRegistry::Get().GetDeviceProviderName(InMediaOutputClass);
+#endif
 }
 
 FName UE::AvaBroadcastOutputUtils::GetDeviceProviderName(const UMediaOutput* InMediaOutput)
 {
-#if WITH_EDITORONLY_DATA
 	const UClass* const MediaOutputClass = IsValid(InMediaOutput) ? InMediaOutput->GetClass() : nullptr;
-	if (MediaOutputClass
-		&& IsValid(MediaOutputClass)
-		&& MediaOutputClass->HasMetaData(Private::NAME_MediaIOCustomLayout))
+	if (MediaOutputClass && IsValid(MediaOutputClass))
 	{
-		return FName(MediaOutputClass->GetMetaData(Private::NAME_MediaIOCustomLayout));
-	}
+#if WITH_EDITORONLY_DATA
+		if (MediaOutputClass->HasMetaData(Private::NAME_MediaIOCustomLayout))
+		{
+			return FName(MediaOutputClass->GetMetaData(Private::NAME_MediaIOCustomLayout));
+		}
+#else
+		return FAvaBroadcastDeviceProviderRegistry::Get().GetDeviceProviderName(MediaOutputClass);
 #endif
-	return FName();
+	}
+	return FName(); 
 }
 
 bool UE::AvaBroadcastOutputUtils::IsGameViewportOutput(const UMediaOutput* InMediaOutput)
