[Motion Design] Fix Sub-PropertyId in Rundown.

diff --git a/Engine/Plugins/Experimental/Avalanche/Source/AvalancheMedia/Private/Playable/AvaPlayableRemoteControl.cpp b/Engine/Plugins/Experimental/Avalanche/Source/AvalancheMedia/Private/Playable/AvaPlayableRemoteControl.cpp
--- a/Engine/Plugins/Experimental/Avalanche/Source/AvalancheMedia/Private/Playable/AvaPlayableRemoteControl.cpp
+++ b/Engine/Plugins/Experimental/Avalanche/Source/AvalancheMedia/Private/Playable/AvaPlayableRemoteControl.cpp
@@ -291,6 +291,12 @@ bool UE::AvaPlayableRemoteControl::GetEntitiesControlledByController(const URemo
 					if (const URemoteControlPropertyIdRegistry* IdentityRegistry = InRemoteControlPreset->GetPropertyIdRegistry())
 					{
 						OutEntityIds.Append(IdentityRegistry->GetEntityIdsForPropertyId(IdentityAction->PropertyId));
+
+						// Explicit support for Sub-PropertyId
+						for (const TPair<FPropertyIdContainerKey, TObjectPtr<URCVirtualPropertySelfContainer>>& PropertyContainer : IdentityAction->PropertySelfContainer)
+						{
+							OutEntityIds.Append(IdentityRegistry->GetEntityIdsForPropertyId(PropertyContainer.Key.PropertyId));
+						}
 					}
 				}
 				else
