[Motion Design] Ease curve tool
 - Fade out ease curve color and display text when no ease curve (broken cubic tangents) is detected for the key selection
 - Fix crash when switching between motion design modes (introduced with last ease curve CL for 5.5)

diff --git a/Engine/Plugins/Experimental/Avalanche/Source/AvalancheSequencer/Private/EaseCurveTool/AvaEaseCurveKeySelection.cpp b/Engine/Plugins/Experimental/Avalanche/Source/AvalancheSequencer/Private/EaseCurveTool/AvaEaseCurveKeySelection.cpp
--- a/Engine/Plugins/Experimental/Avalanche/Source/AvalancheSequencer/Private/EaseCurveTool/AvaEaseCurveKeySelection.cpp
+++ b/Engine/Plugins/Experimental/Avalanche/Source/AvalancheSequencer/Private/EaseCurveTool/AvaEaseCurveKeySelection.cpp
@@ -1,7 +1,7 @@
 ﻿// Copyright Epic Games, Inc. All Rights Reserved.
 
 #include "AvaEaseCurveKeySelection.h"
-
+#include "AvaSequencer.h"
 #include "Channels/MovieSceneChannel.h"
 #include "Channels/MovieSceneDoubleChannel.h"
 #include "Channels/MovieSceneFloatChannel.h"
@@ -9,18 +9,58 @@
 #include "EaseCurveTool/AvaEaseCurveTool.h"
 #include "IKeyArea.h"
 #include "MVVM/ViewModels/ChannelModel.h"
+#include "MVVM/ViewModels/SequencerEditorViewModel.h"
 
 using namespace UE::Sequencer;
 
-FAvaEaseCurveKeySelection::FAvaEaseCurveKeySelection(const TWeakPtr<FSequencerSelection>& InSequencerSelectionWeak)
-	: SequencerSelectionWeak(InSequencerSelectionWeak)
+namespace UE::Ava::EaseCurveTool::Private
+{
+	template <class InChannelHandle, class InChannelValue>
+	bool IsEaseCurve(const FKeyHandle& InKeyHandle, const FAvaEaseCurveKeySelection::FChannelData& InChannelData)
+	{
+		const TMovieSceneChannelHandle<InChannelHandle> Channel = InChannelData.Channel.Cast<InChannelHandle>();
+		TMovieSceneChannelData<InChannelValue> ChannelData = Channel.Get()->GetData();
+
+		const int32 KeyIndex = ChannelData.GetIndex(InKeyHandle);
+		if (KeyIndex == INDEX_NONE)
+		{
+			return false;
+		}
+
+		const TArrayView<InChannelValue> ChannelValues = ChannelData.GetValues();
+
+		if (!FAvaEaseCurveTangents::IsEaseCurveKey(ChannelValues[KeyIndex]))
+		{
+			return false;
+		}
+
+		return true;
+	}
+}
+
+FAvaEaseCurveKeySelection::FAvaEaseCurveKeySelection(const TSharedPtr<FAvaSequencer>& InSequencer)
 {
-	const TSharedPtr<FSequencerSelection> SequencerSelection = InSequencerSelectionWeak.Pin();
+	using namespace UE::Ava::EaseCurveTool::Private;
+
+	if (!InSequencer.IsValid())
+	{
+		return;
+	}
+
+	TSharedPtr<FSequencerSelection> SequencerSelection;
+
+	if (const TSharedPtr<FSequencerEditorViewModel> SequencerViewModel = InSequencer->GetSequencer()->GetViewModel())
+	{
+		SequencerSelection = SequencerViewModel->GetSelection();
+	}
+
 	if (!SequencerSelection.IsValid())
 	{
 		return;
 	}
 
+	bAreAllEaseCurves = true;
+
 	for (const FKeyHandle Key : SequencerSelection->KeySelection)
 	{
 		if (Key == FKeyHandle::Invalid())
@@ -61,9 +101,21 @@ FAvaEaseCurveKeySelection::FAvaEaseCurveKeySelection(const TWeakPtr<FSequencerSe
 				bIsLastOnlySelectedKey = true;
 			}
 		}
-		else if (SelectedKeyCount > 1)
+
+		const FName ChannelTypeName = Entry.Channel.GetChannelTypeName();
+		if (ChannelTypeName == FMovieSceneDoubleChannel::StaticStruct()->GetFName())
+		{
+			if (!IsEaseCurve<FMovieSceneDoubleChannel, FMovieSceneDoubleValue>(Key, Entry))
+			{
+				bAreAllEaseCurves = false;
+			}
+		}
+		else if (ChannelTypeName == FMovieSceneFloatChannel::StaticStruct()->GetFName())
 		{
-			bAllChannelSingleKeySelections = false;
+			if (!IsEaseCurve<FMovieSceneFloatChannel, FMovieSceneFloatValue>(Key, Entry))
+			{
+				bAreAllEaseCurves = false;
+			}
 		}
 	}
 }
@@ -93,18 +145,18 @@ void FAvaEaseCurveKeySelection::ForEachEaseableKey(const bool bInIncludeEqualVal
 
 namespace UE::Ava::EaseCurveTool::Private
 {
-	template <class ChannelHandle, class ChannelValue>
-	void UpdateChannelValues(const FKeyHandle& InKeyHandle, const FKeyHandle& InNextKeyHandle
+	template <class InChannelHandle, class InChannelValue>
+	void NormalizeChannelValues(const FKeyHandle& InKeyHandle, const FKeyHandle& InNextKeyHandle
 		, const FAvaEaseCurveKeySelection::FChannelData& InChannelData
 		, const bool bInAutoFlipTangents
 		, const FFrameRate& InDisplayRate, const FFrameRate& InTickResolution
 		, TArray<FAvaEaseCurveTangents>& OutKeySetTangents
 		, TArray<FAvaEaseCurveTangents>& OutChangingTangents)
 	{
-		TMovieSceneChannelHandle<ChannelHandle> Channel = InChannelData.Channel.Cast<ChannelHandle>();
-		TMovieSceneChannelData<ChannelValue> ChannelData = Channel.Get()->GetData();
+		TMovieSceneChannelHandle<InChannelHandle> Channel = InChannelData.Channel.Cast<InChannelHandle>();
+		TMovieSceneChannelData<InChannelValue> ChannelData = Channel.Get()->GetData();
 
-		const TArrayView<ChannelValue> ChannelValues = ChannelData.GetValues();
+		const TArrayView<InChannelValue> ChannelValues = ChannelData.GetValues();
 		const TArrayView<FFrameNumber> ChannelTimes = ChannelData.GetTimes();
 
 		const int32 KeyIndex = ChannelData.GetIndex(InKeyHandle);
@@ -141,7 +193,7 @@ FAvaEaseCurveTangents FAvaEaseCurveKeySelection::AverageTangents(const FFrameRat
 	, const bool bInAutoFlipTangents)
 {
 	using namespace UE::Ava::EaseCurveTool::Private;
-	
+
 	TArray<FAvaEaseCurveTangents> KeySetTangents;
 	TArray<FAvaEaseCurveTangents> ChangingTangents;
 
@@ -151,24 +203,24 @@ FAvaEaseCurveTangents FAvaEaseCurveKeySelection::AverageTangents(const FFrameRat
 			const FName ChannelTypeName = InChannelData.Channel.GetChannelTypeName();
 			if (ChannelTypeName == FMovieSceneDoubleChannel::StaticStruct()->GetFName())
 			{
-				UpdateChannelValues<FMovieSceneDoubleChannel, FMovieSceneDoubleValue>(InKeyHandle, InNextKeyHandle, InChannelData
+				NormalizeChannelValues<FMovieSceneDoubleChannel, FMovieSceneDoubleValue>(InKeyHandle, InNextKeyHandle, InChannelData
 					, bInAutoFlipTangents, InDisplayRate, InTickResolution, KeySetTangents, ChangingTangents);
 			}
 			else if (ChannelTypeName == FMovieSceneFloatChannel::StaticStruct()->GetFName())
 			{
-				UpdateChannelValues<FMovieSceneFloatChannel, FMovieSceneFloatValue>(InKeyHandle, InNextKeyHandle, InChannelData
+				NormalizeChannelValues<FMovieSceneFloatChannel, FMovieSceneFloatValue>(InKeyHandle, InNextKeyHandle, InChannelData
 					, bInAutoFlipTangents, InDisplayRate, InTickResolution, KeySetTangents, ChangingTangents);
 			}
 
 			return true;
 		});
-	
+
 	return FAvaEaseCurveTangents::Average(ChangingTangents);
 }
 
 namespace UE::Ava::EaseCurveTool::Private
 {
-	template <class ChannelHandle, class ChannelValue>
+	template <class InChannelHandle, class InChannelValue>
 	void SetChannelValues(const FAvaEaseCurveTangents& InTangents
 		, const EAvaEaseCurveToolOperation InOperation
 		, const FKeyHandle& InKeyHandle, const FKeyHandle& InNextKeyHandle
@@ -176,10 +228,10 @@ namespace UE::Ava::EaseCurveTool::Private
 		, const bool bInAutoFlipTangents
 		, const FFrameRate& InDisplayRate, const FFrameRate& InTickResolution)
 	{
-		TMovieSceneChannelHandle<ChannelHandle> Channel = InChannelData.Channel.Cast<ChannelHandle>();
-		TMovieSceneChannelData<ChannelValue> ChannelData = Channel.Get()->GetData();
+		TMovieSceneChannelHandle<InChannelHandle> Channel = InChannelData.Channel.Cast<InChannelHandle>();
+		TMovieSceneChannelData<InChannelValue> ChannelData = Channel.Get()->GetData();
 
-		TArrayView<ChannelValue> ChannelValues = ChannelData.GetValues();
+		TArrayView<InChannelValue> ChannelValues = ChannelData.GetValues();
 		const TArrayView<FFrameNumber> ChannelTimes = ChannelData.GetTimes();
 
 		const int32 KeyIndex = ChannelData.GetIndex(InKeyHandle);
diff --git a/Engine/Plugins/Experimental/Avalanche/Source/AvalancheSequencer/Private/EaseCurveTool/AvaEaseCurveKeySelection.h b/Engine/Plugins/Experimental/Avalanche/Source/AvalancheSequencer/Private/EaseCurveTool/AvaEaseCurveKeySelection.h
--- a/Engine/Plugins/Experimental/Avalanche/Source/AvalancheSequencer/Private/EaseCurveTool/AvaEaseCurveKeySelection.h
+++ b/Engine/Plugins/Experimental/Avalanche/Source/AvalancheSequencer/Private/EaseCurveTool/AvaEaseCurveKeySelection.h
@@ -6,10 +6,11 @@
 #include "Channels/MovieSceneChannelHandle.h"
 #include "MVVM/Selection/Selection.h"
 
-struct FAvaEaseCurveTangents;
+class FAvaSequencer;
 class UMovieSceneSection;
-struct FKeyHandle;
+struct FAvaEaseCurveTangents;
 enum class EAvaEaseCurveToolOperation : uint8;
+struct FKeyHandle;
 
 namespace UE::Sequencer
 {
@@ -18,6 +19,7 @@ namespace UE::Sequencer
 
 struct FAvaEaseCurveKeySelection
 {
+public:
 	struct FChannelData
 	{
 		TSharedPtr<UE::Sequencer::FChannelModel> ChannelModel;
@@ -27,7 +29,7 @@ struct FAvaEaseCurveKeySelection
 	};
 
 	FAvaEaseCurveKeySelection() {}
-	FAvaEaseCurveKeySelection(const TWeakPtr<UE::Sequencer::FSequencerSelection>& InSequencerSelectionWeak);
+	FAvaEaseCurveKeySelection(const TSharedPtr<FAvaSequencer>& InSequencer);
 
 	void ForEachEaseableKey(const bool bInIncludeEqualValueKeys
 		, const TFunctionRef<bool(const FKeyHandle& /*InKeyHandle*/, const FKeyHandle& /*InNextKeyHandle*/, const FChannelData&)>& InCallable);
@@ -100,15 +102,34 @@ struct FAvaEaseCurveKeySelection
 		, const FFrameRate& InTickResolution
 		, const bool bInAutoFlipTangents);
 
-	TWeakPtr<UE::Sequencer::FSequencerSelection> SequencerSelectionWeak;
+	const TMap<FName, FChannelData>& GetChannelKeyData() const
+	{
+		return ChannelKeyData;
+	}
+
+	int32 GetTotalSelectedKeys() const
+	{
+		return TotalSelectedKeys;
+	}
 
+	/** @return True if there is only one selected key and it is the last key of the channel. */
+	bool IsLastOnlySelectedKey() const
+	{
+		return bIsLastOnlySelectedKey;
+	}
+
+	/** @return True if all selected keys are detected to be ease curves (broken, weighted, cubic tangents). */
+	bool AreAllEaseCurves() const
+	{
+		return bAreAllEaseCurves;
+	}
+
+protected:
 	TMap<FName, FChannelData> ChannelKeyData;
 
 	int32 TotalSelectedKeys = 0;
 
-	/** Indicates only one selected key on each of the selected channels. */
-	bool bAllChannelSingleKeySelections = true;
-
-	/** Indicates only one key selected and it is the last key of the channel. */
 	bool bIsLastOnlySelectedKey = false;
+
+	bool bAreAllEaseCurves = true;
 };
diff --git a/Engine/Plugins/Experimental/Avalanche/Source/AvalancheSequencer/Private/EaseCurveTool/AvaEaseCurveTangents.cpp b/Engine/Plugins/Experimental/Avalanche/Source/AvalancheSequencer/Private/EaseCurveTool/AvaEaseCurveTangents.cpp
--- a/Engine/Plugins/Experimental/Avalanche/Source/AvalancheSequencer/Private/EaseCurveTool/AvaEaseCurveTangents.cpp
+++ b/Engine/Plugins/Experimental/Avalanche/Source/AvalancheSequencer/Private/EaseCurveTool/AvaEaseCurveTangents.cpp
@@ -311,4 +311,25 @@ double FAvaEaseCurveTangents::CalculateCurveLength(const int32 SampleCount) cons
 	return Length;
 }
 
+bool FAvaEaseCurveTangents::IsEaseCurveKey(const FRichCurveKey& InKey)
+{
+	return InKey.TangentMode == RCTM_Break
+		&& InKey.TangentWeightMode == RCTWM_WeightedBoth
+		&& InKey.InterpMode == RCIM_Cubic;
+}
+
+bool FAvaEaseCurveTangents::IsEaseCurveKey(const FMovieSceneDoubleValue& InValue)
+{
+	return InValue.TangentMode == RCTM_Break
+		&& InValue.Tangent.TangentWeightMode == RCTWM_WeightedBoth
+		&& InValue.InterpMode == RCIM_Cubic;
+}
+
+bool FAvaEaseCurveTangents::IsEaseCurveKey(const FMovieSceneFloatValue& InValue)
+{
+	return InValue.TangentMode == RCTM_Break
+		&& InValue.Tangent.TangentWeightMode == RCTWM_WeightedBoth
+		&& InValue.InterpMode == RCIM_Cubic;
+}
+
 #undef LOCTEXT_NAMESPACE
diff --git a/Engine/Plugins/Experimental/Avalanche/Source/AvalancheSequencer/Private/EaseCurveTool/AvaEaseCurveTangents.h b/Engine/Plugins/Experimental/Avalanche/Source/AvalancheSequencer/Private/EaseCurveTool/AvaEaseCurveTangents.h
--- a/Engine/Plugins/Experimental/Avalanche/Source/AvalancheSequencer/Private/EaseCurveTool/AvaEaseCurveTangents.h
+++ b/Engine/Plugins/Experimental/Avalanche/Source/AvalancheSequencer/Private/EaseCurveTool/AvaEaseCurveTangents.h
@@ -26,6 +26,10 @@ public:
 
 	static FAvaEaseCurveTangents Average(const TArray<FAvaEaseCurveTangents>& InTangentArray);
 
+	static bool IsEaseCurveKey(const FRichCurveKey& InKey);
+	static bool IsEaseCurveKey(const FMovieSceneDoubleValue& InValue);
+	static bool IsEaseCurveKey(const FMovieSceneFloatValue& InValue);
+
 	FAvaEaseCurveTangents() {}
 	FAvaEaseCurveTangents(const double InStart, const double InStartWeight, const double InEnd, const double InEndWeight)
 		: Start(InStart), StartWeight(InStartWeight)
diff --git a/Engine/Plugins/Experimental/Avalanche/Source/AvalancheSequencer/Private/EaseCurveTool/AvaEaseCurveTool.cpp b/Engine/Plugins/Experimental/Avalanche/Source/AvalancheSequencer/Private/EaseCurveTool/AvaEaseCurveTool.cpp
--- a/Engine/Plugins/Experimental/Avalanche/Source/AvalancheSequencer/Private/EaseCurveTool/AvaEaseCurveTool.cpp
+++ b/Engine/Plugins/Experimental/Avalanche/Source/AvalancheSequencer/Private/EaseCurveTool/AvaEaseCurveTool.cpp
@@ -8,6 +8,7 @@
 #include "EaseCurveTool/AvaEaseCurveKeySelection.h"
 #include "EaseCurveTool/AvaEaseCurveToolSettings.h"
 #include "EaseCurveTool/Widgets/SAvaEaseCurveTool.h"
+#include "Editor.h"
 #include "EngineAnalytics.h"
 #include "Factories/CurveFactory.h"
 #include "Framework/Notifications/NotificationManager.h"
@@ -15,7 +16,6 @@
 #include "ISequencer.h"
 #include "ISettingsModule.h"
 #include "Math/UnrealMathUtility.h"
-#include "MVVM/ViewModels/SequencerEditorViewModel.h"
 #include "Settings/AvaSequencerSettings.h"
 #include "Widgets/Notifications/SNotificationList.h"
 
@@ -28,16 +28,9 @@ FAvaEaseCurveTool::FAvaEaseCurveTool(const TSharedRef<FAvaSequencer>& InSequence
 {
 	EaseCurve = NewObject<UAvaEaseCurve>(GetTransientPackage(), NAME_None, RF_Transient | RF_Transactional);
 
-	const TSharedRef<ISequencer> Sequencer = InSequencer->GetSequencer();
-
-	if (const TSharedPtr<FSequencerEditorViewModel> SequencerViewModel = Sequencer->GetViewModel())
-	{
-		SequencerSelectionWeak = SequencerViewModel->GetSelection();
-	}
-
-	Sequencer->GetSelectionChangedObjectGuids().AddRaw(this, &FAvaEaseCurveTool::OnSequencerSelectionChanged);
-
 	UpdateEaseCurveFromSequencerKeySelections();
+
+	InSequencer->GetSequencer()->GetSelectionChangedObjectGuids().AddRaw(this, &FAvaEaseCurveTool::OnSequencerSelectionChanged);
 }
 
 FAvaEaseCurveTool::~FAvaEaseCurveTool()
@@ -80,7 +73,7 @@ TSharedRef<SWidget> FAvaEaseCurveTool::GenerateWidget()
 
 EVisibility FAvaEaseCurveTool::GetVisibility() const
 {
-	return (KeyCache.TotalSelectedKeys > 0 && !KeyCache.bIsLastOnlySelectedKey) ? EVisibility::Visible : EVisibility::Collapsed;
+	return (KeyCache.GetTotalSelectedKeys() > 0 && !KeyCache.IsLastOnlySelectedKey()) ? EVisibility::Visible : EVisibility::Collapsed;
 }
 
 TObjectPtr<UAvaEaseCurve> FAvaEaseCurveTool::GetToolCurve() const
@@ -252,13 +245,13 @@ void FAvaEaseCurveTool::ApplyQuickEaseToSequencerKeySelections(const EAvaEaseCur
 
 void FAvaEaseCurveTool::SetSequencerKeySelectionTangents(const FAvaEaseCurveTangents& InTangents, const EAvaEaseCurveToolOperation InOperation)
 {
-	KeyCache = FAvaEaseCurveKeySelection(SequencerSelectionWeak);
+	KeyCache = FAvaEaseCurveKeySelection(AvaSequencerWeak.Pin());
 
-	if (KeyCache.TotalSelectedKeys == 0)
+	if (KeyCache.GetTotalSelectedKeys() == 0)
 	{
 		return;
 	}
-	
+
 	const UAvaEaseCurveToolSettings* const EaseCurveToolSettings = GetDefault<UAvaEaseCurveToolSettings>();
 	check(EaseCurveToolSettings);
 	const bool bAutoFlipTangents = EaseCurveToolSettings->GetAutoFlipTangents();
@@ -268,12 +261,12 @@ void FAvaEaseCurveTool::SetSequencerKeySelectionTangents(const FAvaEaseCurveTang
 
 void FAvaEaseCurveTool::UpdateEaseCurveFromSequencerKeySelections()
 {
-	KeyCache = FAvaEaseCurveKeySelection(SequencerSelectionWeak);
-	
+	KeyCache = FAvaEaseCurveKeySelection(AvaSequencerWeak.Pin());
+
 	const UAvaEaseCurveToolSettings* const EaseCurveToolSettings = GetDefault<UAvaEaseCurveToolSettings>();
 	check(EaseCurveToolSettings);
 	const bool bAutoFlipTangents = EaseCurveToolSettings->GetAutoFlipTangents();
-	
+
 	const FAvaEaseCurveTangents AverageTangents = KeyCache.AverageTangents(GetDisplayRate(), GetTickResolution(), bAutoFlipTangents);
 
 	SetEaseCurveTangents(AverageTangents, EAvaEaseCurveToolOperation::InOut, /*bInBroadcastUpdate=*/true, false);
@@ -522,4 +515,9 @@ bool FAvaEaseCurveTool::HasCachedKeysToEase()
 	return bEaseableKeys;
 }
 
+bool FAvaEaseCurveTool::AreAllEaseCurves()
+{
+	return KeyCache.AreAllEaseCurves();
+}
+
 #undef LOCTEXT_NAMESPACE
diff --git a/Engine/Plugins/Experimental/Avalanche/Source/AvalancheSequencer/Private/EaseCurveTool/AvaEaseCurveTool.h b/Engine/Plugins/Experimental/Avalanche/Source/AvalancheSequencer/Private/EaseCurveTool/AvaEaseCurveTool.h
--- a/Engine/Plugins/Experimental/Avalanche/Source/AvalancheSequencer/Private/EaseCurveTool/AvaEaseCurveTool.h
+++ b/Engine/Plugins/Experimental/Avalanche/Source/AvalancheSequencer/Private/EaseCurveTool/AvaEaseCurveTool.h
@@ -122,6 +122,8 @@ public:
 
 	bool HasCachedKeysToEase();
 
+	bool AreAllEaseCurves();
+
 	//~ Begin FGCObject
 	virtual void AddReferencedObjects(FReferenceCollector& Collector) override;
 	virtual FString GetReferencerName() const override;
@@ -140,7 +142,6 @@ protected:
 	void SetEaseCurveTangents_Internal(const FAvaEaseCurveTangents& InTangents, const EAvaEaseCurveToolOperation InOperation, const bool bInBroadcastUpdate) const;
 
 	TWeakPtr<FAvaSequencer> AvaSequencerWeak;
-	TWeakPtr<UE::Sequencer::FSequencerSelection> SequencerSelectionWeak;
 
 	TObjectPtr<UAvaEaseCurve> EaseCurve;
 
@@ -148,7 +149,7 @@ protected:
 	EAvaEaseCurveToolOperation OperationMode = EAvaEaseCurveToolOperation::InOut;
 
 	TSharedPtr<SAvaEaseCurveTool> ToolWidget;
-	
+
 	/** Cached data set when a new sequencer selection is made. */
 	FAvaEaseCurveKeySelection KeyCache;
 
diff --git a/Engine/Plugins/Experimental/Avalanche/Source/AvalancheSequencer/Private/EaseCurveTool/Widgets/SAvaEaseCurveEditor.cpp b/Engine/Plugins/Experimental/Avalanche/Source/AvalancheSequencer/Private/EaseCurveTool/Widgets/SAvaEaseCurveEditor.cpp
--- a/Engine/Plugins/Experimental/Avalanche/Source/AvalancheSequencer/Private/EaseCurveTool/Widgets/SAvaEaseCurveEditor.cpp
+++ b/Engine/Plugins/Experimental/Avalanche/Source/AvalancheSequencer/Private/EaseCurveTool/Widgets/SAvaEaseCurveEditor.cpp
@@ -38,7 +38,7 @@ namespace UE::EaseCurveTool::Private
 	static const FVector2D KeyHitSize = FVector2D(8.f);
 	static const FVector2D TangentHitSize = FVector2D(20.f);
 	static const FVector2D TangentDrawSize = FVector2D(16.f);
-	static const FVector2D CurveHitSize = FVector2D(6.f);
+	static const FVector2D CurveHitSize = FVector2D(8.f);
 
 	static const FVector2D HalfKeyHitSize = KeyHitSize * 0.5f;
 	static const FVector2D HalfTangentHitSize = TangentHitSize * 0.5f;
@@ -90,6 +90,7 @@ void SAvaEaseCurveEditor::Construct(const FArguments& InArgs, const TObjectPtr<U
 	CurveColor = InArgs._CurveColor;
 	Operation = InArgs._Operation;
 	ShowEqualValueKeyError = InArgs._ShowEqualValueKeyError;
+	IsEaseCurveSelection = InArgs._IsEaseCurveSelection;
 
 	StartText = InArgs._StartText;
 	StartTooltipText = InArgs._StartTooltipText;
@@ -179,7 +180,28 @@ int32 SAvaEaseCurveEditor::OnPaint(const FPaintArgs& InArgs, const FGeometry& In
 			WhiteBrush, DrawEffects, NormalAreaColor);
 	}
 
+	auto PaintCenteredText = [&InAllottedGeometry, &OutDrawElements, &InLayerId, DrawEffects, &ScaleInfo](const FText& InText)
+	{
+		const FSlateFontInfo FontInfo = FAvaEaseCurveStyle::Get().GetFontStyle(TEXT("Editor.ErrorFont"));
+
+		const TSharedRef<FSlateFontMeasure> FontMeasure = FSlateApplication::Get().GetRenderer()->GetFontMeasureService();
+		const FVector2D TextSize = FontMeasure->Measure(InText, FontInfo);
+
+		const float HalfInputX = ScaleInfo.InputToLocalX(0.5f);
+		const float HalfOutputY = ScaleInfo.OutputToLocalY(0.5f);
+
+		const FVector2D ActualOffset(HalfInputX - (TextSize.X * 0.5f), HalfOutputY - (TextSize.Y * 0.5f));
+
+		FSlateDrawElement::MakeText(OutDrawElements, ++InLayerId
+			, InAllottedGeometry.ToPaintGeometry(InAllottedGeometry.Size, FSlateLayoutTransform(ActualOffset))
+			, InText, FontInfo, DrawEffects, FStyleColors::Foreground.GetSpecifiedColor());
+	};
+
 	if (ShowEqualValueKeyError.Get(false))
+	{
+		PaintCenteredText(LOCTEXT("EqualValueKeys", "No different key values!"));
+	}
+	else
 	{
 		InLayerId = PaintGrid(ScaleInfo, InAllottedGeometry, OutDrawElements, ++InLayerId, InMyCullingRect, DrawEffects);
 
@@ -218,6 +240,11 @@ int32 SAvaEaseCurveEditor::OnPaint(const FPaintArgs& InArgs, const FGeometry& In
 
 		InLayerId = PaintKeys(ScaleInfo, InAllottedGeometry, OutDrawElements, ++InLayerId, InMyCullingRect, DrawEffects, InWidgetStyle);
 
+		if (!IsEaseCurveSelection.Get(false))
+		{
+			PaintCenteredText(LOCTEXT("NoEaseCurve", "No weighted, broken, cubic tangents!"));
+		}
+
 		if (DragState == EDragState::MarqueeSelect)
 		{
 			const FVector2D MarqueTopLeft(FMath::Min(MouseDownLocation.X, MouseMoveLocation.X), FMath::Min(MouseDownLocation.Y, MouseMoveLocation.Y));
@@ -228,23 +255,6 @@ int32 SAvaEaseCurveEditor::OnPaint(const FPaintArgs& InArgs, const FGeometry& In
 				, FAppStyle::GetBrush(TEXT("MarqueeSelection")));
 		}
 	}
-	else
-	{
-		const FSlateFontInfo FontInfo = FAvaEaseCurveStyle::Get().GetFontStyle(TEXT("Editor.ErrorFont"));
-		const FText Text = LOCTEXT("EqualValueKeys", "No different key values to create ease curve!");
-
-		const TSharedRef<FSlateFontMeasure> FontMeasure = FSlateApplication::Get().GetRenderer()->GetFontMeasureService();
-		const FVector2D TextSize = FontMeasure->Measure(Text, FontInfo);
-
-		const float HalfInputX = ScaleInfo.InputToLocalX(0.5f);
-		const float HalfOutputY = ScaleInfo.OutputToLocalY(0.5f);
-
-		const FVector2D ActualOffset(HalfInputX - (TextSize.X * 0.5f), HalfOutputY - (TextSize.Y * 0.5f));
-
-		FSlateDrawElement::MakeText(OutDrawElements, ++InLayerId
-			, InAllottedGeometry.ToPaintGeometry(InAllottedGeometry.Size, FSlateLayoutTransform(ActualOffset))
-			, Text, FontInfo, DrawEffects, FStyleColors::Foreground.GetSpecifiedColor());
-	}
 
 	return InLayerId;
 }
@@ -381,9 +391,8 @@ void SAvaEaseCurveEditor::PaintCurve(const FTrackScaleInfo& InScaleInfo, const F
 	{
 		return;
 	}
-	
-	TArray<FVector2D> LinePoints;
-	TArray<FLinearColor> LineColors;
+
+	// Gather key handles and time/value pairs
 	TArray<FKeyHandle> KeyHandles;
 	TArray<TPair<float, float>> Key_TimeValuePairs;
 	
@@ -398,46 +407,64 @@ void SAvaEaseCurveEditor::PaintCurve(const FTrackScaleInfo& InScaleInfo, const F
 		Key_TimeValuePairs.Emplace(EaseCurve->FloatCurve.GetKeyTimeValuePair(KeyHandle));
 	}
 
-	// Add enclosed segments
+	// Draw curve lines
+	TArray<FVector2D> LinePoints;
+	TArray<FLinearColor> LineColors;
+
 	for (int32 Index = 0; Index < NumKeys - 1; ++Index)
 	{
-		CreateLinesForSegment(EaseCurve->FloatCurve.GetKeyInterpMode(KeyHandles[Index])
+		CreateLinesForSegment(EaseCurve->FloatCurve.GetKeyInterpMode(KeyHandles[Index]), EaseCurve->FloatCurve.GetKeyTangentMode(KeyHandles[Index])
 			, Key_TimeValuePairs[Index], Key_TimeValuePairs[Index + 1], LinePoints, LineColors, InScaleInfo);
 
 		FSlateDrawElement::MakeLines(OutDrawElements, LayerId, InAllottedGeometry.ToPaintGeometry()
 			, LinePoints, LineColors, DrawEffects, FLinearColor::White/*Color*/, true, CurveThickness * InAllottedGeometry.Scale);
 
-		LinePoints.Empty();
-		LineColors.Empty();
+		LinePoints.Reset();
+		LineColors.Reset();
 	}
 }
 
 void SAvaEaseCurveEditor::CreateLinesForSegment(const ERichCurveInterpMode InInterpMode
+	, const ERichCurveTangentMode InTangentMode
 	, const TPair<float, float>& InStartKeyTimeValue, const TPair<float, float>& InEndKeyTimeValue
 	, TArray<FVector2D>& OutLinePoints, TArray<FLinearColor>& OutLineColors
 	, const FTrackScaleInfo& InScaleInfo) const
 {
+	static constexpr FLinearColor DisabledFadeFactor(0.03f, 0.03f, 0.03f, 1.f);
+
+	const bool bIsEaseCurve = IsEaseCurveSelection.Get(false);
+
 	FLinearColor FadedCurveColor = CurveColor;
 	FadedCurveColor.A = 0.05f;
 
+	auto AddSegment = [this, &OutLinePoints, &OutLineColors, &bIsEaseCurve](const FVector2D& InPoint, FLinearColor InColor)
+	{
+		OutLinePoints.Add(InPoint);
+		OutLineColors.Add(bIsEaseCurve ? MoveTemp(InColor) : (CurveColor * DisabledFadeFactor));
+	};
+
 	switch (InInterpMode)
 	{
 	case RCIM_Constant:
 	{
 		//@todo: should really only need 3 points here but something about the line rendering isn't quite behaving as I'd expect, so need extras
-		OutLinePoints.Add(FVector2D(InStartKeyTimeValue.Key, InStartKeyTimeValue.Value));
-		OutLinePoints.Add(FVector2D(InEndKeyTimeValue.Key, InStartKeyTimeValue.Value));
-		OutLinePoints.Add(FVector2D(InEndKeyTimeValue.Key, InStartKeyTimeValue.Value));
-		OutLinePoints.Add(FVector2D(InEndKeyTimeValue.Key, InEndKeyTimeValue.Value));
-		OutLinePoints.Add(FVector2D(InEndKeyTimeValue.Key, InStartKeyTimeValue.Value));
+		AddSegment(FVector2D(InStartKeyTimeValue.Key, InStartKeyTimeValue.Value), CurveColor);
+		AddSegment(FVector2D(InEndKeyTimeValue.Key, InStartKeyTimeValue.Value), CurveColor);
+		AddSegment(FVector2D(InEndKeyTimeValue.Key, InStartKeyTimeValue.Value), CurveColor);
+		AddSegment(FVector2D(InEndKeyTimeValue.Key, InEndKeyTimeValue.Value), CurveColor);
+		AddSegment(FVector2D(InEndKeyTimeValue.Key, InStartKeyTimeValue.Value), CurveColor);
+
 		break;
 	}
+
 	case RCIM_Linear:
 	{
-		OutLinePoints.Add(FVector2D(InStartKeyTimeValue.Key, InStartKeyTimeValue.Value));
-		OutLinePoints.Add(FVector2D(InEndKeyTimeValue.Key, InEndKeyTimeValue.Value));
+		AddSegment(FVector2D(InStartKeyTimeValue.Key, InStartKeyTimeValue.Value), CurveColor);
+		AddSegment(FVector2D(InEndKeyTimeValue.Key, InEndKeyTimeValue.Value), CurveColor);
+
 		break;
 	}
+
 	case RCIM_Cubic:
 	{
 		// Clamp to screen to avoid massive slowdown when zoomed in
@@ -449,35 +476,38 @@ void SAvaEaseCurveEditor::CreateLinesForSegment(const ERichCurveInterpMode InInt
 
 		for (float CurrentX = StartX; CurrentX < EndX; CurrentX += StepSize)
 		{
-			// Add line point
-			const float CurveIn = InScaleInfo.LocalXToInput(FMath::Min(CurrentX, EndX));
-			const float CurveOut = EaseCurve->FloatCurve.Eval(CurveIn);
-			OutLinePoints.Add(FVector2D(CurveIn, CurveOut));
-			
+			// Get color
+			FLinearColor Color = CurveColor;
 			switch (Operation.Get(EAvaEaseCurveToolOperation::InOut))
 			{
 			case EAvaEaseCurveToolOperation::Out:
-			{
-				const float Alpha = CurrentX / CurveLengthX;
-				OutLineColors.Add(FLinearColor::LerpUsingHSV(FadedCurveColor, CurveColor,  1.f - Alpha));
-				break;
-			}
+				{
+					const float Alpha = CurrentX / CurveLengthX;
+					Color = FLinearColor::LerpUsingHSV(FadedCurveColor, CurveColor,  1.f - Alpha);
+					break;
+				}
 			case EAvaEaseCurveToolOperation::In:
-			{
-				const float Alpha = CurrentX / CurveLengthX;
-				OutLineColors.Add(FLinearColor::LerpUsingHSV(FadedCurveColor, CurveColor, Alpha));
-				break;
-			}
+				{
+					const float Alpha = CurrentX / CurveLengthX;
+					Color = FLinearColor::LerpUsingHSV(FadedCurveColor, CurveColor, Alpha);
+					break;
+				}
 			case EAvaEaseCurveToolOperation::InOut:
 			default:
-			{
-				OutLineColors.Add(CurveColor);
-				break;
-			}
+				{
+					Color = CurveColor;
+					break;
+				}
 			}
+			
+			// Add segment
+			const float CurveIn = InScaleInfo.LocalXToInput(FMath::Min(CurrentX, EndX));
+			const float CurveOut = EaseCurve->FloatCurve.Eval(CurveIn);
+			AddSegment(FVector2D(CurveIn, CurveOut), Color);
 		}
-		OutLinePoints.Add(FVector2D(InEndKeyTimeValue.Key, InEndKeyTimeValue.Value));
-		OutLineColors.Add(CurveColor);
+
+		AddSegment(FVector2D(InEndKeyTimeValue.Key, InEndKeyTimeValue.Value), CurveColor);
+
 		break;
 	}
 	default:
diff --git a/Engine/Plugins/Experimental/Avalanche/Source/AvalancheSequencer/Private/EaseCurveTool/Widgets/SAvaEaseCurveEditor.h b/Engine/Plugins/Experimental/Avalanche/Source/AvalancheSequencer/Private/EaseCurveTool/Widgets/SAvaEaseCurveEditor.h
--- a/Engine/Plugins/Experimental/Avalanche/Source/AvalancheSequencer/Private/EaseCurveTool/Widgets/SAvaEaseCurveEditor.h
+++ b/Engine/Plugins/Experimental/Avalanche/Source/AvalancheSequencer/Private/EaseCurveTool/Widgets/SAvaEaseCurveEditor.h
@@ -46,6 +46,8 @@ public:
 		, _ExtendedGridColor(FLinearColor(0.02f, 0.02f, 0.02f, 0.5f))
 		, _CurveThickness(2.5f)
 		, _CurveColor(FLinearColor::White)
+		, _ShowEqualValueKeyError(false)
+		, _IsEaseCurveSelection(true)
 	{
 		_Clipping = EWidgetClipping::ClipToBounds;
 	}
@@ -68,6 +70,7 @@ public:
 		SLATE_ARGUMENT(FLinearColor, CurveColor)
 		SLATE_ATTRIBUTE(EAvaEaseCurveToolOperation, Operation)
 		SLATE_ATTRIBUTE(bool, ShowEqualValueKeyError)
+		SLATE_ATTRIBUTE(bool, IsEaseCurveSelection)
 
 		SLATE_ATTRIBUTE(FText, StartText)
 		SLATE_ATTRIBUTE(FText, StartTooltipText)
@@ -200,6 +203,7 @@ protected:
 
 	/* Generates the line(s) for rendering between KeyIndex and the following key. */
 	void CreateLinesForSegment(const ERichCurveInterpMode InInterpMode
+		, const ERichCurveTangentMode InTangentMode
 		, const TPair<float, float>& InStartKeyTimeValue, const TPair<float, float>& InEndKeyTimeValue
 		, TArray<FVector2D>& OutLinePoints, TArray<FLinearColor>& OutLineColors
 		, const FTrackScaleInfo& InScaleInfo) const;
@@ -248,6 +252,7 @@ protected:
 	FLinearColor CurveColor;
 	TAttribute<EAvaEaseCurveToolOperation> Operation;
 	TAttribute<bool> ShowEqualValueKeyError;
+	TAttribute<bool> IsEaseCurveSelection;
 
 	TAttribute<FText> StartText;
 	TAttribute<FText> StartTooltipText;
diff --git a/Engine/Plugins/Experimental/Avalanche/Source/AvalancheSequencer/Private/EaseCurveTool/Widgets/SAvaEaseCurveTool.cpp b/Engine/Plugins/Experimental/Avalanche/Source/AvalancheSequencer/Private/EaseCurveTool/Widgets/SAvaEaseCurveTool.cpp
--- a/Engine/Plugins/Experimental/Avalanche/Source/AvalancheSequencer/Private/EaseCurveTool/Widgets/SAvaEaseCurveTool.cpp
+++ b/Engine/Plugins/Experimental/Avalanche/Source/AvalancheSequencer/Private/EaseCurveTool/Widgets/SAvaEaseCurveTool.cpp
@@ -119,9 +119,15 @@ TSharedRef<SWidget> SAvaEaseCurveTool::ConstructCurveEditorPanel()
 					{
 						return FVector2D(CurrentGraphSize);
 					})
-				.ShowEqualValueKeyError_Lambda([this]() -> bool
+				// These two are lambda because the functions aren't const and the functions aren't const
+				// because there is no const ForEachEaseableKey implementation
+				.ShowEqualValueKeyError_Lambda([this]()
 					{
-						return EaseCurveTool->HasCachedKeysToEase();
+						return !EaseCurveTool->HasCachedKeysToEase();
+					})
+				.IsEaseCurveSelection_Lambda([this]()
+					{
+						return EaseCurveTool->AreAllEaseCurves();
 					})
 				.OnTangentsChanged(this, &SAvaEaseCurveTool::HandleEditorTangentsChanged)
 				.GridSnap_UObject(GetDefault<UAvaEaseCurveToolSettings>(), &UAvaEaseCurveToolSettings::GetGridSnap)
