MotionDesign : SVG
- Fixed crash when SVG Importer module is not yet loaded when AvaSVGActorTool gets FSVGImporterEditorCommands

diff --git a/Engine/Plugins/Experimental/Avalanche/Source/AvalancheSVGEditor/Private/AvaSVGEditorModule.cpp b/Engine/Plugins/Experimental/Avalanche/Source/AvalancheSVGEditor/Private/AvaSVGEditorModule.cpp
--- a/Engine/Plugins/Experimental/Avalanche/Source/AvalancheSVGEditor/Private/AvaSVGEditorModule.cpp
+++ b/Engine/Plugins/Experimental/Avalanche/Source/AvalancheSVGEditor/Private/AvaSVGEditorModule.cpp
@@ -7,6 +7,7 @@
 #include "Modifiers/ActorModifierCoreStack.h"
 #include "Modifiers/AvaBevelModifier.h"
 #include "Modifiers/AvaExtrudeModifier.h"
+#include "Modules/ModuleObserver.h"
 #include "ProceduralMeshes/SVGDynamicMeshComponent.h"
 #include "SVGEngineSubsystem.h"
 #include "SVGShapesParentActor.h"
@@ -33,9 +34,22 @@ void FAvaSVGEditorModule::ShutdownModule()
 
 void FAvaSVGEditorModule::RegisterTools(IAvalancheInteractiveToolsModule* InModule)
 {
-	InModule->RegisterTool(
-		IAvalancheInteractiveToolsModule::Get().CategoryNameActor,
-		GetDefault<UAvaSVGActorTool>()->GetToolParameters()
+	// Use module observer to make sure SVGImporterEditor is loaded
+	FModuleObserver ModuleObserver(
+		"SVGImporterEditor",
+		FSimpleDelegate::CreateLambda([InModule]()
+		{
+			if (!InModule)
+			{
+				return;
+			}
+
+			InModule->RegisterTool(
+				IAvalancheInteractiveToolsModule::Get().CategoryNameActor,
+				GetDefault<UAvaSVGActorTool>()->GetToolParameters()
+			);
+		}),
+		FSimpleDelegate()
 	);
 }
 
