MotionDesign : ClonerEffector
- Refactoring cloner and effector into extensions
- Moving logic and properties from actor to component and extensions
- Adapting viewport visualizers and details customization
- Migrating properties for backward compatibility
- Loading system async instead of duplicating to avoid crash at runtime
- Using a custom mesh builder to merge different meshes together

#virtualized

diff --git a/Engine/Plugins/Experimental/Avalanche/Source/AvalancheComponentVisualizers/Private/AvaVisBase.cpp b/Engine/Plugins/Experimental/Avalanche/Source/AvalancheComponentVisualizers/Private/AvaVisBase.cpp
--- a/Engine/Plugins/Experimental/Avalanche/Source/AvalancheComponentVisualizers/Private/AvaVisBase.cpp
+++ b/Engine/Plugins/Experimental/Avalanche/Source/AvalancheComponentVisualizers/Private/AvaVisBase.cpp
@@ -5,7 +5,7 @@
 #include "AvaComponentVisualizersEdMode.h"
 #include "AvaComponentVisualizersSettings.h"
 #include "AvaViewportUtils.h"
-#include "Components/DynamicMeshComponent.h"
+#include "EditorModeManager.h"
 #include "EditorViewportClient.h"
 #include "Interaction/AvaSnapOperation.h"
 #include "Math/Box.h"
@@ -16,6 +16,8 @@
 #include "UnrealClient.h"
 #include "ViewportClient/IAvaViewportClient.h"
 
+#define LOCTEXT_NAMESPACE "AvaVisualizerBase"
+
 IMPLEMENT_HIT_PROXY(HAvaHitProxy, HComponentVisProxy);
 IMPLEMENT_HIT_PROXY(HAvaDirectHitProxy, HAvaHitProxy);
 
@@ -110,20 +112,20 @@ bool FAvaVisualizerBase::HandleModifiedClick(FEditorViewportClient* InViewportCl
 
 void FAvaVisualizerBase::StartTransaction()
 {
-	if (!GetEditedComponent())
+	UActorComponent* EditedComponent = GetEditedComponent();
+
+	if (!EditedComponent)
 	{
 		return;
 	}
 
-	GetEditedComponent()->SetFlags(RF_Transactional);
+	EditedComponent->SetFlags(RF_Transactional);
+
+	const FString OwnerLabel = EditedComponent->GetOwner() ? EditedComponent->GetOwner()->GetActorNameOrLabel() : TEXT("");
+	TransactionIdx = GEditor->BeginTransaction(FText::Format(LOCTEXT("VisualizerChange", "Edit {0} Visualizer Change"), FText::FromString(OwnerLabel)));
 
-	TransactionIdx = GEngine->BeginTransaction(
-			TEXT("Motion Design Component Visualizer"),
-			NSLOCTEXT("AvaComponentVisualizerBase", "VisualizerChange", "Visualizer Change"),
-			GetEditedComponent()
-		);
+	EditedComponent->Modify();
 
-	GetEditedComponent()->Modify();
 	bHasBeenModified = false;
 }
 
@@ -133,11 +135,11 @@ void FAvaVisualizerBase::EndTransaction()
 	{
 		if (bHasBeenModified)
 		{
-			GEngine->EndTransaction();
+			GEditor->EndTransaction();
 		}
 		else
 		{
-			GEngine->CancelTransaction(TransactionIdx);
+			GEditor->CancelTransaction(TransactionIdx);
 		}
 		TransactionIdx = INDEX_NONE;
 	}
@@ -156,6 +158,14 @@ void FAvaVisualizerBase::StoreInitialValues()
 void FAvaVisualizerBase::StartEditing(FEditorViewportClient* InViewportClient, UActorComponent* InEditedComponent)
 {
 	LastUsedViewportClient = InViewportClient;
+
+	// Disable AutoSelectComponent to keep active actor selected
+	if (IConsoleVariable* CVar = IConsoleManager::Get().FindConsoleVariable(TEXT("Editor.ComponentVisualizer.AutoSelectComponent")))
+	{
+		bCVarAutoSelectComponent = CVar->GetBool();
+		CVar->Set(false);
+	}
+
 	UAvaComponentVisualizersEdMode::OnVisualizerActivated(SharedThis(this));
 }
 
@@ -165,7 +175,14 @@ void FAvaVisualizerBase::EndEditing()
 
 	EndTransaction();
 	UAvaComponentVisualizersEdMode::OnVisualizerDeactivated(SharedThis(this));
+
 	LastUsedViewportClient = nullptr;
+
+	// Restore AutoSelectComponent state
+	if (IConsoleVariable* CVar = IConsoleManager::Get().FindConsoleVariable(TEXT("Editor.ComponentVisualizer.AutoSelectComponent")))
+	{
+		CVar->Set(bCVarAutoSelectComponent);
+	}
 }
 
 void FAvaVisualizerBase::TrackingStarted(FEditorViewportClient* InViewportClient)
@@ -523,9 +540,31 @@ void FAvaVisualizerBase::TrackingStoppedInternal(FEditorViewportClient* InViewpo
 
 void FAvaVisualizerBase::ModifyProperty(UObject* InObject, FProperty* InProperty, EPropertyChangeType::Type InPropertyChangeType, TFunctionRef<void()> InFunction)
 {
+	if (!InObject)
+	{
+		return;
+	}
+
+	InObject->SetFlags(RF_Transactional);
+
 	bHasBeenModified = InObject->Modify();
+
 	InFunction();
-	NotifyPropertyModified(InObject, InProperty, InPropertyChangeType);
+
+	if (InPropertyChangeType != EPropertyChangeType::Interactive)
+	{
+		NotifyPropertyModified(InObject, InProperty, InPropertyChangeType);
+	}
+}
+
+void FAvaVisualizerBase::ModifyProperty(UObject* InObject, FName InPropertyName, EPropertyChangeType::Type InPropertyChangeType, TFunctionRef<void()> InFunction)
+{
+	if (!InObject)
+	{
+		return;
+	}
+
+	ModifyProperty(InObject, FindFProperty<FProperty>(InObject->GetClass(), InPropertyName), InPropertyChangeType, InFunction);
 }
 
 void FAvaVisualizerBase::NotifyPropertyModified(UObject* InObject, FProperty* InProperty,
@@ -610,3 +649,5 @@ void FAvaVisualizerBase::NotifyPropertyChainModified(UObject* InObject, FPropert
 		}
 	}
 }
+
+#undef LOCTEXT_NAMESPACE
\ No newline at end of file
diff --git a/Engine/Plugins/Experimental/Avalanche/Source/AvalancheComponentVisualizers/Public/AvaVisBase.h b/Engine/Plugins/Experimental/Avalanche/Source/AvalancheComponentVisualizers/Public/AvaVisBase.h
--- a/Engine/Plugins/Experimental/Avalanche/Source/AvalancheComponentVisualizers/Public/AvaVisBase.h
+++ b/Engine/Plugins/Experimental/Avalanche/Source/AvalancheComponentVisualizers/Public/AvaVisBase.h
@@ -129,6 +129,7 @@ protected:
 
 	// Last used viewport client
 	FEditorViewportClient* LastUsedViewportClient = nullptr;
+	bool bCVarAutoSelectComponent = true;
 
 	// Delta sums
 	bool bHasInitialWidgetLocation;
@@ -191,4 +192,5 @@ protected:
 	virtual void TrackingStoppedInternal(FEditorViewportClient* InViewportClient);
 
 	void ModifyProperty(UObject* InObject, FProperty* InProperty, EPropertyChangeType::Type InPropertyChangeType, TFunctionRef<void()> InFunction);
+	void ModifyProperty(UObject* InObject, FName InPropertyName, EPropertyChangeType::Type InPropertyChangeType, TFunctionRef<void()> InFunction);
 };
diff --git a/Engine/Plugins/Experimental/Avalanche/Source/AvalancheEffectorsEditor/Private/Cloner/AvaClonerActorVis.cpp b/Engine/Plugins/Experimental/Avalanche/Source/AvalancheEffectorsEditor/Private/Cloner/AvaClonerActorVis.cpp
--- a/Engine/Plugins/Experimental/Avalanche/Source/AvalancheEffectorsEditor/Private/Cloner/AvaClonerActorVis.cpp
+++ b/Engine/Plugins/Experimental/Avalanche/Source/AvalancheEffectorsEditor/Private/Cloner/AvaClonerActorVis.cpp
@@ -36,22 +36,22 @@ void FAvaClonerActorVisualizer::StoreInitialValues()
 		return;
 	}
 
-	const ACEClonerActor* ClonerActor = ClonerActorWeak.Get();
+	const UCEClonerComponent* ClonerComponent = ClonerComponentWeak.Get();
 
-	if (!ClonerActor)
+	if (!ClonerComponent)
 	{
 		return;
 	}
 
-	if (const UCEClonerGridLayout* GridLayout = ClonerActor->GetActiveLayout<UCEClonerGridLayout>())
+	if (const UCEClonerGridLayout* GridLayout = ClonerComponent->GetActiveLayout<UCEClonerGridLayout>())
 	{
 		InitialSpacing = FVector(GridLayout->GetSpacingX(), GridLayout->GetSpacingY(), GridLayout->GetSpacingZ());
 	}
-	else if (const UCEClonerLineLayout* LineLayout = ClonerActor->GetActiveLayout<UCEClonerLineLayout>())
+	else if (const UCEClonerLineLayout* LineLayout = ClonerComponent->GetActiveLayout<UCEClonerLineLayout>())
 	{
 		InitialSpacing = FVector(LineLayout->GetSpacing());
 	}
-	else if (const UCEClonerHoneycombLayout* HoneycombLayout = ClonerActor->GetActiveLayout<UCEClonerHoneycombLayout>())
+	else if (const UCEClonerHoneycombLayout* HoneycombLayout = ClonerComponent->GetActiveLayout<UCEClonerHoneycombLayout>())
 	{
 		const ECEClonerPlane Plane = HoneycombLayout->GetPlane();
 
@@ -71,15 +71,15 @@ void FAvaClonerActorVisualizer::StoreInitialValues()
 
 		InitialSpacing = Spacing;
 	}
-	else if (const UCEClonerCircleLayout* CircleLayout = ClonerActor->GetActiveLayout<UCEClonerCircleLayout>())
+	else if (const UCEClonerCircleLayout* CircleLayout = ClonerComponent->GetActiveLayout<UCEClonerCircleLayout>())
 	{
 		InitialSpacing = FVector(CircleLayout->GetRadius());
 	}
-	else if (const UCEClonerCylinderLayout* CylinderLayout = ClonerActor->GetActiveLayout<UCEClonerCylinderLayout>())
+	else if (const UCEClonerCylinderLayout* CylinderLayout = ClonerComponent->GetActiveLayout<UCEClonerCylinderLayout>())
 	{
 		InitialSpacing = FVector(0, CylinderLayout->GetRadius(), CylinderLayout->GetHeight());
 	}
-	else if (const UCEClonerSphereUniformLayout* SphereLayout = ClonerActor->GetActiveLayout<UCEClonerSphereUniformLayout>())
+	else if (const UCEClonerSphereUniformLayout* SphereLayout = ClonerComponent->GetActiveLayout<UCEClonerSphereUniformLayout>())
 	{
 		InitialSpacing = FVector(0, SphereLayout->GetRadius(), 0);
 	}
@@ -89,7 +89,7 @@ FBox FAvaClonerActorVisualizer::GetComponentBounds(const UActorComponent* InComp
 {
 	if (const UCEClonerComponent* ClonerComponent = Cast<UCEClonerComponent>(InComponent))
 	{
-		if (const ACEClonerActor* ClonerActor = Cast<ACEClonerActor>(ClonerComponent->GetOwner()))
+		if (const AActor* ClonerActor = ClonerComponent->GetOwner())
 		{
 			FVector Origin;
 			FVector Extent;
@@ -108,9 +108,9 @@ bool FAvaClonerActorVisualizer::HandleInputDeltaInternal(FEditorViewportClient*
 		return false;
 	}
 
-	const ACEClonerActor* ClonerActor = ClonerActorWeak.Get();
+	const UCEClonerComponent* ClonerComponent = ClonerComponentWeak.Get();
 
-	if (!ClonerActor)
+	if (!ClonerComponent)
 	{
 		EndEditing();
 	}
@@ -118,30 +118,48 @@ bool FAvaClonerActorVisualizer::HandleInputDeltaInternal(FEditorViewportClient*
 	{
 		const EAxisList::Type AxisList = GetViewportWidgetAxisList(InViewportClient);
 
-		if (UCEClonerGridLayout* GridLayout = ClonerActor->GetActiveLayout<UCEClonerGridLayout>())
+		if (UCEClonerGridLayout* GridLayout = ClonerComponent->GetActiveLayout<UCEClonerGridLayout>())
 		{
 			if (AxisList & EAxisList::X)
 			{
-				const float SpacingX = InitialSpacing.X + InAccumulatedTranslation.X / FMath::Max(1, GridLayout->GetCountX() / 2.f);
-				GridLayout->SetSpacingX(SpacingX);
-				OnPropertyModified(GridLayout, GET_MEMBER_NAME_CHECKED(UCEClonerGridLayout, SpacingX));
+				ModifyProperty(
+					GridLayout
+					, GET_MEMBER_NAME_CHECKED(UCEClonerGridLayout, SpacingX)
+					, EPropertyChangeType::Interactive
+					, [this, &InAccumulatedTranslation, &GridLayout]()
+				{
+					const float SpacingX = InitialSpacing.X + InAccumulatedTranslation.X / FMath::Max(1, GridLayout->GetCountX() / 2.f);
+					GridLayout->SetSpacingX(SpacingX);
+				});
 			}
 			else if (AxisList & EAxisList::Y)
 			{
-				const float SpacingY = InitialSpacing.Y + InAccumulatedTranslation.Y / FMath::Max(1, GridLayout->GetCountY() / 2.f);
-				GridLayout->SetSpacingY(SpacingY);
-				OnPropertyModified(GridLayout, GET_MEMBER_NAME_CHECKED(UCEClonerGridLayout, SpacingY));
+				ModifyProperty(
+					GridLayout
+					, GET_MEMBER_NAME_CHECKED(UCEClonerGridLayout, SpacingY)
+					, EPropertyChangeType::Interactive
+					, [this, &InAccumulatedTranslation, &GridLayout]()
+				{
+					const float SpacingY = InitialSpacing.Y + InAccumulatedTranslation.Y / FMath::Max(1, GridLayout->GetCountY() / 2.f);
+					GridLayout->SetSpacingY(SpacingY);
+				});
 			}
 			else if (AxisList & EAxisList::Z)
 			{
-				const float SpacingZ = InitialSpacing.Z + InAccumulatedTranslation.Z / FMath::Max(1, GridLayout->GetCountZ() / 2.f);
-				GridLayout->SetSpacingZ(SpacingZ);
-				OnPropertyModified(GridLayout, GET_MEMBER_NAME_CHECKED(UCEClonerGridLayout, SpacingZ));
+				ModifyProperty(
+					GridLayout
+					, GET_MEMBER_NAME_CHECKED(UCEClonerGridLayout, SpacingZ)
+					, EPropertyChangeType::Interactive
+					, [this, &InAccumulatedTranslation, &GridLayout]()
+				{
+					const float SpacingZ = InitialSpacing.Z + InAccumulatedTranslation.Z / FMath::Max(1, GridLayout->GetCountZ() / 2.f);
+					GridLayout->SetSpacingZ(SpacingZ);
+				});
 			}
 
 			return true;
 		}
-		else if (UCEClonerLineLayout* LineLayout = ClonerActor->GetActiveLayout<UCEClonerLineLayout>())
+		else if (UCEClonerLineLayout* LineLayout = ClonerComponent->GetActiveLayout<UCEClonerLineLayout>())
 		{
 			const int32 Count = LineLayout->GetCount();
 			float Spacing = LineLayout->GetSpacing();
@@ -159,12 +177,18 @@ bool FAvaClonerActorVisualizer::HandleInputDeltaInternal(FEditorViewportClient*
 				Spacing = InitialSpacing.Z + InAccumulatedTranslation.Z / FMath::Max(1, Count);
 			}
 
-			LineLayout->SetSpacing(Spacing);
-			OnPropertyModified(LineLayout, GET_MEMBER_NAME_CHECKED(UCEClonerLineLayout, Spacing));
+			ModifyProperty(
+				LineLayout
+				, GET_MEMBER_NAME_CHECKED(UCEClonerLineLayout, Spacing)
+				, EPropertyChangeType::Interactive
+				, [this, &Spacing, &LineLayout]()
+			{
+				LineLayout->SetSpacing(Spacing);
+			});
 
 			return true;
 		}
-		else if (UCEClonerHoneycombLayout* HoneycombLayout = ClonerActor->GetActiveLayout<UCEClonerHoneycombLayout>())
+		else if (UCEClonerHoneycombLayout* HoneycombLayout = ClonerComponent->GetActiveLayout<UCEClonerHoneycombLayout>())
 		{
 			const ECEClonerPlane Plane = HoneycombLayout->GetPlane();
 			const int32 WidthCount = HoneycombLayout->GetWidthCount();
@@ -188,15 +212,27 @@ bool FAvaClonerActorVisualizer::HandleInputDeltaInternal(FEditorViewportClient*
 				HeightSpacing = InitialSpacing.Z + InAccumulatedTranslation.Z / FMath::Max(1, HeightCount / 2.f);
 			}
 
-			HoneycombLayout->SetWidthSpacing(WidthSpacing);
-			HoneycombLayout->SetHeightSpacing(HeightSpacing);
-
-			OnPropertyModified(HoneycombLayout, GET_MEMBER_NAME_CHECKED(UCEClonerHoneycombLayout, WidthSpacing));
-			OnPropertyModified(HoneycombLayout, GET_MEMBER_NAME_CHECKED(UCEClonerHoneycombLayout, HeightSpacing));
+			ModifyProperty(
+				HoneycombLayout
+				, GET_MEMBER_NAME_CHECKED(UCEClonerHoneycombLayout, WidthSpacing)
+				, EPropertyChangeType::Interactive
+				, [this, &WidthSpacing, &HoneycombLayout]()
+			{
+				HoneycombLayout->SetWidthSpacing(WidthSpacing);
+			});
+
+			ModifyProperty(
+				HoneycombLayout
+				, GET_MEMBER_NAME_CHECKED(UCEClonerHoneycombLayout, HeightSpacing)
+				, EPropertyChangeType::Interactive
+				, [this, &HeightSpacing, &HoneycombLayout]()
+			{
+				HoneycombLayout->SetHeightSpacing(HeightSpacing);
+			});
 
 			return true;
 		}
-		else if (UCEClonerCircleLayout* CircleLayout = ClonerActor->GetActiveLayout<UCEClonerCircleLayout>())
+		else if (UCEClonerCircleLayout* CircleLayout = ClonerComponent->GetActiveLayout<UCEClonerCircleLayout>())
 		{
 			float Radius = CircleLayout->GetRadius();
 
@@ -213,12 +249,18 @@ bool FAvaClonerActorVisualizer::HandleInputDeltaInternal(FEditorViewportClient*
 				Radius = InitialSpacing.Z + InAccumulatedTranslation.Z;
 			}
 
-			CircleLayout->SetRadius(Radius);
-			OnPropertyModified(CircleLayout, GET_MEMBER_NAME_CHECKED(UCEClonerCircleLayout, Radius));
+			ModifyProperty(
+				CircleLayout
+				, GET_MEMBER_NAME_CHECKED(UCEClonerCircleLayout, Radius)
+				, EPropertyChangeType::Interactive
+				, [this, &Radius, &CircleLayout]()
+			{
+				CircleLayout->SetRadius(Radius);
+			});
 
 			return true;
 		}
-		else if (UCEClonerCylinderLayout* CylinderLayout = ClonerActor->GetActiveLayout<UCEClonerCylinderLayout>())
+		else if (UCEClonerCylinderLayout* CylinderLayout = ClonerComponent->GetActiveLayout<UCEClonerCylinderLayout>())
 		{
 			const ECEClonerPlane Plane = CylinderLayout->GetPlane();
 			float Radius = CylinderLayout->GetRadius();
@@ -258,15 +300,27 @@ bool FAvaClonerActorVisualizer::HandleInputDeltaInternal(FEditorViewportClient*
 				}
 			}
 
-			CylinderLayout->SetRadius(Radius);
-			CylinderLayout->SetHeight(Height);
-
-			OnPropertyModified(CylinderLayout, GET_MEMBER_NAME_CHECKED(UCEClonerCylinderLayout, Radius));
-			OnPropertyModified(CylinderLayout, GET_MEMBER_NAME_CHECKED(UCEClonerCylinderLayout, Height));
+			ModifyProperty(
+				CylinderLayout
+				, GET_MEMBER_NAME_CHECKED(UCEClonerCylinderLayout, Radius)
+				, EPropertyChangeType::Interactive
+				, [this, &Radius, &CylinderLayout]()
+			{
+				CylinderLayout->SetRadius(Radius);
+			});
+
+			ModifyProperty(
+				CylinderLayout
+				, GET_MEMBER_NAME_CHECKED(UCEClonerCylinderLayout, Height)
+				, EPropertyChangeType::Interactive
+				, [this, &Height, &CylinderLayout]()
+			{
+				CylinderLayout->SetHeight(Height);
+			});
 
 			return true;
 		}
-		else if (UCEClonerSphereUniformLayout* SphereLayout = ClonerActor->GetActiveLayout<UCEClonerSphereUniformLayout>())
+		else if (UCEClonerSphereUniformLayout* SphereLayout = ClonerComponent->GetActiveLayout<UCEClonerSphereUniformLayout>())
 		{
 			float Radius = SphereLayout->GetRadius();
 
@@ -275,8 +329,14 @@ bool FAvaClonerActorVisualizer::HandleInputDeltaInternal(FEditorViewportClient*
 				Radius = InitialSpacing.Y + InAccumulatedTranslation.Y;
 			}
 
-			SphereLayout->SetRadius(Radius);
-			OnPropertyModified(SphereLayout, GET_MEMBER_NAME_CHECKED(UCEClonerSphereUniformLayout, Radius));
+			ModifyProperty(
+				SphereLayout
+				, GET_MEMBER_NAME_CHECKED(UCEClonerSphereUniformLayout, Radius)
+				, EPropertyChangeType::Interactive
+				, [this, &Radius, &SphereLayout]()
+			{
+				SphereLayout->SetRadius(Radius);
+			});
 
 			return true;
 		}
@@ -291,121 +351,114 @@ void FAvaClonerActorVisualizer::DrawVisualizationEditing(const UActorComponent*
 
 	const UCEClonerComponent* ClonerComponent = Cast<UCEClonerComponent>(InComponent);
 
-	if (!ClonerComponent)
-	{
-		return;
-	}
-
-	const ACEClonerActor* ClonerActor = Cast<ACEClonerActor>(ClonerComponent->GetOwner());
-
-	if (!ClonerActor || !ClonerActor->GetEnabled() || ClonerActor->GetMeshCount() == 0)
+	if (!ClonerComponent || !ClonerComponent->GetEnabled() || ClonerComponent->GetMeshCount() == 0)
 	{
 		return;
 	}
 
-	if (const UCEClonerGridLayout* GridLayout = ClonerActor->GetActiveLayout<UCEClonerGridLayout>())
+	if (const UCEClonerGridLayout* GridLayout = ClonerComponent->GetActiveLayout<UCEClonerGridLayout>())
 	{
 		if (GridLayout->GetCountX() > 0)
 		{
-			DrawSpacingButton(ClonerActor, InView, InPDI, InOutIconIndex, ECEClonerAxis::X, FAvaVisualizerBase::Inactive);
+			DrawSpacingButton(ClonerComponent, InView, InPDI, InOutIconIndex, ECEClonerAxis::X, FAvaVisualizerBase::Inactive);
 			InOutIconIndex++;
 		}
 
 		if (GridLayout->GetCountY() > 0)
 		{
-			DrawSpacingButton(ClonerActor, InView, InPDI, InOutIconIndex, ECEClonerAxis::Y, FAvaVisualizerBase::Inactive);
+			DrawSpacingButton(ClonerComponent, InView, InPDI, InOutIconIndex, ECEClonerAxis::Y, FAvaVisualizerBase::Inactive);
 			InOutIconIndex++;
 		}
 
 		if (GridLayout->GetCountZ() > 0)
 		{
-			DrawSpacingButton(ClonerActor, InView, InPDI, InOutIconIndex, ECEClonerAxis::Z, FAvaVisualizerBase::Inactive);
+			DrawSpacingButton(ClonerComponent, InView, InPDI, InOutIconIndex, ECEClonerAxis::Z, FAvaVisualizerBase::Inactive);
 			InOutIconIndex++;
 		}
 	}
-	else if (const UCEClonerLineLayout* LineLayout = ClonerActor->GetActiveLayout<UCEClonerLineLayout>())
+	else if (const UCEClonerLineLayout* LineLayout = ClonerComponent->GetActiveLayout<UCEClonerLineLayout>())
 	{
 		if (LineLayout->GetCount() > 0)
 		{
-			DrawSpacingButton(ClonerActor, InView, InPDI, InOutIconIndex, LineLayout->GetAxis(), FAvaVisualizerBase::Inactive);
+			DrawSpacingButton(ClonerComponent, InView, InPDI, InOutIconIndex, LineLayout->GetAxis(), FAvaVisualizerBase::Inactive);
 			InOutIconIndex++;
 		}
 	}
-	else if (const UCEClonerHoneycombLayout* HoneycombLayout = ClonerActor->GetActiveLayout<UCEClonerHoneycombLayout>())
+	else if (const UCEClonerHoneycombLayout* HoneycombLayout = ClonerComponent->GetActiveLayout<UCEClonerHoneycombLayout>())
 	{
 		const ECEClonerPlane Plane = HoneycombLayout->GetPlane();
 
 		if (Plane == ECEClonerPlane::XY)
 		{
-			DrawSpacingButton(ClonerActor, InView, InPDI, InOutIconIndex, ECEClonerAxis::X, FAvaVisualizerBase::Inactive);
+			DrawSpacingButton(ClonerComponent, InView, InPDI, InOutIconIndex, ECEClonerAxis::X, FAvaVisualizerBase::Inactive);
 			InOutIconIndex++;
-			DrawSpacingButton(ClonerActor, InView, InPDI, InOutIconIndex, ECEClonerAxis::Y, FAvaVisualizerBase::Inactive);
+			DrawSpacingButton(ClonerComponent, InView, InPDI, InOutIconIndex, ECEClonerAxis::Y, FAvaVisualizerBase::Inactive);
 			InOutIconIndex++;
 		}
 		else if (Plane == ECEClonerPlane::YZ)
 		{
-			DrawSpacingButton(ClonerActor, InView, InPDI, InOutIconIndex, ECEClonerAxis::Y, FAvaVisualizerBase::Inactive);
+			DrawSpacingButton(ClonerComponent, InView, InPDI, InOutIconIndex, ECEClonerAxis::Y, FAvaVisualizerBase::Inactive);
 			InOutIconIndex++;
-			DrawSpacingButton(ClonerActor, InView, InPDI, InOutIconIndex, ECEClonerAxis::Z, FAvaVisualizerBase::Inactive);
+			DrawSpacingButton(ClonerComponent, InView, InPDI, InOutIconIndex, ECEClonerAxis::Z, FAvaVisualizerBase::Inactive);
 			InOutIconIndex++;
 		}
 		else if (Plane == ECEClonerPlane::XZ)
 		{
-			DrawSpacingButton(ClonerActor, InView, InPDI, InOutIconIndex, ECEClonerAxis::X, FAvaVisualizerBase::Inactive);
+			DrawSpacingButton(ClonerComponent, InView, InPDI, InOutIconIndex, ECEClonerAxis::X, FAvaVisualizerBase::Inactive);
 			InOutIconIndex++;
-			DrawSpacingButton(ClonerActor, InView, InPDI, InOutIconIndex, ECEClonerAxis::Z, FAvaVisualizerBase::Inactive);
+			DrawSpacingButton(ClonerComponent, InView, InPDI, InOutIconIndex, ECEClonerAxis::Z, FAvaVisualizerBase::Inactive);
 			InOutIconIndex++;
 		}
 	}
-	else if (const UCEClonerCircleLayout* CircleLayout = ClonerActor->GetActiveLayout<UCEClonerCircleLayout>())
+	else if (const UCEClonerCircleLayout* CircleLayout = ClonerComponent->GetActiveLayout<UCEClonerCircleLayout>())
 	{
 		const ECEClonerPlane Plane = CircleLayout->GetPlane();
 
 		if (Plane == ECEClonerPlane::XY)
 		{
-			DrawSpacingButton(ClonerActor, InView, InPDI, InOutIconIndex, ECEClonerAxis::Y, FAvaVisualizerBase::Inactive);
+			DrawSpacingButton(ClonerComponent, InView, InPDI, InOutIconIndex, ECEClonerAxis::Y, FAvaVisualizerBase::Inactive);
 			InOutIconIndex++;
 		}
 		else if (Plane == ECEClonerPlane::YZ)
 		{
-			DrawSpacingButton(ClonerActor, InView, InPDI, InOutIconIndex, ECEClonerAxis::Y, FAvaVisualizerBase::Inactive);
+			DrawSpacingButton(ClonerComponent, InView, InPDI, InOutIconIndex, ECEClonerAxis::Y, FAvaVisualizerBase::Inactive);
 			InOutIconIndex++;
 		}
 		else if (Plane == ECEClonerPlane::XZ)
 		{
-			DrawSpacingButton(ClonerActor, InView, InPDI, InOutIconIndex, ECEClonerAxis::Z, FAvaVisualizerBase::Inactive);
+			DrawSpacingButton(ClonerComponent, InView, InPDI, InOutIconIndex, ECEClonerAxis::Z, FAvaVisualizerBase::Inactive);
 			InOutIconIndex++;
 		}
 	}
-	else if (const UCEClonerCylinderLayout* CylinderLayout = ClonerActor->GetActiveLayout<UCEClonerCylinderLayout>())
+	else if (const UCEClonerCylinderLayout* CylinderLayout = ClonerComponent->GetActiveLayout<UCEClonerCylinderLayout>())
 	{
 		const ECEClonerPlane Plane = CylinderLayout->GetPlane();
 
 		if (Plane == ECEClonerPlane::XY)
 		{
-			DrawSpacingButton(ClonerActor, InView, InPDI, InOutIconIndex, ECEClonerAxis::Y, FAvaVisualizerBase::Inactive);
+			DrawSpacingButton(ClonerComponent, InView, InPDI, InOutIconIndex, ECEClonerAxis::Y, FAvaVisualizerBase::Inactive);
 			InOutIconIndex++;
-			DrawSpacingButton(ClonerActor, InView, InPDI, InOutIconIndex, ECEClonerAxis::Z, FAvaVisualizerBase::Inactive);
+			DrawSpacingButton(ClonerComponent, InView, InPDI, InOutIconIndex, ECEClonerAxis::Z, FAvaVisualizerBase::Inactive);
 			InOutIconIndex++;
 		}
 		else if (Plane == ECEClonerPlane::YZ)
 		{
-			DrawSpacingButton(ClonerActor, InView, InPDI, InOutIconIndex, ECEClonerAxis::Y, FAvaVisualizerBase::Inactive);
+			DrawSpacingButton(ClonerComponent, InView, InPDI, InOutIconIndex, ECEClonerAxis::Y, FAvaVisualizerBase::Inactive);
 			InOutIconIndex++;
-			DrawSpacingButton(ClonerActor, InView, InPDI, InOutIconIndex, ECEClonerAxis::X, FAvaVisualizerBase::Inactive);
+			DrawSpacingButton(ClonerComponent, InView, InPDI, InOutIconIndex, ECEClonerAxis::X, FAvaVisualizerBase::Inactive);
 			InOutIconIndex++;
 		}
 		else if (Plane == ECEClonerPlane::XZ)
 		{
-			DrawSpacingButton(ClonerActor, InView, InPDI, InOutIconIndex, ECEClonerAxis::Z, FAvaVisualizerBase::Inactive);
+			DrawSpacingButton(ClonerComponent, InView, InPDI, InOutIconIndex, ECEClonerAxis::Z, FAvaVisualizerBase::Inactive);
 			InOutIconIndex++;
-			DrawSpacingButton(ClonerActor, InView, InPDI, InOutIconIndex, ECEClonerAxis::Y, FAvaVisualizerBase::Inactive);
+			DrawSpacingButton(ClonerComponent, InView, InPDI, InOutIconIndex, ECEClonerAxis::Y, FAvaVisualizerBase::Inactive);
 			InOutIconIndex++;
 		}
 	}
-	else if (ClonerActor->IsActiveLayout<UCEClonerSphereUniformLayout>())
+	else if (ClonerComponent->IsActiveLayout<UCEClonerSphereUniformLayout>())
 	{
-		DrawSpacingButton(ClonerActor, InView, InPDI, InOutIconIndex, ECEClonerAxis::Y, FAvaVisualizerBase::Inactive);
+		DrawSpacingButton(ClonerComponent, InView, InPDI, InOutIconIndex, ECEClonerAxis::Y, FAvaVisualizerBase::Inactive);
 		InOutIconIndex++;
 	}
 }
@@ -416,145 +469,138 @@ void FAvaClonerActorVisualizer::DrawVisualizationNotEditing(const UActorComponen
 
 	const UCEClonerComponent* ClonerComponent = Cast<UCEClonerComponent>(InComponent);
 
-	if (!ClonerComponent)
+	if (!ClonerComponent || !ClonerComponent->GetEnabled() || ClonerComponent->GetMeshCount() == 0)
 	{
 		return;
 	}
 
-	const ACEClonerActor* ClonerActor = Cast<ACEClonerActor>(ClonerComponent->GetOwner());
-
-	if (!ClonerActor || !ClonerActor->GetEnabled() || ClonerActor->GetMeshCount() == 0)
-	{
-		return;
-	}
-
-	if (const UCEClonerGridLayout* GridLayout = ClonerActor->GetActiveLayout<UCEClonerGridLayout>())
+	if (const UCEClonerGridLayout* GridLayout = ClonerComponent->GetActiveLayout<UCEClonerGridLayout>())
 	{
 		if (GridLayout->GetCountX() > 0)
 		{
-			DrawSpacingButton(ClonerActor, InView, InPDI, InOutIconIndex, ECEClonerAxis::X, FAvaVisualizerBase::Inactive);
+			DrawSpacingButton(ClonerComponent, InView, InPDI, InOutIconIndex, ECEClonerAxis::X, FAvaVisualizerBase::Inactive);
 			InOutIconIndex++;
 		}
 
 		if (GridLayout->GetCountY() > 0)
 		{
-			DrawSpacingButton(ClonerActor, InView, InPDI, InOutIconIndex, ECEClonerAxis::Y, FAvaVisualizerBase::Inactive);
+			DrawSpacingButton(ClonerComponent, InView, InPDI, InOutIconIndex, ECEClonerAxis::Y, FAvaVisualizerBase::Inactive);
 			InOutIconIndex++;
 		}
 
 		if (GridLayout->GetCountZ() > 0)
 		{
-			DrawSpacingButton(ClonerActor, InView, InPDI, InOutIconIndex, ECEClonerAxis::Z, FAvaVisualizerBase::Inactive);
+			DrawSpacingButton(ClonerComponent, InView, InPDI, InOutIconIndex, ECEClonerAxis::Z, FAvaVisualizerBase::Inactive);
 			InOutIconIndex++;
 		}
 	}
-	else if (const UCEClonerLineLayout* LineLayout = ClonerActor->GetActiveLayout<UCEClonerLineLayout>())
+	else if (const UCEClonerLineLayout* LineLayout = ClonerComponent->GetActiveLayout<UCEClonerLineLayout>())
 	{
 		if (LineLayout->GetCount() > 0)
 		{
-			DrawSpacingButton(ClonerActor, InView, InPDI, InOutIconIndex, LineLayout->GetAxis(), FAvaVisualizerBase::Inactive);
+			DrawSpacingButton(ClonerComponent, InView, InPDI, InOutIconIndex, LineLayout->GetAxis(), FAvaVisualizerBase::Inactive);
 			InOutIconIndex++;
 		}
 	}
-	else if (const UCEClonerHoneycombLayout* HoneycombLayout = ClonerActor->GetActiveLayout<UCEClonerHoneycombLayout>())
+	else if (const UCEClonerHoneycombLayout* HoneycombLayout = ClonerComponent->GetActiveLayout<UCEClonerHoneycombLayout>())
 	{
 		const ECEClonerPlane Plane = HoneycombLayout->GetPlane();
 
 		if (Plane == ECEClonerPlane::XY)
 		{
-			DrawSpacingButton(ClonerActor, InView, InPDI, InOutIconIndex, ECEClonerAxis::X, FAvaVisualizerBase::Inactive);
+			DrawSpacingButton(ClonerComponent, InView, InPDI, InOutIconIndex, ECEClonerAxis::X, FAvaVisualizerBase::Inactive);
 			InOutIconIndex++;
-			DrawSpacingButton(ClonerActor, InView, InPDI, InOutIconIndex, ECEClonerAxis::Y, FAvaVisualizerBase::Inactive);
+			DrawSpacingButton(ClonerComponent, InView, InPDI, InOutIconIndex, ECEClonerAxis::Y, FAvaVisualizerBase::Inactive);
 			InOutIconIndex++;
 		}
 		else if (Plane == ECEClonerPlane::YZ)
 		{
-			DrawSpacingButton(ClonerActor, InView, InPDI, InOutIconIndex, ECEClonerAxis::Y, FAvaVisualizerBase::Inactive);
+			DrawSpacingButton(ClonerComponent, InView, InPDI, InOutIconIndex, ECEClonerAxis::Y, FAvaVisualizerBase::Inactive);
 			InOutIconIndex++;
-			DrawSpacingButton(ClonerActor, InView, InPDI, InOutIconIndex, ECEClonerAxis::Z, FAvaVisualizerBase::Inactive);
+			DrawSpacingButton(ClonerComponent, InView, InPDI, InOutIconIndex, ECEClonerAxis::Z, FAvaVisualizerBase::Inactive);
 			InOutIconIndex++;
 		}
 		else if (Plane == ECEClonerPlane::XZ)
 		{
-			DrawSpacingButton(ClonerActor, InView, InPDI, InOutIconIndex, ECEClonerAxis::X, FAvaVisualizerBase::Inactive);
+			DrawSpacingButton(ClonerComponent, InView, InPDI, InOutIconIndex, ECEClonerAxis::X, FAvaVisualizerBase::Inactive);
 			InOutIconIndex++;
-			DrawSpacingButton(ClonerActor, InView, InPDI, InOutIconIndex, ECEClonerAxis::Z, FAvaVisualizerBase::Inactive);
+			DrawSpacingButton(ClonerComponent, InView, InPDI, InOutIconIndex, ECEClonerAxis::Z, FAvaVisualizerBase::Inactive);
 			InOutIconIndex++;
 		}
 	}
-	else if (const UCEClonerCircleLayout* CircleLayout = ClonerActor->GetActiveLayout<UCEClonerCircleLayout>())
+	else if (const UCEClonerCircleLayout* CircleLayout = ClonerComponent->GetActiveLayout<UCEClonerCircleLayout>())
 	{
 		const ECEClonerPlane Plane = CircleLayout->GetPlane();
 
 		if (Plane == ECEClonerPlane::XY)
 		{
-			DrawSpacingButton(ClonerActor, InView, InPDI, InOutIconIndex, ECEClonerAxis::Y, FAvaVisualizerBase::Inactive);
+			DrawSpacingButton(ClonerComponent, InView, InPDI, InOutIconIndex, ECEClonerAxis::Y, FAvaVisualizerBase::Inactive);
 			InOutIconIndex++;
 		}
 		else if (Plane == ECEClonerPlane::YZ)
 		{
-			DrawSpacingButton(ClonerActor, InView, InPDI, InOutIconIndex, ECEClonerAxis::Y, FAvaVisualizerBase::Inactive);
+			DrawSpacingButton(ClonerComponent, InView, InPDI, InOutIconIndex, ECEClonerAxis::Y, FAvaVisualizerBase::Inactive);
 			InOutIconIndex++;
 		}
 		else if (Plane == ECEClonerPlane::XZ)
 		{
-			DrawSpacingButton(ClonerActor, InView, InPDI, InOutIconIndex, ECEClonerAxis::Z, FAvaVisualizerBase::Inactive);
+			DrawSpacingButton(ClonerComponent, InView, InPDI, InOutIconIndex, ECEClonerAxis::Z, FAvaVisualizerBase::Inactive);
 			InOutIconIndex++;
 		}
 		else
 		{
-			DrawSpacingButton(ClonerActor, InView, InPDI, InOutIconIndex, ECEClonerAxis::Custom, FAvaVisualizerBase::Inactive);
+			DrawSpacingButton(ClonerComponent, InView, InPDI, InOutIconIndex, ECEClonerAxis::Custom, FAvaVisualizerBase::Inactive);
 			InOutIconIndex++;
 		}
 	}
-	else if (const UCEClonerCylinderLayout* CylinderLayout = ClonerActor->GetActiveLayout<UCEClonerCylinderLayout>())
+	else if (const UCEClonerCylinderLayout* CylinderLayout = ClonerComponent->GetActiveLayout<UCEClonerCylinderLayout>())
 	{
 		const ECEClonerPlane Plane = CylinderLayout->GetPlane();
 
 		if (Plane == ECEClonerPlane::XY)
 		{
-			DrawSpacingButton(ClonerActor, InView, InPDI, InOutIconIndex, ECEClonerAxis::Y, FAvaVisualizerBase::Inactive);
+			DrawSpacingButton(ClonerComponent, InView, InPDI, InOutIconIndex, ECEClonerAxis::Y, FAvaVisualizerBase::Inactive);
 			InOutIconIndex++;
-			DrawSpacingButton(ClonerActor, InView, InPDI, InOutIconIndex, ECEClonerAxis::Z, FAvaVisualizerBase::Inactive);
+			DrawSpacingButton(ClonerComponent, InView, InPDI, InOutIconIndex, ECEClonerAxis::Z, FAvaVisualizerBase::Inactive);
 			InOutIconIndex++;
 		}
 		else if (Plane == ECEClonerPlane::YZ)
 		{
-			DrawSpacingButton(ClonerActor, InView, InPDI, InOutIconIndex, ECEClonerAxis::Y, FAvaVisualizerBase::Inactive);
+			DrawSpacingButton(ClonerComponent, InView, InPDI, InOutIconIndex, ECEClonerAxis::Y, FAvaVisualizerBase::Inactive);
 			InOutIconIndex++;
-			DrawSpacingButton(ClonerActor, InView, InPDI, InOutIconIndex, ECEClonerAxis::X, FAvaVisualizerBase::Inactive);
+			DrawSpacingButton(ClonerComponent, InView, InPDI, InOutIconIndex, ECEClonerAxis::X, FAvaVisualizerBase::Inactive);
 			InOutIconIndex++;
 		}
 		else if (Plane == ECEClonerPlane::XZ)
 		{
-			DrawSpacingButton(ClonerActor, InView, InPDI, InOutIconIndex, ECEClonerAxis::Z, FAvaVisualizerBase::Inactive);
+			DrawSpacingButton(ClonerComponent, InView, InPDI, InOutIconIndex, ECEClonerAxis::Z, FAvaVisualizerBase::Inactive);
 			InOutIconIndex++;
-			DrawSpacingButton(ClonerActor, InView, InPDI, InOutIconIndex, ECEClonerAxis::Y, FAvaVisualizerBase::Inactive);
+			DrawSpacingButton(ClonerComponent, InView, InPDI, InOutIconIndex, ECEClonerAxis::Y, FAvaVisualizerBase::Inactive);
 			InOutIconIndex++;
 		}
 	}
-	else if (ClonerActor->IsActiveLayout<UCEClonerSphereUniformLayout>())
+	else if (ClonerComponent->IsActiveLayout<UCEClonerSphereUniformLayout>())
 	{
-		DrawSpacingButton(ClonerActor, InView, InPDI, InOutIconIndex, ECEClonerAxis::Y, FAvaVisualizerBase::Inactive);
+		DrawSpacingButton(ClonerComponent, InView, InPDI, InOutIconIndex, ECEClonerAxis::Y, FAvaVisualizerBase::Inactive);
 		InOutIconIndex++;
 	}
 }
 
-FVector FAvaClonerActorVisualizer::GetHandleSpacingLocation(const ACEClonerActor* InClonerActor, ECEClonerAxis InAxis) const
+FVector FAvaClonerActorVisualizer::GetHandleSpacingLocation(const UCEClonerComponent* InClonerComponent, ECEClonerAxis InAxis) const
 {
-	const FRotator ClonerRotation = InClonerActor->GetActorRotation();
-	const FVector ClonerScale = InClonerActor->GetActorScale();
+	const FRotator ClonerRotation = InClonerComponent->GetComponentRotation();
+	const FVector ClonerScale = InClonerComponent->GetComponentScale();
 	FVector OutLocation = FVector::ZeroVector;
 	FVector Axis = InAxis == ECEClonerAxis::X ? FVector::XAxisVector : (InAxis == ECEClonerAxis::Z ? FVector::ZAxisVector : FVector::YAxisVector);
 
-	if (const UCEClonerGridLayout* GridLayout = InClonerActor->GetActiveLayout<UCEClonerGridLayout>())
+	if (const UCEClonerGridLayout* GridLayout = InClonerComponent->GetActiveLayout<UCEClonerGridLayout>())
 	{
 		const FVector Spacing(GridLayout->GetSpacingX(), GridLayout->GetSpacingY(), GridLayout->GetSpacingZ());
 		const FVector Count(GridLayout->GetCountX(), GridLayout->GetCountY(), GridLayout->GetCountZ());
 
 		OutLocation = ClonerRotation.RotateVector(Axis * Spacing * Count / 2);
 	}
-	else if (const UCEClonerLineLayout* LineLayout = InClonerActor->GetActiveLayout<UCEClonerLineLayout>())
+	else if (const UCEClonerLineLayout* LineLayout = InClonerComponent->GetActiveLayout<UCEClonerLineLayout>())
 	{
 		if (InAxis == ECEClonerAxis::Custom)
 		{
@@ -563,7 +609,7 @@ FVector FAvaClonerActorVisualizer::GetHandleSpacingLocation(const ACEClonerActor
 
 		OutLocation = ClonerRotation.RotateVector(Axis) * LineLayout->GetSpacing() * FVector(LineLayout->GetCount());
 	}
-	else if (const UCEClonerHoneycombLayout* HoneycombLayout = InClonerActor->GetActiveLayout<UCEClonerHoneycombLayout>())
+	else if (const UCEClonerHoneycombLayout* HoneycombLayout = InClonerComponent->GetActiveLayout<UCEClonerHoneycombLayout>())
 	{
 		const ECEClonerPlane Plane = HoneycombLayout->GetPlane();
 		const float WidthSpacing = HoneycombLayout->GetWidthSpacing();
@@ -591,7 +637,7 @@ FVector FAvaClonerActorVisualizer::GetHandleSpacingLocation(const ACEClonerActor
 
 		OutLocation = ClonerRotation.RotateVector(Axis * Spacing * Count/2);
 	}
-	else if (const UCEClonerCircleLayout* CircleLayout = InClonerActor->GetActiveLayout<UCEClonerCircleLayout>())
+	else if (const UCEClonerCircleLayout* CircleLayout = InClonerComponent->GetActiveLayout<UCEClonerCircleLayout>())
 	{
 		FRotator Rotation = CircleLayout->GetRotation();
 		const FVector Scale = CircleLayout->GetScale();
@@ -613,7 +659,7 @@ FVector FAvaClonerActorVisualizer::GetHandleSpacingLocation(const ACEClonerActor
 
 		OutLocation = ClonerRotation.RotateVector(Rotation.RotateVector(Axis * Scale)) * Radius;
 	}
-	else if (const UCEClonerCylinderLayout* CylinderLayout = InClonerActor->GetActiveLayout<UCEClonerCylinderLayout>())
+	else if (const UCEClonerCylinderLayout* CylinderLayout = InClonerComponent->GetActiveLayout<UCEClonerCylinderLayout>())
 	{
 		FRotator Rotation = CylinderLayout->GetRotation();
 		const FVector Scale = CylinderLayout->GetScale();
@@ -663,7 +709,7 @@ FVector FAvaClonerActorVisualizer::GetHandleSpacingLocation(const ACEClonerActor
 
 		OutLocation = ClonerRotation.RotateVector(Rotation.RotateVector(Axis * Scale)) * Dim;
 	}
-	else if (const UCEClonerSphereUniformLayout* SphereLayout = InClonerActor->GetActiveLayout<UCEClonerSphereUniformLayout>())
+	else if (const UCEClonerSphereUniformLayout* SphereLayout = InClonerComponent->GetActiveLayout<UCEClonerSphereUniformLayout>())
 	{
 		const float Radius = SphereLayout->GetRadius();
 		const FVector Scale = SphereLayout->GetScale();
@@ -671,14 +717,14 @@ FVector FAvaClonerActorVisualizer::GetHandleSpacingLocation(const ACEClonerActor
 		OutLocation = ClonerRotation.RotateVector(Axis) * Scale * Radius;
 	}
 
-	return InClonerActor->GetActorLocation() + OutLocation * ClonerScale;
+	return InClonerComponent->GetComponentLocation() + OutLocation * ClonerScale;
 }
 
-void FAvaClonerActorVisualizer::DrawSpacingButton(const ACEClonerActor* InClonerActor, const FSceneView* InView, FPrimitiveDrawInterface* InPDI, int32 InIconIndex, ECEClonerAxis InAxis, FLinearColor InColor) const
+void FAvaClonerActorVisualizer::DrawSpacingButton(const UCEClonerComponent* InClonerComponent, const FSceneView* InView, FPrimitiveDrawInterface* InPDI, int32 InIconIndex, ECEClonerAxis InAxis, FLinearColor InColor) const
 {
 	UTexture2D* SpacingSprite = IAvalancheComponentVisualizersModule::Get().GetSettings()->GetVisualizerSprite(UE::AvaShapes::BevelSprite);
 
-	if (!SpacingSprite || !SpacingSprite->GetResource() || !InClonerActor->GetEnabled() || InClonerActor->GetMeshCount() == 0)
+	if (!SpacingSprite || !SpacingSprite->GetResource() || !InClonerComponent->GetEnabled() || InClonerComponent->GetMeshCount() == 0)
 	{
 		return;
 	}
@@ -687,29 +733,16 @@ void FAvaClonerActorVisualizer::DrawSpacingButton(const ACEClonerActor* InCloner
 	float IconSize;
 	GetIconMetrics(InView, InIconIndex, IconLocation, IconSize);
 
-	IconLocation = GetHandleSpacingLocation(InClonerActor, InAxis);
+	IconLocation = GetHandleSpacingLocation(InClonerComponent, InAxis);
 
-	InPDI->SetHitProxy(new HAvaClonerActorSpacingHitProxy(InClonerActor->GetRootComponent(), InAxis));
+	InPDI->SetHitProxy(new HAvaClonerActorSpacingHitProxy(InClonerComponent, InAxis));
 	InPDI->DrawSprite(IconLocation, IconSize, IconSize, SpacingSprite->GetResource(), InColor, SDPG_Foreground, 0, 0, 0, 0, SE_BLEND_Opaque);
 	InPDI->SetHitProxy(nullptr);
 }
 
-void FAvaClonerActorVisualizer::OnPropertyModified(UObject* InPropertyObject, FName InPropertyName, EPropertyChangeType::Type InType)
-{
-	if (!InPropertyObject)
-	{
-		return;
-	}
-
-	InPropertyObject->Modify();
-	bHasBeenModified = true;
-	FProperty* PropertyModified = UE::AvaCore::GetProperty<UObject>(InPropertyName);
-	NotifyPropertyModified(InPropertyObject, PropertyModified, InType);
-}
-
 UActorComponent* FAvaClonerActorVisualizer::GetEditedComponent() const
 {
-	return ClonerActorWeak.IsValid() ? ClonerActorWeak->GetRootComponent() : nullptr;
+	return GetClonerComponent();
 }
 
 TMap<UObject*, TArray<FProperty*>> FAvaClonerActorVisualizer::GatherEditableProperties(UObject* InObject) const
@@ -789,7 +822,7 @@ bool FAvaClonerActorVisualizer::VisProxyHandleClick(FEditorViewportClient* InVie
 	if (InVisProxy->IsA(HAvaClonerActorSpacingHitProxy::StaticGetType()))
 	{
 		EndEditing();
-		ClonerActorWeak = Cast<ACEClonerActor>(DynMesh->GetOwner());
+		ClonerComponentWeak = DynMesh;
 		bEditingSpacing = true;
 		EditingAxis = static_cast<HAvaClonerActorSpacingHitProxy*>(InVisProxy)->Axis;
 		StartEditing(InViewportClient, DynMesh);
@@ -802,11 +835,11 @@ bool FAvaClonerActorVisualizer::VisProxyHandleClick(FEditorViewportClient* InVie
 
 bool FAvaClonerActorVisualizer::GetWidgetLocation(const FEditorViewportClient* InViewportClient, FVector& OutLocation) const
 {
-	if (const ACEClonerActor* ClonerActor = ClonerActorWeak.Get())
+	if (const UCEClonerComponent* ClonerComponent = ClonerComponentWeak.Get())
 	{
 		if (bEditingSpacing)
 		{
-			OutLocation = GetHandleSpacingLocation(ClonerActor, EditingAxis);
+			OutLocation = GetHandleSpacingLocation(ClonerComponent, EditingAxis);
 			return true;
 		}
 	}
@@ -873,60 +906,108 @@ bool FAvaClonerActorVisualizer::ResetValue(FEditorViewportClient* InViewportClie
 		return Super::ResetValue(InViewportClient, InHitProxy);
 	}
 
-	if (const ACEClonerActor* ClonerActor = Cast<ACEClonerActor>(ComponentHitProxy->Component->GetOwner()))
+	if (const UCEClonerComponent* ClonerComponent = Cast<UCEClonerComponent>(ComponentHitProxy->Component.Get()))
 	{
-		if (UCEClonerGridLayout* GridLayout = ClonerActor->GetActiveLayout<UCEClonerGridLayout>())
+		if (UCEClonerGridLayout* GridLayout = ClonerComponent->GetActiveLayout<UCEClonerGridLayout>())
 		{
 			FScopedTransaction Transaction(LOCTEXT("VisualizerResetValue", "Visualizer Reset Value"));
-			GridLayout->SetFlags(RF_Transactional);
 
-			GridLayout->SetSpacingX(100.f);
-			GridLayout->SetSpacingY(100.f);
-			GridLayout->SetSpacingZ(100.f);
-
-			OnPropertyModified(GridLayout, GET_MEMBER_NAME_CHECKED(UCEClonerGridLayout, SpacingX), EPropertyChangeType::ValueSet);
-			OnPropertyModified(GridLayout, GET_MEMBER_NAME_CHECKED(UCEClonerGridLayout, SpacingY), EPropertyChangeType::ValueSet);
-			OnPropertyModified(GridLayout, GET_MEMBER_NAME_CHECKED(UCEClonerGridLayout, SpacingZ), EPropertyChangeType::ValueSet);
+			ModifyProperty(
+				GridLayout
+				, GET_MEMBER_NAME_CHECKED(UCEClonerGridLayout, SpacingX)
+				, EPropertyChangeType::ValueSet
+				, [this, &GridLayout]()
+			{
+				GridLayout->SetSpacingX(100.f);
+			});
+
+			ModifyProperty(
+				GridLayout
+				, GET_MEMBER_NAME_CHECKED(UCEClonerGridLayout, SpacingY)
+				, EPropertyChangeType::ValueSet
+				, [this, &GridLayout]()
+			{
+				GridLayout->SetSpacingY(100.f);
+			});
+
+			ModifyProperty(
+				GridLayout
+				, GET_MEMBER_NAME_CHECKED(UCEClonerGridLayout, SpacingZ)
+				, EPropertyChangeType::ValueSet
+				, [this, &GridLayout]()
+			{
+				GridLayout->SetSpacingZ(100.f);
+			});
 		}
-		else if (UCEClonerLineLayout* LineLayout = ClonerActor->GetActiveLayout<UCEClonerLineLayout>())
+		else if (UCEClonerLineLayout* LineLayout = ClonerComponent->GetActiveLayout<UCEClonerLineLayout>())
 		{
 			FScopedTransaction Transaction(LOCTEXT("VisualizerResetValue", "Visualizer Reset Value"));
-			LineLayout->SetFlags(RF_Transactional);
-
-			LineLayout->SetSpacing(500.f);
 
-			OnPropertyModified(LineLayout, GET_MEMBER_NAME_CHECKED(UCEClonerLineLayout, Spacing), EPropertyChangeType::ValueSet);
+			ModifyProperty(
+				LineLayout
+				, GET_MEMBER_NAME_CHECKED(UCEClonerLineLayout, Spacing)
+				, EPropertyChangeType::ValueSet
+				, [this, &LineLayout]()
+			{
+				LineLayout->SetSpacing(500.f);
+			});
 		}
-		else if (UCEClonerHoneycombLayout* HoneycombLayout = ClonerActor->GetActiveLayout<UCEClonerHoneycombLayout>())
+		else if (UCEClonerHoneycombLayout* HoneycombLayout = ClonerComponent->GetActiveLayout<UCEClonerHoneycombLayout>())
 		{
 			FScopedTransaction Transaction(LOCTEXT( "VisualizerResetValue", "Visualizer Reset Value"));
-			HoneycombLayout->SetFlags(RF_Transactional);
 
-			HoneycombLayout->SetWidthSpacing(100.f);
-			HoneycombLayout->SetHeightSpacing(100.f);
-
-			OnPropertyModified(HoneycombLayout, GET_MEMBER_NAME_CHECKED(UCEClonerHoneycombLayout, WidthSpacing), EPropertyChangeType::ValueSet);
-			OnPropertyModified(HoneycombLayout, GET_MEMBER_NAME_CHECKED(UCEClonerHoneycombLayout, HeightSpacing), EPropertyChangeType::ValueSet);
+			ModifyProperty(
+				HoneycombLayout
+				, GET_MEMBER_NAME_CHECKED(UCEClonerHoneycombLayout, WidthSpacing)
+				, EPropertyChangeType::ValueSet
+				, [this, &HoneycombLayout]()
+			{
+				HoneycombLayout->SetWidthSpacing(100.f);
+			});
+
+			ModifyProperty(
+				HoneycombLayout
+				, GET_MEMBER_NAME_CHECKED(UCEClonerHoneycombLayout, HeightSpacing)
+				, EPropertyChangeType::ValueSet
+				, [this, &HoneycombLayout]()
+			{
+				HoneycombLayout->SetHeightSpacing(100.f);
+			});
 		}
-		else if (UCEClonerCircleLayout* CircleLayout = ClonerActor->GetActiveLayout<UCEClonerCircleLayout>())
+		else if (UCEClonerCircleLayout* CircleLayout = ClonerComponent->GetActiveLayout<UCEClonerCircleLayout>())
 		{
 			FScopedTransaction Transaction(LOCTEXT("VisualizerResetValue", "Visualizer Reset Value"));
-			CircleLayout->SetFlags(RF_Transactional);
 
-			CircleLayout->SetRadius(500.f);
-
-			OnPropertyModified(CircleLayout, GET_MEMBER_NAME_CHECKED(UCEClonerCircleLayout, Radius), EPropertyChangeType::ValueSet);
+			ModifyProperty(
+				CircleLayout
+				, GET_MEMBER_NAME_CHECKED(UCEClonerCircleLayout, Radius)
+				, EPropertyChangeType::ValueSet
+				, [this, &CircleLayout]()
+			{
+				CircleLayout->SetRadius(500.f);
+			});
 		}
-		else if (UCEClonerCylinderLayout* CylinderLayout = ClonerActor->GetActiveLayout<UCEClonerCylinderLayout>())
+		else if (UCEClonerCylinderLayout* CylinderLayout = ClonerComponent->GetActiveLayout<UCEClonerCylinderLayout>())
 		{
 			FScopedTransaction Transaction(LOCTEXT("VisualizerResetValue", "Visualizer Reset Value"));
-			CylinderLayout->SetFlags(RF_Transactional);
-
-			CylinderLayout->SetRadius(500.f);
-			CylinderLayout->SetHeight(1000.f);
 
-			OnPropertyModified(CylinderLayout, GET_MEMBER_NAME_CHECKED(UCEClonerCylinderLayout, Radius), EPropertyChangeType::ValueSet);
-			OnPropertyModified(CylinderLayout, GET_MEMBER_NAME_CHECKED(UCEClonerCylinderLayout, Height), EPropertyChangeType::ValueSet);
+			ModifyProperty(
+				CylinderLayout
+				, GET_MEMBER_NAME_CHECKED(UCEClonerCylinderLayout, Radius)
+				, EPropertyChangeType::ValueSet
+				, [this, &CylinderLayout]()
+			{
+				CylinderLayout->SetRadius(500.f);
+			});
+
+			ModifyProperty(
+				CylinderLayout
+				, GET_MEMBER_NAME_CHECKED(UCEClonerCylinderLayout, Height)
+				, EPropertyChangeType::ValueSet
+				, [this, &CylinderLayout]()
+			{
+				CylinderLayout->SetHeight(1000.f);
+			});
 		}
 	}
 
@@ -947,7 +1028,7 @@ void FAvaClonerActorVisualizer::EndEditing()
 {
 	Super::EndEditing();
 
-	ClonerActorWeak.Reset();
+	ClonerComponentWeak.Reset();
 	InitialSpacing = FVector::ZeroVector;
 	bEditingSpacing = false;
 }
diff --git a/Engine/Plugins/Experimental/Avalanche/Source/AvalancheEffectorsEditor/Private/Cloner/AvaClonerActorVis.h b/Engine/Plugins/Experimental/Avalanche/Source/AvalancheEffectorsEditor/Private/Cloner/AvaClonerActorVis.h
--- a/Engine/Plugins/Experimental/Avalanche/Source/AvalancheEffectorsEditor/Private/Cloner/AvaClonerActorVis.h
+++ b/Engine/Plugins/Experimental/Avalanche/Source/AvalancheEffectorsEditor/Private/Cloner/AvaClonerActorVis.h
@@ -6,7 +6,6 @@
 #include "AvaVisBase.h"
 #include "UObject/WeakObjectPtrTemplates.h"
 
-class ACEClonerActor;
 class UCEClonerComponent;
 
 struct HAvaClonerActorSpacingHitProxy : HAvaHitProxy
@@ -48,18 +47,17 @@ public:
 	virtual void DrawVisualizationNotEditing(const UActorComponent* InComponent, const FSceneView* InView, FPrimitiveDrawInterface* InPDI, int32& InOutIconIndex) override;
 	//~ End FAvaVisualizerBase
 
-	ACEClonerActor* GetClonerActor() const
+	UCEClonerComponent* GetClonerComponent() const
 	{
-		return ClonerActorWeak.Get();
+		return ClonerComponentWeak.Get();
 	}
 
 protected:
-	FVector GetHandleSpacingLocation(const ACEClonerActor* InClonerActor, ECEClonerAxis InAxis) const;
-	void DrawSpacingButton(const ACEClonerActor* InClonerActor, const FSceneView* InView, FPrimitiveDrawInterface* InPDI, int32 InIconIndex, ECEClonerAxis InAxis, FLinearColor InColor) const;
+	FVector GetHandleSpacingLocation(const UCEClonerComponent* InClonerComponent, ECEClonerAxis InAxis) const;
+	void DrawSpacingButton(const UCEClonerComponent* InClonerComponent, const FSceneView* InView, FPrimitiveDrawInterface* InPDI, int32 InIconIndex, ECEClonerAxis InAxis, FLinearColor InColor) const;
 
-	void OnPropertyModified(UObject* InPropertyObject, FName InPropertyName, EPropertyChangeType::Type InType = EPropertyChangeType::Interactive);
+	TWeakObjectPtr<UCEClonerComponent> ClonerComponentWeak = nullptr;
 
-	TWeakObjectPtr<ACEClonerActor> ClonerActorWeak = nullptr;
 	FVector InitialSpacing = FVector::ZeroVector;
 	bool bEditingSpacing = false;
 	ECEClonerAxis EditingAxis = ECEClonerAxis::X;
diff --git a/Engine/Plugins/Experimental/Avalanche/Source/AvalancheEffectorsEditor/Private/Effector/AvaEffectorActorVis.cpp b/Engine/Plugins/Experimental/Avalanche/Source/AvalancheEffectorsEditor/Private/Effector/AvaEffectorActorVis.cpp
--- a/Engine/Plugins/Experimental/Avalanche/Source/AvalancheEffectorsEditor/Private/Effector/AvaEffectorActorVis.cpp
+++ b/Engine/Plugins/Experimental/Avalanche/Source/AvalancheEffectorsEditor/Private/Effector/AvaEffectorActorVis.cpp
@@ -5,7 +5,11 @@
 #include "AvaShapeSprites.h"
 #include "AvaVisBase.h"
 #include "EditorViewportClient.h"
-#include "Effector/CEEffectorComponent.h"
+#include "Effector/Types/CEEffectorBoxType.h"
+#include "Effector/Types/CEEffectorPlaneType.h"
+#include "Effector/Types/CEEffectorRadialType.h"
+#include "Effector/Types/CEEffectorSphereType.h"
+#include "Effector/Types/CEEffectorTorusType.h"
 #include "Engine/Texture2D.h"
 #include "Framework/Application/SlateApplication.h"
 #include "IAvalancheComponentVisualizersModule.h"
@@ -22,92 +26,94 @@ FAvaEffectorActorVisualizer::FAvaEffectorActorVisualizer()
 	using namespace UE::AvaCore;
 
 	// Sphere
-	InnerRadiusProperty  = GetProperty<ACEEffectorActor>(GET_MEMBER_NAME_CHECKED(ACEEffectorActor, InnerRadius));
-	OuterRadiusProperty  = GetProperty<ACEEffectorActor>(GET_MEMBER_NAME_CHECKED(ACEEffectorActor, OuterRadius));
+	InnerRadiusProperty  = GetProperty<UCEEffectorSphereType>(GET_MEMBER_NAME_CHECKED(UCEEffectorSphereType, InnerRadius));
+	OuterRadiusProperty  = GetProperty<UCEEffectorSphereType>(GET_MEMBER_NAME_CHECKED(UCEEffectorSphereType, OuterRadius));
 
 	// Box
-	InnerExtentProperty  = GetProperty<ACEEffectorActor>(GET_MEMBER_NAME_CHECKED(ACEEffectorActor, InnerExtent));
-	OuterExtentProperty  = GetProperty<ACEEffectorActor>(GET_MEMBER_NAME_CHECKED(ACEEffectorActor, OuterExtent));
+	InnerExtentProperty  = GetProperty<UCEEffectorBoxType>(GET_MEMBER_NAME_CHECKED(UCEEffectorBoxType, InnerExtent));
+	OuterExtentProperty  = GetProperty<UCEEffectorBoxType>(GET_MEMBER_NAME_CHECKED(UCEEffectorBoxType, OuterExtent));
 
 	// Plane
-	PlaneSpacingProperty = GetProperty<ACEEffectorActor>(GET_MEMBER_NAME_CHECKED(ACEEffectorActor, PlaneSpacing));
+	PlaneSpacingProperty = GetProperty<UCEEffectorPlaneType>(GET_MEMBER_NAME_CHECKED(UCEEffectorPlaneType, PlaneSpacing));
 
 	// Radial
-	RadialAngleProperty = GetProperty<ACEEffectorActor>(GET_MEMBER_NAME_CHECKED(ACEEffectorActor, RadialAngle));
-	RadialMinRadiusProperty = GetProperty<ACEEffectorActor>(GET_MEMBER_NAME_CHECKED(ACEEffectorActor, RadialMinRadius));
-	RadialMaxRadiusProperty = GetProperty<ACEEffectorActor>(GET_MEMBER_NAME_CHECKED(ACEEffectorActor, RadialMaxRadius));
+	RadialAngleProperty = GetProperty<UCEEffectorRadialType>(GET_MEMBER_NAME_CHECKED(UCEEffectorRadialType, RadialAngle));
+	RadialMinRadiusProperty = GetProperty<UCEEffectorRadialType>(GET_MEMBER_NAME_CHECKED(UCEEffectorRadialType, RadialMinRadius));
+	RadialMaxRadiusProperty = GetProperty<UCEEffectorRadialType>(GET_MEMBER_NAME_CHECKED(UCEEffectorRadialType, RadialMaxRadius));
 
 	// Torus
-	TorusRadiusProperty = GetProperty<ACEEffectorActor>(GET_MEMBER_NAME_CHECKED(ACEEffectorActor, TorusRadius));
-	TorusInnerRadiusProperty = GetProperty<ACEEffectorActor>(GET_MEMBER_NAME_CHECKED(ACEEffectorActor, TorusInnerRadius));
-	TorusOuterRadiusProperty = GetProperty<ACEEffectorActor>(GET_MEMBER_NAME_CHECKED(ACEEffectorActor, TorusOuterRadius));
+	TorusRadiusProperty = GetProperty<UCEEffectorTorusType>(GET_MEMBER_NAME_CHECKED(UCEEffectorTorusType, TorusRadius));
+	TorusInnerRadiusProperty = GetProperty<UCEEffectorTorusType>(GET_MEMBER_NAME_CHECKED(UCEEffectorTorusType, TorusInnerRadius));
+	TorusOuterRadiusProperty = GetProperty<UCEEffectorTorusType>(GET_MEMBER_NAME_CHECKED(UCEEffectorTorusType, TorusOuterRadius));
 }
 
 void FAvaEffectorActorVisualizer::StoreInitialValues()
 {
 	Super::StoreInitialValues();
 
-	const ACEEffectorActor* EffectorActor = EffectorActorWeak.Get();
+	const UCEEffectorComponent* EffectorComponent = EffectorComponentWeak.Get();
 
-	if (GetEditedComponent() == nullptr || !EffectorActor)
+	if (!EffectorComponent)
 	{
 		return;
 	}
 
-	// Sphere
-	InitialInnerRadius = EffectorActor->GetInnerRadius();
-	InitialOuterRadius = EffectorActor->GetOuterRadius();
-
-	// Box
-	InitialInnerExtent = EffectorActor->GetInnerExtent();
-	InitialOuterExtent = EffectorActor->GetOuterExtent();
-
-	// Plane
-	InitialPlaneSpacing = EffectorActor->GetPlaneSpacing();
-
-	// Radial
-	InitialRadialAngle = EffectorActor->GetRadialAngle();
-	InitialRadialMinRadius = EffectorActor->GetRadialMinRadius();
-	InitialRadialMaxRadius = EffectorActor->GetRadialMaxRadius();
-
-	// Torus
-	InitialTorusRadius = EffectorActor->GetTorusRadius();
-	InitialTorusInnerRadius = EffectorActor->GetTorusInnerRadius();
-	InitialTorusOuterRadius = EffectorActor->GetTorusOuterRadius();
+	if (UCEEffectorSphereType* SphereType = EffectorComponent->GetActiveType<UCEEffectorSphereType>())
+	{
+		// Sphere
+		InitialInnerRadius = SphereType->GetInnerRadius();
+		InitialOuterRadius = SphereType->GetOuterRadius();
+	}
+	else if (UCEEffectorBoxType* BoxType = EffectorComponent->GetActiveType<UCEEffectorBoxType>())
+	{
+		// Box
+		InitialInnerExtent = BoxType->GetInnerExtent();
+		InitialOuterExtent = BoxType->GetOuterExtent();
+	}
+	else if (UCEEffectorPlaneType* PlaneType = EffectorComponent->GetActiveType<UCEEffectorPlaneType>())
+	{
+		// Plane
+		InitialPlaneSpacing = PlaneType->GetPlaneSpacing();
+	}
+	else if (UCEEffectorRadialType* RadialType = EffectorComponent->GetActiveType<UCEEffectorRadialType>())
+	{
+		// Radial
+		InitialRadialAngle = RadialType->GetRadialAngle();
+		InitialRadialMinRadius = RadialType->GetRadialMinRadius();
+		InitialRadialMaxRadius = RadialType->GetRadialMaxRadius();
+	}
+	else if (UCEEffectorTorusType* TorusType = EffectorComponent->GetActiveType<UCEEffectorTorusType>())
+	{
+		// Torus
+		InitialTorusRadius = TorusType->GetTorusRadius();
+		InitialTorusInnerRadius = TorusType->GetTorusInnerRadius();
+		InitialTorusOuterRadius = TorusType->GetTorusOuterRadius();
+	}
 }
 
 FBox FAvaEffectorActorVisualizer::GetComponentBounds(const UActorComponent* InComponent) const
 {
 	if (const UCEEffectorComponent* EffectorComponent = Cast<UCEEffectorComponent>(InComponent))
 	{
-		if (const ACEEffectorActor* EffectorActor = Cast<ACEEffectorActor>(EffectorComponent->GetOwner()))
+		if (UCEEffectorSphereType* SphereType = EffectorComponent->GetActiveType<UCEEffectorSphereType>())
 		{
-			const ECEClonerEffectorType Type = EffectorActor->GetType();
-
-			if (Type == ECEClonerEffectorType::Box)
-			{
-				return FBox(-EffectorActor->GetOuterExtent(), EffectorActor->GetOuterExtent());
-			}
-
-			if (Type == ECEClonerEffectorType::Sphere)
-			{
-				return FBox(-FVector(EffectorActor->GetOuterRadius() / 2), FVector(EffectorActor->GetOuterRadius() / 2));
-			}
-
-			if (Type == ECEClonerEffectorType::Plane)
-			{
-				return FBox(-FVector(EffectorActor->GetPlaneSpacing() / 2), FVector(EffectorActor->GetPlaneSpacing() / 2));
-			}
-
-			if (Type == ECEClonerEffectorType::Radial)
-			{
-				return FBox(-FVector(EffectorActor->GetRadialMaxRadius() / 2), FVector(EffectorActor->GetRadialMaxRadius() / 2));
-			}
-
-			if (Type == ECEClonerEffectorType::Torus)
-			{
-				return FBox(-FVector((EffectorActor->GetTorusRadius() + EffectorActor->GetTorusOuterRadius()) / 2), FVector((EffectorActor->GetTorusRadius() + EffectorActor->GetTorusOuterRadius()) / 2));
-			}
+			return FBox(-FVector(SphereType->GetOuterRadius() / 2), FVector(SphereType->GetOuterRadius() / 2));
+		}
+		else if (UCEEffectorBoxType* BoxType = EffectorComponent->GetActiveType<UCEEffectorBoxType>())
+		{
+			return FBox(-BoxType->GetOuterExtent(), BoxType->GetOuterExtent());
+		}
+		else if (UCEEffectorPlaneType* PlaneType = EffectorComponent->GetActiveType<UCEEffectorPlaneType>())
+		{
+			return FBox(-FVector(PlaneType->GetPlaneSpacing() / 2), FVector(PlaneType->GetPlaneSpacing() / 2));
+		}
+		else if (UCEEffectorRadialType* RadialType = EffectorComponent->GetActiveType<UCEEffectorRadialType>())
+		{
+			return FBox(-FVector(RadialType->GetRadialMaxRadius() / 2), FVector(RadialType->GetRadialMaxRadius() / 2));
+		}
+		else if (UCEEffectorTorusType* TorusType = EffectorComponent->GetActiveType<UCEEffectorTorusType>())
+		{
+			return FBox(-FVector((TorusType->GetTorusRadius() + TorusType->GetTorusOuterRadius()) / 2), FVector((TorusType->GetTorusRadius() + TorusType->GetTorusOuterRadius()) / 2));
 		}
 	}
 
@@ -121,28 +127,26 @@ bool FAvaEffectorActorVisualizer::HandleInputDeltaInternal(FEditorViewportClient
 		return false;
 	}
 
-	if (ACEEffectorActor* EffectorActor = EffectorActorWeak.Get())
+	if (UCEEffectorComponent* EffectorComponent = EffectorComponentWeak.Get())
 	{
-		const ECEClonerEffectorType Type = EffectorActor->GetType();
-
 		if (GetViewportWidgetMode(InViewportClient) == UE::Widget::WM_Translate)
 		{
 			if (GetViewportWidgetAxisList(InViewportClient) & EAxisList::XYZ)
 			{
-				if (Type == ECEClonerEffectorType::Box)
+				if (UCEEffectorBoxType* BoxType = EffectorComponent->GetActiveType<UCEEffectorBoxType>())
 				{
 					if (EditingHandleType == HandleTypeInnerZone)
 					{
-						ModifyProperty(EffectorActor, InnerExtentProperty, EPropertyChangeType::Interactive, [this, &EffectorActor, &InAccumulatedTranslation]()
+						ModifyProperty(BoxType, InnerExtentProperty, EPropertyChangeType::Interactive, [this, &BoxType, &InAccumulatedTranslation]()
 						{
-							EffectorActor->SetInnerExtent(InitialInnerExtent + InAccumulatedTranslation);
+							BoxType->SetInnerExtent(InitialInnerExtent + InAccumulatedTranslation);
 						});
 					}
 					else if (EditingHandleType == HandleTypeOuterZone)
 					{
-						ModifyProperty(EffectorActor, OuterExtentProperty, EPropertyChangeType::Interactive, [this, &EffectorActor, &InAccumulatedTranslation]()
+						ModifyProperty(BoxType, OuterExtentProperty, EPropertyChangeType::Interactive, [this, &BoxType, &InAccumulatedTranslation]()
 						{
-							EffectorActor->SetOuterExtent(InitialOuterExtent + InAccumulatedTranslation);
+							BoxType->SetOuterExtent(InitialOuterExtent + InAccumulatedTranslation);
 						});
 					}
 
@@ -152,66 +156,66 @@ bool FAvaEffectorActorVisualizer::HandleInputDeltaInternal(FEditorViewportClient
 
 			if (GetViewportWidgetAxisList(InViewportClient) & EAxisList::Y)
 			{
-				if (Type == ECEClonerEffectorType::Plane)
+				if (UCEEffectorPlaneType* PlaneType = EffectorComponent->GetActiveType<UCEEffectorPlaneType>())
 				{
 					if (EditingHandleType == HandleTypeInnerZone || EditingHandleType == HandleTypeOuterZone)
 					{
-						ModifyProperty(EffectorActor, PlaneSpacingProperty, EPropertyChangeType::Interactive, [this, &EffectorActor, &InAccumulatedTranslation]()
+						ModifyProperty(PlaneType, PlaneSpacingProperty, EPropertyChangeType::Interactive, [this, &PlaneType, &InAccumulatedTranslation]()
 						{
-							EffectorActor->SetPlaneSpacing(InitialPlaneSpacing + InAccumulatedTranslation.Y);
+							PlaneType->SetPlaneSpacing(InitialPlaneSpacing + InAccumulatedTranslation.Y);
 						});
 					}
 
 					return true;
 				}
 
-				if (Type == ECEClonerEffectorType::Sphere)
+				if (UCEEffectorSphereType* SphereType = EffectorComponent->GetActiveType<UCEEffectorSphereType>())
 				{
 					if (EditingHandleType == HandleTypeInnerZone)
 					{
-						ModifyProperty(EffectorActor, InnerRadiusProperty, EPropertyChangeType::Interactive, [this, &EffectorActor, &InAccumulatedTranslation]()
+						ModifyProperty(SphereType, InnerRadiusProperty, EPropertyChangeType::Interactive, [this, &SphereType, &InAccumulatedTranslation]()
 						{
-							EffectorActor->SetInnerRadius(InitialInnerRadius + InAccumulatedTranslation.Y);
+							SphereType->SetInnerRadius(InitialInnerRadius + InAccumulatedTranslation.Y);
 						});
 					}
 					else if (EditingHandleType == HandleTypeOuterZone)
 					{
-						ModifyProperty(EffectorActor, OuterRadiusProperty, EPropertyChangeType::Interactive, [this, &EffectorActor, &InAccumulatedTranslation]
+						ModifyProperty(SphereType, OuterRadiusProperty, EPropertyChangeType::Interactive, [this, &SphereType, &InAccumulatedTranslation]
 						{
-							EffectorActor->SetOuterRadius(InitialOuterRadius + InAccumulatedTranslation.Y);
+							SphereType->SetOuterRadius(InitialOuterRadius + InAccumulatedTranslation.Y);
 						});
 					}
 
 					return true;
 				}
 
-				if (Type == ECEClonerEffectorType::Radial)
+				if (UCEEffectorRadialType* RadialType = EffectorComponent->GetActiveType<UCEEffectorRadialType>())
 				{
 					if (EditingHandleType == HandleTypeInnerZone)
 					{
-						ModifyProperty(EffectorActor, RadialMinRadiusProperty, EPropertyChangeType::Interactive, [this, &EffectorActor, &InAccumulatedTranslation]
+						ModifyProperty(RadialType, RadialMinRadiusProperty, EPropertyChangeType::Interactive, [this, &RadialType, &InAccumulatedTranslation]
 						{
-							EffectorActor->SetRadialMinRadius(InitialRadialMinRadius + InAccumulatedTranslation.Y);
+							RadialType->SetRadialMinRadius(InitialRadialMinRadius + InAccumulatedTranslation.Y);
 						});
 					}
 					else if (EditingHandleType == HandleTypeOuterZone)
 					{
-						ModifyProperty(EffectorActor, RadialMaxRadiusProperty, EPropertyChangeType::Interactive, [this, &EffectorActor, &InAccumulatedTranslation]()
+						ModifyProperty(RadialType, RadialMaxRadiusProperty, EPropertyChangeType::Interactive, [this, &RadialType, &InAccumulatedTranslation]()
 						{
-							EffectorActor->SetRadialMaxRadius(InitialRadialMaxRadius + InAccumulatedTranslation.Y);
+							RadialType->SetRadialMaxRadius(InitialRadialMaxRadius + InAccumulatedTranslation.Y);
 						});
 					}
 
 					return true;
 				}
 
-				if (Type == ECEClonerEffectorType::Torus)
+				if (UCEEffectorTorusType* TorusType = EffectorComponent->GetActiveType<UCEEffectorTorusType>())
 				{
 					if (EditingHandleType == HandleTypeRadius)
 					{
-						ModifyProperty(EffectorActor, TorusRadiusProperty, EPropertyChangeType::Interactive, [this, &EffectorActor, &InAccumulatedTranslation]()
+						ModifyProperty(TorusType, TorusRadiusProperty, EPropertyChangeType::Interactive, [this, &TorusType, &InAccumulatedTranslation]()
 						{
-							EffectorActor->SetTorusRadius(InitialTorusRadius + InAccumulatedTranslation.Y);
+							TorusType->SetTorusRadius(InitialTorusRadius + InAccumulatedTranslation.Y);
 						});
 					}
 
@@ -221,20 +225,20 @@ bool FAvaEffectorActorVisualizer::HandleInputDeltaInternal(FEditorViewportClient
 
 			if (GetViewportWidgetAxisList(InViewportClient) & EAxisList::Z)
 			{
-				if (Type == ECEClonerEffectorType::Torus)
+				if (UCEEffectorTorusType* TorusType = EffectorComponent->GetActiveType<UCEEffectorTorusType>())
 				{
 					if (EditingHandleType == HandleTypeInnerZone)
 					{
-						ModifyProperty(EffectorActor, TorusInnerRadiusProperty, EPropertyChangeType::Interactive, [this, &EffectorActor, &InAccumulatedTranslation]()
+						ModifyProperty(TorusType, TorusInnerRadiusProperty, EPropertyChangeType::Interactive, [this, &TorusType, &InAccumulatedTranslation]()
 						{
-							EffectorActor->SetTorusInnerRadius(InitialTorusInnerRadius + InAccumulatedTranslation.Z);
+							TorusType->SetTorusInnerRadius(InitialTorusInnerRadius + InAccumulatedTranslation.Z);
 						});
 					}
 					else if (EditingHandleType == HandleTypeOuterZone)
 					{
-						ModifyProperty(EffectorActor, TorusOuterRadiusProperty, EPropertyChangeType::Interactive, [this, &EffectorActor, &InAccumulatedTranslation]()
+						ModifyProperty(TorusType, TorusOuterRadiusProperty, EPropertyChangeType::Interactive, [this, &TorusType, &InAccumulatedTranslation]()
 						{
-							EffectorActor->SetTorusOuterRadius(InitialTorusOuterRadius + InAccumulatedTranslation.Z);
+							TorusType->SetTorusOuterRadius(InitialTorusOuterRadius + InAccumulatedTranslation.Z);
 						});
 					}
 
@@ -246,13 +250,13 @@ bool FAvaEffectorActorVisualizer::HandleInputDeltaInternal(FEditorViewportClient
 		{
 			if (GetViewportWidgetAxisList(InViewportClient) & EAxisList::Z)
 			{
-				if (Type == ECEClonerEffectorType::Radial)
+				if (UCEEffectorRadialType* RadialType = EffectorComponent->GetActiveType<UCEEffectorRadialType>())
 				{
 					if (EditingHandleType == HandleTypeAngle)
 					{
-						ModifyProperty(EffectorActor, RadialAngleProperty, EPropertyChangeType::Interactive, [this, &EffectorActor, &InAccumulatedRotation]()
+						ModifyProperty(RadialType, RadialAngleProperty, EPropertyChangeType::Interactive, [this, &RadialType, &InAccumulatedRotation]()
 						{
-							EffectorActor->SetRadialAngle(InitialRadialAngle + InAccumulatedRotation.Yaw);
+							RadialType->SetRadialAngle(InitialRadialAngle + InAccumulatedRotation.Yaw);
 						});
 					}
 
@@ -280,31 +284,24 @@ void FAvaEffectorActorVisualizer::DrawVisualizationEditing(const UActorComponent
 		return;
 	}
 
-	const ACEEffectorActor* EffectorActor = Cast<ACEEffectorActor>(EffectorComponent->GetOwner());
-
-	if (!EffectorActor)
-	{
-		return;
-	}
-
-	if (EffectorActor->GetType() != ECEClonerEffectorType::Plane)
+	if (!EffectorComponent->GetActiveType<UCEEffectorPlaneType>())
 	{
-		DrawZoneButton(EffectorActor, InView, InPDI, InOutIconIndex, HandleTypeInnerZone, FAvaVisualizerBase::Active);
+		DrawZoneButton(EffectorComponent, InView, InPDI, InOutIconIndex, HandleTypeInnerZone, FAvaVisualizerBase::Active);
 		InOutIconIndex++;
 	}
 
-	DrawZoneButton(EffectorActor, InView, InPDI, InOutIconIndex, HandleTypeOuterZone, FAvaVisualizerBase::Active);
+	DrawZoneButton(EffectorComponent, InView, InPDI, InOutIconIndex, HandleTypeOuterZone, FAvaVisualizerBase::Active);
 	InOutIconIndex++;
 
-	if (EffectorActor->GetType() == ECEClonerEffectorType::Radial)
+	if (EffectorComponent->GetActiveType<UCEEffectorRadialType>())
 	{
-		DrawZoneButton(EffectorActor, InView, InPDI, InOutIconIndex, HandleTypeAngle, FAvaVisualizerBase::Active);
+		DrawZoneButton(EffectorComponent, InView, InPDI, InOutIconIndex, HandleTypeAngle, FAvaVisualizerBase::Active);
 		InOutIconIndex++;
 	}
 
-	if (EffectorActor->GetType() == ECEClonerEffectorType::Torus)
+	if (EffectorComponent->GetActiveType<UCEEffectorTorusType>())
 	{
-		DrawZoneButton(EffectorActor, InView, InPDI, InOutIconIndex, HandleTypeRadius, FAvaVisualizerBase::Active);
+		DrawZoneButton(EffectorComponent, InView, InPDI, InOutIconIndex, HandleTypeRadius, FAvaVisualizerBase::Active);
 		InOutIconIndex++;
 	}
 }
@@ -320,58 +317,50 @@ void FAvaEffectorActorVisualizer::DrawVisualizationNotEditing(const UActorCompon
 		return;
 	}
 
-	const ACEEffectorActor* EffectorActor = Cast<ACEEffectorActor>(EffectorComponent->GetOwner());
-
-	if (!EffectorActor)
-	{
-		return;
-	}
-
-	if (EffectorActor->GetType() != ECEClonerEffectorType::Plane)
+	if (!EffectorComponent->GetActiveType<UCEEffectorPlaneType>())
 	{
-		DrawZoneButton(EffectorActor, InView, InPDI, InOutIconIndex, HandleTypeInnerZone, FAvaVisualizerBase::Inactive);
+		DrawZoneButton(EffectorComponent, InView, InPDI, InOutIconIndex, HandleTypeInnerZone, FAvaVisualizerBase::Inactive);
 		InOutIconIndex++;
 	}
 
-	DrawZoneButton(EffectorActor, InView, InPDI, InOutIconIndex, HandleTypeOuterZone, FAvaVisualizerBase::Inactive);
+	DrawZoneButton(EffectorComponent, InView, InPDI, InOutIconIndex, HandleTypeOuterZone, FAvaVisualizerBase::Inactive);
 	InOutIconIndex++;
 
-	if (EffectorActor->GetType() == ECEClonerEffectorType::Radial)
+	if (EffectorComponent->GetActiveType<UCEEffectorRadialType>())
 	{
-		DrawZoneButton(EffectorActor, InView, InPDI, InOutIconIndex, HandleTypeAngle, FAvaVisualizerBase::Inactive);
+		DrawZoneButton(EffectorComponent, InView, InPDI, InOutIconIndex, HandleTypeAngle, FAvaVisualizerBase::Inactive);
 		InOutIconIndex++;
 	}
 
-	if (EffectorActor->GetType() == ECEClonerEffectorType::Torus)
+	if (EffectorComponent->GetActiveType<UCEEffectorTorusType>())
 	{
-		DrawZoneButton(EffectorActor, InView, InPDI, InOutIconIndex, HandleTypeRadius, FAvaVisualizerBase::Inactive);
+		DrawZoneButton(EffectorComponent, InView, InPDI, InOutIconIndex, HandleTypeRadius, FAvaVisualizerBase::Inactive);
 		InOutIconIndex++;
 	}
 }
 
-FVector FAvaEffectorActorVisualizer::GetHandleZoneLocation(const ACEEffectorActor* InEffectorActor, int32 InHandleType) const
+FVector FAvaEffectorActorVisualizer::GetHandleZoneLocation(const UCEEffectorComponent* InEffectorComponent, int32 InHandleType) const
 {
-	const ECEClonerEffectorType Type = InEffectorActor->GetType();
-	const FVector EffectorScale = InEffectorActor->GetActorScale();
-	const FRotator EffectorRotation = InEffectorActor->GetActorRotation();
-	FVector OutLocation = InEffectorActor->GetActorLocation();
+	const FVector EffectorScale = InEffectorComponent->GetComponentScale();
+	const FRotator EffectorRotation = InEffectorComponent->GetComponentRotation();
+	FVector OutLocation = InEffectorComponent->GetComponentLocation();
 
 	// To avoid inner/outer handle to be near actor gizmo and hard to select
 	constexpr float MinHandleOffset = 50.f;
 	constexpr float MaxHandleOffset = 100.f;
 
-	if (Type == ECEClonerEffectorType::Box)
+	if (UCEEffectorBoxType* BoxType = InEffectorComponent->GetActiveType<UCEEffectorBoxType>())
 	{
 		if (InHandleType == HandleTypeInnerZone)
 		{
-			OutLocation += EffectorRotation.RotateVector(InEffectorActor->GetInnerExtent()) * EffectorScale;
+			OutLocation += EffectorRotation.RotateVector(BoxType->GetInnerExtent()) * EffectorScale;
 		}
 		else if (InHandleType == HandleTypeOuterZone)
 		{
-			OutLocation += EffectorRotation.RotateVector(InEffectorActor->GetOuterExtent()) * EffectorScale;
+			OutLocation += EffectorRotation.RotateVector(BoxType->GetOuterExtent()) * EffectorScale;
 		}
 	}
-	else if (Type == ECEClonerEffectorType::Plane)
+	else if (UCEEffectorPlaneType* PlaneType = InEffectorComponent->GetActiveType<UCEEffectorPlaneType>())
 	{
 		if (InHandleType == HandleTypeInnerZone
 			|| InHandleType == HandleTypeOuterZone)
@@ -379,24 +368,24 @@ FVector FAvaEffectorActorVisualizer::GetHandleZoneLocation(const ACEEffectorActo
 			const float ComponentScale = (EffectorRotation.RotateVector(-FVector::YAxisVector) * EffectorScale).Length();
 			const FVector HandleAxis = EffectorRotation.RotateVector(FVector::YAxisVector);
 
-			OutLocation += HandleAxis * (InEffectorActor->GetPlaneSpacing() / 2) * ComponentScale;
+			OutLocation += HandleAxis * (PlaneType->GetPlaneSpacing() / 2) * ComponentScale;
 		}
 	}
-	else if (Type == ECEClonerEffectorType::Sphere)
+	else if (UCEEffectorSphereType* SphereType = InEffectorComponent->GetActiveType<UCEEffectorSphereType>())
 	{
 		const float MinComponentScale = FMath::Min<float>(FMath::Min<float>(EffectorScale.X, EffectorScale.Y), EffectorScale.Z);
 		const FVector HandleAxis = EffectorRotation.RotateVector(FVector::YAxisVector);
 
 		if (InHandleType == HandleTypeInnerZone)
 		{
-			OutLocation += FVector::Max(HandleAxis * InEffectorActor->GetInnerRadius(), HandleAxis * MinHandleOffset) * MinComponentScale;
+			OutLocation += FVector::Max(HandleAxis * SphereType->GetInnerRadius(), HandleAxis * MinHandleOffset) * MinComponentScale;
 		}
 		else if (InHandleType == HandleTypeOuterZone)
 		{
-			OutLocation += FVector::Max(HandleAxis * InEffectorActor->GetOuterRadius(), HandleAxis * MaxHandleOffset) * MinComponentScale;
+			OutLocation += FVector::Max(HandleAxis * SphereType->GetOuterRadius(), HandleAxis * MaxHandleOffset) * MinComponentScale;
 		}
 	}
-	else if (Type == ECEClonerEffectorType::Radial)
+	else if (UCEEffectorRadialType* RadialType = InEffectorComponent->GetActiveType<UCEEffectorRadialType>())
 	{
 		const float MinComponentScale = FMath::Min<float>(FMath::Min<float>(EffectorScale.X, EffectorScale.Y), EffectorScale.Z);
 		const FVector UpHandleAxis = EffectorRotation.RotateVector(FVector::ZAxisVector);
@@ -404,18 +393,18 @@ FVector FAvaEffectorActorVisualizer::GetHandleZoneLocation(const ACEEffectorActo
 
 		if (InHandleType == HandleTypeInnerZone)
 		{
-			OutLocation += FVector::Max(RightHandleAxis * InEffectorActor->GetRadialMinRadius(), RightHandleAxis * MinHandleOffset) * MinComponentScale;
+			OutLocation += FVector::Max(RightHandleAxis * RadialType->GetRadialMinRadius(), RightHandleAxis * MinHandleOffset) * MinComponentScale;
 		}
 		else if (InHandleType == HandleTypeOuterZone)
 		{
-			OutLocation += FVector::Max(RightHandleAxis * InEffectorActor->GetRadialMaxRadius(), RightHandleAxis * MaxHandleOffset) * MinComponentScale;
+			OutLocation += FVector::Max(RightHandleAxis * RadialType->GetRadialMaxRadius(), RightHandleAxis * MaxHandleOffset) * MinComponentScale;
 		}
 		else if (InHandleType == HandleTypeAngle)
 		{
 			OutLocation += UpHandleAxis * MaxHandleOffset * MinComponentScale;
 		}
 	}
-	else if (Type == ECEClonerEffectorType::Torus)
+	else if (UCEEffectorTorusType* TorusType = InEffectorComponent->GetActiveType<UCEEffectorTorusType>())
 	{
 		const float MinComponentScale = FMath::Min<float>(FMath::Min<float>(EffectorScale.X, EffectorScale.Y), EffectorScale.Z);
 		const FVector UpHandleAxis = EffectorRotation.RotateVector(FVector::ZAxisVector);
@@ -423,24 +412,24 @@ FVector FAvaEffectorActorVisualizer::GetHandleZoneLocation(const ACEEffectorActo
 
 		if (InHandleType == HandleTypeInnerZone)
 		{
-			OutLocation += RightHandleAxis * InEffectorActor->GetTorusRadius()
-				+ FVector::Max(UpHandleAxis * InEffectorActor->GetTorusInnerRadius(), UpHandleAxis * MinHandleOffset) * MinComponentScale;
+			OutLocation += RightHandleAxis * TorusType->GetTorusRadius()
+				+ FVector::Max(UpHandleAxis * TorusType->GetTorusInnerRadius(), UpHandleAxis * MinHandleOffset) * MinComponentScale;
 		}
 		else if (InHandleType == HandleTypeOuterZone)
 		{
-			OutLocation += RightHandleAxis * InEffectorActor->GetTorusRadius()
-				+ FVector::Max(UpHandleAxis * InEffectorActor->GetTorusOuterRadius(), UpHandleAxis * MaxHandleOffset) * MinComponentScale;
+			OutLocation += RightHandleAxis * TorusType->GetTorusRadius()
+				+ FVector::Max(UpHandleAxis * TorusType->GetTorusOuterRadius(), UpHandleAxis * MaxHandleOffset) * MinComponentScale;
 		}
 		else if (InHandleType == HandleTypeRadius)
 		{
-			OutLocation += FVector::Max(RightHandleAxis * InEffectorActor->GetTorusRadius(), RightHandleAxis * MinHandleOffset) * MinComponentScale;
+			OutLocation += FVector::Max(RightHandleAxis * TorusType->GetTorusRadius(), RightHandleAxis * MinHandleOffset) * MinComponentScale;
 		}
 	}
 
 	return OutLocation;
 }
 
-void FAvaEffectorActorVisualizer::DrawZoneButton(const ACEEffectorActor* InEffectorActor, const FSceneView* InView, FPrimitiveDrawInterface* InPDI, int32 InIconIndex, int32 InHandleType, FLinearColor InColor) const
+void FAvaEffectorActorVisualizer::DrawZoneButton(const UCEEffectorComponent* InEffectorComponent, const FSceneView* InView, FPrimitiveDrawInterface* InPDI, int32 InIconIndex, int32 InHandleType, FLinearColor InColor) const
 {
 	UTexture2D* ZoneSprite = IAvalancheComponentVisualizersModule::Get().GetSettings()->GetVisualizerSprite(UE::AvaShapes::BevelSprite);
 
@@ -453,44 +442,41 @@ void FAvaEffectorActorVisualizer::DrawZoneButton(const ACEEffectorActor* InEffec
 	float IconSize;
 	GetIconMetrics(InView, InIconIndex, IconLocation, IconSize);
 
-	IconLocation = GetHandleZoneLocation(InEffectorActor, InHandleType);
+	IconLocation = GetHandleZoneLocation(InEffectorComponent, InHandleType);
 
-	InPDI->SetHitProxy(new HAvaEffectorActorZoneHitProxy(InEffectorActor->GetRootComponent(), InHandleType));
+	InPDI->SetHitProxy(new HAvaEffectorActorZoneHitProxy(InEffectorComponent, InHandleType));
 	InPDI->DrawSprite(IconLocation, IconSize, IconSize, ZoneSprite->GetResource(), InColor, SDPG_Foreground, 0, 0, 0, 0, SE_BLEND_Opaque);
 	InPDI->SetHitProxy(nullptr);
 }
 
 UActorComponent* FAvaEffectorActorVisualizer::GetEditedComponent() const
 {
-	return EffectorActorWeak.IsValid() ? EffectorActorWeak->GetRootComponent() : nullptr;
+	return EffectorComponentWeak.Get();
 }
 
 TMap<UObject*, TArray<FProperty*>> FAvaEffectorActorVisualizer::GatherEditableProperties(UObject* InObject) const
 {
 	if (const UCEEffectorComponent* EffectorComponent = Cast<UCEEffectorComponent>(InObject))
 	{
-		if (ACEEffectorActor* EffectorActor = EffectorComponent->GetOuterACEEffectorActor())
+		if (UCEEffectorSphereType* SphereType = EffectorComponent->GetActiveType<UCEEffectorSphereType>())
 		{
-			switch (EffectorActor->GetType())
-			{
-				case ECEClonerEffectorType::Plane:
-					return {{EffectorActor, {PlaneSpacingProperty}}};
-
-				case ECEClonerEffectorType::Box:
-					return {{EffectorActor, {InnerExtentProperty, OuterExtentProperty}}};
-
-				case ECEClonerEffectorType::Sphere:
-					return {{EffectorActor, {InnerRadiusProperty, OuterRadiusProperty}}};
-
-				case ECEClonerEffectorType::Radial:
-					return {{EffectorActor, {RadialAngleProperty, RadialMinRadiusProperty, RadialMaxRadiusProperty}}};
-
-				case ECEClonerEffectorType::Torus:
-					return {{EffectorActor, {TorusRadiusProperty, TorusInnerRadiusProperty, TorusOuterRadiusProperty}}};
-
-				default:
-					return {};
-			}
+			return {{SphereType, {InnerRadiusProperty, OuterRadiusProperty}}};
+		}
+		else if (UCEEffectorBoxType* BoxType = EffectorComponent->GetActiveType<UCEEffectorBoxType>())
+		{
+			return {{BoxType, {InnerExtentProperty, OuterExtentProperty}}};
+		}
+		else if (UCEEffectorPlaneType* PlaneType = EffectorComponent->GetActiveType<UCEEffectorPlaneType>())
+		{
+			return {{PlaneType, {PlaneSpacingProperty}}};
+		}
+		else if (UCEEffectorRadialType* RadialType = EffectorComponent->GetActiveType<UCEEffectorRadialType>())
+		{
+			return {{RadialType, {RadialAngleProperty, RadialMinRadiusProperty, RadialMaxRadiusProperty}}};
+		}
+		else if (UCEEffectorTorusType* TorusType = EffectorComponent->GetActiveType<UCEEffectorTorusType>())
+		{
+			return {{TorusType, {TorusRadiusProperty, TorusInnerRadiusProperty, TorusOuterRadiusProperty}}};
 		}
 	}
 
@@ -507,7 +493,7 @@ bool FAvaEffectorActorVisualizer::VisProxyHandleClick(FEditorViewportClient* InV
 
 	UActorComponent* Component = const_cast<UActorComponent*>(InVisProxy->Component.Get());
 
-	if (!Component || !Component->GetOwner()->IsA<ACEEffectorActor>())
+	if (!Component || !Component->IsA<UCEEffectorComponent>())
 	{
 		return Super::VisProxyHandleClick(InViewportClient, InVisProxy, InClick);
 	}
@@ -515,7 +501,7 @@ bool FAvaEffectorActorVisualizer::VisProxyHandleClick(FEditorViewportClient* InV
 	if (InVisProxy->IsA(HAvaEffectorActorZoneHitProxy::StaticGetType()))
 	{
 		EndEditing();
-		EffectorActorWeak = Cast<ACEEffectorActor>(InVisProxy->Component->GetOwner());
+		EffectorComponentWeak = Cast<UCEEffectorComponent>(Component);
 		EditingHandleType = static_cast<HAvaEffectorActorZoneHitProxy*>(InVisProxy)->HandleType;
 		StartEditing(InViewportClient, Component);
 
@@ -527,9 +513,9 @@ bool FAvaEffectorActorVisualizer::VisProxyHandleClick(FEditorViewportClient* InV
 
 bool FAvaEffectorActorVisualizer::GetWidgetLocation(const FEditorViewportClient* InViewportClient, FVector& OutLocation) const
 {
-	if (const ACEEffectorActor* EffectorActor = EffectorActorWeak.Get())
+	if (const UCEEffectorComponent* EffectorComponent = EffectorComponentWeak.Get())
 	{
-		OutLocation = GetHandleZoneLocation(EffectorActor, EditingHandleType);
+		OutLocation = GetHandleZoneLocation(EffectorComponent, EditingHandleType);
 		return true;
 	}
 
@@ -559,43 +545,46 @@ bool FAvaEffectorActorVisualizer::GetWidgetMode(const FEditorViewportClient* InV
 bool FAvaEffectorActorVisualizer::GetWidgetAxisList(const FEditorViewportClient* InViewportClient,
 	UE::Widget::EWidgetMode InWidgetMode, EAxisList::Type& OutAxisList) const
 {
-	if (EditingHandleType == HandleTypeInnerZone
-		|| EditingHandleType == HandleTypeOuterZone)
+	if (UCEEffectorComponent* EffectorComponent = EffectorComponentWeak.Get())
 	{
-		if (EffectorActorWeak->GetType() == ECEClonerEffectorType::Torus)
-		{
-			OutAxisList = EAxisList::Type::Z;
-		}
-		else if (EffectorActorWeak->GetType() == ECEClonerEffectorType::Box)
-		{
-			OutAxisList = EAxisList::Type::XYZ;
-		}
-		else
+		if (EditingHandleType == HandleTypeInnerZone
+		|| EditingHandleType == HandleTypeOuterZone)
 		{
-			OutAxisList = EAxisList::Type::Y;
-		}
+			if (!!EffectorComponent->GetActiveType<UCEEffectorTorusType>())
+			{
+				OutAxisList = EAxisList::Type::Z;
+			}
+			else if (!!EffectorComponent->GetActiveType<UCEEffectorBoxType>())
+			{
+				OutAxisList = EAxisList::Type::XYZ;
+			}
+			else
+			{
+				OutAxisList = EAxisList::Type::Y;
+			}
 
-		return true;
-	}
+			return true;
+		}
 
-	if (EditingHandleType == HandleTypeRadius)
-	{
-		if (EffectorActorWeak->GetType() == ECEClonerEffectorType::Torus)
+		if (EditingHandleType == HandleTypeRadius)
 		{
-			OutAxisList = EAxisList::Type::Y;
-		}
+			if (!!EffectorComponent->GetActiveType<UCEEffectorTorusType>())
+			{
+				OutAxisList = EAxisList::Type::Y;
+			}
 
-		return true;
-	}
+			return true;
+		}
 
-	if (EditingHandleType == HandleTypeAngle)
-	{
-		if (EffectorActorWeak->GetType() == ECEClonerEffectorType::Radial)
+		if (EditingHandleType == HandleTypeAngle)
 		{
-			OutAxisList = EAxisList::Type::Z;
-		}
+			if (!!EffectorComponent->GetActiveType<UCEEffectorRadialType>())
+			{
+				OutAxisList = EAxisList::Type::Z;
+			}
 
-		return true;
+			return true;
+		}
 	}
 
 	return Super::GetWidgetAxisList(InViewportClient, InWidgetMode, OutAxisList);
@@ -604,39 +593,42 @@ bool FAvaEffectorActorVisualizer::GetWidgetAxisList(const FEditorViewportClient*
 bool FAvaEffectorActorVisualizer::GetWidgetAxisListDragOverride(const FEditorViewportClient* InViewportClient,
 	UE::Widget::EWidgetMode InWidgetMode, EAxisList::Type& OutAxisList) const
 {
-	if (EditingHandleType == HandleTypeInnerZone
-		|| EditingHandleType == HandleTypeOuterZone)
+	if (UCEEffectorComponent* EffectorComponent = EffectorComponentWeak.Get())
 	{
-		if (EffectorActorWeak->GetType() == ECEClonerEffectorType::Torus)
+		if (EditingHandleType == HandleTypeInnerZone
+		|| EditingHandleType == HandleTypeOuterZone)
 		{
-			OutAxisList = EAxisList::Type::Z;
-			return true;
-		}
+			if (!!EffectorComponent->GetActiveType<UCEEffectorTorusType>())
+			{
+				OutAxisList = EAxisList::Type::Z;
+				return true;
+			}
 
-		if (EffectorActorWeak->GetType() != ECEClonerEffectorType::Box)
-		{
-			OutAxisList = EAxisList::Type::Y;
-			return true;
+			if (!EffectorComponent->GetActiveType<UCEEffectorBoxType>())
+			{
+				OutAxisList = EAxisList::Type::Y;
+				return true;
+			}
 		}
-	}
 
-	if (EditingHandleType == HandleTypeRadius)
-	{
-		if (EffectorActorWeak->GetType() == ECEClonerEffectorType::Torus)
+		if (EditingHandleType == HandleTypeRadius)
 		{
-			OutAxisList = EAxisList::Type::Y;
-			return true;
+			if (!!EffectorComponent->GetActiveType<UCEEffectorTorusType>())
+			{
+				OutAxisList = EAxisList::Type::Y;
+				return true;
+			}
 		}
-	}
 
-	if (EditingHandleType == HandleTypeAngle)
-	{
-		if (EffectorActorWeak->GetType() == ECEClonerEffectorType::Radial)
+		if (EditingHandleType == HandleTypeAngle)
 		{
-			OutAxisList = EAxisList::Type::Z;
-		}
+			if (!!EffectorComponent->GetActiveType<UCEEffectorRadialType>())
+			{
+				OutAxisList = EAxisList::Type::Z;
+			}
 
-		return true;
+			return true;
+		}
 	}
 
 	return Super::GetWidgetAxisListDragOverride(InViewportClient, InWidgetMode, OutAxisList);
@@ -650,109 +642,109 @@ bool FAvaEffectorActorVisualizer::ResetValue(FEditorViewportClient* InViewportCl
 	}
 
 	const HAvaEffectorActorZoneHitProxy* ComponentHitProxy = static_cast<HAvaEffectorActorZoneHitProxy*>(InHitProxy);
+	UActorComponent* Component = const_cast<UActorComponent*>(ComponentHitProxy->Component.Get());
 
-	if (!ComponentHitProxy->Component.IsValid() || !ComponentHitProxy->Component->IsA<UCEEffectorComponent>())
+	if (!Component || !Component->IsA<UCEEffectorComponent>())
 	{
 		return Super::ResetValue(InViewportClient, InHitProxy);
 	}
 
-	if (ACEEffectorActor* EffectorActor = Cast<ACEEffectorActor>(ComponentHitProxy->Component->GetOwner()))
+	if (UCEEffectorComponent* EffectorComponent = Cast<UCEEffectorComponent>(Component))
 	{
-		const ECEClonerEffectorType Type = EffectorActor->GetType();
 		const int32 HandleType = ComponentHitProxy->HandleType;
 
 		FScopedTransaction Transaction(LOCTEXT("VisualizerResetValue", "Visualizer Reset Value"));
 
-		if (Type == ECEClonerEffectorType::Box)
+		if (UCEEffectorBoxType* BoxType = EffectorComponent->GetActiveType<UCEEffectorBoxType>())
 		{
 			if (HandleType == HandleTypeInnerZone)
 			{
-				ModifyProperty(EffectorActor, InnerExtentProperty, EPropertyChangeType::ValueSet, [&EffectorActor]()
+				ModifyProperty(BoxType, InnerExtentProperty, EPropertyChangeType::ValueSet, [&BoxType]()
 				{
-					EffectorActor->SetInnerExtent(FVector(50.f));
+					BoxType->SetInnerExtent(FVector(50.f));
 				});
 			}
 			else if (HandleType == HandleTypeOuterZone)
 			{
-				ModifyProperty(EffectorActor, OuterExtentProperty, EPropertyChangeType::ValueSet, [&EffectorActor]()
+				ModifyProperty(BoxType, OuterExtentProperty, EPropertyChangeType::ValueSet, [&BoxType]()
 				{
-					EffectorActor->SetOuterExtent(FVector(200.f));
+					BoxType->SetOuterExtent(FVector(200.f));
 				});
 			}
 		}
-		else if (Type == ECEClonerEffectorType::Plane)
+		else if (UCEEffectorPlaneType* PlaneType = EffectorComponent->GetActiveType<UCEEffectorPlaneType>())
 		{
 			if (HandleType == HandleTypeInnerZone
 				|| HandleType == HandleTypeOuterZone)
 			{
-				ModifyProperty(EffectorActor, PlaneSpacingProperty, EPropertyChangeType::ValueSet, [&EffectorActor]()
+				ModifyProperty(PlaneType, PlaneSpacingProperty, EPropertyChangeType::ValueSet, [&PlaneType]()
 				{
-					EffectorActor->SetPlaneSpacing(200.f);
+					PlaneType->SetPlaneSpacing(200.f);
 				});
 			}
 		}
-		else if (Type == ECEClonerEffectorType::Sphere)
+		else if (UCEEffectorSphereType* SphereType = EffectorComponent->GetActiveType<UCEEffectorSphereType>())
 		{
 			if (HandleType == HandleTypeInnerZone)
 			{
-				ModifyProperty(EffectorActor, InnerRadiusProperty, EPropertyChangeType::ValueSet, [&EffectorActor]()
+				ModifyProperty(SphereType, InnerRadiusProperty, EPropertyChangeType::ValueSet, [&SphereType]()
 				{
-					EffectorActor->SetInnerRadius(50.f);
+					SphereType->SetInnerRadius(50.f);
 				});
 			}
 			else if (HandleType == HandleTypeOuterZone)
 			{
-				ModifyProperty(EffectorActor, OuterRadiusProperty, EPropertyChangeType::ValueSet, [&EffectorActor]()
+				ModifyProperty(SphereType, OuterRadiusProperty, EPropertyChangeType::ValueSet, [&SphereType]()
 				{
-					EffectorActor->SetOuterRadius(200.f);
+					SphereType->SetOuterRadius(200.f);
 				});
 			}
 		}
-		else if (Type == ECEClonerEffectorType::Radial)
+		else if (UCEEffectorRadialType* RadialType = EffectorComponent->GetActiveType<UCEEffectorRadialType>())
 		{
 			if (HandleType == HandleTypeInnerZone)
 			{
-				ModifyProperty(EffectorActor, RadialMinRadiusProperty, EPropertyChangeType::ValueSet, [&EffectorActor]()
+				ModifyProperty(RadialType, RadialMinRadiusProperty, EPropertyChangeType::ValueSet, [&RadialType]()
 				{
-					EffectorActor->SetRadialMinRadius(0.f);
+					RadialType->SetRadialMinRadius(0.f);
 				});
 			}
 			else if (HandleType == HandleTypeOuterZone)
 			{
-				ModifyProperty(EffectorActor, RadialMaxRadiusProperty, EPropertyChangeType::ValueSet, [&EffectorActor]()
+				ModifyProperty(RadialType, RadialMaxRadiusProperty, EPropertyChangeType::ValueSet, [&RadialType]()
 				{
-					EffectorActor->SetRadialMaxRadius(1000.f);
+					RadialType->SetRadialMaxRadius(1000.f);
 				});
 			}
 			else if (HandleType == HandleTypeAngle)
 			{
-				ModifyProperty(EffectorActor, RadialAngleProperty, EPropertyChangeType::ValueSet, [&EffectorActor]()
+				ModifyProperty(RadialType, RadialAngleProperty, EPropertyChangeType::ValueSet, [&RadialType]()
 				{
-					EffectorActor->SetRadialAngle(180.f);
+					RadialType->SetRadialAngle(180.f);
 				});
 			}
 		}
-		else if (Type == ECEClonerEffectorType::Torus)
+		else if (UCEEffectorTorusType* TorusType = EffectorComponent->GetActiveType<UCEEffectorTorusType>())
 		{
 			if (HandleType == HandleTypeInnerZone)
 			{
-				ModifyProperty(EffectorActor, TorusInnerRadiusProperty, EPropertyChangeType::ValueSet, [&EffectorActor]()
+				ModifyProperty(TorusType, TorusInnerRadiusProperty, EPropertyChangeType::ValueSet, [&TorusType]()
 				{
-					EffectorActor->SetTorusInnerRadius(50.f);
+					TorusType->SetTorusInnerRadius(50.f);
 				});
 			}
 			else if (HandleType == HandleTypeOuterZone)
 			{
-				ModifyProperty(EffectorActor, TorusOuterRadiusProperty, EPropertyChangeType::ValueSet, [&EffectorActor]()
+				ModifyProperty(TorusType, TorusOuterRadiusProperty, EPropertyChangeType::ValueSet, [&TorusType]()
 				{
-					EffectorActor->SetTorusOuterRadius(200.f);
+					TorusType->SetTorusOuterRadius(200.f);
 				});
 			}
 			else if (HandleType == HandleTypeRadius)
 			{
-				ModifyProperty(EffectorActor, TorusRadiusProperty, EPropertyChangeType::ValueSet, [&EffectorActor]()
+				ModifyProperty(TorusType, TorusRadiusProperty, EPropertyChangeType::ValueSet, [&TorusType]()
 				{
-					EffectorActor->SetTorusRadius(250.f);
+					TorusType->SetTorusRadius(250.f);
 				});
 			}
 		}
@@ -775,7 +767,7 @@ void FAvaEffectorActorVisualizer::EndEditing()
 {
 	Super::EndEditing();
 
-	EffectorActorWeak.Reset();
+	EffectorComponentWeak.Reset();
 	EditingHandleType = INDEX_NONE;
 }
 
diff --git a/Engine/Plugins/Experimental/Avalanche/Source/AvalancheEffectorsEditor/Private/Effector/AvaEffectorActorVis.h b/Engine/Plugins/Experimental/Avalanche/Source/AvalancheEffectorsEditor/Private/Effector/AvaEffectorActorVis.h
--- a/Engine/Plugins/Experimental/Avalanche/Source/AvalancheEffectorsEditor/Private/Effector/AvaEffectorActorVis.h
+++ b/Engine/Plugins/Experimental/Avalanche/Source/AvalancheEffectorsEditor/Private/Effector/AvaEffectorActorVis.h
@@ -2,11 +2,9 @@
 
 #pragma once
 
-#include "Effector/CEEffectorActor.h"
+#include "Effector/CEEffectorComponent.h"
 #include "AvaVisBase.h"
 
-class UCEEffectorComponent;
-
 struct HAvaEffectorActorZoneHitProxy : HAvaHitProxy
 {
 	DECLARE_HIT_PROXY();
@@ -45,14 +43,14 @@ public:
 	virtual void DrawVisualizationNotEditing(const UActorComponent* InComponent, const FSceneView* InView, FPrimitiveDrawInterface* InPDI, int32& InOutIconIndex) override;
 	//~ End FAvaVisualizerBase
 
-	ACEEffectorActor* GetEffectorActor() const
+	UCEEffectorComponent* GetEffectorComponent() const
 	{
-		return EffectorActorWeak.Get();
+		return EffectorComponentWeak.Get();
 	}
 
 protected:
-	FVector GetHandleZoneLocation(const ACEEffectorActor* InEffectorActor, int32 InHandleType) const;
-	void DrawZoneButton(const ACEEffectorActor* InEffectorActor, const FSceneView* InView, FPrimitiveDrawInterface* InPDI, int32 InIconIndex, int32 InHandleType, FLinearColor InColor) const;
+	FVector GetHandleZoneLocation(const UCEEffectorComponent* InEffectorComponent, int32 InHandleType) const;
+	void DrawZoneButton(const UCEEffectorComponent* InEffectorComponent, const FSceneView* InView, FPrimitiveDrawInterface* InPDI, int32 InIconIndex, int32 InHandleType, FLinearColor InColor) const;
 
 	FProperty* InnerRadiusProperty;
 	FProperty* OuterRadiusProperty;
@@ -70,7 +68,7 @@ protected:
 	FProperty* TorusInnerRadiusProperty;
 	FProperty* TorusOuterRadiusProperty;
 
-	TWeakObjectPtr<ACEEffectorActor> EffectorActorWeak = nullptr;
+	TWeakObjectPtr<UCEEffectorComponent> EffectorComponentWeak = nullptr;
 
 	float InitialInnerRadius = 0.f;
 	float InitialOuterRadius = 0.f;
