[Motion Design] Rundown Server API: save broadcast command and some comment clean up.

Also, for Add/Edit/Remove channel device, prevent the operation if the channel is live (the back end doesn't support it).

diff --git a/Engine/Plugins/Experimental/Avalanche/Source/AvalancheMedia/Internal/Rundown/AvaRundownMessages.h b/Engine/Plugins/Experimental/Avalanche/Source/AvalancheMedia/Internal/Rundown/AvaRundownMessages.h
--- a/Engine/Plugins/Experimental/Avalanche/Source/AvalancheMedia/Internal/Rundown/AvaRundownMessages.h
+++ b/Engine/Plugins/Experimental/Avalanche/Source/AvalancheMedia/Internal/Rundown/AvaRundownMessages.h
@@ -154,6 +154,7 @@ public:
  * The rundown asset must have been loaded, either by an edit command
  * or playback, prior to this command.
  * Unloaded assets will not be loaded by this command.
+ * This command is not supported in game builds.
  */
 USTRUCT()
 struct FAvaRundownSaveRundown : public FAvaRundownMsgBase
@@ -168,7 +169,6 @@ public:
 	bool bOnlyIfIsDirty = false;
 };
 
-
 /**
  * Request the list of pages from the given rundown.
  */
@@ -665,7 +665,10 @@ public:
 	UPROPERTY()
 	FString Name;
 
-	/** Name of the server this class was seen on. */
+	/**
+	 * Name of the playback server this class was seen on.
+	 * The name will be empty for the "local process" device.
+	 */
 	UPROPERTY()
 	FString Server;
 
@@ -807,15 +810,28 @@ public:
 	FString NewChannelName;
 };
 
+/**
+ * Request a list of devices from the rundown server.
+ * The server will reply with FAvaRundownDevicesList containing
+ * the devices that can be enumerated from the local host and all connected hosts
+ * through the motion design playback service.
+ */
 USTRUCT()
 struct FAvaRundownGetDevices : public FAvaRundownMsgBase
 {
 	GENERATED_BODY()
 
+	/**
+	 * If true, listing all media output classes on the server, even if they don't have a device provider.
+	 */
 	UPROPERTY()
 	bool bShowAllMediaOutputClasses = false;
 };
 
+/**
+ * Add an enumerated device to the given channel.
+ * This command will fail if the channel is live.
+ */
 USTRUCT()
 struct FAvaRundownAddChannelDevice : public FAvaRundownMsgBase
 {
@@ -824,10 +840,21 @@ public:
 	UPROPERTY()
 	FString ChannelName;
 
+	/**
+	 * The specified name is one of the enumerated device from FAvaRundownDevicesList,
+	 * FAvaRundownOutputDeviceItem::Name.
+	 */
 	UPROPERTY()
 	FString MediaOutputName;
+
+	UPROPERTY()
+	bool bSaveBroadcast = true;
 };
 
+/**
+ * Modify an existing device in the given channel.
+ * This command will fail if the channel is live.
+ */
 USTRUCT()
 struct FAvaRundownEditChannelDevice : public FAvaRundownMsgBase
 {
@@ -836,13 +863,26 @@ public:
 	UPROPERTY()
 	FString ChannelName;
 
+	/**
+	 * The specified name is one of the enumerated device from FAvaRundownChannel::Devices,
+	 * FAvaRundownOutputDeviceItem::Name field.
+	 * Must be the instanced devices from either FAvaRundownChannels, FAvaRundownChannelResponse
+	 * or FAvaRundownChannelListChanged. These names are not the same as when adding a device.
+	 */
 	UPROPERTY()
 	FString MediaOutputName;
 
 	UPROPERTY()
 	FString Data;
+
+	UPROPERTY()
+	bool bSaveBroadcast = true;
 };
 
+/**
+ * Remove an existing device from the given channel.
+ * This command will fail if the channel is live.
+ */
 USTRUCT()
 struct FAvaRundownRemoveChannelDevice : public FAvaRundownMsgBase
 {
@@ -851,21 +891,17 @@ public:
 	UPROPERTY()
 	FString ChannelName;
 
+	/**
+	 * The specified name is one of the enumerated device from FAvaRundownChannel::Devices,
+	 * FAvaRundownOutputDeviceItem::Name field.
+	 * Must be the instanced devices from either FAvaRundownChannels, FAvaRundownChannelResponse
+	 * or FAvaRundownChannelListChanged. These names are not the same as when adding a device.
+	 */
 	UPROPERTY()
 	FString MediaOutputName;
-};
-
-/* No difference from FAvaRundownOutputDeviceItem except this is meant to return a single device response. */
-USTRUCT()
-struct FAvaRundownOutputDeviceItemResponse : public FAvaRundownMsgBase
-{
-	GENERATED_BODY()
-public:
-	UPROPERTY()
-	FString Name;
 
 	UPROPERTY()
-	FString Data;
+	bool bSaveBroadcast = true;
 };
 
 USTRUCT()
@@ -925,4 +961,11 @@ public:
 	/** Advanced viewport client engine features indexed by FEngineShowFlags names. */
 	UPROPERTY()
 	TArray<FAvaViewportQualitySettingsFeature> Features;
+};
+
+/** Save current broadcast configuration to file on the server. */
+USTRUCT()
+struct FAvaRundownSaveBroadcast : public FAvaRundownMsgBase
+{
+	GENERATED_BODY()
 };
\ No newline at end of file
diff --git a/Engine/Plugins/Experimental/Avalanche/Source/AvalancheMedia/Private/Rundown/AvaRundownServer.cpp b/Engine/Plugins/Experimental/Avalanche/Source/AvalancheMedia/Private/Rundown/AvaRundownServer.cpp
--- a/Engine/Plugins/Experimental/Avalanche/Source/AvalancheMedia/Private/Rundown/AvaRundownServer.cpp
+++ b/Engine/Plugins/Experimental/Avalanche/Source/AvalancheMedia/Private/Rundown/AvaRundownServer.cpp
@@ -271,6 +271,7 @@ void FAvaRundownServer::Init(const FString& InAssignedHostName)
 	.Handling<FAvaRundownGetChannelImage>(this, &FAvaRundownServer::HandleGetChannelImage)
 	.Handling<FAvaRundownGetChannelQualitySettings>(this, &FAvaRundownServer::HandleGetChannelQualitySettings)
 	.Handling<FAvaRundownSetChannelQualitySettings>(this, &FAvaRundownServer::HandleSetChannelQualitySettings)
+	.Handling<FAvaRundownSaveBroadcast>(this, &FAvaRundownServer::HandleSaveBroadcast)
 	.NotificationHandling(FOnBusNotification::CreateRaw(this, &FAvaRundownServer::OnMessageBusNotification));
 	
 	if (MessageEndpoint.IsValid())
@@ -1310,6 +1311,12 @@ void FAvaRundownServer::HandleAddChannelDevice(const FAvaRundownAddChannelDevice
 		return;    
 	}
 
+	if (OutputChannel.GetState() == EAvaBroadcastChannelState::Live)
+	{
+		LogAndSendMessage(InContext->GetSender(), InMessage.RequestId, ELogVerbosity::Error, TEXT("\"AddChannelDevice\" Failed. Reason: Channel is Live."), *InMessage.ChannelName);
+		return;    
+	}
+
 	/*
 		We're essentially replicating the UI editor here. The editor:
 		1. Builds an output tree
@@ -1336,7 +1343,11 @@ void FAvaRundownServer::HandleAddChannelDevice(const FAvaRundownAddChannelDevice
 
 	const FAvaBroadcastMediaOutputInfo OutputInfo;
 	const UMediaOutput* OutputDevice = TreeItem->AddMediaOutputToChannel(OutputChannel.GetChannelName(), OutputInfo);
-	Broadcast.SaveBroadcast();
+
+	if (InMessage.bSaveBroadcast)
+	{
+		Broadcast.SaveBroadcast();
+	}
 
 	LogAndSendMessage(InContext->GetSender(), InMessage.RequestId, ELogVerbosity::Log, TEXT("\"AddChannelDevice\" successfully added device \"%s\""), *OutputDevice->GetFName().ToString());
 }
@@ -1359,6 +1370,12 @@ void FAvaRundownServer::HandleEditChannelDevice(const FAvaRundownEditChannelDevi
 		return;    
 	}
 
+	if (OutputChannel.GetState() == EAvaBroadcastChannelState::Live)
+	{
+		LogAndSendMessage(InContext->GetSender(), InMessage.RequestId, ELogVerbosity::Error, TEXT("\"EditChannelDevice\" Failed. Reason: Channel is Live."), *InMessage.ChannelName);
+		return;    
+	}
+
 	UMediaOutput* MediaOutput = UE::AvaRundownServer::Private::FindChannelMediaOutput(OutputChannel, InMessage.MediaOutputName);
 	
 	if (!MediaOutput)
@@ -1369,7 +1386,10 @@ void FAvaRundownServer::HandleEditChannelDevice(const FAvaRundownEditChannelDevi
 
 	FAvaRundownServerMediaOutputUtils::EditMediaOutput(MediaOutput, InMessage.Data);
 	
-	Broadcast.SaveBroadcast();
+	if (InMessage.bSaveBroadcast)
+	{
+		Broadcast.SaveBroadcast();
+	}
 	
 	SendMessage(InContext->GetSender(), InMessage.RequestId, ELogVerbosity::Log, TEXT("\"EditChannelDevice\". Successfully edited device \"%s\" on \"%s\""), *InMessage.MediaOutputName, *InMessage.ChannelName); 
 }
@@ -1391,6 +1411,12 @@ void FAvaRundownServer::HandleRemoveChannelDevice(const FAvaRundownRemoveChannel
 		return;    
 	}
 
+	if (OutputChannel.GetState() == EAvaBroadcastChannelState::Live)
+	{
+		LogAndSendMessage(InContext->GetSender(), InMessage.RequestId, ELogVerbosity::Error, TEXT("\"RemoveChannelDevice\" Failed. Reason: Channel is Live."), *InMessage.ChannelName);
+		return;    
+	}
+
 	UMediaOutput* MediaOutput = UE::AvaRundownServer::Private::FindChannelMediaOutput(OutputChannel, InMessage.MediaOutputName);
 	
 	if (!MediaOutput)
@@ -1414,8 +1440,11 @@ void FAvaRundownServer::HandleRemoveChannelDevice(const FAvaRundownRemoveChannel
 		LogAndSendMessage(InContext->GetSender(), InMessage.RequestId, ELogVerbosity::Error, TEXT("\"RemoveChannelDevice\" Didn't remove device."));
 		return;
 	}
-
-	Broadcast.SaveBroadcast();
+	
+	if (InMessage.bSaveBroadcast)
+	{
+		Broadcast.SaveBroadcast();
+	}
 	
 	LogAndSendMessage(InContext->GetSender(), InMessage.RequestId, ELogVerbosity::Log, TEXT("\"RemoveChannelDevice\" Removed Device \"%s\""), *InMessage.MediaOutputName);
 }
@@ -1451,6 +1480,13 @@ void FAvaRundownServer::HandleSetChannelQualitySettings(const FAvaRundownSetChan
 	LogAndSendMessage(InContext->GetSender(), InMessage.RequestId, ELogVerbosity::Log, TEXT("\"SetChannelQualitySettings\" Channel \"%s\" success."), *InMessage.ChannelName);
 }
 
+void FAvaRundownServer::HandleSaveBroadcast(const FAvaRundownSaveBroadcast& InMessage, const TSharedRef<IMessageContext, ESPMode::ThreadSafe>& InContext)
+{
+	UAvaBroadcast& Broadcast = UAvaBroadcast::Get();
+	Broadcast.SaveBroadcast();
+	LogAndSendMessage(InContext->GetSender(), InMessage.RequestId, ELogVerbosity::Log, TEXT("\"SaveBroadcast\" success."));
+}
+
 void FAvaRundownServer::HandleGetDevices(const FAvaRundownGetDevices& InMessage,
 	const TSharedRef<IMessageContext, ESPMode::ThreadSafe>& InContext)
 {
diff --git a/Engine/Plugins/Experimental/Avalanche/Source/AvalancheMedia/Private/Rundown/AvaRundownServer.h b/Engine/Plugins/Experimental/Avalanche/Source/AvalancheMedia/Private/Rundown/AvaRundownServer.h
--- a/Engine/Plugins/Experimental/Avalanche/Source/AvalancheMedia/Private/Rundown/AvaRundownServer.h
+++ b/Engine/Plugins/Experimental/Avalanche/Source/AvalancheMedia/Private/Rundown/AvaRundownServer.h
@@ -84,6 +84,7 @@ public:
 	void HandleGetChannelImage(const FAvaRundownGetChannelImage& InMessage, const TSharedRef<IMessageContext, ESPMode::ThreadSafe>& InContext);
 	void HandleGetChannelQualitySettings(const FAvaRundownGetChannelQualitySettings& InMessage, const TSharedRef<IMessageContext, ESPMode::ThreadSafe>& InContext);
 	void HandleSetChannelQualitySettings(const FAvaRundownSetChannelQualitySettings& InMessage, const TSharedRef<IMessageContext, ESPMode::ThreadSafe>& InContext);
+	void HandleSaveBroadcast(const FAvaRundownSaveBroadcast& InMessage, const TSharedRef<IMessageContext, ESPMode::ThreadSafe>& InContext);
 	
 	void HandleGetDevices(const FAvaRundownGetDevices& InMessage, const TSharedRef<IMessageContext, ESPMode::ThreadSafe>& InContext);
 	
@@ -223,8 +224,8 @@ private:
 		/**
 		 * Checks if previous RCP was registered.
 		 * If so, save modified values to corresponding page.
- 		 * This may result in the Rundown to be modified.
- 		 * Will also unregister RCP from RC Module if requested.
+		 * This may result in the Rundown to be modified.
+		 * Will also unregister RCP from RC Module if requested.
 		 */
 		void SaveCurrentRemoteControlPresetToPage(bool bInUnregister);
 	};
