Motion Design
- Toggle to set whether the motion design viewport automatically activates when entering motion design mode

diff --git a/Engine/Plugins/Experimental/Avalanche/Source/AvalancheEditor/Private/AvaEditorSettings.h b/Engine/Plugins/Experimental/Avalanche/Source/AvalancheEditor/Private/AvaEditorSettings.h
--- a/Engine/Plugins/Experimental/Avalanche/Source/AvalancheEditor/Private/AvaEditorSettings.h
+++ b/Engine/Plugins/Experimental/Avalanche/Source/AvalancheEditor/Private/AvaEditorSettings.h
@@ -42,6 +42,12 @@ public:
 	UPROPERTY(Config, EditAnywhere, BlueprintReadWrite, Category = "Behavior")
 	float CameraDistance = 500.0f;
 
+	/**
+	 * Whether to automatically switch to the Motion Design viewport when the mode is activated
+	 * or a Motion Design level is opened. */
+	UPROPERTY(Config, EditAnywhere, BlueprintReadWrite, Category = "Behavior")
+	bool bAutoActivateMotionDesignViewport = true;
+
 	/** Default viewport quality settings for all newly created Motion Design blueprints. */
 	UPROPERTY(Config, EditAnywhere, BlueprintReadWrite, Category = "Quality")
 	FAvaViewportQualitySettings DefaultViewportQualitySettings = FAvaViewportQualitySettings(true);
diff --git a/Engine/Plugins/Experimental/Avalanche/Source/AvalancheEditor/Private/Viewport/AvaLevelViewportExtension.cpp b/Engine/Plugins/Experimental/Avalanche/Source/AvalancheEditor/Private/Viewport/AvaLevelViewportExtension.cpp
--- a/Engine/Plugins/Experimental/Avalanche/Source/AvalancheEditor/Private/Viewport/AvaLevelViewportExtension.cpp
+++ b/Engine/Plugins/Experimental/Avalanche/Source/AvalancheEditor/Private/Viewport/AvaLevelViewportExtension.cpp
@@ -2,6 +2,7 @@
 
 #include "AvaLevelViewportExtension.h"
 #include "AvaEditorCommands.h"
+#include "AvaEditorSettings.h"
 #include "AvaLevelViewportCommands.h"
 #include "AvaLevelViewportLayoutEntity.h"
 #include "AvaViewportUtils.h"
@@ -410,7 +411,11 @@ void FAvaLevelViewportExtension::CheckValidViewportType()
 	// Set the active viewport to Motion Design if there is a valid scene object (not all the viewports)
 	if (Editor.IsValid() && Editor->IsActive() && GetSceneObject())
 	{
-		SetMotionDesignViewportType();
+		if (GetDefault<UAvaEditorSettings>()->bAutoActivateMotionDesignViewport)
+		{
+			SetMotionDesignViewportType();
+		}
+
 		return;
 	}
 
