Motion Design Transition: Updated UI to latest State Tree changes

diff --git a/Engine/Plugins/Experimental/Avalanche/Source/AvalancheTransitionEditor/Private/AvaTransitionEditorStyle.cpp b/Engine/Plugins/Experimental/Avalanche/Source/AvalancheTransitionEditor/Private/AvaTransitionEditorStyle.cpp
--- a/Engine/Plugins/Experimental/Avalanche/Source/AvalancheTransitionEditor/Private/AvaTransitionEditorStyle.cpp
+++ b/Engine/Plugins/Experimental/Avalanche/Source/AvalancheTransitionEditor/Private/AvaTransitionEditorStyle.cpp
@@ -8,6 +8,7 @@
 #include "Misc/Paths.h"
 #include "Styling/SlateStyleMacros.h"
 #include "Styling/SlateStyleRegistry.h"
+#include "Textures/SlateIcon.h"
 
 FAvaTransitionEditorStyle::FAvaTransitionEditorStyle()
 	: FSlateStyleSet(TEXT("AvaTransitionEditor"))
@@ -57,3 +58,36 @@ FLinearColor FAvaTransitionEditorStyle::LerpColorSRGB(const FLinearColor& InA, c
 		static_cast<uint8>(FMath::RoundToInt(static_cast<float>(A.B) * (1.f - InAlpha) + static_cast<float>(B.B) * InAlpha)),
 		static_cast<uint8>(FMath::RoundToInt(static_cast<float>(A.A) * (1.f - InAlpha) + static_cast<float>(B.A) * InAlpha))));
 }
+
+FSlateIcon FAvaTransitionEditorStyle::ParseIcon(FName InIconName)
+{
+	FString IconPath = InIconName.ToString();
+
+	constexpr int32 NumOfIconPathNames = 4;
+
+	FName IconPathNames[NumOfIconPathNames] = {
+		NAME_None, // StyleSetName
+		NAME_None, // StyleName
+		NAME_None, // SmallStyleName
+		NAME_None  // StatusOverlayStyleName
+	};
+
+	int32 NameIndex = 0;
+	while (!IconPath.IsEmpty() && NameIndex < NumOfIconPathNames)
+	{
+		FString Left;
+		FString Right;
+
+		if (!IconPath.Split(TEXT("|"), &Left, &Right))
+		{
+			Left = IconPath;
+		}
+
+		IconPathNames[NameIndex] = FName(*Left);
+
+		NameIndex++;
+		IconPath = Right;
+	}
+
+	return FSlateIcon(IconPathNames[0], IconPathNames[1], IconPathNames[2], IconPathNames[3]);
+}
diff --git a/Engine/Plugins/Experimental/Avalanche/Source/AvalancheTransitionEditor/Private/AvaTransitionEditorStyle.h b/Engine/Plugins/Experimental/Avalanche/Source/AvalancheTransitionEditor/Private/AvaTransitionEditorStyle.h
--- a/Engine/Plugins/Experimental/Avalanche/Source/AvalancheTransitionEditor/Private/AvaTransitionEditorStyle.h
+++ b/Engine/Plugins/Experimental/Avalanche/Source/AvalancheTransitionEditor/Private/AvaTransitionEditorStyle.h
@@ -4,6 +4,8 @@
 
 #include "Styling/SlateStyle.h"
 
+struct FSlateIcon;
+
 class FAvaTransitionEditorStyle final : public FSlateStyleSet
 {
 public:
@@ -17,4 +19,6 @@ public:
 	virtual ~FAvaTransitionEditorStyle() override;
 
 	static FLinearColor LerpColorSRGB(const FLinearColor& InA, const FLinearColor& InB, float InAlpha);
+
+	static FSlateIcon ParseIcon(FName InIconName);
 };
diff --git a/Engine/Plugins/Experimental/Avalanche/Source/AvalancheTransitionEditor/Private/Customizations/AvaStateTreeStateCustomization.cpp b/Engine/Plugins/Experimental/Avalanche/Source/AvalancheTransitionEditor/Private/Customizations/AvaStateTreeStateCustomization.cpp
--- a/Engine/Plugins/Experimental/Avalanche/Source/AvalancheTransitionEditor/Private/Customizations/AvaStateTreeStateCustomization.cpp
+++ b/Engine/Plugins/Experimental/Avalanche/Source/AvalancheTransitionEditor/Private/Customizations/AvaStateTreeStateCustomization.cpp
@@ -17,6 +17,14 @@ TSharedRef<IDetailCustomization> FAvaStateTreeStateCustomization::MakeInstance()
 	return MakeShared<FAvaStateTreeStateCustomization>();
 }
 
+bool FAvaStateTreeStateCustomization::IsPropertyVisible(const FPropertyAndParent& InPropertyAndParent)
+{
+	const FName PropertyName = InPropertyAndParent.Property.GetFName();
+
+	return PropertyName != GET_MEMBER_NAME_CHECKED(UStateTreeState, bCheckPrerequisitesWhenActivatingChildDirectly)
+		&& PropertyName != GET_MEMBER_NAME_CHECKED(UStateTreeState, RequiredEventToEnter);
+}
+
 void FAvaStateTreeStateCustomization::CustomizeDetails(IDetailLayoutBuilder& InDetailBuilder)
 {
 	if (TSharedPtr<IDetailCustomization> Customization = GetDefaultCustomization())
@@ -26,6 +34,7 @@ void FAvaStateTreeStateCustomization::CustomizeDetails(IDetailLayoutBuilder& InD
 
 	InDetailBuilder.HideCategory(TEXT("State"));
 	InDetailBuilder.HideCategory(TEXT("Transitions"));
+	InDetailBuilder.HideCategory(TEXT("Parameters"));
 
 	InDetailBuilder.EditCategory(TEXT("Enter Conditions")).InitiallyCollapsed(false);
 	InDetailBuilder.EditCategory(TEXT("Task")).InitiallyCollapsed(false);
diff --git a/Engine/Plugins/Experimental/Avalanche/Source/AvalancheTransitionEditor/Private/Customizations/AvaStateTreeStateCustomization.h b/Engine/Plugins/Experimental/Avalanche/Source/AvalancheTransitionEditor/Private/Customizations/AvaStateTreeStateCustomization.h
--- a/Engine/Plugins/Experimental/Avalanche/Source/AvalancheTransitionEditor/Private/Customizations/AvaStateTreeStateCustomization.h
+++ b/Engine/Plugins/Experimental/Avalanche/Source/AvalancheTransitionEditor/Private/Customizations/AvaStateTreeStateCustomization.h
@@ -4,12 +4,16 @@
 
 #include "IDetailCustomization.h"
 
+struct FPropertyAndParent;
+
 /** Customization that re-uses the module-registered Customization for UStateTreeState, and tweaks a few settings for it */
 class FAvaStateTreeStateCustomization : public IDetailCustomization
 {
 public:
 	static TSharedRef<IDetailCustomization> MakeInstance();
 
+	static bool IsPropertyVisible(const FPropertyAndParent& InPropertyAndParent);
+
 protected:
 	//~ Begin IDetailCustomization
 	virtual void CustomizeDetails(IDetailLayoutBuilder& InDetailBuilder) override;
diff --git a/Engine/Plugins/Experimental/Avalanche/Source/AvalancheTransitionEditor/Private/Customizations/Widgets/SAvaTransitionTransitionType.cpp b/Engine/Plugins/Experimental/Avalanche/Source/AvalancheTransitionEditor/Private/Customizations/Widgets/SAvaTransitionTransitionType.cpp
--- a/Engine/Plugins/Experimental/Avalanche/Source/AvalancheTransitionEditor/Private/Customizations/Widgets/SAvaTransitionTransitionType.cpp
+++ b/Engine/Plugins/Experimental/Avalanche/Source/AvalancheTransitionEditor/Private/Customizations/Widgets/SAvaTransitionTransitionType.cpp
@@ -38,7 +38,7 @@ void SAvaTransitionTransitionType::Construct(const FArguments& InArgs, const TSh
 
 	AddTransitionType(EStateTreeTransitionType::None, LOCTEXT("TransitionTypeNone", "None"));
 	AddTransitionType(EStateTreeTransitionType::NextSelectableState, LOCTEXT("TransitionTypeNextSelectableState", "Next"));
-	AddTransitionType(EStateTreeTransitionType::Succeeded, LOCTEXT("TransitionTypeSucceeded", "Finish"));
+	AddTransitionType(EStateTreeTransitionType::Succeeded, LOCTEXT("TransitionTypeSucceeded", "Succeed"));
 
 	ChildSlot
 	[
diff --git a/Engine/Plugins/Experimental/Avalanche/Source/AvalancheTransitionEditor/Private/ViewModels/AvaTransitionNodeViewModel.h b/Engine/Plugins/Experimental/Avalanche/Source/AvalancheTransitionEditor/Private/ViewModels/AvaTransitionNodeViewModel.h
--- a/Engine/Plugins/Experimental/Avalanche/Source/AvalancheTransitionEditor/Private/ViewModels/AvaTransitionNodeViewModel.h
+++ b/Engine/Plugins/Experimental/Avalanche/Source/AvalancheTransitionEditor/Private/ViewModels/AvaTransitionNodeViewModel.h
@@ -31,8 +31,8 @@ public:
 
 	const FInstancedStruct* GetNode() const;
 
-	template<typename T UE_REQUIRES(TIsDerivedFrom<T, FStateTreeNodeBase>::Value)>
-	const T* GetNodeOfType() const
+	template<typename T = FStateTreeNodeBase UE_REQUIRES(TIsDerivedFrom<T, FStateTreeNodeBase>::Value)>
+	const T* GetTypedNode() const
 	{
 		if (const FInstancedStruct* Node = GetNode())
 		{
diff --git a/Engine/Plugins/Experimental/Avalanche/Source/AvalancheTransitionEditor/Private/ViewModels/Task/AvaTransitionTaskViewModel.cpp b/Engine/Plugins/Experimental/Avalanche/Source/AvalancheTransitionEditor/Private/ViewModels/Task/AvaTransitionTaskViewModel.cpp
--- a/Engine/Plugins/Experimental/Avalanche/Source/AvalancheTransitionEditor/Private/ViewModels/Task/AvaTransitionTaskViewModel.cpp
+++ b/Engine/Plugins/Experimental/Avalanche/Source/AvalancheTransitionEditor/Private/ViewModels/Task/AvaTransitionTaskViewModel.cpp
@@ -1,12 +1,14 @@
 // Copyright Epic Games, Inc. All Rights Reserved.
 
 #include "AvaTransitionTaskViewModel.h"
+#include "AvaTransitionEditorStyle.h"
 #include "AvaTransitionNodeContext.h"
 #include "AvaTransitionTreeEditorData.h"
 #include "StateTreeEditorData.h"
 #include "StateTreeEditorStyle.h"
 #include "StateTreeState.h"
 #include "Tasks/AvaTransitionTask.h"
+#include "Textures/SlateIcon.h"
 #include "ViewModels/AvaTransitionViewModelSharedData.h"
 #include "ViewModels/AvaTransitionViewModelUtils.h"
 #include "Widgets/Images/SImage.h"
@@ -28,6 +30,42 @@ FText FAvaTransitionTaskViewModel::GetTaskDescription() const
 	return TaskDescription;
 }
 
+FSlateColor FAvaTransitionTaskViewModel::GetTaskColor() const
+{
+	FLinearColor TaskColor = FAvaTransitionEditorStyle::LerpColorSRGB(GetStateColor().GetSpecifiedColor(), FColor::White, 0.25f);
+	return TaskColor.CopyWithNewOpacity(0.25f);
+}
+
+EVisibility FAvaTransitionTaskViewModel::GetTaskIconVisibility() const
+{
+	if (const FStateTreeNodeBase* Node = GetTypedNode<>())
+	{
+		if (!Node->GetIconName().IsNone())
+		{
+			return EVisibility::SelfHitTestInvisible;
+		}
+	}
+	return EVisibility::Collapsed;
+}
+
+const FSlateBrush* FAvaTransitionTaskViewModel::GetTaskIcon() const
+{
+	if (const FStateTreeNodeBase* Node = GetTypedNode<>())
+	{
+		return FAvaTransitionEditorStyle::ParseIcon(Node->GetIconName()).GetIcon();
+	}
+	return nullptr;	
+}
+
+FSlateColor FAvaTransitionTaskViewModel::GetTaskIconColor() const
+{
+	if (const FStateTreeNodeBase* Node = GetTypedNode<>())
+	{
+		return FLinearColor(Node->GetIconColor());
+	}
+	return FSlateColor::UseForeground();
+}
+
 void FAvaTransitionTaskViewModel::UpdateTaskDescription()
 {
 	TaskDescription = FText::GetEmpty();
@@ -63,7 +101,7 @@ void FAvaTransitionTaskViewModel::UpdateTaskDescription()
 
 bool FAvaTransitionTaskViewModel::IsEnabled() const
 {
-	const FStateTreeTaskBase* Task = GetNodeOfType<FStateTreeTaskBase>();
+	const FStateTreeTaskBase* Task = GetTypedNode<FStateTreeTaskBase>();
 	return Task && Task->bTaskEnabled;
 }
 
@@ -122,20 +160,41 @@ TSharedRef<SWidget> FAvaTransitionTaskViewModel::CreateWidget()
 		.VAlign(VAlign_Fill)
 		.Padding(0)
 		.IsEnabled(this, &FAvaTransitionTaskViewModel::IsEnabled)
-		.BorderImage(FStateTreeEditorStyle::Get().GetBrush("StateTree.Task.Rect"))
-		.BorderBackgroundColor(this, &FAvaTransitionTaskViewModel::GetStateColor)
+		.BorderBackgroundColor(this, &FAvaTransitionTaskViewModel::GetTaskColor)
+		.BorderImage(FAppStyle::GetBrush("WhiteBrush"))
 		[
 			SNew(SOverlay)
 			// Task Description
 			+ SOverlay::Slot()
-			.VAlign(VAlign_Center)
+			.Padding(6.f, 0.f)
 			[
-				SNew(STextBlock)
-				.Margin(FMargin(4.f, 0.f))
-				.Text(this, &FAvaTransitionTaskViewModel::GetTaskDescription)
-				.ToolTipText(this, &FAvaTransitionTaskViewModel::GetTaskDescription)
-				.TextStyle(FStateTreeEditorStyle::Get(), "StateTree.Task.Title")
-				.OverflowPolicy(ETextOverflowPolicy::Ellipsis)
+				SNew(SHorizontalBox)
+				+ SHorizontalBox::Slot()
+				.VAlign(VAlign_Center)
+				.HAlign(HAlign_Left)
+				.AutoWidth()
+				[
+					SNew(SBox)
+					.Padding(FMargin(0.f, 0.f, 2.f, 0.f))
+					.Visibility(this, &FAvaTransitionTaskViewModel::GetTaskIconVisibility)
+					[
+						SNew(SImage)
+						.Image(this, &FAvaTransitionTaskViewModel::GetTaskIcon)
+						.ColorAndOpacity(this, &FAvaTransitionTaskViewModel::GetTaskIconColor)
+					]
+				]
+				+ SHorizontalBox::Slot()
+				.AutoWidth()
+				.VAlign(VAlign_Center)
+				.HAlign(HAlign_Left)
+				[
+					SNew(STextBlock)
+					.Margin(FMargin(4.f, 0.f))
+					.Text(this, &FAvaTransitionTaskViewModel::GetTaskDescription)
+					.ToolTipText(this, &FAvaTransitionTaskViewModel::GetTaskDescription)
+					.TextStyle(FStateTreeEditorStyle::Get(), "StateTree.Task.Title")
+					.OverflowPolicy(ETextOverflowPolicy::Ellipsis)
+				]
 			]
 			// Task Breakpoint
 			+ SOverlay::Slot()
diff --git a/Engine/Plugins/Experimental/Avalanche/Source/AvalancheTransitionEditor/Private/ViewModels/Task/AvaTransitionTaskViewModel.h b/Engine/Plugins/Experimental/Avalanche/Source/AvalancheTransitionEditor/Private/ViewModels/Task/AvaTransitionTaskViewModel.h
--- a/Engine/Plugins/Experimental/Avalanche/Source/AvalancheTransitionEditor/Private/ViewModels/Task/AvaTransitionTaskViewModel.h
+++ b/Engine/Plugins/Experimental/Avalanche/Source/AvalancheTransitionEditor/Private/ViewModels/Task/AvaTransitionTaskViewModel.h
@@ -5,8 +5,10 @@
 #include "Extensions/IAvaTransitionWidgetExtension.h"
 #include "ViewModels/AvaTransitionNodeViewModel.h"
 
-struct FStateTreeTaskBase;
 struct EVisibility;
+struct FSlateBrush;
+struct FSlateColor;
+struct FStateTreeTaskBase;
 
 /** View Model for a Task Node */
 class FAvaTransitionTaskViewModel : public FAvaTransitionNodeViewModel, public IAvaTransitionWidgetExtension
@@ -18,6 +20,14 @@ public:
 
 	FText GetTaskDescription() const;
 
+	FSlateColor GetTaskColor() const;
+
+	EVisibility GetTaskIconVisibility() const;
+
+	const FSlateBrush* GetTaskIcon() const;
+
+	FSlateColor GetTaskIconColor() const;
+
 	bool IsEnabled() const;
 
 	EVisibility GetBreakpointVisibility() const;
diff --git a/Engine/Plugins/Experimental/Avalanche/Source/AvalancheTransitionEditor/Private/ViewModels/Transition/AvaTransitionTransitionViewModel.cpp b/Engine/Plugins/Experimental/Avalanche/Source/AvalancheTransitionEditor/Private/ViewModels/Transition/AvaTransitionTransitionViewModel.cpp
--- a/Engine/Plugins/Experimental/Avalanche/Source/AvalancheTransitionEditor/Private/ViewModels/Transition/AvaTransitionTransitionViewModel.cpp
+++ b/Engine/Plugins/Experimental/Avalanche/Source/AvalancheTransitionEditor/Private/ViewModels/Transition/AvaTransitionTransitionViewModel.cpp
@@ -53,12 +53,12 @@ FStateTreeTransition* FAvaTransitionTransitionViewModel::GetTransition() const
 	return nullptr;
 }
 
-FText FAvaTransitionTransitionViewModel::GetIcon() const
+const FSlateBrush* FAvaTransitionTransitionViewModel::GetIcon() const
 {
 	const FStateTreeTransition* Transition = GetTransition();
 	if (!Transition)
 	{
-		return FText::GetEmpty();
+		return nullptr;
 	}
 
 	switch (Transition->State.LinkType)
@@ -68,7 +68,7 @@ FText FAvaTransitionTransitionViewModel::GetIcon() const
 			const UStateTreeState* State = GetState();
 			if (State && State->Children.IsEmpty() && State->Type == EStateTreeStateType::State && EnumHasAnyFlags(Transition->Trigger, EStateTreeTransitionTrigger::OnStateCompleted))
 			{
-				return FEditorFontGlyphs::Level_Up;
+				return FStateTreeEditorStyle::Get().GetBrush("StateTreeEditor.Transition.Parent");
 			}
 		}
 		// falls through
@@ -76,18 +76,18 @@ FText FAvaTransitionTransitionViewModel::GetIcon() const
 	case EStateTreeTransitionType::Succeeded:
 	case EStateTreeTransitionType::Failed:
 	case EStateTreeTransitionType::GotoState:
-		return FEditorFontGlyphs::Long_Arrow_Right;
+		return FStateTreeEditorStyle::Get().GetBrush("StateTreeEditor.Transition.Goto");
 
 	case EStateTreeTransitionType::NextState:
 	case EStateTreeTransitionType::NextSelectableState:
-		return FEditorFontGlyphs::Long_Arrow_Down;
+		return FStateTreeEditorStyle::Get().GetBrush("StateTreeEditor.Transition.Next");
 
 	default:
 		ensureMsgf(false, TEXT("Unhandled transition type."));
 		break;
 	}
 
-	return FText::GetEmpty();
+	return nullptr;
 }
 
 FText FAvaTransitionTransitionViewModel::GetDescription() const
@@ -101,19 +101,19 @@ FText FAvaTransitionTransitionViewModel::GetDescription() const
 	switch (Transition->State.LinkType)
 	{
 	case EStateTreeTransitionType::None:
-		return LOCTEXT("TransitionNone", "[None]");
+		return LOCTEXT("TransitionNone", "None");
 
 	case EStateTreeTransitionType::Succeeded:
-		return LOCTEXT("TransitionSucceed", "[Succeed]");
+		return LOCTEXT("TransitionSucceed", "Succeed");
 
 	case EStateTreeTransitionType::Failed:
-		return LOCTEXT("TransitionFail", "[Fail]");
+		return LOCTEXT("TransitionFail", "Fail");
 
 	case EStateTreeTransitionType::NextState:
-		return LOCTEXT("TransitionNext", "[Next]");
+		return LOCTEXT("TransitionNext", "Next");
 
 	case EStateTreeTransitionType::NextSelectableState:
-		return LOCTEXT("TransitionNextSelectable", "[Next Selectable]");
+		return LOCTEXT("TransitionNextSelectable", "Next Selectable");
 
 	case EStateTreeTransitionType::GotoState:
 		return FText::FromName(Transition->State.Name);
@@ -161,9 +161,9 @@ TSharedRef<SWidget> FAvaTransitionTransitionViewModel::CreateWidget()
 			SNew(SOverlay)
 			+ SOverlay::Slot()
 			[
-				SNew(STextBlock)
-				.Text(this, &FAvaTransitionTransitionViewModel::GetIcon)
-				.TextStyle(FStateTreeEditorStyle::Get(), "StateTree.Icon")
+				SNew(SImage)
+				.Image(this, &FAvaTransitionTransitionViewModel::GetIcon)
+				.ColorAndOpacity(FLinearColor(1, 1, 1, 0.5f))
 			]
 			// Breakpoint box
 			+ SOverlay::Slot()
diff --git a/Engine/Plugins/Experimental/Avalanche/Source/AvalancheTransitionEditor/Private/ViewModels/Transition/AvaTransitionTransitionViewModel.h b/Engine/Plugins/Experimental/Avalanche/Source/AvalancheTransitionEditor/Private/ViewModels/Transition/AvaTransitionTransitionViewModel.h
--- a/Engine/Plugins/Experimental/Avalanche/Source/AvalancheTransitionEditor/Private/ViewModels/Transition/AvaTransitionTransitionViewModel.h
+++ b/Engine/Plugins/Experimental/Avalanche/Source/AvalancheTransitionEditor/Private/ViewModels/Transition/AvaTransitionTransitionViewModel.h
@@ -9,6 +9,7 @@
 class UAvaTransitionTreeEditorData;
 class UStateTreeState;
 struct EVisibility;
+struct FSlateBrush;
 struct FStateTreeTransition;
 
 /** View Model for a State Tree Transition */
@@ -25,7 +26,7 @@ public:
 
 	FStateTreeTransition* GetTransition() const;
 
-	FText GetIcon() const;
+	const FSlateBrush* GetIcon() const;
 
 	FText GetDescription() const;
 
diff --git a/Engine/Plugins/Experimental/Avalanche/Source/AvalancheTransitionEditor/Private/Views/SAvaTransitionSelectionDetails.cpp b/Engine/Plugins/Experimental/Avalanche/Source/AvalancheTransitionEditor/Private/Views/SAvaTransitionSelectionDetails.cpp
--- a/Engine/Plugins/Experimental/Avalanche/Source/AvalancheTransitionEditor/Private/Views/SAvaTransitionSelectionDetails.cpp
+++ b/Engine/Plugins/Experimental/Avalanche/Source/AvalancheTransitionEditor/Private/Views/SAvaTransitionSelectionDetails.cpp
@@ -27,6 +27,8 @@ void SAvaTransitionSelectionDetails::Construct(const FArguments& InArgs, const T
 	{
 		DetailsView->RegisterInstancedCustomPropertyLayout(UStateTreeState::StaticClass()
 			, FOnGetDetailCustomizationInstance::CreateStatic(&FAvaStateTreeStateCustomization::MakeInstance));
+
+		DetailsView->SetIsPropertyVisibleDelegate(FIsPropertyVisible::CreateStatic(&FAvaStateTreeStateCustomization::IsPropertyVisible));
 	}
 
 	// Read-only
