Fixed multiply-inherited TSharedFromThis.

diff --git a/Engine/Plugins/Experimental/Avalanche/Source/AvalancheMaskEditor/Private/Details/AvaMask2DModifierDetails.h b/Engine/Plugins/Experimental/Avalanche/Source/AvalancheMaskEditor/Private/Details/AvaMask2DModifierDetails.h
--- a/Engine/Plugins/Experimental/Avalanche/Source/AvalancheMaskEditor/Private/Details/AvaMask2DModifierDetails.h
+++ b/Engine/Plugins/Experimental/Avalanche/Source/AvalancheMaskEditor/Private/Details/AvaMask2DModifierDetails.h
@@ -6,8 +6,7 @@
 #include "Mask2D/AvaMask2DBaseModifier.h"
 
 class FAvaMask2DModifierDetails
-	: public TSharedFromThis<FAvaMask2DModifierDetails>
-	, public IDetailCustomization
+	: public IDetailCustomization
 {
 public:
 	static TSharedRef<IDetailCustomization> MakeInstance()
