Motion Design Transition: Moved from GenerateDescription to the now StateTree-built in func of "GetDescription".

diff --git a/Engine/Plugins/Experimental/Avalanche/Source/Avalanche/Private/Transition/Conditions/AvaSceneContainsTagAttributeCondition.cpp b/Engine/Plugins/Experimental/Avalanche/Source/Avalanche/Private/Transition/Conditions/AvaSceneContainsTagAttributeCondition.cpp
--- a/Engine/Plugins/Experimental/Avalanche/Source/Avalanche/Private/Transition/Conditions/AvaSceneContainsTagAttributeCondition.cpp
+++ b/Engine/Plugins/Experimental/Avalanche/Source/Avalanche/Private/Transition/Conditions/AvaSceneContainsTagAttributeCondition.cpp
@@ -13,7 +13,8 @@
 
 #define LOCTEXT_NAMESPACE "AvaSceneContainsTagAttributeConditionBase"
 
-FText FAvaSceneContainsTagAttributeConditionBase::GenerateDescription(const FAvaTransitionNodeContext& InContext) const
+#if WITH_EDITOR
+FText FAvaSceneContainsTagAttributeConditionBase::GetDescription(const FGuid& InId, FStateTreeDataView InInstanceDataView, const IStateTreeBindingLookup& InBindingLookup, EStateTreeNodeFormatting InFormatting) const
 {
 	FFormatNamedArguments Arguments;
 
@@ -40,6 +41,7 @@ FText FAvaSceneContainsTagAttributeConditionBase::GenerateDescription(const FAva
 
 	return FText::Format(LOCTEXT("ConditionDescription", "{IndefinitePronoun}{Scene} {Contains} tag attribute '{TagAttribute}'"), Arguments);
 }
+#endif
 
 bool FAvaSceneContainsTagAttributeConditionBase::Link(FStateTreeLinker& InLinker)
 {
diff --git a/Engine/Plugins/Experimental/Avalanche/Source/Avalanche/Private/Transition/Conditions/AvaTransitionRCControllerMatchCondition.cpp b/Engine/Plugins/Experimental/Avalanche/Source/Avalanche/Private/Transition/Conditions/AvaTransitionRCControllerMatchCondition.cpp
--- a/Engine/Plugins/Experimental/Avalanche/Source/Avalanche/Private/Transition/Conditions/AvaTransitionRCControllerMatchCondition.cpp
+++ b/Engine/Plugins/Experimental/Avalanche/Source/Avalanche/Private/Transition/Conditions/AvaTransitionRCControllerMatchCondition.cpp
@@ -18,12 +18,14 @@
 
 #define LOCTEXT_NAMESPACE "AvaTransitionRCControllerMatchCondition"
 
-FText FAvaTransitionRCControllerMatchCondition::GenerateDescription(const FAvaTransitionNodeContext& InContext) const
+#if WITH_EDITOR
+FText FAvaTransitionRCControllerMatchCondition::GetDescription(const FGuid& InId, FStateTreeDataView InInstanceDataView, const IStateTreeBindingLookup& InBindingLookup, EStateTreeNodeFormatting InFormatting) const
 {
 	return FText::Format(LOCTEXT("ConditionDescription", "'{0}' is {1}")
 		, ControllerId.ToText()
 		, UEnum::GetDisplayValueAsText(ValueComparisonType).ToLower());
 }
+#endif
 
 bool FAvaTransitionRCControllerMatchCondition::Link(FStateTreeLinker& InLinker)
 {
diff --git a/Engine/Plugins/Experimental/Avalanche/Source/Avalanche/Private/Transition/Tasks/AvaSceneAddTagAttributeTask.cpp b/Engine/Plugins/Experimental/Avalanche/Source/Avalanche/Private/Transition/Tasks/AvaSceneAddTagAttributeTask.cpp
--- a/Engine/Plugins/Experimental/Avalanche/Source/Avalanche/Private/Transition/Tasks/AvaSceneAddTagAttributeTask.cpp
+++ b/Engine/Plugins/Experimental/Avalanche/Source/Avalanche/Private/Transition/Tasks/AvaSceneAddTagAttributeTask.cpp
@@ -6,10 +6,12 @@
 
 #define LOCTEXT_NAMESPACE "AvaSceneAddTagAttributeTask"
 
-FText FAvaSceneAddTagAttributeTask::GenerateDescription(const FAvaTransitionNodeContext& InContext) const
+#if WITH_EDITOR
+FText FAvaSceneAddTagAttributeTask::GetDescription(const FGuid& InId, FStateTreeDataView InInstanceDataView, const IStateTreeBindingLookup& InBindingLookup, EStateTreeNodeFormatting InFormatting) const
 {
 	return FText::Format(LOCTEXT("TaskDescription", "Add '{0}' tag attribute to this scene"), FText::FromName(TagAttribute.ToName()));
 }
+#endif
 
 EStateTreeRunStatus FAvaSceneAddTagAttributeTask::EnterState(FStateTreeExecutionContext& InContext, const FStateTreeTransitionResult& InTransition) const
 {
diff --git a/Engine/Plugins/Experimental/Avalanche/Source/Avalanche/Private/Transition/Tasks/AvaSceneRemoveTagAttributeTask.cpp b/Engine/Plugins/Experimental/Avalanche/Source/Avalanche/Private/Transition/Tasks/AvaSceneRemoveTagAttributeTask.cpp
--- a/Engine/Plugins/Experimental/Avalanche/Source/Avalanche/Private/Transition/Tasks/AvaSceneRemoveTagAttributeTask.cpp
+++ b/Engine/Plugins/Experimental/Avalanche/Source/Avalanche/Private/Transition/Tasks/AvaSceneRemoveTagAttributeTask.cpp
@@ -6,10 +6,12 @@
 
 #define LOCTEXT_NAMESPACE "AvaSceneRemoveTagAttributeTask"
 
-FText FAvaSceneRemoveTagAttributeTask::GenerateDescription(const FAvaTransitionNodeContext& InContext) const
+#if WITH_EDITOR
+FText FAvaSceneRemoveTagAttributeTask::GetDescription(const FGuid& InId, FStateTreeDataView InInstanceDataView, const IStateTreeBindingLookup& InBindingLookup, EStateTreeNodeFormatting InFormatting) const
 {
 	return FText::Format(LOCTEXT("TaskDescription", "Remove '{0}' tag attribute from this scene"), FText::FromName(TagAttribute.ToName()));
 }
+#endif
 
 EStateTreeRunStatus FAvaSceneRemoveTagAttributeTask::EnterState(FStateTreeExecutionContext& InContext, const FStateTreeTransitionResult& InTransition) const
 {
diff --git a/Engine/Plugins/Experimental/Avalanche/Source/Avalanche/Public/Transition/Conditions/AvaSceneContainsTagAttributeCondition.h b/Engine/Plugins/Experimental/Avalanche/Source/Avalanche/Public/Transition/Conditions/AvaSceneContainsTagAttributeCondition.h
--- a/Engine/Plugins/Experimental/Avalanche/Source/Avalanche/Public/Transition/Conditions/AvaSceneContainsTagAttributeCondition.h
+++ b/Engine/Plugins/Experimental/Avalanche/Source/Avalanche/Public/Transition/Conditions/AvaSceneContainsTagAttributeCondition.h
@@ -31,11 +31,10 @@ struct AVALANCHE_API FAvaSceneContainsTagAttributeConditionBase : public FAvaTra
 
 	using FInstanceDataType = FAvaSceneContainsTagAttributeConditionInstanceData;
 
-	//~ Begin FAvaTransitionCondition
-	virtual FText GenerateDescription(const FAvaTransitionNodeContext& InContext) const override;
-	//~ End FAvaTransitionCondition
-
 	//~ Begin FStateTreeNodeBase
+#if WITH_EDITOR
+	virtual FText GetDescription(const FGuid& InId, FStateTreeDataView InInstanceDataView, const IStateTreeBindingLookup& InBindingLookup, EStateTreeNodeFormatting InFormatting) const override;
+#endif
 	virtual const UStruct* GetInstanceDataType() const override { return FInstanceDataType::StaticStruct(); }
 	virtual bool Link(FStateTreeLinker& InLinker) override;
 	//~ End FStateTreeNodeBase
diff --git a/Engine/Plugins/Experimental/Avalanche/Source/Avalanche/Public/Transition/Conditions/AvaTransitionRCControllerMatchCondition.h b/Engine/Plugins/Experimental/Avalanche/Source/Avalanche/Public/Transition/Conditions/AvaTransitionRCControllerMatchCondition.h
--- a/Engine/Plugins/Experimental/Avalanche/Source/Avalanche/Public/Transition/Conditions/AvaTransitionRCControllerMatchCondition.h
+++ b/Engine/Plugins/Experimental/Avalanche/Source/Avalanche/Public/Transition/Conditions/AvaTransitionRCControllerMatchCondition.h
@@ -24,11 +24,10 @@ struct AVALANCHE_API FAvaTransitionRCControllerMatchCondition : public FAvaTrans
 
 	using FInstanceDataType = FAvaTransitionRCControllerMatchConditionInstanceData;
 
-	//~ Begin FAvaTransitionCondition
-	virtual FText GenerateDescription(const FAvaTransitionNodeContext& InContext) const override;
-	//~ End FAvaTransitionCondition
-
 	//~ Begin FStateTreeNodeBase
+#if WITH_EDITOR
+	virtual FText GetDescription(const FGuid& InId, FStateTreeDataView InInstanceDataView, const IStateTreeBindingLookup& InBindingLookup, EStateTreeNodeFormatting InFormatting) const override;
+#endif
 	virtual const UStruct* GetInstanceDataType() const override { return FInstanceDataType::StaticStruct(); }
 	virtual bool Link(FStateTreeLinker& InLinker) override;
 	//~ End FStateTreeNodeBase
diff --git a/Engine/Plugins/Experimental/Avalanche/Source/Avalanche/Public/Transition/Tasks/AvaSceneAddTagAttributeTask.h b/Engine/Plugins/Experimental/Avalanche/Source/Avalanche/Public/Transition/Tasks/AvaSceneAddTagAttributeTask.h
--- a/Engine/Plugins/Experimental/Avalanche/Source/Avalanche/Public/Transition/Tasks/AvaSceneAddTagAttributeTask.h
+++ b/Engine/Plugins/Experimental/Avalanche/Source/Avalanche/Public/Transition/Tasks/AvaSceneAddTagAttributeTask.h
@@ -10,9 +10,11 @@ struct AVALANCHE_API FAvaSceneAddTagAttributeTask : public FAvaSceneTask
 {
 	GENERATED_BODY()
 
-	//~ Begin FAvaTransitionTask
-	virtual FText GenerateDescription(const FAvaTransitionNodeContext& InContext) const override;
-	//~ End FAvaTransitionTask
+	//~ Begin FStateTreeNodeBase
+#if WITH_EDITOR
+	virtual FText GetDescription(const FGuid& InId, FStateTreeDataView InInstanceDataView, const IStateTreeBindingLookup& InBindingLookup, EStateTreeNodeFormatting InFormatting) const override;
+#endif
+	//~ End FStateTreeNodeBase
 
 	//~ Begin FStateTreeTaskBase
 	virtual EStateTreeRunStatus EnterState(FStateTreeExecutionContext& InContext, const FStateTreeTransitionResult& InTransition) const override;
diff --git a/Engine/Plugins/Experimental/Avalanche/Source/Avalanche/Public/Transition/Tasks/AvaSceneRemoveTagAttributeTask.h b/Engine/Plugins/Experimental/Avalanche/Source/Avalanche/Public/Transition/Tasks/AvaSceneRemoveTagAttributeTask.h
--- a/Engine/Plugins/Experimental/Avalanche/Source/Avalanche/Public/Transition/Tasks/AvaSceneRemoveTagAttributeTask.h
+++ b/Engine/Plugins/Experimental/Avalanche/Source/Avalanche/Public/Transition/Tasks/AvaSceneRemoveTagAttributeTask.h
@@ -10,9 +10,11 @@ struct AVALANCHE_API FAvaSceneRemoveTagAttributeTask : public FAvaSceneTask
 {
 	GENERATED_BODY()
 
-	//~ Begin FAvaTransitionTask
-	virtual FText GenerateDescription(const FAvaTransitionNodeContext& InContext) const override;
-	//~ End FAvaTransitionTask
+	//~ Begin FStateTreeNodeBase
+#if WITH_EDITOR
+	virtual FText GetDescription(const FGuid& InId, FStateTreeDataView InInstanceDataView, const IStateTreeBindingLookup& InBindingLookup, EStateTreeNodeFormatting InFormatting) const override;
+#endif
+	//~ End FStateTreeNodeBase
 
 	//~ Begin FStateTreeTaskBase
 	virtual EStateTreeRunStatus EnterState(FStateTreeExecutionContext& InContext, const FStateTreeTransitionResult& InTransition) const override;
diff --git a/Engine/Plugins/Experimental/Avalanche/Source/AvalancheSequence/Private/Transition/AvaTransitionPlaySequenceTask.cpp b/Engine/Plugins/Experimental/Avalanche/Source/AvalancheSequence/Private/Transition/AvaTransitionPlaySequenceTask.cpp
--- a/Engine/Plugins/Experimental/Avalanche/Source/AvalancheSequence/Private/Transition/AvaTransitionPlaySequenceTask.cpp
+++ b/Engine/Plugins/Experimental/Avalanche/Source/AvalancheSequence/Private/Transition/AvaTransitionPlaySequenceTask.cpp
@@ -10,7 +10,8 @@
 
 #define LOCTEXT_NAMESPACE "AvaTransitionPlaySequenceTask"
 
-FText FAvaTransitionPlaySequenceTask::GenerateDescription(const FAvaTransitionNodeContext& InContext) const
+#if WITH_EDITOR
+FText FAvaTransitionPlaySequenceTask::GetDescription(const FGuid& InId, FStateTreeDataView InInstanceDataView, const IStateTreeBindingLookup& InBindingLookup, EStateTreeNodeFormatting InFormatting) const
 {
 	FText DescriptionFormat = LOCTEXT("TaskDescription", "Play {SequenceQueryText} {AdditionalArgs}");
 
@@ -31,6 +32,7 @@ FText FAvaTransitionPlaySequenceTask::GenerateDescription(const FAvaTransitionNo
 
 	return FText::Format(DescriptionFormat, NamedArgs);
 }
+#endif
 
 TArray<UAvaSequencePlayer*> FAvaTransitionPlaySequenceTask::ExecuteSequenceTask(FStateTreeExecutionContext& InContext) const
 {
diff --git a/Engine/Plugins/Experimental/Avalanche/Source/AvalancheSequence/Public/Transition/AvaTransitionPlaySequenceTask.h b/Engine/Plugins/Experimental/Avalanche/Source/AvalancheSequence/Public/Transition/AvaTransitionPlaySequenceTask.h
--- a/Engine/Plugins/Experimental/Avalanche/Source/AvalancheSequence/Public/Transition/AvaTransitionPlaySequenceTask.h
+++ b/Engine/Plugins/Experimental/Avalanche/Source/AvalancheSequence/Public/Transition/AvaTransitionPlaySequenceTask.h
@@ -11,9 +11,11 @@ struct AVALANCHESEQUENCE_API FAvaTransitionPlaySequenceTask : public FAvaTransit
 {
 	GENERATED_BODY()
 
-	//~ Begin FAvaTransitionTask
-	virtual FText GenerateDescription(const FAvaTransitionNodeContext& InContext) const override;
-	//~ End FAvaTransitionTask
+	//~ Begin FStateTreeNodeBase
+#if WITH_EDITOR
+	virtual FText GetDescription(const FGuid& InId, FStateTreeDataView InInstanceDataView, const IStateTreeBindingLookup& InBindingLookup, EStateTreeNodeFormatting InFormatting) const override;
+#endif
+	//~ End FStateTreeNodeBase
 
 	//~ Begin FAvaTransitionSequenceTaskBase
 	virtual TArray<UAvaSequencePlayer*> ExecuteSequenceTask(FStateTreeExecutionContext& InContext) const override;
diff --git a/Engine/Plugins/Experimental/Avalanche/Source/AvalancheTransition/Private/Conditions/AvaTransitionLayerMatchCondition.cpp b/Engine/Plugins/Experimental/Avalanche/Source/AvalancheTransition/Private/Conditions/AvaTransitionLayerMatchCondition.cpp
--- a/Engine/Plugins/Experimental/Avalanche/Source/AvalancheTransition/Private/Conditions/AvaTransitionLayerMatchCondition.cpp
+++ b/Engine/Plugins/Experimental/Avalanche/Source/AvalancheTransition/Private/Conditions/AvaTransitionLayerMatchCondition.cpp
@@ -7,10 +7,12 @@
 
 #define LOCTEXT_NAMESPACE "AvaTransitionLayerMatchCondition"
 
-FText FAvaTransitionLayerMatchCondition::GenerateDescription(const FAvaTransitionNodeContext& InContext) const
+#if WITH_EDITOR
+FText FAvaTransitionLayerMatchCondition::GetDescription(const FGuid& InId, FStateTreeDataView InInstanceDataView, const IStateTreeBindingLookup& InBindingLookup, EStateTreeNodeFormatting InFormatting) const
 {
 	return FText::Format(LOCTEXT("ConditionDescription", "scenes transitioning in {0}"), GetLayerQueryText());
 }
+#endif
 
 bool FAvaTransitionLayerMatchCondition::TestCondition(FStateTreeExecutionContext& InContext) const
 {
diff --git a/Engine/Plugins/Experimental/Avalanche/Source/AvalancheTransition/Private/Conditions/AvaTransitionSceneMatchCondition.cpp b/Engine/Plugins/Experimental/Avalanche/Source/AvalancheTransition/Private/Conditions/AvaTransitionSceneMatchCondition.cpp
--- a/Engine/Plugins/Experimental/Avalanche/Source/AvalancheTransition/Private/Conditions/AvaTransitionSceneMatchCondition.cpp
+++ b/Engine/Plugins/Experimental/Avalanche/Source/AvalancheTransition/Private/Conditions/AvaTransitionSceneMatchCondition.cpp
@@ -8,12 +8,14 @@
 
 #define LOCTEXT_NAMESPACE "AvaTransitionSceneMatchCondition"
 
-FText FAvaTransitionSceneMatchCondition::GenerateDescription(const FAvaTransitionNodeContext& InContext) const
+#if WITH_EDITOR
+FText FAvaTransitionSceneMatchCondition::GetDescription(const FGuid& InId, FStateTreeDataView InInstanceDataView, const IStateTreeBindingLookup& InBindingLookup, EStateTreeNodeFormatting InFormatting) const
 {
 	return FText::Format(LOCTEXT("ConditionDescription", "{0} scene in {1}")
 		, UEnum::GetDisplayValueAsText(SceneComparisonType).ToLower()
 		, GetLayerQueryText());
 }
+#endif
 
 bool FAvaTransitionSceneMatchCondition::TestCondition(FStateTreeExecutionContext& InContext) const
 {
diff --git a/Engine/Plugins/Experimental/Avalanche/Source/AvalancheTransition/Private/Conditions/AvaTransitionStateMatchCondition.cpp b/Engine/Plugins/Experimental/Avalanche/Source/AvalancheTransition/Private/Conditions/AvaTransitionStateMatchCondition.cpp
--- a/Engine/Plugins/Experimental/Avalanche/Source/AvalancheTransition/Private/Conditions/AvaTransitionStateMatchCondition.cpp
+++ b/Engine/Plugins/Experimental/Avalanche/Source/AvalancheTransition/Private/Conditions/AvaTransitionStateMatchCondition.cpp
@@ -5,12 +5,14 @@
 
 #define LOCTEXT_NAMESPACE "AvaTransitionSceneMatchCondition"
 
-FText FAvaTransitionStateMatchCondition::GenerateDescription(const FAvaTransitionNodeContext& InContext) const
+#if WITH_EDITOR
+FText FAvaTransitionStateMatchCondition::GetDescription(const FGuid& InId, FStateTreeDataView InInstanceDataView, const IStateTreeBindingLookup& InBindingLookup, EStateTreeNodeFormatting InFormatting) const
 {
 	return FText::Format(LOCTEXT("ConditionDescription", "{0} scene in {1}")
 		, UEnum::GetDisplayValueAsText(TransitionState).ToLower()
 		, GetLayerQueryText());
 }
+#endif
 
 bool FAvaTransitionStateMatchCondition::TestCondition(FStateTreeExecutionContext& InContext) const
 {
diff --git a/Engine/Plugins/Experimental/Avalanche/Source/AvalancheTransition/Private/Conditions/AvaTransitionTypeMatchCondition.cpp b/Engine/Plugins/Experimental/Avalanche/Source/AvalancheTransition/Private/Conditions/AvaTransitionTypeMatchCondition.cpp
--- a/Engine/Plugins/Experimental/Avalanche/Source/AvalancheTransition/Private/Conditions/AvaTransitionTypeMatchCondition.cpp
+++ b/Engine/Plugins/Experimental/Avalanche/Source/AvalancheTransition/Private/Conditions/AvaTransitionTypeMatchCondition.cpp
@@ -6,11 +6,13 @@
 
 #define LOCTEXT_NAMESPACE "AvaTransitionTypeMatchCondition"
 
-FText FAvaTransitionTypeMatchCondition::GenerateDescription(const FAvaTransitionNodeContext& InContext) const
+#if WITH_EDITOR
+FText FAvaTransitionTypeMatchCondition::GetDescription(const FGuid& InId, FStateTreeDataView InInstanceDataView, const IStateTreeBindingLookup& InBindingLookup, EStateTreeNodeFormatting InFormatting) const
 {
 	return FText::Format(LOCTEXT("ConditionDescription", "transitioning {0}")
 		, UEnum::GetDisplayValueAsText(TransitionType).ToLower());
 }
+#endif
 
 bool FAvaTransitionTypeMatchCondition::TestCondition(FStateTreeExecutionContext& InContext) const
 {
diff --git a/Engine/Plugins/Experimental/Avalanche/Source/AvalancheTransition/Private/Tasks/AvaTransitionDelayTask.cpp b/Engine/Plugins/Experimental/Avalanche/Source/AvalancheTransition/Private/Tasks/AvaTransitionDelayTask.cpp
--- a/Engine/Plugins/Experimental/Avalanche/Source/AvalancheTransition/Private/Tasks/AvaTransitionDelayTask.cpp
+++ b/Engine/Plugins/Experimental/Avalanche/Source/AvalancheTransition/Private/Tasks/AvaTransitionDelayTask.cpp
@@ -5,10 +5,12 @@
 
 #define LOCTEXT_NAMESPACE "AvaTransitionDelayTask"
 
-FText FAvaTransitionDelayTask::GenerateDescription(const FAvaTransitionNodeContext& InContext) const
+#if WITH_EDITOR
+FText FAvaTransitionDelayTask::GetDescription(const FGuid& InId, FStateTreeDataView InInstanceDataView, const IStateTreeBindingLookup& InBindingLookup, EStateTreeNodeFormatting InFormatting) const
 {
 	return FText::Format(LOCTEXT("TaskDescription", "Delay {0} seconds"), FText::AsNumber(Duration));
 }
+#endif
 
 EStateTreeRunStatus FAvaTransitionDelayTask::EnterState(FStateTreeExecutionContext& InContext, const FStateTreeTransitionResult& InTransition) const
 {
diff --git a/Engine/Plugins/Experimental/Avalanche/Source/AvalancheTransition/Private/Tasks/AvaTransitionWaitForLayerTask.cpp b/Engine/Plugins/Experimental/Avalanche/Source/AvalancheTransition/Private/Tasks/AvaTransitionWaitForLayerTask.cpp
--- a/Engine/Plugins/Experimental/Avalanche/Source/AvalancheTransition/Private/Tasks/AvaTransitionWaitForLayerTask.cpp
+++ b/Engine/Plugins/Experimental/Avalanche/Source/AvalancheTransition/Private/Tasks/AvaTransitionWaitForLayerTask.cpp
@@ -6,11 +6,12 @@
 
 #define LOCTEXT_NAMESPACE "AvaTransitionWaitForLayerTask"
 
-FText FAvaTransitionWaitForLayerTask::GenerateDescription(const FAvaTransitionNodeContext& InContext) const
+#if WITH_EDITOR
+FText FAvaTransitionWaitForLayerTask::GetDescription(const FGuid& InId, FStateTreeDataView InInstanceDataView, const IStateTreeBindingLookup& InBindingLookup, EStateTreeNodeFormatting InFormatting) const
 {
-	return FText::Format(LOCTEXT("TaskDescription", "Wait for others in {0} to finish")
-		, GetLayerQueryText());
+	return FText::Format(LOCTEXT("TaskDescription", "Wait for others in {0} to finish"), GetLayerQueryText());
 }
+#endif
 
 EStateTreeRunStatus FAvaTransitionWaitForLayerTask::EnterState(FStateTreeExecutionContext& InContext, const FStateTreeTransitionResult& InTransition) const
 {
diff --git a/Engine/Plugins/Experimental/Avalanche/Source/AvalancheTransition/Public/AvaTransitionNode.h b/Engine/Plugins/Experimental/Avalanche/Source/AvalancheTransition/Public/AvaTransitionNode.h
--- a/Engine/Plugins/Experimental/Avalanche/Source/AvalancheTransition/Public/AvaTransitionNode.h
+++ b/Engine/Plugins/Experimental/Avalanche/Source/AvalancheTransition/Public/AvaTransitionNode.h
@@ -4,7 +4,6 @@
 
 #include "StateTreeExecutionTypes.h"
 
-class FAvaTransitionNodeContext;
 class UAvaTransitionSubsystem;
 struct FAvaTransitionContext;
 struct FStateTreeLinker;
@@ -14,12 +13,6 @@ class FAvaTransitionNode
 public:
 	AVALANCHETRANSITION_API bool LinkNode(FStateTreeLinker& InLinker);
 
-	/** Generates a more user-friendly dynamic description of the Node based on its settings */
-	virtual FText GenerateDescription(const FAvaTransitionNodeContext& InContext) const
-	{
-		return FText::GetEmpty();
-	}
-
 	TStateTreeExternalDataHandle<FAvaTransitionContext> TransitionContextHandle;
 
 	TStateTreeExternalDataHandle<UAvaTransitionSubsystem> TransitionSubsystemHandle;
diff --git a/Engine/Plugins/Experimental/Avalanche/Source/AvalancheTransition/Public/AvaTransitionNodeContext.h b/Engine/Plugins/Experimental/Avalanche/Source/AvalancheTransition/Public/AvaTransitionNodeContext.h
deleted file mode 100644
diff --git a/Engine/Plugins/Experimental/Avalanche/Source/AvalancheTransition/Public/Conditions/AvaTransitionLayerMatchCondition.h b/Engine/Plugins/Experimental/Avalanche/Source/AvalancheTransition/Public/Conditions/AvaTransitionLayerMatchCondition.h
--- a/Engine/Plugins/Experimental/Avalanche/Source/AvalancheTransition/Public/Conditions/AvaTransitionLayerMatchCondition.h
+++ b/Engine/Plugins/Experimental/Avalanche/Source/AvalancheTransition/Public/Conditions/AvaTransitionLayerMatchCondition.h
@@ -18,11 +18,10 @@ struct AVALANCHETRANSITION_API FAvaTransitionLayerMatchCondition : public FAvaTr
 
 	using FInstanceDataType = FAvaTransitionLayerMatchConditionInstanceData;
 
-	//~ Begin FAvaTransitionCondition
-	virtual FText GenerateDescription(const FAvaTransitionNodeContext& InContext) const override;
-	//~ End FAvaTransitionCondition
-
 	//~ Begin FStateTreeNodeBase
+#if WITH_EDITOR
+	virtual FText GetDescription(const FGuid& InId, FStateTreeDataView InInstanceDataView, const IStateTreeBindingLookup& InBindingLookup, EStateTreeNodeFormatting InFormatting) const override;
+#endif
 	virtual const UStruct* GetInstanceDataType() const override { return FInstanceDataType::StaticStruct(); }
 	//~ End FStateTreeNodeBase
 
diff --git a/Engine/Plugins/Experimental/Avalanche/Source/AvalancheTransition/Public/Conditions/AvaTransitionSceneMatchCondition.h b/Engine/Plugins/Experimental/Avalanche/Source/AvalancheTransition/Public/Conditions/AvaTransitionSceneMatchCondition.h
--- a/Engine/Plugins/Experimental/Avalanche/Source/AvalancheTransition/Public/Conditions/AvaTransitionSceneMatchCondition.h
+++ b/Engine/Plugins/Experimental/Avalanche/Source/AvalancheTransition/Public/Conditions/AvaTransitionSceneMatchCondition.h
@@ -19,11 +19,10 @@ struct AVALANCHETRANSITION_API FAvaTransitionSceneMatchCondition : public FAvaTr
 
 	using FInstanceDataType = FAvaTransitionSceneMatchConditionInstanceData;
 
-	//~ Begin FAvaTransitionCondition
-	virtual FText GenerateDescription(const FAvaTransitionNodeContext& InContext) const override;
-	//~ End FAvaTransitionCondition
-
 	//~ Begin FStateTreeNodeBase
+#if WITH_EDITOR
+	virtual FText GetDescription(const FGuid& InId, FStateTreeDataView InInstanceDataView, const IStateTreeBindingLookup& InBindingLookup, EStateTreeNodeFormatting InFormatting) const override;
+#endif
 	virtual const UStruct* GetInstanceDataType() const override { return FInstanceDataType::StaticStruct(); }
 	//~ End FStateTreeNodeBase
 
diff --git a/Engine/Plugins/Experimental/Avalanche/Source/AvalancheTransition/Public/Conditions/AvaTransitionStateMatchCondition.h b/Engine/Plugins/Experimental/Avalanche/Source/AvalancheTransition/Public/Conditions/AvaTransitionStateMatchCondition.h
--- a/Engine/Plugins/Experimental/Avalanche/Source/AvalancheTransition/Public/Conditions/AvaTransitionStateMatchCondition.h
+++ b/Engine/Plugins/Experimental/Avalanche/Source/AvalancheTransition/Public/Conditions/AvaTransitionStateMatchCondition.h
@@ -19,11 +19,10 @@ struct AVALANCHETRANSITION_API FAvaTransitionStateMatchCondition : public FAvaTr
 
 	using FInstanceDataType = FAvaTransitionStateMatchConditionInstanceData;
 
-	//~ Begin FAvaTransitionCondition
-	virtual FText GenerateDescription(const FAvaTransitionNodeContext& InContext) const override;
-	//~ End FAvaTransitionCondition
-
 	//~ Begin FStateTreeNodeBase
+#if WITH_EDITOR
+	virtual FText GetDescription(const FGuid& InId, FStateTreeDataView InInstanceDataView, const IStateTreeBindingLookup& InBindingLookup, EStateTreeNodeFormatting InFormatting) const override;
+#endif
 	virtual const UStruct* GetInstanceDataType() const override { return FAvaTransitionStateMatchConditionInstanceData::StaticStruct(); }
 	//~ End FStateTreeNodeBase
 
diff --git a/Engine/Plugins/Experimental/Avalanche/Source/AvalancheTransition/Public/Conditions/AvaTransitionTypeMatchCondition.h b/Engine/Plugins/Experimental/Avalanche/Source/AvalancheTransition/Public/Conditions/AvaTransitionTypeMatchCondition.h
--- a/Engine/Plugins/Experimental/Avalanche/Source/AvalancheTransition/Public/Conditions/AvaTransitionTypeMatchCondition.h
+++ b/Engine/Plugins/Experimental/Avalanche/Source/AvalancheTransition/Public/Conditions/AvaTransitionTypeMatchCondition.h
@@ -21,11 +21,10 @@ struct AVALANCHETRANSITION_API FAvaTransitionTypeMatchCondition : public FAvaTra
 
 	FAvaTransitionTypeMatchCondition() = default;
 
-	//~ Begin FAvaTransitionCondition
-	virtual FText GenerateDescription(const FAvaTransitionNodeContext& InContext) const override;
-	//~ End FAvaTransitionCondition
-
 	//~ Begin FStateTreeNodeBase
+#if WITH_EDITOR
+	virtual FText GetDescription(const FGuid& InId, FStateTreeDataView InInstanceDataView, const IStateTreeBindingLookup& InBindingLookup, EStateTreeNodeFormatting InFormatting) const override;
+#endif
 	virtual const UStruct* GetInstanceDataType() const override { return FInstanceDataType::StaticStruct(); }
 	//~ End FStateTreeNodeBase
 
diff --git a/Engine/Plugins/Experimental/Avalanche/Source/AvalancheTransition/Public/Tasks/AvaTransitionDelayTask.h b/Engine/Plugins/Experimental/Avalanche/Source/AvalancheTransition/Public/Tasks/AvaTransitionDelayTask.h
--- a/Engine/Plugins/Experimental/Avalanche/Source/AvalancheTransition/Public/Tasks/AvaTransitionDelayTask.h
+++ b/Engine/Plugins/Experimental/Avalanche/Source/AvalancheTransition/Public/Tasks/AvaTransitionDelayTask.h
@@ -21,11 +21,10 @@ struct AVALANCHETRANSITION_API FAvaTransitionDelayTask : public FAvaTransitionTa
 
 	using FInstanceDataType = FAvaTransitionDelayTaskInstanceData;
 
-	//~ Begin FAvaTransitionTask
-	virtual FText GenerateDescription(const FAvaTransitionNodeContext& InContext) const override;
-	//~ End FAvaTransitionTask
-
 	//~ Begin FStateTreeNodeBase
+#if WITH_EDITOR
+	virtual FText GetDescription(const FGuid& InId, FStateTreeDataView InInstanceDataView, const IStateTreeBindingLookup& InBindingLookup, EStateTreeNodeFormatting InFormatting) const override;
+#endif
 	virtual const UStruct* GetInstanceDataType() const override { return FInstanceDataType::StaticStruct(); }
 	//~ End FStateTreeNodeBase
 
diff --git a/Engine/Plugins/Experimental/Avalanche/Source/AvalancheTransition/Public/Tasks/AvaTransitionWaitForLayerTask.h b/Engine/Plugins/Experimental/Avalanche/Source/AvalancheTransition/Public/Tasks/AvaTransitionWaitForLayerTask.h
--- a/Engine/Plugins/Experimental/Avalanche/Source/AvalancheTransition/Public/Tasks/AvaTransitionWaitForLayerTask.h
+++ b/Engine/Plugins/Experimental/Avalanche/Source/AvalancheTransition/Public/Tasks/AvaTransitionWaitForLayerTask.h
@@ -18,11 +18,10 @@ struct AVALANCHETRANSITION_API FAvaTransitionWaitForLayerTask : public FAvaTrans
 
 	using FInstanceDataType = FAvaTransitionWaitForLayerTaskInstanceData;
 
-	//~ Begin FAvaTransitionTask
-	virtual FText GenerateDescription(const FAvaTransitionNodeContext& InContext) const override;
-	//~ End FAvaTransitionTask
-
 	//~ Begin FStateTreeNodeBase
+#if WITH_EDITOR
+	virtual FText GetDescription(const FGuid& InId, FStateTreeDataView InInstanceDataView, const IStateTreeBindingLookup& InBindingLookup, EStateTreeNodeFormatting InFormatting) const override;
+#endif
 	virtual const UStruct* GetInstanceDataType() const override { return FInstanceDataType::StaticStruct(); }
 	//~ End FStateTreeNodeBase
 
diff --git a/Engine/Plugins/Experimental/Avalanche/Source/AvalancheTransitionEditor/Private/ViewModels/AvaTransitionContainerViewModel.cpp b/Engine/Plugins/Experimental/Avalanche/Source/AvalancheTransitionEditor/Private/ViewModels/AvaTransitionContainerViewModel.cpp
--- a/Engine/Plugins/Experimental/Avalanche/Source/AvalancheTransitionEditor/Private/ViewModels/AvaTransitionContainerViewModel.cpp
+++ b/Engine/Plugins/Experimental/Avalanche/Source/AvalancheTransitionEditor/Private/ViewModels/AvaTransitionContainerViewModel.cpp
@@ -1,6 +1,9 @@
 // Copyright Epic Games, Inc. All Rights Reserved.
 
 #include "AvaTransitionContainerViewModel.h"
+#include "AvaTransitionEditorViewModel.h"
+#include "AvaTransitionTreeEditorData.h"
+#include "AvaTransitionViewModelSharedData.h"
 #include "StateTreeState.h"
 
 FAvaTransitionContainerViewModel::FAvaTransitionContainerViewModel(UStateTreeState* InState)
@@ -13,6 +16,15 @@ UStateTreeState* FAvaTransitionContainerViewModel::GetState() const
 	return StateWeak.Get();
 }
 
+UAvaTransitionTreeEditorData* FAvaTransitionContainerViewModel::GetEditorData() const
+{
+	if (TSharedPtr<FAvaTransitionEditorViewModel> EditorViewModel = GetSharedData()->GetEditorViewModel())
+	{
+		return EditorViewModel->GetEditorData();
+	}
+	return nullptr;
+}
+
 bool FAvaTransitionContainerViewModel::IsValid() const
 {
 	return StateWeak.IsValid();
diff --git a/Engine/Plugins/Experimental/Avalanche/Source/AvalancheTransitionEditor/Private/ViewModels/AvaTransitionContainerViewModel.h b/Engine/Plugins/Experimental/Avalanche/Source/AvalancheTransitionEditor/Private/ViewModels/AvaTransitionContainerViewModel.h
--- a/Engine/Plugins/Experimental/Avalanche/Source/AvalancheTransitionEditor/Private/ViewModels/AvaTransitionContainerViewModel.h
+++ b/Engine/Plugins/Experimental/Avalanche/Source/AvalancheTransitionEditor/Private/ViewModels/AvaTransitionContainerViewModel.h
@@ -5,6 +5,7 @@
 #include "AvaTransitionViewModel.h"
 #include "UObject/WeakObjectPtr.h"
 
+class UAvaTransitionTreeEditorData;
 class UStateTreeState;
 
 /** View Model serving as a container for a particular group of elements in a State */
@@ -17,6 +18,8 @@ public:
 
 	UStateTreeState* GetState() const;
 
+	UAvaTransitionTreeEditorData* GetEditorData() const;
+
 	//~ Begin FAvaTransitionViewModel
 	virtual bool IsValid() const override;
 	//~ End FAvaTransitionViewModel
diff --git a/Engine/Plugins/Experimental/Avalanche/Source/AvalancheTransitionEditor/Private/ViewModels/Condition/AvaTransitionConditionContainerViewModel.cpp b/Engine/Plugins/Experimental/Avalanche/Source/AvalancheTransitionEditor/Private/ViewModels/Condition/AvaTransitionConditionContainerViewModel.cpp
--- a/Engine/Plugins/Experimental/Avalanche/Source/AvalancheTransitionEditor/Private/ViewModels/Condition/AvaTransitionConditionContainerViewModel.cpp
+++ b/Engine/Plugins/Experimental/Avalanche/Source/AvalancheTransitionEditor/Private/ViewModels/Condition/AvaTransitionConditionContainerViewModel.cpp
@@ -2,7 +2,7 @@
 
 #include "AvaTransitionConditionContainerViewModel.h"
 #include "AvaTransitionConditionViewModel.h"
-#include "AvaTransitionNodeContext.h"
+#include "AvaTransitionTreeEditorData.h"
 #include "Conditions/AvaTransitionCondition.h"
 #include "StateTreeEditorStyle.h"
 #include "StateTreeState.h"
@@ -24,22 +24,6 @@ namespace UE::AvaTransitionEditor::Private
 		return FText::GetEmpty();
 	}
 
-	FText GetDescriptionText(const FStateTreeEditorNode& InEditorNode)
-	{
-		FText ConditionDescription;
-		if (const FAvaTransitionCondition* Condition = InEditorNode.Node.GetPtr<FAvaTransitionCondition>())
-		{
-			ConditionDescription = Condition->GenerateDescription(FAvaTransitionNodeContext(InEditorNode.GetInstance()));
-		}
-
-		const UScriptStruct* Struct = InEditorNode.Node.GetScriptStruct();
-		if (Struct && ConditionDescription.IsEmpty())
-		{
-			ConditionDescription = Struct->GetDisplayNameText();
-		}
-		return ConditionDescription;
-	}
-
 	FText GetParenthesisText(int8 InDeltaIndent)
 	{
 		if (int8 IndentAmount = FMath::Abs(InDeltaIndent))
@@ -87,6 +71,12 @@ FText FAvaTransitionConditionContainerViewModel::UpdateStateDescription() const
 		return FText::GetEmpty();
 	}
 
+	const UAvaTransitionTreeEditorData* EditorData = GetEditorData();
+	if (!EditorData)
+	{
+		return FText::GetEmpty();
+	}
+
 	TArray<FText> ConditionDescriptions;
 	for (int32 ConditionIndex = 0; ConditionIndex < ConditionViewModels.Num(); ++ConditionIndex)
 	{
@@ -114,7 +104,7 @@ FText FAvaTransitionConditionContainerViewModel::UpdateStateDescription() const
 
 		FFormatNamedArguments TextArguments;
 		TextArguments.Add(TEXT("Operand"), Private::GetOperandText(ConditionIndex, *EditorNode));
-		TextArguments.Add(TEXT("Description"), Private::GetDescriptionText(*EditorNode));
+		TextArguments.Add(TEXT("Description"), EditorData->GetNodeDescription(*EditorNode, EStateTreeNodeFormatting::RichText));
 		TextArguments.Add(TEXT("Parenthesis"), Private::GetParenthesisText(DeltaIndent));
 
 		ConditionDescriptions.Add(FText::Format(TextFormat, TextArguments));
diff --git a/Engine/Plugins/Experimental/Avalanche/Source/AvalancheTransitionEditor/Private/ViewModels/Task/AvaTransitionTaskViewModel.cpp b/Engine/Plugins/Experimental/Avalanche/Source/AvalancheTransitionEditor/Private/ViewModels/Task/AvaTransitionTaskViewModel.cpp
--- a/Engine/Plugins/Experimental/Avalanche/Source/AvalancheTransitionEditor/Private/ViewModels/Task/AvaTransitionTaskViewModel.cpp
+++ b/Engine/Plugins/Experimental/Avalanche/Source/AvalancheTransitionEditor/Private/ViewModels/Task/AvaTransitionTaskViewModel.cpp
@@ -2,7 +2,6 @@
 
 #include "AvaTransitionTaskViewModel.h"
 #include "AvaTransitionEditorStyle.h"
-#include "AvaTransitionNodeContext.h"
 #include "AvaTransitionTreeEditorData.h"
 #include "StateTreeEditorData.h"
 #include "StateTreeEditorStyle.h"
@@ -68,34 +67,16 @@ FSlateColor FAvaTransitionTaskViewModel::GetTaskIconColor() const
 
 void FAvaTransitionTaskViewModel::UpdateTaskDescription()
 {
-	TaskDescription = FText::GetEmpty();
-
+	const UAvaTransitionTreeEditorData* EditorData = GetEditorData();
 	const FStateTreeEditorNode* EditorNode = GetEditorNode();
-	if (!EditorNode)
-	{
-		return;
-	}
-
-	// Advanced Mode: Show Task Descriptions from the Node Name (Raw) if name not none
-	if (GetSharedData()->GetEditorMode() == EAvaTransitionEditorMode::Advanced)
-	{
-		const FStateTreeNodeBase* Task = EditorNode->Node.GetPtr<FStateTreeNodeBase>();
-		if (Task && Task->Name != NAME_None)
-		{
-			TaskDescription = FText::FromName(Task->Name);
-			return;
-		}
-	}
 
-	if (const FAvaTransitionTask* Task = EditorNode->Node.GetPtr<FAvaTransitionTask>())
+	if (EditorData && EditorNode)
 	{
-		TaskDescription = Task->GenerateDescription(FAvaTransitionNodeContext(EditorNode->GetInstance()));
+		TaskDescription = EditorData->GetNodeDescription(*EditorNode, EStateTreeNodeFormatting::RichText);
 	}
-
-	const UScriptStruct* Struct = EditorNode->Node.GetScriptStruct();
-	if (Struct && TaskDescription.IsEmpty())
+	else
 	{
-		TaskDescription = Struct->GetDisplayNameText();
+		TaskDescription = FText::GetEmpty();
 	}
 }
 
