[Motion Design] Rundown Server API: Server Info

Provides additional information about the rundown server process:
- Engine build and mode: given that some commands may not be implemented in game build for instance, knowing this, the client can disable the corresponding UI and not generate errors.
- Project directory can help with deploy tools.
- Web RC ports (http, ws)
- Other information may help with other available api.

diff --git a/Engine/Plugins/Experimental/Avalanche/Source/AvalancheMedia/AvalancheMedia.Build.cs b/Engine/Plugins/Experimental/Avalanche/Source/AvalancheMedia/AvalancheMedia.Build.cs
--- a/Engine/Plugins/Experimental/Avalanche/Source/AvalancheMedia/AvalancheMedia.Build.cs
+++ b/Engine/Plugins/Experimental/Avalanche/Source/AvalancheMedia/AvalancheMedia.Build.cs
@@ -45,6 +45,7 @@ public class AvalancheMedia : ModuleRules
 				"InputCore",
 				"Projects",
 				"RHI",
+				"RemoteControlCommon",
 				"RenderCore",
 				"Renderer",
 				"Slate",
diff --git a/Engine/Plugins/Experimental/Avalanche/Source/AvalancheMedia/Internal/Rundown/AvaRundownMessages.h b/Engine/Plugins/Experimental/Avalanche/Source/AvalancheMedia/Internal/Rundown/AvaRundownMessages.h
--- a/Engine/Plugins/Experimental/Avalanche/Source/AvalancheMedia/Internal/Rundown/AvaRundownMessages.h
+++ b/Engine/Plugins/Experimental/Avalanche/Source/AvalancheMedia/Internal/Rundown/AvaRundownMessages.h
@@ -19,6 +19,8 @@ namespace EAvaRundownApiVersion
 	 */
 	enum Type
 	{
+		Unspecified = -1,
+		
 		Initial = 1,
 		/**
 		 * The rundown server has been moved to the runtime module.
@@ -36,6 +38,37 @@ namespace EAvaRundownApiVersion
 	};
 }
 
+/**
+ * Build targets.
+ * This will help determine the set of features that are available.
+ */
+UENUM()
+enum class EAvaRundownServerBuildTargetType : uint8
+{
+	Unknown = 0,
+	Editor,
+	Game,
+	Server,
+	Client,
+	Program
+};
+
+/**
+ * An editor build can be launched in different modes but it could also be
+ * a dedicated build target. The engine mode combined with the build target
+ * will determine the set of functionalities available.
+ */
+UENUM()
+enum class EAvaRundownServerEngineMode : uint8
+{
+	Unknown = 0,
+	Editor,
+	Game,
+	Server,
+	Commandlet,
+	Other
+};
+
 
 USTRUCT()
 struct FAvaRundownMsgBase
@@ -43,7 +76,7 @@ struct FAvaRundownMsgBase
 	GENERATED_BODY()
 public:
 	UPROPERTY()
-	int32 RequestId = -1;
+	int32 RequestId = INDEX_NONE;
 };
 
 USTRUCT()
@@ -70,10 +103,10 @@ public:
 
 	/**
 	 * API Version the client has been implemented against.
-	 * If none (-1) the server will consider the initial version is requested.
+	 * If unspecified the server will consider the latest version is requested.
 	 */
 	UPROPERTY()
-	int32 RequestedApiVersion = -1;
+	int32 RequestedApiVersion = EAvaRundownApiVersion::Unspecified;
 };
 
 /** Response sent by server to client to be discovered. */
@@ -94,20 +127,87 @@ public:
 	 * Clients should expect an older server to reply with an older version.
 	 */
 	UPROPERTY()
-	int32 ApiVersion = -1;
+	int32 ApiVersion = EAvaRundownApiVersion::Unspecified;
 
 	/** Minimum API Version the server implements. */
 	UPROPERTY()
-	int32 MinimumApiVersion = -1;
+	int32 MinimumApiVersion = EAvaRundownApiVersion::Unspecified;
 
 	/** Latest API Version the server support. */
 	UPROPERTY()
-	int32 LatestApiVersion = -1;
+	int32 LatestApiVersion = EAvaRundownApiVersion::Unspecified;
 
 	UPROPERTY()
 	FString HostName;
 };
 
+/**
+ * Request the extended server information.
+ */
+USTRUCT()
+struct FAvaRundownGetServerInfo : public FAvaRundownMsgBase
+{
+	GENERATED_BODY()
+};
+
+/**
+ * Extended server information 
+ */
+USTRUCT()
+struct FAvaRundownServerInfo : public FAvaRundownMsgBase
+{
+	GENERATED_BODY()
+
+	/** API Version the server will communicate with for this client. */
+	UPROPERTY()
+	int32 ApiVersion = EAvaRundownApiVersion::Unspecified;
+
+	/** Minimum API Version the server implements. */
+	UPROPERTY()
+	int32 MinimumApiVersion = EAvaRundownApiVersion::Unspecified;
+
+	/** Latest API Version the server support. */
+	UPROPERTY()
+	int32 LatestApiVersion = EAvaRundownApiVersion::Unspecified;
+
+	UPROPERTY()
+	FString HostName;
+
+	/** Holds the engine version checksum */
+	UPROPERTY()
+	uint32 EngineVersion = 0;
+	
+	/** Holds the instance identifier. */
+	UPROPERTY()
+	FGuid InstanceId;
+
+	UPROPERTY()
+	EAvaRundownServerBuildTargetType InstanceBuild = EAvaRundownServerBuildTargetType::Unknown;
+
+	UPROPERTY()
+	EAvaRundownServerEngineMode InstanceMode = EAvaRundownServerEngineMode::Unknown;
+
+	/** Holds the identifier of the session that the application belongs to. */
+	UPROPERTY()
+	FGuid SessionId;
+	
+	/** The unreal project name this server is running from. */
+	UPROPERTY()
+	FString ProjectName;
+	
+	/** The unreal project directory this server is running from. */
+	UPROPERTY()
+	FString ProjectDir;
+
+	/** Http Server Port of the remote control service. */
+	UPROPERTY()
+	uint32 RemoteControlHttpServerPort = 0;
+
+	/** WebSocket Server Port of the remote control service. */
+	UPROPERTY()
+	uint32 RemoteControlWebSocketServerPort = 0;
+};
+
 /**
  *	Request list of rundown that can be opened on the current server.
  */
diff --git a/Engine/Plugins/Experimental/Avalanche/Source/AvalancheMedia/Private/Rundown/AvaRundownServer.cpp b/Engine/Plugins/Experimental/Avalanche/Source/AvalancheMedia/Private/Rundown/AvaRundownServer.cpp
--- a/Engine/Plugins/Experimental/Avalanche/Source/AvalancheMedia/Private/Rundown/AvaRundownServer.cpp
+++ b/Engine/Plugins/Experimental/Avalanche/Source/AvalancheMedia/Private/Rundown/AvaRundownServer.cpp
@@ -16,14 +16,17 @@
 #include "Broadcast/OutputDevices/AvaBroadcastOutputServerItem.h"
 #include "Broadcast/OutputDevices/AvaBroadcastOutputTreeItem.h"
 #include "Broadcast/OutputDevices/AvaBroadcastRenderTargetMediaUtils.h"
+#include "Engine/Engine.h"
 #include "IAvaMediaModule.h"
 #include "IRemoteControlModule.h"
 #include "ImageUtils.h"
 #include "MediaOutput.h"
 #include "MessageEndpointBuilder.h"
 #include "Misc/FileHelper.h"
+#include "Misc/Paths.h"
 #include "Playback/AvaPlaybackManager.h"
 #include "Playback/AvaPlaybackUtils.h"
+#include "RemoteControlSettings.h"
 #include "RenderingThread.h"
 #include "Rundown/AvaRundown.h"
 #include "Rundown/AvaRundownManagedInstanceCache.h"
@@ -44,6 +47,61 @@ DEFINE_LOG_CATEGORY_STATIC(LogAvaRundownServer, Log, All);
 
 namespace UE::AvaRundownServer::Private
 {
+	// We still support the initial version.
+	constexpr int32 CurrentMinimumApiVersion = EAvaRundownApiVersion::Initial;
+
+	EAvaRundownServerBuildTargetType GetRundownEngineBuild()
+	{
+		switch(FApp::GetBuildTargetType())
+		{
+		case EBuildTargetType::Unknown:
+			return EAvaRundownServerBuildTargetType::Unknown;
+		case EBuildTargetType::Game:
+			return EAvaRundownServerBuildTargetType::Game;
+		case EBuildTargetType::Server:
+			return EAvaRundownServerBuildTargetType::Server;
+		case EBuildTargetType::Client:
+			return EAvaRundownServerBuildTargetType::Client;
+		case EBuildTargetType::Editor:
+			return EAvaRundownServerBuildTargetType::Editor;
+		case EBuildTargetType::Program:
+			return EAvaRundownServerBuildTargetType::Program;
+		default:
+			return EAvaRundownServerBuildTargetType::Unknown;
+		}
+	}
+
+	EAvaRundownServerEngineMode GetRundownEngineMode()
+	{
+		if (!GEngine)
+		{
+			return EAvaRundownServerEngineMode::Unknown;
+		}
+		
+		if (IsRunningDedicatedServer())
+		{
+			return EAvaRundownServerEngineMode::Server;
+		}
+		
+		if (IsRunningCommandlet())
+		{
+			return EAvaRundownServerEngineMode::Commandlet;
+		}
+
+		// This is checking GIsEditor
+		if (GEngine->IsEditor())
+		{
+			return EAvaRundownServerEngineMode::Editor;
+		}
+		
+		if (IsRunningGame())
+		{
+			return EAvaRundownServerEngineMode::Game;
+		}
+		
+		return EAvaRundownServerEngineMode::Other;
+	}
+	
 	inline FAvaRundownPageInfo GetPageInfo(const UAvaRundown* InRundown, const FAvaRundownPage& InPage)
 	{
 		FAvaRundownPageInfo PageInfo;
@@ -237,6 +295,7 @@ void FAvaRundownServer::Init(const FString& InAssignedHostName)
 	
 	MessageEndpoint = FMessageEndpoint::Builder("MotionDesignRundownServer")
 	.Handling<FAvaRundownPing>(this, &FAvaRundownServer::HandleRundownPing)
+	.Handling<FAvaRundownGetServerInfo>(this, &FAvaRundownServer::HandleGetRundownServerInfo)
 	.Handling<FAvaRundownGetRundowns>(this, &FAvaRundownServer::HandleGetRundowns)
 	.Handling<FAvaRundownLoadRundown>(this, &FAvaRundownServer::HandleLoadRundown)
 	.Handling<FAvaRundownSaveRundown>(this, &FAvaRundownServer::HandleSaveRundown)
@@ -462,8 +521,7 @@ void FAvaRundownServer::HandleRundownPing(const FAvaRundownPing& InMessage, cons
 	ReplyMessage->RequestId = InMessage.RequestId;
 	ReplyMessage->bAuto = InMessage.bAuto;
 
-	// We still support the initial version.
-	constexpr int32 CurrentMinimumApiVersion = EAvaRundownApiVersion::Initial;
+	using namespace UE::AvaRundownServer::Private;
 
 	// Consider clients that didn't request a version to be the latest version.
 	const int32 RequestedApiVersion = InMessage.RequestedApiVersion != -1 ? InMessage.RequestedApiVersion : EAvaRundownApiVersion::LatestVersion;
@@ -487,6 +545,31 @@ void FAvaRundownServer::HandleRundownPing(const FAvaRundownPing& InMessage, cons
 	SendResponse(ReplyMessage, InContext->GetSender());
 }
 
+void FAvaRundownServer::HandleGetRundownServerInfo(const FAvaRundownGetServerInfo& InMessage, const TSharedRef<IMessageContext, ESPMode::ThreadSafe>& InContext)
+{
+	FAvaRundownServerInfo* ServerInfo = FMessageEndpoint::MakeMessage<FAvaRundownServerInfo>();
+	ServerInfo->RequestId = InMessage.RequestId;
+	
+	const FClientInfo* ClientInfo = GetClientInfo(InContext->GetSender());
+	ServerInfo->ApiVersion = ClientInfo ? ClientInfo->ApiVersion : EAvaRundownApiVersion::Unspecified;
+	ServerInfo->MinimumApiVersion = UE::AvaRundownServer::Private::CurrentMinimumApiVersion;
+	ServerInfo->LatestApiVersion = EAvaRundownApiVersion::LatestVersion;
+	ServerInfo->HostName = HostName;
+	ServerInfo->EngineVersion = FNetworkVersion::GetLocalNetworkVersion();
+	ServerInfo->InstanceId = FApp::GetInstanceId();
+	ServerInfo->InstanceBuild = UE::AvaRundownServer::Private::GetRundownEngineBuild();
+	ServerInfo->InstanceMode = UE::AvaRundownServer::Private::GetRundownEngineMode();
+	ServerInfo->SessionId = FApp::GetSessionId();
+	ServerInfo->ProjectName = FApp::GetProjectName();
+	ServerInfo->ProjectDir = FPaths::ConvertRelativePathToFull(FPaths::ProjectDir());
+
+	const URemoteControlSettings* RCSettings = GetDefault<URemoteControlSettings>();
+	ServerInfo->RemoteControlHttpServerPort = RCSettings->RemoteControlHttpServerPort;
+	ServerInfo->RemoteControlWebSocketServerPort = RCSettings->RemoteControlWebSocketServerPort;
+
+	SendResponse(ServerInfo, InContext->GetSender());
+}
+
 void FAvaRundownServer::HandleGetRundowns(const FAvaRundownGetRundowns& InMessage, const TSharedRef<IMessageContext, ESPMode::ThreadSafe>& InContext)
 {
 	FAvaRundownRundowns* ReplyMessage = FMessageEndpoint::MakeMessage<FAvaRundownRundowns>();
diff --git a/Engine/Plugins/Experimental/Avalanche/Source/AvalancheMedia/Private/Rundown/AvaRundownServer.h b/Engine/Plugins/Experimental/Avalanche/Source/AvalancheMedia/Private/Rundown/AvaRundownServer.h
--- a/Engine/Plugins/Experimental/Avalanche/Source/AvalancheMedia/Private/Rundown/AvaRundownServer.h
+++ b/Engine/Plugins/Experimental/Avalanche/Source/AvalancheMedia/Private/Rundown/AvaRundownServer.h
@@ -49,6 +49,7 @@ public:
 	
 	// Rundown message handlers
 	void HandleRundownPing(const FAvaRundownPing& InMessage, const TSharedRef<IMessageContext, ESPMode::ThreadSafe>& InContext);
+	void HandleGetRundownServerInfo(const FAvaRundownGetServerInfo& InMessage, const TSharedRef<IMessageContext, ESPMode::ThreadSafe>& InContext);
 	void HandleGetRundowns(const FAvaRundownGetRundowns& InMessage, const TSharedRef<IMessageContext, ESPMode::ThreadSafe>& InContext);
 	void HandleLoadRundown(const FAvaRundownLoadRundown& InMessage, const TSharedRef<IMessageContext, ESPMode::ThreadSafe>& InContext);
 	void HandleSaveRundown(const FAvaRundownSaveRundown& InMessage, const TSharedRef<IMessageContext, ESPMode::ThreadSafe>& InContext);
