Motion Design Font: Fix issue with project fonts not appearing

diff --git a/Engine/Plugins/Experimental/Avalanche/Source/AvalancheTextEditor/Private/Font/AvaFontManagerSubsystem.cpp b/Engine/Plugins/Experimental/Avalanche/Source/AvalancheTextEditor/Private/Font/AvaFontManagerSubsystem.cpp
--- a/Engine/Plugins/Experimental/Avalanche/Source/AvalancheTextEditor/Private/Font/AvaFontManagerSubsystem.cpp
+++ b/Engine/Plugins/Experimental/Avalanche/Source/AvalancheTextEditor/Private/Font/AvaFontManagerSubsystem.cpp
@@ -670,20 +670,19 @@ UAvaFontObject* UAvaFontManagerSubsystem::CreateProjectFont(UFont* InSourceFont,
 				}
 
 				// If the font object already exist, early return
-				if (UAvaFontObject* ExistingAvaFont = FindObject<UAvaFontObject>(ProjectFontsPackage, *SanitizedFontName))
+				UAvaFontObject* FontObject = FindObject<UAvaFontObject>(ProjectFontsPackage, *SanitizedFontName);
+				if (!FontObject)
 				{
-					return ExistingAvaFont;
+					FontObject = NewObject<UAvaFontObject>(ProjectFontsPackage, *SanitizedFontName, RF_Public | RF_Standalone | RF_Transient);
+					FontObject->AddToRoot();
 				}
 
-				UAvaFontObject* NewAvaFont = NewObject<UAvaFontObject>(ProjectFontsPackage, *SanitizedFontName, RF_Public | RF_Standalone | RF_Transient);
-				NewAvaFont->AddToRoot();
+				FontObject->InitProjectFont(InSourceFont, FontName);
+				UAvaFontManagerSubsystem::SetupMetrics(FontObject);
 
-				NewAvaFont->InitProjectFont(InSourceFont, FontName);
-				UAvaFontManagerSubsystem::SetupMetrics(NewAvaFont);
+				ProjectFontsMap.Add(SanitizedFontName, FontObject);
 
-				ProjectFontsMap.Add(SanitizedFontName, NewAvaFont);
-
-				return NewAvaFont;
+				return FontObject;
 			}
 			else
 			{
