Motion Design Outliner: Fix issue with Editor Replace Actor causing the actor item to disappear

diff --git a/Engine/Plugins/Experimental/Avalanche/Source/AvalancheOutliner/Private/AvaOutliner.cpp b/Engine/Plugins/Experimental/Avalanche/Source/AvalancheOutliner/Private/AvaOutliner.cpp
--- a/Engine/Plugins/Experimental/Avalanche/Source/AvalancheOutliner/Private/AvaOutliner.cpp
+++ b/Engine/Plugins/Experimental/Avalanche/Source/AvalancheOutliner/Private/AvaOutliner.cpp
@@ -66,6 +66,7 @@ FAvaOutliner::FAvaOutliner(IAvaOutlinerProvider& InOutlinerProvider)
 
 	//Listen to Object Replacement Changes
 	FCoreUObjectDelegates::OnObjectsReplaced.AddRaw(this, &FAvaOutliner::OnObjectsReplaced);
+	FEditorDelegates::OnEditorActorReplaced.AddRaw(this, &FAvaOutliner::OnActorReplaced);
 }
 
 FAvaOutliner::~FAvaOutliner()
@@ -84,6 +85,7 @@ FAvaOutliner::~FAvaOutliner()
 	}
 
 	FCoreUObjectDelegates::OnObjectsReplaced.RemoveAll(this);
+	FEditorDelegates::OnEditorActorReplaced.RemoveAll(this);
 }
 
 UAvaOutlinerSubsystem* FAvaOutliner::GetOutlinerSubsystem() const
@@ -1284,6 +1286,13 @@ void FAvaOutliner::OnObjectsReplaced(const TMap<UObject*, UObject*>& InReplaceme
 	});
 }
 
+void FAvaOutliner::OnActorReplaced(AActor* InOldActor, AActor* InNewActor)
+{
+	TMap<UObject*, UObject*> ReplacementMap;
+	ReplacementMap.Add(InOldActor, InNewActor);
+	OnObjectsReplaced(ReplacementMap);
+}
+
 void FAvaOutliner::SetOutlinerModified()
 {
 	if (!bOutlinerDirty)
diff --git a/Engine/Plugins/Experimental/Avalanche/Source/AvalancheOutliner/Private/AvaOutliner.h b/Engine/Plugins/Experimental/Avalanche/Source/AvalancheOutliner/Private/AvaOutliner.h
--- a/Engine/Plugins/Experimental/Avalanche/Source/AvalancheOutliner/Private/AvaOutliner.h
+++ b/Engine/Plugins/Experimental/Avalanche/Source/AvalancheOutliner/Private/AvaOutliner.h
@@ -247,9 +247,12 @@ public:
 	/** Called when an Actor's attachment has changed. This triggers a refresh */
 	void OnActorAttachmentChanged(AActor* InActor, const AActor* InParent, bool bAttach);
 
-	/** Called the engine replaces an object. A common example is when a BP Component is destroyed, and replaced */
+	/** Called when the engine replaces an object. A common example is when a BP Component is destroyed, and replaced */
 	void OnObjectsReplaced(const TMap<UObject*, UObject*>& InReplacementMap);
 
+	/** Called when the engine replaces an Actor */
+	void OnActorReplaced(AActor* InOldActor, AActor* InNewActor);
+
 	/** Marks the Outliner dirty. This triggers IAvaOutlinerProvider::OnOutlinerModified on next tick */
 	void SetOutlinerModified();
 
