[Motion Design] Remove indent from ease curve tool settings menu entries

diff --git a/Engine/Plugins/Experimental/Avalanche/Source/AvalancheSequencer/Private/EaseCurveTool/Widgets/AvaEaseCurveToolContextMenu.cpp b/Engine/Plugins/Experimental/Avalanche/Source/AvalancheSequencer/Private/EaseCurveTool/Widgets/AvaEaseCurveToolContextMenu.cpp
--- a/Engine/Plugins/Experimental/Avalanche/Source/AvalancheSequencer/Private/EaseCurveTool/Widgets/AvaEaseCurveToolContextMenu.cpp
+++ b/Engine/Plugins/Experimental/Avalanche/Source/AvalancheSequencer/Private/EaseCurveTool/Widgets/AvaEaseCurveToolContextMenu.cpp
@@ -292,7 +292,7 @@ void FAvaEaseCurveToolContextMenu::PopulateContextMenuSettings(UToolMenu* const
 					EaseCurveToolSettings->SaveConfig();
 				});
 
-		Section.AddEntry(FToolMenuEntry::InitWidget(TEXT("ToolSize"), GraphSizeWidget, LOCTEXT("ToolSizeLabel", "Tool Size")));
+		Section.AddEntry(FToolMenuEntry::InitWidget(TEXT("ToolSize"), GraphSizeWidget, LOCTEXT("ToolSizeLabel", "Tool Size"), true));
 	}
 
 	// Grid Size
@@ -345,7 +345,7 @@ void FAvaEaseCurveToolContextMenu::PopulateContextMenuSettings(UToolMenu* const
 					SetEaseCurveToolGridSize(InNewValue);
 				});
 
-		Section.AddEntry(FToolMenuEntry::InitWidget(TEXT("GridSize"), GridSizeWidget, LOCTEXT("GridSizeLabel", "Grid Size")));
+		Section.AddEntry(FToolMenuEntry::InitWidget(TEXT("GridSize"), GridSizeWidget, LOCTEXT("GridSizeLabel", "Grid Size"), true));
 	}
 
 	Section.AddSeparator(NAME_None);
