[Motion Design] Fix cooking error: don't generate an error if the broadcast config doesn't exist because it hasn't been saved.

This will just lead to an empty broadcast config being loaded in the game. But it can be configured through the rundown server and saved.

diff --git a/Engine/Plugins/Experimental/Avalanche/Source/AvalancheMedia/Private/Broadcast/AvaBroadcastSerialization.cpp b/Engine/Plugins/Experimental/Avalanche/Source/AvalancheMedia/Private/Broadcast/AvaBroadcastSerialization.cpp
--- a/Engine/Plugins/Experimental/Avalanche/Source/AvalancheMedia/Private/Broadcast/AvaBroadcastSerialization.cpp
+++ b/Engine/Plugins/Experimental/Avalanche/Source/AvalancheMedia/Private/Broadcast/AvaBroadcastSerialization.cpp
@@ -8,6 +8,7 @@
 #include "Formatters/XmlArchiveOutputFormatter.h"
 #include "JsonObjectConverter.h"
 #include "Misc/FileHelper.h"
+#include "Misc/Paths.h"
 
 namespace UE::AvaMedia::BroadcastSerialization::Private
 {
@@ -178,6 +179,14 @@ bool FAvaBroadcastSerialization::SaveBroadcastToJson(const UAvaBroadcast* InBroa
 
 bool FAvaBroadcastSerialization::LoadBroadcastFromJson(const FString& InFilename, UAvaBroadcast* OutBroadcast)
 {
+	// A missing file will not be considered an error as part of initializing a broadcast object
+	// since it may not have been saved yet. However, we log it to help troubleshoot potential issues.
+	if (!FPaths::FileExists(InFilename))
+	{
+		UE_LOG(LogAvaBroadcast, Log, TEXT("Json Configuration file \"%s\" not found."), *InFilename);
+		return false;
+	}
+
 	FString JsonText;
 
 	// Load json text to the string object
