Fixed constraints on FAvaMediaEditorModule::RegisterRundownFilter*Factory() definitions - which don't compile under C++17 - and switched to std:: traits.

diff --git a/Engine/Plugins/Experimental/Avalanche/Source/AvalancheMediaEditor/Private/AvaMediaEditorModule.cpp b/Engine/Plugins/Experimental/Avalanche/Source/AvalancheMediaEditor/Private/AvaMediaEditorModule.cpp
--- a/Engine/Plugins/Experimental/Avalanche/Source/AvalancheMediaEditor/Private/AvaMediaEditorModule.cpp
+++ b/Engine/Plugins/Experimental/Avalanche/Source/AvalancheMediaEditor/Private/AvaMediaEditorModule.cpp
@@ -328,7 +328,7 @@ void FAvaMediaEditorModule::HandleMapChanged(UWorld* InWorld, EMapChangeType InM
 template <
 	typename InRundownFilterExpressionFactoryType,
 	typename ... InArgsType
-	UE_REQUIRES(TIsDerivedFrom<InRundownFilterExpressionFactoryType, IAvaRundownFilterExpressionFactory>::Value)
+	UE_REQUIRES_DEFINITION(std::is_base_of_v<IAvaRundownFilterExpressionFactory, InRundownFilterExpressionFactoryType>)
 >
 void FAvaMediaEditorModule::RegisterRundownFilterExpressionFactory(InArgsType&&... InArgs)
 {
@@ -344,7 +344,7 @@ void FAvaMediaEditorModule::RegisterRundownFilterExpressionFactory(InArgsType&&.
 template <
 	typename InRundownSuggestionFactoryType,
 	typename ... InArgsType
-	UE_REQUIRES(TIsDerivedFrom<InRundownSuggestionFactoryType, IAvaRundownFilterSuggestionFactory>::Value)
+	UE_REQUIRES_DEFINITION(std::is_base_of_v<IAvaRundownFilterSuggestionFactory, InRundownSuggestionFactoryType>)
 >
 void FAvaMediaEditorModule::RegisterRundownFilterSuggestionFactory(InArgsType&&... InArgs)
 {
diff --git a/Engine/Plugins/Experimental/Avalanche/Source/AvalancheMediaEditor/Private/AvaMediaEditorModule.h b/Engine/Plugins/Experimental/Avalanche/Source/AvalancheMediaEditor/Private/AvaMediaEditorModule.h
--- a/Engine/Plugins/Experimental/Avalanche/Source/AvalancheMediaEditor/Private/AvaMediaEditorModule.h
+++ b/Engine/Plugins/Experimental/Avalanche/Source/AvalancheMediaEditor/Private/AvaMediaEditorModule.h
@@ -5,6 +5,8 @@
 #include "IAvaMediaEditorModule.h"
 #include "Templates/UnrealTypeTraits.h"
 
+#include <type_traits>
+
 class IAvaRundownFilterExpressionFactory;
 class IAvaRundownFilterSuggestionFactory;
 enum class EAvaRundownSearchListType : uint8;
@@ -54,14 +56,14 @@ private:
 	template <
 		typename InRundownFilterExpressionFactoryType,
 		typename... InArgsType
-		UE_REQUIRES(TIsDerivedFrom<InRundownFilterExpressionFactoryType, IAvaRundownFilterExpressionFactory>::Value)
+		UE_REQUIRES(std::is_base_of_v<IAvaRundownFilterExpressionFactory, InRundownFilterExpressionFactoryType>)
 	>
 	void RegisterRundownFilterExpressionFactory(InArgsType&&... InArgs);
 
 	template <
 		typename InRundownSuggestionFactoryType,
 		typename... InArgsType
-		UE_REQUIRES(TIsDerivedFrom<InRundownSuggestionFactoryType, IAvaRundownFilterSuggestionFactory>::Value)
+		UE_REQUIRES(std::is_base_of_v<IAvaRundownFilterSuggestionFactory, InRundownSuggestionFactoryType>)
 	>
 	void RegisterRundownFilterSuggestionFactory(InArgsType&&... InArgs);
 
