Motion Design
- Escape now works as expected when a Motion Design ITF tool is active and the viewport is not focused.

diff --git a/Engine/Plugins/Experimental/Avalanche/Source/AvalancheInteractiveTools/Private/AvaInteractiveToolsCommands.cpp b/Engine/Plugins/Experimental/Avalanche/Source/AvalancheInteractiveTools/Private/AvaInteractiveToolsCommands.cpp
--- a/Engine/Plugins/Experimental/Avalanche/Source/AvalancheInteractiveTools/Private/AvaInteractiveToolsCommands.cpp
+++ b/Engine/Plugins/Experimental/Avalanche/Source/AvalancheInteractiveTools/Private/AvaInteractiveToolsCommands.cpp
@@ -17,6 +17,7 @@ FAvaInteractiveToolsCommands::FAvaInteractiveToolsCommands()
 
 void FAvaInteractiveToolsCommands::RegisterCommands()
 {
+	RegisterEdModeCommands();
 	RegisterCategoryCommands();
 	Register2DCommands();
 	Register3DCommands();
@@ -24,6 +25,16 @@ void FAvaInteractiveToolsCommands::RegisterCommands()
 	RegisterLayoutCommands();
 }
 
+void FAvaInteractiveToolsCommands::RegisterEdModeCommands()
+{
+	UI_COMMAND(CancelActiveTool
+		, "Cancel"
+		, "Cancel the active Tool"
+		, EUserInterfaceActionType::Button
+		, FInputChord(EKeys::Escape)
+	);
+}
+
 void FAvaInteractiveToolsCommands::RegisterCategoryCommands()
 {
 	UI_COMMAND(Category_2D
diff --git a/Engine/Plugins/Experimental/Avalanche/Source/AvalancheInteractiveTools/Private/AvaInteractiveToolsCommands.h b/Engine/Plugins/Experimental/Avalanche/Source/AvalancheInteractiveTools/Private/AvaInteractiveToolsCommands.h
--- a/Engine/Plugins/Experimental/Avalanche/Source/AvalancheInteractiveTools/Private/AvaInteractiveToolsCommands.h
+++ b/Engine/Plugins/Experimental/Avalanche/Source/AvalancheInteractiveTools/Private/AvaInteractiveToolsCommands.h
@@ -13,12 +13,16 @@ public:
 	virtual void RegisterCommands() override;
 	//~ End TCommands
 
+	void RegisterEdModeCommands();
 	void RegisterCategoryCommands();
 	void Register2DCommands();
 	void Register3DCommands();
 	void RegisterActorCommands();
 	void RegisterLayoutCommands();
 
+	// Ed Mode
+	TSharedPtr<FUICommandInfo> CancelActiveTool;
+
 	// Categories
 	TSharedPtr<FUICommandInfo> Category_2D;
 	TSharedPtr<FUICommandInfo> Category_3D;
diff --git a/Engine/Plugins/Experimental/Avalanche/Source/AvalancheInteractiveTools/Private/EdMode/AvaInteractiveToolsEdMode.cpp b/Engine/Plugins/Experimental/Avalanche/Source/AvalancheInteractiveTools/Private/EdMode/AvaInteractiveToolsEdMode.cpp
--- a/Engine/Plugins/Experimental/Avalanche/Source/AvalancheInteractiveTools/Private/EdMode/AvaInteractiveToolsEdMode.cpp
+++ b/Engine/Plugins/Experimental/Avalanche/Source/AvalancheInteractiveTools/Private/EdMode/AvaInteractiveToolsEdMode.cpp
@@ -1,12 +1,14 @@
 // Copyright Epic Games, Inc. All Rights Reserved.
 
 #include "AvaInteractiveToolsEdMode.h"
+
+#include "AvaInteractiveToolsCommands.h"
 #include "AvaInteractiveToolsEdModeToolkit.h"
 #include "AvaInteractiveToolsStyle.h"
 #include "ContextObjectStore.h"
-#include "EdModeInteractiveToolsContext.h"
 #include "EditorModeManager.h"
 #include "EditorModes.h"
+#include "EdModeInteractiveToolsContext.h"
 #include "IAvalancheInteractiveToolsModule.h"
 #include "Selection.h"
 #include "Tools/AvaInteractiveToolsToolBase.h"
@@ -42,6 +44,26 @@ bool UAvaInteractiveToolsEdMode::IsCompatibleWith(FEditorModeID OtherModeID) con
 	return !UE::AvaInteractiveToolsEditorMode::Private::IncompatibleEdModes.Contains(OtherModeID);
 }
 
+void UAvaInteractiveToolsEdMode::BindCommands()
+{
+	Super::BindCommands();
+
+	const FAvaInteractiveToolsCommands& ToolManagerCommands = FAvaInteractiveToolsCommands::Get();
+	const TSharedRef<FUICommandList>& CommandList = Toolkit->GetToolkitCommands();
+
+	CommandList->MapAction(
+		ToolManagerCommands.CancelActiveTool,
+		FExecuteAction::CreateWeakLambda(this, [this]()
+			{
+				GetInteractiveToolsContext()->EndTool(EToolShutdownType::Cancel);
+			}),
+		FCanExecuteAction(),
+		FGetActionCheckState(),
+		FIsActionButtonVisible(),
+		EUIActionRepeatMode::RepeatDisabled
+	);
+}
+
 void UAvaInteractiveToolsEdMode::Enter()
 {
 	UEdMode::Enter();
diff --git a/Engine/Plugins/Experimental/Avalanche/Source/AvalancheInteractiveTools/Private/EdMode/AvaInteractiveToolsEdMode.h b/Engine/Plugins/Experimental/Avalanche/Source/AvalancheInteractiveTools/Private/EdMode/AvaInteractiveToolsEdMode.h
--- a/Engine/Plugins/Experimental/Avalanche/Source/AvalancheInteractiveTools/Private/EdMode/AvaInteractiveToolsEdMode.h
+++ b/Engine/Plugins/Experimental/Avalanche/Source/AvalancheInteractiveTools/Private/EdMode/AvaInteractiveToolsEdMode.h
@@ -20,6 +20,7 @@ public:
 
 	//~ Begin UEdMode interface
 	virtual bool IsCompatibleWith(FEditorModeID OtherModeID) const override;
+	virtual void BindCommands() override;
 	virtual void Enter() override;
 	virtual void CreateToolkit() override;
 	virtual TMap<FName, TArray<TSharedPtr<FUICommandInfo>>> GetModeCommands() const override;
@@ -36,7 +37,7 @@ public:
 protected:
 	FString LastActiveTool;
 	TWeakObjectPtr<UTypedElementSelectionSet> WeakActorSelectionSet;
-	double LastToolActivateTime;
+	double LastToolActivateTime = -1.0;
 
 	void OnActorSelectionChange(const UTypedElementSelectionSet* InSelectionSet);
 };
