Motion Design Modifiers:
- Fix translucency priority crash when importing a 5.4 scene and playing it in 5.5.

diff --git a/Engine/Plugins/Experimental/Avalanche/Source/AvalancheModifiers/Private/Modifiers/AvaTranslucentPriorityModifier.cpp b/Engine/Plugins/Experimental/Avalanche/Source/AvalancheModifiers/Private/Modifiers/AvaTranslucentPriorityModifier.cpp
--- a/Engine/Plugins/Experimental/Avalanche/Source/AvalancheModifiers/Private/Modifiers/AvaTranslucentPriorityModifier.cpp
+++ b/Engine/Plugins/Experimental/Avalanche/Source/AvalancheModifiers/Private/Modifiers/AvaTranslucentPriorityModifier.cpp
@@ -176,11 +176,11 @@ void UAvaTranslucentPriorityModifier::Apply()
 		int32 TranslucentSortPriority = GlobalOffset + SortPriority;
 
 		// Sets all components with the same priority
-		for (const FAvaTranslucentPriorityModifierComponentState* SortedComponentState : CachedSortedComponentStates)
+		for (const FAvaTranslucentPriorityModifierComponentState& SortedComponentState : CachedSortedComponentStates)
 		{
-			if (SortedComponentState && SortedComponentState->ModifierWeak == this)
+			if (SortedComponentState.ModifierWeak == this)
 			{
-				if (UPrimitiveComponent* Component = SortedComponentState->PrimitiveComponentWeak.Get())
+				if (UPrimitiveComponent* Component = SortedComponentState.PrimitiveComponentWeak.Get())
 				{
 					LastSortPriorities.Add(Component, TranslucentSortPriority);
 					Component->SetTranslucentSortPriority(TranslucentSortPriority);
@@ -195,15 +195,10 @@ void UAvaTranslucentPriorityModifier::Apply()
 		int32 TranslucentSortPriority = GlobalOffset;
 
 		// Increment sort priority for each component that this modifier handles
-		for (const FAvaTranslucentPriorityModifierComponentState* SortedComponentState : CachedSortedComponentStates)
+		for (const FAvaTranslucentPriorityModifierComponentState& SortedComponentState : CachedSortedComponentStates)
 		{
-			if (!SortedComponentState)
-			{
-				continue;
-			}
-
-			UPrimitiveComponent* Component = SortedComponentState->PrimitiveComponentWeak.Get();
-			UAvaTranslucentPriorityModifier* ComponentModifier = SortedComponentState->ModifierWeak.Get();
+			UPrimitiveComponent* Component = SortedComponentState.PrimitiveComponentWeak.Get();
+			UAvaTranslucentPriorityModifier* ComponentModifier = SortedComponentState.ModifierWeak.Get();
 
 			if (!Component || !ComponentModifier)
 			{
diff --git a/Engine/Plugins/Experimental/Avalanche/Source/AvalancheModifiers/Private/Shared/AvaTranslucentPriorityModifierShared.cpp b/Engine/Plugins/Experimental/Avalanche/Source/AvalancheModifiers/Private/Shared/AvaTranslucentPriorityModifierShared.cpp
--- a/Engine/Plugins/Experimental/Avalanche/Source/AvalancheModifiers/Private/Shared/AvaTranslucentPriorityModifierShared.cpp
+++ b/Engine/Plugins/Experimental/Avalanche/Source/AvalancheModifiers/Private/Shared/AvaTranslucentPriorityModifierShared.cpp
@@ -142,9 +142,9 @@ UAvaTranslucentPriorityModifier* UAvaTranslucentPriorityModifierShared::FindModi
 	return ModifierContext;
 }
 
-TArray<const FAvaTranslucentPriorityModifierComponentState*> UAvaTranslucentPriorityModifierShared::GetSortedComponentStates(UAvaTranslucentPriorityModifier* InModifierContext) const
+TArray<FAvaTranslucentPriorityModifierComponentState> UAvaTranslucentPriorityModifierShared::GetSortedComponentStates(UAvaTranslucentPriorityModifier* InModifierContext) const
 {
-	TArray<const FAvaTranslucentPriorityModifierComponentState*> SortedComponentStates;
+	TArray<FAvaTranslucentPriorityModifierComponentState> SortedComponentStates;
 
 	if (!InModifierContext)
 	{
@@ -168,7 +168,7 @@ TArray<const FAvaTranslucentPriorityModifierComponentState*> UAvaTranslucentPrio
 
 			if (ModifierContext == InModifierContext)
 			{
-				SortedComponentStates.Add(&ComponentState);
+				SortedComponentStates.Add(ComponentState);
 			}
 		}
 	}
@@ -195,7 +195,7 @@ TArray<const FAvaTranslucentPriorityModifierComponentState*> UAvaTranslucentPrio
 			if (ModifierContext->GetMode() == InModifierContext->GetMode()
 				&& ModifierContext->CameraActorWeak == CameraActor)
 			{
-				SortedComponentStates.Add(&ComponentState);
+				SortedComponentStates.Add(ComponentState);
 			}
 		}
 
@@ -235,7 +235,7 @@ TArray<const FAvaTranslucentPriorityModifierComponentState*> UAvaTranslucentPrio
 			if (ModifierContext->GetMode() == InModifierContext->GetMode()
 				&& OwningActor->GetTypedOuter<UWorld>() == World)
 			{
-				SortedComponentStates.Add(&ComponentState);
+				SortedComponentStates.Add(ComponentState);
 			}
 		}
 
diff --git a/Engine/Plugins/Experimental/Avalanche/Source/AvalancheModifiers/Public/Modifiers/AvaTranslucentPriorityModifier.h b/Engine/Plugins/Experimental/Avalanche/Source/AvalancheModifiers/Public/Modifiers/AvaTranslucentPriorityModifier.h
--- a/Engine/Plugins/Experimental/Avalanche/Source/AvalancheModifiers/Public/Modifiers/AvaTranslucentPriorityModifier.h
+++ b/Engine/Plugins/Experimental/Avalanche/Source/AvalancheModifiers/Public/Modifiers/AvaTranslucentPriorityModifier.h
@@ -145,5 +145,5 @@ private:
 	TMap<FObjectKey, int32> LastSortPriorities;
 
 	/** Used to avoid querying again the full list of component states */
-	TArray<const FAvaTranslucentPriorityModifierComponentState*> CachedSortedComponentStates;
+	TArray<FAvaTranslucentPriorityModifierComponentState> CachedSortedComponentStates;
 };
diff --git a/Engine/Plugins/Experimental/Avalanche/Source/AvalancheModifiers/Public/Shared/AvaTranslucentPriorityModifierShared.h b/Engine/Plugins/Experimental/Avalanche/Source/AvalancheModifiers/Public/Shared/AvaTranslucentPriorityModifierShared.h
--- a/Engine/Plugins/Experimental/Avalanche/Source/AvalancheModifiers/Public/Shared/AvaTranslucentPriorityModifierShared.h
+++ b/Engine/Plugins/Experimental/Avalanche/Source/AvalancheModifiers/Public/Shared/AvaTranslucentPriorityModifierShared.h
@@ -89,7 +89,7 @@ public:
 	UAvaTranslucentPriorityModifier* FindModifierContext(UPrimitiveComponent* InComponent) const;
 
 	/** Get sorted components state based on modifier context */
-	TArray<const FAvaTranslucentPriorityModifierComponentState*> GetSortedComponentStates(UAvaTranslucentPriorityModifier* InModifierContext) const;
+	TArray<FAvaTranslucentPriorityModifierComponentState> GetSortedComponentStates(UAvaTranslucentPriorityModifier* InModifierContext) const;
 
 	void SetSortPriorityOffset(int32 InOffset);
 
