Motion Design: fix non-unity

diff --git a/Engine/Plugins/Experimental/Avalanche/Source/AvalancheModifiers/Private/Shared/AvaTranslucentPriorityModifierShared.cpp b/Engine/Plugins/Experimental/Avalanche/Source/AvalancheModifiers/Private/Shared/AvaTranslucentPriorityModifierShared.cpp
--- a/Engine/Plugins/Experimental/Avalanche/Source/AvalancheModifiers/Private/Shared/AvaTranslucentPriorityModifierShared.cpp
+++ b/Engine/Plugins/Experimental/Avalanche/Source/AvalancheModifiers/Private/Shared/AvaTranslucentPriorityModifierShared.cpp
@@ -9,6 +9,7 @@
 #include "Camera/CameraComponent.h"
 #include "Framework/AvaGameInstance.h"
 #include "IAvaSceneInterface.h"
+#include "Modifiers/AvaTranslucentPriorityModifier.h"
 
 #if WITH_EDITOR
 #include "AvaOutlinerUtils.h"
diff --git a/Engine/Plugins/Experimental/Avalanche/Source/AvalancheModifiers/Public/Modifiers/AvaTranslucentPriorityModifier.h b/Engine/Plugins/Experimental/Avalanche/Source/AvalancheModifiers/Public/Modifiers/AvaTranslucentPriorityModifier.h
--- a/Engine/Plugins/Experimental/Avalanche/Source/AvalancheModifiers/Public/Modifiers/AvaTranslucentPriorityModifier.h
+++ b/Engine/Plugins/Experimental/Avalanche/Source/AvalancheModifiers/Public/Modifiers/AvaTranslucentPriorityModifier.h
@@ -4,9 +4,9 @@
 
 #include "AvaArrangeBaseModifier.h"
 #include "AvaBaseModifier.h"
+#include "Shared/AvaTranslucentPriorityModifierShared.h"
 #include "AvaTranslucentPriorityModifier.generated.h"
 
-struct FAvaTranslucentPriorityModifierComponentState;
 class ACameraActor;
 class UPrimitiveComponent;
 
diff --git a/Engine/Plugins/Experimental/Avalanche/Source/AvalancheModifiers/Public/Shared/AvaTranslucentPriorityModifierShared.h b/Engine/Plugins/Experimental/Avalanche/Source/AvalancheModifiers/Public/Shared/AvaTranslucentPriorityModifierShared.h
--- a/Engine/Plugins/Experimental/Avalanche/Source/AvalancheModifiers/Public/Shared/AvaTranslucentPriorityModifierShared.h
+++ b/Engine/Plugins/Experimental/Avalanche/Source/AvalancheModifiers/Public/Shared/AvaTranslucentPriorityModifierShared.h
@@ -5,13 +5,13 @@
 #include "Components/PrimitiveComponent.h"
 #include "Containers/Set.h"
 #include "Modifiers/ActorModifierCoreSharedObject.h"
-#include "Modifiers/AvaTranslucentPriorityModifier.h"
 #include "Templates/SharedPointer.h"
 #include "UObject/WeakObjectPtr.h"
 #include "AvaTranslucentPriorityModifierShared.generated.h"
 
 class AActor;
 class ACameraActor;
+class UAvaTranslucentPriorityModifier;
 
 USTRUCT()
 struct FAvaTranslucentPriorityModifierComponentState
