MotionDesign : Geometry Modifiers
- Adding more profiling stat (vertices, edges)
- Renaming VertexIn/VertexOut by TriIn/TriOut

diff --git a/Engine/Plugins/Experimental/Avalanche/Source/AvalancheModifiers/Private/Profilers/AvaGeometryModifierProfiler.cpp b/Engine/Plugins/Experimental/Avalanche/Source/AvalancheModifiers/Private/Profilers/AvaGeometryModifierProfiler.cpp
--- a/Engine/Plugins/Experimental/Avalanche/Source/AvalancheModifiers/Private/Profilers/AvaGeometryModifierProfiler.cpp
+++ b/Engine/Plugins/Experimental/Avalanche/Source/AvalancheModifiers/Private/Profilers/AvaGeometryModifierProfiler.cpp
@@ -11,6 +11,10 @@ void FAvaGeometryModifierProfiler::SetupProfilingStats()
 
 	ProfilerStats.AddProperty(VertexInName, EPropertyBagPropertyType::Int32);
 	ProfilerStats.AddProperty(VertexOutName, EPropertyBagPropertyType::Int32);
+	ProfilerStats.AddProperty(TriangleInName, EPropertyBagPropertyType::Int32);
+	ProfilerStats.AddProperty(TriangleOutName, EPropertyBagPropertyType::Int32);
+	ProfilerStats.AddProperty(EdgeInName, EPropertyBagPropertyType::Int32);
+	ProfilerStats.AddProperty(EdgeOutName, EPropertyBagPropertyType::Int32);
 }
 
 void FAvaGeometryModifierProfiler::BeginProfiling()
@@ -18,10 +22,15 @@ void FAvaGeometryModifierProfiler::BeginProfiling()
 	FActorModifierCoreProfiler::BeginProfiling();
 
 	const UAvaGeometryBaseModifier* GeometryModifier = GetModifier<UAvaGeometryBaseModifier>();
+
 	if (GeometryModifier && GeometryModifier->IsMeshValid())
 	{
-		const int32 TriangleCount = GeometryModifier->GetMeshComponent()->GetDynamicMesh()->GetTriangleCount();
-		ProfilerStats.SetValueInt32(VertexInName, TriangleCount);
+		GeometryModifier->GetMeshComponent()->ProcessMesh([this](const FDynamicMesh3& InMesh)
+		{
+			ProfilerStats.SetValueInt32(VertexInName, InMesh.VertexCount());
+			ProfilerStats.SetValueInt32(TriangleInName, InMesh.TriangleCount());
+			ProfilerStats.SetValueInt32(EdgeInName, InMesh.EdgeCount());
+		});
 	}
 }
 
@@ -30,10 +39,15 @@ void FAvaGeometryModifierProfiler::EndProfiling()
 	FActorModifierCoreProfiler::EndProfiling();
 
 	const UAvaGeometryBaseModifier* GeometryModifier = GetModifier<UAvaGeometryBaseModifier>();
+
 	if (GeometryModifier && GeometryModifier->IsMeshValid())
 	{
-		const int32 TriangleCount = GeometryModifier->GetMeshComponent()->GetDynamicMesh()->GetTriangleCount();
-		ProfilerStats.SetValueInt32(VertexOutName, TriangleCount);
+		GeometryModifier->GetMeshComponent()->ProcessMesh([this](const FDynamicMesh3& InMesh)
+		{
+			ProfilerStats.SetValueInt32(VertexOutName, InMesh.VertexCount());
+			ProfilerStats.SetValueInt32(TriangleOutName, InMesh.TriangleCount());
+			ProfilerStats.SetValueInt32(EdgeOutName, InMesh.EdgeCount());
+		});
 	}
 }
 
@@ -42,8 +56,8 @@ TSet<FName> FAvaGeometryModifierProfiler::GetMainProfilingStats() const
 	return TSet<FName>
 	{
 		ExecutionTimeName,
-		VertexInName,
-		VertexOutName
+		TriangleInName,
+		TriangleOutName
 	};
 }
 
@@ -58,3 +72,27 @@ int32 FAvaGeometryModifierProfiler::GetVertexOut() const
 	TValueOrError<int32, EPropertyBagResult> ValueResult = ProfilerStats.GetValueInt32(VertexOutName);
 	return ValueResult.HasValue() ? ValueResult.GetValue() : 0;
 }
+
+int32 FAvaGeometryModifierProfiler::GetTriangleIn() const
+{
+	TValueOrError<int32, EPropertyBagResult> ValueResult = ProfilerStats.GetValueInt32(TriangleInName);
+	return ValueResult.HasValue() ? ValueResult.GetValue() : 0;
+}
+
+int32 FAvaGeometryModifierProfiler::GetTriangleOut() const
+{
+	TValueOrError<int32, EPropertyBagResult> ValueResult = ProfilerStats.GetValueInt32(TriangleOutName);
+	return ValueResult.HasValue() ? ValueResult.GetValue() : 0;
+}
+
+int32 FAvaGeometryModifierProfiler::GetEdgeIn() const
+{
+	TValueOrError<int32, EPropertyBagResult> ValueResult = ProfilerStats.GetValueInt32(EdgeInName);
+	return ValueResult.HasValue() ? ValueResult.GetValue() : 0;
+}
+
+int32 FAvaGeometryModifierProfiler::GetEdgeOut() const
+{
+	TValueOrError<int32, EPropertyBagResult> ValueResult = ProfilerStats.GetValueInt32(EdgeOutName);
+	return ValueResult.HasValue() ? ValueResult.GetValue() : 0;
+}
diff --git a/Engine/Plugins/Experimental/Avalanche/Source/AvalancheModifiers/Public/Profilers/AvaGeometryModifierProfiler.h b/Engine/Plugins/Experimental/Avalanche/Source/AvalancheModifiers/Public/Profilers/AvaGeometryModifierProfiler.h
--- a/Engine/Plugins/Experimental/Avalanche/Source/AvalancheModifiers/Public/Profilers/AvaGeometryModifierProfiler.h
+++ b/Engine/Plugins/Experimental/Avalanche/Source/AvalancheModifiers/Public/Profilers/AvaGeometryModifierProfiler.h
@@ -10,6 +10,10 @@ class FAvaGeometryModifierProfiler : public FActorModifierCoreProfiler
 public:
 	static inline const FName VertexInName = TEXT("VertexIn");
 	static inline const FName VertexOutName = TEXT("VertexOut");
+	static inline const FName TriangleInName = TEXT("TriIn");
+	static inline const FName TriangleOutName = TEXT("TriOut");
+	static inline const FName EdgeInName = TEXT("EdgeIn");
+	static inline const FName EdgeOutName = TEXT("EdgeOut");
 
 	//~ Begin FActorModifierCoreProfiler
 	virtual void SetupProfilingStats() override;
@@ -20,4 +24,8 @@ public:
 
 	AVALANCHEMODIFIERS_API int32 GetVertexIn() const;
 	AVALANCHEMODIFIERS_API int32 GetVertexOut() const;
+	AVALANCHEMODIFIERS_API int32 GetTriangleIn() const;
+	AVALANCHEMODIFIERS_API int32 GetTriangleOut() const;
+	AVALANCHEMODIFIERS_API int32 GetEdgeIn() const;
+	AVALANCHEMODIFIERS_API int32 GetEdgeOut() const;
 };
