[MotionDesign] Ease Curve Tool
 - Fix settings menu actions not bound
 - Fix flatten/straighten tangents not propagating update
 - Remove layout HAlign setting
 - Make sure key cache is always up to date, so the "no ease curve detected" message gets displayed correctly

diff --git a/Engine/Plugins/Experimental/Avalanche/Source/AvalancheSequencer/Private/DetailsView/SAvaKeyFrameEdit.cpp b/Engine/Plugins/Experimental/Avalanche/Source/AvalancheSequencer/Private/DetailsView/SAvaKeyFrameEdit.cpp
--- a/Engine/Plugins/Experimental/Avalanche/Source/AvalancheSequencer/Private/DetailsView/SAvaKeyFrameEdit.cpp
+++ b/Engine/Plugins/Experimental/Avalanche/Source/AvalancheSequencer/Private/DetailsView/SAvaKeyFrameEdit.cpp
@@ -3,62 +3,15 @@
 #include "SAvaKeyFrameEdit.h"
 #include "AvaSequencer.h"
 #include "EaseCurveTool/AvaEaseCurveTool.h"
-#include "EaseCurveTool/AvaEaseCurveToolSettings.h"
-#include "EaseCurveTool/Widgets/SAvaEaseCurveTool.h"
 #include "MVVM/Selection/Selection.h"
 #include "MVVM/ViewModels/SequencerEditorViewModel.h"
+#include "Widgets/Layout/SScaleBox.h"
 #include "Widgets/Layout/SScrollBox.h"
 
 using namespace UE::Sequencer;
 
 #define LOCTEXT_NAMESPACE "SAvaKeyFrameEdit"
 
-namespace UE::AvaSequencer::Private
-{
-	/** Utility widget to auto update horizontal alignment. */
-	class SDynamicHAlign : public SCompoundWidget
-	{
-	public:
-		SLATE_BEGIN_ARGS(SDynamicHAlign)
-			: _HAlign(HAlign_Fill)
-		{}
-			SLATE_ATTRIBUTE(EHorizontalAlignment, HAlign)
-			SLATE_ARGUMENT(TSharedPtr<SWidget>, Content)
-		SLATE_END_ARGS()
-
-		void Construct(const FArguments& InArgs)
-		{
-			SetCanTick(true);
-
-			HAlign = InArgs._HAlign;
-
-			if (InArgs._Content.IsValid())
-			{
-				ChildSlot
-					.HAlign(HAlign.Get())
-					[
-						InArgs._Content.ToSharedRef()
-					];
-			}
-		}
-
-	protected:
-		//~ Begin SWidget
-		virtual void Tick(const FGeometry& InAllottedGeometry, const double InCurrentTime, const float InDeltaTime) override
-		{
-			SCompoundWidget::Tick(InAllottedGeometry, InCurrentTime, InDeltaTime);
-
-			if (ChildSlot.GetHorizontalAlignment() != HAlign.Get())
-			{
-				ChildSlot.SetHorizontalAlignment(HAlign.Get());
-			}
-		}
-		//~ End SWidget
-
-		TAttribute<EHorizontalAlignment> HAlign;
-	};
-}
-
 void SAvaKeyFrameEdit::Construct(const FArguments& InArgs, const TSharedRef<FAvaSequencer>& InSequencer)
 {
 	AvaSequencerWeak = InSequencer;
@@ -85,14 +38,13 @@ void SAvaKeyFrameEdit::Construct(const FArguments& InArgs, const TSharedRef<FAva
 			+ SScrollBox::Slot()
 			.AutoSize()
 			.HAlign(HAlign_Fill)
-			.Padding(1.f, 2.f, 2.f, 3.f)
+			.Padding(1.f, 1.f, 1.f, 3.f)
 			[
-				/** Horizontal alignment is polled every frame from settings.
-				Fill will stretch the ease curve graph editor horizontally in the details panel.
-				Left/Right/Center will maintain 1:1 ratio. */
-				SNew(UE::AvaSequencer::Private::SDynamicHAlign)
-				.HAlign_UObject(GetDefault<UAvaEaseCurveToolSettings>(), &UAvaEaseCurveToolSettings::GetGraphHAlign)
-				.Content(InSequencer->GetEaseCurveTool()->GenerateWidget())
+				SNew(SScaleBox)
+				.Stretch(EStretch::Fill)
+				[
+					InSequencer->GetEaseCurveTool()->GenerateWidget()
+				]
 			]
 			+ SScrollBox::Slot()
 			.AutoSize()
diff --git a/Engine/Plugins/Experimental/Avalanche/Source/AvalancheSequencer/Private/EaseCurveTool/AvaEaseCurveTool.cpp b/Engine/Plugins/Experimental/Avalanche/Source/AvalancheSequencer/Private/EaseCurveTool/AvaEaseCurveTool.cpp
--- a/Engine/Plugins/Experimental/Avalanche/Source/AvalancheSequencer/Private/EaseCurveTool/AvaEaseCurveTool.cpp
+++ b/Engine/Plugins/Experimental/Avalanche/Source/AvalancheSequencer/Private/EaseCurveTool/AvaEaseCurveTool.cpp
@@ -16,6 +16,7 @@
 #include "ISequencer.h"
 #include "ISettingsModule.h"
 #include "Math/UnrealMathUtility.h"
+#include "ScopedTransaction.h"
 #include "Settings/AvaSequencerSettings.h"
 #include "Widgets/Notifications/SNotificationList.h"
 
@@ -112,15 +113,18 @@ void FAvaEaseCurveTool::SetEaseCurveTangents_Internal(const FAvaEaseCurveTangent
 	}
 }
 
-void FAvaEaseCurveTool::SetEaseCurveTangents(const FAvaEaseCurveTangents& InTangents, const EAvaEaseCurveToolOperation InOperation
-	, const bool bInBroadcastUpdate, const bool bInSetSequencerTangents)
+void FAvaEaseCurveTool::SetEaseCurveTangents(const FAvaEaseCurveTangents& InTangents
+	, const EAvaEaseCurveToolOperation InOperation
+	, const bool bInBroadcastUpdate
+	, const bool bInSetSequencerTangents
+	, const FText& InTransactionText)
 {
 	if (InTangents == GetEaseCurveTangents())
 	{
 		return;
 	}
 
-	const FScopedTransaction Transaction(LOCTEXT("SetEaseCurveTangents", "Set Ease Curve Tangents"), !GIsTransacting);
+	const FScopedTransaction Transaction(InTransactionText, !GIsTransacting);
 	EaseCurve->Modify();
 
 	SetEaseCurveTangents_Internal(InTangents, InOperation, bInBroadcastUpdate);
@@ -129,6 +133,8 @@ void FAvaEaseCurveTool::SetEaseCurveTangents(const FAvaEaseCurveTangents& InTang
 	{
 		SetSequencerKeySelectionTangents(InTangents, InOperation);
 	}
+
+	KeyCache = FAvaEaseCurveKeySelection(AvaSequencerWeak.Pin());
 }
 
 void FAvaEaseCurveTool::ResetEaseCurveTangents(const EAvaEaseCurveToolOperation InOperation)
@@ -152,11 +158,16 @@ void FAvaEaseCurveTool::ResetEaseCurveTangents(const EAvaEaseCurveToolOperation
 	EaseCurve->ModifyOwner();
 
 	const FAvaEaseCurveTangents ZeroTangents;
-	SetEaseCurveTangents_Internal(ZeroTangents, InOperation, /*bInBroadcastUpdate=*/true);
-	SetSequencerKeySelectionTangents(ZeroTangents, InOperation);
+	SetEaseCurveTangents(ZeroTangents, InOperation, /*bInBroadcastUpdate=*/true, /*bInSetSequencerTangents=*/true, TransactionText);
+
+	// Update ease curve combobox widget and zoom graph editor
+	if (ToolWidget.IsValid())
+	{
+		ToolWidget->SetTangents(ZeroTangents, InOperation, /*bInSetEaseCurve=*/false, /*bInBroadcastUpdate=*/false, /*bInSetSequencerTangents=*/false);
+	}
 }
 
-void FAvaEaseCurveTool::FlattenOrStraightenTangents(const EAvaEaseCurveToolOperation InOperation, const bool bInFlattenTangents) const
+void FAvaEaseCurveTool::FlattenOrStraightenTangents(const EAvaEaseCurveToolOperation InOperation, const bool bInFlattenTangents)
 {
 	FText TransactionText;
 	if (bInFlattenTangents)
@@ -201,7 +212,14 @@ void FAvaEaseCurveTool::FlattenOrStraightenTangents(const EAvaEaseCurveToolOpera
 		EaseCurve->FlattenOrStraightenTangents(EaseCurve->GetEndKeyHandle(), bInFlattenTangents);
 	}
 
-	EaseCurve->BroadcastUpdate();
+	const FAvaEaseCurveTangents NewTangents = EaseCurve->GetTangents();
+	SetEaseCurveTangents(NewTangents, InOperation, /*bInBroadcastUpdate=*/true, /*bInSetSequencerTangents=*/true, TransactionText);
+
+	// Update ease curve combobox widget and zoom graph editor
+	if (ToolWidget.IsValid())
+	{
+		ToolWidget->SetTangents(NewTangents, InOperation, /*bInSetEaseCurve=*/false, /*bInBroadcastUpdate=*/false, /*bInSetSequencerTangents=*/false);
+	}
 }
 
 void FAvaEaseCurveTool::ApplyQuickEaseToSequencerKeySelections(const EAvaEaseCurveToolOperation InOperation)
@@ -467,9 +485,11 @@ void FAvaEaseCurveTool::PostRedo(bool bInSuccess)
 
 void FAvaEaseCurveTool::OpenToolSettings() const
 {
-	if (ISettingsModule* SettingsModule = FModuleManager::GetModulePtr<ISettingsModule>(TEXT("Settings")))
+	if (ISettingsModule* const SettingsModule = FModuleManager::GetModulePtr<ISettingsModule>(TEXT("Settings")))
 	{
-		SettingsModule->ShowViewer(TEXT("Editor"), TEXT("Motion Design"), TEXT("Ease Curve Tool"));
+		const UAvaEaseCurveToolSettings* const Settings = GetDefault<UAvaEaseCurveToolSettings>();
+		check(IsValid(Settings));
+		SettingsModule->ShowViewer(Settings->GetContainerName(), Settings->GetCategoryName(), Settings->GetSectionName());
 	}
 }
 
diff --git a/Engine/Plugins/Experimental/Avalanche/Source/AvalancheSequencer/Private/EaseCurveTool/AvaEaseCurveTool.h b/Engine/Plugins/Experimental/Avalanche/Source/AvalancheSequencer/Private/EaseCurveTool/AvaEaseCurveTool.h
--- a/Engine/Plugins/Experimental/Avalanche/Source/AvalancheSequencer/Private/EaseCurveTool/AvaEaseCurveTool.h
+++ b/Engine/Plugins/Experimental/Avalanche/Source/AvalancheSequencer/Private/EaseCurveTool/AvaEaseCurveTool.h
@@ -8,7 +8,7 @@
 #include "Curves/KeyHandle.h"
 #include "EaseCurveTool/Widgets/SAvaEaseCurvePreset.h"
 #include "EditorUndoClient.h"
-#include "ScopedTransaction.h"
+#include "Internationalization/Text.h"
 #include "Templates/SharedPointer.h"
 #include "UObject/WeakObjectPtrTemplates.h"
 #include "Widgets/DeclarativeSyntaxSupport.h"
@@ -81,11 +81,15 @@ public:
 	 * This is different from SetEaseCurveTangents_Internal in that it performs undo/redo transactions and
 	 * optionally sets the selected tangents in the actual sequence.
 	 */
-	void SetEaseCurveTangents(const FAvaEaseCurveTangents& InTangents, const EAvaEaseCurveToolOperation InOperation, const bool bInBroadcastUpdate, const bool bInSetSequencerTangents);
+	void SetEaseCurveTangents(const FAvaEaseCurveTangents& InTangents
+		, const EAvaEaseCurveToolOperation InOperation
+		, const bool bInBroadcastUpdate
+		, const bool bInSetSequencerTangents
+		, const FText& InTransactionText = NSLOCTEXT("AvaEaseCurveTool", "SetEaseCurveTangents", "Set Ease Curve Tangents"));
 
 	void ResetEaseCurveTangents(const EAvaEaseCurveToolOperation InOperation);
 
-	void FlattenOrStraightenTangents(const EAvaEaseCurveToolOperation InOperation, const bool bInFlattenTangents) const;
+	void FlattenOrStraightenTangents(const EAvaEaseCurveToolOperation InOperation, const bool bInFlattenTangents);
 
 	/** Creates a new external float curve from the internet curve editor curve. */
 	UCurveBase* CreateCurveAsset() const;
diff --git a/Engine/Plugins/Experimental/Avalanche/Source/AvalancheSequencer/Private/EaseCurveTool/AvaEaseCurveToolSettings.cpp b/Engine/Plugins/Experimental/Avalanche/Source/AvalancheSequencer/Private/EaseCurveTool/AvaEaseCurveToolSettings.cpp
--- a/Engine/Plugins/Experimental/Avalanche/Source/AvalancheSequencer/Private/EaseCurveTool/AvaEaseCurveToolSettings.cpp
+++ b/Engine/Plugins/Experimental/Avalanche/Source/AvalancheSequencer/Private/EaseCurveTool/AvaEaseCurveToolSettings.cpp
@@ -1,7 +1,6 @@
 // Copyright Epic Games, Inc. All Rights Reserved.
 
 #include "AvaEaseCurveToolSettings.h"
-#include "Misc/Paths.h"
 
 UAvaEaseCurveToolSettings::UAvaEaseCurveToolSettings()
 {
diff --git a/Engine/Plugins/Experimental/Avalanche/Source/AvalancheSequencer/Private/EaseCurveTool/AvaEaseCurveToolSettings.h b/Engine/Plugins/Experimental/Avalanche/Source/AvalancheSequencer/Private/EaseCurveTool/AvaEaseCurveToolSettings.h
--- a/Engine/Plugins/Experimental/Avalanche/Source/AvalancheSequencer/Private/EaseCurveTool/AvaEaseCurveToolSettings.h
+++ b/Engine/Plugins/Experimental/Avalanche/Source/AvalancheSequencer/Private/EaseCurveTool/AvaEaseCurveToolSettings.h
@@ -3,7 +3,6 @@
 #pragma once
 
 #include "Containers/EnumAsByte.h"
-#include "EaseCurveTool/AvaEaseCurveTangents.h"
 #include "Engine/DeveloperSettings.h"
 #include "Types/SlateEnums.h"
 #include "UObject/Object.h"
@@ -20,9 +19,6 @@ public:
 	int32 GetGraphSize() const { return GraphSize; }
 	void SetGraphSize(const int32 InSize) { GraphSize = InSize; }
 
-	EHorizontalAlignment GetGraphHAlign() const { return GraphHAlign.GetValue(); }
-	void SetGraphHAlign(const EHorizontalAlignment InHAlign) { GraphHAlign = InHAlign; }
-
 	bool GetGridSnap() const { return bGridSnap; }
 	void SetGridSnap(const bool bInGridSnap) { bGridSnap = bInGridSnap; }
 	void ToggleGridSnap() { bGridSnap = !bGridSnap; }
@@ -49,10 +45,6 @@ private:
 	UPROPERTY(Config, EditAnywhere, Category = "EaseCurveTool", meta = (DisplayName = "Graph Size", UIMin = 64, ClampMin = 64, UIMax = 256, ClampMax = 256, Delta = 1))
 	int32 GraphSize = 140;
 
-	/** The horizontal alignment of the curve ease tool in the details panel. */
-	UPROPERTY(Config, EditAnywhere, Category = "EaseCurveTool", meta = (DisplayName = "Graph Horizontal Alignment"))
-	TEnumAsByte<EHorizontalAlignment> GraphHAlign = EHorizontalAlignment::HAlign_Fill;
-
 	/** If true, snaps tangents to grid. */
 	UPROPERTY(Config, EditAnywhere, Category = "EaseCurveTool", meta = (DisplayName = "Grid Snap"))
 	bool bGridSnap = false;
