[MotionDesign] Add missing include

diff --git a/Engine/Plugins/Experimental/Avalanche/Source/AvalancheSequencer/Private/EaseCurveTool/AvaEaseCurveKeySelection.cpp b/Engine/Plugins/Experimental/Avalanche/Source/AvalancheSequencer/Private/EaseCurveTool/AvaEaseCurveKeySelection.cpp
--- a/Engine/Plugins/Experimental/Avalanche/Source/AvalancheSequencer/Private/EaseCurveTool/AvaEaseCurveKeySelection.cpp
+++ b/Engine/Plugins/Experimental/Avalanche/Source/AvalancheSequencer/Private/EaseCurveTool/AvaEaseCurveKeySelection.cpp
@@ -8,6 +8,7 @@
 #include "EaseCurveTool/AvaEaseCurveTangents.h"
 #include "EaseCurveTool/AvaEaseCurveTool.h"
 #include "IKeyArea.h"
+#include "ScopedTransaction.h"
 #include "MVVM/ViewModels/ChannelModel.h"
 #include "MVVM/ViewModels/SequencerEditorViewModel.h"
 
