Motion Design: Upgrade conditions and tasks to use instance data for its properties

diff --git a/Engine/Plugins/Experimental/Avalanche/Source/Avalanche/Private/Transition/Conditions/AvaSceneContainsTagAttributeCondition.cpp b/Engine/Plugins/Experimental/Avalanche/Source/Avalanche/Private/Transition/Conditions/AvaSceneContainsTagAttributeCondition.cpp
--- a/Engine/Plugins/Experimental/Avalanche/Source/Avalanche/Private/Transition/Conditions/AvaSceneContainsTagAttributeCondition.cpp
+++ b/Engine/Plugins/Experimental/Avalanche/Source/Avalanche/Private/Transition/Conditions/AvaSceneContainsTagAttributeCondition.cpp
@@ -6,6 +6,7 @@
 #include "AvaTransitionLayer.h"
 #include "AvaTransitionLayerUtils.h"
 #include "AvaTransitionSubsystem.h"
+#include "AvaTransitionUtils.h"
 #include "Behavior/AvaTransitionBehaviorInstance.h"
 #include "IAvaSceneInterface.h"
 #include "StateTreeExecutionContext.h"
@@ -16,9 +17,11 @@
 #if WITH_EDITOR
 FText FAvaSceneContainsTagAttributeConditionBase::GetDescription(const FGuid& InId, FStateTreeDataView InInstanceDataView, const IStateTreeBindingLookup& InBindingLookup, EStateTreeNodeFormatting InFormatting) const
 {
+	const FInstanceDataType& InstanceData = InInstanceDataView.Get<FInstanceDataType>();
+
 	FFormatNamedArguments Arguments;
 
-	switch (SceneType)
+	switch (InstanceData.SceneType)
 	{
 	case EAvaTransitionSceneType::This:
 		Arguments.Add(TEXT("IndefinitePronoun"), FText::GetEmpty());
@@ -31,18 +34,36 @@ FText FAvaSceneContainsTagAttributeConditionBase::GetDescription(const FGuid& In
 	case EAvaTransitionSceneType::Other:
 		Arguments.Add(TEXT("IndefinitePronoun"), bInvertCondition ? LOCTEXT("NoScene", "no ") : LOCTEXT("AnyScene", "a "));
 
-		Arguments.Add(TEXT("Scene"), FText::Format(LOCTEXT("OtherScene", "scene in {0}"), FAvaTransitionLayerUtils::GetLayerQueryText(LayerType, *SpecificLayers.ToString())));
+		Arguments.Add(TEXT("Scene"), FText::Format(LOCTEXT("OtherScene", "scene in {0}"), FAvaTransitionLayerUtils::GetLayerQueryText(InstanceData.LayerType, *InstanceData.SpecificLayers.ToString())));
 
 		Arguments.Add(TEXT("Contains"), LOCTEXT("OtherSceneContains", "contains"));
 		break;
 	}
 
-	Arguments.Add(TEXT("TagAttribute"), FText::FromName(TagAttribute.ToName()));
+	Arguments.Add(TEXT("TagAttribute"), FText::FromName(InstanceData.TagAttribute.ToName()));
 
 	return FText::Format(LOCTEXT("ConditionDescription", "{IndefinitePronoun}{Scene} {Contains} tag attribute '{TagAttribute}'"), Arguments);
 }
 #endif
 
+void FAvaSceneContainsTagAttributeConditionBase::PostLoad(FStateTreeDataView InInstanceDataView)
+{
+	Super::PostLoad(InInstanceDataView);
+
+	PRAGMA_DISABLE_DEPRECATION_WARNINGS
+	if (LayerType_DEPRECATED != EAvaTransitionLayerCompareType::None)
+	{
+		if (FInstanceDataType* InstanceData = UE::AvaTransition::TryGetInstanceData(*this, InInstanceDataView))
+		{
+			InstanceData->SceneType      = SceneType_DEPRECATED;
+			InstanceData->LayerType      = LayerType_DEPRECATED;
+			InstanceData->SpecificLayers = SpecificLayers_DEPRECATED;
+			InstanceData->TagAttribute   = TagAttribute_DEPRECATED;
+		}
+	}
+	PRAGMA_ENABLE_DEPRECATION_WARNINGS
+}
+
 bool FAvaSceneContainsTagAttributeConditionBase::Link(FStateTreeLinker& InLinker)
 {
 	FAvaTransitionCondition::Link(InLinker);
@@ -63,6 +84,7 @@ bool FAvaSceneContainsTagAttributeConditionBase::ContainsTagAttribute(FStateTree
 		return false;
 	}
 
+	const FInstanceDataType& InstanceData = InContext.GetInstanceData(*this);
 	const UAvaSceneSubsystem& SceneSubsystem = InContext.GetExternalData(SceneSubsystemHandle);
 
 	for (const FAvaTransitionScene* TransitionScene : TransitionScenes)
@@ -79,7 +101,7 @@ bool FAvaSceneContainsTagAttributeConditionBase::ContainsTagAttribute(FStateTree
 		}
 
 		const UAvaSceneState* SceneState = Scene->GetSceneState();
-		if (SceneState && SceneState->ContainsTagAttribute(TagAttribute))
+		if (SceneState && SceneState->ContainsTagAttribute(InstanceData.TagAttribute))
 		{
 			return true;
 		}
@@ -90,23 +112,24 @@ bool FAvaSceneContainsTagAttributeConditionBase::ContainsTagAttribute(FStateTree
 
 TArray<const FAvaTransitionScene*> FAvaSceneContainsTagAttributeConditionBase::GetTransitionScenes(FStateTreeExecutionContext& InContext) const
 {
+	const FInstanceDataType& InstanceData = InContext.GetInstanceData(*this);
 	const FAvaTransitionContext& TransitionContext = InContext.GetExternalData(TransitionContextHandle);
 
-	if (SceneType == EAvaTransitionSceneType::This)
+	if (InstanceData.SceneType == EAvaTransitionSceneType::This)
 	{
 		return { TransitionContext.GetTransitionScene() };
 	}
 
-	ensureMsgf(SceneType == EAvaTransitionSceneType::Other
+	ensureMsgf(InstanceData.SceneType == EAvaTransitionSceneType::Other
 		, TEXT("FAvaSceneContainsAttributeCondition::GetTargetSceneContexts did not recognize the provided transition scene type")
-		, *UEnum::GetValueAsString(SceneType));
+		, *UEnum::GetValueAsString(InstanceData.SceneType));
 
 	// Get all the Behavior Instances from Query 
 	TArray<const FAvaTransitionBehaviorInstance*> BehaviorInstances;
 	{
 		UAvaTransitionSubsystem& TransitionSubsystem = InContext.GetExternalData(TransitionSubsystemHandle);
 
-		FAvaTransitionLayerComparator Comparator = FAvaTransitionLayerUtils::BuildComparator(TransitionContext, LayerType, SpecificLayers);
+		FAvaTransitionLayerComparator Comparator = FAvaTransitionLayerUtils::BuildComparator(TransitionContext, InstanceData.LayerType, InstanceData.SpecificLayers);
 
 		BehaviorInstances = FAvaTransitionLayerUtils::QueryBehaviorInstances(TransitionSubsystem, Comparator);
 	}
diff --git a/Engine/Plugins/Experimental/Avalanche/Source/Avalanche/Private/Transition/Conditions/AvaTransitionRCControllerMatchCondition.cpp b/Engine/Plugins/Experimental/Avalanche/Source/Avalanche/Private/Transition/Conditions/AvaTransitionRCControllerMatchCondition.cpp
--- a/Engine/Plugins/Experimental/Avalanche/Source/Avalanche/Private/Transition/Conditions/AvaTransitionRCControllerMatchCondition.cpp
+++ b/Engine/Plugins/Experimental/Avalanche/Source/Avalanche/Private/Transition/Conditions/AvaTransitionRCControllerMatchCondition.cpp
@@ -8,6 +8,7 @@
 #include "AvaTransitionLog.h"
 #include "AvaTransitionScene.h"
 #include "AvaTransitionSubsystem.h"
+#include "AvaTransitionUtils.h"
 #include "Behavior/AvaTransitionBehaviorInstance.h"
 #include "IAvaSceneInterface.h"
 #include "RCVirtualProperty.h"
@@ -21,12 +22,30 @@
 #if WITH_EDITOR
 FText FAvaTransitionRCControllerMatchCondition::GetDescription(const FGuid& InId, FStateTreeDataView InInstanceDataView, const IStateTreeBindingLookup& InBindingLookup, EStateTreeNodeFormatting InFormatting) const
 {
+	const FInstanceDataType& InstanceData = InInstanceDataView.Get<FInstanceDataType>();
+
 	return FText::Format(LOCTEXT("ConditionDescription", "'{0}' is {1}")
-		, ControllerId.ToText()
-		, UEnum::GetDisplayValueAsText(ValueComparisonType).ToLower());
+		, InstanceData.ControllerId.ToText()
+		, UEnum::GetDisplayValueAsText(InstanceData.ValueComparisonType).ToLower());
 }
 #endif
 
+void FAvaTransitionRCControllerMatchCondition::PostLoad(FStateTreeDataView InInstanceDataView)
+{
+	Super::PostLoad(InInstanceDataView);
+
+	PRAGMA_DISABLE_DEPRECATION_WARNINGS
+	if (ControllerId_DEPRECATED.Name != NAME_None)
+	{
+		if (FInstanceDataType* InstanceData = UE::AvaTransition::TryGetInstanceData(*this, InInstanceDataView))
+		{
+			InstanceData->ControllerId        = ControllerId_DEPRECATED;
+			InstanceData->ValueComparisonType = ValueComparisonType_DEPRECATED;
+		}
+	}
+	PRAGMA_ENABLE_DEPRECATION_WARNINGS
+}
+
 bool FAvaTransitionRCControllerMatchCondition::Link(FStateTreeLinker& InLinker)
 {
 	FAvaTransitionCondition::Link(InLinker);
@@ -36,12 +55,13 @@ bool FAvaTransitionRCControllerMatchCondition::Link(FStateTreeLinker& InLinker)
 
 bool FAvaTransitionRCControllerMatchCondition::TestCondition(FStateTreeExecutionContext& InContext) const
 {
+	const FInstanceDataType& InstanceData          = InContext.GetInstanceData(*this);
 	const FAvaTransitionContext& TransitionContext = InContext.GetExternalData(TransitionContextHandle);
 	UAvaTransitionSubsystem& TransitionSubsystem   = InContext.GetExternalData(TransitionSubsystemHandle);
 	const UAvaSceneSubsystem& SceneSubsystem       = InContext.GetExternalData(SceneSubsystemHandle);
 	const FAvaTransitionScene* TransitionScene     = TransitionContext.GetTransitionScene();
 
-	URCVirtualPropertyBase* Controller = GetController(SceneSubsystem, TransitionScene);
+	URCVirtualPropertyBase* Controller = GetController(InstanceData.ControllerId, SceneSubsystem, TransitionScene);
 	if (!Controller)
 	{
 		return false;
@@ -79,7 +99,7 @@ bool FAvaTransitionRCControllerMatchCondition::TestCondition(FStateTreeExecution
 				, *TransitionScene
 				, *OtherTransitionScene);
 		}
-		else if (URCVirtualPropertyBase* OtherController = GetController(SceneSubsystem, OtherTransitionScene))
+		else if (URCVirtualPropertyBase* OtherController = GetController(InstanceData.ControllerId, SceneSubsystem, OtherTransitionScene))
 		{
 			Result = Controller->IsValueEqual(OtherController)
 				? EAvaTransitionComparisonResult::Same
@@ -90,7 +110,7 @@ bool FAvaTransitionRCControllerMatchCondition::TestCondition(FStateTreeExecution
 			Result = EAvaTransitionComparisonResult::None;
 		}
 
-		if (ValueComparisonType == Result)
+		if (InstanceData.ValueComparisonType == Result)
 		{
 			return true;
 		}
@@ -99,7 +119,7 @@ bool FAvaTransitionRCControllerMatchCondition::TestCondition(FStateTreeExecution
 	return false;
 }
 
-URCVirtualPropertyBase* FAvaTransitionRCControllerMatchCondition::GetController(const UAvaSceneSubsystem& InSceneSubsystem, const FAvaTransitionScene* InTransitionScene) const
+URCVirtualPropertyBase* FAvaTransitionRCControllerMatchCondition::GetController(const FAvaRCControllerId& InControllerId, const UAvaSceneSubsystem& InSceneSubsystem, const FAvaTransitionScene* InTransitionScene) const
 {
 	if (!InTransitionScene)
 	{
@@ -114,7 +134,7 @@ URCVirtualPropertyBase* FAvaTransitionRCControllerMatchCondition::GetController(
 
 	if (URemoteControlPreset* RemoteControlPreset = SceneInterface->GetRemoteControlPreset())
 	{
-		return ControllerId.FindController(RemoteControlPreset);
+		return InControllerId.FindController(RemoteControlPreset);
 	}
 
 	return nullptr;
diff --git a/Engine/Plugins/Experimental/Avalanche/Source/Avalanche/Private/Transition/Tasks/AvaSceneAddTagAttributeTask.cpp b/Engine/Plugins/Experimental/Avalanche/Source/Avalanche/Private/Transition/Tasks/AvaSceneAddTagAttributeTask.cpp
--- a/Engine/Plugins/Experimental/Avalanche/Source/Avalanche/Private/Transition/Tasks/AvaSceneAddTagAttributeTask.cpp
+++ b/Engine/Plugins/Experimental/Avalanche/Source/Avalanche/Private/Transition/Tasks/AvaSceneAddTagAttributeTask.cpp
@@ -3,13 +3,15 @@
 #include "Transition/Tasks/AvaSceneAddTagAttributeTask.h"
 #include "AvaSceneState.h"
 #include "IAvaSceneInterface.h"
+#include "StateTreeExecutionContext.h"
 
 #define LOCTEXT_NAMESPACE "AvaSceneAddTagAttributeTask"
 
 #if WITH_EDITOR
 FText FAvaSceneAddTagAttributeTask::GetDescription(const FGuid& InId, FStateTreeDataView InInstanceDataView, const IStateTreeBindingLookup& InBindingLookup, EStateTreeNodeFormatting InFormatting) const
 {
-	return FText::Format(LOCTEXT("TaskDescription", "Add '{0}' tag attribute to this scene"), FText::FromName(TagAttribute.ToName()));
+	const FInstanceDataType& InstanceData = InInstanceDataView.Get<FInstanceDataType>();
+	return FText::Format(LOCTEXT("TaskDescription", "Add '{0}' tag attribute to this scene"), FText::FromName(InstanceData.TagAttribute.ToName()));
 }
 #endif
 
@@ -27,7 +29,8 @@ EStateTreeRunStatus FAvaSceneAddTagAttributeTask::EnterState(FStateTreeExecution
 		return EStateTreeRunStatus::Failed;
 	}
 
-	if (SceneState->AddTagAttribute(TagAttribute))
+	const FInstanceDataType& InstanceData = InContext.GetInstanceData(*this);
+	if (SceneState->AddTagAttribute(InstanceData.TagAttribute))
 	{
 		return EStateTreeRunStatus::Succeeded;
 	}
diff --git a/Engine/Plugins/Experimental/Avalanche/Source/Avalanche/Private/Transition/Tasks/AvaSceneRemoveTagAttributeTask.cpp b/Engine/Plugins/Experimental/Avalanche/Source/Avalanche/Private/Transition/Tasks/AvaSceneRemoveTagAttributeTask.cpp
--- a/Engine/Plugins/Experimental/Avalanche/Source/Avalanche/Private/Transition/Tasks/AvaSceneRemoveTagAttributeTask.cpp
+++ b/Engine/Plugins/Experimental/Avalanche/Source/Avalanche/Private/Transition/Tasks/AvaSceneRemoveTagAttributeTask.cpp
@@ -3,13 +3,15 @@
 #include "Transition/Tasks/AvaSceneRemoveTagAttributeTask.h"
 #include "AvaSceneState.h"
 #include "IAvaSceneInterface.h"
+#include "StateTreeExecutionContext.h"
 
 #define LOCTEXT_NAMESPACE "AvaSceneRemoveTagAttributeTask"
 
 #if WITH_EDITOR
 FText FAvaSceneRemoveTagAttributeTask::GetDescription(const FGuid& InId, FStateTreeDataView InInstanceDataView, const IStateTreeBindingLookup& InBindingLookup, EStateTreeNodeFormatting InFormatting) const
 {
-	return FText::Format(LOCTEXT("TaskDescription", "Remove '{0}' tag attribute from this scene"), FText::FromName(TagAttribute.ToName()));
+	const FInstanceDataType& InstanceData = InInstanceDataView.Get<FInstanceDataType>();
+	return FText::Format(LOCTEXT("TaskDescription", "Remove '{0}' tag attribute from this scene"), FText::FromName(InstanceData.TagAttribute.ToName()));
 }
 #endif
 
@@ -27,7 +29,8 @@ EStateTreeRunStatus FAvaSceneRemoveTagAttributeTask::EnterState(FStateTreeExecut
 		return EStateTreeRunStatus::Failed;
 	}
 
-	if (SceneState->RemoveTagAttribute(TagAttribute))
+	const FInstanceDataType& InstanceData = InContext.GetInstanceData(*this);
+	if (SceneState->RemoveTagAttribute(InstanceData.TagAttribute))
 	{
 		return EStateTreeRunStatus::Succeeded;
 	}
diff --git a/Engine/Plugins/Experimental/Avalanche/Source/Avalanche/Private/Transition/Tasks/AvaSceneTask.cpp b/Engine/Plugins/Experimental/Avalanche/Source/Avalanche/Private/Transition/Tasks/AvaSceneTask.cpp
--- a/Engine/Plugins/Experimental/Avalanche/Source/Avalanche/Private/Transition/Tasks/AvaSceneTask.cpp
+++ b/Engine/Plugins/Experimental/Avalanche/Source/Avalanche/Private/Transition/Tasks/AvaSceneTask.cpp
@@ -3,12 +3,28 @@
 #include "Transition/Tasks/AvaSceneTask.h"
 #include "AvaSceneSubsystem.h"
 #include "AvaTransitionContext.h"
+#include "AvaTransitionUtils.h"
 #include "StateTreeExecutionContext.h"
 #include "StateTreeLinker.h"
 
+void FAvaSceneTask::PostLoad(FStateTreeDataView InInstanceDataView)
+{
+	Super::PostLoad(InInstanceDataView);
+
+	PRAGMA_DISABLE_DEPRECATION_WARNINGS
+	if (TagAttribute_DEPRECATED.IsValid())
+	{
+		if (FInstanceDataType* InstanceData = UE::AvaTransition::TryGetInstanceData(*this, InInstanceDataView))
+		{
+			InstanceData->TagAttribute = TagAttribute_DEPRECATED;
+		}
+	}
+	PRAGMA_ENABLE_DEPRECATION_WARNINGS
+}
+
 bool FAvaSceneTask::Link(FStateTreeLinker& InLinker)
 {
-	FAvaTransitionTask::Link(InLinker);
+	Super::Link(InLinker);
 	InLinker.LinkExternalData(SceneSubsystemHandle);
 	return true;
 }
diff --git a/Engine/Plugins/Experimental/Avalanche/Source/Avalanche/Public/Transition/Conditions/AvaSceneContainsTagAttributeCondition.h b/Engine/Plugins/Experimental/Avalanche/Source/Avalanche/Public/Transition/Conditions/AvaSceneContainsTagAttributeCondition.h
--- a/Engine/Plugins/Experimental/Avalanche/Source/Avalanche/Public/Transition/Conditions/AvaSceneContainsTagAttributeCondition.h
+++ b/Engine/Plugins/Experimental/Avalanche/Source/Avalanche/Public/Transition/Conditions/AvaSceneContainsTagAttributeCondition.h
@@ -15,6 +15,22 @@ USTRUCT()
 struct FAvaSceneContainsTagAttributeConditionInstanceData
 {
 	GENERATED_BODY()
+
+	/** Whether the scene to check should be this scene or other scene */
+	UPROPERTY(EditAnywhere, Category="Parameter")
+	EAvaTransitionSceneType SceneType = EAvaTransitionSceneType::This;
+
+	/** Which Layer should be queried for the Scene Attributes */
+	UPROPERTY(EditAnywhere, Category="Parameter", meta=(EditCondition="SceneType==EAvaTransitionSceneType::Other", EditConditionHides))
+	EAvaTransitionLayerCompareType LayerType = EAvaTransitionLayerCompareType::Same;
+
+	/** Specific layer tags to check */
+	UPROPERTY(EditAnywhere, Category="Parameter", meta=(EditCondition="SceneType==EAvaTransitionSceneType::Other && LayerType==EAvaTransitionLayerCompareType::MatchingTag", EditConditionHides))
+	FAvaTagHandleContainer SpecificLayers;
+
+	/** The Tag Attribute to check if it's contained in the scene(s) or not */
+	UPROPERTY(EditAnywhere, Category="Parameter")
+	FAvaTagHandle TagAttribute;
 };
 
 USTRUCT(meta=(Hidden))
@@ -22,20 +38,28 @@ struct AVALANCHE_API FAvaSceneContainsTagAttributeConditionBase : public FAvaTra
 {
 	GENERATED_BODY()
 
+	using FInstanceDataType = FAvaSceneContainsTagAttributeConditionInstanceData;
+
+	PRAGMA_DISABLE_DEPRECATION_WARNINGS
 	FAvaSceneContainsTagAttributeConditionBase() = default;
+	virtual ~FAvaSceneContainsTagAttributeConditionBase() override = default;
+	FAvaSceneContainsTagAttributeConditionBase(const FAvaSceneContainsTagAttributeConditionBase&) = default;
+	FAvaSceneContainsTagAttributeConditionBase(FAvaSceneContainsTagAttributeConditionBase&&) = default;
+	FAvaSceneContainsTagAttributeConditionBase& operator=(const FAvaSceneContainsTagAttributeConditionBase&) = default;
+	FAvaSceneContainsTagAttributeConditionBase& operator=(FAvaSceneContainsTagAttributeConditionBase&&) = default;
+	PRAGMA_ENABLE_DEPRECATION_WARNINGS
 
 	FAvaSceneContainsTagAttributeConditionBase(bool bInInvertCondition)
 		: bInvertCondition(bInInvertCondition)
 	{
 	}
 
-	using FInstanceDataType = FAvaSceneContainsTagAttributeConditionInstanceData;
-
 	//~ Begin FStateTreeNodeBase
 #if WITH_EDITOR
 	virtual FText GetDescription(const FGuid& InId, FStateTreeDataView InInstanceDataView, const IStateTreeBindingLookup& InBindingLookup, EStateTreeNodeFormatting InFormatting) const override;
 #endif
 	virtual const UStruct* GetInstanceDataType() const override { return FInstanceDataType::StaticStruct(); }
+	virtual void PostLoad(FStateTreeDataView InInstanceDataView) override;
 	virtual bool Link(FStateTreeLinker& InLinker) override;
 	//~ End FStateTreeNodeBase
 
@@ -47,21 +71,21 @@ struct AVALANCHE_API FAvaSceneContainsTagAttributeConditionBase : public FAvaTra
 
 	TArray<const FAvaTransitionScene*> GetTransitionScenes(FStateTreeExecutionContext& InContext) const;
 
-	/** Whether the scene to check should be this scene or other scene */
-	UPROPERTY(EditAnywhere, Category="Parameter")
-	EAvaTransitionSceneType SceneType = EAvaTransitionSceneType::This;
+	UE_DEPRECATED(5.5, "SceneType has been moved to Instance Data")
+	UPROPERTY(meta = (DeprecatedProperty, DeprecationMessage = "Use the Instance Data SceneType instead"))
+	EAvaTransitionSceneType SceneType_DEPRECATED = EAvaTransitionSceneType::This;
 
-	/** Which Layer should be queried for the Scene Attributes */
-	UPROPERTY(EditAnywhere, Category="Parameter", meta=(EditCondition="SceneType == EAvaTransitionSceneType::Other", EditConditionHides))
-	EAvaTransitionLayerCompareType LayerType = EAvaTransitionLayerCompareType::Same;
+	UE_DEPRECATED(5.5, "LayerType has been moved to Instance Data")
+	UPROPERTY(meta = (DeprecatedProperty, DeprecationMessage = "Use the Instance Data LayerType instead"))
+	EAvaTransitionLayerCompareType LayerType_DEPRECATED = EAvaTransitionLayerCompareType::None;
 
-	/** Specific layer tags to check */
-	UPROPERTY(EditAnywhere, Category="Parameter", meta=(EditCondition="SceneType == EAvaTransitionSceneType::Other && LayerType==EAvaTransitionLayerCompareType::MatchingTag", EditConditionHides))
-	FAvaTagHandleContainer SpecificLayers;
+	UE_DEPRECATED(5.5, "SpecificLayers has been moved to Instance Data")
+	UPROPERTY(meta = (DeprecatedProperty, DeprecationMessage = "Use the Instance Data SpecificLayers instead"))
+	FAvaTagHandleContainer SpecificLayers_DEPRECATED;
 
-	/** The Tag Attribute to check if it's contained in the scene(s) or not */
-	UPROPERTY(EditAnywhere, Category="Parameter")
-	FAvaTagHandle TagAttribute;
+	UE_DEPRECATED(5.5, "TagAttribute has been moved to Instance Data")
+	UPROPERTY(meta = (DeprecatedProperty, DeprecationMessage = "Use the Instance Data TagAttribute instead"))
+	FAvaTagHandle TagAttribute_DEPRECATED;
 
 protected:
 	bool bInvertCondition = false;
diff --git a/Engine/Plugins/Experimental/Avalanche/Source/Avalanche/Public/Transition/Conditions/AvaTransitionRCControllerMatchCondition.h b/Engine/Plugins/Experimental/Avalanche/Source/Avalanche/Public/Transition/Conditions/AvaTransitionRCControllerMatchCondition.h
--- a/Engine/Plugins/Experimental/Avalanche/Source/Avalanche/Public/Transition/Conditions/AvaTransitionRCControllerMatchCondition.h
+++ b/Engine/Plugins/Experimental/Avalanche/Source/Avalanche/Public/Transition/Conditions/AvaTransitionRCControllerMatchCondition.h
@@ -15,6 +15,12 @@ USTRUCT()
 struct FAvaTransitionRCControllerMatchConditionInstanceData
 {
 	GENERATED_BODY()
+
+	UPROPERTY(EditAnywhere, Category="Parameter")
+	FAvaRCControllerId ControllerId;
+
+	UPROPERTY(EditAnywhere, Category="Parameter")
+	EAvaTransitionComparisonResult ValueComparisonType = EAvaTransitionComparisonResult::None;
 };
 
 USTRUCT(DisplayName="Compare RC Controller Values", Category="Remote Control")
@@ -24,11 +30,21 @@ struct AVALANCHE_API FAvaTransitionRCControllerMatchCondition : public FAvaTrans
 
 	using FInstanceDataType = FAvaTransitionRCControllerMatchConditionInstanceData;
 
+	PRAGMA_DISABLE_DEPRECATION_WARNINGS
+	FAvaTransitionRCControllerMatchCondition() = default;
+	virtual ~FAvaTransitionRCControllerMatchCondition() override = default;
+	FAvaTransitionRCControllerMatchCondition(const FAvaTransitionRCControllerMatchCondition&) = default;
+	FAvaTransitionRCControllerMatchCondition(FAvaTransitionRCControllerMatchCondition&&) = default;
+	FAvaTransitionRCControllerMatchCondition& operator=(const FAvaTransitionRCControllerMatchCondition&) = default;
+	FAvaTransitionRCControllerMatchCondition& operator=(FAvaTransitionRCControllerMatchCondition&&) = default;
+	PRAGMA_ENABLE_DEPRECATION_WARNINGS
+
 	//~ Begin FStateTreeNodeBase
 #if WITH_EDITOR
 	virtual FText GetDescription(const FGuid& InId, FStateTreeDataView InInstanceDataView, const IStateTreeBindingLookup& InBindingLookup, EStateTreeNodeFormatting InFormatting) const override;
 #endif
 	virtual const UStruct* GetInstanceDataType() const override { return FInstanceDataType::StaticStruct(); }
+	virtual void PostLoad(FStateTreeDataView InInstanceDataView) override;
 	virtual bool Link(FStateTreeLinker& InLinker) override;
 	//~ End FStateTreeNodeBase
 
@@ -36,13 +52,15 @@ struct AVALANCHE_API FAvaTransitionRCControllerMatchCondition : public FAvaTrans
 	virtual bool TestCondition(FStateTreeExecutionContext& InContext) const override;
 	//~ End FStateTreeConditionBase
 
-	URCVirtualPropertyBase* GetController(const UAvaSceneSubsystem& InSceneSubsystem, const FAvaTransitionScene* InTransitionScene) const;
+	URCVirtualPropertyBase* GetController(const FAvaRCControllerId& InControllerId, const UAvaSceneSubsystem& InSceneSubsystem, const FAvaTransitionScene* InTransitionScene) const;
 
-	UPROPERTY(EditAnywhere, Category="Parameter")
-	FAvaRCControllerId ControllerId;
+	UE_DEPRECATED(5.5, "ControllerId has been moved to Instance Data")
+	UPROPERTY(meta = (DeprecatedProperty, DeprecationMessage = "Use the Instance Data ControllerId instead"))
+	FAvaRCControllerId ControllerId_DEPRECATED;
 
-	UPROPERTY(EditAnywhere, Category="Parameter")
-	EAvaTransitionComparisonResult ValueComparisonType = EAvaTransitionComparisonResult::None;
+	UE_DEPRECATED(5.5, "ValueComparisonType has been moved to Instance Data")
+	UPROPERTY(meta = (DeprecatedProperty, DeprecationMessage = "Use the Instance Data ValueComparisonType instead"))
+	EAvaTransitionComparisonResult ValueComparisonType_DEPRECATED = EAvaTransitionComparisonResult::None;
 
 	TStateTreeExternalDataHandle<UAvaSceneSubsystem> SceneSubsystemHandle;
 };
diff --git a/Engine/Plugins/Experimental/Avalanche/Source/Avalanche/Public/Transition/Tasks/AvaSceneTask.h b/Engine/Plugins/Experimental/Avalanche/Source/Avalanche/Public/Transition/Tasks/AvaSceneTask.h
--- a/Engine/Plugins/Experimental/Avalanche/Source/Avalanche/Public/Transition/Tasks/AvaSceneTask.h
+++ b/Engine/Plugins/Experimental/Avalanche/Source/Avalanche/Public/Transition/Tasks/AvaSceneTask.h
@@ -13,6 +13,9 @@ USTRUCT()
 struct FAvaSceneTaskInstanceData
 {
 	GENERATED_BODY()
+
+	UPROPERTY(EditAnywhere, Category="Parameter")
+	FAvaTagHandle TagAttribute;
 };
 
 USTRUCT(meta=(Hidden))
@@ -22,15 +25,26 @@ struct AVALANCHE_API FAvaSceneTask : public FAvaTransitionTask
 
 	using FInstanceDataType = FAvaSceneTaskInstanceData;
 
+	PRAGMA_DISABLE_DEPRECATION_WARNINGS
+	FAvaSceneTask() = default;
+	virtual ~FAvaSceneTask() override = default;
+	FAvaSceneTask(const FAvaSceneTask&) = default;
+	FAvaSceneTask(FAvaSceneTask&&) = default;
+	FAvaSceneTask& operator=(const FAvaSceneTask&) = default;
+	FAvaSceneTask& operator=(FAvaSceneTask&&) = default;
+	PRAGMA_ENABLE_DEPRECATION_WARNINGS
+
 	//~ Begin FStateTreeNodeBase
 	virtual const UStruct* GetInstanceDataType() const override { return FInstanceDataType::StaticStruct(); }
+	virtual void PostLoad(FStateTreeDataView InInstanceDataView) override;
 	virtual bool Link(FStateTreeLinker& InLinker) override;
 	//~ End FStateTreeNodeBase
 
 	IAvaSceneInterface* GetScene(FStateTreeExecutionContext& InContext) const;
 
-	UPROPERTY(EditAnywhere, Category="Parameter")
-	FAvaTagHandle TagAttribute;
+	UE_DEPRECATED(5.5, "TagAttribute has been moved to Instance Data")
+	UPROPERTY(meta = (DeprecatedProperty, DeprecationMessage = "Use the Instance Data TagAttribute instead"))
+	FAvaTagHandle TagAttribute_DEPRECATED;
 
 	TStateTreeExternalDataHandle<UAvaSceneSubsystem> SceneSubsystemHandle;
 };
diff --git a/Engine/Plugins/Experimental/Avalanche/Source/AvalancheEditor/Private/Transition/AvaTransitionExtension.cpp b/Engine/Plugins/Experimental/Avalanche/Source/AvalancheEditor/Private/Transition/AvaTransitionExtension.cpp
--- a/Engine/Plugins/Experimental/Avalanche/Source/AvalancheEditor/Private/Transition/AvaTransitionExtension.cpp
+++ b/Engine/Plugins/Experimental/Avalanche/Source/AvalancheEditor/Private/Transition/AvaTransitionExtension.cpp
@@ -31,7 +31,7 @@ namespace UE::AvaEditor::Private
 	{
 		TStateTreeEditorNode<FAvaTransitionSceneMatchCondition>& ConditionNode = InState.AddEnterCondition<FAvaTransitionSceneMatchCondition>();
 
-		FAvaTransitionSceneMatchCondition& Condition = ConditionNode.GetNode();
+		FAvaTransitionSceneMatchCondition::FInstanceDataType& Condition = ConditionNode.GetInstanceData();
 
 		Condition.SceneComparisonType  = InComparisonType;
 		ConditionNode.ExpressionOperand = EStateTreeExpressionOperand::Or;
@@ -39,7 +39,7 @@ namespace UE::AvaEditor::Private
 
 	void AddTransitionCondition(UStateTreeState& InState, EAvaTransitionType InTransitionType)
 	{
-		FAvaTransitionTypeMatchCondition& Condition = InState.AddEnterCondition<FAvaTransitionTypeMatchCondition>().GetNode();
+		FAvaTransitionTypeMatchCondition::FInstanceDataType& Condition = InState.AddEnterCondition<FAvaTransitionTypeMatchCondition>().GetInstanceData();
 		Condition.TransitionType = InTransitionType;
 	}
 
@@ -51,7 +51,7 @@ namespace UE::AvaEditor::Private
 
 	FAvaSequencePlayParams& AddPlayTask(UStateTreeState& InState, const FAvaTagHandle& InSequenceTag)
 	{
-		FAvaTransitionPlaySequenceTask& PlayTask = InState.AddTask<FAvaTransitionPlaySequenceTask>().GetNode();
+		FAvaTransitionPlaySequenceTask::FInstanceDataType& PlayTask = InState.AddTask<FAvaTransitionPlaySequenceTask>().GetInstanceData();
 
 		PlayTask.QueryType   = EAvaTransitionSequenceQueryType::Tag;
 		PlayTask.SequenceTag = InSequenceTag;
@@ -191,7 +191,7 @@ void FAvaTransitionExtension::BuildDefaultTransitionTree(UAvaTransitionTreeEdito
 			UStateTreeState& WaitState = ChangeIn.AddChildState(TEXT("Wait for change out"));
 			SetStateCompletedResult(WaitState, EStateTreeTransitionType::NextSelectableState);
 
-			FAvaTransitionWaitForLayerTask& WaitTask = WaitState.AddTask<FAvaTransitionWaitForLayerTask>().GetNode();
+			FAvaTransitionWaitForLayerTask::FInstanceDataType& WaitTask = WaitState.AddTask<FAvaTransitionWaitForLayerTask>().GetInstanceData();
 			WaitTask.LayerType = EAvaTransitionLayerCompareType::Same;
 
 			UStateTreeState& PlayChangeInState = ChangeIn.AddChildState(TEXT("Play change in"));
diff --git a/Engine/Plugins/Experimental/Avalanche/Source/AvalancheSequence/Private/Transition/AvaTransitionContinueSequenceTask.cpp b/Engine/Plugins/Experimental/Avalanche/Source/AvalancheSequence/Private/Transition/AvaTransitionContinueSequenceTask.cpp
--- a/Engine/Plugins/Experimental/Avalanche/Source/AvalancheSequence/Private/Transition/AvaTransitionContinueSequenceTask.cpp
+++ b/Engine/Plugins/Experimental/Avalanche/Source/AvalancheSequence/Private/Transition/AvaTransitionContinueSequenceTask.cpp
@@ -12,15 +12,17 @@ TArray<UAvaSequencePlayer*> FAvaTransitionContinueSequenceTask::ExecuteSequenceT
 		return TArray<UAvaSequencePlayer*>();
 	}
 
-	switch (QueryType)
+	const FInstanceDataType& InstanceData = InContext.GetInstanceData(*this);
+
+	switch (InstanceData.QueryType)
 	{
 	case EAvaTransitionSequenceQueryType::Name:
-		return PlaybackObject->ContinueSequencesByLabel(SequenceName);
+		return PlaybackObject->ContinueSequencesByLabel(InstanceData.SequenceName);
 
 	case EAvaTransitionSequenceQueryType::Tag:
-		if (const FAvaTag* Tag = SequenceTag.GetTag())
+		if (const FAvaTag* Tag = InstanceData.SequenceTag.GetTag())
 		{
-			return PlaybackObject->ContinueSequencesByTag(*Tag, bPerformExactMatch);
+			return PlaybackObject->ContinueSequencesByTag(*Tag, InstanceData.bPerformExactMatch);
 		}
 		return TArray<UAvaSequencePlayer*>();
 	}
diff --git a/Engine/Plugins/Experimental/Avalanche/Source/AvalancheSequence/Private/Transition/AvaTransitionInitializeSequence.cpp b/Engine/Plugins/Experimental/Avalanche/Source/AvalancheSequence/Private/Transition/AvaTransitionInitializeSequence.cpp
--- a/Engine/Plugins/Experimental/Avalanche/Source/AvalancheSequence/Private/Transition/AvaTransitionInitializeSequence.cpp
+++ b/Engine/Plugins/Experimental/Avalanche/Source/AvalancheSequence/Private/Transition/AvaTransitionInitializeSequence.cpp
@@ -2,7 +2,31 @@
 
 #include "Transition/AvaTransitionInitializeSequence.h"
 #include "AvaSequencePlaybackObject.h"
+#include "AvaTransitionUtils.h"
 #include "Math/NumericLimits.h"
+#include "StateTreeExecutionContext.h"
+#include "Transition/AvaTransitionSequenceUtils.h"
+
+void FAvaTransitionInitializeSequence::PostLoad(FStateTreeDataView InInstanceDataView)
+{
+	Super::PostLoad(InInstanceDataView);
+
+	PRAGMA_DISABLE_DEPRECATION_WARNINGS
+	if (QueryType_DEPRECATED != EAvaTransitionSequenceQueryType::None)
+	{
+		if (FInstanceDataType* InstanceData = UE::AvaTransition::TryGetInstanceData(*this, InInstanceDataView))
+		{
+			InstanceData->InitializeTime = InitializeTime_DEPRECATED;
+			InstanceData->PlayMode       = PlayMode_DEPRECATED;
+		}
+	}
+	PRAGMA_ENABLE_DEPRECATION_WARNINGS
+}
+
+EAvaTransitionSequenceWaitType FAvaTransitionInitializeSequence::GetWaitType(FStateTreeExecutionContext& InContext) const
+{
+	return EAvaTransitionSequenceWaitType::NoWait;
+}
 
 TArray<UAvaSequencePlayer*> FAvaTransitionInitializeSequence::ExecuteSequenceTask(FStateTreeExecutionContext& InContext) const
 {
@@ -14,20 +38,22 @@ TArray<UAvaSequencePlayer*> FAvaTransitionInitializeSequence::ExecuteSequenceTas
 
 	FAvaSequencePlayParams PlaySettings;
 
-	// Set start to largest double, so that it gets clamped down to be at the End Time (i.e. time that the Sequence should evaluate)
+	const FInstanceDataType& InstanceData = InContext.GetInstanceData(*this);
+
+	// Set start to the largest double, so that it gets clamped down to be at the End Time (i.e. time that the Sequence should evaluate)
 	PlaySettings.Start = FAvaSequenceTime(TNumericLimits<double>::Max());
-	PlaySettings.End = InitializeTime;
-	PlaySettings.PlayMode = PlayMode;
+	PlaySettings.End = InstanceData.InitializeTime;
+	PlaySettings.PlayMode = InstanceData.PlayMode;
 
-	switch (QueryType)
+	switch (InstanceData.QueryType)
 	{
 	case EAvaTransitionSequenceQueryType::Name:
-		return PlaybackObject->PlaySequencesByLabel(SequenceName, PlaySettings);
+		return PlaybackObject->PlaySequencesByLabel(InstanceData.SequenceName, PlaySettings);
 
 	case EAvaTransitionSequenceQueryType::Tag:
-		if (const FAvaTag* Tag = SequenceTag.GetTag())
+		if (const FAvaTag* Tag = InstanceData.SequenceTag.GetTag())
 		{
-			return PlaybackObject->PlaySequencesByTag(*Tag, bPerformExactMatch, PlaySettings);
+			return PlaybackObject->PlaySequencesByTag(*Tag, InstanceData.bPerformExactMatch, PlaySettings);
 		}
 		return TArray<UAvaSequencePlayer*>();
 	}
diff --git a/Engine/Plugins/Experimental/Avalanche/Source/AvalancheSequence/Private/Transition/AvaTransitionPlaySequenceTask.cpp b/Engine/Plugins/Experimental/Avalanche/Source/AvalancheSequence/Private/Transition/AvaTransitionPlaySequenceTask.cpp
--- a/Engine/Plugins/Experimental/Avalanche/Source/AvalancheSequence/Private/Transition/AvaTransitionPlaySequenceTask.cpp
+++ b/Engine/Plugins/Experimental/Avalanche/Source/AvalancheSequence/Private/Transition/AvaTransitionPlaySequenceTask.cpp
@@ -5,32 +5,39 @@
 #include "AvaSequencePlaybackObject.h"
 #include "AvaSequencePlayer.h"
 #include "AvaTransitionContext.h"
+#include "AvaTransitionUtils.h"
 #include "StateTreeExecutionContext.h"
 #include "StateTreeLinker.h"
+#include "Transition/AvaTransitionSequenceUtils.h"
 
 #define LOCTEXT_NAMESPACE "AvaTransitionPlaySequenceTask"
 
-#if WITH_EDITOR
-FText FAvaTransitionPlaySequenceTask::GetDescription(const FGuid& InId, FStateTreeDataView InInstanceDataView, const IStateTreeBindingLookup& InBindingLookup, EStateTreeNodeFormatting InFormatting) const
+void FAvaTransitionPlaySequenceTask::PostLoad(FStateTreeDataView InInstanceDataView)
 {
-	FText DescriptionFormat = LOCTEXT("TaskDescription", "Play {SequenceQueryText} {AdditionalArgs}");
+	Super::PostLoad(InInstanceDataView);
 
-	FFormatNamedArguments NamedArgs;
-	NamedArgs.Add(TEXT("SequenceQueryText"), GetSequenceQueryText());
-
-	TArray<FText> AdditionalArgs;
-	if (PlaySettings.PlayMode != EAvaSequencePlayMode::Forward)
+	PRAGMA_DISABLE_DEPRECATION_WARNINGS
+	if (QueryType_DEPRECATED != EAvaTransitionSequenceQueryType::None)
 	{
-		AdditionalArgs.Add(UEnum::GetDisplayValueAsText(PlaySettings.PlayMode));
+		if (FInstanceDataType* InstanceData = UE::AvaTransition::TryGetInstanceData(*this, InInstanceDataView))
+		{
+			InstanceData->PlaySettings = PlaySettings_DEPRECATED;
+		}
 	}
-	AdditionalArgs.Add(UEnum::GetDisplayValueAsText(WaitType));
+	PRAGMA_ENABLE_DEPRECATION_WARNINGS
+}
 
-	if (!AdditionalArgs.IsEmpty())
-	{
-		NamedArgs.Add(TEXT("AdditionalArgs"), FText::Format(INVTEXT("( {0} )"), FText::Join(FText::FromString(TEXT(" | ")), AdditionalArgs)));	
-	}
+#if WITH_EDITOR
+FText FAvaTransitionPlaySequenceTask::GetDescription(const FGuid& InId, FStateTreeDataView InInstanceDataView, const IStateTreeBindingLookup& InBindingLookup, EStateTreeNodeFormatting InFormatting) const
+{
+	const FInstanceDataType& InstanceData = InInstanceDataView.Get<FInstanceDataType>();
 
-	return FText::Format(DescriptionFormat, NamedArgs);
+	TArray<FText> AdditionalArgs;
+	AdditionalArgs.Add(UEnum::GetDisplayValueAsText(InstanceData.PlaySettings.PlayMode));
+	AdditionalArgs.Add(UEnum::GetDisplayValueAsText(InstanceData.WaitType));
+
+	FText AddOnText = FText::Format(INVTEXT("( {0} )"), FText::Join(FText::FromString(TEXT(" | ")), AdditionalArgs));
+	return FText::Format(LOCTEXT("TaskDescription", "Play {0} {1}"), GetSequenceQueryText(InstanceData), AddOnText);
 }
 #endif
 
@@ -42,15 +49,17 @@ TArray<UAvaSequencePlayer*> FAvaTransitionPlaySequenceTask::ExecuteSequenceTask(
 		return TArray<UAvaSequencePlayer*>();
 	}
 
-	switch (QueryType)
+	const FInstanceDataType& InstanceData = InContext.GetInstanceData(*this);
+
+	switch (InstanceData.QueryType)
 	{
 	case EAvaTransitionSequenceQueryType::Name:
-		return PlaybackObject->PlaySequencesByLabel(SequenceName, PlaySettings);
+		return PlaybackObject->PlaySequencesByLabel(InstanceData.SequenceName, InstanceData.PlaySettings);
 
 	case EAvaTransitionSequenceQueryType::Tag:
-		if (const FAvaTag* Tag = SequenceTag.GetTag())
+		if (const FAvaTag* Tag = InstanceData.SequenceTag.GetTag())
 		{
-			return PlaybackObject->PlaySequencesByTag(*Tag, bPerformExactMatch, PlaySettings);
+			return PlaybackObject->PlaySequencesByTag(*Tag, InstanceData.bPerformExactMatch, InstanceData.PlaySettings);
 		}
 		return TArray<UAvaSequencePlayer*>();
 	}
diff --git a/Engine/Plugins/Experimental/Avalanche/Source/AvalancheSequence/Private/Transition/AvaTransitionSequenceTask.cpp b/Engine/Plugins/Experimental/Avalanche/Source/AvalancheSequence/Private/Transition/AvaTransitionSequenceTask.cpp
--- a/Engine/Plugins/Experimental/Avalanche/Source/AvalancheSequence/Private/Transition/AvaTransitionSequenceTask.cpp
+++ b/Engine/Plugins/Experimental/Avalanche/Source/AvalancheSequence/Private/Transition/AvaTransitionSequenceTask.cpp
@@ -1,158 +1,26 @@
 // Copyright Epic Games, Inc. All Rights Reserved.
 
 #include "Transition/AvaTransitionSequenceTask.h"
-#include "AvaSequence.h"
-#include "AvaSequencePlaybackObject.h"
-#include "AvaSequencePlayer.h"
-#include "AvaSequenceSubsystem.h"
-#include "AvaTransitionContext.h"
+#include "AvaTransitionUtils.h"
 #include "StateTreeExecutionContext.h"
-#include "StateTreeLinker.h"
 #include "Transition/AvaTransitionSequenceUtils.h"
 
-#define LOCTEXT_NAMESPACE "AvaTransitionSequenceTaskBase"
-
-TArray<UAvaSequencePlayer*> FAvaTransitionSequenceTaskBase::ExecuteSequenceTask(FStateTreeExecutionContext& InContext) const
-{
-	return TArray<UAvaSequencePlayer*>();
-}
-
-bool FAvaTransitionSequenceTaskBase::IsSequenceQueryValid() const
-{
-	switch (QueryType)
-	{
-	case EAvaTransitionSequenceQueryType::Name:
-		return !SequenceName.IsNone();
-
-	case EAvaTransitionSequenceQueryType::Tag:
-		return SequenceTag.IsValid();
-	}
-	checkNoEntry();
-	return false;
-}
-
-IAvaSequencePlaybackObject* FAvaTransitionSequenceTaskBase::GetPlaybackObject(FStateTreeExecutionContext& InContext) const
-{
-	const FAvaTransitionContext& TransitionContext = InContext.GetExternalData(TransitionContextHandle);
-	const UAvaSequenceSubsystem& SequenceSubsystem = InContext.GetExternalData(SequenceSubsystemHandle);
-
-	return FAvaTransitionSequenceUtils::GetPlaybackObject(TransitionContext, SequenceSubsystem);
-}
-
-TArray<TWeakObjectPtr<UAvaSequence>> FAvaTransitionSequenceTaskBase::GetActiveSequences(TConstArrayView<UAvaSequencePlayer*> InSequencePlayers) const
-{
-	TArray<TWeakObjectPtr<UAvaSequence>> OutActiveSequences;
-	OutActiveSequences.Reserve(InSequencePlayers.Num());
-
-	Algo::TransformIf(InSequencePlayers, OutActiveSequences,
-		[](UAvaSequencePlayer* InPlayer)->bool
-		{
-			return IsValid(InPlayer);
-		},
-		[](UAvaSequencePlayer* InPlayer)->TWeakObjectPtr<UAvaSequence>
-		{
-			return TWeakObjectPtr<UAvaSequence>(InPlayer->GetAvaSequence());
-		});
-
-	return OutActiveSequences;
-}
-
-EStateTreeRunStatus FAvaTransitionSequenceTaskBase::WaitForActiveSequences(FStateTreeExecutionContext& InContext) const
-{
-	FAvaTransitionSequenceInstanceData& InstanceData = InContext.GetInstanceData(*this);
-
-	EAvaTransitionSequenceWaitType WaitType = GetWaitType();
-
-	// No wait should only succeed in Enter State, and so return failed run if active sequences is empty 
-	if (WaitType == EAvaTransitionSequenceWaitType::NoWait)
-	{
-		if (InstanceData.ActiveSequences.IsEmpty())
-		{
-			return EStateTreeRunStatus::Failed;
-		}
-		return EStateTreeRunStatus::Succeeded;
-	}
-
-	IAvaSequencePlaybackObject* PlaybackObject = GetPlaybackObject(InContext);
-	if (!PlaybackObject)
-	{
-		return EStateTreeRunStatus::Failed;
-	}
-
-	return FAvaTransitionSequenceUtils::UpdatePlayerRunStatus(*PlaybackObject, InstanceData.ActiveSequences, WaitType);
-}
-
-void FAvaTransitionSequenceTaskBase::StopActiveSequences(FStateTreeExecutionContext& InContext) const
+void FAvaTransitionSequenceTask::PostLoad(FStateTreeDataView InInstanceDataView)
 {
-	IAvaSequencePlaybackObject* PlaybackObject = GetPlaybackObject(InContext);
-	if (!PlaybackObject)
-	{
-		return;
-	}
-
-	FAvaTransitionSequenceInstanceData& InstanceData = InContext.GetInstanceData(*this);
+	Super::PostLoad(InInstanceDataView);
 
-	for (const TWeakObjectPtr<UAvaSequence>& SequenceWeak : InstanceData.ActiveSequences)
+	PRAGMA_DISABLE_DEPRECATION_WARNINGS
+	if (QueryType_DEPRECATED != EAvaTransitionSequenceQueryType::None)
 	{
-		if (UAvaSequence* Sequence = SequenceWeak.Get())
+		if (FInstanceDataType* InstanceData = UE::AvaTransition::TryGetInstanceData(*this, InInstanceDataView))
 		{
-			PlaybackObject->StopSequence(Sequence);
+			InstanceData->WaitType = WaitType_DEPRECATED;
 		}
 	}
+	PRAGMA_ENABLE_DEPRECATION_WARNINGS
 }
 
-FText FAvaTransitionSequenceTaskBase::GetSequenceQueryText() const
+EAvaTransitionSequenceWaitType FAvaTransitionSequenceTask::GetWaitType(FStateTreeExecutionContext& InContext) const
 {
-	switch (QueryType)
-	{
-	case EAvaTransitionSequenceQueryType::Name:
-		return FText::Format(INVTEXT("'{0}'"), FText::FromName(SequenceName));
-
-	case EAvaTransitionSequenceQueryType::Tag:
-		return FText::Format(LOCTEXT("SequenceQueryTag", "tag '{0}'"), FText::FromName(SequenceTag.ToName()));
-	}
-
-	checkNoEntry();
-	return FText::GetEmpty();
+	return InContext.GetInstanceData(*this).WaitType;
 }
-
-EStateTreeRunStatus FAvaTransitionSequenceTaskBase::EnterState(FStateTreeExecutionContext& InContext
-	, const FStateTreeTransitionResult& InTransition) const
-{
-	if (!IsSequenceQueryValid())
-	{
-		return EStateTreeRunStatus::Failed;
-	}
-
-	TArray<UAvaSequencePlayer*> SequencePlayers = ExecuteSequenceTask(InContext);
-	if (SequencePlayers.IsEmpty())
-	{
-		return EStateTreeRunStatus::Succeeded;
-	}
-
-	FInstanceDataType& InstanceData = InContext.GetInstanceData(*this);
-	InstanceData.ActiveSequences = FAvaTransitionSequenceUtils::GetSequences(SequencePlayers);
-	return WaitForActiveSequences(InContext);
-}
-
-EStateTreeRunStatus FAvaTransitionSequenceTaskBase::Tick(FStateTreeExecutionContext& InContext, const float InDeltaTime) const
-{
-	return WaitForActiveSequences(InContext);
-}
-
-void FAvaTransitionSequenceTaskBase::ExitState(FStateTreeExecutionContext& InContext, const FStateTreeTransitionResult& InTransition) const
-{
-	if (InTransition.CurrentRunStatus == EStateTreeRunStatus::Stopped)
-	{
-		StopActiveSequences(InContext);
-	}
-}
-
-bool FAvaTransitionSequenceTaskBase::Link(FStateTreeLinker& InLinker)
-{
-	FAvaTransitionTask::Link(InLinker);
-	InLinker.LinkExternalData(SequenceSubsystemHandle);
-	return true;
-}
-
-#undef LOCTEXT_NAMESPACE
diff --git a/Engine/Plugins/Experimental/Avalanche/Source/AvalancheSequence/Private/Transition/AvaTransitionSequenceTaskBase.cpp b/Engine/Plugins/Experimental/Avalanche/Source/AvalancheSequence/Private/Transition/AvaTransitionSequenceTaskBase.cpp
new file mode 100644
--- /dev/null
+++ b/Engine/Plugins/Experimental/Avalanche/Source/AvalancheSequence/Private/Transition/AvaTransitionSequenceTaskBase.cpp
@@ -0,0 +1,175 @@
+// Copyright Epic Games, Inc. All Rights Reserved.
+
+#include "Transition/AvaTransitionSequenceTaskBase.h"
+#include "AvaSequence.h"
+#include "AvaSequencePlayer.h"
+#include "AvaSequenceSubsystem.h"
+#include "AvaTransitionContext.h"
+#include "AvaTransitionUtils.h"
+#include "StateTreeExecutionContext.h"
+#include "StateTreeLinker.h"
+#include "Transition/AvaTransitionSequenceUtils.h"
+
+#define LOCTEXT_NAMESPACE "AvaTransitionSequenceTaskBase"
+
+TArray<UAvaSequencePlayer*> FAvaTransitionSequenceTaskBase::ExecuteSequenceTask(FStateTreeExecutionContext& InContext) const
+{
+	return TArray<UAvaSequencePlayer*>();
+}
+
+bool FAvaTransitionSequenceTaskBase::IsSequenceQueryValid(const FInstanceDataType& InInstanceData) const
+{
+	switch (InInstanceData.QueryType)
+	{
+	case EAvaTransitionSequenceQueryType::Name:
+		return !InInstanceData.SequenceName.IsNone();
+
+	case EAvaTransitionSequenceQueryType::Tag:
+		return InInstanceData.SequenceTag.IsValid();
+	}
+	checkNoEntry();
+	return false;
+}
+
+IAvaSequencePlaybackObject* FAvaTransitionSequenceTaskBase::GetPlaybackObject(FStateTreeExecutionContext& InContext) const
+{
+	const FAvaTransitionContext& TransitionContext = InContext.GetExternalData(TransitionContextHandle);
+	const UAvaSequenceSubsystem& SequenceSubsystem = InContext.GetExternalData(SequenceSubsystemHandle);
+
+	return FAvaTransitionSequenceUtils::GetPlaybackObject(TransitionContext, SequenceSubsystem);
+}
+
+TArray<TWeakObjectPtr<UAvaSequence>> FAvaTransitionSequenceTaskBase::GetActiveSequences(TConstArrayView<UAvaSequencePlayer*> InSequencePlayers) const
+{
+	TArray<TWeakObjectPtr<UAvaSequence>> OutActiveSequences;
+	OutActiveSequences.Reserve(InSequencePlayers.Num());
+
+	Algo::TransformIf(InSequencePlayers, OutActiveSequences,
+		[](UAvaSequencePlayer* InPlayer)->bool
+		{
+			return IsValid(InPlayer);
+		},
+		[](UAvaSequencePlayer* InPlayer)->TWeakObjectPtr<UAvaSequence>
+		{
+			return TWeakObjectPtr<UAvaSequence>(InPlayer->GetAvaSequence());
+		});
+
+	return OutActiveSequences;
+}
+
+EStateTreeRunStatus FAvaTransitionSequenceTaskBase::WaitForActiveSequences(FStateTreeExecutionContext& InContext) const
+{
+	FInstanceDataType& InstanceData = InContext.GetInstanceData(*this);
+
+	EAvaTransitionSequenceWaitType WaitType = GetWaitType(InContext);
+
+	// No wait should only succeed in Enter State, and so return failed run if active sequences is empty 
+	if (WaitType == EAvaTransitionSequenceWaitType::NoWait)
+	{
+		if (InstanceData.ActiveSequences.IsEmpty())
+		{
+			return EStateTreeRunStatus::Failed;
+		}
+		return EStateTreeRunStatus::Succeeded;
+	}
+
+	IAvaSequencePlaybackObject* PlaybackObject = GetPlaybackObject(InContext);
+	if (!PlaybackObject)
+	{
+		return EStateTreeRunStatus::Failed;
+	}
+
+	return FAvaTransitionSequenceUtils::UpdatePlayerRunStatus(*PlaybackObject, InstanceData.ActiveSequences, WaitType);
+}
+
+void FAvaTransitionSequenceTaskBase::StopActiveSequences(FStateTreeExecutionContext& InContext) const
+{
+	IAvaSequencePlaybackObject* PlaybackObject = GetPlaybackObject(InContext);
+	if (!PlaybackObject)
+	{
+		return;
+	}
+
+	const FInstanceDataType& InstanceData = InContext.GetInstanceData(*this);
+
+	for (const TWeakObjectPtr<UAvaSequence>& SequenceWeak : InstanceData.ActiveSequences)
+	{
+		if (UAvaSequence* Sequence = SequenceWeak.Get())
+		{
+			PlaybackObject->StopSequence(Sequence);
+		}
+	}
+}
+
+FText FAvaTransitionSequenceTaskBase::GetSequenceQueryText(const FInstanceDataType& InInstanceData) const
+{
+	switch (InInstanceData.QueryType)
+	{
+	case EAvaTransitionSequenceQueryType::Name:
+		return FText::Format(INVTEXT("'{0}'"), FText::FromName(InInstanceData.SequenceName));
+
+	case EAvaTransitionSequenceQueryType::Tag:
+		return FText::Format(LOCTEXT("SequenceQueryTag", "tag '{0}'"), FText::FromName(InInstanceData.SequenceTag.ToName()));
+	}
+
+	checkNoEntry();
+	return FText::GetEmpty();
+}
+
+EStateTreeRunStatus FAvaTransitionSequenceTaskBase::EnterState(FStateTreeExecutionContext& InContext, const FStateTreeTransitionResult& InTransition) const
+{
+	FInstanceDataType& InstanceData = InContext.GetInstanceData(*this);
+
+	if (!IsSequenceQueryValid(InstanceData))
+	{
+		return EStateTreeRunStatus::Failed;
+	}
+
+	TArray<UAvaSequencePlayer*> SequencePlayers = ExecuteSequenceTask(InContext);
+	if (SequencePlayers.IsEmpty())
+	{
+		return EStateTreeRunStatus::Succeeded;
+	}
+
+	InstanceData.ActiveSequences = FAvaTransitionSequenceUtils::GetSequences(SequencePlayers);
+	return WaitForActiveSequences(InContext);
+}
+
+EStateTreeRunStatus FAvaTransitionSequenceTaskBase::Tick(FStateTreeExecutionContext& InContext, const float InDeltaTime) const
+{
+	return WaitForActiveSequences(InContext);
+}
+
+void FAvaTransitionSequenceTaskBase::ExitState(FStateTreeExecutionContext& InContext, const FStateTreeTransitionResult& InTransition) const
+{
+	if (InTransition.CurrentRunStatus == EStateTreeRunStatus::Stopped)
+	{
+		StopActiveSequences(InContext);
+	}
+}
+
+void FAvaTransitionSequenceTaskBase::PostLoad(FStateTreeDataView InInstanceDataView)
+{
+	Super::PostLoad(InInstanceDataView);
+
+	PRAGMA_DISABLE_DEPRECATION_WARNINGS
+	if (QueryType_DEPRECATED != EAvaTransitionSequenceQueryType::None)
+	{
+		if (FInstanceDataType* InstanceData = UE::AvaTransition::TryGetInstanceData(*this, InInstanceDataView))
+		{
+			InstanceData->QueryType    = QueryType_DEPRECATED;
+			InstanceData->SequenceName = SequenceName_DEPRECATED;
+			InstanceData->SequenceTag  = SequenceTag_DEPRECATED;
+		}
+	}
+	PRAGMA_ENABLE_DEPRECATION_WARNINGS
+}
+
+bool FAvaTransitionSequenceTaskBase::Link(FStateTreeLinker& InLinker)
+{
+	Super::Link(InLinker);
+	InLinker.LinkExternalData(SequenceSubsystemHandle);
+	return true;
+}
+
+#undef LOCTEXT_NAMESPACE
diff --git a/Engine/Plugins/Experimental/Avalanche/Source/AvalancheSequence/Private/Transition/AvaTransitionWaitForAllSequencesTask.cpp b/Engine/Plugins/Experimental/Avalanche/Source/AvalancheSequence/Private/Transition/AvaTransitionWaitForAllSequencesTask.cpp
--- a/Engine/Plugins/Experimental/Avalanche/Source/AvalancheSequence/Private/Transition/AvaTransitionWaitForAllSequencesTask.cpp
+++ b/Engine/Plugins/Experimental/Avalanche/Source/AvalancheSequence/Private/Transition/AvaTransitionWaitForAllSequencesTask.cpp
@@ -34,7 +34,7 @@ EStateTreeRunStatus FAvaTransitionWaitForAllSequencesTask::Tick(FStateTreeExecut
 
 bool FAvaTransitionWaitForAllSequencesTask::Link(FStateTreeLinker& InLinker)
 {
-	FAvaTransitionTask::Link(InLinker);
+	Super::Link(InLinker);
 	InLinker.LinkExternalData(SequenceSubsystemHandle);
 	return true;
 }
@@ -54,6 +54,6 @@ EStateTreeRunStatus FAvaTransitionWaitForAllSequencesTask::WaitForAllSequences(F
 		return EStateTreeRunStatus::Failed;
 	}
 
-	FAvaTransitionSequenceInstanceData& InstanceData = InContext.GetInstanceData(*this);
+	FInstanceDataType& InstanceData = InContext.GetInstanceData(*this);
 	return FAvaTransitionSequenceUtils::UpdatePlayerRunStatus(*PlaybackObject, InstanceData.ActiveSequences, EAvaTransitionSequenceWaitType::WaitUntilStop);
 }
diff --git a/Engine/Plugins/Experimental/Avalanche/Source/AvalancheSequence/Private/Transition/AvaTransitionWaitForSequenceTask.cpp b/Engine/Plugins/Experimental/Avalanche/Source/AvalancheSequence/Private/Transition/AvaTransitionWaitForSequenceTask.cpp
--- a/Engine/Plugins/Experimental/Avalanche/Source/AvalancheSequence/Private/Transition/AvaTransitionWaitForSequenceTask.cpp
+++ b/Engine/Plugins/Experimental/Avalanche/Source/AvalancheSequence/Private/Transition/AvaTransitionWaitForSequenceTask.cpp
@@ -13,15 +13,17 @@ TArray<UAvaSequencePlayer*> FAvaTransitionWaitForSequenceTask::ExecuteSequenceTa
 		return TArray<UAvaSequencePlayer*>();
 	}
 
-	switch (QueryType)
+	const FInstanceDataType& InstanceData = InContext.GetInstanceData(*this);
+
+	switch (InstanceData.QueryType)
 	{
 	case EAvaTransitionSequenceQueryType::Name:
-		return PlaybackObject->GetSequencePlayersByLabel(SequenceName);
+		return PlaybackObject->GetSequencePlayersByLabel(InstanceData.SequenceName);
 
 	case EAvaTransitionSequenceQueryType::Tag:
-		if (const FAvaTag* Tag = SequenceTag.GetTag())
+		if (const FAvaTag* Tag = InstanceData.SequenceTag.GetTag())
 		{
-			return PlaybackObject->GetSequencePlayersByTag(*Tag, bPerformExactMatch);
+			return PlaybackObject->GetSequencePlayersByTag(*Tag, InstanceData.bPerformExactMatch);
 		}
 		return TArray<UAvaSequencePlayer*>();
 	}
diff --git a/Engine/Plugins/Experimental/Avalanche/Source/AvalancheSequence/Public/Transition/AvaTransitionInitializeSequence.h b/Engine/Plugins/Experimental/Avalanche/Source/AvalancheSequence/Public/Transition/AvaTransitionInitializeSequence.h
--- a/Engine/Plugins/Experimental/Avalanche/Source/AvalancheSequence/Public/Transition/AvaTransitionInitializeSequence.h
+++ b/Engine/Plugins/Experimental/Avalanche/Source/AvalancheSequence/Public/Transition/AvaTransitionInitializeSequence.h
@@ -6,19 +6,49 @@
 #include "AvaTransitionSequenceTask.h"
 #include "AvaTransitionInitializeSequence.generated.h"
 
-USTRUCT(DisplayName = "Initialize Sequence", Category="Sequence Playback")
+USTRUCT()
+struct FAvaTransitionInitSequenceTaskInstanceData : public FAvaTransitionSequenceTaskBaseInstanceData
+{
+	GENERATED_BODY()
+
+	UPROPERTY(EditAnywhere, Category = "Parameter")
+	FAvaSequenceTime InitializeTime = FAvaSequenceTime(0.0);
+
+	UPROPERTY(EditAnywhere, Category = "Parameter")
+	EAvaSequencePlayMode PlayMode = EAvaSequencePlayMode::Forward;
+};
+
+USTRUCT(DisplayName = "Initialize Sequence", Category = "Sequence Playback")
 struct AVALANCHESEQUENCE_API FAvaTransitionInitializeSequence : public FAvaTransitionSequenceTaskBase
 {
 	GENERATED_BODY()
 
+	using FInstanceDataType = FAvaTransitionInitSequenceTaskInstanceData;
+
+	PRAGMA_DISABLE_DEPRECATION_WARNINGS
+	FAvaTransitionInitializeSequence() = default;
+	virtual ~FAvaTransitionInitializeSequence() override = default;
+	FAvaTransitionInitializeSequence(const FAvaTransitionInitializeSequence&) = default;
+	FAvaTransitionInitializeSequence(FAvaTransitionInitializeSequence&&) = default;
+	FAvaTransitionInitializeSequence& operator=(const FAvaTransitionInitializeSequence&) = default;
+	FAvaTransitionInitializeSequence& operator=(FAvaTransitionInitializeSequence&&) = default;
+	PRAGMA_ENABLE_DEPRECATION_WARNINGS
+
+	//~ Begin FStateTreeNodeBase
+	virtual const UStruct* GetInstanceDataType() const override { return FInstanceDataType::StaticStruct(); }
+	virtual void PostLoad(FStateTreeDataView InInstanceDataView) override;
+	//~ End FStateTreeNodeBase
+
 	//~ Begin FAvaTransitionSequenceTaskBase
-	virtual EAvaTransitionSequenceWaitType GetWaitType() const { return EAvaTransitionSequenceWaitType::NoWait; }
+	virtual EAvaTransitionSequenceWaitType GetWaitType(FStateTreeExecutionContext& InContext) const override;
 	virtual TArray<UAvaSequencePlayer*> ExecuteSequenceTask(FStateTreeExecutionContext& InContext) const override;
 	//~ End FAvaTransitionSequenceTaskBase
 
-	UPROPERTY(EditAnywhere, Category = "Motion Design Sequence")
-	FAvaSequenceTime InitializeTime = FAvaSequenceTime(0.0);
+	UE_DEPRECATED(5.5, "InitializeTime has been moved to Instance Data")
+	UPROPERTY(meta = (DeprecatedProperty, DeprecationMessage = "Use the Instance Data InitializeTime instead"))
+	FAvaSequenceTime InitializeTime_DEPRECATED = FAvaSequenceTime(0.0);
 
-	UPROPERTY(EditAnywhere, Category="Parameter")
-	EAvaSequencePlayMode PlayMode = EAvaSequencePlayMode::Forward;
+	UE_DEPRECATED(5.5, "PlayMode has been moved to Instance Data")
+	UPROPERTY(meta = (DeprecatedProperty, DeprecationMessage = "Use the Instance Data PlayMode instead"))
+	EAvaSequencePlayMode PlayMode_DEPRECATED = EAvaSequencePlayMode::Forward;
 };
diff --git a/Engine/Plugins/Experimental/Avalanche/Source/AvalancheSequence/Public/Transition/AvaTransitionPlaySequenceTask.h b/Engine/Plugins/Experimental/Avalanche/Source/AvalancheSequence/Public/Transition/AvaTransitionPlaySequenceTask.h
--- a/Engine/Plugins/Experimental/Avalanche/Source/AvalancheSequence/Public/Transition/AvaTransitionPlaySequenceTask.h
+++ b/Engine/Plugins/Experimental/Avalanche/Source/AvalancheSequence/Public/Transition/AvaTransitionPlaySequenceTask.h
@@ -6,12 +6,35 @@
 #include "AvaTransitionSequenceTask.h"
 #include "AvaTransitionPlaySequenceTask.generated.h"
 
+USTRUCT()
+struct FAvaTransitionPlaySequenceTaskInstanceData : public FAvaTransitionSequenceTaskInstanceData
+{
+	GENERATED_BODY()
+
+	/** Sequence Play Settings */
+	UPROPERTY(EditAnywhere, Category="Parameter")
+	FAvaSequencePlayParams PlaySettings;
+};
+
 USTRUCT(DisplayName = "Play Sequence", Category="Sequence Playback")
 struct AVALANCHESEQUENCE_API FAvaTransitionPlaySequenceTask : public FAvaTransitionSequenceTask
 {
 	GENERATED_BODY()
 
+	using FInstanceDataType = FAvaTransitionPlaySequenceTaskInstanceData;
+
+	PRAGMA_DISABLE_DEPRECATION_WARNINGS
+	FAvaTransitionPlaySequenceTask() = default;
+	virtual ~FAvaTransitionPlaySequenceTask() override = default;
+	FAvaTransitionPlaySequenceTask(const FAvaTransitionPlaySequenceTask&) = default;
+	FAvaTransitionPlaySequenceTask(FAvaTransitionPlaySequenceTask&&) = default;
+	FAvaTransitionPlaySequenceTask& operator=(const FAvaTransitionPlaySequenceTask&) = default;
+	FAvaTransitionPlaySequenceTask& operator=(FAvaTransitionPlaySequenceTask&&) = default;
+	PRAGMA_ENABLE_DEPRECATION_WARNINGS
+
 	//~ Begin FStateTreeNodeBase
+	virtual const UStruct* GetInstanceDataType() const override { return FInstanceDataType::StaticStruct(); }
+	virtual void PostLoad(FStateTreeDataView InInstanceDataView) override;
 #if WITH_EDITOR
 	virtual FText GetDescription(const FGuid& InId, FStateTreeDataView InInstanceDataView, const IStateTreeBindingLookup& InBindingLookup, EStateTreeNodeFormatting InFormatting) const override;
 #endif
@@ -21,7 +44,7 @@ struct AVALANCHESEQUENCE_API FAvaTransitionPlaySequenceTask : public FAvaTransit
 	virtual TArray<UAvaSequencePlayer*> ExecuteSequenceTask(FStateTreeExecutionContext& InContext) const override;
 	//~ End FAvaTransitionSequenceTaskBase
 
-	/** Sequence Play Settings */
-	UPROPERTY(EditAnywhere, Category="Parameter")
-	FAvaSequencePlayParams PlaySettings;
+	UE_DEPRECATED(5.5, "PlaySettings has been moved to Instance Data")
+	UPROPERTY(meta = (DeprecatedProperty, DeprecationMessage = "Use the Instance Data PlaySettings instead"))
+	FAvaSequencePlayParams PlaySettings_DEPRECATED;
 };
diff --git a/Engine/Plugins/Experimental/Avalanche/Source/AvalancheSequence/Public/Transition/AvaTransitionSequenceTask.h b/Engine/Plugins/Experimental/Avalanche/Source/AvalancheSequence/Public/Transition/AvaTransitionSequenceTask.h
--- a/Engine/Plugins/Experimental/Avalanche/Source/AvalancheSequence/Public/Transition/AvaTransitionSequenceTask.h
+++ b/Engine/Plugins/Experimental/Avalanche/Source/AvalancheSequence/Public/Transition/AvaTransitionSequenceTask.h
@@ -2,105 +2,46 @@
 
 #pragma once
 
-#include "AvaSequenceShared.h"
-#include "AvaTagHandle.h"
-#include "AvaTransitionSequenceEnums.h"
-#include "Tasks/AvaTransitionTask.h"
+#include "AvaTransitionSequenceTaskBase.h"
 #include "AvaTransitionSequenceTask.generated.h"
 
-class IAvaSequencePlaybackObject;
-class UAvaSequence;
-class UAvaSequencePlayer;
-class UAvaSequenceSubsystem;
-
 USTRUCT()
-struct FAvaTransitionSequenceInstanceData
+struct FAvaTransitionSequenceTaskInstanceData : public FAvaTransitionSequenceTaskBaseInstanceData
 {
 	GENERATED_BODY()
 
-	/** Sequences being played */
-	UPROPERTY()
-	TArray<TWeakObjectPtr<UAvaSequence>> ActiveSequences;
+	/** The wait type before this task completes */
+	UPROPERTY(EditAnywhere, Category="Parameter")
+	EAvaTransitionSequenceWaitType WaitType = EAvaTransitionSequenceWaitType::WaitUntilStop;
 };
 
+/** Base Task but with additional Parameters */
 USTRUCT(meta=(Hidden))
-struct AVALANCHESEQUENCE_API FAvaTransitionSequenceTaskBase : public FAvaTransitionTask
+struct AVALANCHESEQUENCE_API FAvaTransitionSequenceTask : public FAvaTransitionSequenceTaskBase
 {
 	GENERATED_BODY()
 
-	using FInstanceDataType = FAvaTransitionSequenceInstanceData;
-
-	/**
-	 * Execute the Sequence Task (overriden by implementation)
-	 * @return the sequence players that are relevant to the task
-	 */
-	virtual TArray<UAvaSequencePlayer*> ExecuteSequenceTask(FStateTreeExecutionContext& InContext) const;
-
-	/** Gets the Wait Type to use when Waiting for Active Sequences */
-	virtual EAvaTransitionSequenceWaitType GetWaitType() const { return EAvaTransitionSequenceWaitType::None; }
+	using FInstanceDataType = FAvaTransitionSequenceTaskInstanceData;
 
-	/** Determines whether the Current Sequence information is valid for query */
-	bool IsSequenceQueryValid() const;
-
-	/** Attempts to Retrieve the Playback Object from the given Execution Context */
-	IAvaSequencePlaybackObject* GetPlaybackObject(FStateTreeExecutionContext& InContext) const;
-
-	/** Gets all the Sequences from the provided Sequence Players that are Active (playing) */
-	TArray<TWeakObjectPtr<UAvaSequence>> GetActiveSequences(TConstArrayView<UAvaSequencePlayer*> InSequencePlayers) const;
-
-	/**
-	 * Helper function to determine the Tree Run Status by updating and checking if all activated Sequence Players
-	 * are in a state that match the Wait Type
-	 * @param InContext the execution context where the Instance Data of the Active Sequences is retrieved
-	 * @return the suggested tree run status
-	 */
-	EStateTreeRunStatus WaitForActiveSequences(FStateTreeExecutionContext& InContext) const;
-
-	/**
-	 * Helper function to stop all the currently active sequences
-	 * @param InContext the execution context where the Instance Data of the Active Sequences is retrieved
-	 */
-	void StopActiveSequences(FStateTreeExecutionContext& InContext) const;
-
-	FText GetSequenceQueryText() const;
-
-	//~ Begin FStateTreeTaskBase
-	virtual EStateTreeRunStatus EnterState(FStateTreeExecutionContext& InContext, const FStateTreeTransitionResult& InTransition) const override;
-	virtual EStateTreeRunStatus Tick(FStateTreeExecutionContext& InContext, const float InDeltaTime) const override;
-	virtual void ExitState(FStateTreeExecutionContext& InContext, const FStateTreeTransitionResult& InTransition) const override;
-	//~ End FStateTreeTaskBase
+	PRAGMA_DISABLE_DEPRECATION_WARNINGS
+	FAvaTransitionSequenceTask() = default;
+	virtual ~FAvaTransitionSequenceTask() override = default;
+	FAvaTransitionSequenceTask(const FAvaTransitionSequenceTask&) = default;
+	FAvaTransitionSequenceTask(FAvaTransitionSequenceTask&&) = default;
+	FAvaTransitionSequenceTask& operator=(const FAvaTransitionSequenceTask&) = default;
+	FAvaTransitionSequenceTask& operator=(FAvaTransitionSequenceTask&&) = default;
+	PRAGMA_ENABLE_DEPRECATION_WARNINGS
 
 	//~ Begin FStateTreeNodeBase
 	virtual const UStruct* GetInstanceDataType() const override { return FInstanceDataType::StaticStruct(); }
-	virtual bool Link(FStateTreeLinker& InLinker) override;
+	virtual void PostLoad(FStateTreeDataView InInstanceDataView) override;
 	//~ End FStateTreeNodeBase
 
-	UPROPERTY(EditAnywhere, Category="Parameter", DisplayName="Sequence Query Type", meta=(DisplayPriority=1))
-	EAvaTransitionSequenceQueryType QueryType = EAvaTransitionSequenceQueryType::Name;
-
-	UPROPERTY(EditAnywhere, Category="Parameter", meta=(DisplayPriority=1, EditCondition="QueryType==EAvaTransitionSequenceQueryType::Name", EditConditionHides))
-	FName SequenceName;
-
-	UPROPERTY(EditAnywhere, Category="Parameter", meta=(DisplayPriority=1, EditCondition="QueryType==EAvaTransitionSequenceQueryType::Tag", EditConditionHides))
-	FAvaTagHandle SequenceTag;
-
-	UPROPERTY()
-	bool bPerformExactMatch = false;
-
-	TStateTreeExternalDataHandle<UAvaSequenceSubsystem> SequenceSubsystemHandle;
-};
-
-/** Base Task but with additional Parameters */
-USTRUCT(meta=(Hidden))
-struct AVALANCHESEQUENCE_API FAvaTransitionSequenceTask : public FAvaTransitionSequenceTaskBase
-{
-	GENERATED_BODY()
-
 	//~ Begin FAvaTransitionSequenceTaskBase
-	virtual EAvaTransitionSequenceWaitType GetWaitType() const override { return WaitType; }
+	virtual EAvaTransitionSequenceWaitType GetWaitType(FStateTreeExecutionContext& InContext) const override;
 	//~ End FAvaTransitionSequenceTaskBase
 
-	/** The wait type before this task completes */
-	UPROPERTY(EditAnywhere, Category="Parameter")
-	EAvaTransitionSequenceWaitType WaitType = EAvaTransitionSequenceWaitType::WaitUntilStop;
+	UE_DEPRECATED(5.5, "WaitType has been moved to Instance Data")
+	UPROPERTY(meta = (DeprecatedProperty, DeprecationMessage = "Use the Instance Data WaitType instead"))
+	EAvaTransitionSequenceWaitType WaitType_DEPRECATED = EAvaTransitionSequenceWaitType::WaitUntilStop;
 };
diff --git a/Engine/Plugins/Experimental/Avalanche/Source/AvalancheSequence/Public/Transition/AvaTransitionSequenceTaskBase.h b/Engine/Plugins/Experimental/Avalanche/Source/AvalancheSequence/Public/Transition/AvaTransitionSequenceTaskBase.h
new file mode 100644
--- /dev/null
+++ b/Engine/Plugins/Experimental/Avalanche/Source/AvalancheSequence/Public/Transition/AvaTransitionSequenceTaskBase.h
@@ -0,0 +1,122 @@
+// Copyright Epic Games, Inc. All Rights Reserved.
+
+#pragma once
+
+#include "AvaSequenceShared.h"
+#include "AvaTagHandle.h"
+#include "AvaTransitionSequenceEnums.h"
+#include "Tasks/AvaTransitionTask.h"
+#include "AvaTransitionSequenceTaskBase.generated.h"
+
+class IAvaSequencePlaybackObject;
+class UAvaSequence;
+class UAvaSequencePlayer;
+class UAvaSequenceSubsystem;
+
+USTRUCT()
+struct FAvaTransitionSequenceInstanceData
+{
+	GENERATED_BODY()
+
+	/** Sequences being played */
+	UPROPERTY()
+	TArray<TWeakObjectPtr<UAvaSequence>> ActiveSequences;
+};
+
+USTRUCT()
+struct FAvaTransitionSequenceTaskBaseInstanceData : public FAvaTransitionSequenceInstanceData
+{
+	GENERATED_BODY()
+
+	UPROPERTY(EditAnywhere, Category="Parameter", DisplayName="Sequence Query Type", meta=(DisplayPriority=1))
+	EAvaTransitionSequenceQueryType QueryType = EAvaTransitionSequenceQueryType::Name;
+
+	UPROPERTY(EditAnywhere, Category="Parameter", meta=(DisplayPriority=1, EditCondition="QueryType==EAvaTransitionSequenceQueryType::Name", EditConditionHides))
+	FName SequenceName;
+
+	UPROPERTY(EditAnywhere, Category="Parameter", meta=(DisplayPriority=1, EditCondition="QueryType==EAvaTransitionSequenceQueryType::Tag", EditConditionHides))
+	FAvaTagHandle SequenceTag;
+
+	UPROPERTY()
+	bool bPerformExactMatch = false;
+};
+
+USTRUCT(meta=(Hidden))
+struct AVALANCHESEQUENCE_API FAvaTransitionSequenceTaskBase : public FAvaTransitionTask
+{
+	GENERATED_BODY()
+
+	PRAGMA_DISABLE_DEPRECATION_WARNINGS
+	FAvaTransitionSequenceTaskBase() = default;
+	virtual ~FAvaTransitionSequenceTaskBase() override = default;
+	FAvaTransitionSequenceTaskBase(const FAvaTransitionSequenceTaskBase&) = default;
+	FAvaTransitionSequenceTaskBase(FAvaTransitionSequenceTaskBase&&) = default;
+	FAvaTransitionSequenceTaskBase& operator=(const FAvaTransitionSequenceTaskBase&) = default;
+	FAvaTransitionSequenceTaskBase& operator=(FAvaTransitionSequenceTaskBase&&) = default;
+	PRAGMA_ENABLE_DEPRECATION_WARNINGS
+
+	using FInstanceDataType = FAvaTransitionSequenceTaskBaseInstanceData;
+
+	/**
+	 * Execute the Sequence Task (overriden by implementation)
+	 * @return the sequence players that are relevant to the task
+	 */
+	virtual TArray<UAvaSequencePlayer*> ExecuteSequenceTask(FStateTreeExecutionContext& InContext) const;
+
+	/** Gets the Wait Type to use when Waiting for Active Sequences */
+	virtual EAvaTransitionSequenceWaitType GetWaitType(FStateTreeExecutionContext& InContext) const
+	{
+		return EAvaTransitionSequenceWaitType::None;
+	}
+
+	/** Determines whether the Current Sequence information is valid for query */
+	bool IsSequenceQueryValid(const FInstanceDataType& InInstanceData) const;
+
+	/** Attempts to Retrieve the Playback Object from the given Execution Context */
+	IAvaSequencePlaybackObject* GetPlaybackObject(FStateTreeExecutionContext& InContext) const;
+
+	/** Gets all the Sequences from the provided Sequence Players that are Active (playing) */
+	TArray<TWeakObjectPtr<UAvaSequence>> GetActiveSequences(TConstArrayView<UAvaSequencePlayer*> InSequencePlayers) const;
+
+	/**
+	 * Helper function to determine the Tree Run Status by updating and checking if all activated Sequence Players
+	 * are in a state that match the Wait Type
+	 * @param InContext the execution context where the Instance Data of the Active Sequences is retrieved
+	 * @return the suggested tree run status
+	 */
+	EStateTreeRunStatus WaitForActiveSequences(FStateTreeExecutionContext& InContext) const;
+
+	/**
+	 * Helper function to stop all the currently active sequences
+	 * @param InContext the execution context where the Instance Data of the Active Sequences is retrieved
+	 */
+	void StopActiveSequences(FStateTreeExecutionContext& InContext) const;
+
+	FText GetSequenceQueryText(const FInstanceDataType& InInstanceData) const;
+
+	//~ Begin FStateTreeTaskBase
+	virtual EStateTreeRunStatus EnterState(FStateTreeExecutionContext& InContext, const FStateTreeTransitionResult& InTransition) const override;
+	virtual EStateTreeRunStatus Tick(FStateTreeExecutionContext& InContext, const float InDeltaTime) const override;
+	virtual void ExitState(FStateTreeExecutionContext& InContext, const FStateTreeTransitionResult& InTransition) const override;
+	//~ End FStateTreeTaskBase
+
+	//~ Begin FStateTreeNodeBase
+	virtual const UStruct* GetInstanceDataType() const override { return FInstanceDataType::StaticStruct(); }
+	virtual void PostLoad(FStateTreeDataView InInstanceDataView) override;
+	virtual bool Link(FStateTreeLinker& InLinker) override;
+	//~ End FStateTreeNodeBase
+
+	UE_DEPRECATED(5.5, "QueryType has been moved to Instance Data")
+    UPROPERTY(meta = (DeprecatedProperty, DeprecationMessage = "Use the Instance Data QueryType instead"))
+	EAvaTransitionSequenceQueryType QueryType_DEPRECATED = EAvaTransitionSequenceQueryType::None;
+
+	UE_DEPRECATED(5.5, "SequenceName has been moved to Instance Data")
+	UPROPERTY(meta = (DeprecatedProperty, DeprecationMessage = "Use the Instance Data SequenceName instead"))
+	FName SequenceName_DEPRECATED;
+
+	UE_DEPRECATED(5.5, "SequenceTag has been moved to Instance Data")
+	UPROPERTY(meta = (DeprecatedProperty, DeprecationMessage = "Use the Instance Data SequenceTag instead"))
+	FAvaTagHandle SequenceTag_DEPRECATED;
+
+	TStateTreeExternalDataHandle<UAvaSequenceSubsystem> SequenceSubsystemHandle;
+};
diff --git a/Engine/Plugins/Experimental/Avalanche/Source/AvalancheSequencer/Private/Transition/AvaSequenceTransitionCompiler.cpp b/Engine/Plugins/Experimental/Avalanche/Source/AvalancheSequencer/Private/Transition/AvaSequenceTransitionCompiler.cpp
--- a/Engine/Plugins/Experimental/Avalanche/Source/AvalancheSequencer/Private/Transition/AvaSequenceTransitionCompiler.cpp
+++ b/Engine/Plugins/Experimental/Avalanche/Source/AvalancheSequencer/Private/Transition/AvaSequenceTransitionCompiler.cpp
@@ -80,9 +80,9 @@ namespace UE::AvaSequencer::Private
 		{
 			if (PlayTaskNode)
 			{
-				const FAvaTransitionPlaySequenceTask& PlayTask = PlayTaskNode->Node.Get<FAvaTransitionPlaySequenceTask>();
+				const FAvaTransitionPlaySequenceTask::FInstanceDataType& PlayTask = PlayTaskNode->Instance.Get<FAvaTransitionPlaySequenceTask::FInstanceDataType>();
 
-				FAvaTransitionInitializeSequence& InitTask = InInitState.AddTask<FAvaTransitionInitializeSequence>().GetNode();
+				FAvaTransitionInitializeSequence::FInstanceDataType& InitTask = InInitState.AddTask<FAvaTransitionInitializeSequence>().GetInstanceData();
 				InitTask.InitializeTime     = PlayTask.PlaySettings.Start;
 				InitTask.PlayMode           = PlayTask.PlaySettings.PlayMode;
 				InitTask.QueryType          = PlayTask.QueryType;
diff --git a/Engine/Plugins/Experimental/Avalanche/Source/AvalancheTransition/Private/Conditions/AvaTransitionLayerCondition.cpp b/Engine/Plugins/Experimental/Avalanche/Source/AvalancheTransition/Private/Conditions/AvaTransitionLayerCondition.cpp
--- a/Engine/Plugins/Experimental/Avalanche/Source/AvalancheTransition/Private/Conditions/AvaTransitionLayerCondition.cpp
+++ b/Engine/Plugins/Experimental/Avalanche/Source/AvalancheTransition/Private/Conditions/AvaTransitionLayerCondition.cpp
@@ -4,18 +4,41 @@
 #include "AvaTransitionContext.h"
 #include "AvaTransitionLayerUtils.h"
 #include "AvaTransitionSubsystem.h"
+#include "AvaTransitionUtils.h"
 #include "StateTreeExecutionContext.h"
 
+#define LOCTEXT_NAMESPACE "AvaTransitionLayerCondition"
+
+void FAvaTransitionLayerCondition::PostLoad(FStateTreeDataView InInstanceDataView)
+{
+	Super::PostLoad(InInstanceDataView);
+
+	PRAGMA_DISABLE_DEPRECATION_WARNINGS
+	if (LayerType_DEPRECATED != EAvaTransitionLayerCompareType::None)
+	{
+		if (FInstanceDataType* InstanceData = UE::AvaTransition::TryGetInstanceData(*this, InInstanceDataView))
+		{
+			InstanceData->LayerType     = LayerType_DEPRECATED;
+			InstanceData->SpecificLayer = SpecificLayer_DEPRECATED;
+		}
+	}
+	PRAGMA_ENABLE_DEPRECATION_WARNINGS
+}
+
 TArray<const FAvaTransitionBehaviorInstance*> FAvaTransitionLayerCondition::QueryBehaviorInstances(const FStateTreeExecutionContext& InContext) const
 {
+	const FInstanceDataType& InstanceData = InContext.GetInstanceData(*this);
+
 	UAvaTransitionSubsystem& TransitionSubsystem   = InContext.GetExternalData(TransitionSubsystemHandle);
 	const FAvaTransitionContext& TransitionContext = InContext.GetExternalData(TransitionContextHandle);
-	const FAvaTransitionLayerComparator Comparator = FAvaTransitionLayerUtils::BuildComparator(TransitionContext, LayerType, SpecificLayer);
+	const FAvaTransitionLayerComparator Comparator = FAvaTransitionLayerUtils::BuildComparator(TransitionContext, InstanceData.LayerType, InstanceData.SpecificLayer);
 
 	return FAvaTransitionLayerUtils::QueryBehaviorInstances(TransitionSubsystem, Comparator);
 }
 
-FText FAvaTransitionLayerCondition::GetLayerQueryText() const
+FText FAvaTransitionLayerCondition::GetLayerQueryText(const FInstanceDataType& InInstanceData) const
 {
-	return FAvaTransitionLayerUtils::GetLayerQueryText(LayerType, SpecificLayer.ToName());
+	return FAvaTransitionLayerUtils::GetLayerQueryText(InInstanceData.LayerType, InInstanceData.SpecificLayer.ToName());
 }
+
+#undef LOCTEXT_NAMESPACE
diff --git a/Engine/Plugins/Experimental/Avalanche/Source/AvalancheTransition/Private/Conditions/AvaTransitionLayerMatchCondition.cpp b/Engine/Plugins/Experimental/Avalanche/Source/AvalancheTransition/Private/Conditions/AvaTransitionLayerMatchCondition.cpp
--- a/Engine/Plugins/Experimental/Avalanche/Source/AvalancheTransition/Private/Conditions/AvaTransitionLayerMatchCondition.cpp
+++ b/Engine/Plugins/Experimental/Avalanche/Source/AvalancheTransition/Private/Conditions/AvaTransitionLayerMatchCondition.cpp
@@ -10,15 +10,17 @@
 #if WITH_EDITOR
 FText FAvaTransitionLayerMatchCondition::GetDescription(const FGuid& InId, FStateTreeDataView InInstanceDataView, const IStateTreeBindingLookup& InBindingLookup, EStateTreeNodeFormatting InFormatting) const
 {
-	return FText::Format(LOCTEXT("ConditionDescription", "scenes transitioning in {0}"), GetLayerQueryText());
+	const FInstanceDataType& InstanceData = InInstanceDataView.Get<FInstanceDataType>();
+	return FText::Format(LOCTEXT("ConditionDescription", "scenes transitioning in {0}"), GetLayerQueryText(InstanceData));
 }
 #endif
 
 bool FAvaTransitionLayerMatchCondition::TestCondition(FStateTreeExecutionContext& InContext) const
 {
 	const FAvaTransitionContext& TransitionContext = InContext.GetExternalData(TransitionContextHandle);
+	const FInstanceDataType& InstanceData = InContext.GetInstanceData(*this);
 
-	if (LayerType == EAvaTransitionLayerCompareType::Same && TransitionContext.GetTransitionType() == EAvaTransitionType::In)
+	if (InstanceData.LayerType == EAvaTransitionLayerCompareType::Same && TransitionContext.GetTransitionType() == EAvaTransitionType::In)
 	{
 		return true;
 	}
diff --git a/Engine/Plugins/Experimental/Avalanche/Source/AvalancheTransition/Private/Conditions/AvaTransitionSceneMatchCondition.cpp b/Engine/Plugins/Experimental/Avalanche/Source/AvalancheTransition/Private/Conditions/AvaTransitionSceneMatchCondition.cpp
--- a/Engine/Plugins/Experimental/Avalanche/Source/AvalancheTransition/Private/Conditions/AvaTransitionSceneMatchCondition.cpp
+++ b/Engine/Plugins/Experimental/Avalanche/Source/AvalancheTransition/Private/Conditions/AvaTransitionSceneMatchCondition.cpp
@@ -3,6 +3,7 @@
 #include "Conditions/AvaTransitionSceneMatchCondition.h"
 #include "AvaTransitionContext.h"
 #include "AvaTransitionScene.h"
+#include "AvaTransitionUtils.h"
 #include "Behavior/AvaTransitionBehaviorInstance.h"
 #include "StateTreeExecutionContext.h"
 
@@ -11,12 +12,29 @@
 #if WITH_EDITOR
 FText FAvaTransitionSceneMatchCondition::GetDescription(const FGuid& InId, FStateTreeDataView InInstanceDataView, const IStateTreeBindingLookup& InBindingLookup, EStateTreeNodeFormatting InFormatting) const
 {
+	const FInstanceDataType& InstanceData = InInstanceDataView.Get<FInstanceDataType>();
+
 	return FText::Format(LOCTEXT("ConditionDescription", "{0} scene in {1}")
-		, UEnum::GetDisplayValueAsText(SceneComparisonType).ToLower()
-		, GetLayerQueryText());
+		, UEnum::GetDisplayValueAsText(InstanceData.SceneComparisonType).ToLower()
+		, GetLayerQueryText(InstanceData));
 }
 #endif
 
+void FAvaTransitionSceneMatchCondition::PostLoad(FStateTreeDataView InInstanceDataView)
+{
+	Super::PostLoad(InInstanceDataView);
+
+	PRAGMA_DISABLE_DEPRECATION_WARNINGS
+	if (LayerType_DEPRECATED != EAvaTransitionLayerCompareType::None)
+	{
+		if (FInstanceDataType* InstanceData = UE::AvaTransition::TryGetInstanceData(*this, InInstanceDataView))
+		{
+			InstanceData->SceneComparisonType = SceneComparisonType_DEPRECATED;
+		}
+	}
+	PRAGMA_ENABLE_DEPRECATION_WARNINGS
+}
+
 bool FAvaTransitionSceneMatchCondition::TestCondition(FStateTreeExecutionContext& InContext) const
 {
 	const FAvaTransitionContext& TransitionContext = InContext.GetExternalData(TransitionContextHandle);
@@ -33,8 +51,10 @@ bool FAvaTransitionSceneMatchCondition::TestCondition(FStateTreeExecutionContext
 		return false;
 	}
 
+	const FInstanceDataType& InstanceData = InContext.GetInstanceData(*this);
+
 	bool bHasMatchingInstance = BehaviorInstances.ContainsByPredicate(
-		[TransitionScene, this](const FAvaTransitionBehaviorInstance* InInstance)
+		[TransitionScene, &InstanceData](const FAvaTransitionBehaviorInstance* InInstance)
 		{
 			const FAvaTransitionScene* OtherTransitionScene = InInstance->GetTransitionContext().GetTransitionScene();
 
@@ -42,7 +62,7 @@ bool FAvaTransitionSceneMatchCondition::TestCondition(FStateTreeExecutionContext
 				? TransitionScene->Compare(*OtherTransitionScene)
 				: EAvaTransitionComparisonResult::None;
 
-			return ComparisonResult == SceneComparisonType;
+			return ComparisonResult == InstanceData.SceneComparisonType;
 		});
 
 	return bHasMatchingInstance;
diff --git a/Engine/Plugins/Experimental/Avalanche/Source/AvalancheTransition/Private/Conditions/AvaTransitionStateMatchCondition.cpp b/Engine/Plugins/Experimental/Avalanche/Source/AvalancheTransition/Private/Conditions/AvaTransitionStateMatchCondition.cpp
--- a/Engine/Plugins/Experimental/Avalanche/Source/AvalancheTransition/Private/Conditions/AvaTransitionStateMatchCondition.cpp
+++ b/Engine/Plugins/Experimental/Avalanche/Source/AvalancheTransition/Private/Conditions/AvaTransitionStateMatchCondition.cpp
@@ -1,19 +1,38 @@
 // Copyright Epic Games, Inc. All Rights Reserved.
 
 #include "Conditions/AvaTransitionStateMatchCondition.h"
+#include "AvaTransitionUtils.h"
 #include "Behavior/AvaTransitionBehaviorInstance.h"
+#include "StateTreeExecutionContext.h"
 
 #define LOCTEXT_NAMESPACE "AvaTransitionSceneMatchCondition"
 
 #if WITH_EDITOR
 FText FAvaTransitionStateMatchCondition::GetDescription(const FGuid& InId, FStateTreeDataView InInstanceDataView, const IStateTreeBindingLookup& InBindingLookup, EStateTreeNodeFormatting InFormatting) const
 {
+	const FInstanceDataType& InstanceData = InInstanceDataView.Get<FInstanceDataType>();
+
 	return FText::Format(LOCTEXT("ConditionDescription", "{0} scene in {1}")
-		, UEnum::GetDisplayValueAsText(TransitionState).ToLower()
-		, GetLayerQueryText());
+		, UEnum::GetDisplayValueAsText(InstanceData.TransitionState).ToLower()
+		, GetLayerQueryText(InstanceData));
 }
 #endif
 
+void FAvaTransitionStateMatchCondition::PostLoad(FStateTreeDataView InInstanceDataView)
+{
+	Super::PostLoad(InInstanceDataView);
+
+	PRAGMA_DISABLE_DEPRECATION_WARNINGS
+	if (TransitionState_DEPRECATED != EAvaTransitionRunState::Unknown)
+	{
+		if (FInstanceDataType* InstanceData = UE::AvaTransition::TryGetInstanceData(*this, InInstanceDataView))
+		{
+			InstanceData->TransitionState = TransitionState_DEPRECATED;
+		}
+	}
+	PRAGMA_ENABLE_DEPRECATION_WARNINGS
+}
+
 bool FAvaTransitionStateMatchCondition::TestCondition(FStateTreeExecutionContext& InContext) const
 {
 	TArray<const FAvaTransitionBehaviorInstance*> BehaviorInstances = QueryBehaviorInstances(InContext);
@@ -25,7 +44,9 @@ bool FAvaTransitionStateMatchCondition::TestCondition(FStateTreeExecutionContext
 			return InInstance->IsRunning();
 		});
 
-	switch (TransitionState)
+	const FInstanceDataType& InstanceData = InContext.GetInstanceData(*this);
+
+	switch (InstanceData.TransitionState)
 	{
 	case EAvaTransitionRunState::Running:
 		return bIsLayerRunning;
diff --git a/Engine/Plugins/Experimental/Avalanche/Source/AvalancheTransition/Private/Conditions/AvaTransitionTypeMatchCondition.cpp b/Engine/Plugins/Experimental/Avalanche/Source/AvalancheTransition/Private/Conditions/AvaTransitionTypeMatchCondition.cpp
--- a/Engine/Plugins/Experimental/Avalanche/Source/AvalancheTransition/Private/Conditions/AvaTransitionTypeMatchCondition.cpp
+++ b/Engine/Plugins/Experimental/Avalanche/Source/AvalancheTransition/Private/Conditions/AvaTransitionTypeMatchCondition.cpp
@@ -2,6 +2,7 @@
 
 #include "Conditions/AvaTransitionTypeMatchCondition.h"
 #include "AvaTransitionContext.h"
+#include "AvaTransitionUtils.h"
 #include "StateTreeExecutionContext.h"
 
 #define LOCTEXT_NAMESPACE "AvaTransitionTypeMatchCondition"
@@ -9,15 +10,34 @@
 #if WITH_EDITOR
 FText FAvaTransitionTypeMatchCondition::GetDescription(const FGuid& InId, FStateTreeDataView InInstanceDataView, const IStateTreeBindingLookup& InBindingLookup, EStateTreeNodeFormatting InFormatting) const
 {
+	const FInstanceDataType& InstanceData = InInstanceDataView.Get<FInstanceDataType>();
+
 	return FText::Format(LOCTEXT("ConditionDescription", "transitioning {0}")
-		, UEnum::GetDisplayValueAsText(TransitionType).ToLower());
+		, UEnum::GetDisplayValueAsText(InstanceData.TransitionType).ToLower());
 }
 #endif
 
+void FAvaTransitionTypeMatchCondition::PostLoad(FStateTreeDataView InInstanceDataView)
+{
+	Super::PostLoad(InInstanceDataView);
+
+	PRAGMA_DISABLE_DEPRECATION_WARNINGS
+	if (TransitionType_DEPRECATED != EAvaTransitionType::None)
+	{
+		if (FInstanceDataType* InstanceData = UE::AvaTransition::TryGetInstanceData(*this, InInstanceDataView))
+		{
+			InstanceData->TransitionType = TransitionType_DEPRECATED;
+		}
+	}
+	PRAGMA_ENABLE_DEPRECATION_WARNINGS
+}
+
 bool FAvaTransitionTypeMatchCondition::TestCondition(FStateTreeExecutionContext& InContext) const
 {
 	const FAvaTransitionContext& TransitionContext = InContext.GetExternalData(TransitionContextHandle);
-	return TransitionContext.GetTransitionType() == TransitionType;
+	const FInstanceDataType& InstanceData = InContext.GetInstanceData(*this);
+
+	return TransitionContext.GetTransitionType() == InstanceData.TransitionType;
 }
 
 #undef LOCTEXT_NAMESPACE
diff --git a/Engine/Plugins/Experimental/Avalanche/Source/AvalancheTransition/Private/Tasks/AvaTransitionDelayTask.cpp b/Engine/Plugins/Experimental/Avalanche/Source/AvalancheTransition/Private/Tasks/AvaTransitionDelayTask.cpp
--- a/Engine/Plugins/Experimental/Avalanche/Source/AvalancheTransition/Private/Tasks/AvaTransitionDelayTask.cpp
+++ b/Engine/Plugins/Experimental/Avalanche/Source/AvalancheTransition/Private/Tasks/AvaTransitionDelayTask.cpp
@@ -1,6 +1,7 @@
 // Copyright Epic Games, Inc. All Rights Reserved.
 
 #include "Tasks/AvaTransitionDelayTask.h"
+#include "AvaTransitionUtils.h"
 #include "StateTreeExecutionContext.h"
 
 #define LOCTEXT_NAMESPACE "AvaTransitionDelayTask"
@@ -8,20 +9,36 @@
 #if WITH_EDITOR
 FText FAvaTransitionDelayTask::GetDescription(const FGuid& InId, FStateTreeDataView InInstanceDataView, const IStateTreeBindingLookup& InBindingLookup, EStateTreeNodeFormatting InFormatting) const
 {
-	return FText::Format(LOCTEXT("TaskDescription", "Delay {0} seconds"), FText::AsNumber(Duration));
+	const FInstanceDataType& InstanceData = InInstanceDataView.Get<FInstanceDataType>();
+	return FText::Format(LOCTEXT("TaskDescription", "Delay {0} seconds"), FText::AsNumber(InstanceData.Duration));
 }
 #endif
 
+void FAvaTransitionDelayTask::PostLoad(FStateTreeDataView InInstanceDataView)
+{
+	Super::PostLoad(InInstanceDataView);
+
+	PRAGMA_DISABLE_DEPRECATION_WARNINGS
+	if (Duration_DEPRECATED >= 0.f)
+	{
+		if (FInstanceDataType* InstanceData = UE::AvaTransition::TryGetInstanceData(*this, InInstanceDataView))
+		{
+			InstanceData->Duration = Duration_DEPRECATED;
+		}
+	}
+	PRAGMA_ENABLE_DEPRECATION_WARNINGS
+}
+
 EStateTreeRunStatus FAvaTransitionDelayTask::EnterState(FStateTreeExecutionContext& InContext, const FStateTreeTransitionResult& InTransition) const
 {
 	FInstanceDataType& InstanceData = InContext.GetInstanceData(*this);
-
-	InstanceData.RemainingTime = Duration;
+	InstanceData.RemainingTime = InstanceData.Duration;
 
 	if (InstanceData.RemainingTime <= 0.f)
 	{
 		return EStateTreeRunStatus::Succeeded;
 	}
+
 	return EStateTreeRunStatus::Running;
 }
 
diff --git a/Engine/Plugins/Experimental/Avalanche/Source/AvalancheTransition/Private/Tasks/AvaTransitionLayerTask.cpp b/Engine/Plugins/Experimental/Avalanche/Source/AvalancheTransition/Private/Tasks/AvaTransitionLayerTask.cpp
--- a/Engine/Plugins/Experimental/Avalanche/Source/AvalancheTransition/Private/Tasks/AvaTransitionLayerTask.cpp
+++ b/Engine/Plugins/Experimental/Avalanche/Source/AvalancheTransition/Private/Tasks/AvaTransitionLayerTask.cpp
@@ -4,18 +4,41 @@
 #include "AvaTransitionContext.h"
 #include "AvaTransitionLayerUtils.h"
 #include "AvaTransitionSubsystem.h"
+#include "AvaTransitionUtils.h"
 #include "StateTreeExecutionContext.h"
 
+#define LOCTEXT_NAMESPACE "AvaTransitionLayerTask"
+
+void FAvaTransitionLayerTask::PostLoad(FStateTreeDataView InInstanceDataView)
+{
+	Super::PostLoad(InInstanceDataView);
+
+	PRAGMA_DISABLE_DEPRECATION_WARNINGS
+	if (LayerType_DEPRECATED != EAvaTransitionLayerCompareType::None)
+	{
+		if (FInstanceDataType* InstanceData = UE::AvaTransition::TryGetInstanceData(*this, InInstanceDataView))
+		{
+			InstanceData->LayerType     = LayerType_DEPRECATED;
+			InstanceData->SpecificLayer = SpecificLayer_DEPRECATED;
+		}
+	}
+	PRAGMA_ENABLE_DEPRECATION_WARNINGS
+}
+
+FText FAvaTransitionLayerTask::GetLayerQueryText(const FInstanceDataType& InInstanceData) const
+{
+	return FAvaTransitionLayerUtils::GetLayerQueryText(InInstanceData.LayerType, InInstanceData.SpecificLayer.ToName());
+}
+
 TArray<const FAvaTransitionBehaviorInstance*> FAvaTransitionLayerTask::QueryBehaviorInstances(const FStateTreeExecutionContext& InContext) const
 {
+	const FInstanceDataType& InstanceData = InContext.GetInstanceData(*this);
+
 	UAvaTransitionSubsystem& TransitionSubsystem   = InContext.GetExternalData(TransitionSubsystemHandle);
 	const FAvaTransitionContext& TransitionContext = InContext.GetExternalData(TransitionContextHandle);
-	const FAvaTransitionLayerComparator Comparator = FAvaTransitionLayerUtils::BuildComparator(TransitionContext, LayerType, SpecificLayer);
+	const FAvaTransitionLayerComparator Comparator = FAvaTransitionLayerUtils::BuildComparator(TransitionContext, InstanceData.LayerType, InstanceData.SpecificLayer);
 
 	return FAvaTransitionLayerUtils::QueryBehaviorInstances(TransitionSubsystem, Comparator);
 }
 
-FText FAvaTransitionLayerTask::GetLayerQueryText() const
-{
-	return FAvaTransitionLayerUtils::GetLayerQueryText(LayerType, SpecificLayer.ToName());
-}
+#undef LOCTEXT_NAMESPACE
diff --git a/Engine/Plugins/Experimental/Avalanche/Source/AvalancheTransition/Private/Tasks/AvaTransitionWaitForLayerTask.cpp b/Engine/Plugins/Experimental/Avalanche/Source/AvalancheTransition/Private/Tasks/AvaTransitionWaitForLayerTask.cpp
--- a/Engine/Plugins/Experimental/Avalanche/Source/AvalancheTransition/Private/Tasks/AvaTransitionWaitForLayerTask.cpp
+++ b/Engine/Plugins/Experimental/Avalanche/Source/AvalancheTransition/Private/Tasks/AvaTransitionWaitForLayerTask.cpp
@@ -1,6 +1,7 @@
 // Copyright Epic Games, Inc. All Rights Reserved.
 
 #include "Tasks/AvaTransitionWaitForLayerTask.h"
+#include "AvaTransitionLayerUtils.h"
 #include "Behavior/AvaTransitionBehaviorInstance.h"
 #include "StateTreeExecutionContext.h"
 
@@ -9,7 +10,8 @@
 #if WITH_EDITOR
 FText FAvaTransitionWaitForLayerTask::GetDescription(const FGuid& InId, FStateTreeDataView InInstanceDataView, const IStateTreeBindingLookup& InBindingLookup, EStateTreeNodeFormatting InFormatting) const
 {
-	return FText::Format(LOCTEXT("TaskDescription", "Wait for others in {0} to finish"), GetLayerQueryText());
+	const FInstanceDataType& InstanceData = InInstanceDataView.Get<FInstanceDataType>();
+	return FText::Format(LOCTEXT("TaskDescription", "Wait for others in {0} to finish"), GetLayerQueryText(InstanceData));
 }
 #endif
 
diff --git a/Engine/Plugins/Experimental/Avalanche/Source/AvalancheTransition/Public/AvaTransitionUtils.h b/Engine/Plugins/Experimental/Avalanche/Source/AvalancheTransition/Public/AvaTransitionUtils.h
new file mode 100644
--- /dev/null
+++ b/Engine/Plugins/Experimental/Avalanche/Source/AvalancheTransition/Public/AvaTransitionUtils.h
@@ -0,0 +1,21 @@
+// Copyright Epic Games, Inc. All Rights Reserved.
+
+#pragma once
+
+#include "StateTreeTypes.h"
+
+namespace UE::AvaTransition
+{
+	/** Safely gets the Instance Data if the Struct type matches or nullptr if there is type mismatch */
+	template<typename InNodeType>
+	static typename InNodeType::FInstanceDataType* TryGetInstanceData(const InNodeType& InNode, FStateTreeDataView InInstanceDataView)
+	{
+		// Current safety check in case this is being called prior to Instance Data types getting fixed. (e.g. in UStateTree::PostLoad)
+		const UStruct* InstanceDataType = InInstanceDataView.GetStruct();
+		if (InstanceDataType && InstanceDataType->IsChildOf(InNodeType::FInstanceDataType::StaticStruct()))
+		{
+			return InInstanceDataView.GetMutablePtr<typename InNodeType::FInstanceDataType>();
+		}
+		return nullptr;
+	}
+}
diff --git a/Engine/Plugins/Experimental/Avalanche/Source/AvalancheTransition/Public/Conditions/AvaTransitionLayerCondition.h b/Engine/Plugins/Experimental/Avalanche/Source/AvalancheTransition/Public/Conditions/AvaTransitionLayerCondition.h
--- a/Engine/Plugins/Experimental/Avalanche/Source/AvalancheTransition/Public/Conditions/AvaTransitionLayerCondition.h
+++ b/Engine/Plugins/Experimental/Avalanche/Source/AvalancheTransition/Public/Conditions/AvaTransitionLayerCondition.h
@@ -9,19 +9,49 @@
 
 struct FAvaTransitionContext;
 
+USTRUCT()
+struct FAvaTransitionLayerConditionInstanceData
+{
+	GENERATED_BODY()
+
+	UPROPERTY(EditAnywhere, Category="Parameter")
+	EAvaTransitionLayerCompareType LayerType = EAvaTransitionLayerCompareType::Same;
+
+	UPROPERTY(EditAnywhere, Category="Parameter", meta=(EditCondition="LayerType==EAvaTransitionLayerCompareType::MatchingTag", EditConditionHides))
+	FAvaTagHandle SpecificLayer;
+};
+
 USTRUCT(Category="Transition Logic", meta=(Hidden))
 struct AVALANCHETRANSITION_API FAvaTransitionLayerCondition : public FAvaTransitionCondition
 {
 	GENERATED_BODY()
 
+	using FInstanceDataType = FAvaTransitionLayerConditionInstanceData;
+
+	PRAGMA_DISABLE_DEPRECATION_WARNINGS
+	FAvaTransitionLayerCondition() = default;
+	virtual ~FAvaTransitionLayerCondition() override = default;
+	FAvaTransitionLayerCondition(const FAvaTransitionLayerCondition&) = default;
+	FAvaTransitionLayerCondition(FAvaTransitionLayerCondition&&) = default;
+	FAvaTransitionLayerCondition& operator=(const FAvaTransitionLayerCondition&) = default;
+	FAvaTransitionLayerCondition& operator=(FAvaTransitionLayerCondition&&) = default;
+	PRAGMA_ENABLE_DEPRECATION_WARNINGS
+
+	//~ Begin FStateTreeNodeBase
+	virtual const UStruct* GetInstanceDataType() const override { return FInstanceDataType::StaticStruct(); }
+	virtual void PostLoad(FStateTreeDataView InInstanceDataView) override;
+	//~ End FStateTreeNodeBase
+
 	/** Gets all the Behavior Instances that match the Layer Query. Always excludes the Instance belonging to this Transition */
 	TArray<const FAvaTransitionBehaviorInstance*> QueryBehaviorInstances(const FStateTreeExecutionContext& InContext) const;
 
-	FText GetLayerQueryText() const;
+	FText GetLayerQueryText(const FInstanceDataType& InInstanceData) const;
 
-	UPROPERTY(EditAnywhere, Category="Parameter")
-	EAvaTransitionLayerCompareType LayerType = EAvaTransitionLayerCompareType::Same;
+	UE_DEPRECATED(5.5, "LayerType has been moved to Instance Data")
+	UPROPERTY(meta = (DeprecatedProperty, DeprecationMessage = "Use the Instance Data LayerType instead"))
+	EAvaTransitionLayerCompareType LayerType_DEPRECATED = EAvaTransitionLayerCompareType::None;
 
-	UPROPERTY(EditAnywhere, Category="Parameter", meta=(EditCondition="LayerType==EAvaTransitionLayerCompareType::MatchingTag", EditConditionHides))
-	FAvaTagHandle SpecificLayer;
+	UE_DEPRECATED(5.5, "SpecificLayer has been moved to Instance Data")
+	UPROPERTY(meta = (DeprecatedProperty, DeprecationMessage = "Use the Instance Data SpecificLayer instead"))
+	FAvaTagHandle SpecificLayer_DEPRECATED;
 };
diff --git a/Engine/Plugins/Experimental/Avalanche/Source/AvalancheTransition/Public/Conditions/AvaTransitionLayerMatchCondition.h b/Engine/Plugins/Experimental/Avalanche/Source/AvalancheTransition/Public/Conditions/AvaTransitionLayerMatchCondition.h
--- a/Engine/Plugins/Experimental/Avalanche/Source/AvalancheTransition/Public/Conditions/AvaTransitionLayerMatchCondition.h
+++ b/Engine/Plugins/Experimental/Avalanche/Source/AvalancheTransition/Public/Conditions/AvaTransitionLayerMatchCondition.h
@@ -6,7 +6,7 @@
 #include "AvaTransitionLayerMatchCondition.generated.h"
 
 USTRUCT()
-struct FAvaTransitionLayerMatchConditionInstanceData
+struct FAvaTransitionLayerMatchConditionInstanceData : public FAvaTransitionLayerConditionInstanceData
 {
 	GENERATED_BODY()
 };
diff --git a/Engine/Plugins/Experimental/Avalanche/Source/AvalancheTransition/Public/Conditions/AvaTransitionSceneMatchCondition.h b/Engine/Plugins/Experimental/Avalanche/Source/AvalancheTransition/Public/Conditions/AvaTransitionSceneMatchCondition.h
--- a/Engine/Plugins/Experimental/Avalanche/Source/AvalancheTransition/Public/Conditions/AvaTransitionSceneMatchCondition.h
+++ b/Engine/Plugins/Experimental/Avalanche/Source/AvalancheTransition/Public/Conditions/AvaTransitionSceneMatchCondition.h
@@ -7,9 +7,12 @@
 #include "AvaTransitionSceneMatchCondition.generated.h"
 
 USTRUCT()
-struct FAvaTransitionSceneMatchConditionInstanceData
+struct FAvaTransitionSceneMatchConditionInstanceData : public FAvaTransitionLayerConditionInstanceData
 {
 	GENERATED_BODY()
+
+	UPROPERTY(EditAnywhere, Category="Parameter")
+    EAvaTransitionComparisonResult SceneComparisonType = EAvaTransitionComparisonResult::None;
 };
 
 USTRUCT(DisplayName="Compare other Scene in Layer", Category="Transition Logic")
@@ -19,17 +22,28 @@ struct AVALANCHETRANSITION_API FAvaTransitionSceneMatchCondition : public FAvaTr
 
 	using FInstanceDataType = FAvaTransitionSceneMatchConditionInstanceData;
 
+	PRAGMA_DISABLE_DEPRECATION_WARNINGS
+	FAvaTransitionSceneMatchCondition() = default;
+	virtual ~FAvaTransitionSceneMatchCondition() override = default;
+	FAvaTransitionSceneMatchCondition(const FAvaTransitionSceneMatchCondition&) = default;
+	FAvaTransitionSceneMatchCondition(FAvaTransitionSceneMatchCondition&&) = default;
+	FAvaTransitionSceneMatchCondition& operator=(const FAvaTransitionSceneMatchCondition&) = default;
+	FAvaTransitionSceneMatchCondition& operator=(FAvaTransitionSceneMatchCondition&&) = default;
+	PRAGMA_ENABLE_DEPRECATION_WARNINGS
+
 	//~ Begin FStateTreeNodeBase
 #if WITH_EDITOR
 	virtual FText GetDescription(const FGuid& InId, FStateTreeDataView InInstanceDataView, const IStateTreeBindingLookup& InBindingLookup, EStateTreeNodeFormatting InFormatting) const override;
 #endif
 	virtual const UStruct* GetInstanceDataType() const override { return FInstanceDataType::StaticStruct(); }
+	virtual void PostLoad(FStateTreeDataView InInstanceDataView) override;
 	//~ End FStateTreeNodeBase
 
 	//~ Begin FStateTreeConditionBase
 	virtual bool TestCondition(FStateTreeExecutionContext& InContext) const override;
 	//~ End FStateTreeConditionBase
 
-	UPROPERTY(EditAnywhere, Category="Parameter")
-	EAvaTransitionComparisonResult SceneComparisonType = EAvaTransitionComparisonResult::None;
+	UE_DEPRECATED(5.5, "SceneComparisonType has been moved to Instance Data")
+	UPROPERTY(meta = (DeprecatedProperty, DeprecationMessage = "Use the Instance Data SceneComparisonType instead"))
+	EAvaTransitionComparisonResult SceneComparisonType_DEPRECATED = EAvaTransitionComparisonResult::None;
 };
diff --git a/Engine/Plugins/Experimental/Avalanche/Source/AvalancheTransition/Public/Conditions/AvaTransitionStateMatchCondition.h b/Engine/Plugins/Experimental/Avalanche/Source/AvalancheTransition/Public/Conditions/AvaTransitionStateMatchCondition.h
--- a/Engine/Plugins/Experimental/Avalanche/Source/AvalancheTransition/Public/Conditions/AvaTransitionStateMatchCondition.h
+++ b/Engine/Plugins/Experimental/Avalanche/Source/AvalancheTransition/Public/Conditions/AvaTransitionStateMatchCondition.h
@@ -7,9 +7,12 @@
 #include "AvaTransitionStateMatchCondition.generated.h"
 
 USTRUCT()
-struct FAvaTransitionStateMatchConditionInstanceData
+struct FAvaTransitionStateMatchConditionInstanceData : public FAvaTransitionLayerConditionInstanceData
 {
 	GENERATED_BODY()
+
+	UPROPERTY(EditAnywhere, Category="Parameter")
+	EAvaTransitionRunState TransitionState = EAvaTransitionRunState::Running;
 };
 
 USTRUCT(DisplayName="Check State of other Scene in Layer", Category="Transition Logic")
@@ -19,17 +22,28 @@ struct AVALANCHETRANSITION_API FAvaTransitionStateMatchCondition : public FAvaTr
 
 	using FInstanceDataType = FAvaTransitionStateMatchConditionInstanceData;
 
+	PRAGMA_DISABLE_DEPRECATION_WARNINGS
+	FAvaTransitionStateMatchCondition() = default;
+	virtual ~FAvaTransitionStateMatchCondition() override = default;
+	FAvaTransitionStateMatchCondition(const FAvaTransitionStateMatchCondition&) = default;
+	FAvaTransitionStateMatchCondition(FAvaTransitionStateMatchCondition&&) = default;
+	FAvaTransitionStateMatchCondition& operator=(const FAvaTransitionStateMatchCondition&) = default;
+	FAvaTransitionStateMatchCondition& operator=(FAvaTransitionStateMatchCondition&&) = default;
+	PRAGMA_ENABLE_DEPRECATION_WARNINGS
+
 	//~ Begin FStateTreeNodeBase
 #if WITH_EDITOR
 	virtual FText GetDescription(const FGuid& InId, FStateTreeDataView InInstanceDataView, const IStateTreeBindingLookup& InBindingLookup, EStateTreeNodeFormatting InFormatting) const override;
 #endif
 	virtual const UStruct* GetInstanceDataType() const override { return FAvaTransitionStateMatchConditionInstanceData::StaticStruct(); }
+	virtual void PostLoad(FStateTreeDataView InInstanceDataView) override;
 	//~ End FStateTreeNodeBase
 
 	//~ Begin FStateTreeConditionBase
 	virtual bool TestCondition(FStateTreeExecutionContext& InContext) const override;
 	//~ End FStateTreeConditionBase
 
-	UPROPERTY(EditAnywhere, Category="Parameter")
-	EAvaTransitionRunState TransitionState = EAvaTransitionRunState::Running;
+	UE_DEPRECATED(5.5, "TransitionState has been moved to Instance Data")
+	UPROPERTY(meta = (DeprecatedProperty, DeprecationMessage = "Use the Instance Data TransitionState instead"))
+	EAvaTransitionRunState TransitionState_DEPRECATED = EAvaTransitionRunState::Unknown;
 };
diff --git a/Engine/Plugins/Experimental/Avalanche/Source/AvalancheTransition/Public/Conditions/AvaTransitionTypeMatchCondition.h b/Engine/Plugins/Experimental/Avalanche/Source/AvalancheTransition/Public/Conditions/AvaTransitionTypeMatchCondition.h
--- a/Engine/Plugins/Experimental/Avalanche/Source/AvalancheTransition/Public/Conditions/AvaTransitionTypeMatchCondition.h
+++ b/Engine/Plugins/Experimental/Avalanche/Source/AvalancheTransition/Public/Conditions/AvaTransitionTypeMatchCondition.h
@@ -10,6 +10,9 @@ USTRUCT()
 struct FAvaTransitionTypeMatchConditionInstanceData
 {
 	GENERATED_BODY()
+
+	UPROPERTY(EditAnywhere, Category="Parameter")
+	EAvaTransitionType TransitionType = EAvaTransitionType::In;
 };
 
 USTRUCT(DisplayName="My Transition Type is", Category="Transition Logic")
@@ -19,19 +22,28 @@ struct AVALANCHETRANSITION_API FAvaTransitionTypeMatchCondition : public FAvaTra
 
 	using FInstanceDataType = FAvaTransitionTypeMatchConditionInstanceData;
 
+	PRAGMA_DISABLE_DEPRECATION_WARNINGS
 	FAvaTransitionTypeMatchCondition() = default;
+	virtual ~FAvaTransitionTypeMatchCondition() override = default;
+	FAvaTransitionTypeMatchCondition(const FAvaTransitionTypeMatchCondition&) = default;
+	FAvaTransitionTypeMatchCondition(FAvaTransitionTypeMatchCondition&&) = default;
+	FAvaTransitionTypeMatchCondition& operator=(const FAvaTransitionTypeMatchCondition&) = default;
+	FAvaTransitionTypeMatchCondition& operator=(FAvaTransitionTypeMatchCondition&&) = default;
+	PRAGMA_ENABLE_DEPRECATION_WARNINGS
 
 	//~ Begin FStateTreeNodeBase
 #if WITH_EDITOR
 	virtual FText GetDescription(const FGuid& InId, FStateTreeDataView InInstanceDataView, const IStateTreeBindingLookup& InBindingLookup, EStateTreeNodeFormatting InFormatting) const override;
 #endif
 	virtual const UStruct* GetInstanceDataType() const override { return FInstanceDataType::StaticStruct(); }
+	virtual void PostLoad(FStateTreeDataView InInstanceDataView) override;
 	//~ End FStateTreeNodeBase
 
 	//~ Begin FStateTreeConditionBase
 	virtual bool TestCondition(FStateTreeExecutionContext& InContext) const override;
 	//~ End FStateTreeConditionBase
 
-	UPROPERTY(EditAnywhere, Category="Parameter")
-	EAvaTransitionType TransitionType = EAvaTransitionType::In;
+	UE_DEPRECATED(5.5, "TransitionType has been moved to Instance Data")
+	UPROPERTY(meta = (DeprecatedProperty, DeprecationMessage = "Use the Instance Data TransitionType instead"))
+	EAvaTransitionType TransitionType_DEPRECATED = EAvaTransitionType::None;
 };
diff --git a/Engine/Plugins/Experimental/Avalanche/Source/AvalancheTransition/Public/Tasks/AvaTransitionDelayTask.h b/Engine/Plugins/Experimental/Avalanche/Source/AvalancheTransition/Public/Tasks/AvaTransitionDelayTask.h
--- a/Engine/Plugins/Experimental/Avalanche/Source/AvalancheTransition/Public/Tasks/AvaTransitionDelayTask.h
+++ b/Engine/Plugins/Experimental/Avalanche/Source/AvalancheTransition/Public/Tasks/AvaTransitionDelayTask.h
@@ -10,6 +10,10 @@ struct FAvaTransitionDelayTaskInstanceData
 {
 	GENERATED_BODY()
 
+	/** Delay in seconds before the task ends. */
+	UPROPERTY(EditAnywhere, Category="Parameter", meta=(ClampMin="0.0"))
+	float Duration = 0.5f;
+
 	/** Internal countdown in seconds. */
 	float RemainingTime = 0.f;
 };
@@ -21,11 +25,21 @@ struct AVALANCHETRANSITION_API FAvaTransitionDelayTask : public FAvaTransitionTa
 
 	using FInstanceDataType = FAvaTransitionDelayTaskInstanceData;
 
+	PRAGMA_DISABLE_DEPRECATION_WARNINGS
+	FAvaTransitionDelayTask() = default;
+	virtual ~FAvaTransitionDelayTask() override = default;
+	FAvaTransitionDelayTask(const FAvaTransitionDelayTask&) = default;
+	FAvaTransitionDelayTask(FAvaTransitionDelayTask&&) = default;
+	FAvaTransitionDelayTask& operator=(const FAvaTransitionDelayTask&) = default;
+	FAvaTransitionDelayTask& operator=(FAvaTransitionDelayTask&&) = default;
+	PRAGMA_ENABLE_DEPRECATION_WARNINGS
+
 	//~ Begin FStateTreeNodeBase
 #if WITH_EDITOR
 	virtual FText GetDescription(const FGuid& InId, FStateTreeDataView InInstanceDataView, const IStateTreeBindingLookup& InBindingLookup, EStateTreeNodeFormatting InFormatting) const override;
 #endif
 	virtual const UStruct* GetInstanceDataType() const override { return FInstanceDataType::StaticStruct(); }
+	virtual void PostLoad(FStateTreeDataView InInstanceDataView) override;
 	//~ End FStateTreeNodeBase
 
 	//~ Begin FStateTreeTaskBase
@@ -33,7 +47,7 @@ struct AVALANCHETRANSITION_API FAvaTransitionDelayTask : public FAvaTransitionTa
 	virtual EStateTreeRunStatus Tick(FStateTreeExecutionContext& InContext, const float InDeltaTime) const override;
 	//~ End FStateTreeTaskBase
 
-	/** Delay in seconds before the task ends. */
-	UPROPERTY(EditAnywhere, Category="Parameter", meta=(ClampMin="0.0"))
-	float Duration = 0.5f;
+	UE_DEPRECATED(5.5, "Duration has been moved to Instance Data")
+	UPROPERTY(meta = (DeprecatedProperty, DeprecationMessage = "Use the Instance Data Duration instead"))
+	float Duration_DEPRECATED = -1.f;
 };
diff --git a/Engine/Plugins/Experimental/Avalanche/Source/AvalancheTransition/Public/Tasks/AvaTransitionLayerTask.h b/Engine/Plugins/Experimental/Avalanche/Source/AvalancheTransition/Public/Tasks/AvaTransitionLayerTask.h
--- a/Engine/Plugins/Experimental/Avalanche/Source/AvalancheTransition/Public/Tasks/AvaTransitionLayerTask.h
+++ b/Engine/Plugins/Experimental/Avalanche/Source/AvalancheTransition/Public/Tasks/AvaTransitionLayerTask.h
@@ -7,19 +7,49 @@
 #include "AvaTransitionTask.h"
 #include "AvaTransitionLayerTask.generated.h"
 
+USTRUCT()
+struct FAvaTransitionLayerTaskInstanceData
+{
+	GENERATED_BODY()
+
+	UPROPERTY(EditAnywhere, Category="Parameter")
+	EAvaTransitionLayerCompareType LayerType = EAvaTransitionLayerCompareType::Same;
+
+	UPROPERTY(EditAnywhere, Category="Parameter", meta=(EditCondition="LayerType==EAvaTransitionLayerCompareType::MatchingTag", EditConditionHides))
+	FAvaTagHandle SpecificLayer;
+};
+
 USTRUCT(meta=(Hidden))
 struct AVALANCHETRANSITION_API FAvaTransitionLayerTask : public FAvaTransitionTask
 {
 	GENERATED_BODY()
 
+	using FInstanceDataType = FAvaTransitionLayerTaskInstanceData;
+
+	PRAGMA_DISABLE_DEPRECATION_WARNINGS
+	FAvaTransitionLayerTask() = default;
+	virtual ~FAvaTransitionLayerTask() override = default;
+	FAvaTransitionLayerTask(const FAvaTransitionLayerTask&) = default;
+	FAvaTransitionLayerTask(FAvaTransitionLayerTask&&) = default;
+	FAvaTransitionLayerTask& operator=(const FAvaTransitionLayerTask&) = default;
+	FAvaTransitionLayerTask& operator=(FAvaTransitionLayerTask&&) = default;
+	PRAGMA_ENABLE_DEPRECATION_WARNINGS
+
+	//~ Begin FStateTreeNodeBase
+	virtual const UStruct* GetInstanceDataType() const override { return FInstanceDataType::StaticStruct(); }
+	virtual void PostLoad(FStateTreeDataView InInstanceDataView) override;
+	//~ End FStateTreeNodeBase
+
+	FText GetLayerQueryText(const FInstanceDataType& InInstanceData) const;
+
 	/** Gets all the Behavior Instances that match the Layer Query. Always excludes the Instance belonging to this Transition */
 	TArray<const FAvaTransitionBehaviorInstance*> QueryBehaviorInstances(const FStateTreeExecutionContext& InContext) const;
 
-	FText GetLayerQueryText() const;
+	UE_DEPRECATED(5.5, "LayerType has been moved to Instance Data")
+	UPROPERTY(meta = (DeprecatedProperty, DeprecationMessage = "Use the Instance Data LayerType instead"))
+	EAvaTransitionLayerCompareType LayerType_DEPRECATED = EAvaTransitionLayerCompareType::None;
 
-	UPROPERTY(EditAnywhere, Category="Parameter")
-	EAvaTransitionLayerCompareType LayerType = EAvaTransitionLayerCompareType::Same;
-
-	UPROPERTY(EditAnywhere, Category="Parameter", meta=(EditCondition="LayerType==EAvaTransitionLayerCompareType::MatchingTag", EditConditionHides))
-	FAvaTagHandle SpecificLayer;
+	UE_DEPRECATED(5.5, "SpecificLayer has been moved to Instance Data")
+	UPROPERTY(meta = (DeprecatedProperty, DeprecationMessage = "Use the Instance Data SpecificLayer instead"))
+	FAvaTagHandle SpecificLayer_DEPRECATED;
 };
diff --git a/Engine/Plugins/Experimental/Avalanche/Source/AvalancheTransition/Public/Tasks/AvaTransitionWaitForLayerTask.h b/Engine/Plugins/Experimental/Avalanche/Source/AvalancheTransition/Public/Tasks/AvaTransitionWaitForLayerTask.h
--- a/Engine/Plugins/Experimental/Avalanche/Source/AvalancheTransition/Public/Tasks/AvaTransitionWaitForLayerTask.h
+++ b/Engine/Plugins/Experimental/Avalanche/Source/AvalancheTransition/Public/Tasks/AvaTransitionWaitForLayerTask.h
@@ -6,7 +6,7 @@
 #include "AvaTransitionWaitForLayerTask.generated.h"
 
 USTRUCT()
-struct FAvaTransitionWaitForLayerTaskInstanceData
+struct FAvaTransitionWaitForLayerTaskInstanceData : public FAvaTransitionLayerTaskInstanceData
 {
 	GENERATED_BODY()
 };
