MotionDesign : ClonerEffector
- Add MotionDesign editor extension for cloner/effector
- Add keyboard shortcuts to enable/disable cloners : SHIFT+C / ALT+C
- Add keyboard shortcuts to enable/disable effectors : SHIFT+E / ALT+E
- When nothing is selected, it will target the whole world with the specific action

diff --git a/Engine/Plugins/Experimental/Avalanche/Source/AvalancheEditor/AvalancheEditor.Build.cs b/Engine/Plugins/Experimental/Avalanche/Source/AvalancheEditor/AvalancheEditor.Build.cs
--- a/Engine/Plugins/Experimental/Avalanche/Source/AvalancheEditor/AvalancheEditor.Build.cs
+++ b/Engine/Plugins/Experimental/Avalanche/Source/AvalancheEditor/AvalancheEditor.Build.cs
@@ -55,6 +55,7 @@ public class AvalancheEditor : ModuleRules
 				"AvalancheViewport",
 				"BlueprintGraph",
 				"CinematicCamera",
+				"ClonerEffector",
 				"CommonMenuExtensions",
 				"ContentBrowser",
 				"CustomDetailsView",
diff --git a/Engine/Plugins/Experimental/Avalanche/Source/AvalancheEditor/Private/AvaEditorCommands.cpp b/Engine/Plugins/Experimental/Avalanche/Source/AvalancheEditor/Private/AvaEditorCommands.cpp
--- a/Engine/Plugins/Experimental/Avalanche/Source/AvalancheEditor/Private/AvaEditorCommands.cpp
+++ b/Engine/Plugins/Experimental/Avalanche/Source/AvalancheEditor/Private/AvaEditorCommands.cpp
@@ -24,6 +24,7 @@ void FAvaEditorCommands::RegisterCommands()
 	RegisterToolsCommands();
 	RegisterAdvancedRenamerCommands();
 	RegisterAnimatorCommands();
+	RegisterClonerEffectorCommands();
 }
 
 void FAvaEditorCommands::RegisterViewportCommands()
@@ -429,4 +430,31 @@ void FAvaEditorCommands::RegisterAnimatorCommands()
 		, FInputChord(EKeys::U, EModifierKey::Control))
 }
 
+void FAvaEditorCommands::RegisterClonerEffectorCommands()
+{
+	UI_COMMAND(DisableEffectors
+		, "Disable Effectors"
+		, "Disable effectors on the selected actors in the level"
+		, EUserInterfaceActionType::Button
+		, FInputChord(EKeys::E, EModifierKey::Alt))
+
+	UI_COMMAND(EnableEffectors
+		, "Enable Effectors"
+		, "Enable effectors on the selected actors in the level"
+		, EUserInterfaceActionType::Button
+		, FInputChord(EKeys::E, EModifierKey::Shift))
+
+	UI_COMMAND(DisableCloners
+		, "Disable Cloners"
+		, "Disable cloners on the selected actors in the level"
+		, EUserInterfaceActionType::Button
+		, FInputChord(EKeys::C, EModifierKey::Alt))
+
+	UI_COMMAND(EnableCloners
+		, "Enable Cloners"
+		, "Enable cloners on the selected actors in the level"
+		, EUserInterfaceActionType::Button
+		, FInputChord(EKeys::C, EModifierKey::Shift))
+}
+
 #undef LOCTEXT_NAMESPACE
diff --git a/Engine/Plugins/Experimental/Avalanche/Source/AvalancheEditor/Private/AvaEditorCommands.h b/Engine/Plugins/Experimental/Avalanche/Source/AvalancheEditor/Private/AvaEditorCommands.h
--- a/Engine/Plugins/Experimental/Avalanche/Source/AvalancheEditor/Private/AvaEditorCommands.h
+++ b/Engine/Plugins/Experimental/Avalanche/Source/AvalancheEditor/Private/AvaEditorCommands.h
@@ -16,6 +16,7 @@ public:
 	void RegisterToolsCommands();
 	void RegisterAdvancedRenamerCommands();
 	void RegisterAnimatorCommands();
+	void RegisterClonerEffectorCommands();
 
 	/** Switches between Camera "2D" viewport and 3D viewport. */
 	TSharedPtr<FUICommandInfo> SwitchViewports;
@@ -98,4 +99,10 @@ public:
 	// Animator
 	TSharedPtr<FUICommandInfo> DisableAnimators;
 	TSharedPtr<FUICommandInfo> EnableAnimators;
+
+	// ClonerEffector
+	TSharedPtr<FUICommandInfo> DisableCloners;
+	TSharedPtr<FUICommandInfo> EnableCloners;
+	TSharedPtr<FUICommandInfo> DisableEffectors;
+	TSharedPtr<FUICommandInfo> EnableEffectors;
 };
diff --git a/Engine/Plugins/Experimental/Avalanche/Source/AvalancheEditor/Private/AvaEditorIntegration.cpp b/Engine/Plugins/Experimental/Avalanche/Source/AvalancheEditor/Private/AvaEditorIntegration.cpp
--- a/Engine/Plugins/Experimental/Avalanche/Source/AvalancheEditor/Private/AvaEditorIntegration.cpp
+++ b/Engine/Plugins/Experimental/Avalanche/Source/AvalancheEditor/Private/AvaEditorIntegration.cpp
@@ -2,6 +2,7 @@
 
 #include "AvaEditorIntegration.h"
 #include "AdvancedRenamer/AvaAdvancedRenamerExtension.h"
+#include "ClonerEffector/AvaClonerEffectorExtension.h"
 #include "ComponentVisualizer/AvaComponentVisualizerExtension.h"
 #include "DetailView/AvaDetailsExtension.h"
 #include "EditorProvider/AvaEditorProvider.h"
@@ -38,5 +39,6 @@ TSharedRef<IAvaEditor> FAvaLevelEditorIntegration::BuildEditor()
 		.AddExtension<FAvaTransitionExtension>()
 		.AddExtension<FAvaAdvancedRenamerExtension>()
 		.AddExtension<FAvaComponentVisualizerExtension>()
+		.AddExtension<FAvaClonerEffectorExtension>()
 		.Build();
 }
diff --git a/Engine/Plugins/Experimental/Avalanche/Source/AvalancheEditor/Private/ClonerEffector/AvaClonerEffectorExtension.cpp b/Engine/Plugins/Experimental/Avalanche/Source/AvalancheEditor/Private/ClonerEffector/AvaClonerEffectorExtension.cpp
new file mode 100644
--- /dev/null
+++ b/Engine/Plugins/Experimental/Avalanche/Source/AvalancheEditor/Private/ClonerEffector/AvaClonerEffectorExtension.cpp
@@ -0,0 +1,130 @@
+// Copyright Epic Games, Inc. All Rights Reserved.
+
+#include "AvaClonerEffectorExtension.h"
+
+#include "AvaEditorCommands.h"
+#include "Cloner/CEClonerComponent.h"
+#include "EditorModeManager.h"
+#include "Effector/CEEffectorComponent.h"
+#include "Selection.h"
+#include "Subsystems/CEClonerSubsystem.h"
+#include "Subsystems/CEEffectorSubsystem.h"
+
+#define LOCTEXT_NAMESPACE "AvaClonerEffectorExtension"
+
+FAvaClonerEffectorExtension::FAvaClonerEffectorExtension()
+	: ClonerEffectorCommands(MakeShared<FUICommandList>())
+{
+}
+
+void FAvaClonerEffectorExtension::BindCommands(const TSharedRef<FUICommandList>& InCommandList)
+{
+	InCommandList->Append(ClonerEffectorCommands);
+
+	const FAvaEditorCommands& EditorCommands = FAvaEditorCommands::Get();
+
+	ClonerEffectorCommands->MapAction(EditorCommands.DisableEffectors
+		, FExecuteAction::CreateSP(this, &FAvaClonerEffectorExtension::EnableEffectors, false));
+
+	ClonerEffectorCommands->MapAction(EditorCommands.EnableEffectors
+		, FExecuteAction::CreateSP(this, &FAvaClonerEffectorExtension::EnableEffectors, true));
+
+	ClonerEffectorCommands->MapAction(EditorCommands.DisableCloners
+		, FExecuteAction::CreateSP(this, &FAvaClonerEffectorExtension::EnableCloners, false));
+
+	ClonerEffectorCommands->MapAction(EditorCommands.EnableCloners
+		, FExecuteAction::CreateSP(this, &FAvaClonerEffectorExtension::EnableCloners, true));
+}
+
+TSet<AActor*> FAvaClonerEffectorExtension::GetSelectedActors() const
+{
+	TSet<AActor*> SelectedActors;
+
+	const FEditorModeTools* ModeTools = GetEditorModeTools();
+	if (!ModeTools)
+	{
+		return SelectedActors;
+	}
+
+	const UTypedElementSelectionSet* SelectionSet = ModeTools->GetEditorSelectionSet();
+	if (!SelectionSet)
+	{
+		return SelectedActors;
+	}
+
+	SelectedActors.Append(SelectionSet->GetSelectedObjects<AActor>());
+
+	return SelectedActors;
+}
+
+void FAvaClonerEffectorExtension::EnableEffectors(bool bInEnable) const
+{
+	TSet<AActor*> SelectedActors = GetSelectedActors();
+
+	TSet<UCEEffectorComponent*> EffectorComponents;
+	for (AActor* SelectedActor : SelectedActors)
+	{
+		if (IsValid(SelectedActor))
+		{
+			TArray<UCEEffectorComponent*> Components;
+			SelectedActor->GetComponents(Components, /** IncludeChildren */false);
+			EffectorComponents.Append(Components);
+		}
+	}
+
+	UCEEffectorSubsystem* EffectorSubsystem = UCEEffectorSubsystem::Get();
+
+	if (!EffectorSubsystem)
+	{
+		return;
+	}
+
+	constexpr bool bTransact = true;
+
+	// If nothing selected, target all animators in level
+	if (!SelectedActors.IsEmpty())
+	{
+		EffectorSubsystem->SetEffectorsEnabled(EffectorComponents, bInEnable, bTransact);
+	}
+	else
+	{
+		EffectorSubsystem->SetLevelEffectorsEnabled(GetWorld(), bInEnable, bTransact);
+	}
+}
+
+void FAvaClonerEffectorExtension::EnableCloners(bool bInEnable) const
+{
+	TSet<AActor*> SelectedActors = GetSelectedActors();
+
+	TSet<UCEClonerComponent*> ClonerComponents;
+	for (AActor* SelectedActor : SelectedActors)
+	{
+		if (IsValid(SelectedActor))
+		{
+			TArray<UCEClonerComponent*> Components;
+			SelectedActor->GetComponents(Components, /** IncludeChildren */false);
+			ClonerComponents.Append(Components);
+		}
+	}
+
+	UCEClonerSubsystem* ClonerSubsystem = UCEClonerSubsystem::Get();
+
+	if (!ClonerSubsystem)
+	{
+		return;
+	}
+
+	constexpr bool bTransact = true;
+
+	// If nothing selected, target all animators in level
+	if (!SelectedActors.IsEmpty())
+	{
+		ClonerSubsystem->SetClonersEnabled(ClonerComponents, bInEnable, bTransact);
+	}
+	else
+	{
+		ClonerSubsystem->SetLevelClonersEnabled(GetWorld(), bInEnable, bTransact);
+	}
+}
+
+#undef LOCTEXT_NAMESPACE
diff --git a/Engine/Plugins/Experimental/Avalanche/Source/AvalancheEditor/Private/ClonerEffector/AvaClonerEffectorExtension.h b/Engine/Plugins/Experimental/Avalanche/Source/AvalancheEditor/Private/ClonerEffector/AvaClonerEffectorExtension.h
new file mode 100644
--- /dev/null
+++ b/Engine/Plugins/Experimental/Avalanche/Source/AvalancheEditor/Private/ClonerEffector/AvaClonerEffectorExtension.h
@@ -0,0 +1,30 @@
+// Copyright Epic Games, Inc. All Rights Reserved.
+
+#pragma once
+
+#include "Containers/Set.h"
+#include "IAvaEditorExtension.h"
+#include "Templates/SharedPointer.h"
+
+class AActor;
+class FUICommandList;
+
+class FAvaClonerEffectorExtension : public FAvaEditorExtension
+{
+public:
+	UE_AVA_INHERITS(FAvaClonerEffectorExtension, FAvaEditorExtension);
+
+	FAvaClonerEffectorExtension();
+
+	//~ Begin IAvaEditorExtension
+	virtual void BindCommands(const TSharedRef<FUICommandList>& InCommandList) override;
+	//~ End IAvaEditorExtension
+
+private:
+	TSet<AActor*> GetSelectedActors() const;
+
+	void EnableEffectors(bool bInEnable) const;
+	void EnableCloners(bool bInEnable) const;
+
+	TSharedRef<FUICommandList> ClonerEffectorCommands;
+};
